UE4 - Reworked virtual memory API to support page caches and heterogeneous memory architectures. Switch is the only platform with substantive changes, the rest are just API upgrades.

diff --git a/Engine/Source/Runtime/Core/Private/Android/AndroidPlatformMemory.cpp b/Engine/Source/Runtime/Core/Private/Android/AndroidPlatformMemory.cpp
--- a/Engine/Source/Runtime/Core/Private/Android/AndroidPlatformMemory.cpp
+++ b/Engine/Source/Runtime/Core/Private/Android/AndroidPlatformMemory.cpp
@@ -399,11 +399,68 @@ void FAndroidPlatformMemory::BinnedFreeToOS(void* Ptr, SIZE_T Size)
 	}
 }
 
-bool FAndroidPlatformMemory::MemoryRangeDecommit(void* Ptr, SIZE_T Size)
+size_t FAndroidPlatformMemory::FPlatformVirtualMemoryBlock::GetVirtualSizeAlignment()
 {
-	return madvise(Ptr, Size, MADV_DONTNEED) == 0;
+	static SIZE_T OSPageSize = FPlatformMemory::GetConstants().PageSize;
+	return OSPageSize;
 }
 
+size_t FAndroidPlatformMemory::FPlatformVirtualMemoryBlock::GetCommitAlignment()
+{
+	static SIZE_T OSPageSize = FPlatformMemory::GetConstants().PageSize;
+	return OSPageSize;
+}
+
+FAndroidPlatformMemory::FPlatformVirtualMemoryBlock FAndroidPlatformMemory::FPlatformVirtualMemoryBlock::AllocateVirtual(size_t InSize, size_t InAlignment)
+{
+	FPlatformVirtualMemoryBlock Result;
+	InSize = Align(InSize, GetVirtualSizeAlignment());
+	Result.VMSizeDivVirtualSizeAlignment = InSize / GetVirtualSizeAlignment();
+
+	size_t Alignment = FMath::Max(InAlignment, GetVirtualSizeAlignment());
+	check(Alignment <= GetVirtualSizeAlignment());
+
+	Result.Ptr = mmap(nullptr, Result.GetActualSize(), PROT_READ | PROT_WRITE, MAP_PRIVATE | MAP_ANON, -1, 0);
+	if (!LIKELY(Result.Ptr != MAP_FAILED))
+	{
+		FPlatformMemory::OnOutOfMemory(Result.GetActualSize(), InAlignment);
+	}
+	check(Result.Ptr && IsAligned(Result.Ptr, Alignment));
+	return Result;
+}
+
+
+
+void FAndroidPlatformMemory::FPlatformVirtualMemoryBlock::FreeVirtual()
+{
+	if (Ptr)
+	{
+		check(VMSizeDivVirtualSizeAlignment > 0);
+		if (munmap(Ptr, GetActualSize()) != 0)
+		{
+			// we can ran out of VMAs here
+			FPlatformMemory::OnOutOfMemory(GetActualSize(), 0);
+			// unreachable
+		}
+		Ptr = nullptr;
+		VMSizeDivVirtualSizeAlignment = 0;
+	}
+}
+
+void FAndroidPlatformMemory::FPlatformVirtualMemoryBlock::Commit(size_t InOffset, size_t InSize)
+{
+	check(IsAligned(InOffset, GetCommitAlignment()) && IsAligned(InSize, GetCommitAlignment()));
+	check(InOffset >= 0 && InSize >= 0 && InOffset + InSize <= GetActualSize() && Ptr);
+}
+
+void FAndroidPlatformMemory::FPlatformVirtualMemoryBlock::Decommit(size_t InOffset, size_t InSize)
+{
+	check(IsAligned(InOffset, GetCommitAlignment()) && IsAligned(InSize, GetCommitAlignment()));
+	check(InOffset >= 0 && InSize >= 0 && InOffset + InSize <= GetActualSize() && Ptr);
+	madvise(((uint8*)Ptr) + InOffset, InSize, MADV_DONTNEED);
+}
+
+
 /**
 * LLM uses these low level functions (LLMAlloc and LLMFree) to allocate memory. It grabs
 * the function pointers by calling FPlatformMemory::GetLLMAllocFunctions. If these functions
diff --git a/Engine/Source/Runtime/Core/Private/Apple/ApplePlatformMemory.cpp b/Engine/Source/Runtime/Core/Private/Apple/ApplePlatformMemory.cpp
--- a/Engine/Source/Runtime/Core/Private/Apple/ApplePlatformMemory.cpp
+++ b/Engine/Source/Runtime/Core/Private/Apple/ApplePlatformMemory.cpp
@@ -343,39 +343,278 @@ bool FApplePlatformMemory::PageProtect(void* const Ptr, const SIZE_T Size, const
 	return mprotect(Ptr, Size, static_cast<int32>(ProtectMode)) == 0;
 }
 
-void* FApplePlatformMemory::BinnedAllocFromOS( SIZE_T Size )
+#ifndef MALLOC_LEAKDETECTION
+	#define MALLOC_LEAKDETECTION 0
+#endif
+#define UE4_PLATFORM_REDUCE_NUMBER_OF_MAPS 0
+
+// check bookkeeping info against the passed in parameters in Debug and Development (the latter only in games and servers. also, only if leak detection is disabled, otherwise things are very slow)
+#define UE4_PLATFORM_SANITY_CHECK_OS_ALLOCATIONS			(UE_BUILD_DEBUG || (UE_BUILD_DEVELOPMENT && (UE_GAME || UE_SERVER) && !MALLOC_LEAKDETECTION))
+
+/** This structure is stored in the page after each OS allocation and checks that its properties are valid on Free. Should be less than the page size (4096 on all platforms we support) */
+struct FOSAllocationDescriptor
+{
+	enum class MagicType : uint64
+	{
+		Marker = 0xd0c233ccf493dfb0
+	};
+
+	/** Magic that makes sure that we are not passed a pointer somewhere into the middle of the allocation (and/or the structure wasn't stomped). */
+	MagicType	Magic;
+
+	/** This should include the descriptor itself. */
+	void*		PointerToUnmap;
+
+	/** This should include the total size of allocation, so after unmapping it everything is gone, including the descriptor */
+	SIZE_T		SizeToUnmap;
+
+	/** Debug info that makes sure that the correct size is preserved. */
+	SIZE_T		OriginalSizeAsPassed;
+};
+
+void* _Nullable FApplePlatformMemory::BinnedAllocFromOS(SIZE_T Size)
 {
-// Binned2 requires allocations to be BinnedPageSize-aligned. Simple mmap() does not guarantee this for recommended BinnedPageSize (64KB).
+	// Binned2 requires allocations to be BinnedPageSize-aligned. Simple mmap() does not guarantee this for recommended BinnedPageSize (64KB).
 #if USE_MALLOC_BINNED2
-    return FGenericPlatformMemory::BinnedAllocFromOS(Size);
+	static SIZE_T OSPageSize = FPlatformMemory::GetConstants().PageSize;
+	// guard against someone not passing size in whole pages
+	SIZE_T SizeInWholePages = (Size % OSPageSize) ? (Size + OSPageSize - (Size % OSPageSize)) : Size;
+	void* Pointer = nullptr;
+
+	// Binned expects OS allocations to be BinnedPageSize-aligned, and that page is at least 64KB. mmap() alone cannot do this, so carve out the needed chunks.
+	const SIZE_T ExpectedAlignment = FPlatformMemory::GetConstants().BinnedPageSize;
+	// Descriptor is only used if we're sanity checking. However, #ifdef'ing its use would make the code more fragile. Size needs to be at least one page.
+	const SIZE_T DescriptorSize = (UE4_PLATFORM_SANITY_CHECK_OS_ALLOCATIONS != 0) ? OSPageSize : 0;
+
+	SIZE_T ActualSizeMapped = SizeInWholePages + ExpectedAlignment;
+
+	// the remainder of the map will be used for the descriptor, if any.
+	// we always allocate at least one page more than needed
+	void* PointerWeGotFromMMap = mmap(nullptr, ActualSizeMapped, PROT_READ | PROT_WRITE, MAP_PRIVATE | MAP_ANON, -1, 0);
+	// store these values, to unmap later
+
+	Pointer = PointerWeGotFromMMap;
+	if (Pointer == MAP_FAILED)
+	{
+		FPlatformMemory::OnOutOfMemory(ActualSizeMapped, ExpectedAlignment);
+		// unreachable
+		return nullptr;
+	}
+
+	SIZE_T Offset = (reinterpret_cast<SIZE_T>(Pointer) % ExpectedAlignment);
+
+	// See if we need to unmap anything in the front. If the pointer happened to be aligned, we don't need to unmap anything.
+	if (LIKELY(Offset != 0))
+	{
+		// figure out how much to unmap before the alignment.
+		SIZE_T SizeToNextAlignedPointer = ExpectedAlignment - Offset;
+		void* AlignedPointer = reinterpret_cast<void*>(reinterpret_cast<SIZE_T>(Pointer) + SizeToNextAlignedPointer);
+
+		// do not unmap if we're trying to reduce the number of distinct maps, since holes prevent the Linux kernel from coalescing two adjoining mmap()s into a single VMA
+		if (!UE4_PLATFORM_REDUCE_NUMBER_OF_MAPS)
+		{
+			// unmap the part before
+			if (munmap(Pointer, SizeToNextAlignedPointer) != 0)
+			{
+				FPlatformMemory::OnOutOfMemory(SizeToNextAlignedPointer, ExpectedAlignment);
+				// unreachable
+				return nullptr;
+			}
+
+			// account for reduced mmaped size
+			ActualSizeMapped -= SizeToNextAlignedPointer;
+		}
+
+		// now, make it appear as if we initially got the allocation right
+		Pointer = AlignedPointer;
+	}
+
+	// at this point, Pointer is aligned at the expected alignment - either we lucked out on the initial allocation
+	// or we already got rid of the extra memory that was allocated in the front.
+	checkf((reinterpret_cast<SIZE_T>(Pointer) % ExpectedAlignment) == 0, TEXT("BinnedAllocFromOS(): Internal error: did not align the pointer as expected."));
+
+	// do not unmap if we're trying to reduce the number of distinct maps, since holes prevent the Linux kernel from coalescing two adjoining mmap()s into a single VMA
+	if (!UE4_PLATFORM_REDUCE_NUMBER_OF_MAPS)
+	{
+		// Now unmap the tail only, if any, but leave just enough space for the descriptor
+		void* TailPtr = reinterpret_cast<void*>(reinterpret_cast<SIZE_T>(Pointer) + SizeInWholePages + DescriptorSize);
+		SIZE_T TailSize = ActualSizeMapped - SizeInWholePages - DescriptorSize;
+
+		if (LIKELY(TailSize > 0))
+		{
+			if (munmap(TailPtr, TailSize) != 0)
+			{
+				FPlatformMemory::OnOutOfMemory(TailSize, ExpectedAlignment);
+				// unreachable
+				return nullptr;
+			}
+		}
+	}
+
+	// we're done with this allocation, fill in the descriptor with the info
+	if (LIKELY(DescriptorSize > 0))
+	{
+		FOSAllocationDescriptor* AllocDescriptor = reinterpret_cast<FOSAllocationDescriptor*>(reinterpret_cast<SIZE_T>(Pointer) + Size);
+		AllocDescriptor->Magic = FOSAllocationDescriptor::MagicType::Marker;
+		if (!UE4_PLATFORM_REDUCE_NUMBER_OF_MAPS)
+		{
+			AllocDescriptor->PointerToUnmap = Pointer;
+			AllocDescriptor->SizeToUnmap = SizeInWholePages + DescriptorSize;
+		}
+		else
+		{
+			AllocDescriptor->PointerToUnmap = PointerWeGotFromMMap;
+			AllocDescriptor->SizeToUnmap = ActualSizeMapped;
+		}
+		AllocDescriptor->OriginalSizeAsPassed = Size;
+	}
+
+	LLM(FLowLevelMemTracker::Get().OnLowLevelAlloc(ELLMTracker::Platform, Pointer, Size));
+	return Pointer;
 #else
-    void* Ptr = mmap(nullptr, Size, PROT_READ | PROT_WRITE, MAP_PRIVATE | MAP_ANON, -1, 0);
-    if (Ptr == (void*)-1)
-    {
-        UE_LOG(LogTemp, Warning, TEXT("mmap failure allocating %d, error code: %d"), Size, errno);
-        Ptr = nullptr;
-    }
-    LLM(FLowLevelMemTracker::Get().OnLowLevelAlloc(ELLMTracker::Platform, Ptr, Size));
-    return Ptr;
+	void* Ptr = mmap(nullptr, Size, PROT_READ | PROT_WRITE, MAP_PRIVATE | MAP_ANON, -1, 0);
+	if (Ptr == (void*)-1)
+	{
+		UE_LOG(LogTemp, Warning, TEXT("mmap failure allocating %d, error code: %d"), Size, errno);
+		Ptr = nullptr;
+	}
+	LLM(FLowLevelMemTracker::Get().OnLowLevelAlloc(ELLMTracker::Platform, Ptr, Size));
+	return Ptr;
 #endif // USE_MALLOC_BINNED2
 }
 
-void FApplePlatformMemory::BinnedFreeToOS( void* Ptr, SIZE_T Size )
+void FApplePlatformMemory::BinnedFreeToOS(void* Ptr, SIZE_T Size)
 {
-// Binned2 requires allocations to be BinnedPageSize-aligned. Simple mmap() does not guarantee this for recommended BinnedPageSize (64KB).
+	// Binned2 requires allocations to be BinnedPageSize-aligned. Simple mmap() does not guarantee this for recommended BinnedPageSize (64KB).
 #if USE_MALLOC_BINNED2
-    return FGenericPlatformMemory::BinnedFreeToOS(Ptr, Size);
+	LLM(FLowLevelMemTracker::Get().OnLowLevelFree(ELLMTracker::Platform, Ptr));
+	// guard against someone not passing size in whole pages
+	static SIZE_T OSPageSize = FPlatformMemory::GetConstants().PageSize;
+	SIZE_T SizeInWholePages = (Size % OSPageSize) ? (Size + OSPageSize - (Size % OSPageSize)) : Size;
+
+	if (UE4_PLATFORM_SANITY_CHECK_OS_ALLOCATIONS)
+	{
+		const SIZE_T DescriptorSize = OSPageSize;
+
+		FOSAllocationDescriptor* AllocDescriptor = reinterpret_cast<FOSAllocationDescriptor*>(reinterpret_cast<SIZE_T>(Ptr) + Size);
+		if (UNLIKELY(AllocDescriptor->Magic != FOSAllocationDescriptor::MagicType::Marker))
+		{
+			UE_LOG(LogHAL, Fatal, TEXT("BinnedFreeToOS() has been passed an address %p (size %llu) not allocated through it."), Ptr, (uint64)Size);
+			// unreachable
+			return;
+		}
+
+		void* PointerToUnmap = AllocDescriptor->PointerToUnmap;
+		SIZE_T SizeToUnmap = AllocDescriptor->SizeToUnmap;
+
+		// do checks, from most to least serious
+		if (UE4_PLATFORM_SANITY_CHECK_OS_ALLOCATIONS != 0)
+		{
+			// this check only makes sense when we're not reducing number of maps, since the pointer will have to be different.
+			if (UNLIKELY(PointerToUnmap != Ptr || SizeToUnmap != SizeInWholePages + DescriptorSize))
+			{
+				UE_LOG(LogHAL, Fatal, TEXT("BinnedFreeToOS(): info mismatch: descriptor ptr: %p, size %llu, but our pointer is %p and size %llu."), PointerToUnmap, SizeToUnmap, AllocDescriptor, (uint64)(SizeInWholePages + DescriptorSize));
+				// unreachable
+				return;
+			}
+
+			if (UNLIKELY(AllocDescriptor->OriginalSizeAsPassed != Size))
+			{
+				UE_LOG(LogHAL, Fatal, TEXT("BinnedFreeToOS(): info mismatch: descriptor original size %llu, our size is %llu for pointer %p"), AllocDescriptor->OriginalSizeAsPassed, Size, Ptr);
+				// unreachable
+				return;
+			}
+		}
+
+		AllocDescriptor = nullptr;	// just so no one touches it
+
+		if (UNLIKELY(munmap(PointerToUnmap, SizeToUnmap) != 0))
+		{
+			FPlatformMemory::OnOutOfMemory(SizeToUnmap, 0);
+			// unreachable
+		}
+	}
+	else
+	{
+		if (UNLIKELY(munmap(Ptr, SizeInWholePages) != 0))
+		{
+			FPlatformMemory::OnOutOfMemory(SizeInWholePages, 0);
+			// unreachable
+		}
+	}
 #else
-    LLM(FLowLevelMemTracker::Get().OnLowLevelFree(ELLMTracker::Platform, Ptr));
-    if (munmap(Ptr, Size) != 0)
-    {
-        const int ErrNo = errno;
-        UE_LOG(LogHAL, Fatal, TEXT("munmap(addr=%p, len=%llu) failed with errno = %d (%s)"), Ptr, Size,
-               ErrNo, StringCast< TCHAR >(strerror(ErrNo)).Get());
-    }
+	LLM(FLowLevelMemTracker::Get().OnLowLevelFree(ELLMTracker::Platform, Ptr));
+	if (munmap(Ptr, Size) != 0)
+	{
+		const int ErrNo = errno;
+		UE_LOG(LogHAL, Fatal, TEXT("munmap(addr=%p, len=%llu) failed with errno = %d (%s)"), Ptr, Size,
+			ErrNo, StringCast< TCHAR >(strerror(ErrNo)).Get());
+	}
 #endif // USE_MALLOC_BINNED2
 }
 
+size_t FApplePlatformMemory::FPlatformVirtualMemoryBlock::GetVirtualSizeAlignment()
+{
+	static SIZE_T OSPageSize = FPlatformMemory::GetConstants().PageSize;
+	return OSPageSize;
+}
+
+size_t FApplePlatformMemory::FPlatformVirtualMemoryBlock::GetCommitAlignment()
+{
+	static SIZE_T OSPageSize = FPlatformMemory::GetConstants().PageSize;
+	return OSPageSize;
+}
+
+FApplePlatformMemory::FPlatformVirtualMemoryBlock FApplePlatformMemory::FPlatformVirtualMemoryBlock::AllocateVirtual(size_t InSize, size_t InAlignment)
+{
+	FPlatformVirtualMemoryBlock Result;
+	InSize = Align(InSize, GetVirtualSizeAlignment());
+	Result.VMSizeDivVirtualSizeAlignment = InSize / GetVirtualSizeAlignment();
+	size_t Alignment = FMath::Max(InAlignment, GetVirtualSizeAlignment());
+	check(Alignment <= GetVirtualSizeAlignment());
+
+	Result.Ptr = mmap(nullptr, Result.GetActualSize(), PROT_READ | PROT_WRITE, MAP_PRIVATE | MAP_ANON, -1, 0);
+	if (!LIKELY(Result.Ptr != MAP_FAILED))
+	{
+		FPlatformMemory::OnOutOfMemory(Result.GetActualSize(), InAlignment);
+	}
+	check(Result.Ptr && IsAligned(Result.Ptr, Alignment));
+	return Result;
+}
+
+
+
+void FApplePlatformMemory::FPlatformVirtualMemoryBlock::FreeVirtual()
+{
+	if (Ptr)
+	{
+		check(GetActualSize() > 0);
+		if (munmap(Ptr, GetActualSize()) != 0)
+		{
+			// we can ran out of VMAs here
+			FPlatformMemory::OnOutOfMemory(GetActualSize(), 0);
+			// unreachable
+		}
+		Ptr = nullptr;
+		VMSizeDivVirtualSizeAlignment = 0;
+	}
+}
+
+void FApplePlatformMemory::FPlatformVirtualMemoryBlock::Commit(size_t InOffset, size_t InSize)
+{
+	check(IsAligned(InOffset, GetCommitAlignment()) && IsAligned(InSize, GetCommitAlignment()));
+	check(InOffset >= 0 && InSize >= 0 && InOffset + InSize <= GetActualSize() && Ptr);
+}
+
+void FApplePlatformMemory::FPlatformVirtualMemoryBlock::Decommit(size_t InOffset, size_t InSize)
+{
+	check(IsAligned(InOffset, GetCommitAlignment()) && IsAligned(InSize, GetCommitAlignment()));
+	check(InOffset >= 0 && InSize >= 0 && InOffset + InSize <= GetActualSize() && Ptr);
+	madvise(((uint8*)Ptr) + InOffset, InSize, MADV_DONTNEED);
+}
+
+
+
 /**
  * LLM uses these low level functions (LLMAlloc and LLMFree) to allocate memory. It grabs
  * the function pointers by calling FPlatformMemory::GetLLMAllocFunctions. If these functions
diff --git a/Engine/Source/Runtime/Core/Private/GenericPlatform/GenericPlatformMemory.cpp b/Engine/Source/Runtime/Core/Private/GenericPlatform/GenericPlatformMemory.cpp
--- a/Engine/Source/Runtime/Core/Private/GenericPlatform/GenericPlatformMemory.cpp
+++ b/Engine/Source/Runtime/Core/Private/GenericPlatform/GenericPlatformMemory.cpp
@@ -21,32 +21,6 @@
 #include "Misc/MessageDialog.h"
 #include "Templates/UnrealTemplate.h"
 
-#if PLATFORM_UNIX || PLATFORM_MAC || PLATFORM_IOS || PLATFORM_ANDROID
-	#include <sys/mman.h>
-	// more mmap()-based platforms can be added
-	#define UE4_PLATFORM_USES_MMAP_FOR_BINNED_OS_ALLOCS			1
-#else
-	#define UE4_PLATFORM_USES_MMAP_FOR_BINNED_OS_ALLOCS			0
-#endif
-
-// on 64 bit Linux, it is easier to run out of vm.max_map_count than of other limits. Due to that, trade VIRT (address space) size for smaller amount of distinct mappings
-// by not leaving holes between them (kernel will coalesce the adjoining mappings into a single one)
-// Disable by default as this causes large wasted virtual memory areas as we've to cut out 64k aligned pointers from a larger memory area then requested but then leave them mmap'd
-#define UE4_PLATFORM_REDUCE_NUMBER_OF_MAPS					(0 && PLATFORM_UNIX && PLATFORM_64BITS)
-
-#ifndef MALLOC_LEAKDETECTION
-	#define MALLOC_LEAKDETECTION 0
-#endif
-
-// check bookkeeping info against the passed in parameters in Debug and Development (the latter only in games and servers. also, only if leak detection is disabled, otherwise things are very slow)
-#define UE4_PLATFORM_SANITY_CHECK_OS_ALLOCATIONS			(UE_BUILD_DEBUG || (UE_BUILD_DEVELOPMENT && (UE_GAME || UE_SERVER) && !MALLOC_LEAKDETECTION))
-
-#if PLATFORM_UNIX
-// Defined in UnixPlatformMemory.cpp. If enabled advise mmap'ed memory that its pages are mergable
-extern bool GUseKSM;
-extern bool GKSMMergeAllPages;
-#endif
-
 DEFINE_STAT(MCR_Physical);
 DEFINE_STAT(MCR_PhysicalLLM);
 DEFINE_STAT(MCR_GPU);
@@ -260,292 +234,6 @@ bool FGenericPlatformMemory::PageProtect(void* const Ptr, const SIZE_T Size, con
 	return false;
 }
 
-/** This structure is stored in the page after each OS allocation and checks that its properties are valid on Free. Should be less than the page size (4096 on all platforms we support) */
-struct FOSAllocationDescriptor
-{
-	enum class MagicType : uint64
-	{
-		Marker = 0xd0c233ccf493dfb0
-	};
-
-	/** Magic that makes sure that we are not passed a pointer somewhere into the middle of the allocation (and/or the structure wasn't stomped). */
-	MagicType	Magic;
-
-	/** This should include the descriptor itself. */
-	void*		PointerToUnmap;
-
-	/** This should include the total size of allocation, so after unmapping it everything is gone, including the descriptor */
-	SIZE_T		SizeToUnmap;
-
-	/** Debug info that makes sure that the correct size is preserved. */
-	SIZE_T		OriginalSizeAsPassed;
-};
-
-// TODO Move all uses of madvise to UnixPlatformMemory. This requires refactoring BinnedAllocFromOS so we can apply
-// platform specific logic to the memory.
-#if PLATFORM_UNIX
-namespace
-{
-	void MarkMappedMemoryMergable(void* Pointer, SIZE_T Size)
-	{
-		const SIZE_T BinnedPageSize = FPlatformMemory::GetConstants().BinnedPageSize;
-
-		// If we dont want to merge all pages only merge chunks larger then BinnedPageSize
-		if (GUseKSM && (GKSMMergeAllPages || Size > BinnedPageSize))
-		{
-			int Ret = madvise(Pointer, Size, MADV_MERGEABLE);
-			if (Ret != 0)
-			{
-				int ErrNo = errno;
-				UE_LOG(LogHAL, Fatal, TEXT("madvise(addr=%p, length=%d, advice=MADV_MERGEABLE) failed with errno = %d (%s)"),
-					Pointer, Size, ErrNo, StringCast< TCHAR >(strerror(ErrNo)).Get());
-				// unreachable
-			}
-		}
-	}
-}
-#endif
-
-void* FGenericPlatformMemory::BinnedAllocFromOS( SIZE_T Size )
-{
-#if !UE4_PLATFORM_USES_MMAP_FOR_BINNED_OS_ALLOCS
-	UE_LOG(LogMemory, Error, TEXT("FGenericPlatformMemory::BinnedAllocFromOS not implemented on this platform"));
-	return nullptr;
-#else
-
-	static SIZE_T OSPageSize = FPlatformMemory::GetConstants().PageSize;
-	// guard against someone not passing size in whole pages
-	SIZE_T SizeInWholePages = (Size % OSPageSize) ? (Size + OSPageSize - (Size % OSPageSize)) : Size;
-	void* Pointer = nullptr;
-
-	// Binned expects OS allocations to be BinnedPageSize-aligned, and that page is at least 64KB. mmap() alone cannot do this, so carve out the needed chunks.
-	const SIZE_T ExpectedAlignment = FPlatformMemory::GetConstants().BinnedPageSize;
-	// Descriptor is only used if we're sanity checking. However, #ifdef'ing its use would make the code more fragile. Size needs to be at least one page.
-	const SIZE_T DescriptorSize = (UE4_PLATFORM_REDUCE_NUMBER_OF_MAPS != 0 || UE4_PLATFORM_SANITY_CHECK_OS_ALLOCATIONS != 0) ? OSPageSize : 0;
-
-#if PLATFORM_UNIX
-	// If we are not using descriptors lets see if a mmap of the requested size happens to return a pointer aligned with ExpectedAlignment
-	// Such pointer will likely be coalesced with an existing mapping, reducing the number of disjoint mappings (VMAs)
-	// This is important on Linux where we can run out VMAs; other platforms seem to be fine with holes in the virtual address space, so spare them an extra syscall
-	if (DescriptorSize == 0)
-	{
-		void* PointerWeGotFromMMap = mmap(nullptr, Size, PROT_READ | PROT_WRITE, MAP_PRIVATE | MAP_ANON, -1, 0);
-
-		if (PointerWeGotFromMMap == MAP_FAILED)
-		{
-			FPlatformMemory::OnOutOfMemory(Size, ExpectedAlignment);
-			// unreachable
-			return nullptr;
-		}
-
-		// If we happen to be aligned we can return early with out having to unmap anything
-		if (reinterpret_cast<SIZE_T>(PointerWeGotFromMMap) % ExpectedAlignment == 0)
-		{
-			MarkMappedMemoryMergable(PointerWeGotFromMMap, Size);
-
-			return PointerWeGotFromMMap;
-		}
-
-		// We have failed to get an aligned pointer, lets unmap and fall back
-		if (munmap(PointerWeGotFromMMap, Size) != 0)
-		{
-			FPlatformMemory::OnOutOfMemory(Size, ExpectedAlignment);
-			// unreachable
-			return nullptr;
-		}
-	}
-#endif
-
-	SIZE_T ActualSizeMapped = SizeInWholePages + ExpectedAlignment;
-
-	// the remainder of the map will be used for the descriptor, if any.
-	// we always allocate at least one page more than needed
-	void* PointerWeGotFromMMap = mmap(nullptr, ActualSizeMapped, PROT_READ | PROT_WRITE, MAP_PRIVATE | MAP_ANON, -1, 0);
-	// store these values, to unmap later
-
-	Pointer = PointerWeGotFromMMap;
-	if (Pointer == MAP_FAILED)
-	{
-		FPlatformMemory::OnOutOfMemory(ActualSizeMapped, ExpectedAlignment);
-		// unreachable
-		return nullptr;
-	}
-
-	SIZE_T Offset = (reinterpret_cast<SIZE_T>(Pointer) % ExpectedAlignment);
-
-	// See if we need to unmap anything in the front. If the pointer happened to be aligned, we don't need to unmap anything.
-	if (LIKELY(Offset != 0))
-	{
-		// figure out how much to unmap before the alignment.
-		SIZE_T SizeToNextAlignedPointer = ExpectedAlignment - Offset;
-		void* AlignedPointer = reinterpret_cast<void*>(reinterpret_cast<SIZE_T>(Pointer) + SizeToNextAlignedPointer);
-
-		// do not unmap if we're trying to reduce the number of distinct maps, since holes prevent the Linux kernel from coalescing two adjoining mmap()s into a single VMA
-		if (!UE4_PLATFORM_REDUCE_NUMBER_OF_MAPS)
-		{
-			// unmap the part before
-			if (munmap(Pointer, SizeToNextAlignedPointer) != 0)
-			{
-				FPlatformMemory::OnOutOfMemory(SizeToNextAlignedPointer, ExpectedAlignment);
-				// unreachable
-				return nullptr;
-			}
-
-			// account for reduced mmaped size
-			ActualSizeMapped -= SizeToNextAlignedPointer;
-		}
-
-		// now, make it appear as if we initially got the allocation right
-		Pointer = AlignedPointer;
-	}
-
-	// at this point, Pointer is aligned at the expected alignment - either we lucked out on the initial allocation
-	// or we already got rid of the extra memory that was allocated in the front.
-	checkf((reinterpret_cast<SIZE_T>(Pointer) % ExpectedAlignment) == 0, TEXT("BinnedAllocFromOS(): Internal error: did not align the pointer as expected."));
-
-#if PLATFORM_UNIX
-	MarkMappedMemoryMergable(Pointer, ActualSizeMapped);
-#endif
-
-	// do not unmap if we're trying to reduce the number of distinct maps, since holes prevent the Linux kernel from coalescing two adjoining mmap()s into a single VMA
-	if (!UE4_PLATFORM_REDUCE_NUMBER_OF_MAPS)
-	{
-		// Now unmap the tail only, if any, but leave just enough space for the descriptor
-		void* TailPtr = reinterpret_cast<void*>(reinterpret_cast<SIZE_T>(Pointer) + SizeInWholePages + DescriptorSize);
-		SIZE_T TailSize = ActualSizeMapped - SizeInWholePages - DescriptorSize;
-	
-		if (LIKELY(TailSize > 0))
-		{
-			if (munmap(TailPtr, TailSize) != 0)
-			{
-				FPlatformMemory::OnOutOfMemory(TailSize, ExpectedAlignment);
-				// unreachable
-				return nullptr;
-			}
-		}
-	}
-
-	// we're done with this allocation, fill in the descriptor with the info
-	if (LIKELY(DescriptorSize > 0))
-	{
-		FOSAllocationDescriptor* AllocDescriptor = reinterpret_cast<FOSAllocationDescriptor*>(reinterpret_cast<SIZE_T>(Pointer) + Size);
-		AllocDescriptor->Magic = FOSAllocationDescriptor::MagicType::Marker;
-		if (!UE4_PLATFORM_REDUCE_NUMBER_OF_MAPS)
-		{
-			AllocDescriptor->PointerToUnmap = Pointer;
-			AllocDescriptor->SizeToUnmap = SizeInWholePages + DescriptorSize;
-		}
-		else
-		{
-			AllocDescriptor->PointerToUnmap = PointerWeGotFromMMap;
-			AllocDescriptor->SizeToUnmap = ActualSizeMapped;
-		}
-		AllocDescriptor->OriginalSizeAsPassed = Size;
-	}
-
-	return Pointer;
-#endif // UE4_PLATFORM_USES_MMAP_FOR_BINNED_OS_ALLOCS
-}
-
-void FGenericPlatformMemory::BinnedFreeToOS( void* Ptr, SIZE_T Size )
-{
-#if !UE4_PLATFORM_USES_MMAP_FOR_BINNED_OS_ALLOCS
-	UE_LOG(LogMemory, Error, TEXT("FGenericPlatformMemory::BinnedFreeToOS not implemented on this platform"));
-#else
-
-	// guard against someone not passing size in whole pages
-	static SIZE_T OSPageSize = FPlatformMemory::GetConstants().PageSize;
-	SIZE_T SizeInWholePages = (Size % OSPageSize) ? (Size + OSPageSize - (Size % OSPageSize)) : Size;
-
-	if (UE4_PLATFORM_REDUCE_NUMBER_OF_MAPS || UE4_PLATFORM_SANITY_CHECK_OS_ALLOCATIONS)
-	{
-		const SIZE_T DescriptorSize = OSPageSize;
-
-		FOSAllocationDescriptor* AllocDescriptor = reinterpret_cast<FOSAllocationDescriptor*>(reinterpret_cast<SIZE_T>(Ptr) + Size);
-		if (UNLIKELY(AllocDescriptor->Magic != FOSAllocationDescriptor::MagicType::Marker))
-		{
-			UE_LOG(LogHAL, Fatal, TEXT("BinnedFreeToOS() has been passed an address %p (size %llu) not allocated through it."), Ptr, (uint64)Size);
-			// unreachable
-			return;
-		}
-
-		void* PointerToUnmap = AllocDescriptor->PointerToUnmap;
-		SIZE_T SizeToUnmap = AllocDescriptor->SizeToUnmap;
-
-		// do checks, from most to least serious
-		if (UE4_PLATFORM_SANITY_CHECK_OS_ALLOCATIONS != 0)
-		{
-			// this check only makes sense when we're not reducing number of maps, since the pointer will have to be different.
-			if (UE4_PLATFORM_REDUCE_NUMBER_OF_MAPS == 0)
-			{
-				if (UNLIKELY(PointerToUnmap != Ptr || SizeToUnmap != SizeInWholePages + DescriptorSize))
-				{
-					UE_LOG(LogHAL, Fatal, TEXT("BinnedFreeToOS(): info mismatch: descriptor ptr: %p, size %llu, but our pointer is %p and size %llu."), PointerToUnmap, SizeToUnmap, AllocDescriptor, (uint64)(SizeInWholePages + DescriptorSize));
-					// unreachable
-					return;
-				}
-			}
-	
-			if (UNLIKELY(AllocDescriptor->OriginalSizeAsPassed != Size))
-			{
-				UE_LOG(LogHAL, Fatal, TEXT("BinnedFreeToOS(): info mismatch: descriptor original size %llu, our size is %llu for pointer %p"), AllocDescriptor->OriginalSizeAsPassed, Size, Ptr);
-				// unreachable
-				return;
-			}
-		}
-
-		AllocDescriptor = nullptr;	// just so no one touches it
-
-		if (UNLIKELY(munmap(PointerToUnmap, SizeToUnmap) != 0))
-		{
-			FPlatformMemory::OnOutOfMemory(SizeToUnmap, 0);
-			// unreachable
-		}
-	}
-	else
-	{
-		if (UNLIKELY(munmap(Ptr, SizeInWholePages) != 0))
-		{
-			FPlatformMemory::OnOutOfMemory(SizeInWholePages, 0);
-			// unreachable
-		}
-	}
-#endif // UE4_PLATFORM_USES_MMAP_FOR_BINNED_OS_ALLOCS
-}
-
-void* FGenericPlatformMemory::MemoryRangeReserve(SIZE_T Size, bool bCommit, int32 Node)
-{
-#if UE4_PLATFORM_USES_MMAP_FOR_BINNED_OS_ALLOCS
-	void* Ptr = mmap(nullptr, Size, PROT_READ | PROT_WRITE, MAP_PRIVATE | MAP_ANON, -1, 0);
-	if (LIKELY(Ptr != MAP_FAILED))
-	{
-#if PLATFORM_UNIX
-		MarkMappedMemoryMergable(Ptr, Size);
-#endif
-		return Ptr;
-	}
-	return nullptr;
-#else
-	UE_LOG(LogMemory, Fatal, TEXT("FGenericPlatformMemory::MemoryRangeReserve not implemented on this platform"));
-	// unreachable
-	return nullptr;
-#endif
-}
-
-void FGenericPlatformMemory::MemoryRangeFree(void* Ptr, SIZE_T Size)
-{
-#if UE4_PLATFORM_USES_MMAP_FOR_BINNED_OS_ALLOCS
-	if (munmap(Ptr, Size) != 0)
-	{
-		// we can ran out of VMAs here
-		FPlatformMemory::OnOutOfMemory(Size, 0);
-		// unreachable
-	}
-#else
-	UE_LOG(LogMemory, Fatal, TEXT("FGenericPlatformMemory::MemoryRangeFree not implemented on this platform"));
-	// unreachable
-#endif
-}
 
 void FGenericPlatformMemory::DumpStats( class FOutputDevice& Ar )
 {
diff --git a/Engine/Source/Runtime/Core/Private/HAL/Allocators/PooledVirtualMemoryAllocator.cpp b/Engine/Source/Runtime/Core/Private/HAL/Allocators/PooledVirtualMemoryAllocator.cpp
--- a/Engine/Source/Runtime/Core/Private/HAL/Allocators/PooledVirtualMemoryAllocator.cpp
+++ b/Engine/Source/Runtime/Core/Private/HAL/Allocators/PooledVirtualMemoryAllocator.cpp
@@ -10,7 +10,9 @@
 #include "Templates/AlignmentTemplates.h"
 #include "GenericPlatform/OSAllocationPool.h"
 
-using T64KBAlignedPool = TMemoryPool< FPlatformMemory::MemoryRangeCommit, FPlatformMemory::MemoryRangeDecommit, 65536 >;
+#if PLATFORM_HAS_FPlatformVirtualMemoryBlock
+
+using T64KBAlignedPool = TMemoryPool<65536>;
 
 /** Scale parameter used when growing the pools on allocation (and scaling them back), configurable from the commandline */
 float GVMAPoolScale = 1.4f;
@@ -186,17 +188,19 @@ FPooledVirtualMemoryAllocator::FPoolDescriptorBase* FPooledVirtualMemoryAllocato
 	// now add the main memory requirements
 	TotalSize += AllocationSize * static_cast<SIZE_T>(NumPooledAllocations);
 
-	uint8* RawPtr = static_cast<uint8*>(FPlatformMemory::MemoryRangeReserve(TotalSize));
+	FPlatformMemory::FPlatformVirtualMemoryBlock VMBlock = FPlatformMemory::FPlatformVirtualMemoryBlock::AllocateVirtual(Align(TotalSize, FPlatformMemory::FPlatformVirtualMemoryBlock::GetVirtualSizeAlignment()));
+
+	uint8* RawPtr = static_cast<uint8*>(VMBlock.GetVirtualPointer());
 	if (RawPtr == nullptr)
 	{
 		return nullptr;
 	}    
 
 	// Commit the header so we can touch it
-	FPlatformMemory::MemoryRangeCommit(RawPtr, HeaderSize);
+	VMBlock.Commit(0, Align(HeaderSize, FPlatformMemory::FPlatformVirtualMemoryBlock::GetCommitAlignment()));
 	FPoolDescriptor* Ptr = reinterpret_cast<FPoolDescriptor*>(RawPtr);
 	// store the total size to deallocate
-	Ptr->TotalSize = TotalSize;
+	Ptr->VMSizeDivVirtualSizeAlignment = VMBlock.GetActualSizeInPages();
 
 	// find different offsets
 	uint8* PointerToPool = RawPtr + DescriptorSize;
@@ -205,7 +209,7 @@ FPooledVirtualMemoryAllocator::FPoolDescriptorBase* FPooledVirtualMemoryAllocato
 
 	uint8* AlignedMemoryForThePool = Align(MemoryAfterTheHeader, 65536);
 	Ptr->Pool = new (PointerToPool) T64KBAlignedPool(AllocationSize, reinterpret_cast<SIZE_T>(AlignedMemoryForThePool), NumPooledAllocations, 
-		PointerToBookkeepingMemory);
+		PointerToBookkeepingMemory, VMBlock);
 
 	return Ptr;
 }
@@ -219,7 +223,8 @@ void FPooledVirtualMemoryAllocator::DestroyPool(FPoolDescriptorBase* Pool)
 	// allocated with placement new, do not call delete
 	PoolDesc.Pool->~T64KBAlignedPool();
 
-	FPlatformMemory::MemoryRangeFree(Pool, Pool->TotalSize);
+	FPlatformMemory::FPlatformVirtualMemoryBlock VMBlock(Pool, Pool->VMSizeDivVirtualSizeAlignment);
+	VMBlock.FreeVirtual();
 }
 
 void FPooledVirtualMemoryAllocator::FreeAll()
@@ -252,3 +257,4 @@ uint64 FPooledVirtualMemoryAllocator::GetCachedFreeTotal()
 
 	return TotalFree;
 }
+#endif
\ No newline at end of file
diff --git a/Engine/Source/Runtime/Core/Private/HAL/MallocBinned3.cpp b/Engine/Source/Runtime/Core/Private/HAL/MallocBinned3.cpp
--- a/Engine/Source/Runtime/Core/Private/HAL/MallocBinned3.cpp
+++ b/Engine/Source/Runtime/Core/Private/HAL/MallocBinned3.cpp
@@ -1,4 +1,4 @@
-// Copyright 1998-2019 Epic Games, Inc. All Rights Reserved.
+// Copyright 1998-2018 Epic Games, Inc. All Rights Reserved.
 
 /*=============================================================================
 	MallocBinned.cpp: Binned memory allocator
@@ -6,7 +6,7 @@
 
 #include "HAL/MallocBinned3.h"
 
-#if PLATFORM_64BITS
+#if PLATFORM_64BITS && PLATFORM_HAS_FPlatformVirtualMemoryBlock
 #include "Logging/LogMacros.h"
 #include "Misc/ScopeLock.h"
 #include "Templates/Function.h"
@@ -18,6 +18,20 @@
 
 
 
+#if USE_CACHED_PAGE_ALLOCATOR_FOR_LARGE_ALLOCS
+#include "HAL/Allocators/CachedOSPageAllocator.h"
+#define BINNED3_MAX_CACHED_OS_FREES (64)
+#define BINNED3_MAX_CACHED_OS_FREES_BYTE_LIMIT (64*1024*1024)
+
+typedef TCachedOSPageAllocator<BINNED3_MAX_CACHED_OS_FREES, BINNED3_MAX_CACHED_OS_FREES_BYTE_LIMIT> TBinned3CachedOSPageAllocator;
+
+TBinned3CachedOSPageAllocator& GetCachedOSPageAllocator()
+{
+	static TBinned3CachedOSPageAllocator Singleton;
+	return Singleton;
+}
+
+#endif
 
 #if BINNED3_ALLOW_RUNTIME_TWEAKING
 
@@ -65,6 +79,8 @@ int64 Binned3AllocatedOSSmallPoolMemory = 0;
 int64 Binned3AllocatedLargePoolMemory = 0; // memory requests to the OS which don't fit in the small pool
 int64 Binned3AllocatedLargePoolMemoryWAlignment = 0; // when we allocate at OS level we need to align to a size
 
+TAtomic<int64> Binned3Commits;
+TAtomic<int64> Binned3Decommits;
 int64 Binned3PoolInfoMemory = 0;
 int64 Binned3HashMemory = 0;
 int64 Binned3FreeBitsMemory = 0;
@@ -85,179 +101,12 @@ TAtomic<double> MemoryRangeFreeTotalTime(0.0);
 TAtomic<int32> MemoryRangeFreeTotalCount(0);
 #endif
 
-#define BINNED3_LARGE_POOL_CANARIES 0 // need to repad the data structure so that the pagesize divides by this (!(UE_BUILD_SHIPPING || UE_BUILD_TEST))
-
-// Block sizes are based around getting the maximum amount of allocations per pool, with as little alignment waste as possible.
-// Block sizes should be close to even divisors of the system page size, and well distributed.
-// They must be 16-byte aligned as well.
-static uint32 Binned3SmallBlockSizes4k[] = 
-{
-	16, 32, 48, 64, 80, 96, 112, 128, 160, // +16
-	192, 224, 256, 288, 320, // +32
-	368,  // /11 ish
-	400,  // /10 ish
-	448,  // /9 ish
-	512,  // /8
-	576,  // /7 ish
-	672,  // /6 ish
-	816,  // /5 ish 
-	1024, // /4 
-	1360, // /3 ish
-	2048, // /2
-	4096 // /1
-};
-
-static uint32 Binned3SmallBlockSizes8k[] =
-{
-	736,  // /11 ish
-	1168, // /7 ish
-	1632, // /5 ish
-	2720, // /3 ish
-	8192  // /1
-};
-
-static uint32 Binned3SmallBlockSizes12k[] =
-{
-	//1104, // /11 ish
-	//1216, // /10 ish
-	1536, // /8
-	1744, // /7 ish
-	2448, // /5 ish
-	3072, // /4
-	6144, // /2
-	12288 // /1
-};
-
-static uint32 Binned3SmallBlockSizes16k[] =
-{
-	//1488, // /11 ish
-	//1808, // /9 ish
-	//2336, // /7 ish
-	3264, // /5 ish
-	5456, // /3 ish
-	16384 // /1
-};
-
-static uint32 Binned3SmallBlockSizes20k[] =
-{
-	// 2912, // /7 ish
-	//3408, // /6 ish
-	5120, // /4
-	//6186, // /3 ish
-	10240, // /2
-	20480  // /1
-};
-
-static uint32 Binned3SmallBlockSizes24k[] = // 1 total
-{
-	24576  // /1
-};
-
-static uint32 Binned3SmallBlockSizes28k[] = // 6 total
-{
-	4768,  // /6 ish
-	5728,  // /5 ish
-	7168,  // /4
-	9552,  // /3
-	14336, // /2
-	28672  // /1
-};
-
-
-
-struct FSizeTableEntry
-{
-	uint32 BlockSize;
-	uint16 BlocksPerBlockOfBlocks;
-	uint8 PagesPlatformForBlockOfBlocks;
-
-	FSizeTableEntry()
-	{
-	}
-
-	FSizeTableEntry(uint32 InBlockSize, uint64 PlatformPageSize, uint8 Pages4k)
-		: BlockSize(InBlockSize)
-	{
-		check((PlatformPageSize & (BINNED3_BASE_PAGE_SIZE - 1)) == 0 && PlatformPageSize >= BINNED3_BASE_PAGE_SIZE && InBlockSize % BINNED3_MINIMUM_ALIGNMENT == 0);
-
-		uint64 Page4kPerPlatformPage = PlatformPageSize / BINNED3_BASE_PAGE_SIZE;
-
-		PagesPlatformForBlockOfBlocks = 0;
-		while (true)
-		{
-			check(PagesPlatformForBlockOfBlocks < MAX_uint8);
-			PagesPlatformForBlockOfBlocks++;
-			if (PagesPlatformForBlockOfBlocks * Page4kPerPlatformPage < Pages4k)
-			{
-				continue;
-			}
-			if (PagesPlatformForBlockOfBlocks * Page4kPerPlatformPage % Pages4k != 0)
-			{
-				continue;
-			}
-			break;
-		}
-		check((PlatformPageSize * PagesPlatformForBlockOfBlocks) / BlockSize <= MAX_uint16);
-		BlocksPerBlockOfBlocks = (PlatformPageSize * PagesPlatformForBlockOfBlocks) / BlockSize;
-	}
-
-	bool operator<(const FSizeTableEntry& Other) const
-	{
-		return BlockSize < Other.BlockSize;
-	}
-};
-
-static FSizeTableEntry SizeTable[BINNED3_SMALL_POOL_COUNT];
-
-static void FillSizeTable(uint64 PlatformPageSize)
-{
-	int32 Index = 0;
-
-	for (int32 Sub = 0; Sub < ARRAY_COUNT(Binned3SmallBlockSizes4k); Sub++)
-	{
-		SizeTable[Index++] = FSizeTableEntry(Binned3SmallBlockSizes4k[Sub], PlatformPageSize, 1);
-	}
-	for (int32 Sub = 0; Sub < ARRAY_COUNT(Binned3SmallBlockSizes8k); Sub++)
-	{
-		SizeTable[Index++] = FSizeTableEntry(Binned3SmallBlockSizes8k[Sub], PlatformPageSize, 2);
-	}
-	for (int32 Sub = 0; Sub < ARRAY_COUNT(Binned3SmallBlockSizes12k); Sub++)
-	{
-		SizeTable[Index++] = FSizeTableEntry(Binned3SmallBlockSizes12k[Sub], PlatformPageSize, 3);
-	}
-	for (int32 Sub = 0; Sub < ARRAY_COUNT(Binned3SmallBlockSizes16k); Sub++)
-	{
-		SizeTable[Index++] = FSizeTableEntry(Binned3SmallBlockSizes16k[Sub], PlatformPageSize, 4);
-	}
-	for (int32 Sub = 0; Sub < ARRAY_COUNT(Binned3SmallBlockSizes20k); Sub++)
-	{
-		SizeTable[Index++] = FSizeTableEntry(Binned3SmallBlockSizes20k[Sub], PlatformPageSize, 5);
-	}
-	for (int32 Sub = 0; Sub < ARRAY_COUNT(Binned3SmallBlockSizes24k); Sub++)
-	{
-		SizeTable[Index++] = FSizeTableEntry(Binned3SmallBlockSizes24k[Sub], PlatformPageSize, 6);
-	}
-	for (int32 Sub = 0; Sub < ARRAY_COUNT(Binned3SmallBlockSizes28k); Sub++)
-	{
-		SizeTable[Index++] = FSizeTableEntry(Binned3SmallBlockSizes28k[Sub], PlatformPageSize, 7);
-	}
-	Sort(&SizeTable[0], Index);
-	check(SizeTable[Index - 1].BlockSize == BINNED3_MAX_LISTED_SMALL_POOL_SIZE);
-	check(IsAligned(BINNED3_MAX_LISTED_SMALL_POOL_SIZE, BINNED3_BASE_PAGE_SIZE));
-	for (uint32 Size = BINNED3_MAX_LISTED_SMALL_POOL_SIZE + BINNED3_BASE_PAGE_SIZE; Size <= BINNED3_MAX_SMALL_POOL_SIZE; Size += BINNED3_BASE_PAGE_SIZE)
-	{
-		SizeTable[Index++] = FSizeTableEntry(Size, PlatformPageSize, Size / BINNED3_BASE_PAGE_SIZE);
-	}
-	check(Index == ARRAY_COUNT(SizeTable));
-	check(Index == BINNED3_SMALL_POOL_COUNT);
-
-}
+#define BINNED3_LARGE_POOL_CANARIES 1 // need to repad the data structure so that the pagesize divides by this to disable
 
 MS_ALIGN(PLATFORM_CACHE_LINE_SIZE) static uint8 Binned3UnusedAlignPadding[PLATFORM_CACHE_LINE_SIZE] GCC_ALIGN(PLATFORM_CACHE_LINE_SIZE) = { 0 };
 uint16 FMallocBinned3::SmallBlockSizesReversedShifted[BINNED3_SMALL_POOL_COUNT] = { 0 };
 uint32 FMallocBinned3::Binned3TlsSlot = 0;
 uint32 FMallocBinned3::OsAllocationGranularity = 0;
-uint32 FMallocBinned3::MaxAlignmentForMemoryRangeReserve = 0;
 
 #if !BINNED3_USE_SEPARATE_VM_PER_POOL
 	uint8* FMallocBinned3::Binned3BaseVMPtr = nullptr;
@@ -275,14 +124,14 @@ struct FMallocBinned3::FPoolInfoSmall
 {
 	enum ECanary
 	{
-		SmallUnassigned = 0x39,
-		SmallAssigned = 0x71
+		SmallUnassigned = 0x3,
+		SmallAssigned = 0x1
 	};
 
-	uint32 Canary : 7;
-	uint32 Taken : 12;
+	uint32 Canary : 2;
+	uint32 Taken : 15;
 	uint32 NoFirstFreeIndex : 1;
-	uint32 FirstFreeIndex : 12;
+	uint32 FirstFreeIndex : 14;
 
 	FPoolInfoSmall()
 		: Canary(ECanary::SmallUnassigned)
@@ -345,6 +194,7 @@ struct FMallocBinned3::FPoolInfoSmall
 	{
 		check(HasFreeRegularBlock());
 		++Taken;
+		check(Taken != 0);
 		FFreeBlock* Free = (FFreeBlock*)(BlockOfBlocksPtr + BlockSize * FirstFreeIndex);
 		void* Result = Free->AllocateRegularBlock();
 		if (Free->GetNumFreeRegularBlocks() == 0)
@@ -368,41 +218,36 @@ struct FMallocBinned3::FPoolInfoSmall
 
 struct FMallocBinned3::FPoolInfoLarge
 {
-	enum class ECanary : uint16
+	enum class ECanary : uint32
 	{
-		LargeUnassigned = 0x3943,
-		LargeAssigned = 0x17ea,
+		LargeUnassigned = 0x39431234,
+		LargeAssigned = 0x17ea5678,
 	};
 
 public:
-#if BINNED3_LARGE_POOL_CANARIES
 	ECanary	Canary;	// See ECanary
-#endif
 private:
 	uint32 AllocSize;      // Number of bytes allocated
-	uint32 OSAllocSize;      // Number of bytes allocated aligned for OS
+	uint32 VMSizeDivVirtualSizeAlignment;    // Number of VM bytes allocated aligned for OS
+	uint32 CommitSize;    // Number of bytes committed by the OS
 
 public:
 	FPoolInfoLarge() :
-#if BINNED3_LARGE_POOL_CANARIES
 		Canary(ECanary::LargeUnassigned),
-#endif
 		AllocSize(0),
-		OSAllocSize(0)
+		VMSizeDivVirtualSizeAlignment(0),
+		CommitSize(0)
 	{
 	}
 	void CheckCanary(ECanary ShouldBe) const
 	{
-#if BINNED3_LARGE_POOL_CANARIES
 		if (Canary != ShouldBe)
 		{
 			UE_LOG(LogMemory, Fatal, TEXT("MallocBinned3 Corruption Canary was 0x%x, should be 0x%x"), int32(Canary), int32(ShouldBe));
 		}
-#endif
 	}
 	void SetCanary(ECanary ShouldBe, bool bPreexisting, bool bGuarnteedToBeNew)
 	{
-#if BINNED3_LARGE_POOL_CANARIES
 		if (bPreexisting)
 		{
 			if (bGuarnteedToBeNew)
@@ -436,27 +281,37 @@ public:
 			}
 		}
 		Canary = ShouldBe;
-#endif
 	}
 	uint32 GetOSRequestedBytes() const
 	{
 		return AllocSize;
 	}
 
-	UPTRINT GetOsAllocatedBytes() const
+	UPTRINT GetOsCommittedBytes() const
+	{
+		return (UPTRINT)CommitSize;
+	}
+
+	uint32 GetOsVMPages() const
 	{
 		CheckCanary(ECanary::LargeAssigned);
-		return (UPTRINT)OSAllocSize;
+		return VMSizeDivVirtualSizeAlignment;
 	}
 
-	void SetOSAllocationSizes(uint32 InRequestedBytes, UPTRINT InAllocatedBytes)
+	void SetOSAllocationSizes(uint32 InRequestedBytes)
 	{
 		CheckCanary(ECanary::LargeAssigned);
-		check(InRequestedBytes != 0);                // Shouldn't be pooling zero byte allocations
-		check(InAllocatedBytes >= InRequestedBytes); // We must be allocating at least as much as we requested
+		AllocSize = InRequestedBytes;
+		check(AllocSize > 0 && CommitSize >= AllocSize && VMSizeDivVirtualSizeAlignment * FPlatformMemory::FPlatformVirtualMemoryBlock::GetVirtualSizeAlignment() >= CommitSize);
+	}
 
+	void SetOSAllocationSizes(uint32 InRequestedBytes, UPTRINT InCommittedBytes, uint32 InVMSizeDivVirtualSizeAlignment)
+	{
+		CheckCanary(ECanary::LargeAssigned);
 		AllocSize = InRequestedBytes;
-		OSAllocSize = InAllocatedBytes;
+		CommitSize = InCommittedBytes;
+		VMSizeDivVirtualSizeAlignment = InVMSizeDivVirtualSizeAlignment;
+		check(AllocSize > 0 && CommitSize >= AllocSize && VMSizeDivVirtualSizeAlignment * FPlatformMemory::FPlatformVirtualMemoryBlock::GetVirtualSizeAlignment() >= CommitSize);
 	}
 };
 
@@ -513,11 +368,7 @@ struct FMallocBinned3::Private
 		FPoolInfoSmall*& InfoBlock = Allocator.SmallPoolTables[InPoolIndex].PoolInfos[BlockOfBlocksIndex / Allocator.SmallPoolInfosPerPlatformPage];
 		if (!InfoBlock)
 		{
-			InfoBlock = (FPoolInfoSmall*)FPlatformMemory::MemoryRangeReserve(Allocator.OsAllocationGranularity, true);
-			if (!InfoBlock)
-			{
-				Private::OutOfMemory(Allocator.OsAllocationGranularity);
-			}
+			InfoBlock = (FPoolInfoSmall*)Allocator.AllocateMetaDataMemory(Allocator.OsAllocationGranularity);
 #if BINNED3_ALLOCATOR_STATS
 			Binned3PoolInfoMemory += Allocator.OsAllocationGranularity;
 #endif
@@ -543,14 +394,15 @@ struct FMallocBinned3::Private
 		/** 
 		 * Creates an array of FPoolInfo structures for tracking allocations.
 		 */
-		auto CreatePoolArray = [](uint64 NumPools)
+		auto CreatePoolArray = [&Allocator](uint64 NumPools)
 		{
 			uint64 PoolArraySize = NumPools * sizeof(FPoolInfoLarge);
 
 			void* Result;
 			{
+
 				LLM_PLATFORM_SCOPE(ELLMTag::FMalloc);
-				Result = FPlatformMemory::MemoryRangeReserve(PoolArraySize, true);
+				Result = Allocator.AllocateMetaDataMemory(PoolArraySize);
 #if BINNED3_ALLOCATOR_STATS
 				Binned3PoolInfoMemory += PoolArraySize;
 #endif
@@ -596,15 +448,10 @@ struct FMallocBinned3::Private
 		if (!Allocator.HashBucketFreeList)
 		{
 			{
-				LLM_PLATFORM_SCOPE(ELLMTag::FMalloc);
-				Allocator.HashBucketFreeList = (PoolHashBucket*)FPlatformMemory::MemoryRangeReserve(FMallocBinned3::OsAllocationGranularity, true);
+				Allocator.HashBucketFreeList = (PoolHashBucket*)Allocator.AllocateMetaDataMemory(FMallocBinned3::OsAllocationGranularity);
 #if BINNED3_ALLOCATOR_STATS
 				Binned3HashMemory += FMallocBinned3::OsAllocationGranularity;
 #endif
-				if (!Allocator.HashBucketFreeList)
-				{
-					OutOfMemory(FMallocBinned3::OsAllocationGranularity);
-				}
 			}
 
 			for (UPTRINT i = 0, n = FMallocBinned3::OsAllocationGranularity / sizeof(PoolHashBucket); i < n; ++i)
@@ -773,7 +620,7 @@ struct FMallocBinned3::Private
 						Table.BlockOfBlockIsExhausted.AllocBit(OutBlockOfBlocksIndex);
 					}
 
-					verify(FPlatformMemory::MemoryRangeDecommit(BasePtrOfNode, AllocSize));
+					Allocator.Decommit(InPoolIndex, BasePtrOfNode, AllocSize);
 #if BINNED3_ALLOCATOR_STATS
 					Binned3AllocatedOSSmallPoolMemory -= AllocSize;
 #endif
@@ -819,382 +666,8 @@ struct FMallocBinned3::Private
 FMallocBinned3::Private::FGlobalRecycler FMallocBinned3::Private::GGlobalRecycler;
 
 #if BINNED3_ALLOCATOR_STATS
-int64 FMallocBinned3::FPerThreadFreeBlockLists::ConsolidatedMemory = 0;
-#endif
-
-void FBitTree::FBitTreeInit(uint32 InDesiredCapacity, uint32 OsAllocationGranularity, bool InitialValue)
-{
-	DesiredCapacity = InDesiredCapacity;
-	AllocationSize = 8;
-	Rows = 1;
-	uint32 RowsUint64s = 1;
-	Capacity = 64;
-	OffsetOfLastRow = 0;
-
-	uint32 RowOffsets[10]; // 10 is way more than enough
-	RowOffsets[0] = 0;
-	uint32 RowNum[10]; // 10 is way more than enough
-	RowNum[0] = 1;
-
-	while (Capacity < DesiredCapacity)
-	{
-		Capacity *= 64;
-		RowsUint64s *= 64;
-		OffsetOfLastRow = AllocationSize / 8;
-		check(Rows < 10);
-		RowOffsets[Rows] = OffsetOfLastRow;
-		RowNum[Rows] = RowsUint64s;
-		AllocationSize += 8 * RowsUint64s;
-		Rows++;
-	}
-
-	uint32 LastRowTotal = (AllocationSize - OffsetOfLastRow * 8) * 8;
-	uint32 ExtraBits = LastRowTotal - DesiredCapacity;
-	AllocationSize -= (ExtraBits / 64) * 8;
-
-	int64 AlignedAllocationSize = Align(AllocationSize, OsAllocationGranularity);
-	LLM_PLATFORM_SCOPE(ELLMTag::FMalloc);
-	Bits = (uint64*)FPlatformMemory::MemoryRangeReserve(AlignedAllocationSize, true);
-#if BINNED3_ALLOCATOR_STATS
-	Binned3FreeBitsMemory += AlignedAllocationSize;
-#endif
-	verify(Bits);
-
-	FMemory::Memset(Bits, InitialValue ? 0xff : 0, AllocationSize);
-
-	if (!InitialValue)
-	{
-		// we fill everything beyond the desired size with occupied
-		uint32 ItemsPerBit = 64;
-		for (int32 FillRow = Rows - 2; FillRow >= 0; FillRow--)
-		{
-			uint32 NeededOneBits = RowNum[FillRow] * 64 - (DesiredCapacity + ItemsPerBit - 1) / ItemsPerBit;
-			uint32 NeededOne64s = NeededOneBits / 64;
-			NeededOneBits %= 64;
-			for (uint32 Fill = RowNum[FillRow] - NeededOne64s; Fill < RowNum[FillRow]; Fill++)
-			{
-				Bits[RowOffsets[FillRow] + Fill] = MAX_uint64;
-			}
-			if (NeededOneBits)
-			{
-				Bits[RowOffsets[FillRow] + RowNum[FillRow] - NeededOne64s - 1] = (MAX_uint64 << (64 - NeededOneBits));
-			}
-			ItemsPerBit *= 64;
-		}
-
-		if (DesiredCapacity % 64)
-		{
-			Bits[AllocationSize / 8 - 1] = (MAX_uint64 << (DesiredCapacity % 64));
-		}
-	}
-}
-
-uint32 FBitTree::AllocBit()
-{
-	uint32 Result = MAX_uint32;
-	if (*Bits != MAX_uint64) // else we are full
-	{
-		Result = 0;
-		uint32 Offset = 0;
-		uint32 Row = 0;
-		while (true)
-		{
-			uint64* At = Bits + Offset;
-			check(At >= Bits && At < Bits + AllocationSize / 8);
-			uint32 LowestZeroBit = FMath::CountTrailingZeros64(~*At);
-			check(LowestZeroBit < 64);
-			Result = Result * 64 + LowestZeroBit;
-			if (Row == Rows - 1)
-			{
-				check(!((*At) & (1ull << LowestZeroBit))); // this was already allocated?
-				*At |= (1ull << LowestZeroBit);
-				if (Row > 0)
-				{
-					if (*At == MAX_uint64)
-					{
-						do
-						{
-							uint32 Rem = (Offset - 1) % 64;
-							Offset = (Offset - 1) / 64;
-							At = Bits + Offset;
-							check(At >= Bits && At < Bits + AllocationSize / 8);
-							check(*At != MAX_uint64); // this should not already be marked full
-							*At |= (1ull << Rem);
-							if (*At != MAX_uint64)
-							{
-								break;
-							}
-							Row--;
-						} while (Row);
-					}
-				}
-				break;
-			}
-			Offset = Offset * 64 + 1 + LowestZeroBit;
-			Row++;
-		}
-	}
-
-	return Result;
-}
-
-void FBitTree::AllocBit(uint32 Index)
-{
-	check(Index < DesiredCapacity);
-	uint32 Row = Rows - 1;
-	uint32 Rem = Index % 64;
-	uint32 Offset = OffsetOfLastRow + Index / 64;
-	uint64* At = Bits + Offset;
-	check(At >= Bits && At < Bits + AllocationSize / 8);
-	check(!((*At) & (1ull << Rem))); // this was already allocated?
-	*At |= (1ull << Rem);
-	if (*At == MAX_uint64 && Row > 0)
-	{
-		do
-		{
-			Rem = (Offset - 1) % 64;
-			Offset = (Offset - 1) / 64;
-			At = Bits + Offset;
-			check(At >= Bits && At < Bits + AllocationSize / 8);
-			check(!((*At) & (1ull << Rem))); // this was already allocated?
-			*At |= (1ull << Rem);
-			if (*At != MAX_uint64)
-			{
-				break;
-			}
-			Row--;
-		} while (Row);
-	}
-
-}
-uint32 FBitTree::NextAllocBit()
-{
-	uint32 Result = MAX_uint32;
-	if (*Bits != MAX_uint64) // else we are full
-	{
-		Result = 0;
-		uint32 Offset = 0;
-		uint32 Row = 0;
-		while (true)
-		{
-			uint64* At = Bits + Offset;
-			check(At >= Bits && At < Bits + AllocationSize / 8);
-			uint32 LowestZeroBit = FMath::CountTrailingZeros64(~*At);
-			check(LowestZeroBit < 64);
-			Result = Result * 64 + LowestZeroBit;
-			if (Row == Rows - 1)
-			{
-				check(!((*At) & (1ull << LowestZeroBit))); // this was already allocated?
-				break;
-			}
-			Offset = Offset * 64 + 1 + LowestZeroBit;
-			Row++;
-		}
-	}
-
-	return Result;
-}
-
-
-void FBitTree::FreeBit(uint32 Index)
-{
-	check(Index < DesiredCapacity);
-	uint32 Row = Rows - 1;
-	uint32 Rem = Index % 64;
-	uint32 Offset = OffsetOfLastRow + Index / 64;
-	uint64* At = Bits + Offset;
-	check(At >= Bits && At < Bits + AllocationSize / 8);
-	bool bWasFull = *At == MAX_uint64;
-	check((*At) & (1ull << Rem)); // this was not already allocated?
-	*At &= ~(1ull << Rem);
-	if (bWasFull && Row > 0)
-	{
-		do
-		{
-			Rem = (Offset - 1) % 64;
-			Offset = (Offset - 1) / 64;
-			At = Bits + Offset;
-			check(At >= Bits && At < Bits + AllocationSize / 8);
-			bWasFull = *At == MAX_uint64;
-			*At &= ~(1ull << Rem);
-			if (!bWasFull)
-			{
-				break;
-			}
-			Row--;
-		} while (Row);
-	}
-}
-
-uint32 FBitTree::CountOnes(uint32 UpTo)
-{
-	uint32 Result = 0;
-	uint64* At = Bits + OffsetOfLastRow;
-	while (UpTo >= 64)
-	{
-		Result += FMath::CountBits(*At);
-		At++;
-		UpTo -= 64;
-	}
-	if (UpTo)
-	{
-		Result += FMath::CountBits((*At) << (64 - UpTo));
-	}
-	return Result;
-}
-
-void FBitTree::Test()
-{
-#if 0
-	int32 Sizes[] =
-	{
-		1,
-		63,
-		64,
-		65,
-		64 * 64 - 1,
-		64 * 64,
-		64 * 64 + 1,
-		64 * 64 * 64 - 99,
-		64 * 64 * 64 + 99,
-
-		1024 * 1024 * 1024 / (1 * 4096),
-		1024 * 1024 * 1024 / (2 * 4096),
-		1024 * 1024 * 1024 / (3 * 4096),
-		1024 * 1024 * 1024 / (4 * 4096),
-		1024 * 1024 * 1024 / (5 * 4096),
-		1024 * 1024 * 1024 / (6 * 4096),
-		1024 * 1024 * 1024 / (7 * 4096),
-
-		1024 * 1024 * 1024 / (1 * 16384),
-		1024 * 1024 * 1024 / (2 * 16384),
-		1024 * 1024 * 1024 / (3 * 16384),
-		1024 * 1024 * 1024 / (4 * 16384),
-		1024 * 1024 * 1024 / (5 * 16384),
-		1024 * 1024 * 1024 / (6 * 16384),
-		1024 * 1024 * 1024 / (7 * 16384)
-
-	};
-
-	{
-		for (int i = 0; i < ARRAY_COUNT(Sizes); i++)
-		{
-			FBitTree Tree;
-			Tree.FBitTreeInit(Sizes[i], (i & 1) ? 4096 : 16384, false);
-
-			for (int k = 0; k < Sizes[i]; k++)
-			{
-				uint32 TestI = Tree.NextAllocBit();
-				check(TestI == k);
-				check(Tree.AllocBit() == k);
-			}
-			check(Tree.NextAllocBit() == MAX_uint32);
-			check(Tree.AllocBit() == MAX_uint32);
-			Tree.FreeBit(0);
-			check(Tree.NextAllocBit() == 0);
-			check(Tree.AllocBit() == 0);
-
-			Tree.FreeBit(Sizes[i] - 1);
-			check(Tree.NextAllocBit() == Sizes[i] - 1);
-			check(Tree.AllocBit() == Sizes[i] - 1);
-
-			for (int k = 0; k < Sizes[i] - 1; k++)
-			{
-				Tree.FreeBit(k);
-			}
-			Tree.FreeBit(Sizes[i] - 1);
-
-			TArray<TPair<float, int32>> Pairs;
-			for (int k = 0; k < Sizes[i]; k++)
-			{
-				check(Tree.NextAllocBit() == k);
-				check(Tree.AllocBit() == k);
-				Pairs.Add(TPair<float, int32>(FMath::FRand(), k));
-			}
-			check(Tree.AllocBit() == MAX_uint32);
-			Pairs.Sort();
-			for (auto& p : Pairs)
-			{
-				Tree.FreeBit(p.Value);
-			}
-			Pairs.Empty();
-			for (int k = 0; k < Sizes[i]; k++)
-			{
-				check(Tree.NextAllocBit() == k);
-				check(Tree.AllocBit() == k);
-				Pairs.Add(TPair<float, int32>(FMath::FRand(), k));
-			}
-			check(Tree.NextAllocBit() == MAX_uint32);
-			check(Tree.AllocBit() == MAX_uint32);
-			Pairs.Sort();
-			for (auto& p : Pairs)
-			{
-				Tree.FreeBit(p.Value);
-				uint32 Temp = Tree.AllocBit();
-				Tree.FreeBit(Temp);
-			}
-		}
-	}
-	{
-		for (int i = 0; i < ARRAY_COUNT(Sizes); i++)
-		{
-			FBitTree Tree;
-			Tree.FBitTreeInit(Sizes[i], (i & 1) ? 4096 : 16384, false);
-
-			for (int k = 0; k < Sizes[i]; k++)
-			{
-				uint32 TestI = Tree.NextAllocBit();
-				Tree.AllocBit(TestI);
-			}
-			check(Tree.NextAllocBit() == MAX_uint32);
-			check(Tree.AllocBit() == MAX_uint32);
-			Tree.FreeBit(0);
-			check(Tree.NextAllocBit() == 0);
-			Tree.AllocBit(0);
-
-			Tree.FreeBit(Sizes[i] - 1);
-			check(Tree.NextAllocBit() == Sizes[i] - 1);
-			Tree.AllocBit(Sizes[i] - 1);
-
-			for (int k = 0; k < Sizes[i] - 1; k++)
-			{
-				Tree.FreeBit(k);
-			}
-			Tree.FreeBit(Sizes[i] - 1);
-
-			TArray<TPair<float, int32>> Pairs;
-			for (int k = 0; k < Sizes[i]; k++)
-			{
-				check(Tree.NextAllocBit() == k);
-				Tree.AllocBit(k);
-				Pairs.Add(TPair<float, int32>(FMath::FRand(), k));
-			}
-			check(Tree.AllocBit() == MAX_uint32);
-			Pairs.Sort();
-			for (auto& p : Pairs)
-			{
-				Tree.FreeBit(p.Value);
-			}
-			Pairs.Empty();
-			for (int k = 0; k < Sizes[i]; k++)
-			{
-				check(Tree.NextAllocBit() == k);
-				Tree.AllocBit(k);
-				Pairs.Add(TPair<float, int32>(FMath::FRand(), k));
-			}
-			check(Tree.NextAllocBit() == MAX_uint32);
-			check(Tree.AllocBit() == MAX_uint32);
-			Pairs.Sort();
-			for (auto& p : Pairs)
-			{
-				Tree.FreeBit(p.Value);
-				uint32 Temp = Tree.AllocBit();
-				Tree.FreeBit(Temp);
-			}
-		}
-	}
+TAtomic<int64> FMallocBinned3::FPerThreadFreeBlockLists::ConsolidatedMemory;
 #endif
-}
 
 FMallocBinned3::FPoolInfoSmall* FMallocBinned3::PushNewPoolToFront(FMallocBinned3::FPoolTable& Table, uint32 InBlockSize, uint32 InPoolIndex, uint32& OutBlockOfBlocksIndex)
 {
@@ -1210,10 +683,7 @@ FMallocBinned3::FPoolInfoSmall* FMallocBinned3::PushNewPoolToFront(FMallocBinned
 	uint8* FreePtr = BlockPointerFromIndecies(InPoolIndex, BlockOfBlocksIndex, BlockOfBlocksSize);
 
 	LLM_PLATFORM_SCOPE(ELLMTag::FMalloc);
-	if (!FPlatformMemory::MemoryRangeCommit(FreePtr, BlockOfBlocksSize))
-	{
-		Private::OutOfMemory(BlockOfBlocksSize);
-	}
+	Commit(InPoolIndex, FreePtr, BlockOfBlocksSize);
 	uint64 EndOffset = UPTRINT(FreePtr + BlockOfBlocksSize) - UPTRINT(PoolBasePtr(InPoolIndex));
 	if (EndOffset > Table.UnusedAreaOffsetLow)
 	{
@@ -1258,8 +728,8 @@ FMallocBinned3::FMallocBinned3()
 	check(!PLATFORM_32BITS);
 
 	FGenericPlatformMemoryConstants Constants = FPlatformMemory::GetConstants();
-	OsAllocationGranularity = Constants.BinnedAllocationGranularity ? Constants.BinnedAllocationGranularity : Constants.PageSize;
-	MaxAlignmentForMemoryRangeReserve = Constants.OsAllocationGranularity;
+	// large slab sizes are possible OsAllocationGranularity = 65536;
+	OsAllocationGranularity = FPlatformMemory::FPlatformVirtualMemoryBlock::GetCommitAlignment();
 	NumLargePoolsPerPage = OsAllocationGranularity / sizeof(FPoolInfoLarge);
 	check(OsAllocationGranularity % sizeof(FPoolInfoLarge) == 0);  // these need to divide evenly!
 	PtrToPoolMapping.Init(OsAllocationGranularity, NumLargePoolsPerPage, Constants.AddressLimit);
@@ -1267,14 +737,16 @@ FMallocBinned3::FMallocBinned3()
 	checkf(FMath::IsPowerOfTwo(OsAllocationGranularity), TEXT("OS page size must be a power of two"));
 	checkf(FMath::IsPowerOfTwo(Constants.AddressLimit), TEXT("OS address limit must be a power of two"));
 	checkf(Constants.AddressLimit > OsAllocationGranularity, TEXT("OS address limit must be greater than the page size")); // Check to catch 32 bit overflow in AddressLimit
-	checkf(sizeof(FMallocBinned3::FFreeBlock) <= Binned3SmallBlockSizes4k[0], TEXT("Pool header must be able to fit into the smallest block"));
 	static_assert(BINNED3_SMALL_POOL_COUNT <= 256, "Small block size array size must fit in a byte");
 	static_assert(sizeof(FFreeBlock) <= BINNED3_MINIMUM_ALIGNMENT, "Free block struct must be small enough to fit into a block.");
 
 	// Init pool tables.
 
-	FillSizeTable(OsAllocationGranularity);
+	FSizeTableEntry SizeTable[BINNED3_SMALL_POOL_COUNT];
+
+	verify(FSizeTableEntry::FillSizeTable(OsAllocationGranularity, SizeTable, BINNED3_BASE_PAGE_SIZE, BINNED3_MINIMUM_ALIGNMENT, BINNED3_MAX_SMALL_POOL_SIZE, BINNED3_BASE_PAGE_SIZE) == BINNED3_SMALL_POOL_COUNT);
 	checkf(SizeTable[BINNED3_SMALL_POOL_COUNT - 1].BlockSize == BINNED3_MAX_SMALL_POOL_SIZE, TEXT("BINNED3_MAX_SMALL_POOL_SIZE must equal the largest block size"));
+	checkf(sizeof(FMallocBinned3::FFreeBlock) <= SizeTable[0].BlockSize, TEXT("Pool header must be able to fit into the smallest block"));
 
 	SmallPoolInfosPerPlatformPage = OsAllocationGranularity / sizeof(FPoolInfoSmall);
 
@@ -1298,14 +770,33 @@ FMallocBinned3::FMallocBinned3()
 		int64 TotalNumberOfBlocksOfBlocks = MAX_MEMORY_PER_BLOCK_SIZE / (SizeTable[Index].PagesPlatformForBlockOfBlocks * OsAllocationGranularity);
 
 		int64 MaxPoolInfoMemory = Align(sizeof(FPoolInfoSmall**) * (TotalNumberOfBlocksOfBlocks + SmallPoolInfosPerPlatformPage - 1) / SmallPoolInfosPerPlatformPage, OsAllocationGranularity);
-		SmallPoolTables[Index].PoolInfos = (FPoolInfoSmall**)FPlatformMemory::MemoryRangeReserve(MaxPoolInfoMemory, true);
+		SmallPoolTables[Index].PoolInfos = (FPoolInfoSmall**)AllocateMetaDataMemory(MaxPoolInfoMemory);
 		FMemory::Memzero(SmallPoolTables[Index].PoolInfos, MaxPoolInfoMemory);
 #if BINNED3_ALLOCATOR_STATS
 		Binned3PoolInfoMemory += MaxPoolInfoMemory;
 #endif
 
-		SmallPoolTables[Index].BlockOfBlockAllocationBits.FBitTreeInit(TotalNumberOfBlocksOfBlocks, OsAllocationGranularity, false);
-		SmallPoolTables[Index].BlockOfBlockIsExhausted.FBitTreeInit(TotalNumberOfBlocksOfBlocks, OsAllocationGranularity, true);
+		{
+			int64 AllocationSize = FBitTree::GetMemoryRequirements(TotalNumberOfBlocksOfBlocks);
+			int64 AlignedAllocationSize = Align(AllocationSize, OsAllocationGranularity);
+
+			{
+				void *Bits = AllocateMetaDataMemory(AlignedAllocationSize);
+				check(Bits);
+#if BINNED3_ALLOCATOR_STATS
+				Binned3FreeBitsMemory += AlignedAllocationSize;
+#endif
+				SmallPoolTables[Index].BlockOfBlockAllocationBits.FBitTreeInit(TotalNumberOfBlocksOfBlocks, Bits, AlignedAllocationSize, false);
+			}
+			{
+				void *Bits = AllocateMetaDataMemory(AlignedAllocationSize);
+				check(Bits);
+#if BINNED3_ALLOCATOR_STATS
+				Binned3FreeBitsMemory += AlignedAllocationSize;
+#endif
+				SmallPoolTables[Index].BlockOfBlockIsExhausted.FBitTreeInit(TotalNumberOfBlocksOfBlocks, Bits, AlignedAllocationSize, true);
+			}
+		}
 	}
 
 
@@ -1334,9 +825,8 @@ FMallocBinned3::FMallocBinned3()
 	uint64 MaxHashBuckets = PtrToPoolMapping.GetMaxHashBuckets();
 
 	{
-		LLM_PLATFORM_SCOPE(ELLMTag::FMalloc);
 		int64 HashAllocSize = Align(MaxHashBuckets * sizeof(PoolHashBucket), OsAllocationGranularity);
-		HashBuckets = (PoolHashBucket*)FPlatformMemory::MemoryRangeReserve(HashAllocSize, true);
+		HashBuckets = (PoolHashBucket*)AllocateMetaDataMemory(HashAllocSize);
 #if BINNED3_ALLOCATOR_STATS
 		Binned3HashMemory += HashAllocSize;
 #endif
@@ -1348,13 +838,18 @@ FMallocBinned3::FMallocBinned3()
 	GFixedMallocLocationPtr = (FMalloc**)(&MallocBinned3);
 
 #if !BINNED3_USE_SEPARATE_VM_PER_POOL
-	Binned3BaseVMPtr = (uint8*)FPlatformMemory::MemoryRangeReserve(BINNED3_SMALL_POOL_COUNT * MAX_MEMORY_PER_BLOCK_SIZE);
+	Binned3BaseVMBlock = FPlatformMemory::FPlatformVirtualMemoryBlock::AllocateVirtual(BINNED3_SMALL_POOL_COUNT * MAX_MEMORY_PER_BLOCK_SIZE, OsAllocationGranularity);
+	Binned3BaseVMPtr = (uint8*)Binned3BaseVMBlock.GetVirtualPointer();
+	check(IsAligned(Binned3BaseVMPtr, OsAllocationGranularity));
 	verify(Binned3BaseVMPtr);
 #else
 
 	for (uint32 Index = 0; Index < BINNED3_SMALL_POOL_COUNT; ++Index)
 	{
-		uint8* NewVM = (uint8*)FPlatformMemory::MemoryRangeReserve(MAX_MEMORY_PER_BLOCK_SIZE);
+		FPlatformMemory::FPlatformVirtualMemoryBlock NewBLock = FPlatformMemory::FPlatformVirtualMemoryBlock::AllocateVirtual(MAX_MEMORY_PER_BLOCK_SIZE, OsAllocationGranularity);
+
+		uint8* NewVM = (uint8*)NewBLock.GetVirtualPointer();
+		check(IsAligned(NewVM, OsAllocationGranularity));
 		// insertion sort
 		if (Index && NewVM < PoolBaseVMPtr[Index - 1])
 		{
@@ -1370,12 +865,15 @@ FMallocBinned3::FMallocBinned3()
 			for (uint32 MoveIndex = Index; MoveIndex > InsertIndex; --MoveIndex)
 			{
 				PoolBaseVMPtr[MoveIndex] = PoolBaseVMPtr[MoveIndex - 1];
+				PoolBaseVMBlock[MoveIndex] = PoolBaseVMBlock[MoveIndex - 1];
 			}
 			PoolBaseVMPtr[InsertIndex] = NewVM;
+			PoolBaseVMBlock[InsertIndex] = NewBLock;
 		}
 		else
 		{
 			PoolBaseVMPtr[Index] = NewVM;
+			PoolBaseVMBlock[Index] = NewBLock;
 		}
 	}
 	HighestPoolBaseVMPtr = PoolBaseVMPtr[BINNED3_SMALL_POOL_COUNT - 1];
@@ -1405,6 +903,41 @@ FMallocBinned3::~FMallocBinned3()
 {
 }
 
+void FMallocBinned3::Commit(uint32 InPoolIndex, void *Ptr, SIZE_T Size)
+{
+#if BINNED3_ALLOCATOR_STATS
+	Binned3Commits++;
+#endif
+#if !BINNED3_USE_SEPARATE_VM_PER_POOL
+	Binned3BaseVMBlock.CommitByPtr(Ptr, Size);
+#else
+	PoolBaseVMBlock[InPoolIndex].CommitByPtr(Ptr, Size);
+#endif
+}
+void FMallocBinned3::Decommit(uint32 InPoolIndex, void *Ptr, SIZE_T Size)
+{
+#if BINNED3_ALLOCATOR_STATS
+	Binned3Decommits++;
+#endif
+#if !BINNED3_USE_SEPARATE_VM_PER_POOL
+	Binned3BaseVMBlock.DecommitByPtr(Ptr, Size);
+#else
+	PoolBaseVMBlock[InPoolIndex].DecommitByPtr(Ptr, Size);
+#endif
+}
+
+void* FMallocBinned3::AllocateMetaDataMemory(SIZE_T Size)
+{
+	LLM_PLATFORM_SCOPE(ELLMTag::FMalloc);
+	size_t VirtualAlignedSize = Align(Size, FPlatformMemory::FPlatformVirtualMemoryBlock::GetVirtualSizeAlignment());
+	FPlatformMemory::FPlatformVirtualMemoryBlock Block = FPlatformMemory::FPlatformVirtualMemoryBlock::AllocateVirtual(VirtualAlignedSize);
+	size_t CommitAlignedSize = Align(Size, FPlatformMemory::FPlatformVirtualMemoryBlock::GetCommitAlignment());
+	Block.Commit(0, CommitAlignedSize);
+	return Block.GetVirtualPointer();
+}
+
+
+
 bool FMallocBinned3::IsInternallyThreadSafe() const
 { 
 	return true;
@@ -1482,19 +1015,25 @@ void* FMallocBinned3::MallocExternal(SIZE_T Size, uint32 Alignment)
 	Size = Align(FMath::Max((SIZE_T)1, Size), Alignment);
 
 	check(FMath::IsPowerOfTwo(Alignment));
-	UE_CLOG(Alignment > MaxAlignmentForMemoryRangeReserve ,LogMemory, Fatal, TEXT("Requested alignment was too large for OS. Alignment=%dKB MaxAlignmentForMemoryRangeReserve=%dKB"), Alignment/1024, MaxAlignmentForMemoryRangeReserve/1024);
-
-	FScopeLock Lock(&Mutex); 
 
 	// Use OS for non-pooled allocations.
-	UPTRINT AlignedSize = Align(Size, FMath::Max(OsAllocationGranularity,Alignment));
+	UPTRINT AlignedSize = Align(Size, FPlatformMemory::FPlatformVirtualMemoryBlock::GetCommitAlignment());
 
 #if BINNED3_TIME_LARGE_BLOCKS
 	double StartTime = FPlatformTime::Seconds();
 #endif
 
 	LLM_PLATFORM_SCOPE(ELLMTag::FMalloc);
-	void* Result = FPlatformMemory::MemoryRangeReserve(AlignedSize, true);
+
+#if USE_CACHED_PAGE_ALLOCATOR_FOR_LARGE_ALLOCS
+	FScopeLock Lock(&Mutex);
+	void* Result = GetCachedOSPageAllocator().Allocate(AlignedSize);
+	check(IsAligned(Result, Alignment));
+#else
+	FPlatformMemory::FPlatformVirtualMemoryBlock Block = FPlatformMemory::FPlatformVirtualMemoryBlock::AllocateVirtual(AlignedSize, Alignment);
+	Block.Commit(0, AlignedSize);
+	void* Result = Block.GetVirtualPointer();
+#endif
 
 #if BINNED3_TIME_LARGE_BLOCKS
 	double Add = FPlatformTime::Seconds() - StartTime;
@@ -1513,6 +1052,10 @@ void* FMallocBinned3::MallocExternal(SIZE_T Size, uint32 Alignment)
 		Private::OutOfMemory(AlignedSize);
 	}
 	check(IsOSAllocation(Result));
+#if! USE_CACHED_PAGE_ALLOCATOR_FOR_LARGE_ALLOCS
+	FScopeLock Lock(&Mutex);
+#endif
+
 
 #if BINNED3_ALLOCATOR_STATS
 	Binned3AllocatedLargePoolMemory += (int64)Size;
@@ -1521,8 +1064,12 @@ void* FMallocBinned3::MallocExternal(SIZE_T Size, uint32 Alignment)
 
 	// Create pool.
 	FPoolInfoLarge* Pool = Private::GetOrCreatePoolInfoLarge(*this, Result);
-	check(Size > 0 && Size <= AlignedSize && AlignedSize >= OsAllocationGranularity);
-	Pool->SetOSAllocationSizes(Size, AlignedSize);
+	check(Size > 0 && Size <= AlignedSize && AlignedSize >= FPlatformMemory::FPlatformVirtualMemoryBlock::GetCommitAlignment());
+#if USE_CACHED_PAGE_ALLOCATOR_FOR_LARGE_ALLOCS
+	Pool->SetOSAllocationSizes(Size, AlignedSize, AlignedSize / FPlatformMemory::FPlatformVirtualMemoryBlock::GetCommitAlignment());
+#else
+	Pool->SetOSAllocationSizes(Size, AlignedSize, Block.GetActualSizeInPages());
+#endif
 
 	return Result;
 }
@@ -1575,35 +1122,28 @@ void* FMallocBinned3::ReallocExternal(void* Ptr, SIZE_T NewSize, uint32 Alignmen
 	{
 		UE_LOG(LogMemory, Fatal, TEXT("FMallocBinned3 Attempt to realloc an unrecognized block %p"), Ptr);
 	}
-	UPTRINT PoolOsBytes = Pool->GetOsAllocatedBytes();
+	UPTRINT PoolOsBytes = Pool->GetOsCommittedBytes();
 	uint32 PoolOSRequestedBytes = Pool->GetOSRequestedBytes();
 	checkf(PoolOSRequestedBytes <= PoolOsBytes, TEXT("FMallocBinned3::ReallocExternal %d %d"), int32(PoolOSRequestedBytes), int32(PoolOsBytes));
 	if (NewSize > PoolOsBytes || // can't fit in the old block
 		(NewSize <= BINNED3_MAX_SMALL_POOL_SIZE && Alignment <= BINNED3_MINIMUM_ALIGNMENT) || // can switch to the small block allocator
 		Align(NewSize, OsAllocationGranularity) < PoolOsBytes) // we can get some pages back
 	{
+		Mutex.Unlock();
 		// Grow or shrink.
 		void* Result = FMallocBinned3::MallocExternal(NewSize, Alignment);
 		SIZE_T CopySize = FMath::Min<SIZE_T>(NewSize, PoolOSRequestedBytes);
-		if (CopySize > 4096) // don't release the lock for small stuff
-		{
-			Mutex.Unlock();
-		}
 		FMemory::Memcpy(Result, Ptr, CopySize);
 		FMallocBinned3::FreeExternal(Ptr);
-		if (CopySize <= 4096) // Release here for small stuff
-		{
-			Mutex.Unlock();
-		}
 		return Result;
 	}
 
 #if BINNED3_ALLOCATOR_STATS
-	Binned3AllocatedLargePoolMemory += ((int64)NewSize) - ((int64)Pool->GetOSRequestedBytes());
+	Binned3AllocatedLargePoolMemory += ((int64)NewSize) - ((int64)PoolOSRequestedBytes);
 	// don't need to change the Binned3AllocatedLargePoolMemoryWAlignment because we didn't reallocate so it's the same size
 #endif
 	
-	Pool->SetOSAllocationSizes(NewSize, PoolOsBytes);
+	Pool->SetOSAllocationSizes(NewSize);
 	Mutex.Unlock();
 	return Ptr;
 }
@@ -1650,27 +1190,44 @@ void FMallocBinned3::FreeExternal(void* Ptr)
 	}
 	else if (Ptr)
 	{
+#if USE_CACHED_PAGE_ALLOCATOR_FOR_LARGE_ALLOCS
 		FScopeLock Lock(&Mutex);
-		FPoolInfoLarge* Pool = Private::FindPoolInfo(*this, Ptr);
-		if (!Pool)
+#endif
+		uint32 VMPages;
 		{
-			UE_LOG(LogMemory, Fatal, TEXT("FMallocBinned3 Attempt to free an unrecognized block %p"), Ptr);
-		}
-		UPTRINT PoolOsBytes = Pool->GetOsAllocatedBytes();
-		uint32 PoolOSRequestedBytes = Pool->GetOSRequestedBytes();
+#if !USE_CACHED_PAGE_ALLOCATOR_FOR_LARGE_ALLOCS
+			FScopeLock Lock(&Mutex);
+#endif
+			FPoolInfoLarge* Pool = Private::FindPoolInfo(*this, Ptr);
+			if (!Pool)
+			{
+				UE_LOG(LogMemory, Fatal, TEXT("FMallocBinned3 Attempt to free an unrecognized block %p"), Ptr);
+			}
+			UPTRINT PoolOsBytes = Pool->GetOsCommittedBytes();
+			uint32 PoolOSRequestedBytes = Pool->GetOSRequestedBytes();
+			VMPages = Pool->GetOsVMPages();
 
 #if BINNED3_ALLOCATOR_STATS
-		Binned3AllocatedLargePoolMemory -= ((int64)PoolOSRequestedBytes);
-		Binned3AllocatedLargePoolMemoryWAlignment -= ((int64)PoolOsBytes);
+			Binned3AllocatedLargePoolMemory -= ((int64)PoolOSRequestedBytes);
+			Binned3AllocatedLargePoolMemoryWAlignment -= ((int64)PoolOsBytes);
 #endif
 
-		checkf(PoolOSRequestedBytes <= PoolOsBytes, TEXT("FMallocBinned3::FreeExternal %d %d"), int32(PoolOSRequestedBytes), int32(PoolOsBytes));
-		Pool->SetCanary(FPoolInfoLarge::ECanary::LargeUnassigned, true, false);
+			checkf(PoolOSRequestedBytes <= PoolOsBytes, TEXT("FMallocBinned3::FreeExternal %d %d"), int32(PoolOSRequestedBytes), int32(PoolOsBytes));
+			Pool->SetCanary(FPoolInfoLarge::ECanary::LargeUnassigned, true, false);
+		}
+
 		// Free an OS allocation.
 #if BINNED3_TIME_LARGE_BLOCKS
 		double StartTime = FPlatformTime::Seconds();
 #endif
-		FPlatformMemory::MemoryRangeFree(Ptr, PoolOsBytes);
+		{
+#if USE_CACHED_PAGE_ALLOCATOR_FOR_LARGE_ALLOCS
+			GetCachedOSPageAllocator().Free(Ptr, VMPages * FPlatformMemory::FPlatformVirtualMemoryBlock::GetCommitAlignment());
+#else
+			FPlatformMemory::FPlatformVirtualMemoryBlock Block(Ptr, VMPages);
+			Block.FreeVirtual();
+#endif
+		}
 #if BINNED3_TIME_LARGE_BLOCKS
 		double Add = FPlatformTime::Seconds() - StartTime;
 		double Old;
@@ -1702,7 +1259,7 @@ bool FMallocBinned3::GetAllocationSizeExternal(void* Ptr, SIZE_T& SizeOut)
 	{
 		UE_LOG(LogMemory, Fatal, TEXT("FMallocBinned3 Attempt to GetAllocationSizeExternal an unrecognized block %p"), Ptr);
 	}
-	UPTRINT PoolOsBytes = Pool->GetOsAllocatedBytes();
+	UPTRINT PoolOsBytes = Pool->GetOsCommittedBytes();
 	uint32 PoolOSRequestedBytes = Pool->GetOSRequestedBytes();
 	checkf(PoolOSRequestedBytes <= PoolOsBytes, TEXT("FMallocBinned3::GetAllocationSizeExternal %d %d"), int32(PoolOSRequestedBytes), int32(PoolOsBytes));
 	SizeOut = PoolOsBytes;
@@ -1760,10 +1317,10 @@ void FMallocBinned3::FlushCurrentThreadCache()
 
 void FMallocBinned3::Trim(bool bTrimThreadCaches)
 {
-	QUICK_SCOPE_CYCLE_COUNTER(STAT_FMallocBinned3_Trim);
 
 	if (GMallocBinned3PerThreadCaches  &&  bTrimThreadCaches)
 	{
+		QUICK_SCOPE_CYCLE_COUNTER(STAT_FMallocBinned3_Trim);
 		//double StartTime = FPlatformTime::Seconds();
 		TFunction<void(ENamedThreads::Type CurrentThread)> Broadcast =
 			[this](ENamedThreads::Type MyThread)
@@ -1793,6 +1350,10 @@ void FMallocBinned3::SetupTLSCachesOnCurrentThread()
 void FMallocBinned3::ClearAndDisableTLSCachesOnCurrentThread()
 {
 	FlushCurrentThreadCache();
+	if (!BINNED3_ALLOW_RUNTIME_TWEAKING && !GMallocBinned3PerThreadCaches)
+	{
+		return;
+	}
 	FPerThreadFreeBlockLists::ClearTLS();
 }
 
@@ -1865,9 +1426,8 @@ void FMallocBinned3::FPerThreadFreeBlockLists::SetTLS()
 	FPerThreadFreeBlockLists* ThreadSingleton = (FPerThreadFreeBlockLists*)FPlatformTLS::GetTlsValue(FMallocBinned3::Binned3TlsSlot);
 	if (!ThreadSingleton)
 	{
-		LLM_PLATFORM_SCOPE(ELLMTag::FMalloc);
 		int64 TLSSize = Align(sizeof(FPerThreadFreeBlockLists), FMallocBinned3::OsAllocationGranularity);
-		ThreadSingleton = new (FPlatformMemory::MemoryRangeReserve(TLSSize, true)) FPerThreadFreeBlockLists();
+		ThreadSingleton = new (FMallocBinned3::AllocateMetaDataMemory(TLSSize)) FPerThreadFreeBlockLists();
 #if BINNED3_ALLOCATOR_STATS
 		Binned3TLSMemory += TLSSize;
 #endif
@@ -1957,6 +1517,8 @@ void FMallocBinned3::DumpAllocatorStats(class FOutputDevice& Ar)
 	Ar.Logf(TEXT("Hash: %fmb"), ((double)Binned3HashMemory) / (1024.0f * 1024.0f));
 	Ar.Logf(TEXT("Free Bits: %fmb"), ((double)Binned3FreeBitsMemory) / (1024.0f * 1024.0f));
 	Ar.Logf(TEXT("TLS: %fmb"), ((double)Binned3TLSMemory) / (1024.0f * 1024.0f));
+	Ar.Logf(TEXT("Slab Commits: %llu"), Binned3Commits.Load());
+	Ar.Logf(TEXT("Slab Decommits: %llu"), Binned3Decommits.Load());
 #if BINNED3_USE_SEPARATE_VM_PER_POOL
 	Ar.Logf(TEXT("BINNED3_USE_SEPARATE_VM_PER_POOL is true - VM is Contiguous = %d"), PoolSearchDiv == 0);
 	if (PoolSearchDiv)
diff --git a/Engine/Source/Runtime/Core/Private/HAL/MallocBinnedArena.cpp b/Engine/Source/Runtime/Core/Private/HAL/MallocBinnedArena.cpp
new file mode 100644
--- /dev/null
+++ b/Engine/Source/Runtime/Core/Private/HAL/MallocBinnedArena.cpp
@@ -0,0 +1,1357 @@
+// Copyright 1998-2018 Epic Games, Inc. All Rights Reserved.
+
+/*=============================================================================
+	MallocBinned.cpp: Binned memory allocator
+=============================================================================*/
+
+#include "HAL/MallocBinnedArena.h"
+
+#if PLATFORM_64BITS && PLATFORM_HAS_FPlatformVirtualMemoryBlock
+#include "Logging/LogMacros.h"
+#include "Misc/ScopeLock.h"
+#include "Templates/Function.h"
+#include "GenericPlatform/GenericPlatformProcess.h"
+#include "Stats/Stats.h"
+#include "HAL/IConsoleManager.h"
+#include "HAL/MemoryMisc.h"
+#include "HAL/PlatformMisc.h"
+
+struct FMallocBinnedArena::FPoolInfoSmall
+{
+	enum ECanary
+	{
+		SmallUnassigned = 0x2,
+		SmallAssigned = 0x1
+	};
+
+	uint32 Canary : 2;
+	uint32 Taken : 15;
+	uint32 NoFirstFreeIndex : 1;
+	uint32 FirstFreeIndex : 14;
+
+	FPoolInfoSmall()
+		: Canary(ECanary::SmallUnassigned)
+		, Taken(0)
+		, NoFirstFreeIndex(1)
+		, FirstFreeIndex(0)
+	{
+		static_assert(sizeof(FPoolInfoSmall) == 4, "Padding fail");
+	}
+	void CheckCanary(ECanary ShouldBe) const
+	{
+		if (Canary != ShouldBe)
+		{
+			UE_LOG(LogMemory, Fatal, TEXT("MallocBinnedArena Corruption Canary was 0x%x, should be 0x%x"), int32(Canary), int32(ShouldBe));
+		}
+	}
+	void SetCanary(ECanary ShouldBe, bool bPreexisting, bool bGuarnteedToBeNew)
+	{
+		if (bPreexisting)
+		{
+			if (bGuarnteedToBeNew)
+			{
+				UE_LOG(LogMemory, Fatal, TEXT("MallocBinnedArena Corruption Canary was 0x%x, should be 0x%x. This block is both preexisting and guaranteed to be new; which makes no sense."), int32(Canary), int32(ShouldBe));
+			}
+			if (ShouldBe == ECanary::SmallUnassigned)
+			{
+				if (Canary != ECanary::SmallAssigned)
+				{
+					UE_LOG(LogMemory, Fatal, TEXT("MallocBinnedArena Corruption Canary was 0x%x, will be 0x%x because this block should be preexisting and in use."), int32(Canary), int32(ShouldBe));
+				}
+			}
+			else if (Canary != ShouldBe)
+			{
+				UE_LOG(LogMemory, Fatal, TEXT("MallocBinnedArena Corruption Canary was 0x%x, should be 0x%x because this block should be preexisting."), int32(Canary), int32(ShouldBe));
+			}
+		}
+		else
+		{
+			if (bGuarnteedToBeNew)
+			{
+				if (Canary != ECanary::SmallUnassigned)
+				{
+					UE_LOG(LogMemory, Fatal, TEXT("MallocBinnedArena Corruption Canary was 0x%x, will be 0x%x. This block is guaranteed to be new yet is it already assigned."), int32(Canary), int32(ShouldBe));
+				}
+			}
+			else if (Canary != ShouldBe && Canary != ECanary::SmallUnassigned)
+			{
+				UE_LOG(LogMemory, Fatal, TEXT("MallocBinnedArena Corruption Canary was 0x%x, will be 0x%x does not have an expected value."), int32(Canary), int32(ShouldBe));
+			}
+		}
+		Canary = ShouldBe;
+	}
+	bool HasFreeRegularBlock() const
+	{
+		CheckCanary(ECanary::SmallAssigned);
+		return !NoFirstFreeIndex;
+	}
+
+	void* AllocateRegularBlock(uint8* BlockOfBlocksPtr, uint32 BlockSize, uint8 MinimumAlignmentShift)
+	{
+		check(HasFreeRegularBlock());
+		++Taken;
+		FFreeBlock* Free = (FFreeBlock*)(BlockOfBlocksPtr + BlockSize * FirstFreeIndex);
+		void* Result = Free->AllocateRegularBlock(MinimumAlignmentShift);
+		if (Free->GetNumFreeRegularBlocks() == 0)
+		{
+			if (Free->NextFreeIndex == MAX_uint32)
+			{
+				FirstFreeIndex = 0;
+				NoFirstFreeIndex = 1;
+			}
+			else
+			{
+				FirstFreeIndex = Free->NextFreeIndex;
+				check(uint32(FirstFreeIndex) == Free->NextFreeIndex);
+				check(((FFreeBlock*)(BlockOfBlocksPtr + BlockSize * FirstFreeIndex))->GetNumFreeRegularBlocks());
+			}
+		}
+
+		return Result;
+	}
+};
+
+struct FMallocBinnedArena::FPoolInfoLarge
+{
+	enum ECanary
+	{
+		LargeUnassigned = 673,
+		LargeAssigned = 3917,
+	};
+
+public:
+	uint32 Canary;      
+private:
+	uint32 VMSizeDivVirtualSizeAlignment;
+	uint32 OSCommitSize;
+	uint32 AllocSize;      // Number of bytes allocated
+public:
+
+	FPoolInfoLarge() :
+		Canary(ECanary::LargeUnassigned),
+		VMSizeDivVirtualSizeAlignment(0),
+		OSCommitSize(0),
+		AllocSize(0)
+	{
+	}
+	void CheckCanary(ECanary ShouldBe) const
+	{
+		if (Canary != ShouldBe)
+		{
+			UE_LOG(LogMemory, Fatal, TEXT("MallocBinnedArena Corruption Canary was 0x%x, should be 0x%x"), int32(Canary), int32(ShouldBe));
+		}
+	}
+	void SetCanary(ECanary ShouldBe, bool bPreexisting, bool bGuarnteedToBeNew)
+	{
+		if (bPreexisting)
+		{
+			if (bGuarnteedToBeNew)
+			{
+				UE_LOG(LogMemory, Fatal, TEXT("MallocBinnedArena Corruption Canary was 0x%x, should be 0x%x. This block is both preexisting and guaranteed to be new; which makes no sense."), int32(Canary), int32(ShouldBe));
+			}
+			if (ShouldBe == ECanary::LargeUnassigned)
+			{
+				if (Canary != ECanary::LargeAssigned)
+				{
+					UE_LOG(LogMemory, Fatal, TEXT("MallocBinnedArena Corruption Canary was 0x%x, will be 0x%x because this block should be preexisting and in use."), int32(Canary), int32(ShouldBe));
+				}
+			}
+			else if (Canary != ShouldBe)
+			{
+				UE_LOG(LogMemory, Fatal, TEXT("MallocBinnedArena Corruption Canary was 0x%x, should be 0x%x because this block should be preexisting."), int32(Canary), int32(ShouldBe));
+			}
+		}
+		else
+		{
+			if (bGuarnteedToBeNew)
+			{
+				if (Canary != ECanary::LargeUnassigned)
+				{
+					UE_LOG(LogMemory, Fatal, TEXT("MallocBinnedArena Corruption Canary was 0x%x, will be 0x%x. This block is guaranteed to be new yet is it already assigned."), int32(Canary), int32(ShouldBe));
+				}
+			}
+			else if (Canary != ShouldBe && Canary != ECanary::LargeUnassigned)
+			{
+				UE_LOG(LogMemory, Fatal, TEXT("MallocBinnedArena Corruption Canary was 0x%x, will be 0x%x does not have an expected value."), int32(Canary), int32(ShouldBe));
+			}
+		}
+		Canary = ShouldBe;
+	}
+	uint32 GetOSRequestedBytes() const
+	{
+		return AllocSize;
+	}
+
+	uint32 GetOsVMSizeDivVirtualSizeAlignment() const
+	{
+		CheckCanary(ECanary::LargeAssigned);
+		return VMSizeDivVirtualSizeAlignment;
+	}
+
+	UPTRINT GetOsCommittedBytes() const
+	{
+		CheckCanary(ECanary::LargeAssigned);
+		return (UPTRINT)OSCommitSize;
+	}
+
+	void SetOSAllocationSizes(uint32 InRequestedBytes, UPTRINT InCommittedBytes, uint32 InVMSizeDivVirtualSizeAlignment)
+	{
+		CheckCanary(ECanary::LargeAssigned);
+		check(InRequestedBytes != 0);                // Shouldn't be pooling zero byte allocations
+		check(InCommittedBytes >= InRequestedBytes); // We must be allocating at least as much as we requested
+
+		AllocSize = InRequestedBytes;
+		OSCommitSize = InCommittedBytes;
+		VMSizeDivVirtualSizeAlignment = InVMSizeDivVirtualSizeAlignment;
+	}
+};
+
+
+/** Hash table struct for retrieving allocation book keeping information */
+struct FMallocBinnedArena::PoolHashBucket
+{
+	UPTRINT BucketIndex;
+	FPoolInfoLarge* FirstPool;
+	PoolHashBucket* Prev;
+	PoolHashBucket* Next;
+
+	PoolHashBucket()
+	{
+		BucketIndex = 0;
+		FirstPool = nullptr;
+		Prev = this;
+		Next = this;
+	}
+
+	void Link(PoolHashBucket* After)
+	{
+		After->Prev = Prev;
+		After->Next = this;
+		Prev->Next = After;
+		this->Prev = After;
+	}
+
+	void Unlink()
+	{
+		Next->Prev = Prev;
+		Prev->Next = Next;
+		Prev = this;
+		Next = this;
+	}
+};
+
+
+
+struct FMallocBinnedArena::Private
+{
+	// Implementation. 
+	static CA_NO_RETURN void OutOfMemory(uint64 Size, uint32 Alignment = 0)
+	{
+		// this is expected not to return
+		FPlatformMemory::OnOutOfMemory(Size, Alignment);
+	}
+
+	/**
+	* Gets the FPoolInfoSmall for a small block memory address. If no valid info exists one is created.
+	*/
+	static FPoolInfoSmall* GetOrCreatePoolInfoSmall(FMallocBinnedArena& Allocator, uint32 InPoolIndex, uint32 BlockOfBlocksIndex)
+	{
+		FPoolInfoSmall*& InfoBlock = Allocator.SmallPoolTables[InPoolIndex].PoolInfos[BlockOfBlocksIndex / Allocator.SmallPoolInfosPerPlatformPage];
+		if (!InfoBlock)
+		{
+			InfoBlock = (FPoolInfoSmall*)FMemory::Malloc(Allocator.ArenaParams.BasePageSize);
+			Allocator.MallocedPointers.Add(InfoBlock);
+			MBA_STAT(Allocator.BinnedArenaPoolInfoMemory += Allocator.ArenaParams.AllocationGranularity;)
+			DefaultConstructItems<FPoolInfoSmall>((void*)InfoBlock, Allocator.SmallPoolInfosPerPlatformPage);
+		}
+
+		FPoolInfoSmall* Result = &InfoBlock[BlockOfBlocksIndex % Allocator.SmallPoolInfosPerPlatformPage];
+
+		bool bGuaranteedToBeNew = false;
+		if (BlockOfBlocksIndex >= Allocator.SmallPoolTables[InPoolIndex].NumEverUsedBlockOfBlocks)
+		{
+			bGuaranteedToBeNew = true;
+			Allocator.SmallPoolTables[InPoolIndex].NumEverUsedBlockOfBlocks = BlockOfBlocksIndex + 1;
+		}
+		Result->SetCanary(FPoolInfoSmall::ECanary::SmallAssigned, false, bGuaranteedToBeNew);
+		return Result;
+	}
+	/**
+	 * Gets the FPoolInfoLarge for a large block memory address. If no valid info exists one is created.
+	 */
+	static FPoolInfoLarge* GetOrCreatePoolInfoLarge(FMallocBinnedArena& Allocator, void* InPtr)
+	{
+		/**
+		 * Creates an array of FPoolInfo structures for tracking allocations.
+		 */
+		auto CreatePoolArray = [](FMallocBinnedArena& LocalAllocator)
+		{
+			uint64 PoolArraySize = LocalAllocator.NumLargePoolsPerPage * sizeof(FPoolInfoLarge);
+
+			void* Result;
+			Result = FMemory::Malloc(PoolArraySize);
+			LocalAllocator.MallocedPointers.Add(Result);
+			MBA_STAT(LocalAllocator.BinnedArenaPoolInfoMemory += PoolArraySize;)
+
+			DefaultConstructItems<FPoolInfoLarge>(Result, LocalAllocator.NumLargePoolsPerPage);
+			return (FPoolInfoLarge*)Result;
+		};
+
+		uint32 BucketIndex;
+		UPTRINT BucketIndexCollision;
+		uint32  PoolIndex;
+		Allocator.PtrToPoolMapping.GetHashBucketAndPoolIndices(InPtr, BucketIndex, BucketIndexCollision, PoolIndex);
+
+		PoolHashBucket* FirstBucket = &Allocator.HashBuckets[BucketIndex];
+		PoolHashBucket* Collision = FirstBucket;
+		do
+		{
+			if (!Collision->FirstPool)
+			{
+				Collision->BucketIndex = BucketIndexCollision;
+				Collision->FirstPool = CreatePoolArray(Allocator);
+				Collision->FirstPool[PoolIndex].SetCanary(FPoolInfoLarge::ECanary::LargeAssigned, false, true);
+				return &Collision->FirstPool[PoolIndex];
+			}
+
+			if (Collision->BucketIndex == BucketIndexCollision)
+			{
+				Collision->FirstPool[PoolIndex].SetCanary(FPoolInfoLarge::ECanary::LargeAssigned, false, false);
+				return &Collision->FirstPool[PoolIndex];
+			}
+
+			Collision = Collision->Next;
+		} while (Collision != FirstBucket);
+
+		// Create a new hash bucket entry
+		if (!Allocator.HashBucketFreeList)
+		{
+			{
+				Allocator.HashBucketFreeList = (PoolHashBucket*)FMemory::Malloc(Allocator.ArenaParams.AllocationGranularity);
+				Allocator.MallocedPointers.Add(Allocator.HashBucketFreeList);
+				MBA_STAT(Allocator.BinnedArenaHashMemory += Allocator.ArenaParams.AllocationGranularity;)
+			}
+
+			for (UPTRINT i = 0, n = Allocator.ArenaParams.AllocationGranularity / sizeof(PoolHashBucket); i < n; ++i)
+			{
+				Allocator.HashBucketFreeList->Link(new (Allocator.HashBucketFreeList + i) PoolHashBucket());
+			}
+		}
+
+		PoolHashBucket* NextFree = Allocator.HashBucketFreeList->Next;
+		PoolHashBucket* NewBucket = Allocator.HashBucketFreeList;
+
+		NewBucket->Unlink();
+
+		if (NextFree == NewBucket)
+		{
+			NextFree = nullptr;
+		}
+		Allocator.HashBucketFreeList = NextFree;
+
+		if (!NewBucket->FirstPool)
+		{
+			NewBucket->FirstPool = CreatePoolArray(Allocator);
+			NewBucket->FirstPool[PoolIndex].SetCanary(FPoolInfoLarge::ECanary::LargeAssigned, false, true);
+		}
+		else
+		{
+			NewBucket->FirstPool[PoolIndex].SetCanary(FPoolInfoLarge::ECanary::LargeAssigned, false, false);
+		}
+
+		NewBucket->BucketIndex = BucketIndexCollision;
+
+		FirstBucket->Link(NewBucket);
+
+		return &NewBucket->FirstPool[PoolIndex];
+	}
+
+	static FPoolInfoLarge* FindPoolInfo(FMallocBinnedArena& Allocator, void* InPtr)
+	{
+		uint32 BucketIndex;
+		UPTRINT BucketIndexCollision;
+		uint32  PoolIndex;
+		Allocator.PtrToPoolMapping.GetHashBucketAndPoolIndices(InPtr, BucketIndex, BucketIndexCollision, PoolIndex);
+
+
+		PoolHashBucket* FirstBucket = &Allocator.HashBuckets[BucketIndex];
+		PoolHashBucket* Collision = FirstBucket;
+		do
+		{
+			if (Collision->BucketIndex == BucketIndexCollision)
+			{
+				return &Collision->FirstPool[PoolIndex];
+			}
+
+			Collision = Collision->Next;
+		} while (Collision != FirstBucket);
+
+		return nullptr;
+	}
+
+
+	static void FreeBundles(FMallocBinnedArena& Allocator, FBundleNode* BundlesToRecycle, uint32 InBlockSize, uint32 InPoolIndex)
+	{
+		FPoolTable& Table = Allocator.SmallPoolTables[InPoolIndex];
+
+		FBundleNode* Bundle = BundlesToRecycle;
+		while (Bundle)
+		{
+			FBundleNode* NextBundle = Bundle->NextBundle;
+
+			FBundleNode* Node = Bundle;
+			do
+			{
+				FBundleNode* NextNode = Node->NextNodeInCurrentBundle;
+
+				uint32 OutBlockOfBlocksIndex;
+				void* BasePtrOfNode = Allocator.BlockOfBlocksPointerFromContainedPtr(Node, Allocator.SmallPoolTables[InPoolIndex].PagesPlatformForBlockOfBlocks, OutBlockOfBlocksIndex);
+				uint32 BlockWithinIndex = (((uint8*)Node) - ((uint8*)BasePtrOfNode)) / Allocator.SmallPoolTables[InPoolIndex].BlockSize;
+
+				FPoolInfoSmall* NodePoolBlock = Allocator.SmallPoolTables[InPoolIndex].PoolInfos[OutBlockOfBlocksIndex / Allocator.SmallPoolInfosPerPlatformPage];
+				if (!NodePoolBlock)
+				{
+					UE_LOG(LogMemory, Fatal, TEXT("FMallocBinnedArena Attempt to free an unrecognized small block %p"), Node);
+				}
+				FPoolInfoSmall* NodePool = &NodePoolBlock[OutBlockOfBlocksIndex % Allocator.SmallPoolInfosPerPlatformPage];
+
+				NodePool->CheckCanary(FPoolInfoSmall::ECanary::SmallAssigned);
+
+				bool bWasExhaused = NodePool->NoFirstFreeIndex;
+
+				// Free a pooled allocation.
+				FFreeBlock* Free = (FFreeBlock*)Node;
+				Free->NumFreeBlocks = 1;
+				Free->NextFreeIndex = NodePool->NoFirstFreeIndex ? MAX_uint32 : NodePool->FirstFreeIndex;
+				Free->BlockSizeShifted = (InBlockSize >> Allocator.ArenaParams.MinimumAlignmentShift);
+				Free->Canary = FFreeBlock::CANARY_VALUE;
+				Free->PoolIndex = InPoolIndex;
+				NodePool->FirstFreeIndex = BlockWithinIndex;
+				NodePool->NoFirstFreeIndex = 0;
+				check(uint32(NodePool->FirstFreeIndex) == BlockWithinIndex);
+
+				// Free this pool.
+				check(NodePool->Taken >= 1);
+				if (--NodePool->Taken == 0)
+				{
+					NodePool->SetCanary(FPoolInfoSmall::ECanary::SmallUnassigned, true, false);
+					Table.BlockOfBlockAllocationBits.FreeBit(OutBlockOfBlocksIndex);
+
+					uint64 AllocSize = Allocator.SmallPoolTables[InPoolIndex].PagesPlatformForBlockOfBlocks * Allocator.ArenaParams.AllocationGranularity;
+
+					if (!bWasExhaused)
+					{
+						Table.BlockOfBlockIsExhausted.AllocBit(OutBlockOfBlocksIndex);
+					}
+
+					Allocator.Decommit(InPoolIndex, BasePtrOfNode, AllocSize);
+					MBA_STAT(Allocator.BinnedArenaAllocatedOSSmallPoolMemory -= AllocSize;)
+				}
+				else if (bWasExhaused)
+				{
+					Table.BlockOfBlockIsExhausted.FreeBit(OutBlockOfBlocksIndex);
+				}
+
+				Node = NextNode;
+			} while (Node);
+
+			Bundle = NextBundle;
+		}
+	}
+
+};
+
+
+FMallocBinnedArena::FPoolInfoSmall* FMallocBinnedArena::PushNewPoolToFront(FMallocBinnedArena& Allocator, uint32 InBlockSize, uint32 InPoolIndex, uint32& OutBlockOfBlocksIndex)
+{
+	FMallocBinnedArena::FPoolTable& Table = Allocator.SmallPoolTables[InPoolIndex];
+	const uint32 BlockOfBlocksSize = Allocator.ArenaParams.AllocationGranularity * Table.PagesPlatformForBlockOfBlocks;
+
+	// Allocate memory.
+
+	uint32 BlockOfBlocksIndex = Table.BlockOfBlockAllocationBits.AllocBit();
+	if (BlockOfBlocksIndex == MAX_uint32)
+	{
+		Private::OutOfMemory(InBlockSize + 1); // The + 1 will hopefully be a hint that we actually ran out of our 1GB space.
+	}
+	uint8* FreePtr = BlockPointerFromIndecies(InPoolIndex, BlockOfBlocksIndex, BlockOfBlocksSize);
+
+	Allocator.Commit(InPoolIndex, FreePtr, BlockOfBlocksSize);
+	uint64 EndOffset = UPTRINT(FreePtr + BlockOfBlocksSize) - UPTRINT(PoolBasePtr(InPoolIndex));
+	if (EndOffset > Table.UnusedAreaOffsetLow)
+	{
+		Table.UnusedAreaOffsetLow = EndOffset;
+	}
+	FFreeBlock* Free = new ((void*)FreePtr) FFreeBlock(BlockOfBlocksSize, InBlockSize, InPoolIndex, Allocator.ArenaParams.MinimumAlignmentShift);
+	MBA_STAT(BinnedArenaAllocatedOSSmallPoolMemory += (int64)BlockOfBlocksSize;)
+	check(IsAligned(Free, Allocator.ArenaParams.AllocationGranularity));
+	// Create pool
+	FPoolInfoSmall* Result = Private::GetOrCreatePoolInfoSmall(*this, InPoolIndex, BlockOfBlocksIndex);
+	Result->CheckCanary(FPoolInfoSmall::ECanary::SmallAssigned);
+	Result->Taken = 0;
+	Result->FirstFreeIndex = 0;
+	Result->NoFirstFreeIndex = 0;
+	Table.BlockOfBlockIsExhausted.FreeBit(BlockOfBlocksIndex);
+
+	OutBlockOfBlocksIndex = BlockOfBlocksIndex;
+
+	return Result;
+}
+
+FMallocBinnedArena::FPoolInfoSmall* FMallocBinnedArena::GetFrontPool(FPoolTable& Table, uint32 InPoolIndex, uint32& OutBlockOfBlocksIndex)
+{
+	OutBlockOfBlocksIndex = Table.BlockOfBlockIsExhausted.NextAllocBit();
+	if (OutBlockOfBlocksIndex == MAX_uint32)
+	{
+		return nullptr;
+	}
+	return Private::GetOrCreatePoolInfoSmall(*this, InPoolIndex, OutBlockOfBlocksIndex);
+}
+
+
+FMallocBinnedArena::FMallocBinnedArena()
+	: HashBucketFreeList(nullptr)
+{
+	MBA_STAT(ConsolidatedMemory = 0;)
+	check(!PLATFORM_32BITS);
+	ArenaParams.BasePageSize = 4096;
+	ArenaParams.AllocationGranularity = FPlatformMemory::FPlatformVirtualMemoryBlock::GetCommitAlignment();
+
+	ArenaParams.ReserveVM = [](SIZE_T Size) -> FPlatformMemory::FPlatformVirtualMemoryBlock
+	{
+		return FPlatformMemory::FPlatformVirtualMemoryBlock::AllocateVirtual(Size, FPlatformMemory::FPlatformVirtualMemoryBlock::GetCommitAlignment());
+	};
+
+	ArenaParams.LargeBlockAlloc = [](SIZE_T Size, SIZE_T Alignment, SIZE_T& OutCommitSize, uint32& OutVMSizeDivVirtualSizeAlignment) -> void*
+	{
+		FPlatformMemory::FPlatformVirtualMemoryBlock Block = FPlatformMemory::FPlatformVirtualMemoryBlock::AllocateVirtual(Size, Alignment);
+		check(IsAligned(Block.GetVirtualPointer(), Alignment)); 
+		OutCommitSize = Align(Size, FPlatformMemory::FPlatformVirtualMemoryBlock::GetCommitAlignment());
+		Block.Commit(0, OutCommitSize);
+		OutVMSizeDivVirtualSizeAlignment = Block.GetActualSizeInPages();
+		return Block.GetVirtualPointer();
+	};
+	ArenaParams.LargeBlockFree = [](void* Ptr, uint32 VMSizeDivVirtualSizeAlignment)
+	{
+		FPlatformMemory::FPlatformVirtualMemoryBlock Block(Ptr, VMSizeDivVirtualSizeAlignment);
+		Block.FreeVirtual();
+	};
+
+}
+
+void FMallocBinnedArena::InitMallocBinned()
+{
+	ArenaParams.MinimumAlignment = uint32(1) << ArenaParams.MinimumAlignmentShift;
+	check(sizeof(FBundleNode) <= ArenaParams.MinimumAlignment);
+
+	ArenaParams.MaxMemoryPerBlockSize = uint64(1) << ArenaParams.MaxMemoryPerBlockSizeShift;
+
+	check(ArenaParams.BasePageSize % sizeof(FPoolInfoLarge) == 0);  // these need to divide evenly!
+	NumLargePoolsPerPage = ArenaParams.BasePageSize / sizeof(FPoolInfoLarge);
+	PtrToPoolMapping.Init(ArenaParams.BasePageSize, NumLargePoolsPerPage, ArenaParams.AddressLimit);
+
+	checkf(FMath::IsPowerOfTwo(ArenaParams.AllocationGranularity), TEXT("OS page size must be a power of two"));
+	checkf(FMath::IsPowerOfTwo(ArenaParams.BasePageSize), TEXT("OS page size must be a power of two"));
+	check(ArenaParams.PoolCount <= 256);
+
+	// Init pool tables.
+
+	TArray<FSizeTableEntry> SizeTable;
+
+	if (ArenaParams.bUseStandardSmallPoolSizes)
+	{
+		SizeTable.AddZeroed(BINNEDCOMMON_NUM_LISTED_SMALL_POOLS + ArenaParams.MaxStandardPoolSize / ArenaParams.BasePageSize); // overestimate
+		ArenaParams.PoolCount = FSizeTableEntry::FillSizeTable(ArenaParams.AllocationGranularity, &SizeTable[0], ArenaParams.BasePageSize, ArenaParams.MinimumAlignment, ArenaParams.MaxStandardPoolSize, ArenaParams.BasePageSize);
+		SizeTable.RemoveAt(ArenaParams.PoolCount, SizeTable.Num() - ArenaParams.PoolCount);
+	}
+	else
+	{
+		ArenaParams.PoolCount = 0;
+	}
+	for (uint32 Size : ArenaParams.AdditionalBlockSizes)
+	{
+		if (Size >= ArenaParams.BasePageSize)
+		{
+			check(Size % ArenaParams.BasePageSize == 0); // we can actually use blocks that are not divisible by the base page size, but we need a bit more work
+			SizeTable.Emplace(Size, ArenaParams.AllocationGranularity, Size / ArenaParams.BasePageSize, ArenaParams.BasePageSize, ArenaParams.MinimumAlignment);
+		}
+		else
+		{
+			// it is difficult to test what would actually make a good bucket size here, wouldn't want a prime number, 33 for example because that would take 33 pages a slab
+			SizeTable.Emplace(Size, ArenaParams.AllocationGranularity, 1, ArenaParams.BasePageSize, ArenaParams.MinimumAlignment);
+		}
+		ArenaParams.PoolCount++;
+	}
+	if (ArenaParams.AdditionalBlockSizes.Num())
+	{
+		Sort(&SizeTable[0], SizeTable.Num());
+	}
+	check(ArenaParams.PoolCount == SizeTable.Num());
+	check(SizeTable.Num() < 256);
+	ArenaParams.MaxPoolSize = SizeTable[ArenaParams.PoolCount - 1].BlockSize;
+
+	check(ArenaParams.BasePageSize % sizeof(FPoolInfoSmall) == 0);
+	SmallPoolInfosPerPlatformPage = ArenaParams.BasePageSize / sizeof(FPoolInfoSmall);
+
+	GGlobalRecycler.Init(ArenaParams.PoolCount);
+	SmallPoolTables.AddDefaulted(ArenaParams.PoolCount);
+	SmallBlockSizesReversedShifted.AddDefaulted(ArenaParams.PoolCount);
+	PoolBaseVMPtr.AddDefaulted(ArenaParams.PoolCount);
+	PoolBaseVMBlocks.AddDefaulted(ArenaParams.PoolCount);
+	MemSizeToIndex.AddDefaulted(1 + (ArenaParams.MaxPoolSize >> ArenaParams.MinimumAlignmentShift));
+
+	ArenaParams.MaximumAlignmentForSmallBlock = ArenaParams.MinimumAlignment;
+	check(ArenaParams.MaximumAlignmentForSmallBlock > 0);
+
+	for (uint32 Index = 0; Index < ArenaParams.PoolCount; ++Index)
+	{
+		checkf(Index == 0 || SizeTable[Index - 1].BlockSize < SizeTable[Index].BlockSize, TEXT("Small block sizes must be strictly increasing"));
+		checkf(SizeTable[Index].BlockSize % ArenaParams.MinimumAlignment == 0, TEXT("Small block size must be a multiple of ArenaParams.MinimumAlignment"));
+
+		// determine the largest alignment that we can cover with a small block
+		while (ArenaParams.MaximumAlignmentForSmallBlock < ArenaParams.AllocationGranularity && IsAligned(SizeTable[Index].BlockSize, ArenaParams.MaximumAlignmentForSmallBlock * 2))
+		{
+			ArenaParams.MaximumAlignmentForSmallBlock *= 2;
+		}
+
+		SmallPoolTables[Index].BlockSize = SizeTable[Index].BlockSize;
+		SmallPoolTables[Index].BlocksPerBlockOfBlocks = SizeTable[Index].BlocksPerBlockOfBlocks;
+		SmallPoolTables[Index].PagesPlatformForBlockOfBlocks = SizeTable[Index].PagesPlatformForBlockOfBlocks;
+
+		SmallPoolTables[Index].UnusedAreaOffsetLow = 0;
+		SmallPoolTables[Index].NumEverUsedBlockOfBlocks = 0;
+
+		int64 TotalNumberOfBlocksOfBlocks = ArenaParams.MaxMemoryPerBlockSize / (SizeTable[Index].PagesPlatformForBlockOfBlocks * ArenaParams.AllocationGranularity);
+
+		int64 MaxPoolInfoMemory = sizeof(FPoolInfoSmall**) * (TotalNumberOfBlocksOfBlocks + SmallPoolInfosPerPlatformPage - 1) / SmallPoolInfosPerPlatformPage;
+		SmallPoolTables[Index].PoolInfos = (FPoolInfoSmall**)FMemory::Malloc(MaxPoolInfoMemory);
+		MallocedPointers.Add(SmallPoolTables[Index].PoolInfos);
+
+		FMemory::Memzero(SmallPoolTables[Index].PoolInfos, MaxPoolInfoMemory);
+		MBA_STAT(BinnedArenaPoolInfoMemory += MaxPoolInfoMemory;)
+
+		{
+			int64 AllocationSize = FBitTree::GetMemoryRequirements(TotalNumberOfBlocksOfBlocks);
+
+			{
+				void *Bits = FMemory::Malloc(AllocationSize);
+				MallocedPointers.Add(Bits);
+				check(Bits);
+				MBA_STAT(BinnedArenaFreeBitsMemory += AllocationSize;)
+				SmallPoolTables[Index].BlockOfBlockAllocationBits.FBitTreeInit(TotalNumberOfBlocksOfBlocks, Bits, AllocationSize, false);
+			}
+			{
+				void *Bits = FMemory::Malloc(AllocationSize);
+				MallocedPointers.Add(Bits);
+				check(Bits);
+				MBA_STAT(BinnedArenaFreeBitsMemory += AllocationSize;)
+				SmallPoolTables[Index].BlockOfBlockIsExhausted.FBitTreeInit(TotalNumberOfBlocksOfBlocks, Bits, AllocationSize, true);
+			}
+		}
+	}
+
+
+	// Set up pool mappings
+	uint8* IndexEntry = &MemSizeToIndex[0];
+	uint32  PoolIndex = 0;
+	for (uint32 Index = 0; Index != 1 + (ArenaParams.MaxPoolSize >> ArenaParams.MinimumAlignmentShift); ++Index)
+	{
+
+		uint32 BlockSize = Index << ArenaParams.MinimumAlignmentShift; // inverse of int32 Index = int32((Size >> ArenaParams.MinimumAlignmentShift));
+		while (SizeTable[PoolIndex].BlockSize < BlockSize)
+		{
+			++PoolIndex;
+			check(PoolIndex != ArenaParams.PoolCount);
+		}
+		check(PoolIndex < 256);
+		*IndexEntry++ = uint8(PoolIndex);
+	}
+	// now reverse the pool sizes for cache coherency
+
+	for (uint32 Index = 0; Index != ArenaParams.PoolCount; ++Index)
+	{
+		uint32 Partner = ArenaParams.PoolCount - Index - 1;
+		SmallBlockSizesReversedShifted[Index] = (SizeTable[Partner].BlockSize >> ArenaParams.MinimumAlignmentShift);
+	}
+	uint64 MaxHashBuckets = PtrToPoolMapping.GetMaxHashBuckets();
+
+	{
+		int64 HashAllocSize = MaxHashBuckets * sizeof(PoolHashBucket);
+		HashBuckets = (PoolHashBucket*)FMemory::Malloc(HashAllocSize);
+		MallocedPointers.Add(HashBuckets);
+
+		MBA_STAT(BinnedArenaHashMemory += HashAllocSize;)
+		verify(HashBuckets);
+	}
+
+	DefaultConstructItems<PoolHashBucket>(HashBuckets, MaxHashBuckets);
+
+	uint8* BinnedArenaBaseVMPtr;
+	if (!ArenaParams.bUseSeparateVMPerPool)
+	{
+		PoolBaseVMBlock = ArenaParams.ReserveVM(ArenaParams.PoolCount * ArenaParams.MaxMemoryPerBlockSize);
+		BinnedArenaBaseVMPtr = (uint8*)PoolBaseVMBlock.GetVirtualPointer();
+	}
+	else
+	{
+		BinnedArenaBaseVMPtr = nullptr;
+	}
+	for (uint32 Index = 0; Index < ArenaParams.PoolCount; ++Index)
+	{
+		uint8* NewVM;
+		FPlatformMemory::FPlatformVirtualMemoryBlock NewBlock;
+		if (BinnedArenaBaseVMPtr)
+		{
+			NewVM = BinnedArenaBaseVMPtr;
+			BinnedArenaBaseVMPtr += ArenaParams.MaxMemoryPerBlockSize;
+		}
+		else
+		{
+			NewBlock = ArenaParams.ReserveVM(ArenaParams.MaxMemoryPerBlockSize);
+			NewVM = (uint8*)NewBlock.GetVirtualPointer();
+		}
+
+		// insertion sort
+		if (Index && NewVM < PoolBaseVMPtr[Index - 1])
+		{
+			uint32 InsertIndex = 0;
+			for (; InsertIndex < Index; ++InsertIndex)
+			{
+				if (NewVM < PoolBaseVMPtr[InsertIndex])
+				{
+					break;
+				}
+			}
+			check(InsertIndex < Index);
+			for (uint32 MoveIndex = Index; MoveIndex > InsertIndex; --MoveIndex)
+			{
+				PoolBaseVMPtr[MoveIndex] = PoolBaseVMPtr[MoveIndex - 1];
+				PoolBaseVMBlocks[MoveIndex] = PoolBaseVMBlocks[MoveIndex - 1];
+			}
+			PoolBaseVMPtr[InsertIndex] = NewVM;
+			PoolBaseVMBlocks[InsertIndex] = NewBlock;
+		}
+		else
+		{
+			PoolBaseVMPtr[Index] = NewVM;
+			PoolBaseVMBlocks[Index] = NewBlock;
+		}
+	}
+	HighestPoolBaseVMPtr = PoolBaseVMPtr[ArenaParams.PoolCount - 1];
+	uint64 TotalGaps = 0;
+	for (uint32 Index = 0; Index < uint32(ArenaParams.PoolCount - 1); ++Index)
+	{
+		check(PoolBaseVMPtr[Index + 1] > PoolBaseVMPtr[Index]); // we sorted it
+		check(PoolBaseVMPtr[Index + 1] >= PoolBaseVMPtr[Index] + ArenaParams.MaxMemoryPerBlockSize); // and blocks are non-overlapping
+		TotalGaps += PoolBaseVMPtr[Index + 1] - (PoolBaseVMPtr[Index] + ArenaParams.MaxMemoryPerBlockSize);
+	}
+	if (TotalGaps == 0)
+	{
+		PoolSearchDiv = 0;
+	}
+	else if (TotalGaps < ArenaParams.MaxMemoryPerBlockSize)
+	{
+		check(ArenaParams.bUseSeparateVMPerPool);
+		PoolSearchDiv = ArenaParams.MaxMemoryPerBlockSize; // the gaps are not significant, ignoring them should give accurate searches
+	}
+	else
+	{
+		check(ArenaParams.bUseSeparateVMPerPool);
+		PoolSearchDiv = ArenaParams.MaxMemoryPerBlockSize + ((TotalGaps + ArenaParams.PoolCount - 2) / (ArenaParams.PoolCount - 1));
+	}
+}
+
+void FMallocBinnedArena::Commit(uint32 InPoolIndex, void *Ptr, SIZE_T Size)
+{
+	if (!ArenaParams.bUseSeparateVMPerPool)
+	{
+		PoolBaseVMBlock.CommitByPtr(Ptr, Size);
+	}
+	else
+	{
+		PoolBaseVMBlocks[InPoolIndex].CommitByPtr(Ptr, Size);
+	}
+}
+void FMallocBinnedArena::Decommit(uint32 InPoolIndex, void *Ptr, SIZE_T Size)
+{
+	if (!ArenaParams.bUseSeparateVMPerPool)
+	{
+		PoolBaseVMBlock.DecommitByPtr(Ptr, Size);
+	}
+	else
+	{
+		PoolBaseVMBlocks[InPoolIndex].DecommitByPtr(Ptr, Size);
+	}
+}
+
+
+FMallocBinnedArena::~FMallocBinnedArena()
+{
+	FScopeLock Lock(&Mutex);
+	FScopeLock Lock2(&GetFreeBlockListsRegistrationMutex());
+
+	MBA_STAT(
+		UE_CLOG(BinnedArenaAllocatedOSSmallPoolMemory > 0, LogCore, Error, TEXT("FMallocBinnedArena leaked small block memory: %fmb"), ((double)BinnedArenaAllocatedOSSmallPoolMemory) / (1024.0f * 1024.0f));
+		UE_CLOG(BinnedArenaAllocatedLargePoolMemoryWAlignment > 0, LogCore, Error, TEXT("FMallocBinnedArena leaked large block memory: %fmb"), ((double)BinnedArenaAllocatedLargePoolMemoryWAlignment) / (1024.0f * 1024.0f));
+	)
+	for (FPerThreadFreeBlockLists* Lists : GetRegisteredFreeBlockLists())
+	{
+		if (Lists)
+		{
+			for (int32 PoolIndex = 0; PoolIndex != ArenaParams.PoolCount; ++PoolIndex)
+			{
+				FBundleNode* Bundles = Lists->PopBundles(PoolIndex);
+				if (Bundles)
+				{
+					Private::FreeBundles(*this, Bundles, PoolIndexToBlockSize(PoolIndex), PoolIndex);
+				}
+			}
+			delete Lists;
+		}
+	}
+
+	if (ArenaParams.bUseSeparateVMPerPool)
+	{
+		for (int32 PoolIndex = 0; PoolIndex != ArenaParams.PoolCount; ++PoolIndex)
+		{
+			PoolBaseVMBlocks[PoolIndex].FreeVirtual();
+		}
+	}
+	else
+	{
+		PoolBaseVMBlock.FreeVirtual();
+	}
+
+
+	for (void* Ptr : MallocedPointers)
+	{
+		FMemory::Free(Ptr);
+	}
+
+	FPlatformTLS::FreeTlsSlot(BinnedArenaTlsSlot);
+}
+
+bool FMallocBinnedArena::IsInternallyThreadSafe() const
+{
+	return true;
+}
+
+void* FMallocBinnedArena::MallocExternal(SIZE_T Size, uint32 Alignment)
+{
+	check(FMath::IsPowerOfTwo(Alignment));
+	checkf(DEFAULT_ALIGNMENT <= ArenaParams.MinimumAlignment, TEXT("DEFAULT_ALIGNMENT is assumed to be zero")); // used below
+
+	if (AdjustSmallBlockSizeForAlignment(Size, Alignment)) // there is some redundant work here... we already adjusted the size for alignment
+	{
+		uint32 PoolIndex = BoundSizeToPoolIndex(Size);
+		FPerThreadFreeBlockLists* Lists = ArenaParams.bPerThreadCaches ? FPerThreadFreeBlockLists::Get(BinnedArenaTlsSlot) : nullptr;
+		if (Lists)
+		{
+			if (Lists->ObtainRecycledPartial(ArenaParams, GGlobalRecycler, PoolIndex))
+			{
+				if (void* Result = Lists->Malloc(PoolIndex))
+				{
+					uint32 BlockSize = PoolIndexToBlockSize(PoolIndex);
+					Lists->AllocatedMemory += BlockSize;
+					checkSlow(IsAligned(Result, Alignment));
+					return Result;
+				}
+			}
+		}
+
+		FScopeLock Lock(&Mutex);
+
+		// Allocate from small object pool.
+		FPoolTable& Table = SmallPoolTables[PoolIndex];
+
+		uint32 BlockOfBlocksIndex = MAX_uint32;
+		FPoolInfoSmall* Pool = GetFrontPool(Table, PoolIndex, BlockOfBlocksIndex);
+		if (!Pool)
+		{
+			Pool = PushNewPoolToFront(*this, Table.BlockSize, PoolIndex, BlockOfBlocksIndex);
+		}
+
+		const uint32 BlockOfBlocksSize = ArenaParams.AllocationGranularity * Table.PagesPlatformForBlockOfBlocks;
+		uint8* BlockOfBlocksPtr = BlockPointerFromIndecies(PoolIndex, BlockOfBlocksIndex, BlockOfBlocksSize);
+
+		void* Result = Pool->AllocateRegularBlock(BlockOfBlocksPtr, Table.BlockSize, ArenaParams.MinimumAlignmentShift);
+		MBA_STAT(BinnedArenaAllocatedSmallPoolMemory += PoolIndexToBlockSize(PoolIndex);)
+		if (ArenaParams.EmptyCacheAllocExtra)
+		{
+			if (Lists)
+			{
+				// prefill the free list with some allocations so we are less likely to hit this slow path with the mutex 
+				for (int32 Index = 0; Index < ArenaParams.EmptyCacheAllocExtra && Pool->HasFreeRegularBlock(); Index++)
+				{
+					if (!Lists->Free(Result, PoolIndex, Table.BlockSize, ArenaParams))
+					{
+						break;
+					}
+					Result = Pool->AllocateRegularBlock(BlockOfBlocksPtr, Table.BlockSize, ArenaParams.MinimumAlignmentShift);
+				}
+			}
+		}
+		if (!Pool->HasFreeRegularBlock())
+		{
+			Table.BlockOfBlockIsExhausted.AllocBit(BlockOfBlocksIndex);
+		}
+		checkSlow(IsAligned(Result, Alignment));
+		return Result;
+	}
+
+	uint32 VMSizeDivVirtualSizeAlignment = 0;
+	SIZE_T CommitSize = 0;
+	void* Result = ArenaParams.LargeBlockAlloc(Size, Alignment, CommitSize, VMSizeDivVirtualSizeAlignment);
+
+	UE_CLOG(!IsAligned(Result, Alignment), LogMemory, Fatal, TEXT("FMallocBinnedArena alignment was too large for OS. Alignment=%d   Ptr=%p"), Alignment, Result);
+
+	if (!Result)
+	{
+		Private::OutOfMemory(Size);
+	}
+	check(IsOSAllocation(Result));
+	FScopeLock Lock(&Mutex);
+
+	MBA_STAT(BinnedArenaAllocatedLargePoolMemory += (int64)Size;)
+	MBA_STAT(BinnedArenaAllocatedLargePoolMemoryWAlignment += CommitSize;)
+
+	// Create pool.
+	FPoolInfoLarge* Pool = Private::GetOrCreatePoolInfoLarge(*this, Result);
+	check(Size > 0 && Size <= CommitSize && CommitSize >= ArenaParams.BasePageSize);
+	Pool->SetOSAllocationSizes(Size, CommitSize, VMSizeDivVirtualSizeAlignment);
+
+	return Result;
+}
+
+
+void* FMallocBinnedArena::ReallocExternal(void* Ptr, SIZE_T NewSize, uint32 Alignment)
+{
+	if (NewSize == 0)
+	{
+		FMallocBinnedArena::FreeExternal(Ptr);
+		return nullptr;
+	}
+	check(DEFAULT_ALIGNMENT <= ArenaParams.MinimumAlignment); // used below
+	check(FMath::IsPowerOfTwo(Alignment));
+
+	uint64 PoolIndex = PoolIndexFromPtr(Ptr);
+	if (PoolIndex < ArenaParams.PoolCount)
+	{
+		check(Ptr); // null is an OS allocation because it will not fall in our VM block
+		uint32 BlockSize = PoolIndexToBlockSize(PoolIndex);
+		if (
+			((NewSize <= BlockSize) & IsAligned(Ptr, Alignment)) && // one branch, not two
+			(PoolIndex == 0 || NewSize > PoolIndexToBlockSize(PoolIndex - 1)))
+		{
+			return Ptr;
+		}
+
+		// Reallocate and copy the data across
+		void* Result = FMallocBinnedArena::MallocExternal(NewSize, Alignment);
+		FMemory::Memcpy(Result, Ptr, FMath::Min<SIZE_T>(NewSize, BlockSize));
+		FMallocBinnedArena::FreeExternal(Ptr);
+		return Result;
+	}
+	if (!Ptr)
+	{
+		void* Result = FMallocBinnedArena::MallocExternal(NewSize, Alignment);
+		return Result;
+	}
+
+	Mutex.Lock();
+
+	// Allocated from OS.
+	FPoolInfoLarge* Pool = Private::FindPoolInfo(*this, Ptr);
+	if (!Pool)
+	{
+		UE_LOG(LogMemory, Fatal, TEXT("FMallocBinnedArena Attempt to realloc an unrecognized block %p"), Ptr);
+	}
+	UPTRINT PoolOsCommittedBytes = Pool->GetOsCommittedBytes();
+	uint32 PoolOSRequestedBytes = Pool->GetOSRequestedBytes();
+	uint32 VMSizeDivVirtualSizeAlignment = Pool->GetOsVMSizeDivVirtualSizeAlignment();
+	checkf(PoolOSRequestedBytes <= PoolOsCommittedBytes, TEXT("FMallocBinnedArena::ReallocExternal %d %d"), int32(PoolOSRequestedBytes), int32(PoolOsCommittedBytes));
+	if (NewSize > PoolOsCommittedBytes || // can't fit in the old block
+		(NewSize <= ArenaParams.MaxPoolSize && Alignment <= ArenaParams.MaximumAlignmentForSmallBlock) || // can switch to the small block allocator
+		!IsAligned(Ptr, Alignment) || // can't reuse the pointer because it doesn't have the right alignment
+		Align(NewSize, ArenaParams.AllocationGranularity) < PoolOsCommittedBytes) // we can get some pages back
+	{
+		// Grow or shrink.
+		Mutex.Unlock();
+		void* Result = FMallocBinnedArena::MallocExternal(NewSize, Alignment);
+		SIZE_T CopySize = FMath::Min<SIZE_T>(NewSize, PoolOSRequestedBytes);
+		FMemory::Memcpy(Result, Ptr, CopySize);
+		FMallocBinnedArena::FreeExternal(Ptr);
+		return Result;
+	}
+
+	MBA_STAT(BinnedArenaAllocatedLargePoolMemory += ((int64)NewSize) - ((int64)Pool->GetOSRequestedBytes());)
+	// don't need to change the BinnedArenaAllocatedLargePoolMemoryWAlignment because we didn't reallocate so it's the same size
+
+	Pool->SetOSAllocationSizes(NewSize, PoolOsCommittedBytes, VMSizeDivVirtualSizeAlignment);
+	Mutex.Unlock();
+	return Ptr;
+}
+
+void FMallocBinnedArena::FreeExternal(void* Ptr)
+{
+	uint64 PoolIndex = PoolIndexFromPtr(Ptr);
+	if (PoolIndex < ArenaParams.PoolCount)
+	{
+		check(Ptr); // null is an OS allocation because it will not fall in our VM block
+		uint32 BlockSize = PoolIndexToBlockSize(PoolIndex);
+
+		FBundleNode* BundlesToRecycle = nullptr;
+		FPerThreadFreeBlockLists* Lists = ArenaParams.bPerThreadCaches ? FPerThreadFreeBlockLists::Get(BinnedArenaTlsSlot) : nullptr;
+		if (Lists)
+		{
+			BundlesToRecycle = Lists->RecycleFullBundle(ArenaParams, GGlobalRecycler, PoolIndex);
+			bool bPushed = Lists->Free(Ptr, PoolIndex, BlockSize, ArenaParams);
+			check(bPushed);
+			Lists->AllocatedMemory -= BlockSize;
+		}
+		else
+		{
+			BundlesToRecycle = (FBundleNode*)Ptr;
+			BundlesToRecycle->NextNodeInCurrentBundle = nullptr;
+		}
+		if (BundlesToRecycle)
+		{
+			BundlesToRecycle->NextBundle = nullptr;
+			FScopeLock Lock(&Mutex);
+			Private::FreeBundles(*this, BundlesToRecycle, BlockSize, PoolIndex);
+			if (!Lists)
+			{
+				// lists track their own stat track them instead in the global stat if we don't have lists
+				MBA_STAT(BinnedArenaAllocatedSmallPoolMemory -= ((int64)(BlockSize));)
+			}
+		}
+	}
+	else if (Ptr)
+	{
+		uint32 VMSizeDivVirtualSizeAlignment;
+
+		{
+			FScopeLock Lock(&Mutex);
+			FPoolInfoLarge* Pool = Private::FindPoolInfo(*this, Ptr);
+			if (!Pool)
+			{
+				UE_LOG(LogMemory, Fatal, TEXT("FMallocBinnedArena Attempt to free an unrecognized block %p"), Ptr);
+			}
+			UPTRINT PoolOsCommittedBytes = Pool->GetOsCommittedBytes();
+			UPTRINT PoolOSRequestedBytes = Pool->GetOSRequestedBytes();
+			VMSizeDivVirtualSizeAlignment = Pool->GetOsVMSizeDivVirtualSizeAlignment();
+
+			MBA_STAT(BinnedArenaAllocatedLargePoolMemory -= ((int64)PoolOSRequestedBytes);)
+			MBA_STAT(BinnedArenaAllocatedLargePoolMemoryWAlignment -= ((int64)PoolOsCommittedBytes);)
+
+			checkf(PoolOSRequestedBytes <= PoolOsCommittedBytes, TEXT("FMallocBinnedArena::FreeExternal %d %d"), int32(PoolOSRequestedBytes), int32(PoolOsCommittedBytes));
+			Pool->SetCanary(FPoolInfoLarge::ECanary::LargeUnassigned, true, false);
+		}
+		// Free an OS allocation.
+		ArenaParams.LargeBlockFree(Ptr, VMSizeDivVirtualSizeAlignment);
+	}
+}
+
+bool FMallocBinnedArena::GetAllocationSizeExternal(void* Ptr, SIZE_T& SizeOut)
+{
+	uint64 PoolIndex = PoolIndexFromPtr(Ptr);
+	if (PoolIndex < ArenaParams.PoolCount)
+	{
+		check(Ptr); // null is an OS allocation because it will not fall in our VM block
+		SizeOut = PoolIndexToBlockSize(PoolIndex);
+		return true;
+	}
+	if (!Ptr)
+	{
+		return false;
+	}
+	FScopeLock Lock(&Mutex);
+	FPoolInfoLarge* Pool = Private::FindPoolInfo(*this, Ptr);
+	if (!Pool)
+	{
+		UE_LOG(LogMemory, Fatal, TEXT("FMallocBinnedArena Attempt to GetAllocationSizeExternal an unrecognized block %p"), Ptr);
+	}
+	UPTRINT PoolOsBytes = Pool->GetOsCommittedBytes();
+	uint32 PoolOSRequestedBytes = Pool->GetOSRequestedBytes();
+	checkf(PoolOSRequestedBytes <= PoolOsBytes, TEXT("FMallocBinnedArena::GetAllocationSizeExternal %d %d"), int32(PoolOSRequestedBytes), int32(PoolOsBytes));
+	SizeOut = PoolOsBytes;
+	return true;
+}
+
+bool FMallocBinnedArena::ValidateHeap()
+{
+	// Not implemented
+	// NumEverUsedBlockOfBlocks gives us all of the information we need to examine each pool, so it is doable.
+	return true;
+}
+
+const TCHAR* FMallocBinnedArena::GetDescriptiveName()
+{
+	return TEXT("BinnedArena");
+}
+
+void FMallocBinnedArena::FlushCurrentThreadCache()
+{
+	double StartTimeInner = FPlatformTime::Seconds();
+	QUICK_SCOPE_CYCLE_COUNTER(STAT_FMallocBinnedArena_FlushCurrentThreadCache);
+	FPerThreadFreeBlockLists* Lists = FPerThreadFreeBlockLists::Get(BinnedArenaTlsSlot);
+
+	float WaitForMutexTime = 0.0f;
+	float WaitForMutexAndTrimTime = 0.0f;
+
+	if (Lists)
+	{
+		FScopeLock Lock(&Mutex);
+		WaitForMutexTime = FPlatformTime::Seconds() - StartTimeInner;
+		for (int32 PoolIndex = 0; PoolIndex != ArenaParams.PoolCount; ++PoolIndex)
+		{
+			FBundleNode* Bundles = Lists->PopBundles(PoolIndex);
+			if (Bundles)
+			{
+				Private::FreeBundles(*this, Bundles, PoolIndexToBlockSize(PoolIndex), PoolIndex);
+			}
+		}
+		WaitForMutexAndTrimTime = FPlatformTime::Seconds() - StartTimeInner;
+	}
+
+	// These logs must happen outside the above mutex to avoid deadlocks
+	if (WaitForMutexTime > 0.02f)
+	{
+		UE_LOG(LogMemory, Warning, TEXT("FMallocBinnedArena took %6.2fms to wait for mutex for trim."), WaitForMutexTime * 1000.0f);
+	}
+	if (WaitForMutexAndTrimTime > 0.02f)
+	{
+		UE_LOG(LogMemory, Warning, TEXT("FMallocBinnedArena took %6.2fms to wait for mutex AND trim."), WaitForMutexAndTrimTime * 1000.0f);
+	}
+}
+
+#include "Async/TaskGraphInterfaces.h"
+
+void FMallocBinnedArena::Trim(bool bTrimThreadCaches)
+{
+	if (bTrimThreadCaches && ArenaParams.bPerThreadCaches)
+	{
+		QUICK_SCOPE_CYCLE_COUNTER(STAT_FMallocBinnedArena_Trim);
+		//double StartTime = FPlatformTime::Seconds();
+		TFunction<void(ENamedThreads::Type CurrentThread)> Broadcast =
+			[this](ENamedThreads::Type MyThread)
+		{
+			FlushCurrentThreadCache();
+		};
+		// Skip task threads on desktop platforms as it is too slow and they don't have much memory
+		FTaskGraphInterface::BroadcastSlow_OnlyUseForSpecialPurposes(!PLATFORM_DESKTOP, false, Broadcast);
+		//UE_LOG(LogTemp, Display, TEXT("Trim Broadcast = %6.2fms"), 1000.0f * float(FPlatformTime::Seconds() - StartTime));
+	}
+}
+
+void FMallocBinnedArena::SetupTLSCachesOnCurrentThread()
+{
+	if (!ArenaParams.bPerThreadCaches)
+	{
+		return;
+	}
+	if (!BinnedArenaTlsSlot)
+	{
+		BinnedArenaTlsSlot = FPlatformTLS::AllocTlsSlot();
+	}
+	check(BinnedArenaTlsSlot);
+	FPerThreadFreeBlockLists::SetTLS(*this);
+}
+
+void FMallocBinnedArena::ClearAndDisableTLSCachesOnCurrentThread()
+{
+	FlushCurrentThreadCache();
+	MBA_STAT(ConsolidatedMemory +=) FPerThreadFreeBlockLists::ClearTLS(*this);
+}
+
+
+bool FMallocBinnedArena::FFreeBlockList::ObtainPartial(FArenaParams& LocalArenaParams, FGlobalRecycler& GlobalRecycler, uint32 InPoolIndex)
+{
+	if (!PartialBundle.Head)
+	{
+		PartialBundle.Count = 0;
+		PartialBundle.Head = GlobalRecycler.PopBundle(LocalArenaParams.MaxGlobalBundles, InPoolIndex);
+		if (PartialBundle.Head)
+		{
+			PartialBundle.Count = PartialBundle.Head->Count;
+			PartialBundle.Head->NextBundle = nullptr;
+			return true;
+		}
+		return false;
+	}
+	return true;
+}
+
+FMallocBinnedArena::FBundleNode* FMallocBinnedArena::FFreeBlockList::RecyleFull(FArenaParams& LocalArenaParams, FGlobalRecycler& GlobalRecycler, uint32 InPoolIndex)
+{
+	FMallocBinnedArena::FBundleNode* Result = nullptr;
+	if (FullBundle.Head)
+	{
+		FullBundle.Head->Count = FullBundle.Count;
+		if (!GlobalRecycler.PushBundle(LocalArenaParams.MaxGlobalBundles, InPoolIndex, FullBundle.Head))
+		{
+			Result = FullBundle.Head;
+			Result->NextBundle = nullptr;
+		}
+		FullBundle.Reset();
+	}
+	return Result;
+}
+
+FMallocBinnedArena::FBundleNode* FMallocBinnedArena::FFreeBlockList::PopBundles(uint32 InPoolIndex)
+{
+	FBundleNode* Partial = PartialBundle.Head;
+	if (Partial)
+	{
+		PartialBundle.Reset();
+		Partial->NextBundle = nullptr;
+	}
+
+	FBundleNode* Full = FullBundle.Head;
+	if (Full)
+	{
+		FullBundle.Reset();
+		Full->NextBundle = nullptr;
+	}
+
+	FBundleNode* Result = Partial;
+	if (Result)
+	{
+		Result->NextBundle = Full;
+	}
+	else
+	{
+		Result = Full;
+	}
+
+	return Result;
+}
+
+void FMallocBinnedArena::FPerThreadFreeBlockLists::SetTLS(FMallocBinnedArena& Allocator)
+{
+	uint32 BinnedArenaTlsSlot = Allocator.BinnedArenaTlsSlot;
+	check(BinnedArenaTlsSlot);
+	FPerThreadFreeBlockLists* ThreadSingleton = (FPerThreadFreeBlockLists*)FPlatformTLS::GetTlsValue(BinnedArenaTlsSlot);
+	if (!ThreadSingleton)
+	{
+		int64 TLSSize = sizeof(FPerThreadFreeBlockLists);
+		ThreadSingleton = new FPerThreadFreeBlockLists(Allocator.ArenaParams.PoolCount);
+		MBA_STAT(Allocator.BinnedArenaTLSMemory += TLSSize;)
+		verify(ThreadSingleton);
+		FPlatformTLS::SetTlsValue(BinnedArenaTlsSlot, ThreadSingleton);
+		Allocator.RegisterThreadFreeBlockLists(ThreadSingleton);
+	}
+}
+
+int64 FMallocBinnedArena::FPerThreadFreeBlockLists::ClearTLS(FMallocBinnedArena& Allocator)
+{
+	uint32 BinnedArenaTlsSlot = Allocator.BinnedArenaTlsSlot;
+	check(BinnedArenaTlsSlot);
+	int64 Result = 0;
+	FPerThreadFreeBlockLists* ThreadSingleton = (FPerThreadFreeBlockLists*)FPlatformTLS::GetTlsValue(BinnedArenaTlsSlot);
+	if (ThreadSingleton)
+	{
+		int64 TLSSize = sizeof(FPerThreadFreeBlockLists);
+		MBA_STAT(Allocator.BinnedArenaTLSMemory -= TLSSize;)
+		Result = Allocator.UnregisterThreadFreeBlockLists(ThreadSingleton);
+	}
+	FPlatformTLS::SetTlsValue(BinnedArenaTlsSlot, nullptr);
+	return Result;
+}
+
+void FMallocBinnedArena::FFreeBlock::CanaryFail() const
+{
+	UE_LOG(LogMemory, Fatal, TEXT("FMallocBinnedArena Attempt to realloc an unrecognized block %p   canary == 0x%x != 0x%x"), (void*)this, (int32)Canary, (int32)FMallocBinnedArena::FFreeBlock::CANARY_VALUE);
+}
+
+#if COLLECT_BINNEDARENA_STATS
+int64 FMallocBinnedArena::GetTotalAllocatedSmallPoolMemory()
+{
+	int64 FreeBlockAllocatedMemory = 0;
+	{
+		FScopeLock Lock(&GetFreeBlockListsRegistrationMutex());
+		for (const FPerThreadFreeBlockLists* FreeBlockLists : GetRegisteredFreeBlockLists())
+		{
+			FreeBlockAllocatedMemory += FreeBlockLists->AllocatedMemory;
+		}
+		FreeBlockAllocatedMemory += ConsolidatedMemory;
+	}
+
+	return BinnedArenaAllocatedSmallPoolMemory + FreeBlockAllocatedMemory;
+}
+#endif
+
+void FMallocBinnedArena::GetAllocatorStats(FGenericMemoryStats& OutStats)
+{
+	MBA_STAT(
+		int64 TotalAllocatedSmallPoolMemory = GetTotalAllocatedSmallPoolMemory();
+
+		OutStats.Add(TEXT("BinnedArenaAllocatedSmallPoolMemory"), TotalAllocatedSmallPoolMemory);
+		OutStats.Add(TEXT("BinnedArenaAllocatedOSSmallPoolMemory"), BinnedArenaAllocatedOSSmallPoolMemory);
+		OutStats.Add(TEXT("BinnedArenaAllocatedLargePoolMemory"), BinnedArenaAllocatedLargePoolMemory);
+		OutStats.Add(TEXT("BinnedArenaAllocatedLargePoolMemoryWAlignment"), BinnedArenaAllocatedLargePoolMemoryWAlignment);
+
+		uint64 TotalAllocated = TotalAllocatedSmallPoolMemory + BinnedArenaAllocatedLargePoolMemory;
+		uint64 TotalOSAllocated = BinnedArenaAllocatedOSSmallPoolMemory + BinnedArenaAllocatedLargePoolMemoryWAlignment;
+
+		OutStats.Add(TEXT("TotalAllocated"), TotalAllocated);
+		OutStats.Add(TEXT("TotalOSAllocated"), TotalOSAllocated);
+	)
+	FMalloc::GetAllocatorStats(OutStats);
+}
+
+void FMallocBinnedArena::DumpAllocatorStats(class FOutputDevice& Ar)
+{
+
+	Ar.Logf(TEXT("FMallocBinnedArena Mem report"));
+	Ar.Logf(TEXT("Constants.BinnedAllocationGranularity = %d"), int32(ArenaParams.AllocationGranularity));
+	Ar.Logf(TEXT("ArenaParams.MaxPoolSize = %d"), int32(ArenaParams.MaxPoolSize));
+	Ar.Logf(TEXT("MAX_MEMORY_PER_BLOCK_SIZE = %llu"), uint64(ArenaParams.MaxMemoryPerBlockSize));
+
+	MBA_STAT(
+		int64 TotalAllocatedSmallPoolMemory = GetTotalAllocatedSmallPoolMemory();
+
+		Ar.Logf(TEXT("Small Pool Allocations: %fmb  (%lld) (including block size padding)"), ((double)TotalAllocatedSmallPoolMemory) / (1024.0f * 1024.0f), TotalAllocatedSmallPoolMemory);
+		Ar.Logf(TEXT("Small Pool OS Allocated: %fmb"), ((double)BinnedArenaAllocatedOSSmallPoolMemory) / (1024.0f * 1024.0f));
+		Ar.Logf(TEXT("Large Pool Requested Allocations: %fmb"), ((double)BinnedArenaAllocatedLargePoolMemory) / (1024.0f * 1024.0f));
+		Ar.Logf(TEXT("Large Pool OS Allocated: %fmb"), ((double)BinnedArenaAllocatedLargePoolMemoryWAlignment) / (1024.0f * 1024.0f));
+		Ar.Logf(TEXT("PoolInfo: %fmb"), ((double)BinnedArenaPoolInfoMemory) / (1024.0f * 1024.0f));
+		Ar.Logf(TEXT("Hash: %fmb"), ((double)BinnedArenaHashMemory) / (1024.0f * 1024.0f));
+		Ar.Logf(TEXT("Free Bits: %fmb"), ((double)BinnedArenaFreeBitsMemory) / (1024.0f * 1024.0f));
+		Ar.Logf(TEXT("TLS: %fmb"), ((double)BinnedArenaTLSMemory) / (1024.0f * 1024.0f));
+		Ar.Logf(TEXT("Total allocated from OS: %fmb"),
+			((double)
+				BinnedArenaAllocatedOSSmallPoolMemory + BinnedArenaAllocatedLargePoolMemoryWAlignment + BinnedArenaPoolInfoMemory + BinnedArenaHashMemory + BinnedArenaFreeBitsMemory + BinnedArenaTLSMemory
+				) / (1024.0f * 1024.0f));
+	)
+	Ar.Logf(TEXT("BINNEDARENA_USE_SEPARATE_VM_PER_POOL is true - VM is Contiguous = %d"), PoolSearchDiv == 0);
+	if (PoolSearchDiv)
+	{
+		uint64 TotalMem = PoolBaseVMPtr[ArenaParams.PoolCount - 1] + ArenaParams.MaxMemoryPerBlockSize - PoolBaseVMPtr[0];
+		uint64 MinimumMem = uint64(ArenaParams.PoolCount) * ArenaParams.MaxMemoryPerBlockSize;
+		Ar.Logf(TEXT("Percent of gaps in the address range %6.4f  (hopefully < 1, or the searches above will suffer)"), 100.0f * (1.0f - float(MinimumMem) / float(TotalMem)));
+	}
+
+	for (int32 PoolIndex = 0; PoolIndex < ArenaParams.PoolCount; PoolIndex++)
+	{
+
+		int64 VM = SmallPoolTables[PoolIndex].UnusedAreaOffsetLow;
+		uint32 CommittedBlocks = SmallPoolTables[PoolIndex].BlockOfBlockAllocationBits.CountOnes(SmallPoolTables[PoolIndex].NumEverUsedBlockOfBlocks);
+		uint32 PartialBlocks = SmallPoolTables[PoolIndex].NumEverUsedBlockOfBlocks - SmallPoolTables[PoolIndex].BlockOfBlockIsExhausted.CountOnes(SmallPoolTables[PoolIndex].NumEverUsedBlockOfBlocks);
+		uint32 FullBlocks = CommittedBlocks - PartialBlocks;
+		int64 ComittedVM = VM - (SmallPoolTables[PoolIndex].NumEverUsedBlockOfBlocks - CommittedBlocks) * SmallPoolTables[PoolIndex].PagesPlatformForBlockOfBlocks * ArenaParams.AllocationGranularity;
+
+
+		Ar.Logf(TEXT("Pool %2d   Size %6d   UsedVM %3dMB  CommittedVM %3dMB  HighSlabs %6d  CommittedSlabs %6d  FullSlabs %6d  PartialSlabs  %6d"),
+			PoolIndex,
+			PoolIndexToBlockSize(PoolIndex),
+			VM / (1024 * 1024),
+			ComittedVM / (1024 * 1024),
+			SmallPoolTables[PoolIndex].NumEverUsedBlockOfBlocks,
+			CommittedBlocks,
+			FullBlocks,
+			PartialBlocks
+		);
+	}
+}
+#endif
\ No newline at end of file
diff --git a/Engine/Source/Runtime/Core/Private/HAL/MallocBinnedCommon.cpp b/Engine/Source/Runtime/Core/Private/HAL/MallocBinnedCommon.cpp
new file mode 100644
--- /dev/null
+++ b/Engine/Source/Runtime/Core/Private/HAL/MallocBinnedCommon.cpp
@@ -0,0 +1,476 @@
+// Copyright 1998-2018 Epic Games, Inc. All Rights Reserved.
+
+/*=============================================================================
+	MallocBinned.cpp: Binned memory allocator
+=============================================================================*/
+
+#include "HAL/MallocBinnedCommon.h"
+#include "Misc/AssertionMacros.h"
+#include "Math/NumericLimits.h"
+#include "Templates/AlignmentTemplates.h"
+#include "Templates/UnrealTemplate.h"
+#include "Templates/Sorting.h"
+
+#if PLATFORM_HAS_FPlatformVirtualMemoryBlock
+
+// Block sizes are based around getting the maximum amount of allocations per pool, with as little alignment waste as possible.
+// Block sizes should be close to even divisors of the system page size, and well distributed.
+// They must be 16-byte aligned as well.
+static uint32 BinnedCommonSmallBlockSizes4k[] = 
+{
+	16, 32, 48, 64, 80, 96, 112, 128, 160, // +16
+	192, 224, 256, 288, 320, // +32
+	368,  // /11 ish
+	400,  // /10 ish
+	448,  // /9 ish
+	512,  // /8
+	576,  // /7 ish
+	672,  // /6 ish
+	816,  // /5 ish 
+	1024, // /4 
+	1360, // /3 ish
+	2048, // /2
+	4096 // /1
+};
+
+static uint32 BinnedCommonSmallBlockSizes8k[] =
+{
+	736,  // /11 ish
+	1168, // /7 ish
+	1632, // /5 ish
+	2720, // /3 ish
+	8192  // /1
+};
+
+static uint32 BinnedCommonSmallBlockSizes12k[] =
+{
+	//1104, // /11 ish
+	//1216, // /10 ish
+	1536, // /8
+	1744, // /7 ish
+	2448, // /5 ish
+	3072, // /4
+	6144, // /2
+	12288 // /1
+};
+
+static uint32 BinnedCommonSmallBlockSizes16k[] =
+{
+	//1488, // /11 ish
+	//1808, // /9 ish
+	//2336, // /7 ish
+	3264, // /5 ish
+	5456, // /3 ish
+	16384 // /1
+};
+
+static uint32 BinnedCommonSmallBlockSizes20k[] =
+{
+	// 2912, // /7 ish
+	//3408, // /6 ish
+	5120, // /4
+	//6186, // /3 ish
+	10240, // /2
+	20480  // /1
+};
+
+static uint32 BinnedCommonSmallBlockSizes24k[] = // 1 total
+{
+	24576  // /1
+};
+
+static uint32 BinnedCommonSmallBlockSizes28k[] = // 6 total
+{
+	4768,  // /6 ish
+	5728,  // /5 ish
+	7168,  // /4
+	9552,  // /3
+	14336, // /2
+	28672  // /1
+};
+
+FSizeTableEntry::FSizeTableEntry(uint32 InBlockSize, uint64 PlatformPageSize, uint8 Pages4k, uint32 BasePageSize, uint32 MinimumAlignment)
+	: BlockSize(InBlockSize)
+{
+	check((PlatformPageSize & (BasePageSize - 1)) == 0 && PlatformPageSize >= BasePageSize && InBlockSize % MinimumAlignment == 0);
+
+	uint64 Page4kPerPlatformPage = PlatformPageSize / BasePageSize;
+
+	PagesPlatformForBlockOfBlocks = 0;
+	while (true)
+	{
+		check(PagesPlatformForBlockOfBlocks < MAX_uint8);
+		PagesPlatformForBlockOfBlocks++;
+		if (PagesPlatformForBlockOfBlocks * Page4kPerPlatformPage < Pages4k)
+		{
+			continue;
+		}
+		if (PagesPlatformForBlockOfBlocks * Page4kPerPlatformPage % Pages4k != 0)
+		{
+			continue;
+		}
+		break;
+	}
+	check((PlatformPageSize * PagesPlatformForBlockOfBlocks) / BlockSize <= MAX_uint16);
+	BlocksPerBlockOfBlocks = (PlatformPageSize * PagesPlatformForBlockOfBlocks) / BlockSize;
+}
+
+uint8 FSizeTableEntry::FillSizeTable(uint64 PlatformPageSize, FSizeTableEntry* SizeTable, uint32 BasePageSize, uint32 MinimumAlignment, uint32 MaxSize, uint32 SizeIncrement)
+{
+	int32 Index = 0;
+
+	for (int32 Sub = 0; Sub < ARRAY_COUNT(BinnedCommonSmallBlockSizes4k); Sub++)
+	{
+		SizeTable[Index++] = FSizeTableEntry(BinnedCommonSmallBlockSizes4k[Sub], PlatformPageSize, 1, BasePageSize, MinimumAlignment);
+	}
+	for (int32 Sub = 0; Sub < ARRAY_COUNT(BinnedCommonSmallBlockSizes8k); Sub++)
+	{
+		SizeTable[Index++] = FSizeTableEntry(BinnedCommonSmallBlockSizes8k[Sub], PlatformPageSize, 2, BasePageSize, MinimumAlignment);
+	}
+	for (int32 Sub = 0; Sub < ARRAY_COUNT(BinnedCommonSmallBlockSizes12k); Sub++)
+	{
+		SizeTable[Index++] = FSizeTableEntry(BinnedCommonSmallBlockSizes12k[Sub], PlatformPageSize, 3, BasePageSize, MinimumAlignment);
+	}
+	for (int32 Sub = 0; Sub < ARRAY_COUNT(BinnedCommonSmallBlockSizes16k); Sub++)
+	{
+		SizeTable[Index++] = FSizeTableEntry(BinnedCommonSmallBlockSizes16k[Sub], PlatformPageSize, 4, BasePageSize, MinimumAlignment);
+	}
+	for (int32 Sub = 0; Sub < ARRAY_COUNT(BinnedCommonSmallBlockSizes20k); Sub++)
+	{
+		SizeTable[Index++] = FSizeTableEntry(BinnedCommonSmallBlockSizes20k[Sub], PlatformPageSize, 5, BasePageSize, MinimumAlignment);
+	}
+	for (int32 Sub = 0; Sub < ARRAY_COUNT(BinnedCommonSmallBlockSizes24k); Sub++)
+	{
+		SizeTable[Index++] = FSizeTableEntry(BinnedCommonSmallBlockSizes24k[Sub], PlatformPageSize, 6, BasePageSize, MinimumAlignment);
+	}
+	for (int32 Sub = 0; Sub < ARRAY_COUNT(BinnedCommonSmallBlockSizes28k); Sub++)
+	{
+		SizeTable[Index++] = FSizeTableEntry(BinnedCommonSmallBlockSizes28k[Sub], PlatformPageSize, 7, BasePageSize, MinimumAlignment);
+	}
+	Sort(&SizeTable[0], Index);
+	check(SizeTable[Index - 1].BlockSize == BINNEDCOMMON_MAX_LISTED_SMALL_POOL_SIZE);
+	check(IsAligned(BINNEDCOMMON_MAX_LISTED_SMALL_POOL_SIZE, BasePageSize));
+	for (uint32 Size = BINNEDCOMMON_MAX_LISTED_SMALL_POOL_SIZE + BasePageSize; Size <= MaxSize; Size += SizeIncrement)
+	{
+		SizeTable[Index++] = FSizeTableEntry(Size, PlatformPageSize, Size / BasePageSize, BasePageSize, MinimumAlignment);
+	}
+	check(Index < 256);
+	return Index;
+}
+
+uint32 FBitTree::GetMemoryRequirements(uint32 DesiredCapacity)
+{
+	uint32 AllocationSize = 8;
+	uint32 RowsUint64s = 1;
+	uint32 Capacity = 64;
+	uint32 OffsetOfLastRow = 0;
+
+	while (Capacity < DesiredCapacity)
+	{
+		Capacity *= 64;
+		RowsUint64s *= 64;
+		OffsetOfLastRow = AllocationSize / 8;
+		AllocationSize += 8 * RowsUint64s;
+	}
+
+	uint32 LastRowTotal = (AllocationSize - OffsetOfLastRow * 8) * 8;
+	uint32 ExtraBits = LastRowTotal - DesiredCapacity;
+	AllocationSize -= (ExtraBits / 64) * 8;
+	return AllocationSize;
+}
+
+void FBitTree::FBitTreeInit(uint32 InDesiredCapacity, void * Memory, uint32 MemorySize, bool InitialValue)
+{
+	Bits = (uint64*)Memory;
+	DesiredCapacity = InDesiredCapacity;
+	AllocationSize = 8;
+	Rows = 1;
+	uint32 RowsUint64s = 1;
+	Capacity = 64;
+	OffsetOfLastRow = 0;
+
+	uint32 RowOffsets[10]; // 10 is way more than enough
+	RowOffsets[0] = 0;
+	uint32 RowNum[10]; // 10 is way more than enough
+	RowNum[0] = 1;
+
+	while (Capacity < DesiredCapacity)
+	{
+		Capacity *= 64;
+		RowsUint64s *= 64;
+		OffsetOfLastRow = AllocationSize / 8;
+		check(Rows < 10);
+		RowOffsets[Rows] = OffsetOfLastRow;
+		RowNum[Rows] = RowsUint64s;
+		AllocationSize += 8 * RowsUint64s;
+		Rows++;
+	}
+
+	uint32 LastRowTotal = (AllocationSize - OffsetOfLastRow * 8) * 8;
+	uint32 ExtraBits = LastRowTotal - DesiredCapacity;
+	AllocationSize -= (ExtraBits / 64) * 8;
+	check(AllocationSize <= MemorySize && Bits);
+
+	FMemory::Memset(Bits, InitialValue ? 0xff : 0, AllocationSize);
+
+	if (!InitialValue)
+	{
+		// we fill everything beyond the desired size with occupied
+		uint32 ItemsPerBit = 64;
+		for (int32 FillRow = Rows - 2; FillRow >= 0; FillRow--)
+		{
+			uint32 NeededOneBits = RowNum[FillRow] * 64 - (DesiredCapacity + ItemsPerBit - 1) / ItemsPerBit;
+			uint32 NeededOne64s = NeededOneBits / 64;
+			NeededOneBits %= 64;
+			for (uint32 Fill = RowNum[FillRow] - NeededOne64s; Fill < RowNum[FillRow]; Fill++)
+			{
+				Bits[RowOffsets[FillRow] + Fill] = MAX_uint64;
+			}
+			if (NeededOneBits)
+			{
+				Bits[RowOffsets[FillRow] + RowNum[FillRow] - NeededOne64s - 1] = (MAX_uint64 << (64 - NeededOneBits));
+			}
+			ItemsPerBit *= 64;
+		}
+
+		if (DesiredCapacity % 64)
+		{
+			Bits[AllocationSize / 8 - 1] = (MAX_uint64 << (DesiredCapacity % 64));
+		}
+	}
+}
+
+uint32 FBitTree::AllocBit()
+{
+	uint32 Result = MAX_uint32;
+	if (*Bits != MAX_uint64) // else we are full
+	{
+		Result = 0;
+		uint32 Offset = 0;
+		uint32 Row = 0;
+		while (true)
+		{
+			uint64* At = Bits + Offset;
+			check(At >= Bits && At < Bits + AllocationSize / 8);
+			uint32 LowestZeroBit = FMath::CountTrailingZeros64(~*At);
+			check(LowestZeroBit < 64);
+			Result = Result * 64 + LowestZeroBit;
+			if (Row == Rows - 1)
+			{
+				check(!((*At) & (1ull << LowestZeroBit))); // this was already allocated?
+				*At |= (1ull << LowestZeroBit);
+				if (Row > 0)
+				{
+					if (*At == MAX_uint64)
+					{
+						do
+						{
+							uint32 Rem = (Offset - 1) % 64;
+							Offset = (Offset - 1) / 64;
+							At = Bits + Offset;
+							check(At >= Bits && At < Bits + AllocationSize / 8);
+							check(*At != MAX_uint64); // this should not already be marked full
+							*At |= (1ull << Rem);
+							if (*At != MAX_uint64)
+							{
+								break;
+							}
+							Row--;
+						} while (Row);
+					}
+				}
+				break;
+			}
+			Offset = Offset * 64 + 1 + LowestZeroBit;
+			Row++;
+		}
+	}
+
+	return Result;
+}
+
+bool FBitTree::IsAllocated(uint32 Index) const
+{
+	check(Index < DesiredCapacity);
+	uint32 Row = Rows - 1;
+	uint32 Rem = Index % 64;
+	uint32 Offset = OffsetOfLastRow + Index / 64;
+	uint64* At = Bits + Offset;
+	check(At >= Bits && At < Bits + AllocationSize / 8);
+	return !!((*At) & (1ull << Rem));
+}
+
+void FBitTree::AllocBit(uint32 Index)
+{
+	check(Index < DesiredCapacity);
+	uint32 Row = Rows - 1;
+	uint32 Rem = Index % 64;
+	uint32 Offset = OffsetOfLastRow + Index / 64;
+	uint64* At = Bits + Offset;
+	check(At >= Bits && At < Bits + AllocationSize / 8);
+	check(!((*At) & (1ull << Rem))); // this was already allocated?
+	*At |= (1ull << Rem);
+	if (*At == MAX_uint64 && Row > 0)
+	{
+		do
+		{
+			Rem = (Offset - 1) % 64;
+			Offset = (Offset - 1) / 64;
+			At = Bits + Offset;
+			check(At >= Bits && At < Bits + AllocationSize / 8);
+			check(!((*At) & (1ull << Rem))); // this was already allocated?
+			*At |= (1ull << Rem);
+			if (*At != MAX_uint64)
+			{
+				break;
+			}
+			Row--;
+		} while (Row);
+	}
+
+}
+uint32 FBitTree::NextAllocBit() const
+{
+	uint32 Result = MAX_uint32;
+	if (*Bits != MAX_uint64) // else we are full
+	{
+		Result = 0;
+		uint32 Offset = 0;
+		uint32 Row = 0;
+		while (true)
+		{
+			uint64* At = Bits + Offset;
+			check(At >= Bits && At < Bits + AllocationSize / 8);
+			uint32 LowestZeroBit = FMath::CountTrailingZeros64(~*At);
+			check(LowestZeroBit < 64);
+			Result = Result * 64 + LowestZeroBit;
+			if (Row == Rows - 1)
+			{
+				check(!((*At) & (1ull << LowestZeroBit))); // this was already allocated?
+				break;
+			}
+			Offset = Offset * 64 + 1 + LowestZeroBit;
+			Row++;
+		}
+	}
+
+	return Result;
+}
+
+uint32 FBitTree::NextAllocBit(uint32 StartIndex) const
+{
+	if (*Bits != MAX_uint64) // else we are full
+	{
+		uint32 Index = StartIndex;
+		check(Index < DesiredCapacity);
+		uint32 Row = Rows - 1;
+		uint32 Rem = Index % 64;
+		uint32 Offset = OffsetOfLastRow + Index / 64;
+		uint64* At = Bits + Offset;
+		check(At >= Bits && At < Bits + AllocationSize / 8);
+		uint64 LocalAt = *At;
+		if (!(LocalAt & (1ull << Rem)))
+		{
+			return Index; // lucked out, start was unallocated
+		}
+		// start was allocated, search for an unallocated one
+		LocalAt |= MAX_uint64 >> (63 - Rem); // set and ignore the bits representing items before (and including) the start
+		if (LocalAt != MAX_uint64)
+		{
+			// this qword has an item we can use
+			uint32 LowestZeroBit = FMath::CountTrailingZeros64(~LocalAt);
+			check(LowestZeroBit < 64);
+			return Index - Rem + LowestZeroBit;
+
+		}
+		// rest of qword was also allocated, search up the tree for the next free item
+		if (Row > 0)
+		{
+			do
+			{
+				Row--;
+				Rem = (Offset - 1) % 64;
+				Offset = (Offset - 1) / 64;
+				At = Bits + Offset;
+				check(At >= Bits && At < Bits + AllocationSize / 8);
+				LocalAt = *At;
+				LocalAt |= MAX_uint64 >> (63 - Rem); // set and ignore the bits representing items before (and including) the start
+				if (LocalAt != MAX_uint64)
+				{
+					// this qword has an item we can use
+					// now search down the tree
+					while (true)
+					{
+						uint32 LowestZeroBit = FMath::CountTrailingZeros64(~LocalAt);
+						check(LowestZeroBit < 64);
+						if (Row == Rows - 1)
+						{
+							check(!(LocalAt & (1ull << LowestZeroBit))); // this was already allocated?
+							uint32 Result = (Offset - OffsetOfLastRow) * 64 + LowestZeroBit;
+							check(Result < DesiredCapacity);
+							return Result;
+
+						}
+						Offset = Offset * 64 + 1 + LowestZeroBit;
+						At = Bits + Offset;
+						check(At >= Bits && At < Bits + AllocationSize / 8);
+						LocalAt = *At;
+						Row++;
+					}
+				}
+			} while (Row);
+		}
+	}
+
+	return MAX_uint32;
+}
+
+
+void FBitTree::FreeBit(uint32 Index)
+{
+	check(Index < DesiredCapacity);
+	uint32 Row = Rows - 1;
+	uint32 Rem = Index % 64;
+	uint32 Offset = OffsetOfLastRow + Index / 64;
+	uint64* At = Bits + Offset;
+	check(At >= Bits && At < Bits + AllocationSize / 8);
+	bool bWasFull = *At == MAX_uint64;
+	check((*At) & (1ull << Rem)); // this was not already allocated?
+	*At &= ~(1ull << Rem);
+	if (bWasFull && Row > 0)
+	{
+		do
+		{
+			Rem = (Offset - 1) % 64;
+			Offset = (Offset - 1) / 64;
+			At = Bits + Offset;
+			check(At >= Bits && At < Bits + AllocationSize / 8);
+			bWasFull = *At == MAX_uint64;
+			*At &= ~(1ull << Rem);
+			if (!bWasFull)
+			{
+				break;
+			}
+			Row--;
+		} while (Row);
+	}
+}
+
+uint32 FBitTree::CountOnes(uint32 UpTo) const
+{
+	uint32 Result = 0;
+	uint64* At = Bits + OffsetOfLastRow;
+	while (UpTo >= 64)
+	{
+		Result += FMath::CountBits(*At);
+		At++;
+		UpTo -= 64;
+	}
+	if (UpTo)
+	{
+		Result += FMath::CountBits((*At) << (64 - UpTo));
+	}
+	return Result;
+}
+
+#endif
\ No newline at end of file
diff --git a/Engine/Source/Runtime/Core/Private/HAL/MallocBinnedGPU.cpp b/Engine/Source/Runtime/Core/Private/HAL/MallocBinnedGPU.cpp
new file mode 100644
--- /dev/null
+++ b/Engine/Source/Runtime/Core/Private/HAL/MallocBinnedGPU.cpp
@@ -0,0 +1,1299 @@
+// Copyright 1998-2018 Epic Games, Inc. All Rights Reserved.
+
+/*=============================================================================
+	MallocBinned.cpp: Binned memory allocator
+=============================================================================*/
+
+#include "HAL/MallocBinnedGPU.h"
+
+#if PLATFORM_64BITS && PLATFORM_HAS_FPlatformVirtualMemoryBlock
+#include "Logging/LogMacros.h"
+#include "Misc/ScopeLock.h"
+#include "Templates/Function.h"
+#include "GenericPlatform/GenericPlatformProcess.h"
+#include "Stats/Stats.h"
+#include "HAL/IConsoleManager.h"
+#include "HAL/MemoryMisc.h"
+#include "HAL/PlatformMisc.h"
+
+struct FMallocBinnedGPU::FPoolInfoSmall
+{
+	enum ECanary
+	{
+		SmallUnassigned = 0x5effeed,
+		SmallAssigned = 0x69ffeed
+	};
+
+	uint32 Canary;
+	uint32 Taken;
+	FFreeBlock* FirstFreeProxy;
+
+	FPoolInfoSmall()
+		: Canary(ECanary::SmallUnassigned)
+		, Taken(0)
+		, FirstFreeProxy(nullptr)
+	{
+		static_assert(sizeof(FPoolInfoSmall) == 16, "Padding fail");
+	}
+	void CheckCanary(ECanary ShouldBe) const
+	{
+		if (Canary != ShouldBe)
+		{
+			UE_LOG(LogMemory, Fatal, TEXT("MallocBinnedGPU Corruption Canary was 0x%x, should be 0x%x"), int32(Canary), int32(ShouldBe));
+		}
+	}
+	void SetCanary(ECanary ShouldBe, bool bPreexisting, bool bGuarnteedToBeNew)
+	{
+		if (bPreexisting)
+		{
+			if (bGuarnteedToBeNew)
+			{
+				UE_LOG(LogMemory, Fatal, TEXT("MallocBinnedGPU Corruption Canary was 0x%x, should be 0x%x. This block is both preexisting and guaranteed to be new; which makes no sense."), int32(Canary), int32(ShouldBe));
+			}
+			if (ShouldBe == ECanary::SmallUnassigned)
+			{
+				if (Canary != ECanary::SmallAssigned)
+				{
+					UE_LOG(LogMemory, Fatal, TEXT("MallocBinnedGPU Corruption Canary was 0x%x, will be 0x%x because this block should be preexisting and in use."), int32(Canary), int32(ShouldBe));
+				}
+			}
+			else if (Canary != ShouldBe)
+			{
+				UE_LOG(LogMemory, Fatal, TEXT("MallocBinnedGPU Corruption Canary was 0x%x, should be 0x%x because this block should be preexisting."), int32(Canary), int32(ShouldBe));
+			}
+		}
+		else
+		{
+			if (bGuarnteedToBeNew)
+			{
+				if (Canary != ECanary::SmallUnassigned)
+				{
+					UE_LOG(LogMemory, Fatal, TEXT("MallocBinnedGPU Corruption Canary was 0x%x, will be 0x%x. This block is guaranteed to be new yet is it already assigned."), int32(Canary), int32(ShouldBe));
+				}
+			}
+			else if (Canary != ShouldBe && Canary != ECanary::SmallUnassigned)
+			{
+				UE_LOG(LogMemory, Fatal, TEXT("MallocBinnedGPU Corruption Canary was 0x%x, will be 0x%x does not have an expected value."), int32(Canary), int32(ShouldBe));
+			}
+		}
+		Canary = ShouldBe;
+	}
+	bool HasFreeRegularBlock() const
+	{
+		CheckCanary(ECanary::SmallAssigned);
+		return !!FirstFreeProxy;
+	}
+
+	void* AllocateRegularBlock(FMallocBinnedGPU& Allocator, uint8 MinimumAlignmentShift)
+	{
+		check(HasFreeRegularBlock());
+		++Taken;
+		FFreeBlock* Free = FirstFreeProxy;
+		void* Result = Free->AllocateRegularBlock(MinimumAlignmentShift);
+		if (Free->GetNumFreeRegularBlocks() == 0)
+		{
+			FirstFreeProxy = Free->NextFreeBlock;
+			delete (FGPUMemoryBlockProxy*)Free;
+			MBG_STAT(Allocator.GPUProxyMemory -= sizeof(FGPUMemoryBlockProxy);)
+		}
+
+		return Result;
+	}
+};
+
+struct FMallocBinnedGPU::FPoolInfoLarge
+{
+	enum ECanary
+	{
+		LargeUnassigned = 673,
+		LargeAssigned = 3917,
+	};
+
+public:
+	uint32 Canary;      
+private:
+	uint32 VMSizeDivVirtualSizeAlignment;
+	uint32 OSCommitSize;
+	uint32 AllocSize;      // Number of bytes allocated
+public:
+
+	FPoolInfoLarge() :
+		Canary(ECanary::LargeUnassigned),
+		VMSizeDivVirtualSizeAlignment(0),
+		OSCommitSize(0),
+		AllocSize(0)
+	{
+	}
+	void CheckCanary(ECanary ShouldBe) const
+	{
+		if (Canary != ShouldBe)
+		{
+			UE_LOG(LogMemory, Fatal, TEXT("MallocBinnedGPU Corruption Canary was 0x%x, should be 0x%x"), int32(Canary), int32(ShouldBe));
+		}
+	}
+	void SetCanary(ECanary ShouldBe, bool bPreexisting, bool bGuarnteedToBeNew)
+	{
+		if (bPreexisting)
+		{
+			if (bGuarnteedToBeNew)
+			{
+				UE_LOG(LogMemory, Fatal, TEXT("MallocBinnedGPU Corruption Canary was 0x%x, should be 0x%x. This block is both preexisting and guaranteed to be new; which makes no sense."), int32(Canary), int32(ShouldBe));
+			}
+			if (ShouldBe == ECanary::LargeUnassigned)
+			{
+				if (Canary != ECanary::LargeAssigned)
+				{
+					UE_LOG(LogMemory, Fatal, TEXT("MallocBinnedGPU Corruption Canary was 0x%x, will be 0x%x because this block should be preexisting and in use."), int32(Canary), int32(ShouldBe));
+				}
+			}
+			else if (Canary != ShouldBe)
+			{
+				UE_LOG(LogMemory, Fatal, TEXT("MallocBinnedGPU Corruption Canary was 0x%x, should be 0x%x because this block should be preexisting."), int32(Canary), int32(ShouldBe));
+			}
+		}
+		else
+		{
+			if (bGuarnteedToBeNew)
+			{
+				if (Canary != ECanary::LargeUnassigned)
+				{
+					UE_LOG(LogMemory, Fatal, TEXT("MallocBinnedGPU Corruption Canary was 0x%x, will be 0x%x. This block is guaranteed to be new yet is it already assigned."), int32(Canary), int32(ShouldBe));
+				}
+			}
+			else if (Canary != ShouldBe && Canary != ECanary::LargeUnassigned)
+			{
+				UE_LOG(LogMemory, Fatal, TEXT("MallocBinnedGPU Corruption Canary was 0x%x, will be 0x%x does not have an expected value."), int32(Canary), int32(ShouldBe));
+			}
+		}
+		Canary = ShouldBe;
+	}
+	uint32 GetOSRequestedBytes() const
+	{
+		return AllocSize;
+	}
+
+	uint32 GetOsVMSizeDivVirtualSizeAlignment() const
+	{
+		CheckCanary(ECanary::LargeAssigned);
+		return VMSizeDivVirtualSizeAlignment;
+	}
+
+	UPTRINT GetOsCommittedBytes() const
+	{
+		CheckCanary(ECanary::LargeAssigned);
+		return (UPTRINT)OSCommitSize;
+	}
+
+	void SetOSAllocationSizes(uint32 InRequestedBytes, UPTRINT InCommittedBytes, uint32 InVMSizeDivVirtualSizeAlignment)
+	{
+		CheckCanary(ECanary::LargeAssigned);
+		check(InRequestedBytes != 0);                // Shouldn't be pooling zero byte allocations
+		check(InCommittedBytes >= InRequestedBytes); // We must be allocating at least as much as we requested
+
+		AllocSize = InRequestedBytes;
+		OSCommitSize = InCommittedBytes;
+		VMSizeDivVirtualSizeAlignment = InVMSizeDivVirtualSizeAlignment;
+	}
+};
+
+
+/** Hash table struct for retrieving allocation book keeping information */
+struct FMallocBinnedGPU::PoolHashBucket
+{
+	UPTRINT BucketIndex;
+	FPoolInfoLarge* FirstPool;
+	PoolHashBucket* Prev;
+	PoolHashBucket* Next;
+
+	PoolHashBucket()
+	{
+		BucketIndex = 0;
+		FirstPool = nullptr;
+		Prev = this;
+		Next = this;
+	}
+
+	void Link(PoolHashBucket* After)
+	{
+		After->Prev = Prev;
+		After->Next = this;
+		Prev->Next = After;
+		this->Prev = After;
+	}
+
+	void Unlink()
+	{
+		Next->Prev = Prev;
+		Prev->Next = Next;
+		Prev = this;
+		Next = this;
+	}
+};
+
+
+
+struct FMallocBinnedGPU::Private
+{
+	// Implementation. 
+	static CA_NO_RETURN void OutOfMemory(uint64 Size, uint32 Alignment = 0)
+	{
+		// this is expected not to return
+		FPlatformMemory::OnOutOfMemory(Size, Alignment);
+	}
+
+	/**
+	* Gets the FPoolInfoSmall for a small block memory address. If no valid info exists one is created.
+	*/
+	static FPoolInfoSmall* GetOrCreatePoolInfoSmall(FMallocBinnedGPU& Allocator, uint32 InPoolIndex, uint32 BlockOfBlocksIndex)
+	{
+		FPoolInfoSmall*& InfoBlock = Allocator.SmallPoolTables[InPoolIndex].PoolInfos[BlockOfBlocksIndex / Allocator.SmallPoolInfosPerPlatformPage];
+		if (!InfoBlock)
+		{
+			InfoBlock = (FPoolInfoSmall*)FMemory::Malloc(Allocator.ArenaParams.BasePageSize);
+			Allocator.MallocedPointers.Add(InfoBlock);
+			MBG_STAT(Allocator.BinnedGPUPoolInfoMemory += Allocator.ArenaParams.AllocationGranularity;)
+			DefaultConstructItems<FPoolInfoSmall>((void*)InfoBlock, Allocator.SmallPoolInfosPerPlatformPage);
+		}
+
+		FPoolInfoSmall* Result = &InfoBlock[BlockOfBlocksIndex % Allocator.SmallPoolInfosPerPlatformPage];
+
+		bool bGuaranteedToBeNew = false;
+		if (BlockOfBlocksIndex >= Allocator.SmallPoolTables[InPoolIndex].NumEverUsedBlockOfBlocks)
+		{
+			bGuaranteedToBeNew = true;
+			Allocator.SmallPoolTables[InPoolIndex].NumEverUsedBlockOfBlocks = BlockOfBlocksIndex + 1;
+		}
+		Result->SetCanary(FPoolInfoSmall::ECanary::SmallAssigned, false, bGuaranteedToBeNew);
+		return Result;
+	}
+	/**
+	 * Gets the FPoolInfoLarge for a large block memory address. If no valid info exists one is created.
+	 */
+	static FPoolInfoLarge* GetOrCreatePoolInfoLarge(FMallocBinnedGPU& Allocator, void* InPtr)
+	{
+		/**
+		 * Creates an array of FPoolInfo structures for tracking allocations.
+		 */
+		auto CreatePoolArray = [](FMallocBinnedGPU& LocalAllocator)
+		{
+			uint64 PoolArraySize = LocalAllocator.NumLargePoolsPerPage * sizeof(FPoolInfoLarge);
+
+			void* Result;
+			Result = FMemory::Malloc(PoolArraySize);
+			LocalAllocator.MallocedPointers.Add(Result);
+			MBG_STAT(LocalAllocator.BinnedGPUPoolInfoMemory += PoolArraySize;)
+
+			DefaultConstructItems<FPoolInfoLarge>(Result, LocalAllocator.NumLargePoolsPerPage);
+			return (FPoolInfoLarge*)Result;
+		};
+
+		uint32 BucketIndex;
+		UPTRINT BucketIndexCollision;
+		uint32  PoolIndex;
+		Allocator.PtrToPoolMapping.GetHashBucketAndPoolIndices(InPtr, BucketIndex, BucketIndexCollision, PoolIndex);
+
+		PoolHashBucket* FirstBucket = &Allocator.HashBuckets[BucketIndex];
+		PoolHashBucket* Collision = FirstBucket;
+		do
+		{
+			if (!Collision->FirstPool)
+			{
+				Collision->BucketIndex = BucketIndexCollision;
+				Collision->FirstPool = CreatePoolArray(Allocator);
+				Collision->FirstPool[PoolIndex].SetCanary(FPoolInfoLarge::ECanary::LargeAssigned, false, true);
+				return &Collision->FirstPool[PoolIndex];
+			}
+
+			if (Collision->BucketIndex == BucketIndexCollision)
+			{
+				Collision->FirstPool[PoolIndex].SetCanary(FPoolInfoLarge::ECanary::LargeAssigned, false, false);
+				return &Collision->FirstPool[PoolIndex];
+			}
+
+			Collision = Collision->Next;
+		} while (Collision != FirstBucket);
+
+		// Create a new hash bucket entry
+		if (!Allocator.HashBucketFreeList)
+		{
+			{
+				Allocator.HashBucketFreeList = (PoolHashBucket*)FMemory::Malloc(Allocator.ArenaParams.AllocationGranularity);
+				Allocator.MallocedPointers.Add(Allocator.HashBucketFreeList);
+				MBG_STAT(Allocator.BinnedGPUHashMemory += Allocator.ArenaParams.AllocationGranularity;)
+			}
+
+			for (UPTRINT i = 0, n = Allocator.ArenaParams.AllocationGranularity / sizeof(PoolHashBucket); i < n; ++i)
+			{
+				Allocator.HashBucketFreeList->Link(new (Allocator.HashBucketFreeList + i) PoolHashBucket());
+			}
+		}
+
+		PoolHashBucket* NextFree = Allocator.HashBucketFreeList->Next;
+		PoolHashBucket* NewBucket = Allocator.HashBucketFreeList;
+
+		NewBucket->Unlink();
+
+		if (NextFree == NewBucket)
+		{
+			NextFree = nullptr;
+		}
+		Allocator.HashBucketFreeList = NextFree;
+
+		if (!NewBucket->FirstPool)
+		{
+			NewBucket->FirstPool = CreatePoolArray(Allocator);
+			NewBucket->FirstPool[PoolIndex].SetCanary(FPoolInfoLarge::ECanary::LargeAssigned, false, true);
+		}
+		else
+		{
+			NewBucket->FirstPool[PoolIndex].SetCanary(FPoolInfoLarge::ECanary::LargeAssigned, false, false);
+		}
+
+		NewBucket->BucketIndex = BucketIndexCollision;
+
+		FirstBucket->Link(NewBucket);
+
+		return &NewBucket->FirstPool[PoolIndex];
+	}
+
+	static FPoolInfoLarge* FindPoolInfo(FMallocBinnedGPU& Allocator, void* InPtr)
+	{
+		uint32 BucketIndex;
+		UPTRINT BucketIndexCollision;
+		uint32  PoolIndex;
+		Allocator.PtrToPoolMapping.GetHashBucketAndPoolIndices(InPtr, BucketIndex, BucketIndexCollision, PoolIndex);
+
+
+		PoolHashBucket* FirstBucket = &Allocator.HashBuckets[BucketIndex];
+		PoolHashBucket* Collision = FirstBucket;
+		do
+		{
+			if (Collision->BucketIndex == BucketIndexCollision)
+			{
+				return &Collision->FirstPool[PoolIndex];
+			}
+
+			Collision = Collision->Next;
+		} while (Collision != FirstBucket);
+
+		return nullptr;
+	}
+
+
+	static void FreeBundles(FMallocBinnedGPU& Allocator, FBundleNode* BundlesToRecycle, uint32 InBlockSize, uint32 InPoolIndex)
+	{
+		FPoolTable& Table = Allocator.SmallPoolTables[InPoolIndex];
+
+		FBundleNode* Bundle = BundlesToRecycle;
+		while (Bundle)
+		{
+			FBundleNode* NextBundle = Bundle->NextBundle;
+
+			FBundleNode* Node = Bundle;
+			do
+			{
+				FBundleNode* NextNode = Node->NextNodeInCurrentBundle;
+
+				uint32 OutBlockOfBlocksIndex;
+
+				void *GPUNode = ((FGPUMemoryBlockProxy*)Node)->GPUMemory;
+				check(GPUNode);
+				void* BasePtrOfNode = Allocator.BlockOfBlocksPointerFromContainedPtr(GPUNode, Allocator.SmallPoolTables[InPoolIndex].PagesPlatformForBlockOfBlocks, OutBlockOfBlocksIndex);
+				uint32 BlockWithinIndex = (((uint8*)GPUNode) - ((uint8*)BasePtrOfNode)) / Allocator.SmallPoolTables[InPoolIndex].BlockSize;
+
+				FPoolInfoSmall* NodePoolBlock = Allocator.SmallPoolTables[InPoolIndex].PoolInfos[OutBlockOfBlocksIndex / Allocator.SmallPoolInfosPerPlatformPage];
+				if (!NodePoolBlock)
+				{
+					UE_LOG(LogMemory, Fatal, TEXT("FMallocBinnedGPU Attempt to free an unrecognized small block %p"), GPUNode);
+				}
+				FPoolInfoSmall* NodePool = &NodePoolBlock[OutBlockOfBlocksIndex % Allocator.SmallPoolInfosPerPlatformPage];
+
+				NodePool->CheckCanary(FPoolInfoSmall::ECanary::SmallAssigned);
+
+				bool bWasExhaused = !NodePool->FirstFreeProxy;
+
+				// Free a pooled allocation.
+				FFreeBlock* Free = (FFreeBlock*)Node;
+				Free->NumFreeBlocks = 1;
+				Free->NextFreeBlock = NodePool->FirstFreeProxy;
+				Free->BlockSizeShifted = (InBlockSize >> Allocator.ArenaParams.MinimumAlignmentShift);
+				Free->Canary = FFreeBlock::CANARY_VALUE;
+				Free->PoolIndex = InPoolIndex;
+				NodePool->FirstFreeProxy = Free;
+
+				// Free this pool.
+				check(NodePool->Taken >= 1);
+				if (--NodePool->Taken == 0)
+				{
+					NodePool->SetCanary(FPoolInfoSmall::ECanary::SmallUnassigned, true, false);
+					Table.BlockOfBlockAllocationBits.FreeBit(OutBlockOfBlocksIndex);
+
+					uint64 AllocSize = Allocator.SmallPoolTables[InPoolIndex].PagesPlatformForBlockOfBlocks * Allocator.ArenaParams.AllocationGranularity;
+
+					if (!bWasExhaused)
+					{
+						Table.BlockOfBlockIsExhausted.AllocBit(OutBlockOfBlocksIndex);
+					}
+
+					Allocator.Decommit(InPoolIndex, BasePtrOfNode, AllocSize);
+					MBG_STAT(Allocator.BinnedGPUAllocatedOSSmallPoolMemory -= AllocSize;)
+
+					// this can be quite slow and it will sort of randomly happen when a whole page becomes free
+					int64 TotalFree = 0;
+					do 
+					{
+						FFreeBlock* NextFree = Free->NextFreeBlock;
+						delete (FGPUMemoryBlockProxy*)Free;
+						TotalFree += sizeof(FGPUMemoryBlockProxy);
+						Free = NextFree;
+					} while (Free);
+					MBG_STAT(Allocator.GPUProxyMemory -= TotalFree;)
+					NodePool->FirstFreeProxy = nullptr;
+				}
+				else if (bWasExhaused)
+				{
+					Table.BlockOfBlockIsExhausted.FreeBit(OutBlockOfBlocksIndex);
+				}
+
+				Node = NextNode;
+			} while (Node);
+
+			Bundle = NextBundle;
+		}
+	}
+
+};
+
+
+FMallocBinnedGPU::FPoolInfoSmall* FMallocBinnedGPU::PushNewPoolToFront(FMallocBinnedGPU& Allocator, uint32 InBlockSize, uint32 InPoolIndex, uint32& OutBlockOfBlocksIndex)
+{
+	FMallocBinnedGPU::FPoolTable& Table = Allocator.SmallPoolTables[InPoolIndex];
+	const uint32 BlockOfBlocksSize = Allocator.ArenaParams.AllocationGranularity * Table.PagesPlatformForBlockOfBlocks;
+
+	// Allocate memory.
+
+	uint32 BlockOfBlocksIndex = Table.BlockOfBlockAllocationBits.AllocBit();
+	if (BlockOfBlocksIndex == MAX_uint32)
+	{
+		Private::OutOfMemory(InBlockSize + 1); // The + 1 will hopefully be a hint that we actually ran out of our 1GB space.
+	}
+	uint8* FreePtr = BlockPointerFromIndecies(InPoolIndex, BlockOfBlocksIndex, BlockOfBlocksSize);
+
+	Allocator.Commit(InPoolIndex, FreePtr, BlockOfBlocksSize);
+	uint64 EndOffset = UPTRINT(FreePtr + BlockOfBlocksSize) - UPTRINT(PoolBasePtr(InPoolIndex));
+	if (EndOffset > Table.UnusedAreaOffsetLow)
+	{
+		Table.UnusedAreaOffsetLow = EndOffset;
+	}
+
+	FGPUMemoryBlockProxy* Proxy = new FGPUMemoryBlockProxy(FreePtr);
+	MBG_STAT(Allocator.GPUProxyMemory += sizeof(FGPUMemoryBlockProxy);)
+
+
+	FFreeBlock* Free = new ((void*)Proxy) FFreeBlock(BlockOfBlocksSize, InBlockSize, InPoolIndex, Allocator.ArenaParams.MinimumAlignmentShift);
+
+
+	MBG_STAT(BinnedGPUAllocatedOSSmallPoolMemory += (int64)BlockOfBlocksSize;)
+	check(IsAligned(FreePtr, Allocator.ArenaParams.AllocationGranularity));
+	// Create pool
+	FPoolInfoSmall* Result = Private::GetOrCreatePoolInfoSmall(*this, InPoolIndex, BlockOfBlocksIndex);
+	Result->CheckCanary(FPoolInfoSmall::ECanary::SmallAssigned);
+	Result->Taken = 0;
+	Result->FirstFreeProxy = Free;
+	Table.BlockOfBlockIsExhausted.FreeBit(BlockOfBlocksIndex);
+
+	OutBlockOfBlocksIndex = BlockOfBlocksIndex;
+
+	return Result;
+}
+
+FMallocBinnedGPU::FPoolInfoSmall* FMallocBinnedGPU::GetFrontPool(FPoolTable& Table, uint32 InPoolIndex, uint32& OutBlockOfBlocksIndex)
+{
+	OutBlockOfBlocksIndex = Table.BlockOfBlockIsExhausted.NextAllocBit();
+	if (OutBlockOfBlocksIndex == MAX_uint32)
+	{
+		return nullptr;
+	}
+	return Private::GetOrCreatePoolInfoSmall(*this, InPoolIndex, OutBlockOfBlocksIndex);
+}
+
+
+FMallocBinnedGPU::FMallocBinnedGPU()
+	: HashBucketFreeList(nullptr)
+{
+	MBG_STAT(ConsolidatedMemory = 0;)
+	MBG_STAT(GPUProxyMemory = 0;)
+
+	check(!PLATFORM_32BITS);
+	ArenaParams.BasePageSize = 4096;
+	ArenaParams.AllocationGranularity = FPlatformMemory::FPlatformVirtualMemoryBlock::GetCommitAlignment();
+
+	ArenaParams.ReserveVM = [](SIZE_T Size) -> FPlatformMemory::FPlatformVirtualMemoryBlock
+	{
+		return FPlatformMemory::FPlatformVirtualMemoryBlock::AllocateVirtual(Size, FPlatformMemory::FPlatformVirtualMemoryBlock::GetCommitAlignment());
+	};
+
+	ArenaParams.LargeBlockAlloc = [](SIZE_T Size, SIZE_T Alignment, SIZE_T& OutCommitSize, uint32& OutVMSizeDivVirtualSizeAlignment) -> void*
+	{
+		FPlatformMemory::FPlatformVirtualMemoryBlock Block = FPlatformMemory::FPlatformVirtualMemoryBlock::AllocateVirtual(Size, Alignment);
+		check(IsAligned(Block.GetVirtualPointer(), Alignment)); 
+		OutCommitSize = Align(Size, FPlatformMemory::FPlatformVirtualMemoryBlock::GetCommitAlignment());
+		Block.Commit(0, OutCommitSize);
+		OutVMSizeDivVirtualSizeAlignment = Block.GetActualSizeInPages();
+		return Block.GetVirtualPointer();
+	};
+	ArenaParams.LargeBlockFree = [](void* Ptr, uint32 VMSizeDivVirtualSizeAlignment)
+	{
+		FPlatformMemory::FPlatformVirtualMemoryBlock Block(Ptr, VMSizeDivVirtualSizeAlignment);
+		Block.FreeVirtual();
+	};
+
+}
+
+void FMallocBinnedGPU::InitMallocBinned()
+{
+	ArenaParams.MinimumAlignment = uint32(1) << ArenaParams.MinimumAlignmentShift;
+	check(sizeof(FBundleNode) <= ArenaParams.MinimumAlignment);
+
+	ArenaParams.MaxMemoryPerBlockSize = uint64(1) << ArenaParams.MaxMemoryPerBlockSizeShift;
+
+	check(ArenaParams.BasePageSize % sizeof(FPoolInfoLarge) == 0);  // these need to divide evenly!
+	NumLargePoolsPerPage = ArenaParams.BasePageSize / sizeof(FPoolInfoLarge);
+	PtrToPoolMapping.Init(ArenaParams.BasePageSize, NumLargePoolsPerPage, ArenaParams.AddressLimit);
+
+	checkf(FMath::IsPowerOfTwo(ArenaParams.AllocationGranularity), TEXT("OS page size must be a power of two"));
+	checkf(FMath::IsPowerOfTwo(ArenaParams.BasePageSize), TEXT("OS page size must be a power of two"));
+	check(ArenaParams.PoolCount <= 256);
+
+	// Init pool tables.
+
+	TArray<FSizeTableEntry> SizeTable;
+
+	if (ArenaParams.bUseStandardSmallPoolSizes)
+	{
+		SizeTable.AddZeroed(BINNEDCOMMON_NUM_LISTED_SMALL_POOLS + ArenaParams.MaxStandardPoolSize / ArenaParams.BasePageSize); // overestimate
+		ArenaParams.PoolCount = FSizeTableEntry::FillSizeTable(ArenaParams.AllocationGranularity, &SizeTable[0], ArenaParams.BasePageSize, ArenaParams.MinimumAlignment, ArenaParams.MaxStandardPoolSize, ArenaParams.BasePageSize);
+		SizeTable.RemoveAt(ArenaParams.PoolCount, SizeTable.Num() - ArenaParams.PoolCount);
+	}
+	else
+	{
+		ArenaParams.PoolCount = 0;
+	}
+	for (uint32 Size : ArenaParams.AdditionalBlockSizes)
+	{
+		if (Size >= ArenaParams.BasePageSize)
+		{
+			check(Size % ArenaParams.BasePageSize == 0); // we can actually use blocks that are not divisible by the base page size, but we need a bit more work
+			SizeTable.Emplace(Size, ArenaParams.AllocationGranularity, Size / ArenaParams.BasePageSize, ArenaParams.BasePageSize, ArenaParams.MinimumAlignment);
+		}
+		else
+		{
+			// it is difficult to test what would actually make a good bucket size here, wouldn't want a prime number, 33 for example because that would take 33 pages a slab
+			SizeTable.Emplace(Size, ArenaParams.AllocationGranularity, 1, ArenaParams.BasePageSize, ArenaParams.MinimumAlignment);
+		}
+		ArenaParams.PoolCount++;
+	}
+	if (ArenaParams.AdditionalBlockSizes.Num())
+	{
+		Sort(&SizeTable[0], SizeTable.Num());
+	}
+	check(ArenaParams.PoolCount == SizeTable.Num());
+	check(SizeTable.Num() < 256);
+	ArenaParams.MaxPoolSize = SizeTable[ArenaParams.PoolCount - 1].BlockSize;
+
+	check(ArenaParams.BasePageSize % sizeof(FPoolInfoSmall) == 0);
+	SmallPoolInfosPerPlatformPage = ArenaParams.BasePageSize / sizeof(FPoolInfoSmall);
+
+	GGlobalRecycler.Init(ArenaParams.PoolCount);
+	SmallPoolTables.AddDefaulted(ArenaParams.PoolCount);
+	SmallBlockSizesReversedShifted.AddDefaulted(ArenaParams.PoolCount);
+	PoolBaseVMPtr.AddDefaulted(ArenaParams.PoolCount);
+	PoolBaseVMBlocks.AddDefaulted(ArenaParams.PoolCount);
+	MemSizeToIndex.AddDefaulted(1 + (ArenaParams.MaxPoolSize >> ArenaParams.MinimumAlignmentShift));
+
+	ArenaParams.MaximumAlignmentForSmallBlock = ArenaParams.MinimumAlignment;
+	check(ArenaParams.MaximumAlignmentForSmallBlock > 0);
+
+	for (uint32 Index = 0; Index < ArenaParams.PoolCount; ++Index)
+	{
+		checkf(Index == 0 || SizeTable[Index - 1].BlockSize < SizeTable[Index].BlockSize, TEXT("Small block sizes must be strictly increasing"));
+		checkf(SizeTable[Index].BlockSize % ArenaParams.MinimumAlignment == 0, TEXT("Small block size must be a multiple of ArenaParams.MinimumAlignment"));
+
+		// determine the largest alignment that we can cover with a small block
+		while (ArenaParams.MaximumAlignmentForSmallBlock < ArenaParams.AllocationGranularity && IsAligned(SizeTable[Index].BlockSize, ArenaParams.MaximumAlignmentForSmallBlock * 2))
+		{
+			ArenaParams.MaximumAlignmentForSmallBlock *= 2;
+		}
+
+		SmallPoolTables[Index].BlockSize = SizeTable[Index].BlockSize;
+		SmallPoolTables[Index].BlocksPerBlockOfBlocks = SizeTable[Index].BlocksPerBlockOfBlocks;
+		SmallPoolTables[Index].PagesPlatformForBlockOfBlocks = SizeTable[Index].PagesPlatformForBlockOfBlocks;
+
+		SmallPoolTables[Index].UnusedAreaOffsetLow = 0;
+		SmallPoolTables[Index].NumEverUsedBlockOfBlocks = 0;
+
+		int64 TotalNumberOfBlocksOfBlocks = ArenaParams.MaxMemoryPerBlockSize / (SizeTable[Index].PagesPlatformForBlockOfBlocks * ArenaParams.AllocationGranularity);
+
+		int64 MaxPoolInfoMemory = sizeof(FPoolInfoSmall**) * (TotalNumberOfBlocksOfBlocks + SmallPoolInfosPerPlatformPage - 1) / SmallPoolInfosPerPlatformPage;
+		SmallPoolTables[Index].PoolInfos = (FPoolInfoSmall**)FMemory::Malloc(MaxPoolInfoMemory);
+		MallocedPointers.Add(SmallPoolTables[Index].PoolInfos);
+
+		FMemory::Memzero(SmallPoolTables[Index].PoolInfos, MaxPoolInfoMemory);
+		MBG_STAT(BinnedGPUPoolInfoMemory += MaxPoolInfoMemory;)
+
+		{
+			int64 AllocationSize = FBitTree::GetMemoryRequirements(TotalNumberOfBlocksOfBlocks);
+
+			{
+				void *Bits = FMemory::Malloc(AllocationSize);
+				MallocedPointers.Add(Bits);
+				check(Bits);
+				MBG_STAT(BinnedGPUFreeBitsMemory += AllocationSize;)
+				SmallPoolTables[Index].BlockOfBlockAllocationBits.FBitTreeInit(TotalNumberOfBlocksOfBlocks, Bits, AllocationSize, false);
+			}
+			{
+				void *Bits = FMemory::Malloc(AllocationSize);
+				MallocedPointers.Add(Bits);
+				check(Bits);
+				MBG_STAT(BinnedGPUFreeBitsMemory += AllocationSize;)
+				SmallPoolTables[Index].BlockOfBlockIsExhausted.FBitTreeInit(TotalNumberOfBlocksOfBlocks, Bits, AllocationSize, true);
+			}
+		}
+	}
+
+
+	// Set up pool mappings
+	uint8* IndexEntry = &MemSizeToIndex[0];
+	uint32  PoolIndex = 0;
+	for (uint32 Index = 0; Index != 1 + (ArenaParams.MaxPoolSize >> ArenaParams.MinimumAlignmentShift); ++Index)
+	{
+
+		uint32 BlockSize = Index << ArenaParams.MinimumAlignmentShift; // inverse of int32 Index = int32((Size >> ArenaParams.MinimumAlignmentShift));
+		while (SizeTable[PoolIndex].BlockSize < BlockSize)
+		{
+			++PoolIndex;
+			check(PoolIndex != ArenaParams.PoolCount);
+		}
+		check(PoolIndex < 256);
+		*IndexEntry++ = uint8(PoolIndex);
+	}
+	// now reverse the pool sizes for cache coherency
+
+	for (uint32 Index = 0; Index != ArenaParams.PoolCount; ++Index)
+	{
+		uint32 Partner = ArenaParams.PoolCount - Index - 1;
+		SmallBlockSizesReversedShifted[Index] = (SizeTable[Partner].BlockSize >> ArenaParams.MinimumAlignmentShift);
+	}
+	uint64 MaxHashBuckets = PtrToPoolMapping.GetMaxHashBuckets();
+
+	{
+		int64 HashAllocSize = MaxHashBuckets * sizeof(PoolHashBucket);
+		HashBuckets = (PoolHashBucket*)FMemory::Malloc(HashAllocSize);
+		MallocedPointers.Add(HashBuckets);
+
+		MBG_STAT(BinnedGPUHashMemory += HashAllocSize;)
+		verify(HashBuckets);
+	}
+
+	DefaultConstructItems<PoolHashBucket>(HashBuckets, MaxHashBuckets);
+
+	uint8* BinnedGPUBaseVMPtr;
+	if (!ArenaParams.bUseSeparateVMPerPool)
+	{
+		PoolBaseVMBlock = ArenaParams.ReserveVM(ArenaParams.PoolCount * ArenaParams.MaxMemoryPerBlockSize);
+		BinnedGPUBaseVMPtr = (uint8*)PoolBaseVMBlock.GetVirtualPointer();
+	}
+	else
+	{
+		BinnedGPUBaseVMPtr = nullptr;
+	}
+	for (uint32 Index = 0; Index < ArenaParams.PoolCount; ++Index)
+	{
+		uint8* NewVM;
+		FPlatformMemory::FPlatformVirtualMemoryBlock NewBlock;
+		if (BinnedGPUBaseVMPtr)
+		{
+			NewVM = BinnedGPUBaseVMPtr;
+			BinnedGPUBaseVMPtr += ArenaParams.MaxMemoryPerBlockSize;
+		}
+		else
+		{
+			NewBlock = ArenaParams.ReserveVM(ArenaParams.MaxMemoryPerBlockSize);
+			NewVM = (uint8*)NewBlock.GetVirtualPointer();
+		}
+
+		// insertion sort
+		if (Index && NewVM < PoolBaseVMPtr[Index - 1])
+		{
+			uint32 InsertIndex = 0;
+			for (; InsertIndex < Index; ++InsertIndex)
+			{
+				if (NewVM < PoolBaseVMPtr[InsertIndex])
+				{
+					break;
+				}
+			}
+			check(InsertIndex < Index);
+			for (uint32 MoveIndex = Index; MoveIndex > InsertIndex; --MoveIndex)
+			{
+				PoolBaseVMPtr[MoveIndex] = PoolBaseVMPtr[MoveIndex - 1];
+				PoolBaseVMBlocks[MoveIndex] = PoolBaseVMBlocks[MoveIndex - 1];
+			}
+			PoolBaseVMPtr[InsertIndex] = NewVM;
+			PoolBaseVMBlocks[InsertIndex] = NewBlock;
+		}
+		else
+		{
+			PoolBaseVMPtr[Index] = NewVM;
+			PoolBaseVMBlocks[Index] = NewBlock;
+		}
+	}
+	HighestPoolBaseVMPtr = PoolBaseVMPtr[ArenaParams.PoolCount - 1];
+	uint64 TotalGaps = 0;
+	for (uint32 Index = 0; Index < uint32(ArenaParams.PoolCount - 1); ++Index)
+	{
+		check(PoolBaseVMPtr[Index + 1] > PoolBaseVMPtr[Index]); // we sorted it
+		check(PoolBaseVMPtr[Index + 1] >= PoolBaseVMPtr[Index] + ArenaParams.MaxMemoryPerBlockSize); // and blocks are non-overlapping
+		TotalGaps += PoolBaseVMPtr[Index + 1] - (PoolBaseVMPtr[Index] + ArenaParams.MaxMemoryPerBlockSize);
+	}
+	if (TotalGaps == 0)
+	{
+		PoolSearchDiv = 0;
+	}
+	else if (TotalGaps < ArenaParams.MaxMemoryPerBlockSize)
+	{
+		check(ArenaParams.bUseSeparateVMPerPool);
+		PoolSearchDiv = ArenaParams.MaxMemoryPerBlockSize; // the gaps are not significant, ignoring them should give accurate searches
+	}
+	else
+	{
+		check(ArenaParams.bUseSeparateVMPerPool);
+		PoolSearchDiv = ArenaParams.MaxMemoryPerBlockSize + ((TotalGaps + ArenaParams.PoolCount - 2) / (ArenaParams.PoolCount - 1));
+	}
+}
+
+void FMallocBinnedGPU::Commit(uint32 InPoolIndex, void *Ptr, SIZE_T Size)
+{
+	if (!ArenaParams.bUseSeparateVMPerPool)
+	{
+		PoolBaseVMBlock.CommitByPtr(Ptr, Size);
+	}
+	else
+	{
+		PoolBaseVMBlocks[InPoolIndex].CommitByPtr(Ptr, Size);
+	}
+}
+void FMallocBinnedGPU::Decommit(uint32 InPoolIndex, void *Ptr, SIZE_T Size)
+{
+	if (!ArenaParams.bUseSeparateVMPerPool)
+	{
+		PoolBaseVMBlock.DecommitByPtr(Ptr, Size);
+	}
+	else
+	{
+		PoolBaseVMBlocks[InPoolIndex].DecommitByPtr(Ptr, Size);
+	}
+}
+
+
+FMallocBinnedGPU::~FMallocBinnedGPU()
+{
+	FScopeLock Lock(&Mutex);
+	FScopeLock Lock2(&GetFreeBlockListsRegistrationMutex());
+
+	MBG_STAT(
+		UE_CLOG(BinnedGPUAllocatedOSSmallPoolMemory > 0, LogCore, Error, TEXT("FMallocBinnedGPU leaked small block memory: %fmb"), ((double)BinnedGPUAllocatedOSSmallPoolMemory) / (1024.0f * 1024.0f));
+		UE_CLOG(BinnedGPUAllocatedLargePoolMemoryWAlignment > 0, LogCore, Error, TEXT("FMallocBinnedGPU leaked large block memory: %fmb"), ((double)BinnedGPUAllocatedLargePoolMemoryWAlignment) / (1024.0f * 1024.0f));
+	)
+	for (FPerThreadFreeBlockLists* Lists : GetRegisteredFreeBlockLists())
+	{
+		if (Lists)
+		{
+			for (int32 PoolIndex = 0; PoolIndex != ArenaParams.PoolCount; ++PoolIndex)
+			{
+				FBundleNode* Bundles = Lists->PopBundles(PoolIndex);
+				if (Bundles)
+				{
+					Private::FreeBundles(*this, Bundles, PoolIndexToBlockSize(PoolIndex), PoolIndex);
+				}
+			}
+			delete Lists;
+		}
+	}
+
+	if (ArenaParams.bUseSeparateVMPerPool)
+	{
+		for (int32 PoolIndex = 0; PoolIndex != ArenaParams.PoolCount; ++PoolIndex)
+		{
+			PoolBaseVMBlocks[PoolIndex].FreeVirtual();
+		}
+	}
+	else
+	{
+		PoolBaseVMBlock.FreeVirtual();
+	}
+
+
+	for (void* Ptr : MallocedPointers)
+	{
+		FMemory::Free(Ptr);
+	}
+
+	FPlatformTLS::FreeTlsSlot(BinnedGPUTlsSlot);
+}
+
+bool FMallocBinnedGPU::IsInternallyThreadSafe() const
+{
+	return true;
+}
+
+void* FMallocBinnedGPU::MallocExternal(SIZE_T Size, uint32 Alignment)
+{
+	check(FMath::IsPowerOfTwo(Alignment));
+	checkf(DEFAULT_ALIGNMENT <= ArenaParams.MinimumAlignment, TEXT("DEFAULT_ALIGNMENT is assumed to be zero")); // used below
+
+	if (AdjustSmallBlockSizeForAlignment(Size, Alignment)) // there is some redundant work here... we already adjusted the size for alignment
+	{
+		uint32 PoolIndex = BoundSizeToPoolIndex(Size);
+		FPerThreadFreeBlockLists* Lists = ArenaParams.bPerThreadCaches ? FPerThreadFreeBlockLists::Get(BinnedGPUTlsSlot) : nullptr;
+		if (Lists)
+		{
+			if (Lists->ObtainRecycledPartial(ArenaParams, GGlobalRecycler, PoolIndex))
+			{
+				if (void* Result = Lists->Malloc(*this, PoolIndex))
+				{
+					uint32 BlockSize = PoolIndexToBlockSize(PoolIndex);
+					Lists->AllocatedMemory += BlockSize;
+					checkSlow(IsAligned(Result, Alignment));
+					return Result;
+				}
+			}
+		}
+
+		FScopeLock Lock(&Mutex);
+
+		// Allocate from small object pool.
+		FPoolTable& Table = SmallPoolTables[PoolIndex];
+
+		uint32 BlockOfBlocksIndex = MAX_uint32;
+		FPoolInfoSmall* Pool = GetFrontPool(Table, PoolIndex, BlockOfBlocksIndex);
+		if (!Pool)
+		{
+			Pool = PushNewPoolToFront(*this, Table.BlockSize, PoolIndex, BlockOfBlocksIndex);
+		}
+
+		const uint32 BlockOfBlocksSize = ArenaParams.AllocationGranularity * Table.PagesPlatformForBlockOfBlocks;
+		uint8* BlockOfBlocksPtr = BlockPointerFromIndecies(PoolIndex, BlockOfBlocksIndex, BlockOfBlocksSize);
+
+		void* Result = Pool->AllocateRegularBlock(*this, ArenaParams.MinimumAlignmentShift);
+		MBG_STAT(BinnedGPUAllocatedSmallPoolMemory += PoolIndexToBlockSize(PoolIndex);)
+		if (ArenaParams.EmptyCacheAllocExtra)
+		{
+			if (Lists)
+			{
+				// prefill the free list with some allocations so we are less likely to hit this slow path with the mutex 
+				for (int32 Index = 0; Index < ArenaParams.EmptyCacheAllocExtra && Pool->HasFreeRegularBlock(); Index++)
+				{
+					if (!Lists->Free(*this, Result, PoolIndex, Table.BlockSize, ArenaParams))
+					{
+						break;
+					}
+					Result = Pool->AllocateRegularBlock(*this, ArenaParams.MinimumAlignmentShift);
+				}
+			}
+		}
+		if (!Pool->HasFreeRegularBlock())
+		{
+			Table.BlockOfBlockIsExhausted.AllocBit(BlockOfBlocksIndex);
+		}
+		checkSlow(IsAligned(Result, Alignment));
+		return Result;
+	}
+
+	uint32 VMSizeDivVirtualSizeAlignment = 0;
+	SIZE_T CommitSize = 0;
+	void* Result = ArenaParams.LargeBlockAlloc(Size, Alignment, CommitSize, VMSizeDivVirtualSizeAlignment);
+
+	UE_CLOG(!IsAligned(Result, Alignment), LogMemory, Fatal, TEXT("FMallocBinnedGPU alignment was too large for OS. Alignment=%d   Ptr=%p"), Alignment, Result);
+
+	if (!Result)
+	{
+		Private::OutOfMemory(Size);
+	}
+	check(IsOSAllocation(Result));
+	FScopeLock Lock(&Mutex);
+
+	MBG_STAT(BinnedGPUAllocatedLargePoolMemory += (int64)Size;)
+	MBG_STAT(BinnedGPUAllocatedLargePoolMemoryWAlignment += CommitSize;)
+
+	// Create pool.
+	FPoolInfoLarge* Pool = Private::GetOrCreatePoolInfoLarge(*this, Result);
+	check(Size > 0 && Size <= CommitSize && CommitSize >= ArenaParams.BasePageSize);
+	Pool->SetOSAllocationSizes(Size, CommitSize, VMSizeDivVirtualSizeAlignment);
+
+	return Result;
+}
+
+
+
+void FMallocBinnedGPU::FreeExternal(void* Ptr)
+{
+	uint64 PoolIndex = PoolIndexFromPtr(Ptr);
+	if (PoolIndex < ArenaParams.PoolCount)
+	{
+		check(Ptr); // null is an OS allocation because it will not fall in our VM block
+		uint32 BlockSize = PoolIndexToBlockSize(PoolIndex);
+
+		FBundleNode* BundlesToRecycle = nullptr;
+		FPerThreadFreeBlockLists* Lists = ArenaParams.bPerThreadCaches ? FPerThreadFreeBlockLists::Get(BinnedGPUTlsSlot) : nullptr;
+		if (Lists)
+		{
+			BundlesToRecycle = Lists->RecycleFullBundle(ArenaParams, GGlobalRecycler, PoolIndex);
+			bool bPushed = Lists->Free(*this, Ptr, PoolIndex, BlockSize, ArenaParams);
+			check(bPushed);
+			Lists->AllocatedMemory -= BlockSize;
+		}
+		else
+		{
+			BundlesToRecycle = (FBundleNode*)new FGPUMemoryBlockProxy(Ptr);
+			MBG_STAT(GPUProxyMemory += sizeof(FGPUMemoryBlockProxy);)
+			BundlesToRecycle->NextNodeInCurrentBundle = nullptr;
+		}
+		if (BundlesToRecycle)
+		{
+			BundlesToRecycle->NextBundle = nullptr;
+			FScopeLock Lock(&Mutex);
+			Private::FreeBundles(*this, BundlesToRecycle, BlockSize, PoolIndex);
+			if (!Lists)
+			{
+				// lists track their own stat track them instead in the global stat if we don't have lists
+				MBG_STAT(BinnedGPUAllocatedSmallPoolMemory -= ((int64)(BlockSize));)
+			}
+		}
+	}
+	else if (Ptr)
+	{
+		uint32 VMSizeDivVirtualSizeAlignment;
+
+		{
+			FScopeLock Lock(&Mutex);
+			FPoolInfoLarge* Pool = Private::FindPoolInfo(*this, Ptr);
+			if (!Pool)
+			{
+				UE_LOG(LogMemory, Fatal, TEXT("FMallocBinnedGPU Attempt to free an unrecognized block %p"), Ptr);
+			}
+			UPTRINT PoolOsCommittedBytes = Pool->GetOsCommittedBytes();
+			UPTRINT PoolOSRequestedBytes = Pool->GetOSRequestedBytes();
+			VMSizeDivVirtualSizeAlignment = Pool->GetOsVMSizeDivVirtualSizeAlignment();
+
+			MBG_STAT(BinnedGPUAllocatedLargePoolMemory -= ((int64)PoolOSRequestedBytes);)
+			MBG_STAT(BinnedGPUAllocatedLargePoolMemoryWAlignment -= ((int64)PoolOsCommittedBytes);)
+
+			checkf(PoolOSRequestedBytes <= PoolOsCommittedBytes, TEXT("FMallocBinnedGPU::FreeExternal %d %d"), int32(PoolOSRequestedBytes), int32(PoolOsCommittedBytes));
+			Pool->SetCanary(FPoolInfoLarge::ECanary::LargeUnassigned, true, false);
+		}
+		// Free an OS allocation.
+		ArenaParams.LargeBlockFree(Ptr, VMSizeDivVirtualSizeAlignment);
+	}
+}
+
+bool FMallocBinnedGPU::GetAllocationSizeExternal(void* Ptr, SIZE_T& SizeOut)
+{
+	uint64 PoolIndex = PoolIndexFromPtr(Ptr);
+	if (PoolIndex < ArenaParams.PoolCount)
+	{
+		check(Ptr); // null is an OS allocation because it will not fall in our VM block
+		SizeOut = PoolIndexToBlockSize(PoolIndex);
+		return true;
+	}
+	if (!Ptr)
+	{
+		return false;
+	}
+	FScopeLock Lock(&Mutex);
+	FPoolInfoLarge* Pool = Private::FindPoolInfo(*this, Ptr);
+	if (!Pool)
+	{
+		UE_LOG(LogMemory, Fatal, TEXT("FMallocBinnedGPU Attempt to GetAllocationSizeExternal an unrecognized block %p"), Ptr);
+	}
+	UPTRINT PoolOsBytes = Pool->GetOsCommittedBytes();
+	uint32 PoolOSRequestedBytes = Pool->GetOSRequestedBytes();
+	checkf(PoolOSRequestedBytes <= PoolOsBytes, TEXT("FMallocBinnedGPU::GetAllocationSizeExternal %d %d"), int32(PoolOSRequestedBytes), int32(PoolOsBytes));
+	SizeOut = PoolOsBytes;
+	return true;
+}
+
+bool FMallocBinnedGPU::ValidateHeap()
+{
+	// Not implemented
+	// NumEverUsedBlockOfBlocks gives us all of the information we need to examine each pool, so it is doable.
+	return true;
+}
+
+const TCHAR* FMallocBinnedGPU::GetDescriptiveName()
+{
+	return TEXT("BinnedGPU");
+}
+
+void FMallocBinnedGPU::FlushCurrentThreadCache()
+{
+	double StartTimeInner = FPlatformTime::Seconds();
+	QUICK_SCOPE_CYCLE_COUNTER(STAT_FMallocBinnedGPU_FlushCurrentThreadCache);
+	FPerThreadFreeBlockLists* Lists = FPerThreadFreeBlockLists::Get(BinnedGPUTlsSlot);
+
+	float WaitForMutexTime = 0.0f;
+	float WaitForMutexAndTrimTime = 0.0f;
+
+	if (Lists)
+	{
+		FScopeLock Lock(&Mutex);
+		WaitForMutexTime = FPlatformTime::Seconds() - StartTimeInner;
+		for (int32 PoolIndex = 0; PoolIndex != ArenaParams.PoolCount; ++PoolIndex)
+		{
+			FBundleNode* Bundles = Lists->PopBundles(PoolIndex);
+			if (Bundles)
+			{
+				Private::FreeBundles(*this, Bundles, PoolIndexToBlockSize(PoolIndex), PoolIndex);
+			}
+		}
+		WaitForMutexAndTrimTime = FPlatformTime::Seconds() - StartTimeInner;
+	}
+
+	// These logs must happen outside the above mutex to avoid deadlocks
+	if (WaitForMutexTime > 0.02f)
+	{
+		UE_LOG(LogMemory, Warning, TEXT("FMallocBinnedGPU took %6.2fms to wait for mutex for trim."), WaitForMutexTime * 1000.0f);
+	}
+	if (WaitForMutexAndTrimTime > 0.02f)
+	{
+		UE_LOG(LogMemory, Warning, TEXT("FMallocBinnedGPU took %6.2fms to wait for mutex AND trim."), WaitForMutexAndTrimTime * 1000.0f);
+	}
+}
+
+#include "Async/TaskGraphInterfaces.h"
+
+void FMallocBinnedGPU::Trim(bool bTrimThreadCaches)
+{
+	if (bTrimThreadCaches && ArenaParams.bPerThreadCaches)
+	{
+		QUICK_SCOPE_CYCLE_COUNTER(STAT_FMallocBinnedGPU_Trim);
+		//double StartTime = FPlatformTime::Seconds();
+		TFunction<void(ENamedThreads::Type CurrentThread)> Broadcast =
+			[this](ENamedThreads::Type MyThread)
+		{
+			FlushCurrentThreadCache();
+		};
+		// Skip task threads on desktop platforms as it is too slow and they don't have much memory
+		FTaskGraphInterface::BroadcastSlow_OnlyUseForSpecialPurposes(!PLATFORM_DESKTOP, false, Broadcast);
+		//UE_LOG(LogTemp, Display, TEXT("Trim Broadcast = %6.2fms"), 1000.0f * float(FPlatformTime::Seconds() - StartTime));
+	}
+}
+
+void FMallocBinnedGPU::SetupTLSCachesOnCurrentThread()
+{
+	if (!ArenaParams.bPerThreadCaches)
+	{
+		return;
+	}
+	if (!BinnedGPUTlsSlot)
+	{
+		BinnedGPUTlsSlot = FPlatformTLS::AllocTlsSlot();
+	}
+	check(BinnedGPUTlsSlot);
+	FPerThreadFreeBlockLists::SetTLS(*this);
+}
+
+void FMallocBinnedGPU::ClearAndDisableTLSCachesOnCurrentThread()
+{
+	FlushCurrentThreadCache();
+	MBG_STAT(ConsolidatedMemory +=) FPerThreadFreeBlockLists::ClearTLS(*this);
+}
+
+
+bool FMallocBinnedGPU::FFreeBlockList::ObtainPartial(FArenaParams& LocalArenaParams, FGlobalRecycler& GlobalRecycler, uint32 InPoolIndex)
+{
+	if (!PartialBundle.Head)
+	{
+		PartialBundle.Count = 0;
+		PartialBundle.Head = GlobalRecycler.PopBundle(LocalArenaParams.MaxGlobalBundles, InPoolIndex);
+		if (PartialBundle.Head)
+		{
+			PartialBundle.Count = PartialBundle.Head->Count;
+			PartialBundle.Head->NextBundle = nullptr;
+			return true;
+		}
+		return false;
+	}
+	return true;
+}
+
+FMallocBinnedGPU::FBundleNode* FMallocBinnedGPU::FFreeBlockList::RecyleFull(FArenaParams& LocalArenaParams, FGlobalRecycler& GlobalRecycler, uint32 InPoolIndex)
+{
+	FMallocBinnedGPU::FBundleNode* Result = nullptr;
+	if (FullBundle.Head)
+	{
+		FullBundle.Head->Count = FullBundle.Count;
+		if (!GlobalRecycler.PushBundle(LocalArenaParams.MaxGlobalBundles, InPoolIndex, FullBundle.Head))
+		{
+			Result = FullBundle.Head;
+			Result->NextBundle = nullptr;
+		}
+		FullBundle.Reset();
+	}
+	return Result;
+}
+
+FMallocBinnedGPU::FBundleNode* FMallocBinnedGPU::FFreeBlockList::PopBundles(uint32 InPoolIndex)
+{
+	FBundleNode* Partial = PartialBundle.Head;
+	if (Partial)
+	{
+		PartialBundle.Reset();
+		Partial->NextBundle = nullptr;
+	}
+
+	FBundleNode* Full = FullBundle.Head;
+	if (Full)
+	{
+		FullBundle.Reset();
+		Full->NextBundle = nullptr;
+	}
+
+	FBundleNode* Result = Partial;
+	if (Result)
+	{
+		Result->NextBundle = Full;
+	}
+	else
+	{
+		Result = Full;
+	}
+
+	return Result;
+}
+
+void FMallocBinnedGPU::FPerThreadFreeBlockLists::SetTLS(FMallocBinnedGPU& Allocator)
+{
+	uint32 BinnedGPUTlsSlot = Allocator.BinnedGPUTlsSlot;
+	check(BinnedGPUTlsSlot);
+	FPerThreadFreeBlockLists* ThreadSingleton = (FPerThreadFreeBlockLists*)FPlatformTLS::GetTlsValue(BinnedGPUTlsSlot);
+	if (!ThreadSingleton)
+	{
+		int64 TLSSize = sizeof(FPerThreadFreeBlockLists);
+		ThreadSingleton = new FPerThreadFreeBlockLists(Allocator.ArenaParams.PoolCount);
+		MBG_STAT(Allocator.BinnedGPUTLSMemory += TLSSize;)
+		verify(ThreadSingleton);
+		FPlatformTLS::SetTlsValue(BinnedGPUTlsSlot, ThreadSingleton);
+		Allocator.RegisterThreadFreeBlockLists(ThreadSingleton);
+	}
+}
+
+int64 FMallocBinnedGPU::FPerThreadFreeBlockLists::ClearTLS(FMallocBinnedGPU& Allocator)
+{
+	uint32 BinnedGPUTlsSlot = Allocator.BinnedGPUTlsSlot;
+	check(BinnedGPUTlsSlot);
+	int64 Result = 0;
+	FPerThreadFreeBlockLists* ThreadSingleton = (FPerThreadFreeBlockLists*)FPlatformTLS::GetTlsValue(BinnedGPUTlsSlot);
+	if (ThreadSingleton)
+	{
+		int64 TLSSize = sizeof(FPerThreadFreeBlockLists);
+		MBG_STAT(Allocator.BinnedGPUTLSMemory -= TLSSize;)
+		Result = Allocator.UnregisterThreadFreeBlockLists(ThreadSingleton);
+	}
+	FPlatformTLS::SetTlsValue(BinnedGPUTlsSlot, nullptr);
+	return Result;
+}
+
+void FMallocBinnedGPU::FFreeBlock::CanaryFail() const
+{
+	UE_LOG(LogMemory, Fatal, TEXT("FMallocBinnedGPU Attempt to realloc an unrecognized block %p   canary == 0x%x != 0x%x"), (void*)this, (int32)Canary, (int32)FMallocBinnedGPU::FFreeBlock::CANARY_VALUE);
+}
+
+MBG_STAT(
+int64 FMallocBinnedGPU::GetTotalAllocatedSmallPoolMemory()
+{
+	int64 FreeBlockAllocatedMemory = 0;
+	{
+		FScopeLock Lock(&GetFreeBlockListsRegistrationMutex());
+		for (const FPerThreadFreeBlockLists* FreeBlockLists : GetRegisteredFreeBlockLists())
+		{
+			FreeBlockAllocatedMemory += FreeBlockLists->AllocatedMemory;
+		}
+		FreeBlockAllocatedMemory += ConsolidatedMemory;
+	}
+	return BinnedGPUAllocatedSmallPoolMemory + FreeBlockAllocatedMemory;
+})
+
+void FMallocBinnedGPU::GetAllocatorStats(FGenericMemoryStats& OutStats)
+{
+	MBG_STAT(
+		int64 TotalAllocatedSmallPoolMemory = GetTotalAllocatedSmallPoolMemory();
+
+		OutStats.Add(TEXT("BinnedGPUAllocatedSmallPoolMemory"), TotalAllocatedSmallPoolMemory);
+		OutStats.Add(TEXT("BinnedGPUAllocatedOSSmallPoolMemory"), BinnedGPUAllocatedOSSmallPoolMemory);
+		OutStats.Add(TEXT("BinnedGPUAllocatedLargePoolMemory"), BinnedGPUAllocatedLargePoolMemory);
+		OutStats.Add(TEXT("BinnedGPUAllocatedLargePoolMemoryWAlignment"), BinnedGPUAllocatedLargePoolMemoryWAlignment);
+
+		uint64 TotalAllocated = TotalAllocatedSmallPoolMemory + BinnedGPUAllocatedLargePoolMemory;
+		uint64 TotalOSAllocated = BinnedGPUAllocatedOSSmallPoolMemory + BinnedGPUAllocatedLargePoolMemoryWAlignment;
+
+		OutStats.Add(TEXT("TotalAllocated"), TotalAllocated);
+		OutStats.Add(TEXT("TotalOSAllocated"), TotalOSAllocated);
+	)
+	FMalloc::GetAllocatorStats(OutStats);
+}
+
+void FMallocBinnedGPU::DumpAllocatorStats(class FOutputDevice& Ar)
+{
+
+	Ar.Logf(TEXT("FMallocBinnedGPU Mem report"));
+	Ar.Logf(TEXT("Constants.BinnedAllocationGranularity = %d"), int32(ArenaParams.AllocationGranularity));
+	Ar.Logf(TEXT("ArenaParams.MaxPoolSize = %d"), int32(ArenaParams.MaxPoolSize));
+	Ar.Logf(TEXT("MAX_MEMORY_PER_BLOCK_SIZE = %llu"), uint64(ArenaParams.MaxMemoryPerBlockSize));
+	MBG_STAT(
+		int64 TotalAllocatedSmallPoolMemory = GetTotalAllocatedSmallPoolMemory();
+		Ar.Logf(TEXT("Small Pool Allocations: %fmb  (including block size padding)"), ((double)TotalAllocatedSmallPoolMemory) / (1024.0f * 1024.0f));
+		Ar.Logf(TEXT("Small Pool OS Allocated: %fmb"), ((double)BinnedGPUAllocatedOSSmallPoolMemory) / (1024.0f * 1024.0f));
+		Ar.Logf(TEXT("Large Pool Requested Allocations: %fmb"), ((double)BinnedGPUAllocatedLargePoolMemory) / (1024.0f * 1024.0f));
+		Ar.Logf(TEXT("Large Pool OS Allocated: %fmb"), ((double)BinnedGPUAllocatedLargePoolMemoryWAlignment) / (1024.0f * 1024.0f));
+		Ar.Logf(TEXT("PoolInfo: %fmb"), ((double)BinnedGPUPoolInfoMemory) / (1024.0f * 1024.0f));
+		Ar.Logf(TEXT("Hash: %fmb"), ((double)BinnedGPUHashMemory) / (1024.0f * 1024.0f));
+		Ar.Logf(TEXT("Free Bits: %fmb"), ((double)BinnedGPUFreeBitsMemory) / (1024.0f * 1024.0f));
+		Ar.Logf(TEXT("TLS: %fmb"), ((double)BinnedGPUTLSMemory) / (1024.0f * 1024.0f));
+		Ar.Logf(TEXT("GPU Memory Proxies: %fmb"), ((double)GPUProxyMemory) / (1024.0f * 1024.0f));
+		Ar.Logf(TEXT("Total allocated from OS: %fmb"),
+			((double)
+				BinnedGPUAllocatedOSSmallPoolMemory + BinnedGPUAllocatedLargePoolMemoryWAlignment + BinnedGPUPoolInfoMemory + BinnedGPUHashMemory + BinnedGPUFreeBitsMemory + BinnedGPUTLSMemory + GPUProxyMemory
+				) / (1024.0f * 1024.0f));
+	)
+	Ar.Logf(TEXT("BINNEDGPU_USE_SEPARATE_VM_PER_POOL is true - VM is Contiguous = %d"), PoolSearchDiv == 0);
+	if (PoolSearchDiv)
+	{
+		uint64 TotalMem = PoolBaseVMPtr[ArenaParams.PoolCount - 1] + ArenaParams.MaxMemoryPerBlockSize - PoolBaseVMPtr[0];
+		uint64 MinimumMem = uint64(ArenaParams.PoolCount) * ArenaParams.MaxMemoryPerBlockSize;
+		Ar.Logf(TEXT("Percent of gaps in the address range %6.4f  (hopefully < 1, or the searches above will suffer)"), 100.0f * (1.0f - float(MinimumMem) / float(TotalMem)));
+	}
+
+	for (int32 PoolIndex = 0; PoolIndex < ArenaParams.PoolCount; PoolIndex++)
+	{
+
+		int64 VM = SmallPoolTables[PoolIndex].UnusedAreaOffsetLow;
+		uint32 CommittedBlocks = SmallPoolTables[PoolIndex].BlockOfBlockAllocationBits.CountOnes(SmallPoolTables[PoolIndex].NumEverUsedBlockOfBlocks);
+		uint32 PartialBlocks = SmallPoolTables[PoolIndex].NumEverUsedBlockOfBlocks - SmallPoolTables[PoolIndex].BlockOfBlockIsExhausted.CountOnes(SmallPoolTables[PoolIndex].NumEverUsedBlockOfBlocks);
+		uint32 FullBlocks = CommittedBlocks - PartialBlocks;
+		int64 ComittedVM = VM - (SmallPoolTables[PoolIndex].NumEverUsedBlockOfBlocks - CommittedBlocks) * SmallPoolTables[PoolIndex].PagesPlatformForBlockOfBlocks * ArenaParams.AllocationGranularity;
+
+
+		Ar.Logf(TEXT("Pool %2d   Size %6d   UsedVM %3dMB  CommittedVM %3dMB  HighSlabs %6d  CommittedSlabs %6d  FullSlabs %6d  PartialSlabs  %6d"),
+			PoolIndex,
+			PoolIndexToBlockSize(PoolIndex),
+			VM / (1024 * 1024),
+			ComittedVM / (1024 * 1024),
+			SmallPoolTables[PoolIndex].NumEverUsedBlockOfBlocks,
+			CommittedBlocks,
+			FullBlocks,
+			PartialBlocks
+		);
+	}
+}
+#endif
\ No newline at end of file
diff --git a/Engine/Source/Runtime/Core/Private/HAL/UnrealMemory.cpp b/Engine/Source/Runtime/Core/Private/HAL/UnrealMemory.cpp
--- a/Engine/Source/Runtime/Core/Private/HAL/UnrealMemory.cpp
+++ b/Engine/Source/Runtime/Core/Private/HAL/UnrealMemory.cpp
@@ -385,7 +385,7 @@ int FMemory_GCreateMalloc_ThreadUnsafe()
 	FPlatformMisc::LowLevelOutputDebugStringf(TEXT("Used memory before allocating anything was %.2fMB\n"), SizeInMb);
 	UE_LOG(LogMemory, Display, TEXT("Used memory before allocating anything was %.2fMB"), SizeInMb);
 #endif
-	
+
 	return 0;
 }
 
@@ -538,21 +538,6 @@ void FMemory::ClearAndDisableTLSCachesOnCurrentThread()
 	}
 }
 
-void* FMemory::GPUMalloc(SIZE_T Count, uint32 Alignment /* = DEFAULT_ALIGNMENT */)
-{
-	return FPlatformMemory::GPUMalloc(Count, Alignment);
-}
-
-void* FMemory::GPURealloc(void* Original, SIZE_T Count, uint32 Alignment /* = DEFAULT_ALIGNMENT */)
-{
-	return FPlatformMemory::GPURealloc(Original, Count, Alignment);
-}
-
-void FMemory::GPUFree(void* Original)
-{
-	return FPlatformMemory::GPUFree(Original);
-}
-
 void FMemory::TestMemory()
 {
 #if !UE_BUILD_SHIPPING
diff --git a/Engine/Source/Runtime/Core/Private/Unix/UnixPlatformMemory.cpp b/Engine/Source/Runtime/Core/Private/Unix/UnixPlatformMemory.cpp
--- a/Engine/Source/Runtime/Core/Private/Unix/UnixPlatformMemory.cpp
+++ b/Engine/Source/Runtime/Core/Private/Unix/UnixPlatformMemory.cpp
@@ -32,6 +32,11 @@
 #include "GenericPlatform/OSAllocationPool.h"
 #include "Misc/ScopeLock.h"
 
+// on 64 bit Linux, it is easier to run out of vm.max_map_count than of other limits. Due to that, trade VIRT (address space) size for smaller amount of distinct mappings
+// by not leaving holes between them (kernel will coalesce the adjoining mappings into a single one)
+// Disable by default as this causes large wasted virtual memory areas as we've to cut out 64k aligned pointers from a larger memory area then requested but then leave them mmap'd
+#define UE4_PLATFORM_REDUCE_NUMBER_OF_MAPS					(0 && PLATFORM_UNIX && PLATFORM_64BITS)
+
 // do not do a root privilege check on non-x86-64 platforms (assume an embedded device)
 #if defined(_M_X64) || defined(__x86_64__) || defined (__amd64__) 
 	#define UE4_DO_ROOT_PRIVILEGE_CHECK	 1
@@ -302,28 +307,282 @@ bool FUnixPlatformMemory::PageProtect(void* const Ptr, const SIZE_T Size, const
 	return mprotect(Ptr, Size, ProtectMode) == 0;
 }
 
+
+static void MarkMappedMemoryMergable(void* Pointer, SIZE_T Size)
+{
+	const SIZE_T BinnedPageSize = FPlatformMemory::GetConstants().BinnedPageSize;
+
+	// If we dont want to merge all pages only merge chunks larger then BinnedPageSize
+	if (GUseKSM && (GKSMMergeAllPages || Size > BinnedPageSize))
+	{
+		int Ret = madvise(Pointer, Size, MADV_MERGEABLE);
+		if (Ret != 0)
+		{
+			int ErrNo = errno;
+			UE_LOG(LogHAL, Fatal, TEXT("madvise(addr=%p, length=%d, advice=MADV_MERGEABLE) failed with errno = %d (%s)"),
+				Pointer, Size, ErrNo, StringCast< TCHAR >(strerror(ErrNo)).Get());
+			// unreachable
+		}
+	}
+}
+
+#ifndef MALLOC_LEAKDETECTION
+	#define MALLOC_LEAKDETECTION 0
+#endif
+// check bookkeeping info against the passed in parameters in Debug and Development (the latter only in games and servers. also, only if leak detection is disabled, otherwise things are very slow)
+#define UE4_PLATFORM_SANITY_CHECK_OS_ALLOCATIONS			(UE_BUILD_DEBUG || (UE_BUILD_DEVELOPMENT && (UE_GAME || UE_SERVER) && !MALLOC_LEAKDETECTION))
+
+/** This structure is stored in the page after each OS allocation and checks that its properties are valid on Free. Should be less than the page size (4096 on all platforms we support) */
+struct FOSAllocationDescriptor
+{
+	enum class MagicType : uint64
+	{
+		Marker = 0xd0c233ccf493dfb0
+	};
+
+	/** Magic that makes sure that we are not passed a pointer somewhere into the middle of the allocation (and/or the structure wasn't stomped). */
+	MagicType	Magic;
+
+	/** This should include the descriptor itself. */
+	void*		PointerToUnmap;
+
+	/** This should include the total size of allocation, so after unmapping it everything is gone, including the descriptor */
+	SIZE_T		SizeToUnmap;
+
+	/** Debug info that makes sure that the correct size is preserved. */
+	SIZE_T		OriginalSizeAsPassed;
+};
+
 void* FUnixPlatformMemory::BinnedAllocFromOS(SIZE_T Size)
 {
-	void* Ret = FGenericPlatformMemory::BinnedAllocFromOS(Size);
-	LLM(FLowLevelMemTracker::Get().OnLowLevelAlloc(ELLMTracker::Platform, Ret, Size));
-	return Ret;
+	static SIZE_T OSPageSize = FPlatformMemory::GetConstants().PageSize;
+	// guard against someone not passing size in whole pages
+	SIZE_T SizeInWholePages = (Size % OSPageSize) ? (Size + OSPageSize - (Size % OSPageSize)) : Size;
+	void* Pointer = nullptr;
+
+	// Binned expects OS allocations to be BinnedPageSize-aligned, and that page is at least 64KB. mmap() alone cannot do this, so carve out the needed chunks.
+	const SIZE_T ExpectedAlignment = FPlatformMemory::GetConstants().BinnedPageSize;
+	// Descriptor is only used if we're sanity checking. However, #ifdef'ing its use would make the code more fragile. Size needs to be at least one page.
+	const SIZE_T DescriptorSize = (UE4_PLATFORM_REDUCE_NUMBER_OF_MAPS != 0 || UE4_PLATFORM_SANITY_CHECK_OS_ALLOCATIONS != 0) ? OSPageSize : 0;
+
+	SIZE_T ActualSizeMapped = SizeInWholePages + ExpectedAlignment;
+
+	// the remainder of the map will be used for the descriptor, if any.
+	// we always allocate at least one page more than needed
+	void* PointerWeGotFromMMap = mmap(nullptr, ActualSizeMapped, PROT_READ | PROT_WRITE, MAP_PRIVATE | MAP_ANON, -1, 0);
+	// store these values, to unmap later
+
+	Pointer = PointerWeGotFromMMap;
+	if (Pointer == MAP_FAILED)
+	{
+		FPlatformMemory::OnOutOfMemory(ActualSizeMapped, ExpectedAlignment);
+		// unreachable
+		return nullptr;
+	}
+
+	SIZE_T Offset = (reinterpret_cast<SIZE_T>(Pointer) % ExpectedAlignment);
+
+	// See if we need to unmap anything in the front. If the pointer happened to be aligned, we don't need to unmap anything.
+	if (LIKELY(Offset != 0))
+	{
+		// figure out how much to unmap before the alignment.
+		SIZE_T SizeToNextAlignedPointer = ExpectedAlignment - Offset;
+		void* AlignedPointer = reinterpret_cast<void*>(reinterpret_cast<SIZE_T>(Pointer) + SizeToNextAlignedPointer);
+
+		// do not unmap if we're trying to reduce the number of distinct maps, since holes prevent the Linux kernel from coalescing two adjoining mmap()s into a single VMA
+		if (!UE4_PLATFORM_REDUCE_NUMBER_OF_MAPS)
+		{
+			// unmap the part before
+			if (munmap(Pointer, SizeToNextAlignedPointer) != 0)
+			{
+				FPlatformMemory::OnOutOfMemory(SizeToNextAlignedPointer, ExpectedAlignment);
+				// unreachable
+				return nullptr;
+			}
+
+			// account for reduced mmaped size
+			ActualSizeMapped -= SizeToNextAlignedPointer;
+		}
+
+		// now, make it appear as if we initially got the allocation right
+		Pointer = AlignedPointer;
+	}
+
+	// at this point, Pointer is aligned at the expected alignment - either we lucked out on the initial allocation
+	// or we already got rid of the extra memory that was allocated in the front.
+	checkf((reinterpret_cast<SIZE_T>(Pointer) % ExpectedAlignment) == 0, TEXT("BinnedAllocFromOS(): Internal error: did not align the pointer as expected."));
+
+	// do not unmap if we're trying to reduce the number of distinct maps, since holes prevent the Linux kernel from coalescing two adjoining mmap()s into a single VMA
+	if (!UE4_PLATFORM_REDUCE_NUMBER_OF_MAPS)
+	{
+		// Now unmap the tail only, if any, but leave just enough space for the descriptor
+		void* TailPtr = reinterpret_cast<void*>(reinterpret_cast<SIZE_T>(Pointer) + SizeInWholePages + DescriptorSize);
+		SIZE_T TailSize = ActualSizeMapped - SizeInWholePages - DescriptorSize;
+
+		if (LIKELY(TailSize > 0))
+		{
+			if (munmap(TailPtr, TailSize) != 0)
+			{
+				FPlatformMemory::OnOutOfMemory(TailSize, ExpectedAlignment);
+				// unreachable
+				return nullptr;
+			}
+		}
+	}
+
+	// we're done with this allocation, fill in the descriptor with the info
+	if (LIKELY(DescriptorSize > 0))
+	{
+		FOSAllocationDescriptor* AllocDescriptor = reinterpret_cast<FOSAllocationDescriptor*>(reinterpret_cast<SIZE_T>(Pointer) + Size);
+		AllocDescriptor->Magic = FOSAllocationDescriptor::MagicType::Marker;
+		if (!UE4_PLATFORM_REDUCE_NUMBER_OF_MAPS)
+		{
+			AllocDescriptor->PointerToUnmap = Pointer;
+			AllocDescriptor->SizeToUnmap = SizeInWholePages + DescriptorSize;
+		}
+		else
+		{
+			AllocDescriptor->PointerToUnmap = PointerWeGotFromMMap;
+			AllocDescriptor->SizeToUnmap = ActualSizeMapped;
+		}
+		AllocDescriptor->OriginalSizeAsPassed = Size;
+	}
+
+	LLM(FLowLevelMemTracker::Get().OnLowLevelAlloc(ELLMTracker::Platform, Pointer, Size));
+	return Pointer;
 }
 
 void FUnixPlatformMemory::BinnedFreeToOS(void* Ptr, SIZE_T Size)
 {
 	LLM(FLowLevelMemTracker::Get().OnLowLevelFree(ELLMTracker::Platform, Ptr));
-	return FGenericPlatformMemory::BinnedFreeToOS(Ptr, Size);
+	// guard against someone not passing size in whole pages
+	static SIZE_T OSPageSize = FPlatformMemory::GetConstants().PageSize;
+	SIZE_T SizeInWholePages = (Size % OSPageSize) ? (Size + OSPageSize - (Size % OSPageSize)) : Size;
+
+	if (UE4_PLATFORM_REDUCE_NUMBER_OF_MAPS || UE4_PLATFORM_SANITY_CHECK_OS_ALLOCATIONS)
+	{
+		const SIZE_T DescriptorSize = OSPageSize;
+
+		FOSAllocationDescriptor* AllocDescriptor = reinterpret_cast<FOSAllocationDescriptor*>(reinterpret_cast<SIZE_T>(Ptr) + Size);
+		if (UNLIKELY(AllocDescriptor->Magic != FOSAllocationDescriptor::MagicType::Marker))
+		{
+			UE_LOG(LogHAL, Fatal, TEXT("BinnedFreeToOS() has been passed an address %p (size %llu) not allocated through it."), Ptr, (uint64)Size);
+			// unreachable
+			return;
+		}
+
+		void* PointerToUnmap = AllocDescriptor->PointerToUnmap;
+		SIZE_T SizeToUnmap = AllocDescriptor->SizeToUnmap;
+
+		// do checks, from most to least serious
+		if (UE4_PLATFORM_SANITY_CHECK_OS_ALLOCATIONS != 0)
+		{
+			// this check only makes sense when we're not reducing number of maps, since the pointer will have to be different.
+			if (UE4_PLATFORM_REDUCE_NUMBER_OF_MAPS == 0)
+			{
+				if (UNLIKELY(PointerToUnmap != Ptr || SizeToUnmap != SizeInWholePages + DescriptorSize))
+				{
+					UE_LOG(LogHAL, Fatal, TEXT("BinnedFreeToOS(): info mismatch: descriptor ptr: %p, size %llu, but our pointer is %p and size %llu."), PointerToUnmap, SizeToUnmap, AllocDescriptor, (uint64)(SizeInWholePages + DescriptorSize));
+					// unreachable
+					return;
+				}
+			}
+
+			if (UNLIKELY(AllocDescriptor->OriginalSizeAsPassed != Size))
+			{
+				UE_LOG(LogHAL, Fatal, TEXT("BinnedFreeToOS(): info mismatch: descriptor original size %llu, our size is %llu for pointer %p"), AllocDescriptor->OriginalSizeAsPassed, Size, Ptr);
+				// unreachable
+				return;
+			}
+		}
+
+		AllocDescriptor = nullptr;	// just so no one touches it
+
+		if (UNLIKELY(munmap(PointerToUnmap, SizeToUnmap) != 0))
+		{
+			FPlatformMemory::OnOutOfMemory(SizeToUnmap, 0);
+			// unreachable
+		}
+	}
+	else
+	{
+		if (UNLIKELY(munmap(Ptr, SizeInWholePages) != 0))
+		{
+			FPlatformMemory::OnOutOfMemory(SizeInWholePages, 0);
+			// unreachable
+		}
+	}
 }
 
-bool FUnixPlatformMemory::MemoryRangeDecommit(void* Ptr, SIZE_T Size)
+size_t FUnixPlatformMemory::FPlatformVirtualMemoryBlock::GetVirtualSizeAlignment()
 {
-	if (LIKELY(GMemoryRangeDecommitIsNoOp))
+	static SIZE_T OSPageSize = FPlatformMemory::GetConstants().PageSize;
+	return OSPageSize;
+}
+
+size_t FUnixPlatformMemory::FPlatformVirtualMemoryBlock::GetCommitAlignment()
+{
+	static SIZE_T OSPageSize = FPlatformMemory::GetConstants().PageSize;
+	return OSPageSize;
+}
+
+FUnixPlatformMemory::FPlatformVirtualMemoryBlock FUnixPlatformMemory::FPlatformVirtualMemoryBlock::AllocateVirtual(size_t InSize, size_t InAlignment)
+{
+	FPlatformVirtualMemoryBlock Result;
+	InSize = Align(InSize, GetVirtualSizeAlignment());
+	Result.VMSizeDivVirtualSizeAlignment = InSize / GetVirtualSizeAlignment();
+
+	size_t Alignment = FMath::Max(InAlignment, GetVirtualSizeAlignment());
+	check(Alignment <= GetVirtualSizeAlignment());
+
+	Result.Ptr = mmap(nullptr, Result.GetActualSize(), PROT_READ | PROT_WRITE, MAP_PRIVATE | MAP_ANON, -1, 0);
+	if (LIKELY(Result.Ptr != MAP_FAILED))
+	{
+		MarkMappedMemoryMergable(Result.Ptr, Result.GetActualSize());
+	}
+	else
 	{
-		return true;
+		FPlatformMemory::OnOutOfMemory(Result.GetActualSize(), InAlignment);
 	}
-	return madvise(Ptr, Size, MADV_DONTNEED) == 0;
+	check(Result.Ptr && IsAligned(Result.Ptr, Alignment));
+	return Result;
 }
 
+
+
+void FUnixPlatformMemory::FPlatformVirtualMemoryBlock::FreeVirtual()
+{
+	if (Ptr)
+	{
+		check(GetActualSize() > 0);
+		if (munmap(Ptr, GetActualSize()) != 0)
+		{
+			// we can ran out of VMAs here
+			FPlatformMemory::OnOutOfMemory(GetActualSize(), 0);
+			// unreachable
+		}
+		Ptr = nullptr;
+		VMSizeDivVirtualSizeAlignment = 0;
+	}
+}
+
+void FUnixPlatformMemory::FPlatformVirtualMemoryBlock::Commit(size_t InOffset, size_t InSize)
+{
+	check(IsAligned(InOffset, GetCommitAlignment()) && IsAligned(InSize, GetCommitAlignment()));
+	check(InOffset >= 0 && InSize >= 0 && InOffset + InSize <= GetActualSize() && Ptr);
+}
+
+void FUnixPlatformMemory::FPlatformVirtualMemoryBlock::Decommit(size_t InOffset, size_t InSize)
+{
+	check(IsAligned(InOffset, GetCommitAlignment()) && IsAligned(InSize, GetCommitAlignment()));
+	check(InOffset >= 0 && InSize >= 0 && InOffset + InSize <= GetActualSize() && Ptr);
+	if (!LIKELY(GMemoryRangeDecommitIsNoOp))
+	{
+		madvise(((uint8*)Ptr) + InOffset, InSize, MADV_DONTNEED);
+	}
+}
+
+
 namespace UnixPlatformMemory
 {
 	/**
diff --git a/Engine/Source/Runtime/Core/Private/Windows/WindowsPlatformMemory.cpp b/Engine/Source/Runtime/Core/Private/Windows/WindowsPlatformMemory.cpp
--- a/Engine/Source/Runtime/Core/Private/Windows/WindowsPlatformMemory.cpp
+++ b/Engine/Source/Runtime/Core/Private/Windows/WindowsPlatformMemory.cpp
@@ -301,37 +301,89 @@ void FWindowsPlatformMemory::BinnedFreeToOS( void* Ptr, SIZE_T Size )
 	verify(VirtualFree( Ptr, 0, MEM_RELEASE ) != 0);
 }
 
-void* FWindowsPlatformMemory::MemoryRangeReserve(SIZE_T Size, bool bCommit, int32 Node)
+
+size_t FWindowsPlatformMemory::FPlatformVirtualMemoryBlock::GetVirtualSizeAlignment()
+{
+	static SIZE_T OsAllocationGranularity = FPlatformMemory::GetConstants().OsAllocationGranularity;
+	return OsAllocationGranularity;
+}
+
+size_t FWindowsPlatformMemory::FPlatformVirtualMemoryBlock::GetCommitAlignment()
 {
-	if (bCommit)
+	static SIZE_T OSPageSize = FPlatformMemory::GetConstants().PageSize;
+	return OSPageSize;
+}
+
+FWindowsPlatformMemory::FPlatformVirtualMemoryBlock FWindowsPlatformMemory::FPlatformVirtualMemoryBlock::AllocateVirtual(size_t InSize, size_t InAlignment)
+{
+	FPlatformVirtualMemoryBlock Result;
+	InSize = Align(InSize, GetVirtualSizeAlignment());
+	Result.VMSizeDivVirtualSizeAlignment = InSize / GetVirtualSizeAlignment();
+
+	size_t Alignment = FMath::Max(InAlignment, GetVirtualSizeAlignment());
+	check(Alignment <= GetVirtualSizeAlignment());
+
+	bool bTopDown = Result.GetActualSize() > 100ll * 1024 * 1024; // this is hacky, but we want to allocate huge VM blocks (like for MB3) top down
+
+	Result.Ptr = VirtualAlloc(NULL, Result.GetActualSize(), MEM_RESERVE | (bTopDown ? MEM_TOP_DOWN : 0), PAGE_NOACCESS);
+
+
+	if (!LIKELY(Result.Ptr))
 	{
-		return BinnedAllocFromOS(Size);
+		FPlatformMemory::OnOutOfMemory(Result.GetActualSize(), Alignment);
 	}
-	void* Ptr = VirtualAlloc(NULL, Size, MEM_RESERVE | MEM_TOP_DOWN, PAGE_NOACCESS);
-	// no LLM for uncommitted memory!
-	return Ptr;
-
+	check(Result.Ptr && IsAligned(Result.Ptr, Alignment));
+	return Result;
 }
 
-void FWindowsPlatformMemory::MemoryRangeFree(void* Ptr, SIZE_T Size)
+
+
+void FWindowsPlatformMemory::FPlatformVirtualMemoryBlock::FreeVirtual()
 {
-	// this is an iffy assumption, mallocbinned3 will never free decommitted memory, but we don't know that anyone else will
-	BinnedFreeToOS(Ptr, Size);
+	if (Ptr)
+	{
+		check(GetActualSize() > 0);
+		// this is an iffy assumption, we don't know how much of this memory is really committed, we will assume none of it is
+		//LLM(FLowLevelMemTracker::Get().OnLowLevelFree(ELLMTracker::Platform, Ptr));
+
+		CA_SUPPRESS(6001)
+		// Windows maintains the size of allocation internally, so Size is unused
+		verify(VirtualFree(Ptr, 0, MEM_RELEASE) != 0);
+
+		Ptr = nullptr;
+		VMSizeDivVirtualSizeAlignment = 0;
+	}
 }
 
-bool FWindowsPlatformMemory::MemoryRangeCommit(void* Ptr, SIZE_T Size)
+void FWindowsPlatformMemory::FPlatformVirtualMemoryBlock::Commit(size_t InOffset, size_t InSize)
 {
-	LLM(FLowLevelMemTracker::Get().OnLowLevelAlloc(ELLMTracker::Platform, Ptr, Size));
-	return VirtualAlloc(Ptr, Size, MEM_COMMIT, PAGE_READWRITE) == Ptr;
+	check(IsAligned(InOffset, GetCommitAlignment()) && IsAligned(InSize, GetCommitAlignment()));
+	check(InOffset >= 0 && InSize >= 0 && InOffset + InSize <= GetActualSize() && Ptr);
+
+	// There are no guarantees LLM is going to be able to deal with this
+	uint8* UsePtr = ((uint8*)Ptr) + InOffset;
+	LLM(FLowLevelMemTracker::Get().OnLowLevelAlloc(ELLMTracker::Platform, UsePtr, InSize));
+	if (VirtualAlloc(UsePtr, InSize, MEM_COMMIT, PAGE_READWRITE) != UsePtr)
+	{
+		FPlatformMemory::OnOutOfMemory(InSize, 0);
+	}
 }
 
-bool FWindowsPlatformMemory::MemoryRangeDecommit(void* Ptr, SIZE_T Size)
+void FWindowsPlatformMemory::FPlatformVirtualMemoryBlock::Decommit(size_t InOffset, size_t InSize)
 {
-	LLM(FLowLevelMemTracker::Get().OnLowLevelFree(ELLMTracker::Platform, Ptr));
-	return VirtualFree(Ptr, Size, MEM_DECOMMIT) != 0;
+	check(IsAligned(InOffset, GetCommitAlignment()) && IsAligned(InSize, GetCommitAlignment()));
+	check(InOffset >= 0 && InSize >= 0 && InOffset + InSize <= GetActualSize() && Ptr);
+	uint8* UsePtr = ((uint8*)Ptr) + InOffset;
+	// There are no guarantees LLM is going to be able to deal with this
+	LLM(FLowLevelMemTracker::Get().OnLowLevelFree(ELLMTracker::Platform, UsePtr));
+	VirtualFree(UsePtr, InSize, MEM_DECOMMIT);
 }
 
 
+
+
+
+
 FPlatformMemory::FSharedMemoryRegion* FWindowsPlatformMemory::MapNamedSharedMemoryRegion(const FString& InName, bool bCreate, uint32 AccessMode, SIZE_T Size)
 {
 	FString Name(TEXT("Global\\"));
diff --git a/Engine/Source/Runtime/Core/Public/Android/AndroidPlatformMemory.h b/Engine/Source/Runtime/Core/Public/Android/AndroidPlatformMemory.h
--- a/Engine/Source/Runtime/Core/Public/Android/AndroidPlatformMemory.h
+++ b/Engine/Source/Runtime/Core/Public/Android/AndroidPlatformMemory.h
@@ -26,10 +26,57 @@ struct CORE_API FAndroidPlatformMemory : public FGenericPlatformMemory
 	static const FPlatformMemoryConstants& GetConstants();
 	static EPlatformMemorySizeBucket GetMemorySizeBucket();
 	static FMalloc* BaseAllocator();
-	static void* BinnedAllocFromOS( SIZE_T Size );
-	static void BinnedFreeToOS( void* Ptr, SIZE_T Size );
-	static bool MemoryRangeDecommit(void* Ptr, SIZE_T Size);
+	static void* BinnedAllocFromOS(SIZE_T Size);
+	static void BinnedFreeToOS(void* Ptr, SIZE_T Size);
 
+	class FPlatformVirtualMemoryBlock : public FBasicVirtualMemoryBlock
+	{
+	public:
+
+		FPlatformVirtualMemoryBlock()
+		{
+		}
+
+		FPlatformVirtualMemoryBlock(void *InPtr, uint32 InVMSizeDivVirtualSizeAlignment)
+			: FBasicVirtualMemoryBlock(InPtr, InVMSizeDivVirtualSizeAlignment)
+		{
+		}
+		FPlatformVirtualMemoryBlock(const FPlatformVirtualMemoryBlock& Other) = default;
+		FPlatformVirtualMemoryBlock& operator=(const FPlatformVirtualMemoryBlock& Other) = default;
+
+		void Commit(size_t InOffset, size_t InSize);
+		void Decommit(size_t InOffset, size_t InSize);
+		void FreeVirtual();
+
+		FORCEINLINE void CommitByPtr(void *InPtr, size_t InSize)
+		{
+			Commit(size_t(((uint8*)InPtr) - ((uint8*)Ptr)), InSize);
+		}
+
+		FORCEINLINE void DecommitByPtr(void *InPtr, size_t InSize)
+		{
+			Decommit(size_t(((uint8*)InPtr) - ((uint8*)Ptr)), InSize);
+		}
+
+		FORCEINLINE void Commit()
+		{
+			Commit(0, GetActualSize());
+		}
+
+		FORCEINLINE void Decommit()
+		{
+			Decommit(0, GetActualSize());
+		}
+
+		FORCEINLINE size_t GetActualSize() const
+		{
+			return VMSizeDivVirtualSizeAlignment * GetVirtualSizeAlignment();
+		}
+
+		static FPlatformVirtualMemoryBlock AllocateVirtual(size_t Size, size_t InAlignment = FPlatformVirtualMemoryBlock::GetVirtualSizeAlignment());
+		static size_t GetCommitAlignment();
+		static size_t GetVirtualSizeAlignment();
+	};
 
 	static bool GetLLMAllocFunctions(void*(*&OutAllocFunction)(size_t), void(*&OutFreeFunction)(void*, size_t), int32& OutAlignment);
 	//~ End FGenericPlatformMemory Interface
diff --git a/Engine/Source/Runtime/Core/Public/Apple/ApplePlatformMemory.h b/Engine/Source/Runtime/Core/Public/Apple/ApplePlatformMemory.h
--- a/Engine/Source/Runtime/Core/Public/Apple/ApplePlatformMemory.h
+++ b/Engine/Source/Runtime/Core/Public/Apple/ApplePlatformMemory.h
@@ -67,8 +67,57 @@ struct CORE_API FApplePlatformMemory : public FGenericPlatformMemory
 	static const FPlatformMemoryConstants& GetConstants();
 	static FMalloc* BaseAllocator();
 	static bool PageProtect(void* const Ptr, const SIZE_T Size, const bool bCanRead, const bool bCanWrite);
-	static void* BinnedAllocFromOS( SIZE_T Size );
-	static void BinnedFreeToOS( void* Ptr, SIZE_T Size );
+	static void* BinnedAllocFromOS(SIZE_T Size);
+	static void BinnedFreeToOS(void* Ptr, SIZE_T Size);
+
+	class FPlatformVirtualMemoryBlock : public FBasicVirtualMemoryBlock
+	{
+	public:
+
+		FPlatformVirtualMemoryBlock()
+		{
+		}
+
+		FPlatformVirtualMemoryBlock(void *InPtr, uint32 InVMSizeDivVirtualSizeAlignment)
+			: FBasicVirtualMemoryBlock(InPtr, InVMSizeDivVirtualSizeAlignment)
+		{
+		}
+		FPlatformVirtualMemoryBlock(const FPlatformVirtualMemoryBlock& Other) = default;
+		FPlatformVirtualMemoryBlock& operator=(const FPlatformVirtualMemoryBlock& Other) = default;
+
+		void Commit(size_t InOffset, size_t InSize);
+		void Decommit(size_t InOffset, size_t InSize);
+		void FreeVirtual();
+
+		FORCEINLINE void CommitByPtr(void *InPtr, size_t InSize)
+		{
+			Commit(size_t(((uint8*)InPtr) - ((uint8*)Ptr)), InSize);
+		}
+
+		FORCEINLINE void DecommitByPtr(void *InPtr, size_t InSize)
+		{
+			Decommit(size_t(((uint8*)InPtr) - ((uint8*)Ptr)), InSize);
+		}
+
+		FORCEINLINE void Commit()
+		{
+			Commit(0, GetActualSize());
+		}
+
+		FORCEINLINE void Decommit()
+		{
+			Decommit(0, GetActualSize());
+		}
+
+		FORCEINLINE size_t GetActualSize() const
+		{
+			return VMSizeDivVirtualSizeAlignment * GetVirtualSizeAlignment();
+		}
+
+		static FPlatformVirtualMemoryBlock AllocateVirtual(size_t Size, size_t InAlignment = FPlatformVirtualMemoryBlock::GetVirtualSizeAlignment());
+		static size_t GetCommitAlignment();
+		static size_t GetVirtualSizeAlignment();
+	};
 
     static bool GetLLMAllocFunctions(void* _Nonnull (* _Nonnull &OutAllocFunction)(size_t), void(* _Nonnull &OutFreeFunction)(void*, size_t), int32& OutAlignment);
     //~ End FGenericPlatformMemory Interface
diff --git a/Engine/Source/Runtime/Core/Public/GenericPlatform/GenericPlatformMemory.h b/Engine/Source/Runtime/Core/Public/GenericPlatform/GenericPlatformMemory.h
--- a/Engine/Source/Runtime/Core/Public/GenericPlatform/GenericPlatformMemory.h
+++ b/Engine/Source/Runtime/Core/Public/GenericPlatform/GenericPlatformMemory.h
@@ -327,28 +327,81 @@ struct CORE_API FGenericPlatformMemory
 	 */
 	static void BinnedFreeToOS( void* Ptr, SIZE_T Size );
 
-	/**
-	 * Reserves a virtual memory range.
-	 * 
-	 * Unlike BinnedAllocFromOS, this function does not have to return pointers with alignment larger than the VM page.
-	 */
-	static void* MemoryRangeReserve(SIZE_T Size, bool bCommit = false, int32 Node = -1);
+	class FBasicVirtualMemoryBlock
+	{
+	protected:
+		void *Ptr;
+		uint32 VMSizeDivVirtualSizeAlignment;
 
-	/**
-	 * Frees a virtual memory range.
-	 * 
-	 */
-	static void MemoryRangeFree(void* Ptr, SIZE_T Size);
+	public:
 
-	/**
-	 * Commits a virtual memory range. Can be no-op on platforms where just reservation of the range is not supported.
-	 */
-	static bool MemoryRangeCommit(void* Ptr, SIZE_T Size) { return true; }
+		FBasicVirtualMemoryBlock()
+			: Ptr(nullptr)
+			, VMSizeDivVirtualSizeAlignment(0)
+		{
+		}
 
-	/**
-	 * Decommits a virtual memory range. Can be no-op on platforms where just reservation of the range is not supported.
-	 */
-	static bool MemoryRangeDecommit(void* Ptr, SIZE_T Size) { return true; }
+		FBasicVirtualMemoryBlock(void *InPtr, uint32 InVMSizeDivVirtualSizeAlignment)
+			: Ptr(InPtr)
+			, VMSizeDivVirtualSizeAlignment(InVMSizeDivVirtualSizeAlignment)
+		{
+		}
+
+		FBasicVirtualMemoryBlock(const FBasicVirtualMemoryBlock& Other) = default;
+		FBasicVirtualMemoryBlock& operator=(const FBasicVirtualMemoryBlock& Other) = default;
+
+		FORCEINLINE uint32 GetActualSizeInPages() const
+		{
+			return VMSizeDivVirtualSizeAlignment;
+		}
+
+		FORCEINLINE void* GetVirtualPointer() const
+		{
+			return Ptr;
+		}
+
+		
+#if 0 // documentation, must be defined by the platform
+		void Commit(size_t InOffset, size_t InSize);
+		void Decommit(size_t InOffset, size_t InSize);
+		void FreeVirtual();
+
+		FORCEINLINE void CommitByPtr(void *InPtr, size_t InSize)
+		{
+			Commit(size_t(((uint8*)InPtr) - ((uint8*)Ptr)), InSize);
+		}
+
+		FORCEINLINE void DecommitByPtr(void *InPtr, size_t InSize)
+		{
+			Decommit(size_t(((uint8*)InPtr) - ((uint8*)Ptr)), InSize);
+		}
+
+		FORCEINLINE void Commit()
+		{
+			Commit(0, GetActualSize());
+		}
+
+		FORCEINLINE void Decommit()
+		{
+			Decommit(0, GetActualSize());
+		}
+
+		FORCEINLINE size_t GetActualSize() const
+		{
+			return VMSizeDivVirtualSizeAlignment * GetVirtualSizeAlignment();
+		}
+
+		/**
+		 * Returns uncommitted virtual memory
+		 *
+		 * @param Size Size of virtual memory block, must be aligned to GetVirtualSizeAlignment
+		 * @param InAlignment Alignment of returned virtual Ptr
+		 */
+		static FPlatformVirtualMemoryBlock AllocateVirtual(size_t Size, size_t InAlignment = FPlatformVirtualMemoryBlock::GetVirtualSizeAlignment());
+		static size_t GetCommitAlignment();
+		static size_t GetVirtualSizeAlignment();
+#endif
+	};
 
 	/**
 	 * Some platforms may pool allocations of this size to reduce OS calls. This function
@@ -359,12 +412,6 @@ struct CORE_API FGenericPlatformMemory
 		return false;
 	}
 
-	// These alloc/free memory that is mapped to the GPU
-	// Only for platforms with UMA (XB1/PS4/etc)
-	static void* GPUMalloc(SIZE_T Count, uint32 Alignment = 0) { return nullptr; };
-	static void* GPURealloc(void* Original, SIZE_T Count, uint32 Alignment = 0) { return nullptr; };
-	static void GPUFree(void* Original) { };
-
 	/** Dumps basic platform memory statistics into the specified output device. */
 	static void DumpStats( FOutputDevice& Ar );
 
diff --git a/Engine/Source/Runtime/Core/Public/GenericPlatform/OSAllocationPool.h b/Engine/Source/Runtime/Core/Public/GenericPlatform/OSAllocationPool.h
--- a/Engine/Source/Runtime/Core/Public/GenericPlatform/OSAllocationPool.h
+++ b/Engine/Source/Runtime/Core/Public/GenericPlatform/OSAllocationPool.h
@@ -2,6 +2,7 @@
 
 #pragma once
 
+#if PLATFORM_HAS_FPlatformVirtualMemoryBlock
 /*
  *  Perform checks that ensure that the pools are working as intended. This is not necessary in builds that are used for Shipping/Test or for the Development editor.
  */
@@ -15,8 +16,6 @@
 * @param RequiredAlignment alignment for the addresses returned from this pool
 */
 template<
-	bool(*CommitAddressRange)(void* AddrStart, SIZE_T Size),
-	bool(*EvictAddressRange)(void* AddrStart, SIZE_T Size),
 	SIZE_T RequiredAlignment
 >
 class TMemoryPool
@@ -47,13 +46,15 @@ protected:
 	/** When we're allocating less than block size, only BlockSize - Size is going to be used. */
 	SIZE_T 		UsefulMemorySize;
 
+	FPlatformMemory::FPlatformVirtualMemoryBlock VMBlock;
+
 #if UE4_TMEMORY_POOL_DO_SANITY_CHECKS
 	FThreadSafeCounter NoConcurrentAccess; // Tests that this is not being accessed on multiple threads at once, see TestPAL mallocthreadtest
 #endif // UE4_TMEMORY_POOL_DO_SANITY_CHECKS
 
 public:
 
-	TMemoryPool(SIZE_T InBlockSize, SIZE_T InAlignedPoolStart, SIZE_T InNumBlocks, uint8* InBitmask)
+	TMemoryPool(SIZE_T InBlockSize, SIZE_T InAlignedPoolStart, SIZE_T InNumBlocks, uint8* InBitmask, FPlatformMemory::FPlatformVirtualMemoryBlock& InVMBlock)
 		: BlockSize(InBlockSize)
 		, AlignedPoolStart(InAlignedPoolStart)
 		, AlignedPoolEnd(InAlignedPoolStart + InBlockSize * (SIZE_T)InNumBlocks)
@@ -62,6 +63,7 @@ public:
 		, BitmaskSizeInBytes(BitmaskMemorySize(NumBlocks))
 		, NumFreeBlocks(InNumBlocks)
 		, UsefulMemorySize(0)
+		, VMBlock(InVMBlock)
 	{
 #if UE4_TMEMORY_POOL_DO_SANITY_CHECKS
 		checkf((AlignedPoolStart % RequiredAlignment) == 0, TEXT("Non-aligned pool address passed to a TMemoryPool"));
@@ -71,7 +73,7 @@ public:
 		//checkf(NumFreeBlocks == CalculateFreeBlocksInBitmap(), TEXT("Mismatch between a bitmap and NumFreeBlocks at the very beginning"));
 
 		// decommit all the memory
-		EvictAddressRange(reinterpret_cast<void *>(AlignedPoolStart), NumBlocks * BlockSize);
+		VMBlock.DecommitByPtr(reinterpret_cast<void *>(AlignedPoolStart), Align(NumBlocks * BlockSize, FPlatformMemory::FPlatformVirtualMemoryBlock::GetCommitAlignment()));
 	}
 
 	/** We always allocate in BlockSize chunks, Size is only passed for more accurate Commit() */
@@ -97,7 +99,8 @@ public:
 #endif // UE4_TMEMORY_POOL_DO_SANITY_CHECKS
 
 			// make sure this address range is backed by the actual RAM
-			CommitAddressRange(Address, Size);
+			VMBlock.CommitByPtr(Address, Align(Size, FPlatformMemory::FPlatformVirtualMemoryBlock::GetCommitAlignment()));
+
 		}
 
 #if UE4_TMEMORY_POOL_DO_SANITY_CHECKS
@@ -136,7 +139,7 @@ public:
 		//checkf(NumFreeBlocks == CalculateFreeBlocksInBitmap(), TEXT("Mismatch between a bitmap and NumFreeBlocks after Free()ing"));
 
 		// evict this memory
-		EvictAddressRange(Ptr, BlockSize);
+		VMBlock.DecommitByPtr(Ptr, BlockSize);
 
 #if UE4_TMEMORY_POOL_DO_SANITY_CHECKS
 		// check that we never free more than we allocated
@@ -310,3 +313,5 @@ public:
 		printf("BlockSize: %llu NumAllocated/TotalBlocks = %llu/%llu\n", (uint64)BlockSize, (uint64)(NumBlocks - NumFreeBlocks), (uint64)NumBlocks);
 	}
 };
+
+#endif
\ No newline at end of file
diff --git a/Engine/Source/Runtime/Core/Public/HAL/Allocators/PooledVirtualMemoryAllocator.h b/Engine/Source/Runtime/Core/Public/HAL/Allocators/PooledVirtualMemoryAllocator.h
--- a/Engine/Source/Runtime/Core/Public/HAL/Allocators/PooledVirtualMemoryAllocator.h
+++ b/Engine/Source/Runtime/Core/Public/HAL/Allocators/PooledVirtualMemoryAllocator.h
@@ -72,7 +72,7 @@ struct FPooledVirtualMemoryAllocator
 		FPoolDescriptorBase* Next;
 
 		/** Total size to be deallocated, which includes both pool memory and all the descriptor/bookkeeping memory */
-		SIZE_T TotalSize;
+		SIZE_T VMSizeDivVirtualSizeAlignment;
 	};
 
 	/** Returns free memory in the pools */
diff --git a/Engine/Source/Runtime/Core/Public/HAL/MallocBinned3.h b/Engine/Source/Runtime/Core/Public/HAL/MallocBinned3.h
--- a/Engine/Source/Runtime/Core/Public/HAL/MallocBinned3.h
+++ b/Engine/Source/Runtime/Core/Public/HAL/MallocBinned3.h
@@ -1,10 +1,11 @@
-// Copyright 1998-2019 Epic Games, Inc. All Rights Reserved.
+// Copyright 1998-2018 Epic Games, Inc. All Rights Reserved.
 
 #pragma once
 
 #include "CoreTypes.h"
 
-#if PLATFORM_64BITS
+#if PLATFORM_64BITS && PLATFORM_HAS_FPlatformVirtualMemoryBlock
+#include "HAL/MallocBinnedCommon.h"
 #include "Misc/AssertionMacros.h"
 #include "HAL/MemoryBase.h"
 #include "HAL/UnrealMemory.h"
@@ -17,18 +18,19 @@
 #include "HAL/PlatformMath.h"
 #include "HAL/LowLevelMemTracker.h"
 
+#define USE_CACHED_PAGE_ALLOCATOR_FOR_LARGE_ALLOCS (0)
+
+
 #define BINNED3_BASE_PAGE_SIZE				4096			// Minimum "page size" for binned3
-#define BINNED3_MAX_PAGE_SIZE				16384			// Maximum "page size" for binned3
 #define BINNED3_MINIMUM_ALIGNMENT_SHIFT		4				// Alignment of blocks, expressed as a shift
 #define BINNED3_MINIMUM_ALIGNMENT			16				// Alignment of blocks
-#define BINNED3_MAX_LISTED_SMALL_POOL_SIZE	28672			// Maximum block size in GMallocBinned3SmallBlockSizes
-#define BINNED3_MAX_SMALL_POOL_SIZE			(128 * 1024)		// Maximum medium block size
-#define BINNED3_SMALL_POOL_COUNT			(49 + (BINNED3_MAX_SMALL_POOL_SIZE - BINNED3_MAX_LISTED_SMALL_POOL_SIZE) / BINNED3_BASE_PAGE_SIZE)
-#if PLATFORM_SWITCH
-	#define MAX_MEMORY_PER_BLOCK_SIZE_SHIFT (29) // maximum of 512MB per block size
+#if USE_CACHED_PAGE_ALLOCATOR_FOR_LARGE_ALLOCS
+#define BINNED3_MAX_SMALL_POOL_SIZE			(BINNEDCOMMON_MAX_LISTED_SMALL_POOL_SIZE)	// Maximum medium block size
 #else
-	#define MAX_MEMORY_PER_BLOCK_SIZE_SHIFT (30) // maximum of 1GB per block size
+#define BINNED3_MAX_SMALL_POOL_SIZE			(128 * 1024)	// Maximum medium block size
 #endif
+#define BINNED3_SMALL_POOL_COUNT			(BINNEDCOMMON_NUM_LISTED_SMALL_POOLS + (BINNED3_MAX_SMALL_POOL_SIZE - BINNEDCOMMON_MAX_LISTED_SMALL_POOL_SIZE) / BINNED3_BASE_PAGE_SIZE)
+#define MAX_MEMORY_PER_BLOCK_SIZE_SHIFT (29) // maximum of 512MB per block size
 #define MAX_MEMORY_PER_BLOCK_SIZE (1ull << MAX_MEMORY_PER_BLOCK_SIZE_SHIFT) 
 
 // This choice depends on how efficient the OS is with sparse commits in large VM blocks
@@ -39,7 +41,6 @@
 #endif
 
 #define DEFAULT_GMallocBinned3PerThreadCaches 1
-#define DEFAULT_GMallocBinned3LockFreeCaches 0
 #define DEFAULT_GMallocBinned3BundleCount 64
 #define DEFAULT_GMallocBinned3AllocExtra 32
 #define BINNED3_MAX_GMallocBinned3MaxBundlesBeforeRecycle 8
@@ -89,30 +90,6 @@
 	#define BINNED3_ALLOCATOR_PER_BIN_STATS 0
 #endif
 
-class FBitTree
-{
-	uint64* Bits; // one bits in middle layers mean "all allocated"
-	uint32 Capacity; // rounded up to a power of two
-	uint32 DesiredCapacity;
-	uint32 Rows;
-	uint32 OffsetOfLastRow;
-	uint32 AllocationSize;
-
-public:
-	FBitTree()
-		: Bits(nullptr)
-	{
-	}
-
-	void FBitTreeInit(uint32 InDesiredCapacity, uint32 OsAllocationGranularity, bool InitialValue);
-	uint32 AllocBit();
-	void AllocBit(uint32 Index);
-	uint32 NextAllocBit();
-	void FreeBit(uint32 Index);
-	uint32 CountOnes(uint32 UpTo);
-	static void Test();
-};
-
 //
 // Optimized virtual memory allocator.
 //
@@ -414,7 +391,7 @@ class CORE_API FMallocBinned3 final : public FMalloc
 #if BINNED3_ALLOCATOR_STATS
 	public:
 		int64 AllocatedMemory;
-		static int64 ConsolidatedMemory;
+		static TAtomic<int64> ConsolidatedMemory;
 #endif
 	private:
 		FFreeBlockList FreeLists[BINNED3_SMALL_POOL_COUNT];
@@ -700,13 +677,14 @@ public:
 	static FMallocBinned3* MallocBinned3;
 	static uint32 Binned3TlsSlot;
 	static uint32 OsAllocationGranularity;
-	static uint32 MaxAlignmentForMemoryRangeReserve;
 #if !BINNED3_USE_SEPARATE_VM_PER_POOL
 	static uint8* Binned3BaseVMPtr;
+	FPlatformMemory::FPlatformVirtualMemoryBlock Binned3BaseVMBlock;
 #else
 	static uint64 PoolSearchDiv; // if this is zero, the VM turned out to be contiguous anyway so we use a simple subtract and shift
 	static uint8* HighestPoolBaseVMPtr; // this is a duplicate of PoolBaseVMPtr[BINNED3_SMALL_POOL_COUNT - 1]
 	static uint8* PoolBaseVMPtr[BINNED3_SMALL_POOL_COUNT];
+	FPlatformMemory::FPlatformVirtualMemoryBlock PoolBaseVMBlock[BINNED3_SMALL_POOL_COUNT];
 #endif
 	// Mapping of sizes to small table indices
 	static uint8 MemSizeToIndex[1 + (BINNED3_MAX_SMALL_POOL_SIZE >> BINNED3_MINIMUM_ALIGNMENT_SHIFT)];
@@ -723,6 +701,11 @@ public:
 	{
 		return uint32(SmallBlockSizesReversedShifted[BINNED3_SMALL_POOL_COUNT - PoolIndex - 1]) << BINNED3_MINIMUM_ALIGNMENT_SHIFT;
 	}
+
+	void Commit(uint32 InPoolIndex, void *Ptr, SIZE_T Size);
+	void Decommit(uint32 InPoolIndex, void *Ptr, SIZE_T Size);
+
+	static void* AllocateMetaDataMemory(SIZE_T Size);
 };
 
 #define BINNED3_INLINE (1)
diff --git a/Engine/Source/Runtime/Core/Public/HAL/MallocBinnedArena.h b/Engine/Source/Runtime/Core/Public/HAL/MallocBinnedArena.h
new file mode 100644
--- /dev/null
+++ b/Engine/Source/Runtime/Core/Public/HAL/MallocBinnedArena.h
@@ -0,0 +1,688 @@
+// Copyright 1998-2018 Epic Games, Inc. All Rights Reserved.
+
+#pragma once
+
+#include "CoreTypes.h"
+
+#if PLATFORM_64BITS && PLATFORM_HAS_FPlatformVirtualMemoryBlock
+#include "Misc/AssertionMacros.h"
+#include "HAL/MemoryBase.h"
+#include "HAL/UnrealMemory.h"
+#include "Math/NumericLimits.h"
+#include "Templates/AlignmentTemplates.h"
+#include "HAL/CriticalSection.h"
+#include "HAL/PlatformTLS.h"
+#include "HAL/Allocators/CachedOSPageAllocator.h"
+#include "HAL/Allocators/PooledVirtualMemoryAllocator.h"
+#include "HAL/PlatformMath.h"
+#include "HAL/LowLevelMemTracker.h"
+#include "HAL/MallocBinnedCommon.h"
+
+#define BINNEDARENA_MAX_GMallocBinnedArenaMaxBundlesBeforeRecycle (8)
+
+#define COLLECT_BINNEDARENA_STATS (!UE_BUILD_SHIPPING)
+
+#if COLLECT_BINNEDARENA_STATS
+#define MBA_STAT(x) x
+#else
+#define MBA_STAT(x)
+#endif
+
+
+class CORE_API FMallocBinnedArena final : public FMalloc
+{
+	struct Private;
+
+	struct FPoolInfoSmall;
+	struct FPoolInfoLarge;
+	struct PoolHashBucket;
+	struct FPoolTable;
+	struct FGlobalRecycler;
+
+
+	struct FFreeBlock
+	{
+		enum
+		{
+			CANARY_VALUE = 0xc3
+		};
+
+		FORCEINLINE FFreeBlock(uint32 InPageSize, uint32 InBlockSize, uint32 InPoolIndex, uint8 MinimumAlignmentShift)
+			: BlockSizeShifted(InBlockSize >> MinimumAlignmentShift)
+			, PoolIndex(InPoolIndex)
+			, Canary(CANARY_VALUE)
+			, NextFreeIndex(MAX_uint32)
+		{
+			check(InPoolIndex < MAX_uint8 && (InBlockSize >> MinimumAlignmentShift) <= MAX_uint16);
+			NumFreeBlocks = InPageSize / InBlockSize;
+		}
+
+		FORCEINLINE uint32 GetNumFreeRegularBlocks() const
+		{
+			return NumFreeBlocks;
+		}
+		FORCEINLINE bool IsCanaryOk() const
+		{
+			return Canary == FFreeBlock::CANARY_VALUE;
+		}
+
+		FORCEINLINE void CanaryTest() const
+		{
+			if (!IsCanaryOk())
+			{
+				CanaryFail();
+			}
+		}
+		void CanaryFail() const;
+
+		FORCEINLINE void* AllocateRegularBlock(uint8 MinimumAlignmentShift)
+		{
+			--NumFreeBlocks;
+			return (uint8*)this + NumFreeBlocks * (uint32(BlockSizeShifted) << MinimumAlignmentShift);
+		}
+
+		uint16 BlockSizeShifted;		// Size of the blocks that this list points to >> ArenaParams.MinimumAlignmentShift
+		uint8 PoolIndex;				// Index of this pool
+		uint8 Canary;					// Constant value of 0xe3
+		uint32 NumFreeBlocks;          // Number of consecutive free blocks here, at least 1.
+		uint32 NextFreeIndex;          // Next free block or MAX_uint32
+	};
+
+	struct FPoolTable
+	{
+		uint32 BlockSize;
+		uint16 BlocksPerBlockOfBlocks;
+		uint8 PagesPlatformForBlockOfBlocks;
+
+		FBitTree BlockOfBlockAllocationBits; // one bits in here mean the virtual memory is committed
+		FBitTree BlockOfBlockIsExhausted;    // one bit in here means the pool is completely full
+
+		uint32 NumEverUsedBlockOfBlocks;
+		FPoolInfoSmall** PoolInfos;
+
+		uint64 UnusedAreaOffsetLow;
+	};
+
+	struct FPtrToPoolMapping
+	{
+		FPtrToPoolMapping()
+			: PtrToPoolPageBitShift(0)
+			, HashKeyShift(0)
+			, PoolMask(0)
+			, MaxHashBuckets(0)
+		{
+		}
+		explicit FPtrToPoolMapping(uint32 InPageSize, uint64 InNumPoolsPerPage, uint64 AddressLimit)
+		{
+			Init(InPageSize, InNumPoolsPerPage, AddressLimit);
+		}
+
+		void Init(uint32 InPageSize, uint64 InNumPoolsPerPage, uint64 AddressLimit)
+		{
+			uint64 PoolPageToPoolBitShift = FPlatformMath::CeilLogTwo(InNumPoolsPerPage);
+
+			PtrToPoolPageBitShift = FPlatformMath::CeilLogTwo(InPageSize);
+			HashKeyShift = PtrToPoolPageBitShift + PoolPageToPoolBitShift;
+			PoolMask = (1ull << PoolPageToPoolBitShift) - 1;
+			MaxHashBuckets = AddressLimit >> HashKeyShift;
+		}
+
+		FORCEINLINE void GetHashBucketAndPoolIndices(const void* InPtr, uint32& OutBucketIndex, UPTRINT& OutBucketCollision, uint32& OutPoolIndex) const
+		{
+			OutBucketCollision = (UPTRINT)InPtr >> HashKeyShift;
+			OutBucketIndex = uint32(OutBucketCollision & (MaxHashBuckets - 1));
+			OutPoolIndex = ((UPTRINT)InPtr >> PtrToPoolPageBitShift) & PoolMask;
+		}
+
+		FORCEINLINE uint64 GetMaxHashBuckets() const
+		{
+			return MaxHashBuckets;
+		}
+
+	private:
+		/** Shift to apply to a pointer to get the reference from the indirect tables */
+		uint64 PtrToPoolPageBitShift;
+		/** Shift required to get required hash table key. */
+		uint64 HashKeyShift;
+		/** Used to mask off the bits that have been used to lookup the indirect table */
+		uint64 PoolMask;
+		// PageSize dependent constants
+		uint64 MaxHashBuckets;
+	};
+
+	struct FBundleNode
+	{
+		FBundleNode* NextNodeInCurrentBundle;
+		union
+		{
+			FBundleNode* NextBundle;
+			int32 Count;
+		};
+	};
+
+	struct FBundle
+	{
+		FORCEINLINE FBundle()
+		{
+			Reset();
+		}
+
+		FORCEINLINE void Reset()
+		{
+			Head = nullptr;
+			Count = 0;
+		}
+
+		FORCEINLINE void PushHead(FBundleNode* Node)
+		{
+			Node->NextNodeInCurrentBundle = Head;
+			Node->NextBundle = nullptr;
+			Head = Node;
+			Count++;
+		}
+
+		FORCEINLINE FBundleNode* PopHead()
+		{
+			FBundleNode* Result = Head;
+
+			Count--;
+			Head = Head->NextNodeInCurrentBundle;
+			return Result;
+		}
+
+		FBundleNode* Head;
+		uint32       Count;
+	};
+
+	struct FFreeBlockList
+	{
+		// return true if we actually pushed it
+		FORCEINLINE bool PushToFront(void* InPtr, uint32 InPoolIndex, uint32 InBlockSize, const FArenaParams& LocalArenaParams)
+		{
+			checkSlow(InPtr);
+
+			if ((PartialBundle.Count >= (uint32)LocalArenaParams.MaxBlocksPerBundle) | (PartialBundle.Count * InBlockSize >= (uint32)LocalArenaParams.MaxSizePerBundle))
+			{
+				if (FullBundle.Head)
+				{
+					return false;
+				}
+				FullBundle = PartialBundle;
+				PartialBundle.Reset();
+			}
+			PartialBundle.PushHead((FBundleNode*)InPtr);
+			return true;
+		}
+		FORCEINLINE bool CanPushToFront(uint32 InPoolIndex, uint32 InBlockSize, const FArenaParams& LocalArenaParams)
+		{
+			return !((!!FullBundle.Head) & ((PartialBundle.Count >= (uint32)LocalArenaParams.MaxBlocksPerBundle) | (PartialBundle.Count * InBlockSize >= (uint32)LocalArenaParams.MaxSizePerBundle)));
+		}
+		FORCEINLINE void* PopFromFront(uint32 InPoolIndex)
+		{
+			if ((!PartialBundle.Head) & (!!FullBundle.Head))
+			{
+				PartialBundle = FullBundle;
+				FullBundle.Reset();
+			}
+			return PartialBundle.Head ? PartialBundle.PopHead() : nullptr;
+		}
+
+		// tries to recycle the full bundle, if that fails, it is returned for freeing
+		FBundleNode* RecyleFull(FArenaParams& LocalArenaParams, FGlobalRecycler& GGlobalRecycler, uint32 InPoolIndex);
+		bool ObtainPartial(FArenaParams& LocalArenaParams, FGlobalRecycler& GGlobalRecycler, uint32 InPoolIndex);
+		FBundleNode* PopBundles(uint32 InPoolIndex);
+	private:
+		FBundle PartialBundle;
+		FBundle FullBundle;
+	};
+
+	struct FPerThreadFreeBlockLists
+	{
+		FORCEINLINE static FPerThreadFreeBlockLists* Get(uint32 BinnedArenaTlsSlot)
+		{
+			return BinnedArenaTlsSlot ? (FPerThreadFreeBlockLists*)FPlatformTLS::GetTlsValue(BinnedArenaTlsSlot) : nullptr;
+		}
+		static void SetTLS(FMallocBinnedArena& Allocator);
+		static int64 ClearTLS(FMallocBinnedArena& Allocator);
+
+		FPerThreadFreeBlockLists(uint32 PoolCount)
+			: AllocatedMemory(0)
+		{ 
+			FreeLists.AddDefaulted(PoolCount);
+		}
+
+		FORCEINLINE void* Malloc(uint32 InPoolIndex)
+		{
+			return FreeLists[InPoolIndex].PopFromFront(InPoolIndex);
+		}
+		// return true if the pointer was pushed
+		FORCEINLINE bool Free(void* InPtr, uint32 InPoolIndex, uint32 InBlockSize, const FArenaParams& LocalArenaParams)
+		{
+			return FreeLists[InPoolIndex].PushToFront(InPtr, InPoolIndex, InBlockSize, LocalArenaParams);
+		}
+		// return true if a pointer can be pushed
+		FORCEINLINE bool CanFree(uint32 InPoolIndex, uint32 InBlockSize, const FArenaParams& LocalArenaParams)
+		{
+			return FreeLists[InPoolIndex].CanPushToFront(InPoolIndex, InBlockSize, LocalArenaParams);
+		}
+		// returns a bundle that needs to be freed if it can't be recycled
+		FBundleNode* RecycleFullBundle(FArenaParams& LocalArenaParams, FGlobalRecycler& GlobalRecycler, uint32 InPoolIndex)
+		{
+			return FreeLists[InPoolIndex].RecyleFull(LocalArenaParams, GlobalRecycler, InPoolIndex);
+		}
+		// returns true if we have anything to pop
+		bool ObtainRecycledPartial(FArenaParams& LocalArenaParams, FGlobalRecycler& GlobalRecycler, uint32 InPoolIndex)
+		{
+			return FreeLists[InPoolIndex].ObtainPartial(LocalArenaParams, GlobalRecycler, InPoolIndex);
+		}
+		FBundleNode* PopBundles(uint32 InPoolIndex)
+		{
+			return FreeLists[InPoolIndex].PopBundles(InPoolIndex);
+		}
+		int64 AllocatedMemory;
+		TArray<FFreeBlockList> FreeLists;
+	};
+
+	struct FGlobalRecycler
+	{
+		void Init(uint32 PoolCount)
+		{
+			Bundles.AddDefaulted(PoolCount);
+		}
+		bool PushBundle(uint32 NumCachedBundles, uint32 InPoolIndex, FBundleNode* InBundle)
+		{
+			for (uint32 Slot = 0; Slot < NumCachedBundles && Slot < BINNEDARENA_MAX_GMallocBinnedArenaMaxBundlesBeforeRecycle; Slot++)
+			{
+				if (!Bundles[InPoolIndex].FreeBundles[Slot])
+				{
+					if (!FPlatformAtomics::InterlockedCompareExchangePointer((void**)&Bundles[InPoolIndex].FreeBundles[Slot], InBundle, nullptr))
+					{
+						return true;
+					}
+				}
+			}
+			return false;
+		}
+
+		FBundleNode* PopBundle(uint32 NumCachedBundles, uint32 InPoolIndex)
+		{
+			for (uint32 Slot = 0; Slot < NumCachedBundles && Slot < BINNEDARENA_MAX_GMallocBinnedArenaMaxBundlesBeforeRecycle; Slot++)
+			{
+				FBundleNode* Result = Bundles[InPoolIndex].FreeBundles[Slot];
+				if (Result)
+				{
+					if (FPlatformAtomics::InterlockedCompareExchangePointer((void**)&Bundles[InPoolIndex].FreeBundles[Slot], nullptr, Result) == Result)
+					{
+						return Result;
+					}
+				}
+			}
+			return nullptr;
+		}
+
+	private:
+		struct FPaddedBundlePointer
+		{
+			FBundleNode* FreeBundles[BINNEDARENA_MAX_GMallocBinnedArenaMaxBundlesBeforeRecycle];
+			FPaddedBundlePointer()
+			{
+				DefaultConstructItems<FBundleNode*>(FreeBundles, BINNEDARENA_MAX_GMallocBinnedArenaMaxBundlesBeforeRecycle);
+			}
+		};
+		TArray<FPaddedBundlePointer> Bundles;
+	};
+
+
+	FORCEINLINE uint64 PoolIndexFromPtr(const void* Ptr) 
+	{
+		if (PoolSearchDiv == 0)
+		{
+			return (UPTRINT(Ptr) - UPTRINT(PoolBaseVMPtr[0])) >> ArenaParams.MaxMemoryPerBlockSizeShift;
+		}
+		uint64 PoolIndex = ArenaParams.PoolCount;
+		if (((uint8*)Ptr >= PoolBaseVMPtr[0]) & ((uint8*)Ptr < HighestPoolBaseVMPtr + ArenaParams.MaxMemoryPerBlockSize))
+		{
+			PoolIndex = uint64((uint8*)Ptr - PoolBaseVMPtr[0]) / PoolSearchDiv;
+			if (PoolIndex >= ArenaParams.PoolCount)
+			{
+				PoolIndex = ArenaParams.PoolCount - 1;
+			}
+			if ((uint8*)Ptr < PoolBaseVMPtr[PoolIndex])
+			{
+				do
+				{
+					PoolIndex--;
+					check(PoolIndex < ArenaParams.PoolCount);
+				} while ((uint8*)Ptr < PoolBaseVMPtr[PoolIndex]);
+				if ((uint8*)Ptr >= PoolBaseVMPtr[PoolIndex] + ArenaParams.MaxMemoryPerBlockSize)
+				{
+					PoolIndex = ArenaParams.PoolCount; // was in the gap
+				}
+			}
+			else if ((uint8*)Ptr >= PoolBaseVMPtr[PoolIndex] + ArenaParams.MaxMemoryPerBlockSize)
+			{
+				do
+				{
+					PoolIndex++;
+					check(PoolIndex < ArenaParams.PoolCount);
+				} while ((uint8*)Ptr >= PoolBaseVMPtr[PoolIndex] + ArenaParams.MaxMemoryPerBlockSize);
+				if ((uint8*)Ptr < PoolBaseVMPtr[PoolIndex])
+				{
+					PoolIndex = ArenaParams.PoolCount; // was in the gap
+				}
+			}
+		}
+		return PoolIndex;
+	}
+
+	FORCEINLINE uint8* PoolBasePtr(uint32 InPoolIndex)
+	{
+		return PoolBaseVMPtr[InPoolIndex];
+	}
+	FORCEINLINE uint32 PoolIndexFromPtrChecked(const void* Ptr)
+	{
+		uint64 Result = PoolIndexFromPtr(Ptr);
+		check(Result < ArenaParams.PoolCount);
+		return Result;
+	}
+
+	FORCEINLINE bool IsOSAllocation(const void* Ptr)
+	{
+		return PoolIndexFromPtr(Ptr) >= ArenaParams.PoolCount;
+	}
+
+
+	FORCEINLINE void* BlockOfBlocksPointerFromContainedPtr(const void* Ptr, uint8 PagesPlatformForBlockOfBlocks, uint32& OutBlockOfBlocksIndex)
+	{
+		uint32 PoolIndex = PoolIndexFromPtrChecked(Ptr);
+		uint8* PoolStart = PoolBasePtr(PoolIndex);
+		uint64 BlockOfBlocksIndex = (UPTRINT(Ptr) - UPTRINT(PoolStart)) / (UPTRINT(PagesPlatformForBlockOfBlocks) * UPTRINT(ArenaParams.AllocationGranularity));
+		OutBlockOfBlocksIndex = BlockOfBlocksIndex;
+
+		uint8* Result = PoolStart + BlockOfBlocksIndex * UPTRINT(PagesPlatformForBlockOfBlocks) * UPTRINT(ArenaParams.AllocationGranularity);
+
+		check(Result < PoolStart + ArenaParams.MaxMemoryPerBlockSize);
+		return Result;
+	}
+	FORCEINLINE uint8* BlockPointerFromIndecies(uint32 InPoolIndex, uint32 BlockOfBlocksIndex, uint32 BlockOfBlocksSize)
+	{
+		uint8* PoolStart = PoolBasePtr(InPoolIndex);
+		uint8* Ptr = PoolStart + BlockOfBlocksIndex * uint64(BlockOfBlocksSize);
+		check(Ptr + BlockOfBlocksSize <= PoolStart + ArenaParams.MaxMemoryPerBlockSize);
+		return Ptr;
+	}
+	FPoolInfoSmall* PushNewPoolToFront(FMallocBinnedArena& Allocator, uint32 InBlockSize, uint32 InPoolIndex, uint32& OutBlockOfBlocksIndex);
+	FPoolInfoSmall* GetFrontPool(FPoolTable& Table, uint32 InPoolIndex, uint32& OutBlockOfBlocksIndex);
+
+	FORCEINLINE bool AdjustSmallBlockSizeForAlignment(SIZE_T& InOutSize, uint32 Alignment)
+	{
+		if ((InOutSize <= ArenaParams.MaxPoolSize) & (Alignment <= ArenaParams.MinimumAlignment)) // one branch, not two
+		{
+			return true;
+		}
+		if (ArenaParams.bAttemptToAlignSmallBocks & (InOutSize <= ArenaParams.MaxPoolSize) & (Alignment <= ArenaParams.MaximumAlignmentForSmallBlock)) // one branch, not three
+		{
+			uint32 PoolIndex = BoundSizeToPoolIndex(Align(InOutSize, Alignment));
+			while (true)
+			{
+				uint32 BlockSize = PoolIndexToBlockSize(PoolIndex);
+				if (IsAligned(BlockSize, Alignment))
+				{
+					InOutSize = SIZE_T(BlockSize);
+					return true;
+				}
+				PoolIndex++;
+				check(PoolIndex < ArenaParams.PoolCount);
+			}
+		}
+		return false;
+	}
+
+public:
+
+
+	FMallocBinnedArena();
+	FArenaParams& GetParams()
+	{
+		return ArenaParams;
+	}
+	void InitMallocBinned();
+
+	virtual ~FMallocBinnedArena();
+
+
+	// FMalloc interface.
+	virtual bool IsInternallyThreadSafe() const override;
+	FORCEINLINE virtual void* Malloc(SIZE_T Size, uint32 Alignment) override
+	{
+		Alignment = FMath::Max<uint32>(Alignment, ArenaParams.MinimumAlignment);
+
+		void* Result = nullptr;
+
+		// Only allocate from the small pools if the size is small enough and the alignment isn't crazy large.
+		// With large alignments, we'll waste a lot of memory allocating an entire page, but such alignments are highly unlikely in practice.
+		if (AdjustSmallBlockSizeForAlignment(Size, Alignment))
+		{
+			FPerThreadFreeBlockLists* Lists = ArenaParams.bPerThreadCaches ? FPerThreadFreeBlockLists::Get(BinnedArenaTlsSlot) : nullptr;
+			if (Lists)
+			{
+				uint32 PoolIndex = BoundSizeToPoolIndex(Size);
+				uint32 BlockSize = PoolIndexToBlockSize(PoolIndex);
+				Result = Lists->Malloc(PoolIndex);
+				if (Result)
+				{
+					Lists->AllocatedMemory += BlockSize;
+					checkSlow(IsAligned(Result, Alignment));
+				}
+			}
+		}
+		if (Result == nullptr)
+		{
+			Result = MallocExternal(Size, Alignment);
+		}
+
+		return Result;
+	}
+	FORCEINLINE virtual void* Realloc(void* Ptr, SIZE_T NewSize, uint32 Alignment) override
+	{
+		Alignment = FMath::Max<uint32>(Alignment, ArenaParams.MinimumAlignment);
+		if (AdjustSmallBlockSizeForAlignment(NewSize, Alignment))
+		{
+			FPerThreadFreeBlockLists* Lists = ArenaParams.bPerThreadCaches ? FPerThreadFreeBlockLists::Get(BinnedArenaTlsSlot) : nullptr;
+
+			uint64 PoolIndex = PoolIndexFromPtr(Ptr);
+			if ((!!Lists) & ((!Ptr) | (PoolIndex < ArenaParams.PoolCount)))
+			{
+				uint32 BlockSize = 0;
+
+				bool bCanFree = true; // the nullptr is always "freeable"
+				if (Ptr)
+				{
+					// Reallocate to a smaller/bigger pool if necessary
+					BlockSize = PoolIndexToBlockSize(PoolIndex);
+					if ((!!NewSize) & (NewSize <= BlockSize) & IsAligned(BlockSize, Alignment) & ((!PoolIndex) | (NewSize > PoolIndexToBlockSize(PoolIndex - 1))))
+					{
+						return Ptr;
+					}
+					bCanFree = Lists->CanFree(PoolIndex, BlockSize, ArenaParams);
+				}
+				if (bCanFree)
+				{
+					uint32 NewPoolIndex = BoundSizeToPoolIndex(NewSize);
+					uint32 NewBlockSize = PoolIndexToBlockSize(NewPoolIndex);
+					void* Result = NewSize ? Lists->Malloc(NewPoolIndex) : nullptr;
+					if (Result)
+					{
+						Lists->AllocatedMemory += NewBlockSize;
+					}
+					if (Result || !NewSize)
+					{
+						if (Result && Ptr)
+						{
+							FMemory::Memcpy(Result, Ptr, FPlatformMath::Min<SIZE_T>(NewSize, BlockSize));
+						}
+						if (Ptr)
+						{
+							bool bDidPush = Lists->Free(Ptr, PoolIndex, BlockSize, ArenaParams);
+							checkSlow(bDidPush);
+							Lists->AllocatedMemory -= BlockSize;
+						}
+
+						return Result;
+					}
+				}
+			}
+		}
+		void* Result = ReallocExternal(Ptr, NewSize, Alignment);
+		return Result;
+	}
+
+	FORCEINLINE virtual void Free(void* Ptr) override
+	{
+		uint64 PoolIndex = PoolIndexFromPtr(Ptr);
+		if (PoolIndex < ArenaParams.PoolCount)
+		{
+			FPerThreadFreeBlockLists* Lists = ArenaParams.bPerThreadCaches ? FPerThreadFreeBlockLists::Get(BinnedArenaTlsSlot) : nullptr;
+			if (Lists)
+			{
+				int32 BlockSize = PoolIndexToBlockSize(PoolIndex);
+				if (Lists->Free(Ptr, PoolIndex, BlockSize, ArenaParams))
+				{
+					Lists->AllocatedMemory -= BlockSize;
+					return;
+				}
+			}
+		}
+		FreeExternal(Ptr);
+	}
+	FORCEINLINE virtual bool GetAllocationSize(void *Ptr, SIZE_T &SizeOut) override
+	{
+		uint64 PoolIndex = PoolIndexFromPtr(Ptr);
+		if (PoolIndex < ArenaParams.PoolCount)
+		{
+			SizeOut = PoolIndexToBlockSize(PoolIndex);
+			return true;
+		}
+		return GetAllocationSizeExternal(Ptr, SizeOut);
+	}
+
+	FORCEINLINE virtual SIZE_T QuantizeSize(SIZE_T Count, uint32 Alignment) override
+	{
+		check(DEFAULT_ALIGNMENT <= ArenaParams.MinimumAlignment); // used below
+		checkSlow((Alignment & (Alignment - 1)) == 0); // Check the alignment is a power of two
+		SIZE_T SizeOut;
+		if ((Count <= ArenaParams.MaxPoolSize) & (Alignment <= ArenaParams.MinimumAlignment)) // one branch, not two
+		{
+			SizeOut = PoolIndexToBlockSize(BoundSizeToPoolIndex(Count));
+		}
+		else
+		{
+			Alignment = FPlatformMath::Max<uint32>(Alignment, ArenaParams.AllocationGranularity);
+			SizeOut = Align(Count, Alignment);
+		}
+		check(SizeOut >= Count);
+		return SizeOut;
+	}
+
+	virtual bool ValidateHeap() override;
+	virtual void Trim(bool bTrimThreadCaches) override;
+	virtual void SetupTLSCachesOnCurrentThread() override;
+	virtual void ClearAndDisableTLSCachesOnCurrentThread() override;
+	virtual const TCHAR* GetDescriptiveName() override;
+	// End FMalloc interface.
+
+	void FlushCurrentThreadCache();
+	void* MallocExternal(SIZE_T Size, uint32 Alignment);
+	void* ReallocExternal(void* Ptr, SIZE_T NewSize, uint32 Alignment);
+	void FreeExternal(void *Ptr);
+	bool GetAllocationSizeExternal(void* Ptr, SIZE_T& SizeOut);
+
+	MBA_STAT(int64 GetTotalAllocatedSmallPoolMemory();)
+	virtual void GetAllocatorStats(FGenericMemoryStats& out_Stats) override;
+	/** Dumps current allocator stats to the log. */
+	virtual void DumpAllocatorStats(class FOutputDevice& Ar) override;
+
+	FORCEINLINE uint32 BoundSizeToPoolIndex(SIZE_T Size)
+	{
+		auto Index = ((Size + ArenaParams.MinimumAlignment - 1) >> ArenaParams.MinimumAlignmentShift);
+		checkSlow(Index >= 0 && Index <= (ArenaParams.MaxPoolSize >> ArenaParams.MinimumAlignmentShift)); // and it should be in the table
+		uint32 PoolIndex = uint32(MemSizeToIndex[Index]);
+		checkSlow(PoolIndex >= 0 && PoolIndex < ArenaParams.PoolCount);
+		return PoolIndex;
+	}
+	FORCEINLINE uint32 PoolIndexToBlockSize(uint32 PoolIndex)
+	{
+		return uint32(SmallBlockSizesReversedShifted[ArenaParams.PoolCount - PoolIndex - 1]) << ArenaParams.MinimumAlignmentShift;
+	}
+
+	void Commit(uint32 InPoolIndex, void *Ptr, SIZE_T Size);
+	void Decommit(uint32 InPoolIndex, void *Ptr, SIZE_T Size);
+
+
+	// Pool tables for different pool sizes
+	TArray<FPoolTable> SmallPoolTables;
+
+	uint32 SmallPoolInfosPerPlatformPage;
+
+	PoolHashBucket* HashBuckets;
+	PoolHashBucket* HashBucketFreeList;
+	uint64 NumLargePoolsPerPage;
+
+	FCriticalSection Mutex;
+	FGlobalRecycler GGlobalRecycler;
+	FPtrToPoolMapping PtrToPoolMapping;
+
+	FArenaParams ArenaParams;
+
+	TArray<uint16> SmallBlockSizesReversedShifted; // this is reversed to get the smallest elements on our main cache line
+	uint32 BinnedArenaTlsSlot;
+	uint64 PoolSearchDiv; // if this is zero, the VM turned out to be contiguous anyway so we use a simple subtract and shift
+	uint8* HighestPoolBaseVMPtr; // this is a duplicate of PoolBaseVMPtr[ArenaParams.PoolCount - 1]
+	FPlatformMemory::FPlatformVirtualMemoryBlock PoolBaseVMBlock;
+	TArray<uint8*> PoolBaseVMPtr;
+	TArray<FPlatformMemory::FPlatformVirtualMemoryBlock> PoolBaseVMBlocks;
+	// Mapping of sizes to small table indices
+	TArray<uint8> MemSizeToIndex;
+
+	MBA_STAT(
+		int64 BinnedArenaAllocatedSmallPoolMemory = 0; // memory that's requested to be allocated by the game
+		int64 BinnedArenaAllocatedOSSmallPoolMemory = 0;
+
+		int64 BinnedArenaAllocatedLargePoolMemory = 0; // memory requests to the OS which don't fit in the small pool
+		int64 BinnedArenaAllocatedLargePoolMemoryWAlignment = 0; // when we allocate at OS level we need to align to a size
+
+		int64 BinnedArenaPoolInfoMemory = 0;
+		int64 BinnedArenaHashMemory = 0;
+		int64 BinnedArenaFreeBitsMemory = 0;
+		int64 BinnedArenaTLSMemory = 0;
+		TAtomic<int64> ConsolidatedMemory;
+	)
+
+	FCriticalSection FreeBlockListsRegistrationMutex;
+	FCriticalSection& GetFreeBlockListsRegistrationMutex()
+	{
+		return FreeBlockListsRegistrationMutex;
+	}
+	TArray<FPerThreadFreeBlockLists*> RegisteredFreeBlockLists;
+	TArray<FPerThreadFreeBlockLists*>& GetRegisteredFreeBlockLists()
+	{
+		return RegisteredFreeBlockLists;
+	}
+	void RegisterThreadFreeBlockLists(FPerThreadFreeBlockLists* FreeBlockLists)
+	{
+		FScopeLock Lock(&GetFreeBlockListsRegistrationMutex());
+		GetRegisteredFreeBlockLists().Add(FreeBlockLists);
+	}
+	int64 UnregisterThreadFreeBlockLists(FPerThreadFreeBlockLists* FreeBlockLists)
+	{
+		FScopeLock Lock(&GetFreeBlockListsRegistrationMutex());
+		int64 Result = FreeBlockLists->AllocatedMemory;
+		check(Result >= 0);
+		GetRegisteredFreeBlockLists().Remove(FreeBlockLists);
+		return Result;
+	}
+
+	TArray<void*> MallocedPointers;
+};
+
+#endif
\ No newline at end of file
diff --git a/Engine/Source/Runtime/Core/Public/HAL/MallocBinnedCommon.h b/Engine/Source/Runtime/Core/Public/HAL/MallocBinnedCommon.h
new file mode 100644
--- /dev/null
+++ b/Engine/Source/Runtime/Core/Public/HAL/MallocBinnedCommon.h
@@ -0,0 +1,96 @@
+// Copyright 1998-2018 Epic Games, Inc. All Rights Reserved.
+
+#pragma once
+
+#include "CoreTypes.h"
+#include "Containers/Array.h"
+
+#if PLATFORM_HAS_FPlatformVirtualMemoryBlock
+
+#define BINNEDCOMMON_MAX_LISTED_SMALL_POOL_SIZE	28672
+#define BINNEDCOMMON_NUM_LISTED_SMALL_POOLS	49
+
+class FBitTree
+{
+	uint64* Bits; // one bits in middle layers mean "all allocated"
+	uint32 Capacity; // rounded up to a power of two
+	uint32 DesiredCapacity;
+	uint32 Rows;
+	uint32 OffsetOfLastRow;
+	uint32 AllocationSize;
+
+public:
+	FBitTree()
+		: Bits(nullptr)
+	{
+	}
+
+	static uint32 GetMemoryRequirements(uint32 DesiredCapacity);
+	void FBitTreeInit(uint32 InDesiredCapacity, void * Memory, uint32 MemorySize, bool InitialValue);
+	uint32 AllocBit();
+	bool IsAllocated(uint32 Index) const;
+	void AllocBit(uint32 Index);
+	uint32 NextAllocBit() const;
+	uint32 NextAllocBit(uint32 StartIndex) const;
+	void FreeBit(uint32 Index);
+	uint32 CountOnes(uint32 UpTo) const;
+};
+
+struct FSizeTableEntry
+{
+	uint32 BlockSize;
+	uint16 BlocksPerBlockOfBlocks;
+	uint8 PagesPlatformForBlockOfBlocks;
+
+	FSizeTableEntry()
+	{
+	}
+
+	FSizeTableEntry(uint32 InBlockSize, uint64 PlatformPageSize, uint8 Pages4k, uint32 BasePageSize, uint32 MinimumAlignment);
+
+	bool operator<(const FSizeTableEntry& Other) const
+	{
+		return BlockSize < Other.BlockSize;
+	}
+	static uint8 FillSizeTable(uint64 PlatformPageSize, FSizeTableEntry* SizeTable, uint32 BasePageSize, uint32 MinimumAlignment, uint32 MaxSize, uint32 SizeIncrement);
+};
+
+struct FArenaParams
+{
+	// these are parameters you set
+	uint64 AddressLimit = 1024 * 1024 * 1024; // this controls the size of the root hash table
+	uint32 BasePageSize = 4096; // this is used to make sensible calls to malloc and figures into the standard pool sizes if bUseStandardSmallPoolSizes is true
+	uint32 AllocationGranularity = 4096; // this is the granularity of the commit and decommit calls used on the VM slabs
+	uint32 MaxSizePerBundle = 8192;
+	uint32 MaxStandardPoolSize = 128 * 1024; // these are added to the standard pool sizes, mainly to use the TLS caches, they are typically one block per slab
+	uint16 MaxBlocksPerBundle = 64;
+	uint8 MaxMemoryPerBlockSizeShift = 29;
+	uint8 EmptyCacheAllocExtra = 32;
+	uint8 MaxGlobalBundles = 32;
+	uint8 MinimumAlignmentShift = 4;
+	uint8 PoolCount;
+	bool bUseSeparateVMPerPool = !!(PLATFORM_XBOXONE || PLATFORM_WINDOWS);
+	bool bPerThreadCaches = true;
+	bool bUseStandardSmallPoolSizes = true;
+	bool bAttemptToAlignSmallBocks = true;
+	TArray<uint32> AdditionalBlockSizes;
+
+	// This lambdas is similar to the platform virtual memory HAL and by default just call that. 
+	TFunction<FPlatformMemory::FPlatformVirtualMemoryBlock(SIZE_T)> ReserveVM;
+
+	// These allow you to override the large block allocator. The value add here is that MBA tracks the metadata for you and call tell the difference between a large block pointer and a small block pointer.
+	// By defaults these just use the platform VM interface to allocate some committed memory
+	TFunction<void*(SIZE_T, SIZE_T, SIZE_T&, uint32&)> LargeBlockAlloc;
+	TFunction<void(void*, uint32)> LargeBlockFree;
+
+
+	// these are parameters are derived from other parameters
+	uint64 MaxMemoryPerBlockSize;
+	uint32 MaxPoolSize;
+	uint32 MinimumAlignment;
+	uint32 MaximumAlignmentForSmallBlock;
+
+};
+
+
+#endif
\ No newline at end of file
diff --git a/Engine/Source/Runtime/Core/Public/HAL/MallocBinnedGPU.h b/Engine/Source/Runtime/Core/Public/HAL/MallocBinnedGPU.h
new file mode 100644
--- /dev/null
+++ b/Engine/Source/Runtime/Core/Public/HAL/MallocBinnedGPU.h
@@ -0,0 +1,660 @@
+// Copyright 1998-2018 Epic Games, Inc. All Rights Reserved.
+
+#pragma once
+
+#include "CoreTypes.h"
+
+#if PLATFORM_64BITS && PLATFORM_HAS_FPlatformVirtualMemoryBlock
+#include "Misc/AssertionMacros.h"
+#include "HAL/MemoryBase.h"
+#include "HAL/UnrealMemory.h"
+#include "Math/NumericLimits.h"
+#include "Templates/AlignmentTemplates.h"
+#include "HAL/CriticalSection.h"
+#include "HAL/PlatformTLS.h"
+#include "HAL/Allocators/CachedOSPageAllocator.h"
+#include "HAL/Allocators/PooledVirtualMemoryAllocator.h"
+#include "HAL/PlatformMath.h"
+#include "HAL/LowLevelMemTracker.h"
+#include "HAL/MallocBinnedCommon.h"
+
+#define BINNEDGPU_MAX_GMallocBinnedGPUMaxBundlesBeforeRecycle (8)
+
+#define COLLECT_BINNEDGPU_STATS (!UE_BUILD_SHIPPING)
+
+#if COLLECT_BINNEDGPU_STATS
+	#define MBG_STAT(x) x
+#else
+	#define MBG_STAT(x)
+#endif
+
+
+class CORE_API FMallocBinnedGPU final : public FMalloc
+{
+	struct Private;
+
+	struct FPoolInfoSmall;
+	struct FPoolInfoLarge;
+	struct PoolHashBucket;
+	struct FPoolTable;
+	struct FGlobalRecycler;
+
+
+	struct FGPUMemoryBlockProxy
+	{
+		uint8 MemoryModifiedByCPU[32 - sizeof(void*)]; // might be modified for free list links, etc
+		void *GPUMemory;  // pointer to the actual GPU memory, which we cannot modify with the CPU
+
+		FGPUMemoryBlockProxy(void *InGPUMemory)
+			: GPUMemory(InGPUMemory)
+		{
+			check(GPUMemory);
+		}
+	};
+
+	struct FFreeBlock
+	{
+		enum
+		{
+			CANARY_VALUE = 0xc3
+		};
+
+		FORCEINLINE FFreeBlock(uint32 InPageSize, uint32 InBlockSize, uint32 InPoolIndex, uint8 MinimumAlignmentShift)
+			: BlockSizeShifted(InBlockSize >> MinimumAlignmentShift)
+			, PoolIndex(InPoolIndex)
+			, Canary(CANARY_VALUE)
+			, NextFreeBlock(nullptr)
+		{
+			check(InPoolIndex < MAX_uint8 && (InBlockSize >> MinimumAlignmentShift) <= MAX_uint16);
+			NumFreeBlocks = InPageSize / InBlockSize;
+		}
+
+		FORCEINLINE uint32 GetNumFreeRegularBlocks() const
+		{
+			return NumFreeBlocks;
+		}
+		FORCEINLINE bool IsCanaryOk() const
+		{
+			return Canary == FFreeBlock::CANARY_VALUE;
+		}
+
+		FORCEINLINE void CanaryTest() const
+		{
+			if (!IsCanaryOk())
+			{
+				CanaryFail();
+			}
+		}
+		void CanaryFail() const;
+
+		FORCEINLINE void* AllocateRegularBlock(uint8 MinimumAlignmentShift)
+		{
+			--NumFreeBlocks;
+			return (uint8*)(((FGPUMemoryBlockProxy*)this)->GPUMemory) + NumFreeBlocks * (uint32(BlockSizeShifted) << MinimumAlignmentShift);
+		}
+
+		uint16 BlockSizeShifted;		// Size of the blocks that this list points to >> ArenaParams.MinimumAlignmentShift
+		uint8 PoolIndex;				// Index of this pool
+		uint8 Canary;					// Constant value of 0xe3
+		uint32 NumFreeBlocks;          // Number of consecutive free blocks here, at least 1.
+		FFreeBlock* NextFreeBlock;     // Next free block or nullptr
+	};
+
+	struct FPoolTable
+	{
+		uint32 BlockSize;
+		uint16 BlocksPerBlockOfBlocks;
+		uint8 PagesPlatformForBlockOfBlocks;
+
+		FBitTree BlockOfBlockAllocationBits; // one bits in here mean the virtual memory is committed
+		FBitTree BlockOfBlockIsExhausted;    // one bit in here means the pool is completely full
+
+		uint32 NumEverUsedBlockOfBlocks;
+		FPoolInfoSmall** PoolInfos;
+
+		uint64 UnusedAreaOffsetLow;
+	};
+
+	struct FPtrToPoolMapping
+	{
+		FPtrToPoolMapping()
+			: PtrToPoolPageBitShift(0)
+			, HashKeyShift(0)
+			, PoolMask(0)
+			, MaxHashBuckets(0)
+		{
+		}
+		explicit FPtrToPoolMapping(uint32 InPageSize, uint64 InNumPoolsPerPage, uint64 AddressLimit)
+		{
+			Init(InPageSize, InNumPoolsPerPage, AddressLimit);
+		}
+
+		void Init(uint32 InPageSize, uint64 InNumPoolsPerPage, uint64 AddressLimit)
+		{
+			uint64 PoolPageToPoolBitShift = FPlatformMath::CeilLogTwo(InNumPoolsPerPage);
+
+			PtrToPoolPageBitShift = FPlatformMath::CeilLogTwo(InPageSize);
+			HashKeyShift = PtrToPoolPageBitShift + PoolPageToPoolBitShift;
+			PoolMask = (1ull << PoolPageToPoolBitShift) - 1;
+			MaxHashBuckets = AddressLimit >> HashKeyShift;
+		}
+
+		FORCEINLINE void GetHashBucketAndPoolIndices(const void* InPtr, uint32& OutBucketIndex, UPTRINT& OutBucketCollision, uint32& OutPoolIndex) const
+		{
+			OutBucketCollision = (UPTRINT)InPtr >> HashKeyShift;
+			OutBucketIndex = uint32(OutBucketCollision & (MaxHashBuckets - 1));
+			OutPoolIndex = ((UPTRINT)InPtr >> PtrToPoolPageBitShift) & PoolMask;
+		}
+
+		FORCEINLINE uint64 GetMaxHashBuckets() const
+		{
+			return MaxHashBuckets;
+		}
+
+	private:
+		/** Shift to apply to a pointer to get the reference from the indirect tables */
+		uint64 PtrToPoolPageBitShift;
+		/** Shift required to get required hash table key. */
+		uint64 HashKeyShift;
+		/** Used to mask off the bits that have been used to lookup the indirect table */
+		uint64 PoolMask;
+		// PageSize dependent constants
+		uint64 MaxHashBuckets;
+	};
+
+	struct FBundleNode
+	{
+		FBundleNode* NextNodeInCurrentBundle;
+		union
+		{
+			FBundleNode* NextBundle;
+			int32 Count;
+		};
+	};
+
+	struct FBundle
+	{
+		FORCEINLINE FBundle()
+		{
+			Reset();
+		}
+
+		FORCEINLINE void Reset()
+		{
+			Head = nullptr;
+			Count = 0;
+		}
+
+		FORCEINLINE void PushHead(FBundleNode* Node)
+		{
+			Node->NextNodeInCurrentBundle = Head;
+			Node->NextBundle = nullptr;
+			Head = Node;
+			Count++;
+		}
+
+		FORCEINLINE FBundleNode* PopHead()
+		{
+			FBundleNode* Result = Head;
+
+			Count--;
+			Head = Head->NextNodeInCurrentBundle;
+			return Result;
+		}
+
+		FBundleNode* Head;
+		uint32       Count;
+	};
+
+	struct FFreeBlockList
+	{
+		// return true if we actually pushed it
+		FORCEINLINE bool PushToFront(FMallocBinnedGPU& Allocator, void* InPtr, uint32 InPoolIndex, uint32 InBlockSize, const FArenaParams& LocalArenaParams)
+		{
+			check(InPtr);
+
+			if ((PartialBundle.Count >= (uint32)LocalArenaParams.MaxBlocksPerBundle) | (PartialBundle.Count * InBlockSize >= (uint32)LocalArenaParams.MaxSizePerBundle))
+			{
+				if (FullBundle.Head)
+				{
+					return false;
+				}
+				FullBundle = PartialBundle;
+				PartialBundle.Reset();
+			}
+			PartialBundle.PushHead((FBundleNode*)new FGPUMemoryBlockProxy(InPtr));
+			MBG_STAT(Allocator.GPUProxyMemory += sizeof(FGPUMemoryBlockProxy);)
+			return true;
+		}
+		FORCEINLINE bool CanPushToFront(uint32 InPoolIndex, uint32 InBlockSize, const FArenaParams& LocalArenaParams)
+		{
+			return !((!!FullBundle.Head) & ((PartialBundle.Count >= (uint32)LocalArenaParams.MaxBlocksPerBundle) | (PartialBundle.Count * InBlockSize >= (uint32)LocalArenaParams.MaxSizePerBundle)));
+		}
+		FORCEINLINE void* PopFromFront(FMallocBinnedGPU& Allocator, uint32 InPoolIndex)
+		{
+			if ((!PartialBundle.Head) & (!!FullBundle.Head))
+			{
+				PartialBundle = FullBundle;
+				FullBundle.Reset();
+			}
+			void *Result = nullptr;
+			if (PartialBundle.Head)
+			{
+				FGPUMemoryBlockProxy* Proxy = (FGPUMemoryBlockProxy*)PartialBundle.PopHead();
+				Result = Proxy->GPUMemory;
+				check(Result);
+				delete Proxy;
+				MBG_STAT(Allocator.GPUProxyMemory -= sizeof(FGPUMemoryBlockProxy);)
+			}
+			return Result;
+		}
+
+		// tries to recycle the full bundle, if that fails, it is returned for freeing
+		FBundleNode* RecyleFull(FArenaParams& LocalArenaParams, FGlobalRecycler& GGlobalRecycler, uint32 InPoolIndex);
+		bool ObtainPartial(FArenaParams& LocalArenaParams, FGlobalRecycler& GGlobalRecycler, uint32 InPoolIndex);
+		FBundleNode* PopBundles(uint32 InPoolIndex);
+	private:
+		FBundle PartialBundle;
+		FBundle FullBundle;
+	};
+
+	struct FPerThreadFreeBlockLists
+	{
+		FORCEINLINE static FPerThreadFreeBlockLists* Get(uint32 BinnedGPUTlsSlot)
+		{
+			return BinnedGPUTlsSlot ? (FPerThreadFreeBlockLists*)FPlatformTLS::GetTlsValue(BinnedGPUTlsSlot) : nullptr;
+		}
+		static void SetTLS(FMallocBinnedGPU& Allocator);
+		static int64 ClearTLS(FMallocBinnedGPU& Allocator);
+
+		FPerThreadFreeBlockLists(uint32 PoolCount)
+			: AllocatedMemory(0)
+		{ 
+			FreeLists.AddDefaulted(PoolCount);
+		}
+
+		FORCEINLINE void* Malloc(FMallocBinnedGPU& Allocator, uint32 InPoolIndex)
+		{
+			return FreeLists[InPoolIndex].PopFromFront(Allocator, InPoolIndex);
+		}
+		// return true if the pointer was pushed
+		FORCEINLINE bool Free(FMallocBinnedGPU& Allocator, void* InPtr, uint32 InPoolIndex, uint32 InBlockSize, const FArenaParams& LocalArenaParams)
+		{
+			return FreeLists[InPoolIndex].PushToFront(Allocator, InPtr, InPoolIndex, InBlockSize, LocalArenaParams);
+		}
+		// return true if a pointer can be pushed
+		FORCEINLINE bool CanFree(uint32 InPoolIndex, uint32 InBlockSize, const FArenaParams& LocalArenaParams)
+		{
+			return FreeLists[InPoolIndex].CanPushToFront(InPoolIndex, InBlockSize, LocalArenaParams);
+		}
+		// returns a bundle that needs to be freed if it can't be recycled
+		FBundleNode* RecycleFullBundle(FArenaParams& LocalArenaParams, FGlobalRecycler& GlobalRecycler, uint32 InPoolIndex)
+		{
+			return FreeLists[InPoolIndex].RecyleFull(LocalArenaParams, GlobalRecycler, InPoolIndex);
+		}
+		// returns true if we have anything to pop
+		bool ObtainRecycledPartial(FArenaParams& LocalArenaParams, FGlobalRecycler& GlobalRecycler, uint32 InPoolIndex)
+		{
+			return FreeLists[InPoolIndex].ObtainPartial(LocalArenaParams, GlobalRecycler, InPoolIndex);
+		}
+		FBundleNode* PopBundles(uint32 InPoolIndex)
+		{
+			return FreeLists[InPoolIndex].PopBundles(InPoolIndex);
+		}
+		int64 AllocatedMemory;
+		TArray<FFreeBlockList> FreeLists;
+	};
+
+	struct FGlobalRecycler
+	{
+		void Init(uint32 PoolCount)
+		{
+			Bundles.AddDefaulted(PoolCount);
+		}
+		bool PushBundle(uint32 NumCachedBundles, uint32 InPoolIndex, FBundleNode* InBundle)
+		{
+			for (uint32 Slot = 0; Slot < NumCachedBundles && Slot < BINNEDGPU_MAX_GMallocBinnedGPUMaxBundlesBeforeRecycle; Slot++)
+			{
+				if (!Bundles[InPoolIndex].FreeBundles[Slot])
+				{
+					if (!FPlatformAtomics::InterlockedCompareExchangePointer((void**)&Bundles[InPoolIndex].FreeBundles[Slot], InBundle, nullptr))
+					{
+						return true;
+					}
+				}
+			}
+			return false;
+		}
+
+		FBundleNode* PopBundle(uint32 NumCachedBundles, uint32 InPoolIndex)
+		{
+			for (uint32 Slot = 0; Slot < NumCachedBundles && Slot < BINNEDGPU_MAX_GMallocBinnedGPUMaxBundlesBeforeRecycle; Slot++)
+			{
+				FBundleNode* Result = Bundles[InPoolIndex].FreeBundles[Slot];
+				if (Result)
+				{
+					if (FPlatformAtomics::InterlockedCompareExchangePointer((void**)&Bundles[InPoolIndex].FreeBundles[Slot], nullptr, Result) == Result)
+					{
+						return Result;
+					}
+				}
+			}
+			return nullptr;
+		}
+
+	private:
+		struct FPaddedBundlePointer
+		{
+			FBundleNode* FreeBundles[BINNEDGPU_MAX_GMallocBinnedGPUMaxBundlesBeforeRecycle];
+			FPaddedBundlePointer()
+			{
+				DefaultConstructItems<FBundleNode*>(FreeBundles, BINNEDGPU_MAX_GMallocBinnedGPUMaxBundlesBeforeRecycle);
+			}
+		};
+		TArray<FPaddedBundlePointer> Bundles;
+	};
+
+
+	FORCEINLINE uint64 PoolIndexFromPtr(const void* Ptr) 
+	{
+		if (PoolSearchDiv == 0)
+		{
+			return (UPTRINT(Ptr) - UPTRINT(PoolBaseVMPtr[0])) >> ArenaParams.MaxMemoryPerBlockSizeShift;
+		}
+		uint64 PoolIndex = ArenaParams.PoolCount;
+		if (((uint8*)Ptr >= PoolBaseVMPtr[0]) & ((uint8*)Ptr < HighestPoolBaseVMPtr + ArenaParams.MaxMemoryPerBlockSize))
+		{
+			PoolIndex = uint64((uint8*)Ptr - PoolBaseVMPtr[0]) / PoolSearchDiv;
+			if (PoolIndex >= ArenaParams.PoolCount)
+			{
+				PoolIndex = ArenaParams.PoolCount - 1;
+			}
+			if ((uint8*)Ptr < PoolBaseVMPtr[PoolIndex])
+			{
+				do
+				{
+					PoolIndex--;
+					check(PoolIndex < ArenaParams.PoolCount);
+				} while ((uint8*)Ptr < PoolBaseVMPtr[PoolIndex]);
+				if ((uint8*)Ptr >= PoolBaseVMPtr[PoolIndex] + ArenaParams.MaxMemoryPerBlockSize)
+				{
+					PoolIndex = ArenaParams.PoolCount; // was in the gap
+				}
+			}
+			else if ((uint8*)Ptr >= PoolBaseVMPtr[PoolIndex] + ArenaParams.MaxMemoryPerBlockSize)
+			{
+				do
+				{
+					PoolIndex++;
+					check(PoolIndex < ArenaParams.PoolCount);
+				} while ((uint8*)Ptr >= PoolBaseVMPtr[PoolIndex] + ArenaParams.MaxMemoryPerBlockSize);
+				if ((uint8*)Ptr < PoolBaseVMPtr[PoolIndex])
+				{
+					PoolIndex = ArenaParams.PoolCount; // was in the gap
+				}
+			}
+		}
+		return PoolIndex;
+	}
+
+	FORCEINLINE uint8* PoolBasePtr(uint32 InPoolIndex)
+	{
+		return PoolBaseVMPtr[InPoolIndex];
+	}
+	FORCEINLINE uint32 PoolIndexFromPtrChecked(const void* Ptr)
+	{
+		uint64 Result = PoolIndexFromPtr(Ptr);
+		check(Result < ArenaParams.PoolCount);
+		return Result;
+	}
+
+	FORCEINLINE bool IsOSAllocation(const void* Ptr)
+	{
+		return PoolIndexFromPtr(Ptr) >= ArenaParams.PoolCount;
+	}
+
+
+	FORCEINLINE void* BlockOfBlocksPointerFromContainedPtr(const void* Ptr, uint8 PagesPlatformForBlockOfBlocks, uint32& OutBlockOfBlocksIndex)
+	{
+		uint32 PoolIndex = PoolIndexFromPtrChecked(Ptr);
+		uint8* PoolStart = PoolBasePtr(PoolIndex);
+		uint64 BlockOfBlocksIndex = (UPTRINT(Ptr) - UPTRINT(PoolStart)) / (UPTRINT(PagesPlatformForBlockOfBlocks) * UPTRINT(ArenaParams.AllocationGranularity));
+		OutBlockOfBlocksIndex = BlockOfBlocksIndex;
+
+		uint8* Result = PoolStart + BlockOfBlocksIndex * UPTRINT(PagesPlatformForBlockOfBlocks) * UPTRINT(ArenaParams.AllocationGranularity);
+
+		check(Result < PoolStart + ArenaParams.MaxMemoryPerBlockSize);
+		return Result;
+	}
+	FORCEINLINE uint8* BlockPointerFromIndecies(uint32 InPoolIndex, uint32 BlockOfBlocksIndex, uint32 BlockOfBlocksSize)
+	{
+		uint8* PoolStart = PoolBasePtr(InPoolIndex);
+		uint8* Ptr = PoolStart + BlockOfBlocksIndex * uint64(BlockOfBlocksSize);
+		check(Ptr + BlockOfBlocksSize <= PoolStart + ArenaParams.MaxMemoryPerBlockSize);
+		return Ptr;
+	}
+	FPoolInfoSmall* PushNewPoolToFront(FMallocBinnedGPU& Allocator, uint32 InBlockSize, uint32 InPoolIndex, uint32& OutBlockOfBlocksIndex);
+	FPoolInfoSmall* GetFrontPool(FPoolTable& Table, uint32 InPoolIndex, uint32& OutBlockOfBlocksIndex);
+
+	FORCEINLINE bool AdjustSmallBlockSizeForAlignment(SIZE_T& InOutSize, uint32 Alignment)
+	{
+		if ((InOutSize <= ArenaParams.MaxPoolSize) & (Alignment <= ArenaParams.MinimumAlignment)) // one branch, not two
+		{
+			return true;
+		}
+		if (ArenaParams.bAttemptToAlignSmallBocks & (InOutSize <= ArenaParams.MaxPoolSize) & (Alignment <= ArenaParams.MaximumAlignmentForSmallBlock)) // one branch, not three
+		{
+			uint32 PoolIndex = BoundSizeToPoolIndex(Align(InOutSize, Alignment));
+			while (true)
+			{
+				uint32 BlockSize = PoolIndexToBlockSize(PoolIndex);
+				if (IsAligned(BlockSize, Alignment))
+				{
+					InOutSize = SIZE_T(BlockSize);
+					return true;
+				}
+				PoolIndex++;
+				check(PoolIndex < ArenaParams.PoolCount);
+			}
+		}
+		return false;
+	}
+
+public:
+
+
+	FMallocBinnedGPU();
+	FArenaParams& GetParams()
+	{
+		return ArenaParams;
+	}
+	void InitMallocBinned();
+
+	virtual ~FMallocBinnedGPU();
+
+
+	// FMalloc interface.
+	virtual bool IsInternallyThreadSafe() const override;
+	FORCEINLINE virtual void* Malloc(SIZE_T Size, uint32 Alignment) override
+	{
+		Alignment = FMath::Max<uint32>(Alignment, ArenaParams.MinimumAlignment);
+
+		void* Result = nullptr;
+
+		// Only allocate from the small pools if the size is small enough and the alignment isn't crazy large.
+		// With large alignments, we'll waste a lot of memory allocating an entire page, but such alignments are highly unlikely in practice.
+		if (AdjustSmallBlockSizeForAlignment(Size, Alignment))
+		{
+			FPerThreadFreeBlockLists* Lists = ArenaParams.bPerThreadCaches ? FPerThreadFreeBlockLists::Get(BinnedGPUTlsSlot) : nullptr;
+			if (Lists)
+			{
+				uint32 PoolIndex = BoundSizeToPoolIndex(Size);
+				uint32 BlockSize = PoolIndexToBlockSize(PoolIndex);
+				Result = Lists->Malloc(*this, PoolIndex);
+				if (Result)
+				{
+					Lists->AllocatedMemory += BlockSize;
+					checkSlow(IsAligned(Result, Alignment));
+				}
+			}
+		}
+		if (Result == nullptr)
+		{
+			Result = MallocExternal(Size, Alignment);
+		}
+
+		return Result;
+	}
+	FORCEINLINE virtual void* Realloc(void* Ptr, SIZE_T NewSize, uint32 Alignment) override
+	{
+		check(!"MallocBinnedGPU cannot realloc memory because the memory is assumed to not be writable by the CPU");
+		return nullptr;
+	}
+
+	FORCEINLINE virtual void Free(void* Ptr) override
+	{
+		uint64 PoolIndex = PoolIndexFromPtr(Ptr);
+		if (PoolIndex < ArenaParams.PoolCount)
+		{
+			FPerThreadFreeBlockLists* Lists = ArenaParams.bPerThreadCaches ? FPerThreadFreeBlockLists::Get(BinnedGPUTlsSlot) : nullptr;
+			if (Lists)
+			{
+				int32 BlockSize = PoolIndexToBlockSize(PoolIndex);
+				if (Lists->Free(*this, Ptr, PoolIndex, BlockSize, ArenaParams))
+				{
+					Lists->AllocatedMemory -= BlockSize;
+					return;
+				}
+			}
+		}
+		FreeExternal(Ptr);
+	}
+	FORCEINLINE virtual bool GetAllocationSize(void *Ptr, SIZE_T &SizeOut) override
+	{
+		uint64 PoolIndex = PoolIndexFromPtr(Ptr);
+		if (PoolIndex < ArenaParams.PoolCount)
+		{
+			SizeOut = PoolIndexToBlockSize(PoolIndex);
+			return true;
+		}
+		return GetAllocationSizeExternal(Ptr, SizeOut);
+	}
+
+	FORCEINLINE virtual SIZE_T QuantizeSize(SIZE_T Count, uint32 Alignment) override
+	{
+		check(DEFAULT_ALIGNMENT <= ArenaParams.MinimumAlignment); // used below
+		checkSlow((Alignment & (Alignment - 1)) == 0); // Check the alignment is a power of two
+		SIZE_T SizeOut;
+		if ((Count <= ArenaParams.MaxPoolSize) & (Alignment <= ArenaParams.MinimumAlignment)) // one branch, not two
+		{
+			SizeOut = PoolIndexToBlockSize(BoundSizeToPoolIndex(Count));
+		}
+		else
+		{
+			Alignment = FPlatformMath::Max<uint32>(Alignment, ArenaParams.AllocationGranularity);
+			SizeOut = Align(Count, Alignment);
+		}
+		check(SizeOut >= Count);
+		return SizeOut;
+	}
+
+	virtual bool ValidateHeap() override;
+	virtual void Trim(bool bTrimThreadCaches) override;
+	virtual void SetupTLSCachesOnCurrentThread() override;
+	virtual void ClearAndDisableTLSCachesOnCurrentThread() override;
+	virtual const TCHAR* GetDescriptiveName() override;
+	// End FMalloc interface.
+
+	void FlushCurrentThreadCache();
+	void* MallocExternal(SIZE_T Size, uint32 Alignment);
+	void FreeExternal(void *Ptr);
+	bool GetAllocationSizeExternal(void* Ptr, SIZE_T& SizeOut);
+
+	MBG_STAT(int64 GetTotalAllocatedSmallPoolMemory();)
+	virtual void GetAllocatorStats(FGenericMemoryStats& out_Stats) override;
+	/** Dumps current allocator stats to the log. */
+	virtual void DumpAllocatorStats(class FOutputDevice& Ar) override;
+
+	FORCEINLINE uint32 BoundSizeToPoolIndex(SIZE_T Size)
+	{
+		auto Index = ((Size + ArenaParams.MinimumAlignment - 1) >> ArenaParams.MinimumAlignmentShift);
+		checkSlow(Index >= 0 && Index <= (ArenaParams.MaxPoolSize >> ArenaParams.MinimumAlignmentShift)); // and it should be in the table
+		uint32 PoolIndex = uint32(MemSizeToIndex[Index]);
+		checkSlow(PoolIndex >= 0 && PoolIndex < ArenaParams.PoolCount);
+		return PoolIndex;
+	}
+	FORCEINLINE uint32 PoolIndexToBlockSize(uint32 PoolIndex)
+	{
+		return uint32(SmallBlockSizesReversedShifted[ArenaParams.PoolCount - PoolIndex - 1]) << ArenaParams.MinimumAlignmentShift;
+	}
+
+	void Commit(uint32 InPoolIndex, void *Ptr, SIZE_T Size);
+	void Decommit(uint32 InPoolIndex, void *Ptr, SIZE_T Size);
+
+
+	// Pool tables for different pool sizes
+	TArray<FPoolTable> SmallPoolTables;
+
+	uint32 SmallPoolInfosPerPlatformPage;
+
+	PoolHashBucket* HashBuckets;
+	PoolHashBucket* HashBucketFreeList;
+	uint64 NumLargePoolsPerPage;
+
+	FCriticalSection Mutex;
+	FGlobalRecycler GGlobalRecycler;
+	FPtrToPoolMapping PtrToPoolMapping;
+
+	FArenaParams ArenaParams;
+
+	TArray<uint16> SmallBlockSizesReversedShifted; // this is reversed to get the smallest elements on our main cache line
+	uint32 BinnedGPUTlsSlot;
+	uint64 PoolSearchDiv; // if this is zero, the VM turned out to be contiguous anyway so we use a simple subtract and shift
+	uint8* HighestPoolBaseVMPtr; // this is a duplicate of PoolBaseVMPtr[ArenaParams.PoolCount - 1]
+	FPlatformMemory::FPlatformVirtualMemoryBlock PoolBaseVMBlock;
+	TArray<uint8*> PoolBaseVMPtr;
+	TArray<FPlatformMemory::FPlatformVirtualMemoryBlock> PoolBaseVMBlocks;
+	// Mapping of sizes to small table indices
+	TArray<uint8> MemSizeToIndex;
+
+	MBG_STAT(
+		int64 BinnedGPUAllocatedSmallPoolMemory = 0; // memory that's requested to be allocated by the game
+		int64 BinnedGPUAllocatedOSSmallPoolMemory = 0;
+
+		int64 BinnedGPUAllocatedLargePoolMemory = 0; // memory requests to the OS which don't fit in the small pool
+		int64 BinnedGPUAllocatedLargePoolMemoryWAlignment = 0; // when we allocate at OS level we need to align to a size
+
+		int64 BinnedGPUPoolInfoMemory = 0;
+		int64 BinnedGPUHashMemory = 0;
+		int64 BinnedGPUFreeBitsMemory = 0;
+		int64 BinnedGPUTLSMemory = 0;
+		TAtomic<int64> ConsolidatedMemory;
+		TAtomic<int64> GPUProxyMemory;
+	)
+
+	FCriticalSection FreeBlockListsRegistrationMutex;
+	FCriticalSection& GetFreeBlockListsRegistrationMutex()
+	{
+		return FreeBlockListsRegistrationMutex;
+	}
+	TArray<FPerThreadFreeBlockLists*> RegisteredFreeBlockLists;
+	TArray<FPerThreadFreeBlockLists*>& GetRegisteredFreeBlockLists()
+	{
+		return RegisteredFreeBlockLists;
+	}
+	void RegisterThreadFreeBlockLists(FPerThreadFreeBlockLists* FreeBlockLists)
+	{
+		FScopeLock Lock(&GetFreeBlockListsRegistrationMutex());
+		GetRegisteredFreeBlockLists().Add(FreeBlockLists);
+	}
+	int64 UnregisterThreadFreeBlockLists(FPerThreadFreeBlockLists* FreeBlockLists)
+	{
+		FScopeLock Lock(&GetFreeBlockListsRegistrationMutex());
+		GetRegisteredFreeBlockLists().Remove(FreeBlockLists);
+		return FreeBlockLists->AllocatedMemory;
+	}
+
+	TArray<void*> MallocedPointers;
+};
+
+#endif
\ No newline at end of file
diff --git a/Engine/Source/Runtime/Core/Public/HAL/PageCache.h b/Engine/Source/Runtime/Core/Public/HAL/PageCache.h
new file mode 100644
--- /dev/null
+++ b/Engine/Source/Runtime/Core/Public/HAL/PageCache.h
@@ -0,0 +1,299 @@
+// Copyright 1998-2018 Epic Games, Inc. All Rights Reserved.
+
+#pragma once
+
+#include "CoreTypes.h"
+#include "HAL/MallocBinnedCommon.h"
+
+class FPageCache
+{
+	FCriticalSection CriticalSection;
+
+	uint8* LowAddress;
+	uint8* HighAddress;
+	size_t PageSize;
+	uint32 NumPages;
+	uint32 MemSize;
+	uint32 SweepPage;
+	uint32 CommittedPages;
+	uint32 DecommittedPages;
+	uint32 PendingDecommittedPages;
+
+	uint64 CommitHits;
+	uint64 CommitMisses;
+
+
+	FBitTree CurrentlyCommitted;
+	FBitTree NotPendingDecommit;
+
+	uint32 AddrToPageIndex(void* Addr) const
+	{
+		uint32 PageIndex = (((uint8*)Addr) - LowAddress) / PageSize;
+		check(PageIndex < NumPages && (uint8*)Addr >= LowAddress);
+		return PageIndex;
+	}
+
+	void* PageIndexToAddr(uint32 PageIndex) const
+	{
+		check(PageIndex < NumPages);
+		return (void*)(LowAddress + size_t(PageIndex) * PageSize);
+	}
+
+	bool IsCommitted(void* Addr) const
+	{
+		uint32 PageIndex = AddrToPageIndex(Addr);
+		bool bCommitted = CurrentlyCommitted.IsAllocated(PageIndex);
+		check(bCommitted || NotPendingDecommit.IsAllocated(PageIndex));  // can't be decommited and pending decommit
+		return bCommitted;
+	}
+
+	bool IsPendingDecommit(void* Addr) const
+	{
+		uint32 PageIndex = AddrToPageIndex(Addr);
+		bool bPendingDecommit = !NotPendingDecommit.IsAllocated(PageIndex);
+		check(!bPendingDecommit || CurrentlyCommitted.IsAllocated(PageIndex));  // can't be decommited and pending decommit
+		return bPendingDecommit;
+	}
+
+	void Commit(void* Addr)
+	{
+		uint32 PageIndex = AddrToPageIndex(Addr);
+		CurrentlyCommitted.AllocBit(PageIndex);
+	}
+	void Decommit(void* Addr)
+	{
+		uint32 PageIndex = AddrToPageIndex(Addr);
+		CurrentlyCommitted.FreeBit(PageIndex);
+	}
+	void MarkPendingDecommit(void* Addr)
+	{
+		uint32 PageIndex = AddrToPageIndex(Addr);
+		check(IsCommitted(Addr));
+		NotPendingDecommit.FreeBit(PageIndex);
+	}
+	void UnMarkPendingDecommit(void* Addr)
+	{
+		uint32 PageIndex = AddrToPageIndex(Addr);
+		check(IsCommitted(Addr));
+		NotPendingDecommit.AllocBit(PageIndex);
+	}
+
+public:
+
+	FPageCache(void *InLowAdress, void* InHighAddress, size_t InPageSize)
+		: LowAddress((uint8*)InLowAdress)
+		, HighAddress((uint8*)InHighAddress)
+		, PageSize(InPageSize)
+		, NumPages(0)
+		, MemSize(0)
+		, SweepPage(0)
+		, CommittedPages(0)
+		, PendingDecommittedPages(0)
+	{
+	}
+
+	uint32 GetMemoryRequirements()
+	{
+		check(PageSize && LowAddress && HighAddress > LowAddress);
+		NumPages = (HighAddress - LowAddress) / PageSize;
+		MemSize = FBitTree::GetMemoryRequirements(NumPages);
+		return MemSize * 2;
+	}
+
+	void InitPageCache(void* Memory)
+	{
+		check(NumPages && MemSize);
+		FScopeLock Lock(&CriticalSection);
+		DecommittedPages = NumPages;
+		CurrentlyCommitted.FBitTreeInit(NumPages, Memory, MemSize, false);
+		NotPendingDecommit.FBitTreeInit(NumPages, (uint8*)Memory + MemSize, MemSize, true);
+	}
+
+	size_t MarkForPendingDecommit(void *InAddr, size_t Size)
+	{
+		check(Size > 0 && IsAligned(Size, PageSize));
+		uint32 StartPage = AddrToPageIndex(InAddr);
+		uint32 EndPage = StartPage + Size / PageSize;
+
+		uint32 NumFound = 0;
+
+		FScopeLock Lock(&CriticalSection);
+		// this loop could be accelerated by using the hierarchical info in the bit tree 
+		for (uint32 Index = StartPage; Index < EndPage; Index++)
+		{
+			if (CurrentlyCommitted.IsAllocated(Index) && NotPendingDecommit.IsAllocated(Index))
+			{
+				NumFound++;
+				PendingDecommittedPages++;
+				NotPendingDecommit.FreeBit(Index);
+			}
+		}
+		return size_t(NumFound) * PageSize;
+	}
+
+	template<typename T>
+	size_t Commit(void *InAddr, size_t Size, size_t& OutUnPending, T&& CommitFunction)
+	{
+		check(Size > 0 && IsAligned(Size, PageSize));
+		uint32 StartPage = AddrToPageIndex(InAddr);
+		uint32 EndPage = StartPage + Size / PageSize;
+
+		uint32 LastCommitPage = MAX_uint32;
+		uint32 StartCommitPage = MAX_uint32;
+		FScopeLock Lock(&CriticalSection);
+		// this loop could be accelerated by using the hierarchical info in the bit tree 
+		uint32 NumFound = 0;
+		uint32 NumUnPending = 0;
+		for (uint32 Index = StartPage; Index < EndPage; Index++)
+		{
+			if (CurrentlyCommitted.IsAllocated(Index))
+			{
+				if (!NotPendingDecommit.IsAllocated(Index))
+				{
+					check(PendingDecommittedPages);
+					PendingDecommittedPages--;
+					NumUnPending++;
+					NotPendingDecommit.AllocBit(Index);
+				}
+			}
+			else
+			{
+				NumFound++;
+				CommittedPages++;
+				check(DecommittedPages);
+				DecommittedPages--;
+				check(NotPendingDecommit.IsAllocated(Index));
+				CurrentlyCommitted.AllocBit(Index);
+				if (StartCommitPage == MAX_uint32 || LastCommitPage + 1 != Index)
+				{
+					if (StartCommitPage != MAX_uint32)
+					{
+						CommitFunction(PageIndexToAddr(StartCommitPage), (1 + LastCommitPage - StartCommitPage) * PageSize);
+					}
+					StartCommitPage = Index;
+				}
+				LastCommitPage = Index;
+			}
+		}
+		if (StartCommitPage != MAX_uint32)
+		{
+			CommitFunction(PageIndexToAddr(StartCommitPage), (1 + LastCommitPage - StartCommitPage) * PageSize);
+		}
+		CommitHits += NumUnPending;
+		CommitMisses += NumFound;
+		OutUnPending = size_t(NumUnPending) * PageSize;
+		return size_t(NumFound) * PageSize;
+	}
+	template<typename T>
+	size_t ForceDecommit(void *InAddr, size_t Size, size_t& OutUnPending, T&& DecommitFunction)
+	{
+		check(Size > 0 && IsAligned(Size, PageSize));
+		uint32 StartPage = AddrToPageIndex(InAddr);
+		uint32 EndPage = StartPage + Size / PageSize;
+
+		uint32 LastCommitPage = MAX_uint32;
+		uint32 StartCommitPage = MAX_uint32;
+		FScopeLock Lock(&CriticalSection);
+		// this loop could be accelerated by using the hierarchical info in the bit tree 
+		uint32 NumFound = 0;
+		uint32 NumUnPending = 0;
+		for (uint32 Index = StartPage; Index < EndPage; Index++)
+		{
+			if (CurrentlyCommitted.IsAllocated(Index))
+			{
+				if (!NotPendingDecommit.IsAllocated(Index))
+				{
+					check(PendingDecommittedPages);
+					PendingDecommittedPages--;
+					NotPendingDecommit.AllocBit(Index);
+					NumUnPending++;
+				}
+				NumFound++;
+				check(CommittedPages);
+				CommittedPages--;
+				DecommittedPages++;
+				CurrentlyCommitted.FreeBit(Index);
+				if (StartCommitPage == MAX_uint32 || LastCommitPage + 1 != Index)
+				{
+					if (StartCommitPage != MAX_uint32)
+					{
+						DecommitFunction(PageIndexToAddr(StartCommitPage), (1 + LastCommitPage - StartCommitPage) * PageSize);
+					}
+					StartCommitPage = Index;
+				}
+				LastCommitPage = Index;
+			}
+		}
+		if (StartCommitPage != MAX_uint32)
+		{
+			DecommitFunction(PageIndexToAddr(StartCommitPage), (1 + LastCommitPage - StartCommitPage) * PageSize);
+		}
+		OutUnPending = size_t(NumUnPending) * PageSize;
+		return size_t(NumFound) * PageSize;
+	}
+	template<typename T>
+	size_t DecommitPending(size_t Size, T&& DecommitFunction)
+	{
+		if (!NumPages)
+		{
+			return 0; // this page cache was never really set up, nothing to sweep
+		}
+		check(Size > 0 && IsAligned(Size, PageSize));
+		uint32 NumNeed = Size / PageSize;
+
+		uint32 LastDecommitPage = MAX_uint32;
+		uint32 StartDecommitPage = MAX_uint32;
+		uint32 NumFound = 0;
+		FScopeLock Lock(&CriticalSection);
+
+		while (NumFound < NumNeed)
+		{
+			check(SweepPage < NumPages);
+			uint32 Index = NotPendingDecommit.NextAllocBit(SweepPage);
+			if (Index == MAX_uint32)
+			{
+				SweepPage = 0;
+				break;
+			}
+			check(CurrentlyCommitted.IsAllocated(Index) && !NotPendingDecommit.IsAllocated(Index));
+			check(CommittedPages);
+			CommittedPages--;
+			DecommittedPages++;
+			check(PendingDecommittedPages);
+			PendingDecommittedPages--;
+			NumFound++;
+			CurrentlyCommitted.FreeBit(Index);
+			NotPendingDecommit.AllocBit(Index);
+			if (StartDecommitPage == MAX_uint32 || LastDecommitPage + 1 != Index)
+			{
+				if (StartDecommitPage != MAX_uint32)
+				{
+					DecommitFunction(PageIndexToAddr(StartDecommitPage), (1 + LastDecommitPage - StartDecommitPage) * PageSize);
+				}
+				StartDecommitPage = Index;
+			}
+			LastDecommitPage = Index;
+			SweepPage = Index + 1;
+			if (SweepPage >= NumPages)
+			{
+				SweepPage = 0;
+				break;
+			}
+		}
+		if (StartDecommitPage != MAX_uint32)
+		{
+			DecommitFunction(PageIndexToAddr(StartDecommitPage), (1 + LastDecommitPage - StartDecommitPage) * PageSize);
+		}
+		return size_t(NumFound) * PageSize;
+	}
+	size_t GetFreeableMemory()
+	{
+		FScopeLock Lock(&CriticalSection);
+		return size_t(PendingDecommittedPages) * PageSize;
+	}
+	float GetHitRate()
+	{
+		FScopeLock Lock(&CriticalSection);
+		return 100.0f * float(CommitHits) / float(CommitHits + CommitMisses + 1); // +1 to avoid divide by zero
+	}
+};
\ No newline at end of file
diff --git a/Engine/Source/Runtime/Core/Public/HAL/UnrealMemory.h b/Engine/Source/Runtime/Core/Public/HAL/UnrealMemory.h
--- a/Engine/Source/Runtime/Core/Public/HAL/UnrealMemory.h
+++ b/Engine/Source/Runtime/Core/Public/HAL/UnrealMemory.h
@@ -190,15 +190,6 @@ struct CORE_API FMemory
 	*/
 	static void ClearAndDisableTLSCachesOnCurrentThread();
 
-	//
-	// Malloc for GPU mapped memory on UMA systems (XB1/PS4/etc)
-	// It is expected that the RHI on platforms that use these knows what to 
-	// do with the memory and avoids unnecessary copies into GPU resources, etc.
-	//
-	static void* GPUMalloc(SIZE_T Count, uint32 Alignment = DEFAULT_ALIGNMENT);
-	static void* GPURealloc(void* Original, SIZE_T Count, uint32 Alignment = DEFAULT_ALIGNMENT);
-	static void GPUFree(void* Original);
-
 	/**
 	 * A helper function that will perform a series of random heap allocations to test
 	 * the internal validity of the heap. Note, this function will "leak" memory, but another call
diff --git a/Engine/Source/Runtime/Core/Public/HAL/VirtualAllocator.h b/Engine/Source/Runtime/Core/Public/HAL/VirtualAllocator.h
new file mode 100644
--- /dev/null
+++ b/Engine/Source/Runtime/Core/Public/HAL/VirtualAllocator.h
@@ -0,0 +1,236 @@
+// Copyright 1998-2018 Epic Games, Inc. All Rights Reserved.
+
+#pragma once
+
+#include "CoreTypes.h"
+#include "HAL/ThreadSafeCounter.h"
+#include "HAL/PlatformProcess.h"
+#include "Templates/AlignmentTemplates.h"
+
+class FVirtualAllocator
+{
+	struct FFreeLink
+	{
+		void *Ptr = nullptr;
+		FFreeLink* Next = nullptr;
+	};
+
+	struct FPerBlockSize
+	{
+		int64 AllocBlocksSize = 0;
+		int64 FreeBlocksSize = 0;
+		FFreeLink* FirstFree = nullptr;
+	};
+
+	FCriticalSection CriticalSection;
+
+	uint8* LowAddress;
+	uint8* HighAddress;
+
+	size_t TotalSize;
+	size_t PageSize;
+	size_t MaximumAlignment;
+
+	uint8* NextAlloc;
+
+	FFreeLink* RecycledLinks;
+	int64 LinkSize;
+
+	bool bBacksMalloc;
+
+	FPerBlockSize Blocks[64]; 
+
+	void FreeVirtualByBlock(void* Ptr, FPerBlockSize& Block, size_t AlignedSize)
+	{
+		// already locked
+		if (!RecycledLinks)
+		{
+			void *Alloc;
+			// If we ARE malloc, then we know we can (and must) commit part of our range for free links. Otherwise, we can just use malloc (and can't use our VM space anyway)
+			if (bBacksMalloc)
+			{
+				Alloc = AllocNewVM(MaximumAlignment);
+				uint32 Pages = MaximumAlignment / FPlatformMemory::FPlatformVirtualMemoryBlock::GetVirtualSizeAlignment();
+				check(Pages && MaximumAlignment % FPlatformMemory::FPlatformVirtualMemoryBlock::GetVirtualSizeAlignment() == 0);
+				FPlatformMemory::FPlatformVirtualMemoryBlock VMBlock(Alloc, Pages);
+				VMBlock.Commit();
+			}
+			else
+			{
+				Alloc = FMemory::Malloc(MaximumAlignment);
+			}
+			for (int32 Index = 0; (Index + 1) * sizeof(FFreeLink) <= MaximumAlignment; Index++)
+			{
+				FFreeLink* NewLink = new ((void*)(((uint8*)Alloc + Index * sizeof(FFreeLink)))) FFreeLink;
+				NewLink->Next = RecycledLinks;
+				RecycledLinks = NewLink;
+			}
+			LinkSize += MaximumAlignment;
+		}
+
+		FFreeLink* Link = RecycledLinks;
+		RecycledLinks = Link->Next;
+		check(!Link->Ptr && Ptr);
+		Link->Ptr = Ptr;
+		Link->Next = Block.FirstFree;
+		Block.FirstFree = Link;
+
+		Block.FreeBlocksSize += AlignedSize;
+	}
+protected:
+	virtual uint8* AllocNewVM(size_t AlignedSize)
+	{
+		uint8* Result = NextAlloc;
+		check(IsAligned(Result, MaximumAlignment) && IsAligned(AlignedSize, MaximumAlignment));
+		NextAlloc = Result + AlignedSize;
+		return Result;
+	}
+
+public:
+
+	FVirtualAllocator(void *InLowAdress, void* InHighAddress, size_t InPageSize, size_t InMaximumAlignment, bool bInBacksMalloc)
+		: LowAddress((uint8*)InLowAdress)
+		, HighAddress((uint8*)InHighAddress)
+		, PageSize(InPageSize)
+		, MaximumAlignment(FMath::Max(InMaximumAlignment, InPageSize))
+		, NextAlloc((uint8*)InLowAdress)
+		, RecycledLinks(nullptr)
+		, LinkSize(0)
+		, bBacksMalloc(bInBacksMalloc)
+	{
+		TotalSize = HighAddress - LowAddress;
+		check(LowAddress && HighAddress && LowAddress < HighAddress && IsAligned(LowAddress, MaximumAlignment));
+	}
+
+	uint32 GetPagesForSizeAndAlignment(size_t Size, size_t Alignment = 1) const
+	{
+		check(Alignment <= MaximumAlignment && Alignment > 0);
+		size_t SizeAndAlignment = FMath::Max(Align(Size, Alignment), PageSize);
+		if (SizeAndAlignment * 2 >= TotalSize)
+		{
+			// this is hack, MB3 will ask for tons of virtual and never free it, so we won't round up to power of two
+			size_t Pages = SizeAndAlignment / PageSize;
+			check(Pages == uint32(Pages)); // overflow of uint32
+			return Pages;
+		}
+		int32 BlockIndex = FMath::CeilLogTwo64(FMath::Max(SizeAndAlignment, PageSize));
+		size_t AlignedSize = size_t(1) << BlockIndex;
+		check(AlignedSize % PageSize == 0);
+		size_t Pages = AlignedSize / PageSize;
+		check(Pages == uint32(Pages)); // overflow of uint32
+		return Pages;
+	}
+
+	void* AllocateVirtualPages(uint32 NumPages, size_t AlignmentForCheck = 1) 
+	{
+		check(AlignmentForCheck <= MaximumAlignment && AlignmentForCheck > 0 && NumPages);
+
+
+		int32 BlockIndex = FMath::CeilLogTwo64(NumPages * PageSize);
+		size_t AlignedSize = size_t(1) << BlockIndex;
+		bool bHackForHugeBlock = false;
+		if (size_t(NumPages) * PageSize * 2 >= TotalSize)
+		{
+			// this is hack, MB3 will ask for tons of virtual and never free it, so we won't round up to power of two
+			AlignedSize = size_t(NumPages) * PageSize;
+			bHackForHugeBlock = true;
+		}
+		FPerBlockSize& Block = Blocks[BlockIndex];
+
+		FScopeLock Lock(&CriticalSection);
+
+
+		uint8* Result;
+		if (Block.FirstFree)
+		{
+			check(!bHackForHugeBlock);
+			Result = (uint8*)Block.FirstFree->Ptr;
+			check(Result);
+			Block.FirstFree->Ptr = nullptr;
+
+			FFreeLink* Next = RecycledLinks;
+			RecycledLinks = Block.FirstFree;
+			Block.FirstFree = Block.FirstFree->Next;
+			check(!Block.FirstFree || Block.FirstFree->Ptr);
+			RecycledLinks->Next = Next;
+
+			check(IsAligned(Result, FMath::Min(AlignedSize, MaximumAlignment)));
+			Block.FreeBlocksSize -= AlignedSize;
+		}
+		else
+		{
+			size_t AllocSize = FMath::Max(AlignedSize, MaximumAlignment);
+			uint8* LocalNextAlloc = AllocNewVM(AllocSize);
+			check(IsAligned(LocalNextAlloc, MaximumAlignment));
+
+			uint8* NewNextAlloc = LocalNextAlloc + AllocSize;
+
+			if (NewNextAlloc > HighAddress)
+			{
+				FPlatformMemory::OnOutOfMemory(HighAddress - LowAddress, 0);
+			}
+
+			Block.AllocBlocksSize += AllocSize;
+			Result = LocalNextAlloc;
+			check(Result);
+			LocalNextAlloc += AlignedSize;
+			while (LocalNextAlloc < NewNextAlloc && !bHackForHugeBlock)
+			{
+				FreeVirtualByBlock(LocalNextAlloc, Block, AlignedSize);
+				LocalNextAlloc += AlignedSize;
+			}
+		}
+		check(IsAligned(Result, AlignmentForCheck));
+		return Result;
+	}
+	
+	void FreeVirtual(void* Ptr, uint32 NumPages)
+	{
+		if (size_t(NumPages) * PageSize * 2 >= TotalSize)
+		{
+			// this is hack, MB3 will ask for tons of virtual and never free it, so we won't round up to power of two
+			check(!"Huge vm blocks may not be freed.");
+			return;
+		}
+		int32 BlockIndex = FMath::CeilLogTwo64(NumPages * PageSize);
+		size_t AlignedSize = size_t(1) << BlockIndex;
+
+		FPerBlockSize& Block = Blocks[BlockIndex];
+		FScopeLock Lock(&CriticalSection);
+		FreeVirtualByBlock(Ptr, Block, AlignedSize);
+	}
+
+	struct FVirtualAllocatorStatsPerBlockSize
+	{
+		size_t AllocBlocksSize;
+		size_t FreeBlocksSize;
+	};
+	struct FVirtualAllocatorStats
+	{
+		size_t PageSize;
+		size_t MaximumAlignment;
+		size_t VMSpaceTotal;
+		size_t VMSpaceConsumed;
+		size_t VMSpaceConsumedPeak;
+
+		size_t FreeListLinks;
+
+		FVirtualAllocatorStatsPerBlockSize BlockStats[64];
+	};
+
+	void GetStats(FVirtualAllocatorStats& OutStats)
+	{
+		FScopeLock Lock(&CriticalSection);
+		OutStats.PageSize = PageSize;
+		OutStats.MaximumAlignment = MaximumAlignment;
+		OutStats.VMSpaceTotal = HighAddress - LowAddress;
+		OutStats.VMSpaceConsumed = NextAlloc - LowAddress;
+		OutStats.FreeListLinks = LinkSize;
+
+		for (int32 Index = 0; Index < 64; Index++)
+		{
+			OutStats.BlockStats[Index].AllocBlocksSize = Blocks[Index].AllocBlocksSize;
+			OutStats.BlockStats[Index].FreeBlocksSize = Blocks[Index].FreeBlocksSize;
+		}
+	}
+};
\ No newline at end of file
diff --git a/Engine/Source/Runtime/Core/Public/Unix/UnixPlatformMemory.h b/Engine/Source/Runtime/Core/Public/Unix/UnixPlatformMemory.h
--- a/Engine/Source/Runtime/Core/Public/Unix/UnixPlatformMemory.h
+++ b/Engine/Source/Runtime/Core/Public/Unix/UnixPlatformMemory.h
@@ -73,7 +73,57 @@ struct CORE_API FUnixPlatformMemory : public FGenericPlatformMemory
 	static bool PageProtect(void* const Ptr, const SIZE_T Size, const bool bCanRead, const bool bCanWrite);
 	static void* BinnedAllocFromOS(SIZE_T Size);
 	static void BinnedFreeToOS(void* Ptr, SIZE_T Size);
-	static bool MemoryRangeDecommit(void* Ptr, SIZE_T Size);
+
+	class FPlatformVirtualMemoryBlock : public FBasicVirtualMemoryBlock
+	{
+	public:
+
+		FPlatformVirtualMemoryBlock()
+		{
+		}
+
+		FPlatformVirtualMemoryBlock(void *InPtr, uint32 InVMSizeDivVirtualSizeAlignment)
+			: FBasicVirtualMemoryBlock(InPtr, InVMSizeDivVirtualSizeAlignment)
+		{
+		}
+		FPlatformVirtualMemoryBlock(const FPlatformVirtualMemoryBlock& Other) = default;
+		FPlatformVirtualMemoryBlock& operator=(const FPlatformVirtualMemoryBlock& Other) = default;
+
+		void Commit(size_t InOffset, size_t InSize);
+		void Decommit(size_t InOffset, size_t InSize);
+		void FreeVirtual();
+
+		FORCEINLINE void CommitByPtr(void *InPtr, size_t InSize)
+		{
+			Commit(size_t(((uint8*)InPtr) - ((uint8*)Ptr)), InSize);
+		}
+
+		FORCEINLINE void DecommitByPtr(void *InPtr, size_t InSize)
+		{
+			Decommit(size_t(((uint8*)InPtr) - ((uint8*)Ptr)), InSize);
+		}
+
+		FORCEINLINE void Commit()
+		{
+			Commit(0, GetActualSize());
+		}
+
+		FORCEINLINE void Decommit()
+		{
+			Decommit(0, GetActualSize());
+		}
+
+		FORCEINLINE size_t GetActualSize() const
+		{
+			return VMSizeDivVirtualSizeAlignment * GetVirtualSizeAlignment();
+		}
+
+		static FPlatformVirtualMemoryBlock AllocateVirtual(size_t Size, size_t InAlignment = FPlatformVirtualMemoryBlock::GetVirtualSizeAlignment());
+		static size_t GetCommitAlignment();
+		static size_t GetVirtualSizeAlignment();
+	};
+
+
 	static FSharedMemoryRegion * MapNamedSharedMemoryRegion(const FString& InName, bool bCreate, uint32 AccessMode, SIZE_T Size);
 	static bool UnmapNamedSharedMemoryRegion(FSharedMemoryRegion * MemoryRegion);
 	static bool GetLLMAllocFunctions(void*(*&OutAllocFunction)(size_t), void(*&OutFreeFunction)(void*, size_t), int32& OutAlignment);
diff --git a/Engine/Source/Runtime/Core/Public/Windows/WindowsPlatformMemory.h b/Engine/Source/Runtime/Core/Public/Windows/WindowsPlatformMemory.h
--- a/Engine/Source/Runtime/Core/Public/Windows/WindowsPlatformMemory.h
+++ b/Engine/Source/Runtime/Core/Public/Windows/WindowsPlatformMemory.h
@@ -90,10 +90,56 @@ struct CORE_API FWindowsPlatformMemory
 	static bool PageProtect(void* const Ptr, const SIZE_T Size, const bool bCanRead, const bool bCanWrite);
 	static void* BinnedAllocFromOS( SIZE_T Size );
 	static void BinnedFreeToOS( void* Ptr, SIZE_T Size );
-	static void* MemoryRangeReserve(SIZE_T Size, bool bCommit = false, int32 Node = -1);
-	static void MemoryRangeFree(void* Ptr, SIZE_T Size);
-	static bool MemoryRangeCommit(void* Ptr, SIZE_T Size);
-	static bool MemoryRangeDecommit(void* Ptr, SIZE_T Size);
+
+	class FPlatformVirtualMemoryBlock : public FBasicVirtualMemoryBlock
+	{
+	public:
+
+		FPlatformVirtualMemoryBlock()
+		{
+		}
+
+		FPlatformVirtualMemoryBlock(void *InPtr, uint32 InVMSizeDivVirtualSizeAlignment)
+			: FBasicVirtualMemoryBlock(InPtr, InVMSizeDivVirtualSizeAlignment)
+		{
+		}
+		FPlatformVirtualMemoryBlock(const FPlatformVirtualMemoryBlock& Other) = default;
+		FPlatformVirtualMemoryBlock& operator=(const FPlatformVirtualMemoryBlock& Other) = default;
+
+		void Commit(size_t InOffset, size_t InSize);
+		void Decommit(size_t InOffset, size_t InSize);
+		void FreeVirtual();
+
+		FORCEINLINE void CommitByPtr(void *InPtr, size_t InSize)
+		{
+			Commit(size_t(((uint8*)InPtr) - ((uint8*)Ptr)), InSize);
+		}
+
+		FORCEINLINE void DecommitByPtr(void *InPtr, size_t InSize)
+		{
+			Decommit(size_t(((uint8*)InPtr) - ((uint8*)Ptr)), InSize);
+		}
+
+		FORCEINLINE void Commit()
+		{
+			Commit(0, GetActualSize());
+		}
+
+		FORCEINLINE void Decommit()
+		{
+			Decommit(0, GetActualSize());
+		}
+
+		FORCEINLINE size_t GetActualSize() const
+		{
+			return VMSizeDivVirtualSizeAlignment * GetVirtualSizeAlignment();
+		}
+
+		static FPlatformVirtualMemoryBlock AllocateVirtual(size_t Size, size_t InAlignment = FPlatformVirtualMemoryBlock::GetVirtualSizeAlignment());
+		static size_t GetCommitAlignment();
+		static size_t GetVirtualSizeAlignment();
+	};
+
 	static FSharedMemoryRegion* MapNamedSharedMemoryRegion(const FString& InName, bool bCreate, uint32 AccessMode, SIZE_T Size);
 	static bool UnmapNamedSharedMemoryRegion(FSharedMemoryRegion * MemoryRegion);
 	static bool GetLLMAllocFunctions(void*(*&OutAllocFunction)(size_t), void(*&OutFreeFunction)(void*, size_t), int32& OutAlignment);
diff --git a/Engine/Source/Runtime/D3D12RHI/Private/D3D12RHIPrivate.h b/Engine/Source/Runtime/D3D12RHI/Private/D3D12RHIPrivate.h
--- a/Engine/Source/Runtime/D3D12RHI/Private/D3D12RHIPrivate.h
+++ b/Engine/Source/Runtime/D3D12RHI/Private/D3D12RHIPrivate.h
@@ -930,8 +930,8 @@ public:
 	static __declspec(thread) FD3D12FastAllocator* HelperThreadDynamicHeapAllocator;
 
 #if	PLATFORM_SUPPORTS_VIRTUAL_TEXTURES
-	virtual void* CreateVirtualTexture(uint32 Flags, D3D12_RESOURCE_DESC& ResourceDesc, const struct FD3D12TextureLayout& TextureLayout, FD3D12Resource** ppResource, D3D12_RESOURCE_STATES InitialUsage = D3D12_RESOURCE_STATE_PIXEL_SHADER_RESOURCE) = 0;
-	virtual void DestroyVirtualTexture(uint32 Flags, void* RawTextureMemory, uint64 CommittedTextureSize) = 0;
+	virtual void* CreateVirtualTexture(uint32 Flags, D3D12_RESOURCE_DESC& ResourceDesc, const struct FD3D12TextureLayout& TextureLayout, FD3D12Resource** ppResource, FPlatformMemory::FPlatformVirtualMemoryBlock& RawTextureBlock, D3D12_RESOURCE_STATES InitialUsage = D3D12_RESOURCE_STATE_PIXEL_SHADER_RESOURCE) = 0;
+	virtual void DestroyVirtualTexture(uint32 Flags, void* RawTextureMemory, FPlatformMemory::FPlatformVirtualMemoryBlock& RawTextureBlock, uint64 CommittedTextureSize) = 0;
 	virtual bool HandleSpecialLock(void*& MemoryOut, uint32 MipIndex, uint32 ArrayIndex, uint32 Flags, EResourceLockMode LockMode, const FD3D12TextureLayout& TextureLayout, void* RawTextureMemory, uint32& DestStride) = 0;
 	virtual bool HandleSpecialUnlock(FRHICommandListBase* RHICmdList, uint32 MipIndex, uint32 Flags, const struct FD3D12TextureLayout& TextureLayout, void* RawTextureMemory) = 0;
 #endif
diff --git a/Engine/Source/Runtime/D3D12RHI/Private/D3D12Texture.cpp b/Engine/Source/Runtime/D3D12RHI/Private/D3D12Texture.cpp
--- a/Engine/Source/Runtime/D3D12RHI/Private/D3D12Texture.cpp
+++ b/Engine/Source/Runtime/D3D12RHI/Private/D3D12Texture.cpp
@@ -400,7 +400,7 @@ TD3D12Texture2D<BaseResourceType>::~TD3D12Texture2D()
 		FD3D12TextureStats::D3D12TextureDeleted(*this);
 	}
 #if PLATFORM_SUPPORTS_VIRTUAL_TEXTURES
-	GetParentDevice()->GetOwningRHI()->DestroyVirtualTexture(GetFlags(), GetRawTextureMemory(), GetMemorySize());
+	GetParentDevice()->GetOwningRHI()->DestroyVirtualTexture(GetFlags(), GetRawTextureMemory(), GetRawTextureBlock(), GetMemorySize());
 #endif
 }
 
diff --git a/Engine/Source/Runtime/D3D12RHI/Private/D3D12Texture.h b/Engine/Source/Runtime/D3D12RHI/Private/D3D12Texture.h
--- a/Engine/Source/Runtime/D3D12RHI/Private/D3D12Texture.h
+++ b/Engine/Source/Runtime/D3D12RHI/Private/D3D12Texture.h
@@ -322,6 +322,10 @@ public:
 	{
 		RawTextureMemory = Memory;
 	}
+	FPlatformMemory::FPlatformVirtualMemoryBlock& GetRawTextureBlock()
+	{
+		return RawTextureBlock;
+	}
 #endif
 
 	const FD3D12TextureLayout& GetTextureLayout() const { return TextureLayout; }
@@ -338,6 +342,7 @@ private:
 
 #if PLATFORM_SUPPORTS_VIRTUAL_TEXTURES
 	void* RawTextureMemory;
+	FPlatformMemory::FPlatformVirtualMemoryBlock RawTextureBlock;
 #endif
 
 	FD3D12TextureLayout TextureLayout;
