Fix copyrights

diff --git a/Engine/Source/Runtime/Core/Private/HAL/MallocBinned3.cpp b/Engine/Source/Runtime/Core/Private/HAL/MallocBinned3.cpp
--- a/Engine/Source/Runtime/Core/Private/HAL/MallocBinned3.cpp
+++ b/Engine/Source/Runtime/Core/Private/HAL/MallocBinned3.cpp
@@ -1,4 +1,4 @@
-// Copyright 1998-2018 Epic Games, Inc. All Rights Reserved.
+// Copyright 1998-2019 Epic Games, Inc. All Rights Reserved.
 
 /*=============================================================================
 	MallocBinned.cpp: Binned memory allocator
diff --git a/Engine/Source/Runtime/Core/Private/HAL/MallocBinnedArena.cpp b/Engine/Source/Runtime/Core/Private/HAL/MallocBinnedArena.cpp
--- a/Engine/Source/Runtime/Core/Private/HAL/MallocBinnedArena.cpp
+++ b/Engine/Source/Runtime/Core/Private/HAL/MallocBinnedArena.cpp
@@ -1,4 +1,4 @@
-// Copyright 1998-2018 Epic Games, Inc. All Rights Reserved.
+// Copyright 1998-2019 Epic Games, Inc. All Rights Reserved.
 
 /*=============================================================================
 	MallocBinned.cpp: Binned memory allocator
diff --git a/Engine/Source/Runtime/Core/Private/HAL/MallocBinnedCommon.cpp b/Engine/Source/Runtime/Core/Private/HAL/MallocBinnedCommon.cpp
--- a/Engine/Source/Runtime/Core/Private/HAL/MallocBinnedCommon.cpp
+++ b/Engine/Source/Runtime/Core/Private/HAL/MallocBinnedCommon.cpp
@@ -1,4 +1,4 @@
-// Copyright 1998-2018 Epic Games, Inc. All Rights Reserved.
+// Copyright 1998-2019 Epic Games, Inc. All Rights Reserved.
 
 /*=============================================================================
 	MallocBinned.cpp: Binned memory allocator
diff --git a/Engine/Source/Runtime/Core/Private/HAL/MallocBinnedGPU.cpp b/Engine/Source/Runtime/Core/Private/HAL/MallocBinnedGPU.cpp
--- a/Engine/Source/Runtime/Core/Private/HAL/MallocBinnedGPU.cpp
+++ b/Engine/Source/Runtime/Core/Private/HAL/MallocBinnedGPU.cpp
@@ -1,4 +1,4 @@
-// Copyright 1998-2018 Epic Games, Inc. All Rights Reserved.
+// Copyright 1998-2019 Epic Games, Inc. All Rights Reserved.
 
 /*=============================================================================
 	MallocBinned.cpp: Binned memory allocator
diff --git a/Engine/Source/Runtime/Core/Public/HAL/MallocBinned3.h b/Engine/Source/Runtime/Core/Public/HAL/MallocBinned3.h
--- a/Engine/Source/Runtime/Core/Public/HAL/MallocBinned3.h
+++ b/Engine/Source/Runtime/Core/Public/HAL/MallocBinned3.h
@@ -1,4 +1,4 @@
-// Copyright 1998-2018 Epic Games, Inc. All Rights Reserved.
+// Copyright 1998-2019 Epic Games, Inc. All Rights Reserved.
 
 #pragma once
 
diff --git a/Engine/Source/Runtime/Core/Public/HAL/MallocBinnedArena.h b/Engine/Source/Runtime/Core/Public/HAL/MallocBinnedArena.h
--- a/Engine/Source/Runtime/Core/Public/HAL/MallocBinnedArena.h
+++ b/Engine/Source/Runtime/Core/Public/HAL/MallocBinnedArena.h
@@ -1,4 +1,4 @@
-// Copyright 1998-2018 Epic Games, Inc. All Rights Reserved.
+// Copyright 1998-2019 Epic Games, Inc. All Rights Reserved.
 
 #pragma once
 
diff --git a/Engine/Source/Runtime/Core/Public/HAL/MallocBinnedCommon.h b/Engine/Source/Runtime/Core/Public/HAL/MallocBinnedCommon.h
--- a/Engine/Source/Runtime/Core/Public/HAL/MallocBinnedCommon.h
+++ b/Engine/Source/Runtime/Core/Public/HAL/MallocBinnedCommon.h
@@ -1,4 +1,4 @@
-// Copyright 1998-2018 Epic Games, Inc. All Rights Reserved.
+// Copyright 1998-2019 Epic Games, Inc. All Rights Reserved.
 
 #pragma once
 
diff --git a/Engine/Source/Runtime/Core/Public/HAL/MallocBinnedGPU.h b/Engine/Source/Runtime/Core/Public/HAL/MallocBinnedGPU.h
--- a/Engine/Source/Runtime/Core/Public/HAL/MallocBinnedGPU.h
+++ b/Engine/Source/Runtime/Core/Public/HAL/MallocBinnedGPU.h
@@ -1,4 +1,4 @@
-// Copyright 1998-2018 Epic Games, Inc. All Rights Reserved.
+// Copyright 1998-2019 Epic Games, Inc. All Rights Reserved.
 
 #pragma once
 
diff --git a/Engine/Source/Runtime/Core/Public/HAL/PageCache.h b/Engine/Source/Runtime/Core/Public/HAL/PageCache.h
--- a/Engine/Source/Runtime/Core/Public/HAL/PageCache.h
+++ b/Engine/Source/Runtime/Core/Public/HAL/PageCache.h
@@ -1,4 +1,4 @@
-// Copyright 1998-2018 Epic Games, Inc. All Rights Reserved.
+// Copyright 1998-2019 Epic Games, Inc. All Rights Reserved.
 
 #pragma once
 
diff --git a/Engine/Source/Runtime/Core/Public/HAL/VirtualAllocator.h b/Engine/Source/Runtime/Core/Public/HAL/VirtualAllocator.h
--- a/Engine/Source/Runtime/Core/Public/HAL/VirtualAllocator.h
+++ b/Engine/Source/Runtime/Core/Public/HAL/VirtualAllocator.h
@@ -1,4 +1,4 @@
-// Copyright 1998-2018 Epic Games, Inc. All Rights Reserved.
+// Copyright 1998-2019 Epic Games, Inc. All Rights Reserved.
 
 #pragma once
 
