GetContentState() now returns individual state information for each bundle included.

diff --git a/Engine/Source/Runtime/Core/Public/GenericPlatform/GenericPlatformInstallBundleManager.h b/Engine/Source/Runtime/Core/Public/GenericPlatform/GenericPlatformInstallBundleManager.h
--- a/Engine/Source/Runtime/Core/Public/GenericPlatform/GenericPlatformInstallBundleManager.h
+++ b/Engine/Source/Runtime/Core/Public/GenericPlatform/GenericPlatformInstallBundleManager.h
@@ -221,6 +221,7 @@ inline const TCHAR* LexToString(EInstallBundleContentState State)
 struct FInstallBundleContentState
 {
 	EInstallBundleContentState State = EInstallBundleContentState::InitializationError;
+	TMap<FName, EInstallBundleContentState> IndividualBundleStates;
 	uint64 DownloadSize = 0;
 	uint64 InstallSize = 0;	
 	uint64 FreeSpace = 0;
