HTML5 - set PLATFORM_HAS_FPlatformVirtualMemoryBlock to 0

#rnx

diff --git a/Engine/Source/Runtime/Core/Public/HTML5/HTML5Platform.h b/Engine/Source/Runtime/Core/Public/HTML5/HTML5Platform.h
--- a/Engine/Source/Runtime/Core/Public/HTML5/HTML5Platform.h
+++ b/Engine/Source/Runtime/Core/Public/HTML5/HTML5Platform.h
@@ -44,6 +44,7 @@ typedef FHTML5Types FPlatformTypes;
 #define PLATFORM_SUPPORTS_LLM						0
 #define	PLATFORM_SUPPORTS_OPUS_CODEC				0
 #define PLATFORM_SUPPORTS_VORBIS_CODEC				0
+#define	PLATFORM_HAS_FPlatformVirtualMemoryBlock	0
 
 #define PLATFORM_BREAK() \
 	[] () -> void \
