Adding a driver date filter option rather than version alone.

diff --git a/Engine/Source/Runtime/Core/Public/GenericPlatform/GenericPlatformDriver.h b/Engine/Source/Runtime/Core/Public/GenericPlatform/GenericPlatformDriver.h
--- a/Engine/Source/Runtime/Core/Public/GenericPlatform/GenericPlatformDriver.h
+++ b/Engine/Source/Runtime/Core/Public/GenericPlatform/GenericPlatformDriver.h
@@ -324,16 +324,20 @@ struct FGPUDriverInfo
 // one entry in the Hardware.ini file
 struct FBlackListEntry
 {
-	// required, e.g. "<=223.112.21.1", might includes comparison operators, later even things multiple ">12.22 <=12.44"
+	// optional, e.g. "<=223.112.21.1", might includes comparison operators, later even things multiple ">12.22 <=12.44"
 	FString DriverVersionString;
+	// optional, e.g. "<=MM-DD-YYYY"
+	FString DriverDateString;
 	// required
 	FString Reason;
 
 	// @param e.g. "DriverVersion=\"361.43\", Reason=\"UE-25096 Viewport flashes black and white when moving in the scene on latest Nvidia drivers\""
+	// At least one specifier of driver date or version is required to check for a match
 	void LoadFromINIString(const TCHAR* In)
 	{
 		FParse::Value(In, TEXT("DriverVersion="), DriverVersionString);
-		ensure(!DriverVersionString.IsEmpty());
+		FParse::Value(In, TEXT("DriverDate="), DriverDateString);
+		ensure(!DriverVersionString.IsEmpty() || !DriverDateString.IsEmpty());
 
 		// later:
 //		FParse::Value(In, TEXT("DeviceId="), DeviceId);
@@ -349,12 +353,48 @@ struct FBlackListEntry
 	// @return true:yes, inform used, false otherwise
 	bool Test(const FGPUDriverInfo& Info) const
 	{
-		return CompareStringOp(*DriverVersionString, *Info.GetUnifiedDriverVersion());
+		if (IsValid())
+		{
+			if (!DriverVersionString.IsEmpty())
+			{
+				return CompareStringOp(*DriverVersionString, *Info.GetUnifiedDriverVersion());
+			}
+			else
+			{
+				FString TempDay, TempMonth, TempMonthDay, TempYear;
+
+				// Trim 1-2 character operator then reformat for comparison
+				int32 OpLength = FChar::IsDigit((*DriverDateString)[1]) ? 1 : 2;
+
+				const TCHAR* DriverDateStringChars = *DriverDateString;
+				EComparisonOp Op = ParseComparisonOp(DriverDateStringChars);
+				if (Op == ECO_Unknown)
+				{
+					Op = ECO_Equal;
+				}
+
+				FString DriverDateComparisonOp = DriverDateString.Left(OpLength);
+				FString TrimmedDriverDateString = DriverDateString.RightChop(OpLength);
+				TrimmedDriverDateString.Split(TEXT("-"), &TempMonthDay, &TempYear, ESearchCase::IgnoreCase, ESearchDir::FromEnd);
+				TempMonthDay.Split(TEXT("-"), &TempMonth, &TempDay, ESearchCase::IgnoreCase, ESearchDir::FromEnd);
+				
+				FDateTime TestDate(FCString::Atoi(*TempYear), FCString::Atoi(*TempMonth), FCString::Atoi(*TempDay));
+
+				Info.DriverDate.Split(TEXT("-"), &TempMonthDay, &TempYear, ESearchCase::IgnoreCase, ESearchDir::FromEnd);
+				TempMonthDay.Split(TEXT("-"), &TempMonth, &TempDay, ESearchCase::IgnoreCase, ESearchDir::FromEnd);
+
+				FDateTime InfoDate(FCString::Atoi(*TempYear), FCString::Atoi(*TempMonth), FCString::Atoi(*TempDay));
+
+				return Compare(InfoDate, Op, TestDate);
+			}
+		}
+
+		return true;
 	}
 
 	bool IsValid() const
 	{
-		return !DriverVersionString.IsEmpty();
+		return DriverVersionString.Len() > 1 || DriverDateString.Len() > 1;
 	}
 
 	/**
@@ -366,7 +406,7 @@ struct FBlackListEntry
 		bool bLatestBlacklisted = false;
 		if (IsValid())
 		{
-			const TCHAR* DriverVersionTchar = *DriverVersionString;
+			const TCHAR* DriverVersionTchar = !DriverVersionString.IsEmpty() ? *DriverVersionString : *DriverDateString;
 			EComparisonOp ComparisonOp = ParseComparisonOp(DriverVersionTchar);
 			bLatestBlacklisted = (ComparisonOp == ECO_Larger) || (ComparisonOp == ECO_LargerThan);
 		}
