Fix useless static analysis warnings

diff --git a/Engine/Source/Runtime/Core/Private/HAL/MallocBinnedArena.cpp b/Engine/Source/Runtime/Core/Private/HAL/MallocBinnedArena.cpp
--- a/Engine/Source/Runtime/Core/Private/HAL/MallocBinnedArena.cpp
+++ b/Engine/Source/Runtime/Core/Private/HAL/MallocBinnedArena.cpp
@@ -839,7 +839,7 @@ bool FMallocBinnedArena::IsInternallyThreadSafe() const
 void* FMallocBinnedArena::MallocExternal(SIZE_T Size, uint32 Alignment)
 {
 	check(FMath::IsPowerOfTwo(Alignment));
-	checkf(DEFAULT_ALIGNMENT <= ArenaParams.MinimumAlignment, TEXT("DEFAULT_ALIGNMENT is assumed to be zero")); // used below
+	checkf(DEFAULT_ALIGNMENT <= ArenaParams.MinimumAlignment, TEXT("DEFAULT_ALIGNMENT is assumed to be zero")); //-V547
 
 	if (AdjustSmallBlockSizeForAlignment(Size, Alignment)) // there is some redundant work here... we already adjusted the size for alignment
 	{
@@ -940,7 +940,7 @@ void* FMallocBinnedArena::ReallocExternal(void* Ptr, SIZE_T NewSize, uint32 Alig
 		check(Ptr); // null is an OS allocation because it will not fall in our VM block
 		uint32 BlockSize = PoolIndexToBlockSize(PoolIndex);
 		if (
-			((NewSize <= BlockSize) & IsAligned(Ptr, Alignment)) && // one branch, not two
+			((NewSize <= BlockSize) & IsAligned(Ptr, Alignment)) && //-V792
 			(PoolIndex == 0 || NewSize > PoolIndexToBlockSize(PoolIndex - 1)))
 		{
 			return Ptr;
diff --git a/Engine/Source/Runtime/Core/Private/HAL/MallocBinnedGPU.cpp b/Engine/Source/Runtime/Core/Private/HAL/MallocBinnedGPU.cpp
--- a/Engine/Source/Runtime/Core/Private/HAL/MallocBinnedGPU.cpp
+++ b/Engine/Source/Runtime/Core/Private/HAL/MallocBinnedGPU.cpp
@@ -850,7 +850,7 @@ bool FMallocBinnedGPU::IsInternallyThreadSafe() const
 void* FMallocBinnedGPU::MallocExternal(SIZE_T Size, uint32 Alignment)
 {
 	check(FMath::IsPowerOfTwo(Alignment));
-	checkf(DEFAULT_ALIGNMENT <= ArenaParams.MinimumAlignment, TEXT("DEFAULT_ALIGNMENT is assumed to be zero")); // used below
+	checkf(DEFAULT_ALIGNMENT <= ArenaParams.MinimumAlignment, TEXT("DEFAULT_ALIGNMENT is assumed to be zero")); //-V547
 
 	if (AdjustSmallBlockSizeForAlignment(Size, Alignment)) // there is some redundant work here... we already adjusted the size for alignment
 	{
diff --git a/Engine/Source/Runtime/Core/Public/HAL/MallocBinnedArena.h b/Engine/Source/Runtime/Core/Public/HAL/MallocBinnedArena.h
--- a/Engine/Source/Runtime/Core/Public/HAL/MallocBinnedArena.h
+++ b/Engine/Source/Runtime/Core/Public/HAL/MallocBinnedArena.h
@@ -500,7 +500,7 @@ public:
 				{
 					// Reallocate to a smaller/bigger pool if necessary
 					BlockSize = PoolIndexToBlockSize(PoolIndex);
-					if ((!!NewSize) & (NewSize <= BlockSize) & IsAligned(BlockSize, Alignment) & ((!PoolIndex) | (NewSize > PoolIndexToBlockSize(PoolIndex - 1))))
+					if ((!!NewSize) & (NewSize <= BlockSize) & IsAligned(BlockSize, Alignment) & ((!PoolIndex) | (NewSize > PoolIndexToBlockSize(PoolIndex - 1)))) //-V792
 					{
 						return Ptr;
 					}
