fix non-unity issue

diff --git a/Engine/Source/Runtime/Core/Private/HAL/MallocBinnedArena.cpp b/Engine/Source/Runtime/Core/Private/HAL/MallocBinnedArena.cpp
--- a/Engine/Source/Runtime/Core/Private/HAL/MallocBinnedArena.cpp
+++ b/Engine/Source/Runtime/Core/Private/HAL/MallocBinnedArena.cpp
@@ -8,7 +8,6 @@
 
 #if PLATFORM_64BITS && PLATFORM_HAS_FPlatformVirtualMemoryBlock
 #include "Logging/LogMacros.h"
-#include "Misc/ScopeLock.h"
 #include "Templates/Function.h"
 #include "GenericPlatform/GenericPlatformProcess.h"
 #include "Stats/Stats.h"
diff --git a/Engine/Source/Runtime/Core/Private/HAL/MallocBinnedGPU.cpp b/Engine/Source/Runtime/Core/Private/HAL/MallocBinnedGPU.cpp
--- a/Engine/Source/Runtime/Core/Private/HAL/MallocBinnedGPU.cpp
+++ b/Engine/Source/Runtime/Core/Private/HAL/MallocBinnedGPU.cpp
@@ -8,7 +8,6 @@
 
 #if PLATFORM_64BITS && PLATFORM_HAS_FPlatformVirtualMemoryBlock
 #include "Logging/LogMacros.h"
-#include "Misc/ScopeLock.h"
 #include "Templates/Function.h"
 #include "GenericPlatform/GenericPlatformProcess.h"
 #include "Stats/Stats.h"
diff --git a/Engine/Source/Runtime/Core/Public/HAL/MallocBinnedArena.h b/Engine/Source/Runtime/Core/Public/HAL/MallocBinnedArena.h
--- a/Engine/Source/Runtime/Core/Public/HAL/MallocBinnedArena.h
+++ b/Engine/Source/Runtime/Core/Public/HAL/MallocBinnedArena.h
@@ -17,6 +17,8 @@
 #include "HAL/PlatformMath.h"
 #include "HAL/LowLevelMemTracker.h"
 #include "HAL/MallocBinnedCommon.h"
+#include "Misc/ScopeLock.h"
+
 
 #define BINNEDARENA_MAX_GMallocBinnedArenaMaxBundlesBeforeRecycle (8)
 
diff --git a/Engine/Source/Runtime/Core/Public/HAL/MallocBinnedGPU.h b/Engine/Source/Runtime/Core/Public/HAL/MallocBinnedGPU.h
--- a/Engine/Source/Runtime/Core/Public/HAL/MallocBinnedGPU.h
+++ b/Engine/Source/Runtime/Core/Public/HAL/MallocBinnedGPU.h
@@ -17,6 +17,8 @@
 #include "HAL/PlatformMath.h"
 #include "HAL/LowLevelMemTracker.h"
 #include "HAL/MallocBinnedCommon.h"
+#include "Misc/ScopeLock.h"
+
 
 #define BINNEDGPU_MAX_GMallocBinnedGPUMaxBundlesBeforeRecycle (8)
 
