Create files with 664 permissions over 644 for files created to unify to what other platforms are doing

diff --git a/Engine/Source/Runtime/Core/Private/Unix/UnixPlatformFile.cpp b/Engine/Source/Runtime/Core/Private/Unix/UnixPlatformFile.cpp
--- a/Engine/Source/Runtime/Core/Private/Unix/UnixPlatformFile.cpp
+++ b/Engine/Source/Runtime/Core/Private/Unix/UnixPlatformFile.cpp
@@ -917,7 +917,7 @@ IFileHandle* FUnixPlatformFile::OpenWrite(const TCHAR* Filename, bool bAppend, b
 	}
 
 	// Caveat: cannot specify O_TRUNC in flags, as this will corrupt the file which may be "locked" by other process. We will ftruncate() it once we "lock" it
-	int32 Handle = open(TCHAR_TO_UTF8(*NormalizeFilename(Filename, true)), Flags, S_IRUSR | S_IWUSR | S_IRGRP | S_IROTH);
+	int32 Handle = open(TCHAR_TO_UTF8(*NormalizeFilename(Filename, true)), Flags, S_IRUSR | S_IWUSR | S_IRGRP | S_IWGRP | S_IROTH);
 	if (Handle != -1)
 	{
 		// mimic Windows "exclusive write" behavior (we don't use FILE_SHARE_WRITE) by locking the file.
