Add begin and end analytics for game feature install requests.

diff --git a/Engine/Source/Runtime/Core/Public/GenericPlatform/GenericPlatformInstallBundleManager.h b/Engine/Source/Runtime/Core/Public/GenericPlatform/GenericPlatformInstallBundleManager.h
--- a/Engine/Source/Runtime/Core/Public/GenericPlatform/GenericPlatformInstallBundleManager.h
+++ b/Engine/Source/Runtime/Core/Public/GenericPlatform/GenericPlatformInstallBundleManager.h
@@ -8,6 +8,9 @@
 #include "Misc/ConfigCacheIni.h"
 #include "Logging/LogMacros.h"
 #include "Internationalization/Text.h"
+//#include "IAnalyticsProviderET.h"
+
+class IAnalyticsProviderET;
 
 enum class EInstallBundleModuleInitResult : int
 {
@@ -314,6 +317,8 @@ public:
 	virtual bool IsNullInterface() const = 0;
 
 	virtual void SetErrorSimulationCommands(const FString& CommandLine) {}
+
+	virtual TSharedPtr<IAnalyticsProviderET> GetAnalyticsProvider() const { return TSharedPtr<IAnalyticsProviderET>(); }
 };
 
 class IPlatformInstallBundleManagerModule : public IModuleInterface
