Fix for backwards RemoveIf() implementation.

diff --git a/Engine/Source/Runtime/Core/Public/Algo/RemoveIf.h b/Engine/Source/Runtime/Core/Public/Algo/RemoveIf.h
--- a/Engine/Source/Runtime/Core/Public/Algo/RemoveIf.h
+++ b/Engine/Source/Runtime/Core/Public/Algo/RemoveIf.h
@@ -34,7 +34,7 @@ namespace Algo
 					return IterStart - First;
 				}
 
-				if (!Invoke(Pred, *IterStart))
+				if (Invoke(Pred, *IterStart))
 				{
 					break;
 				}
@@ -45,7 +45,7 @@ namespace Algo
 			// Skip removed elements at the end
 			for (;;)
 			{
-				if (Invoke(Pred, *(IterEnd - 1)))
+				if (!Invoke(Pred, *(IterEnd - 1)))
 				{
 					break;
 				}
@@ -91,7 +91,7 @@ namespace Algo
 				return IterStart - First;
 			}
 
-			if (!Invoke(Pred, *IterStart))
+			if (Invoke(Pred, *IterStart))
 			{
 				break;
 			}
@@ -109,7 +109,7 @@ namespace Algo
 				return IterStart - First;
 			}
 
-			if (Invoke(Pred, *IterKeep))
+			if (!Invoke(Pred, *IterKeep))
 			{
 				*IterStart++ = MoveTemp(*IterKeep);
 			}
