PR #5779: Add Build Support & Customization Options For DirectX Math on Windows Platform (Contributed by [REDACTED_PERSON])

#github https://github.com/EpicGames/UnrealEngine/pull/5779

diff --git a/Engine/Source/Programs/UnrealBuildTool/Configuration/TargetRules.cs b/Engine/Source/Programs/UnrealBuildTool/Configuration/TargetRules.cs
--- a/Engine/Source/Programs/UnrealBuildTool/Configuration/TargetRules.cs
+++ b/Engine/Source/Programs/UnrealBuildTool/Configuration/TargetRules.cs
@@ -544,6 +544,12 @@ namespace UnrealBuildTool
 		[XmlConfigFile(Category = "BuildConfiguration")]
 		public bool bUseDebugLiveCodingConsole = false;
 
+		/// <summary>
+		/// Whether to enable support for DirectX Math
+		/// </summary>
+		[RequiresUniqueBuildEnvironment]
+		public bool bWithDirectXMath = false;
+
 		/// <summary>
 		/// Whether to turn on logging for test/shipping builds.
 		/// </summary>
@@ -1883,6 +1889,11 @@ namespace UnrealBuildTool
 			get { return Inner.bUseDebugLiveCodingConsole; }
 		}
 
+		public bool bWithDirectXMath
+		{
+			get { return Inner.bWithDirectXMath; }
+		}
+
 		public bool bUseLoggingInShipping
 		{
 			get { return Inner.bUseLoggingInShipping; }
diff --git a/Engine/Source/Runtime/Core/Core.Build.cs b/Engine/Source/Runtime/Core/Core.Build.cs
--- a/Engine/Source/Runtime/Core/Core.Build.cs
+++ b/Engine/Source/Runtime/Core/Core.Build.cs
@@ -171,21 +171,11 @@ public class Core : ModuleRules
 
 		WhitelistRestrictedFolders.Add("Private/NoRedist");
 
-        if (Target.Platform == UnrealTargetPlatform.XboxOne)
+        if (Target.Platform == UnrealTargetPlatform.XboxOne ||
+            (Target.bWithDirectXMath && (Target.Platform == UnrealTargetPlatform.Win64 || Target.Platform == UnrealTargetPlatform.Win32)))
         {
             PublicDefinitions.Add("WITH_DIRECTXMATH=1");
         }
-        else if ((Target.Platform == UnrealTargetPlatform.Win64) ||
-                (Target.Platform == UnrealTargetPlatform.Win32))
-        {
-			// To enable this requires Win8 SDK
-            PublicDefinitions.Add("WITH_DIRECTXMATH=0");  // Enable to test on Win64/32.
-
-            //PublicDependencyModuleNames.AddRange(  // Enable to test on Win64/32.
-			//    new string[] {
-			//    "DirectXMath"
-            //});
-        }
         else
         {
             PublicDefinitions.Add("WITH_DIRECTXMATH=0");
diff --git a/Engine/Source/Runtime/Core/Public/Math/UnrealMathDirectX.h b/Engine/Source/Runtime/Core/Public/Math/UnrealMathDirectX.h
--- a/Engine/Source/Runtime/Core/Public/Math/UnrealMathDirectX.h
+++ b/Engine/Source/Runtime/Core/Public/Math/UnrealMathDirectX.h
@@ -2,12 +2,16 @@
 
 #pragma once
 
-#if __cplusplus_cli
+#ifdef __cplusplus_cli
 // there are compile issues with this file in managed mode, so use the FPU version
 #include "Math/UnrealMathFPU.h"
 #else
 
+#pragma warning( push )
+// We're in non-managed mode, ignore warnings about _M_CEE and _MANAGED for DirectXMath.h
+#pragma warning ( disable : 4668 )
 #include <DirectXMath.h>
+#pragma warning( pop )
 #include <DirectXPackedVector.h>
 
 /*=============================================================================
@@ -632,6 +636,21 @@ FORCEINLINE VectorRegister VectorCombineLow(const VectorRegister& Vec1, const Ve
 	return VectorShuffle(Vec1, Vec2, 0, 1, 0, 1);
 }
 
+/**
+ * These functions return a vector mask to indicate which components pass the comparison.
+ * Each component is 0xffffffff if it passes, 0x00000000 if it fails.
+ *
+ * @param Vec1			1st source vector
+ * @param Vec2			2nd source vector
+ * @return				Vector with a mask for each component.
+ */
+#define VectorMask_LT( Vec1, Vec2 )	_mm_cmplt_ps( Vec1, Vec2 )
+#define VectorMask_LE( Vec1, Vec2 )	_mm_cmple_ps( Vec1, Vec2 )
+#define VectorMask_GT( Vec1, Vec2 )	DirectX::XMVectorGreater( Vec1, Vec2 )
+#define VectorMask_GE( Vec1, Vec2 )	DirectX::XMVectorGreaterOrEqual( Vec1, Vec2 )
+#define VectorMask_EQ( Vec1, Vec2 )	DirectX::XMVectorEqual( Vec1, Vec2 )
+#define VectorMask_NE( Vec1, Vec2 )	DirectX::XMVectorNotEqual( Vec1, Vec2 )
+
 /**
  * Merges the XYZ components of one vector with the W component of another vector and returns the result.
  *
