Copying //UE4/Dev-Core to Dev-Main (//UE4/Dev-Main)

diff --git a/Engine/Source/Developer/TraceServices/Private/Common/Utils.h b/Engine/Source/Developer/TraceServices/Private/Common/Utils.h
new file mode 100644
--- /dev/null
+++ b/Engine/Source/Developer/TraceServices/Private/Common/Utils.h
@@ -0,0 +1,29 @@
+// Copyright 1998-2019 Epic Games, Inc. All Rights Reserved.
+
+#pragma once
+
+#include "CoreTypes.h"
+
+struct FTraceAnalyzerUtils
+{
+	static uint64 Decode7bit(const uint8*& BufferPtr)
+	{
+		uint64 Value = 0;
+		uint64 ByteIndex = 0;
+		bool HasMoreBytes;
+		do
+		{
+			uint8 ByteValue = *BufferPtr++;
+			HasMoreBytes = ByteValue & 0x80;
+			Value |= uint64(ByteValue & 0x7f) << (ByteIndex * 7);
+			++ByteIndex;
+		} while (HasMoreBytes);
+		return Value;
+	}
+
+	static int64 DecodeZigZag(const uint8*& BufferPtr)
+	{
+		uint64 Z = Decode7bit(BufferPtr);
+		return (Z & 1) ? (Z >> 1) ^ -1 : (Z >> 1);
+	}
+};
diff --git a/Engine/Source/Developer/TraceServices/Private/Model/Bookmarks.cpp b/Engine/Source/Developer/TraceServices/Private/Model/Bookmarks.cpp
new file mode 100644
--- /dev/null
+++ b/Engine/Source/Developer/TraceServices/Private/Model/Bookmarks.cpp
@@ -0,0 +1,86 @@
+// Copyright 1998-2019 Epic Games, Inc. All Rights Reserved.
+
+#include "TraceServices/Model/Bookmarks.h"
+#include "Model/BookmarksPrivate.h"
+#include "AnalysisServicePrivate.h"
+#include "Common/FormatArgs.h"
+
+namespace Trace
+{
+
+const FName FBookmarkProvider::ProviderName("BookmarkProvider");
+
+FBookmarkProvider::FBookmarkProvider(IAnalysisSession& InSession)
+	: Session(InSession)
+{
+
+}
+
+FBookmarkSpec& FBookmarkProvider::GetSpec(uint64 BookmarkPoint)
+{
+	Session.WriteAccessCheck();
+	if (SpecMap.Contains(BookmarkPoint))
+	{
+		return *SpecMap[BookmarkPoint].Get();
+	}
+	else
+	{
+		TSharedPtr<FBookmarkSpec> Spec = MakeShared<FBookmarkSpec>();
+		SpecMap.Add(BookmarkPoint, Spec);
+		return *Spec.Get();
+	}
+}
+
+void FBookmarkProvider::AppendBookmark(double Time, uint64 BookmarkPoint, const uint8* FormatArgs)
+{
+	Session.WriteAccessCheck();
+	FBookmarkSpec& Spec = GetSpec(BookmarkPoint);
+	TSharedRef<FBookmarkInternal> Bookmark = MakeShared<FBookmarkInternal>();
+	Bookmark->Time = Time;
+	FFormatArgsHelper::Format(FormatBuffer, FormatBufferSize - 1, Spec.FormatString, FormatArgs);
+	Bookmark->Text = Session.StoreString(FormatBuffer);
+	Bookmarks.Add(Bookmark);
+	Session.UpdateDurationSeconds(Time);
+}
+
+void FBookmarkProvider::EnumerateBookmarks(double IntervalStart, double IntervalEnd, TFunctionRef<void(const FBookmark &)> Callback) const
+{
+	Session.ReadAccessCheck();
+	if (IntervalStart > IntervalEnd)
+	{
+		return;
+	}
+	uint64 FirstBookmarkIndex = Algo::LowerBoundBy(Bookmarks, IntervalStart, [](const TSharedRef<FBookmarkInternal>& B)
+	{
+		return B->Time;
+	});
+	uint64 BookmarkCount = Bookmarks.Num();
+	if (FirstBookmarkIndex >= BookmarkCount)
+	{
+		return;
+	}
+	uint64 LastBookmarkIndex = Algo::UpperBoundBy(Bookmarks, IntervalEnd, [](const TSharedRef<FBookmarkInternal>& B)
+	{
+		return B->Time;
+	});
+	if (LastBookmarkIndex == 0)
+	{
+		return;
+	}
+	--LastBookmarkIndex;
+	for (uint64 Index = FirstBookmarkIndex; Index <= LastBookmarkIndex; ++Index)
+	{
+		const FBookmarkInternal& InternalBookmark = Bookmarks[Index].Get();
+		FBookmark Bookmark;
+		Bookmark.Time = InternalBookmark.Time;
+		Bookmark.Text = InternalBookmark.Text;
+		Callback(Bookmark);
+	}
+}
+
+const IBookmarkProvider& ReadBookmarkProvider(const IAnalysisSession& Session)
+{
+	return *Session.ReadProvider<IBookmarkProvider>(FBookmarkProvider::ProviderName);
+}
+
+}
\ No newline at end of file
diff --git a/Engine/Source/Developer/TraceServices/Private/Model/Counters.cpp b/Engine/Source/Developer/TraceServices/Private/Model/Counters.cpp
new file mode 100644
--- /dev/null
+++ b/Engine/Source/Developer/TraceServices/Private/Model/Counters.cpp
@@ -0,0 +1,188 @@
+// Copyright 1998-2019 Epic Games, Inc. All Rights Reserved.
+
+#include "Model/Counters.h"
+#include "AnalysisServicePrivate.h"
+
+namespace Trace
+{
+
+FCounterInternal::FCounterInternal(ILinearAllocator& Allocator, uint32 InId, const TCHAR* InName, const TCHAR* InDescription, ECounterDisplayHint InDisplayHint)
+	: Name(InName)
+	, Description(InDescription)
+	, Id(InId)
+	, DisplayHint(InDisplayHint)
+	, Timestamps(Allocator, 1024)
+	, Ops(Allocator, 1024)
+	, IntOpArguments(Allocator, 1024)
+	, DoubleOpArguments(Allocator, 1024)
+{
+
+}
+
+template<typename CounterType, typename EnumerationType>
+static void EnumerateCounterValuesInternal(const TPagedArray<double>& Timestamps, const TPagedArray<ECounterOpType>& Ops, const TPagedArray<CounterType>& OpArgs, double IntervalStart, double IntervalEnd, TFunctionRef<void(double, EnumerationType)> Callback)
+{
+	CounterType Value = 0;
+	auto TimeIterator = Timestamps.GetIteratorFromItem(0);
+	auto OpIterator = Ops.GetIteratorFromItem(0);
+	auto OpArgIterator = OpArgs.GetIteratorFromItem(0);
+	const double* Time = TimeIterator.GetCurrentItem();
+	const ECounterOpType* Op = OpIterator.GetCurrentItem();
+	const CounterType* OpArg = OpArgIterator.GetCurrentItem();
+	
+	double PrevTime = 0.0;
+	CounterType PrevValue = CounterType();
+	bool IsFirstValue = true;
+	bool IsFirstIncludedValue = true;
+	while (Time)
+	{
+		switch (*Op)
+		{
+		case CounterOpType_Add:
+			Value += *OpArg;
+			break;
+		case CounterOpType_Set:
+			Value = *OpArg;
+			break;
+		}
+
+		if (*Time >= IntervalStart)
+		{
+			if (IsFirstIncludedValue && !IsFirstValue)
+			{
+				Callback(PrevTime, PrevValue);
+			}
+			PrevTime = *Time;
+			PrevValue = EnumerationType(Value);
+			Callback(PrevTime, PrevValue);
+			IsFirstIncludedValue = false;
+		}
+		if (*Time > IntervalEnd)
+		{
+			break;
+		}
+
+		IsFirstValue = false;
+
+		Time = TimeIterator.NextItem();
+		Op = OpIterator.NextItem();
+		OpArg = OpArgIterator.NextItem();
+	}
+}
+
+void FCounterInternal::EnumerateValues(double IntervalStart, double IntervalEnd, TFunctionRef<void(double, int64)> Callback) const
+{
+	if (DisplayHint == CounterDisplayHint_FloatingPoint)
+	{
+		EnumerateCounterValuesInternal<double, int64>(Timestamps, Ops, DoubleOpArguments, IntervalStart, IntervalEnd, Callback);
+	}
+	else
+	{
+		EnumerateCounterValuesInternal<int64, int64>(Timestamps, Ops, IntOpArguments, IntervalStart, IntervalEnd, Callback);
+	}
+}
+
+void FCounterInternal::EnumerateFloatValues(double IntervalStart, double IntervalEnd, TFunctionRef<void(double, double)> Callback) const
+{
+	if (DisplayHint == CounterDisplayHint_FloatingPoint)
+	{
+		EnumerateCounterValuesInternal<double, double>(Timestamps, Ops, DoubleOpArguments, IntervalStart, IntervalEnd, Callback);
+	}
+	else
+	{
+		EnumerateCounterValuesInternal<int64, double>(Timestamps, Ops, IntOpArguments, IntervalStart, IntervalEnd, Callback);
+	}
+}
+
+FCounterProvider::FCounterProvider(IAnalysisSession& InSession)
+	: Session(InSession)
+{
+
+}
+
+FCounterProvider::~FCounterProvider()
+{
+	for (FCounterInternal* Counter : Counters)
+	{
+		delete Counter;
+	}
+}
+
+void FCounterProvider::EnumerateCounters(TFunctionRef<void(const ICounter &)> Callback) const
+{
+	uint32 Id = 0;
+	for (FCounterInternal* Counter : Counters)
+	{
+		Callback(*Counter);
+	}
+}
+
+FCounterInternal* FCounterProvider::CreateCounter(const TCHAR* Name, const TCHAR* Description, ECounterDisplayHint DisplayHint)
+{
+	FCounterInternal* Counter = new FCounterInternal(Session.GetLinearAllocator(), Counters.Num(), Name, Description, DisplayHint);
+	Counters.Add(Counter);
+	return Counter;
+}
+
+template<typename T>
+static void InsertOp(FCounterInternal& Counter, double Time, ECounterOpType Type, T Arg)
+{
+	TPagedArray<double>& Timestamps = Counter.EditTimestamps();
+	uint64 InsertionIndex;
+	if (Timestamps.Num() == 0 || Timestamps[Timestamps.Num() - 1] <= Time)
+	{
+		InsertionIndex = Timestamps.Num();
+	}
+	else
+	{
+		auto TimestampIterator = Timestamps.GetIteratorFromItem(Timestamps.Num() - 1);
+		auto CurrentPage = TimestampIterator.GetCurrentPage();
+		while (CurrentPage && *GetFirstItem(*CurrentPage) > Time)
+		{
+			CurrentPage = TimestampIterator.PrevPage();
+		}
+		if (!CurrentPage)
+		{
+			InsertionIndex = 0;
+		}
+		else
+		{
+			uint64 PageInsertionIndex = Algo::LowerBound(*CurrentPage, Time);
+			InsertionIndex = TimestampIterator.GetCurrentPageIndex() * Timestamps.GetPageSize() + PageInsertionIndex;
+		}
+	}
+	Timestamps.Insert(InsertionIndex) = Time;
+	Counter.EditOps().Insert(InsertionIndex) = Type;
+	if (Counter.GetDisplayHint() == CounterDisplayHint_FloatingPoint)
+	{
+		Counter.EditDoubleOpArguments().Insert(InsertionIndex) = double(Arg);
+	}
+	else
+	{
+		Counter.EditIntOpArguments().Insert(InsertionIndex) = int64(Arg);
+	}
+}
+
+void FCounterProvider::Add(FCounterInternal& Counter, double Time, int64 Value)
+{
+	InsertOp(Counter, Time, CounterOpType_Add, Value);
+}
+
+void FCounterProvider::Add(FCounterInternal& Counter, double Time, double Value)
+{
+	InsertOp(Counter, Time, CounterOpType_Add, Value);
+}
+
+void FCounterProvider::Set(FCounterInternal& Counter, double Time, int64 Value)
+{
+	InsertOp(Counter, Time, CounterOpType_Set, Value);
+}
+
+void FCounterProvider::Set(FCounterInternal& Counter, double Time, double Value)
+{
+	InsertOp(Counter, Time, CounterOpType_Set, Value);
+}
+
+
+
+}
\ No newline at end of file
diff --git a/Engine/Source/Developer/TraceServices/Private/Model/Counters.h b/Engine/Source/Developer/TraceServices/Private/Model/Counters.h
new file mode 100644
--- /dev/null
+++ b/Engine/Source/Developer/TraceServices/Private/Model/Counters.h
@@ -0,0 +1,72 @@
+// Copyright 1998-2019 Epic Games, Inc. All Rights Reserved.
+
+#pragma once
+
+#include "TraceServices/Model/AnalysisSession.h"
+#include "Templates/SharedPointer.h"
+#include "Common/PagedArray.h"
+#include "TraceServices/Model/Stats.h"
+
+namespace Trace
+{
+
+class FAnalysisSessionLock;
+
+enum ECounterOpType : uint8
+{
+	CounterOpType_Add,
+	CounterOpType_Set,
+};
+
+class FCounterInternal
+	: public ICounter
+{
+public:
+	FCounterInternal(ILinearAllocator& Allocator, uint32 Id, const TCHAR* Name, const TCHAR* Description, ECounterDisplayHint DisplayHint);
+	virtual const TCHAR* GetName() const override { return *Name; }
+	virtual uint32 GetId() const override { return Id; }
+	virtual ECounterDisplayHint GetDisplayHint() const { return DisplayHint; }
+	virtual void EnumerateValues(double IntervalStart, double IntervalEnd, TFunctionRef<void(double, int64)> Callback) const override;
+	virtual void EnumerateFloatValues(double IntervalStart, double IntervalEnd, TFunctionRef<void(double, double)> Callback) const override;
+
+	const TPagedArray<int64>& ReadIntOpArguments() const { return IntOpArguments; }
+	const TPagedArray<double>& ReadDoubleOpArguments() const { return DoubleOpArguments; }
+	TPagedArray<double>& EditTimestamps() { return Timestamps; }
+	TPagedArray<ECounterOpType>& EditOps() { return Ops; }
+	TPagedArray<int64>& EditIntOpArguments() { return IntOpArguments; }
+	TPagedArray<double>& EditDoubleOpArguments() { return DoubleOpArguments; }
+
+private:
+	friend class FCountersProvider;
+
+	FString Name;
+	FString Description;
+	uint32 Id;
+	ECounterDisplayHint DisplayHint;
+	TPagedArray<double> Timestamps;
+	TPagedArray<ECounterOpType> Ops;
+	TPagedArray<int64> IntOpArguments;
+	TPagedArray<double> DoubleOpArguments;
+};
+
+class FCounterProvider
+	: public ICounterProvider
+{
+public:
+	FCounterProvider(IAnalysisSession& Session);
+	virtual ~FCounterProvider();
+	virtual uint64 GetCounterCount() const override { return Counters.Num(); }
+	virtual void EnumerateCounters(TFunctionRef<void(const ICounter&)> Callback) const override;
+
+	FCounterInternal* CreateCounter(const TCHAR* Name, const TCHAR* Description, ECounterDisplayHint DisplayHint);
+	void Add(FCounterInternal& Counter, double Time, int64 Value);
+	void Add(FCounterInternal& Counter, double Time, double Value);
+	void Set(FCounterInternal& Counter, double Time, int64 Value);
+	void Set(FCounterInternal& Counter, double Time, double Value);
+
+private:
+	IAnalysisSession& Session;
+	TArray<FCounterInternal*> Counters;
+};
+
+}
diff --git a/Engine/Source/Developer/TraceServices/Private/Model/FileActivity.cpp b/Engine/Source/Developer/TraceServices/Private/Model/FileActivity.cpp
new file mode 100644
--- /dev/null
+++ b/Engine/Source/Developer/TraceServices/Private/Model/FileActivity.cpp
@@ -0,0 +1,63 @@
+// Copyright 1998-2019 Epic Games, Inc. All Rights Reserved.
+
+#include "Model/FileActivity.h"
+#include "Model/IntervalTimeline.h"
+#include "AnalysisServicePrivate.h"
+
+namespace Trace
+{
+
+FFileActivityProvider::FFileActivityProvider(IAnalysisSession& InSession)
+	: Session(InSession)
+	, Files(InSession.GetLinearAllocator(), 1024)
+{
+}
+
+void FFileActivityProvider::EnumerateFileActivity(TFunctionRef<bool(const FFileInfo&, const ITimeline<FFileActivity>&)> Callback) const
+{
+	for (uint64 FileIndex = 0, FileCount = Files.Num(); FileIndex < FileCount; ++FileIndex)
+	{
+		const FFileInfoInternal& FileInfoInternal = Files[FileIndex];
+		FFileInfo FileInfo;
+		FileInfo.Id = FileIndex;
+		FileInfo.Path = *FileInfoInternal.Path;
+		if (!Callback(FileInfo, *FileInfoInternal.ActivityTimeline))
+		{
+			return;
+		}
+	}
+}
+
+uint32 FFileActivityProvider::GetFileIndex(const TCHAR* Path)
+{
+	uint32 FileIndex = Files.Num();
+	FFileInfoInternal& FileInfo = Files.PushBack();
+	FileInfo.Path = Path;
+	FileInfo.ActivityTimeline = MakeShared<TimelineInternal>(Session.GetLinearAllocator());
+	return FileIndex;
+}
+
+uint64 FFileActivityProvider::BeginActivity(uint32 FileIndex, EFileActivityType Type, uint64 Offset, uint64 Size, double Time)
+{
+	FFileInfoInternal& FileInfo = Files[FileIndex];
+	FFileActivity FileActivity;
+	FileActivity.ActivityType = Type;
+	FileActivity.Offset = Offset;
+	FileActivity.Size = Size;
+	FileActivity.Failed = false;
+	return FileInfo.ActivityTimeline->AppendBeginEvent(Time, FileActivity);
+}
+
+void FFileActivityProvider::EndActivity(uint32 FileIndex, uint64 ActivityIndex, double Time, bool Failed)
+{
+	FFileInfoInternal& FileInfo = Files[FileIndex];
+	FFileActivity& Activity = FileInfo.ActivityTimeline->EndEvent(ActivityIndex, Time);
+	Activity.Failed = Failed;
+}
+
+const TCHAR* FFileActivityProvider::GetFilePath(uint32 FileIndex) const
+{
+	return *Files[FileIndex].Path;
+}
+
+}
\ No newline at end of file
diff --git a/Engine/Source/Developer/TraceServices/Private/Model/FileActivity.h b/Engine/Source/Developer/TraceServices/Private/Model/FileActivity.h
new file mode 100644
--- /dev/null
+++ b/Engine/Source/Developer/TraceServices/Private/Model/FileActivity.h
@@ -0,0 +1,48 @@
+// Copyright 1998-2019 Epic Games, Inc. All Rights Reserved.
+
+#pragma once
+
+#include "TraceServices/Model/AnalysisSession.h"
+#include "Common/PagedArray.h"
+#include "Model/IntervalTimeline.h"
+#include "TraceServices/Model/LoadTimeProfiler.h"
+#include "Containers/UnrealString.h"
+#include "Templates/SharedPointer.h"
+
+namespace Trace
+{
+
+class FAnalysisSessionLock;
+
+class FFileActivityProvider
+	: public IFileActivityProvider
+{
+public:
+	struct FTimelineSettings
+	{
+		enum
+		{
+			EventsPerPage = 128
+		};
+	};
+	typedef TIntervalTimeline<FFileActivity, FTimelineSettings> TimelineInternal;
+
+	FFileActivityProvider(IAnalysisSession& Session);
+	virtual void EnumerateFileActivity(TFunctionRef<bool(const FFileInfo&, const Timeline&)> Callback) const override;
+	uint32 GetFileIndex(const TCHAR* Path);
+	uint64 BeginActivity(uint32 FileIndex, EFileActivityType Type, uint64 Offset, uint64 Size, double Time);
+	void EndActivity(uint32 FileIndex, uint64 ActivityIndex, double Time, bool Failed);
+	const TCHAR* GetFilePath(uint32 FileIndex) const;
+
+private:
+	struct FFileInfoInternal
+	{
+		FString Path;
+		TSharedPtr<TimelineInternal> ActivityTimeline;
+	};
+
+	IAnalysisSession& Session;
+	TPagedArray<FFileInfoInternal> Files;
+};
+
+}
diff --git a/Engine/Source/Developer/TraceServices/Private/Model/Frames.cpp b/Engine/Source/Developer/TraceServices/Private/Model/Frames.cpp
new file mode 100644
--- /dev/null
+++ b/Engine/Source/Developer/TraceServices/Private/Model/Frames.cpp
@@ -0,0 +1,71 @@
+// Copyright 1998-2019 Epic Games, Inc. All Rights Reserved.
+
+#include "TraceServices/Model/Frames.h"
+#include "Model/FramesPrivate.h"
+#include "AnalysisServicePrivate.h"
+#include <limits>
+
+namespace Trace
+{
+
+const FName FFrameProvider::ProviderName("FrameProvider");
+
+FFrameProvider::FFrameProvider(IAnalysisSession& InSession)
+	: Session(InSession)
+{
+	for (int32 FrameType = 0; FrameType < TraceFrameType_Count; ++FrameType)
+	{
+		Frames.Emplace(Session.GetLinearAllocator(), 65536);
+	}
+}
+
+uint64 FFrameProvider::GetFrameCount(ETraceFrameType FrameType) const
+{
+	Session.ReadAccessCheck();
+	return Frames[FrameType].Num();
+}
+
+void FFrameProvider::EnumerateFrames(ETraceFrameType FrameType, uint64 Start, uint64 End, TFunctionRef<void(const FFrame&)> Callback) const
+{
+	Session.ReadAccessCheck();
+
+	End = FMath::Min(End, Frames[FrameType].Num());
+	if (Start >= End)
+	{
+		return;
+	}
+	auto Iterator = Frames[FrameType].GetIteratorFromItem(Start);
+	const FFrame* Frame = Iterator.GetCurrentItem();
+	for (uint64 Index = Start; Index < End; ++Index)
+	{
+		Callback(*Frame);
+		Frame = Iterator.NextItem();
+	}
+}
+
+void FFrameProvider::BeginFrame(ETraceFrameType FrameType, double Time)
+{
+	Session.WriteAccessCheck();
+	
+	uint64 Index = Frames[FrameType].Num();
+	FFrame& Frame = Frames[FrameType].PushBack();
+	Frame.StartTime = Time;
+	Frame.EndTime = std::numeric_limits<double>::infinity();
+	Frame.Index = Index;
+	Session.UpdateDurationSeconds(Time);
+}
+
+void FFrameProvider::EndFrame(ETraceFrameType FrameType, double Time)
+{
+	Session.WriteAccessCheck();
+	FFrame& Frame = Frames[FrameType][Frames[FrameType].Num() - 1];
+	Frame.EndTime = Time;
+	Session.UpdateDurationSeconds(Time);
+}
+
+const IFrameProvider& ReadFrameProvider(const IAnalysisSession& Session)
+{
+	return *Session.ReadProvider<IFrameProvider>(FFrameProvider::ProviderName);
+}
+
+}
\ No newline at end of file
diff --git a/Engine/Source/Developer/TraceServices/Private/Model/IntervalTimeline.h b/Engine/Source/Developer/TraceServices/Private/Model/IntervalTimeline.h
new file mode 100644
--- /dev/null
+++ b/Engine/Source/Developer/TraceServices/Private/Model/IntervalTimeline.h
@@ -0,0 +1,136 @@
+// Copyright 1998-2019 Epic Games, Inc. All Rights Reserved.
+
+#pragma once
+
+#include "TraceServices/Containers/Timelines.h"
+#include "Common/PagedArray.h"
+#include "TraceServices/Containers/Allocators.h"
+#include <limits>
+
+namespace Trace
+{
+
+struct FIntervalTimelineDefaultSettings
+{
+	enum
+	{
+		EventsPerPage = 65536
+	};
+};
+
+template<typename EventType, typename SettingsType = FIntervalTimelineDefaultSettings>
+class TIntervalTimeline
+	: public ITimeline<EventType>
+{
+public:
+	TIntervalTimeline(ILinearAllocator& Allocator)
+		: Events(Allocator, SettingsType::EventsPerPage)
+	{
+	}
+
+	virtual ~TIntervalTimeline() = default;
+	virtual uint64 GetModCount() const override { return ModCount; }
+	virtual uint64 GetEventCount() const override { return Events.Num(); }
+	
+	virtual void EnumerateEventsDownSampled(double IntervalStart, double IntervalEnd, double Resolution, TFunctionRef<void(bool, double, const EventType&)> Callback) const override { check(false); }
+	virtual void EnumerateEventsDownSampled(double IntervalStart, double IntervalEnd, double Resolution, TFunctionRef<void(double, double, uint32, const EventType&)> Callback) const override { check(false); }
+	virtual void EnumerateEvents(double IntervalStart, double IntervalEnd, TFunctionRef<void(bool, double, const EventType&)> Callback) const override { check(false); }
+	
+	virtual void EnumerateEvents(double IntervalStart, double IntervalEnd, TFunctionRef<void(double, double, uint32, const EventType&)> Callback) const override
+	{
+		auto EventIterator = Events.GetIteratorFromPage(0);
+		const FEventInternal* Event = EventIterator.GetCurrentItem();
+		while (Event)
+		{
+			const FEventPage* Page = EventIterator.GetCurrentPage();
+			if (Page->EndTime < IntervalStart || IntervalEnd < Page->BeginTime)
+			{
+				if (!EventIterator.NextPage())
+				{
+					return;
+				}
+				Event = EventIterator.GetCurrentItem();
+			}
+			else
+			{
+				if (!(Event->EndTime < IntervalStart || IntervalEnd < Event->StartTime))
+				{
+					Callback(Event->StartTime, Event->EndTime, 0, Event->Event);
+				}
+				Event = EventIterator.NextItem();
+			}
+		}
+	}
+	
+	uint64 AppendBeginEvent(double StartTime, const EventType& Event)
+	{
+		//check(StartTime >= LastTime);
+		//LastTime = StartTime;
+
+		uint64 Index = Events.Num();
+		FEventInternal& EventInternal = Events.PushBack();
+		EventInternal.StartTime = StartTime;
+		EventInternal.EndTime = std::numeric_limits<double>::infinity();
+		EventInternal.Event = Event;
+
+		FEventPage* LastPage = Events.GetLastPage();
+		if (LastPage != CurrentPage)
+		{
+			CurrentPage = LastPage;
+			LastPage->BeginTime = StartTime;
+			LastPage->EndTime = StartTime;
+			LastPage->BeginEventIndex = Index;
+		}
+		else
+		{
+			LastPage->BeginTime = FMath::Min(LastPage->BeginTime, StartTime);
+			LastPage->EndTime = FMath::Max(LastPage->EndTime, StartTime);
+		}
+
+		LastPage->EndEventIndex = Index + 1;
+
+		++ModCount;
+
+		return Index;
+	}
+
+	EventType& EndEvent(uint64 EventIndex, double EndTime)
+	{
+		//check(EndTime >= LastTime);
+		//LastTime = EndTime;
+		FEventInternal& EventInternal = Events[EventIndex];
+		check(EventIndex < Events.Num());
+		EventInternal.EndTime = EndTime;
+		FEventPage* Page = Events.GetItemPage(EventIndex);
+		Page->EndTime = FMath::Max(Page->EndTime, EndTime);
+
+		++ModCount;
+
+		return EventInternal.Event;
+	}
+
+private:
+	struct FEventInternal
+	{
+		double StartTime;
+		double EndTime;
+		EventType Event;
+	};
+
+	struct FEventPage
+	{
+		double BeginTime = 0.0;
+		double EndTime = 0.0;
+		uint64 BeginEventIndex = 0;
+		uint64 EndEventIndex = 0;
+		FEventInternal* Items = nullptr;
+		uint64 Count = 0;
+	};
+
+	TPagedArray<FEventInternal, FEventPage> Events;
+	FEventPage* CurrentPage = nullptr;
+	double LastTime = 0.0;
+	uint64 ModCount = 0;
+};
+
+}
\ No newline at end of file
diff --git a/Engine/Source/Developer/TraceServices/Private/Model/LoadTimeProfiler.cpp b/Engine/Source/Developer/TraceServices/Private/Model/LoadTimeProfiler.cpp
new file mode 100644
--- /dev/null
+++ b/Engine/Source/Developer/TraceServices/Private/Model/LoadTimeProfiler.cpp
@@ -0,0 +1,175 @@
+// Copyright 1998-2019 Epic Games, Inc. All Rights Reserved.
+
+#include "TraceServices/Model/LoadTimeProfiler.h"
+#include "Model/LoadTimeProfilerPrivate.h"
+#include "AnalysisServicePrivate.h"
+#include "Common/TimelineStatistics.h"
+
+namespace Trace
+{
+
+FLoadTimeProfilerProvider::FLoadTimeProfilerProvider(IAnalysisSession& InSession)
+	: Session(InSession)
+	, ClassInfos(Session.GetLinearAllocator(), 4096)
+	, Packages(Session.GetLinearAllocator(), 4096)
+	, Exports(Session.GetLinearAllocator(), 4090)
+	, MainThreadCpuTimeline(MakeShared<CpuTimelineInternal>(Session.GetLinearAllocator()))
+	, AsyncLoadingThreadCpuTimeline(MakeShared<CpuTimelineInternal>(Session.GetLinearAllocator()))
+{
+	
+}
+
+void FLoadTimeProfilerProvider::EnumeratePackages(TFunctionRef<void(const FPackageInfo&)> Callback) const
+{
+	Session.ReadAccessCheck();
+
+	auto Iterator = Packages.GetIteratorFromItem(0);
+	const FPackageInfo* Package = Iterator.GetCurrentItem();
+	while (Package)
+	{
+		Callback(*Package);
+		Package = Iterator.NextItem();
+	}
+}
+
+void FLoadTimeProfilerProvider::ReadMainThreadCpuTimeline(TFunctionRef<void(const CpuTimeline &)> Callback) const
+{
+	Session.ReadAccessCheck();
+	Callback(*MainThreadCpuTimeline);
+}
+
+void FLoadTimeProfilerProvider::ReadAsyncLoadingThreadCpuTimeline(TFunctionRef<void(const CpuTimeline &)> Callback) const
+{
+	Session.ReadAccessCheck();
+	Callback(*AsyncLoadingThreadCpuTimeline);
+}
+
+ITable<FLoadTimeProfilerAggregatedStats>* FLoadTimeProfilerProvider::CreateEventAggregation(double IntervalStart, double IntervalEnd) const
+{
+	TArray<const CpuTimelineInternal*> Timelines;
+	Timelines.Add(&MainThreadCpuTimeline.Get());
+	Timelines.Add(&AsyncLoadingThreadCpuTimeline.Get());
+
+	auto BucketMapper = [](const FLoadTimeProfilerCpuEvent& Event)
+	{
+		return Event.PackageEventType;
+	};
+	TMap<ELoadTimeProfilerPackageEventType, FAggregatedTimingStats> Aggregation;
+	FTimelineStatistics::CreateAggregation(Timelines, BucketMapper, IntervalStart, IntervalEnd, Aggregation);
+	TTable<FAggregatedStatsTableLayout>* Table = new TTable<FAggregatedStatsTableLayout>();
+	for (const auto& KV : Aggregation)
+	{
+		FLoadTimeProfilerAggregatedStats& Row = Table->AddRow();
+		switch (KV.Key)
+		{
+		case LoadTimeProfilerPackageEventType_CreateLinker:
+			Row.Name = TEXT("CreateLinker");
+			break;
+		case LoadTimeProfilerPackageEventType_FinishLinker:
+			Row.Name = TEXT("FinishLinker");
+			break;
+		case LoadTimeProfilerPackageEventType_StartImportPackages:
+			Row.Name = TEXT("StartImportPackages");
+			break;
+		case LoadTimeProfilerPackageEventType_SetupImports:
+			Row.Name = TEXT("SetupImports");
+			break;
+		case LoadTimeProfilerPackageEventType_SetupExports:
+			Row.Name = TEXT("SetupExports");
+			break;
+		case LoadTimeProfilerPackageEventType_ProcessImportsAndExports:
+			Row.Name = TEXT("ProcessImportsAndExports");
+			break;
+		case LoadTimeProfilerPackageEventType_ExportsDone:
+			Row.Name = TEXT("ExportsDone");
+			break;
+		case LoadTimeProfilerPackageEventType_PostLoadWait:
+			Row.Name = TEXT("PostLoadWait");
+			break;
+		case LoadTimeProfilerPackageEventType_StartPostLoad:
+			Row.Name = TEXT("StartPostLoad");
+			break;
+		case LoadTimeProfilerPackageEventType_Tick:
+			Row.Name = TEXT("Tick");
+			break;
+		case LoadTimeProfilerPackageEventType_DeferredPostLoad:
+			Row.Name = TEXT("DeferredPostLoad");
+			break;
+		case LoadTimeProfilerPackageEventType_Finish:
+			Row.Name = TEXT("Finish");
+			break;
+		}
+		const FAggregatedTimingStats& Stats = KV.Value;
+		Row.Count = Stats.InstanceCount;
+		Row.Total = Stats.TotalExclusiveTime;
+		Row.Min = Stats.MinExclusiveTime;
+		Row.Max = Stats.MaxExclusiveTime;
+		Row.Average = Stats.AverageExclusiveTime;
+		Row.Median = Stats.MedianExclusiveTime;
+	}
+	return Table;
+}
+
+ITable<FLoadTimeProfilerAggregatedStats>* FLoadTimeProfilerProvider::CreateObjectTypeAggregation(double IntervalStart, double IntervalEnd) const
+{
+	TArray<const CpuTimelineInternal*> Timelines;
+	Timelines.Add(&MainThreadCpuTimeline.Get());
+	Timelines.Add(&AsyncLoadingThreadCpuTimeline.Get());
+
+	auto BucketMapper = [](const FLoadTimeProfilerCpuEvent& Event) -> const FClassInfo*
+	{
+		return Event.Export ? Event.Export->Class : nullptr;
+	};
+	TMap<const Trace::FClassInfo*, FAggregatedTimingStats> Aggregation;
+	FTimelineStatistics::CreateAggregation(Timelines, BucketMapper, IntervalStart, IntervalEnd, Aggregation);
+	TTable<FAggregatedStatsTableLayout>* Table = new TTable<FAggregatedStatsTableLayout>();
+	for (const auto& KV : Aggregation)
+	{
+		const FClassInfo* ClassInfo = KV.Key;
+		if (ClassInfo)
+		{
+			FLoadTimeProfilerAggregatedStats& Row = Table->AddRow();
+			const FAggregatedTimingStats& Stats = KV.Value;
+			Row.Name = ClassInfo->Name;
+			Row.Count = Stats.InstanceCount;
+			Row.Total = Stats.TotalExclusiveTime;
+			Row.Min = Stats.MinExclusiveTime;
+			Row.Max = Stats.MaxExclusiveTime;
+			Row.Average = Stats.AverageExclusiveTime;
+			Row.Median = Stats.MedianExclusiveTime;
+		}
+	}
+	return Table;
+}
+
+const Trace::FClassInfo& FLoadTimeProfilerProvider::AddClassInfo(const TCHAR* ClassName)
+{
+	Session.WriteAccessCheck();
+
+	FClassInfo& ClassInfo = ClassInfos.PushBack();
+	ClassInfo.Name = Session.StoreString(ClassName);
+	return ClassInfo;
+}
+
+Trace::FPackageInfo& FLoadTimeProfilerProvider::CreatePackage(const TCHAR* PackageName)
+{
+	Session.WriteAccessCheck();
+
+	uint32 PackageId = Packages.Num();
+	FPackageInfo& Package = Packages.PushBack();
+	Package.Id = PackageId;
+	Package.Name = Session.StoreString(PackageName);
+	return Package;
+}
+
+Trace::FPackageExportInfo& FLoadTimeProfilerProvider::CreateExport()
+{
+	Session.WriteAccessCheck();
+
+	uint32 ExportId = Exports.Num();
+	FPackageExportInfo& Export = Exports.PushBack();
+	Export.Id = ExportId;
+	return Export;
+}
+
+}
\ No newline at end of file
diff --git a/Engine/Source/Developer/TraceServices/Private/Model/Log.cpp b/Engine/Source/Developer/TraceServices/Private/Model/Log.cpp
new file mode 100644
--- /dev/null
+++ b/Engine/Source/Developer/TraceServices/Private/Model/Log.cpp
@@ -0,0 +1,153 @@
+// Copyright 1998-2019 Epic Games, Inc. All Rights Reserved.
+
+#include "TraceServices/Model/Log.h"
+#include "Model/LogPrivate.h"
+#include "AnalysisServicePrivate.h"
+#include "Common/FormatArgs.h"
+
+namespace Trace
+{
+
+const FName FLogProvider::ProviderName("LogProvider");
+
+FLogProvider::FLogProvider(IAnalysisSession& InSession)
+	: Session(InSession)
+	, Categories(InSession.GetLinearAllocator(), 128)
+	, MessageSpecs(InSession.GetLinearAllocator(), 1024)
+	, Messages(InSession.GetLinearAllocator(), 1024)
+	, MessagesTable(Messages)
+{
+
+}
+
+FLogCategory& FLogProvider::GetCategory(uint64 CategoryPointer)
+{
+	Session.WriteAccessCheck();
+	if (CategoryMap.Contains(CategoryPointer))
+	{
+		return *CategoryMap[CategoryPointer];
+	}
+	else
+	{
+		FLogCategory& Category = Categories.PushBack();
+		CategoryMap.Add(CategoryPointer, &Category);
+		return Category;
+	}
+}
+
+FLogMessageSpec& FLogProvider::GetMessageSpec(uint64 LogPoint)
+{
+	Session.WriteAccessCheck();
+	if (SpecMap.Contains(LogPoint))
+	{
+		return *SpecMap[LogPoint];
+	}
+	else
+	{
+		FLogMessageSpec& Spec = MessageSpecs.PushBack();
+		SpecMap.Add(LogPoint, &Spec);
+		return Spec;
+	}
+}
+
+void FLogProvider::AppendMessage(uint64 LogPoint, double Time, const uint8* FormatArgs)
+{
+	Session.WriteAccessCheck();
+	check(SpecMap.Contains(LogPoint));
+	FLogMessageInternal& InternalMessage = Messages.PushBack();
+	InternalMessage.Time = Time;
+	InternalMessage.Spec = SpecMap[LogPoint];
+	FFormatArgsHelper::Format(FormatBuffer, FormatBufferSize - 1, InternalMessage.Spec->FormatString, FormatArgs);
+	InternalMessage.Message = Session.StoreString(FormatBuffer);
+	Session.UpdateDurationSeconds(Time);
+}
+
+uint64 FLogProvider::GetMessageCount() const
+{
+	Session.ReadAccessCheck();
+	return Messages.Num();
+}
+
+bool FLogProvider::ReadMessage(uint64 Index, TFunctionRef<void(const FLogMessage &)> Callback) const
+{
+	Session.ReadAccessCheck();
+	if (Index >= Messages.Num())
+	{
+		return false;
+	}
+	ConstructMessage(Index, Callback);
+	return true;
+}
+
+void FLogProvider::EnumerateMessages(double IntervalStart, double IntervalEnd, TFunctionRef<void(const FLogMessage&)> Callback) const
+{
+	Session.ReadAccessCheck();
+	if (IntervalStart > IntervalEnd)
+	{
+		return;
+	}
+	uint64 MessageCount = Messages.Num();
+	for (uint64 Index = 0; Index < MessageCount; ++Index)
+	{
+		double Time = Messages[Index].Time;
+		if (IntervalStart <= Time && Time <= IntervalEnd)
+		{
+			ConstructMessage(Index, Callback);
+		}
+	}
+}
+
+void FLogProvider::EnumerateMessagesByIndex(uint64 Start, uint64 End, TFunctionRef<void(const FLogMessage &)> Callback) const
+{
+	Session.ReadAccessCheck();
+	uint64 Count = Messages.Num();
+	if (Start >= Count)
+	{
+		return;
+	}
+	if (End > Count)
+	{
+		End = Count;
+	}
+	if (Start >= End)
+	{
+		return;
+	}
+	for (uint64 Index = Start; Index < End; ++Index)
+	{
+		ConstructMessage(Index, Callback);
+	}
+}
+
+void FLogProvider::ConstructMessage(uint64 Index, TFunctionRef<void(const FLogMessage &)> Callback) const
+{
+	const FLogMessageInternal& InternalMessage = Messages[Index];
+	FLogMessage Message;
+	Message.Index = Index;
+	Message.Time = InternalMessage.Time;
+	Message.Category = InternalMessage.Spec->Category;
+	Message.File = InternalMessage.Spec->File;
+	Message.Line = InternalMessage.Spec->Line;
+	Message.Verbosity = InternalMessage.Spec->Verbosity;
+	Message.Message = InternalMessage.Message;
+	Callback(Message);
+}
+
+void FLogProvider::EnumerateCategories(TFunctionRef<void(const FLogCategory &)> Callback) const
+{
+	Session.ReadAccessCheck();
+	auto Iterator = Categories.GetIteratorFromItem(0);
+	const FLogCategory* Category = Iterator.GetCurrentItem();
+	while (Category)
+	{
+		Callback(*Category);
+		Category = Iterator.NextItem();
+	}
+}
+
+const ILogProvider& ReadLogProvider(const IAnalysisSession& Session)
+{
+	return *Session.ReadProvider<ILogProvider>(FLogProvider::ProviderName);
+}
+
+}
\ No newline at end of file
diff --git a/Engine/Source/Developer/TraceServices/Private/Model/TimingProfiler.cpp b/Engine/Source/Developer/TraceServices/Private/Model/TimingProfiler.cpp
new file mode 100644
--- /dev/null
+++ b/Engine/Source/Developer/TraceServices/Private/Model/TimingProfiler.cpp
@@ -0,0 +1,179 @@
+// Copyright 1998-2019 Epic Games, Inc. All Rights Reserved.
+
+#include "TraceServices/Model/TimingProfiler.h"
+#include "Model/TimingProfilerPrivate.h"
+#include "AnalysisServicePrivate.h"
+#include "Common/StringStore.h"
+#include "Common/TimelineStatistics.h"
+
+namespace Trace
+{
+
+FTimingProfilerProvider::FTimingProfilerProvider(IAnalysisSession& InSession)
+	: Session(InSession)
+{
+	Timelines.Add(MakeShared<TimelineInternal>(Session.GetLinearAllocator()));
+}
+
+FTimingProfilerProvider::~FTimingProfilerProvider()
+{
+}
+
+uint32 FTimingProfilerProvider::AddCpuTimer(const TCHAR* Name)
+{
+	Session.WriteAccessCheck();
+
+	FTimingProfilerTimer& Timer = AddTimerInternal(Name, false);
+	return Timer.Id;
+}
+
+void FTimingProfilerProvider::SetTimerName(uint32 TimerId, const TCHAR* Name)
+{
+	Session.WriteAccessCheck();
+	
+	FTimingProfilerTimer& Timer = Timers[TimerId];
+	Timer.Name = Session.StoreString(Name);
+	uint32 NameHash = 0;
+	for (const TCHAR* c = Name; *c; ++c)
+	{
+		NameHash = (NameHash + *c) * 0x2c2c57ed;
+	}
+	Timer.NameHash = NameHash;
+}
+
+uint32 FTimingProfilerProvider::AddGpuTimer(const TCHAR* Name)
+{
+	Session.WriteAccessCheck();
+
+	FTimingProfilerTimer& Timer = AddTimerInternal(Name, true);
+	return Timer.Id;
+}
+
+FTimingProfilerTimer& FTimingProfilerProvider::AddTimerInternal(const TCHAR* Name, bool IsGpuTimer)
+{
+	FTimingProfilerTimer& Timer = Timers.AddDefaulted_GetRef();
+	Timer.Id = Timers.Num() - 1;
+	Timer.Name = Session.StoreString(Name);
+	uint32 NameHash = 0;
+	for (const TCHAR* c = Name; *c; ++c)
+	{
+		NameHash = (NameHash + *c) * 0x2c2c57ed;
+	}
+	Timer.NameHash = NameHash;
+	Timer.IsGpuTimer = IsGpuTimer;
+	return Timer;
+}
+
+FTimingProfilerProvider::TimelineInternal& FTimingProfilerProvider::EditCpuThreadTimeline(uint32 ThreadId)
+{
+	Session.WriteAccessCheck();
+
+	if (!CpuThreadTimelineIndexMap.Contains(ThreadId))
+	{
+		TSharedRef<TimelineInternal> Timeline = MakeShared<TimelineInternal>(Session.GetLinearAllocator());
+		uint32 TimelineIndex = Timelines.Num();
+		CpuThreadTimelineIndexMap.Add(ThreadId, TimelineIndex);
+		Timelines.Add(Timeline);
+		return Timeline.Get();
+	}
+	else
+	{
+		uint32 TimelineIndex = CpuThreadTimelineIndexMap[ThreadId];
+		return Timelines[TimelineIndex].Get();
+	}
+}
+
+FTimingProfilerProvider::TimelineInternal& FTimingProfilerProvider::EditGpuTimeline()
+{
+	Session.WriteAccessCheck();
+
+	return Timelines[GpuTimelineIndex].Get();
+}
+
+bool FTimingProfilerProvider::GetCpuThreadTimelineIndex(uint32 ThreadId, uint32& OutTimelineIndex) const
+{
+	Session.ReadAccessCheck();
+
+	if (CpuThreadTimelineIndexMap.Contains(ThreadId))
+	{
+		OutTimelineIndex = CpuThreadTimelineIndexMap[ThreadId];
+		return true;
+	}
+	return false;
+}
+
+bool FTimingProfilerProvider::GetGpuTimelineIndex(uint32& OutTimelineIndex) const
+{
+	Session.ReadAccessCheck();
+
+	OutTimelineIndex = GpuTimelineIndex;
+	return true;
+}
+
+bool FTimingProfilerProvider::ReadTimeline(uint32 Index, TFunctionRef<void(const Timeline &)> Callback) const
+{
+	Session.ReadAccessCheck();
+
+	if (Index < uint32(Timelines.Num()))
+	{
+		Callback(*Timelines[Index]);
+		return true;
+	}
+	else
+	{
+		return false;
+	}
+}
+
+void FTimingProfilerProvider::EnumerateTimelines(TFunctionRef<void(const Timeline&)> Callback) const
+{
+	Session.ReadAccessCheck();
+
+	for (const auto& Timeline : Timelines)
+	{
+		Callback(*Timeline);
+	}
+}
+
+void FTimingProfilerProvider::ReadTimers(TFunctionRef<void(const FTimingProfilerTimer*, uint64)> Callback) const
+{
+	Session.ReadAccessCheck();
+
+	Callback(Timers.GetData(), Timers.Num());
+}
+
+ITable<FTimingProfilerAggregatedStats>* FTimingProfilerProvider::CreateAggregation(double IntervalStart, double IntervalEnd, TFunctionRef<bool(uint32)> CpuThreadFilter, bool IncludeGpu) const
+{
+	Session.ReadAccessCheck();
+
+	TArray<const TimelineInternal*> IncludedTimelines;
+	if (IncludeGpu)
+	{
+		IncludedTimelines.Add(&Timelines[GpuTimelineIndex].Get());
+	}
+	for (const auto& KV : CpuThreadTimelineIndexMap)
+	{
+		if (CpuThreadFilter(KV.Key))
+		{
+			IncludedTimelines.Add(&Timelines[KV.Value].Get());
+		}
+	}
+
+	auto BucketMappingFunc = [this](const TimelineInternal::EventType& Event) -> const FTimingProfilerTimer*
+	{
+		return &Timers[Event.TimerIndex];
+	};
+
+	TMap<const FTimingProfilerTimer*, FAggregatedTimingStats> Aggregation;
+	FTimelineStatistics::CreateAggregation(IncludedTimelines, BucketMappingFunc, IntervalStart, IntervalEnd, Aggregation);
+	TTable<FAggregatedStatsTableLayout>* Table = new TTable<FAggregatedStatsTableLayout>();
+	for (const auto& KV : Aggregation)
+	{
+		FTimingProfilerAggregatedStats& Row = Table->AddRow();
+		*static_cast<FAggregatedTimingStats*>(&Row) = KV.Value;
+		Row.Timer = KV.Key;
+	}
+	return Table;
+}
+
+}
diff --git a/Engine/Source/Developer/TraceServices/Private/TraceServicesModule.cpp b/Engine/Source/Developer/TraceServices/Private/TraceServicesModule.cpp
new file mode 100644
--- /dev/null
+++ b/Engine/Source/Developer/TraceServices/Private/TraceServicesModule.cpp
@@ -0,0 +1,82 @@
+// Copyright 1998-2019 Epic Games, Inc. All Rights Reserved.
+
+#include "Modules/ModuleManager.h"
+#include "TraceServices/ITraceServicesModule.h"
+#include "SessionServicePrivate.h"
+#include "AnalysisServicePrivate.h"
+#include "ModuleServicePrivate.h"
+#include "Features/IModularFeatures.h"
+#include "Modules/TimingProfilerModule.h"
+#include "Modules/LoadTimeProfilerModule.h"
+#include "Modules/StatsModule.h"
+#include "Stats/StatsTrace.h"
+
+class FTraceServicesModule
+	: public ITraceServicesModule
+{
+public:
+	virtual TSharedPtr<Trace::ISessionService> GetSessionService() override;
+	virtual TSharedPtr<Trace::IAnalysisService> GetAnalysisService() override;
+	virtual TSharedPtr<Trace::IModuleService> GetModuleService() override;
+
+	virtual void StartupModule() override;
+	virtual void ShutdownModule() override;
+
+private:
+	TSharedPtr<Trace::FSessionService> SessionService;
+	TSharedPtr<Trace::FAnalysisService> AnalysisService;
+	TSharedPtr<Trace::FModuleService> ModuleService;
+
+	Trace::FTimingProfilerModule TimingProfilerModule;
+	Trace::FLoadTimeProfilerModule LoadTimeProfilerModule;
+	Trace::FStatsModule StatsModule;
+};
+
+TSharedPtr<Trace::ISessionService> FTraceServicesModule::GetSessionService()
+{
+	if (!SessionService.IsValid())
+	{
+		GetModuleService();
+		SessionService = MakeShared<Trace::FSessionService>(*ModuleService.Get());
+	}
+	return SessionService;
+}
+
+TSharedPtr<Trace::IAnalysisService> FTraceServicesModule::GetAnalysisService()
+{
+	if (!AnalysisService.IsValid())
+	{
+		GetModuleService();
+		AnalysisService = MakeShared<Trace::FAnalysisService>(*ModuleService.Get());
+	}
+	return AnalysisService;
+}
+
+TSharedPtr<Trace::IModuleService> FTraceServicesModule::GetModuleService()
+{
+	if (!ModuleService.IsValid())
+	{
+		ModuleService = MakeShared<Trace::FModuleService>();
+	}
+	return ModuleService;
+}
+
+void FTraceServicesModule::StartupModule()
+{
+	IModularFeatures::Get().RegisterModularFeature(Trace::ModuleFeatureName, &TimingProfilerModule);
+	IModularFeatures::Get().RegisterModularFeature(Trace::ModuleFeatureName, &LoadTimeProfilerModule);
+#if EXPERIMENTAL_STATSTRACE_ENABLED
+	IModularFeatures::Get().RegisterModularFeature(Trace::ModuleFeatureName, &StatsModule);
+#endif
+}
+
+void FTraceServicesModule::ShutdownModule()
+{
+#if EXPERIMENTAL_STATSTRACE_ENABLED
+	IModularFeatures::Get().UnregisterModularFeature(Trace::ModuleFeatureName, &StatsModule);
+#endif
+	IModularFeatures::Get().UnregisterModularFeature(Trace::ModuleFeatureName, &LoadTimeProfilerModule);
+	IModularFeatures::Get().UnregisterModularFeature(Trace::ModuleFeatureName, &TimingProfilerModule);
+}
+
+IMPLEMENT_MODULE(FTraceServicesModule, TraceServices)
diff --git a/Engine/Source/Developer/TraceServices/Public/TraceServices/AnalysisService.h b/Engine/Source/Developer/TraceServices/Public/TraceServices/AnalysisService.h
new file mode 100644
--- /dev/null
+++ b/Engine/Source/Developer/TraceServices/Public/TraceServices/AnalysisService.h
@@ -0,0 +1,35 @@
+// Copyright 1998-2019 Epic Games, Inc. All Rights Reserved.
+
+#pragma once
+
+#include "CoreTypes.h"
+#include "Templates/SharedPointer.h"
+#include "Templates/UniquePtr.h"
+#include "Delegates/Delegate.h"
+#include "Trace/DataStream.h"
+#include "TraceServices/Model/AnalysisSession.h"
+#include "TraceServices/Model/Log.h"
+#include "TraceServices/Model/Bookmarks.h"
+#include "TraceServices/Model/Frames.h"
+#include "TraceServices/Model/Threads.h"
+#include "TraceServices/Model/TimingProfiler.h"
+#include "TraceServices/Model/LoadTimeProfiler.h"
+#include "TraceServices/Model/Stats.h"
+
+namespace Trace
+{
+
+class IAnalysisService
+{
+public:
+	virtual TSharedPtr<const IAnalysisSession> Analyze(const TCHAR* SessionName, TUniquePtr<Trace::IInDataStream>&& DataStream) = 0;
+	virtual TSharedPtr<const IAnalysisSession> StartAnalysis(const TCHAR* SessionName, TUniquePtr<Trace::IInDataStream>&& DataStream) = 0;
+
+	DECLARE_EVENT_OneParam(IAnalysisService, FAnalysisStartedEvent, TSharedRef<const IAnalysisSession>)
+	virtual FAnalysisStartedEvent& OnAnalysisStarted() = 0;
+
+	DECLARE_EVENT_OneParam(IAnalysisService, FAnalysisFinishedEvent, TSharedRef<const IAnalysisSession>)
+	virtual FAnalysisFinishedEvent& OnAnalysisFinished() = 0;
+};
+
+}
diff --git a/Engine/Source/Developer/TraceServices/Public/TraceServices/ITraceServicesModule.h b/Engine/Source/Developer/TraceServices/Public/TraceServices/ITraceServicesModule.h
new file mode 100644
--- /dev/null
+++ b/Engine/Source/Developer/TraceServices/Public/TraceServices/ITraceServicesModule.h
@@ -0,0 +1,27 @@
+// Copyright 1998-2019 Epic Games, Inc. All Rights Reserved.
+
+#pragma once
+
+#include "CoreMinimal.h"
+#include "Trace/Store.h"
+#include "Modules/ModuleInterface.h"
+
+namespace Trace
+{
+
+class IAnalysisService;
+class ISessionService;
+class IModuleService;
+
+}
+
+class ITraceServicesModule
+	: public IModuleInterface
+{
+public:
+	virtual TSharedPtr<Trace::ISessionService> GetSessionService() = 0;
+	virtual TSharedPtr<Trace::IAnalysisService> GetAnalysisService() = 0;
+	virtual TSharedPtr<Trace::IModuleService> GetModuleService() = 0;
+
+	virtual ~ITraceServicesModule() = default;
+};
diff --git a/Engine/Source/Developer/TraceServices/Public/TraceServices/SessionService.h b/Engine/Source/Developer/TraceServices/Public/TraceServices/SessionService.h
new file mode 100644
--- /dev/null
+++ b/Engine/Source/Developer/TraceServices/Public/TraceServices/SessionService.h
@@ -0,0 +1,40 @@
+// Copyright 1998-2019 Epic Games, Inc. All Rights Reserved.
+
+#pragma once
+
+#include "CoreMinimal.h"
+#include "Templates/SharedPointer.h"
+#include "Trace/Store.h"
+#include "Trace/DataStream.h"
+
+namespace Trace
+{
+
+typedef uint64 FSessionHandle;
+
+struct FSessionInfo
+{
+	const TCHAR* Uri;
+	const TCHAR* Name;
+	bool bIsLive;
+};
+
+class ISessionService
+{
+public:
+	virtual ~ISessionService() = default;
+	virtual bool StartRecorderServer() = 0;
+	virtual bool IsRecorderServerRunning() const = 0;
+	virtual void StopRecorderServer() = 0;
+	virtual const TCHAR* GetLocalSessionDirectory() const = 0;
+	virtual void GetAvailableSessions(TArray<FSessionHandle>& OutSessions) const = 0;
+	virtual void GetLiveSessions(TArray<FSessionHandle>& OutSessions) const = 0;
+	virtual bool GetSessionInfo(FSessionHandle SessionHandle, FSessionInfo& OutSessionInfo) const = 0;
+	virtual Trace::IInDataStream* OpenSessionStream(FSessionHandle SessionHandle) = 0;
+	virtual Trace::IInDataStream* OpenSessionFromFile(const TCHAR* FilePath) = 0;
+	virtual void SetModuleEnabled(Trace::FSessionHandle SessionHandle, const FName& ModuleName, bool bState) = 0;
+	virtual bool IsModuleEnabled(Trace::FSessionHandle SessionHandle, const FName& ModuleName) const = 0;
+	virtual bool ConnectSession(const TCHAR* ControlClientAddress) = 0;
+};
+
+}
diff --git a/Engine/Source/Programs/UnrealInsights/Private/Linux/LinuxUnrealInsightsMain.cpp b/Engine/Source/Programs/UnrealInsights/Private/Linux/LinuxUnrealInsightsMain.cpp
new file mode 100644
--- /dev/null
+++ b/Engine/Source/Programs/UnrealInsights/Private/Linux/LinuxUnrealInsightsMain.cpp
@@ -0,0 +1,9 @@
+// Copyright 1998-2019 Epic Games, Inc. All Rights Reserved.
+
+#include "UnrealInsightsMain.h"
+#include "UnixCommonStartup.h"
+
+int main(int argc, char *argv[])
+{
+	return CommonUnixMain(argc, argv, &UnrealInsightsMain);
+}
diff --git a/Engine/Source/Programs/UnrealInsights/Private/Mac/MacUnrealInsightsMain.cpp b/Engine/Source/Programs/UnrealInsights/Private/Mac/MacUnrealInsightsMain.cpp
new file mode 100644
--- /dev/null
+++ b/Engine/Source/Programs/UnrealInsights/Private/Mac/MacUnrealInsightsMain.cpp
@@ -0,0 +1,117 @@
+// Copyright 1998-2019 Epic Games, Inc. All Rights Reserved.
+
+#include "UnrealInsightsMain.h"
+#include "HAL/ExceptionHandling.h"
+#include "LaunchEngineLoop.h"
+#include "Mac/CocoaThread.h"
+
+
+static FString GSavedCommandLine;
+
+
+@interface UE4AppDelegate : NSObject<NSApplicationDelegate, NSFileManagerDelegate>
+{
+}
+
+@end
+
+
+@implementation UE4AppDelegate
+
+//handler for the quit apple event used by the Dock menu
+- (void)handleQuitEvent:(NSAppleEventDescriptor*)Event withReplyEvent:(NSAppleEventDescriptor*)ReplyEvent
+{
+	[self requestQuit:self];
+}
+
+- (IBAction)requestQuit:(id)Sender
+{
+	GIsRequestingExit = true;
+}
+
+- (void) runGameThread:(id)Arg
+{
+	FPlatformMisc::SetGracefulTerminationHandler();
+	FPlatformMisc::SetCrashHandler(nullptr);
+	
+#if !UE_BUILD_SHIPPING
+	if (FParse::Param(*GSavedCommandLine,TEXT("crashreports")))
+	{
+		GAlwaysReportCrash = true;
+	}
+#endif
+	
+#if UE_BUILD_DEBUG
+	if (!GAlwaysReportCrash)
+#else
+		if (FPlatformMisc::IsDebuggerPresent() && !GAlwaysReportCrash)
+#endif
+		{
+			UnrealInsightsMain(*GSavedCommandLine);
+		}
+		else
+		{
+			GIsGuarded = 1;
+			UnrealInsightsMain(*GSavedCommandLine);
+			GIsGuarded = 0;
+		}
+	
+	FEngineLoop::AppExit();
+
+	[NSApp terminate: self];
+}
+
+- (NSApplicationTerminateReply)applicationShouldTerminate:(NSApplication *)Sender;
+{
+	if(!GIsRequestingExit || ([NSThread gameThread] && [NSThread gameThread] != [NSThread mainThread]))
+	{
+		[self requestQuit:self];
+		return NSTerminateLater;
+	}
+	else
+	{
+		return NSTerminateNow;
+	}
+}
+
+- (void)applicationDidFinishLaunching:(NSNotification *)Notification
+{
+	//install the custom quit event handler
+	NSAppleEventManager* appleEventManager = [NSAppleEventManager sharedAppleEventManager];
+	[appleEventManager setEventHandler:self andSelector:@selector(handleQuitEvent:withReplyEvent:) forEventClass:kCoreEventClass andEventID:kAEQuitApplication];
+
+	RunGameThread(self, @selector(runGameThread:));
+}
+
+@end
+
+
+int main(int argc, char *argv[])
+{
+	for (int32 Option = 1; Option < argc; Option++)
+	{
+		GSavedCommandLine += TEXT(" ");
+		FString Argument(ANSI_TO_TCHAR(argv[Option]));
+		if (Argument.Contains(TEXT(" ")))
+		{
+			if (Argument.Contains(TEXT("=")))
+			{
+				FString ArgName;
+				FString ArgValue;
+				Argument.Split( TEXT("="), &ArgName, &ArgValue );
+				Argument = FString::Printf( TEXT("%s=\"%s\""), *ArgName, *ArgValue );
+			}
+			else
+			{
+				Argument = FString::Printf(TEXT("\"%s\""), *Argument);
+			}
+		}
+		GSavedCommandLine += Argument;
+	}
+
+	SCOPED_AUTORELEASE_POOL;
+	[NSApplication sharedApplication];
+	[NSApp setDelegate:[UE4AppDelegate new]];
+	[NSApp run];
+	return 0;
+}
diff --git a/Engine/Source/Programs/UnrealInsights/Private/UnrealInsightsMain.cpp b/Engine/Source/Programs/UnrealInsights/Private/UnrealInsightsMain.cpp
new file mode 100644
--- /dev/null
+++ b/Engine/Source/Programs/UnrealInsights/Private/UnrealInsightsMain.cpp
@@ -0,0 +1,46 @@
+// Copyright 1998-2019 Epic Games, Inc. All Rights Reserved.
+
+#include "UnrealInsightsMain.h"
+
+#include "RequiredProgramMainCPPInclude.h"
+#include "UserInterfaceCommand.h"
+
+
+IMPLEMENT_APPLICATION(UnrealInsights, "UnrealInsights");
+
+
+/**
+ * Platform agnostic implementation of the main entry point.
+ */
+int32 UnrealInsightsMain(const TCHAR* CommandLine)
+{
+	// Override the stack size for the thread pool.
+	FQueuedThreadPool::OverrideStackSize = 256 * 1024;
+
+	//im: ???
+	FCommandLine::Set(CommandLine);
+
+	// Initialize core.
+	GEngineLoop.PreInit(CommandLine);
+
+	// Make sure all UObject classes are registered and default properties have been initialized.
+	//im: ??? ProcessNewlyLoadedUObjects();
+
+	// Tell the module manager it may now process newly-loaded UObjects when new C++ modules are loaded.
+	FModuleManager::Get().StartProcessingNewlyLoadedObjects();
+
+	FUserInterfaceCommand::Run();
+
+	// Shut down.
+	//im: ??? FCoreDelegates::OnExit.Broadcast();
+	FModuleManager::Get().UnloadModulesAtShutdown();
+	FEngineLoop::AppPreExit(); //im: ???
+
+#if STATS
+	FThreadStats::StopThread();
+#endif
+
+	FTaskGraphInterface::Shutdown(); //im: ???
+
+	return 0;
+}
diff --git a/Engine/Source/Programs/UnrealInsights/Private/UserInterfaceCommand.h b/Engine/Source/Programs/UnrealInsights/Private/UserInterfaceCommand.h
new file mode 100644
--- /dev/null
+++ b/Engine/Source/Programs/UnrealInsights/Private/UserInterfaceCommand.h
@@ -0,0 +1,37 @@
+// Copyright 1998-2019 Epic Games, Inc. All Rights Reserved.
+
+#pragma once
+
+#include "CoreMinimal.h"
+
+class SDockTab;
+class FSpawnTabArgs;
+
+class FUserInterfaceCommand
+{
+public:
+
+	/** Executes the command. */
+	static void Run();
+
+protected:
+
+	/**
+	 * Initializes the Slate application.
+	 *
+	 * @param LayoutInit The path to the layout configuration file.
+	 */
+	static void InitializeSlateApplication(const FString& LayoutIni);
+
+	/**
+	 * Shuts down the Slate application.
+	 *
+	 * @param LayoutInit The path to the layout configuration file.
+	 */
+	static void ShutdownSlateApplication(const FString& LayoutIni);
+
+private:
+
+	/** Callback for spawning tabs. */
+	TSharedRef<SDockTab> HandleTabManagerSpawnTab(const FSpawnTabArgs& Args, FName TabIdentifier) const;
+};
diff --git a/Engine/Source/Runtime/Core/Private/Async/TaskGraph.cpp b/Engine/Source/Runtime/Core/Private/Async/TaskGraph.cpp
--- a/Engine/Source/Runtime/Core/Private/Async/TaskGraph.cpp
+++ b/Engine/Source/Runtime/Core/Private/Async/TaskGraph.cpp
@@ -28,6 +28,7 @@
 #include "HAL/LowLevelMemTracker.h"
 #include "HAL/ThreadHeartBeat.h"
 #include "ProfilingDebugging/ExternalProfiler.h"
+#include "ProfilingDebugging/MiscTrace.h"
 
 DEFINE_LOG_CATEGORY_STATIC(LogTaskGraph, Log, All);
 
@@ -1210,12 +1211,14 @@ public:
 		for (int32 ThreadIndex = LastExternalThread + 1; ThreadIndex < NumThreads; ThreadIndex++)
 		{
 			FString Name;
+			const ANSICHAR* GroupName = "TaskGraphNormal";
 			int32 Priority = ThreadIndexToPriorityIndex(ThreadIndex);
 			EThreadPriority ThreadPri;
 			uint64 Affinity = FPlatformAffinity::GetTaskGraphThreadMask();
 			if (Priority == 1)
 			{
 				Name = FString::Printf(TEXT("TaskGraphThreadHP %d"), ThreadIndex - (LastExternalThread + 1));
+				GroupName = "TaskGraphHigh";
 				ThreadPri = TPri_SlightlyBelowNormal; // we want even hi priority tasks below the normal threads
 
 				// If the platform defines FPlatformAffinity::GetTaskGraphHighPriorityTaskMask then use it
@@ -1227,6 +1230,7 @@ public:
 			else if (Priority == 2)
 			{
 				Name = FString::Printf(TEXT("TaskGraphThreadBP %d"), ThreadIndex - (LastExternalThread + 1));
+				GroupName = "TaskGraphLow";
 				ThreadPri = TPri_Lowest;
 				// If the platform defines FPlatformAffinity::GetTaskGraphBackgroundTaskMask then use it
 				if ( FPlatformAffinity::GetTaskGraphBackgroundTaskMask() != 0xFFFFFFFFFFFFFFFF )
@@ -1248,6 +1252,10 @@ public:
 #endif
 			WorkerThreads[ThreadIndex].RunnableThread = FRunnableThread::Create(&Thread(ThreadIndex), *Name, StackSize, ThreadPri, Affinity); // these are below normal threads so that they sleep when the named threads are active
 			WorkerThreads[ThreadIndex].bAttached = true;
+			if (WorkerThreads[ThreadIndex].RunnableThread)
+			{
+				TRACE_SET_THREAD_GROUP(WorkerThreads[ThreadIndex].RunnableThread->GetThreadID(), GroupName);
+			}
 		}
 	}
 
diff --git a/Engine/Source/Runtime/Core/Private/HAL/ThreadingBase.cpp b/Engine/Source/Runtime/Core/Private/HAL/ThreadingBase.cpp
--- a/Engine/Source/Runtime/Core/Private/HAL/ThreadingBase.cpp
+++ b/Engine/Source/Runtime/Core/Private/HAL/ThreadingBase.cpp
@@ -8,6 +8,7 @@
 #include "Misc/EventPool.h"
 #include "Templates/Atomic.h"
 #include "HAL/PlatformStackWalk.h"
+#include "ProfilingDebugging/MiscTrace.h"
 
 DEFINE_STAT( STAT_EventWaitWithId );
 DEFINE_STAT( STAT_EventTriggerWithId );
@@ -410,6 +411,10 @@ FRunnableThread* FRunnableThread::Create(
 		}
 	}
 
+	if (NewThread)
+	{
+		TRACE_CREATE_THREAD(NewThread->GetThreadID(), *NewThread->GetThreadName(), InThreadPri);
+	}
 #if	STATS
 	if( NewThread )
 	{
diff --git a/Engine/Source/Runtime/Core/Private/Logging/LogCategory.cpp b/Engine/Source/Runtime/Core/Private/Logging/LogCategory.cpp
--- a/Engine/Source/Runtime/Core/Private/Logging/LogCategory.cpp
+++ b/Engine/Source/Runtime/Core/Private/Logging/LogCategory.cpp
@@ -10,6 +10,7 @@ FLogCategoryBase::FLogCategoryBase(const TCHAR *CategoryName, ELogVerbosity::Typ
 	, CompileTimeVerbosity(InCompileTimeVerbosity)
 	, CategoryFName(CategoryName)
 {
+	TRACE_LOG_CATEGORY(this, CategoryName, InDefaultVerbosity);
 	ResetFromDefault();
 	if (CompileTimeVerbosity > ELogVerbosity::NoLogging)
 	{
diff --git a/Engine/Source/Runtime/Core/Private/Logging/LogTrace.cpp b/Engine/Source/Runtime/Core/Private/Logging/LogTrace.cpp
new file mode 100644
--- /dev/null
+++ b/Engine/Source/Runtime/Core/Private/Logging/LogTrace.cpp
@@ -0,0 +1,63 @@
+// Copyright 1998-2019 Epic Games, Inc. All Rights Reserved.
+#include "Logging/LogTrace.h"
+
+#if LOGTRACE_ENABLED
+
+#include "Trace/Trace.h"
+#include "Templates/Function.h"
+#include "HAL/PlatformTime.h"
+#include "HAL/PlatformTLS.h"
+
+UE_TRACE_EVENT_BEGIN(Logging, LogCategory, Always|Important)
+	UE_TRACE_EVENT_FIELD(const void*, CategoryPointer)
+	UE_TRACE_EVENT_FIELD(uint8, DefaultVerbosity)
+UE_TRACE_EVENT_END()
+
+UE_TRACE_EVENT_BEGIN(Logging, LogMessageSpec, Always|Important)
+	UE_TRACE_EVENT_FIELD(const void*, LogPoint)
+	UE_TRACE_EVENT_FIELD(const void*, CategoryPointer)
+	UE_TRACE_EVENT_FIELD(int32, Line)
+	UE_TRACE_EVENT_FIELD(uint8, Verbosity)
+UE_TRACE_EVENT_END()
+
+UE_TRACE_EVENT_BEGIN(Logging, LogMessage, Always)
+	UE_TRACE_EVENT_FIELD(const void*, LogPoint)
+	UE_TRACE_EVENT_FIELD(uint64, Cycle)
+	UE_TRACE_EVENT_FIELD(uint32, ThreadId)
+UE_TRACE_EVENT_END()
+
+void FLogTrace::OutputLogCategory(const FLogCategoryBase* Category, const TCHAR* Name, ELogVerbosity::Type DefaultVerbosity)
+{
+	uint16 NameSize = (FCString::Strlen(Name) + 1) * sizeof(TCHAR);
+	UE_TRACE_LOG(Logging, LogCategory, NameSize)
+		<< LogCategory.CategoryPointer(Category)
+		<< LogCategory.DefaultVerbosity(DefaultVerbosity)
+		<< LogCategory.Attachment(Name, NameSize);
+}
+
+void FLogTrace::OutputLogMessageSpec(const void* LogPoint, const FLogCategoryBase* Category, ELogVerbosity::Type Verbosity, const ANSICHAR* File, int32 Line, const TCHAR* Format)
+{
+	uint16 FileNameSize = strlen(File) + 1;
+	uint16 FormatStringSize = (FCString::Strlen(Format) + 1) * sizeof(TCHAR);
+	auto StringCopyFunc = [FileNameSize, FormatStringSize, File, Format](uint8* Out) {
+		memcpy(Out, File, FileNameSize);
+		memcpy(Out + FileNameSize, Format, FormatStringSize);
+	};
+	UE_TRACE_LOG(Logging, LogMessageSpec, FileNameSize + FormatStringSize)
+		<< LogMessageSpec.LogPoint(LogPoint)
+		<< LogMessageSpec.CategoryPointer(Category)
+		<< LogMessageSpec.Line(Line)
+		<< LogMessageSpec.Verbosity(Verbosity)
+		<< LogMessageSpec.Attachment(StringCopyFunc);
+}
+
+void FLogTrace::OutputLogMessageInternal(const void* LogPoint, uint16 EncodedFormatArgsSize, uint8* EncodedFormatArgs)
+{
+	UE_TRACE_LOG(Logging, LogMessage, EncodedFormatArgsSize)
+		<< LogMessage.LogPoint(LogPoint)
+		<< LogMessage.Cycle(FPlatformTime::Cycles64())
+		<< LogMessage.ThreadId(FPlatformTLS::GetCurrentThreadId())
+		<< LogMessage.Attachment(EncodedFormatArgs, EncodedFormatArgsSize);
+}
+
+#endif
diff --git a/Engine/Source/Runtime/Core/Private/ProfilingDebugging/CpuProfilerTrace.cpp b/Engine/Source/Runtime/Core/Private/ProfilingDebugging/CpuProfilerTrace.cpp
new file mode 100644
--- /dev/null
+++ b/Engine/Source/Runtime/Core/Private/ProfilingDebugging/CpuProfilerTrace.cpp
@@ -0,0 +1,157 @@
+// Copyright 1998-2019 Epic Games, Inc. All Rights Reserved.
+#include "ProfilingDebugging/CpuProfilerTrace.h"
+#include "Trace/Trace.h"
+#include "HAL/PlatformTime.h"
+#include "HAL/PlatformTLS.h"
+#include "ProfilingDebugging/MiscTrace.h"
+
+#if CPUPROFILERTRACE_ENABLED
+
+UE_TRACE_EVENT_BEGIN(CpuProfiler, EventSpec, Always)
+	UE_TRACE_EVENT_FIELD(uint16, Id)
+	UE_TRACE_EVENT_FIELD(uint16, Group)
+UE_TRACE_EVENT_END()
+
+UE_TRACE_EVENT_BEGIN(CpuProfiler, EventBatch)
+	UE_TRACE_EVENT_FIELD(uint32, ThreadId)
+UE_TRACE_EVENT_END()
+
+UE_TRACE_EVENT_BEGIN(CpuProfiler, EndCapture, Always)
+	UE_TRACE_EVENT_FIELD(uint32, ThreadId)
+UE_TRACE_EVENT_END()
+
+struct FCpuProfilerTraceInternal
+{
+public:
+	enum
+	{
+		MaxBufferSize = 255,
+		MaxEncodedEventSize = 12, // 9 + 3
+		FullBufferThreshold = MaxBufferSize - MaxEncodedEventSize,
+	};
+
+	struct FThreadState
+	{
+		uint64 LastCycle = 0;
+		uint32 ThreadId;
+		uint16 Depth = 0;
+		uint16 BufferSize = 0;
+		bool bEnabled = false;
+		uint8 Buffer[MaxBufferSize];
+	};
+
+	FORCENOINLINE static FThreadState* InitThreadState();
+	FORCENOINLINE static void FlushThreadBuffer(FThreadState* ThreadState);
+	FORCENOINLINE static void EndCapture(FThreadState* ThreadState);
+
+	static thread_local FThreadState* ThreadLocalThreadState;
+};
+
+thread_local FCpuProfilerTraceInternal::FThreadState* FCpuProfilerTraceInternal::ThreadLocalThreadState = nullptr;
+
+FCpuProfilerTraceInternal::FThreadState* FCpuProfilerTraceInternal::InitThreadState()
+{
+	ThreadLocalThreadState = new FThreadState();
+	ThreadLocalThreadState->ThreadId = FPlatformTLS::GetCurrentThreadId();
+	return ThreadLocalThreadState;
+}
+
+void FCpuProfilerTraceInternal::FlushThreadBuffer(FThreadState* ThreadState)
+{
+	UE_TRACE_LOG(CpuProfiler, EventBatch, ThreadState->BufferSize)
+		<< EventBatch.ThreadId(ThreadState->ThreadId)
+		<< EventBatch.Attachment(ThreadState->Buffer, ThreadState->BufferSize);
+	ThreadState->BufferSize = 0;
+}
+
+void FCpuProfilerTraceInternal::EndCapture(FThreadState* ThreadState)
+{
+	UE_TRACE_LOG(CpuProfiler, EndCapture, ThreadState->BufferSize)
+		<< EndCapture.ThreadId(ThreadState->ThreadId)
+		<< EndCapture.Attachment(ThreadState->Buffer, ThreadState->BufferSize);
+	ThreadState->BufferSize = 0;
+}
+
+void FCpuProfilerTrace::OutputBeginEvent(uint16 SpecId)
+{
+	FCpuProfilerTraceInternal::FThreadState* ThreadState = FCpuProfilerTraceInternal::ThreadLocalThreadState;
+	if (!ThreadState)
+	{
+		ThreadState = FCpuProfilerTraceInternal::InitThreadState();
+	}
+	++ThreadState->Depth;
+	bool bEventEnabled = UE_TRACE_EVENT_IS_ENABLED(CpuProfiler, EventBatch);
+	if ((!bEventEnabled) & (bEventEnabled == ThreadState->bEnabled))
+	{
+		return;
+	}
+	ThreadState->bEnabled = bEventEnabled;
+	if (!bEventEnabled)
+	{
+		FCpuProfilerTraceInternal::EndCapture(ThreadState);
+		return;
+	}
+	if (ThreadState->BufferSize >= FCpuProfilerTraceInternal::FullBufferThreshold)
+	{
+		FCpuProfilerTraceInternal::FlushThreadBuffer(ThreadState);
+	}
+	uint64 Cycle = FPlatformTime::Cycles64();
+	uint64 CycleDiff = Cycle - ThreadState->LastCycle;
+	ThreadState->LastCycle = Cycle;
+	uint8* BufferPtr = ThreadState->Buffer + ThreadState->BufferSize;
+	FTraceUtils::Encode7bit((CycleDiff << 1) | 1ull, BufferPtr);
+	FTraceUtils::Encode7bit(SpecId, BufferPtr);
+	ThreadState->BufferSize = BufferPtr - ThreadState->Buffer;
+}
+
+void FCpuProfilerTrace::OutputEndEvent()
+{
+	FCpuProfilerTraceInternal::FThreadState* ThreadState = FCpuProfilerTraceInternal::ThreadLocalThreadState;
+	--ThreadState->Depth;
+	bool bEventEnabled = UE_TRACE_EVENT_IS_ENABLED(CpuProfiler, EventBatch);
+	if ((!bEventEnabled) & (bEventEnabled == ThreadState->bEnabled))
+	{
+		return;
+	}
+	ThreadState->bEnabled = bEventEnabled;
+	if (!bEventEnabled)
+	{
+		FCpuProfilerTraceInternal::EndCapture(ThreadState);
+		return;
+	}
+	if ((ThreadState->Depth == 0) | (ThreadState->BufferSize >= FCpuProfilerTraceInternal::FullBufferThreshold))
+	{
+		FCpuProfilerTraceInternal::FlushThreadBuffer(ThreadState);
+	}
+	uint64 Cycle = FPlatformTime::Cycles64();
+	uint64 CycleDiff = Cycle - ThreadState->LastCycle;
+	ThreadState->LastCycle = Cycle;
+	uint8* BufferPtr = ThreadState->Buffer + ThreadState->BufferSize;
+	FTraceUtils::Encode7bit(CycleDiff << 1, BufferPtr);
+	ThreadState->BufferSize = BufferPtr - ThreadState->Buffer;
+}
+
+uint16 FCpuProfilerTrace::OutputEventType(const TCHAR* Name, ECpuProfilerGroup Group)
+{
+	static TAtomic<uint32> NextSpecId(1);
+	uint32 WideId = NextSpecId++;
+	check(WideId <= 0xFFFF);
+	uint16 Id = uint16(WideId);
+	uint16 NameSize = (FCString::Strlen(Name) + 1) * sizeof(TCHAR);
+	UE_TRACE_LOG(CpuProfiler, EventSpec, NameSize)
+		<< EventSpec.Id(Id)
+		<< EventSpec.Group(uint16(Group))
+		<< EventSpec.Attachment(Name, NameSize);
+	return Id;
+}
+
+void FCpuProfilerTrace::Init(bool bStartEnabled)
+{
+	if (bStartEnabled)
+	{
+		UE_TRACE_EVENT_IS_ENABLED(CpuProfiler, EventBatch);
+		Trace::ToggleEvent(TEXT("CpuProfiler.EventBatch"), true);
+	}
+}
+
+#endif
\ No newline at end of file
diff --git a/Engine/Source/Runtime/Core/Private/ProfilingDebugging/MiscTrace.cpp b/Engine/Source/Runtime/Core/Private/ProfilingDebugging/MiscTrace.cpp
new file mode 100644
--- /dev/null
+++ b/Engine/Source/Runtime/Core/Private/ProfilingDebugging/MiscTrace.cpp
@@ -0,0 +1,166 @@
+// Copyright 1998-2019 Epic Games, Inc. All Rights Reserved.
+#include "ProfilingDebugging/MiscTrace.h"
+
+#if MISCTRACE_ENABLED
+
+#include "Trace/Trace.h"
+#include "Misc/CString.h"
+#include "HAL/PlatformTLS.h"
+#include "HAL/PlatformTime.h"
+
+UE_TRACE_EVENT_BEGIN(Misc, RegisterGameThread, Always)
+	UE_TRACE_EVENT_FIELD(uint32, ThreadId)
+UE_TRACE_EVENT_END()
+
+UE_TRACE_EVENT_BEGIN(Misc, CreateThread, Always)
+	UE_TRACE_EVENT_FIELD(uint32, CurrentThreadId)
+	UE_TRACE_EVENT_FIELD(uint32, CreatedThreadId)
+	UE_TRACE_EVENT_FIELD(uint32, Priority)
+	UE_TRACE_EVENT_FIELD(uint16, NameSize)
+UE_TRACE_EVENT_END()
+
+UE_TRACE_EVENT_BEGIN(Misc, SetThreadGroup, Always)
+	UE_TRACE_EVENT_FIELD(uint32, ThreadId)
+UE_TRACE_EVENT_END()
+
+UE_TRACE_EVENT_BEGIN(Misc, BeginThreadGroupScope, Always)
+	UE_TRACE_EVENT_FIELD(uint32, CurrentThreadId)
+UE_TRACE_EVENT_END()
+
+UE_TRACE_EVENT_BEGIN(Misc, EndThreadGroupScope, Always)
+	UE_TRACE_EVENT_FIELD(uint32, CurrentThreadId)
+UE_TRACE_EVENT_END()
+
+UE_TRACE_EVENT_BEGIN(Misc, BookmarkSpec, Always)
+	UE_TRACE_EVENT_FIELD(const void*, BookmarkPoint)
+	UE_TRACE_EVENT_FIELD(int32, Line)
+UE_TRACE_EVENT_END()
+
+UE_TRACE_EVENT_BEGIN(Misc, Bookmark, Always)
+	UE_TRACE_EVENT_FIELD(uint64, Cycle)
+	UE_TRACE_EVENT_FIELD(const void*, BookmarkPoint)
+UE_TRACE_EVENT_END()
+
+UE_TRACE_EVENT_BEGIN(Misc, BeginGameFrame, Always)
+UE_TRACE_EVENT_END()
+
+UE_TRACE_EVENT_BEGIN(Misc, EndGameFrame, Always)
+UE_TRACE_EVENT_END()
+
+UE_TRACE_EVENT_BEGIN(Misc, BeginRenderFrame, Always)
+UE_TRACE_EVENT_END()
+
+UE_TRACE_EVENT_BEGIN(Misc, EndRenderFrame, Always)
+UE_TRACE_EVENT_END()
+
+
+struct FMiscTraceInternal
+{
+	static uint64 LastFrameCycle[TraceFrameType_Count];
+};
+
+uint64 FMiscTraceInternal::LastFrameCycle[TraceFrameType_Count] = { 0, 0 };
+
+void FMiscTrace::OutputRegisterGameThread(uint32 Id)
+{
+	UE_TRACE_LOG(Misc, RegisterGameThread)
+		<< RegisterGameThread.ThreadId(FPlatformTLS::GetCurrentThreadId());
+}
+
+void FMiscTrace::OutputCreateThread(uint32 Id, const TCHAR* Name, uint32 Priority)
+{
+	uint16 NameSize = (FCString::Strlen(Name) + 1) * sizeof(TCHAR);
+	UE_TRACE_LOG(Misc, CreateThread, NameSize)
+		<< CreateThread.CurrentThreadId(FPlatformTLS::GetCurrentThreadId())
+		<< CreateThread.CreatedThreadId(Id)
+		<< CreateThread.Priority(Priority)
+		<< CreateThread.Attachment(Name, NameSize);
+}
+
+void FMiscTrace::OutputSetThreadGroup(uint32 Id, const ANSICHAR* GroupName)
+{
+	uint16 NameSize = strlen(GroupName) + 1;
+	UE_TRACE_LOG(Misc, SetThreadGroup, NameSize)
+		<< SetThreadGroup.ThreadId(Id)
+		<< SetThreadGroup.Attachment(GroupName, NameSize);
+}
+
+void FMiscTrace::OutputBeginThreadGroupScope(const ANSICHAR* GroupName)
+{
+	uint16 NameSize = strlen(GroupName) + 1;
+	UE_TRACE_LOG(Misc, BeginThreadGroupScope, NameSize)
+		<< BeginThreadGroupScope.CurrentThreadId(FPlatformTLS::GetCurrentThreadId())
+		<< BeginThreadGroupScope.Attachment(GroupName, NameSize);
+}
+
+void FMiscTrace::OutputEndThreadGroupScope()
+{
+	UE_TRACE_LOG(Misc, EndThreadGroupScope)
+		<< EndThreadGroupScope.CurrentThreadId(FPlatformTLS::GetCurrentThreadId());
+}
+
+void FMiscTrace::OutputBookmarkSpec(const void* BookmarkPoint, const ANSICHAR* File, int32 Line, const TCHAR* Format)
+{
+	uint16 FileNameSize = strlen(File) + 1;
+	uint16 FormatStringSize = (FCString::Strlen(Format) + 1) * sizeof(TCHAR);
+	auto StringCopyFunc = [FileNameSize, FormatStringSize, File, Format](uint8* Out) {
+		memcpy(Out, File, FileNameSize);
+		memcpy(Out + FileNameSize, Format, FormatStringSize);
+	};
+	UE_TRACE_LOG(Misc, BookmarkSpec, FileNameSize + FormatStringSize)
+		<< BookmarkSpec.BookmarkPoint(BookmarkPoint)
+		<< BookmarkSpec.Line(Line)
+		<< BookmarkSpec.Attachment(StringCopyFunc);
+}
+
+void FMiscTrace::OutputBookmarkInternal(const void* BookmarkPoint, uint16 EncodedFormatArgsSize, uint8* EncodedFormatArgs)
+{
+	UE_TRACE_LOG(Misc, Bookmark, EncodedFormatArgsSize)
+		<< Bookmark.Cycle(FPlatformTime::Cycles64())
+		<< Bookmark.BookmarkPoint(BookmarkPoint)
+		<< Bookmark.Attachment(EncodedFormatArgs, EncodedFormatArgsSize);
+}
+
+void FMiscTrace::OutputBeginFrame(ETraceFrameType FrameType)
+{
+	uint64 Cycle = FPlatformTime::Cycles64();
+	uint64 CycleDiff = Cycle - FMiscTraceInternal::LastFrameCycle[FrameType];
+	FMiscTraceInternal::LastFrameCycle[FrameType] = Cycle;
+	uint8 Buffer[9];
+	uint8* BufferPtr = Buffer;
+	FTraceUtils::Encode7bit(CycleDiff, BufferPtr);
+	uint16 BufferSize = BufferPtr - Buffer;
+	if (FrameType == TraceFrameType_Game)
+	{
+		UE_TRACE_LOG(Misc, BeginGameFrame, BufferSize)
+			<< BeginGameFrame.Attachment(&Buffer, BufferSize);
+	}
+	else if (FrameType == TraceFrameType_Rendering)
+	{
+		UE_TRACE_LOG(Misc, BeginRenderFrame, BufferSize)
+			<< BeginRenderFrame.Attachment(&Buffer, BufferSize);
+	}
+}
+
+void FMiscTrace::OutputEndFrame(ETraceFrameType FrameType)
+{
+	uint64 Cycle = FPlatformTime::Cycles64();
+	uint64 CycleDiff = Cycle - FMiscTraceInternal::LastFrameCycle[FrameType];
+	FMiscTraceInternal::LastFrameCycle[FrameType] = Cycle;
+	uint8 Buffer[9];
+	uint8* BufferPtr = Buffer;
+	FTraceUtils::Encode7bit(CycleDiff, BufferPtr);
+	uint16 BufferSize = BufferPtr - Buffer;
+	if (FrameType == TraceFrameType_Game)
+	{
+		UE_TRACE_LOG(Misc, EndGameFrame, BufferSize)
+			<< EndGameFrame.Attachment(&Buffer, BufferSize);
+	}
+	else if (FrameType == TraceFrameType_Rendering)
+	{
+		UE_TRACE_LOG(Misc, EndRenderFrame, BufferSize)
+			<< EndRenderFrame.Attachment(&Buffer, BufferSize);
+	}
+}
+
+#endif
diff --git a/Engine/Source/Runtime/Core/Private/ProfilingDebugging/PlatformFileTrace.cpp b/Engine/Source/Runtime/Core/Private/ProfilingDebugging/PlatformFileTrace.cpp
new file mode 100644
--- /dev/null
+++ b/Engine/Source/Runtime/Core/Private/ProfilingDebugging/PlatformFileTrace.cpp
@@ -0,0 +1,140 @@
+// Copyright 1998-2019 Epic Games, Inc. All Rights Reserved.
+#include "ProfilingDebugging/PlatformFileTrace.h"
+
+#if PLATFORMFILETRACE_ENABLED
+
+#include "Misc/CString.h"
+#include "HAL/PlatformTime.h"
+#include "HAL/PlatformTLS.h"
+
+UE_TRACE_EVENT_BEGIN(PlatformFile, BeginOpen, Always)
+	UE_TRACE_EVENT_FIELD(uint64, Cycle)
+	UE_TRACE_EVENT_FIELD(uint64, TempHandle)
+	UE_TRACE_EVENT_FIELD(uint32, ThreadId)
+UE_TRACE_EVENT_END()
+
+UE_TRACE_EVENT_BEGIN(PlatformFile, EndOpen, Always)
+	UE_TRACE_EVENT_FIELD(uint64, Cycle)
+	UE_TRACE_EVENT_FIELD(uint64, TempHandle)
+	UE_TRACE_EVENT_FIELD(uint64, FileHandle)
+UE_TRACE_EVENT_END()
+
+UE_TRACE_EVENT_BEGIN(PlatformFile, BeginClose, Always)
+	UE_TRACE_EVENT_FIELD(uint64, Cycle)
+	UE_TRACE_EVENT_FIELD(uint64, TempHandle)
+	UE_TRACE_EVENT_FIELD(uint64, FileHandle)
+	UE_TRACE_EVENT_FIELD(uint32, ThreadId)
+UE_TRACE_EVENT_END()
+
+UE_TRACE_EVENT_BEGIN(PlatformFile, EndClose, Always)
+	UE_TRACE_EVENT_FIELD(uint64, Cycle)
+	UE_TRACE_EVENT_FIELD(uint64, TempHandle)
+UE_TRACE_EVENT_END()
+
+UE_TRACE_EVENT_BEGIN(PlatformFile, BeginRead, Always)
+	UE_TRACE_EVENT_FIELD(uint64, Cycle)
+	UE_TRACE_EVENT_FIELD(uint64, ReadHandle)
+	UE_TRACE_EVENT_FIELD(uint64, FileHandle)
+	UE_TRACE_EVENT_FIELD(uint64, Offset)
+	UE_TRACE_EVENT_FIELD(uint64, Size)
+	UE_TRACE_EVENT_FIELD(uint32, ThreadId)
+UE_TRACE_EVENT_END()
+
+UE_TRACE_EVENT_BEGIN(PlatformFile, EndRead, Always)
+	UE_TRACE_EVENT_FIELD(uint64, Cycle)
+	UE_TRACE_EVENT_FIELD(uint64, ReadHandle)
+	UE_TRACE_EVENT_FIELD(uint64, SizeRead)
+	UE_TRACE_EVENT_FIELD(uint32, ThreadId)
+UE_TRACE_EVENT_END()
+
+UE_TRACE_EVENT_BEGIN(PlatformFile, BeginWrite, Always)
+	UE_TRACE_EVENT_FIELD(uint64, Cycle)
+	UE_TRACE_EVENT_FIELD(uint64, WriteHandle)
+	UE_TRACE_EVENT_FIELD(uint64, FileHandle)
+	UE_TRACE_EVENT_FIELD(uint64, Offset)
+	UE_TRACE_EVENT_FIELD(uint64, Size)
+	UE_TRACE_EVENT_FIELD(uint32, ThreadId)
+UE_TRACE_EVENT_END()
+
+UE_TRACE_EVENT_BEGIN(PlatformFile, EndWrite, Always)
+	UE_TRACE_EVENT_FIELD(uint64, Cycle)
+	UE_TRACE_EVENT_FIELD(uint64, WriteHandle)
+	UE_TRACE_EVENT_FIELD(uint64, SizeWritten)
+	UE_TRACE_EVENT_FIELD(uint32, ThreadId)
+UE_TRACE_EVENT_END()
+
+void FPlatformFileTrace::BeginOpen(uint64 TempHandle, const TCHAR* Path)
+{
+	uint16 PathSize = (FCString::Strlen(Path) + 1) * sizeof(TCHAR);
+	UE_TRACE_LOG(PlatformFile, BeginOpen, PathSize)
+		<< BeginOpen.Cycle(FPlatformTime::Cycles64())
+		<< BeginOpen.TempHandle(TempHandle)
+		<< BeginOpen.Attachment(Path, PathSize)
+		<< BeginOpen.ThreadId(FPlatformTLS::GetCurrentThreadId());
+}
+
+void FPlatformFileTrace::EndOpen(uint64 TempHandle, uint64 FileHandle)
+{
+	UE_TRACE_LOG(PlatformFile, EndOpen)
+		<< EndOpen.Cycle(FPlatformTime::Cycles64())
+		<< EndOpen.TempHandle(TempHandle)
+		<< EndOpen.FileHandle(FileHandle);
+}
+
+void FPlatformFileTrace::BeginClose(uint64 TempHandle, uint64 FileHandle)
+{
+	UE_TRACE_LOG(PlatformFile, BeginClose)
+		<< BeginClose.Cycle(FPlatformTime::Cycles64())
+		<< BeginClose.TempHandle(TempHandle)
+		<< BeginClose.FileHandle(FileHandle)
+		<< BeginClose.ThreadId(FPlatformTLS::GetCurrentThreadId());
+}
+
+void FPlatformFileTrace::EndClose(uint64 TempHandle)
+{
+	UE_TRACE_LOG(PlatformFile, EndClose)
+		<< EndClose.Cycle(FPlatformTime::Cycles64())
+		<< EndClose.TempHandle(TempHandle);
+}
+
+void FPlatformFileTrace::BeginRead(uint64 ReadHandle, uint64 FileHandle, uint64 Offset, uint64 Size)
+{
+	UE_TRACE_LOG(PlatformFile, BeginRead)
+		<< BeginRead.Cycle(FPlatformTime::Cycles64())
+		<< BeginRead.ReadHandle(ReadHandle)
+		<< BeginRead.FileHandle(FileHandle)
+		<< BeginRead.Offset(Offset)
+		<< BeginRead.Size(Size)
+		<< BeginRead.ThreadId(FPlatformTLS::GetCurrentThreadId());
+}
+
+void FPlatformFileTrace::EndRead(uint64 ReadHandle, uint64 SizeRead)
+{
+	UE_TRACE_LOG(PlatformFile, EndRead)
+		<< EndRead.Cycle(FPlatformTime::Cycles64())
+		<< EndRead.ReadHandle(ReadHandle)
+		<< EndRead.SizeRead(SizeRead)
+		<< EndRead.ThreadId(FPlatformTLS::GetCurrentThreadId());
+}
+
+void FPlatformFileTrace::BeginWrite(uint64 WriteHandle, uint64 FileHandle, uint64 Offset, uint64 Size)
+{
+	UE_TRACE_LOG(PlatformFile, BeginWrite)
+		<< BeginWrite.Cycle(FPlatformTime::Cycles64())
+		<< BeginWrite.WriteHandle(WriteHandle)
+		<< BeginWrite.FileHandle(FileHandle)
+		<< BeginWrite.Offset(Offset)
+		<< BeginWrite.Size(Size)
+		<< BeginWrite.ThreadId(FPlatformTLS::GetCurrentThreadId());
+}
+
+void FPlatformFileTrace::EndWrite(uint64 WriteHandle, uint64 SizeWritten)
+{
+	UE_TRACE_LOG(PlatformFile, EndWrite)
+		<< EndWrite.Cycle(FPlatformTime::Cycles64())
+		<< EndWrite.WriteHandle(WriteHandle)
+		<< EndWrite.SizeWritten(SizeWritten)
+		<< EndWrite.ThreadId(FPlatformTLS::GetCurrentThreadId());
+}
+
+#endif
\ No newline at end of file
diff --git a/Engine/Source/Runtime/Core/Private/Stats/Stats2.cpp b/Engine/Source/Runtime/Core/Private/Stats/Stats2.cpp
--- a/Engine/Source/Runtime/Core/Private/Stats/Stats2.cpp
+++ b/Engine/Source/Runtime/Core/Private/Stats/Stats2.cpp
@@ -602,6 +602,21 @@ public:
 		{
 			Found->AlwaysEnabledNamesInThisGroup.Add( Stat, Result );
 		}
+
+#if CPUPROFILERTRACE_ENABLED
+		if (bCycleStat)
+		{
+			Result->TraceCpuProfilerSpecId = FCpuProfilerTrace::OutputEventType(*StatDescription, CpuProfilerGroup_Stats);
+		}
+#endif
+#if STATSTRACE_ENABLED
+		if (!bCycleStat && (InStatType == EStatDataType::ST_int64 || InStatType == EStatDataType::ST_double))
+		{
+			ANSICHAR NameBuffer[1024];
+			StatShortName.GetPlainANSIString(NameBuffer);
+			FStatsTrace::DeclareStat(Stat, NameBuffer, *StatDescription, InStatType == EStatDataType::ST_double, MemoryRegion != FPlatformMemory::MCR_Invalid, bShouldClearEveryFrame);
+		}
+#endif
 		return TStatId(Result);
 	}
 
diff --git a/Engine/Source/Runtime/Core/Private/Windows/WindowsAsyncIO.h b/Engine/Source/Runtime/Core/Private/Windows/WindowsAsyncIO.h
--- a/Engine/Source/Runtime/Core/Private/Windows/WindowsAsyncIO.h
+++ b/Engine/Source/Runtime/Core/Private/Windows/WindowsAsyncIO.h
@@ -2,6 +2,8 @@
 
 #pragma once
 
+#include "ProfilingDebugging/PlatformFileTrace.h"
+
 class FWindowsReadRequest;
 class FWindowsAsyncReadFileHandle;
 
@@ -125,6 +127,7 @@ public:
 				OverlappedIO.OffsetHigh = LI.HighPart;
 			}
 			OverlappedIO.hEvent = GetIOPooledEvent();
+			TRACE_PLATFORMFILE_BEGIN_READ(&OverlappedIO, FileHandle, AlignedOffset, AlignedBytesToRead);
 			if (!ReadFile(FileHandle, TempMemory ? TempMemory : Memory, AlignedBytesToRead, (LPDWORD)&NumRead, &OverlappedIO))
 			{
 				uint32 ErrorCode = GetLastError();
@@ -162,11 +165,16 @@ public:
 
 		if (GTriggerFailedWindowsRead || !GetOverlappedResult(FileHandle, &OverlappedIO, (LPDWORD)&BytesRead, TRUE))
 		{
+			TRACE_PLATFORMFILE_END_READ(&OverlappedIO, 0);
 			GTriggerFailedWindowsRead = false;
 			uint32 ErrorCode = GetLastError();
 			FailedMessage = FString::Printf(TEXT("FWindowsReadRequest GetOverlappedResult Code = %x Offset = %lld Size = %lld FileSize = %lld File = %s"), ErrorCode, AlignedOffset, AlignedBytesToRead, FileSize, GetFileNameForErrorMessagesAndPanicRetry());
 			bFailed = true;
 		}
+		else
+		{
+			TRACE_PLATFORMFILE_END_READ(&OverlappedIO, BytesRead);
+		}
 		if (!bFailed && int64(BytesRead) < BytesToRead + (Offset - AlignedOffset))
 		{
 			uint32 ErrorCode = GetLastError();
@@ -383,7 +391,9 @@ public:
 		FScopeLock Lock(&LiveRequestsCritical);
 		check(!LiveRequests.Num()); // must delete all requests before you delete the handle
 #endif
+		TRACE_PLATFORMFILE_BEGIN_CLOSE(FileHandle);
 		CloseHandle(FileHandle);
+		TRACE_PLATFORMFILE_END_CLOSE();
 	}
 	void RemoveRequest(FWindowsReadRequest* Req)
 	{
diff --git a/Engine/Source/Runtime/Core/Private/Windows/WindowsPlatformFile.cpp b/Engine/Source/Runtime/Core/Private/Windows/WindowsPlatformFile.cpp
--- a/Engine/Source/Runtime/Core/Private/Windows/WindowsPlatformFile.cpp
+++ b/Engine/Source/Runtime/Core/Private/Windows/WindowsPlatformFile.cpp
@@ -19,6 +19,7 @@
 #include "Async/AsyncWork.h"
 #include "Misc/ScopeLock.h"
 #include "HAL/IConsoleManager.h"
+#include "ProfilingDebugging/PlatformFileTrace.h"
 
 #include "Windows/AllowWindowsPlatformTypes.h"
 
@@ -178,7 +179,9 @@ protected:
 		if (Handle != nullptr)
 		{
 			// Close the file handle
+			TRACE_PLATFORMFILE_BEGIN_CLOSE(Handle);
 			CloseHandle(Handle);
+			TRACE_PLATFORMFILE_END_CLOSE();
 			Handle = nullptr;
 		}
 		return true;
@@ -225,11 +228,13 @@ protected:
 		uint32 NumRead = 0;
 		if (GetOverlappedResult(Handle, &OverlappedIO, (::DWORD*)&NumRead, true) != false)
 		{
+			TRACE_PLATFORMFILE_END_READ(&OverlappedIO, NumRead);
 			UpdateFileOffsetAfterRead(NumRead);
 			return true;
 		}
 		else if (GetLastError() == ERROR_HANDLE_EOF)
 		{
+			TRACE_PLATFORMFILE_END_READ(&OverlappedIO, 0);
 			bIsAtEOF = true;
 			return true;
 		}
@@ -244,11 +249,13 @@ protected:
 			CurrentAsyncReadBuffer = BufferToReadInto;
 			uint32 NumRead = 0;
 			// Now kick off an async read
+			TRACE_PLATFORMFILE_BEGIN_READ(&OverlappedIO, Handle, OverlappedFilePos, BufferSize);
 			if (!ReadFile(Handle, Buffers[BufferToReadInto], BufferSize, (::DWORD*)&NumRead, &OverlappedIO))
 			{
 				uint32 ErrorCode = GetLastError();
 				if (ErrorCode != ERROR_IO_PENDING)
 				{
+					TRACE_PLATFORMFILE_END_READ(&OverlappedIO, 0);
 					bIsAtEOF = true;
 					bHasReadOutstanding = false;
 				}
@@ -256,6 +263,7 @@ protected:
 			else
 			{
 				// Read completed immediately
+				TRACE_PLATFORMFILE_END_READ(&OverlappedIO, NumRead);
 				UpdateFileOffsetAfterRead(NumRead);
 			}
 		}
@@ -506,7 +514,9 @@ public:
 	}
 	virtual ~FFileHandleWindows()
 	{
+		TRACE_PLATFORMFILE_BEGIN_CLOSE(FileHandle);
 		CloseHandle(FileHandle);
+		TRACE_PLATFORMFILE_END_CLOSE();
 		FileHandle = NULL;
 	}
 	virtual int64 Tell(void) override
@@ -541,12 +551,14 @@ public:
 		check(IsValid());
 		uint32 NumRead = 0;
 		// Now kick off an async read
+		TRACE_PLATFORMFILE_BEGIN_READ(&OverlappedIO, FileHandle, FilePos, BytesToRead);
 		if (!ReadFile(FileHandle, Destination, BytesToRead, (::DWORD*)&NumRead, &OverlappedIO))
 		{
 			uint32 ErrorCode = GetLastError();
 			if (ErrorCode != ERROR_IO_PENDING)
 			{
 				// Read failed
+				TRACE_PLATFORMFILE_END_READ(&OverlappedIO, 0);
 				return false;
 			}
 			// Wait for the read to complete
@@ -554,9 +566,11 @@ public:
 			if (!GetOverlappedResult(FileHandle, &OverlappedIO, (::DWORD*)&NumRead, true))
 			{
 				// Read failed
+				TRACE_PLATFORMFILE_END_READ(&OverlappedIO, 0);
 				return false;
 			}
 		}
+		TRACE_PLATFORMFILE_END_READ(&OverlappedIO, NumRead);
 		// Update where we are in the file
 		FilePos += NumRead;
 		UpdateOverlappedPos();
@@ -567,12 +581,14 @@ public:
 		check(IsValid());
 		uint32 NumWritten = 0;
 		// Now kick off an async write
+		TRACE_PLATFORMFILE_BEGIN_WRITE(this, FileHandle, FilePos, BytesToWrite);
 		if (!WriteFile(FileHandle, Source, BytesToWrite, (::DWORD*)&NumWritten, &OverlappedIO))
 		{
 			uint32 ErrorCode = GetLastError();
 			if (ErrorCode != ERROR_IO_PENDING)
 			{
 				// Write failed
+				TRACE_PLATFORMFILE_END_WRITE(this, 0);
 				return false;
 			}
 			// Wait for the write to complete
@@ -580,9 +596,11 @@ public:
 			if (!GetOverlappedResult(FileHandle, &OverlappedIO, (::DWORD*)&NumWritten, true))
 			{
 				// Write failed
+				TRACE_PLATFORMFILE_END_WRITE(this, 0);
 				return false;
 			}
 		}
+		TRACE_PLATFORMFILE_END_WRITE(this, NumWritten);
 		// Update where we are in the file
 		FilePos += NumWritten;
 		UpdateOverlappedPos();
@@ -657,7 +675,9 @@ public:
 	{
 		check(!NumOutstandingRegions); // can't delete the file before you delete all outstanding regions
 		CloseHandle(MappingHandle);
+		TRACE_PLATFORMFILE_BEGIN_CLOSE(Handle);
 		CloseHandle(Handle);
+		TRACE_PLATFORMFILE_END_CLOSE();
 	}
 	virtual IMappedFileRegion* MapRegion(int64 Offset = 0, int64 BytesToMap = MAX_int64, bool bPreloadHint = false) override
 	{
@@ -788,7 +808,9 @@ public:
 
 	virtual void SetTimeStamp(const TCHAR* Filename, FDateTime DateTime) override
 	{
+		TRACE_PLATFORMFILE_BEGIN_OPEN(Filename);
 		HANDLE Handle = CreateFileW(*NormalizeFilename(Filename), FILE_WRITE_ATTRIBUTES, FILE_SHARE_READ, nullptr, OPEN_EXISTING, FILE_ATTRIBUTE_NORMAL | FILE_FLAG_OVERLAPPED, nullptr);
+		TRACE_PLATFORMFILE_END_OPEN(Handle);
 		if (Handle != INVALID_HANDLE_VALUE)
 		{
 			const FILETIME ModificationFileTime = UEDateTimeToWindowsFileTime(DateTime);
@@ -796,7 +818,9 @@ public:
 			{
 				UE_LOG(LogTemp, Warning, TEXT("SetTimeStamp: Failed to SetFileTime on %s"), Filename);
 			}
+			TRACE_PLATFORMFILE_BEGIN_CLOSE(Handle);
 			CloseHandle(Handle);
+			TRACE_PLATFORMFILE_END_CLOSE();
 		}
 		else
 		{
@@ -817,7 +841,9 @@ public:
 
 	virtual FString GetFilenameOnDisk(const TCHAR* Filename) override
 	{
+		TRACE_PLATFORMFILE_BEGIN_OPEN(Filename);
 		HANDLE hFile = CreateFile(Filename, FILE_READ_ATTRIBUTES, FILE_SHARE_READ | FILE_SHARE_WRITE | FILE_SHARE_DELETE, NULL, OPEN_EXISTING, 0, NULL);
+		TRACE_PLATFORMFILE_END_OPEN(hFile);
 		if(hFile == INVALID_HANDLE_VALUE)
 		{
 			return NormalizeFilename(Filename);
@@ -842,8 +868,9 @@ public:
 					break;
 				}
 			}
-
+			TRACE_PLATFORMFILE_BEGIN_CLOSE(hFile);
 			CloseHandle(hFile);
+			TRACE_PLATFORMFILE_END_CLOSE();
 
 			NormalizedFileName.RemoveFromStart(TEXT("\\\\?\\"), ESearchCase::CaseSensitive);
 			NormalizedFileName.ReplaceInline(TEXT("\\"), TEXT("/"));
@@ -862,7 +889,9 @@ public:
 
 
 		FString NormalizedFilename = NormalizeFilename(Filename);
+		TRACE_PLATFORMFILE_BEGIN_OPEN(Filename);
 		HANDLE Handle = CreateFileW(*NormalizedFilename, Access, WinFlags, NULL, Create, FILE_ATTRIBUTE_NORMAL | FILE_FLAG_OVERLAPPED, NULL);
+		TRACE_PLATFORMFILE_END_OPEN(Handle);
 		// we can't really fail here because this is intended to be an async open
 		return new FWindowsAsyncReadFileHandle(Handle, *NormalizedFilename);
 
@@ -877,13 +906,17 @@ public:
 #define USE_OVERLAPPED_IO (!IS_PROGRAM && !WITH_EDITOR)		// Use straightforward synchronous I/O in cooker/editor
 
 #if USE_OVERLAPPED_IO
+		TRACE_PLATFORMFILE_BEGIN_OPEN(Filename);
 		HANDLE Handle    = CreateFileW(*NormalizeFilename(Filename), Access, WinFlags, NULL, Create, FILE_ATTRIBUTE_NORMAL | FILE_FLAG_OVERLAPPED, NULL);
+		TRACE_PLATFORMFILE_END_OPEN(Handle);
 		if (Handle != INVALID_HANDLE_VALUE)
 		{
 			return new FAsyncBufferedFileReaderWindows(Handle);
 		}
 #else
+		TRACE_PLATFORMFILE_BEGIN_OPEN(Filename);
 		HANDLE Handle = CreateFileW(*NormalizeFilename(Filename), Access, WinFlags, NULL, Create, FILE_ATTRIBUTE_NORMAL, NULL);
+		TRACE_PLATFORMFILE_END_OPEN(Handle);
 		if (Handle != INVALID_HANDLE_VALUE)
 		{
 			return new FFileHandleWindows(Handle);
@@ -897,7 +930,9 @@ public:
 		uint32  Access = GENERIC_READ;
 		uint32  WinFlags = FILE_SHARE_READ | (bAllowWrite ? FILE_SHARE_WRITE : 0);
 		uint32  Create = OPEN_EXISTING;
+		TRACE_PLATFORMFILE_BEGIN_OPEN(Filename);
 		HANDLE Handle = CreateFileW(*NormalizeFilename(Filename), Access, WinFlags, NULL, Create, FILE_ATTRIBUTE_NORMAL | FILE_FLAG_OVERLAPPED, NULL);
+		TRACE_PLATFORMFILE_END_OPEN(Handle);
 		if (Handle != INVALID_HANDLE_VALUE)
 		{
 			return new FFileHandleWindows(Handle);
@@ -910,7 +945,9 @@ public:
 		uint32  Access    = GENERIC_WRITE | (bAllowRead ? GENERIC_READ : 0);
 		uint32  WinFlags  = bAllowRead ? FILE_SHARE_READ : 0;
 		uint32  Create    = bAppend ? OPEN_ALWAYS : CREATE_ALWAYS;
+		TRACE_PLATFORMFILE_BEGIN_OPEN(Filename);
 		HANDLE Handle    = CreateFileW(*NormalizeFilename(Filename), Access, WinFlags, NULL, Create, FILE_ATTRIBUTE_NORMAL, NULL);
+		TRACE_PLATFORMFILE_END_OPEN(Handle);
 		if(Handle != INVALID_HANDLE_VALUE)
 		{
 			FFileHandleWindows *PlatformFileHandle = new FFileHandleWindows(Handle);
@@ -934,7 +971,9 @@ public:
 		uint32  Access = GENERIC_READ;
 		uint32  WinFlags = FILE_SHARE_READ;
 		uint32  Create = OPEN_EXISTING;
+		TRACE_PLATFORMFILE_BEGIN_OPEN(Filename);
 		HANDLE Handle = CreateFileW(*NormalizeFilename(Filename), Access, WinFlags, NULL, Create, FILE_ATTRIBUTE_NORMAL, NULL);
+		TRACE_PLATFORMFILE_END_OPEN(Handle);
 		if (Handle == INVALID_HANDLE_VALUE)
 		{
 			return nullptr;
@@ -942,7 +981,9 @@ public:
 		HANDLE MappingHandle = CreateFileMapping(Handle, NULL, PAGE_READONLY, 0, 0, NULL);
 		if (MappingHandle == INVALID_HANDLE_VALUE)
 		{
+			TRACE_PLATFORMFILE_BEGIN_CLOSE(Handle);
 			CloseHandle(Handle);
+			TRACE_PLATFORMFILE_END_CLOSE();
 			return nullptr;
 		}
 		return new FMappedFileHandleWindows(Handle, MappingHandle, Size, Filename);
diff --git a/Engine/Source/Runtime/Core/Public/HAL/PlatformMisc.h b/Engine/Source/Runtime/Core/Public/HAL/PlatformMisc.h
--- a/Engine/Source/Runtime/Core/Public/HAL/PlatformMisc.h
+++ b/Engine/Source/Runtime/Core/Public/HAL/PlatformMisc.h
@@ -30,6 +30,8 @@
 
 #if ENABLE_NAMED_EVENTS
 
+#include "ProfilingDebugging/CpuProfilerTrace.h"
+
 class CORE_API FScopedNamedEvent
 {
 public:
@@ -110,9 +112,16 @@ public:
 #define NAMED_EVENT_STR(x) L##x
 #endif
 
-#define SCOPED_NAMED_EVENT(Name, Color)          FScopedNamedEventStatic ANONYMOUS_VARIABLE(NamedEvent_##Name##_)(Color, NAMED_EVENT_STR(#Name));
+#define SCOPED_NAMED_EVENT(Name, Color) \
+	TRACE_CPUPROFILER_EVENT_SCOPE(Name) \
+	FScopedNamedEventStatic ANONYMOUS_VARIABLE(NamedEvent_##Name##_)(Color, NAMED_EVENT_STR(#Name));
+
 #define SCOPED_NAMED_EVENT_FSTRING(Text, Color)  FScopedNamedEvent       ANONYMOUS_VARIABLE(NamedEvent_)         (Color, *Text);
-#define SCOPED_NAMED_EVENT_TEXT(Text, Color)     FScopedNamedEventStatic ANONYMOUS_VARIABLE(NamedEvent_)         (Color, NAMED_EVENT_STR(Text));
+
+#define SCOPED_NAMED_EVENT_TEXT(Text, Color) \
+	TRACE_CPUPROFILER_EVENT_SCOPE_TEXT(TEXT(Text)) \
+	FScopedNamedEventStatic ANONYMOUS_VARIABLE(NamedEvent_)         (Color, NAMED_EVENT_STR(Text));
+
 #define SCOPED_NAMED_EVENT_F(Format, Color, ...) FScopedNamedEvent       ANONYMOUS_VARIABLE(NamedEvent_)         (Color, *FString::Printf(Format, __VA_ARGS__));
 #define SCOPED_PROFILER_COLOR(Color)			 FScopedProfilerColor    ANONYMOUS_VARIABLE(ProfilerColor_##Name##_)(Color);
 
diff --git a/Engine/Source/Runtime/Core/Public/Logging/LogTrace.h b/Engine/Source/Runtime/Core/Public/Logging/LogTrace.h
new file mode 100644
--- /dev/null
+++ b/Engine/Source/Runtime/Core/Public/Logging/LogTrace.h
@@ -0,0 +1,53 @@
+// Copyright 1998-2019 Epic Games, Inc. All Rights Reserved.
+
+#pragma once
+
+#include "CoreTypes.h"
+
+#if !IS_PROGRAM && !UE_BUILD_SHIPPING && (PLATFORM_WINDOWS || PLATFORM_PS4 || PLATFORM_XBOXONE)
+#define LOGTRACE_ENABLED 1
+#else
+#define LOGTRACE_ENABLED 0
+#endif
+
+#if LOGTRACE_ENABLED
+#include "ProfilingDebugging/FormatArgsTrace.h"
+
+struct FLogCategoryBase;
+
+struct FLogTrace
+{
+	CORE_API static void OutputLogCategory(const FLogCategoryBase* Category, const TCHAR* Name, ELogVerbosity::Type DefaultVerbosity);
+	CORE_API static void OutputLogMessageSpec(const void* LogPoint, const FLogCategoryBase* Category, ELogVerbosity::Type Verbosity, const ANSICHAR* File, int32 Line, const TCHAR* Format);
+
+	template <typename... Types>
+	static void OutputLogMessage(const void* LogPoint, Types... FormatArgs)
+	{
+		uint8 FormatArgsBuffer[4096];
+		uint16 FormatArgsSize = FFormatArgsTrace::EncodeArguments(FormatArgsBuffer, FormatArgs...);
+		if (FormatArgsSize)
+		{
+			OutputLogMessageInternal(LogPoint, FormatArgsSize, FormatArgsBuffer);
+		}
+	}
+
+private:
+	CORE_API static void OutputLogMessageInternal(const void* LogPoint, uint16 EncodedFormatArgsSize, uint8* EncodedFormatArgs);
+};
+
+#define TRACE_LOG_CATEGORY(Category, Name, DefaultVerbosity) \
+	FLogTrace::OutputLogCategory(Category, Name, DefaultVerbosity);
+
+#define TRACE_LOG_MESSAGE(Category, Verbosity, Format, ...) \
+	static bool PREPROCESSOR_JOIN(__LogPoint, __LINE__); \
+	if (!PREPROCESSOR_JOIN(__LogPoint, __LINE__)) \
+	{ \
+		FLogTrace::OutputLogMessageSpec(&PREPROCESSOR_JOIN(__LogPoint, __LINE__), &Category, ELogVerbosity::Verbosity, __FILE__, __LINE__, Format); \
+		PREPROCESSOR_JOIN(__LogPoint, __LINE__) = true; \
+	} \
+	FLogTrace::OutputLogMessage(&PREPROCESSOR_JOIN(__LogPoint, __LINE__), ##__VA_ARGS__);
+
+#else
+#define TRACE_LOG_CATEGORY(Category, Name, DefaultVerbosity)
+#define TRACE_LOG_MESSAGE(Category, Verbosity, Format, ...)
+#endif
\ No newline at end of file
diff --git a/Engine/Source/Runtime/Core/Public/ProfilingDebugging/CpuProfilerTrace.h b/Engine/Source/Runtime/Core/Public/ProfilingDebugging/CpuProfilerTrace.h
new file mode 100644
--- /dev/null
+++ b/Engine/Source/Runtime/Core/Public/ProfilingDebugging/CpuProfilerTrace.h
@@ -0,0 +1,67 @@
+// Copyright 1998-2019 Epic Games, Inc. All Rights Reserved.
+
+#pragma once
+
+#include "CoreTypes.h"
+
+#if !IS_PROGRAM && !UE_BUILD_SHIPPING && (PLATFORM_WINDOWS || PLATFORM_PS4 || PLATFORM_XBOXONE)
+#define CPUPROFILERTRACE_ENABLED 1
+#else
+#define CPUPROFILERTRACE_ENABLED 0
+#endif
+
+enum ECpuProfilerGroup
+{
+	CpuProfilerGroup_Default,
+	CpuProfilerGroup_Stats,
+	CpuProfilerGroup_LoadTime,
+	CpuProfilerGroup_CsvProfiler,
+};
+
+#if CPUPROFILERTRACE_ENABLED
+
+struct FCpuProfilerTrace
+{
+	CORE_API static void Init(bool bStartEnabled);
+	CORE_API static uint16 OutputEventType(const TCHAR* Name, ECpuProfilerGroup Group);
+	CORE_API static void OutputBeginEvent(uint16 SpecId);
+	CORE_API static void OutputEndEvent();
+
+	struct FEventScope
+	{
+		FEventScope(uint16 InSpecId)
+		{
+			OutputBeginEvent(InSpecId);
+		}
+
+		~FEventScope()
+		{
+			OutputEndEvent();
+		}
+	};
+};
+
+#define TRACE_CPUPROFILER_EVENT_SCOPE_TEXT_GROUP(Name, Group) \
+	static uint16 PREPROCESSOR_JOIN(__CpuProfilerEventSpecId, __LINE__); \
+	if (PREPROCESSOR_JOIN(__CpuProfilerEventSpecId, __LINE__) == 0) { \
+		PREPROCESSOR_JOIN(__CpuProfilerEventSpecId, __LINE__) = FCpuProfilerTrace::OutputEventType(Name, Group); \
+	} \
+	FCpuProfilerTrace::FEventScope PREPROCESSOR_JOIN(__CpuProfilerEventScope, __LINE__)(PREPROCESSOR_JOIN(__CpuProfilerEventSpecId, __LINE__));
+
+#define TRACE_CPUPROFILER_EVENT_SCOPE_TEXT(Name) \
+	TRACE_CPUPROFILER_EVENT_SCOPE_TEXT_GROUP(Name, CpuProfilerGroup_Default)
+
+#define TRACE_CPUPROFILER_EVENT_SCOPE_GROUP(Name, Group) \
+	TRACE_CPUPROFILER_EVENT_SCOPE_TEXT_GROUP(TEXT(#Name), Group);
+
+#define TRACE_CPUPROFILER_EVENT_SCOPE(Name) \
+	TRACE_CPUPROFILER_EVENT_SCOPE_TEXT_GROUP(TEXT(#Name), CpuProfilerGroup_Default);
+
+#else
+
+#define TRACE_CPUPROFILER_EVENT_SCOPE_GROUP(Name, Group)
+#define TRACE_CPUPROFILER_EVENT_SCOPE(Name)
+#define TRACE_CPUPROFILER_EVENT_SCOPE_TEXT_GROUP(Name, Group)
+#define TRACE_CPUPROFILER_EVENT_SCOPE_TEXT(Name)
+
+#endif
diff --git a/Engine/Source/Runtime/Core/Public/ProfilingDebugging/CsvProfiler.h b/Engine/Source/Runtime/Core/Public/ProfilingDebugging/CsvProfiler.h
--- a/Engine/Source/Runtime/Core/Public/ProfilingDebugging/CsvProfiler.h
+++ b/Engine/Source/Runtime/Core/Public/ProfilingDebugging/CsvProfiler.h
@@ -14,6 +14,7 @@
 #include "Async/Future.h"
 #include "Async/TaskGraphInterfaces.h"
 #include "Misc/EnumClassFlags.h"
+#include "ProfilingDebugging/MiscTrace.h"
 
 // Whether to allow the CSV profiler in shipping builds.
 // Enable in a .Target.cs file if required.
@@ -44,9 +45,15 @@
 #define CSV_STAT_FNAME(StatName)								(_GCsvStat_##StatName.Name)
 
 // Inline stats (no up front definition)
-#define CSV_SCOPED_TIMING_STAT(Category,StatName)				FScopedCsvStat _ScopedCsvStat_ ## StatName (#StatName, CSV_CATEGORY_INDEX(Category));
-#define CSV_SCOPED_TIMING_STAT_GLOBAL(StatName)					FScopedCsvStat _ScopedCsvStat_ ## StatName (#StatName, CSV_CATEGORY_INDEX_GLOBAL);
-#define CSV_SCOPED_TIMING_STAT_EXCLUSIVE(StatName)				FScopedCsvStatExclusive _ScopedCsvStatExclusive_ ## StatName (#StatName);
+#define CSV_SCOPED_TIMING_STAT(Category,StatName) \
+	FScopedCsvStat _ScopedCsvStat_ ## StatName (#StatName, CSV_CATEGORY_INDEX(Category)); \
+	TRACE_CPUPROFILER_EVENT_SCOPE_GROUP(StatName, CpuProfilerGroup_CsvProfiler)
+#define CSV_SCOPED_TIMING_STAT_GLOBAL(StatName) \
+	FScopedCsvStat _ScopedCsvStat_ ## StatName (#StatName, CSV_CATEGORY_INDEX_GLOBAL); \
+	TRACE_CPUPROFILER_EVENT_SCOPE_GROUP(StatName, CpuProfilerGroup_CsvProfiler)
+#define CSV_SCOPED_TIMING_STAT_EXCLUSIVE(StatName) \
+	FScopedCsvStatExclusive _ScopedCsvStatExclusive_ ## StatName (#StatName); \
+	TRACE_CPUPROFILER_EVENT_SCOPE_GROUP(StatName, CpuProfilerGroup_CsvProfiler)
 #define CSV_SCOPED_TIMING_STAT_EXCLUSIVE_CONDITIONAL(StatName,Condition) FScopedCsvStatExclusiveConditional _ScopedCsvStatExclusive_ ## StatName (#StatName,Condition);
 
 #define CSV_SCOPED_WAIT_CONDITIONAL(Condition)					FScopedCsvWaitConditional _ScopedCsvWait(Condition);
@@ -70,8 +77,13 @@
 #define CSV_DECLARE_CATEGORY_MODULE_EXTERN(Module_API,CategoryName)			extern Module_API FCsvCategory _GCsvCategory_##CategoryName
 
 // Events
-#define CSV_EVENT(Category, Format, ...) 						FCsvProfiler::RecordEventf( CSV_CATEGORY_INDEX(Category), Format, ##__VA_ARGS__ )
-#define CSV_EVENT_GLOBAL(Format, ...) 							FCsvProfiler::RecordEventf( CSV_CATEGORY_INDEX_GLOBAL, Format, ##__VA_ARGS__ )
+#define CSV_EVENT(Category, Format, ...) \
+	FCsvProfiler::RecordEventf( CSV_CATEGORY_INDEX(Category), Format, ##__VA_ARGS__ ); \
+	TRACE_BOOKMARK(Format, ##__VA_ARGS__)
+
+#define CSV_EVENT_GLOBAL(Format, ...) \
+	FCsvProfiler::RecordEventf( CSV_CATEGORY_INDEX_GLOBAL, Format, ##__VA_ARGS__ ); \
+	TRACE_BOOKMARK(Format, ##__VA_ARGS__)
 
 // Metadata
 #define CSV_METADATA(Key,Value)									FCsvProfiler::SetMetadata( Key, Value )
diff --git a/Engine/Source/Runtime/Core/Public/ProfilingDebugging/FormatArgsTrace.h b/Engine/Source/Runtime/Core/Public/ProfilingDebugging/FormatArgsTrace.h
new file mode 100644
--- /dev/null
+++ b/Engine/Source/Runtime/Core/Public/ProfilingDebugging/FormatArgsTrace.h
@@ -0,0 +1,130 @@
+// Copyright 1998-2019 Epic Games, Inc. All Rights Reserved.
+
+#pragma once
+
+#include "CoreTypes.h"
+#include "Misc/CString.h"
+#include "Templates/IsFloatingPoint.h"
+#include "Templates/UnrealTypeTraits.h"
+#include "Templates/UnrealTemplate.h"
+
+struct FFormatArgsTrace
+{
+	enum EFormatArgTypeCode
+	{
+		FormatArgTypeCode_CategoryBitShift = 6,
+		FormatArgTypeCode_SizeBitMask = (1 << FormatArgTypeCode_CategoryBitShift) - 1,
+		FormatArgTypeCode_CategoryBitMask = ~FormatArgTypeCode_SizeBitMask,
+		FormatArgTypeCode_CategoryInteger = 1 << FormatArgTypeCode_CategoryBitShift,
+		FormatArgTypeCode_CategoryFloatingPoint = 2 << FormatArgTypeCode_CategoryBitShift,
+		FormatArgTypeCode_CategoryString = 3 << FormatArgTypeCode_CategoryBitShift,
+	};
+
+	template <int BufferSize, typename... Types>
+	static uint16 EncodeArguments(uint8(&Buffer)[BufferSize], Types... FormatArgs)
+	{
+		static_assert(BufferSize < 65536, "Maximum buffer size of 16 bits exceeded");
+		uint64 FormatArgsCount = sizeof...(FormatArgs);
+		if (FormatArgsCount >= 256)
+		{
+			// Nope
+			return 0;
+		}
+		uint64 FormatArgsSize = 1 + FormatArgsCount + GetArgumentsEncodedSize(FormatArgs...);
+		if (FormatArgsSize > BufferSize)
+		{
+			// Nope
+			return 0;
+		}
+		uint8* TypeCodesBufferPtr = Buffer;
+		*TypeCodesBufferPtr++ = uint8(FormatArgsCount);
+		uint8* PayloadBufferPtr = TypeCodesBufferPtr + FormatArgsCount;
+		EncodeArgumentsInternal(TypeCodesBufferPtr, PayloadBufferPtr, FormatArgs...);
+		check(PayloadBufferPtr - Buffer == FormatArgsSize);
+		return FormatArgsSize;
+	}
+
+private:
+	template <typename T>
+	struct TIsStringArgument
+	{
+		enum { Value = TAnd<TIsPointer<T>, TIsCharType<typename TRemoveCV<typename TRemovePointer<T>::Type>::Type>>::Value };
+	};
+
+	template <typename T>
+	constexpr static typename TEnableIf<!TIsStringArgument<T>::Value, uint64>::Type GetArgumentEncodedSize(T Argument)
+	{
+		return sizeof(T);
+	}
+
+	template <typename T, typename CharType = typename TRemoveCV<typename TRemovePointer<T>::Type>::Type>
+	static typename TEnableIf<TIsStringArgument<T>::Value, uint64>::Type GetArgumentEncodedSize(T Argument)
+	{
+		return (TCString<CharType>::Strlen(Argument) + 1) * sizeof(CharType);
+	}
+
+	constexpr static uint64 GetArgumentsEncodedSize()
+	{
+		return 0;
+	}
+
+	template <typename T, typename... Types>
+	static uint64 GetArgumentsEncodedSize(T Head, Types... Tail)
+	{
+		return GetArgumentEncodedSize(Head) + GetArgumentsEncodedSize(Tail...);
+	}
+
+	template <typename T>
+	static typename TEnableIf<TAnd<TNot<TIsFloatingPoint<T>>, TNot<TIsStringArgument<T>>>::Value>::Type EncodeArgumentInternal(uint8*& TypeCodesPtr, uint8*& PayloadPtr, T Argument)
+	{
+		*TypeCodesPtr++ = FormatArgTypeCode_CategoryInteger | sizeof(T);
+
+#if PLATFORM_SUPPORTS_UNALIGNED_LOADS
+		*reinterpret_cast<T*>(PayloadPtr) = Argument;
+#else
+		// For ARM targets, it's possible that using __packed here would be preferable
+		// but I have not checked the codegen -- it's possible that the compiler generates
+		// the same code for this fixed size memcpy
+		memcpy(PayloadPtr, &Argument, sizeof Argument);
+#endif
+
+		PayloadPtr += sizeof(T);
+	}
+
+	template <typename T>
+	static typename TEnableIf<TIsFloatingPoint<T>::Value>::Type EncodeArgumentInternal(uint8*& TypeCodesPtr, uint8*& PayloadPtr, T Argument)
+	{
+		*TypeCodesPtr++ = FormatArgTypeCode_CategoryFloatingPoint | sizeof(T);
+
+#if PLATFORM_SUPPORTS_UNALIGNED_LOADS
+		*reinterpret_cast<T*>(PayloadPtr) = Argument;
+#else
+		// For ARM targets, it's possible that using __packed here would be preferable
+		// but I have not checked the codegen -- it's possible that the compiler generates
+		// the same code for this fixed size memcpy
+		memcpy(PayloadPtr, &Argument, sizeof Argument);
+#endif
+
+		PayloadPtr += sizeof(T);
+	}
+
+	template <typename T, typename CharType = typename TRemoveCV<typename TRemovePointer<T>::Type>::Type>
+	static typename TEnableIf<TIsStringArgument<T>::Value>::Type EncodeArgumentInternal(uint8*& TypeCodesPtr, uint8*& PayloadPtr, T Argument)
+	{
+		*TypeCodesPtr++ = FormatArgTypeCode_CategoryString | sizeof(CharType);
+		uint16 Length = (TCString<CharType>::Strlen(Argument) + 1) * sizeof(CharType);
+		memcpy(PayloadPtr, Argument, Length);
+		PayloadPtr += Length;
+	}
+
+	constexpr static void EncodeArgumentsInternal(uint8*& TypeCodesPtr, uint8*& PayloadPtr)
+	{
+	}
+
+	template <typename T, typename... Types>
+	static void EncodeArgumentsInternal(uint8*& ArgDescriptorsPtr, uint8*& ArgPayloadPtr, T Head, Types... Tail)
+	{
+		EncodeArgumentInternal(ArgDescriptorsPtr, ArgPayloadPtr, Head);
+		EncodeArgumentsInternal(ArgDescriptorsPtr, ArgPayloadPtr, Tail...);
+	}
+};
\ No newline at end of file
diff --git a/Engine/Source/Runtime/Core/Public/ProfilingDebugging/LoadTimeTracker.h b/Engine/Source/Runtime/Core/Public/ProfilingDebugging/LoadTimeTracker.h
--- a/Engine/Source/Runtime/Core/Public/ProfilingDebugging/LoadTimeTracker.h
+++ b/Engine/Source/Runtime/Core/Public/ProfilingDebugging/LoadTimeTracker.h
@@ -11,6 +11,7 @@
 #include "Containers/Map.h"
 #include "UObject/NameTypes.h"
 #include "ProfilingDebugging/ScopedTimers.h"
+#include "Trace/Trace.h"
 
 #ifndef ENABLE_LOADTIME_TRACKING
 	#define ENABLE_LOADTIME_TRACKING 0
@@ -211,6 +212,7 @@ struct CORE_API FScopedLoadTimeAccumulatorTimer : public FScopedDurationTimer
 #define SCOPED_LOADTIMER(TimerName) FScopedDurationTimer DurationTimer_##TimerName(FLoadTimeTracker::Get().TimerName);
 #define SCOPED_LOADTIMER_CNT(TimerName) FScopedDurationTimer DurationTimer_##TimerName(FLoadTimeTracker::Get().TimerName); FLoadTimeTracker::Get().TimerName##Cnt++;
 #else
-#define SCOPED_LOADTIMER(TimerName)
+
+#define SCOPED_LOADTIMER(TimerName) TRACE_CPUPROFILER_EVENT_SCOPE_GROUP(TimerName, CpuProfilerGroup_LoadTime)
 #define SCOPED_LOADTIMER_CNT(TimerName)
 #endif
