Allow for implicit uploading of crashes based on an ini setting

diff --git a/Engine/Source/Runtime/Core/Private/Mac/MacPlatformMisc.cpp b/Engine/Source/Runtime/Core/Private/Mac/MacPlatformMisc.cpp
--- a/Engine/Source/Runtime/Core/Private/Mac/MacPlatformMisc.cpp
+++ b/Engine/Source/Runtime/Core/Private/Mac/MacPlatformMisc.cpp
@@ -1882,6 +1882,12 @@ void FMacCrashContext::GenerateCrashInfoAndLaunchReporter() const
 	// Prevent CrashReportClient from spawning another CrashReportClient.
 	bool bCanRunCrashReportClient = FCString::Stristr( *(GMacAppInfo.ExecutableName), TEXT( "CrashReportClient" ) ) == nullptr;
 
+	bool bImplicitSend = false;
+	if (GConfig)
+	{
+		GConfig->GetBool(TEXT("CrashReportClient"), TEXT("bImplicitSend"), bImplicitSend, GEngineIni);
+	}
+
 	bool bSendUnattendedBugReports = true;
 	if (GConfig)
 	{
@@ -1925,7 +1931,7 @@ void FMacCrashContext::GenerateCrashInfoAndLaunchReporter() const
 		if(ForkPID == 0)
 		{
 			// Child
-			if(GMacAppInfo.bIsUnattended)
+			if(GMacAppInfo.bIsUnattended || bImplicitSend)
 			{
 				execl(GMacAppInfo.CrashReportClient, "CrashReportClient", CrashInfoFolder, "-Unattended", NULL);
 			}
@@ -1934,6 +1940,18 @@ void FMacCrashContext::GenerateCrashInfoAndLaunchReporter() const
 				execl(GMacAppInfo.CrashReportClient, "CrashReportClient", CrashInfoFolder, NULL);
 			}
 		}
+
+		if (bImplicitSend)
+		{
+			// We will not have any GUI for the crash reporter if we are sending implicitly, so pop a message box up at least
+			if (FApp::CanEverRender())
+			{
+				FPlatformMisc::MessageBoxExt(EAppMsgType::Ok,
+					*NSLOCTEXT("MessageDialog", "ReportCrash_Body", "The application has crashed and will now close. We apologize for the inconvenience.").ToString(),
+					*NSLOCTEXT("MessageDialog", "ReportCrash_Title", "Application Crash Detected").ToString());
+			}
+		}
+
 		// We no longer wait here because on return the OS will scribble & crash again due to the behaviour of the XPC function
 		// macOS uses internally to launch/wait on the CrashReportClient. It is simpler, easier & safer to just die here like a good little Mac.app.
 	}
diff --git a/Engine/Source/Runtime/Core/Private/Unix/UnixPlatformCrashContext.cpp b/Engine/Source/Runtime/Core/Private/Unix/UnixPlatformCrashContext.cpp
--- a/Engine/Source/Runtime/Core/Private/Unix/UnixPlatformCrashContext.cpp
+++ b/Engine/Source/Runtime/Core/Private/Unix/UnixPlatformCrashContext.cpp
@@ -422,6 +422,12 @@ void FUnixCrashContext::GenerateCrashInfoAndLaunchReporter(bool bReportingNonCra
 	}
 #endif
 
+	bool bImplicitSend = false;
+	if (GConfig)
+	{
+		GConfig->GetBool(TEXT("CrashReportClient"), TEXT("bImplicitSend"), bImplicitSend, GEngineIni);
+	}
+
 	// By default we wont upload unless the *.ini has set this to true
 	bool bSendUnattendedBugReports = false;
 	if (GConfig)
@@ -559,7 +565,7 @@ void FUnixCrashContext::GenerateCrashInfoAndLaunchReporter(bool bReportingNonCra
 			CrashReportClientArguments += TEXT("\"\"") + CrashReportLogFilepath + TEXT("\"\"");
 			CrashReportClientArguments += TEXT(" ");
 
-			if (bUnattended)
+			if (bUnattended || bImplicitSend)
 			{
 				CrashReportClientArguments += TEXT(" -Unattended ");
 			}
@@ -606,6 +612,17 @@ void FUnixCrashContext::GenerateCrashInfoAndLaunchReporter(bool bReportingNonCra
 				// spin here until CrashReporter exits
 				FProcHandle RunningProc = FPlatformProcess::CreateProc(RelativePathToCrashReporter, *CrashReportClientArguments, true, false, false, NULL, 0, NULL, NULL);
 
+				// We will not have any GUI for the crash reporter if we are sending implicitly, so pop a message box up at least
+				if (bImplicitSend)
+				{
+					if (FApp::CanEverRender())
+					{
+						FPlatformMisc::MessageBoxExt(EAppMsgType::Ok,
+							*NSLOCTEXT("MessageDialog", "ReportCrash_Body", "The application has crashed and will now close. We apologize for the inconvenience.").ToString(),
+							*NSLOCTEXT("MessageDialog", "ReportCrash_Title", "Application Crash Detected").ToString());
+					}
+				}
+
 				// do not wait indefinitely - can be more generous about the hitch than in ensure() case
 				// NOTE: Chris.Wood - increased from 3 to 8 mins because server crashes were timing out and getting lost
 				// NOTE: Do not increase above 8.5 mins without altering watchdog scripts to match
diff --git a/Engine/Source/Runtime/Core/Private/Windows/WindowsPlatformCrashContext.cpp b/Engine/Source/Runtime/Core/Private/Windows/WindowsPlatformCrashContext.cpp
--- a/Engine/Source/Runtime/Core/Private/Windows/WindowsPlatformCrashContext.cpp
+++ b/Engine/Source/Runtime/Core/Private/Windows/WindowsPlatformCrashContext.cpp
@@ -314,6 +314,12 @@ int32 ReportCrashUsingCrashReportClient(FWindowsPlatformCrashContext& InContext,
 	// Suppress the user input dialog if we're running in unattended mode
 	bool bNoDialog = FApp::IsUnattended() || ReportUI == EErrorReportUI::ReportInUnattendedMode || IsRunningDedicatedServer();
 
+	bool bImplicitSend = false;
+	if (GConfig)
+	{
+		GConfig->GetBool(TEXT("CrashReportClient"), TEXT("bImplicitSend"), bImplicitSend, GEngineIni);
+	}
+
 	bool bSendUnattendedBugReports = true;
 	if (GConfig)
 	{
@@ -430,7 +436,7 @@ int32 ReportCrashUsingCrashReportClient(FWindowsPlatformCrashContext& InContext,
 			// Pass nullrhi to CRC when the engine is in this mode to stop the CRC attempting to initialize RHI when the capability isn't available
 			bool bNullRHI = !FApp::CanEverRender();
 
-			if (bNoDialog || bNullRHI)
+			if (bNoDialog || bNullRHI || bImplicitSend)
 			{
 				CrashReportClientArguments += TEXT(" -Unattended");
 			}
@@ -474,6 +480,17 @@ int32 ReportCrashUsingCrashReportClient(FWindowsPlatformCrashContext& InContext,
 			FString CrashClientPath = FPaths::Combine(*FPaths::EngineDir(), TEXT("Binaries"), FPlatformProcess::GetBinariesSubdirectory(), CrashReportClientExeName);
 			bCrashReporterRan = FPlatformProcess::CreateProc(*CrashClientPath, *CrashReportClientArguments, true, false, false, NULL, 0, NULL, NULL).IsValid();
 
+			if (bImplicitSend)
+			{
+				// We will not have any GUI for the crash reporter if we are sending implicitly, so pop a message box up at least
+				if (FApp::CanEverRender())
+				{
+					FPlatformMisc::MessageBoxExt(EAppMsgType::Ok,
+						*NSLOCTEXT("MessageDialog", "ReportCrash_Body", "The application has crashed and will now close. We apologize for the inconvenience.").ToString(),
+						*NSLOCTEXT("MessageDialog", "ReportCrash_Title", "Application Crash Detected").ToString());
+				}
+			}
+
 			// Restore the dll directory
 			if (CurrentDllDirectory)
 			{
