Correct fix for the missing messages from failed checkf() assertions.

(p4-cherrypick of 6829622 by [REDACTED_PERSON])

diff --git a/Engine/Source/Runtime/Core/Private/Misc/AssertionMacros.cpp b/Engine/Source/Runtime/Core/Private/Misc/AssertionMacros.cpp
--- a/Engine/Source/Runtime/Core/Private/Misc/AssertionMacros.cpp
+++ b/Engine/Source/Runtime/Core/Private/Misc/AssertionMacros.cpp
@@ -80,6 +80,27 @@ void PrintScriptCallstack()
 	}
 }
 
+static void AssertFailedImplV(const ANSICHAR* Expr, const ANSICHAR* File, int32 Line, const TCHAR* Format, va_list Args)
+{
+	if (GIsCriticalError)
+	{
+		return;
+	}
+
+	// This is not perfect because another thread might crash and be handled before this assert
+	// but this static varible will report the crash as an assert. Given complexity of a thread
+	// aware solution, this should be good enough. If crash reports are obviously wrong we can
+	// look into fixing this.
+	bHasAsserted = true;
+
+	TCHAR DescriptionString[4096];
+	FCString::GetVarArgs(DescriptionString, ARRAY_COUNT(DescriptionString), Format, Args);
+
+	TCHAR ErrorString[MAX_SPRINTF];
+	FCString::Sprintf(ErrorString, TEXT("%s"), ANSI_TO_TCHAR(Expr));
+	GError->Logf(TEXT("Assertion failed: %s") FILE_LINE_DESC TEXT("\n%s\n"), ErrorString, ANSI_TO_TCHAR(File), Line, DescriptionString);
+}
+
 /**
  *	Prints error to the debug output, 
  *	prompts for the remote debugging if there is not debugger, breaks into the debugger 
@@ -418,7 +439,7 @@ void FORCENOINLINE FDebug::CheckVerifyFailedImpl(
 	if (!FPlatformMisc::IsDebuggerPresent())
 	{
 		FPlatformMisc::PromptForRemoteDebugging(false);
-		FDebug::AssertFailed(Expr, File, Line, Format, Args);
+		AssertFailedImplV(Expr, File, Line, Format, Args);
 	}
 	va_end(Args);
 }
@@ -427,23 +448,10 @@ void FORCENOINLINE FDebug::CheckVerifyFailedImpl(
 
 void VARARGS FDebug::AssertFailed(const ANSICHAR* Expr, const ANSICHAR* File, int32 Line, const TCHAR* Format/* = TEXT("")*/, ...)
 {
-	if (GIsCriticalError)
-	{
-		return;
-	}
-
-	// This is not perfect because another thread might crash and be handled before this assert
-	// but this static varible will report the crash as an assert. Given complexity of a thread
-	// aware solution, this should be good enough. If crash reports are obviously wrong we can
-	// look into fixing this.
-	bHasAsserted = true;
-
-	TCHAR DescriptionString[4096];
-	GET_VARARGS(DescriptionString, ARRAY_COUNT(DescriptionString), ARRAY_COUNT(DescriptionString) - 1, Format, Format);
-
-	TCHAR ErrorString[MAX_SPRINTF];
-	FCString::Sprintf(ErrorString, TEXT("%s"), ANSI_TO_TCHAR(Expr));
-	GError->Logf(TEXT("Assertion failed: %s") FILE_LINE_DESC TEXT("\n%s\n"), ErrorString, ANSI_TO_TCHAR(File), Line, DescriptionString);
+	va_list Args;
+	va_start(Args, Format);
+	AssertFailedImplV(Expr, File, Line, Format, Args);
+	va_end(Args);
 }
 
 #if DO_CHECK || DO_GUARD_SLOW
