Message boxes are a no-no in crash handling code, undo for now and look at a better fix here

diff --git a/Engine/Source/Runtime/Core/Private/Mac/MacPlatformMisc.cpp b/Engine/Source/Runtime/Core/Private/Mac/MacPlatformMisc.cpp
--- a/Engine/Source/Runtime/Core/Private/Mac/MacPlatformMisc.cpp
+++ b/Engine/Source/Runtime/Core/Private/Mac/MacPlatformMisc.cpp
@@ -1931,7 +1931,7 @@ void FMacCrashContext::GenerateCrashInfoAndLaunchReporter() const
 		if(ForkPID == 0)
 		{
 			// Child
-			if(GMacAppInfo.bIsUnattended || bImplicitSend)
+			if(GMacAppInfo.bIsUnattended)
 			{
 				execl(GMacAppInfo.CrashReportClient, "CrashReportClient", CrashInfoFolder, "-Unattended", NULL);
 			}
@@ -1941,17 +1941,6 @@ void FMacCrashContext::GenerateCrashInfoAndLaunchReporter() const
 			}
 		}
 
-		if (bImplicitSend)
-		{
-			// We will not have any GUI for the crash reporter if we are sending implicitly, so pop a message box up at least
-			if (FApp::CanEverRender())
-			{
-				FPlatformMisc::MessageBoxExt(EAppMsgType::Ok,
-					*NSLOCTEXT("MessageDialog", "ReportCrash_Body", "The application has crashed and will now close. We apologize for the inconvenience.").ToString(),
-					*NSLOCTEXT("MessageDialog", "ReportCrash_Title", "Application Crash Detected").ToString());
-			}
-		}
-
 		// We no longer wait here because on return the OS will scribble & crash again due to the behaviour of the XPC function
 		// macOS uses internally to launch/wait on the CrashReportClient. It is simpler, easier & safer to just die here like a good little Mac.app.
 	}
