- Fixed a deadlock in IOSAsyncTask where it could block adding new items to the queue if a task was also blocked
- Should address various backgrounding badnesses

diff --git a/Engine/Source/Runtime/Core/Private/IOS/IOSAsyncTask.cpp b/Engine/Source/Runtime/Core/Private/IOS/IOSAsyncTask.cpp
--- a/Engine/Source/Runtime/Core/Private/IOS/IOSAsyncTask.cpp
+++ b/Engine/Source/Runtime/Core/Private/IOS/IOSAsyncTask.cpp
@@ -90,25 +90,31 @@ NSMutableArray* RunningTasks;
  */
 + (void)ProcessAsyncTasks
 {
+	FIOSAsyncTask* CurrentTask = nil;
+
+	// grab one out of the queue to process outside of the lock
 	@synchronized(RunningTasks)
 	{
-		// check all outstanding tasks to see if they are done
-		// (we go backwards so we can remove as we go)
-		for (int32 TaskIndex = 0; TaskIndex < [RunningTasks count]; TaskIndex++)
+		if ([RunningTasks count] > 0)
+		{
+			CurrentTask = [RunningTasks objectAtIndex:0];
+			[RunningTasks removeObjectAtIndex:0];
+		}
+	}
+
+	if (CurrentTask != nil)
+	{
+		if ([CurrentTask CheckForCompletion])
+		{
+			// release the object count
+			[CurrentTask release];
+		}
+		else
 		{
-			FIOSAsyncTask* Task = [RunningTasks objectAtIndex:TaskIndex];
-			// if it finished, remove it from the list of running tasks
-			if ([Task CheckForCompletion])
+			// if it's not done, but it back in the list
+			@synchronized(RunningTasks)
 			{
-				// release the object count
-				[Task release];
-
-				// and remove from the list
-				[RunningTasks removeObjectAtIndex:TaskIndex];
-				TaskIndex--;
-				
-				// do one at once
-				break;
+				[RunningTasks addObject:CurrentTask];
 			}
 		}
 	}
