Undoing change in 6236960 that set commandline logging to display

diff --git a/Engine/Source/Runtime/Core/Private/Misc/App.cpp b/Engine/Source/Runtime/Core/Private/Misc/App.cpp
--- a/Engine/Source/Runtime/Core/Private/Misc/App.cpp
+++ b/Engine/Source/Runtime/Core/Private/Misc/App.cpp
@@ -342,7 +342,7 @@ void FApp::PrintStartupLogMessages()
 	UE_LOG(LogInit, Log, TEXT("Build Configuration: %s"), EBuildConfigurations::ToString(FApp::GetBuildConfiguration()));
 	UE_LOG(LogInit, Log, TEXT("Branch Name: %s"), *FApp::GetBranchName());
 	FString FilteredString = FCommandLine::IsCommandLineLoggingFiltered() ? TEXT("Filtered ") : TEXT("");
-	UE_LOG(LogInit, Display, TEXT("%sCommand Line: %s"), *FilteredString, FCommandLine::GetForLogging());
+	UE_LOG(LogInit, Log, TEXT("%sCommand Line: %s"), *FilteredString, FCommandLine::GetForLogging());
 	UE_LOG(LogInit, Log, TEXT("Base Directory: %s"), FPlatformProcess::BaseDir());
 	//UE_LOG(LogInit, Log, TEXT("Character set: %s"), sizeof(TCHAR)==1 ? TEXT("ANSI") : TEXT("Unicode") );
 	UE_LOG(LogInit, Log, TEXT("Installed Engine Build: %d"), FApp::IsEngineInstalled() ? 1 : 0);
