Reduced noise caused by double const violation warnings in delegates.

diff --git a/Engine/Source/Runtime/Core/Public/Delegates/DelegateSignatureImpl.inl b/Engine/Source/Runtime/Core/Public/Delegates/DelegateSignatureImpl.inl
--- a/Engine/Source/Runtime/Core/Public/Delegates/DelegateSignatureImpl.inl
+++ b/Engine/Source/Runtime/Core/Public/Delegates/DelegateSignatureImpl.inl
@@ -454,7 +454,7 @@ public:
 	{
 		UE_STATIC_DEPRECATE(4.23, TIsConst<UserClass>::Value, "Binding a delegate with a const object pointer and non-const function is deprecated.");
 
-		*this = CreateRaw(InUserObject, InFunc, Vars...);
+		*this = CreateRaw(const_cast<typename TRemoveConst<UserClass>::Type*>(InUserObject), InFunc, Vars...);
 	}
 	template <typename UserClass, typename... VarTypes>
 	inline void BindRaw(UserClass* InUserObject, typename TMemFunPtrType<true, UserClass, RetValType (ParamTypes..., VarTypes...)>::Type InFunc, VarTypes... Vars)
@@ -470,7 +470,7 @@ public:
 	{
 		UE_STATIC_DEPRECATE(4.23, TIsConst<UserClass>::Value, "Binding a delegate with a const object pointer and non-const function is deprecated.");
 
-		*this = CreateSP(InUserObjectRef, InFunc, Vars...);
+		*this = CreateSP(ConstCastSharedRef<typename TRemoveConst<UserClass>::Type>(InUserObjectRef), InFunc, Vars...);
 	}
 	template <typename UserClass, typename... VarTypes>
 	inline void BindSP(const TSharedRef<UserClass, ESPMode::Fast>& InUserObjectRef, typename TMemFunPtrType<true, UserClass, RetValType (ParamTypes..., VarTypes...)>::Type InFunc, VarTypes... Vars)
@@ -489,7 +489,7 @@ public:
 	{
 		UE_STATIC_DEPRECATE(4.23, TIsConst<UserClass>::Value, "Binding a delegate with a const object pointer and non-const function is deprecated.");
 
-		*this = CreateSP(InUserObject, InFunc, Vars...);
+		*this = CreateSP(const_cast<typename TRemoveConst<UserClass>::Type*>(InUserObject), InFunc, Vars...);
 	}
 	template <typename UserClass, typename... VarTypes>
 	inline void BindSP(UserClass* InUserObject, typename TMemFunPtrType<true, UserClass, RetValType (ParamTypes..., VarTypes...)>::Type InFunc, VarTypes... Vars)
@@ -508,7 +508,7 @@ public:
 	{
 		UE_STATIC_DEPRECATE(4.23, TIsConst<UserClass>::Value, "Binding a delegate with a const object pointer and non-const function is deprecated.");
 
-		*this = CreateThreadSafeSP(InUserObjectRef, InFunc, Vars...);
+		*this = CreateThreadSafeSP(ConstCastSharedRef<typename TRemoveConst<UserClass>::Type>(InUserObjectRef), InFunc, Vars...);
 	}
 	template <typename UserClass, typename... VarTypes>
 	inline void BindThreadSafeSP(const TSharedRef<UserClass, ESPMode::ThreadSafe>& InUserObjectRef, typename TMemFunPtrType<true, UserClass, RetValType (ParamTypes..., VarTypes...)>::Type InFunc, VarTypes... Vars)
@@ -527,7 +527,7 @@ public:
 	{
 		UE_STATIC_DEPRECATE(4.23, TIsConst<UserClass>::Value, "Binding a delegate with a const object pointer and non-const function is deprecated.");
 
-		*this = CreateThreadSafeSP(InUserObject, InFunc, Vars...);
+		*this = CreateThreadSafeSP(const_cast<typename TRemoveConst<UserClass>::Type*>(InUserObject), InFunc, Vars...);
 	}
 	template <typename UserClass, typename... VarTypes>
 	inline void BindThreadSafeSP(UserClass* InUserObject, typename TMemFunPtrType<true, UserClass, RetValType (ParamTypes..., VarTypes...)>::Type InFunc, VarTypes... Vars)
@@ -558,7 +558,7 @@ public:
 	{
 		UE_STATIC_DEPRECATE(4.23, TIsConst<UserClass>::Value, "Binding a delegate with a const object pointer and non-const function is deprecated.");
 
-		*this = CreateUObject(InUserObject, InFunc, Vars...);
+		*this = CreateUObject(const_cast<typename TRemoveConst<UserClass>::Type*>(InUserObject), InFunc, Vars...);
 	}
 	template <typename UserClass, typename... VarTypes>
 	inline void BindUObject(UserClass* InUserObject, typename TMemFunPtrType<true, UserClass, RetValType (ParamTypes..., VarTypes...)>::Type InFunc, VarTypes... Vars)
@@ -791,7 +791,7 @@ public:
 	{
 		UE_STATIC_DEPRECATE(4.23, TIsConst<UserClass>::Value, "Binding a delegate with a const object pointer and non-const function is deprecated.");
 
-		return Add(FDelegate::CreateRaw(InUserObject, InFunc, Vars...));
+		return Add(FDelegate::CreateRaw(const_cast<typename TRemoveConst<UserClass>::Type*>(InUserObject), InFunc, Vars...));
 	}
 	template <typename UserClass, typename... VarTypes>
 	inline FDelegateHandle AddRaw(UserClass* InUserObject, typename TMemFunPtrType<true, UserClass, void (ParamTypes..., VarTypes...)>::Type InFunc, VarTypes... Vars)
@@ -812,7 +812,7 @@ public:
 	{
 		UE_STATIC_DEPRECATE(4.23, TIsConst<UserClass>::Value, "Binding a delegate with a const object pointer and non-const function is deprecated.");
 
-		return Add(FDelegate::CreateSP(InUserObjectRef, InFunc, Vars...));
+		return Add(FDelegate::CreateSP(ConstCastSharedRef<typename TRemoveConst<UserClass>::Type>(InUserObjectRef), InFunc, Vars...));
 	}
 	template <typename UserClass, typename... VarTypes>
 	inline FDelegateHandle AddSP(const TSharedRef<UserClass, ESPMode::Fast>& InUserObjectRef, typename TMemFunPtrType<true, UserClass, void (ParamTypes..., VarTypes...)>::Type InFunc, VarTypes... Vars)
@@ -833,7 +833,7 @@ public:
 	{
 		UE_STATIC_DEPRECATE(4.23, TIsConst<UserClass>::Value, "Binding a delegate with a const object pointer and non-const function is deprecated.");
 
-		return Add(FDelegate::CreateSP(InUserObject, InFunc, Vars...));
+		return Add(FDelegate::CreateSP(const_cast<typename TRemoveConst<UserClass>::Type*>(InUserObject), InFunc, Vars...));
 	}
 	template <typename UserClass, typename... VarTypes>
 	inline FDelegateHandle AddSP(UserClass* InUserObject, typename TMemFunPtrType<true, UserClass, void (ParamTypes..., VarTypes...)>::Type InFunc, VarTypes... Vars)
@@ -852,7 +852,7 @@ public:
 	{
 		UE_STATIC_DEPRECATE(4.23, TIsConst<UserClass>::Value, "Binding a delegate with a const object pointer and non-const function is deprecated.");
 
-		return Add(FDelegate::CreateThreadSafeSP(InUserObjectRef, InFunc, Vars...));
+		return Add(FDelegate::CreateThreadSafeSP(ConstCastSharedRef<typename TRemoveConst<UserClass>::Type>(InUserObjectRef), InFunc, Vars...));
 	}
 	template <typename UserClass, typename... VarTypes>
 	inline FDelegateHandle AddThreadSafeSP(const TSharedRef<UserClass, ESPMode::ThreadSafe>& InUserObjectRef, typename TMemFunPtrType<true, UserClass, void (ParamTypes..., VarTypes...)>::Type InFunc, VarTypes... Vars)
@@ -873,7 +873,7 @@ public:
 	{
 		UE_STATIC_DEPRECATE(4.23, TIsConst<UserClass>::Value, "Binding a delegate with a const object pointer and non-const function is deprecated.");
 
-		return Add(FDelegate::CreateThreadSafeSP(InUserObject, InFunc, Vars...));
+		return Add(FDelegate::CreateThreadSafeSP(const_cast<typename TRemoveConst<UserClass>::Type*>(InUserObject), InFunc, Vars...));
 	}
 	template <typename UserClass, typename... VarTypes>
 	inline FDelegateHandle AddThreadSafeSP(UserClass* InUserObject, typename TMemFunPtrType<true, UserClass, void (ParamTypes..., VarTypes...)>::Type InFunc, VarTypes... Vars)
@@ -908,7 +908,7 @@ public:
 	{
 		UE_STATIC_DEPRECATE(4.23, TIsConst<UserClass>::Value, "Binding a delegate with a const object pointer and non-const function is deprecated.");
 
-		return Add(FDelegate::CreateUObject(InUserObject, InFunc, Vars...));
+		return Add(FDelegate::CreateUObject(const_cast<typename TRemoveConst<UserClass>::Type*>(InUserObject), InFunc, Vars...));
 	}
 	template <typename UserClass, typename... VarTypes>
 	inline FDelegateHandle AddUObject(UserClass* InUserObject, typename TMemFunPtrType<true, UserClass, void (ParamTypes..., VarTypes...)>::Type InFunc, VarTypes... Vars)
