Only implicit send in non editor builds and for crashes only

diff --git a/Engine/Source/Runtime/Core/Private/Mac/MacPlatformMisc.cpp b/Engine/Source/Runtime/Core/Private/Mac/MacPlatformMisc.cpp
--- a/Engine/Source/Runtime/Core/Private/Mac/MacPlatformMisc.cpp
+++ b/Engine/Source/Runtime/Core/Private/Mac/MacPlatformMisc.cpp
@@ -1883,8 +1883,9 @@ void FMacCrashContext::GenerateCrashInfoAndLaunchReporter() const
 	bool bCanRunCrashReportClient = FCString::Stristr( *(GMacAppInfo.ExecutableName), TEXT( "CrashReportClient" ) ) == nullptr;
 
 	bool bImplicitSend = false;
-	if (GConfig)
+	if (!UE_EDITOR && GConfig)
 	{
+		// Only check if we are in a non-editor build
 		GConfig->GetBool(TEXT("CrashReportClient"), TEXT("bImplicitSend"), bImplicitSend, GEngineIni);
 	}
 
diff --git a/Engine/Source/Runtime/Core/Private/Unix/UnixPlatformCrashContext.cpp b/Engine/Source/Runtime/Core/Private/Unix/UnixPlatformCrashContext.cpp
--- a/Engine/Source/Runtime/Core/Private/Unix/UnixPlatformCrashContext.cpp
+++ b/Engine/Source/Runtime/Core/Private/Unix/UnixPlatformCrashContext.cpp
@@ -423,8 +423,9 @@ void FUnixCrashContext::GenerateCrashInfoAndLaunchReporter(bool bReportingNonCra
 #endif
 
 	bool bImplicitSend = false;
-	if (GConfig)
+	if (!UE_EDITOR && GConfig && !bReportingNonCrash)
 	{
+		// Only check if we are in a non-editor build
 		GConfig->GetBool(TEXT("CrashReportClient"), TEXT("bImplicitSend"), bImplicitSend, GEngineIni);
 	}
 
diff --git a/Engine/Source/Runtime/Core/Private/Windows/WindowsPlatformCrashContext.cpp b/Engine/Source/Runtime/Core/Private/Windows/WindowsPlatformCrashContext.cpp
--- a/Engine/Source/Runtime/Core/Private/Windows/WindowsPlatformCrashContext.cpp
+++ b/Engine/Source/Runtime/Core/Private/Windows/WindowsPlatformCrashContext.cpp
@@ -315,8 +315,9 @@ int32 ReportCrashUsingCrashReportClient(FWindowsPlatformCrashContext& InContext,
 	bool bNoDialog = FApp::IsUnattended() || ReportUI == EErrorReportUI::ReportInUnattendedMode || IsRunningDedicatedServer();
 
 	bool bImplicitSend = false;
-	if (GConfig)
+	if (!UE_EDITOR && GConfig && ReportUI != EErrorReportUI::ReportInUnattendedMode)
 	{
+		// Only check if we are in a non-editor build
 		GConfig->GetBool(TEXT("CrashReportClient"), TEXT("bImplicitSend"), bImplicitSend, GEngineIni);
 	}
 
