StandardPlatformString::GetVarArgs bug fixes.

size_t and ptrdiff_t are now properly treated as 64-bit types on 64-bit platforms.
Honor width and precision for strings.

diff --git a/Engine/Source/Runtime/Core/Private/GenericPlatform/StandardPlatformString.cpp b/Engine/Source/Runtime/Core/Private/GenericPlatform/StandardPlatformString.cpp
--- a/Engine/Source/Runtime/Core/Private/GenericPlatform/StandardPlatformString.cpp
+++ b/Engine/Source/Runtime/Core/Private/GenericPlatform/StandardPlatformString.cpp
@@ -14,7 +14,7 @@
 
 DEFINE_LOG_CATEGORY_STATIC(LogStandardPlatformString, Log, All);
 
-#if PLATFORM_IOS
+#if PLATFORM_IOS || PLATFORM_WINDOWS
 	#define VA_LIST_REF va_list&
 #else
 	#define VA_LIST_REF va_list
@@ -67,8 +67,8 @@ static int32 GetFormattingInfo(const WIDECHAR* Format, FFormatInfo& OutInfo)
 	}
 
 	OutInfo.LengthModifier = 0;
-	if (*Format == LITERAL(WIDECHAR, 'h') || *Format == LITERAL(WIDECHAR, 'l') || *Format == LITERAL(WIDECHAR, 'j') || *Format == LITERAL(WIDECHAR, 't')
-		|| *Format == LITERAL(WIDECHAR, 'z') || *Format == LITERAL(WIDECHAR, 'q') || *Format == LITERAL(WIDECHAR, 'L'))
+	if (*Format == LITERAL(WIDECHAR, 'h') || *Format == LITERAL(WIDECHAR, 'l') || *Format == LITERAL(WIDECHAR, 'j')
+		|| *Format == LITERAL(WIDECHAR, 'q') || *Format == LITERAL(WIDECHAR, 'L'))
 	{
 		OutInfo.LengthModifier = *Format++;
 		if (*Format == LITERAL(WIDECHAR, 'h'))
@@ -82,6 +82,15 @@ static int32 GetFormattingInfo(const WIDECHAR* Format, FFormatInfo& OutInfo)
 			Format++;
 		}
 	}
+	else if (*Format == LITERAL(WIDECHAR, 't') || *Format == LITERAL(WIDECHAR, 'z'))
+	{
+#if PLATFORM_64BITS
+		OutInfo.LengthModifier = LITERAL(WIDECHAR, 'l');
+		++Format;
+#else
+		OutInfo.LengthModifier = *Format++;
+#endif
+	}
 
 	OutInfo.Type = *Format++;
 
@@ -160,12 +169,20 @@ static const WIDECHAR* GetFormattedArgument(const FFormatInfo& Info, VA_LIST_REF
 				return TEXT("(null)");
 			}
 		}
-		else
+		// Is it a plain string?
+		else if (Info.Format[1] == LITERAL(WIDECHAR, 's'))
 		{
 			const WIDECHAR* String = va_arg(ArgPtr, WIDECHAR*);
 			InOutLength = FCString::Strlen(String);
 			return String ? String : TEXT("(null)");
 		}
+		// Some form of string requiring formatting, such as a left- or right-justified string
+		else
+		{
+			const WIDECHAR* String = va_arg(ArgPtr, WIDECHAR*);
+			InOutLength = swprintf(Formatted, InOutLength, Info.Format, String);
+			return Formatted;
+		}
 	}
 
 	if (FChar::ToLower(Info.Type) == LITERAL(WIDECHAR, 'c'))
