Fix for FramePro when running with -nothreading

diff --git a/Engine/Source/Runtime/Core/Private/FramePro/FrameProUE4.cpp b/Engine/Source/Runtime/Core/Private/FramePro/FrameProUE4.cpp
--- a/Engine/Source/Runtime/Core/Private/FramePro/FrameProUE4.cpp
+++ b/Engine/Source/Runtime/Core/Private/FramePro/FrameProUE4.cpp
@@ -73,6 +73,14 @@
 		#include <cpuid.h>
 	#endif
 
+	//@EPIC BEGIN: workaround for -nothreading
+	#if PLATFORM_WINDOWS
+		#include "Windows/WindowsEvent.h"
+	#elif PLATFORM_USE_PTHREADS
+		#include <pthread.h>
+		#include "HAL/PThreadEvent.h"
+	#endif
+
 	//------------------------------------------------------------------------
 	namespace FramePro
 	{
@@ -429,7 +437,26 @@
 		//------------------------------------------------------------------------
 		void Platform::CreateEventX(void* p_os_event_mem, int os_event_mem_size, bool initial_state, bool auto_reset)
 		{
-			FEvent* p_event = FPlatformProcess::GetSynchEventFromPool(!auto_reset);
+//@EPIC BEGIN: workaround for -nothreading
+			FEvent* p_event;
+			if (FPlatformProcess::SupportsMultithreading())
+			{
+				p_event = FPlatformProcess::GetSynchEventFromPool(!auto_reset);
+			}
+			else
+			{				
+#if PLATFORM_WINDOWS
+				p_event = new FEventWin();
+				p_event->Create(!auto_reset);
+#elif PLATFORM_USE_PTHREADS
+				p_event = new FPThreadEvent();
+				p_event->Create(!auto_reset);
+#else
+				checkf(false,"unsupported platform for -nothreading");
+				return;
+#endif
+			}
+//@EPIC END
 			GetOSEvent(p_os_event_mem) = p_event;
 
 			if(initial_state)
@@ -440,7 +467,16 @@
 		void Platform::DestroyEvent(void* p_os_event_mem)
 		{
 			FEvent* p_event = GetOSEvent(p_os_event_mem);
-			FPlatformProcess::ReturnSynchEventToPool(p_event);
+//@EPIC BEGIN: workaround for -nothreading
+			if (FPlatformProcess::SupportsMultithreading())
+			{
+				FPlatformProcess::ReturnSynchEventToPool(p_event);
+			}
+			else
+			{
+				delete p_event;
+			}
+//@EPIC END
 			GetOSEvent(p_os_event_mem) = nullptr;
 		}
 
@@ -544,7 +580,24 @@
 			:	mp_ThreadMain(p_thread_main),
 				mp_Context(p_context)
 			{
-				m_Runnable = TUniquePtr<FRunnableThread>(FRunnableThread::Create(this, TEXT("FramePro")));
+//@EPIC BEGIN: workaround for -nothreading
+				if (FPlatformProcess::SupportsMultithreading())
+				{
+					m_Runnable = TUniquePtr<FRunnableThread>(FRunnableThread::Create(this, TEXT("FramePro")));
+				}
+				else
+				{
+#if PLATFORM_WINDOWS
+					CreateThread( NULL, 0, (LPTHREAD_START_ROUTINE)mp_ThreadMain, mp_Context, 0, NULL );
+#elif PLATFORM_USE_PTHREADS
+					typedef void *(*PthreadEntryPoint)(void *arg);
+					pthread_t thread_id;
+					pthread_create( &thread_id, nullptr, (PthreadEntryPoint)mp_ThreadMain, mp_Context );
+#else
+					checkf(false,"unsupported platform for -nothreading");
+#endif
+				}
+//@EPIC END
 			}
 
 			uint32 Run() override
