Added output buffer address to decompression error log

diff --git a/Engine/Source/Runtime/Core/Private/Misc/Compression.cpp b/Engine/Source/Runtime/Core/Private/Misc/Compression.cpp
--- a/Engine/Source/Runtime/Core/Private/Misc/Compression.cpp
+++ b/Engine/Source/Runtime/Core/Private/Misc/Compression.cpp
@@ -418,7 +418,7 @@ bool FCompression::UncompressMemory(FName FormatName, void* UncompressedBuffer,
 			bUncompressSucceeded = true;
 		}
 		// Always log an error
-		UE_LOG(LogCompression, Error, TEXT("FCompression::UncompressMemory - Failed to uncompress memory (%d/%d) using format %s, this may indicate the asset is corrupt!"), CompressedSize, UncompressedSize, *FormatName.ToString());
+		UE_LOG(LogCompression, Error, TEXT("FCompression::UncompressMemory - Failed to uncompress memory (%d/%d) from address 0x%016X using format %s, this may indicate the asset is corrupt!"), CompressedSize, UncompressedSize, CompressedBuffer, *FormatName.ToString());
 	}
 
 #if	STATS
