UE4 - Fix LLM for Windows and Switch

diff --git a/Engine/Source/Runtime/Core/Private/HAL/MallocBinned3.cpp b/Engine/Source/Runtime/Core/Private/HAL/MallocBinned3.cpp
--- a/Engine/Source/Runtime/Core/Private/HAL/MallocBinned3.cpp
+++ b/Engine/Source/Runtime/Core/Private/HAL/MallocBinned3.cpp
@@ -905,6 +905,7 @@ FMallocBinned3::~FMallocBinned3()
 
 void FMallocBinned3::Commit(uint32 InPoolIndex, void *Ptr, SIZE_T Size)
 {
+
 #if BINNED3_ALLOCATOR_STATS
 	Binned3Commits++;
 #endif
@@ -913,9 +914,12 @@ void FMallocBinned3::Commit(uint32 InPoolIndex, void *Ptr, SIZE_T Size)
 #else
 	PoolBaseVMBlock[InPoolIndex].CommitByPtr(Ptr, Size);
 #endif
+	LLM(FLowLevelMemTracker::Get().OnLowLevelAlloc(ELLMTracker::Platform, Ptr, Size));
 }
 void FMallocBinned3::Decommit(uint32 InPoolIndex, void *Ptr, SIZE_T Size)
 {
+	LLM(FLowLevelMemTracker::Get().OnLowLevelFree(ELLMTracker::Platform, Ptr));
+
 #if BINNED3_ALLOCATOR_STATS
 	Binned3Decommits++;
 #endif
@@ -933,6 +937,7 @@ void* FMallocBinned3::AllocateMetaDataMemory(SIZE_T Size)
 	FPlatformMemory::FPlatformVirtualMemoryBlock Block = FPlatformMemory::FPlatformVirtualMemoryBlock::AllocateVirtual(VirtualAlignedSize);
 	size_t CommitAlignedSize = Align(Size, FPlatformMemory::FPlatformVirtualMemoryBlock::GetCommitAlignment());
 	Block.Commit(0, CommitAlignedSize);
+	LLM(FLowLevelMemTracker::Get().OnLowLevelAlloc(ELLMTracker::Platform, Block.GetVirtualPointer(), CommitAlignedSize));
 	return Block.GetVirtualPointer();
 }
 
@@ -1030,10 +1035,12 @@ void* FMallocBinned3::MallocExternal(SIZE_T Size, uint32 Alignment)
 	void* Result = GetCachedOSPageAllocator().Allocate(AlignedSize);
 	check(IsAligned(Result, Alignment));
 #else
+
 	FPlatformMemory::FPlatformVirtualMemoryBlock Block = FPlatformMemory::FPlatformVirtualMemoryBlock::AllocateVirtual(AlignedSize, Alignment);
 	Block.Commit(0, AlignedSize);
 	void* Result = Block.GetVirtualPointer();
 #endif
+	LLM(FLowLevelMemTracker::Get().OnLowLevelAlloc(ELLMTracker::Platform, Result, AlignedSize));
 
 #if BINNED3_TIME_LARGE_BLOCKS
 	double Add = FPlatformTime::Seconds() - StartTime;
@@ -1227,6 +1234,8 @@ void FMallocBinned3::FreeExternal(void* Ptr)
 			FPlatformMemory::FPlatformVirtualMemoryBlock Block(Ptr, VMPages);
 			Block.FreeVirtual();
 #endif
+			LLM(FLowLevelMemTracker::Get().OnLowLevelFree(ELLMTracker::Platform, Ptr));
+
 		}
 #if BINNED3_TIME_LARGE_BLOCKS
 		double Add = FPlatformTime::Seconds() - StartTime;
diff --git a/Engine/Source/Runtime/Core/Private/Windows/WindowsPlatformMemory.cpp b/Engine/Source/Runtime/Core/Private/Windows/WindowsPlatformMemory.cpp
--- a/Engine/Source/Runtime/Core/Private/Windows/WindowsPlatformMemory.cpp
+++ b/Engine/Source/Runtime/Core/Private/Windows/WindowsPlatformMemory.cpp
@@ -343,8 +343,6 @@ void FWindowsPlatformMemory::FPlatformVirtualMemoryBlock::FreeVirtual()
 	if (Ptr)
 	{
 		check(GetActualSize() > 0);
-		// this is an iffy assumption, we don't know how much of this memory is really committed, we will assume none of it is
-		//LLM(FLowLevelMemTracker::Get().OnLowLevelFree(ELLMTracker::Platform, Ptr));
 
 		CA_SUPPRESS(6001)
 		// Windows maintains the size of allocation internally, so Size is unused
@@ -360,9 +358,7 @@ void FWindowsPlatformMemory::FPlatformVirtualMemoryBlock::Commit(size_t InOffset
 	check(IsAligned(InOffset, GetCommitAlignment()) && IsAligned(InSize, GetCommitAlignment()));
 	check(InOffset >= 0 && InSize >= 0 && InOffset + InSize <= GetActualSize() && Ptr);
 
-	// There are no guarantees LLM is going to be able to deal with this
 	uint8* UsePtr = ((uint8*)Ptr) + InOffset;
-	LLM(FLowLevelMemTracker::Get().OnLowLevelAlloc(ELLMTracker::Platform, UsePtr, InSize));
 	if (VirtualAlloc(UsePtr, InSize, MEM_COMMIT, PAGE_READWRITE) != UsePtr)
 	{
 		FPlatformMemory::OnOutOfMemory(InSize, 0);
@@ -374,8 +370,6 @@ void FWindowsPlatformMemory::FPlatformVirtualMemoryBlock::Decommit(size_t InOffs
 	check(IsAligned(InOffset, GetCommitAlignment()) && IsAligned(InSize, GetCommitAlignment()));
 	check(InOffset >= 0 && InSize >= 0 && InOffset + InSize <= GetActualSize() && Ptr);
 	uint8* UsePtr = ((uint8*)Ptr) + InOffset;
-	// There are no guarantees LLM is going to be able to deal with this
-	LLM(FLowLevelMemTracker::Get().OnLowLevelFree(ELLMTracker::Platform, UsePtr));
 	VirtualFree(UsePtr, InSize, MEM_DECOMMIT);
 }
 
