Fix ReapplyRecordedCVarSettingsFromIni() wiping out plugin config.
Disable INI reloading because it's still not safe.  Plugins and Hotfix manager will modify config using CombineFromBuffer().  Because this isn't part of the INI hierarchy we lose these changes when reloading.  We need a solution of this before we can safely enable this.

diff --git a/Engine/Source/Runtime/Core/Private/Misc/ConfigCacheIni.cpp b/Engine/Source/Runtime/Core/Private/Misc/ConfigCacheIni.cpp
--- a/Engine/Source/Runtime/Core/Private/Misc/ConfigCacheIni.cpp
+++ b/Engine/Source/Runtime/Core/Private/Misc/ConfigCacheIni.cpp
@@ -4177,7 +4177,7 @@ public:
 		bRecurseCheck=false;
 	}
 };
-FCVarIniHistoryHelper *IniHistoryHelper = nullptr;
+TUniquePtr<FCVarIniHistoryHelper> IniHistoryHelper;
 
 #if !UE_BUILD_SHIPPING
 class FConfigHistoryHelper
@@ -4296,39 +4296,27 @@ public:
 		Writer = nullptr;
 	}
 };
-FConfigHistoryHelper* ConfigHistoryHelper = nullptr;
+TUniquePtr<FConfigHistoryHelper> ConfigHistoryHelper;
 #endif // !UE_BUILD_SHIPPING
 
 
 void RecordApplyCVarSettingsFromIni()
 {
 	check(IniHistoryHelper == nullptr);
-	IniHistoryHelper = new FCVarIniHistoryHelper();
+	IniHistoryHelper = MakeUnique<FCVarIniHistoryHelper>();
 }
 
 void ReapplyRecordedCVarSettingsFromIni()
 {
 	// first we need to reload the inis 
-	struct FIniToReload
+	for (TPair<FString, FConfigFile>& IniPair : *GConfig)
 	{
-		FIniToReload(const FString& InGlobalName, const FString& InIniBaseName) : IniBaseName(InIniBaseName), IniGlobalName(InGlobalName) { }
-		FString IniBaseName;
-		FString IniGlobalName;
-	};
-	TArray<FIniToReload> InisToReloadList;
-	InisToReloadList.Empty(GConfig->Num());
-	for (const auto& GlobalIni : *GConfig)
-	{
-		InisToReloadList.Add(FIniToReload(GlobalIni.Key, GlobalIni.Value.Name.ToString()));
-	}
-
-	for (const FIniToReload& IniToReload : InisToReloadList)
-	{
-		FString TempGlobalName = IniToReload.IniGlobalName;
-		GConfig->LoadGlobalIniFile(TempGlobalName, *IniToReload.IniBaseName, nullptr, true);
-		if (TempGlobalName.Compare(IniToReload.IniGlobalName, ESearchCase::IgnoreCase) != 0)
+		FConfigFile& ConfigFile = IniPair.Value;
+		if (ConfigFile.Num() > 0)
 		{
-			UE_LOG(LogConfig, Warning, TEXT("Tried to reload ini %s final name was %s"), *IniToReload.IniGlobalName, *TempGlobalName );
+			FName BaseName = ConfigFile.Name;
+			// Must call LoadLocalIniFile (NOT LoadGlobalIniFile) to preserve original enginedir/sourcedir for plugins
+			verify(FConfigCacheIni::LoadLocalIniFile(ConfigFile, *BaseName.ToString(), true, nullptr, true));
 		}
 	}
 
@@ -4339,7 +4327,6 @@ void ReapplyRecordedCVarSettingsFromIni()
 void DeleteRecordedCVarSettingsFromIni()
 {
 	check(IniHistoryHelper != nullptr);
-	delete IniHistoryHelper;
 	IniHistoryHelper = nullptr;
 }
 
@@ -4347,7 +4334,7 @@ void RecordConfigReadsFromIni()
 {
 #if !UE_BUILD_SHIPPING
 	check(ConfigHistoryHelper == nullptr);
-	ConfigHistoryHelper = new FConfigHistoryHelper();
+	ConfigHistoryHelper = MakeUnique<FConfigHistoryHelper>();
 #endif
 }
 
@@ -4363,7 +4350,6 @@ void DeleteRecordedConfigReadsFromIni()
 {
 #if !UE_BUILD_SHIPPING
 	check(ConfigHistoryHelper);
-	delete ConfigHistoryHelper;
 	ConfigHistoryHelper = nullptr;
 #endif
 }
