Core: Exposed a version of IConsoleManager::UnregisterConsoleObject which takes a name so unregistering can be done faster than O(N)

diff --git a/Engine/Source/Runtime/Core/Public/HAL/ConsoleManager.h b/Engine/Source/Runtime/Core/Public/HAL/ConsoleManager.h
--- a/Engine/Source/Runtime/Core/Public/HAL/ConsoleManager.h
+++ b/Engine/Source/Runtime/Core/Public/HAL/ConsoleManager.h
@@ -127,11 +127,12 @@ private: // ----------------------------------------------------
 	IConsoleObject* FindConsoleObjectUnfiltered(const TCHAR* Name) const;
 
 	/**
-	 * Unregisters a console variable or command, if that object was registered.  In the case of variables, this will actually only
-	 * "deactivate" the variable, so if it becomes registered again the state may persist.
+	 * Unregisters a console variable or command, if that object was registered.  For console variables, this will
+	 * actually only "deactivate" the variable so if it becomes registered again the state may persist
+	 * (unless bKeepState is false).
 	 *
 	 * @param	Name	Name of the console object to remove (not case sensitive)
-	 * @param bool bKeepState if the current state is kept in memory until a cvar with the same name is registered
+	 * @param	bKeepState	if the current state is kept in memory until a cvar with the same name is registered
 	 */
 	void UnregisterConsoleObject(const TCHAR* Name, bool bKeepState);
 
diff --git a/Engine/Source/Runtime/Core/Public/HAL/IConsoleManager.h b/Engine/Source/Runtime/Core/Public/HAL/IConsoleManager.h
--- a/Engine/Source/Runtime/Core/Public/HAL/IConsoleManager.h
+++ b/Engine/Source/Runtime/Core/Public/HAL/IConsoleManager.h
@@ -628,10 +628,18 @@ struct CORE_API IConsoleManager
 	 * Unregisters a console object, if that object was registered. O(n), n is the console object count
 	 *
 	 * @param ConsoleObject - object to remove
-	 * @param bool bKeepState if the current state is kept in memory until a cvar with the same name is registered
+	 * @param bKeepState if the current state is kept in memory until a cvar with the same name is registered
 	 */
 	virtual void UnregisterConsoleObject( IConsoleObject* ConsoleObject, bool bKeepState = true) = 0;
 
+	/**
+	 * Unregisters a console variable or command by name, if an object of that name was registered.
+	 *
+	 * @param Name - name of object to remove
+	 * @param bKeepState if the current state is kept in memory until a cvar with the same name is registered
+	 */
+	virtual void UnregisterConsoleObject(const TCHAR* Name, bool bKeepState = true) = 0;
+
 	/**
 	 * Find a console variable
 	 * @param Name must not be 0
