Prevent a spam warning about doing CVar FindByNames from URuntimeOptionsBase (they're a relatively cheap guarded map lookup and this use case is validation code designed to spot & prevent problems)

diff --git a/Engine/Source/Runtime/Core/Private/HAL/ConsoleManager.cpp b/Engine/Source/Runtime/Core/Private/HAL/ConsoleManager.cpp
--- a/Engine/Source/Runtime/Core/Private/HAL/ConsoleManager.cpp
+++ b/Engine/Source/Runtime/Core/Private/HAL/ConsoleManager.cpp
@@ -999,9 +999,9 @@ IConsoleCommand* FConsoleManager::RegisterConsoleCommand(const TCHAR* Name, cons
 }
 
 
-IConsoleVariable* FConsoleManager::FindConsoleVariable(const TCHAR* Name) const
+IConsoleVariable* FConsoleManager::FindConsoleVariable(const TCHAR* Name, bool bTrackFrequentCalls) const
 {
-	IConsoleObject* Obj = FindConsoleObject(Name);
+	IConsoleObject* Obj = FindConsoleObject(Name, bTrackFrequentCalls);
 
 	if(Obj)
 	{
@@ -1016,11 +1016,12 @@ IConsoleVariable* FConsoleManager::FindConsoleVariable(const TCHAR* Name) const
 	return 0;
 }
 
-IConsoleObject* FConsoleManager::FindConsoleObject(const TCHAR* Name) const
+IConsoleObject* FConsoleManager::FindConsoleObject(const TCHAR* Name, bool bTrackFrequentCalls) const
 {
 	IConsoleObject* CVar = FindConsoleObjectUnfiltered(Name);
 
 #if TRACK_CONSOLE_FIND_COUNT
+	if (bTrackFrequentCalls)
 	{
 		const bool bEarlyAppPhase = GFrameCounter < 1000;
 		if(CVar)
diff --git a/Engine/Source/Runtime/Core/Public/HAL/ConsoleManager.h b/Engine/Source/Runtime/Core/Public/HAL/ConsoleManager.h
--- a/Engine/Source/Runtime/Core/Public/HAL/ConsoleManager.h
+++ b/Engine/Source/Runtime/Core/Public/HAL/ConsoleManager.h
@@ -67,8 +67,8 @@ public:
 	virtual IConsoleCommand* RegisterConsoleCommand(const TCHAR* Name, const TCHAR* Help, const FConsoleCommandWithWorldArgsAndOutputDeviceDelegate& Command, uint32 Flags) override;
 	virtual IConsoleCommand* RegisterConsoleCommand(const TCHAR* Name, const TCHAR* Help, const FConsoleCommandWithOutputDeviceDelegate& Command, uint32 Flags) override;
 	virtual IConsoleCommand* RegisterConsoleCommand(const TCHAR* Name, const TCHAR* Help, uint32 Flags) override;
-	virtual IConsoleObject* FindConsoleObject(const TCHAR* Name) const override;
-	virtual IConsoleVariable* FindConsoleVariable(const TCHAR* Name) const override;
+	virtual IConsoleObject* FindConsoleObject(const TCHAR* Name, bool bTrackFrequentCalls = true) const override;
+	virtual IConsoleVariable* FindConsoleVariable(const TCHAR* Name, bool bTrackFrequentCalls = true) const override;
 	virtual void ForEachConsoleObjectThatStartsWith(const FConsoleObjectVisitor& Visitor, const TCHAR* ThatStartsWith) const override;
 	virtual void ForEachConsoleObjectThatContains(const FConsoleObjectVisitor& Visitor, const TCHAR* ThatContains) const override;
 	virtual bool ProcessUserConsoleInput(const TCHAR* InInput, FOutputDevice& Ar, UWorld* InWorld) override;
diff --git a/Engine/Source/Runtime/Core/Public/HAL/IConsoleManager.h b/Engine/Source/Runtime/Core/Public/HAL/IConsoleManager.h
--- a/Engine/Source/Runtime/Core/Public/HAL/IConsoleManager.h
+++ b/Engine/Source/Runtime/Core/Public/HAL/IConsoleManager.h
@@ -645,14 +645,14 @@ struct CORE_API IConsoleManager
 	 * @param Name must not be 0
 	 * @return 0 if the object wasn't found
 	 */
-	virtual IConsoleVariable* FindConsoleVariable(const TCHAR* Name) const = 0;
+	virtual IConsoleVariable* FindConsoleVariable(const TCHAR* Name, bool bTrackFrequentCalls = true) const = 0;
 
 	/**
 	* Find a console variable or command
 	* @param Name must not be 0
 	* @return 0 if the object wasn't found
 	*/
-	virtual IConsoleObject* FindConsoleObject(const TCHAR* Name) const = 0;
+	virtual IConsoleObject* FindConsoleObject(const TCHAR* Name, bool bTrackFrequentCalls = true) const = 0;
 
 	/**
 	 * Find a typed console variable (faster access to the value, no virtual function call)
diff --git a/Engine/Source/Runtime/Engine/Private/RuntimeOptionsBase.cpp b/Engine/Source/Runtime/Engine/Private/RuntimeOptionsBase.cpp
--- a/Engine/Source/Runtime/Engine/Private/RuntimeOptionsBase.cpp
+++ b/Engine/Source/Runtime/Engine/Private/RuntimeOptionsBase.cpp
@@ -116,7 +116,7 @@ void URuntimeOptionsBase::RegisterSupportedConsoleVariables(bool bDuringReload)
 			IConsoleManager::Get().UnregisterConsoleObject(*FullyQualifiedName, /*bKeepState=*/ false);
 		}
 		
-		if (IConsoleObject* ExistingObj = IConsoleManager::Get().FindConsoleObject(*FullyQualifiedName))
+		if (IConsoleObject* ExistingObj = IConsoleManager::Get().FindConsoleObject(*FullyQualifiedName, /*bTrackFrequentCalls=*/ false))
 		{
 			UE_LOG(LogRuntimeOptionsBase, Error, TEXT("A conflicting CVar '%s' already exists when trying to register a runtime option for the property %s%s::%s"),
 				*FullyQualifiedName,
