Fix FName crash when constructing from WIDECHAR / TCHAR nullptr and added page protection to catch FNameEntry stomps

Discussed cost of page protection on different platforms with Arciel Rekman and we both think it should be non-noticable.

diff --git a/Engine/Source/Programs/UnrealHeaderTool/UnrealHeaderTool.Target.cs b/Engine/Source/Programs/UnrealHeaderTool/UnrealHeaderTool.Target.cs
--- a/Engine/Source/Programs/UnrealHeaderTool/UnrealHeaderTool.Target.cs
+++ b/Engine/Source/Programs/UnrealHeaderTool/UnrealHeaderTool.Target.cs
@@ -33,8 +33,8 @@ public class UnrealHeaderToolTarget : TargetRules
 		// UnrealHeaderTool is a console application, not a Windows app (sets entry point to main(), instead of WinMain())
 		bIsBuildingConsoleApplication = true;
 
-		GlobalDefinitions.Add("HACK_HEADER_GENERATOR=1");
-
-		GlobalDefinitions.Add("USE_LOCALIZED_PACKAGE_CACHE=0");
+        GlobalDefinitions.Add("HACK_HEADER_GENERATOR=1");
+        GlobalDefinitions.Add("FNAME_WRITE_PROTECT_PAGES=0");
+        GlobalDefinitions.Add("USE_LOCALIZED_PACKAGE_CACHE=0");
 	}
 }
diff --git a/Engine/Source/Runtime/Core/Private/UObject/UnrealNames.cpp b/Engine/Source/Runtime/Core/Private/UObject/UnrealNames.cpp
--- a/Engine/Source/Runtime/Core/Private/UObject/UnrealNames.cpp
+++ b/Engine/Source/Runtime/Core/Private/UObject/UnrealNames.cpp
@@ -25,6 +25,11 @@
 #include "Serialization/ArchiveFromStructuredArchive.h"
 #include "Hash/CityHash.h"
 
+// Page protection to catch FNameEntry stomps
+#ifndef FNAME_WRITE_PROTECT_PAGES
+#define FNAME_WRITE_PROTECT_PAGES 1
+#endif
+
 DEFINE_LOG_CATEGORY_STATIC(LogUnrealNames, Log, All);
 
 const TCHAR* LexToString(EName Ename)
@@ -260,7 +265,7 @@ public:
 	/** Initializes all member variables. */
 	FNameEntryAllocator()
 	{
-		Blocks[0] = (uint8*)FMemory::Malloc(BlockSizeBytes);
+		Blocks[0] = (uint8*)FMemory::Malloc(BlockSizeBytes, FPlatformMemory::GetConstants().PageSize);
 	}
 
 	/**
@@ -370,13 +375,16 @@ private:
 			Terminator->Header.Len = 0;
 		}
 
+#if FNAME_WRITE_PROTECT_PAGES
+		FPlatformMemory::PageProtect(Blocks[CurrentBlock], BlockSizeBytes, /* read */ true, /* write */ false);
+#endif
 		++CurrentBlock;
 		CurrentByteCursor = 0;
 
 		check(CurrentBlock < FNameMaxBlocks);
 		check(Blocks[CurrentBlock] == nullptr);
 
-		Blocks[CurrentBlock] = (uint8*)FMemory::Malloc(BlockSizeBytes);
+		Blocks[CurrentBlock] = (uint8*)FMemory::Malloc(BlockSizeBytes, FPlatformMemory::GetConstants().PageSize);
 	}
 
 	mutable FRWLock Lock;
@@ -1443,10 +1451,13 @@ static int32 GetLengthAndWidth(const WIDECHAR* Str, bool& bOutIsWide)
 {
 	uint32 UserCharBits = 0;
 	const WIDECHAR* It = Str;
-	while (*It)
+	if (Str)
 	{
-		UserCharBits |= TChar<WIDECHAR>::ToUnsigned(*It);
-		++It;
+		while (*It)
+		{
+			UserCharBits |= TChar<WIDECHAR>::ToUnsigned(*It);
+			++It;
+		}
 	}
 
 	bOutIsWide = UserCharBits & 0xffffff80u;
@@ -1575,6 +1586,9 @@ struct FNameHelper
 		{
 			check(FindType == FNAME_Replace_Not_Safe_For_Threading);
 
+#if FNAME_WRITE_PROTECT_PAGES
+			checkf(false, TEXT("FNAME_Replace_Not_Safe_For_Threading can't be used together with page protection."));
+#endif
 			DisplayId = Pool.Store(View);
 #if WITH_CASE_PRESERVING_NAME
 			ComparisonId = Pool.Resolve(DisplayId).ComparisonId;
@@ -1936,6 +1950,7 @@ void FName::AutoTest()
 
 	const FName NullName(static_cast<ANSICHAR*>(nullptr));
 	check(NullName.IsNone());
+	check(NullName == FName(static_cast<WIDECHAR*>(nullptr)));
 	check(NullName == FName(NAME_None));
 	check(NullName == FName());
 	check(NullName == FName(""));
