Reenabled FName page protection and disabled the unit tests that tripped it

diff --git a/Engine/Source/Runtime/Core/Private/UObject/UnrealNames.cpp b/Engine/Source/Runtime/Core/Private/UObject/UnrealNames.cpp
--- a/Engine/Source/Runtime/Core/Private/UObject/UnrealNames.cpp
+++ b/Engine/Source/Runtime/Core/Private/UObject/UnrealNames.cpp
@@ -27,7 +27,7 @@
 
 // Page protection to catch FNameEntry stomps
 #ifndef FNAME_WRITE_PROTECT_PAGES
-#define FNAME_WRITE_PROTECT_PAGES 0
+#define FNAME_WRITE_PROTECT_PAGES 1
 #endif
 
 DEFINE_LOG_CATEGORY_STATIC(LogUnrealNames, Log, All);
@@ -1997,6 +1997,7 @@ void FName::AutoTest()
 		check(FName(TEXT("UNIQUEUNICORN!!"), FNAME_Find) == UniqueName);
 		check(FName("uniqueunicorn!!", FNAME_Find) == UniqueName);
 
+#if !FNAME_WRITE_PROTECT_PAGES
 		// Check FNAME_Replace_Not_Safe_For_Threading updates casing
 		check(0 != UniqueName.GetPlainNameString().Compare("UNIQUEunicorn!!", ESearchCase::CaseSensitive));
 		const FName UniqueNameReplaced("UNIQUEunicorn!!", FNAME_Replace_Not_Safe_For_Threading);
@@ -2011,6 +2012,7 @@ void FName::AutoTest()
 		// Check FNAME_Replace_Not_Safe_For_Threading adds entries that do not exist
 		const FName AddedByReplace("WasAdded!!", FNAME_Replace_Not_Safe_For_Threading);
 		check(FName("WasAdded!!", FNAME_Find) == AddedByReplace);
+#endif
 	
 		Once = false;
 	}
