Bundle Manager
- Add Requested progress status
- Fix FBackgroundHtttpBundleWatchDog sometimes killing downloads that were enqueued in the background http manager but hadn't started yet.

diff --git a/Engine/Source/Runtime/Core/Public/GenericPlatform/GenericPlatformInstallBundleManager.h b/Engine/Source/Runtime/Core/Public/GenericPlatform/GenericPlatformInstallBundleManager.h
--- a/Engine/Source/Runtime/Core/Public/GenericPlatform/GenericPlatformInstallBundleManager.h
+++ b/Engine/Source/Runtime/Core/Public/GenericPlatform/GenericPlatformInstallBundleManager.h
@@ -114,6 +114,7 @@ ENUM_CLASS_FLAGS(EInstallBundleRequestFlags)
 
 enum class EInstallBundleStatus : int
 {
+	Requested,
 	QueuedForDownload,
 	Downloading,
 	QueuedForInstall,
@@ -129,6 +130,7 @@ inline const TCHAR* LexToString(EInstallBundleStatus Status)
 	using UnderType = __underlying_type(EInstallBundleStatus);
 	static const TCHAR* Strings[] =
 	{
+		TEXT("Requested"),
 		TEXT("QueuedForDownload"),
 		TEXT("Downloading"),
 		TEXT("QueuedForInstall"),
