FConfigCacheIni::LoadGlobalIniFile now retains config directories when doing a force reload.  This brings it inline with FConfigCacheIni::LoadLocalIniFile and means that calling it with the force reload flag now works correctly for plugin INIs.

diff --git a/Engine/Source/Runtime/Core/Private/Misc/ConfigCacheIni.cpp b/Engine/Source/Runtime/Core/Private/Misc/ConfigCacheIni.cpp
--- a/Engine/Source/Runtime/Core/Private/Misc/ConfigCacheIni.cpp
+++ b/Engine/Source/Runtime/Core/Private/Misc/ConfigCacheIni.cpp
@@ -3515,10 +3515,33 @@ bool FConfigCacheIni::LoadGlobalIniFile(FString& FinalIniFilename, const TCHAR*
 		return true;
 	}
 
+	FString EngineConfigDir = FPaths::EngineConfigDir();
+	FString SourceConfigDir = FPaths::SourceConfigDir();
+
+	if (bForceReload) // If reloading we should preserve the existing config dirs
+	{
+		// If base ini, try to use an existing GConfig file to set the config directories instead of assuming defaults
+		FConfigFile* BaseConfig = GConfig->FindConfigFileWithBaseName(BaseIniName);
+		if (BaseConfig)
+		{
+			FIniFilename* EngineFilename = BaseConfig->SourceIniHierarchy.Find(EConfigFileHierarchy::EngineDirBase);
+			if (EngineFilename)
+			{
+				EngineConfigDir = FPaths::GetPath(EngineFilename->Filename) + TEXT("/");
+			}
+
+			FIniFilename* GameFilename = BaseConfig->SourceIniHierarchy.Find(EConfigFileHierarchy::GameDirDefault);
+			if (GameFilename)
+			{
+				SourceConfigDir = FPaths::GetPath(GameFilename->Filename) + TEXT("/");
+			}
+		}
+	}
+
 	// make a new entry in GConfig (overwriting what's already there)
 	FConfigFile& NewConfigFile = GConfig->Add(FinalIniFilename, FConfigFile());
 
-	return LoadExternalIniFile(NewConfigFile, BaseIniName, *FPaths::EngineConfigDir(), *FPaths::SourceConfigDir(), true, Platform, bForceReload, true, bAllowGeneratedIniWhenCooked, GeneratedConfigDir);
+	return LoadExternalIniFile(NewConfigFile, BaseIniName, *EngineConfigDir, *SourceConfigDir, true, Platform, bForceReload, true, bAllowGeneratedIniWhenCooked, GeneratedConfigDir);
 }
 
 bool FConfigCacheIni::LoadLocalIniFile(FConfigFile& ConfigFile, const TCHAR* IniName, bool bIsBaseIniName, const TCHAR* Platform, bool bForceReload )
@@ -3530,9 +3553,8 @@ bool FConfigCacheIni::LoadLocalIniFile(FConfigFile& ConfigFile, const TCHAR* Ini
 
 	if (bIsBaseIniName)
 	{
-		FConfigFile* BaseConfig = GConfig->FindConfigFileWithBaseName(IniName);
 		// If base ini, try to use an existing GConfig file to set the config directories instead of assuming defaults
-
+		FConfigFile* BaseConfig = GConfig->FindConfigFileWithBaseName(IniName);
 		if (BaseConfig)
 		{
 			if (BaseConfig->SourceEngineConfigDir.Len())
