Add flag for bundle requests to control whether notifications are sent when background downloads are complete.

diff --git a/Engine/Source/Runtime/Core/Public/GenericPlatform/GenericPlatformInstallBundleManager.h b/Engine/Source/Runtime/Core/Public/GenericPlatform/GenericPlatformInstallBundleManager.h
--- a/Engine/Source/Runtime/Core/Public/GenericPlatform/GenericPlatformInstallBundleManager.h
+++ b/Engine/Source/Runtime/Core/Public/GenericPlatform/GenericPlatformInstallBundleManager.h
@@ -105,9 +105,10 @@ inline const TCHAR* GetInstallBundlePauseReason(EInstallBundlePauseFlags Flags)
 
 enum class EInstallBundleRequestFlags : uint32
 {
-	None = 0,
-	CheckForCellularDataUsage = (1 << 0),
-	UseBackgroundDownloads = (1 << 1),
+	None											= 0,
+	CheckForCellularDataUsage						= (1 << 0),
+	UseBackgroundDownloads							= (1 << 1),
+	SendNotificationIfDownloadCompletesInBackground = (1 << 2),
 	Defaults = UseBackgroundDownloads,
 };
 ENUM_CLASS_FLAGS(EInstallBundleRequestFlags)
