Restore a portion of the config cache changes lost in resolve of changes coming in CL# 6981836
#rnx

diff --git a/Engine/Source/Runtime/Core/Private/Misc/ConfigCacheIni.cpp b/Engine/Source/Runtime/Core/Private/Misc/ConfigCacheIni.cpp
--- a/Engine/Source/Runtime/Core/Private/Misc/ConfigCacheIni.cpp
+++ b/Engine/Source/Runtime/Core/Private/Misc/ConfigCacheIni.cpp
@@ -3524,16 +3524,14 @@ bool FConfigCacheIni::LoadGlobalIniFile(FString& FinalIniFilename, const TCHAR*
 		FConfigFile* BaseConfig = GConfig->FindConfigFileWithBaseName(BaseIniName);
 		if (BaseConfig)
 		{
-			FIniFilename* EngineFilename = BaseConfig->SourceIniHierarchy.Find(EConfigFileHierarchy::EngineDirBase);
-			if (EngineFilename)
+			if (BaseConfig->SourceEngineConfigDir.Len())
 			{
-				EngineConfigDir = FPaths::GetPath(EngineFilename->Filename) + TEXT("/");
+				EngineConfigDir = BaseConfig->SourceEngineConfigDir;
 			}
 
-			FIniFilename* GameFilename = BaseConfig->SourceIniHierarchy.Find(EConfigFileHierarchy::GameDirDefault);
-			if (GameFilename)
+			if (BaseConfig->SourceProjectConfigDir.Len())
 			{
-				SourceConfigDir = FPaths::GetPath(GameFilename->Filename) + TEXT("/");
+				SourceConfigDir = BaseConfig->SourceProjectConfigDir;
 			}
 		}
 	}
