We dont need to sigqueue anymore so just crash

diff --git a/Engine/Source/Runtime/Core/Private/Unix/UnixPlatformStackWalk.cpp b/Engine/Source/Runtime/Core/Private/Unix/UnixPlatformStackWalk.cpp
--- a/Engine/Source/Runtime/Core/Private/Unix/UnixPlatformStackWalk.cpp
+++ b/Engine/Source/Runtime/Core/Private/Unix/UnixPlatformStackWalk.cpp
@@ -867,16 +867,7 @@ void ReportAssert(const TCHAR* ErrorMessage, int NumStackFramesToIgnore)
 	GCrashErrorMessage = ErrorMessage;
 	GCrashErrorType = ECrashContextType::Assert;
 
-	// Store NumStackFramesToIgnore in signal data	
-	sigval UserData;
-	UserData.sival_int = NumStackFramesToIgnore + 2; // +2 for this function and sigqueue()
-	sigqueue(getpid(),  SIGSEGV, UserData);
-	
-	// Make sure we never return
-	for (;;)
-	{
-		FPlatformProcess::Sleep(60.0f);
-	}
+	FPlatformMisc::RaiseException(1);
 }
 
 void ReportGPUCrash(const TCHAR* ErrorMessage, int NumStackFramesToIgnore)
@@ -884,16 +875,7 @@ void ReportGPUCrash(const TCHAR* ErrorMessage, int NumStackFramesToIgnore)
 	GCrashErrorMessage = ErrorMessage;
 	GCrashErrorType = ECrashContextType::GPUCrash;
 
-	// Store NumStackFramesToIgnore in signal data	
-	sigval UserData;
-	UserData.sival_int = NumStackFramesToIgnore + 2; // +2 for this function and sigqueue()
-	sigqueue(getpid(),  SIGSEGV, UserData);
-	
-	// Make sure we never return
-	for (;;)
-	{
-		FPlatformProcess::Sleep(60.0f);
-	}
+	FPlatformMisc::RaiseException(1);
 }
 
 static FCriticalSection EnsureLock;
