Install Bundle Manager
- Add configuration error for initialization

diff --git a/Engine/Source/Runtime/Core/Public/GenericPlatform/GenericPlatformInstallBundleManager.h b/Engine/Source/Runtime/Core/Public/GenericPlatform/GenericPlatformInstallBundleManager.h
--- a/Engine/Source/Runtime/Core/Public/GenericPlatform/GenericPlatformInstallBundleManager.h
+++ b/Engine/Source/Runtime/Core/Public/GenericPlatform/GenericPlatformInstallBundleManager.h
@@ -24,6 +24,7 @@ enum class EInstallBundleModuleInitResult : int
 	ManifestDownloadError,
 	BackgroundDownloadsIniDownloadError,
 	NoInternetConnectionError,
+	ConfigurationError,
 	Count
 };
 
@@ -42,6 +43,7 @@ inline const TCHAR* LexToString(EInstallBundleModuleInitResult Result)
 		TEXT("ManifestDownloadError"),
 		TEXT("BackgroundDownloadsIniDownloadError"),
 		TEXT("NoInternetConnectionError"),
+		TEXT("ConfigurationError"),
 	};
 	static_assert(static_cast<UnderType>(EInstallBundleModuleInitResult::Count) == ARRAY_COUNT(Strings), "");
 
