Fix Win32 compile errors with the Windows 10.0.18362.0 SDK because struct packing is not set to the default value.

#rnx

diff --git a/Engine/Source/Runtime/Core/Public/Windows/MinimalWindowsApi.h b/Engine/Source/Runtime/Core/Public/Windows/MinimalWindowsApi.h
--- a/Engine/Source/Runtime/Core/Public/Windows/MinimalWindowsApi.h
+++ b/Engine/Source/Runtime/Core/Public/Windows/MinimalWindowsApi.h
@@ -18,6 +18,12 @@
 	#define MINIMAL_WINDOWS_API extern "C" __declspec(dllimport)
 #endif
 
+// Set up the default struct packing on Win32
+#if PLATFORM_WINDOWS && PLATFORM_32BITS
+	#pragma pack(push)
+	#pragma pack(8)
+#endif
+
 // Undefine Windows types that we're going to redefine. This should be done by HideWindowsPlatformTypes.h after including any system headers, 
 // but it's less friction to just undefine them here.
 #pragma push_macro("TRUE")
@@ -203,3 +209,8 @@ namespace Windows
 // Restore the definitions for TRUE and FALSE
 #pragma pop_macro("FALSE")
 #pragma pop_macro("TRUE")
+
+// Restore the packing setting
+#if PLATFORM_WINDOWS && PLATFORM_32BITS
+	#pragma pack(pop)
+#endif
diff --git a/Engine/Source/Runtime/Core/Public/Windows/PostWindowsApi.h b/Engine/Source/Runtime/Core/Public/Windows/PostWindowsApi.h
--- a/Engine/Source/Runtime/Core/Public/Windows/PostWindowsApi.h
+++ b/Engine/Source/Runtime/Core/Public/Windows/PostWindowsApi.h
@@ -73,6 +73,11 @@ THIRD_PARTY_INCLUDES_END
 #pragma pop_macro("TRUE")
 #pragma pop_macro("FALSE")
 
+// Restore the struct packing setting
+#if PLATFORM_WINDOWS && PLATFORM_32BITS
+	#pragma pack(pop)
+#endif
+
 // Redefine CDECL to our version of the #define.  <AJS> Is this really necessary?
 #define CDECL	    __cdecl					/* Standard C function */
 
diff --git a/Engine/Source/Runtime/Core/Public/Windows/PreWindowsApi.h b/Engine/Source/Runtime/Core/Public/Windows/PreWindowsApi.h
--- a/Engine/Source/Runtime/Core/Public/Windows/PreWindowsApi.h
+++ b/Engine/Source/Runtime/Core/Public/Windows/PreWindowsApi.h
@@ -3,6 +3,12 @@
 // Disable the warning that the pack size is changed in this header.
 #pragma warning(disable:4103)
 
+// The 10.0.18362.0 SDK introduces an error if the packing isn't the default for the platform.
+#if PLATFORM_WINDOWS && PLATFORM_32BITS
+	#pragma pack(push)
+	#pragma pack(8)
+#endif
+
 // Save these macros for later; Windows redefines them
 #pragma push_macro("MAX_uint8")
 #pragma push_macro("MAX_uint16")
