improve error handling when trying to log an unregistered tag.

diff --git a/Engine/Source/Runtime/Core/Private/HAL/LowLevelMemTracker.cpp b/Engine/Source/Runtime/Core/Private/HAL/LowLevelMemTracker.cpp
--- a/Engine/Source/Runtime/Core/Private/HAL/LowLevelMemTracker.cpp
+++ b/Engine/Source/Runtime/Core/Private/HAL/LowLevelMemTracker.cpp
@@ -2032,6 +2032,7 @@ FString FLLMCsvWriter::GetTagName(int64 Tag, FLLMCustomTag* CustomTags, const in
 			Result = GetTagName( ParentTags[Tag], CustomTags, nullptr ) + TEXT("/");
 		}
 
+		LLMCheckf(CustomTags[Tag - LLM_CUSTOM_TAG_START].Name != nullptr, TEXT("Tag %lld has no name"), Tag ); 
 		Result += CustomTags[Tag - LLM_CUSTOM_TAG_START].Name;
 	}
 	else
