disable the small block allocator on IOS

diff --git a/Engine/Source/Runtime/Core/Private/HAL/MallocBinned.cpp b/Engine/Source/Runtime/Core/Private/HAL/MallocBinned.cpp
--- a/Engine/Source/Runtime/Core/Private/HAL/MallocBinned.cpp
+++ b/Engine/Source/Runtime/Core/Private/HAL/MallocBinned.cpp
@@ -136,6 +136,8 @@ struct FMallocBinned::PoolHashBucket
 	}
 };
 
+#define USE_OS_SMALL_BLOCK_ALLOC 0
+
 #if PLATFORM_IOS
 #define PLAT_PAGE_SIZE_LIMIT 16384
 #define PLAT_BINNED_ALLOC_POOLSIZE 16384
@@ -158,8 +160,10 @@ struct FMallocBinned::Private
 	// On IOS can push small allocs in to a pre-allocated small block pool
 	enum { SMALL_BLOCK_POOL_SIZE = PLAT_SMALL_BLOCK_POOL_SIZE };
 
+#if USE_OS_SMALL_BLOCK_ALLOC
 	static uint64 SmallBlockStartPtr;
 	static uint64 SmallBlockEndPtr;
+#endif
 	
 	// Implementation. 
 	static CA_NO_RETURN void OutOfMemory(uint64 Size, uint32 Alignment=0)
@@ -451,12 +455,14 @@ struct FMallocBinned::Private
 #if PLATFORM_IOS || PLATFORM_MAC
         if (Pool == NULL)
         {
+#if USE_OS_SMALL_BLOCK_ALLOC
 			// check the small block allocation range
 			if ((uint64)Ptr >= SmallBlockStartPtr && (uint64)Ptr <= SmallBlockEndPtr)
 			{
 				SmallOSFree(Allocator, Ptr, Private::SMALL_BLOCK_POOL_SIZE);
 				return;
 			}
+#endif
             UE_LOG(LogMemory, Warning, TEXT("Attempting to free a pointer we didn't allocate!"));
             return;
         }
@@ -727,8 +733,10 @@ struct FMallocBinned::Private
 	}
 };
 
+#if USE_OS_SMALL_BLOCK_ALLOC
 uint64 FMallocBinned::Private::SmallBlockStartPtr = 0xffffffffffffffff;
 uint64 FMallocBinned::Private::SmallBlockEndPtr = 0x0000000000000000;
+#endif
 
 void FMallocBinned::GetAllocatorStats( FGenericMemoryStats& out_Stats )
 {
@@ -937,6 +945,7 @@ void* FMallocBinned::Malloc(SIZE_T Size, uint32 Alignment)
 	
 	FFreeMem* Free = nullptr;
 	bool bUsePools = true;
+#if USE_OS_SMALL_BLOCK_ALLOC
 	if (Size <= Private::SMALL_BLOCK_POOL_SIZE)
 	{
 		UPTRINT AlignedSize = Align(Size, Alignment);
@@ -964,88 +973,89 @@ void* FMallocBinned::Malloc(SIZE_T Size, uint32 Alignment)
 			bUsePools = false;
 		}
 	}
+#endif
 	if (bUsePools)
 	{
-	if( Size < BinnedSizeLimit)
-	{
-		// Allocate from pool.
-		FPoolTable* Table = MemSizeToPoolTable[Size];
+		if( Size < BinnedSizeLimit)
+		{
+			// Allocate from pool.
+			FPoolTable* Table = MemSizeToPoolTable[Size];
 #ifdef USE_FINE_GRAIN_LOCKS
-		FScopeLock TableLock(&Table->CriticalSection);
+			FScopeLock TableLock(&Table->CriticalSection);
 #endif
-		checkSlow(Size <= Table->BlockSize);
+			checkSlow(Size <= Table->BlockSize);
 
-		Private::TrackStats(Table, Size);
+			Private::TrackStats(Table, Size);
 
-		FPoolInfo* Pool = Table->FirstPool;
-		if( !Pool )
-		{
-			Pool = Private::AllocatePoolMemory(*this, Table, Private::BINNED_ALLOC_POOL_SIZE/*PageSize*/, Size);
-		}
+			FPoolInfo* Pool = Table->FirstPool;
+			if( !Pool )
+			{
+				Pool = Private::AllocatePoolMemory(*this, Table, Private::BINNED_ALLOC_POOL_SIZE/*PageSize*/, Size);
+			}
 
-		Free = Private::AllocateBlockFromPool(*this, Table, Pool, Alignment);
-	}
-	else if ( ((Size >= BinnedSizeLimit && Size <= PagePoolTable[0].BlockSize) ||
-		(Size > PageSize && Size <= PagePoolTable[1].BlockSize)))
-	{
-		// Bucket in a pool of 3*PageSize or 6*PageSize
-		uint32 BinType = Size < PageSize ? 0 : 1;
-		uint32 PageCount = 3*BinType + 3;
-		FPoolTable* Table = &PagePoolTable[BinType];
+			Free = Private::AllocateBlockFromPool(*this, Table, Pool, Alignment);
+		}
+		else if ( ((Size >= BinnedSizeLimit && Size <= PagePoolTable[0].BlockSize) ||
+				   (Size > PageSize && Size <= PagePoolTable[1].BlockSize)))
+		{
+			// Bucket in a pool of 3*PageSize or 6*PageSize
+			uint32 BinType = Size < PageSize ? 0 : 1;
+			uint32 PageCount = 3*BinType + 3;
+			FPoolTable* Table = &PagePoolTable[BinType];
 #ifdef USE_FINE_GRAIN_LOCKS
-		FScopeLock TableLock(&Table->CriticalSection);
+			FScopeLock TableLock(&Table->CriticalSection);
 #endif
-		checkSlow(Size <= Table->BlockSize);
+			checkSlow(Size <= Table->BlockSize);
 
-		Private::TrackStats(Table, Size);
+			Private::TrackStats(Table, Size);
 
-		FPoolInfo* Pool = Table->FirstPool;
+			FPoolInfo* Pool = Table->FirstPool;
 		if( !Pool )
-		{
-			Pool = Private::AllocatePoolMemory(*this, Table, PageCount*PageSize, BinnedSizeLimit+BinType);
-		}
+			{
+				Pool = Private::AllocatePoolMemory(*this, Table, PageCount*PageSize, BinnedSizeLimit+BinType);
+			}
 
-		Free = Private::AllocateBlockFromPool(*this, Table, Pool, Alignment);
-	}
-	else
-	{
-		// Use OS for large allocations.
-		UPTRINT AlignedSize = Align(Size,PageSize);
-		SIZE_T ActualPoolSize; //TODO: use this to reduce waste?
-		Free = (FFreeMem*)Private::OSAlloc(*this, AlignedSize, ActualPoolSize);
-		if( !Free )
-		{
-			Private::OutOfMemory(AlignedSize);
+			Free = Private::AllocateBlockFromPool(*this, Table, Pool, Alignment);
 		}
+		else
+		{
+			// Use OS for large allocations.
+			UPTRINT AlignedSize = Align(Size,PageSize);
+			SIZE_T ActualPoolSize; //TODO: use this to reduce waste?
+			Free = (FFreeMem*)Private::OSAlloc(*this, AlignedSize, ActualPoolSize);
+			if( !Free )
+			{
+				Private::OutOfMemory(AlignedSize);
+			}
 
-		void* AlignedFree = Align(Free, Alignment);
+			void* AlignedFree = Align(Free, Alignment);
 
-		// Create indirect.
-		FPoolInfo* Pool;
-		{
+			// Create indirect.
+			FPoolInfo* Pool;
+			{
 #ifdef USE_FINE_GRAIN_LOCKS
-			FScopeLock PoolInfoLock(&AccessGuard);
+				FScopeLock PoolInfoLock(&AccessGuard);
 #endif
-			Pool = Private::GetPoolInfo(*this, (UPTRINT)Free);
+				Pool = Private::GetPoolInfo(*this, (UPTRINT)Free);
 
-			if ((UPTRINT)Free != ((UPTRINT)AlignedFree & ~((UPTRINT)PageSize - 1)))
-			{
-				// Mark the FPoolInfo for AlignedFree to jump back to the FPoolInfo for ptr.
-				for (UPTRINT i = (UPTRINT)PageSize, Offset = 0; i < AlignedSize; i += PageSize, ++Offset)
+				if ((UPTRINT)Free != ((UPTRINT)AlignedFree & ~((UPTRINT)PageSize - 1)))
 				{
-					FPoolInfo* TrailingPool = Private::GetPoolInfo(*this, ((UPTRINT)Free) + i);
-					check(TrailingPool);
-					//Set trailing pools to point back to first pool
-					TrailingPool->SetAllocationSizes(0, 0, Offset, BinnedOSTableIndex);
+					// Mark the FPoolInfo for AlignedFree to jump back to the FPoolInfo for ptr.
+					for (UPTRINT i = (UPTRINT)PageSize, Offset = 0; i < AlignedSize; i += PageSize, ++Offset)
+					{
+						FPoolInfo* TrailingPool = Private::GetPoolInfo(*this, ((UPTRINT)Free) + i);
+						check(TrailingPool);
+						//Set trailing pools to point back to first pool
+						TrailingPool->SetAllocationSizes(0, 0, Offset, BinnedOSTableIndex);
+					}
 				}
 			}
+			Free = (FFreeMem*)AlignedFree;
+			Pool->SetAllocationSizes(Size, AlignedSize, BinnedOSTableIndex, BinnedOSTableIndex);
+			BINNED_PEAK_STATCOUNTER(OsPeak, BINNED_ADD_STATCOUNTER(OsCurrent, AlignedSize));
+			BINNED_PEAK_STATCOUNTER(UsedPeak, BINNED_ADD_STATCOUNTER(UsedCurrent, Size));
+			BINNED_PEAK_STATCOUNTER(WastePeak, BINNED_ADD_STATCOUNTER(WasteCurrent, (int64)(AlignedSize - Size)));
 		}
-		Free = (FFreeMem*)AlignedFree;
-		Pool->SetAllocationSizes(Size, AlignedSize, BinnedOSTableIndex, BinnedOSTableIndex);
-		BINNED_PEAK_STATCOUNTER(OsPeak, BINNED_ADD_STATCOUNTER(OsCurrent, AlignedSize));
-		BINNED_PEAK_STATCOUNTER(UsedPeak, BINNED_ADD_STATCOUNTER(UsedCurrent, Size));
-		BINNED_PEAK_STATCOUNTER(WastePeak, BINNED_ADD_STATCOUNTER(WasteCurrent, (int64)(AlignedSize - Size)));
-	}
 	}
 
 	MEM_TIME(MemTime += FPlatformTime::Seconds());
@@ -1072,60 +1082,70 @@ void* FMallocBinned::Realloc( void* Ptr, SIZE_T NewSize, uint32 Alignment )
 	void* NewPtr = Ptr;
 	if( Ptr && NewSize )
 	{
+#if USE_OS_SMALL_BLOCK_ALLOC
 		// special case for really small blocks
 		if ((uint64)Ptr >= Private::SmallBlockStartPtr && (uint64)Ptr <= Private::SmallBlockEndPtr)
 		{
 			// Grow or shrink.
 			NewPtr = Malloc(NewSizeUnmodified, Alignment);
-			FMemory::Memcpy(NewPtr, Ptr, FMath::Min<SIZE_T>(NewSizeUnmodified, Private::SMALL_BLOCK_POOL_SIZE)); // this is bad need to figure out the old size :(
-			Free( Ptr );
-		}
-		else
-		{
-		FPoolInfo* Pool = Private::FindPoolInfo(*this, (UPTRINT)Ptr, BasePtr);
-
-		if( Pool->TableIndex < BinnedOSTableIndex )
-		{
-			// Allocated from pool, so grow or shrink if necessary.
-			check(Pool->TableIndex > 0); // it isn't possible to allocate a size of 0, Malloc will increase the size to DEFAULT_BINNED_ALLOCATOR_ALIGNMENT
-			if (NewSizeUnmodified > MemSizeToPoolTable[Pool->TableIndex]->BlockSize || NewSizeUnmodified <= MemSizeToPoolTable[Pool->TableIndex - 1]->BlockSize)
+			SIZE_T CopySize = FMath::Min<SIZE_T>(NewSizeUnmodified, Private::SMALL_BLOCK_POOL_SIZE);
+#if UE_BUILD_TEST
+			if ((uint64)Ptr + CopySize >= 0x2a0000000)
 			{
-				NewPtr = Malloc(NewSizeUnmodified, Alignment);
-				FMemory::Memcpy(NewPtr, Ptr, FMath::Min<SIZE_T>(NewSizeUnmodified, MemSizeToPoolTable[Pool->TableIndex]->BlockSize - (Alignment - SpareBytesCount)));
-				Free( Ptr );
-			}
-			else if (((UPTRINT)Ptr & (UPTRINT)(Alignment - 1)) != 0)
-			{
-				NewPtr = Align(Ptr, Alignment);
-				FMemory::Memmove(NewPtr, Ptr, NewSize);
+				UE_LOG(LogMemory, Warning, TEXT("Going to memory copy past the end of the small block!!"));
 			}
+#endif
+			CopySize = FMath::Min<SIZE_T>(CopySize, (0x29fffffff-(uint64)Ptr)); // guard against going off the end of the VM area for the small block when doing the memcpy
+			FMemory::Memcpy(NewPtr, Ptr, CopySize); // this is bad need to figure out the old size :(
+			Free( Ptr );
 		}
 		else
+#endif
 		{
-			// Allocated from OS.
-			if( NewSize > Pool->GetOsBytes(PageSize, BinnedOSTableIndex) || NewSize * 3 < Pool->GetOsBytes(PageSize, BinnedOSTableIndex) * 2 )
+			FPoolInfo* Pool = Private::FindPoolInfo(*this, (UPTRINT)Ptr, BasePtr);
+
+			if( Pool->TableIndex < BinnedOSTableIndex )
 			{
-				// Grow or shrink.
-				NewPtr = Malloc(NewSizeUnmodified, Alignment);
-				FMemory::Memcpy(NewPtr, Ptr, FMath::Min<SIZE_T>(NewSizeUnmodified, Pool->GetBytes()));
-				Free( Ptr );
+				// Allocated from pool, so grow or shrink if necessary.
+				check(Pool->TableIndex > 0); // it isn't possible to allocate a size of 0, Malloc will increase the size to DEFAULT_BINNED_ALLOCATOR_ALIGNMENT
+				if (NewSizeUnmodified > MemSizeToPoolTable[Pool->TableIndex]->BlockSize || NewSizeUnmodified <= MemSizeToPoolTable[Pool->TableIndex - 1]->BlockSize)
+				{
+					NewPtr = Malloc(NewSizeUnmodified, Alignment);
+					FMemory::Memcpy(NewPtr, Ptr, FMath::Min<SIZE_T>(NewSizeUnmodified, MemSizeToPoolTable[Pool->TableIndex]->BlockSize - (Alignment - SpareBytesCount)));
+					Free( Ptr );
+				}
+				else if (((UPTRINT)Ptr & (UPTRINT)(Alignment - 1)) != 0)
+				{
+					NewPtr = Align(Ptr, Alignment);
+					FMemory::Memmove(NewPtr, Ptr, NewSize);
+				}
 			}
 			else
 			{
+				// Allocated from OS.
+				if( NewSize > Pool->GetOsBytes(PageSize, BinnedOSTableIndex) || NewSize * 3 < Pool->GetOsBytes(PageSize, BinnedOSTableIndex) * 2 )
+				{
+					// Grow or shrink.
+					NewPtr = Malloc(NewSizeUnmodified, Alignment);
+					FMemory::Memcpy(NewPtr, Ptr, FMath::Min<SIZE_T>(NewSizeUnmodified, Pool->GetBytes()));
+					Free( Ptr );
+				}
+				else
+				{
 //need a lock to cover the SetAllocationSizes()
 #ifdef USE_FINE_GRAIN_LOCKS
-				FScopeLock PoolInfoLock(&AccessGuard);
+					FScopeLock PoolInfoLock(&AccessGuard);
 #endif
-				int32 UsedChange = (NewSize - Pool->GetBytes());
+					int32 UsedChange = (NewSize - Pool->GetBytes());
 				
-				// Keep as-is, reallocation isn't worth the overhead.
-				BINNED_ADD_STATCOUNTER(UsedCurrent, UsedChange);
-				BINNED_PEAK_STATCOUNTER(UsedPeak, UsedCurrent);
-				BINNED_ADD_STATCOUNTER(WasteCurrent, (Pool->GetBytes() - NewSize));
-				Pool->SetAllocationSizes(NewSizeUnmodified, Pool->GetOsBytes(PageSize, BinnedOSTableIndex), BinnedOSTableIndex, BinnedOSTableIndex);
+					// Keep as-is, reallocation isn't worth the overhead.
+					BINNED_ADD_STATCOUNTER(UsedCurrent, UsedChange);
+					BINNED_PEAK_STATCOUNTER(UsedPeak, UsedCurrent);
+					BINNED_ADD_STATCOUNTER(WasteCurrent, (Pool->GetBytes() - NewSize));
+					Pool->SetAllocationSizes(NewSizeUnmodified, Pool->GetOsBytes(PageSize, BinnedOSTableIndex), BinnedOSTableIndex, BinnedOSTableIndex);
+				}
 			}
 		}
-		}
 	}
 	else if( Ptr == nullptr )
 	{
@@ -1164,7 +1184,7 @@ bool FMallocBinned::GetAllocationSize(void *Original, SIZE_T &SizeOut)
 #if PLATFORM_IOS || PLATFORM_MAC
 	if (Pool == NULL)
 	{
-#if PLATFORM_IOS
+#if USE_OS_SMALL_BLOCK_ALLOC
 		if ((uint64)Original >= Private::SmallBlockStartPtr && (uint64)Original <= Private::SmallBlockEndPtr)
 		{
 			SizeOut = Private::SMALL_BLOCK_POOL_SIZE;
@@ -1210,12 +1230,15 @@ SIZE_T FMallocBinned::QuantizeSize(SIZE_T Size, uint32 Alignment)
 	Size = FMath::Max<SIZE_T>(PoolTable[0].BlockSize, Size + (Alignment - SpareBytesCount));
 
 	SIZE_T Result;
+#if USE_OS_SMALL_BLOCK_ALLOC
 	if (Size <= Private::SMALL_BLOCK_POOL_SIZE)
 	{
 		UPTRINT AlignedSize = Align(Size, Alignment);
 		Result = SIZE_T(AlignedSize);
 	}
-	else if (Size < BinnedSizeLimit)
+	else
+#endif
+	if (Size < BinnedSizeLimit)
 	{
 		// Allocate from pool.
 		FPoolTable* Table = MemSizeToPoolTable[Size];
