Generic Octree
- Reduced the amount of useless operations (building collapsed element list, calls to OctreeSemantics::SetElementId) when collapsing nodes by skipping updates when node is a leaf or if its number of elements is already the number of inclusive elements.
- Append collapsed elements directly into the collapse node list instead of using a temporary list to swap after the operation. That ensures that OctreeSemantics::SetElementId is called with information (Node & ElementIndex) that is already reflected in the octree.

diff --git a/Engine/Source/Runtime/Core/Public/Math/GenericOctree.inl b/Engine/Source/Runtime/Core/Public/Math/GenericOctree.inl
--- a/Engine/Source/Runtime/Core/Public/Math/GenericOctree.inl
+++ b/Engine/Source/Runtime/Core/Public/Math/GenericOctree.inl
@@ -225,22 +225,25 @@ void TOctree<ElementType,OctreeSemantics>::RemoveElement(FOctreeElementId Elemen
 	}
 
 	// Collapse the largest node that was pushed below the threshold for collapse by the removal.
-	if(CollapseNode)
+	if (CollapseNode && !CollapseNode->bIsLeaf && (CollapseNode->Elements.Num() < (int32)CollapseNode->InclusiveNumElements))
 	{
+		CollapseNode->Elements.Reserve(CollapseNode->InclusiveNumElements);
+
 		// Gather the elements contained in this node and its children.
-		TArray<ElementType,TInlineAllocator<OctreeSemantics::MaxElementsPerLeaf> > CollapsedChildElements;
-		CollapsedChildElements.Empty(CollapseNode->InclusiveNumElements);
 		for(TConstIterator<> ChildNodeIt(*CollapseNode,RootNodeContext);ChildNodeIt.HasPendingNodes();ChildNodeIt.Advance())
 		{
 			const FNode& ChildNode = ChildNodeIt.GetCurrentNode();
 
-			// Add the child's elements to the collapsed element list.
-			for(ElementConstIt ElementIt(ChildNode.Elements);ElementIt;++ElementIt)
+			if (&ChildNode != CollapseNode)
 			{
-				const int32 NewElementIndex = CollapsedChildElements.Add(*ElementIt);
+				// Child node will be collapsed so move the child's elements to the collapse node element list.
+				for (ElementType& Element : ChildNode.Elements)
+				{
+					const int32 NewElementIndex = CollapseNode->Elements.Add(MoveTemp(Element));
 
-				// Update the external element id for the element that's being collapsed.
-				OctreeSemantics::SetElementId(*ElementIt,FOctreeElementId(CollapseNode,NewElementIndex));
+					// Update the external element id for the element that's being collapsed.
+					OctreeSemantics::SetElementId(CollapseNode->Elements[NewElementIndex], FOctreeElementId(CollapseNode,NewElementIndex));
+				}
 			}
 
 			// Recursively visit all child nodes.
@@ -253,9 +256,6 @@ void TOctree<ElementType,OctreeSemantics>::RemoveElement(FOctreeElementId Elemen
 			}
 		}
 
-		// Replace the node's elements with the collapsed element list.
-		Exchange(CollapseNode->Elements,CollapsedChildElements);
-
 		// Mark the node as a leaf.
 		CollapseNode->bIsLeaf = true;
 
