Generic Octree
- Reduced the amount of useless operations (building collapsed element list, calls to OctreeSemantics::SetElementId) when collapsing nodes by skipping updates when node is a leaf or if its number of elements is already the number of inclusive elements.
- Append collapsed elements directly into the collapse node list instead of using a temporary list to swap after the operation. That ensures that OctreeSemantics::SetElementId is called with information (Node & ElementIndex) that is already reflected in the octree.

diff --git a/Engine/Source/Runtime/Core/Public/Math/GenericOctree.inl b/Engine/Source/Runtime/Core/Public/Math/GenericOctree.inl
--- a/Engine/Source/Runtime/Core/Public/Math/GenericOctree.inl
+++ b/Engine/Source/Runtime/Core/Public/Math/GenericOctree.inl
@@ -225,51 +225,54 @@ void TOctree<ElementType,OctreeSemantics>::RemoveElement(FOctreeElementId Elemen
 	}
 
 	// Collapse the largest node that was pushed below the threshold for collapse by the removal.
-	if (CollapseNode && !CollapseNode->bIsLeaf && (CollapseNode->Elements.Num() < (int32)CollapseNode->InclusiveNumElements))
+	if (CollapseNode && !CollapseNode->bIsLeaf)
 	{
-		CollapseNode->Elements.Reserve(CollapseNode->InclusiveNumElements);
-
-		// Gather the elements contained in this node and its children.
-		for(TConstIterator<> ChildNodeIt(*CollapseNode,RootNodeContext);ChildNodeIt.HasPendingNodes();ChildNodeIt.Advance())
+		if (CollapseNode->Elements.Num() < (int32)CollapseNode->InclusiveNumElements)
 		{
-			const FNode& ChildNode = ChildNodeIt.GetCurrentNode();
+			CollapseNode->Elements.Reserve(CollapseNode->InclusiveNumElements);
 
-			if (&ChildNode != CollapseNode)
+			// Gather the elements contained in this node and its children.
+			for (TConstIterator<> ChildNodeIt(*CollapseNode, RootNodeContext); ChildNodeIt.HasPendingNodes(); ChildNodeIt.Advance())
 			{
-				// Child node will be collapsed so move the child's elements to the collapse node element list.
-				for (ElementType& Element : ChildNode.Elements)
+				const FNode& ChildNode = ChildNodeIt.GetCurrentNode();
+
+				if (&ChildNode != CollapseNode)
 				{
-					const int32 NewElementIndex = CollapseNode->Elements.Add(MoveTemp(Element));
+					// Child node will be collapsed so move the child's elements to the collapse node element list.
+					for (ElementType& Element : ChildNode.Elements)
+					{
+						const int32 NewElementIndex = CollapseNode->Elements.Add(MoveTemp(Element));
+
+						// Update the external element id for the element that's being collapsed.
+						OctreeSemantics::SetElementId(CollapseNode->Elements[NewElementIndex], FOctreeElementId(CollapseNode, NewElementIndex));
+					}
+				}
 
-					// Update the external element id for the element that's being collapsed.
-					OctreeSemantics::SetElementId(CollapseNode->Elements[NewElementIndex], FOctreeElementId(CollapseNode,NewElementIndex));
+				// Recursively visit all child nodes.
+				FOREACH_OCTREE_CHILD_NODE(ChildRef)
+				{
+					if (ChildNode.HasChild(ChildRef))
+					{
+						ChildNodeIt.PushChild(ChildRef);
+					}
 				}
 			}
 
-			// Recursively visit all child nodes.
+			// Free the child nodes.
 			FOREACH_OCTREE_CHILD_NODE(ChildRef)
 			{
-				if(ChildNode.HasChild(ChildRef))
+				if (CollapseNode->Children[ChildRef.Index])
 				{
-					ChildNodeIt.PushChild(ChildRef);
+					SetOctreeMemoryUsage(this, TotalSizeBytes - sizeof(*CollapseNode->Children[ChildRef.Index]));
 				}
+
+				delete CollapseNode->Children[ChildRef.Index];
+				CollapseNode->Children[ChildRef.Index] = NULL;
 			}
 		}
 
 		// Mark the node as a leaf.
-		CollapseNode->bIsLeaf = true;
-
-		// Free the child nodes.
-		FOREACH_OCTREE_CHILD_NODE(ChildRef)
-		{
-			if (CollapseNode->Children[ChildRef.Index])
-			{
-				SetOctreeMemoryUsage(this, TotalSizeBytes - sizeof(*CollapseNode->Children[ChildRef.Index]));
-			}
-
-			delete CollapseNode->Children[ChildRef.Index];
-			CollapseNode->Children[ChildRef.Index] = NULL;
-		}
+		CollapseNode->bIsLeaf = true;		
 	}
 }
 
