Add API to support installing language chunks on XboxOne
- Also added a console command to install language chunks for testing purpose

diff --git a/Engine/Source/Runtime/Core/Private/GenericPlatform/GenericPlatformMisc.cpp b/Engine/Source/Runtime/Core/Private/GenericPlatform/GenericPlatformMisc.cpp
--- a/Engine/Source/Runtime/Core/Private/GenericPlatform/GenericPlatformMisc.cpp
+++ b/Engine/Source/Runtime/Core/Private/GenericPlatform/GenericPlatformMisc.cpp
@@ -1320,9 +1320,9 @@ bool FGenericPlatformMisc::RequestDeviceCheckToken(TFunction<void(const TArray<u
 	return false;
 }
 
-TArray<FChunkTagID> FGenericPlatformMisc::GetOnDemandChunkTagIDs()
+TArray<FCustomChunk> FGenericPlatformMisc::GetAllOnDemandChunks()
 {
-	return TArray<FChunkTagID>();
+	return TArray<FCustomChunk>();
 }
 
 FString FGenericPlatformMisc::LoadTextFileFromPlatformPackage(const FString& RelativePath)
diff --git a/Engine/Source/Runtime/Core/Public/GenericPlatform/GenericPlatformChunkInstall.h b/Engine/Source/Runtime/Core/Public/GenericPlatform/GenericPlatformChunkInstall.h
--- a/Engine/Source/Runtime/Core/Public/GenericPlatform/GenericPlatformChunkInstall.h
+++ b/Engine/Source/Runtime/Core/Public/GenericPlatform/GenericPlatformChunkInstall.h
@@ -78,13 +78,20 @@ DECLARE_DELEGATE_OneParam(FPlatformChunkInstallCompleteDelegate, uint32);
 DECLARE_DELEGATE_TwoParams(FPlatformChunkInstallDelegate, uint32, bool);
 DECLARE_MULTICAST_DELEGATE_TwoParams(FPlatformChunkInstallMultiDelegate, uint32, bool);
 
-struct FChunkTagID
+enum class ECustomChunkType : uint8
+{
+	OnDemandChunk,
+	LanguageChunk
+};
+
+struct FCustomChunk
 {
 	FString ChunkTag;
 	uint32	ChunkID;
+	ECustomChunkType ChunkType;
 
-	FChunkTagID(FString InTag, uint32 InID) :
-		ChunkTag(InTag), ChunkID(InID)
+	FCustomChunk(FString InTag, uint32 InID, ECustomChunkType InChunkType) :
+		ChunkTag(InTag), ChunkID(InID), ChunkType(InChunkType)
 	{}
 };
 
@@ -185,19 +192,19 @@ public:
 	* Check whether installation of chunks are pending
 	* @return				whether installation task has been kicked
 	*/
-	virtual bool IsChunkInstallationPending(const TArray<FChunkTagID> ChunkTagsID) = 0;
+	virtual bool IsChunkInstallationPending(const TArray<FCustomChunk>& ChunkTagsID) = 0;
 
 	/**
 	* Install chunks with Intelligent Delivery API
 	* @return				whether installation task has been kicked
 	*/
-	virtual bool InstallChunks(const TArray<FChunkTagID> ChunkTagsID) = 0;
+	virtual bool InstallChunks(const TArray<FCustomChunk>& ChunkTagsID) = 0;
 
 	/**
 	* Uninstall chunks with Intelligent Delivery API
 	* @return				whether uninstallation task has been kicked
 	*/
-	virtual bool UninstallChunks(const TArray<FChunkTagID> ChunkTagsID) = 0;
+	virtual bool UninstallChunks(const TArray<FCustomChunk>& ChunkTagsID) = 0;
 };
 
 PRAGMA_DISABLE_DEPRECATION_WARNINGS
@@ -281,17 +288,17 @@ public:
 		return false;
 	}
 
-	virtual bool IsChunkInstallationPending(const TArray<FChunkTagID> ChunkTags) override
+	virtual bool IsChunkInstallationPending(const TArray<FCustomChunk>& ChunkTags) override
 	{
 		return false;
 	}
 
-	virtual bool InstallChunks(const TArray<FChunkTagID> ChunkTagIDs) override
+	virtual bool InstallChunks(const TArray<FCustomChunk>& ChunkTagIDs) override
 	{
 		return false;
 	}
 
-	virtual bool UninstallChunks(const TArray<FChunkTagID> ChunkTagsID) override
+	virtual bool UninstallChunks(const TArray<FCustomChunk>& ChunkTagsID) override
 	{
 		return false;
 	}
diff --git a/Engine/Source/Runtime/Core/Public/GenericPlatform/GenericPlatformMisc.h b/Engine/Source/Runtime/Core/Public/GenericPlatform/GenericPlatformMisc.h
--- a/Engine/Source/Runtime/Core/Public/GenericPlatform/GenericPlatformMisc.h
+++ b/Engine/Source/Runtime/Core/Public/GenericPlatform/GenericPlatformMisc.h
@@ -14,7 +14,8 @@ class IPlatformInstallBundleManager;
 class IPlatformCompression;
 struct FGenericCrashContext;
 struct FGenericMemoryWarningContext;
-struct FChunkTagID;
+struct FCustomChunk;
+enum class ECustomChunkType : uint8;
 
 template <typename FuncType>
 class TFunction;
@@ -1277,7 +1278,9 @@ public:
 
 	static bool RequestDeviceCheckToken(TFunction<void(const TArray<uint8>&)> QuerySucceededFunc, TFunction<void(const FString&, const FString&)> QueryFailedFunc);
 
-	static TArray<FChunkTagID> GetOnDemandChunkTagIDs();
+	static TArray<FCustomChunk> GetAllOnDemandChunks();
+	static TArray<FCustomChunk> GetAllLanguageChunks();
+	static TArray<FCustomChunk> GetCustomChunksByType(ECustomChunkType DesiredChunkType);
 
 	/*
 	 * Loads a text file relative to the package root on platforms that distribute apps in package formats.
