Fix linking error introduced in CL 7092562

diff --git a/Engine/Source/Runtime/Core/Private/GenericPlatform/GenericPlatformMisc.cpp b/Engine/Source/Runtime/Core/Private/GenericPlatform/GenericPlatformMisc.cpp
--- a/Engine/Source/Runtime/Core/Private/GenericPlatform/GenericPlatformMisc.cpp
+++ b/Engine/Source/Runtime/Core/Private/GenericPlatform/GenericPlatformMisc.cpp
@@ -1325,6 +1325,23 @@ TArray<FCustomChunk> FGenericPlatformMisc::GetAllOnDemandChunks()
 	return TArray<FCustomChunk>();
 }
 
+TArray<FCustomChunk> FGenericPlatformMisc::GetAllLanguageChunks()
+{
+	return TArray<FCustomChunk>();
+}
+
+TArray<FCustomChunk> FGenericPlatformMisc::GetCustomChunksByType(ECustomChunkType DesiredChunkType)
+{
+	if (DesiredChunkType == ECustomChunkType::OnDemandChunk)
+	{
+		return GetAllOnDemandChunks();
+	}
+	else
+	{
+		return GetAllLanguageChunks();
+	}
+}
+
 FString FGenericPlatformMisc::LoadTextFileFromPlatformPackage(const FString& RelativePath)
 {
 	FString Path = RootDir() / RelativePath;
