[FramePro] Fix thread leak and tls crash.

diff --git a/Engine/Source/Runtime/Core/Private/FramePro/FramePro.cpp b/Engine/Source/Runtime/Core/Private/FramePro/FramePro.cpp
--- a/Engine/Source/Runtime/Core/Private/FramePro/FramePro.cpp
+++ b/Engine/Source/Runtime/Core/Private/FramePro/FramePro.cpp
@@ -86,6 +86,18 @@ namespace FramePro
 		return p_framepro_tls ? p_framepro_tls : CreateFrameProTLS();
 	}
 
+	//------------------------------------------------------------------------
+	FRAMEPRO_FORCE_INLINE FrameProTLS* TryGetFrameProTLS()
+	{
+		return (FrameProTLS*)Platform::GetTLSValue(GetFrameProTLSSlot());
+	}
+
+	//------------------------------------------------------------------------
+	FRAMEPRO_FORCE_INLINE void ClearFrameProTLS()
+	{
+		Platform::SetTLSValue(GetFrameProTLSSlot(), NULL);
+	}
+
 	//------------------------------------------------------------------------
 	void DebugWrite(const char* p_str, ...);
 
@@ -3055,6 +3067,8 @@ namespace FramePro
 	public:
 		Thread();
 
+		~Thread();
+
 		void CreateThread(ThreadMain p_thread_main, void* p_param, Allocator* p_allocator);
 
 		bool IsAlive() const { return m_Alive; }
@@ -3073,6 +3087,8 @@ namespace FramePro
 	private:
 		static const int m_OSThreadMaxSize = 16;
 		char m_OSThread[m_OSThreadMaxSize];
+		
+		bool m_Created;
 
 		bool m_Alive;
 
@@ -4318,9 +4334,16 @@ void FramePro::AddGlobalHiResTimer(GlobalHiResTimer* p_timer)
 //------------------------------------------------------------------------
 void FramePro::CleanupThread()
 {
-	GetFrameProTLS()->FlushCurrentSendBuffer();
+	FrameProTLS* p_framepro_tls = TryGetFrameProTLS();
 
-	GetFrameProTLS()->Shutdown();		// will get cleaned up the next time the buffers are sent on the send thread
+	if (p_framepro_tls)
+	{
+		p_framepro_tls->FlushCurrentSendBuffer();
+
+		p_framepro_tls->Shutdown();		// will get cleaned up the next time the buffers are sent on the send thread
+
+		ClearFrameProTLS();
+	}
 }
 
 //------------------------------------------------------------------------
@@ -4874,7 +4897,14 @@ namespace FramePro
 	int FrameProSession::StaticSendThreadMain(void* p_arg)
 	{
 		FrameProSession* p_this = (FrameProSession*)p_arg;
-		return p_this->SendThreadMain();
+		int ret = p_this->SendThreadMain();
+
+		#if !FRAMEPRO_PLATFORM_UE4
+			DestroyFrameProTLS(GetFrameProTLS());
+			ClearFrameProTLS();
+		#endif
+
+		return ret;
 	}
 
 	//------------------------------------------------------------------------
@@ -4920,7 +4950,14 @@ namespace FramePro
 	int FrameProSession::StaticConnectThreadMain(void* p_arg)
 	{
 		FrameProSession* p_this = (FrameProSession*)p_arg;
-		return p_this->ConnectThreadMain();
+		int ret = p_this->ConnectThreadMain();
+
+		#if !FRAMEPRO_PLATFORM_UE4
+			DestroyFrameProTLS(GetFrameProTLS());
+			ClearFrameProTLS();
+		#endif
+
+		return ret;
 	}
 #endif
 
@@ -4956,7 +4993,10 @@ namespace FramePro
 		FrameProSession* p_this = (FrameProSession*)p_arg;
 		int ret = p_this->ReceiveThreadMain();
 
-		DestroyFrameProTLS(GetFrameProTLS());
+		#if !FRAMEPRO_PLATFORM_UE4
+			DestroyFrameProTLS(GetFrameProTLS());
+			ClearFrameProTLS();
+		#endif
 
 		return ret;
 	}
@@ -7834,18 +7874,31 @@ namespace FramePro
 {
 	//------------------------------------------------------------------------
 	Thread::Thread()
-	:	m_Alive(false),
+	:	m_Created(false),
+		m_Alive(false),
 		m_ThreadTerminatedEvent(false, false)
 	{
 	}
 
+	//------------------------------------------------------------------------
+	Thread::~Thread()
+	{
+		if(m_Created)
+			Platform::DestroyThread(m_OSThread);
+	}
+
 	//------------------------------------------------------------------------
 	void Thread::CreateThread(FramePro::ThreadMain p_thread_main, void* p_param, Allocator* p_allocator)
 	{
+		if(m_Created)
+			Platform::DestroyThread(m_OSThread);
+
 		mp_ThreadMain = p_thread_main;
 		mp_Param = p_param;
 
 		Platform::CreateThread(m_OSThread, sizeof(m_OSThread), ThreadMain, this, p_allocator);
+
+		m_Created = true;
 	}
 
 	//------------------------------------------------------------------------
@@ -8490,7 +8543,7 @@ namespace FramePro
 			//------------------------------------------------------------------------
 			void HandleSocketError()
 			{
-				#if FRAMEPRO_SOCKETS_ENABLED
+				#if FRAMEPRO_SOCKETS_ENABLED && !FRAMEPRO_PLATFORM_XBOXONE
 					if (WSAGetLastError() == WSAEADDRINUSE)
 					{
 						Platform::DebugWrite("FramePro: Network connection conflict. Please make sure that other FramePro enabled applications are shut down, or change the port in the the FramePro lib and FramePro settings.\n");
@@ -8793,6 +8846,12 @@ namespace FramePro
 				GetOSThread(p_os_thread_mem) = ::CreateThread(NULL, 0, PlatformThreadMain, p_thread_context, 0, NULL);
 			}
 
+			//------------------------------------------------------------------------
+			void DestroyThread(void* p_os_thread_mem)
+			{
+				GetOSThread(p_os_thread_mem) = 0;
+			}
+
 			//------------------------------------------------------------------------
 			void SetThreadPriority(void* p_os_thread_mem, int priority)
 			{
@@ -9314,6 +9373,12 @@ namespace FramePro
 				pthread_create(&GetOSThread(p_os_thread_mem), NULL, PlatformThreadMain, p_thread_context);
 			}
 
+			//------------------------------------------------------------------------
+			void DestroyThread(void*)
+			{
+				// do nothing
+			}
+
 			//------------------------------------------------------------------------
 			void SetThreadPriority(void*, int)
 			{
@@ -9817,6 +9882,12 @@ namespace FramePro
 				p_allocator);
 		}
 
+		//------------------------------------------------------------------------
+		void Platform::DestroyThread(void* p_os_thread_mem)
+		{
+			GenericPlatform::DestroyThread(p_os_thread_mem);
+		}
+
 		//------------------------------------------------------------------------
 		void Platform::SetThreadPriority(void* p_os_thread_mem, int priority)
 		{
@@ -10244,6 +10315,12 @@ namespace FramePro
 				p_allocator);
 		}
 
+		//------------------------------------------------------------------------
+		void Platform::DestroyThread(void* p_os_thread_mem)
+		{
+			GenericPlatform::DestroyThread(p_os_thread_mem);
+		}
+
 		//------------------------------------------------------------------------
 		void Platform::SetThreadPriority(void* p_os_thread_mem, int priority)
 		{
@@ -10642,6 +10719,12 @@ namespace FramePro
 				p_allocator);
 		}
 
+		//------------------------------------------------------------------------
+		void Platform::DestroyThread(void* p_os_thread_mem)
+		{
+			GenericPlatform::DestroyThread(p_os_thread_mem);
+		}
+
 		//------------------------------------------------------------------------
 		void Platform::SetThreadPriority(void* p_os_thread_mem, int priority)
 		{
@@ -11046,6 +11129,12 @@ namespace FramePro
 				p_allocator);
 		}
 
+		//------------------------------------------------------------------------
+		void Platform::DestroyThread(void* p_os_thread_mem)
+		{
+			GenericPlatform::DestroyThread(p_os_thread_mem);
+		}
+
 		//------------------------------------------------------------------------
 		void Platform::SetThreadPriority(void* p_os_thread_mem, int priority)
 		{
diff --git a/Engine/Source/Runtime/Core/Private/FramePro/FrameProUE4.cpp b/Engine/Source/Runtime/Core/Private/FramePro/FrameProUE4.cpp
--- a/Engine/Source/Runtime/Core/Private/FramePro/FrameProUE4.cpp
+++ b/Engine/Source/Runtime/Core/Private/FramePro/FrameProUE4.cpp
@@ -73,14 +73,6 @@
 		#include <cpuid.h>
 	#endif
 
-	//@EPIC BEGIN: workaround for -nothreading
-	#if PLATFORM_WINDOWS
-		#include "Windows/WindowsEvent.h"
-	#elif PLATFORM_USE_PTHREADS
-		#include <pthread.h>
-		#include "HAL/PThreadEvent.h"
-	#endif
-
 	//------------------------------------------------------------------------
 	namespace FramePro
 	{
@@ -437,26 +429,7 @@
 		//------------------------------------------------------------------------
 		void Platform::CreateEventX(void* p_os_event_mem, int os_event_mem_size, bool initial_state, bool auto_reset)
 		{
-//@EPIC BEGIN: workaround for -nothreading
-			FEvent* p_event;
-			if (FPlatformProcess::SupportsMultithreading())
-			{
-				p_event = FPlatformProcess::GetSynchEventFromPool(!auto_reset);
-			}
-			else
-			{				
-#if PLATFORM_WINDOWS
-				p_event = new FEventWin();
-				p_event->Create(!auto_reset);
-#elif PLATFORM_USE_PTHREADS
-				p_event = new FPThreadEvent();
-				p_event->Create(!auto_reset);
-#else
-				checkf(false,"unsupported platform for -nothreading");
-				return;
-#endif
-			}
-//@EPIC END
+			FEvent* p_event = FPlatformProcess::GetSynchEventFromPool(!auto_reset);
 			GetOSEvent(p_os_event_mem) = p_event;
 
 			if(initial_state)
@@ -467,16 +440,7 @@
 		void Platform::DestroyEvent(void* p_os_event_mem)
 		{
 			FEvent* p_event = GetOSEvent(p_os_event_mem);
-//@EPIC BEGIN: workaround for -nothreading
-			if (FPlatformProcess::SupportsMultithreading())
-			{
-				FPlatformProcess::ReturnSynchEventToPool(p_event);
-			}
-			else
-			{
-				delete p_event;
-			}
-//@EPIC END
+			FPlatformProcess::ReturnSynchEventToPool(p_event);
 			GetOSEvent(p_os_event_mem) = nullptr;
 		}
 
@@ -580,24 +544,7 @@
 			:	mp_ThreadMain(p_thread_main),
 				mp_Context(p_context)
 			{
-//@EPIC BEGIN: workaround for -nothreading
-				if (FPlatformProcess::SupportsMultithreading())
-				{
-					m_Runnable = TUniquePtr<FRunnableThread>(FRunnableThread::Create(this, TEXT("FramePro")));
-				}
-				else
-				{
-#if PLATFORM_WINDOWS
-					CreateThread( NULL, 0, (LPTHREAD_START_ROUTINE)mp_ThreadMain, mp_Context, 0, NULL );
-#elif PLATFORM_USE_PTHREADS
-					typedef void *(*PthreadEntryPoint)(void *arg);
-					pthread_t thread_id;
-					pthread_create( &thread_id, nullptr, (PthreadEntryPoint)mp_ThreadMain, mp_Context );
-#else
-					checkf(false,"unsupported platform for -nothreading");
-#endif
-				}
-//@EPIC END
+				m_Runnable = TUniquePtr<FRunnableThread>(FRunnableThread::Create(this, TEXT("FramePro")));
 			}
 
 			uint32 Run() override
@@ -636,6 +583,12 @@
 			GetOSThread(p_os_thread_mem) = new UE4Thread(p_thread_main, p_context);
 		}
 
+		//------------------------------------------------------------------------
+		void Platform::DestroyThread(void* p_os_thread_mem)
+		{
+			delete GetOSThread(p_os_thread_mem);
+		}
+
 		//------------------------------------------------------------------------
 		void Platform::SetThreadPriority(void* p_os_thread_mem, int priority)
 		{
diff --git a/Engine/Source/Runtime/Core/Public/FramePro/FramePro.h b/Engine/Source/Runtime/Core/Public/FramePro/FramePro.h
--- a/Engine/Source/Runtime/Core/Public/FramePro/FramePro.h
+++ b/Engine/Source/Runtime/Core/Public/FramePro/FramePro.h
@@ -29,7 +29,7 @@
 //------------------------------------------------------------------------
 /*
 	FramePro
-	Version:	1.5.15.0
+	Version:	1.5.20.0
 */
 //------------------------------------------------------------------------
 #ifndef FRAMEPRO_H_INCLUDED
@@ -1482,6 +1482,8 @@ namespace FramePro
 			void* p_context,
 			Allocator* p_allocator);
 
+		void DestroyThread(void* p_os_thread_mem);
+
 		void SetThreadPriority(void* p_os_thread_mem, int priority);
 
 		void SetThreadAffinity(void* p_os_thread_mem, int affinity);
