re-enable the IOS OS small block allocations after fixing several issues

diff --git a/Engine/Source/Runtime/Core/Private/HAL/MallocBinned.cpp b/Engine/Source/Runtime/Core/Private/HAL/MallocBinned.cpp
--- a/Engine/Source/Runtime/Core/Private/HAL/MallocBinned.cpp
+++ b/Engine/Source/Runtime/Core/Private/HAL/MallocBinned.cpp
@@ -136,16 +136,20 @@ struct FMallocBinned::PoolHashBucket
 	}
 };
 
-#define USE_OS_SMALL_BLOCK_ALLOC 0
+#define USE_OS_SMALL_BLOCK_ALLOC 1
 
 #if PLATFORM_IOS
 #define PLAT_PAGE_SIZE_LIMIT 16384
 #define PLAT_BINNED_ALLOC_POOLSIZE 16384
 #define PLAT_SMALL_BLOCK_POOL_SIZE 224
+#define PLAT_SMALL_BLOCK_START 0x280000000
+#define PLAT_SMALL_BLOCK_END 0x2a0000000
 #else
 #define PLAT_PAGE_SIZE_LIMIT 65536
 #define PLAT_BINNED_ALLOC_POOLSIZE 65536
 #define PLAT_SMALL_BLOCK_POOL_SIZE 0
+#define PLAT_SMALL_BLOCK_START 0
+#define PLAT_SMALL_BLOCK_END 0
 #endif
 
 struct FMallocBinned::Private
@@ -734,8 +738,8 @@ struct FMallocBinned::Private
 };
 
 #if USE_OS_SMALL_BLOCK_ALLOC
-uint64 FMallocBinned::Private::SmallBlockStartPtr = 0xffffffffffffffff;
-uint64 FMallocBinned::Private::SmallBlockEndPtr = 0x0000000000000000;
+uint64 FMallocBinned::Private::SmallBlockStartPtr = PLAT_SMALL_BLOCK_START;
+uint64 FMallocBinned::Private::SmallBlockEndPtr = PLAT_SMALL_BLOCK_END;
 #endif
 
 void FMallocBinned::GetAllocatorStats( FGenericMemoryStats& out_Stats )
@@ -948,11 +952,12 @@ void* FMallocBinned::Malloc(SIZE_T Size, uint32 Alignment)
 #if USE_OS_SMALL_BLOCK_ALLOC
 	if (Size <= Private::SMALL_BLOCK_POOL_SIZE)
 	{
+		bUsePools = false;
 		UPTRINT AlignedSize = Align(Size, Alignment);
 		SIZE_T ActualPoolSize; //TODO: use this to reduce waste?
 		Free = (FFreeMem*)Private::SmallOSAlloc(*this, AlignedSize, ActualPoolSize);
 		
-		if ((uint64)Free < 0x280000000)
+		if ((uint64)Free < Private::SmallBlockStartPtr || (uint64)Free >= Private::SmallBlockEndPtr)
 		{
 			// if this happens then we need to fall in to the bins below
 //			FPlatformMisc::LowLevelOutputDebugStringf(TEXT("Overflowed!!! Aligned Size:%d"), AlignedSize);
@@ -960,18 +965,6 @@ void* FMallocBinned::Malloc(SIZE_T Size, uint32 Alignment)
 			Free = nullptr;
 			bUsePools = true;
 		}
-		else
-		{
-			if ((uint64)Free < Private::SmallBlockStartPtr)
-			{
-				Private::SmallBlockStartPtr = (uint64)Free;
-			}
-			if ((uint64)Free > Private::SmallBlockEndPtr)
-			{
-				Private::SmallBlockEndPtr = (uint64)Free;
-			}
-			bUsePools = false;
-		}
 	}
 #endif
 	if (bUsePools)
@@ -1084,20 +1077,18 @@ void* FMallocBinned::Realloc( void* Ptr, SIZE_T NewSize, uint32 Alignment )
 	{
 #if USE_OS_SMALL_BLOCK_ALLOC
 		// special case for really small blocks
-		if ((uint64)Ptr >= Private::SmallBlockStartPtr && (uint64)Ptr <= Private::SmallBlockEndPtr)
+		if ((uint64)Ptr >= Private::SmallBlockStartPtr && (uint64)Ptr < Private::SmallBlockEndPtr)
 		{
-			// Grow or shrink.
-			NewPtr = Malloc(NewSizeUnmodified, Alignment);
-			SIZE_T CopySize = FMath::Min<SIZE_T>(NewSizeUnmodified, Private::SMALL_BLOCK_POOL_SIZE);
-#if UE_BUILD_TEST
-			if ((uint64)Ptr + CopySize >= 0x2a0000000)
+			NewPtr = ::realloc(Ptr, NewSize);
+
+			if ((uint64)NewPtr < Private::SmallBlockStartPtr || (uint64)NewPtr >= Private::SmallBlockEndPtr)
 			{
-				UE_LOG(LogMemory, Warning, TEXT("Going to memory copy past the end of the small block!!"));
+				// if this happens then we need to Malloc and copy
+				Ptr = NewPtr;
+				NewPtr = Malloc(NewSizeUnmodified, Alignment);
+				FMemory::Memcpy(NewPtr, Ptr, NewSize);
+				::free( Ptr );
 			}
-#endif
-			CopySize = FMath::Min<SIZE_T>(CopySize, (0x29fffffff-(uint64)Ptr)); // guard against going off the end of the VM area for the small block when doing the memcpy
-			FMemory::Memcpy(NewPtr, Ptr, CopySize); // this is bad need to figure out the old size :(
-			Free( Ptr );
 		}
 		else
 #endif
@@ -1185,7 +1176,7 @@ bool FMallocBinned::GetAllocationSize(void *Original, SIZE_T &SizeOut)
 	if (Pool == NULL)
 	{
 #if USE_OS_SMALL_BLOCK_ALLOC
-		if ((uint64)Original >= Private::SmallBlockStartPtr && (uint64)Original <= Private::SmallBlockEndPtr)
+		if ((uint64)Original >= Private::SmallBlockStartPtr && (uint64)Original < Private::SmallBlockEndPtr)
 		{
 			SizeOut = Private::SMALL_BLOCK_POOL_SIZE;
 			return true;
