changed #define to just be IOS platform

diff --git a/Engine/Source/Runtime/Core/Private/HAL/MallocBinned.cpp b/Engine/Source/Runtime/Core/Private/HAL/MallocBinned.cpp
--- a/Engine/Source/Runtime/Core/Private/HAL/MallocBinned.cpp
+++ b/Engine/Source/Runtime/Core/Private/HAL/MallocBinned.cpp
@@ -136,7 +136,7 @@ struct FMallocBinned::PoolHashBucket
 	}
 };
 
-#define USE_OS_SMALL_BLOCK_ALLOC 1
+#define USE_OS_SMALL_BLOCK_ALLOC PLATFORM_IOS
 
 #if PLATFORM_IOS
 #define PLAT_PAGE_SIZE_LIMIT 16384
