Have WakeGameThread let the engine do 2 full ticks to ensure that game systems get a chance to tick

diff --git a/Engine/Source/Runtime/Core/Private/Misc/EmbeddedCommunication.cpp b/Engine/Source/Runtime/Core/Private/Misc/EmbeddedCommunication.cpp
--- a/Engine/Source/Runtime/Core/Private/Misc/EmbeddedCommunication.cpp
+++ b/Engine/Source/Runtime/Core/Private/Misc/EmbeddedCommunication.cpp
@@ -43,7 +43,7 @@ void* FEmbeddedDelegates::GetNamedObject(const FString& Name)
 #if BUILD_EMBEDDED_APP
 
 static FEvent* GSleepEvent = nullptr;
-static TAtomic<bool> GTickAnotherFrame;
+static TAtomic<int32> GTickWithoutSleepCount(0);
 static FCriticalSection GEmbeddedLock;
 //static TArray<TFunction<void()>> GEmbeddedQueues[5];
 static TQueue<TFunction<void()>> GEmbeddedQueues[5];
@@ -71,7 +71,7 @@ void FEmbeddedCommunication::Init()
 #if BUILD_EMBEDDED_APP
  	FScopeLock Lock(&GEmbeddedLock);
  	GSleepEvent = FPlatformProcess::GetSynchEventFromPool(false);
-	GTickAnotherFrame = false;
+	GTickWithoutSleepCount = 0;
 
 	FTicker::GetCoreTicker().AddTicker(FTickerDelegate::CreateStatic(&FEmbeddedCommunication::TickGameThread));
 #endif
@@ -290,8 +290,10 @@ void FEmbeddedCommunication::RunOnGameThread(int Priority, TFunction<void()> Lam
 void FEmbeddedCommunication::WakeGameThread()
 {
 #if BUILD_EMBEDDED_APP
-	// If the game thread is already running, this will make it loop again.
-	GTickAnotherFrame = true;
+	// Allow 2 ticks without a sleep
+	// Our sleep happens in the core ticker's tick, and that order gets reversed every tick,
+	// so the caller isn't guaranteed to get a tick before our next sleep
+	GTickWithoutSleepCount = 2;
 	// wake up the game thread!
 	if (GSleepEvent)
 	{
@@ -326,14 +328,17 @@ bool FEmbeddedCommunication::TickGameThread(float DeltaTime)
 	// sleep if nothing is going on
 	else if (GRenderingWakeMap.Num() == 0
 		&& GTickWakeMap.Num() == 0
-		&& !GTickAnotherFrame) // Don't sleep if we have been asked to tick another frame
+		&& GTickWithoutSleepCount <= 0)
  	{
- 		// wake up every 5 seconds even if nothing to do
+		// wake up every 5 seconds even if nothing to do
 		UE_LOG(LogInit, VeryVerbose, TEXT("FEmbeddedCommunication Sleeping GameThread..."));
- 		bool bWasTriggered = GSleepEvent->Wait(5000);
+		bool bWasTriggered = GSleepEvent->Wait(5000);
 		UE_LOG(LogInit, VeryVerbose, TEXT("FEmbeddedCommunication Woke up. Reason=[%s]"), bWasTriggered ? TEXT("Triggered") : TEXT("TimedOut"));
  	}
-	GTickAnotherFrame = false;
+	if (GTickWithoutSleepCount > 0)
+	{
+		--GTickWithoutSleepCount;
+	}
 #endif
 
 	return true;
