Updated FramePro drop with disconnection fixes + restoring workaround for -nothreading

diff --git a/Engine/Source/Runtime/Core/Private/FramePro/FramePro.cpp b/Engine/Source/Runtime/Core/Private/FramePro/FramePro.cpp
--- a/Engine/Source/Runtime/Core/Private/FramePro/FramePro.cpp
+++ b/Engine/Source/Runtime/Core/Private/FramePro/FramePro.cpp
@@ -3323,9 +3323,9 @@ namespace FramePro
 
 		int SendThreadMain();
 
-		void HandleDisconnect(bool wait_for_threads_to_exit=true);
+		void Disconect(bool wait_for_threads_to_exit=true);
 
-		void HandleDisconnect_NoLock(bool wait_for_threads_to_exit=true);
+		void Disconect_NoLock(bool wait_for_threads_to_exit=true);
 
 		void SendRecordedDataAndDisconnect();
 
@@ -4805,7 +4805,7 @@ namespace FramePro
 	//------------------------------------------------------------------------
 	FrameProSession::~FrameProSession()
 	{
-		HandleDisconnect(false);
+		Disconect(false);
 
 		m_NamedThreads.Clear();
 
@@ -5176,7 +5176,7 @@ namespace FramePro
 #if FRAMEPRO_SOCKETS_ENABLED
 	int FrameProSession::OnReceiveThreadExit()
 	{
-		HandleDisconnect();
+		Disconect();
 		return 0;
 	}
 #endif
@@ -5544,20 +5544,20 @@ namespace FramePro
 		read_file.Close();
 		mp_Allocator->Free(p_read_buffer);
 
-		HandleDisconnect_NoLock();
+		Disconect_NoLock();
 #endif
 	}
 
 	//------------------------------------------------------------------------
-	void FrameProSession::HandleDisconnect(bool wait_for_threads_to_exit)
+	void FrameProSession::Disconect(bool wait_for_threads_to_exit)
 	{
 		CriticalSectionScope lock(m_CriticalSection);
 		if(g_Connected)
-			HandleDisconnect_NoLock(wait_for_threads_to_exit);
+			Disconect_NoLock(wait_for_threads_to_exit);
 	}
 
 	//------------------------------------------------------------------------
-	void FrameProSession::HandleDisconnect_NoLock(bool wait_for_threads_to_exit)
+	void FrameProSession::Disconect_NoLock(bool wait_for_threads_to_exit)
 	{
 		Platform::StopRecordingContextSitches(mp_ContextSwitchRecorder);
 
@@ -5982,6 +5982,10 @@ namespace FramePro
 	//------------------------------------------------------------------------
 	void FrameProSession::StartRecording(const char* p_filename, bool context_switches, bool callstacks, int64 max_file_size)
 	{
+		Disconect(true);
+
+		CreateDefaultAllocator();
+
 		CriticalSectionScope lock(m_CriticalSection);
 
 		if(m_RecordingFile.IsOpened())
@@ -6020,6 +6024,8 @@ namespace FramePro
 	//------------------------------------------------------------------------
 	void FrameProSession::StartRecording(const wchar_t* p_filename, bool context_switches, bool callstacks, int64 max_file_size)
 	{
+		Disconect(true);
+
 		CreateDefaultAllocator();
 
 		CriticalSectionScope lock(m_CriticalSection);
@@ -6067,7 +6073,7 @@ namespace FramePro
 			#if FRAMEPRO_SOCKETS_ENABLED
 				OpenListenSocket();					// start the listening socket again so that we can accept new connections
 			#endif
-			HandleDisconnect_NoLock();
+				Disconect_NoLock();
 		}
 	}
 
diff --git a/Engine/Source/Runtime/Core/Private/FramePro/FrameProUE4.cpp b/Engine/Source/Runtime/Core/Private/FramePro/FrameProUE4.cpp
--- a/Engine/Source/Runtime/Core/Private/FramePro/FrameProUE4.cpp
+++ b/Engine/Source/Runtime/Core/Private/FramePro/FrameProUE4.cpp
@@ -73,6 +73,15 @@
 		#include <cpuid.h>
 	#endif
 
+	//@EPIC BEGIN: workaround for -nothreading
+	#if PLATFORM_WINDOWS
+		#include "Windows/WindowsEvent.h"
+	#elif PLATFORM_USE_PTHREADS
+		#include <pthread.h>
+		#include "HAL/PThreadEvent.h"
+	#endif
+	//@EPIC END
+
 	//------------------------------------------------------------------------
 	namespace FramePro
 	{
@@ -429,7 +438,26 @@
 		//------------------------------------------------------------------------
 		void Platform::CreateEventX(void* p_os_event_mem, int os_event_mem_size, bool initial_state, bool auto_reset)
 		{
-			FEvent* p_event = FPlatformProcess::GetSynchEventFromPool(!auto_reset);
+//@EPIC BEGIN: workaround for -nothreading
+			FEvent* p_event;
+			if (FPlatformProcess::SupportsMultithreading())
+			{
+				p_event = FPlatformProcess::GetSynchEventFromPool(!auto_reset);
+			}
+			else
+			{				
+#if PLATFORM_WINDOWS
+				p_event = new FEventWin();
+				p_event->Create(!auto_reset);
+#elif PLATFORM_USE_PTHREADS
+				p_event = new FPThreadEvent();
+				p_event->Create(!auto_reset);
+#else
+				checkf(false,"unsupported platform for -nothreading");
+				return;
+#endif
+			}
+//@EPIC END
 			GetOSEvent(p_os_event_mem) = p_event;
 
 			if(initial_state)
@@ -440,7 +468,16 @@
 		void Platform::DestroyEvent(void* p_os_event_mem)
 		{
 			FEvent* p_event = GetOSEvent(p_os_event_mem);
-			FPlatformProcess::ReturnSynchEventToPool(p_event);
+//@EPIC BEGIN: workaround for -nothreading
+			if (FPlatformProcess::SupportsMultithreading())
+			{
+				FPlatformProcess::ReturnSynchEventToPool(p_event);
+			}
+			else
+			{
+				delete p_event;
+			}
+//@EPIC END
 			GetOSEvent(p_os_event_mem) = nullptr;
 		}
 
@@ -544,7 +581,24 @@
 			:	mp_ThreadMain(p_thread_main),
 				mp_Context(p_context)
 			{
-				m_Runnable = TUniquePtr<FRunnableThread>(FRunnableThread::Create(this, TEXT("FramePro")));
+//@EPIC BEGIN: workaround for -nothreading
+				if (FPlatformProcess::SupportsMultithreading())
+				{
+					m_Runnable = TUniquePtr<FRunnableThread>(FRunnableThread::Create(this, TEXT("FramePro")));
+				}
+				else
+				{
+#if PLATFORM_WINDOWS
+					CreateThread( NULL, 0, (LPTHREAD_START_ROUTINE)mp_ThreadMain, mp_Context, 0, NULL );
+#elif PLATFORM_USE_PTHREADS
+					typedef void *(*PthreadEntryPoint)(void *arg);
+					pthread_t thread_id;
+					pthread_create( &thread_id, nullptr, (PthreadEntryPoint)mp_ThreadMain, mp_Context );
+#else
+					checkf(false,"unsupported platform for -nothreading");
+#endif
+				}
+//@EPIC END
 			}
 
 			uint32 Run() override
