Implement Linux FPlatformProcess::OpenProcess()

CL 7109034 added code which calls FPlatformProcess::OpenProcess() in EngineSessionManager.cpp checking previous crash data, and it's crashing on Linux since this method isn't implemented for FUnixPlatformProcess.

Handle returned from OpenProcess() is limited compared to CreateProcess() since all we can really do is check if the process id is running.

diff --git a/Engine/Source/Runtime/Core/Private/Unix/UnixPlatformProcess.cpp b/Engine/Source/Runtime/Core/Private/Unix/UnixPlatformProcess.cpp
--- a/Engine/Source/Runtime/Core/Private/Unix/UnixPlatformProcess.cpp
+++ b/Engine/Source/Runtime/Core/Private/Unix/UnixPlatformProcess.cpp
@@ -1013,6 +1013,28 @@ FProcHandle FUnixPlatformProcess::CreateProc(const TCHAR* URL, const TCHAR* Parm
 	return FProcHandle(new FProcState(ChildPid, bLaunchDetached));
 }
 
+/*
+ * Return a limited use FProcHandle from a PID. Currently can only use w/ IsProcRunning().
+ *
+ * WARNING (from Arciel): PIDs can and will be reused. We have had that issue
+ * before: the editor was tracking ShaderCompileWorker by their PIDs, and a
+ * long-running process (something from PS4 SDK) got a reused SCW PID,
+ * resulting in compilation never ending.
+ */
+FProcHandle FUnixPlatformProcess::OpenProcess(uint32 ProcessID)
+{
+	pid_t Pid = static_cast< pid_t >(ProcessID);
+
+	// check if actually running
+	int KillResult = kill(Pid, 0);	// no actual signal is sent
+	check(KillResult != -1 || errno != EINVAL);
+
+	// errno == EPERM: don't have permissions to send signal
+	// errno == ESRCH: proc doesn't exist
+	bool bIsRunning = (KillResult == 0);
+	return FProcHandle(bIsRunning ? Pid : -1);
+}
+
 /** Initialization constructor. */
 FProcState::FProcState(pid_t InProcessId, bool bInFireAndForget)
 	:	ProcessId(InProcessId)
@@ -1165,8 +1187,25 @@ void FProcState::Wait()
 
 bool FUnixPlatformProcess::IsProcRunning( FProcHandle & ProcessHandle )
 {
+	bool bIsRunning = false;
 	FProcState * ProcInfo = ProcessHandle.GetProcessInfo();
-	return ProcInfo ? ProcInfo->IsRunning() : false;
+
+	if (ProcInfo)
+	{
+		bIsRunning = ProcInfo->IsRunning();
+	}
+	else if (ProcessHandle.Get() != -1)
+	{
+		// Process opened with OpenProcess() call (we only have pid)
+		int KillResult = kill(ProcessHandle.Get(), 0);	// no actual signal is sent
+		check(KillResult != -1 || errno != EINVAL);
+
+		// errno == EPERM: don't have permissions to send signal
+		// errno == ESRCH: proc doesn't exist
+		bIsRunning = (KillResult == 0);
+	}
+
+	return bIsRunning;
 }
 
 void FUnixPlatformProcess::WaitForProc( FProcHandle & ProcessHandle )
@@ -1176,6 +1215,10 @@ void FUnixPlatformProcess::WaitForProc( FProcHandle & ProcessHandle )
 	{
 		ProcInfo->Wait();
 	}
+	else if (ProcessHandle.Get() != -1)
+	{
+		STUBBED("FUnixPlatformProcess::WaitForProc() : Waiting on OpenProcess() handle not implemented yet");
+	}
 }
 
 void FUnixPlatformProcess::CloseProc(FProcHandle & ProcessHandle)
@@ -1201,6 +1244,10 @@ void FUnixPlatformProcess::TerminateProc( FProcHandle & ProcessHandle, bool Kill
 		int KillResult = kill(ProcInfo->GetProcessId(), SIGTERM);	// graceful
 		check(KillResult != -1 || errno != EINVAL);
 	}
+	else if (ProcessHandle.Get() != -1)
+	{
+		STUBBED("FUnixPlatformProcess::TerminateProc() : Terminating OpenProcess() handle not implemented");
+	}
 }
 
 /*
@@ -1480,7 +1527,16 @@ bool FUnixPlatformProcess::GetProcReturnCode(FProcHandle& ProcHandle, int32* Ret
 	}
 
 	FProcState * ProcInfo = ProcHandle.GetProcessInfo();
-	return ProcInfo ? ProcInfo->GetReturnCode(ReturnCode) : false;
+	if (ProcInfo)
+	{
+		return ProcInfo->GetReturnCode(ReturnCode);
+	}
+	else if (ProcHandle.Get() != -1)
+	{
+		STUBBED("FUnixPlatformProcess::GetProcReturnCode() : Return code of OpenProcess() handle not implemented yet");
+	}
+
+	return false;
 }
 
 bool FUnixPlatformProcess::Daemonize()
diff --git a/Engine/Source/Runtime/Core/Public/GenericPlatform/GenericPlatformProcess.h b/Engine/Source/Runtime/Core/Public/GenericPlatform/GenericPlatformProcess.h
--- a/Engine/Source/Runtime/Core/Public/GenericPlatform/GenericPlatformProcess.h
+++ b/Engine/Source/Runtime/Core/Public/GenericPlatform/GenericPlatformProcess.h
@@ -332,7 +332,7 @@ struct CORE_API FGenericPlatformProcess
 	 * @param URL					executable name
 	 * @param Parms					command line arguments
 	 * @param bLaunchDetached		if true, the new process will have its own window
-	 * @param bLaunchHidded			if true, the new process will be minimized in the task bar
+	 * @param bLaunchHidden			if true, the new process will be minimized in the task bar
 	 * @param bLaunchReallyHidden	if true, the new process will not have a window or be in the task bar
 	 * @param OutProcessId			if non-NULL, this will be filled in with the ProcessId
 	 * @param PriorityModifier		-2 idle, -1 low, 0 normal, 1 high, 2 higher
diff --git a/Engine/Source/Runtime/Core/Public/Unix/UnixPlatformProcess.h b/Engine/Source/Runtime/Core/Public/Unix/UnixPlatformProcess.h
--- a/Engine/Source/Runtime/Core/Public/Unix/UnixPlatformProcess.h
+++ b/Engine/Source/Runtime/Core/Public/Unix/UnixPlatformProcess.h
@@ -105,34 +105,48 @@ protected:  // the below is not a public API!
 /** FProcHandle can be copied (and thus passed by value). */
 struct FProcHandle
 {
+	/** Child proc state set from CreateProc() call */
 	FProcState* 		ProcInfo;
 
+	/** Pid of external process opened with OpenProcess() call.
+	  * Added to FProcHandle so we don't have to special case FProcState with process
+	  * we can only check for running state, and even then the PID could be reused so
+	  * we don't ever want to terminate, etc.
+	  */
+	pid_t				OpenedPid;
+
 	FProcHandle()
-	:	ProcInfo(nullptr)
+	:	ProcInfo(nullptr), OpenedPid(-1)
 	{
 	}
 
 	FProcHandle(FProcState* InHandle)
-	:	ProcInfo(InHandle)
+	:	ProcInfo(InHandle), OpenedPid(-1)
+	{
+	}
+
+	FProcHandle(pid_t InProcPid)
+	:	ProcInfo(nullptr), OpenedPid(InProcPid)
 	{
 	}
 
 	/** Accessors. */
 	FORCEINLINE pid_t Get() const
 	{
-		return ProcInfo ? ProcInfo->GetProcessId() : -1;
+		return ProcInfo ? ProcInfo->GetProcessId() : OpenedPid;
 	}
 
 	/** Resets the handle to invalid */
 	FORCEINLINE void Reset()
 	{
 		ProcInfo = nullptr;
+		OpenedPid = -1;
 	}
 
 	/** Checks the validity of handle */
 	FORCEINLINE bool IsValid() const
 	{
-		return ProcInfo != nullptr;
+		return ProcInfo != nullptr || OpenedPid != -1;
 	}
 
 	// the below is not part of FProcHandle API and is specific to Unix implementation
@@ -264,6 +278,7 @@ struct CORE_API FUnixPlatformProcess : public FGenericPlatformProcess
 	static bool CanLaunchURL(const TCHAR* URL);
 	static void LaunchURL(const TCHAR* URL, const TCHAR* Parms, FString* Error);
 	static FProcHandle CreateProc(const TCHAR* URL, const TCHAR* Parms, bool bLaunchDetached, bool bLaunchHidden, bool bLaunchReallyHidden, uint32* OutProcessID, int32 PriorityModifier, const TCHAR* OptionalWorkingDirectory, void* PipeWriteChild, void* PipeReadChild = nullptr);
+	static FProcHandle OpenProcess(uint32 ProcessID);
 	static bool IsProcRunning( FProcHandle & ProcessHandle );
 	static void WaitForProc( FProcHandle & ProcessHandle );
 	static void CloseProc( FProcHandle & ProcessHandle );
