Enable boot profiling for embedded app

diff --git a/Engine/Source/Runtime/Core/Private/Misc/CoreGlobals.cpp b/Engine/Source/Runtime/Core/Private/Misc/CoreGlobals.cpp
--- a/Engine/Source/Runtime/Core/Private/Misc/CoreGlobals.cpp
+++ b/Engine/Source/Runtime/Core/Private/Misc/CoreGlobals.cpp
@@ -346,7 +346,7 @@ static struct FBootTimingStart
 } GBootTimingStart;
 
 
-#define USE_BOOT_PROFILING (0)
+#define USE_BOOT_PROFILING (BUILD_EMBEDDED_APP)
 
 #if !USE_BOOT_PROFILING
 FScopedBootTiming::FScopedBootTiming(const ANSICHAR *InMessage)
