#core Removed deprecated comments from TBaseMulticastDelegate::Add and TBaseMulticastDelegate::FDelegate. No intention to ever hide these, and they can be useful for binding the same handler to multiple delegates.

diff --git a/Engine/Source/Runtime/Core/Public/Delegates/DelegateSignatureImpl.inl b/Engine/Source/Runtime/Core/Public/Delegates/DelegateSignatureImpl.inl
--- a/Engine/Source/Runtime/Core/Public/Delegates/DelegateSignatureImpl.inl
+++ b/Engine/Source/Runtime/Core/Public/Delegates/DelegateSignatureImpl.inl
@@ -697,7 +697,7 @@ class TBaseMulticastDelegate<void, ParamTypes...> : public FMulticastDelegateBas
 	typedef FMulticastDelegateBase<FWeakObjectPtr> Super;
 
 public:
-	/** DEPRECATED: Type definition for unicast delegate classes whose delegate instances are compatible with this delegate. */
+	/** Type definition for unicast delegate classes whose delegate instances are compatible with this delegate. */
 	typedef TBaseDelegate< void, ParamTypes... > FDelegate;
 
 	/** Type definition for the shared interface of delegate instance types compatible with this delegate class. */
@@ -706,9 +706,7 @@ public:
 public:
 
 	/**
-	 * DEPRECATED: Adds a delegate instance to this multicast delegate's invocation list.
-	 *
-	 * This method is retained for backwards compatibility.
+	 * Adds a delegate instance to this multicast delegate's invocation list.
 	 *
 	 * @param Delegate The delegate to add.
 	 */
@@ -724,9 +722,7 @@ public:
 	}
 
 	/**
-	 * DEPRECATED: Adds a unicast delegate to this multi-cast delegate's invocation list.
-	 *
-	 * This method is retained for backwards compatibility.
+	 * Adds a delegate instance to this multicast delegate's invocation list.
 	 *
 	 * @param Delegate The delegate to add.
 	 */
