UE4 - Added a page cache to the XBox and adjusted some external code to properly call through the HAL instead of going directly to the OS for memory requests. The page cache is currently disabled as it does not produce measurable gains (yet).

diff --git a/Engine/Source/Runtime/Core/Private/HAL/MallocBinnedArena.cpp b/Engine/Source/Runtime/Core/Private/HAL/MallocBinnedArena.cpp
--- a/Engine/Source/Runtime/Core/Private/HAL/MallocBinnedArena.cpp
+++ b/Engine/Source/Runtime/Core/Private/HAL/MallocBinnedArena.cpp
@@ -569,10 +569,10 @@ void FMallocBinnedArena::InitMallocBinned()
 	}
 	for (uint32 Size : ArenaParams.AdditionalBlockSizes)
 	{
-		if (Size >= ArenaParams.BasePageSize)
+		if (Size > ArenaParams.BasePageSize)
 		{
-			check(Size % ArenaParams.BasePageSize == 0); // we can actually use blocks that are not divisible by the base page size, but we need a bit more work
-			SizeTable.Emplace(Size, ArenaParams.AllocationGranularity, Size / ArenaParams.BasePageSize, ArenaParams.BasePageSize, ArenaParams.MinimumAlignment);
+			check(Size % 4096 == 0); // calculations are done assume 4k is the smallest page size we will ever see
+			SizeTable.Emplace(Size, ArenaParams.AllocationGranularity, Size / 4096, ArenaParams.BasePageSize, ArenaParams.MinimumAlignment);
 		}
 		else
 		{
diff --git a/Engine/Source/Runtime/Core/Private/HAL/MallocBinnedGPU.cpp b/Engine/Source/Runtime/Core/Private/HAL/MallocBinnedGPU.cpp
--- a/Engine/Source/Runtime/Core/Private/HAL/MallocBinnedGPU.cpp
+++ b/Engine/Source/Runtime/Core/Private/HAL/MallocBinnedGPU.cpp
@@ -580,10 +580,10 @@ void FMallocBinnedGPU::InitMallocBinned()
 	}
 	for (uint32 Size : ArenaParams.AdditionalBlockSizes)
 	{
-		if (Size >= ArenaParams.BasePageSize)
+		if (Size > ArenaParams.BasePageSize)
 		{
-			check(Size % ArenaParams.BasePageSize == 0); // we can actually use blocks that are not divisible by the base page size, but we need a bit more work
-			SizeTable.Emplace(Size, ArenaParams.AllocationGranularity, Size / ArenaParams.BasePageSize, ArenaParams.BasePageSize, ArenaParams.MinimumAlignment);
+			check(Size % 4096 == 0); // calculations are done assume 4k is the smallest page size we will ever see
+			SizeTable.Emplace(Size, ArenaParams.AllocationGranularity, Size / 4096, ArenaParams.BasePageSize, ArenaParams.MinimumAlignment);
 		}
 		else
 		{
diff --git a/Engine/Source/Runtime/Core/Public/HAL/MallocBinnedArena.h b/Engine/Source/Runtime/Core/Public/HAL/MallocBinnedArena.h
--- a/Engine/Source/Runtime/Core/Public/HAL/MallocBinnedArena.h
+++ b/Engine/Source/Runtime/Core/Public/HAL/MallocBinnedArena.h
@@ -423,9 +423,10 @@ class CORE_API FMallocBinnedArena final : public FMalloc
 		{
 			return true;
 		}
-		if (ArenaParams.bAttemptToAlignSmallBocks & (InOutSize <= ArenaParams.MaxPoolSize) & (Alignment <= ArenaParams.MaximumAlignmentForSmallBlock)) // one branch, not three
+		SIZE_T AlignedSize = Align(InOutSize, Alignment);
+		if (ArenaParams.bAttemptToAlignSmallBocks & (AlignedSize <= ArenaParams.MaxPoolSize) & (Alignment <= ArenaParams.MaximumAlignmentForSmallBlock)) // one branch, not three
 		{
-			uint32 PoolIndex = BoundSizeToPoolIndex(Align(InOutSize, Alignment));
+			uint32 PoolIndex = BoundSizeToPoolIndex(AlignedSize);
 			while (true)
 			{
 				uint32 BlockSize = PoolIndexToBlockSize(PoolIndex);
diff --git a/Engine/Source/Runtime/Core/Public/HAL/MallocBinnedGPU.h b/Engine/Source/Runtime/Core/Public/HAL/MallocBinnedGPU.h
--- a/Engine/Source/Runtime/Core/Public/HAL/MallocBinnedGPU.h
+++ b/Engine/Source/Runtime/Core/Public/HAL/MallocBinnedGPU.h
@@ -445,9 +445,10 @@ class CORE_API FMallocBinnedGPU final : public FMalloc
 		{
 			return true;
 		}
-		if (ArenaParams.bAttemptToAlignSmallBocks & (InOutSize <= ArenaParams.MaxPoolSize) & (Alignment <= ArenaParams.MaximumAlignmentForSmallBlock)) // one branch, not three
+		SIZE_T AlignedSize = Align(InOutSize, Alignment);
+		if (ArenaParams.bAttemptToAlignSmallBocks & (AlignedSize <= ArenaParams.MaxPoolSize) & (Alignment <= ArenaParams.MaximumAlignmentForSmallBlock)) // one branch, not three
 		{
-			uint32 PoolIndex = BoundSizeToPoolIndex(Align(InOutSize, Alignment));
+			uint32 PoolIndex = BoundSizeToPoolIndex(AlignedSize);
 			while (true)
 			{
 				uint32 BlockSize = PoolIndexToBlockSize(PoolIndex);
diff --git a/Engine/Source/Runtime/Core/Public/HAL/VirtualAllocator.h b/Engine/Source/Runtime/Core/Public/HAL/VirtualAllocator.h
--- a/Engine/Source/Runtime/Core/Public/HAL/VirtualAllocator.h
+++ b/Engine/Source/Runtime/Core/Public/HAL/VirtualAllocator.h
@@ -78,11 +78,13 @@ class FVirtualAllocator
 		Block.FreeBlocksSize += AlignedSize;
 	}
 protected:
+	size_t SpaceConsumed;
 	virtual uint8* AllocNewVM(size_t AlignedSize)
 	{
 		uint8* Result = NextAlloc;
 		check(IsAligned(Result, MaximumAlignment) && IsAligned(AlignedSize, MaximumAlignment));
 		NextAlloc = Result + AlignedSize;
+		SpaceConsumed = NextAlloc - LowAddress;
 		return Result;
 	}
 
@@ -97,6 +99,7 @@ public:
 		, RecycledLinks(nullptr)
 		, LinkSize(0)
 		, bBacksMalloc(bInBacksMalloc)
+		, SpaceConsumed(0)
 	{
 		TotalSize = HighAddress - LowAddress;
 		check(LowAddress && HighAddress && LowAddress < HighAddress && IsAligned(LowAddress, MaximumAlignment));
@@ -224,7 +227,7 @@ public:
 		OutStats.PageSize = PageSize;
 		OutStats.MaximumAlignment = MaximumAlignment;
 		OutStats.VMSpaceTotal = HighAddress - LowAddress;
-		OutStats.VMSpaceConsumed = NextAlloc - LowAddress;
+		OutStats.VMSpaceConsumed = SpaceConsumed;
 		OutStats.FreeListLinks = LinkSize;
 
 		for (int32 Index = 0; Index < 64; Index++)
