UE-74632 - iOS logging dramatically suppressed

diff --git a/Engine/Source/Runtime/Core/Private/IOS/IOSPlatformMisc.cpp b/Engine/Source/Runtime/Core/Private/IOS/IOSPlatformMisc.cpp
--- a/Engine/Source/Runtime/Core/Private/IOS/IOSPlatformMisc.cpp
+++ b/Engine/Source/Runtime/Core/Private/IOS/IOSPlatformMisc.cpp
@@ -1696,7 +1696,13 @@ void FIOSPlatformMisc::SetCrashHandler(void (* CrashHandler)(const FGenericCrash
 
 bool FIOSPlatformMisc::HasSeparateChannelForDebugOutput()
 {
-    return FPlatformMisc::IsDebuggerPresent();
+#if UE_BUILD_SHIPPING
+    return false;
+#else
+    // We should not just check if we are being debugged because you can use the Xcode log even for
+    // apps launched outside the debugger.
+    return true;
+#endif
 }
 
 void FIOSPlatformMisc::GPUAssert()
