addition of OOM detection for IOS uploaded to Crash Reporter as a crash

diff --git a/Engine/Source/Runtime/Core/Private/IOS/IOSPlatformMisc.cpp b/Engine/Source/Runtime/Core/Private/IOS/IOSPlatformMisc.cpp
--- a/Engine/Source/Runtime/Core/Private/IOS/IOSPlatformMisc.cpp
+++ b/Engine/Source/Runtime/Core/Private/IOS/IOSPlatformMisc.cpp
@@ -1987,6 +1987,24 @@ void ReportEnsure( const TCHAR* ErrorMessage, int NumStackFramesToIgnore )
     EnsureLock.Unlock();
 }
 
+FString FIOSCrashContext::CreateCrashFolder() const
+{
+	// create a crash-specific directory
+	char CrashInfoFolder[PATH_MAX] = {};
+	FCStringAnsi::Strncpy(CrashInfoFolder, GIOSAppInfo.CrashReportPath, PATH_MAX);
+	FCStringAnsi::Strcat(CrashInfoFolder, PATH_MAX, "/CrashReport-UE4-");
+	FCStringAnsi::Strcat(CrashInfoFolder, PATH_MAX, GIOSAppInfo.AppNameUTF8);
+	FCStringAnsi::Strcat(CrashInfoFolder, PATH_MAX, "-pid-");
+	FCStringAnsi::Strcat(CrashInfoFolder, PATH_MAX, ItoANSI(getpid(), 10));
+	FCStringAnsi::Strcat(CrashInfoFolder, PATH_MAX, "-");
+	FCStringAnsi::Strcat(CrashInfoFolder, PATH_MAX, ItoANSI(GIOSAppInfo.RunUUID.A, 16));
+	FCStringAnsi::Strcat(CrashInfoFolder, PATH_MAX, ItoANSI(GIOSAppInfo.RunUUID.B, 16));
+	FCStringAnsi::Strcat(CrashInfoFolder, PATH_MAX, ItoANSI(GIOSAppInfo.RunUUID.C, 16));
+	FCStringAnsi::Strcat(CrashInfoFolder, PATH_MAX, ItoANSI(GIOSAppInfo.RunUUID.D, 16));
+	
+	return FString(ANSI_TO_TCHAR(CrashInfoFolder));
+}
+
 
 class FIOSExec : public FSelfRegisteringExec
 {
diff --git a/Engine/Source/Runtime/Core/Public/GenericPlatform/GenericPlatformCrashContext.h b/Engine/Source/Runtime/Core/Public/GenericPlatform/GenericPlatformCrashContext.h
--- a/Engine/Source/Runtime/Core/Public/GenericPlatform/GenericPlatformCrashContext.h
+++ b/Engine/Source/Runtime/Core/Public/GenericPlatform/GenericPlatformCrashContext.h
@@ -89,6 +89,7 @@ enum class ECrashContextType
 	Ensure,
 	GPUCrash,
 	Hang,
+	OutOfMemory,
 
 	Max
 };
diff --git a/Engine/Source/Runtime/Core/Public/IOS/IOSPlatformCrashContext.h b/Engine/Source/Runtime/Core/Public/IOS/IOSPlatformCrashContext.h
--- a/Engine/Source/Runtime/Core/Public/IOS/IOSPlatformCrashContext.h
+++ b/Engine/Source/Runtime/Core/Public/IOS/IOSPlatformCrashContext.h
@@ -20,6 +20,9 @@ struct CORE_API FIOSCrashContext : public FApplePlatformCrashContext
 	
 	/** Generates information for ensures sent via the CrashReporter */
 	void GenerateEnsureInfo() const;
+	
+	/** Creates the crash folder */
+	FString CreateCrashFolder() const;
 };
 
 typedef FIOSCrashContext FPlatformCrashContext;
