There's now a select instruction that can optimize this down from 8-9 instructions to a max of 2 (also eliminates a branch when creating the mask)

diff --git a/Engine/Source/Runtime/Core/Public/Templates/UnrealTemplate.h b/Engine/Source/Runtime/Core/Public/Templates/UnrealTemplate.h
--- a/Engine/Source/Runtime/Core/Public/Templates/UnrealTemplate.h
+++ b/Engine/Source/Runtime/Core/Public/Templates/UnrealTemplate.h
@@ -26,13 +26,7 @@
 template<typename ReferencedType>
 FORCEINLINE ReferencedType* IfAThenAElseB(ReferencedType* A,ReferencedType* B)
 {
-	const PTRINT IntA = reinterpret_cast<PTRINT>(A);
-	const PTRINT IntB = reinterpret_cast<PTRINT>(B);
-
-	// Compute a mask which has all bits set if IntA is zero, and no bits set if it's non-zero.
-	const PTRINT MaskB = -(!IntA);
-
-	return reinterpret_cast<ReferencedType*>(IntA | (MaskB & IntB));
+	return A ? A : B;
 }
 
 /** branchless pointer selection based on predicate
@@ -41,13 +35,7 @@ FORCEINLINE ReferencedType* IfAThenAElseB(ReferencedType* A,ReferencedType* B)
 template<typename PredicateType,typename ReferencedType>
 FORCEINLINE ReferencedType* IfPThenAElseB(PredicateType Predicate,ReferencedType* A,ReferencedType* B)
 {
-	const PTRINT IntA = reinterpret_cast<PTRINT>(A);
-	const PTRINT IntB = reinterpret_cast<PTRINT>(B);
-
-	// Compute a mask which has all bits set if Predicate is zero, and no bits set if it's non-zero.
-	const PTRINT MaskB = -(!PTRINT(Predicate));
-
-	return reinterpret_cast<ReferencedType*>((IntA & ~MaskB) | (IntB & MaskB));
+	return Predicate ? A : B;
 }
 
 /** A logical exclusive or function. */
