- Fix deadlock in LLM when adding to ThreadStates
- This was due to ThreadArraySection and the underlying allocators Mutex (MallocBinned2)

diff --git a/Engine/Source/Runtime/Core/Private/HAL/LowLevelMemTracker.cpp b/Engine/Source/Runtime/Core/Private/HAL/LowLevelMemTracker.cpp
--- a/Engine/Source/Runtime/Core/Private/HAL/LowLevelMemTracker.cpp
+++ b/Engine/Source/Runtime/Core/Private/HAL/LowLevelMemTracker.cpp
@@ -412,10 +412,12 @@ protected:
 
 	uint32 TlsSlot;
 
-	FCriticalSection ThreadArraySection;
 	FLLMObjectAllocator<FLLMThreadState> ThreadStateAllocator;
 	FLLMArray<FLLMThreadState*> ThreadStates;
 
+	FCriticalSection PendingThreadStatesGuard;
+	FLLMArray<FLLMThreadState*> PendingThreadStates;
+
 	int64 TrackedMemoryOverFrames GCC_ALIGN(8);
 
 	LLMMap AllocationMap;
@@ -1167,6 +1169,7 @@ void FLLMTracker::Initialise(
 
 	ThreadStateAllocator.SetAllocator(Allocator);
 	ThreadStates.SetAllocator(Allocator);
+	PendingThreadStates.SetAllocator(Allocator);
 }
 
 FLLMTracker::FLLMThreadState* FLLMTracker::GetOrCreateState()
@@ -1176,12 +1179,14 @@ FLLMTracker::FLLMThreadState* FLLMTracker::GetOrCreateState()
 	// get one if needed
 	if (State == nullptr)
 	{
-		// protect any accesses to the ThreadStates array
-		FScopeLock Lock(&ThreadArraySection);
-
 		State = ThreadStateAllocator.New();
 		State->SetAllocator(Allocator);
-		ThreadStates.Add(State);
+
+		// Add to pending thread states, these will be consumed on the GT
+		{
+			FScopeLock Lock(&PendingThreadStatesGuard);
+			PendingThreadStates.Add(State);
+		}
 
 		// push to Tls
 		FPlatformTLS::SetTlsValue(TlsSlot, State);
@@ -1350,6 +1355,13 @@ bool FLLMTracker::IsPaused(ELLMAllocType AllocType)
 
 void FLLMTracker::Clear()
 {
+	{
+		FScopeLock Lock(&PendingThreadStatesGuard);
+		for (uint32 Index = 0; Index < PendingThreadStates.Num(); ++Index)
+			ThreadStateAllocator.Delete(PendingThreadStates[Index]);
+		PendingThreadStates.Clear(true);
+	}
+
 	for (uint32 Index = 0; Index < ThreadStates.Num(); ++Index)
 		ThreadStateAllocator.Delete(ThreadStates[Index]);
 	ThreadStates.Clear(true);
@@ -1372,11 +1384,29 @@ void FLLMTracker::SetTotalTags(ELLMTag InUntaggedTotalTag, ELLMTag InTrackedTota
 
 void FLLMTracker::Update(FLLMCustomTag* CustomTags, const int32* ParentTags)
 {
-	// protect any accesses to the ThreadStates array
-	FScopeLock Lock(&ThreadArraySection);
+	int ThreadStateNum = ThreadStates.Num();
+
+	// Consume pending thread states
+	// We must be careful to do all allocations outside of the PendingThreadStatesGuard guard as that can lead to a deadlock due to contention with PendingThreadStatesGuard & Locks inside the underlying allocator (i.e. MallocBinned2 -> Mutex)
+	{
+		PendingThreadStatesGuard.Lock();
+		const int NumPendingThreadStatesToConsume = PendingThreadStates.Num();
+		if (NumPendingThreadStatesToConsume > 0 )
+		{
+			PendingThreadStatesGuard.Unlock();
+			ThreadStates.Reserve(ThreadStateNum + NumPendingThreadStatesToConsume);
+			PendingThreadStatesGuard.Lock();
+
+			for ( int32 i=0; i < NumPendingThreadStatesToConsume; ++i )
+			{
+				ThreadStates.Add(PendingThreadStates.RemoveLast());
+			}
+			ThreadStateNum += NumPendingThreadStatesToConsume;
+		}
+		PendingThreadStatesGuard.Unlock();
+	}
 
 	// accumulate the totals for each thread
-	int ThreadStateNum = ThreadStates.Num();
 	for (int32 ThreadIndex = 0; ThreadIndex < ThreadStateNum; ThreadIndex++)
 	{
 		ThreadStates[ThreadIndex]->UpdateFrameStatGroups(CustomTags,ParentTags);
