Re-enabled LLM on Editor builds since the performance issue has been addressed

diff --git a/Engine/Source/Runtime/Core/Public/HAL/LowLevelMemTracker.h b/Engine/Source/Runtime/Core/Public/HAL/LowLevelMemTracker.h
--- a/Engine/Source/Runtime/Core/Public/HAL/LowLevelMemTracker.h
+++ b/Engine/Source/Runtime/Core/Public/HAL/LowLevelMemTracker.h
@@ -25,7 +25,7 @@
 
 	// *** enable/disable LLM here ***
 #ifndef ENABLE_LOW_LEVEL_MEM_TRACKER
-	#define ENABLE_LOW_LEVEL_MEM_TRACKER (!UE_BUILD_SHIPPING && (!UE_BUILD_TEST || ALLOW_LOW_LEVEL_MEM_TRACKER_IN_TEST) && LLM_SUPPORTED_PLATFORM && WITH_ENGINE && !WITH_EDITOR && 1)
+	#define ENABLE_LOW_LEVEL_MEM_TRACKER (!UE_BUILD_SHIPPING && (!UE_BUILD_TEST || ALLOW_LOW_LEVEL_MEM_TRACKER_IN_TEST) && LLM_SUPPORTED_PLATFORM && WITH_ENGINE && 1)
 #endif
 
 	// using asset tagging requires a significantly higher number of per-thread tags, so make it optional
