Doubles in the form of 1e+9999 can be larger then 30 chars, set aside more room

diff --git a/Engine/Source/Runtime/Core/Private/GenericPlatform/GenericWidePlatformString.cpp b/Engine/Source/Runtime/Core/Private/GenericPlatform/GenericWidePlatformString.cpp
--- a/Engine/Source/Runtime/Core/Private/GenericPlatform/GenericWidePlatformString.cpp
+++ b/Engine/Source/Runtime/Core/Private/GenericPlatform/GenericWidePlatformString.cpp
@@ -576,7 +576,8 @@ int32 FGenericWidePlatformString::GetVarArgs(WIDECHAR* Dest, SIZE_T DestSize, in
 		{
 			Src++;
 			double Val = va_arg(ArgPtr, double);
-			ANSICHAR AnsiNum[30];
+			// doubles in the form of 1e+9999 can get quite large, make sure we have enough room for them
+			ANSICHAR AnsiNum[48];
 			ANSICHAR FmtBuf[30];
 
 			// Yes, this is lame.
@@ -594,6 +595,12 @@ int32 FGenericWidePlatformString::GetVarArgs(WIDECHAR* Dest, SIZE_T DestSize, in
 			{
 				return -1;	// Fail - the app needs to create a larger buffer and try again
 			}
+			if (RetCnt >= ARRAY_COUNT(AnsiNum))
+			{
+				// We should print what we have written into AnsiNum but ensure we null terminate before printing
+				AnsiNum[ARRAY_COUNT(AnsiNum) - 1] = '\0';
+				checkf(0, TEXT("Attempting to read past the size our buffer. Buffer Size: %d Size to read: %d. Current contents: '%s'\n"), ARRAY_COUNT(AnsiNum), RetCnt, UTF8_TO_TCHAR(AnsiNum));
+			}
 			for (int i = 0; i < RetCnt; i++)
 			{
 				*Dst = (TCHAR) AnsiNum[i];
