Copied 4359025 Francis.Hurteau
Test if have a valid continuation before running it when setting continuation on a ready future

diff --git a/Engine/Source/Runtime/Core/Public/Async/Future.h b/Engine/Source/Runtime/Core/Public/Async/Future.h
--- a/Engine/Source/Runtime/Core/Public/Async/Future.h
+++ b/Engine/Source/Runtime/Core/Public/Async/Future.h
@@ -90,7 +90,7 @@ public:
 				CompletionCallback = MoveTemp(Continuation);
 			}
 		}
-		if (bShouldJustRun)
+		if (bShouldJustRun && Continuation)
 		{
 			Continuation();
 		}
