Merging CL 4304711 from Dev-VR - Downgrading sandbox jail warning to log, as it is expected output

diff --git a/Engine/Source/Runtime/Core/Private/Lumin/LuminPlatformFile.cpp b/Engine/Source/Runtime/Core/Private/Lumin/LuminPlatformFile.cpp
--- a/Engine/Source/Runtime/Core/Private/Lumin/LuminPlatformFile.cpp
+++ b/Engine/Source/Runtime/Core/Private/Lumin/LuminPlatformFile.cpp
@@ -903,7 +903,7 @@ bool FLuminPlatformFile::DeleteDirectory(const TCHAR* Directory)
 void FLuminPlatformFile::SetSandboxEnabled(bool bInEnabled)
 {
 	bIsSandboxEnabled = bInEnabled;
-	UE_LOG(LogLuminPlatformFile, Warning, TEXT("Application sandbox jail has been %s."), bIsSandboxEnabled ? TEXT("enabled") : TEXT("disabled"));
+	UE_LOG(LogLuminPlatformFile, Log, TEXT("Application sandbox jail has been %s."), bIsSandboxEnabled ? TEXT("enabled") : TEXT("disabled"));
 }
 
 bool FLuminPlatformFile::IsSandboxEnabled() const
