Include portable callstacks in ensures.

diff --git a/Engine/Source/Runtime/Core/Private/Apple/ApplePlatformCrashContext.cpp b/Engine/Source/Runtime/Core/Private/Apple/ApplePlatformCrashContext.cpp
--- a/Engine/Source/Runtime/Core/Private/Apple/ApplePlatformCrashContext.cpp
+++ b/Engine/Source/Runtime/Core/Private/Apple/ApplePlatformCrashContext.cpp
@@ -75,7 +75,7 @@ int32 FApplePlatformCrashContext::ReportCrash() const
 		*StackBuffer = 0;
 		
 		// Generate the portable callstack for the crash xml
-		GeneratePortableCallStack(IgnoreDepth, ARRAY_COUNT(MinidumpCallstackInfo) - 1, Context);		
+		CapturePortableCallStack(IgnoreDepth, Context);		
 		// Walk the stack and dump it to the allocated memory (ignore first 2 callstack lines as those are in stack walking code)
 		FPlatformStackWalk::StackWalkAndDump( StackBuffer, ARRAY_COUNT(MinidumpCallstackInfo) - 1, IgnoreDepth, Context );
 		
diff --git a/Engine/Source/Runtime/Core/Private/GenericPlatform/GenericPlatformCrashContext.cpp b/Engine/Source/Runtime/Core/Private/GenericPlatform/GenericPlatformCrashContext.cpp
--- a/Engine/Source/Runtime/Core/Private/GenericPlatform/GenericPlatformCrashContext.cpp
+++ b/Engine/Source/Runtime/Core/Private/GenericPlatform/GenericPlatformCrashContext.cpp
@@ -30,20 +30,6 @@ extern CORE_API bool GIsGPUCrashed;
 	FGenericCrashContext
 -----------------------------------------------------------------------------*/
 
-struct FCrashStackFrame
-{
-	FCrashStackFrame(const FString& ModuleNameIn, uint64 BaseAddressIn, uint64 OffsetIn)
-	{
-		ModuleName = ModuleNameIn;
-		BaseAddress = BaseAddressIn;
-		Offset = OffsetIn;
-	}
-
-	FString ModuleName;
-	uint64 BaseAddress;
-	uint64 Offset;
-};
-
 const ANSICHAR* FGenericCrashContext::CrashContextRuntimeXMLNameA = "CrashContext.runtime-xml";
 const TCHAR* FGenericCrashContext::CrashContextRuntimeXMLNameW = TEXT( "CrashContext.runtime-xml" );
 
@@ -109,7 +95,6 @@ namespace NCachedCrashContextProperties
 	static FString CrashReportClientRichText;
 	static FString GameStateName;
 	static TArray<FString> EnabledPluginsList;
-	static TArray<FCrashStackFrame> CrashStack;	
 }
 
 void FGenericCrashContext::Initialize()
@@ -472,7 +457,7 @@ void FGenericCrashContext::AddPlatformSpecificProperties() const
 void FGenericCrashContext::AddPortableCallStack() const
 {	
 
-	if (NCachedCrashContextProperties::CrashStack.Num() == 0)
+	if (CallStack.Num() == 0)
 	{
 		AddCrashProperty(TEXT("PCallStack"), TEXT(""));
 		return;
@@ -482,12 +467,12 @@ void FGenericCrashContext::AddPortableCallStack() const
 
 	// Get the max module name length for padding
 	int32 MaxModuleLength = 0;
-	for (TArray<FCrashStackFrame>::TConstIterator It(NCachedCrashContextProperties::CrashStack); It; ++It)
+	for (TArray<FCrashStackFrame>::TConstIterator It(CallStack); It; ++It)
 	{
 		MaxModuleLength = FMath::Max(MaxModuleLength, It->ModuleName.Len());
 	}
 
-	for (TArray<FCrashStackFrame>::TConstIterator It(NCachedCrashContextProperties::CrashStack); It; ++It)
+	for (TArray<FCrashStackFrame>::TConstIterator It(CallStack); It; ++It)
 	{
 		CrashStackBuffer += FString::Printf(TEXT("%-*s 0x%016x + %-8x"),MaxModuleLength + 1, *It->ModuleName, It->BaseAddress, It->Offset);
 		CrashStackBuffer += LINE_TERMINATOR;
@@ -498,8 +483,6 @@ void FGenericCrashContext::AddPortableCallStack() const
 	AppendEscapedXMLString(EscapedStackBuffer, *CrashStackBuffer);
 
 	AddCrashProperty(TEXT("PCallStack"), *EscapedStackBuffer);
-
-	NCachedCrashContextProperties::CrashStack.Empty();
 }
 
 void FGenericCrashContext::AddHeader() const
@@ -654,12 +637,7 @@ void FGenericCrashContext::AddPlugin(const FString& PluginDesc)
 	NCachedCrashContextProperties::EnabledPluginsList.Add(PluginDesc);
 }
 
-void FGenericCrashContext::AddStackFrame(const FString& ModuleName, uint64 BaseAddress, uint64 Offset)
-{
-	NCachedCrashContextProperties::CrashStack.Add(FCrashStackFrame(ModuleName, BaseAddress, Offset));
-}
-
-void FGenericCrashContext::GeneratePortableCallStack(int32 IgnoreCount, int32 MaxDepth, void* Context)
+FORCENOINLINE void FGenericCrashContext::CapturePortableCallStack(int32 NumStackFramesToIgnore, void* Context)
 {
 	uint32 ModuleEntries = (uint32)FPlatformStackWalk::GetProcessModuleCount();
 
@@ -669,10 +647,11 @@ void FGenericCrashContext::GeneratePortableCallStack(int32 IgnoreCount, int32 Ma
 		ProcessModules.AddUninitialized(ModuleEntries);
 		FPlatformStackWalk::GetProcessModuleSignatures(ProcessModules.GetData(), ProcessModules.Max());
 
-		TArray<FProgramCounterSymbolInfo> Stack = FPlatformStackWalk::GetStack(IgnoreCount, MaxDepth, Context);
+		const int32 MaxDepth = 100;
+		TArray<FProgramCounterSymbolInfo> Stack = FPlatformStackWalk::GetStack(NumStackFramesToIgnore, MaxDepth, Context);
+
 		TMap<FString, uint64> ImageBases;
 		int32 ModuleIndex = 0;
-
 		for (TArray<FProgramCounterSymbolInfo>::TConstIterator Itr(Stack); Itr; ++Itr)
 		{
 			FString ModuleName = FPaths::GetBaseFilename(Itr->ModuleName);
@@ -699,7 +678,7 @@ void FGenericCrashContext::GeneratePortableCallStack(int32 IgnoreCount, int32 Ma
 				BaseOfImage = ImageBases[ModuleName];
 			}
 
-			FGenericCrashContext::AddStackFrame(ModuleName, BaseOfImage, Itr->ProgramCounter > BaseOfImage ? Itr->ProgramCounter - BaseOfImage : MAX_uint64);
+			CallStack.Add(FCrashStackFrame(ModuleName, BaseOfImage, Itr->ProgramCounter > BaseOfImage ? Itr->ProgramCounter - BaseOfImage : MAX_uint64));
 		}
 	}
 
diff --git a/Engine/Source/Runtime/Core/Private/Unix/UnixPlatformCrashContext.cpp b/Engine/Source/Runtime/Core/Private/Unix/UnixPlatformCrashContext.cpp
--- a/Engine/Source/Runtime/Core/Private/Unix/UnixPlatformCrashContext.cpp
+++ b/Engine/Source/Runtime/Core/Private/Unix/UnixPlatformCrashContext.cpp
@@ -285,7 +285,7 @@ void FUnixCrashContext::CaptureStackTrace()
 		StackTrace[0] = 0;
 
 		int32 IgnoreCount = 0;
-		FGenericCrashContext::GeneratePortableCallStack(IgnoreCount, StackTraceSize, this);
+		CapturePortableCallStack(IgnoreCount, this);
 
 		// Walk the stack and dump it to the allocated memory (do not ignore any stack frames to be consistent with check()/ensure() handling)
 		FPlatformStackWalk::StackWalkAndDump( StackTrace, StackTraceSize, IgnoreCount, this);
diff --git a/Engine/Source/Runtime/Core/Private/Windows/WindowsPlatformCrashContext.cpp b/Engine/Source/Runtime/Core/Private/Windows/WindowsPlatformCrashContext.cpp
--- a/Engine/Source/Runtime/Core/Private/Windows/WindowsPlatformCrashContext.cpp
+++ b/Engine/Source/Runtime/Core/Private/Windows/WindowsPlatformCrashContext.cpp
@@ -301,11 +301,14 @@ static bool bReentranceGuard = false;
 /**
  * A wrapper for ReportCrashUsingCrashReportClient that creates a new ensure crash context
  */
-int32 ReportEnsureUsingCrashReportClient(EXCEPTION_POINTERS* ExceptionInfo, const TCHAR* ErrorMessage, EErrorReportUI ReportUI)
+int32 ReportEnsureUsingCrashReportClient(HANDLE Thread, EXCEPTION_POINTERS* ExceptionInfo, int32 IgnoreCount, const TCHAR* ErrorMessage, EErrorReportUI ReportUI)
 {
 	const bool bIsEnsure = true;
 	FWindowsPlatformCrashContext CrashContext(bIsEnsure);
 
+	void* ContextWrapper = FWindowsPlatformStackWalk::MakeThreadContextWrapper(ExceptionInfo->ContextRecord, Thread);
+	CrashContext.CapturePortableCallStack(IgnoreCount, ContextWrapper);
+
 	return ReportCrashUsingCrashReportClient(CrashContext, ExceptionInfo, ErrorMessage, ReportUI, bIsEnsure);
 }
 #endif
@@ -314,7 +317,7 @@ int32 ReportEnsureUsingCrashReportClient(EXCEPTION_POINTERS* ExceptionInfo, cons
 /** 
  * Report an ensure to the crash reporting system
  */
-void NewReportEnsure( const TCHAR* ErrorMessage )
+FORCENOINLINE void NewReportEnsure( const TCHAR* ErrorMessage )
 {
 	if (ReportCrashCallCount > 0)
 	{
@@ -343,6 +346,9 @@ void NewReportEnsure( const TCHAR* ErrorMessage )
 	}
 
 	bReentranceGuard = true;
+	
+	// Three additional frames are skipped inside FWindowsPlatformStackWalk::GetStack() and FPlatformStackWalk::GetStack()
+	const int NumStackFramesToIgnore = 2;
 
 #if WINVER > 0x502	// Windows Error Reporting is not supported on Windows XP
 #if !PLATFORM_SEH_EXCEPTIONS_DISABLED
@@ -352,7 +358,7 @@ void NewReportEnsure( const TCHAR* ErrorMessage )
 		FPlatformMisc::RaiseException( 1 );
 	}
 #if !PLATFORM_SEH_EXCEPTIONS_DISABLED
-	__except(ReportEnsureUsingCrashReportClient(GetExceptionInformation(), ErrorMessage, IsInteractiveEnsureMode() ? EErrorReportUI::ShowDialog : EErrorReportUI::ReportInUnattendedMode))
+	__except(ReportEnsureUsingCrashReportClient(GetCurrentThread(), GetExceptionInformation(), NumStackFramesToIgnore, ErrorMessage, IsInteractiveEnsureMode() ? EErrorReportUI::ShowDialog : EErrorReportUI::ReportInUnattendedMode))
 	CA_SUPPRESS(6322)
 	{
 	}
@@ -544,8 +550,7 @@ private:
 
 		// Generate the portable callstack, for non-Asserts pass in a negative IgnoreCount as need top of crash context record, and IgnoreCount is increased by stack walking methods
 		const int32 IgnoreCount = FDebug::HasAsserted() ? 1 : -3;
-		const int32 MaxDepth = 100;
-		FGenericCrashContext::GeneratePortableCallStack(IgnoreCount, MaxDepth, ContextWrapper);
+		CrashContext.CapturePortableCallStack(IgnoreCount, ContextWrapper);
 
 		// First launch the crash reporter client.
 #if WINVER > 0x502	// Windows Error Reporting is not supported on Windows XP
diff --git a/Engine/Source/Runtime/Core/Private/Windows/WindowsPlatformMisc.cpp b/Engine/Source/Runtime/Core/Private/Windows/WindowsPlatformMisc.cpp
--- a/Engine/Source/Runtime/Core/Private/Windows/WindowsPlatformMisc.cpp
+++ b/Engine/Source/Runtime/Core/Private/Windows/WindowsPlatformMisc.cpp
@@ -1815,7 +1815,7 @@ HWND FWindowsPlatformMisc::GetTopLevelWindowHandle(uint32 ProcessId)
 	return Data.Handle;
 }
 
-void FWindowsPlatformMisc::RaiseException( uint32 ExceptionCode )
+FORCENOINLINE void FWindowsPlatformMisc::RaiseException( uint32 ExceptionCode )
 {
 	/** This is the last place to gather memory stats before exception. */
 	FGenericCrashContext::CrashMemoryStats = FPlatformMemory::GetStats();
diff --git a/Engine/Source/Runtime/Core/Private/Windows/WindowsPlatformStackWalk.cpp b/Engine/Source/Runtime/Core/Private/Windows/WindowsPlatformStackWalk.cpp
--- a/Engine/Source/Runtime/Core/Private/Windows/WindowsPlatformStackWalk.cpp
+++ b/Engine/Source/Runtime/Core/Private/Windows/WindowsPlatformStackWalk.cpp
@@ -235,7 +235,7 @@ void FWindowsPlatformStackWalk::StackWalkAndDump( ANSICHAR* HumanReadableString,
 	FGenericPlatformStackWalk::StackWalkAndDump(HumanReadableString, HumanReadableStringSize, IgnoreCount, Context);
 }
 
-TArray<FProgramCounterSymbolInfo> FWindowsPlatformStackWalk::GetStack(int32 IgnoreCount, int32 MaxDepth, void* Context)
+FORCENOINLINE TArray<FProgramCounterSymbolInfo> FWindowsPlatformStackWalk::GetStack(int32 IgnoreCount, int32 MaxDepth, void* Context)
 {
 	InitStackWalking();
 	return FGenericPlatformStackWalk::GetStack(IgnoreCount + 1, MaxDepth, Context);
diff --git a/Engine/Source/Runtime/Core/Public/GenericPlatform/GenericPlatformCrashContext.h b/Engine/Source/Runtime/Core/Public/GenericPlatform/GenericPlatformCrashContext.h
--- a/Engine/Source/Runtime/Core/Public/GenericPlatform/GenericPlatformCrashContext.h
+++ b/Engine/Source/Runtime/Core/Public/GenericPlatform/GenericPlatformCrashContext.h
@@ -65,6 +65,21 @@ enum class ECrashDumpMode : int32
 	FullDumpAlways = 2,
 };
 
+/** Portable stack frame */
+struct FCrashStackFrame
+{
+	FString ModuleName;
+	uint64 BaseAddress;
+	uint64 Offset;
+
+	FCrashStackFrame(const FString& ModuleNameIn, uint64 BaseAddressIn, uint64 OffsetIn)
+	{
+		ModuleName = ModuleNameIn;
+		BaseAddress = BaseAddressIn;
+		Offset = OffsetIn;
+	}
+};
+
 /**
  *	Contains a runtime crash's properties that are common for all platforms.
  *	This may change in the future.
@@ -189,7 +204,7 @@ public:
 	static void AddPlugin(const FString& PluginDesc);
 	
 	/** Generate raw call stack for crash report (image base + offset) */
-	static void GeneratePortableCallStack(int32 IgnoreCount, int32 MaxDepth = 100, void* Context = nullptr);
+	void CapturePortableCallStack(int32 NumStackFramesToIgnore, void* Context);
 
 	/**
 	 * @return whether this crash is a non-crash event
@@ -203,6 +218,7 @@ public:
 
 protected:
 	bool bIsEnsure;
+	TArray<FCrashStackFrame> CallStack;
 
 private:
 
@@ -212,9 +228,6 @@ private:
 	/** Add callstack information to the crash report xml */
 	void AddPortableCallStack() const;
 
-	/** Add a stack frame to the crash report call stack */
-	static void AddStackFrame(const FString& ModuleName, uint64 BaseAddress, uint64 Offset);
-
 	/** Writes header information to the buffer. */
 	void AddHeader() const;
 
@@ -243,8 +256,8 @@ private:
 	int32 CrashContextIndex;
 
 	// FNoncopyable
-	FGenericCrashContext( const FGenericCrashContext& );
-	FGenericCrashContext& operator=(const FGenericCrashContext&);
+	FGenericCrashContext( const FGenericCrashContext& ) = delete;
+	FGenericCrashContext& operator=(const FGenericCrashContext&) = delete;
 };
 
 struct CORE_API FGenericMemoryWarningContext
