Fix Mac compile error due to calling non-const function.

diff --git a/Engine/Source/Runtime/Core/Private/Apple/ApplePlatformCrashContext.cpp b/Engine/Source/Runtime/Core/Private/Apple/ApplePlatformCrashContext.cpp
--- a/Engine/Source/Runtime/Core/Private/Apple/ApplePlatformCrashContext.cpp
+++ b/Engine/Source/Runtime/Core/Private/Apple/ApplePlatformCrashContext.cpp
@@ -62,7 +62,7 @@ void FApplePlatformCrashContext::InitFromSignal(int32 InSignal, siginfo_t* InInf
 #undef HANDLE_CASE
 }
 
-int32 FApplePlatformCrashContext::ReportCrash() const
+int32 FApplePlatformCrashContext::ReportCrash()
 {
 	static bool GAlreadyCreatedMinidump = false;
 	// Only create a minidump the first time this function is called.
diff --git a/Engine/Source/Runtime/Core/Private/Mac/MacPlatformMisc.cpp b/Engine/Source/Runtime/Core/Private/Mac/MacPlatformMisc.cpp
--- a/Engine/Source/Runtime/Core/Private/Mac/MacPlatformMisc.cpp
+++ b/Engine/Source/Runtime/Core/Private/Mac/MacPlatformMisc.cpp
@@ -1450,7 +1450,7 @@ void (* GCrashHandlerPointer)(const FGenericCrashContext& Context) = NULL;
 /**
  * Good enough default crash reporter.
  */
-static void DefaultCrashHandler(FMacCrashContext const& Context)
+static void DefaultCrashHandler(FMacCrashContext& Context)
 {
 	Context.ReportCrash();
 	if (GLog)
diff --git a/Engine/Source/Runtime/Core/Public/Apple/ApplePlatformCrashContext.h b/Engine/Source/Runtime/Core/Public/Apple/ApplePlatformCrashContext.h
--- a/Engine/Source/Runtime/Core/Public/Apple/ApplePlatformCrashContext.h
+++ b/Engine/Source/Runtime/Core/Public/Apple/ApplePlatformCrashContext.h
@@ -43,7 +43,7 @@ struct CORE_API FApplePlatformCrashContext : public FGenericCrashContext
 	void InitFromSignal(int32 InSignal, siginfo_t* InInfo, void* InContext);
 	
 	/** Initialise context for reporting crash. */
-	int32 ReportCrash(void) const;
+	int32 ReportCrash(void);
 	
 	/** Generates a string representation for the execption/signal Info. */
 	static void CreateExceptionInfoString(int32 Signal, siginfo_t* Info);
