Restore const signature for FApplePlatformCrashContext::ReportCrash() for compatibility with callbacks.

diff --git a/Engine/Source/Runtime/Core/Private/Apple/ApplePlatformCrashContext.cpp b/Engine/Source/Runtime/Core/Private/Apple/ApplePlatformCrashContext.cpp
--- a/Engine/Source/Runtime/Core/Private/Apple/ApplePlatformCrashContext.cpp
+++ b/Engine/Source/Runtime/Core/Private/Apple/ApplePlatformCrashContext.cpp
@@ -62,7 +62,7 @@ void FApplePlatformCrashContext::InitFromSignal(int32 InSignal, siginfo_t* InInf
 #undef HANDLE_CASE
 }
 
-int32 FApplePlatformCrashContext::ReportCrash()
+int32 FApplePlatformCrashContext::ReportCrash() const
 {
 	static bool GAlreadyCreatedMinidump = false;
 	// Only create a minidump the first time this function is called.
@@ -75,7 +75,7 @@ int32 FApplePlatformCrashContext::ReportCrash()
 		*StackBuffer = 0;
 		
 		// Generate the portable callstack for the crash xml
-		CapturePortableCallStack(IgnoreDepth, Context);		
+		const_cast<FApplePlatformCrashContext*>(this)->CapturePortableCallStack(IgnoreDepth, Context);		
 		// Walk the stack and dump it to the allocated memory (ignore first 2 callstack lines as those are in stack walking code)
 		FPlatformStackWalk::StackWalkAndDump( StackBuffer, ARRAY_COUNT(MinidumpCallstackInfo) - 1, IgnoreDepth, Context );
 		
diff --git a/Engine/Source/Runtime/Core/Private/Mac/MacPlatformMisc.cpp b/Engine/Source/Runtime/Core/Private/Mac/MacPlatformMisc.cpp
--- a/Engine/Source/Runtime/Core/Private/Mac/MacPlatformMisc.cpp
+++ b/Engine/Source/Runtime/Core/Private/Mac/MacPlatformMisc.cpp
@@ -1450,7 +1450,7 @@ void (* GCrashHandlerPointer)(const FGenericCrashContext& Context) = NULL;
 /**
  * Good enough default crash reporter.
  */
-static void DefaultCrashHandler(FMacCrashContext& Context)
+static void DefaultCrashHandler(FMacCrashContext const& Context)
 {
 	Context.ReportCrash();
 	if (GLog)
diff --git a/Engine/Source/Runtime/Core/Public/Apple/ApplePlatformCrashContext.h b/Engine/Source/Runtime/Core/Public/Apple/ApplePlatformCrashContext.h
--- a/Engine/Source/Runtime/Core/Public/Apple/ApplePlatformCrashContext.h
+++ b/Engine/Source/Runtime/Core/Public/Apple/ApplePlatformCrashContext.h
@@ -43,7 +43,7 @@ struct CORE_API FApplePlatformCrashContext : public FGenericCrashContext
 	void InitFromSignal(int32 InSignal, siginfo_t* InInfo, void* InContext);
 	
 	/** Initialise context for reporting crash. */
-	int32 ReportCrash(void);
+	int32 ReportCrash(void) const;
 	
 	/** Generates a string representation for the execption/signal Info. */
 	static void CreateExceptionInfoString(int32 Signal, siginfo_t* Info);
