Merging CL# 4357533, CL# 4366557 and CL# 4374101

#rnx

diff --git a/Engine/Source/Runtime/Core/Private/Internationalization/TextLocalizationManager.cpp b/Engine/Source/Runtime/Core/Private/Internationalization/TextLocalizationManager.cpp
--- a/Engine/Source/Runtime/Core/Private/Internationalization/TextLocalizationManager.cpp
+++ b/Engine/Source/Runtime/Core/Private/Internationalization/TextLocalizationManager.cpp
@@ -562,18 +562,25 @@ FTextDisplayStringRef FTextLocalizationManager::GetDisplayString(const FString&
 
 	// In builds with stable keys enabled, we want to use the display string from the "clean" version of the text (if the sources match) as this is the only version that is translated
 	const FString* DisplayString = SourceString;
+	FDisplayStringLookupTable::FDisplayStringEntry* DisplayLiveEntry = nullptr;
 #if USE_STABLE_LOCALIZATION_KEYS
 	if (GIsEditor)
 	{
 		const FString DisplayNamespace = TextNamespaceUtil::StripPackageNamespace(Namespace);
 
 		FDisplayStringLookupTable::FKeysTable* DisplayLiveKeyTable = nullptr;
-		FDisplayStringLookupTable::FDisplayStringEntry* DisplayLiveEntry = nullptr;
 		DisplayStringLookupTable.Find(DisplayNamespace, DisplayLiveKeyTable, Key, DisplayLiveEntry);
 
-		if (DisplayLiveEntry && (!SourceString || DisplayLiveEntry->SourceStringHash == SourceStringHash))
+		if (DisplayLiveEntry)
 		{
-			DisplayString = &DisplayLiveEntry->DisplayString.Get();
+			if (!SourceString || DisplayLiveEntry->SourceStringHash == SourceStringHash)
+			{
+				DisplayString = &DisplayLiveEntry->DisplayString.Get();
+			}
+			else
+			{
+				DisplayLiveEntry = nullptr;
+			}
 		}
 	}
 #endif // USE_STABLE_LOCALIZATION_KEYS
@@ -610,6 +617,27 @@ FTextDisplayStringRef FTextLocalizationManager::GetDisplayString(const FString&
 
 		return LiveEntry->DisplayString;
 	}
+	// Entry is absent, but has a related entry to clone.
+	else if (DisplayLiveEntry)
+	{
+		check(!SourceString || DisplayLiveEntry->SourceStringHash == SourceStringHash);
+		check(DisplayString && &DisplayLiveEntry->DisplayString.Get() == DisplayString);
+
+		// Clone the entry for the active ID, and assign it a new display string instance (as all entries must have a unique display string instance).
+		FDisplayStringLookupTable::FDisplayStringEntry NewEntry(*DisplayLiveEntry);
+		NewEntry.DisplayString = MakeShared<FString, ESPMode::ThreadSafe>(*DisplayString);
+
+		if (!LiveKeyTable)
+		{
+			LiveKeyTable = &(DisplayStringLookupTable.NamespacesTable.Add(Namespace, FDisplayStringLookupTable::FKeysTable()));
+		}
+
+		LiveKeyTable->Add(Key, NewEntry);
+
+		NamespaceKeyLookupTable.Add(NewEntry.DisplayString, FNamespaceKeyEntry(Namespace, Key));
+
+		return NewEntry.DisplayString;
+	}
 	// Entry is absent.
 	else
 	{
@@ -785,7 +813,9 @@ void FTextLocalizationManager::UpdateFromLocalizationResource(const FString& Loc
 {
 	TSharedPtr<FTextLocalizationResource> TextLocalizationResource = MakeShared<FTextLocalizationResource>();
 	TextLocalizationResource->LoadFromFile(LocalizationResourceFilePath);
+#if !UE_BUILD_SHIPPING
 	TextLocalizationResource->DetectAndLogConflicts();
+#endif
 
 	UpdateFromLocalizations(TArrayView<const TSharedPtr<FTextLocalizationResource>>(&TextLocalizationResource, 1));
 }
@@ -873,17 +903,22 @@ void FTextLocalizationManager::LoadLocalizationResourcesForPrioritizedCultures(T
 	// The leet culture is fake. Just leet-ify existing strings.
 	if (PrioritizedCultureNames[0] == TEXT("LEET"))
 	{
-		for (auto NamespaceIterator = DisplayStringLookupTable.NamespacesTable.CreateIterator(); NamespaceIterator; ++NamespaceIterator)
+		// Lock while updating the tables
 		{
-			const FString& Namespace = NamespaceIterator.Key();
-			FDisplayStringLookupTable::FKeysTable& LiveKeyTable = NamespaceIterator.Value();
-			for (auto KeyIterator = LiveKeyTable.CreateIterator(); KeyIterator; ++KeyIterator)
+			FScopeLock ScopeLock(&SynchronizationObject);
+
+			for (auto NamespaceIterator = DisplayStringLookupTable.NamespacesTable.CreateIterator(); NamespaceIterator; ++NamespaceIterator)
 			{
-				const FString& Key = KeyIterator.Key();
-				FDisplayStringLookupTable::FDisplayStringEntry& LiveStringEntry = KeyIterator.Value();
-				LiveStringEntry.bIsLocalized = true;
-				LiveStringEntry.NativeStringBackup = *LiveStringEntry.DisplayString;
-				FInternationalization::Leetify(*LiveStringEntry.DisplayString);
+				const FString& Namespace = NamespaceIterator.Key();
+				FDisplayStringLookupTable::FKeysTable& LiveKeyTable = NamespaceIterator.Value();
+				for (auto KeyIterator = LiveKeyTable.CreateIterator(); KeyIterator; ++KeyIterator)
+				{
+					const FString& Key = KeyIterator.Key();
+					FDisplayStringLookupTable::FDisplayStringEntry& LiveStringEntry = KeyIterator.Value();
+					LiveStringEntry.bIsLocalized = true;
+					LiveStringEntry.NativeStringBackup = *LiveStringEntry.DisplayString;
+					FInternationalization::Leetify(*LiveStringEntry.DisplayString);
+				}
 			}
 		}
 
@@ -917,99 +952,105 @@ void FTextLocalizationManager::LoadLocalizationResourcesForPrioritizedCultures(T
 
 void FTextLocalizationManager::UpdateFromNative(const FTextLocalizationResource& TextLocalizationResource)
 {
-	// Note: This code doesn't handle "leet-ification" itself as it is resetting everything to a known "good" state ("leet-ification" happens later on the "good" native text)
-
-	// Update existing entries to use the new native text
-	for (auto& Namespace : DisplayStringLookupTable.NamespacesTable)
+	// Lock while updating the tables
 	{
-		const FString& NamespaceName = Namespace.Key;
-		FDisplayStringLookupTable::FKeysTable& LiveKeyTable = Namespace.Value;
+		FScopeLock ScopeLock(&SynchronizationObject);
 
-		// In builds with stable keys enabled, we want to use the display string from the "clean" version of the text (if the sources match) as this is the only version that is translated
-		const FString* NamespaceNamePtr = &NamespaceName;
-#if USE_STABLE_LOCALIZATION_KEYS
-		FString DisplayNamespace;
-		if (GIsEditor)
-		{
-			DisplayNamespace = TextNamespaceUtil::StripPackageNamespace(NamespaceName);
-			NamespaceNamePtr = &DisplayNamespace;
-		}
-#endif // USE_STABLE_LOCALIZATION_KEYS
+		// Note: This code doesn't handle "leet-ification" itself as it is resetting everything to a known "good" state ("leet-ification" happens later on the "good" native text)
 
-		for (auto& Key : LiveKeyTable)
+		// Update existing entries to use the new native text
+		for (auto& Namespace : DisplayStringLookupTable.NamespacesTable)
 		{
-			const FString& KeyName = Key.Key;
-			FDisplayStringLookupTable::FDisplayStringEntry& LiveStringEntry = Key.Value;
-
-			const FTextLocalizationResource::FEntry* SourceEntryForUpdate = nullptr;
+			const FString& NamespaceName = Namespace.Key;
+			FDisplayStringLookupTable::FKeysTable& LiveKeyTable = Namespace.Value;
 
-			// Attempt to use resources in prioritized order until we find an entry.
+			// In builds with stable keys enabled, we want to use the display string from the "clean" version of the text (if the sources match) as this is the only version that is translated
+			const FString* NamespaceNamePtr = &NamespaceName;
+#if USE_STABLE_LOCALIZATION_KEYS
+			FString DisplayNamespace;
+			if (GIsEditor)
 			{
-				const FTextLocalizationResource::FKeysTable* const UpdateKeyTable = TextLocalizationResource.Namespaces.Find(*NamespaceNamePtr);
-				const FTextLocalizationResource::FEntryArray* const UpdateEntryArray = UpdateKeyTable ? UpdateKeyTable->Find(KeyName) : nullptr;
-				const FTextLocalizationResource::FEntry* Entry = UpdateEntryArray && UpdateEntryArray->Num() ? &((*UpdateEntryArray)[0]) : nullptr;
-				if (Entry)
-				{
-					SourceEntryForUpdate = Entry;
-				}
+				DisplayNamespace = TextNamespaceUtil::StripPackageNamespace(NamespaceName);
+				NamespaceNamePtr = &DisplayNamespace;
 			}
+#endif // USE_STABLE_LOCALIZATION_KEYS
 
-			// Update the display string with the new native string
-			if (SourceEntryForUpdate && LiveStringEntry.SourceStringHash == SourceEntryForUpdate->SourceStringHash)
+			for (auto& Key : LiveKeyTable)
 			{
-				*LiveStringEntry.DisplayString = SourceEntryForUpdate->LocalizedString;
-			}
-			else
-			{
-				if (!LiveStringEntry.bIsLocalized && LiveStringEntry.DisplayString->Equals(GetAccessedStringBeforeLocLoadedErrorMsg(), ESearchCase::CaseSensitive))
+				const FString& KeyName = Key.Key;
+				FDisplayStringLookupTable::FDisplayStringEntry& LiveStringEntry = Key.Value;
+
+				const FTextLocalizationResource::FEntry* SourceEntryForUpdate = nullptr;
+
+				// Attempt to use resources in prioritized order until we find an entry.
 				{
-					*LiveStringEntry.DisplayString = FString();
+					const FTextLocalizationResource::FKeysTable* const UpdateKeyTable = TextLocalizationResource.Namespaces.Find(*NamespaceNamePtr);
+					const FTextLocalizationResource::FEntryArray* const UpdateEntryArray = UpdateKeyTable ? UpdateKeyTable->Find(KeyName) : nullptr;
+					const FTextLocalizationResource::FEntry* Entry = UpdateEntryArray && UpdateEntryArray->Num() ? &((*UpdateEntryArray)[0]) : nullptr;
+					if (Entry)
+					{
+						SourceEntryForUpdate = Entry;
+					}
 				}
 
-#if ENABLE_LOC_TESTING
-				// Restore the pre-leet state (if any)
-				if (!LiveStringEntry.NativeStringBackup.IsEmpty())
+				// Update the display string with the new native string
+				if (SourceEntryForUpdate && LiveStringEntry.SourceStringHash == SourceEntryForUpdate->SourceStringHash)
 				{
-					*LiveStringEntry.DisplayString = MoveTemp(LiveStringEntry.NativeStringBackup);
+					LiveStringEntry.bIsLocalized = true;
+					LiveStringEntry.LocResID = SourceEntryForUpdate->LocResID;
+					*LiveStringEntry.DisplayString = SourceEntryForUpdate->LocalizedString;
 				}
-#endif
-			}
+				else
+				{
+					if (!LiveStringEntry.bIsLocalized && LiveStringEntry.DisplayString->Equals(GetAccessedStringBeforeLocLoadedErrorMsg(), ESearchCase::CaseSensitive))
+					{
+						*LiveStringEntry.DisplayString = FString();
+					}
 
-			LiveStringEntry.LocResID = FTextLocalizationResourceId();
-			LiveStringEntry.bIsLocalized = false;
+#if ENABLE_LOC_TESTING
+					// Restore the pre-leet state (if any)
+					if (!LiveStringEntry.NativeStringBackup.IsEmpty())
+					{
+						LiveStringEntry.bIsLocalized = false;
+						LiveStringEntry.LocResID = FTextLocalizationResourceId();
+						*LiveStringEntry.DisplayString = MoveTemp(LiveStringEntry.NativeStringBackup);
+					}
+#endif
+				}
 
 #if ENABLE_LOC_TESTING
-			LiveStringEntry.NativeStringBackup.Reset();
+				LiveStringEntry.NativeStringBackup.Reset();
 #endif
+			}
 		}
-	}
 
-	// Add new entries
-	{
-		for (const auto& Namespace : TextLocalizationResource.Namespaces)
+		// Add new entries
 		{
-			const FString& NamespaceName = Namespace.Key;
-			const FTextLocalizationResource::FKeysTable& NewKeyTable = Namespace.Value;
-			FDisplayStringLookupTable::FKeysTable& LiveKeyTable = DisplayStringLookupTable.NamespacesTable.FindOrAdd(NamespaceName);
-			for (const auto& Key : NewKeyTable)
+			for (const auto& Namespace : TextLocalizationResource.Namespaces)
 			{
-				const FString& KeyName = Key.Key;
-				const FTextLocalizationResource::FEntryArray& NewEntryArray = Key.Value;
-				const FTextLocalizationResource::FEntry& NewEntry = NewEntryArray[0];
-
-				FDisplayStringLookupTable::FDisplayStringEntry* const LiveStringEntry = LiveKeyTable.Find(KeyName);
-				// Note: Anything we find in the table has already been updated above.
-				if (!LiveStringEntry)
+				const FString& NamespaceName = Namespace.Key;
+				const FTextLocalizationResource::FKeysTable& NewKeyTable = Namespace.Value;
+				FDisplayStringLookupTable::FKeysTable& LiveKeyTable = DisplayStringLookupTable.NamespacesTable.FindOrAdd(NamespaceName);
+				for (const auto& Key : NewKeyTable)
 				{
-					FDisplayStringLookupTable::FDisplayStringEntry NewLiveEntry(
-						false,																/*bIsLocalized*/
-						FTextLocalizationResourceId(),										/*LocResID*/
-						NewEntry.SourceStringHash,											/*SourceStringHash*/
-						MakeShared<FString, ESPMode::ThreadSafe>(NewEntry.LocalizedString)	/*String*/
+					const FString& KeyName = Key.Key;
+					const FTextLocalizationResource::FEntryArray& NewEntryArray = Key.Value;
+					const FTextLocalizationResource::FEntry& NewEntry = NewEntryArray[0];
+
+					FDisplayStringLookupTable::FDisplayStringEntry* const LiveStringEntry = LiveKeyTable.Find(KeyName);
+					// Note: Anything we find in the table has already been updated above.
+					if (!LiveStringEntry)
+					{
+						FDisplayStringLookupTable::FDisplayStringEntry NewLiveEntry(
+							true,																/*bIsLocalized*/
+							NewEntry.LocResID,													/*LocResID*/
+							NewEntry.SourceStringHash,											/*SourceStringHash*/
+							MakeShared<FString, ESPMode::ThreadSafe>(NewEntry.LocalizedString)	/*String*/
 						);
-					LiveKeyTable.Add(KeyName, NewLiveEntry);
+						LiveKeyTable.Add(KeyName, NewLiveEntry);
 
-					NamespaceKeyLookupTable.Add(NewLiveEntry.DisplayString, FNamespaceKeyEntry(NamespaceName, KeyName));
+						NamespaceKeyLookupTable.Add(NewLiveEntry.DisplayString, FNamespaceKeyEntry(NamespaceName, KeyName));
+					}
 				}
 			}
 		}
@@ -1020,103 +1061,107 @@ void FTextLocalizationManager::UpdateFromNative(const FTextLocalizationResource&
 
 void FTextLocalizationManager::UpdateFromLocalizations(TArrayView<const TSharedPtr<FTextLocalizationResource>> TextLocalizationResources)
 {
-	// Update existing localized entries/flag existing newly-unlocalized entries.
-	for (auto& Namespace : DisplayStringLookupTable.NamespacesTable)
+	// Lock while updating the tables
 	{
-		const FString& NamespaceName = Namespace.Key;
-		FDisplayStringLookupTable::FKeysTable& LiveKeyTable = Namespace.Value;
-
-		// In builds with stable keys enabled, we want to use the display string from the "clean" version of the text (if the sources match) as this is the only version that is translated
-		const FString* NamespaceNamePtr = &NamespaceName;
-#if USE_STABLE_LOCALIZATION_KEYS
-		FString DisplayNamespace;
-		if (GIsEditor)
-		{
-			DisplayNamespace = TextNamespaceUtil::StripPackageNamespace(NamespaceName);
-			NamespaceNamePtr = &DisplayNamespace;
-		}
-#endif // USE_STABLE_LOCALIZATION_KEYS
+		FScopeLock ScopeLock(&SynchronizationObject);
 
-		for(auto& Key : LiveKeyTable)
+		// Update existing localized entries/flag existing newly-unlocalized entries.
+		for (auto& Namespace : DisplayStringLookupTable.NamespacesTable)
 		{
-			const FString& KeyName = Key.Key;
-			FDisplayStringLookupTable::FDisplayStringEntry& LiveStringEntry = Key.Value;
+			const FString& NamespaceName = Namespace.Key;
+			FDisplayStringLookupTable::FKeysTable& LiveKeyTable = Namespace.Value;
 
-			const FTextLocalizationResource::FEntry* SourceEntryForUpdate = nullptr;
+			// In builds with stable keys enabled, we want to use the display string from the "clean" version of the text (if the sources match) as this is the only version that is translated
+			const FString* NamespaceNamePtr = &NamespaceName;
+#if USE_STABLE_LOCALIZATION_KEYS
+			FString DisplayNamespace;
+			if (GIsEditor)
+			{
+				DisplayNamespace = TextNamespaceUtil::StripPackageNamespace(NamespaceName);
+				NamespaceNamePtr = &DisplayNamespace;
+			}
+#endif // USE_STABLE_LOCALIZATION_KEYS
 
-			// Attempt to use resources in prioritized order until we find an entry.
-			for(const TSharedPtr<FTextLocalizationResource>& TextLocalizationResource : TextLocalizationResources)
+			for (auto& Key : LiveKeyTable)
 			{
-				check(TextLocalizationResource.IsValid());
+				const FString& KeyName = Key.Key;
+				FDisplayStringLookupTable::FDisplayStringEntry& LiveStringEntry = Key.Value;
 
-				const FTextLocalizationResource::FKeysTable* const UpdateKeyTable = TextLocalizationResource->Namespaces.Find(*NamespaceNamePtr);
-				const FTextLocalizationResource::FEntryArray* const UpdateEntryArray = UpdateKeyTable ? UpdateKeyTable->Find(KeyName) : nullptr;
-				const FTextLocalizationResource::FEntry* Entry = UpdateEntryArray && UpdateEntryArray->Num() ? &((*UpdateEntryArray)[0]) : nullptr;
-				if(Entry)
+				const FTextLocalizationResource::FEntry* SourceEntryForUpdate = nullptr;
+
+				// Attempt to use resources in prioritized order until we find an entry.
+				for (const TSharedPtr<FTextLocalizationResource>& TextLocalizationResource : TextLocalizationResources)
 				{
-					SourceEntryForUpdate = Entry;
-					break;
+					check(TextLocalizationResource.IsValid());
+
+					const FTextLocalizationResource::FKeysTable* const UpdateKeyTable = TextLocalizationResource->Namespaces.Find(*NamespaceNamePtr);
+					const FTextLocalizationResource::FEntryArray* const UpdateEntryArray = UpdateKeyTable ? UpdateKeyTable->Find(KeyName) : nullptr;
+					const FTextLocalizationResource::FEntry* Entry = UpdateEntryArray && UpdateEntryArray->Num() ? &((*UpdateEntryArray)[0]) : nullptr;
+					if (Entry)
+					{
+						SourceEntryForUpdate = Entry;
+						break;
+					}
 				}
-			}
 
-			// If the source string hashes are are the same, we can replace the display string.
-			// Otherwise, it would suggest the source string has changed and the new localization may be based off of an old source string.
-			if (SourceEntryForUpdate && LiveStringEntry.SourceStringHash == SourceEntryForUpdate->SourceStringHash)
-			{
-				LiveStringEntry.bIsLocalized = true;
-				LiveStringEntry.LocResID = SourceEntryForUpdate->LocResID;
-				*(LiveStringEntry.DisplayString) = SourceEntryForUpdate->LocalizedString;
-			}
-			else
-			{
-				if (!LiveStringEntry.bIsLocalized && LiveStringEntry.DisplayString->Equals(GetAccessedStringBeforeLocLoadedErrorMsg(), ESearchCase::CaseSensitive))
+				// If the source string hashes are are the same, we can replace the display string.
+				// Otherwise, it would suggest the source string has changed and the new localization may be based off of an old source string.
+				if (SourceEntryForUpdate && LiveStringEntry.SourceStringHash == SourceEntryForUpdate->SourceStringHash)
 				{
-					*(LiveStringEntry.DisplayString) = FString();
+					LiveStringEntry.bIsLocalized = true;
+					LiveStringEntry.LocResID = SourceEntryForUpdate->LocResID;
+					*(LiveStringEntry.DisplayString) = SourceEntryForUpdate->LocalizedString;
 				}
-
-				LiveStringEntry.bIsLocalized = false;
-				LiveStringEntry.LocResID = FTextLocalizationResourceId();
+				else
+				{
+					if (!LiveStringEntry.bIsLocalized && LiveStringEntry.DisplayString->Equals(GetAccessedStringBeforeLocLoadedErrorMsg(), ESearchCase::CaseSensitive))
+					{
+						*(LiveStringEntry.DisplayString) = FString();
+					}
 
 #if ENABLE_LOC_TESTING
-				const bool bShouldLEETIFYUnlocalizedString = FParse::Param(FCommandLine::Get(), TEXT("LEETIFYUnlocalized"));
-				if(bShouldLEETIFYUnlocalizedString )
-				{
-					FInternationalization::Leetify(*(LiveStringEntry.DisplayString));
-				}
+					static const bool bShouldLEETIFYUnlocalizedString = FParse::Param(FCommandLine::Get(), TEXT("LEETIFYUnlocalized"));
+					if (bShouldLEETIFYUnlocalizedString)
+					{
+						LiveStringEntry.bIsLocalized = false;
+						LiveStringEntry.LocResID = FTextLocalizationResourceId();
+						FInternationalization::Leetify(*LiveStringEntry.DisplayString);
+					}
 #endif
+				}
 			}
 		}
-	}
 
-	// Add new entries. 
-	for(const TSharedPtr<FTextLocalizationResource>& TextLocalizationResource : TextLocalizationResources)
-	{
-		check(TextLocalizationResource.IsValid());
-
-		for(const auto& Namespace : TextLocalizationResource->Namespaces)
+		// Add new entries. 
+		for (const TSharedPtr<FTextLocalizationResource>& TextLocalizationResource : TextLocalizationResources)
 		{
-			const FString& NamespaceName = Namespace.Key;
-			const FTextLocalizationResource::FKeysTable& NewKeyTable = Namespace.Value;
-			FDisplayStringLookupTable::FKeysTable& LiveKeyTable = DisplayStringLookupTable.NamespacesTable.FindOrAdd(NamespaceName);
-			for(const auto& Key : NewKeyTable)
-			{
-				const FString& KeyName = Key.Key;
-				const FTextLocalizationResource::FEntryArray& NewEntryArray = Key.Value;
-				const FTextLocalizationResource::FEntry& NewEntry = NewEntryArray[0];
+			check(TextLocalizationResource.IsValid());
 
-				FDisplayStringLookupTable::FDisplayStringEntry* const LiveStringEntry = LiveKeyTable.Find(KeyName);
-				// Note: Anything we find in the table has already been updated above.
-				if( !LiveStringEntry )
+			for (const auto& Namespace : TextLocalizationResource->Namespaces)
+			{
+				const FString& NamespaceName = Namespace.Key;
+				const FTextLocalizationResource::FKeysTable& NewKeyTable = Namespace.Value;
+				FDisplayStringLookupTable::FKeysTable& LiveKeyTable = DisplayStringLookupTable.NamespacesTable.FindOrAdd(NamespaceName);
+				for (const auto& Key : NewKeyTable)
 				{
-					FDisplayStringLookupTable::FDisplayStringEntry NewLiveEntry(
-						true,																/*bIsLocalized*/
-						NewEntry.LocResID,													/*LocResID*/
-						NewEntry.SourceStringHash,											/*SourceStringHash*/
-						MakeShared<FString, ESPMode::ThreadSafe>(NewEntry.LocalizedString)	/*String*/
+					const FString& KeyName = Key.Key;
+					const FTextLocalizationResource::FEntryArray& NewEntryArray = Key.Value;
+					const FTextLocalizationResource::FEntry& NewEntry = NewEntryArray[0];
+
+					FDisplayStringLookupTable::FDisplayStringEntry* const LiveStringEntry = LiveKeyTable.Find(KeyName);
+					// Note: Anything we find in the table has already been updated above.
+					if (!LiveStringEntry)
+					{
+						FDisplayStringLookupTable::FDisplayStringEntry NewLiveEntry(
+							true,																/*bIsLocalized*/
+							NewEntry.LocResID,													/*LocResID*/
+							NewEntry.SourceStringHash,											/*SourceStringHash*/
+							MakeShared<FString, ESPMode::ThreadSafe>(NewEntry.LocalizedString)	/*String*/
 						);
-					LiveKeyTable.Add( KeyName, NewLiveEntry );
+						LiveKeyTable.Add(KeyName, NewLiveEntry);
 
-					NamespaceKeyLookupTable.Add(NewLiveEntry.DisplayString, FNamespaceKeyEntry(NamespaceName, KeyName));
+						NamespaceKeyLookupTable.Add(NewLiveEntry.DisplayString, FNamespaceKeyEntry(NamespaceName, KeyName));
+					}
 				}
 			}
 		}
@@ -1127,6 +1172,8 @@ void FTextLocalizationManager::UpdateFromLocalizations(TArrayView<const TSharedP
 
 void FTextLocalizationManager::DirtyLocalRevisionForDisplayString(const FTextDisplayStringRef& InDisplayString)
 {
+	FScopeLock ScopeLock(&SynchronizationObject);
+
 	uint16* FoundLocalRevision = LocalTextRevisions.Find(InDisplayString);
 	if (FoundLocalRevision)
 	{
@@ -1140,8 +1187,13 @@ void FTextLocalizationManager::DirtyLocalRevisionForDisplayString(const FTextDis
 
 void FTextLocalizationManager::DirtyTextRevision()
 {
-	while (++TextRevisionCounter == 0) {} // Zero is special, don't allow an overflow to stay at zero
-	LocalTextRevisions.Empty();
+	// Lock while updating the data
+	{
+		FScopeLock ScopeLock(&SynchronizationObject);
+		while (++TextRevisionCounter == 0) {} // Zero is special, don't allow an overflow to stay at zero
+		LocalTextRevisions.Empty();
+	}
+
 	OnTextRevisionChangedEvent.Broadcast();
 }
 
