Fixed the editor failing to load the native game localization data during init

diff --git a/Engine/Source/Runtime/Core/Private/Internationalization/LocalizationResourceTextSource.cpp b/Engine/Source/Runtime/Core/Private/Internationalization/LocalizationResourceTextSource.cpp
--- a/Engine/Source/Runtime/Core/Private/Internationalization/LocalizationResourceTextSource.cpp
+++ b/Engine/Source/Runtime/Core/Private/Internationalization/LocalizationResourceTextSource.cpp
@@ -45,8 +45,13 @@ void FLocalizationResourceTextSource::GetLocalizedCultureNames(const ELocalizati
 void FLocalizationResourceTextSource::LoadLocalizedResources(const ELocalizationLoadFlags InLoadFlags, TArrayView<const FString> InPrioritizedCultures, FTextLocalizationResource& InOutNativeResource, FTextLocalizationResources& InOutLocalizedResources)
 {
 	// Collect the localization paths to load from.
+	TArray<FString> GameNativePaths;
 	TArray<FString> GameLocalizationPaths;
-	if (ShouldLoadGame(InLoadFlags))
+	if (ShouldLoadNativeGameData(InLoadFlags))
+	{
+		GameNativePaths += FPaths::GetGameLocalizationPaths();
+	}
+	else if (ShouldLoadGame(InLoadFlags))
 	{
 		GameLocalizationPaths += FPaths::GetGameLocalizationPaths();
 	}
@@ -88,10 +93,7 @@ void FLocalizationResourceTextSource::LoadLocalizedResources(const ELocalization
 	}
 
 	TArray<FString> PrioritizedLocalizationPaths;
-	if (ShouldLoadLocalizedGameData(InLoadFlags))
-	{
-		PrioritizedLocalizationPaths += GameLocalizationPaths;
-	}
+	PrioritizedLocalizationPaths += GameLocalizationPaths;
 	PrioritizedLocalizationPaths += EditorLocalizationPaths;
 	PrioritizedLocalizationPaths += EngineLocalizationPaths;
 	PrioritizedLocalizationPaths += AdditionalLocalizationPaths;
@@ -136,14 +138,14 @@ void FLocalizationResourceTextSource::LoadLocalizedResources(const ELocalization
 		}
 	}
 
-	// The editor cheats and loads the native language's localizations.
-	if (!ShouldLoadLocalizedGameData(InLoadFlags))
+	// The editor cheats and loads the games native localizations.
+	if (ShouldLoadNativeGameData(InLoadFlags) && GameNativePaths.Num() > 0)
 	{
 		const FString NativeGameCulture = TextLocalizationResourceUtil::GetNativeProjectCultureName();
-		if (!NativeGameCulture.IsEmpty() && GameLocalizationPaths.Num() > 0)
+		if (!NativeGameCulture.IsEmpty())
 		{
 			TSharedRef<FTextLocalizationResource> TextLocalizationResource = InOutLocalizedResources.EnsureResource(InPrioritizedCultures[0]);
-			for (const FString& LocalizationPath : GameLocalizationPaths)
+			for (const FString& LocalizationPath : GameNativePaths)
 			{
 				const FString CulturePath = LocalizationPath / NativeGameCulture;
 				TextLocalizationResource->LoadFromDirectory(CulturePath);
diff --git a/Engine/Source/Runtime/Core/Private/Internationalization/PolyglotTextSource.cpp b/Engine/Source/Runtime/Core/Private/Internationalization/PolyglotTextSource.cpp
--- a/Engine/Source/Runtime/Core/Private/Internationalization/PolyglotTextSource.cpp
+++ b/Engine/Source/Runtime/Core/Private/Internationalization/PolyglotTextSource.cpp
@@ -107,7 +107,7 @@ void FPolyglotTextSource::LoadLocalizedResources(const ELocalizationLoadFlags In
 
 		if (ShouldLoadLocalizedText(PolyglotTextData))
 		{
-			if (PolyglotTextData.GetCategory() == ELocalizedTextSourceCategory::Game && !ShouldLoadLocalizedGameData(InLoadFlags))
+			if (PolyglotTextData.GetCategory() == ELocalizedTextSourceCategory::Game && ShouldLoadNativeGameData(InLoadFlags))
 			{
 				// The editor cheats and loads the native language's localizations for game data.
 				TSharedRef<FTextLocalizationResource> TextLocalizationResource = InOutLocalizedResources.EnsureResource(InPrioritizedCultures[0]);
diff --git a/Engine/Source/Runtime/Core/Public/Internationalization/ILocalizedTextSource.h b/Engine/Source/Runtime/Core/Public/Internationalization/ILocalizedTextSource.h
--- a/Engine/Source/Runtime/Core/Public/Internationalization/ILocalizedTextSource.h
+++ b/Engine/Source/Runtime/Core/Public/Internationalization/ILocalizedTextSource.h
@@ -74,9 +74,11 @@ public:
 		return EnumHasAnyFlags(InLoadFlags, ELocalizationLoadFlags::Additional);
 	}
 
-	/** Should we load localized game data based on the given load flags and environment? */
-	static FORCEINLINE bool ShouldLoadLocalizedGameData(const ELocalizationLoadFlags InLoadFlags)
+	/** Should we load native game data based on the given load flags and environment? */
+	static FORCEINLINE bool ShouldLoadNativeGameData(const ELocalizationLoadFlags InLoadFlags)
 	{
-		return !GIsEditor || EnumHasAnyFlags(InLoadFlags, ELocalizationLoadFlags::ForceLocalizedGame);
+		// The editor loads native game data by default to prevent authoring issues
+		// It will load localized data only if the request is forced (eg, when entering game localization preview mode)
+		return GIsEditor && !EnumHasAnyFlags(InLoadFlags, ELocalizationLoadFlags::ForceLocalizedGame);
 	}
 };
