Copying from Dev-Enterprise CL 4364256
Fix typo in FGenericPlatformMemory::DumpStats to properly report the free virtual memory.

diff --git a/Engine/Source/Runtime/Core/Private/GenericPlatform/GenericPlatformMemory.cpp b/Engine/Source/Runtime/Core/Private/GenericPlatform/GenericPlatformMemory.cpp
--- a/Engine/Source/Runtime/Core/Private/GenericPlatform/GenericPlatformMemory.cpp
+++ b/Engine/Source/Runtime/Core/Private/GenericPlatform/GenericPlatformMemory.cpp
@@ -551,7 +551,7 @@ void FGenericPlatformMemory::DumpStats( class FOutputDevice& Ar )
 	Ar.CategorizedLogf(CategoryName, ELogVerbosity::Log, TEXT("Physical Memory: %.2f MB used,  %.2f MB free, %.2f MB total"), 
 		(MemoryStats.TotalPhysical - MemoryStats.AvailablePhysical)*InvMB, MemoryStats.AvailablePhysical*InvMB, MemoryStats.TotalPhysical*InvMB);
 	Ar.CategorizedLogf(CategoryName, ELogVerbosity::Log, TEXT("Virtual Memory: %.2f MB used,  %.2f MB free, %.2f MB total"), 
-		(MemoryStats.TotalVirtual - MemoryStats.AvailableVirtual)*InvMB, MemoryStats.AvailablePhysical*InvMB, MemoryStats.TotalVirtual*InvMB);
+		(MemoryStats.TotalVirtual - MemoryStats.AvailableVirtual)*InvMB, MemoryStats.AvailableVirtual*InvMB, MemoryStats.TotalVirtual*InvMB);
 
 }
 
