Fix incorrect usage of a CFDictionary passed into IOServiceGetMatchingServices which releases it for us. #rb none #jira UE-62235

diff --git a/Engine/Source/Runtime/Core/Private/Mac/MacPlatformMisc.cpp b/Engine/Source/Runtime/Core/Private/Mac/MacPlatformMisc.cpp
--- a/Engine/Source/Runtime/Core/Private/Mac/MacPlatformMisc.cpp
+++ b/Engine/Source/Runtime/Core/Private/Mac/MacPlatformMisc.cpp
@@ -1001,7 +1001,6 @@ public:
 						}
 						IOObjectRelease(ServiceEntry);
 					}
-					CFRelease(MatchDictionary);
 				}
 				break;
 			}
