Don't force LocRes files to always use Unicode strings

The automatic Unicode detection works just fine for them, and changing this can produce considerably smaller (up-to 50% smaller!) LocRes files

diff --git a/Engine/Source/Runtime/Core/Private/Internationalization/TextLocalizationResource.cpp b/Engine/Source/Runtime/Core/Private/Internationalization/TextLocalizationResource.cpp
--- a/Engine/Source/Runtime/Core/Private/Internationalization/TextLocalizationResource.cpp
+++ b/Engine/Source/Runtime/Core/Private/Internationalization/TextLocalizationResource.cpp
@@ -143,8 +143,6 @@ bool FTextLocalizationResource::LoadFromFile(const FString& FilePath)
 
 bool FTextLocalizationResource::LoadFromArchive(FArchive& Archive, const FTextLocalizationResourceId& LocResID)
 {
-	Archive.SetForceUnicode(true);
-
 	// Read magic number
 	FGuid MagicNumber;
 	
@@ -259,8 +257,6 @@ bool FTextLocalizationResource::SaveToFile(const FString& FilePath)
 
 bool FTextLocalizationResource::SaveToArchive(FArchive& Archive, const FTextLocalizationResourceId& LocResID)
 {
-	Archive.SetForceUnicode(true);
-
 	// Write the header
 	{
 		FGuid MagicNumber = FTextLocalizationResourceVersion::LocResMagic;
