When we detect a hang and send an ensure crash report, include the callstack of the thread that hung.

diff --git a/Engine/Source/Runtime/Core/Private/GenericPlatform/GenericPlatformCrashContext.cpp b/Engine/Source/Runtime/Core/Private/GenericPlatform/GenericPlatformCrashContext.cpp
--- a/Engine/Source/Runtime/Core/Private/GenericPlatform/GenericPlatformCrashContext.cpp
+++ b/Engine/Source/Runtime/Core/Private/GenericPlatform/GenericPlatformCrashContext.cpp
@@ -638,6 +638,13 @@ void FGenericCrashContext::AddPlugin(const FString& PluginDesc)
 }
 
 FORCENOINLINE void FGenericCrashContext::CapturePortableCallStack(int32 NumStackFramesToIgnore, void* Context)
+{
+	const int32 MaxDepth = 100;
+	TArray<FProgramCounterSymbolInfo> Stack = FPlatformStackWalk::GetStack(NumStackFramesToIgnore, MaxDepth, Context);
+	return SetPortableCallStack(NumStackFramesToIgnore, Stack);
+}
+
+void FGenericCrashContext::SetPortableCallStack(int32 NumStackFramesToIgnore, const TArray<FProgramCounterSymbolInfo>& Stack)
 {
 	uint32 ModuleEntries = (uint32)FPlatformStackWalk::GetProcessModuleCount();
 
@@ -647,9 +654,6 @@ FORCENOINLINE void FGenericCrashContext::CapturePortableCallStack(int32 NumStack
 		ProcessModules.AddUninitialized(ModuleEntries);
 		FPlatformStackWalk::GetProcessModuleSignatures(ProcessModules.GetData(), ProcessModules.Max());
 
-		const int32 MaxDepth = 100;
-		TArray<FProgramCounterSymbolInfo> Stack = FPlatformStackWalk::GetStack(NumStackFramesToIgnore, MaxDepth, Context);
-
 		TMap<FString, uint64> ImageBases;
 		int32 ModuleIndex = 0;
 		for (TArray<FProgramCounterSymbolInfo>::TConstIterator Itr(Stack); Itr; ++Itr)
diff --git a/Engine/Source/Runtime/Core/Private/HAL/ThreadHeartBeat.cpp b/Engine/Source/Runtime/Core/Private/HAL/ThreadHeartBeat.cpp
--- a/Engine/Source/Runtime/Core/Private/HAL/ThreadHeartBeat.cpp
+++ b/Engine/Source/Runtime/Core/Private/HAL/ThreadHeartBeat.cpp
@@ -183,22 +183,26 @@ void FORCENOINLINE FThreadHeartBeat::OnHang(double HangDuration, uint32 ThreadTh
 
 #else // MINIMAL_FATAL_HANG_DETECTION == 0
 
-	const SIZE_T StackTraceSize = 65535;
-	ANSICHAR* StackTrace = (ANSICHAR*)GMalloc->Malloc(StackTraceSize);
-	StackTrace[0] = 0;
-	// Walk the stack and dump it to the allocated memory. This process usually allocates a lot of memory.
-	FPlatformStackWalk::ThreadStackWalkAndDump(StackTrace, StackTraceSize, 0, ThreadThatHung);
+	// Capture the stack in the thread that hung
+	static const int32 MaxStackFrames = 100;
+	uint64 StackFrames[MaxStackFrames];
+	int32 NumStackFrames = FPlatformStackWalk::CaptureThreadStackBackTrace(ThreadThatHung, StackFrames, MaxStackFrames);
 
 	// First verify we're not reporting the same hang over and over again
-	uint32 CallstackCRC = FCrc::StrCrc32(StackTrace);
+	uint32 CallstackCRC = FCrc::MemCrc32(StackFrames, sizeof(StackFrames[0]) * NumStackFrames);
 	if (CallstackCRC != LastHangCallstackCRC || ThreadThatHung != LastHungThreadId)
 	{
 		LastHangCallstackCRC = CallstackCRC;
 		LastHungThreadId = ThreadThatHung;
 
-		FString StackTraceText(StackTrace);
+		// Convert the stack trace to text
 		TArray<FString> StackLines;
-		StackTraceText.ParseIntoArrayLines(StackLines);
+		for (int32 Idx = 0; Idx < NumStackFrames; Idx++)
+		{
+			ANSICHAR Buffer[1024];
+			FPlatformStackWalk::ProgramCounterToHumanReadableString(Idx, StackFrames[Idx], Buffer, sizeof(Buffer));
+			StackLines.Add(Buffer);
+		}
 
 		// Dump the callstack and the thread name to log
 		FString ThreadName(ThreadThatHung == GGameThreadId ? TEXT("GameThread") : FThreadManager::Get().GetThreadName(ThreadThatHung));
@@ -207,9 +211,9 @@ void FORCENOINLINE FThreadHeartBeat::OnHang(double HangDuration, uint32 ThreadTh
 			ThreadName = FString::Printf(TEXT("unknown thread (%u)"), ThreadThatHung);
 		}
 		UE_LOG(LogCore, Error, TEXT("Hang detected on %s (thread hasn't sent a heartbeat for %.2f seconds):"), *ThreadName, HangDuration);
-		for (FString& StackLine : StackLines)
+		for (int32 Idx = 0; Idx < StackLines.Num(); Idx++)
 		{
-			UE_LOG(LogCore, Error, TEXT("  %s"), *StackLine);
+			UE_LOG(LogCore, Error, TEXT("  %s"), *StackLines[Idx]);
 		}
 
 		// Assert (on the current thread unfortunately) with a trimmed stack.
@@ -231,15 +235,20 @@ void FORCENOINLINE FThreadHeartBeat::OnHang(double HangDuration, uint32 ThreadTh
 		GLog->PanicFlushThreadedLogs();
 		// GErrorMessage here is very unfortunate but it's used internally by the crash context code.
 		FCString::Strcpy(GErrorMessage, ARRAY_COUNT(GErrorMessage), *ErrorMessage);
+
 		// Skip macros and FDebug, we always want this to fire
-		NewReportEnsure(*ErrorMessage);
+		TArray<FProgramCounterSymbolInfo> Stack;
+		for (int32 Idx = 0; Idx < NumStackFrames; Idx++)
+		{
+			FPlatformStackWalk::ProgramCounterToSymbolInfo(StackFrames[Idx], Stack.AddDefaulted_GetRef());
+		}
+		ReportHang(*ErrorMessage, Stack);
+
 		GErrorMessage[0] = '\0';
 #endif // PLATFORM_DESKTOP
 
 #endif // UE_ASSERT_ON_HANG == 0
 	}
-
-	GMalloc->Free(StackTrace);
 #endif // MINIMAL_FATAL_HANG_DETECTION == 0
 }
 
diff --git a/Engine/Source/Runtime/Core/Private/Mac/MacPlatformMisc.cpp b/Engine/Source/Runtime/Core/Private/Mac/MacPlatformMisc.cpp
--- a/Engine/Source/Runtime/Core/Private/Mac/MacPlatformMisc.cpp
+++ b/Engine/Source/Runtime/Core/Private/Mac/MacPlatformMisc.cpp
@@ -1876,6 +1876,24 @@ void NewReportEnsure( const TCHAR* ErrorMessage )
 	bReentranceGuard = false;
 	EnsureLock.Unlock();
 }
+
+void ReportHang(const TCHAR* ErrorMessage, const TArray<FProgramCounterSymbolInfo>& Stack)
+{
+	EnsureLock.Lock();
+	if (!bReentranceGuard && FMacApplicationInfo::CrashReporter != nil)
+	{
+		bReentranceGuard = true;
+
+		const bool bIsEnsure = true;
+		FMacCrashContext EnsureContext(bIsEnsure);
+		EnsureContext.SetPortableCallStack(0, Stack);
+		EnsureContext.GenerateEnsureInfoAndLaunchReporter(true);
+
+		bReentranceGuard = false;
+	}
+	EnsureLock.Unlock();
+}
+
 typedef NSArray* (*MTLCopyAllDevices)(void);
 
 bool FMacPlatformMisc::HasPlatformFeature(const TCHAR* FeatureName)
diff --git a/Engine/Source/Runtime/Core/Private/Unix/UnixPlatformStackWalk.cpp b/Engine/Source/Runtime/Core/Private/Unix/UnixPlatformStackWalk.cpp
--- a/Engine/Source/Runtime/Core/Private/Unix/UnixPlatformStackWalk.cpp
+++ b/Engine/Source/Runtime/Core/Private/Unix/UnixPlatformStackWalk.cpp
@@ -635,3 +635,20 @@ void NewReportEnsure(const TCHAR* ErrorMessage)
 	bReentranceGuard = false;
 	EnsureLock.Unlock();
 }
+
+void ReportHang(const TCHAR* ErrorMessage, const TArray<FProgramCounterSymbolInfo>& Stack)
+{
+	EnsureLock.Lock();
+	if (!bReentranceGuard)
+	{
+		bReentranceGuard = true;
+
+		const bool bIsEnsure = true;
+		FUnixCrashContext EnsureContext(bIsEnsure);
+		EnsureContext.SetPortableCallStack(0, Stack);
+		EnsureContext.GenerateCrashInfoAndLaunchReporter(true);
+
+		bReentranceGuard = false;
+	}
+	EnsureLock.Unlock();
+}
diff --git a/Engine/Source/Runtime/Core/Private/Windows/WindowsPlatformCrashContext.cpp b/Engine/Source/Runtime/Core/Private/Windows/WindowsPlatformCrashContext.cpp
--- a/Engine/Source/Runtime/Core/Private/Windows/WindowsPlatformCrashContext.cpp
+++ b/Engine/Source/Runtime/Core/Private/Windows/WindowsPlatformCrashContext.cpp
@@ -313,6 +313,17 @@ int32 ReportEnsureUsingCrashReportClient(HANDLE Thread, EXCEPTION_POINTERS* Exce
 }
 #endif
 
+void ReportHang(const TCHAR* ErrorMessage, const TArray<FProgramCounterSymbolInfo>& Stack)
+{
+	const bool bIsEnsure = true;
+
+	FWindowsPlatformCrashContext CrashContext(bIsEnsure);
+	CrashContext.SetPortableCallStack(0, Stack);
+
+	EErrorReportUI ReportUI = IsInteractiveEnsureMode() ? EErrorReportUI::ShowDialog : EErrorReportUI::ReportInUnattendedMode;
+	ReportCrashUsingCrashReportClient(CrashContext, nullptr, ErrorMessage, ReportUI, bIsEnsure);
+}
+
 // #CrashReport: 2015-05-28 This should be named EngineEnsureHandler
 /** 
  * Report an ensure to the crash reporting system
diff --git a/Engine/Source/Runtime/Core/Public/GenericPlatform/GenericPlatformCrashContext.h b/Engine/Source/Runtime/Core/Public/GenericPlatform/GenericPlatformCrashContext.h
--- a/Engine/Source/Runtime/Core/Public/GenericPlatform/GenericPlatformCrashContext.h
+++ b/Engine/Source/Runtime/Core/Public/GenericPlatform/GenericPlatformCrashContext.h
@@ -6,6 +6,8 @@
 #include "HAL/PlatformMemory.h"
 #include "Containers/UnrealString.h"
 
+struct FProgramCounterSymbolInfo;
+
 /** 
  * Symbol information associated with a program counter. 
  * FString version.
@@ -205,6 +207,9 @@ public:
 	
 	/** Generate raw call stack for crash report (image base + offset) */
 	void CapturePortableCallStack(int32 NumStackFramesToIgnore, void* Context);
+	
+	/** Sets the portable callstack to a specified stack */
+	void SetPortableCallStack(int32 NumStackFramesToIgnore, const TArray<FProgramCounterSymbolInfo>& Stack);
 
 	/**
 	 * @return whether this crash is a non-crash event
diff --git a/Engine/Source/Runtime/Core/Public/HAL/ExceptionHandling.h b/Engine/Source/Runtime/Core/Public/HAL/ExceptionHandling.h
--- a/Engine/Source/Runtime/Core/Public/HAL/ExceptionHandling.h
+++ b/Engine/Source/Runtime/Core/Public/HAL/ExceptionHandling.h
@@ -4,6 +4,8 @@
 
 #include "CoreTypes.h"
 
+struct FProgramCounterSymbolInfo;
+
 /** Whether we should generate crash reports even if the debugger is attached. */
 extern CORE_API bool GAlwaysReportCrash;
 
@@ -22,6 +24,7 @@ extern CORE_API TCHAR MiniDumpFilenameW[1024];
 // #CrashReport: 2014-10-09 These methods are specific to windows, remove from here.
 extern CORE_API int32 ReportCrash( Windows::LPEXCEPTION_POINTERS ExceptionInfo );
 extern CORE_API void NewReportEnsure( const TCHAR* ErrorMessage );
+extern CORE_API void ReportHang(const TCHAR*, const TArray<FProgramCounterSymbolInfo>& Stack);
 #elif PLATFORM_XBOXONE
 #include "XboxOne/XboxOneSystemIncludes.h"
 // #CrashReport: 2014-10-09 Should be move to another file
@@ -32,8 +35,10 @@ extern CORE_API void NewReportEnsure( const TCHAR* ErrorMessage );
 #include <signal.h>
 extern CORE_API int32 ReportCrash( ucontext_t *Context, int32 Signal, struct __siginfo* Info );
 extern CORE_API void NewReportEnsure( const TCHAR* ErrorMessage );
+extern CORE_API void ReportHang(const TCHAR*, const TArray<FProgramCounterSymbolInfo>& Stack);
 #elif PLATFORM_UNIX
 extern CORE_API void NewReportEnsure( const TCHAR* ErrorMessage );
+extern CORE_API void ReportHang(const TCHAR*, const TArray<FProgramCounterSymbolInfo>& Stack);
 #endif
 
 enum class ECrashType
