Add missing include causing CIS error.

diff --git a/Engine/Source/Runtime/Core/Public/HAL/ExceptionHandling.h b/Engine/Source/Runtime/Core/Public/HAL/ExceptionHandling.h
--- a/Engine/Source/Runtime/Core/Public/HAL/ExceptionHandling.h
+++ b/Engine/Source/Runtime/Core/Public/HAL/ExceptionHandling.h
@@ -3,6 +3,7 @@
 #pragma once
 
 #include "CoreTypes.h"
+#include "Containers/ContainersFwd.h"
 
 struct FProgramCounterSymbolInfo;
 
