Fix CIS error.

diff --git a/Engine/Source/Runtime/Core/Private/Mac/MacPlatformMisc.cpp b/Engine/Source/Runtime/Core/Private/Mac/MacPlatformMisc.cpp
--- a/Engine/Source/Runtime/Core/Private/Mac/MacPlatformMisc.cpp
+++ b/Engine/Source/Runtime/Core/Private/Mac/MacPlatformMisc.cpp
@@ -1887,7 +1887,7 @@ void ReportHang(const TCHAR* ErrorMessage, const TArray<FProgramCounterSymbolInf
 		const bool bIsEnsure = true;
 		FMacCrashContext EnsureContext(bIsEnsure);
 		EnsureContext.SetPortableCallStack(0, Stack);
-		EnsureContext.GenerateEnsureInfoAndLaunchReporter(true);
+		EnsureContext.GenerateEnsureInfoAndLaunchReporter();
 
 		bReentranceGuard = false;
 	}
