Lumin needs task threads to avoid saturating a single core. Reverting to previous task thread behavior for lumin.

diff --git a/Engine/Source/Runtime/Core/Private/Misc/App.cpp b/Engine/Source/Runtime/Core/Private/Misc/App.cpp
--- a/Engine/Source/Runtime/Core/Private/Misc/App.cpp
+++ b/Engine/Source/Runtime/Core/Private/Misc/App.cpp
@@ -240,6 +240,13 @@ bool FApp::IsUnattended() // @todo clang: Workaround for missing symbol export
 #endif
 
 #if HAVE_RUNTIME_THREADING_SWITCHES
+
+#if PLATFORM_LUMIN
+#define MIN_CORE_COUNT 2
+#else
+#define MIN_CORE_COUNT 4
+#endif
+
 bool FApp::ShouldUseThreadingForPerformance()
 {
 	static bool OnlyOneThread = 
@@ -247,11 +254,12 @@ bool FApp::ShouldUseThreadingForPerformance()
 		FParse::Param(FCommandLine::Get(), TEXT("noperfthreads")) ||
 		IsRunningDedicatedServer() ||
 		!FPlatformProcess::SupportsMultithreading() ||
-		FPlatformMisc::NumberOfCoresIncludingHyperthreads() < 4;
+		FPlatformMisc::NumberOfCoresIncludingHyperthreads() < MIN_CORE_COUNT;
 	return !OnlyOneThread;
 }
-#endif // HAVE_RUNTIME_THREADING_SWITCHES
+#undef MIN_CORE_COUNT
 
+#endif // HAVE_RUNTIME_THREADING_SWITCHES
 
 static bool GUnfocusedVolumeMultiplierInitialised = false;
 float FApp::GetUnfocusedVolumeMultiplier()
