Fix CIS (UE-64925).

diff --git a/Engine/Source/Runtime/Core/Private/Unix/UnixPlatformCrashContext.cpp b/Engine/Source/Runtime/Core/Private/Unix/UnixPlatformCrashContext.cpp
--- a/Engine/Source/Runtime/Core/Private/Unix/UnixPlatformCrashContext.cpp
+++ b/Engine/Source/Runtime/Core/Private/Unix/UnixPlatformCrashContext.cpp
@@ -647,7 +647,7 @@ void PlatformCrashHandler(int32 Signal, siginfo_t* Info, void* Context)
 	CrashContext.FirstCrashHandlerFrame = static_cast<uint64*>(__builtin_return_address(0));
 
 	// This will ungrab cursor/keyboard and bring down any pointer barriers which will be stuck on when opening the CRC
-	FLinuxPlatformMisc::UngrabAllInput();
+	FPlatformMisc::UngrabAllInput();
 
 	if (GCrashHandlerPointer)
 	{
