[NO NON-PII COMMIT MESSAGE]

diff --git a/Engine/Source/Runtime/Core/Public/Math/TransformNonVectorized.h b/Engine/Source/Runtime/Core/Public/Math/TransformNonVectorized.h
--- a/Engine/Source/Runtime/Core/Public/Math/TransformNonVectorized.h
+++ b/Engine/Source/Runtime/Core/Public/Math/TransformNonVectorized.h
@@ -123,11 +123,12 @@ public:
 		// Note: This can be used to track down initialization issues with bone transform arrays; but it will
 		// cause issues with transient fields such as RootMotionDelta that get initialized to 0 by default
 #if ENABLE_NAN_DIAGNOSTIC
+		// It is unsafe to call normal constructors here as they do their own NaN checks, so would always hit a false positive
 		float qnan = FMath::Log2(-5.3f);
 		check(FMath::IsNaN(qnan));
-		Translation = FVector(qnan, qnan, qnan);
-		Rotation = FQuat(qnan, qnan, qnan, qnan);
-		Scale3D = FVector(qnan, qnan, qnan);
+		Translation.X = Translation.Y = Translation.Z = qnan;
+		Rotation.X = Rotation.Y = Rotation.Z = Rotation.W = qnan;
+		Scale3D.X = Scale3D.Y = Scale3D.Z = qnan;
 #endif
 	}
 
@@ -462,7 +463,6 @@ public:
 
 		Scale3D += Atom.Scale3D;
 
-		DiagnosticCheckNaN_All();
 		return *this;
 	}
 
@@ -479,7 +479,6 @@ public:
 		Rotation.Z *= Mult;
 		Rotation.W *= Mult;
 		Scale3D *= Mult;
-		DiagnosticCheckNaN_All();
 
 		return *this;
 	}
@@ -1070,8 +1069,6 @@ public:
 		FinalAtom.Translation += SourceAtom.Translation;
 		FinalAtom.Scale3D *= (DefaultScale + SourceAtom.Scale3D);
 
-		FinalAtom.DiagnosticCheckNaN_All();
-
 		checkSlow(FinalAtom.IsRotationNormalized());
 	}
 
@@ -1303,7 +1300,6 @@ FORCEINLINE void FTransform::Multiply(FTransform* OutTransform, const FTransform
 
 	// we do not support matrix transform when non-uniform
 	// that was removed at rev 21 with UE4
-	OutTransform->DiagnosticCheckNaN_All();
 }
 /**
 * Apply Scale to this transform
