Add USING_UNDEFINED_BEHAVIOR_SANITIZER define.

diff --git a/Engine/Source/Runtime/Core/Public/Clang/ClangPlatformCodeAnalysis.h b/Engine/Source/Runtime/Core/Public/Clang/ClangPlatformCodeAnalysis.h
--- a/Engine/Source/Runtime/Core/Public/Clang/ClangPlatformCodeAnalysis.h
+++ b/Engine/Source/Runtime/Core/Public/Clang/ClangPlatformCodeAnalysis.h
@@ -96,6 +96,11 @@
 	#define USING_THREAD_SANITISER 0
 #endif
 
+// Clang does not expose __has_feature(undefined_behavior_sanitizer) so this define comes directly from UBT.
+#ifndef USING_UNDEFINED_BEHAVIOR_SANITISER
+	#define USING_UNDEFINED_BEHAVIOR_SANITISER 0
+#endif
+
 #if USING_THREAD_SANITISER
 
 	// Function attribute to disable thread sanitiser validation on specific functions that assume non-atomic load/stores are implicitly atomic
