Fixed the implementation of MakeVectorRegisterInt in DirectX math. Was previously making a float4 register.

diff --git a/Engine/Source/Runtime/Core/Public/Math/UnrealMathDirectX.h b/Engine/Source/Runtime/Core/Public/Math/UnrealMathDirectX.h
--- a/Engine/Source/Runtime/Core/Public/Math/UnrealMathDirectX.h
+++ b/Engine/Source/Runtime/Core/Public/Math/UnrealMathDirectX.h
@@ -64,7 +64,7 @@ FORCEINLINE VectorRegister MakeVectorRegister( float X, float Y, float Z, float
  */
 FORCEINLINE VectorRegisterInt MakeVectorRegisterInt(int32 X, int32 Y, int32 Z, int32 W)
 {
-	return _mm_castps_si128(DirectX::XMVectorSet(X, Y, Z, W));
+	return _mm_castps_si128(DirectX::XMVectorSetInt(X, Y, Z, W));
 }
 
 /*=============================================================================
