Linux: Open can return 0 make sure we dont skip reading if thats the case

diff --git a/Engine/Source/Runtime/Core/Private/Unix/UnixPlatformStackWalk.cpp b/Engine/Source/Runtime/Core/Private/Unix/UnixPlatformStackWalk.cpp
--- a/Engine/Source/Runtime/Core/Private/Unix/UnixPlatformStackWalk.cpp
+++ b/Engine/Source/Runtime/Core/Private/Unix/UnixPlatformStackWalk.cpp
@@ -45,7 +45,7 @@ namespace
 			: SymbolFileFD(open(Path, O_RDONLY)),
 			  StartOffset(sizeof(RecordsHeader))
 		{
-			if (SymbolFileFD)
+			if (SymbolFileFD != -1)
 			{
 				// TODO check for EINTR
 				read(SymbolFileFD, static_cast<void*>(&RecordCount), sizeof(RecordsHeader));
