Fix compile error for Linux servers in Test config.

diff --git a/Engine/Source/Runtime/Core/Private/Unix/UnixPlatformMisc.cpp b/Engine/Source/Runtime/Core/Private/Unix/UnixPlatformMisc.cpp
--- a/Engine/Source/Runtime/Core/Private/Unix/UnixPlatformMisc.cpp
+++ b/Engine/Source/Runtime/Core/Private/Unix/UnixPlatformMisc.cpp
@@ -914,7 +914,7 @@ void FUnixPlatformMisc::CustomNamedStat(const TCHAR* Text, float Value, const TC
 
 void FUnixPlatformMisc::CustomNamedStat(const ANSICHAR* Text, float Value, const ANSICHAR* Graph, const ANSICHAR* Unit)
 {
-	FRAMEPRO_DYNAMIC_CUSTOM_STAT(TCHAR_TO_WCHAR(Text), Value, TCHAR_TO_WCHAR(Graph), TCHAR_TO_WCHAR(Unit));
+	FRAMEPRO_DYNAMIC_CUSTOM_STAT(TCHAR_TO_WCHAR(ANSI_TO_TCHAR(Text)), Value, TCHAR_TO_WCHAR(ANSI_TO_TCHAR(Graph)), TCHAR_TO_WCHAR(ANSI_TO_TCHAR(Unit)));
 }
 #endif
 
