Slow task dialog respects ESlowTaskVisibility setting

diff --git a/Engine/Source/Runtime/Core/Private/Misc/SlowTask.cpp b/Engine/Source/Runtime/Core/Private/Misc/SlowTask.cpp
--- a/Engine/Source/Runtime/Core/Private/Misc/SlowTask.cpp
+++ b/Engine/Source/Runtime/Core/Private/Misc/SlowTask.cpp
@@ -105,7 +105,7 @@ FText FSlowTask::GetCurrentMessage() const
 void FSlowTask::MakeDialog(bool bShowCancelButton, bool bAllowInPIE)
 {
 	const bool bIsDisabledByPIE = Context.IsPlayingInEditor() && !bAllowInPIE;
-	const bool bIsDialogAllowed = bEnabled && !GIsSilent && !bIsDisabledByPIE && !IsRunningCommandlet() && IsInGameThread();
+	const bool bIsDialogAllowed = bEnabled && !GIsSilent && !bIsDisabledByPIE && !IsRunningCommandlet() && IsInGameThread() && Visibility != ESlowTaskVisibility::Invisible;
 	if (!GIsSlowTask && bIsDialogAllowed)
 	{
 		Context.StartSlowTask(GetCurrentMessage(), bShowCancelButton);
