Fix handling of dynamic width modifier for char16_t strings (UE-64798).

- Fixes crash when fracturing a DestructibleMesh in Linux editor (UE-64798), because bone names are generated with that formatting.

diff --git a/Engine/Source/Programs/TestPAL/Private/Main.cpp b/Engine/Source/Programs/TestPAL/Private/Main.cpp
--- a/Engine/Source/Programs/TestPAL/Private/Main.cpp
+++ b/Engine/Source/Programs/TestPAL/Private/Main.cpp
@@ -355,6 +355,14 @@ int32 StringPrecisionTest(const TCHAR* CommandLine)
 	UE_LOG(LogTestPAL, Display, TEXT("Begining of the line %*s"), Indent, *TestString);
 	UE_LOG(LogTestPAL, Display, TEXT("%*s end of the line"), Indent, *TestString);
 
+	int Width = 2;
+	for (uint32 Idx = 0; Idx < 10; ++Idx)
+	{
+		UE_LOG(LogTestPAL, Display, TEXT("DynSize: %d SignedDynFormat%0*d"), Width, Width, Idx);
+		UE_LOG(LogTestPAL, Display, TEXT("DynSize: %d Size_tDynFormat%0*zd"), Width, Width, static_cast<size_t>(Idx));
+		++Width;
+	}
+
 	FEngineLoop::AppPreExit();
 	FEngineLoop::AppExit();
 	return 0;
diff --git a/Engine/Source/Runtime/Core/Private/GenericPlatform/GenericWidePlatformString.cpp b/Engine/Source/Runtime/Core/Private/GenericPlatform/GenericWidePlatformString.cpp
--- a/Engine/Source/Runtime/Core/Private/GenericPlatform/GenericWidePlatformString.cpp
+++ b/Engine/Source/Runtime/Core/Private/GenericPlatform/GenericWidePlatformString.cpp
@@ -363,16 +363,30 @@ int32 FGenericWidePlatformString::GetVarArgs(WIDECHAR* Dest, SIZE_T DestSize, in
 		{
 			Src++;
 			int Val = va_arg(ArgPtr, int);
-			ANSICHAR AnsiNum[30];
-			ANSICHAR FmtBuf[30];
+			ANSICHAR AnsiNum[64];
+			ANSICHAR FmtBuf[30] = {0};
+
+			// limit width to the buffer size
+			FieldLen = FMath::Min(static_cast<int>(sizeof(AnsiNum)) - 1, FieldLen);
 
 			// Yes, this is lame.
 			int CpyIdx = 0;
 			while (Percent < Src && CpyIdx < ARRAY_COUNT(FmtBuf))
 			{
-				FmtBuf[CpyIdx] = (ANSICHAR) *Percent;
-				Percent++;
-				CpyIdx++;
+				if (LIKELY(*Percent != '*'))
+				{
+					FmtBuf[CpyIdx] = (ANSICHAR) *Percent;
+					++CpyIdx;
+				}
+				else
+				{
+					snprintf(&FmtBuf[CpyIdx], sizeof(FmtBuf) - CpyIdx, "%d", FieldLen);
+					while (CpyIdx < ARRAY_COUNT(FmtBuf) && FmtBuf[CpyIdx] != 0)
+					{
+						++CpyIdx;
+					}
+				}
+				++Percent;
 			}
 			FmtBuf[CpyIdx] = 0;
 
@@ -396,16 +410,30 @@ int32 FGenericWidePlatformString::GetVarArgs(WIDECHAR* Dest, SIZE_T DestSize, in
 
 			size_t Val = va_arg(ArgPtr, size_t);
 
-			ANSICHAR AnsiNum[30];
-			ANSICHAR FmtBuf[30];
+			ANSICHAR AnsiNum[64];
+			ANSICHAR FmtBuf[30] = {0};
+
+			// limit width to the buffer size
+			FieldLen = FMath::Min(static_cast<int>(sizeof(AnsiNum)) - 1, FieldLen);
 
 			// Yes, this is lame.
 			int CpyIdx = 0;
 			while (Percent < Src && CpyIdx < ARRAY_COUNT(FmtBuf))
 			{
-				FmtBuf[CpyIdx] = (ANSICHAR) *Percent;
-				Percent++;
-				CpyIdx++;
+				if (LIKELY(*Percent != '*'))
+				{
+					FmtBuf[CpyIdx] = (ANSICHAR) *Percent;
+					++CpyIdx;
+				}
+				else
+				{
+					snprintf(&FmtBuf[CpyIdx], sizeof(FmtBuf) - CpyIdx, "%d", FieldLen);
+					while (CpyIdx < ARRAY_COUNT(FmtBuf) && FmtBuf[CpyIdx] != 0)
+					{
+						++CpyIdx;
+					}
+				}
+				++Percent;
 			}
 			FmtBuf[CpyIdx] = 0;
 
