Fixed a crash when running a cooked game -memoryprofiler in the command line

diff --git a/Engine/Source/Runtime/Core/Private/Stats/StatsMallocProfilerProxy.cpp b/Engine/Source/Runtime/Core/Private/Stats/StatsMallocProfilerProxy.cpp
--- a/Engine/Source/Runtime/Core/Private/Stats/StatsMallocProfilerProxy.cpp
+++ b/Engine/Source/Runtime/Core/Private/Stats/StatsMallocProfilerProxy.cpp
@@ -103,7 +103,18 @@ void FStatsMallocProfilerProxy::TrackAlloc( void* Ptr, int64 Size, int32 Sequenc
 {
 	if( bEnabled )
 	{
-		FThreadStats* ThreadStats = FThreadStats::GetThreadStats();
+		thread_local static bool bReentryGuard = false;
+		if (bReentryGuard)
+		{
+			return;
+		}
+
+		FThreadStats* ThreadStats = nullptr;
+		{
+			bReentryGuard = true;
+			ThreadStats = FThreadStats::GetThreadStats();
+			bReentryGuard = false;
+		}
 
 #ifdef DEBUG_MALLOC_PROXY
 		if (GThreadStatsToDumpMemory == ThreadStats && ThreadStats->MemoryMessageScope == 0)
