Fixed error message generation, by truncating callstack when it won't fit. Callstack is still present in its own field.

diff --git a/Engine/Source/Runtime/Core/Private/Misc/AssertionMacros.cpp b/Engine/Source/Runtime/Core/Private/Misc/AssertionMacros.cpp
--- a/Engine/Source/Runtime/Core/Private/Misc/AssertionMacros.cpp
+++ b/Engine/Source/Runtime/Core/Private/Misc/AssertionMacros.cpp
@@ -124,7 +124,11 @@ void StaticFailDebug( const TCHAR* Error, const ANSICHAR* File, int32 Line, cons
 	FPlatformMisc::LowLevelOutputDebugStringf(TEXT("%s") FILE_LINE_DESC TEXT("\n%s\n"), Error, ANSI_TO_TCHAR(File), Line, Description);
 
 	// Copy the detailed error into the error message.
-	FCString::Snprintf( GErrorMessage, ARRAY_COUNT( GErrorMessage ), TEXT( "%s" ) FILE_LINE_DESC TEXT( "\n%s\n" ), Error, ANSI_TO_TCHAR( File ), Line, DescriptionAndTrace);
+	if (FCString::Snprintf( GErrorMessage, ARRAY_COUNT( GErrorMessage ), TEXT( "%s" ) FILE_LINE_DESC TEXT( "\n%s\n" ), Error, ANSI_TO_TCHAR( File ), Line, DescriptionAndTrace) < 0)
+	{
+		// Description and callstack was too long to fit in GErrorMessage. Use only description
+		FCString::Snprintf( GErrorMessage, ARRAY_COUNT( GErrorMessage ), TEXT( "%s" ) FILE_LINE_DESC TEXT( "\n%s\n<< callstack too long >>" ), Error, ANSI_TO_TCHAR( File ), Line, Description);
+	}
 
 	// Copy the error message to the error history.
 	FCString::Strncpy( GErrorHist, GErrorMessage, ARRAY_COUNT( GErrorHist ) );
