Fix inconsistencies in reported callstacks.

* Number of stack frames to ignore is now passed through all check/ensure handling functions, and incremented by one with each call.
* HAL-level functions for capturing stack frames do not adjust the number of frames to ignore if a context is specified (ie. if the callstack has already been captured). This is only necessary if capturing the stack on the calling thread.
* NewReportEnsure() now includes an additional parameter for number of stack frames to ignore. Currently only observed on Windows.

diff --git a/Engine/Source/Runtime/Core/Private/GenericPlatform/GenericPlatformCrashContext.cpp b/Engine/Source/Runtime/Core/Private/GenericPlatform/GenericPlatformCrashContext.cpp
--- a/Engine/Source/Runtime/Core/Private/GenericPlatform/GenericPlatformCrashContext.cpp
+++ b/Engine/Source/Runtime/Core/Private/GenericPlatform/GenericPlatformCrashContext.cpp
@@ -639,6 +639,12 @@ void FGenericCrashContext::AddPlugin(const FString& PluginDesc)
 
 FORCENOINLINE void FGenericCrashContext::CapturePortableCallStack(int32 NumStackFramesToIgnore, void* Context)
 {
+	// If the callstack is for the executing thread, ignore this function
+	if(Context == nullptr)
+	{
+		NumStackFramesToIgnore++;
+	}
+
 	const int32 MaxDepth = 100;
 	TArray<FProgramCounterSymbolInfo> Stack = FPlatformStackWalk::GetStack(NumStackFramesToIgnore, MaxDepth, Context);
 	return SetPortableCallStack(NumStackFramesToIgnore, Stack);
diff --git a/Engine/Source/Runtime/Core/Private/GenericPlatform/GenericPlatformStackWalk.cpp b/Engine/Source/Runtime/Core/Private/GenericPlatform/GenericPlatformStackWalk.cpp
--- a/Engine/Source/Runtime/Core/Private/GenericPlatform/GenericPlatformStackWalk.cpp
+++ b/Engine/Source/Runtime/Core/Private/GenericPlatform/GenericPlatformStackWalk.cpp
@@ -168,6 +168,12 @@ void FGenericPlatformStackWalk::StackWalkAndDump( ANSICHAR* HumanReadableString,
 	uint64 StackTrace[MAX_DEPTH];
 	FMemory::Memzero( StackTrace );
 
+	// If the callstack is for the executing thread, ignore this function and the FPlatformStackWalk::CaptureStackBackTrace call below
+	if(Context == nullptr)
+	{
+		IgnoreCount += 2;
+	}
+
 	// Capture stack backtrace.
 	uint32 Depth = FPlatformStackWalk::CaptureStackBackTrace( StackTrace, MAX_DEPTH, Context );
 
@@ -196,8 +202,11 @@ TArray<FProgramCounterSymbolInfo> FGenericPlatformStackWalk::GetStack(int32 Igno
 	uint64 StackTrace[MAX_DEPTH];
 	FMemory::Memzero(StackTrace);
 
-	// Add 2 to account for CaptureStackBackTrace and GetStack.
-	IgnoreCount += 2;
+	// If the callstack is for the executing thread, ignore this function and the FPlatformStackWalk::CaptureStackBackTrace call below
+	if(Context == nullptr)
+	{
+		IgnoreCount += 2;
+	}
 
 	MaxDepth = FMath::Min(MAX_DEPTH, IgnoreCount + MaxDepth);
 
diff --git a/Engine/Source/Runtime/Core/Private/HAL/ExceptionHandling.cpp b/Engine/Source/Runtime/Core/Private/HAL/ExceptionHandling.cpp
--- a/Engine/Source/Runtime/Core/Private/HAL/ExceptionHandling.cpp
+++ b/Engine/Source/Runtime/Core/Private/HAL/ExceptionHandling.cpp
@@ -50,7 +50,8 @@ void ReportInteractiveEnsure(const TCHAR* InMessage)
 	// GErrorMessage here is very unfortunate but it's used internally by the crash context code.
 	FCString::Strcpy(GErrorMessage, ARRAY_COUNT(GErrorMessage), InMessage);
 	// Skip macros and FDebug, we always want this to fire
-	NewReportEnsure(InMessage);
+	const int32 NumStackFramesToIgnore = 1;
+	NewReportEnsure(InMessage, NumStackFramesToIgnore);
 	GErrorMessage[0] = '\0';
 #endif
 
diff --git a/Engine/Source/Runtime/Core/Private/IOS/IOSPlatformMisc.cpp b/Engine/Source/Runtime/Core/Private/IOS/IOSPlatformMisc.cpp
--- a/Engine/Source/Runtime/Core/Private/IOS/IOSPlatformMisc.cpp
+++ b/Engine/Source/Runtime/Core/Private/IOS/IOSPlatformMisc.cpp
@@ -1841,7 +1841,7 @@ void FIOSCrashContext::GenerateEnsureInfo() const
 static FCriticalSection EnsureLock;
 static bool bReentranceGuard = false;
 
-void NewReportEnsure( const TCHAR* ErrorMessage )
+void NewReportEnsure( const TCHAR* ErrorMessage, int NumStackFramesToIgnore )
 {
     // Simple re-entrance guard.
     EnsureLock.Lock();
diff --git a/Engine/Source/Runtime/Core/Private/Logging/LogMacros.cpp b/Engine/Source/Runtime/Core/Private/Logging/LogMacros.cpp
--- a/Engine/Source/Runtime/Core/Private/Logging/LogMacros.cpp
+++ b/Engine/Source/Runtime/Core/Private/Logging/LogMacros.cpp
@@ -7,7 +7,7 @@
 #include "Misc/FeedbackContext.h"
 #include "Misc/VarargsHelper.h"
 
-void StaticFailDebug( const TCHAR* Error, const ANSICHAR* File, int32 Line, const TCHAR* Description, bool bIsEnsure );
+void StaticFailDebug( const TCHAR* Error, const ANSICHAR* File, int32 Line, const TCHAR* Description, bool bIsEnsure, int32 NumStackFramesToIgnore );
 
 /** Statics to prevent FMsg::Logf from allocating too much stack memory. */
 static FCriticalSection					MsgLogfStaticBufferGuard;
@@ -57,7 +57,8 @@ void FMsg::LogfImpl(const ANSICHAR* File, int32 Line, const FName& Category, ELo
 			Message[ARRAY_COUNT(Message) - 1] = '\0';
 		}
 
-		StaticFailDebug(TEXT("Fatal error:"), File, Line, Message, false);
+		const int32 NumStackFramesToIgnore = 1;
+		StaticFailDebug(TEXT("Fatal error:"), File, Line, Message, false, NumStackFramesToIgnore);
 		FDebug::AssertFailed("", File, Line, Message);
 	}
 #endif
@@ -106,7 +107,8 @@ void FMsg::Logf_InternalImpl(const ANSICHAR* File, int32 Line, const FName& Cate
 			Message[ARRAY_COUNT(Message) - 1] = '\0';
 		}
 
-		StaticFailDebug(TEXT("Fatal error:"), File, Line, Message, false);
+		const int32 NumStackFramesToIgnore = 1;
+		StaticFailDebug(TEXT("Fatal error:"), File, Line, Message, false, NumStackFramesToIgnore);
 	}
 #endif
 }
diff --git a/Engine/Source/Runtime/Core/Private/Mac/MacPlatformMisc.cpp b/Engine/Source/Runtime/Core/Private/Mac/MacPlatformMisc.cpp
--- a/Engine/Source/Runtime/Core/Private/Mac/MacPlatformMisc.cpp
+++ b/Engine/Source/Runtime/Core/Private/Mac/MacPlatformMisc.cpp
@@ -1848,7 +1848,7 @@ void FMacCrashContext::GenerateEnsureInfoAndLaunchReporter() const
 static FCriticalSection EnsureLock;
 static bool bReentranceGuard = false;
 
-void NewReportEnsure( const TCHAR* ErrorMessage )
+void NewReportEnsure( const TCHAR* ErrorMessage, int NumStackFramesToIgnore )
 {
 	// Simple re-entrance guard.
 	EnsureLock.Lock();
diff --git a/Engine/Source/Runtime/Core/Private/Misc/AssertionMacros.cpp b/Engine/Source/Runtime/Core/Private/Misc/AssertionMacros.cpp
--- a/Engine/Source/Runtime/Core/Private/Misc/AssertionMacros.cpp
+++ b/Engine/Source/Runtime/Core/Private/Misc/AssertionMacros.cpp
@@ -70,17 +70,6 @@ namespace
 
 
 
-/** Number of top function calls to hide when dumping the callstack as text. */
-#if PLATFORM_UNIX
-
-	// Rationale: check() and ensure() handlers have different depth - worse, ensure() can optionally end up calling the same path as check().
-	// It is better to show the full callstack as is than accidentaly ignore a part of the problem
-	#define CALLSTACK_IGNOREDEPTH 0
-
-#else
-	#define CALLSTACK_IGNOREDEPTH 2
-#endif // PLATFORM_UNIX
-
 CORE_API void (*GPrintScriptCallStackFn)() = nullptr;
 
 void PrintScriptCallstack()
@@ -96,7 +85,7 @@ void PrintScriptCallstack()
  *	prompts for the remote debugging if there is not debugger, breaks into the debugger 
  *	and copies the error into the global error message.
  */
-void StaticFailDebug( const TCHAR* Error, const ANSICHAR* File, int32 Line, const TCHAR* Description, bool bIsEnsure )
+FORCENOINLINE void StaticFailDebug( const TCHAR* Error, const ANSICHAR* File, int32 Line, const TCHAR* Description, bool bIsEnsure, int NumStackFramesToIgnore )
 {
 	// Print out the blueprint callstack
 	PrintScriptCallstack();
@@ -113,7 +102,7 @@ void StaticFailDebug( const TCHAR* Error, const ANSICHAR* File, int32 Line, cons
 		if (StackTrace != NULL)
 		{
 			StackTrace[0] = 0;
-			FPlatformStackWalk::StackWalkAndDump(StackTrace, ARRAY_COUNT(StackTrace), CALLSTACK_IGNOREDEPTH);
+			FPlatformStackWalk::StackWalkAndDump(StackTrace, ARRAY_COUNT(StackTrace), NumStackFramesToIgnore + 1);
 
 			FCString::Strncat(DescriptionAndTrace, TEXT("\n"), ARRAY_COUNT(DescriptionAndTrace) - 1);
 			FCString::Strncat(DescriptionAndTrace, ANSI_TO_TCHAR(StackTrace), ARRAY_COUNT(DescriptionAndTrace) - 1);
@@ -221,7 +210,7 @@ void FDebug::LogFormattedMessageWithCallstack(const FName& LogName, const ANSICH
 //warning: May be called at library startup time.
 //
 
-void FDebug::LogAssertFailedMessageImpl(const ANSICHAR* Expr, const ANSICHAR* File, int32 Line, const TCHAR* Fmt, ...)
+FORCENOINLINE void FDebug::LogAssertFailedMessageImpl(const ANSICHAR* Expr, const ANSICHAR* File, int32 Line, const TCHAR* Fmt, ...)
 {
 	// Ignore this assert if we're already forcibly shutting down because of a critical error.
 	if( !GIsCriticalError )
@@ -232,7 +221,8 @@ void FDebug::LogAssertFailedMessageImpl(const ANSICHAR* Expr, const ANSICHAR* Fi
 		TCHAR ErrorString[MAX_SPRINTF];
 		FCString::Sprintf( ErrorString, TEXT( "Assertion failed: %s" ), ANSI_TO_TCHAR( Expr ) );
 
-		StaticFailDebug( ErrorString, File, Line, DescriptionString, false );
+		const int32 NumStackFramesToIgnore = 1;
+		StaticFailDebug( ErrorString, File, Line, DescriptionString, false, NumStackFramesToIgnore );
 	}
 }
 
@@ -243,8 +233,9 @@ void FDebug::LogAssertFailedMessageImpl(const ANSICHAR* Expr, const ANSICHAR* Fi
  * @param	File	File name ANSI string (__FILE__)
  * @param	Line	Line number (__LINE__)
  * @param	Msg		Informative error message text
+ * @param	NumStackFramesToIgnore	Number of stack frames to ignore in the callstack
  */
-void FDebug::EnsureFailed(const ANSICHAR* Expr, const ANSICHAR* File, int32 Line, const TCHAR* Msg)
+FORCENOINLINE void FDebug::EnsureFailed(const ANSICHAR* Expr, const ANSICHAR* File, int32 Line, const TCHAR* Msg, int NumStackFramesToIgnore)
 {
 #if STATS
 	FString EnsureFailedPerfMessage = FString::Printf(TEXT("FDebug::EnsureFailed"));
@@ -267,7 +258,7 @@ void FDebug::EnsureFailed(const ANSICHAR* Expr, const ANSICHAR* File, int32 Line
 	TCHAR ErrorString[MAX_SPRINTF];
 	FCString::Sprintf(ErrorString,TEXT("Ensure condition failed: %s"),ANSI_TO_TCHAR(Expr));
 
-	StaticFailDebug( ErrorString, File, Line, Msg, true );
+	StaticFailDebug( ErrorString, File, Line, Msg, true, NumStackFramesToIgnore + 1 );
 
 	// Is there a debugger attached?  If not we'll submit an error report.
 	if (FPlatformMisc::IsDebuggerPresent())
@@ -302,8 +293,8 @@ void FDebug::EnsureFailed(const ANSICHAR* Expr, const ANSICHAR* File, int32 Line
 				FString StackWalkPerfMessage = FString::Printf(TEXT("FPlatformStackWalk::StackWalkAndDump"));
 				SCOPE_LOG_TIME_IN_SECONDS(*StackWalkPerfMessage, nullptr)
 #endif
-					StackTrace[0] = 0;
-				FPlatformStackWalk::StackWalkAndDumpEx(StackTrace, StackTraceSize, CALLSTACK_IGNOREDEPTH, FGenericPlatformStackWalk::EStackWalkFlags::FlagsUsedWhenHandlingEnsure);
+				StackTrace[0] = 0;
+				FPlatformStackWalk::StackWalkAndDumpEx(StackTrace, StackTraceSize, NumStackFramesToIgnore + 1, FGenericPlatformStackWalk::EStackWalkFlags::FlagsUsedWhenHandlingEnsure);
 			}
 
 			// Create a final string that we'll output to the log (and error history buffer)
@@ -388,9 +379,9 @@ void FDebug::EnsureFailed(const ANSICHAR* Expr, const ANSICHAR* File, int32 Line
 #endif
 
 #if PLATFORM_DESKTOP
-				FScopeLock Lock(&FailDebugCriticalSection);
+			FScopeLock Lock(&FailDebugCriticalSection);
 
-			NewReportEnsure(GErrorMessage);
+			NewReportEnsure(GErrorMessage, NumStackFramesToIgnore + 1);
 
 			GErrorHist[0] = 0;
 			GErrorMessage[0] = 0;
@@ -427,29 +418,32 @@ void VARARGS FDebug::AssertFailed(const ANSICHAR* Expr, const ANSICHAR* File, in
 }
 
 #if DO_CHECK || DO_GUARD_SLOW
-bool VARARGS FDebug::OptionallyLogFormattedEnsureMessageReturningFalseImpl( bool bLog, const ANSICHAR* Expr, const ANSICHAR* File, int32 Line, const TCHAR* FormattedMsg, ... )
+FORCENOINLINE bool VARARGS FDebug::OptionallyLogFormattedEnsureMessageReturningFalseImpl( bool bLog, const ANSICHAR* Expr, const ANSICHAR* File, int32 Line, const TCHAR* FormattedMsg, ... )
 {
 	if (bLog)
 	{
 		const int32 TempStrSize = 4096;
 		TCHAR TempStr[ TempStrSize ];
 		GET_VARARGS( TempStr, TempStrSize, TempStrSize - 1, FormattedMsg, FormattedMsg );
-		EnsureFailed( Expr, File, Line, TempStr );
+
+		const int32 NumStackFramesToIgnore = 1; // Just ignore this frame
+		EnsureFailed( Expr, File, Line, TempStr, NumStackFramesToIgnore );
 	}
 	
 	return false;
 }
 #endif
 
-void VARARGS LowLevelFatalErrorHandler(const ANSICHAR* File, int32 Line, const TCHAR* Format, ...)
+FORCENOINLINE void VARARGS LowLevelFatalErrorHandler(const ANSICHAR* File, int32 Line, const TCHAR* Format, ...)
 {
 	TCHAR DescriptionString[4096];
 	GET_VARARGS( DescriptionString, ARRAY_COUNT(DescriptionString), ARRAY_COUNT(DescriptionString)-1, Format, Format );
 
-	StaticFailDebug(TEXT("LowLevelFatalError"), File, Line, DescriptionString, false);
+	const int32 NumStackFramesToIgnore = 1; // Just ignore this frame
+	StaticFailDebug(TEXT("LowLevelFatalError"), File, Line, DescriptionString, false, NumStackFramesToIgnore);
 }
 
-void FDebug::DumpStackTraceToLog()
+FORCENOINLINE void FDebug::DumpStackTraceToLog()
 {
 #if !NO_LOGGING
 	// Walk the stack and dump it to the allocated memory.
@@ -462,7 +456,9 @@ void FDebug::DumpStackTraceToLog()
 		SCOPE_LOG_TIME_IN_SECONDS(*StackWalkPerfMessage, nullptr)
 #endif
 		StackTrace[0] = 0;
-		FPlatformStackWalk::StackWalkAndDumpEx(StackTrace, StackTraceSize, CALLSTACK_IGNOREDEPTH, FGenericPlatformStackWalk::EStackWalkFlags::FlagsUsedWhenHandlingEnsure);
+
+		const int32 NumStackFramesToIgnore = 1;
+		FPlatformStackWalk::StackWalkAndDumpEx(StackTrace, StackTraceSize, NumStackFramesToIgnore, FGenericPlatformStackWalk::EStackWalkFlags::FlagsUsedWhenHandlingEnsure);
 	}
 
 	// Dump the error and flush the log.
diff --git a/Engine/Source/Runtime/Core/Private/Misc/OutputDevice.cpp b/Engine/Source/Runtime/Core/Private/Misc/OutputDevice.cpp
--- a/Engine/Source/Runtime/Core/Private/Misc/OutputDevice.cpp
+++ b/Engine/Source/Runtime/Core/Private/Misc/OutputDevice.cpp
@@ -55,16 +55,17 @@ void FOutputDevice::Log( const FText& T )
 	Formatted printing and messages.
 -----------------------------------------------------------------------------*/
 
-void FOutputDevice::CategorizedLogfImpl(const FName& Category, ELogVerbosity::Type Verbosity, const TCHAR* Fmt, ...)
+// Do not inline these functions, in case we need to capture a call stack in FOutputDeviceError::Serialize. We need to be certain of how many frames to ignore.
+FORCENOINLINE void FOutputDevice::CategorizedLogfImpl(const FName& Category, ELogVerbosity::Type Verbosity, const TCHAR* Fmt, ...)
 {
 	GROWABLE_LOGF(Serialize(Buffer, Verbosity, Category))
 }
-void FOutputDevice::LogfImpl(ELogVerbosity::Type Verbosity, const TCHAR* Fmt, ...)
+FORCENOINLINE void FOutputDevice::LogfImpl(ELogVerbosity::Type Verbosity, const TCHAR* Fmt, ...)
 {
 	// call serialize with the final buffer
 	GROWABLE_LOGF(Serialize(Buffer, Verbosity, NAME_None))
 }
-void FOutputDevice::LogfImpl(const TCHAR* Fmt, ...)
+FORCENOINLINE void FOutputDevice::LogfImpl(const TCHAR* Fmt, ...)
 {
 	FScopedCategoryAndVerbosityOverride::FOverride* TLS = FScopedCategoryAndVerbosityOverride::GetTLSCurrent();
 	GROWABLE_LOGF(Serialize(Buffer, TLS->Verbosity, TLS->Category))
diff --git a/Engine/Source/Runtime/Core/Private/Unix/UnixPlatformStackWalk.cpp b/Engine/Source/Runtime/Core/Private/Unix/UnixPlatformStackWalk.cpp
--- a/Engine/Source/Runtime/Core/Private/Unix/UnixPlatformStackWalk.cpp
+++ b/Engine/Source/Runtime/Core/Private/Unix/UnixPlatformStackWalk.cpp
@@ -612,7 +612,7 @@ int32 FUnixPlatformStackWalk::GetProcessModuleSignatures(FStackWalkModuleInfo *M
 static FCriticalSection EnsureLock;
 static bool bReentranceGuard = false;
 
-void NewReportEnsure(const TCHAR* ErrorMessage)
+void NewReportEnsure(const TCHAR* ErrorMessage, int NumStackFramesToIgnore)
 {
 	// Simple re-entrance guard.
 	EnsureLock.Lock();
diff --git a/Engine/Source/Runtime/Core/Private/Windows/WindowsPlatformCrashContext.cpp b/Engine/Source/Runtime/Core/Private/Windows/WindowsPlatformCrashContext.cpp
--- a/Engine/Source/Runtime/Core/Private/Windows/WindowsPlatformCrashContext.cpp
+++ b/Engine/Source/Runtime/Core/Private/Windows/WindowsPlatformCrashContext.cpp
@@ -328,7 +328,7 @@ void ReportHang(const TCHAR* ErrorMessage, const TArray<FProgramCounterSymbolInf
 /** 
  * Report an ensure to the crash reporting system
  */
-FORCENOINLINE void NewReportEnsure( const TCHAR* ErrorMessage )
+FORCENOINLINE void NewReportEnsure( const TCHAR* ErrorMessage, int NumStackFramesToIgnore )
 {
 	if (ReportCrashCallCount > 0)
 	{
@@ -358,8 +358,8 @@ FORCENOINLINE void NewReportEnsure( const TCHAR* ErrorMessage )
 
 	bReentranceGuard = true;
 	
-	// Three additional frames are skipped inside FWindowsPlatformStackWalk::GetStack() and FPlatformStackWalk::GetStack()
-	const int NumStackFramesToIgnore = 2;
+	// Ignore this function and the RaiseException() call below.
+	NumStackFramesToIgnore += 2;
 
 #if WINVER > 0x502	// Windows Error Reporting is not supported on Windows XP
 #if !PLATFORM_SEH_EXCEPTIONS_DISABLED
@@ -559,8 +559,13 @@ private:
 		// Thread context wrapper for stack operations
 		void* ContextWrapper = FWindowsPlatformStackWalk::MakeThreadContextWrapper(ExceptionInfo->ContextRecord, CrashingThreadHandle);
 
-		// Generate the portable callstack, for non-Asserts pass in a negative IgnoreCount as need top of crash context record, and IgnoreCount is increased by stack walking methods
-		const int32 IgnoreCount = FDebug::HasAsserted() ? 1 : -3;
+		// Generate the portable callstack. For asserts, we ignore the following frames:
+		//     FDebug::AssertFailed()
+		//   [ FOutputDevice::Logf() ] - force-inlined; ignored
+		//     FOutputDevice::LogfImpl()
+		//     FWindowsErrorOutputDevice::Serialize()
+		//     RaiseException()
+		const int32 IgnoreCount = FDebug::HasAsserted() ? 4 : 0;
 		CrashContext.CapturePortableCallStack(IgnoreCount, ContextWrapper);
 
 		// First launch the crash reporter client.
diff --git a/Engine/Source/Runtime/Core/Private/Windows/WindowsPlatformStackWalk.cpp b/Engine/Source/Runtime/Core/Private/Windows/WindowsPlatformStackWalk.cpp
--- a/Engine/Source/Runtime/Core/Private/Windows/WindowsPlatformStackWalk.cpp
+++ b/Engine/Source/Runtime/Core/Private/Windows/WindowsPlatformStackWalk.cpp
@@ -232,13 +232,25 @@ void DetermineMaxCallstackDepth()
 void FWindowsPlatformStackWalk::StackWalkAndDump( ANSICHAR* HumanReadableString, SIZE_T HumanReadableStringSize, int32 IgnoreCount, void* Context )
 {
 	InitStackWalking();
+
+	// If the callstack is for the executing thread, ignore this function
+	if(Context == nullptr)
+	{
+		IgnoreCount++;
+	}
 	FGenericPlatformStackWalk::StackWalkAndDump(HumanReadableString, HumanReadableStringSize, IgnoreCount, Context);
 }
 
 FORCENOINLINE TArray<FProgramCounterSymbolInfo> FWindowsPlatformStackWalk::GetStack(int32 IgnoreCount, int32 MaxDepth, void* Context)
 {
 	InitStackWalking();
-	return FGenericPlatformStackWalk::GetStack(IgnoreCount + 1, MaxDepth, Context);
+
+	// If the callstack is for the executing thread, ignore this function
+	if(Context == nullptr)
+	{
+		IgnoreCount++;
+	}
+	return FGenericPlatformStackWalk::GetStack(IgnoreCount, MaxDepth, Context);
 }
 
 void FWindowsPlatformStackWalk::ThreadStackWalkAndDump(ANSICHAR* HumanReadableString, SIZE_T HumanReadableStringSize, int32 IgnoreCount, uint32 ThreadId)
@@ -319,18 +331,18 @@ uint32 FWindowsPlatformStackWalk::CaptureStackBackTrace( uint64* BackTrace, uint
 	{
 #if USE_FAST_STACKTRACE
 		if (!GMaxCallstackDepthInitialized)
-			{
-				DetermineMaxCallstackDepth();
-			}
-			PVOID WinBackTrace[MAX_CALLSTACK_DEPTH];
+		{
+			DetermineMaxCallstackDepth();
+		}
+		PVOID WinBackTrace[MAX_CALLSTACK_DEPTH];
 		uint16 NumFrames = RtlCaptureStackBackTrace(0, FMath::Min<ULONG>(GMaxCallstackDepth, MaxDepth), WinBackTrace, NULL);
 		Depth = NumFrames;
 		for (uint16 FrameIndex = 0; FrameIndex < NumFrames; ++FrameIndex)
-			{
+		{
 			BackTrace[FrameIndex] = (uint64)WinBackTrace[FrameIndex];
-			}
+		}
 		while (NumFrames < MaxDepth)
-			{
+		{
 			BackTrace[NumFrames++] = 0;
 		}		
 #elif USE_SLOW_STACKTRACE
diff --git a/Engine/Source/Runtime/Core/Public/HAL/ExceptionHandling.h b/Engine/Source/Runtime/Core/Public/HAL/ExceptionHandling.h
--- a/Engine/Source/Runtime/Core/Public/HAL/ExceptionHandling.h
+++ b/Engine/Source/Runtime/Core/Public/HAL/ExceptionHandling.h
@@ -24,21 +24,21 @@ extern CORE_API TCHAR MiniDumpFilenameW[1024];
 #include <excpt.h>
 // #CrashReport: 2014-10-09 These methods are specific to windows, remove from here.
 extern CORE_API int32 ReportCrash( Windows::LPEXCEPTION_POINTERS ExceptionInfo );
-extern CORE_API void NewReportEnsure( const TCHAR* ErrorMessage );
+extern CORE_API void NewReportEnsure( const TCHAR* ErrorMessage, int NumStackFramesToIgnore );
 extern CORE_API void ReportHang(const TCHAR*, const TArray<FProgramCounterSymbolInfo>& Stack);
 #elif PLATFORM_XBOXONE
 #include "XboxOne/XboxOneSystemIncludes.h"
 // #CrashReport: 2014-10-09 Should be move to another file
 extern CORE_API int32 ReportCrash( int ExceptionCode, LPEXCEPTION_POINTERS ExceptionInfo );
-extern CORE_API void NewReportEnsure( const TCHAR* ErrorMessage );
+extern CORE_API void NewReportEnsure( const TCHAR* ErrorMessage, int NumStackFramesToIgnore );
 #elif PLATFORM_MAC
 // #CrashReport: 2014-10-09 Should be move to another file
 #include <signal.h>
 extern CORE_API int32 ReportCrash( ucontext_t *Context, int32 Signal, struct __siginfo* Info );
-extern CORE_API void NewReportEnsure( const TCHAR* ErrorMessage );
+extern CORE_API void NewReportEnsure( const TCHAR* ErrorMessage, int NumStackFramesToIgnore );
 extern CORE_API void ReportHang(const TCHAR*, const TArray<FProgramCounterSymbolInfo>& Stack);
 #elif PLATFORM_UNIX
-extern CORE_API void NewReportEnsure( const TCHAR* ErrorMessage );
+extern CORE_API void NewReportEnsure( const TCHAR* ErrorMessage, int NumStackFramesToIgnore );
 extern CORE_API void ReportHang(const TCHAR*, const TArray<FProgramCounterSymbolInfo>& Stack);
 #endif
 
diff --git a/Engine/Source/Runtime/Core/Public/Misc/AssertionMacros.h b/Engine/Source/Runtime/Core/Public/Misc/AssertionMacros.h
--- a/Engine/Source/Runtime/Core/Public/Misc/AssertionMacros.h
+++ b/Engine/Source/Runtime/Core/Public/Misc/AssertionMacros.h
@@ -61,10 +61,11 @@ public:
 	 * @param	File	File name ANSI string (__FILE__)
 	 * @param	Line	Line number (__LINE__)
 	 * @param	Msg		Informative error message text
+	 * @param	NumStackFramesToIgnore	Number of stack frames to ignore in the ensure message
 	 * 
 	 * Don't change the name of this function, it's used to detect ensures by the crash reporter.
 	 */
-	static void EnsureFailed( const ANSICHAR* Expr, const ANSICHAR* File, int32 Line, const TCHAR* Msg );
+	static void EnsureFailed( const ANSICHAR* Expr, const ANSICHAR* File, int32 Line, const TCHAR* Msg, int NumStackFramesToIgnore );
 
 private:
 	static bool VARARGS OptionallyLogFormattedEnsureMessageReturningFalseImpl(bool bLog, const ANSICHAR* Expr, const ANSICHAR* File, int32 Line, const TCHAR* FormattedMsg, ...);
diff --git a/Engine/Source/Runtime/Core/Public/Misc/OutputDevice.h b/Engine/Source/Runtime/Core/Public/Misc/OutputDevice.h
--- a/Engine/Source/Runtime/Core/Public/Misc/OutputDevice.h
+++ b/Engine/Source/Runtime/Core/Public/Misc/OutputDevice.h
@@ -219,7 +219,7 @@ private:
 
 public:
 	template <typename FmtType, typename... Types>
-	void Logf(const FmtType& Fmt, Types... Args)
+	FORCEINLINE void Logf(const FmtType& Fmt, Types... Args)
 	{
 		static_assert(TIsArrayOrRefOfType<FmtType, TCHAR>::Value, "Formatting string must be a TCHAR array.");
 		static_assert(TAnd<TIsValidVariadicFunctionArg<Types>...>::Value, "Invalid argument(s) passed to FOutputDevice::Logf");
@@ -228,7 +228,7 @@ public:
 	}
 
 	template <typename FmtType, typename... Types>
-	void Logf(ELogVerbosity::Type Verbosity, const FmtType& Fmt, Types... Args)
+	FORCEINLINE void Logf(ELogVerbosity::Type Verbosity, const FmtType& Fmt, Types... Args)
 	{
 		static_assert(TIsArrayOrRefOfType<FmtType, TCHAR>::Value, "Formatting string must be a TCHAR array.");
 		static_assert(TAnd<TIsValidVariadicFunctionArg<Types>...>::Value, "Invalid argument(s) passed to FOutputDevice::Logf");
@@ -237,7 +237,7 @@ public:
 	}
 
 	template <typename FmtType, typename... Types>
-	void CategorizedLogf(const FName& Category, ELogVerbosity::Type Verbosity, const FmtType& Fmt, Types... Args)
+	FORCEINLINE void CategorizedLogf(const FName& Category, ELogVerbosity::Type Verbosity, const FmtType& Fmt, Types... Args)
 	{
 		static_assert(TIsArrayOrRefOfType<FmtType, TCHAR>::Value, "Formatting string must be a TCHAR array.");
 		static_assert(TAnd<TIsValidVariadicFunctionArg<Types>...>::Value, "Invalid argument(s) passed to FOutputDevice::CategorizedLogf");
