Linux: Remove the implicit -1 to callstacks when overwriting them and add an IgnoreCount to 1 for crashes

diff --git a/Engine/Source/Runtime/Core/Private/Unix/UnixPlatformCrashContext.cpp b/Engine/Source/Runtime/Core/Private/Unix/UnixPlatformCrashContext.cpp
--- a/Engine/Source/Runtime/Core/Private/Unix/UnixPlatformCrashContext.cpp
+++ b/Engine/Source/Runtime/Core/Private/Unix/UnixPlatformCrashContext.cpp
@@ -284,7 +284,10 @@ void FUnixCrashContext::CaptureStackTrace()
 		ANSICHAR* StackTrace = (ANSICHAR*) FMemory::Malloc( StackTraceSize );
 		StackTrace[0] = 0;
 
-		int32 IgnoreCount = 0;
+		// We use __builtin_return_address to save an address location to where our signal handler is called from
+		// We still need to ignore the first frame from pthreads which calls our signal handler, that is why we set it to one
+		// This may need to change something other then glibc is used
+		int32 IgnoreCount = 1;
 		CapturePortableCallStack(IgnoreCount, this);
 
 		// Walk the stack and dump it to the allocated memory (do not ignore any stack frames to be consistent with check()/ensure() handling)
@@ -604,7 +607,6 @@ void DefaultCrashHandler(const FUnixCrashContext & Context)
 	FThreadHeartBeat::Get().Stop();
 
 	// at this point we should already be using malloc crash handler (see PlatformCrashHandler)
-
 	const_cast<FUnixCrashContext&>(Context).CaptureStackTrace();
 	if (GLog)
 	{
diff --git a/Engine/Source/Runtime/Core/Private/Unix/UnixPlatformStackWalk.cpp b/Engine/Source/Runtime/Core/Private/Unix/UnixPlatformStackWalk.cpp
--- a/Engine/Source/Runtime/Core/Private/Unix/UnixPlatformStackWalk.cpp
+++ b/Engine/Source/Runtime/Core/Private/Unix/UnixPlatformStackWalk.cpp
@@ -462,7 +462,6 @@ namespace
 			{
 				if (FirstCrashHandlerFrame == reinterpret_cast<uint64*>(BackTrace[i]))
 				{
-					i++;
 					uint64* OverwriteBackTrace = BackTrace;
 
 					for (int j = i; j < Size; j++)
