Linux: fix crashes in headless mode (change to GL) (UE-63525).

diff --git a/Engine/Source/Programs/CrashReportClient/Private/Linux/CrashReportClientMainLinux.cpp b/Engine/Source/Programs/CrashReportClient/Private/Linux/CrashReportClientMainLinux.cpp
--- a/Engine/Source/Programs/CrashReportClient/Private/Linux/CrashReportClientMainLinux.cpp
+++ b/Engine/Source/Programs/CrashReportClient/Private/Linux/CrashReportClientMainLinux.cpp
@@ -58,6 +58,12 @@ int main(int argc, const char *argv[])
 		SavedCommandLine += UTF8_TO_TCHAR(argv[Option]);	// note: technically it depends on locale
 	}
 
+	// assume unattended if we don't have X11 display
+	if (getenv("DISPLAY") == nullptr)
+	{
+		SavedCommandLine += TEXT(" -unattended");
+	}
+
 	// Run the app
 	RunCrashReportClient(*SavedCommandLine);
 
diff --git a/Engine/Source/Runtime/Core/Private/Unix/UnixPlatformCrashContext.cpp b/Engine/Source/Runtime/Core/Private/Unix/UnixPlatformCrashContext.cpp
--- a/Engine/Source/Runtime/Core/Private/Unix/UnixPlatformCrashContext.cpp
+++ b/Engine/Source/Runtime/Core/Private/Unix/UnixPlatformCrashContext.cpp
@@ -398,6 +398,18 @@ void FUnixCrashContext::GenerateCrashInfoAndLaunchReporter(bool bReportingNonCra
 	// Suppress the user input dialog if we're running in unattended mode
 	bool bUnattended = FApp::IsUnattended() || (!IsInteractiveEnsureMode() && bReportingNonCrash) || IsRunningDedicatedServer();
 
+#if PLATFORM_LINUX
+	// On Linux, count not having a X11 display as also running unattended, because CRC will switch to the unattended mode in that case
+	if (!bUnattended)
+	{
+		// see CrashReportClientMainLinux.cpp
+		if (getenv("DISPLAY") == nullptr)
+		{
+			bUnattended = true;
+		}
+	}
+#endif
+
 	// By default we wont upload unless the *.ini has set this to true
 	bool bAgreedToCrashUpload = false;
 	GConfig->GetBool(TEXT("CrashReportClient"), TEXT("bAgreeToCrashUpload"), bAgreedToCrashUpload, GEngineIni);
diff --git a/Engine/Source/Runtime/VulkanRHI/Private/Linux/VulkanLinuxPlatform.cpp b/Engine/Source/Runtime/VulkanRHI/Private/Linux/VulkanLinuxPlatform.cpp
--- a/Engine/Source/Runtime/VulkanRHI/Private/Linux/VulkanLinuxPlatform.cpp
+++ b/Engine/Source/Runtime/VulkanRHI/Private/Linux/VulkanLinuxPlatform.cpp
@@ -18,6 +18,21 @@ bool FVulkanLinuxPlatform::bAttemptedLoad = false;
 
 bool FVulkanLinuxPlatform::IsSupported()
 {
+	// right now we do not provide an offscreen initialization path, so
+	// report as not supported if we're running without X11 or Wayland
+	bool bHasX11Display = getenv("DISPLAY") != nullptr;
+	bool bHasWaylandSession = false;
+	if (!bHasX11Display)
+	{
+		// check Wayland
+		bHasWaylandSession = getenv("WAYLAND_DISPLAY") != nullptr;
+	}
+
+	if (!bHasX11Display && !bHasWaylandSession)
+	{
+		return false;
+	}
+
 	// just attempt to load the library
 	return LoadVulkanLibrary();
 }
