Copy of CL 4455307

Fixed FIOSPlatformMisc::GetIOSDeviceType() for iPad with Major version == 8 to return IOS_IPadPro2_129 instead of IOS_Unknown

diff --git a/Engine/Source/Runtime/Core/Private/IOS/IOSPlatformMisc.cpp b/Engine/Source/Runtime/Core/Private/IOS/IOSPlatformMisc.cpp
--- a/Engine/Source/Runtime/Core/Private/IOS/IOSPlatformMisc.cpp
+++ b/Engine/Source/Runtime/Core/Private/IOS/IOSPlatformMisc.cpp
@@ -455,7 +455,7 @@ FIOSPlatformMisc::EIOSDevice FIOSPlatformMisc::GetIOSDeviceType()
 		}
 
 		// Default to highest settings currently available for any future device
-		else if (Major > 8)
+		else if (Major >= 8)
 		{
 			DeviceType = IOS_IPadPro2_129;
 		}
