Linux: fix skipping the clock calibration when running on Windows Subsystem for Linux.

diff --git a/Engine/Source/Runtime/Core/Private/Unix/UnixPlatformTime.cpp b/Engine/Source/Runtime/Core/Private/Unix/UnixPlatformTime.cpp
--- a/Engine/Source/Runtime/Core/Private/Unix/UnixPlatformTime.cpp
+++ b/Engine/Source/Runtime/Core/Private/Unix/UnixPlatformTime.cpp
@@ -88,14 +88,7 @@ uint64 FUnixTime::CallsPerSecondBenchmark(clockid_t BenchClockId, const char * B
 	char Buffer[256];
 	const uint64 kBenchmarkPeriodMicroSec = 1000000000ULL / 10;	// 0.1s
 
-	// basic sanity check
-	if (clock_getres(BenchClockId, nullptr) == -1)
-	{
-		FCStringAnsi::Snprintf(Buffer, sizeof(Buffer), "Clock_id %d (%s) is not supported on this system, clock_getres() fails.\n", BenchClockId, BenchClockIdName);
-		FCStringAnsi::Strncat(CalibrationLog, Buffer, sizeof(CalibrationLog));
-		return 0;	// unsupported clock id
-	}
-
+	// clock_getres() can fail when running on Windows Subsystem for Linux (but the clock can still be supported).
 	struct timespec ts;
 	if (clock_gettime(BenchClockId, &ts) == -1)
 	{
