Optimize WIDECHAR ToUpper() for characters less than 128

diff --git a/Engine/Source/Runtime/Core/Public/Misc/Char.h b/Engine/Source/Runtime/Core/Public/Misc/Char.h
--- a/Engine/Source/Runtime/Core/Public/Misc/Char.h
+++ b/Engine/Source/Runtime/Core/Public/Misc/Char.h
@@ -137,7 +137,7 @@ typedef TChar<ANSICHAR> FCharAnsi;
 
 template <> inline TChar<WIDECHAR>::CharType TChar<WIDECHAR>::ToUpper(CharType Char)
 {
-#if PLATFORM_SWITCH && !PLATFORM_TCHAR_IS_4_BYTES
+#if !PLATFORM_TCHAR_IS_4_BYTES
 	// for 2 byte wide chars use ansi toupper() if it is not an actual wide char (value less than 128)
 	// towupper() is slow on some platforms.
 	if (Char < 128)
