Add missing GetMaxPathLength for Android

#ue4
#android

diff --git a/Engine/Source/Runtime/Core/Public/Android/AndroidMisc.h b/Engine/Source/Runtime/Core/Public/Android/AndroidMisc.h
--- a/Engine/Source/Runtime/Core/Public/Android/AndroidMisc.h
+++ b/Engine/Source/Runtime/Core/Public/Android/AndroidMisc.h
@@ -33,6 +33,11 @@ struct CORE_API FAndroidMisc : public FGenericPlatformMisc
 	static void PlatformTearDown();
 	static void PlatformHandleSplashScreen(bool ShowSplashScreen);
 
+	FORCEINLINE static int32 GetMaxPathLength()
+	{
+		return ANDROID_MAX_PATH;
+	}
+
 	DEPRECATED(4.21, "void FPlatformMisc::GetEnvironmentVariable(Name, Result, Length) is deprecated. Use FString FPlatformMisc::GetEnvironmentVariable(Name) instead.")
 	static void GetEnvironmentVariable(const TCHAR* VariableName, TCHAR* Result, int32 ResultLength);
 
