Also clean up IsValid functions to just call the parse functions to avoid duplicated code

diff --git a/Engine/Source/Runtime/Core/Private/Misc/DefaultValueHelper.cpp b/Engine/Source/Runtime/Core/Private/Misc/DefaultValueHelper.cpp
--- a/Engine/Source/Runtime/Core/Private/Misc/DefaultValueHelper.cpp
+++ b/Engine/Source/Runtime/Core/Private/Misc/DefaultValueHelper.cpp
@@ -310,33 +310,8 @@ bool FDefaultValueHelper::IsStringValidFloat(const FString& Source)
 
 bool FDefaultValueHelper::IsStringValidVector(const FString& Source)
 {
-	const TCHAR* Start = StartOf(Source);
-	const TCHAR* FirstComma = FCString::Strstr( Start, TEXT(",") );
-	if(!FirstComma)
-	{
-		return false;
-	}
-
-	const TCHAR* SecondComma = FCString::Strstr( FirstComma + 1,  TEXT(",") );
-	if(!SecondComma)
-	{
-		return false;
-	}
-
-	// there must be exactly 2 commas in the string
-	if( FCString::Strstr( SecondComma + 1, TEXT(",") ) )
-	{
-		return false;
-	}
-
-	const TCHAR* End = EndOf(Source);
-	if(	!IsStringValidFloat( Start, FirstComma ) ||
-		!IsStringValidFloat( FirstComma + 1, SecondComma ) ||
-		!IsStringValidFloat( SecondComma + 1, End ) )
-	{
-		return false;
-	}
-	return true;
+	FVector TempVector;
+	return FDefaultValueHelper::ParseVector(Source, TempVector);
 }
 
 
@@ -348,39 +323,8 @@ bool FDefaultValueHelper::IsStringValidRotator(const FString& Source)
 
 bool FDefaultValueHelper::IsStringValidLinearColor(const FString& Source)
 {
-	if(Source.IsEmpty())
-	{
-		return false;
-	}
-
-	const TCHAR* Start = StartOf(Source);
-	const TCHAR* FirstComma = FCString::Strstr( Start, TEXT(",") );
-	if(!FirstComma)
-	{
-		return false;
-	}
-
-	const TCHAR* SecondComma = FCString::Strstr( FirstComma + 1,  TEXT(",") );
-	if(!SecondComma)
-	{
-		return false;
-	}
-
-	const TCHAR* ThirdComma = FCString::Strstr( SecondComma + 1,  TEXT(",") );
-	const TCHAR* End = EndOf(Source);
-	if( ( NULL != ThirdComma ) && !IsStringValidFloat( ThirdComma + 1, End ) )
-	{
-		return false;
-	}
-
-	if(	!IsStringValidFloat( Start, FirstComma ) ||
-		!IsStringValidFloat( FirstComma + 1, SecondComma ) ||
-		!IsStringValidFloat( SecondComma + 1, ( NULL != ThirdComma ) ? ThirdComma : End ) )
-	{
-		return false;
-	}
-
-	return true;
+	FLinearColor TempColor;
+	return ParseLinearColor(Source, TempColor);
 }
 
 
@@ -510,7 +454,8 @@ bool FDefaultValueHelper::ParseVector(const FString& Source, FVector& OutVal)
 		!IsStringValidFloat( FirstComma + 1, SecondComma ) ||
 		!IsStringValidFloat( SecondComma + 1, End ) )
 	{
-		return false;
+		// Fallback to x= format
+		return OutVal.InitFromString(Source);
 	}
 	
 	OutVal = FVector( 
@@ -542,7 +487,8 @@ bool FDefaultValueHelper::ParseVector2D(const FString& Source, FVector2D& OutVal
 	if(	!IsStringValidFloat( Start, FirstComma ) ||
 		!IsStringValidFloat( FirstComma + 1, End ) )
 	{
-		return false;
+		// Fallback to x= format
+		return OutVal.InitFromString(Source);
 	}
 
 	OutVal = FVector2D( 
@@ -567,7 +513,8 @@ bool FDefaultValueHelper::ParseVector4(const FString& Source, FVector4& OutVal)
 		return true;
 	}
 
-	return false;
+	// Fallback to x= format
+	return OutVal.InitFromString(Source);
 }
 
 
@@ -579,7 +526,9 @@ bool FDefaultValueHelper::ParseRotator(const FString& Source, FRotator& OutVal)
 		OutVal = FRotator(Vector.X, Vector.Y, Vector.Z);
 		return true;
 	}
-	return false;
+
+	// Fallback to x= format
+	return OutVal.InitFromString(Source);
 }
 
 
@@ -671,7 +620,8 @@ bool FDefaultValueHelper::ParseLinearColor(const FString& Source, FLinearColor&
 		!IsStringValidFloat( FirstComma + 1, SecondComma ) ||
 		!IsStringValidFloat( SecondComma + 1, ( NULL != ThirdComma ) ? ThirdComma : End ) )
 	{
-		return false;
+		// Fallback to x= format
+		return OutVal.InitFromString(Source);
 	}
 
 	const float Alpha = ( NULL != ThirdComma ) ? FCString::Atof( ThirdComma + 1 ) : 1.0f;
@@ -718,7 +668,8 @@ bool FDefaultValueHelper::ParseColor(const FString& Source, FColor& OutVal)
 		!IsStringValidInteger( FirstComma + 1, SecondComma ) ||
 		!IsStringValidInteger( SecondComma + 1, ( NULL != ThirdComma ) ? ThirdComma : End ) )
 	{
-		return false;
+		// Fallback to x= format
+		return OutVal.InitFromString(Source);
 	}
 
 	const uint8 Alpha = ( NULL != ThirdComma ) ? FCString::Atoi( ThirdComma + 1 ) : 255;
