Moved radix sort to Sorting.h. Added custom sort key policy and built in support for floats.

diff --git a/Engine/Source/Runtime/Core/Public/Templates/Sorting.h b/Engine/Source/Runtime/Core/Public/Templates/Sorting.h
--- a/Engine/Source/Runtime/Core/Public/Templates/Sorting.h
+++ b/Engine/Source/Runtime/Core/Public/Templates/Sorting.h
@@ -416,3 +416,121 @@ void StableSort(T** First, const int32 Num)
 	StableSortInternal(First, Num, TDereferenceWrapper<T*, TLess<T> >(TLess<T>()));
 }
 
+
+/**
+ * Very fast 32bit radix sort.
+ * SortKeyClass defines operator() that takes ValueType and returns a uint32. Sorting based on key.
+ * No comparisons. Is stable.
+ * Use a smaller CountType for smaller histograms.
+ */
+template< typename ValueType, typename CountType, class SortKeyClass >
+void RadixSort32( ValueType* RESTRICT Dst, ValueType* RESTRICT Src, CountType Num, SortKeyClass SortKey )
+{
+	CountType Histograms[ 1024 + 2048 + 2048 ];
+	CountType* RESTRICT Histogram0 = Histograms + 0;
+	CountType* RESTRICT Histogram1 = Histogram0 + 1024;
+	CountType* RESTRICT Histogram2 = Histogram1 + 2048;
+
+	FMemory::Memzero( Histograms, sizeof( Histograms ) );
+
+	{
+		// Parallel histogram generation pass
+		const ValueType* RESTRICT s = (const ValueType* RESTRICT)Src;
+		for( CountType i = 0; i < Num; i++ )
+		{
+			uint32 Key = SortKey( s[i] );
+			Histogram0[ ( Key >>  0 ) & 1023 ]++;
+			Histogram1[ ( Key >> 10 ) & 2047 ]++;
+			Histogram2[ ( Key >> 21 ) & 2047 ]++;
+		}
+	}
+	{
+		// Prefix sum
+		// Set each histogram entry to the sum of entries preceding it
+		CountType Sum0 = 0;
+		CountType Sum1 = 0;
+		CountType Sum2 = 0;
+		for( CountType i = 0; i < 1024; i++ )
+		{
+			CountType t;
+			t = Histogram0[i] + Sum0; Histogram0[i] = Sum0 - 1; Sum0 = t;
+			t = Histogram1[i] + Sum1; Histogram1[i] = Sum1 - 1; Sum1 = t;
+			t = Histogram2[i] + Sum2; Histogram2[i] = Sum2 - 1; Sum2 = t;
+		}
+		for( CountType i = 1024; i < 2048; i++ )
+		{
+			CountType t;
+			t = Histogram1[i] + Sum1; Histogram1[i] = Sum1 - 1; Sum1 = t;
+			t = Histogram2[i] + Sum2; Histogram2[i] = Sum2 - 1; Sum2 = t;
+		}
+	}
+	{
+		// Sort pass 1
+		const ValueType* RESTRICT s = (const ValueType* RESTRICT)Src;
+		ValueType* RESTRICT d = Dst;
+		for( CountType i = 0; i < Num; i++ )
+		{
+			ValueType Value = s[i];
+			uint32 Key = SortKey( Value );
+			d[ ++Histogram0[ ( (Key >> 0) & 1023 ) ] ] = Value;
+		}
+	}
+	{
+		// Sort pass 2
+		const ValueType* RESTRICT s = (const ValueType* RESTRICT)Dst;
+		ValueType* RESTRICT d = Src;
+		for( CountType i = 0; i < Num; i++ )
+		{
+			ValueType Value = s[i];
+			uint32 Key = SortKey( Value );
+			d[ ++Histogram1[ ( (Key >> 10) & 2047 ) ] ] = Value;
+		}
+	}
+	{
+		// Sort pass 3
+		const ValueType* RESTRICT s = (const ValueType* RESTRICT)Src;
+		ValueType* RESTRICT d = Dst;
+		for( CountType i = 0; i < Num; i++ )
+		{
+			ValueType Value = s[i];
+			uint32 Key = SortKey( Value );
+			d[ ++Histogram2[ ( (Key >> 21) & 2047 ) ] ] = Value;
+		}
+	}
+}
+
+
+template< typename T >
+struct TRadixSortKeyCastUint32
+{
+	FORCEINLINE uint32 operator()( const T& Value ) const
+	{
+		return (uint32)A;
+	}
+};
+
+template< typename ValueType, typename CountType >
+void RadixSort32( ValueType* RESTRICT Dst, ValueType* RESTRICT Src, CountType Num )
+{
+	RadixSort32( Dst, Src, Num, TRadixSortKeyCastUint32< ValueType >() );
+}
+
+// float cast to uint32 which maintains sorted order
+// http://codercorner.com/RadixSortRevisited.htm
+struct FRadixSortKeyFloat
+{
+	FORCEINLINE uint32 operator()( float Value ) const
+	{
+		union { float f; uint32 i; } v;
+		v.f = Value;
+
+		uint32 mask = -int32( v.i >> 31 ) | 0x80000000;
+		return v.i ^ mask;
+	}
+};
+
+template< typename CountType >
+void RadixSort32( float* RESTRICT Dst, float* RESTRICT Src, CountType Num )
+{
+	RadixSort32( Dst, Src, Num, FRadixSortKeyFloat() );
+}
\ No newline at end of file
diff --git a/Engine/Source/Runtime/Renderer/Private/VT/TexturePagePool.cpp b/Engine/Source/Runtime/Renderer/Private/VT/TexturePagePool.cpp
--- a/Engine/Source/Runtime/Renderer/Private/VT/TexturePagePool.cpp
+++ b/Engine/Source/Runtime/Renderer/Private/VT/TexturePagePool.cpp
@@ -5,79 +5,6 @@
 #include "VirtualTextureSpace.h"
 #include "VirtualTextureSystem.h"
 
-template< typename IndexType, typename CountType >
-void RadixSort32( IndexType* RESTRICT Dst, IndexType* RESTRICT Src, CountType Num )
-{
-	CountType Histograms[ 1024 + 2048 + 2048 ];
-	CountType* RESTRICT Histogram0 = Histograms + 0;
-	CountType* RESTRICT Histogram1 = Histogram0 + 1024;
-	CountType* RESTRICT Histogram2 = Histogram1 + 2048;
-
-	FMemory::Memzero( Histograms, sizeof( Histograms ) );
-
-	{
-		// Parallel histogram generation pass
-		const IndexType* RESTRICT s = (const IndexType* RESTRICT)Src;
-		for( CountType i = 0; i < Num; i++ )
-		{
-			IndexType j = s[i];
-			Histogram0[ ( j >>  0 ) & 1023 ]++;
-			Histogram1[ ( j >> 10 ) & 2047 ]++;
-			Histogram2[ ( j >> 21 ) & 2047 ]++;
-		}
-	}
-	{
-		// Set each histogram entry to the sum of entries preceding it
-		CountType Sum0 = 0;
-		CountType Sum1 = 0;
-		CountType Sum2 = 0;
-		for( CountType i = 0; i < 1024; i++ )
-		{
-			CountType t;
-			t = Histogram0[i] + Sum0; Histogram0[i] = Sum0 - 1; Sum0 = t;
-			t = Histogram1[i] + Sum1; Histogram1[i] = Sum1 - 1; Sum1 = t;
-			t = Histogram2[i] + Sum2; Histogram2[i] = Sum2 - 1; Sum2 = t;
-		}
-		for( CountType i = 1024; i < 2048; i++ )
-		{
-			CountType t;
-			t = Histogram1[i] + Sum1; Histogram1[i] = Sum1 - 1; Sum1 = t;
-			t = Histogram2[i] + Sum2; Histogram2[i] = Sum2 - 1; Sum2 = t;
-		}
-	}
-	{
-		// Sort pass 1
-		const IndexType* RESTRICT s = (const IndexType* RESTRICT)Src;
-		IndexType* RESTRICT d = Dst;
-		for( CountType i = 0; i < Num; i++ )
-		{
-			IndexType k = s[i];
-			d[ ++Histogram0[ ( (k >> 0) & 1023 ) ] ] = k;
-		}
-	}
-	{
-		// Sort pass 2
-		const IndexType* RESTRICT s = (const IndexType* RESTRICT)Dst;
-		IndexType* RESTRICT d = Src;
-		for( CountType i = 0; i < Num; i++ )
-		{
-			IndexType k = s[i];
-			d[ ++Histogram1[ ( (k >> 10) & 2047 ) ] ] = k;
-		}
-	}
-	{
-		// Sort pass 3
-		const IndexType* RESTRICT s = (const IndexType* RESTRICT)Src;
-		IndexType* RESTRICT d = Dst;
-		for( CountType i = 0; i < Num; i++ )
-		{
-			IndexType k = s[i];
-			d[ ++Histogram2[ ( (k >> 21) & 2047 ) ] ] = k;
-		}
-	}
-}
-
-
 FORCEINLINE uint64 EncodeSortKey( uint8 ID, uint8 vLevel, uint64 vAddress )
 {
 	uint64 Key;
