Fixed compile error.

diff --git a/Engine/Source/Runtime/Core/Public/Templates/Sorting.h b/Engine/Source/Runtime/Core/Public/Templates/Sorting.h
--- a/Engine/Source/Runtime/Core/Public/Templates/Sorting.h
+++ b/Engine/Source/Runtime/Core/Public/Templates/Sorting.h
@@ -505,7 +505,7 @@ struct TRadixSortKeyCastUint32
 {
 	FORCEINLINE uint32 operator()( const T& Value ) const
 	{
-		return (uint32)A;
+		return (uint32)Value;
 	}
 };
 
