Fix FConfigFile::CombineFromBuffer to correctly handle things like \' in quoted strings
This was already fixed for the other config load path a year ago in 3566672
Copy of 4666806

diff --git a/Engine/Source/Runtime/Core/Private/Misc/ConfigCacheIni.cpp b/Engine/Source/Runtime/Core/Private/Misc/ConfigCacheIni.cpp
--- a/Engine/Source/Runtime/Core/Private/Misc/ConfigCacheIni.cpp
+++ b/Engine/Source/Runtime/Core/Private/Misc/ConfigCacheIni.cpp
@@ -574,40 +574,7 @@ void FConfigFile::CombineFromBuffer(const FString& Buffer)
 				// If this line is delimited by quotes
 				if( *Value=='\"' )
 				{
-					Value++;
-					//epic moelfke: fixed handling of escaped characters in quoted string
-					while (*Value && *Value != '\"')
-					{
-						if (*Value != '\\') // unescaped character
-						{
-							ProcessedValue += *Value++;
-						}
-						else if (*++Value == '\\') // escaped forward slash "\\"
-						{
-							ProcessedValue += '\\';
-							Value++;
-						}
-						else if (*Value == '\"') // escaped double quote "\""
-						{
-							ProcessedValue += '\"';
-							Value++;
-						}
-						else if ( *Value == TEXT('n') )
-						{
-							ProcessedValue += TEXT('\n');
-							Value++;
-						}
-						else if( *Value == TEXT('u') && Value[1] && Value[2] && Value[3] && Value[4] )	// \uXXXX - UNICODE code point
-						{
-							ProcessedValue += (TCHAR)(FParse::HexDigit(Value[1])*(1<<12) + FParse::HexDigit(Value[2])*(1<<8) + FParse::HexDigit(Value[3])*(1<<4) + FParse::HexDigit(Value[4]));
-							Value += 5;
-						}
-						else if( Value[1] ) // some other escape sequence, assume it's a hex character value
-						{
-							ProcessedValue += (TCHAR)(FParse::HexDigit(Value[0])*16 + FParse::HexDigit(Value[1]));
-							Value += 2;
-						}
-					}
+					FParse::QuotedString(Value, ProcessedValue);
 				}
 				else
 				{
