Fixed some places that were clobbering LocRes source info

The text localization manager would reset bIsLocalized and LocResID if it failed to find a matching translation, but this isn't a valid thing to do as that text might still be localized... just not in the current language (or we may be applying a partial set of localization data, eg, when entering PIE).

diff --git a/Engine/Source/Runtime/Core/Private/Internationalization/TextLocalizationManager.cpp b/Engine/Source/Runtime/Core/Private/Internationalization/TextLocalizationManager.cpp
--- a/Engine/Source/Runtime/Core/Private/Internationalization/TextLocalizationManager.cpp
+++ b/Engine/Source/Runtime/Core/Private/Internationalization/TextLocalizationManager.cpp
@@ -952,6 +952,8 @@ void FTextLocalizationManager::UpdateFromNative(const FTextLocalizationResource&
 			// Update the display string with the new native string
 			if (SourceEntryForUpdate && LiveStringEntry.SourceStringHash == SourceEntryForUpdate->SourceStringHash)
 			{
+				LiveStringEntry.bIsLocalized = true;
+				LiveStringEntry.LocResID = SourceEntryForUpdate->LocResID;
 				*LiveStringEntry.DisplayString = SourceEntryForUpdate->LocalizedString;
 			}
 			else
@@ -965,14 +967,13 @@ void FTextLocalizationManager::UpdateFromNative(const FTextLocalizationResource&
 				// Restore the pre-leet state (if any)
 				if (!LiveStringEntry.NativeStringBackup.IsEmpty())
 				{
+					LiveStringEntry.bIsLocalized = false;
+					LiveStringEntry.LocResID = FTextLocalizationResourceId();
 					*LiveStringEntry.DisplayString = MoveTemp(LiveStringEntry.NativeStringBackup);
 				}
 #endif
 			}
 
-			LiveStringEntry.LocResID = FTextLocalizationResourceId();
-			LiveStringEntry.bIsLocalized = false;
-
 #if ENABLE_LOC_TESTING
 			LiveStringEntry.NativeStringBackup.Reset();
 #endif
@@ -997,8 +998,8 @@ void FTextLocalizationManager::UpdateFromNative(const FTextLocalizationResource&
 				if (!LiveStringEntry)
 				{
 					FDisplayStringLookupTable::FDisplayStringEntry NewLiveEntry(
-						false,																/*bIsLocalized*/
-						FTextLocalizationResourceId(),										/*LocResID*/
+						true,																/*bIsLocalized*/
+						NewEntry.LocResID,													/*LocResID*/
 						NewEntry.SourceStringHash,											/*SourceStringHash*/
 						MakeShared<FString, ESPMode::ThreadSafe>(NewEntry.LocalizedString)	/*String*/
 						);
@@ -1069,14 +1070,13 @@ void FTextLocalizationManager::UpdateFromLocalizations(TArrayView<const TSharedP
 					*(LiveStringEntry.DisplayString) = FString();
 				}
 
-				LiveStringEntry.bIsLocalized = false;
-				LiveStringEntry.LocResID = FTextLocalizationResourceId();
-
 #if ENABLE_LOC_TESTING
-				const bool bShouldLEETIFYUnlocalizedString = FParse::Param(FCommandLine::Get(), TEXT("LEETIFYUnlocalized"));
-				if(bShouldLEETIFYUnlocalizedString )
+				static const bool bShouldLEETIFYUnlocalizedString = FParse::Param(FCommandLine::Get(), TEXT("LEETIFYUnlocalized"));
+				if (bShouldLEETIFYUnlocalizedString)
 				{
-					FInternationalization::Leetify(*(LiveStringEntry.DisplayString));
+					LiveStringEntry.bIsLocalized = false;
+					LiveStringEntry.LocResID = FTextLocalizationResourceId();
+					FInternationalization::Leetify(*LiveStringEntry.DisplayString);
 				}
 #endif
 			}
