Fixed localization manager not always correctly applying related translation source information to newly loaded text instances

This would happen in the editor (which has package localization IDs on instances) when the equivalent translation was matched. It would take the correct display string, but lose the information about where that display string came from.

This was mostly fine, but it prevented the translation picker being able to locale the source localization target for the in-editor text.

diff --git a/Engine/Source/Runtime/Core/Private/Internationalization/TextLocalizationManager.cpp b/Engine/Source/Runtime/Core/Private/Internationalization/TextLocalizationManager.cpp
--- a/Engine/Source/Runtime/Core/Private/Internationalization/TextLocalizationManager.cpp
+++ b/Engine/Source/Runtime/Core/Private/Internationalization/TextLocalizationManager.cpp
@@ -557,13 +557,13 @@ FTextDisplayStringRef FTextLocalizationManager::GetDisplayString(const FString&
 
 	// In builds with stable keys enabled, we want to use the display string from the "clean" version of the text (if the sources match) as this is the only version that is translated
 	const FString* DisplayString = SourceString;
+	FDisplayStringLookupTable::FDisplayStringEntry* DisplayLiveEntry = nullptr;
 #if USE_STABLE_LOCALIZATION_KEYS
 	if (GIsEditor)
 	{
 		const FString DisplayNamespace = TextNamespaceUtil::StripPackageNamespace(Namespace);
 
 		FDisplayStringLookupTable::FKeysTable* DisplayLiveKeyTable = nullptr;
-		FDisplayStringLookupTable::FDisplayStringEntry* DisplayLiveEntry = nullptr;
 		DisplayStringLookupTable.Find(DisplayNamespace, DisplayLiveKeyTable, Key, DisplayLiveEntry);
 
 		if (DisplayLiveEntry && (!SourceString || DisplayLiveEntry->SourceStringHash == SourceStringHash))
@@ -605,6 +605,27 @@ FTextDisplayStringRef FTextLocalizationManager::GetDisplayString(const FString&
 
 		return LiveEntry->DisplayString;
 	}
+	// Entry is absent, but has a related entry to clone.
+	else if (DisplayLiveEntry && SourceString)
+	{
+		check(DisplayLiveEntry->SourceStringHash == SourceStringHash);
+		check(&DisplayLiveEntry->DisplayString.Get() == DisplayString && DisplayString);
+
+		// Clone the entry for the active ID, and assign it a new display string instance (as all entries must have a unique display string instance).
+		FDisplayStringLookupTable::FDisplayStringEntry NewEntry(*DisplayLiveEntry);
+		NewEntry.DisplayString = MakeShared<FString, ESPMode::ThreadSafe>(*DisplayString);
+
+		if (!LiveKeyTable)
+		{
+			LiveKeyTable = &(DisplayStringLookupTable.NamespacesTable.Add(Namespace, FDisplayStringLookupTable::FKeysTable()));
+		}
+
+		LiveKeyTable->Add(Key, NewEntry);
+
+		NamespaceKeyLookupTable.Add(NewEntry.DisplayString, FNamespaceKeyEntry(Namespace, Key));
+
+		return NewEntry.DisplayString;
+	}
 	// Entry is absent.
 	else
 	{
