Optimized localization conflict detection loop and log the hash as a hex value

diff --git a/Engine/Source/Runtime/Core/Private/Internationalization/TextLocalizationManager.cpp b/Engine/Source/Runtime/Core/Private/Internationalization/TextLocalizationManager.cpp
--- a/Engine/Source/Runtime/Core/Private/Internationalization/TextLocalizationManager.cpp
+++ b/Engine/Source/Runtime/Core/Private/Internationalization/TextLocalizationManager.cpp
@@ -801,7 +801,9 @@ void FTextLocalizationManager::UpdateFromLocalizationResource(const FString& Loc
 {
 	TSharedPtr<FTextLocalizationResource> TextLocalizationResource = MakeShared<FTextLocalizationResource>();
 	TextLocalizationResource->LoadFromFile(LocalizationResourceFilePath);
+#if !UE_BUILD_SHIPPING
 	TextLocalizationResource->DetectAndLogConflicts();
+#endif
 
 	UpdateFromLocalizations(TArrayView<const TSharedPtr<FTextLocalizationResource>>(&TextLocalizationResource, 1));
 }
diff --git a/Engine/Source/Runtime/Core/Private/Internationalization/TextLocalizationResource.cpp b/Engine/Source/Runtime/Core/Private/Internationalization/TextLocalizationResource.cpp
--- a/Engine/Source/Runtime/Core/Private/Internationalization/TextLocalizationResource.cpp
+++ b/Engine/Source/Runtime/Core/Private/Internationalization/TextLocalizationResource.cpp
@@ -387,32 +387,28 @@ void FTextLocalizationResource::DetectAndLogConflicts() const
 			const FString& KeyName = KeyEntry.Key;
 			const FEntryArray& EntryArray = KeyEntry.Value;
 
-			bool WasConflictDetected = false;
-			for (int32 k = 0; k < EntryArray.Num(); ++k)
+			bool bWasConflictDetected = false;
+			if (EntryArray.Num() > 1)
 			{
-				const FEntry& LeftEntry = EntryArray[k];
-				for (int32 l = k + 1; l < EntryArray.Num(); ++l)
+				const FEntry& PrimaryEntry = EntryArray[0];
+				for (int32 OtherEntryIndex = 1; OtherEntryIndex < EntryArray.Num() && !bWasConflictDetected; ++OtherEntryIndex)
 				{
-					const FEntry& RightEntry = EntryArray[l];
-					const bool bDoesSourceStringHashDiffer = LeftEntry.SourceStringHash != RightEntry.SourceStringHash;
-					const bool bDoesLocalizedStringDiffer = !LeftEntry.LocalizedString.Equals(RightEntry.LocalizedString, ESearchCase::CaseSensitive);
-					WasConflictDetected = bDoesSourceStringHashDiffer || bDoesLocalizedStringDiffer;
+					const FEntry& OtherEntry = EntryArray[OtherEntryIndex];
+					bWasConflictDetected = PrimaryEntry.SourceStringHash != OtherEntry.SourceStringHash || !PrimaryEntry.LocalizedString.Equals(OtherEntry.LocalizedString, ESearchCase::CaseSensitive);
 				}
 			}
 
-			if (WasConflictDetected)
+			if (bWasConflictDetected)
 			{
 				FString CollidingEntryListString;
-				for (int32 k = 0; k < EntryArray.Num(); ++k)
+				for (const FEntry& Entry : EntryArray)
 				{
-					const FEntry& Entry = EntryArray[k];
-
-					if (!(CollidingEntryListString.IsEmpty()))
+					if (!CollidingEntryListString.IsEmpty())
 					{
 						CollidingEntryListString += TEXT('\n');
 					}
 
-					CollidingEntryListString += FString::Printf(TEXT("    Localization Resource: (%s) Source String Hash: (%d) Localized String: (%s)"), *(Entry.LocResID.GetString()), Entry.SourceStringHash, *(Entry.LocalizedString));
+					CollidingEntryListString += FString::Printf(TEXT("    Localization Resource: (%s) Source String Hash: (0x%08x) Localized String: (%s)"), *(Entry.LocResID.GetString()), Entry.SourceStringHash, *(Entry.LocalizedString));
 				}
 
 				UE_LOG(LogTextLocalizationResource, Warning, TEXT("Loaded localization resources contain conflicting entries for (Namespace:%s, Key:%s):\n%s"), *NamespaceName, *KeyName, *CollidingEntryListString);
