Fixed assert processing localized display strings

#rnx

diff --git a/Engine/Source/Runtime/Core/Private/Internationalization/TextLocalizationManager.cpp b/Engine/Source/Runtime/Core/Private/Internationalization/TextLocalizationManager.cpp
--- a/Engine/Source/Runtime/Core/Private/Internationalization/TextLocalizationManager.cpp
+++ b/Engine/Source/Runtime/Core/Private/Internationalization/TextLocalizationManager.cpp
@@ -571,9 +571,16 @@ FTextDisplayStringRef FTextLocalizationManager::GetDisplayString(const FString&
 		FDisplayStringLookupTable::FKeysTable* DisplayLiveKeyTable = nullptr;
 		DisplayStringLookupTable.Find(DisplayNamespace, DisplayLiveKeyTable, Key, DisplayLiveEntry);
 
-		if (DisplayLiveEntry && (!SourceString || DisplayLiveEntry->SourceStringHash == SourceStringHash))
+		if (DisplayLiveEntry)
 		{
-			DisplayString = &DisplayLiveEntry->DisplayString.Get();
+			if (!SourceString || DisplayLiveEntry->SourceStringHash == SourceStringHash)
+			{
+				DisplayString = &DisplayLiveEntry->DisplayString.Get();
+			}
+			else
+			{
+				DisplayLiveEntry = nullptr;
+			}
 		}
 	}
 #endif // USE_STABLE_LOCALIZATION_KEYS
@@ -611,10 +618,10 @@ FTextDisplayStringRef FTextLocalizationManager::GetDisplayString(const FString&
 		return LiveEntry->DisplayString;
 	}
 	// Entry is absent, but has a related entry to clone.
-	else if (DisplayLiveEntry && SourceString)
+	else if (DisplayLiveEntry)
 	{
-		check(DisplayLiveEntry->SourceStringHash == SourceStringHash);
-		check(&DisplayLiveEntry->DisplayString.Get() == DisplayString && DisplayString);
+		check(!SourceString || DisplayLiveEntry->SourceStringHash == SourceStringHash);
+		check(DisplayString && &DisplayLiveEntry->DisplayString.Get() == DisplayString);
 
 		// Clone the entry for the active ID, and assign it a new display string instance (as all entries must have a unique display string instance).
 		FDisplayStringLookupTable::FDisplayStringEntry NewEntry(*DisplayLiveEntry);
