Added comment warning people not to use our regex types as file-level variables

#rnx

diff --git a/Engine/Source/Runtime/Core/Public/Internationalization/Regex.h b/Engine/Source/Runtime/Core/Public/Internationalization/Regex.h
--- a/Engine/Source/Runtime/Core/Public/Internationalization/Regex.h
+++ b/Engine/Source/Runtime/Core/Public/Internationalization/Regex.h
@@ -5,12 +5,17 @@
 #include "Containers/UnrealString.h"
 #include "Templates/SharedPointer.h"
 
-class FRegexMatcherImplementation;
 class FRegexPatternImplementation;
+class FRegexMatcherImplementation;
 
+/**
+ * Implements a regular expression pattern.
+ * @note DO NOT use this class as a file-level variable as its construction relies on the internationalization system being initialized!
+ */
 class CORE_API FRegexPattern
 {
 	friend class FRegexMatcher;
+
 public:
 	FRegexPattern(const FString& SourceString);
 
@@ -18,8 +23,10 @@ private:
 	TSharedRef<FRegexPatternImplementation> Implementation;
 };
 
-class FRegexMatcherImplementation;
-
+/**
+ * Implements a regular expression pattern matcher.
+ * @note DO NOT use this class as a file-level variable as its construction relies on the internationalization system being initialized!
+ */
 class CORE_API FRegexMatcher
 {
 public:
