Make sure FPlatformMisc::GetXcodePath()  returns a valid path only if Xcode app was actually found

diff --git a/Engine/Source/Runtime/Core/Private/Mac/MacPlatformMisc.cpp b/Engine/Source/Runtime/Core/Private/Mac/MacPlatformMisc.cpp
--- a/Engine/Source/Runtime/Core/Private/Mac/MacPlatformMisc.cpp
+++ b/Engine/Source/Runtime/Core/Private/Mac/MacPlatformMisc.cpp
@@ -151,7 +151,8 @@ struct FMacApplicationInfo
 					}
 				}
 			}
-			else
+
+			if (XcodeVersion.majorVersion == 0)
 			{
 				XcodePath.Empty();
 			}
