Optimized LoadModuleChecked by not recording stats when module is already loaded

diff --git a/Engine/Source/Runtime/Core/Private/Modules/ModuleManager.cpp b/Engine/Source/Runtime/Core/Private/Modules/ModuleManager.cpp
--- a/Engine/Source/Runtime/Core/Private/Modules/ModuleManager.cpp
+++ b/Engine/Source/Runtime/Core/Private/Modules/ModuleManager.cpp
@@ -353,7 +353,7 @@ IModuleInterface* FModuleManager::LoadModule( const FName InModuleName )
 	IModuleInterface* Result = LoadModuleWithFailureReason(InModuleName, FailureReason );
 
 	// This should return a valid pointer only if and only if the module is loaded
-	check((Result != nullptr) == IsModuleLoaded(InModuleName));
+	checkSlow((Result != nullptr) == IsModuleLoaded(InModuleName));
 
 	return Result;
 }
@@ -374,6 +374,16 @@ IModuleInterface* FModuleManager::LoadModuleWithFailureReason(const FName InModu
 	ensureMsgf(IsInGameThread(), TEXT("ModuleManager: Attempting to load '%s' outside the main thread.  Please call LoadModule on the main/game thread only.  You can use GetModule or GetModuleChecked instead, those are safe to call outside the game thread."), *InModuleName.ToString());
 #endif
 
+	// Return early if module already loaded
+	ModuleInfoRef* FoundModuleInfo = Modules.Find(InModuleName);
+	if (FoundModuleInfo && (*FoundModuleInfo)->Module.IsValid())
+	{
+#if !(UE_BUILD_SHIPPING || UE_BUILD_TEST)
+		WarnIfItWasntSafeToLoadHere(InModuleName);
+#endif
+		return (*FoundModuleInfo)->Module.Get();
+	}
+
 	DECLARE_SCOPE_CYCLE_COUNTER(TEXT("Module Load"), STAT_ModuleLoad, STATGROUP_LoadTime);
 
 #if	STATS
@@ -565,6 +575,9 @@ IModuleInterface* FModuleManager::LoadModuleWithFailureReason(const FName InModu
 			}
 		}
 #endif
+
+		// This should return a valid pointer only if and only if the module is loaded
+		check((LoadedModule != nullptr) == IsModuleLoaded(InModuleName));
 	}
 
 	return LoadedModule;
