Merged ModuleManager optimization from CL#4355294

diff --git a/Engine/Source/Runtime/Core/Private/Modules/ModuleManager.cpp b/Engine/Source/Runtime/Core/Private/Modules/ModuleManager.cpp
--- a/Engine/Source/Runtime/Core/Private/Modules/ModuleManager.cpp
+++ b/Engine/Source/Runtime/Core/Private/Modules/ModuleManager.cpp
@@ -374,18 +374,26 @@ IModuleInterface* FModuleManager::LoadModuleWithFailureReason(const FName InModu
 	ensureMsgf(IsInGameThread(), TEXT("ModuleManager: Attempting to load '%s' outside the main thread.  Please call LoadModule on the main/game thread only.  You can use GetModule or GetModuleChecked instead, those are safe to call outside the game thread."), *InModuleName.ToString());
 #endif
 
-	// Return early if module already loaded
-	ModuleInfoRef* FoundModuleInfo = Modules.Find(InModuleName);
-	if (FoundModuleInfo && (*FoundModuleInfo)->Module.IsValid())
+	IModuleInterface* LoadedModule = nullptr;
+	OutFailureReason = EModuleLoadResult::Success;
+
+	// Do fast check for existing module, this is the most common case
+	ModuleInfoPtr FoundModulePtr = FindModule(InModuleName);
+
+	if (FoundModulePtr.IsValid())
 	{
+		LoadedModule = FoundModulePtr->Module.Get();
+
+		if (LoadedModule)
+		{
 #if !(UE_BUILD_SHIPPING || UE_BUILD_TEST)
-		WarnIfItWasntSafeToLoadHere(InModuleName);
+			WarnIfItWasntSafeToLoadHere(InModuleName);
 #endif
-		return (*FoundModuleInfo)->Module.Get();
+			return LoadedModule;
+		}
 	}
 
 	DECLARE_SCOPE_CYCLE_COUNTER(TEXT("Module Load"), STAT_ModuleLoad, STATGROUP_LoadTime);
-
 #if	STATS
 	// This is fine here, we only load a handful of modules.
 	static FString Module = TEXT( "Module" );
@@ -394,191 +402,178 @@ IModuleInterface* FModuleManager::LoadModuleWithFailureReason(const FName InModu
 	FScopeCycleCounter CycleCounter( StatId );
 #endif // STATS
 
-	IModuleInterface* LoadedModule = nullptr;
-	OutFailureReason = EModuleLoadResult::Success;
-
-	// Update our set of known modules, in case we don't already know about this module
-	AddModule( InModuleName );
-
-	// Grab the module info.  This has the file name of the module, as well as other info.
-	ModuleInfoRef ModuleInfo = Modules.FindChecked( InModuleName );
-
-	if (ModuleInfo->Module.IsValid())
+	if (!FoundModulePtr.IsValid())
 	{
-		// Assign the already loaded module into the return value, otherwise the return value gives the impression the module failed load!
-		LoadedModule = ModuleInfo->Module.Get();
+		// Update our set of known modules, in case we don't already know about this module
+		AddModule(InModuleName);
 
-#if !(UE_BUILD_SHIPPING || UE_BUILD_TEST)
-		WarnIfItWasntSafeToLoadHere(InModuleName);
-#endif
+		// Ptr will always be valid at this point
+		FoundModulePtr = FindModule(InModuleName);
 	}
-	else
+	
+	// Grab the module info.  This has the file name of the module, as well as other info.
+	ModuleInfoRef ModuleInfo = FoundModulePtr.ToSharedRef();
+
+	// Make sure this isn't a module that we had previously loaded, and then unloaded at shutdown time.
+	//
+	// If this assert goes off, your trying to load a module during the shutdown phase that was already
+	// cleaned up.  The easiest way to fix this is to change your code to query for an already-loaded
+	// module instead of trying to load it directly.
+	checkf((!ModuleInfo->bWasUnloadedAtShutdown), TEXT("Attempted to load module '%s' that was already unloaded at shutdown.  FModuleManager::LoadModule() was called to load a module that was previously loaded, and was unloaded at shutdown time.  If this assert goes off, your trying to load a module during the shutdown phase that was already cleaned up.  The easiest way to fix this is to change your code to query for an already-loaded module instead of trying to load it directly."), *InModuleName.ToString());
+
+	// Check if we're statically linked with the module.  Those modules register with the module manager using a static variable,
+	// so hopefully we already know about the name of the module and how to initialize it.
+	const FInitializeStaticallyLinkedModule* ModuleInitializerPtr = StaticallyLinkedModuleInitializers.Find(InModuleName);
+	if (ModuleInitializerPtr != nullptr)
 	{
-		// Make sure this isn't a module that we had previously loaded, and then unloaded at shutdown time.
-		//
-		// If this assert goes off, your trying to load a module during the shutdown phase that was already
-		// cleaned up.  The easiest way to fix this is to change your code to query for an already-loaded
-		// module instead of trying to load it directly.
-		checkf((!ModuleInfo->bWasUnloadedAtShutdown), TEXT("Attempted to load module '%s' that was already unloaded at shutdown.  FModuleManager::LoadModule() was called to load a module that was previously loaded, and was unloaded at shutdown time.  If this assert goes off, your trying to load a module during the shutdown phase that was already cleaned up.  The easiest way to fix this is to change your code to query for an already-loaded module instead of trying to load it directly."), *InModuleName.ToString());
-
-		// Check if we're statically linked with the module.  Those modules register with the module manager using a static variable,
-		// so hopefully we already know about the name of the module and how to initialize it.
-		const FInitializeStaticallyLinkedModule* ModuleInitializerPtr = StaticallyLinkedModuleInitializers.Find(InModuleName);
-		if (ModuleInitializerPtr != nullptr)
-		{
-			const FInitializeStaticallyLinkedModule& ModuleInitializer(*ModuleInitializerPtr);
+		const FInitializeStaticallyLinkedModule& ModuleInitializer(*ModuleInitializerPtr);
 
-			// Initialize the module!
-			ModuleInfo->Module = TUniquePtr<IModuleInterface>(ModuleInitializer.Execute());
+		// Initialize the module!
+		ModuleInfo->Module = TUniquePtr<IModuleInterface>(ModuleInitializer.Execute());
 
-			if (ModuleInfo->Module.IsValid())
-			{
-				// Startup the module
-				ModuleInfo->Module->StartupModule();
-				// The module might try to load other dependent modules in StartupModule. In this case, we want those modules shut down AFTER this one because we may still depend on the module at shutdown.
-				ModuleInfo->LoadOrder = FModuleInfo::CurrentLoadOrder++;
+		if (ModuleInfo->Module.IsValid())
+		{
+			// Startup the module
+			ModuleInfo->Module->StartupModule();
+			// The module might try to load other dependent modules in StartupModule. In this case, we want those modules shut down AFTER this one because we may still depend on the module at shutdown.
+			ModuleInfo->LoadOrder = FModuleInfo::CurrentLoadOrder++;
 
-				// Module was started successfully!  Fire callbacks.
-				ModulesChangedEvent.Broadcast(InModuleName, EModuleChangeReason::ModuleLoaded);
+			// Module was started successfully!  Fire callbacks.
+			ModulesChangedEvent.Broadcast(InModuleName, EModuleChangeReason::ModuleLoaded);
 
-				// Set the return parameter
-				LoadedModule = ModuleInfo->Module.Get();
-			}
-			else
-			{
-				UE_LOG(LogModuleManager, Warning, TEXT("ModuleManager: Unable to load module '%s' because InitializeModule function failed (returned nullptr.)"), *InModuleName.ToString());
-				OutFailureReason = EModuleLoadResult::FailedToInitialize;
-			}
+			// Set the return parameter
+			LoadedModule = ModuleInfo->Module.Get();
 		}
-#if IS_MONOLITHIC
 		else
 		{
-			// Monolithic builds that do not have the initializer were *not found* during the build step, so return FileNotFound
-			// (FileNotFound is an acceptable error in some case - ie loading a content only project)
-			UE_LOG(LogModuleManager, Warning, TEXT("ModuleManager: Module '%s' not found - its StaticallyLinkedModuleInitializers function is null."), *InModuleName.ToString());
-			OutFailureReason = EModuleLoadResult::FileNotFound;
+			UE_LOG(LogModuleManager, Warning, TEXT("ModuleManager: Unable to load module '%s' because InitializeModule function failed (returned nullptr.)"), *InModuleName.ToString());
+			OutFailureReason = EModuleLoadResult::FailedToInitialize;
 		}
-#endif
-#if !IS_MONOLITHIC
-		else
+	}
+#if IS_MONOLITHIC
+	else
+	{
+		// Monolithic builds that do not have the initializer were *not found* during the build step, so return FileNotFound
+		// (FileNotFound is an acceptable error in some case - ie loading a content only project)
+		UE_LOG(LogModuleManager, Warning, TEXT("ModuleManager: Module '%s' not found - its StaticallyLinkedModuleInitializers function is null."), *InModuleName.ToString());
+		OutFailureReason = EModuleLoadResult::FileNotFound;
+	}
+#else
+	else
+	{
+		// Make sure that any UObjects that need to be registered were already processed before we go and
+		// load another module.  We just do this so that we can easily tell whether UObjects are present
+		// in the module being loaded.
+		if (bCanProcessNewlyLoadedObjects)
 		{
-			// Make sure that any UObjects that need to be registered were already processed before we go and
-			// load another module.  We just do this so that we can easily tell whether UObjects are present
-			// in the module being loaded.
-			if (bCanProcessNewlyLoadedObjects)
-			{
-				ProcessLoadedObjectsCallback.Broadcast();
-			}
-
-			// Try to dynamically load the DLL
+			ProcessLoadedObjectsCallback.Broadcast();
+		}
 
-			UE_LOG(LogModuleManager, Verbose, TEXT("ModuleManager: Load Module '%s' DLL '%s'"), *InModuleName.ToString(), *ModuleInfo->Filename);
+		// Try to dynamically load the DLL
 
-			if (ModuleInfo->Filename.IsEmpty() || !FPaths::FileExists(ModuleInfo->Filename))
-			{
-				TMap<FName, FString> ModulePathMap;
-				FindModulePaths(*InModuleName.ToString(), ModulePathMap);
+		UE_LOG(LogModuleManager, Verbose, TEXT("ModuleManager: Load Module '%s' DLL '%s'"), *InModuleName.ToString(), *ModuleInfo->Filename);
 
-				if (ModulePathMap.Num() != 1)
-				{
-					UE_LOG(LogModuleManager, Warning, TEXT("ModuleManager: Unable to load module '%s'  - %d instances of that module name found."), *InModuleName.ToString(), ModulePathMap.Num());
-					OutFailureReason = EModuleLoadResult::FileNotFound;
-					return nullptr;
-				}
+		if (ModuleInfo->Filename.IsEmpty() || !FPaths::FileExists(ModuleInfo->Filename))
+		{
+			TMap<FName, FString> ModulePathMap;
+			FindModulePaths(*InModuleName.ToString(), ModulePathMap);
 
-				ModuleInfo->Filename = MoveTemp(TMap<FName, FString>::TIterator(ModulePathMap).Value());
+			if (ModulePathMap.Num() != 1)
+			{
+				UE_LOG(LogModuleManager, Warning, TEXT("ModuleManager: Unable to load module '%s'  - %d instances of that module name found."), *InModuleName.ToString(), ModulePathMap.Num());
+				OutFailureReason = EModuleLoadResult::FileNotFound;
+				return nullptr;
 			}
 
-			// Determine which file to load for this module.
-			const FString ModuleFileToLoad = FPaths::ConvertRelativePathToFull(ModuleInfo->Filename);
+			ModuleInfo->Filename = MoveTemp(TMap<FName, FString>::TIterator(ModulePathMap).Value());
+		}
+
+		// Determine which file to load for this module.
+		const FString ModuleFileToLoad = FPaths::ConvertRelativePathToFull(ModuleInfo->Filename);
 
-			// Clear the handle and set it again below if the module is successfully loaded
-			ModuleInfo->Handle = nullptr;
+		// Clear the handle and set it again below if the module is successfully loaded
+		ModuleInfo->Handle = nullptr;
 
-			// Skip this check if file manager has not yet been initialized
-			if (FPaths::FileExists(ModuleFileToLoad))
+		// Skip this check if file manager has not yet been initialized
+		if (FPaths::FileExists(ModuleFileToLoad))
+		{
+			ModuleInfo->Handle = FPlatformProcess::GetDllHandle(*ModuleFileToLoad);
+			if (ModuleInfo->Handle != nullptr)
 			{
-				ModuleInfo->Handle = FPlatformProcess::GetDllHandle(*ModuleFileToLoad);
-				if (ModuleInfo->Handle != nullptr)
+				// First things first.  If the loaded DLL has UObjects in it, then their generated code's
+				// static initialization will have run during the DLL loading phase, and we'll need to
+				// go in and make sure those new UObject classes are properly registered.
 				{
-					// First things first.  If the loaded DLL has UObjects in it, then their generated code's
-					// static initialization will have run during the DLL loading phase, and we'll need to
-					// go in and make sure those new UObject classes are properly registered.
+					// Sometimes modules are loaded before even the UObject systems are ready.  We need to assume
+					// these modules aren't using UObjects.
+					if (bCanProcessNewlyLoadedObjects)
 					{
-						// Sometimes modules are loaded before even the UObject systems are ready.  We need to assume
-						// these modules aren't using UObjects.
-						if (bCanProcessNewlyLoadedObjects)
-						{
-							// OK, we've verified that loading the module caused new UObject classes to be
-							// registered, so we'll treat this module as a module with UObjects in it.
-							ProcessLoadedObjectsCallback.Broadcast();
-						}
+						// OK, we've verified that loading the module caused new UObject classes to be
+						// registered, so we'll treat this module as a module with UObjects in it.
+						ProcessLoadedObjectsCallback.Broadcast();
 					}
+				}
 
-					// Find our "InitializeModule" global function, which must exist for all module DLLs
-					FInitializeModuleFunctionPtr InitializeModuleFunctionPtr =
-						(FInitializeModuleFunctionPtr)FPlatformProcess::GetDllExport(ModuleInfo->Handle, TEXT("InitializeModule"));
-					if (InitializeModuleFunctionPtr != nullptr)
+				// Find our "InitializeModule" global function, which must exist for all module DLLs
+				FInitializeModuleFunctionPtr InitializeModuleFunctionPtr =
+					(FInitializeModuleFunctionPtr)FPlatformProcess::GetDllExport(ModuleInfo->Handle, TEXT("InitializeModule"));
+				if (InitializeModuleFunctionPtr != nullptr)
+				{
+					if ( ModuleInfo->Module.IsValid() )
+					{
+						// Assign the already loaded module into the return value, otherwise the return value gives the impression the module failed load!
+						LoadedModule = ModuleInfo->Module.Get();
+					}
+					else
 					{
+						// Initialize the module!
+						ModuleInfo->Module = TUniquePtr<IModuleInterface>(InitializeModuleFunctionPtr());
+
 						if ( ModuleInfo->Module.IsValid() )
 						{
-							// Assign the already loaded module into the return value, otherwise the return value gives the impression the module failed load!
+							// Startup the module
+							ModuleInfo->Module->StartupModule();
+							// The module might try to load other dependent modules in StartupModule. In this case, we want those modules shut down AFTER this one because we may still depend on the module at shutdown.
+							ModuleInfo->LoadOrder = FModuleInfo::CurrentLoadOrder++;
+
+							// Module was started successfully!  Fire callbacks.
+							ModulesChangedEvent.Broadcast(InModuleName, EModuleChangeReason::ModuleLoaded);
+
+							// Set the return parameter
 							LoadedModule = ModuleInfo->Module.Get();
 						}
 						else
 						{
-							// Initialize the module!
-							ModuleInfo->Module = TUniquePtr<IModuleInterface>(InitializeModuleFunctionPtr());
-
-							if ( ModuleInfo->Module.IsValid() )
-							{
-								// Startup the module
-								ModuleInfo->Module->StartupModule();
-								// The module might try to load other dependent modules in StartupModule. In this case, we want those modules shut down AFTER this one because we may still depend on the module at shutdown.
-								ModuleInfo->LoadOrder = FModuleInfo::CurrentLoadOrder++;
-
-								// Module was started successfully!  Fire callbacks.
-								ModulesChangedEvent.Broadcast(InModuleName, EModuleChangeReason::ModuleLoaded);
-
-								// Set the return parameter
-								LoadedModule = ModuleInfo->Module.Get();
-							}
-							else
-							{
-								UE_LOG(LogModuleManager, Warning, TEXT("ModuleManager: Unable to load module '%s' because InitializeModule function failed (returned nullptr.)"), *ModuleFileToLoad);
-
-								FPlatformProcess::FreeDllHandle(ModuleInfo->Handle);
-								ModuleInfo->Handle = nullptr;
-								OutFailureReason = EModuleLoadResult::FailedToInitialize;
-							}
-						}
-					}
-					else
-					{
-						UE_LOG(LogModuleManager, Warning, TEXT("ModuleManager: Unable to load module '%s' because InitializeModule function was not found."), *ModuleFileToLoad);
+							UE_LOG(LogModuleManager, Warning, TEXT("ModuleManager: Unable to load module '%s' because InitializeModule function failed (returned nullptr.)"), *ModuleFileToLoad);
 
-						FPlatformProcess::FreeDllHandle(ModuleInfo->Handle);
-						ModuleInfo->Handle = nullptr;
-						OutFailureReason = EModuleLoadResult::FailedToInitialize;
+							FPlatformProcess::FreeDllHandle(ModuleInfo->Handle);
+							ModuleInfo->Handle = nullptr;
+							OutFailureReason = EModuleLoadResult::FailedToInitialize;
+						}
 					}
 				}
 				else
 				{
-					UE_LOG(LogModuleManager, Warning, TEXT("ModuleManager: Unable to load module '%s' because the file couldn't be loaded by the OS."), *ModuleFileToLoad);
-					OutFailureReason = EModuleLoadResult::CouldNotBeLoadedByOS;
+					UE_LOG(LogModuleManager, Warning, TEXT("ModuleManager: Unable to load module '%s' because InitializeModule function was not found."), *ModuleFileToLoad);
+
+					FPlatformProcess::FreeDllHandle(ModuleInfo->Handle);
+					ModuleInfo->Handle = nullptr;
+					OutFailureReason = EModuleLoadResult::FailedToInitialize;
 				}
 			}
 			else
 			{
-				UE_LOG(LogModuleManager, Warning, TEXT("ModuleManager: Unable to load module '%s' because the file '%s' was not found."), *InModuleName.ToString(), *ModuleFileToLoad);
-				OutFailureReason = EModuleLoadResult::FileNotFound;
+				UE_LOG(LogModuleManager, Warning, TEXT("ModuleManager: Unable to load module '%s' because the file couldn't be loaded by the OS."), *ModuleFileToLoad);
+				OutFailureReason = EModuleLoadResult::CouldNotBeLoadedByOS;
 			}
 		}
-#endif
-
-		// This should return a valid pointer only if and only if the module is loaded
-		check((LoadedModule != nullptr) == IsModuleLoaded(InModuleName));
+		else
+		{
+			UE_LOG(LogModuleManager, Warning, TEXT("ModuleManager: Unable to load module '%s' because the file '%s' was not found."), *InModuleName.ToString(), *ModuleFileToLoad);
+			OutFailureReason = EModuleLoadResult::FileNotFound;
+		}
 	}
+#endif
 
 	return LoadedModule;
 }
