Linux: Revert CL 4523990 as this is causing the final callstack to possibly hit an unknown chunk when symbolicating

Due to x86(-64) instructions the call sites point to the next instruction. So we -1 for all call sites that is not the crash call site. This means we need
to avoid the signal handler as well to not subtract one to the place we actually crashed in. Adding to the ignore count causes us to lose this guess of -1

diff --git a/Engine/Source/Runtime/Core/Private/Unix/UnixPlatformCrashContext.cpp b/Engine/Source/Runtime/Core/Private/Unix/UnixPlatformCrashContext.cpp
--- a/Engine/Source/Runtime/Core/Private/Unix/UnixPlatformCrashContext.cpp
+++ b/Engine/Source/Runtime/Core/Private/Unix/UnixPlatformCrashContext.cpp
@@ -284,10 +284,7 @@ void FUnixCrashContext::CaptureStackTrace()
 		ANSICHAR* StackTrace = (ANSICHAR*) FMemory::Malloc( StackTraceSize );
 		StackTrace[0] = 0;
 
-		// We use __builtin_return_address to save an address location to where our signal handler is called from
-		// We still need to ignore the first frame from pthreads which calls our signal handler, that is why we set it to one
-		// This may need to change something other then glibc is used
-		int32 IgnoreCount = 1;
+		int32 IgnoreCount = 0;
 		CapturePortableCallStack(IgnoreCount, this);
 
 		// Walk the stack and dump it to the allocated memory (do not ignore any stack frames to be consistent with check()/ensure() handling)
diff --git a/Engine/Source/Runtime/Core/Private/Unix/UnixPlatformStackWalk.cpp b/Engine/Source/Runtime/Core/Private/Unix/UnixPlatformStackWalk.cpp
--- a/Engine/Source/Runtime/Core/Private/Unix/UnixPlatformStackWalk.cpp
+++ b/Engine/Source/Runtime/Core/Private/Unix/UnixPlatformStackWalk.cpp
@@ -634,6 +634,7 @@ namespace
 			{
 				if (FirstCrashHandlerFrame == reinterpret_cast<uint64*>(BackTrace[i]))
 				{
+					i++;
 					uint64* OverwriteBackTrace = BackTrace;
 
 					for (int j = i; j < Size; j++)
