Copy of CL 4612619

Increasing FName table max size to 4194304

diff --git a/Engine/Source/Runtime/Core/Public/UObject/NameTypes.h b/Engine/Source/Runtime/Core/Public/UObject/NameTypes.h
--- a/Engine/Source/Runtime/Core/Public/UObject/NameTypes.h
+++ b/Engine/Source/Runtime/Core/Public/UObject/NameTypes.h
@@ -514,7 +514,7 @@ public:
 
 // Typedef for the threadsafe master name table. 
 // CAUTION: If you change those constants, you probably need to update the debug visualizers.
-typedef TStaticIndirectArrayThreadSafeRead<FNameEntry, 2 * 1024 * 1024 /* 2M unique FNames */, 16384 /* allocated in 64K/128K chunks */ > TNameEntryArray;
+typedef TStaticIndirectArrayThreadSafeRead<FNameEntry, 4 * 1024 * 1024 /* 4M unique FNames */, 16384 /* allocated in 64K/128K chunks */ > TNameEntryArray;
 
 /**
  * The minimum amount of data required to reconstruct a name
