[Switch] Add support for and enable MallocBinned3. Also clean up SwitchPlatformMemory.cpp a bit, a little refactoring and removing some unneeded code.

diff --git a/Engine/Source/Runtime/Core/Public/HAL/MallocBinned3.h b/Engine/Source/Runtime/Core/Public/HAL/MallocBinned3.h
--- a/Engine/Source/Runtime/Core/Public/HAL/MallocBinned3.h
+++ b/Engine/Source/Runtime/Core/Public/HAL/MallocBinned3.h
@@ -24,7 +24,11 @@
 #define BINNED3_MAX_LISTED_SMALL_POOL_SIZE	28672			// Maximum block size in GMallocBinned3SmallBlockSizes
 #define BINNED3_MAX_SMALL_POOL_SIZE			(128 * 1024)		// Maximum medium block size
 #define BINNED3_SMALL_POOL_COUNT			(49 + (BINNED3_MAX_SMALL_POOL_SIZE - BINNED3_MAX_LISTED_SMALL_POOL_SIZE) / BINNED3_BASE_PAGE_SIZE)
-#define MAX_MEMORY_PER_BLOCK_SIZE_SHIFT (30) // maximum of 1GB per block size
+#if PLATFORM_SWITCH
+	#define MAX_MEMORY_PER_BLOCK_SIZE_SHIFT (29) // maximum of 512MB per block size
+#else
+	#define MAX_MEMORY_PER_BLOCK_SIZE_SHIFT (30) // maximum of 1GB per block size
+#endif
 #define MAX_MEMORY_PER_BLOCK_SIZE (1ull << MAX_MEMORY_PER_BLOCK_SIZE_SHIFT) 
 
 // This choice depends on how efficient the OS is with sparse commits in large VM blocks
