Merging //UE4/Dev-Main to Dev-Anim (//UE4/Dev-Anim)

diff --git a/Engine/Plugins/Enterprise/VariantManagerContent/Source/VariantManagerContent/Private/PropertyValueTransform.cpp b/Engine/Plugins/Enterprise/VariantManagerContent/Source/VariantManagerContent/Private/PropertyValueTransform.cpp
--- a/Engine/Plugins/Enterprise/VariantManagerContent/Source/VariantManagerContent/Private/PropertyValueTransform.cpp
+++ b/Engine/Plugins/Enterprise/VariantManagerContent/Source/VariantManagerContent/Private/PropertyValueTransform.cpp
@@ -54,7 +54,7 @@ FVector UPropertyValueTransform::GetScale3D()
 	}
 }
 
-void UPropertyValueTransform::SetLocation(FVector NewValue)
+void UPropertyValueTransform::SetLocation(const FVector& NewValue)
 {
 	if (PropCategory == EPropertyValueCategory::RelativeLocation)
 	{
@@ -62,7 +62,7 @@ void UPropertyValueTransform::SetLocation(FVector NewValue)
 	}
 }
 
-void UPropertyValueTransform::SetRotation(FQuat NewValue)
+void UPropertyValueTransform::SetRotation(const FQuat& NewValue)
 {
 	if (PropCategory == EPropertyValueCategory::RelativeRotation)
 	{
@@ -71,7 +71,7 @@ void UPropertyValueTransform::SetRotation(FQuat NewValue)
 	}
 }
 
-void UPropertyValueTransform::SetScale3D(FVector NewValue)
+void UPropertyValueTransform::SetScale3D(const FVector& NewValue)
 {
 	if (PropCategory == EPropertyValueCategory::RelativeScale3D)
 	{
diff --git a/Engine/Plugins/Enterprise/VariantManagerContent/Source/VariantManagerContent/Public/PropertyValueTransform.h b/Engine/Plugins/Enterprise/VariantManagerContent/Source/VariantManagerContent/Public/PropertyValueTransform.h
--- a/Engine/Plugins/Enterprise/VariantManagerContent/Source/VariantManagerContent/Public/PropertyValueTransform.h
+++ b/Engine/Plugins/Enterprise/VariantManagerContent/Source/VariantManagerContent/Public/PropertyValueTransform.h
@@ -19,9 +19,9 @@ public:
 	FQuat GetRotation();
 	FVector GetScale3D();
 
-	void SetLocation(FVector NewValue);
-	void SetRotation(FQuat NewValue);
-	void SetScale3D(FVector NewValue);
+	void SetLocation(const FVector& NewValue);
+	void SetRotation(const FQuat& NewValue);
+	void SetScale3D(const FVector& NewValue);
 
 	virtual void ApplyDataToResolvedObject() override;
 };
\ No newline at end of file
diff --git a/Engine/Plugins/Runtime/Database/ADOSupport/ADOSupport.uplugin b/Engine/Plugins/Runtime/Database/ADOSupport/ADOSupport.uplugin
--- a/Engine/Plugins/Runtime/Database/ADOSupport/ADOSupport.uplugin
+++ b/Engine/Plugins/Runtime/Database/ADOSupport/ADOSupport.uplugin
@@ -12,7 +12,6 @@
 	"SupportURL": "",
 	"EnabledByDefault": false,
 	"CanContainContent": false,
-	"Hidden": true,
 	"IsBetaVersion": false,
 	"Installed": false,
 	"Modules": [
diff --git a/Engine/Plugins/Runtime/Database/RemoteDatabaseSupport/RemoteDatabaseSupport.uplugin b/Engine/Plugins/Runtime/Database/RemoteDatabaseSupport/RemoteDatabaseSupport.uplugin
--- a/Engine/Plugins/Runtime/Database/RemoteDatabaseSupport/RemoteDatabaseSupport.uplugin
+++ b/Engine/Plugins/Runtime/Database/RemoteDatabaseSupport/RemoteDatabaseSupport.uplugin
@@ -12,7 +12,6 @@
 	"SupportURL": "",
 	"EnabledByDefault": false,
 	"CanContainContent": false,
-	"Hidden": true,
 	"IsBetaVersion": false,
 	"Installed": false,
 	"Modules": [
diff --git a/Engine/Plugins/Runtime/Database/SQLiteCore/SQLiteCore.uplugin b/Engine/Plugins/Runtime/Database/SQLiteCore/SQLiteCore.uplugin
--- a/Engine/Plugins/Runtime/Database/SQLiteCore/SQLiteCore.uplugin
+++ b/Engine/Plugins/Runtime/Database/SQLiteCore/SQLiteCore.uplugin
@@ -12,7 +12,6 @@
 	"SupportURL": "",
 	"EnabledByDefault": false,
 	"CanContainContent": false,
-	"Hidden": true,
 	"IsBetaVersion": false,
 	"Installed": false,
 	"Modules": [
diff --git a/Engine/Plugins/Runtime/Database/SQLiteSupport/SQLiteSupport.uplugin b/Engine/Plugins/Runtime/Database/SQLiteSupport/SQLiteSupport.uplugin
--- a/Engine/Plugins/Runtime/Database/SQLiteSupport/SQLiteSupport.uplugin
+++ b/Engine/Plugins/Runtime/Database/SQLiteSupport/SQLiteSupport.uplugin
@@ -12,7 +12,6 @@
 	"SupportURL": "",
 	"EnabledByDefault": false,
 	"CanContainContent": false,
-	"Hidden": true,
 	"IsBetaVersion": false,
 	"Installed": false,
 	"Modules": [
diff --git a/Engine/Source/Developer/Localization/Private/LocTextHelper.cpp b/Engine/Source/Developer/Localization/Private/LocTextHelper.cpp
--- a/Engine/Source/Developer/Localization/Private/LocTextHelper.cpp
+++ b/Engine/Source/Developer/Localization/Private/LocTextHelper.cpp
@@ -973,7 +973,7 @@ void FLocTextHelper::GetRuntimeText(const FString& InCulture, const FLocKey& InN
 	OutTranslation = InSource;
 
 	TSharedPtr<FArchiveEntry> ArchiveEntry = FindTranslationImpl(InCulture, InNamespace, InKey, InKeyMetadataObj);
-	if (ArchiveEntry.IsValid())
+	if (ArchiveEntry.IsValid() && !ArchiveEntry->Translation.Text.IsEmpty())
 	{
 		if (bSkipSourceCheck)
 		{
diff --git a/Engine/Source/Developer/SlateReflector/Private/VisualTreeCapture.cpp b/Engine/Source/Developer/SlateReflector/Private/VisualTreeCapture.cpp
--- a/Engine/Source/Developer/SlateReflector/Private/VisualTreeCapture.cpp
+++ b/Engine/Source/Developer/SlateReflector/Private/VisualTreeCapture.cpp
@@ -120,6 +120,11 @@ void FVisualTreeCapture::Disable()
 #endif
 }
 
+void FVisualTreeCapture::Reset()
+{
+	VisualTrees.Reset();
+}
+
 TSharedPtr<FVisualTreeSnapshot> FVisualTreeCapture::GetVisualTreeForWindow(SWindow* InWindow)
 {
 	return VisualTrees.FindRef(InWindow);
diff --git a/Engine/Source/Developer/SlateReflector/Private/VisualTreeCapture.h b/Engine/Source/Developer/SlateReflector/Private/VisualTreeCapture.h
--- a/Engine/Source/Developer/SlateReflector/Private/VisualTreeCapture.h
+++ b/Engine/Source/Developer/SlateReflector/Private/VisualTreeCapture.h
@@ -52,6 +52,7 @@ public:
 
 	void Enable();
 	void Disable();
+	void Reset();
 
 	TSharedPtr<FVisualTreeSnapshot> GetVisualTreeForWindow(SWindow* InWindow);
 	
diff --git a/Engine/Source/Developer/SlateReflector/Private/Widgets/SWidgetReflector.cpp b/Engine/Source/Developer/SlateReflector/Private/Widgets/SWidgetReflector.cpp
--- a/Engine/Source/Developer/SlateReflector/Private/Widgets/SWidgetReflector.cpp
+++ b/Engine/Source/Developer/SlateReflector/Private/Widgets/SWidgetReflector.cpp
@@ -281,22 +281,24 @@ private:
 	{
 		if (PickingMode != InMode)
 		{
-			{
-				SInvalidationPanel::SetEnableWidgetCaching(true);
-
-				VisualCapture.Disable();
-			}
+			// Disable visual picking, and renable widget caching.
+			SInvalidationPanel::SetEnableWidgetCaching(true);
+			VisualCapture.Disable();
 
+			// Enable the picking mode.
 			PickingMode = InMode;
 
-			if (PickingMode == EWidgetPickingMode::Drawable || PickingMode == EWidgetPickingMode::HitTesting)
+			// If we're enabling hit test, reset the visual capture entirely, we don't want to use the visual tree.
+			if (PickingMode == EWidgetPickingMode::HitTesting)
 			{
+				VisualCapture.Reset();
 				SInvalidationPanel::SetEnableWidgetCaching(false);
 			}
-
-			if (PickingMode == EWidgetPickingMode::Drawable)
+			// If we're using the drawing picking mode enable it!
+			else if (PickingMode == EWidgetPickingMode::Drawable)
 			{
 				VisualCapture.Enable();
+				SInvalidationPanel::SetEnableWidgetCaching(false);
 			}
 		}
 	}
diff --git a/Engine/Source/Editor/AnimGraph/Private/AnimPreviewInstance.cpp b/Engine/Source/Editor/AnimGraph/Private/AnimPreviewInstance.cpp
--- a/Engine/Source/Editor/AnimGraph/Private/AnimPreviewInstance.cpp
+++ b/Engine/Source/Editor/AnimGraph/Private/AnimPreviewInstance.cpp
@@ -144,6 +144,11 @@ void FAnimPreviewInstanceProxy::PreUpdate(UAnimInstance* InAnimInstance, float D
 {
 	FAnimSingleNodeInstanceProxy::PreUpdate(InAnimInstance, DeltaSeconds);
 
+	if (CopyPoseNode.SourceMeshComponent.IsValid())
+	{
+		CopyPoseNode.PreUpdate(InAnimInstance);
+	}
+
 	if (!bForceRetargetBasePose)
 	{
 		CurveSource.PreUpdate(InAnimInstance);
diff --git a/Engine/Source/Editor/MovieSceneTools/Private/TrackEditors/AudioTrackEditor.cpp b/Engine/Source/Editor/MovieSceneTools/Private/TrackEditors/AudioTrackEditor.cpp
--- a/Engine/Source/Editor/MovieSceneTools/Private/TrackEditors/AudioTrackEditor.cpp
+++ b/Engine/Source/Editor/MovieSceneTools/Private/TrackEditors/AudioTrackEditor.cpp
@@ -306,7 +306,7 @@ void FAudioThumbnail::GenerateWaveformPreview(TArray<uint8>& OutData, TRange<flo
 
 			if ( SoundWave->InitAudioResource( AudioDevice->GetRuntimeFormat( SoundWave ) ) && (SoundWave->DecompressionType != DTYPE_RealTime || SoundWave->CachedRealtimeFirstBuffer == nullptr ) )
 			{
-				FAsyncAudioDecompress TempDecompress(SoundWave, 0);
+				FAsyncAudioDecompress TempDecompress(SoundWave, AudioDevice->NumPrecacheFrames);
 				TempDecompress.StartSynchronousTask();
 			}
 
diff --git a/Engine/Source/Editor/UMGEditor/Private/Settings/UMGEditorProjectSettings.cpp b/Engine/Source/Editor/UMGEditor/Private/Settings/UMGEditorProjectSettings.cpp
--- a/Engine/Source/Editor/UMGEditor/Private/Settings/UMGEditorProjectSettings.cpp
+++ b/Engine/Source/Editor/UMGEditor/Private/Settings/UMGEditorProjectSettings.cpp
@@ -5,6 +5,7 @@
 #include "WidgetCompilerRule.h"
 #include "UObject/Package.h"
 #include "UObject/UObjectIterator.h"
+#include "Components/CanvasPanel.h"
 
 UUMGEditorProjectSettings::UUMGEditorProjectSettings()
 {
@@ -13,6 +14,8 @@ UUMGEditorProjectSettings::UUMGEditorProjectSettings()
 	bShowWidgetsFromEngineContent = false;
 	bShowWidgetsFromDeveloperContent = true;
 
+	DefaultRootWidget = UCanvasPanel::StaticClass();
+
 	// Deprecated
 	bCookSlowConstructionWidgetTree_DEPRECATED = true;
 	bWidgetSupportsDynamicCreation_DEPRECATED = true;
diff --git a/Engine/Source/Editor/UMGEditor/Private/WidgetBlueprintFactory.cpp b/Engine/Source/Editor/UMGEditor/Private/WidgetBlueprintFactory.cpp
--- a/Engine/Source/Editor/UMGEditor/Private/WidgetBlueprintFactory.cpp
+++ b/Engine/Source/Editor/UMGEditor/Private/WidgetBlueprintFactory.cpp
@@ -8,10 +8,8 @@
 #include "WidgetBlueprint.h"
 #include "Kismet2/KismetEditorUtilities.h"
 
-
-
-#include "Components/CanvasPanel.h"
 #include "Blueprint/WidgetTree.h"
+#include "UMGEditorProjectSettings.h"
 
 #define LOCTEXT_NAMESPACE "UWidgetBlueprintFactory"
 
@@ -53,19 +51,22 @@ UObject* UWidgetBlueprintFactory::FactoryCreateNew(UClass* Class, UObject* InPar
 	if ( ( ParentClass == NULL ) || !FKismetEditorUtilities::CanCreateBlueprintOfClass(ParentClass) || !ParentClass->IsChildOf(UUserWidget::StaticClass()) )
 	{
 		FFormatNamedArguments Args;
-		Args.Add( TEXT("ClassName"), (ParentClass != NULL) ? FText::FromString( ParentClass->GetName() ) : LOCTEXT("Null", "(null)") );
+		Args.Add( TEXT("ClassName"), ParentClass ? FText::FromString( ParentClass->GetName() ) : LOCTEXT("Null", "(null)") );
 		FMessageDialog::Open( EAppMsgType::Ok, FText::Format( LOCTEXT("CannotCreateWidgetBlueprint", "Cannot create a Widget Blueprint based on the class '{ClassName}'."), Args ) );
-		return NULL;
+		return nullptr;
 	}
 	else
 	{
 		UWidgetBlueprint* NewBP = CastChecked<UWidgetBlueprint>(FKismetEditorUtilities::CreateBlueprint(ParentClass, InParent, Name, BlueprintType, UWidgetBlueprint::StaticClass(), UWidgetBlueprintGeneratedClass::StaticClass(), CallingContext));
 
-		// Create a CanvasPanel to use as the default root widget
+		// Create the desired root widget specified by the project
 		if ( NewBP->WidgetTree->RootWidget == nullptr )
 		{
-			UWidget* Root = NewBP->WidgetTree->ConstructWidget<UCanvasPanel>(UCanvasPanel::StaticClass());
-			NewBP->WidgetTree->RootWidget = Root;
+			if (TSubclassOf<UPanelWidget> RootWidgetClass = GetDefault<UUMGEditorProjectSettings>()->DefaultRootWidget)
+			{
+				UWidget* Root = NewBP->WidgetTree->ConstructWidget<UWidget>(RootWidgetClass);
+				NewBP->WidgetTree->RootWidget = Root;
+			}
 		}
 
 		return NewBP;
diff --git a/Engine/Source/Editor/UMGEditor/Public/UMGEditorProjectSettings.h b/Engine/Source/Editor/UMGEditor/Public/UMGEditorProjectSettings.h
--- a/Engine/Source/Editor/UMGEditor/Public/UMGEditorProjectSettings.h
+++ b/Engine/Source/Editor/UMGEditor/Public/UMGEditorProjectSettings.h
@@ -13,6 +13,7 @@
 class UWidgetCompilerRule;
 class UUserWidget;
 class UWidgetBlueprint;
+class UPanelWidget;
 
 USTRUCT()
 struct FDebugResolution
@@ -174,6 +175,10 @@ public:
 
 public:
 
+	/** The panel widget to place at the root of all newly constructed widget blueprints. Can be empty. */
+	UPROPERTY(EditAnywhere, config, Category = Designer)
+	TSubclassOf<UPanelWidget> DefaultRootWidget;
+
 	UPROPERTY(EditAnywhere, config, Category=Designer)
 	TArray<FDebugResolution> DebugResolutions;
 
diff --git a/Engine/Source/Programs/AutomationTool/Gauntlet/Platform/IOS/Gauntlet.IOSBuildSource.cs b/Engine/Source/Programs/AutomationTool/Gauntlet/Platform/IOS/Gauntlet.IOSBuildSource.cs
--- a/Engine/Source/Programs/AutomationTool/Gauntlet/Platform/IOS/Gauntlet.IOSBuildSource.cs
+++ b/Engine/Source/Programs/AutomationTool/Gauntlet/Platform/IOS/Gauntlet.IOSBuildSource.cs
@@ -141,8 +141,8 @@ namespace Gauntlet
 
 				string AbsPath = Fi.Directory.FullName;
 
-				// IOS builds are always packaged, and we can always replace the command line
-				BuildFlags Flags = BuildFlags.Packaged | BuildFlags.CanReplaceCommandLine;
+				// IOS builds are always packaged, and can always replace the command line and executable as we cache the unzip'd IPA
+				BuildFlags Flags = BuildFlags.Packaged | BuildFlags.CanReplaceCommandLine | BuildFlags.CanReplaceExecutable;
 
 				if (AbsPath.Contains("Bulk"))
 				{
diff --git a/Engine/Source/Programs/AutomationTool/Gauntlet/Platform/IOS/Gauntlet.TargetDeviceIOS.cs b/Engine/Source/Programs/AutomationTool/Gauntlet/Platform/IOS/Gauntlet.TargetDeviceIOS.cs
--- a/Engine/Source/Programs/AutomationTool/Gauntlet/Platform/IOS/Gauntlet.TargetDeviceIOS.cs
+++ b/Engine/Source/Programs/AutomationTool/Gauntlet/Platform/IOS/Gauntlet.TargetDeviceIOS.cs
@@ -383,6 +383,91 @@
 
 		}
 
+		/// <summary>
+		/// Resign application using local executable and update debug symbols
+		/// </summary>
+		void ResignApplication(UnrealAppConfig AppConfig)
+		{
+			// check that we have the signing stuff we need
+			string SignProvision = Globals.Params.ParseValue("signprovision", String.Empty);
+			string SignEntitlements = Globals.Params.ParseValue("signentitlements", String.Empty);
+			string SigningIdentity =  Globals.Params.ParseValue("signidentity", String.Empty);
+			
+			// handle signing provision
+			if (string.IsNullOrEmpty(SignProvision) || !File.Exists(SignProvision))
+			{
+				throw new AutomationException("Absolute path to existing provision must be specified, example: -signprovision=/path/to/myapp.provision");
+			}
+			
+			// handle entitlements
+			// Note this extracts entitlements: which may be useful when using same provision/entitlements?: codesign -d --entitlements :entitlements.plist ~/.gauntletappcache/Payload/Example.app/
+
+			if (string.IsNullOrEmpty(SignEntitlements) || !File.Exists(SignEntitlements))
+			{
+				throw new AutomationException("Absolute path to existing entitlements must be specified, example: -signprovision=/path/to/entitlements.plist");
+			}
+
+			// signing identity
+			if (string.IsNullOrEmpty(SigningIdentity))
+			{
+				throw new AutomationException("Signing identity must be specified, example: -signidentity=\"iPhone Developer: John Smith\"");
+			}
+
+			string ProjectName = AppConfig.ProjectName;
+			string BundleName = Path.GetFileNameWithoutExtension(LocalAppBundle);
+			string ExecutableName = UnrealHelpers.GetExecutableName(ProjectName, UnrealTargetPlatform.IOS, AppConfig.Configuration, AppConfig.ProcessType, "");
+			string CachedAppPath = Path.Combine(GauntletAppCache, "Payload", string.Format("{0}.app", BundleName));			
+
+			string LocalExecutable = Path.Combine(Environment.CurrentDirectory, ProjectName, string.Format("Binaries/IOS/{0}", ExecutableName));
+			if (!File.Exists(LocalExecutable))
+			{
+				throw new AutomationException("Local executable not found for -dev argument: {0}", LocalExecutable);
+			}
+
+			File.WriteAllText(CacheResignedFilename, "The application has been resigned");
+
+			// copy local executable
+			FileInfo SrcInfo = new FileInfo(LocalExecutable);			
+			string DestPath = Path.Combine(CachedAppPath, BundleName);
+			SrcInfo.CopyTo(DestPath, true);
+			Log.Verbose("Copied local executable from {0} to {1}", LocalExecutable, DestPath);
+
+			// copy provision
+			SrcInfo = new FileInfo(SignProvision);
+			DestPath = Path.Combine(CachedAppPath, "embedded.mobileprovision");
+			SrcInfo.CopyTo(DestPath, true);
+			Log.Verbose("Copied provision from {0} to {1}", SignProvision, DestPath);
+
+			// handle symbols
+			string LocalSymbolsDir = Path.Combine(Environment.CurrentDirectory, ProjectName, string.Format("Binaries/IOS/{0}.dSYM", ExecutableName));
+			DestPath = Path.Combine(GauntletAppCache, string.Format("Symbols/{0}.dSYM", ExecutableName));
+
+			if (Directory.Exists(DestPath))
+			{
+				Directory.Delete(DestPath, true);
+			}
+
+			if (Directory.Exists(LocalSymbolsDir))
+			{				
+				CommandUtils.CopyDirectory_NoExceptions(LocalSymbolsDir, DestPath, true);
+			}
+			else
+			{
+				Log.Warning("No symbols found for local build at {0}, removing cached app symbols", LocalSymbolsDir);
+			}
+
+			// resign application
+			// @todo: this asks for password unless "Always Allow" is selected, also for builders, document how to permanently grant codesign access to keychain
+			string SignArgs = string.Format("-f -s \"{0}\" --entitlements \"{1}\" \"{2}\"", SigningIdentity, SignEntitlements, CachedAppPath);
+			Log.Info("\nResigning app, please enter keychain password if prompted:\n\ncodesign {0}", SignArgs);
+			var Result = IOSBuild.ExecuteCommand("codesign", SignArgs);
+			if (Result.ExitCode != 0)
+			{
+				throw new AutomationException("Failed to resign application");
+			}
+
+		}
+
 		public IAppInstall InstallApplication(UnrealAppConfig AppConfig)
 		{            
             IOSBuild Build = AppConfig.Build as IOSBuild;
@@ -398,9 +483,27 @@
 			// device artifact path
 			DeviceArtifactPath = string.Format("/Documents/{0}/Saved", AppConfig.ProjectName);
 
+			bool CacheResigned = File.Exists(CacheResignedFilename);
+			bool UseLocalExecutable = Globals.Params.ParseParam("dev");
+
+			if (CacheResigned && !UseLocalExecutable)
+			{
+				if (File.Exists(IPAHashFilename))
+				{
+					Log.Verbose("App was resigned, invalidating app cache");
+					File.Delete(IPAHashFilename);
+				}								
+			}
+
 			PrepareIPA(Build);
 
-			if (!CheckDeployedIPA(Build))
+			// local executable support			
+			if (UseLocalExecutable)
+			{					
+				ResignApplication(AppConfig);				
+			}
+
+			if (CacheResigned || UseLocalExecutable || !CheckDeployedIPA(Build))
 			{
 				// uninstall will clean all device artifacts
 				ExecuteIOSDeployCommand(String.Format("--uninstall -b \"{0}\"", LocalAppBundle), 10 * 60);
@@ -606,6 +709,10 @@
 		// the current IPA MD5 hash, which is tracked to avoid unneccessary deployments and unzip operations
 		string IPAHashFilename { get { return Path.Combine(GauntletAppCache, "IPAHash.txt"); } }
 
+		// file whose presence signals that cache was resigned
+		string CacheResignedFilename { get { return Path.Combine(GauntletAppCache, "Resigned.txt"); } }
+
+
 		/// <summary>
 		/// Generate MD5 and cache IPA bundle files
 		/// </summary>
@@ -650,6 +757,11 @@
 						Directory.Delete(SymbolsDir, true);
 					}
 
+					if (File.Exists(CacheResignedFilename))
+					{
+						File.Delete(CacheResignedFilename);
+					}
+
 					Log.Verbose("Unzipping IPA {0} to cache at: {1}", Build.SourceIPAPath, GauntletAppCache);
 					
 					IProcessResult Result = IOSBuild.ExecuteCommand("unzip", String.Format("{0} -d {1}", Build.SourceIPAPath, GauntletAppCache));
diff --git a/Engine/Source/Programs/AutomationTool/Scripts/BuildPhysX.Automation.cs b/Engine/Source/Programs/AutomationTool/Scripts/BuildPhysX.Automation.cs
--- a/Engine/Source/Programs/AutomationTool/Scripts/BuildPhysX.Automation.cs
+++ b/Engine/Source/Programs/AutomationTool/Scripts/BuildPhysX.Automation.cs
@@ -1683,7 +1683,7 @@ class BuildPhysX : BuildCommand
 			{
 				RobomergeLine = Environment.NewLine + RobomergeCommand;
 			}
-            P4ChangeList = P4.CreateChange(P4Env.Client, String.Format("BuildPhysX.Automation: Deploying {0} libs.", LibDeploymentDesc) + Environment.NewLine + "#rb none" + Environment.NewLine + "#lockdown Nick.Penwarden" + Environment.NewLine + "#tests none" + Environment.NewLine + "#jira none" + RobomergeLine);
+            P4ChangeList = P4.CreateChange(P4Env.Client, String.Format("BuildPhysX.Automation: Deploying {0} libs.", LibDeploymentDesc) + Environment.NewLine + "#rb none" + Environment.NewLine + "#lockdown Nick.Penwarden" + Environment.NewLine + "#tests none" + Environment.NewLine + "#jira none" + Environment.NewLine + "#okforgithub ignore" + RobomergeLine);
 		}
 
 
diff --git a/Engine/Source/Runtime/Core/Private/HAL/LowLevelMemTracker.cpp b/Engine/Source/Runtime/Core/Private/HAL/LowLevelMemTracker.cpp
--- a/Engine/Source/Runtime/Core/Private/HAL/LowLevelMemTracker.cpp
+++ b/Engine/Source/Runtime/Core/Private/HAL/LowLevelMemTracker.cpp
@@ -132,7 +132,7 @@ struct FLLMTagInfo
 
 extern const TCHAR* LLMGetTagName(ELLMTag Tag)
 {
-#define LLM_TAG_NAME_ARRAY(Enum,Str,Stat,Group) TEXT(Str),
+#define LLM_TAG_NAME_ARRAY(Enum,Str,Stat,Group,ParentTag) TEXT(Str),
 	static TCHAR const* Names[] = { LLM_ENUM_GENERIC_TAGS(LLM_TAG_NAME_ARRAY) };
 #undef LLM_TAG_NAME_ARRAY
 
@@ -149,7 +149,7 @@ extern const TCHAR* LLMGetTagName(ELLMTag Tag)
 
 extern const ANSICHAR* LLMGetTagNameANSI(ELLMTag Tag)
 {
-#define LLM_TAG_NAME_ARRAY(Enum,Str,Stat,Group) Str,
+#define LLM_TAG_NAME_ARRAY(Enum,Str,Stat,Group,ParentTag) Str,
 	static ANSICHAR const* Names[] = { LLM_ENUM_GENERIC_TAGS(LLM_TAG_NAME_ARRAY) };
 #undef LLM_TAG_NAME_ARRAY
 
@@ -166,7 +166,7 @@ extern const ANSICHAR* LLMGetTagNameANSI(ELLMTag Tag)
 
 extern FName LLMGetTagStat(ELLMTag Tag)
 {
-#define LLM_TAG_STAT_ARRAY(Enum,Str,Stat,Group) Stat,
+#define LLM_TAG_STAT_ARRAY(Enum,Str,Stat,Group,ParentTag) Stat,
 	static FName Names[] = { LLM_ENUM_GENERIC_TAGS(LLM_TAG_STAT_ARRAY) };
 #undef LLM_TAG_STAT_ARRAY
 
@@ -183,7 +183,7 @@ extern FName LLMGetTagStat(ELLMTag Tag)
 
 extern FName LLMGetTagStatGroup(ELLMTag Tag)
 {
-#define LLM_TAG_STATGROUP_ARRAY(Enum,Str,Stat,Group) Group,
+#define LLM_TAG_STATGROUP_ARRAY(Enum,Str,Stat,Group,ParentTag) Group,
 	static FName Names[] = { LLM_ENUM_GENERIC_TAGS(LLM_TAG_STATGROUP_ARRAY) };
 #undef LLM_TAG_STAT_ARRAY
 
@@ -198,6 +198,23 @@ extern FName LLMGetTagStatGroup(ELLMTag Tag)
 	}
 }
 
+extern int32 LLMGetTagParent(ELLMTag Tag)
+{
+#define LLM_TAG_NAME_ARRAY(Enum,Str,Stat,Group,ParentTag) (int32)ParentTag,
+	static const int32 ParentTags [] = { LLM_ENUM_GENERIC_TAGS(LLM_TAG_NAME_ARRAY) };
+#undef LLM_TAG_NAME_ARRAY
+
+	int32 Index = (int32)Tag;
+	if( Index >= 0 && Index < ARRAY_COUNT(ParentTags))
+	{
+		return ParentTags[Index];
+	}
+	else
+	{
+		return -1;
+	}
+}
+
 /**
  * FLLMCsvWriter: class for writing out the LLM stats to a csv file every few seconds
  */
@@ -230,16 +247,16 @@ public:
 
 	void CatchRunawayTags(FLLMPlatformTag* PlatformTags);
 
-	void Update(FLLMPlatformTag* PlatformTags);
+	void Update(FLLMPlatformTag* PlatformTags, const int32* ParentTags);
 
 	void SetEnabled(bool value) { Enabled = value; }
 
 private:
-	void WriteGraph(FLLMPlatformTag* PlatformTags);
+	void WriteGraph(FLLMPlatformTag* PlatformTags, const int32* ParentTags);
 
 	void Write(const FString& Text);
 
-	static FString GetTagName(int64 Tag, FLLMPlatformTag* PlatformTags);
+	static FString GetTagName(int64 Tag, FLLMPlatformTag* PlatformTags, const int32* ParentTags);
 
 	static const TCHAR* GetTrackerCsvName(ELLMTracker InTracker);
 
@@ -304,7 +321,7 @@ public:
 
 	void SetCSVEnabled(bool Value);
 
-	void WriteCsv(FLLMPlatformTag* PlatformTags);
+	void WriteCsv(FLLMPlatformTag* PlatformTags, const int32* ParentTags);
 
 #define LLM_USE_ALLOC_INFO_STRUCT (LLM_STAT_TAGS_ENABLED || LLM_ALLOW_ASSETS_TAGS)
 
@@ -329,7 +346,7 @@ public:
 	}
 
 	void SetTotalTags(ELLMTag Untagged, ELLMTag Tracked);
-	void Update(FLLMPlatformTag* PlatformTags);
+	void Update(FLLMPlatformTag* PlatformTags, const int32* ParentTags);
 	void UpdateTotals();
 	void CatchRunawayTags(FLLMPlatformTag* PlatformTags);
 
@@ -376,6 +393,8 @@ protected:
 			int64* EnumTagAmounts,
 			int64* OutAllocTypeAmounts);
 
+		void UpdateFrameStatGroups( FLLMPlatformTag* PlatformTags, const int32* ParentTags );
+
 		static void IncMemoryStatByFName(FName Name, int64 Amount);
 
 		void ClearAllocTypeAmounts();
@@ -499,6 +518,19 @@ FLowLevelMemTracker::FLowLevelMemTracker()
 	{
 		PlatformTags[Index].Name = nullptr;
 	}
+
+	for (int32 Index = 0; Index < (int32)ELLMTag::PlatformTagEnd; Index++ )
+	{
+		ParentTags[Index] = LLMGetTagParent((ELLMTag)Index);
+	}
+	for (int32 Index = 0; Index < (int32)ELLMTag::PlatformTagEnd; Index++ )
+	{
+		if (ParentTags[Index] != -1)
+		{
+			int32 GrandparentTag = ParentTags[ParentTags[Index]];
+			LLMCheckf( GrandparentTag == -1, TEXT("can only have one level of tag parent") );
+		}
+	}
 }
 
 FLowLevelMemTracker::~FLowLevelMemTracker()
@@ -552,7 +584,7 @@ void FLowLevelMemTracker::UpdateStatsPerFrame(const TCHAR* LogName)
 	// update the trackers
 	for (int32 TrackerIndex = 0; TrackerIndex < (int32)ELLMTracker::Max; TrackerIndex++)
 	{
-		GetTracker((ELLMTracker)TrackerIndex)->Update(PlatformTags);
+		GetTracker((ELLMTracker)TrackerIndex)->Update(PlatformTags,ParentTags);
 #if ENABLE_CATCH_RUNAWAY_TAGS
 		GetTracker((ELLMTracker)TrackerIndex)->CatchRunawayTags(PlatformTags);
 #endif
@@ -602,8 +634,8 @@ void FLowLevelMemTracker::UpdateStatsPerFrame(const TCHAR* LogName)
 
 	if (bCsvWriterEnabled)
 	{
-		GetTracker(ELLMTracker::Default)->WriteCsv(PlatformTags);
-		GetTracker(ELLMTracker::Platform)->WriteCsv(PlatformTags);
+		GetTracker(ELLMTracker::Default)->WriteCsv(PlatformTags,ParentTags);
+		GetTracker(ELLMTracker::Platform)->WriteCsv(PlatformTags,ParentTags);
 	}
 
 	if (LogName != nullptr)
@@ -813,7 +845,7 @@ static bool IsAssetTagForAssets(ELLMTagSet Set)
 	return Set == ELLMTagSet::Assets || Set == ELLMTagSet::AssetClasses;
 }
 
-void FLowLevelMemTracker::RegisterPlatformTag(int32 Tag, const TCHAR* Name, FName StatName, FName SummaryStatName)
+void FLowLevelMemTracker::RegisterPlatformTag(int32 Tag, const TCHAR* Name, FName StatName, FName SummaryStatName, int32 ParentTag)
 {
 	LLMCheck(Tag >= (int32)ELLMTag::PlatformTagStart && Tag <= (int32)ELLMTag::PlatformTagEnd);
 	FLLMPlatformTag& PlatformTag = PlatformTags[Tag - (int32)ELLMTag::PlatformTagStart];
@@ -821,6 +853,12 @@ void FLowLevelMemTracker::RegisterPlatformTag(int32 Tag, const TCHAR* Name, FNam
 	PlatformTag.Name = Name;
 	PlatformTag.StatName = StatName;
 	PlatformTag.SummaryStatName = SummaryStatName;
+	ParentTags[Tag] = ParentTag;
+	if (ParentTag != -1)
+	{
+		int32 GrandparentTag = ParentTags[ParentTag];
+		LLMCheckf( GrandparentTag == -1, TEXT("can only have one level of tag parent") );
+	}
 }
 
 bool FLowLevelMemTracker::FindTagByName( const TCHAR* Name, uint64& OutTag ) const
@@ -1241,7 +1279,7 @@ void FLLMTracker::SetTotalTags(ELLMTag InUntaggedTotalTag, ELLMTag InTrackedTota
 	TrackedTotalTag = InTrackedTotalTag;
 }
 
-void FLLMTracker::Update(FLLMPlatformTag* PlatformTags)
+void FLLMTracker::Update(FLLMPlatformTag* PlatformTags, const int32* ParentTags)
 {
 	// protect any accesses to the ThreadStates array
 	FScopeLock Lock(&ThreadArraySection);
@@ -1250,6 +1288,7 @@ void FLLMTracker::Update(FLLMPlatformTag* PlatformTags)
 	int ThreadStateNum = ThreadStates.Num();
 	for (int32 ThreadIndex = 0; ThreadIndex < ThreadStateNum; ThreadIndex++)
 	{
+		ThreadStates[ThreadIndex]->UpdateFrameStatGroups(PlatformTags,ParentTags);
 		ThreadStates[ThreadIndex]->GetFrameStatTotals(UntaggedTotalTag, StateCopy, CsvWriter, PlatformTags, EnumTagAmounts, AllocTypeAmounts);
 	}
 
@@ -1290,9 +1329,9 @@ void FLLMTracker::CatchRunawayTags(FLLMPlatformTag* PlatformTags)
 	CsvWriter.CatchRunawayTags(PlatformTags);
 }
 
-void FLLMTracker::WriteCsv(FLLMPlatformTag* PlatformTags)
+void FLLMTracker::WriteCsv(FLLMPlatformTag* PlatformTags, const int32* ParentTags)
 {
-	CsvWriter.Update(PlatformTags);
+	CsvWriter.Update(PlatformTags,ParentTags);
 }
 
 int64 FLLMTracker::GetActiveTag()
@@ -1458,7 +1497,11 @@ void FLLMTracker::SetTagAmount(ELLMTag Tag, int64 Amount, bool AddToTotal)
 		FPlatformAtomics::InterlockedAdd(&TrackedMemoryOverFrames, (int64)(Amount - EnumTagAmounts[(int32)Tag]));
 	}
 
-	SET_MEMORY_STAT_FName(LLMGetTagStat(Tag), Amount);
+	FName StatName = LLMGetTagStat(Tag);
+	if (StatName != NAME_None)
+	{
+		SET_MEMORY_STAT_FName(StatName, Amount);
+	}
 
 	EnumTagAmounts[(int32)Tag] = Amount;
 
@@ -1610,6 +1653,31 @@ void FLLMTracker::FLLMThreadState::GetFrameStatTotals(
 	InStateCopy.Clear();
 }
 
+void FLLMTracker::FLLMThreadState::UpdateFrameStatGroups( FLLMPlatformTag* PlatformTags, const int32* ParentTags )
+{
+	FScopeLock Lock(&TagSection);
+
+	uint32 MaxTagIndex = TaggedAllocTags.Num(); //group tags will be added at the end of the array. we don't want groups of groups so don't include them in the loop
+	for (uint32 TagIndex = 0; TagIndex < MaxTagIndex; TagIndex++)
+	{
+		int64 Amount = TaggedAllocs[TagIndex];
+		if( Amount != 0 )
+		{
+			int64 Tag = TaggedAllocTags[TagIndex];
+			if( Tag >= 0 && Tag < (int32)ELLMTag::PlatformTagEnd)
+			{
+				int32 ParentTag = ParentTags[Tag];
+				if( ParentTag != -1 )
+				{
+					IncrTag( ParentTag, Amount, false );
+				}
+			}
+		}
+	}
+}
+
+
+
 void FLLMTracker::FLLMThreadState::IncMemoryStatByFName(FName Name, int64 Amount)
 {
 	if (Name != NAME_None)
@@ -1753,7 +1821,7 @@ void FLLMCsvWriter::CatchRunawayTags(FLLMPlatformTag* PlatformTags)
 
 	for (uint32 i = 0; i < StatValues.Num(); i++)
 	{
-		if (GetTagName(StatValues[i].Tag, PlatformTags) == StatsTagName)
+		if (GetTagName(StatValues[i].Tag, PlatformTags,nullptr) == StatsTagName)
 		{
 			StatsTag = StatValues[i].Tag;
 		}
@@ -1772,12 +1840,12 @@ void FLLMCsvWriter::CatchRunawayTags(FLLMPlatformTag* PlatformTags)
 /*
 * memory can be allocated in this function
 */
-void FLLMCsvWriter::Update(FLLMPlatformTag* PlatformTags)
+void FLLMCsvWriter::Update(FLLMPlatformTag* PlatformTags, const int32* ParentTags)
 {
 	double Now = FPlatformTime::Seconds();
 	if (Now - LastWriteTime >= (double)CVarLLMWriteInterval.GetValueOnGameThread())
 	{
-		WriteGraph(PlatformTags);
+		WriteGraph(PlatformTags, ParentTags);
 
 		LastWriteTime = Now;
 	}
@@ -1805,7 +1873,7 @@ void FLLMCsvWriter::Write(const FString& Text)
  * create the csv file on the first call. When it finds a new stat name it seeks
  * back to the start of the file and re-writes the column names.
 */
-void FLLMCsvWriter::WriteGraph(FLLMPlatformTag* PlatformTags)
+void FLLMCsvWriter::WriteGraph(FLLMPlatformTag* PlatformTags, const int32* ParentTags)
 {
 	// create the csv file
 	if (!Archive)
@@ -1842,7 +1910,7 @@ void FLLMCsvWriter::WriteGraph(FLLMPlatformTag* PlatformTags)
 
 		for (int32 i = 0; i < StatValueCountLocal; ++i)
 		{
-			FString StatName = GetTagName(StatValuesForWrite[i].Tag, PlatformTags);
+			FString StatName = GetTagName(StatValuesForWrite[i].Tag, PlatformTags, ParentTags);
 			FString Text = FString::Printf(TEXT("%s,"), *StatName);
 			Write(Text);
 		}
@@ -1877,8 +1945,10 @@ void FLLMCsvWriter::WriteGraph(FLLMPlatformTag* PlatformTags)
 /*
  * convert a Tag to a string. If the Tag is actually a Stat then extract the name of the stat.
 */
-FString FLLMCsvWriter::GetTagName(int64 Tag, FLLMPlatformTag* PlatformTags)
+FString FLLMCsvWriter::GetTagName(int64 Tag, FLLMPlatformTag* PlatformTags, const int32* ParentTags)
 {
+	FString Result;
+
 	if (Tag > (int64)ELLMTag::PlatformTagEnd)
 	{
 		FString Name = TagToFName(Tag).ToString();
@@ -1895,17 +1965,30 @@ FString FLLMCsvWriter::GetTagName(int64 Tag, FLLMPlatformTag* PlatformTags)
 			}
 		}
 
-		return Name;
+		Result = Name;
 	}
 	else if (Tag >= (int32)ELLMTag::PlatformTagStart && Tag <= (int32)ELLMTag::PlatformTagEnd)
 	{
-		return PlatformTags[Tag - (int32)ELLMTag::PlatformTagStart].Name;
+		if (ParentTags != nullptr && ParentTags[Tag] != -1)
+		{
+			Result = GetTagName( ParentTags[Tag], PlatformTags, nullptr ) + TEXT("/");
+		}
+
+		Result += PlatformTags[Tag - (int32)ELLMTag::PlatformTagStart].Name;
 	}
 	else
 	{
 		LLMCheck(Tag >= 0 && LLMGetTagName((ELLMTag)Tag) != nullptr);
-		return LLMGetTagName((ELLMTag)Tag);
+
+		if (ParentTags != nullptr && ParentTags[Tag] != -1)
+		{
+			Result = GetTagName( ParentTags[Tag], PlatformTags, nullptr ) + TEXT("/");
+		}
+
+		Result += LLMGetTagName((ELLMTag)Tag);
 	}
+
+	return Result;
 }
 
 #endif		// #if ENABLE_LOW_LEVEL_MEM_TRACKER
diff --git a/Engine/Source/Runtime/Core/Public/HAL/LowLevelMemTracker.h b/Engine/Source/Runtime/Core/Public/HAL/LowLevelMemTracker.h
--- a/Engine/Source/Runtime/Core/Public/HAL/LowLevelMemTracker.h
+++ b/Engine/Source/Runtime/Core/Public/HAL/LowLevelMemTracker.h
@@ -153,78 +153,78 @@ enum class ELLMTagSet : uint8
 };
 
 #define LLM_ENUM_GENERIC_TAGS(macro) \
-	macro(Untagged,								"Untagged",						NAME_None,													NAME_None)									\
-	macro(Paused,								"Paused",						NAME_None,													NAME_None)									\
-	macro(Total,								"Total",						GET_STATFNAME(STAT_TotalLLM),								GET_STATFNAME(STAT_TrackedTotalSummaryLLM))	\
-	macro(Untracked,							"Untracked",					GET_STATFNAME(STAT_UntrackedLLM),							GET_STATFNAME(STAT_TrackedTotalSummaryLLM))	\
-	macro(PlatformTotal,						"Total",						GET_STATFNAME(STAT_PlatformTotalLLM),						NAME_None)									\
-	macro(TrackedTotal,							"TrackedTotal",					GET_STATFNAME(STAT_TrackedTotalLLM),						GET_STATFNAME(STAT_TrackedTotalSummaryLLM))	\
-	macro(UntaggedTotal,						"Untagged",						GET_STATFNAME(STAT_UntaggedTotalLLM),						NAME_None)									\
-	macro(WorkingSetSize,						"WorkingSetSize",				GET_STATFNAME(STAT_WorkingSetSizeLLM),						GET_STATFNAME(STAT_TrackedTotalSummaryLLM))	\
-	macro(PagefileUsed,							"PagefileUsed",					GET_STATFNAME(STAT_PagefileUsedLLM),						GET_STATFNAME(STAT_TrackedTotalSummaryLLM))	\
-	macro(PlatformTrackedTotal,					"TrackedTotal",					GET_STATFNAME(STAT_PlatformTrackedTotalLLM),				NAME_None)									\
-	macro(PlatformUntaggedTotal,				"Untagged",						GET_STATFNAME(STAT_PlatformUntaggedTotalLLM),				NAME_None)									\
-	macro(PlatformUntracked,					"Untracked",					GET_STATFNAME(STAT_PlatformUntrackedLLM),					NAME_None)									\
-	macro(PlatformOverhead,						"LLMOverhead",					GET_STATFNAME(STAT_PlatformOverheadLLM),					NAME_None)									\
-	macro(FMalloc,								"FMalloc",						GET_STATFNAME(STAT_FMallocLLM),								NAME_None)									\
-	macro(FMallocUnused,						"FMallocUnused",				GET_STATFNAME(STAT_FMallocUnusedLLM),						GET_STATFNAME(STAT_EngineSummaryLLM))		\
-	macro(ThreadStack,							"ThreadStack",					GET_STATFNAME(STAT_ThreadStackLLM),							GET_STATFNAME(STAT_EngineSummaryLLM))		\
-	macro(ThreadStackPlatform,					"ThreadStack",					GET_STATFNAME(STAT_ThreadStackPlatformLLM),					NAME_None)									\
-	macro(ProgramSizePlatform,					"ProgramSize",					GET_STATFNAME(STAT_ProgramSizePlatformLLM),					NAME_None)									\
-	macro(ProgramSize,							"ProgramSize",					GET_STATFNAME(STAT_ProgramSizeLLM),							GET_STATFNAME(STAT_EngineSummaryLLM))		\
-	macro(BackupOOMMemoryPoolPlatform,			"OOMBackupPool",				GET_STATFNAME(STAT_OOMBackupPoolPlatformLLM),				NAME_None)									\
-	macro(BackupOOMMemoryPool,					"OOMBackupPool",				GET_STATFNAME(STAT_OOMBackupPoolLLM),						GET_STATFNAME(STAT_EngineSummaryLLM))		\
-	macro(GenericPlatformMallocCrash,			"GenericPlatformMallocCrash",	GET_STATFNAME(STAT_GenericPlatformMallocCrashLLM),			GET_STATFNAME(STAT_EngineSummaryLLM))		\
-	macro(GenericPlatformMallocCrashPlatform,	"GenericPlatformMallocCrash",	GET_STATFNAME(STAT_GenericPlatformMallocCrashPlatformLLM),	GET_STATFNAME(STAT_EngineSummaryLLM))		\
-	macro(EngineMisc,							"EngineMisc",					GET_STATFNAME(STAT_EngineMiscLLM),							GET_STATFNAME(STAT_EngineSummaryLLM))		\
-	macro(TaskGraphTasksMisc,					"TaskGraphMiscTasks",			GET_STATFNAME(STAT_TaskGraphTasksMiscLLM),					GET_STATFNAME(STAT_EngineSummaryLLM))		\
-	macro(Audio,								"Audio",						GET_STATFNAME(STAT_AudioLLM),								GET_STATFNAME(STAT_AudioSummaryLLM))		\
-	macro(AudioMixer,							"AudioMixer",					GET_STATFNAME(STAT_AudioMixerLLM),							GET_STATFNAME(STAT_AudioSummaryLLM))		\
-	macro(AudioPrecache,						"AudioPrecache",				GET_STATFNAME(STAT_AudioPrecacheLLM),						GET_STATFNAME(STAT_AudioSummaryLLM))		\
-	macro(AudioDecompress,						"AudioDecompress",				GET_STATFNAME(STAT_AudioDecompressLLM),						GET_STATFNAME(STAT_AudioSummaryLLM))		\
-	macro(AudioRealtimePrecache,				"AudioRealtimePrecache",		GET_STATFNAME(STAT_AudioRealtimePrecacheLLM),				GET_STATFNAME(STAT_AudioSummaryLLM))		\
-	macro(AudioFullDecompress,					"AudioFullDecompress",			GET_STATFNAME(STAT_AudioFullDecompressLLM),					GET_STATFNAME(STAT_AudioSummaryLLM))		\
-	macro(FName,								"FName",						GET_STATFNAME(STAT_FNameLLM),								GET_STATFNAME(STAT_EngineSummaryLLM))		\
-	macro(Networking,							"Networking",					GET_STATFNAME(STAT_NetworkingLLM),							GET_STATFNAME(STAT_EngineSummaryLLM))		\
-	macro(Meshes,								"Meshes",						GET_STATFNAME(STAT_MeshesLLM),								GET_STATFNAME(STAT_MeshesSummaryLLM))		\
-	macro(Stats,								"Stats",						GET_STATFNAME(STAT_StatsLLM),								GET_STATFNAME(STAT_EngineSummaryLLM))		\
-	macro(Shaders,								"Shaders",						GET_STATFNAME(STAT_ShadersLLM),								GET_STATFNAME(STAT_EngineSummaryLLM))		\
-	macro(PSO,									"PSO",							GET_STATFNAME(STAT_PSOLLM),									GET_STATFNAME(STAT_EngineSummaryLLM))		\
-	macro(Textures,								"Textures",						GET_STATFNAME(STAT_TexturesLLM),							GET_STATFNAME(STAT_TexturesSummaryLLM))		\
-	macro(RenderTargets,						"RenderTargets",				GET_STATFNAME(STAT_RenderTargetsLLM),						GET_STATFNAME(STAT_EngineSummaryLLM))		\
-	macro(RHIMisc,								"RHIMisc",						GET_STATFNAME(STAT_RHIMiscLLM),								GET_STATFNAME(STAT_EngineSummaryLLM))		\
-	macro(AsyncLoading,							"AsyncLoading",					GET_STATFNAME(STAT_AsyncLoadingLLM),						GET_STATFNAME(STAT_EngineSummaryLLM))		\
-	macro(UObject,								"UObject",						GET_STATFNAME(STAT_UObjectLLM),								GET_STATFNAME(STAT_UObjectSummaryLLM))		\
-	macro(Animation,							"Animation",					GET_STATFNAME(STAT_AnimationLLM),							GET_STATFNAME(STAT_AnimationSummaryLLM))	\
-	macro(StaticMesh,							"StaticMesh",					GET_STATFNAME(STAT_StaticMeshLLM),							GET_STATFNAME(STAT_StaticMeshSummaryLLM))	\
-	macro(Materials,							"Materials",					GET_STATFNAME(STAT_MaterialsLLM),							GET_STATFNAME(STAT_MaterialsSummaryLLM))	\
-	macro(Particles,							"Particles",					GET_STATFNAME(STAT_ParticlesLLM),							GET_STATFNAME(STAT_ParticlesSummaryLLM))	\
-	macro(GC,									"GC",							GET_STATFNAME(STAT_GCLLM),									GET_STATFNAME(STAT_EngineSummaryLLM))		\
-	macro(UI,									"UI",							GET_STATFNAME(STAT_UILLM),									GET_STATFNAME(STAT_UISummaryLLM))			\
-	macro(PhysX,								"PhysX",						GET_STATFNAME(STAT_PhysXLLM),								GET_STATFNAME(STAT_PhysXSummaryLLM))		\
-	macro(EnginePreInitMemory,					"EnginePreInit",				GET_STATFNAME(STAT_EnginePreInitLLM),						GET_STATFNAME(STAT_EngineSummaryLLM))		\
-	macro(EngineInitMemory,						"EngineInit",					GET_STATFNAME(STAT_EngineInitLLM),							GET_STATFNAME(STAT_EngineSummaryLLM))		\
-	macro(RenderingThreadMemory,				"RenderingThread",				GET_STATFNAME(STAT_RenderingThreadLLM),						GET_STATFNAME(STAT_EngineSummaryLLM))		\
-	macro(LoadMapMisc,							"LoadMapMisc",					GET_STATFNAME(STAT_LoadMapMiscLLM),							GET_STATFNAME(STAT_EngineSummaryLLM))		\
-	macro(StreamingManager,						"StreamingManager",				GET_STATFNAME(STAT_StreamingManagerLLM),					GET_STATFNAME(STAT_EngineSummaryLLM))		\
-	macro(GraphicsPlatform,						"Graphics",						GET_STATFNAME(STAT_GraphicsPlatformLLM),					NAME_None)									\
-	macro(FileSystem,							"FileSystem",					GET_STATFNAME(STAT_FileSystemLLM),							GET_STATFNAME(STAT_EngineSummaryLLM))		\
-	macro(Localization,							"Localization",					GET_STATFNAME(STAT_LocalizationLLM),						GET_STATFNAME(STAT_EngineSummaryLLM))		\
-	macro(VertexBuffer,							"VertexBuffer",					GET_STATFNAME(STAT_VertexBufferLLM),						GET_STATFNAME(STAT_EngineSummaryLLM))		\
-	macro(IndexBuffer,							"IndexBuffer",					GET_STATFNAME(STAT_IndexBufferLLM),							GET_STATFNAME(STAT_EngineSummaryLLM))		\
-	macro(UniformBuffer,						"UniformBuffer",				GET_STATFNAME(STAT_UniformBufferLLM),						GET_STATFNAME(STAT_EngineSummaryLLM))		\
-	macro(AssetRegistry,						"AssetRegistry",				GET_STATFNAME(STAT_AssetRegistryLLM),						GET_STATFNAME(STAT_EngineSummaryLLM))		\
-	macro(ConfigSystem,							"ConfigSystem",					GET_STATFNAME(STAT_ConfigSystemLLM),						GET_STATFNAME(STAT_EngineSummaryLLM))		\
-	macro(InitUObject,							"InitUObject",					GET_STATFNAME(STAT_InitUObjectLLM),							GET_STATFNAME(STAT_EngineSummaryLLM))		\
-	macro(VideoRecording,						"VideoRecording",				GET_STATFNAME(STAT_VideoRecordingLLM),						GET_STATFNAME(STAT_EngineSummaryLLM))		\
-	macro(CsvProfiler,							"CsvProfiler",					GET_STATFNAME(STAT_CsvProfilerLLM),							GET_STATFNAME(STAT_EngineSummaryLLM))
+	macro(Untagged,								"Untagged",						NAME_None,													NAME_None,										-1)\
+	macro(Paused,								"Paused",						NAME_None,													NAME_None,										-1)\
+	macro(Total,								"Total",						GET_STATFNAME(STAT_TotalLLM),								GET_STATFNAME(STAT_TrackedTotalSummaryLLM),		-1)\
+	macro(Untracked,							"Untracked",					GET_STATFNAME(STAT_UntrackedLLM),							GET_STATFNAME(STAT_TrackedTotalSummaryLLM),		-1)\
+	macro(PlatformTotal,						"Total",						GET_STATFNAME(STAT_PlatformTotalLLM),						NAME_None,										-1)\
+	macro(TrackedTotal,							"TrackedTotal",					GET_STATFNAME(STAT_TrackedTotalLLM),						GET_STATFNAME(STAT_TrackedTotalSummaryLLM),		-1)\
+	macro(UntaggedTotal,						"Untagged",						GET_STATFNAME(STAT_UntaggedTotalLLM),						NAME_None,										-1)\
+	macro(WorkingSetSize,						"WorkingSetSize",				GET_STATFNAME(STAT_WorkingSetSizeLLM),						GET_STATFNAME(STAT_TrackedTotalSummaryLLM),		-1)\
+	macro(PagefileUsed,							"PagefileUsed",					GET_STATFNAME(STAT_PagefileUsedLLM),						GET_STATFNAME(STAT_TrackedTotalSummaryLLM),		-1)\
+	macro(PlatformTrackedTotal,					"TrackedTotal",					GET_STATFNAME(STAT_PlatformTrackedTotalLLM),				NAME_None,										-1)\
+	macro(PlatformUntaggedTotal,				"Untagged",						GET_STATFNAME(STAT_PlatformUntaggedTotalLLM),				NAME_None,										-1)\
+	macro(PlatformUntracked,					"Untracked",					GET_STATFNAME(STAT_PlatformUntrackedLLM),					NAME_None,										-1)\
+	macro(PlatformOverhead,						"LLMOverhead",					GET_STATFNAME(STAT_PlatformOverheadLLM),					NAME_None,										-1)\
+	macro(FMalloc,								"FMalloc",						GET_STATFNAME(STAT_FMallocLLM),								NAME_None,										-1)\
+	macro(FMallocUnused,						"FMallocUnused",				GET_STATFNAME(STAT_FMallocUnusedLLM),						GET_STATFNAME(STAT_EngineSummaryLLM),			-1)\
+	macro(ThreadStack,							"ThreadStack",					GET_STATFNAME(STAT_ThreadStackLLM),							GET_STATFNAME(STAT_EngineSummaryLLM),			-1)\
+	macro(ThreadStackPlatform,					"ThreadStack",					GET_STATFNAME(STAT_ThreadStackPlatformLLM),					NAME_None,										-1)\
+	macro(ProgramSizePlatform,					"ProgramSize",					GET_STATFNAME(STAT_ProgramSizePlatformLLM),					NAME_None,										-1)\
+	macro(ProgramSize,							"ProgramSize",					GET_STATFNAME(STAT_ProgramSizeLLM),							GET_STATFNAME(STAT_EngineSummaryLLM),			-1)\
+	macro(BackupOOMMemoryPoolPlatform,			"OOMBackupPool",				GET_STATFNAME(STAT_OOMBackupPoolPlatformLLM),				NAME_None,										-1)\
+	macro(BackupOOMMemoryPool,					"OOMBackupPool",				GET_STATFNAME(STAT_OOMBackupPoolLLM),						GET_STATFNAME(STAT_EngineSummaryLLM),			-1)\
+	macro(GenericPlatformMallocCrash,			"GenericPlatformMallocCrash",	GET_STATFNAME(STAT_GenericPlatformMallocCrashLLM),			GET_STATFNAME(STAT_EngineSummaryLLM),			-1)\
+	macro(GenericPlatformMallocCrashPlatform,	"GenericPlatformMallocCrash",	GET_STATFNAME(STAT_GenericPlatformMallocCrashPlatformLLM),	GET_STATFNAME(STAT_EngineSummaryLLM),			-1)\
+	macro(EngineMisc,							"EngineMisc",					GET_STATFNAME(STAT_EngineMiscLLM),							GET_STATFNAME(STAT_EngineSummaryLLM),			-1)\
+	macro(TaskGraphTasksMisc,					"TaskGraphMiscTasks",			GET_STATFNAME(STAT_TaskGraphTasksMiscLLM),					GET_STATFNAME(STAT_EngineSummaryLLM),			-1)\
+	macro(Audio,								"Audio",						GET_STATFNAME(STAT_AudioLLM),								GET_STATFNAME(STAT_AudioSummaryLLM),			-1)\
+	macro(AudioMixer,							"AudioMixer",					GET_STATFNAME(STAT_AudioMixerLLM),							GET_STATFNAME(STAT_AudioSummaryLLM),			-1)\
+	macro(AudioPrecache,						"AudioPrecache",				GET_STATFNAME(STAT_AudioPrecacheLLM),						GET_STATFNAME(STAT_AudioSummaryLLM),			-1)\
+	macro(AudioDecompress,						"AudioDecompress",				GET_STATFNAME(STAT_AudioDecompressLLM),						GET_STATFNAME(STAT_AudioSummaryLLM),			-1)\
+	macro(AudioRealtimePrecache,				"AudioRealtimePrecache",		GET_STATFNAME(STAT_AudioRealtimePrecacheLLM),				GET_STATFNAME(STAT_AudioSummaryLLM),			-1)\
+	macro(AudioFullDecompress,					"AudioFullDecompress",			GET_STATFNAME(STAT_AudioFullDecompressLLM),					GET_STATFNAME(STAT_AudioSummaryLLM),			-1)\
+	macro(FName,								"FName",						GET_STATFNAME(STAT_FNameLLM),								GET_STATFNAME(STAT_EngineSummaryLLM),			-1)\
+	macro(Networking,							"Networking",					GET_STATFNAME(STAT_NetworkingLLM),							GET_STATFNAME(STAT_EngineSummaryLLM),			-1)\
+	macro(Meshes,								"Meshes",						GET_STATFNAME(STAT_MeshesLLM),								GET_STATFNAME(STAT_MeshesSummaryLLM),			-1)\
+	macro(Stats,								"Stats",						GET_STATFNAME(STAT_StatsLLM),								GET_STATFNAME(STAT_EngineSummaryLLM),			-1)\
+	macro(Shaders,								"Shaders",						GET_STATFNAME(STAT_ShadersLLM),								GET_STATFNAME(STAT_EngineSummaryLLM),			-1)\
+	macro(PSO,									"PSO",							GET_STATFNAME(STAT_PSOLLM),									GET_STATFNAME(STAT_EngineSummaryLLM),			-1)\
+	macro(Textures,								"Textures",						GET_STATFNAME(STAT_TexturesLLM),							GET_STATFNAME(STAT_TexturesSummaryLLM),			-1)\
+	macro(RenderTargets,						"RenderTargets",				GET_STATFNAME(STAT_RenderTargetsLLM),						GET_STATFNAME(STAT_EngineSummaryLLM),			-1)\
+	macro(RHIMisc,								"RHIMisc",						GET_STATFNAME(STAT_RHIMiscLLM),								GET_STATFNAME(STAT_EngineSummaryLLM),			-1)\
+	macro(AsyncLoading,							"AsyncLoading",					GET_STATFNAME(STAT_AsyncLoadingLLM),						GET_STATFNAME(STAT_EngineSummaryLLM),			-1)\
+	macro(UObject,								"UObject",						GET_STATFNAME(STAT_UObjectLLM),								GET_STATFNAME(STAT_UObjectSummaryLLM),			-1)\
+	macro(Animation,							"Animation",					GET_STATFNAME(STAT_AnimationLLM),							GET_STATFNAME(STAT_AnimationSummaryLLM),		-1)\
+	macro(StaticMesh,							"StaticMesh",					GET_STATFNAME(STAT_StaticMeshLLM),							GET_STATFNAME(STAT_StaticMeshSummaryLLM),		-1)\
+	macro(Materials,							"Materials",					GET_STATFNAME(STAT_MaterialsLLM),							GET_STATFNAME(STAT_MaterialsSummaryLLM),		-1)\
+	macro(Particles,							"Particles",					GET_STATFNAME(STAT_ParticlesLLM),							GET_STATFNAME(STAT_ParticlesSummaryLLM),		-1)\
+	macro(GC,									"GC",							GET_STATFNAME(STAT_GCLLM),									GET_STATFNAME(STAT_EngineSummaryLLM),			-1)\
+	macro(UI,									"UI",							GET_STATFNAME(STAT_UILLM),									GET_STATFNAME(STAT_UISummaryLLM),				-1)\
+	macro(PhysX,								"PhysX",						GET_STATFNAME(STAT_PhysXLLM),								GET_STATFNAME(STAT_PhysXSummaryLLM),			-1)\
+	macro(EnginePreInitMemory,					"EnginePreInit",				GET_STATFNAME(STAT_EnginePreInitLLM),						GET_STATFNAME(STAT_EngineSummaryLLM),			-1)\
+	macro(EngineInitMemory,						"EngineInit",					GET_STATFNAME(STAT_EngineInitLLM),							GET_STATFNAME(STAT_EngineSummaryLLM),			-1)\
+	macro(RenderingThreadMemory,				"RenderingThread",				GET_STATFNAME(STAT_RenderingThreadLLM),						GET_STATFNAME(STAT_EngineSummaryLLM),			-1)\
+	macro(LoadMapMisc,							"LoadMapMisc",					GET_STATFNAME(STAT_LoadMapMiscLLM),							GET_STATFNAME(STAT_EngineSummaryLLM),			-1)\
+	macro(StreamingManager,						"StreamingManager",				GET_STATFNAME(STAT_StreamingManagerLLM),					GET_STATFNAME(STAT_EngineSummaryLLM),			-1)\
+	macro(GraphicsPlatform,						"Graphics",						GET_STATFNAME(STAT_GraphicsPlatformLLM),					NAME_None,										-1)\
+	macro(FileSystem,							"FileSystem",					GET_STATFNAME(STAT_FileSystemLLM),							GET_STATFNAME(STAT_EngineSummaryLLM),			-1)\
+	macro(Localization,							"Localization",					GET_STATFNAME(STAT_LocalizationLLM),						GET_STATFNAME(STAT_EngineSummaryLLM),			-1)\
+	macro(VertexBuffer,							"VertexBuffer",					GET_STATFNAME(STAT_VertexBufferLLM),						GET_STATFNAME(STAT_EngineSummaryLLM),			-1)\
+	macro(IndexBuffer,							"IndexBuffer",					GET_STATFNAME(STAT_IndexBufferLLM),							GET_STATFNAME(STAT_EngineSummaryLLM),			-1)\
+	macro(UniformBuffer,						"UniformBuffer",				GET_STATFNAME(STAT_UniformBufferLLM),						GET_STATFNAME(STAT_EngineSummaryLLM),			-1)\
+	macro(AssetRegistry,						"AssetRegistry",				GET_STATFNAME(STAT_AssetRegistryLLM),						GET_STATFNAME(STAT_EngineSummaryLLM),			-1)\
+	macro(ConfigSystem,							"ConfigSystem",					GET_STATFNAME(STAT_ConfigSystemLLM),						GET_STATFNAME(STAT_EngineSummaryLLM),			-1)\
+	macro(InitUObject,							"InitUObject",					GET_STATFNAME(STAT_InitUObjectLLM),							GET_STATFNAME(STAT_EngineSummaryLLM),			-1)\
+	macro(VideoRecording,						"VideoRecording",				GET_STATFNAME(STAT_VideoRecordingLLM),						GET_STATFNAME(STAT_EngineSummaryLLM),			-1)\
+	macro(CsvProfiler,							"CsvProfiler",					GET_STATFNAME(STAT_CsvProfilerLLM),							GET_STATFNAME(STAT_EngineSummaryLLM),			-1)
 
 /*
  * Enum values to be passed in to LLM_SCOPE() macro
  */
 enum class ELLMTag : LLM_TAG_TYPE
 {
-#define LLM_ENUM(Enum,Str,Stat,Group) Enum,
+#define LLM_ENUM(Enum,Str,Stat,Group,Parent) Enum,
 	LLM_ENUM_GENERIC_TAGS(LLM_ENUM)
 #undef LLM_ENUM
 
@@ -430,7 +430,7 @@ public:
     // get the top active tag for the given tracker
     int64 GetActiveTag(ELLMTracker Tracker);
 
-	void RegisterPlatformTag(int32 Tag, const TCHAR* Name, FName StatName, FName SummaryStatName);
+	void RegisterPlatformTag(int32 Tag, const TCHAR* Name, FName StatName, FName SummaryStatName, int32 ParentTag = -1);
     
 	// look up the tag associated with the given name
 	bool FindTagByName( const TCHAR* Name, uint64& OutTag ) const;
@@ -471,6 +471,8 @@ private:
 
 	FLLMTracker* Trackers[(int32)ELLMTracker::Max];
 
+	int32 ParentTags[(int32)ELLMTag::PlatformTagEnd];
+
 	static FLowLevelMemTracker* TrackerInstance;
 
 public: // really internal but needs to be visible for LLM_IF_ENABLED macro
@@ -530,4 +532,3 @@ protected:
 	#define LLM_SCOPED_PAUSE_TRACKING_WITH_STAT_AND_AMOUNT(...)
 	#define LLM_PUSH_STATS_FOR_ASSET_TAGS()
 #endif		// #if ENABLE_LOW_LEVEL_MEM_TRACKER
-
diff --git a/Engine/Source/Runtime/Engine/Classes/Kismet/KismetMathLibrary.h b/Engine/Source/Runtime/Engine/Classes/Kismet/KismetMathLibrary.h
--- a/Engine/Source/Runtime/Engine/Classes/Kismet/KismetMathLibrary.h
+++ b/Engine/Source/Runtime/Engine/Classes/Kismet/KismetMathLibrary.h
@@ -1835,7 +1835,7 @@ class ENGINE_API UKismetMathLibrary : public UBlueprintFunctionLibrary
 
 	/** Breaks a 4D vector apart into X, Y, Z, W. */
 	UFUNCTION(BlueprintPure, meta = (NativeBreakFunc), Category = "Math|Vector4")
-	static void BreakVector4(FVector4 InVec, float& X, float& Y, float& Z, float& W);
+	static void BreakVector4(const FVector4& InVec, float& X, float& Y, float& Z, float& W);
 
 	/** Convert a Vector4 to a Vector (dropping the W element) */
 	UFUNCTION(BlueprintPure, meta = (DisplayName = "To Vector (Vector4)", CompactNodeTitle = "->", ScriptMethod = "Vector", Keywords = "cast convert", BlueprintAutocast), Category = "Math|Conversions")
diff --git a/Engine/Source/Runtime/Engine/Classes/Kismet/KismetMathLibrary.inl b/Engine/Source/Runtime/Engine/Classes/Kismet/KismetMathLibrary.inl
--- a/Engine/Source/Runtime/Engine/Classes/Kismet/KismetMathLibrary.inl
+++ b/Engine/Source/Runtime/Engine/Classes/Kismet/KismetMathLibrary.inl
@@ -1833,7 +1833,7 @@ FVector4 UKismetMathLibrary::MakeVector4(float X, float Y, float Z, float W)
 }
 
 KISMET_MATH_FORCEINLINE
-void UKismetMathLibrary::BreakVector4(FVector4 InVec, float& X, float& Y, float& Z, float& W)
+void UKismetMathLibrary::BreakVector4(const FVector4& InVec, float& X, float& Y, float& Z, float& W)
 {
 	X = InVec.X;
 	Y = InVec.Y;
diff --git a/Engine/Source/Runtime/MediaUtils/Private/AppMediaTimeSource.cpp b/Engine/Source/Runtime/MediaUtils/Private/AppMediaTimeSource.cpp
--- a/Engine/Source/Runtime/MediaUtils/Private/AppMediaTimeSource.cpp
+++ b/Engine/Source/Runtime/MediaUtils/Private/AppMediaTimeSource.cpp
@@ -6,6 +6,8 @@
 #include "Misc/App.h"
 #include "Misc/Timespan.h"
 
+/** Declares a log category for this module. */
+DECLARE_LOG_CATEGORY_CLASS(LogMediaTimeSource, Log, All);
 
 /* IMediaTimeSource interface 
  *****************************************************************************/
@@ -15,7 +17,7 @@ FTimespan FAppMediaTimeSource::GetTimecode()
 	const double CurrentTime = FApp::GetCurrentTime();
 	const FTimespan Timecode = FTimespan::FromSeconds(CurrentTime);
 
-	UE_LOG(LogMediaUtils, VeryVerbose, TEXT("AppMediaTimeSource: Time %.10f, Delta %.10f, Timecode %s"), CurrentTime, FApp::GetDeltaTime(), *Timecode.ToString(TEXT("%h:%m:%s.%t")));
+	UE_LOG(LogMediaTimeSource, VeryVerbose, TEXT("AppMediaTimeSource: Time %.10f, Delta %.10f, Timecode %s"), CurrentTime, FApp::GetDeltaTime(), *Timecode.ToString(TEXT("%h:%m:%s.%t")));
 
 	return Timecode;
 }
diff --git a/Engine/Source/Runtime/NavigationSystem/Private/NavAreas/NavArea.cpp b/Engine/Source/Runtime/NavigationSystem/Private/NavAreas/NavArea.cpp
--- a/Engine/Source/Runtime/NavigationSystem/Private/NavAreas/NavArea.cpp
+++ b/Engine/Source/Runtime/NavigationSystem/Private/NavAreas/NavArea.cpp
@@ -44,7 +44,8 @@ void UNavArea::PostInitProperties()
 
 void UNavArea::RegisterArea()
 {
-	if (HasAnyFlags(RF_ClassDefaultObject)
+	if (HasAnyFlags(RF_ClassDefaultObject) && 
+		!HasAnyFlags(RF_NeedInitialization)  // Don't register BP Area that has still not finished loaded their properties, it was also try again to register later via UNavArea::PostLoad()
 #if WITH_HOT_RELOAD
 		&& !GIsHotReload
 #endif // WITH_HOT_RELOAD
diff --git a/Engine/Source/Runtime/RHI/Public/RHIShaderPlatformDefinitions.inl b/Engine/Source/Runtime/RHI/Public/RHIShaderPlatformDefinitions.inl
--- a/Engine/Source/Runtime/RHI/Public/RHIShaderPlatformDefinitions.inl
+++ b/Engine/Source/Runtime/RHI/Public/RHIShaderPlatformDefinitions.inl
@@ -1,4 +1,4 @@
-// Copyright 1998-2018 Epic Games, Inc. All Rights Reserved.
+// Copyright 1998-2019 Epic Games, Inc. All Rights Reserved.
 
 /*=============================================================================
 	RHIShaderPlatformDefinitions.h: Localizable Friendly Names for Shader Platforms
