Add the ability to call FMemory::Trim() and not broadcast. Allows it to be called from other threads.

diff --git a/Engine/Source/Runtime/Core/Private/HAL/MallocBinned2.cpp b/Engine/Source/Runtime/Core/Private/HAL/MallocBinned2.cpp
--- a/Engine/Source/Runtime/Core/Private/HAL/MallocBinned2.cpp
+++ b/Engine/Source/Runtime/Core/Private/HAL/MallocBinned2.cpp
@@ -1056,11 +1056,11 @@ void FMallocBinned2::FlushCurrentThreadCache()
 
 #include "Async/TaskGraphInterfaces.h"
 
-void FMallocBinned2::Trim()
+void FMallocBinned2::Trim(bool bTrimThreadCaches)
 {
 	QUICK_SCOPE_CYCLE_COUNTER(STAT_FMallocBinned2_Trim);
 
-	if (GMallocBinned2PerThreadCaches)
+	if (GMallocBinned2PerThreadCaches  &&  bTrimThreadCaches)
 	{
 		//double StartTime = FPlatformTime::Seconds();
 		TFunction<void(ENamedThreads::Type CurrentThread)> Broadcast =
diff --git a/Engine/Source/Runtime/Core/Private/HAL/MallocBinned3.cpp b/Engine/Source/Runtime/Core/Private/HAL/MallocBinned3.cpp
--- a/Engine/Source/Runtime/Core/Private/HAL/MallocBinned3.cpp
+++ b/Engine/Source/Runtime/Core/Private/HAL/MallocBinned3.cpp
@@ -1756,11 +1756,11 @@ void FMallocBinned3::FlushCurrentThreadCache()
 
 #include "Async/TaskGraphInterfaces.h"
 
-void FMallocBinned3::Trim()
+void FMallocBinned3::Trim(bool bTrimThreadCaches)
 {
 	QUICK_SCOPE_CYCLE_COUNTER(STAT_FMallocBinned3_Trim);
 
-	if (GMallocBinned3PerThreadCaches)
+	if (GMallocBinned3PerThreadCaches  &&  bTrimThreadCaches)
 	{
 		//double StartTime = FPlatformTime::Seconds();
 		TFunction<void(ENamedThreads::Type CurrentThread)> Broadcast =
diff --git a/Engine/Source/Runtime/Core/Private/HAL/MallocLeakDetectionProxy.h b/Engine/Source/Runtime/Core/Private/HAL/MallocLeakDetectionProxy.h
--- a/Engine/Source/Runtime/Core/Private/HAL/MallocLeakDetectionProxy.h
+++ b/Engine/Source/Runtime/Core/Private/HAL/MallocLeakDetectionProxy.h
@@ -104,9 +104,9 @@ public:
 		return UsedMalloc->QuantizeSize(Count, Alignment);
 	}
 
-	virtual void Trim() override
+	virtual void Trim(bool bTrimThreadCaches) override
 	{
-		return UsedMalloc->Trim();
+		return UsedMalloc->Trim(bTrimThreadCaches);
 	}
 	virtual void SetupTLSCachesOnCurrentThread() override
 	{
diff --git a/Engine/Source/Runtime/Core/Private/HAL/UnrealMemory.cpp b/Engine/Source/Runtime/Core/Private/HAL/UnrealMemory.cpp
--- a/Engine/Source/Runtime/Core/Private/HAL/UnrealMemory.cpp
+++ b/Engine/Source/Runtime/Core/Private/HAL/UnrealMemory.cpp
@@ -197,9 +197,9 @@ public:
 	{
 		return UsedMalloc->QuantizeSize(Count, Alignment);
 	}
-	virtual void Trim() override
+	virtual void Trim(bool bTrimThreadCaches) override
 	{
-		return UsedMalloc->Trim();
+		return UsedMalloc->Trim(bTrimThreadCaches);
 	}
 	virtual void SetupTLSCachesOnCurrentThread() override
 	{
@@ -504,7 +504,7 @@ SIZE_T FMemory::QuantizeSizeExternal(SIZE_T Count, uint32 Alignment)
 }	
 
 
-void FMemory::Trim()
+void FMemory::Trim(bool bTrimThreadCaches)
 {
 	if (!GMalloc)
 	{
@@ -517,7 +517,7 @@ void FMemory::Trim()
 		FCoreDelegates::GetMemoryTrimDelegate().Broadcast();
 	}
 	QUICK_SCOPE_CYCLE_COUNTER(STAT_FMemory_Trim_GMalloc);
-	GMalloc->Trim();
+	GMalloc->Trim(bTrimThreadCaches);
 }
 
 void FMemory::SetupTLSCachesOnCurrentThread()
diff --git a/Engine/Source/Runtime/Core/Public/HAL/MallocBinned2.h b/Engine/Source/Runtime/Core/Public/HAL/MallocBinned2.h
--- a/Engine/Source/Runtime/Core/Public/HAL/MallocBinned2.h
+++ b/Engine/Source/Runtime/Core/Public/HAL/MallocBinned2.h
@@ -650,7 +650,7 @@ public:
 	}
 
 	virtual bool ValidateHeap() override;
-	virtual void Trim() override;
+	virtual void Trim(bool bTrimThreadCaches) override;
 	virtual void SetupTLSCachesOnCurrentThread() override;
 	virtual void ClearAndDisableTLSCachesOnCurrentThread() override;
 	virtual const TCHAR* GetDescriptiveName() override;
diff --git a/Engine/Source/Runtime/Core/Public/HAL/MallocBinned3.h b/Engine/Source/Runtime/Core/Public/HAL/MallocBinned3.h
--- a/Engine/Source/Runtime/Core/Public/HAL/MallocBinned3.h
+++ b/Engine/Source/Runtime/Core/Public/HAL/MallocBinned3.h
@@ -677,7 +677,7 @@ public:
 	}
 
 	virtual bool ValidateHeap() override;
-	virtual void Trim() override;
+	virtual void Trim(bool bTrimThreadCaches) override;
 	virtual void SetupTLSCachesOnCurrentThread() override;
 	virtual void ClearAndDisableTLSCachesOnCurrentThread() override;
 	virtual const TCHAR* GetDescriptiveName() override;
diff --git a/Engine/Source/Runtime/Core/Public/HAL/MallocPoisonProxy.h b/Engine/Source/Runtime/Core/Public/HAL/MallocPoisonProxy.h
--- a/Engine/Source/Runtime/Core/Public/HAL/MallocPoisonProxy.h
+++ b/Engine/Source/Runtime/Core/Public/HAL/MallocPoisonProxy.h
@@ -132,9 +132,9 @@ public:
 		return UsedMalloc->GetDescriptiveName(); 
 	}
 
-	virtual void Trim() override
+	virtual void Trim(bool bTrimThreadCaches) override
 	{
-		UsedMalloc->Trim();
+		UsedMalloc->Trim(bTrimThreadCaches);
 	}
 
 	virtual void SetupTLSCachesOnCurrentThread() override
diff --git a/Engine/Source/Runtime/Core/Public/HAL/MallocReplayProxy.h b/Engine/Source/Runtime/Core/Public/HAL/MallocReplayProxy.h
--- a/Engine/Source/Runtime/Core/Public/HAL/MallocReplayProxy.h
+++ b/Engine/Source/Runtime/Core/Public/HAL/MallocReplayProxy.h
@@ -115,7 +115,7 @@ public:
 
 	virtual const TCHAR* GetDescriptiveName() override;
 
-	virtual void Trim() override;
+	virtual void Trim(bool bTrimThreadCaches) override;
 
 	virtual void SetupTLSCachesOnCurrentThread() override;
 
diff --git a/Engine/Source/Runtime/Core/Public/HAL/MallocThreadSafeProxy.h b/Engine/Source/Runtime/Core/Public/HAL/MallocThreadSafeProxy.h
--- a/Engine/Source/Runtime/Core/Public/HAL/MallocThreadSafeProxy.h
+++ b/Engine/Source/Runtime/Core/Public/HAL/MallocThreadSafeProxy.h
@@ -116,10 +116,10 @@ public:
 		return UsedMalloc->GetDescriptiveName(); 
 	}
 
-	virtual void Trim() override
+	virtual void Trim(bool bTrimThreadCaches) override
 	{
 		FScopeLock ScopeLock(&SynchronizationObject);
 		check(UsedMalloc);
-		UsedMalloc->Trim();
+		UsedMalloc->Trim(bTrimThreadCaches);
 	}
 };
diff --git a/Engine/Source/Runtime/Core/Public/HAL/MemoryBase.h b/Engine/Source/Runtime/Core/Public/HAL/MemoryBase.h
--- a/Engine/Source/Runtime/Core/Public/HAL/MemoryBase.h
+++ b/Engine/Source/Runtime/Core/Public/HAL/MemoryBase.h
@@ -117,7 +117,7 @@ public:
 	/**
 	* Releases as much memory as possible. Must be called from the main thread.
 	*/
-	virtual void Trim()
+	virtual void Trim(bool bTrimThreadCaches)
 	{
 	}
 
diff --git a/Engine/Source/Runtime/Core/Public/HAL/UnrealMemory.h b/Engine/Source/Runtime/Core/Public/HAL/UnrealMemory.h
--- a/Engine/Source/Runtime/Core/Public/HAL/UnrealMemory.h
+++ b/Engine/Source/Runtime/Core/Public/HAL/UnrealMemory.h
@@ -178,7 +178,7 @@ struct CORE_API FMemory
 	/**
 	* Releases as much memory as possible. Must be called from the main thread.
 	*/
-	static void Trim();
+	static void Trim(bool bTrimThreadCaches = true);
 
 	/**
 	* Set up TLS caches on the current thread. These are the threads that we can trim.
