Fix linux compile error caused by CL 4715476

diff --git a/Engine/Source/Runtime/Core/Private/HAL/MallocReplayProxy.cpp b/Engine/Source/Runtime/Core/Private/HAL/MallocReplayProxy.cpp
--- a/Engine/Source/Runtime/Core/Private/HAL/MallocReplayProxy.cpp
+++ b/Engine/Source/Runtime/Core/Private/HAL/MallocReplayProxy.cpp
@@ -156,9 +156,9 @@ const TCHAR* FMallocReplayProxy::GetDescriptiveName()
 	return UsedMalloc->GetDescriptiveName();
 }
 
-void FMallocReplayProxy::Trim()
+void FMallocReplayProxy::Trim(bool bTrimThreadCaches)
 {
-	UsedMalloc->Trim();
+	UsedMalloc->Trim(bTrimThreadCaches);
 }
 
 void FMallocReplayProxy::SetupTLSCachesOnCurrentThread()
