MallocBinned3: Fix assert in MallocExternal when Alignment>OSAllocationGranularity. OSAllocationGranularity is not the correct value to use in this case. We actually want FPlatformMemory::GetConstants().OSAllocationGranularity, which gives the alignment of the VM allocations on the platform.

Fix the AlignedSize calculation to use AlignedSize if it's larger then OSAllocationGranularity.

diff --git a/Engine/Source/Runtime/Core/Private/HAL/MallocBinned3.cpp b/Engine/Source/Runtime/Core/Private/HAL/MallocBinned3.cpp
--- a/Engine/Source/Runtime/Core/Private/HAL/MallocBinned3.cpp
+++ b/Engine/Source/Runtime/Core/Private/HAL/MallocBinned3.cpp
@@ -257,6 +257,7 @@ MS_ALIGN(PLATFORM_CACHE_LINE_SIZE) static uint8 Binned3UnusedAlignPadding[PLATFO
 uint16 FMallocBinned3::SmallBlockSizesReversedShifted[BINNED3_SMALL_POOL_COUNT] = { 0 };
 uint32 FMallocBinned3::Binned3TlsSlot = 0;
 uint32 FMallocBinned3::OsAllocationGranularity = 0;
+uint32 FMallocBinned3::MaxAlignmentForMemoryRangeReserve = 0;
 
 #if !BINNED3_USE_SEPARATE_VM_PER_POOL
 	uint8* FMallocBinned3::Binned3BaseVMPtr = nullptr;
@@ -1258,6 +1259,7 @@ FMallocBinned3::FMallocBinned3()
 
 	FGenericPlatformMemoryConstants Constants = FPlatformMemory::GetConstants();
 	OsAllocationGranularity = Constants.BinnedAllocationGranularity ? Constants.BinnedAllocationGranularity : Constants.PageSize;
+	MaxAlignmentForMemoryRangeReserve = Constants.OsAllocationGranularity;
 	NumLargePoolsPerPage = OsAllocationGranularity / sizeof(FPoolInfoLarge);
 	check(OsAllocationGranularity % sizeof(FPoolInfoLarge) == 0);  // these need to divide evenly!
 	PtrToPoolMapping.Init(OsAllocationGranularity, NumLargePoolsPerPage, Constants.AddressLimit);
@@ -1480,12 +1482,12 @@ void* FMallocBinned3::MallocExternal(SIZE_T Size, uint32 Alignment)
 	Size = Align(FMath::Max((SIZE_T)1, Size), Alignment);
 
 	check(FMath::IsPowerOfTwo(Alignment));
-	check(Alignment <= OsAllocationGranularity);
+	UE_CLOG(Alignment > MaxAlignmentForMemoryRangeReserve ,LogMemory, Fatal, TEXT("Requested alignment was too large for OS. Alignment=%dKB MaxAlignmentForMemoryRangeReserve=%dKB"), Alignment/1024, MaxAlignmentForMemoryRangeReserve/1024);
 
-	FScopeLock Lock(&Mutex);
+	FScopeLock Lock(&Mutex); 
 
 	// Use OS for non-pooled allocations.
-	UPTRINT AlignedSize = Align(Size, OsAllocationGranularity);
+	UPTRINT AlignedSize = Align(Size, FMath::Max(OsAllocationGranularity,Alignment));
 
 #if BINNED3_TIME_LARGE_BLOCKS
 	double StartTime = FPlatformTime::Seconds();
diff --git a/Engine/Source/Runtime/Core/Public/HAL/MallocBinned3.h b/Engine/Source/Runtime/Core/Public/HAL/MallocBinned3.h
--- a/Engine/Source/Runtime/Core/Public/HAL/MallocBinned3.h
+++ b/Engine/Source/Runtime/Core/Public/HAL/MallocBinned3.h
@@ -700,6 +700,7 @@ public:
 	static FMallocBinned3* MallocBinned3;
 	static uint32 Binned3TlsSlot;
 	static uint32 OsAllocationGranularity;
+	static uint32 MaxAlignmentForMemoryRangeReserve;
 #if !BINNED3_USE_SEPARATE_VM_PER_POOL
 	static uint8* Binned3BaseVMPtr;
 #else
