Made FTextFormatter able to deal with different formatting token markers

Expanded plural argument modifiers to work with FText arguments generated from numbers, as well as raw number arguments.

diff --git a/Engine/Source/Runtime/Core/Private/Internationalization/TextFormatArgumentModifier.cpp b/Engine/Source/Runtime/Core/Private/Internationalization/TextFormatArgumentModifier.cpp
--- a/Engine/Source/Runtime/Core/Private/Internationalization/TextFormatArgumentModifier.cpp
+++ b/Engine/Source/Runtime/Core/Private/Internationalization/TextFormatArgumentModifier.cpp
@@ -164,7 +164,7 @@ bool ITextFormatArgumentModifier::ParseValueArgs(const FTextFormatString& InArgs
 }
 
 
-TSharedPtr<ITextFormatArgumentModifier> FTextFormatArgumentModifier_PluralForm::Create(const ETextPluralType InPluralType, const FTextFormatString& InArgsString)
+TSharedPtr<ITextFormatArgumentModifier> FTextFormatArgumentModifier_PluralForm::Create(const ETextPluralType InPluralType, const FTextFormatString& InArgsString, const FTextFormatPatternDefinitionConstRef& InPatternDef)
 {
 	TMap<FTextFormatString, FTextFormatString> ArgKeyValues;
 	if (ParseKeyValueArgs(InArgsString, ArgKeyValues))
@@ -177,7 +177,7 @@ TSharedPtr<ITextFormatArgumentModifier> FTextFormatArgumentModifier_PluralForm::
 		PluralForms.Reserve(ArgKeyValues.Num());
 		for (const auto& Pair : ArgKeyValues)
 		{
-			FTextFormat PluralForm = FTextFormat::FromString(FString(Pair.Value.StringLen, Pair.Value.StringPtr));
+			FTextFormat PluralForm = FTextFormat::FromString(FString(Pair.Value.StringLen, Pair.Value.StringPtr), InPatternDef);
 			if (!PluralForm.IsValid())
 			{
 				break;
@@ -281,23 +281,43 @@ void FTextFormatArgumentModifier_PluralForm::Evaluate(const FFormatArgumentValue
 	FInternationalization& I18N = FInternationalization::Get();
 	const FCulture& Culture = *I18N.GetCurrentLanguage();
 
+	auto TryGetPluralFormForArgument = [this, &Culture](const FFormatArgumentValue& InArgValue, const int32 InArgValueMultiplier, ETextPluralForm& OutArgValuePluralForm) -> bool
+	{
+		switch (InArgValue.GetType())
+		{
+		case EFormatArgumentType::Int:
+			OutArgValuePluralForm = Culture.GetPluralForm(InArgValue.GetIntValue() * static_cast<int64>(InArgValueMultiplier), PluralType);
+			return true;
+		case EFormatArgumentType::UInt:
+			OutArgValuePluralForm = Culture.GetPluralForm(InArgValue.GetUIntValue() * static_cast<uint64>(InArgValueMultiplier), PluralType);
+			return true;
+		case EFormatArgumentType::Float:
+			OutArgValuePluralForm = Culture.GetPluralForm(InArgValue.GetFloatValue() * static_cast<float>(InArgValueMultiplier), PluralType);
+			return true;
+		case EFormatArgumentType::Double:
+			OutArgValuePluralForm = Culture.GetPluralForm(InArgValue.GetDoubleValue() * static_cast<double>(InArgValueMultiplier), PluralType);
+			return true;
+		default:
+			break;
+		}
+		return false;
+	};
+
 	ETextPluralForm ValuePluralForm = ETextPluralForm::Other;
-	switch (InValue.GetType())
+	if (!TryGetPluralFormForArgument(InValue, 1, ValuePluralForm) && InValue.GetType() == EFormatArgumentType::Text)
 	{
-	case EFormatArgumentType::Int:
-		ValuePluralForm = Culture.GetPluralForm(InValue.GetIntValue(), PluralType);
-		break;
-	case EFormatArgumentType::UInt:
-		ValuePluralForm = Culture.GetPluralForm(InValue.GetUIntValue(), PluralType);
-		break;
-	case EFormatArgumentType::Float:
-		ValuePluralForm = Culture.GetPluralForm(InValue.GetFloatValue(), PluralType);
-		break;
-	case EFormatArgumentType::Double:
-		ValuePluralForm = Culture.GetPluralForm(InValue.GetDoubleValue(), PluralType);
-		break;
-	default:
-		break;
+		const FText TextValue = InValue.GetTextValue();
+
+		// If this text was generated from a number, extract its numeric argument and use its number to get the plural form
+		FHistoricTextNumericData TextValueNumericData;
+		if (FTextInspector::GetHistoricNumericData(TextValue, TextValueNumericData))
+		{
+			TryGetPluralFormForArgument(
+				TextValueNumericData.SourceValue, 
+				TextValueNumericData.FormatType == FHistoricTextNumericData::EType::AsPercent ? 100 : 1, 
+				ValuePluralForm
+				);
+		}
 	}
 
 	OutResult += FTextFormatter::Format(CompiledPluralForms[(int32)ValuePluralForm], InFormatArgs);
@@ -340,18 +360,18 @@ FTextFormatArgumentModifier_PluralForm::FTextFormatArgumentModifier_PluralForm(c
 }
 
 
-TSharedPtr<ITextFormatArgumentModifier> FTextFormatArgumentModifier_GenderForm::Create(const FTextFormatString& InArgsString)
+TSharedPtr<ITextFormatArgumentModifier> FTextFormatArgumentModifier_GenderForm::Create(const FTextFormatString& InArgsString, const FTextFormatPatternDefinitionConstRef& InPatternDef)
 {
 	TArray<FTextFormatString> ArgValues;
 	if (ParseValueArgs(InArgsString, ArgValues) && (ArgValues.Num() == 2 || ArgValues.Num() == 3))
 	{
 		// Gender forms may contain format markers, so pre-compile all the variants now so that Evaluate doesn't have to (this also lets us validate the gender form strings and fail if they're not correct)
-		FTextFormat MasculineForm = FTextFormat::FromString(FString(ArgValues[0].StringLen, ArgValues[0].StringPtr));
-		FTextFormat FeminineForm  = FTextFormat::FromString(FString(ArgValues[1].StringLen, ArgValues[1].StringPtr));
+		FTextFormat MasculineForm = FTextFormat::FromString(FString(ArgValues[0].StringLen, ArgValues[0].StringPtr), InPatternDef);
+		FTextFormat FeminineForm  = FTextFormat::FromString(FString(ArgValues[1].StringLen, ArgValues[1].StringPtr), InPatternDef);
 		FTextFormat NeuterForm;
 		if (ArgValues.Num() == 3)
 		{
-			NeuterForm = FTextFormat::FromString(FString(ArgValues[2].StringLen, ArgValues[2].StringPtr));;
+			NeuterForm = FTextFormat::FromString(FString(ArgValues[2].StringLen, ArgValues[2].StringPtr), InPatternDef);
 		}
 
 		// Did everything compile?
diff --git a/Engine/Source/Runtime/Core/Private/Internationalization/TextFormatArgumentModifier.h b/Engine/Source/Runtime/Core/Private/Internationalization/TextFormatArgumentModifier.h
--- a/Engine/Source/Runtime/Core/Private/Internationalization/TextFormatArgumentModifier.h
+++ b/Engine/Source/Runtime/Core/Private/Internationalization/TextFormatArgumentModifier.h
@@ -22,7 +22,7 @@ template<typename KeyType, typename ValueType, typename SetAllocator, typename K
 class FTextFormatArgumentModifier_PluralForm : public ITextFormatArgumentModifier
 {
 public:
-	static TSharedPtr<ITextFormatArgumentModifier> Create(const ETextPluralType InPluralType, const FTextFormatString& InArgsString);
+	static TSharedPtr<ITextFormatArgumentModifier> Create(const ETextPluralType InPluralType, const FTextFormatString& InArgsString, const FTextFormatPatternDefinitionConstRef& InPatternDef);
 
 	virtual bool Validate(const FCultureRef& InCulture, TArray<FString>& OutValidationErrors) const override;
 
@@ -49,7 +49,7 @@ private:
 class FTextFormatArgumentModifier_GenderForm : public ITextFormatArgumentModifier
 {
 public:
-	static TSharedPtr<ITextFormatArgumentModifier> Create(const FTextFormatString& InArgsString);
+	static TSharedPtr<ITextFormatArgumentModifier> Create(const FTextFormatString& InArgsString, const FTextFormatPatternDefinitionConstRef& InPatternDef);
 
 	virtual bool Validate(const FCultureRef& InCulture, TArray<FString>& OutValidationErrors) const override;
 
diff --git a/Engine/Source/Runtime/Core/Private/Internationalization/TextFormatter.cpp b/Engine/Source/Runtime/Core/Private/Internationalization/TextFormatter.cpp
--- a/Engine/Source/Runtime/Core/Private/Internationalization/TextFormatter.cpp
+++ b/Engine/Source/Runtime/Core/Private/Internationalization/TextFormatter.cpp
@@ -19,38 +19,6 @@ DECLARE_CYCLE_STAT(TEXT("TextFormatData Compile"), STAT_TextFormatData_Compile,
 namespace TextFormatTokens
 {
 
-const TCHAR EscapeChar			= TEXT('`');		// Character representing the start of an escape token
-const TCHAR ArgStartChar		= TEXT('{');		// Character representing the start of a format argument token
-const TCHAR ArgEndChar			= TEXT('}');		// Character representing the end of a format argument token
-const TCHAR ArgModChar			= TEXT('|');		// Character representing the start of a format argument modifier token
-const TCHAR ValidEscapeChars[]	= TEXT("{}`|");		// Characters that an escape token may escape
-const TCHAR LiteralBreakChars[] = TEXT("{`");		// Characters that should cause a literal string token to break parsing
-
-FORCEINLINE bool ContainsChar(const TCHAR InChar, const TCHAR* InStr)
-{
-	if (InChar)
-	{
-		while (*InStr)
-		{
-			if (*InStr++ == InChar)
-			{
-				return true;
-			}
-		}
-	}
-	return false;
-}
-
-FORCEINLINE bool IsValidEscapeChar(const TCHAR InChar)
-{
-	return ContainsChar(InChar, ValidEscapeChars);
-}
-
-FORCEINLINE bool IsLiteralBreakChar(const TCHAR InChar)
-{
-	return ContainsChar(InChar, LiteralBreakChars);
-}
-
 /** Token representing a literal string inside the text */
 struct FStringLiteral
 {
@@ -138,12 +106,12 @@ struct FEscapedCharacter
 	TCHAR Character;
 };
 
-TOptional<FExpressionError> ParseArgument(FExpressionTokenConsumer& Consumer)
+TOptional<FExpressionError> ParseArgument(const FTextFormatPatternDefinition& InPatternDef, FExpressionTokenConsumer& Consumer)
 {
 	// An argument token looks like {ArgName}
 	FTokenStream& Stream = Consumer.GetStream();
 
-	TOptional<FStringToken> OpeningChar = Stream.ParseSymbol(ArgStartChar);
+	TOptional<FStringToken> OpeningChar = Stream.ParseSymbol(InPatternDef.ArgStartChar);
 	if (!OpeningChar.IsSet())
 	{
 		return TOptional<FExpressionError>();
@@ -152,9 +120,9 @@ TOptional<FExpressionError> ParseArgument(FExpressionTokenConsumer& Consumer)
 	FStringToken& EntireToken = OpeningChar.GetValue();
 
 	// Parse out the argument name
-	TOptional<FStringToken> Identifier = Stream.ParseToken([](TCHAR InC)
+	TOptional<FStringToken> Identifier = Stream.ParseToken([&InPatternDef](TCHAR InC)
 	{
-		if (InC == ArgEndChar)
+		if (InC == InPatternDef.ArgEndChar)
 		{
 			return EParseState::StopBefore;
 		}
@@ -164,7 +132,7 @@ TOptional<FExpressionError> ParseArgument(FExpressionTokenConsumer& Consumer)
 		}
 	}, &EntireToken);
 
-	if (!Identifier.IsSet() || !Stream.ParseSymbol(ArgEndChar, &EntireToken).IsSet())
+	if (!Identifier.IsSet() || !Stream.ParseSymbol(InPatternDef.ArgEndChar, &EntireToken).IsSet())
 	{
 		return TOptional<FExpressionError>();
 	}
@@ -175,12 +143,12 @@ TOptional<FExpressionError> ParseArgument(FExpressionTokenConsumer& Consumer)
 	return TOptional<FExpressionError>();
 }
 
-TOptional<FExpressionError> ParseArgumentModifier(FExpressionTokenConsumer& Consumer)
+TOptional<FExpressionError> ParseArgumentModifier(const FTextFormatPatternDefinition& InPatternDef, FExpressionTokenConsumer& Consumer)
 {
 	// An argument modifier token looks like |keyword(args, ...)
 	FTokenStream& Stream = Consumer.GetStream();
 
-	TOptional<FStringToken> PipeToken = Stream.ParseSymbol(ArgModChar);
+	TOptional<FStringToken> PipeToken = Stream.ParseSymbol(InPatternDef.ArgModChar);
 	if (!PipeToken.IsSet())
 	{
 		return TOptional<FExpressionError>();
@@ -264,7 +232,7 @@ TOptional<FExpressionError> ParseArgumentModifier(FExpressionTokenConsumer& Cons
 
 	// Compile the parameters for this argument modifier
 	FStringToken& ParametersValue = Parameters.GetValue();
-	TSharedPtr<ITextFormatArgumentModifier> CompiledTextArgumentModifier = CompileTextArgumentModifierFunc(FTextFormatString::MakeReference(ParametersValue.GetTokenStartPos(), ParametersValue.GetTokenEndPos() - ParametersValue.GetTokenStartPos()));
+	TSharedPtr<ITextFormatArgumentModifier> CompiledTextArgumentModifier = CompileTextArgumentModifierFunc(FTextFormatString::MakeReference(ParametersValue.GetTokenStartPos(), ParametersValue.GetTokenEndPos() - ParametersValue.GetTokenStartPos()), InPatternDef.AsShared());
 	if (!CompiledTextArgumentModifier.IsValid())
 	{
 		return TOptional<FExpressionError>();
@@ -275,11 +243,11 @@ TOptional<FExpressionError> ParseArgumentModifier(FExpressionTokenConsumer& Cons
 	return TOptional<FExpressionError>();
 }
 
-TOptional<FExpressionError> ParseEscapedChar(FExpressionTokenConsumer& Consumer)
+TOptional<FExpressionError> ParseEscapedChar(const FTextFormatPatternDefinition& InPatternDef, FExpressionTokenConsumer& Consumer)
 {
 	FTokenStream& Stream = Consumer.GetStream();
 
-	TOptional<FStringToken> Token = Stream.ParseSymbol(EscapeChar);
+	TOptional<FStringToken> Token = Stream.ParseSymbol(InPatternDef.EscapeChar);
 	if (!Token.IsSet())
 	{
 		return TOptional<FExpressionError>();
@@ -296,7 +264,7 @@ TOptional<FExpressionError> ParseEscapedChar(FExpressionTokenConsumer& Consumer)
 
 	// Check for a valid escape character
 	const TCHAR Character = *EscapedChar->GetTokenStartPos();
-	if (IsValidEscapeChar(Character))
+	if (InPatternDef.IsValidEscapeChar(Character))
 	{
 		// Add the token to the consumer - this moves the read position in the stream to the end of the token.
 		Consumer.Add(TokenValue, FEscapedCharacter(Character));
@@ -305,7 +273,7 @@ TOptional<FExpressionError> ParseEscapedChar(FExpressionTokenConsumer& Consumer)
 	return TOptional<FExpressionError>();
 }
 
-TOptional<FExpressionError> ParseLiteral(FExpressionTokenConsumer& Consumer)
+TOptional<FExpressionError> ParseLiteral(const FTextFormatPatternDefinition& InPatternDef, FExpressionTokenConsumer& Consumer)
 {
 	FTokenStream& Stream = Consumer.GetStream();
 
@@ -320,7 +288,7 @@ TOptional<FExpressionError> ParseLiteral(FExpressionTokenConsumer& Consumer)
 				bFirstChar = false;
 				return EParseState::Continue;
 			}
-			else if (!IsLiteralBreakChar(C))
+			else if (!InPatternDef.IsLiteralBreakChar(C))
 			{
 				return EParseState::Continue;
 			}
@@ -373,13 +341,13 @@ public:
 	 * Construct an instance from an FText.
 	 * The text will be immediately compiled. 
 	 */
-	explicit FTextFormatData(FText&& InText);
+	FTextFormatData(FText&& InText, FTextFormatPatternDefinitionConstRef InPatternDef);
 
 	/**
 	 * Construct an instance from an FString.
 	 * The string will be immediately compiled. 
 	 */
-	explicit FTextFormatData(FString&& InString);
+	FTextFormatData(FString&& InString, FTextFormatPatternDefinitionConstRef InPatternDef);
 
 	/**
 	 * Test to see whether this instance contains valid compiled data.
@@ -445,6 +413,14 @@ public:
 		return CompiledExpressionType;
 	}
 
+	/**
+	 * Get the format pattern definition being used.
+	 */
+	FORCEINLINE FTextFormatPatternDefinitionConstRef GetPatternDefinition() const
+	{
+		return PatternDef;
+	}
+
 private:
 	/**
 	 * Test to see whether this instance contains valid compiled data.
@@ -493,6 +469,11 @@ private:
 	 */
 	ESourceType SourceType;
 
+	/**
+	 * Definition of the pattern used during a text format.
+	 */
+	FTextFormatPatternDefinitionConstRef PatternDef;
+
 	/**
 	 * Source localized text that is used as the format specifier.
 	 */
@@ -551,28 +532,43 @@ private:
 
 
 FTextFormat::FTextFormat()
-	: TextFormatData(new FTextFormatData(FText()))
+	: TextFormatData(new FTextFormatData(FText(), FTextFormatPatternDefinition::GetDefault()))
 {
 }
 
 FTextFormat::FTextFormat(const FText& InText)
-	: TextFormatData(new FTextFormatData(CopyTemp(InText)))
+	: TextFormatData(new FTextFormatData(CopyTemp(InText), FTextFormatPatternDefinition::GetDefault()))
 {
 }
 
-FTextFormat::FTextFormat(FString&& InString)
-	: TextFormatData(new FTextFormatData(MoveTemp(InString)))
+FTextFormat::FTextFormat(const FText& InText, FTextFormatPatternDefinitionConstRef InCustomPatternDef)
+	: TextFormatData(new FTextFormatData(CopyTemp(InText), MoveTemp(InCustomPatternDef)))
+{
+}
+
+FTextFormat::FTextFormat(FString&& InString, FTextFormatPatternDefinitionConstRef InCustomPatternDef)
+	: TextFormatData(new FTextFormatData(MoveTemp(InString), MoveTemp(InCustomPatternDef)))
 {
 }
 
 FTextFormat FTextFormat::FromString(const FString& InString)
 {
-	return FTextFormat(CopyTemp(InString));
+	return FTextFormat(CopyTemp(InString), FTextFormatPatternDefinition::GetDefault());
 }
 
 FTextFormat FTextFormat::FromString(FString&& InString)
 {
-	return FTextFormat(MoveTemp(InString));
+	return FTextFormat(MoveTemp(InString), FTextFormatPatternDefinition::GetDefault());
+}
+
+FTextFormat FTextFormat::FromString(const FString& InString, FTextFormatPatternDefinitionConstRef InCustomPatternDef)
+{
+	return FTextFormat(CopyTemp(InString), MoveTemp(InCustomPatternDef));
+}
+
+FTextFormat FTextFormat::FromString(FString&& InString, FTextFormatPatternDefinitionConstRef InCustomPatternDef)
+{
+	return FTextFormat(MoveTemp(InString), MoveTemp(InCustomPatternDef));
 }
 
 bool FTextFormat::IsValid() const
@@ -595,6 +591,11 @@ FTextFormat::EExpressionType FTextFormat::GetExpressionType() const
 	return TextFormatData->GetExpressionType();
 }
 
+FTextFormatPatternDefinitionConstRef FTextFormat::GetPatternDefinition() const
+{
+	return TextFormatData->GetPatternDefinition();
+}
+
 bool FTextFormat::ValidatePattern(const FCulturePtr& InCulture, TArray<FString>& OutValidationErrors) const
 {
 	return TextFormatData->ValidatePattern(InCulture, OutValidationErrors);
@@ -606,15 +607,17 @@ void FTextFormat::GetFormatArgumentNames(TArray<FString>& OutArgumentNames) cons
 }
 
 
-FTextFormatData::FTextFormatData(FText&& InText)
+FTextFormatData::FTextFormatData(FText&& InText, FTextFormatPatternDefinitionConstRef InPatternDef)
 	: SourceType(ESourceType::Text)
+	, PatternDef(MoveTemp(InPatternDef))
 	, SourceText(MoveTemp(InText))
 {
 	Compile_NoLock();
 }
 
-FTextFormatData::FTextFormatData(FString&& InString)
+FTextFormatData::FTextFormatData(FString&& InString, FTextFormatPatternDefinitionConstRef InPatternDef)
 	: SourceType(ESourceType::String)
+	, PatternDef(MoveTemp(InPatternDef))
 	, SourceExpression(MoveTemp(InString))
 {
 	Compile_NoLock();
@@ -639,7 +642,7 @@ void FTextFormatData::Compile_NoLock()
 	BaseFormatStringLength = 0;
 	FormatArgumentEstimateMultiplier = 1;
 
-	TValueOrError<TArray<FExpressionToken>, FExpressionError> Result = ExpressionParser::Lex(*SourceExpression, FTextFormatter::Get().GetTextFormatDefinitions());
+	TValueOrError<TArray<FExpressionToken>, FExpressionError> Result = ExpressionParser::Lex(*SourceExpression, PatternDef->GetTextFormatDefinitions());
 	bool bValidExpression = Result.IsValid();
 	if (bValidExpression)
 	{
@@ -808,16 +811,16 @@ FString FTextFormatData::Format_NoLock(const FPrivateTextFormatArguments& InForm
 			}
 			else
 			{
-				ResultString.AppendChar(TextFormatTokens::ArgStartChar);
+				ResultString.AppendChar(PatternDef->ArgStartChar);
 				ResultString.AppendChars(ArgumentToken->ArgumentNameStartPos, ArgumentToken->ArgumentNameLen);
-				ResultString.AppendChar(TextFormatTokens::ArgEndChar);
+				ResultString.AppendChar(PatternDef->ArgEndChar);
 			}
 		}
 		else if (const auto* ArgumentModifierToken = Token.Node.Cast<TextFormatTokens::FArgumentModifierTokenSpecifier>())
 		{
 			// If we find an argument modifier token on its own then it means an argument value failed to evaluate (likely due to InFormatArgs.GetArgumentValue returning null)
 			// In this case we just write the literal value of the argument modifier back into the final string...
-			ResultString.AppendChar(TextFormatTokens::ArgModChar);
+			ResultString.AppendChar(PatternDef->ArgModChar);
 			ResultString.AppendChars(ArgumentModifierToken->ModifierPatternStartPos, ArgumentModifierToken->ModifierPatternLen);
 		}
 	}
@@ -858,17 +861,32 @@ void FTextFormatData::GetFormatArgumentNames_NoLock(TArray<FString>& OutArgument
 }
 
 
+FTextFormatPatternDefinition::FTextFormatPatternDefinition()
+{
+	TextFormatDefinitions.DefineToken([this](FExpressionTokenConsumer& Consumer) { return TextFormatTokens::ParseArgument(*this, Consumer); });
+	TextFormatDefinitions.DefineToken([this](FExpressionTokenConsumer& Consumer) { return TextFormatTokens::ParseArgumentModifier(*this, Consumer); });
+	TextFormatDefinitions.DefineToken([this](FExpressionTokenConsumer& Consumer) { return TextFormatTokens::ParseEscapedChar(*this, Consumer); });
+	TextFormatDefinitions.DefineToken([this](FExpressionTokenConsumer& Consumer) { return TextFormatTokens::ParseLiteral(*this, Consumer); });
+}
+
+FTextFormatPatternDefinitionConstRef FTextFormatPatternDefinition::GetDefault()
+{
+	static FTextFormatPatternDefinitionConstRef DefaultFormatPatternDefinition = MakeShared<FTextFormatPatternDefinition, ESPMode::ThreadSafe>();
+	return DefaultFormatPatternDefinition;
+}
+
+const FTokenDefinitions& FTextFormatPatternDefinition::GetTextFormatDefinitions() const
+{
+	return TextFormatDefinitions;
+}
+
+
 FTextFormatter::FTextFormatter()
 {
-	TextFormatDefinitions.DefineToken([](FExpressionTokenConsumer& Consumer) { return TextFormatTokens::ParseArgument(Consumer); });
-	TextFormatDefinitions.DefineToken([](FExpressionTokenConsumer& Consumer) { return TextFormatTokens::ParseArgumentModifier(Consumer); });
-	TextFormatDefinitions.DefineToken([](FExpressionTokenConsumer& Consumer) { return TextFormatTokens::ParseEscapedChar(Consumer); });
-	TextFormatDefinitions.DefineToken([](FExpressionTokenConsumer& Consumer) { return TextFormatTokens::ParseLiteral(Consumer); });
-
-	TextArgumentModifiers.Add(FTextFormatString::MakeReference(TEXT("plural")),  [](const FTextFormatString& InArgsString) { return FTextFormatArgumentModifier_PluralForm::Create(ETextPluralType::Cardinal, InArgsString); });
-	TextArgumentModifiers.Add(FTextFormatString::MakeReference(TEXT("ordinal")), [](const FTextFormatString& InArgsString) { return FTextFormatArgumentModifier_PluralForm::Create(ETextPluralType::Ordinal, InArgsString); });
-	TextArgumentModifiers.Add(FTextFormatString::MakeReference(TEXT("gender")),  [](const FTextFormatString& InArgsString) { return FTextFormatArgumentModifier_GenderForm::Create(InArgsString); });
-	TextArgumentModifiers.Add(FTextFormatString::MakeReference(TEXT("hpp")),     [](const FTextFormatString& InArgsString) { return FTextFormatArgumentModifier_HangulPostPositions::Create(InArgsString); });
+	TextArgumentModifiers.Add(FTextFormatString::MakeReference(TEXT("plural")),  [](const FTextFormatString& InArgsString, const FTextFormatPatternDefinitionConstRef& InPatternDef) { return FTextFormatArgumentModifier_PluralForm::Create(ETextPluralType::Cardinal, InArgsString, InPatternDef); });
+	TextArgumentModifiers.Add(FTextFormatString::MakeReference(TEXT("ordinal")), [](const FTextFormatString& InArgsString, const FTextFormatPatternDefinitionConstRef& InPatternDef) { return FTextFormatArgumentModifier_PluralForm::Create(ETextPluralType::Ordinal, InArgsString, InPatternDef); });
+	TextArgumentModifiers.Add(FTextFormatString::MakeReference(TEXT("gender")),  [](const FTextFormatString& InArgsString, const FTextFormatPatternDefinitionConstRef& InPatternDef) { return FTextFormatArgumentModifier_GenderForm::Create(InArgsString, InPatternDef); });
+	TextArgumentModifiers.Add(FTextFormatString::MakeReference(TEXT("hpp")),     [](const FTextFormatString& InArgsString, const FTextFormatPatternDefinitionConstRef& InPatternDef) { return FTextFormatArgumentModifier_HangulPostPositions::Create(InArgsString); });
 }
 
 FTextFormatter& FTextFormatter::Get()
@@ -895,11 +913,6 @@ FTextFormatter::FCompileTextArgumentModifierFuncPtr FTextFormatter::FindTextArgu
 	return TextArgumentModifiers.FindRef(InKeyword);
 }
 
-const FTokenDefinitions& FTextFormatter::GetTextFormatDefinitions() const
-{
-	return TextFormatDefinitions;
-}
-
 FText FTextFormatter::Format(FTextFormat&& InFmt, FFormatNamedArguments&& InArguments, const bool bInRebuildText, const bool bInRebuildAsSource)
 {
 	FString ResultString = FormatStr(InFmt, InArguments, bInRebuildText, bInRebuildAsSource);
@@ -1042,7 +1055,7 @@ FString FTextFormatter::Format(const FTextFormat& InFmt, const FPrivateTextForma
 			FmtText.Rebuild();
 		}
 
-		FmtPattern = FTextFormat(FmtText.BuildSourceString());
+		FmtPattern = FTextFormat(FmtText.BuildSourceString(), FmtPattern.GetPatternDefinition());
 	}
 
 	return FmtPattern.TextFormatData->Format(InFormatArgs);
diff --git a/Engine/Source/Runtime/Core/Public/Internationalization/Text.h b/Engine/Source/Runtime/Core/Public/Internationalization/Text.h
--- a/Engine/Source/Runtime/Core/Public/Internationalization/Text.h
+++ b/Engine/Source/Runtime/Core/Public/Internationalization/Text.h
@@ -21,14 +21,13 @@
 #include "Templates/IsConstructible.h"
 #include "Templates/AndOrNot.h"
 
-class FFormatArgumentValue;
-class FTextFormatData;
-
 class FText;
 class FTextHistory;
 class FTextFormatData;
+class FFormatArgumentValue;
 class FHistoricTextFormatData;
 class FHistoricTextNumericData;
+class FTextFormatPatternDefinition;
 class ITextGenerator;
 
 //DECLARE_CYCLE_STAT_EXTERN( TEXT("Format Text"), STAT_TextFormat, STATGROUP_Text, );
@@ -121,6 +120,11 @@ namespace EFormatArgumentType
 typedef TMap<FString, FFormatArgumentValue, FDefaultSetAllocator, FLocKeyMapFuncs<FFormatArgumentValue>> FFormatNamedArguments;
 typedef TArray<FFormatArgumentValue> FFormatOrderedArguments;
 
+typedef TSharedRef<FTextFormatPatternDefinition, ESPMode::ThreadSafe> FTextFormatPatternDefinitionRef;
+typedef TSharedPtr<FTextFormatPatternDefinition, ESPMode::ThreadSafe> FTextFormatPatternDefinitionPtr;
+typedef TSharedRef<const FTextFormatPatternDefinition, ESPMode::ThreadSafe> FTextFormatPatternDefinitionConstRef;
+typedef TSharedPtr<const FTextFormatPatternDefinition, ESPMode::ThreadSafe> FTextFormatPatternDefinitionConstPtr;
+
 /** Redeclared in KismetTextLibrary for meta-data extraction purposes, be sure to update there as well */
 enum ERoundingMode
 {
@@ -247,6 +251,12 @@ public:
 	 */
 	FTextFormat(const FText& InText);
 
+	/**
+	 * Construct an instance from an FText and custom format pattern definition.
+	 * The text will be immediately compiled.
+	 */
+	FTextFormat(const FText& InText, FTextFormatPatternDefinitionConstRef InCustomPatternDef);
+
 	/**
 	 * Construct an instance from an FString.
 	 * The string will be immediately compiled.
@@ -254,6 +264,13 @@ public:
 	static FTextFormat FromString(const FString& InString);
 	static FTextFormat FromString(FString&& InString);
 
+	/**
+	 * Construct an instance from an FString and custom format pattern definition.
+	 * The string will be immediately compiled.
+	 */
+	static FTextFormat FromString(const FString& InString, FTextFormatPatternDefinitionConstRef InCustomPatternDef);
+	static FTextFormat FromString(FString&& InString, FTextFormatPatternDefinitionConstRef InCustomPatternDef);
+
 	/**
 	 * Test to see whether this instance contains valid compiled data.
 	 */
@@ -276,6 +293,11 @@ public:
 	 */
 	EExpressionType GetExpressionType() const;
 
+	/**
+	 * Get the format pattern definition being used.
+	 */
+	FTextFormatPatternDefinitionConstRef GetPatternDefinition() const;
+
 	/**
 	 * Validate the format pattern is valid based on the rules of the given culture (or null to use the current language).
 	 * @return true if the pattern is valid, or false if not (false may also fill in OutValidationErrors).
@@ -292,7 +314,7 @@ private:
 	 * Construct an instance from an FString.
 	 * The string will be immediately compiled.
 	 */
-	FTextFormat(FString&& InString);
+	FTextFormat(FString&& InString, FTextFormatPatternDefinitionConstRef InCustomPatternDef);
 
 	/** Cached compiled expression data */
 	TSharedRef<FTextFormatData, ESPMode::ThreadSafe> TextFormatData;
diff --git a/Engine/Source/Runtime/Core/Public/Internationalization/TextFormatter.h b/Engine/Source/Runtime/Core/Public/Internationalization/TextFormatter.h
--- a/Engine/Source/Runtime/Core/Public/Internationalization/TextFormatter.h
+++ b/Engine/Source/Runtime/Core/Public/Internationalization/TextFormatter.h
@@ -13,14 +13,66 @@
 
 struct FPrivateTextFormatArguments;
 
+/**
+ * Definition of the pattern used during a text format.
+ */
+class CORE_API FTextFormatPatternDefinition : public TSharedFromThis<FTextFormatPatternDefinition, ESPMode::ThreadSafe>
+{
+public:
+	/** Constructor */
+	FTextFormatPatternDefinition();
+
+	/** Singleton access to the default instance */
+	static FTextFormatPatternDefinitionConstRef GetDefault();
+
+	/** Get the text format definitions used when formatting text */
+	const FTokenDefinitions& GetTextFormatDefinitions() const;
+
+	/** Is the given character one that an escape token may escape? */
+	FORCEINLINE bool IsValidEscapeChar(const TCHAR InChar) const
+	{
+		return InChar == EscapeChar
+			|| InChar == ArgStartChar
+			|| InChar == ArgEndChar
+			|| InChar == ArgModChar;
+	}
+
+	/** Is the given character one that should cause a literal string token to break parsing? */
+	FORCEINLINE bool IsLiteralBreakChar(const TCHAR InChar) const
+	{
+		return InChar == EscapeChar
+			|| InChar == ArgStartChar;
+	}
+
+	/** Character representing the start of an escape token */
+	TCHAR EscapeChar = TEXT('`');
+	FTextFormatPatternDefinition& SetEscapeChar(const TCHAR InChar) { EscapeChar = InChar; return *this; }
+
+	/** Character representing the start of a format argument token */
+	TCHAR ArgStartChar = TEXT('{');
+	FTextFormatPatternDefinition& SetArgStartChar(const TCHAR InChar) { ArgStartChar = InChar; return *this; }
+
+	/** Character representing the end of a format argument token */
+	TCHAR ArgEndChar = TEXT('}');
+	FTextFormatPatternDefinition& SetArgEndChar(const TCHAR InChar) { ArgEndChar = InChar; return *this; }
+
+	/** Character representing the start of a format argument modifier token */
+	TCHAR ArgModChar = TEXT('|');
+	FTextFormatPatternDefinition& SetArgModChar(const TCHAR InChar) { ArgModChar = InChar; return *this; }
+
+private:
+	/** Token definitions for the text format lexer */
+	FTokenDefinitions TextFormatDefinitions;
+};
+
 /**
  * A text formatter is responsible for formatting text patterns using a set of named or ordered arguments.
  */
 class CORE_API FTextFormatter
 {
 public:
-	/** Callback function used to compile an argument modifier. Takes an argument modifier string and returns the compiled result. */
-	typedef TFunction<TSharedPtr<ITextFormatArgumentModifier>(const FTextFormatString&)> FCompileTextArgumentModifierFuncPtr;
+	/** Callback function used to compile an argument modifier. Takes an argument modifier string and pattern definition, then returns the compiled result. */
+	typedef TFunction<TSharedPtr<ITextFormatArgumentModifier>(const FTextFormatString&, FTextFormatPatternDefinitionConstRef)> FCompileTextArgumentModifierFuncPtr;
 
 	/** Singleton access */
 	static FTextFormatter& Get();
@@ -29,9 +81,6 @@ public:
 	void UnregisterTextArgumentModifier(const FTextFormatString& InKeyword);
 	FCompileTextArgumentModifierFuncPtr FindTextArgumentModifier(const FTextFormatString& InKeyword) const;
 
-	/** Get the text format definitions used when formatting text */
-	const FTokenDefinitions& GetTextFormatDefinitions() const;
-
 	/** Low-level versions of Format. You probably want to use FText::Format(...) rather than call these directly */
 	static FText Format(FTextFormat&& InFmt, FFormatNamedArguments&& InArguments, const bool bInRebuildText, const bool bInRebuildAsSource);
 	static FText Format(FTextFormat&& InFmt, FFormatOrderedArguments&& InArguments, const bool bInRebuildText, const bool bInRebuildAsSource);
@@ -53,9 +102,6 @@ private:
 
 	FTextFormatter();
 
-	/** Token definitions for the text format lexer */
-	FTokenDefinitions TextFormatDefinitions;
-
 	/** Functions for constructing argument modifier data */
 	TMap<FTextFormatString, FCompileTextArgumentModifierFuncPtr> TextArgumentModifiers;
 
