When reporting ensure ThreadHeartBeat will be suspended for all threads to prevent detecting hangs on threads that wait for the thread that triggered and is processing the ensure

diff --git a/Engine/Source/Runtime/Core/Private/HAL/ThreadHeartBeat.cpp b/Engine/Source/Runtime/Core/Private/HAL/ThreadHeartBeat.cpp
--- a/Engine/Source/Runtime/Core/Private/HAL/ThreadHeartBeat.cpp
+++ b/Engine/Source/Runtime/Core/Private/HAL/ThreadHeartBeat.cpp
@@ -371,13 +371,15 @@ uint32 FThreadHeartBeat::CheckHeartBeat(double& OutHangDuration)
 {
 	// Editor and debug builds run too slow to measure them correctly
 #if USE_HANG_DETECTION
-	static bool bDisabled = FParse::Param(FCommandLine::Get(), TEXT("nothreadtimeout"));
+	static bool bForceEnabled = FParse::Param(FCommandLine::Get(), TEXT("debughangdetection"));
+	static bool bDisabled = !bForceEnabled && FParse::Param(FCommandLine::Get(), TEXT("nothreadtimeout"));
 
 	bool CheckBeats = (ConfigHangDuration > 0.0 || ConfigPresentDuration > 0.0)
 		&& bReadyToCheckHeartbeat
 		&& !GIsRequestingExit
-		&& !FPlatformMisc::IsDebuggerPresent()
-		&& !bDisabled;
+		&& (bForceEnabled || !FPlatformMisc::IsDebuggerPresent())
+		&& !bDisabled
+		&& !GlobalSuspendCount.GetValue();
 
 	if (CheckBeats)
 	{
@@ -424,33 +426,50 @@ void FThreadHeartBeat::KillHeartBeat()
 #endif
 }
 
-void FThreadHeartBeat::SuspendHeartBeat()
+void FThreadHeartBeat::SuspendHeartBeat(bool bAllThreads)
 {
-#if USE_HANG_DETECTION
-	uint32 ThreadId = FPlatformTLS::GetCurrentThreadId();
+#if USE_HANG_DETECTION	
 	FScopeLock HeartBeatLock(&HeartBeatCritical);
-	FHeartBeatInfo* HeartBeatInfo = ThreadHeartBeat.Find(ThreadId);
-	if (HeartBeatInfo)
+	if (bAllThreads)
+	{
+		GlobalSuspendCount.Increment();
+	}
+	else
 	{
-		HeartBeatInfo->SuspendedCount++;
+		uint32 ThreadId = FPlatformTLS::GetCurrentThreadId();
+		FHeartBeatInfo* HeartBeatInfo = ThreadHeartBeat.Find(ThreadId);
+		if (HeartBeatInfo)
+		{
+			HeartBeatInfo->Suspend();
+		}
 	}
 
 	// Suspend the frame-present based detection at the same time.
 	PresentHeartBeat.SuspendedCount++;
 #endif
 }
-void FThreadHeartBeat::ResumeHeartBeat()
+void FThreadHeartBeat::ResumeHeartBeat(bool bAllThreads)
 {
-#if USE_HANG_DETECTION
-	uint32 ThreadId = FPlatformTLS::GetCurrentThreadId();
-	FScopeLock HeartBeatLock(&HeartBeatCritical);
-	FHeartBeatInfo* HeartBeatInfo = ThreadHeartBeat.Find(ThreadId);
-	if (HeartBeatInfo)
+#if USE_HANG_DETECTION	
+	FScopeLock HeartBeatLock(&HeartBeatCritical);	
+	const double CurrentTime = Clock.Seconds();
+	if (bAllThreads)
+	{
+		if (GlobalSuspendCount.Decrement() == 0)
+		{
+			for (TPair<uint32, FHeartBeatInfo>& HeartBeatEntry : ThreadHeartBeat)
+			{
+				HeartBeatEntry.Value.LastHeartBeatTime = CurrentTime;
+			}
+		}
+	}
+	else
 	{
-		check(HeartBeatInfo->SuspendedCount > 0);
-		if (--HeartBeatInfo->SuspendedCount == 0)
+		uint32 ThreadId = FPlatformTLS::GetCurrentThreadId();
+		FHeartBeatInfo* HeartBeatInfo = ThreadHeartBeat.Find(ThreadId);
+		if (HeartBeatInfo)
 		{
-			HeartBeatInfo->LastHeartBeatTime = Clock.Seconds();
+			HeartBeatInfo->Resume(CurrentTime);
 		}
 	}
 
diff --git a/Engine/Source/Runtime/Core/Private/Windows/WindowsPlatformCrashContext.cpp b/Engine/Source/Runtime/Core/Private/Windows/WindowsPlatformCrashContext.cpp
--- a/Engine/Source/Runtime/Core/Private/Windows/WindowsPlatformCrashContext.cpp
+++ b/Engine/Source/Runtime/Core/Private/Windows/WindowsPlatformCrashContext.cpp
@@ -573,7 +573,7 @@ FORCENOINLINE void ReportEnsure( const TCHAR* ErrorMessage, int NumStackFramesTo
 	// Stop checking heartbeat for this thread (and stop the gamethread hitch detector if we're the game thread).
 	// Ensure can take a lot of time (when stackwalking), so we don't want hitches/hangs firing.
 	// These are no-ops on threads that didn't already have a heartbeat etc.
-	FSlowHeartBeatScope SuspendHeartBeat;
+	FSlowHeartBeatScope SuspendHeartBeat(true);
 	FDisableHitchDetectorScope SuspendGameThreadHitch;
 
 	bReentranceGuard = true;
diff --git a/Engine/Source/Runtime/Core/Public/HAL/ThreadHeartBeat.h b/Engine/Source/Runtime/Core/Public/HAL/ThreadHeartBeat.h
--- a/Engine/Source/Runtime/Core/Public/HAL/ThreadHeartBeat.h
+++ b/Engine/Source/Runtime/Core/Public/HAL/ThreadHeartBeat.h
@@ -57,6 +57,21 @@ class CORE_API FThreadHeartBeat : public FRunnable
 		int32 SuspendedCount;
 		/** The timeout for this thread */
 		double HangDuration;
+
+		/** Suspends this thread's heartbeat */
+		void Suspend()
+		{
+			SuspendedCount++;
+		}
+		/** Resumes this thread's heartbeat */
+		void Resume(double CurrentTime)
+		{
+			check(SuspendedCount > 0);
+			if (--SuspendedCount == 0)
+			{
+				LastHeartBeatTime = CurrentTime;
+			}
+		}
 	};
 	/** Thread to run the worker FRunnable on */
 	FRunnableThread* Thread;
@@ -82,6 +97,9 @@ class CORE_API FThreadHeartBeat : public FRunnable
 	/** Id of the last thread that hung */
 	uint32 LastHungThreadId;
 
+	/** Global suspended count */
+	FThreadSafeCounter GlobalSuspendCount;
+
 	FThreadHeartBeatClock Clock;
 
 	FThreadHeartBeat();
@@ -119,12 +137,14 @@ public:
 	void KillHeartBeat();
 	/** 
 	 * Suspend heartbeat measuring for the current thread if the thread has already had a heartbeat 
+	 * @param bAllThreads If true, suspends heartbeat for all threads, not only the current one
 	 */
-	void SuspendHeartBeat();
+	void SuspendHeartBeat(bool bAllThreads = false);
 	/** 
 	 * Resume heartbeat measuring for the current thread 
+	 * @param bAllThreads If true, resumes heartbeat for all threads, not only the current one
 	 */
-	void ResumeHeartBeat();
+	void ResumeHeartBeat(bool bAllThreads = false);
 
 	/**
 	* Returns true/false if this thread is currently performing heartbeat monitoring
@@ -147,18 +167,22 @@ public:
 /** Suspends heartbeat measuring for the current thread in the current scope */
 struct FSlowHeartBeatScope
 {
-	FORCEINLINE FSlowHeartBeatScope()
+private:
+	bool bSuspendedAllThreads;
+public:
+	FORCEINLINE FSlowHeartBeatScope(bool bAllThreads = false)
+		: bSuspendedAllThreads(bAllThreads)
 	{
 		if (FThreadHeartBeat* HB = FThreadHeartBeat::GetNoInit())
 		{
-			HB->SuspendHeartBeat();
+			HB->SuspendHeartBeat(bSuspendedAllThreads);
 		}
 	}
 	FORCEINLINE ~FSlowHeartBeatScope()
 	{
 		if (FThreadHeartBeat* HB = FThreadHeartBeat::GetNoInit())
 		{
-			HB->ResumeHeartBeat();
+			HB->ResumeHeartBeat(bSuspendedAllThreads);
 		}
 	}
 };
