Profile log timing in CSV profiler, so that log spamming can be caught by perf QA.
- More work might need to do later to catch log spamming in live, such as monitoring log count in analytics.

diff --git a/Engine/Source/Runtime/Core/Private/Logging/LogMacros.cpp b/Engine/Source/Runtime/Core/Private/Logging/LogMacros.cpp
--- a/Engine/Source/Runtime/Core/Private/Logging/LogMacros.cpp
+++ b/Engine/Source/Runtime/Core/Private/Logging/LogMacros.cpp
@@ -7,6 +7,7 @@
 #include "Misc/FeedbackContext.h"
 #include "Misc/VarargsHelper.h"
 #include "Stats/Stats.h"
+#include "ProfilingDebugging/CsvProfiler.h"
 
 void StaticFailDebug( const TCHAR* Error, const ANSICHAR* File, int32 Line, const TCHAR* Description, bool bIsEnsure, int32 NumStackFramesToIgnore );
 
@@ -15,6 +16,8 @@ static FCriticalSection					MsgLogfStaticBufferGuard;
 /** Increased from 4096 to fix crashes in the renderthread without autoreporter. */
 static TCHAR							MsgLogfStaticBuffer[8192];
 
+CSV_DEFINE_CATEGORY(FMsgLogf, true);
+
 void FMsg::LogfImpl(const ANSICHAR* File, int32 Line, const FName& Category, ELogVerbosity::Type Verbosity, const TCHAR* Fmt, ...)
 {
 #if !NO_LOGGING
@@ -68,7 +71,9 @@ void FMsg::LogfImpl(const ANSICHAR* File, int32 Line, const FName& Category, ELo
 void FMsg::Logf_InternalImpl(const ANSICHAR* File, int32 Line, const FName& Category, ELogVerbosity::Type Verbosity, const TCHAR* Fmt, ...)
 {
 #if !NO_LOGGING
-	QUICK_SCOPE_CYCLE_COUNTER(STAT_FMsg_Logf);
+	QUICK_SCOPE_CYCLE_COUNTER(STAT_FMsgLogf);
+	CSV_SCOPED_TIMING_STAT(FMsgLogf, FMsgLogfTime);
+	CSV_CUSTOM_STAT(FMsgLogf, FMsgLogfCount, 1, ECsvCustomStatOp::Accumulate);
 
 	if (Verbosity != ELogVerbosity::Fatal)
 	{
