Remove log timing in CSV and only do log counting

diff --git a/Engine/Source/Runtime/Core/Private/Logging/LogMacros.cpp b/Engine/Source/Runtime/Core/Private/Logging/LogMacros.cpp
--- a/Engine/Source/Runtime/Core/Private/Logging/LogMacros.cpp
+++ b/Engine/Source/Runtime/Core/Private/Logging/LogMacros.cpp
@@ -72,7 +72,6 @@ void FMsg::Logf_InternalImpl(const ANSICHAR* File, int32 Line, const FName& Cate
 {
 #if !NO_LOGGING
 	QUICK_SCOPE_CYCLE_COUNTER(STAT_FMsgLogf);
-	CSV_SCOPED_TIMING_STAT(FMsgLogf, FMsgLogfTime);
 	CSV_CUSTOM_STAT(FMsgLogf, FMsgLogfCount, 1, ECsvCustomStatOp::Accumulate);
 
 	if (Verbosity != ELogVerbosity::Fatal)
