#UE4 Fix for underrun mode in mallocstomp not properly aligning FAllocationData for use by Free().

diff --git a/Engine/Source/Runtime/Core/Private/HAL/MallocStomp.cpp b/Engine/Source/Runtime/Core/Private/HAL/MallocStomp.cpp
--- a/Engine/Source/Runtime/Core/Private/HAL/MallocStomp.cpp
+++ b/Engine/Source/Runtime/Core/Private/HAL/MallocStomp.cpp
@@ -101,23 +101,22 @@ void* FMallocStomp::Malloc(SIZE_T Size, uint32 Alignment)
 	static const SIZE_T AllocationDataSize = sizeof(FAllocationData);
 
 	const FAllocationData AllocData = { FullAllocationPointer, TotalAllocationSize, AlignedSize, SentinelExpectedValue };
-	FAllocationData* AllocDataPointer = nullptr;
 
 	if(bUseUnderrunMode)
 	{
 		const SIZE_T AlignedAllocationData = (Alignment > 0U) ? ((AllocationDataSize + Alignment - 1U) & -static_cast<int32>(Alignment)) : AllocationDataSize;
 		ReturnedPointer = reinterpret_cast<void*>(reinterpret_cast<uint8*>(FullAllocationPointer) + PageSize + AlignedAllocationData);
-		AllocDataPointer = reinterpret_cast<FAllocationData*>(reinterpret_cast<uint8*>(FullAllocationPointer) + PageSize);
+		void* AllocDataPointerStart = reinterpret_cast<FAllocationData*>(reinterpret_cast<uint8*>(FullAllocationPointer) + PageSize);
 
 #if PLATFORM_WINDOWS && MALLOC_STOMP_KEEP_VIRTUAL_MEMORY
 		// Commit physical pages to the used range, leaving the first page unmapped.
-		void* CommittedMemory = VirtualAlloc(AllocDataPointer, AllocFullPageSize, MEM_COMMIT, PAGE_READWRITE);
+		void* CommittedMemory = VirtualAlloc(AllocDataPointerStart, AllocFullPageSize, MEM_COMMIT, PAGE_READWRITE);
 		if (!CommittedMemory)
 		{
 			// Failed to allocate and commit physical memory pages. Report OOM.
 			FPlatformMemory::OnOutOfMemory(Size, Alignment);
 		}
-		check(CommittedMemory == AllocDataPointer);
+		check(CommittedMemory == AllocDataPointerStart);
 #else
 		// Page protect the first page, this will cause the exception in case the is an underrun.
 		FPlatformMemory::PageProtect(FullAllocationPointer, PageSize, false, false);
@@ -126,7 +125,6 @@ void* FMallocStomp::Malloc(SIZE_T Size, uint32 Alignment)
 	else
 	{
 		ReturnedPointer = reinterpret_cast<void*>(reinterpret_cast<uint8*>(FullAllocationPointer) + AllocFullPageSize - AlignedSize);
-		AllocDataPointer = reinterpret_cast<FAllocationData*>(reinterpret_cast<uint8*>(ReturnedPointer) - AllocationDataSize);
 
 #if PLATFORM_WINDOWS && MALLOC_STOMP_KEEP_VIRTUAL_MEMORY
 		// Commit physical pages to the used range, leaving the last page unmapped.
@@ -143,6 +141,7 @@ void* FMallocStomp::Malloc(SIZE_T Size, uint32 Alignment)
 #endif
 	} //-V773
 
+	FAllocationData* AllocDataPointer = reinterpret_cast<FAllocationData*>(reinterpret_cast<uint8*>(ReturnedPointer) - AllocationDataSize);
 	*AllocDataPointer = AllocData;
 
 	return ReturnedPointer;
