Fixed typo in text serialization. Fixes CIS automation test errors

diff --git a/Engine/Source/Runtime/Core/Private/Internationalization/Text.cpp b/Engine/Source/Runtime/Core/Private/Internationalization/Text.cpp
--- a/Engine/Source/Runtime/Core/Private/Internationalization/Text.cpp
+++ b/Engine/Source/Runtime/Core/Private/Internationalization/Text.cpp
@@ -1093,7 +1093,7 @@ void operator<<(FStructuredArchive::FSlot Slot, FFormatArgumentValue& Value)
 	FStructuredArchive::FRecord Record = Slot.EnterRecord();
 	int8 TypeAsInt8 = Value.GetType();
 	Record << NAMED_ITEM("Type", TypeAsInt8);
-	Value  = (EFormatArgumentType::Type)TypeAsInt8;
+	Value.Type = (EFormatArgumentType::Type)TypeAsInt8;
 
 	switch(Value.Type)
 	{
