TBA: FArchiveFromStructuredArchive will now access it's host slot lazily, i.e. only when a value is actually written to the archive.

diff --git a/Engine/Source/Runtime/Core/Private/Serialization/ArchiveFromStructuredArchive.cpp b/Engine/Source/Runtime/Core/Private/Serialization/ArchiveFromStructuredArchive.cpp
--- a/Engine/Source/Runtime/Core/Private/Serialization/ArchiveFromStructuredArchive.cpp
+++ b/Engine/Source/Runtime/Core/Private/Serialization/ArchiveFromStructuredArchive.cpp
@@ -6,26 +6,14 @@
 FArchiveFromStructuredArchive::FArchiveFromStructuredArchive(FStructuredArchive::FSlot Slot)
 	: FArchiveProxy(Slot.GetUnderlyingArchive())
 	, bPendingSerialize(true)
+	, bWasOpened(false)
 	, Pos(0)
+	, RootSlot(Slot)
 {
 	// For some reason, the FArchive copy constructor will copy all the trivial members of the source archive, but then specifically set ArIsFilterEditorOnly to false, with a comment saying
 	// they don't know why it's doing this... make sure we inherit this flag here!
 	ArIsFilterEditorOnly = InnerArchive.ArIsFilterEditorOnly;
 	ArIsTextFormat = false;
-	
-	if (InnerArchive.IsTextFormat())
-	{
-		Root = Slot.EnterRecord();
-
-		if (IsLoading())
-		{
-			Commit();
-		}
-	}
-	else
-	{
-		Slot.EnterStream();
-	}
 }
 
 FArchiveFromStructuredArchive::~FArchiveFromStructuredArchive()
@@ -90,6 +78,8 @@ bool FArchiveFromStructuredArchive::AtEnd()
 
 FArchive& FArchiveFromStructuredArchive::operator<<(class FName& Value)
 {
+	OpenArchive();
+
 	if (InnerArchive.IsTextFormat())
 	{
 		if (IsLoading())
@@ -118,6 +108,8 @@ FArchive& FArchiveFromStructuredArchive::operator<<(class FName& Value)
 
 FArchive& FArchiveFromStructuredArchive::operator<<(class UObject*& Value)
 {
+	OpenArchive();
+
 	if(InnerArchive.IsTextFormat())
 	{
 		if (IsLoading())
@@ -187,6 +179,8 @@ FArchive& FArchiveFromStructuredArchive::operator<<(class UObject*& Value)
 
 FArchive& FArchiveFromStructuredArchive::operator<<(class FText& Value)
 {
+	OpenArchive();
+
 	if (InnerArchive.IsTextFormat())
 	{
 		FText::SerializeText(*this, Value);
@@ -200,6 +194,8 @@ FArchive& FArchiveFromStructuredArchive::operator<<(class FText& Value)
 
 void FArchiveFromStructuredArchive::Serialize(void* V, int64 Length)
 {
+	OpenArchive();
+
 	if (InnerArchive.IsTextFormat())
 	{
 		if (IsLoading())
@@ -229,7 +225,7 @@ void FArchiveFromStructuredArchive::Serialize(void* V, int64 Length)
 
 void FArchiveFromStructuredArchive::Commit()
 {
-	if (InnerArchive.IsTextFormat())
+	if (bWasOpened && InnerArchive.IsTextFormat())
 	{
 		SerializeInternal(Root.GetValue());
 	}
@@ -237,6 +233,8 @@ void FArchiveFromStructuredArchive::Commit()
 
 void FArchiveFromStructuredArchive::SerializeInternal(FStructuredArchive::FRecord Record)
 {
+	check(bWasOpened);
+
 	if (bPendingSerialize)
 	{
 		FStructuredArchive::FSlot DataSlot = Record.EnterField(FIELD_NAME_TEXT("Data"));
@@ -275,3 +273,25 @@ void FArchiveFromStructuredArchive::SerializeInternal(FStructuredArchive::FRecor
 		bPendingSerialize = false;
 	}
 }
+
+void FArchiveFromStructuredArchive::OpenArchive()
+{
+	if (!bWasOpened)
+	{
+		if (InnerArchive.IsTextFormat())
+		{
+			Root = RootSlot.EnterRecord();
+
+			if (IsLoading())
+			{
+				Commit();
+			}
+		}
+		else
+		{
+			RootSlot.EnterStream();
+		}
+
+		bWasOpened = true;
+	}
+}
\ No newline at end of file
diff --git a/Engine/Source/Runtime/Core/Public/Serialization/ArchiveFromStructuredArchive.h b/Engine/Source/Runtime/Core/Public/Serialization/ArchiveFromStructuredArchive.h
--- a/Engine/Source/Runtime/Core/Public/Serialization/ArchiveFromStructuredArchive.h
+++ b/Engine/Source/Runtime/Core/Public/Serialization/ArchiveFromStructuredArchive.h
@@ -40,11 +40,13 @@ protected:
 	TOptional<FStructuredArchive::FRecord> Root;
 	void Commit();
 	virtual void SerializeInternal(FStructuredArchive::FRecord Record);
+	void OpenArchive();
 
 private:
 	static const int32 MaxBufferSize = 128;
 
 	bool bPendingSerialize;
+	bool bWasOpened;
 
 	TArray<uint8> Buffer;
 	int32 Pos;
@@ -55,4 +57,6 @@ private:
 	TArray<UObject*> Objects;
 	TBitArray<> ObjectsValid;
 	TMap<UObject*, int32> ObjectToIndex;
+
+	FStructuredArchive::FSlot RootSlot;
 };
\ No newline at end of file
diff --git a/Engine/Source/Runtime/CoreUObject/Private/Serialization/ArchiveUObjectFromStructuredArchive.cpp b/Engine/Source/Runtime/CoreUObject/Private/Serialization/ArchiveUObjectFromStructuredArchive.cpp
--- a/Engine/Source/Runtime/CoreUObject/Private/Serialization/ArchiveUObjectFromStructuredArchive.cpp
+++ b/Engine/Source/Runtime/CoreUObject/Private/Serialization/ArchiveUObjectFromStructuredArchive.cpp
@@ -6,10 +6,6 @@ FArchiveUObjectFromStructuredArchive::FArchiveUObjectFromStructuredArchive(FStru
 	: FArchiveFromStructuredArchive(Slot)
 	, bPendingSerialize(true)
 {
-	if (InnerArchive.IsTextFormat() && IsLoading())
-	{
-		Commit();
-	}
 }
 
 FArchiveUObjectFromStructuredArchive::~FArchiveUObjectFromStructuredArchive()
@@ -19,6 +15,8 @@ FArchiveUObjectFromStructuredArchive::~FArchiveUObjectFromStructuredArchive()
 
 FArchive& FArchiveUObjectFromStructuredArchive::operator<<(FLazyObjectPtr& Value)
 {
+	OpenArchive();
+
 	if (InnerArchive.IsTextFormat())
 	{
 		if (IsLoading())
@@ -47,6 +45,8 @@ FArchive& FArchiveUObjectFromStructuredArchive::operator<<(FLazyObjectPtr& Value
 
 FArchive& FArchiveUObjectFromStructuredArchive::operator<<(FSoftObjectPtr& Value)
 {
+	OpenArchive();
+
 	if (InnerArchive.IsTextFormat())
 	{
 		if (IsLoading())
@@ -75,6 +75,8 @@ FArchive& FArchiveUObjectFromStructuredArchive::operator<<(FSoftObjectPtr& Value
 
 FArchive& FArchiveUObjectFromStructuredArchive::operator<<(FSoftObjectPath& Value)
 {
+	OpenArchive();
+
 	if (InnerArchive.IsTextFormat())
 	{
 		if (IsLoading())
@@ -103,6 +105,8 @@ FArchive& FArchiveUObjectFromStructuredArchive::operator<<(FSoftObjectPath& Valu
 
 FArchive& FArchiveUObjectFromStructuredArchive::operator<<(FWeakObjectPtr& Value)
 {
+	OpenArchive();
+
 	if (InnerArchive.IsTextFormat())
 	{
 		if (IsLoading())
