Undo //UE4/Dev-Core/Engine/Source/Runtime/... changelist 4060007

Needs some other changes that I haven't checked in yet...

diff --git a/Engine/Source/Runtime/Core/Private/Serialization/ArchiveFromStructuredArchive.cpp b/Engine/Source/Runtime/Core/Private/Serialization/ArchiveFromStructuredArchive.cpp
--- a/Engine/Source/Runtime/Core/Private/Serialization/ArchiveFromStructuredArchive.cpp
+++ b/Engine/Source/Runtime/Core/Private/Serialization/ArchiveFromStructuredArchive.cpp
@@ -6,14 +6,26 @@
 FArchiveFromStructuredArchive::FArchiveFromStructuredArchive(FStructuredArchive::FSlot Slot)
 	: FArchiveProxy(Slot.GetUnderlyingArchive())
 	, bPendingSerialize(true)
-	, bWasOpened(false)
 	, Pos(0)
-	, RootSlot(Slot)
 {
 	// For some reason, the FArchive copy constructor will copy all the trivial members of the source archive, but then specifically set ArIsFilterEditorOnly to false, with a comment saying
 	// they don't know why it's doing this... make sure we inherit this flag here!
 	ArIsFilterEditorOnly = InnerArchive.ArIsFilterEditorOnly;
 	ArIsTextFormat = false;
+	
+	if (InnerArchive.IsTextFormat())
+	{
+		Root = Slot.EnterRecord();
+
+		if (IsLoading())
+		{
+			Commit();
+		}
+	}
+	else
+	{
+		Slot.EnterStream();
+	}
 }
 
 FArchiveFromStructuredArchive::~FArchiveFromStructuredArchive()
@@ -78,8 +90,6 @@ bool FArchiveFromStructuredArchive::AtEnd()
 
 FArchive& FArchiveFromStructuredArchive::operator<<(class FName& Value)
 {
-	OpenArchive();
-
 	if (InnerArchive.IsTextFormat())
 	{
 		if (IsLoading())
@@ -108,8 +118,6 @@ FArchive& FArchiveFromStructuredArchive::operator<<(class FName& Value)
 
 FArchive& FArchiveFromStructuredArchive::operator<<(class UObject*& Value)
 {
-	OpenArchive();
-
 	if(InnerArchive.IsTextFormat())
 	{
 		if (IsLoading())
@@ -179,8 +187,6 @@ FArchive& FArchiveFromStructuredArchive::operator<<(class UObject*& Value)
 
 FArchive& FArchiveFromStructuredArchive::operator<<(class FText& Value)
 {
-	OpenArchive();
-
 	if (InnerArchive.IsTextFormat())
 	{
 		FText::SerializeText(*this, Value);
@@ -194,8 +200,6 @@ FArchive& FArchiveFromStructuredArchive::operator<<(class FText& Value)
 
 void FArchiveFromStructuredArchive::Serialize(void* V, int64 Length)
 {
-	OpenArchive();
-
 	if (InnerArchive.IsTextFormat())
 	{
 		if (IsLoading())
@@ -225,7 +229,7 @@ void FArchiveFromStructuredArchive::Serialize(void* V, int64 Length)
 
 void FArchiveFromStructuredArchive::Commit()
 {
-	if (bWasOpened && InnerArchive.IsTextFormat())
+	if (InnerArchive.IsTextFormat())
 	{
 		SerializeInternal(Root.GetValue());
 	}
@@ -233,8 +237,6 @@ void FArchiveFromStructuredArchive::Commit()
 
 void FArchiveFromStructuredArchive::SerializeInternal(FStructuredArchive::FRecord Record)
 {
-	check(bWasOpened);
-
 	if (bPendingSerialize)
 	{
 		FStructuredArchive::FSlot DataSlot = Record.EnterField(FIELD_NAME_TEXT("Data"));
@@ -273,25 +275,3 @@ void FArchiveFromStructuredArchive::SerializeInternal(FStructuredArchive::FRecor
 		bPendingSerialize = false;
 	}
 }
-
-void FArchiveFromStructuredArchive::OpenArchive()
-{
-	if (!bWasOpened)
-	{
-		if (InnerArchive.IsTextFormat())
-		{
-			Root = RootSlot.EnterRecord();
-
-			if (IsLoading())
-			{
-				Commit();
-			}
-		}
-		else
-		{
-			RootSlot.EnterStream();
-		}
-
-		bWasOpened = true;
-	}
-}
\ No newline at end of file
diff --git a/Engine/Source/Runtime/Core/Public/Serialization/ArchiveFromStructuredArchive.h b/Engine/Source/Runtime/Core/Public/Serialization/ArchiveFromStructuredArchive.h
--- a/Engine/Source/Runtime/Core/Public/Serialization/ArchiveFromStructuredArchive.h
+++ b/Engine/Source/Runtime/Core/Public/Serialization/ArchiveFromStructuredArchive.h
@@ -40,13 +40,11 @@ protected:
 	TOptional<FStructuredArchive::FRecord> Root;
 	void Commit();
 	virtual void SerializeInternal(FStructuredArchive::FRecord Record);
-	void OpenArchive();
 
 private:
 	static const int32 MaxBufferSize = 128;
 
 	bool bPendingSerialize;
-	bool bWasOpened;
 
 	TArray<uint8> Buffer;
 	int32 Pos;
@@ -57,6 +55,4 @@ private:
 	TArray<UObject*> Objects;
 	TBitArray<> ObjectsValid;
 	TMap<UObject*, int32> ObjectToIndex;
-
-	FStructuredArchive::FSlot RootSlot;
 };
\ No newline at end of file
diff --git a/Engine/Source/Runtime/CoreUObject/Private/Serialization/ArchiveUObjectFromStructuredArchive.cpp b/Engine/Source/Runtime/CoreUObject/Private/Serialization/ArchiveUObjectFromStructuredArchive.cpp
--- a/Engine/Source/Runtime/CoreUObject/Private/Serialization/ArchiveUObjectFromStructuredArchive.cpp
+++ b/Engine/Source/Runtime/CoreUObject/Private/Serialization/ArchiveUObjectFromStructuredArchive.cpp
@@ -6,6 +6,10 @@ FArchiveUObjectFromStructuredArchive::FArchiveUObjectFromStructuredArchive(FStru
 	: FArchiveFromStructuredArchive(Slot)
 	, bPendingSerialize(true)
 {
+	if (InnerArchive.IsTextFormat() && IsLoading())
+	{
+		Commit();
+	}
 }
 
 FArchiveUObjectFromStructuredArchive::~FArchiveUObjectFromStructuredArchive()
@@ -15,8 +19,6 @@ FArchiveUObjectFromStructuredArchive::~FArchiveUObjectFromStructuredArchive()
 
 FArchive& FArchiveUObjectFromStructuredArchive::operator<<(FLazyObjectPtr& Value)
 {
-	OpenArchive();
-
 	if (InnerArchive.IsTextFormat())
 	{
 		if (IsLoading())
@@ -45,8 +47,6 @@ FArchive& FArchiveUObjectFromStructuredArchive::operator<<(FLazyObjectPtr& Value
 
 FArchive& FArchiveUObjectFromStructuredArchive::operator<<(FSoftObjectPtr& Value)
 {
-	OpenArchive();
-
 	if (InnerArchive.IsTextFormat())
 	{
 		if (IsLoading())
@@ -75,8 +75,6 @@ FArchive& FArchiveUObjectFromStructuredArchive::operator<<(FSoftObjectPtr& Value
 
 FArchive& FArchiveUObjectFromStructuredArchive::operator<<(FSoftObjectPath& Value)
 {
-	OpenArchive();
-
 	if (InnerArchive.IsTextFormat())
 	{
 		if (IsLoading())
@@ -105,8 +103,6 @@ FArchive& FArchiveUObjectFromStructuredArchive::operator<<(FSoftObjectPath& Valu
 
 FArchive& FArchiveUObjectFromStructuredArchive::operator<<(FWeakObjectPtr& Value)
 {
-	OpenArchive();
-
 	if (InnerArchive.IsTextFormat())
 	{
 		if (IsLoading())
