FString can now be constructed from any char pointer type and length.

diff --git a/Engine/Source/Runtime/Core/Public/Containers/UnrealString.h b/Engine/Source/Runtime/Core/Public/Containers/UnrealString.h
--- a/Engine/Source/Runtime/Core/Public/Containers/UnrealString.h
+++ b/Engine/Source/Runtime/Core/Public/Containers/UnrealString.h
@@ -123,13 +123,22 @@ public:
 	 * @param InCount how many characters to copy
 	 * @param InSrc String to copy from
 	 */
-	FORCEINLINE explicit FString( int32 InCount, const TCHAR* InSrc )
+	template <
+		typename CharType,
+		typename = typename TEnableIf<TIsCharType<CharType>::Value>::Type // This TEnableIf is to ensure we don't instantiate this constructor for non-char types, like id* in Obj-C
+	>
+	FORCEINLINE explicit FString(int32 InCount, const CharType* InSrc)
 	{
-		Data.AddUninitialized(InCount ? InCount + 1 : 0);
-
-		if( Data.Num() > 0 )
+		if (InSrc && *InSrc)
 		{
-			FCString::Strncpy(Data.GetData(), InSrc, InCount + 1);
+			int32 DestLen = FPlatformString::ConvertedLength<TCHAR>(InSrc, InCount);
+			if (DestLen > 0)
+			{
+				Data.AddUninitialized(DestLen + 1);
+
+				FPlatformString::Convert(Data.GetData(), DestLen, InSrc, InCount);
+				*(Data.GetData() + Data.Num() - 1) = TEXT('\0');
+			}
 		}
 	}
 
