TBitArray::Reserve() added.
TBitArray::Add() overloaded to allow adding multiple bits.
TSparseArray::Reserve() optimized to call the overloaded Add().

diff --git a/Engine/Source/Runtime/Core/Public/Containers/BitArray.h b/Engine/Source/Runtime/Core/Public/Containers/BitArray.h
--- a/Engine/Source/Runtime/Core/Public/Containers/BitArray.h
+++ b/Engine/Source/Runtime/Core/Public/Containers/BitArray.h
@@ -340,24 +340,32 @@ public:
 	int32 Add(const bool Value)
 	{
 		const int32 Index = NumBits;
-		const bool bReallocate = (NumBits + 1) > MaxBits;
 
-		NumBits++;
+		Reserve(Index + 1);
+		++NumBits;
+		(*this)[Index] = Value;
+
+		return Index;
+	}
+
+	/**
+	 * Adds multiple bits to the array with the given value.
+	 * @return The index of the added bit.
+	 */
+	int32 Add(const bool Value, int32 NumToAdd)
+	{
+		const int32 Index = NumBits;
 
-		if(bReallocate)
+		if (NumToAdd > 0)
 		{
-			// Allocate memory for the new bits.
-			const uint32 MaxDWORDs = AllocatorInstance.CalculateSlackGrow(
-				FMath::DivideAndRoundUp(NumBits, NumBitsPerDWORD),
-				FMath::DivideAndRoundUp(MaxBits, NumBitsPerDWORD),
-				sizeof(uint32)
-				);
-			MaxBits = MaxDWORDs * NumBitsPerDWORD;
-			Realloc(NumBits - 1);
+			Reserve(Index + NumToAdd);
+			NumBits += NumToAdd;
+			for (int32 It = Index, End = It + NumToAdd; It != End; ++It)
+			{
+				(*this)[It] = Value;
+			}
 		}
 
-		(*this)[Index] = Value;
-
 		return Index;
 	}
 
@@ -378,6 +386,25 @@ public:
 		}
 	}
 
+	/**
+	 * Reserves memory such that the array can contain at least Number bits.
+	 *
+	 * @Number  The number of bits to reserve space for.
+	 */
+	void Reserve(int32 Number)
+	{
+		if (Number > MaxBits)
+		{
+			const uint32 MaxDWORDs = AllocatorInstance.CalculateSlackGrow(
+				FMath::DivideAndRoundUp(Number,  NumBitsPerDWORD),
+				FMath::DivideAndRoundUp(MaxBits, NumBitsPerDWORD),
+				sizeof(uint32)
+				);
+			MaxBits = MaxDWORDs * NumBitsPerDWORD;
+			Realloc(NumBits);
+		}
+	}
+
 	/**
 	 * Removes all bits from the array retaining any space already allocated.
 	 */
diff --git a/Engine/Source/Runtime/Core/Public/Containers/SparseArray.h b/Engine/Source/Runtime/Core/Public/Containers/SparseArray.h
--- a/Engine/Source/Runtime/Core/Public/Containers/SparseArray.h
+++ b/Engine/Source/Runtime/Core/Public/Containers/SparseArray.h
@@ -295,11 +295,7 @@ public:
 				FirstFreeIndex = FreeIndex;
 				++NumFreeIndices;
 			}
-			//@fixme - this will have to do until TBitArray has a Reserve method....
-			for ( int32 i = 0; i < ElementsToAdd; i++ )
-			{
-				AllocationFlags.Add(false);
-			}
+			AllocationFlags.Add(false, ExpectedNumElements);
 		}
 	}
 
