Fix for TSparseArray::Reserve() size.

diff --git a/Engine/Source/Runtime/Core/Public/Containers/SparseArray.h b/Engine/Source/Runtime/Core/Public/Containers/SparseArray.h
--- a/Engine/Source/Runtime/Core/Public/Containers/SparseArray.h
+++ b/Engine/Source/Runtime/Core/Public/Containers/SparseArray.h
@@ -295,7 +295,7 @@ public:
 				FirstFreeIndex = FreeIndex;
 				++NumFreeIndices;
 			}
-			AllocationFlags.Add(false, ExpectedNumElements);
+			AllocationFlags.Add(false, ElementsToAdd);
 		}
 	}
 
