Vast simplification of TFunction, making it smaller in footprint, easier to follow and extend, and more correct.
TUniqueFunction added, which is a move-only TFunction which can hold move-only functors.
Fix for UWidgetBlueprint::ForEachSourceWidget() which should never have compiled but did.
FFunctionGraphTask and TFuture<> updated to use TUniqueFunction to make them more general.
TArray::HeapPop() made to work with move-only types.

diff --git a/Engine/Source/Editor/UMGEditor/Private/WidgetBlueprint.cpp b/Engine/Source/Editor/UMGEditor/Private/WidgetBlueprint.cpp
--- a/Engine/Source/Editor/UMGEditor/Private/WidgetBlueprint.cpp
+++ b/Engine/Source/Editor/UMGEditor/Private/WidgetBlueprint.cpp
@@ -770,7 +770,8 @@ void UWidgetBlueprint::ForEachSourceWidget(TFunctionRef<void(UWidget*)> Fn)
 
 void UWidgetBlueprint::ForEachSourceWidget(TFunctionRef<void(const UWidget*)> Fn) const
 {
-	ForEachSourceWidgetImpl(Fn);
+	// This cast is pretty terrible, but it'll probably work (!) and saves us repeaeting the Impl code.
+	ForEachSourceWidgetImpl(*(const TFunctionRef<void(UWidget*)>*)&Fn);
 }
 
 UPackage* UWidgetBlueprint::GetWidgetTemplatePackage() const
diff --git a/Engine/Source/Runtime/Core/Public/Async/Future.h b/Engine/Source/Runtime/Core/Public/Async/Future.h
--- a/Engine/Source/Runtime/Core/Public/Async/Future.h
+++ b/Engine/Source/Runtime/Core/Public/Async/Future.h
@@ -35,8 +35,8 @@ public:
 	 *
 	 * @param InCompletionCallback A function that is called when the state is completed.
 	 */
-	FFutureState(TFunction<void()>&& InCompletionCallback)
-		: CompletionCallback(InCompletionCallback)
+	FFutureState(TUniqueFunction<void()>&& InCompletionCallback)
+		: CompletionCallback(MoveTemp(InCompletionCallback))
 		, CompletionEvent(FPlatformProcess::GetSynchEventFromPool(true))
 		, Complete(false)
 	{ }
@@ -95,7 +95,7 @@ protected:
 private:
 
 	/** An optional callback function that is executed the state is completed. */
-	TFunction<void()> CompletionCallback;
+	TUniqueFunction<void()> CompletionCallback;
 
 	/** Holds an event signaling that the result is available. */
 	FEvent* CompletionEvent;
@@ -124,7 +124,7 @@ public:
 	 *
 	 * @param CompletionCallback A function that is called when the state is completed.
 	 */
-	TFutureState(TFunction<void()>&& CompletionCallback)
+	TFutureState(TUniqueFunction<void()>&& CompletionCallback)
 		: FFutureState(MoveTemp(CompletionCallback))
 	{ }
 
@@ -745,7 +745,7 @@ public:
 
 	/** Default constructor. */
 	TPromiseBase()
-		: State(MakeShareable(new TFutureState<InternalResultType>))
+		: State(MakeShared<TFutureState<InternalResultType>, ESPMode::ThreadSafe>())
 	{ }
 
 	/**
@@ -764,8 +764,8 @@ public:
 	 *
 	 * @param CompletionCallback A function that is called when the future state is completed.
 	 */
-	TPromiseBase(TFunction<void()>&& CompletionCallback)
-		: State(MakeShareable(new TFutureState<InternalResultType>(MoveTemp(CompletionCallback))))
+	TPromiseBase(TUniqueFunction<void()>&& CompletionCallback)
+		: State(MakeShared<TFutureState<InternalResultType>, ESPMode::ThreadSafe>(MoveTemp(CompletionCallback)))
 	{ }
 
 public:
@@ -844,7 +844,7 @@ public:
 	 *
 	 * @param CompletionCallback A function that is called when the future state is completed.
 	 */
-	TPromise(TFunction<void()>&& CompletionCallback)
+	TPromise(TUniqueFunction<void()>&& CompletionCallback)
 		: BaseType(MoveTemp(CompletionCallback))
 		, FutureRetrieved(false)
 	{ }
@@ -944,7 +944,7 @@ public:
 	 *
 	 * @param CompletionCallback A function that is called when the future state is completed.
 	 */
-	TPromise(TFunction<void()>&& CompletionCallback)
+	TPromise(TUniqueFunction<void()>&& CompletionCallback)
 		: BaseType(MoveTemp(CompletionCallback))
 		, FutureRetrieved(false)
 	{ }
@@ -1031,7 +1031,7 @@ public:
 	 *
 	 * @param CompletionCallback A function that is called when the future state is completed.
 	 */
-	TPromise(TFunction<void()>&& CompletionCallback)
+	TPromise(TUniqueFunction<void()>&& CompletionCallback)
 		: BaseType(MoveTemp(CompletionCallback))
 		, FutureRetrieved(false)
 	{ }
diff --git a/Engine/Source/Runtime/Core/Public/Async/TaskGraphInterfaces.h b/Engine/Source/Runtime/Core/Public/Async/TaskGraphInterfaces.h
--- a/Engine/Source/Runtime/Core/Public/Async/TaskGraphInterfaces.h
+++ b/Engine/Source/Runtime/Core/Public/Async/TaskGraphInterfaces.h
@@ -1268,7 +1268,7 @@ class FFunctionGraphTask : public FCustomStatIDGraphTaskBase
 {
 public:
     /** Function to run **/
-    TFunction<void()> Function;
+    TUniqueFunction<void()> Function;
     /** Thread to run the function on **/
     const ENamedThreads::Type			DesiredThread;
     
@@ -1289,7 +1289,7 @@ public:
      *	@param StatId The stat id for this task.
      * @param InDesiredThread - Thread to run on
      **/
-    FFunctionGraphTask(TFunction<void()>&& InFunction, const TStatId StatId, ENamedThreads::Type InDesiredThread)
+    FFunctionGraphTask(TUniqueFunction<void()>&& InFunction, const TStatId StatId, ENamedThreads::Type InDesiredThread)
     : FCustomStatIDGraphTaskBase(StatId)
     , Function(MoveTemp(InFunction))
     , DesiredThread(InDesiredThread)
@@ -1304,7 +1304,7 @@ public:
      * @param InDesiredThread - Thread to run on
      * @return completion handle for the new task
      **/
-    static FGraphEventRef CreateAndDispatchWhenReady(TFunction<void()> InFunction, const TStatId InStatId, const FGraphEventArray* InPrerequisites = NULL, ENamedThreads::Type InDesiredThread = ENamedThreads::AnyThread)
+    static FGraphEventRef CreateAndDispatchWhenReady(TUniqueFunction<void()> InFunction, const TStatId InStatId, const FGraphEventArray* InPrerequisites = NULL, ENamedThreads::Type InDesiredThread = ENamedThreads::AnyThread)
     {
         return TGraphTask<FFunctionGraphTask>::CreateTask(InPrerequisites).ConstructAndDispatchWhenReady(MoveTemp(InFunction), InStatId, InDesiredThread);
     }
@@ -1316,7 +1316,7 @@ public:
      * @param InDesiredThread - Thread to run on
      * @return completion handle for the new task
      **/
-    static FGraphEventRef CreateAndDispatchWhenReady(TFunction<void()> InFunction, const TStatId&& InStatId, const FGraphEventRef& InPrerequisite, ENamedThreads::Type InDesiredThread = ENamedThreads::AnyThread)
+    static FGraphEventRef CreateAndDispatchWhenReady(TUniqueFunction<void()> InFunction, const TStatId&& InStatId, const FGraphEventRef& InPrerequisite, ENamedThreads::Type InDesiredThread = ENamedThreads::AnyThread)
     {
         FGraphEventArray Prerequisites;
         check(InPrerequisite.GetReference());
diff --git a/Engine/Source/Runtime/Core/Public/Containers/Array.h b/Engine/Source/Runtime/Core/Public/Containers/Array.h
--- a/Engine/Source/Runtime/Core/Public/Containers/Array.h
+++ b/Engine/Source/Runtime/Core/Public/Containers/Array.h
@@ -2650,7 +2650,7 @@ public:
 	template <class PREDICATE_CLASS>
 	void HeapPop(ElementType& OutItem, const PREDICATE_CLASS& Predicate, bool bAllowShrinking = true)
 	{
-		OutItem = (*this)[0];
+		OutItem = MoveTemp((*this)[0]);
 		RemoveAtSwap(0, 1, bAllowShrinking);
 
 		TDereferenceWrapper< ElementType, PREDICATE_CLASS> PredicateWrapper(Predicate);
diff --git a/Engine/Source/Runtime/Core/Public/Templates/Function.h b/Engine/Source/Runtime/Core/Public/Templates/Function.h
--- a/Engine/Source/Runtime/Core/Public/Templates/Function.h
+++ b/Engine/Source/Runtime/Core/Public/Templates/Function.h
@@ -6,7 +6,6 @@
 #include "Misc/AssertionMacros.h"
 #include "HAL/UnrealMemory.h"
 #include "Templates/AndOrNot.h"
-#include "Templates/AreTypesEqual.h"
 #include "Templates/UnrealTypeTraits.h"
 #include "Templates/RemoveReference.h"
 #include "Templates/Decay.h"
@@ -15,6 +14,7 @@
 #include "Templates/IsInvocable.h"
 #include "Templates/IsMemberPointer.h"
 #include "Templates/IsPointer.h"
+#include "Templates/UnrealTemplate.h"
 #include "Containers/ContainerAllocationPolicies.h"
 #include "Math/UnrealMathUtility.h"
 #include <new>
@@ -38,6 +38,14 @@ template <typename FuncType> class TFunctionRef;
 template <typename FuncType>
 class TFunction;
 
+/**
+ * TUniqueFunction<FuncType>
+ *
+ * See the class definition for intended usage.
+ */
+template <typename FuncType>
+class TUniqueFunction;
+
 /**
  * TFunctionRef<FuncType>
  *
@@ -56,6 +64,16 @@ template <typename T> struct TIsTFunction<const          T> { enum { Value = TIs
 template <typename T> struct TIsTFunction<      volatile T> { enum { Value = TIsTFunction<T>::Value }; };
 template <typename T> struct TIsTFunction<const volatile T> { enum { Value = TIsTFunction<T>::Value }; };
 
+/**
+ * Traits class which checks if T is a TFunction<> type.
+ */
+template <typename T> struct TIsTUniqueFunction                     { enum { Value = false }; };
+template <typename T> struct TIsTUniqueFunction<TUniqueFunction<T>> { enum { Value = true  }; };
+
+template <typename T> struct TIsTUniqueFunction<const          T> { enum { Value = TIsTUniqueFunction<T>::Value }; };
+template <typename T> struct TIsTUniqueFunction<      volatile T> { enum { Value = TIsTUniqueFunction<T>::Value }; };
+template <typename T> struct TIsTUniqueFunction<const volatile T> { enum { Value = TIsTUniqueFunction<T>::Value }; };
+
 /**
  * Traits class which checks if T is a TFunctionRef<> type.
  */
@@ -71,7 +89,7 @@ template <typename T> struct TIsTFunctionRef<const volatile T> { enum { Value =
  */
 namespace UE4Function_Private
 {
-	struct FFunctionStorage;
+	struct FUniqueFunctionStorage;
 
 	/**
 	 * Common interface to a callable object owned by TFunction.
@@ -81,7 +99,7 @@ namespace UE4Function_Private
 		/**
 		 * Creates a copy of the object in the allocator and returns a pointer to it.
 		 */
-		virtual IFunction_OwnedObject* CopyToEmptyStorage(FFunctionStorage& Storage) const = 0;
+		virtual IFunction_OwnedObject* CopyToEmptyStorage(FUniqueFunctionStorage& Storage) const = 0;
 
 		/**
 		 * Returns the address of the object.
@@ -91,15 +109,74 @@ namespace UE4Function_Private
 		/**
 		 * Destructor.
 		 */
-		virtual ~IFunction_OwnedObject() = 0;
+		virtual ~IFunction_OwnedObject()
+		{
+		}
 	};
 
 	/**
-	 * Destructor.
+	 * Implementation of IFunction_OwnedObject for a given copyable T.
 	 */
-	inline IFunction_OwnedObject::~IFunction_OwnedObject()
+	template <typename T>
+	struct TFunction_CopyableOwnedObject final : public IFunction_OwnedObject
 	{
-	}
+		/**
+		 * Constructor which creates its T by copying.
+		 */
+		explicit TFunction_CopyableOwnedObject(const T& InObj)
+			: Obj(InObj)
+		{
+		}
+
+		/**
+		 * Constructor which creates its T by moving.
+		 */
+		explicit TFunction_CopyableOwnedObject(T&& InObj)
+			: Obj(MoveTemp(InObj))
+		{
+		}
+
+		IFunction_OwnedObject* CopyToEmptyStorage(FUniqueFunctionStorage& Storage) const override
+		{
+			return new (Storage) TFunction_CopyableOwnedObject(Obj);
+		}
+
+		virtual void* GetAddress() override
+		{
+			return &Obj;
+		}
+
+		T Obj;
+	};
+
+	/**
+	 * Implementation of IFunction_OwnedObject for a given non-copyable T.
+	 */
+	template <typename T>
+	struct TFunction_UniqueOwnedObject final : public IFunction_OwnedObject
+	{
+		/**
+		 * Constructor which creates its T by moving.
+		 */
+		explicit TFunction_UniqueOwnedObject(T&& InObj)
+			: Obj(MoveTemp(InObj))
+		{
+		}
+
+		IFunction_OwnedObject* CopyToEmptyStorage(FUniqueFunctionStorage& Storage) const override
+		{
+			// Should never get here - copy functions are deleted for TUniqueFunction
+			check(false);
+			return nullptr;
+		}
+
+		virtual void* GetAddress() override
+		{
+			return &Obj;
+		}
+
+		T Obj;
+	};
 
 	#if !defined(_WIN32) || defined(_WIN64)
 		// Let TFunction store up to 32 bytes which are 16-byte aligned before we heap allocate
@@ -112,111 +189,140 @@ namespace UE4Function_Private
 		typedef FHeapAllocator FunctionAllocatorType;
 	#endif
 
-	struct FFunctionStorage
+	template <typename T>
+	struct TIsNullableBinding :
+		TOr<
+			TIsPointer<T>,
+			TIsMemberPointer<T>,
+			TIsTFunction<T>
+		>
 	{
-		FFunctionStorage()
-			: AllocatedSize(0)
-		{
-		}
+	};
 
-		FFunctionStorage(FFunctionStorage&& Other)
-			: AllocatedSize(0)
+	template <typename T>
+	FORCEINLINE typename TEnableIf<TIsNullableBinding<T>::Value, bool>::Type IsBound(const T& Func)
+	{
+		// Function pointers, data member pointers, member function pointers and TFunctions
+		// can all be null/unbound, so test them using their boolean state.
+		return !!Func;
+	}
+
+	template <typename T>
+	FORCEINLINE typename TEnableIf<!TIsNullableBinding<T>::Value, bool>::Type IsBound(const T& Func)
+	{
+		// We can't tell if any other generic callable can be invoked, so just assume they can be.
+		return true;
+	}
+
+	struct FFunctionStorage;
+
+	struct FUniqueFunctionStorage
+	{
+		FUniqueFunctionStorage() = default;
+
+		FUniqueFunctionStorage(FUniqueFunctionStorage&& Other)
 		{
 			Allocator.MoveToEmpty(Other.Allocator);
-			AllocatedSize       = Other.AllocatedSize;
-			Other.AllocatedSize = 0;
 		}
 
-		void Empty()
-		{
-			Allocator.ResizeAllocation(0, 0, sizeof(UE4Function_Private::AlignedInlineFunctionType));
-			AllocatedSize = 0;
-		}
+		// Allow moving from a TFunction's storage
+		FUniqueFunctionStorage(FFunctionStorage&& Other);
 
-		typedef FunctionAllocatorType::ForElementType<AlignedInlineFunctionType> AllocatorType;
+		FUniqueFunctionStorage& operator=(FUniqueFunctionStorage&& Other) = delete;
+		FUniqueFunctionStorage& operator=(const FUniqueFunctionStorage& Other) = delete;
 
-		IFunction_OwnedObject* GetBoundObject() const
+		template <typename FunctorType>
+		typename TDecay<FunctorType>::Type* Bind(FunctorType&& InFunc)
 		{
-			return AllocatedSize ? (IFunction_OwnedObject*)Allocator.GetAllocation() : nullptr;
-		}
+			if (!IsBound(InFunc))
+			{
+				return nullptr;
+			}
 
-		AllocatorType Allocator;
-		int32         AllocatedSize;
-	};
-}
+			using OwnedType = TFunction_UniqueOwnedObject<typename TDecay<FunctorType>::Type>;
 
-inline void* operator new(size_t Size, UE4Function_Private::FFunctionStorage& Storage)
-{
-	if (UE4Function_Private::IFunction_OwnedObject* Obj = Storage.GetBoundObject())
-	{
-		Obj->~IFunction_OwnedObject();
-	}
+			OwnedType* NewObj = new (*this) OwnedType(Forward<FunctorType>(InFunc));
+			return &NewObj->Obj;
+		}
 
-	int32 NewSize = FMath::DivideAndRoundUp(Size, sizeof(UE4Function_Private::AlignedInlineFunctionType));
-	if (Storage.AllocatedSize != NewSize)
-	{
-		Storage.Allocator.ResizeAllocation(0, NewSize, sizeof(UE4Function_Private::AlignedInlineFunctionType));
-		Storage.AllocatedSize = NewSize;
-	}
+		void* BindCopy(const FFunctionStorage& Other);
 
-	return Storage.Allocator.GetAllocation();
-}
+		IFunction_OwnedObject* GetBoundObject() const
+		{
+			return (IFunction_OwnedObject*)Allocator.GetAllocation();
+		}
 
-namespace UE4Function_Private
-{
-	/**
-	 * Implementation of IFunction_OwnedObject for a given T.
-	 */
-	template <typename T>
-	struct TFunction_OwnedObject : public IFunction_OwnedObject
-	{
 		/**
-		 * Constructor which creates its T by copying.
+		 * Returns a pointer to the callable object - needed by TFunctionRefBase.
 		 */
-		explicit TFunction_OwnedObject(const T& InObj)
-			: Obj(InObj)
+		void* GetPtr() const
 		{
+			return GetBoundObject()->GetAddress();
 		}
 
 		/**
-		 * Constructor which creates its T by moving.
+		 * Destroy any owned bindings - called by TFunctionRefBase only if Bind() or BindCopy() was called.
 		 */
-		explicit TFunction_OwnedObject(T&& InObj)
-			: Obj(MoveTemp(InObj))
+		void Unbind()
 		{
+			GetBoundObject()->~IFunction_OwnedObject();
 		}
 
-		IFunction_OwnedObject* CopyToEmptyStorage(FFunctionStorage& Storage) const override
+		FunctionAllocatorType::ForElementType<AlignedInlineFunctionType> Allocator;
+	};
+
+	struct FFunctionStorage : public FUniqueFunctionStorage
+	{
+		FFunctionStorage() = default;
+
+		FFunctionStorage(FFunctionStorage&& Other)
+			: FUniqueFunctionStorage(MoveTemp((FUniqueFunctionStorage&)Other))
 		{
-			return new (Storage) TFunction_OwnedObject(Obj);
 		}
 
-		virtual void* GetAddress() override
+		template <typename FunctorType>
+		typename TDecay<FunctorType>::Type* Bind(FunctorType&& InFunc)
 		{
-			return &Obj;
+			if (!IsBound(InFunc))
+			{
+				return nullptr;
+			}
+
+			using OwnedType = TFunction_CopyableOwnedObject<typename TDecay<FunctorType>::Type>;
+
+			OwnedType* NewObj = new (*this) OwnedType(Forward<FunctorType>(InFunc));
+			return &NewObj->Obj;
 		}
 
-		T Obj;
+		void* BindCopy(const FFunctionStorage& Other)
+		{
+			IFunction_OwnedObject* ThisFunc = Other.GetBoundObject()->CopyToEmptyStorage(*this);
+			return ThisFunc->GetAddress();
+		}
 	};
 
-	/**
-	 * A class which is used to instantiate the code needed to call a bound function.
-	 */
-	template <typename Functor, typename FuncType>
-	struct TFunctionRefCaller;
+	inline FUniqueFunctionStorage::FUniqueFunctionStorage(FFunctionStorage&& Other)
+	{
+		Allocator.MoveToEmpty(Other.Allocator);
+	}
 
-	/**
-	 * A class which is used to instantiate the code needed to assert when called - used for unset bindings.
-	 */
-	template <typename FuncType>
-	struct TFunctionRefAsserter;
+	inline void* FUniqueFunctionStorage::BindCopy(const FFunctionStorage& Other)
+	{
+		IFunction_OwnedObject* ThisFunc = Other.GetBoundObject()->CopyToEmptyStorage(*this);
+		return ThisFunc->GetAddress();
+	}
+}
 
-	/**
-	 * A class which defines an operator() which will invoke the TFunctionRefCaller::Call function.
-	 */
-	template <typename DerivedType, typename FuncType>
-	struct TFunctionRefBase;
+inline void* operator new(size_t Size, UE4Function_Private::FUniqueFunctionStorage& Storage)
+{
+	int32 NewSize = FMath::DivideAndRoundUp(Size, sizeof(UE4Function_Private::AlignedInlineFunctionType));
+	Storage.Allocator.ResizeAllocation(0, NewSize, sizeof(UE4Function_Private::AlignedInlineFunctionType));
+
+	return Storage.Allocator.GetAllocation();
+}
 
+namespace UE4Function_Private
+{
 	#if ENABLE_TFUNCTIONREF_VISUALIZATION
 		/**
 		 * Helper classes to help debugger visualization.
@@ -238,8 +344,11 @@ namespace UE4Function_Private
 	#endif
 
 	/**
-	 * Function which invokes the passed callable when invoked.
+	 * A class which is used to instantiate the code needed to call a bound function.
 	 */
+	template <typename Functor, typename FuncType>
+	struct TFunctionRefCaller;
+
 	template <typename Functor, typename Ret, typename... ParamTypes>
 	struct TFunctionRefCaller<Functor, Ret (ParamTypes...)>
 	{
@@ -259,95 +368,182 @@ namespace UE4Function_Private
 	};
 
 	/**
-	 * Function which asserts when invoked.
+	 * A class which defines an operator() which will own storage of a callable and invoke the TFunctionRefCaller::Call function on it.
 	 */
-	template <typename Ret, typename... ParamTypes>
-	struct TFunctionRefAsserter<Ret (ParamTypes...)>
+	template <typename StorageType, typename FuncType>
+	struct TFunctionRefBase;
+
+	template <typename StorageType, typename Ret, typename... ParamTypes>
+	struct TFunctionRefBase<StorageType, Ret (ParamTypes...)>
 	{
-		static Ret Call(void* Obj, ParamTypes&...)
-		{
-			checkf(false, TEXT("Attempting to call an unbound TFunction!"));
+		template <typename OtherStorageType, typename OtherFuncType>
+		friend struct TFunctionRefBase;
 
-			// This doesn't need to be valid, because it'll never be reached, but it does at least need to compile.
-			return Forward<Ret&&>(*(typename TRemoveReference<Ret>::Type*)Obj);
+		TFunctionRefBase()
+			: Callable(nullptr)
+		{
 		}
-	};
 
-	template <typename... ParamTypes>
-	struct TFunctionRefAsserter<void (ParamTypes...)>
-	{
-		static void Call(void*, ParamTypes&...)
+		TFunctionRefBase(TFunctionRefBase&& Other)
+			: Callable(Other.Callable)
+			, Storage (MoveTemp(Other.Storage))
 		{
-			checkf(false, TEXT("Attempting to call an unbound TFunction!"));
+			if (Callable)
+			{
+				#if ENABLE_TFUNCTIONREF_VISUALIZATION
+					// Use Memcpy to copy the other DebugPtrStorage, including vptr (because we don't know the bound type
+					// here), and then reseat the underlying pointer.  Possibly even more evil than the Set code.
+					FMemory::Memcpy(&DebugPtrStorage, &Other.DebugPtrStorage, sizeof(DebugPtrStorage));
+					DebugPtrStorage.Ptr = Storage.GetPtr();
+				#endif
+
+				Other.Callable = nullptr;
+			}
 		}
-	};
 
-	template <typename DerivedType, typename Ret, typename... ParamTypes>
-	struct TFunctionRefBase<DerivedType, Ret (ParamTypes...)>
-	{
-		explicit TFunctionRefBase(ENoInit)
+		template <typename OtherStorage>
+		TFunctionRefBase(TFunctionRefBase<OtherStorage, Ret (ParamTypes...)>&& Other)
+			: Callable(Other.Callable)
+			, Storage (MoveTemp(Other.Storage))
 		{
-			// Not really designed to be initialized directly, but we want to be explicit about that.
+			if (Callable)
+			{
+				#if ENABLE_TFUNCTIONREF_VISUALIZATION
+					// Use Memcpy to copy the other DebugPtrStorage, including vptr (because we don't know the bound type
+					// here), and then reseat the underlying pointer.  Possibly even more evil than the Set code.
+					FMemory::Memcpy(&DebugPtrStorage, &Other.DebugPtrStorage, sizeof(DebugPtrStorage));
+					DebugPtrStorage.Ptr = Storage.GetPtr();
+				#endif
+
+				Other.Callable = nullptr;
+			}
 		}
 
-		Ret operator()(ParamTypes... Params) const
+		template <typename OtherStorage>
+		TFunctionRefBase(const TFunctionRefBase<OtherStorage, Ret (ParamTypes...)>& Other)
+			: Callable(Other.Callable)
 		{
-			const DerivedType* Derived = static_cast<const DerivedType*>(this);
-			return Callable(Derived->GetPtr(), Params...);
+			if (Callable)
+			{
+				void* NewPtr = Storage.BindCopy(Other.Storage);
+
+				#if ENABLE_TFUNCTIONREF_VISUALIZATION
+					// Use Memcpy to copy the other DebugPtrStorage, including vptr (because we don't know the bound type
+					// here), and then reseat the underlying pointer.  Possibly even more evil than the Set code.
+					FMemory::Memcpy(&DebugPtrStorage, &Other.DebugPtrStorage, sizeof(DebugPtrStorage));
+					DebugPtrStorage.Ptr = NewPtr;
+				#endif
+			}
 		}
 
-	protected:
-		template <typename FunctorType>
-		void Set(FunctorType* Functor)
+		TFunctionRefBase(const TFunctionRefBase& Other)
+			: Callable(Other.Callable)
 		{
-			Callable = &UE4Function_Private::TFunctionRefCaller<FunctorType, Ret (ParamTypes...)>::Call;
+			if (Callable)
+			{
+				void* NewPtr = Storage.BindCopy(Other.Storage);
+
+				#if ENABLE_TFUNCTIONREF_VISUALIZATION
+					// Use Memcpy to copy the other DebugPtrStorage, including vptr (because we don't know the bound type
+					// here), and then reseat the underlying pointer.  Possibly even more evil than the Set code.
+					FMemory::Memcpy(&DebugPtrStorage, &Other.DebugPtrStorage, sizeof(DebugPtrStorage));
+					DebugPtrStorage.Ptr = NewPtr;
+				#endif
+			}
+		}
 
-			#if ENABLE_TFUNCTIONREF_VISUALIZATION
-				// We placement new over the top of the same object each time.  This is illegal,
-				// but it ensures that the vptr is set correctly for the bound type, and so is
-				// visualizable.  We never depend on the state of this object at runtime, so it's
-				// ok.
-				new ((void*)&DebugPtrStorage) UE4Function_Private::TDebugHelper<FunctorType>;
-				DebugPtrStorage.Ptr = (void*)Functor;
-			#endif
+		template <
+			typename FunctorType,
+			typename = typename TEnableIf<
+				TNot<
+					TIsSame<TFunctionRefBase, typename TDecay<FunctorType>::Type>
+				>::Value
+			>::Type
+		>
+		TFunctionRefBase(FunctorType&& InFunc)
+		{
+			if (auto* Binding = Storage.Bind(Forward<FunctorType>(InFunc)))
+			{
+				using DecayedFunctorType = typename TRemovePointer<decltype(Binding)>::Type;
+
+				Callable = &TFunctionRefCaller<DecayedFunctorType, Ret (ParamTypes...)>::Call;
+
+				#if ENABLE_TFUNCTIONREF_VISUALIZATION
+					// We placement new over the top of the same object each time.  This is illegal,
+					// but it ensures that the vptr is set correctly for the bound type, and so is
+					// visualizable.  We never depend on the state of this object at runtime, so it's
+					// ok.
+					new ((void*)&DebugPtrStorage) TDebugHelper<DecayedFunctorType>;
+					DebugPtrStorage.Ptr = (void*)Binding;
+				#endif
+			}
 		}
 
-		void CopyAndReseat(const TFunctionRefBase& Other, void* Functor)
+		TFunctionRefBase& operator=(TFunctionRefBase&&) = delete;
+		TFunctionRefBase& operator=(const TFunctionRefBase&) = delete;
+
+		Ret operator()(ParamTypes... Params) const
 		{
-			Callable = Other.Callable;
+			checkf(Callable, TEXT("Attempting to call an unbound TFunction!"));
+			return Callable(Storage.GetPtr(), Params...);
+		}
 
-			#if ENABLE_TFUNCTIONREF_VISUALIZATION
-				// Use Memcpy to copy the other DebugPtrStorage, including vptr (because we don't know the bound type
-				// here), and then reseat the underlying pointer.  Possibly even more evil than the Set code.
-				FMemory::Memcpy(&DebugPtrStorage, &Other.DebugPtrStorage, sizeof(DebugPtrStorage));
-				DebugPtrStorage.Ptr = Functor;
-			#endif
+		~TFunctionRefBase()
+		{
+			if (Callable)
+			{
+				Storage.Unbind();
+			}
 		}
 
-		void Unset()
+	protected:
+		bool IsSet() const
 		{
-			Callable = &UE4Function_Private::TFunctionRefAsserter<Ret (ParamTypes...)>::Call;
+			return !!Callable;
 		}
 
 	private:
 		// A pointer to a function which invokes the call operator on the callable object
 		Ret (*Callable)(void*, ParamTypes&...);
 
+		StorageType Storage;
+
 		#if ENABLE_TFUNCTIONREF_VISUALIZATION
 			// To help debug visualizers
-			UE4Function_Private::TDebugHelper<void> DebugPtrStorage;
+			TDebugHelper<void> DebugPtrStorage;
 		#endif
 	};
 
-	template <typename T>
-	T&& DeclVal();
-
 	template <typename FunctorType, typename Ret, typename... ParamTypes>
 	struct TFunctorReturnTypeIsCompatible
 		: TIsConstructible<Ret, decltype(DeclVal<FunctorType>()(DeclVal<ParamTypes>()...))>
 	{
 	};
 
+	template <typename MemberRet, typename Class, typename Ret, typename... ParamTypes>
+	struct TFunctorReturnTypeIsCompatible<MemberRet Class::*, Ret, ParamTypes...>
+		: TIsConstructible<Ret, MemberRet>
+	{
+	};
+
+	template <typename MemberRet, typename Class, typename Ret, typename... ParamTypes>
+	struct TFunctorReturnTypeIsCompatible<MemberRet Class::* const, Ret, ParamTypes...>
+		: TIsConstructible<Ret, MemberRet>
+	{
+	};
+
+	template <typename MemberRet, typename Class, typename... MemberParamTypes, typename Ret, typename... ParamTypes>
+	struct TFunctorReturnTypeIsCompatible<MemberRet (Class::*)(MemberParamTypes...), Ret, ParamTypes...>
+		: TIsConstructible<Ret, MemberRet>
+	{
+	};
+
+	template <typename MemberRet, typename Class, typename... MemberParamTypes, typename Ret, typename... ParamTypes>
+	struct TFunctorReturnTypeIsCompatible<MemberRet (Class::*)(MemberParamTypes...) const, Ret, ParamTypes...>
+		: TIsConstructible<Ret, MemberRet>
+	{
+	};
+
 	template <typename FuncType, typename FunctorType>
 	struct TFuncCanBindToFunctor;
 
@@ -366,30 +562,44 @@ namespace UE4Function_Private
 	{
 	};
 
-	template <typename T>
-	struct TIsNullableBinding :
-		TOr<
-			TIsPointer<T>,
-			TIsMemberPointer<T>,
-			TIsTFunction<T>
-		>
+	struct FFunctionRefStoragePolicy
 	{
-	};
+		template <typename FunctorType>
+		typename TRemoveReference<FunctorType>::Type* Bind(FunctorType&& InFunc)
+		{
+			checkf(IsBound(InFunc), TEXT("Cannot bind a null/unbound callable to a TFunctionRef"));
 
-	template <typename T>
-	FORCEINLINE typename TEnableIf<TIsNullableBinding<T>::Value, bool>::Type IsBound(const T& Func)
-	{
-		// Function pointers, data member pointers, member function pointers and TFunctions
-		// can all be null/unbound, so test them using their boolean state.
-		return !!Func;
-	}
+			Ptr = (void*)&InFunc;
+			return &InFunc;
+		}
 
-	template <typename T>
-	FORCEINLINE typename TEnableIf<!TIsNullableBinding<T>::Value, bool>::Type IsBound(const T& Func)
-	{
-		// We can't tell if any other generic callable can be invoked, so just assume they can be.
-		return true;
-	}
+		void* BindCopy(const FFunctionRefStoragePolicy& Other)
+		{
+			void* OtherPtr = Other.Ptr;
+			Ptr = OtherPtr;
+			return OtherPtr;
+		}
+
+		/**
+		 * Returns a pointer to the callable object - needed by TFunctionRefBase.
+		 */
+		void* GetPtr() const
+		{
+			return Ptr;
+		}
+
+		/**
+		 * Destroy any owned bindings - called by TFunctionRefBase only if Bind() or BindCopy() was called.
+		 */
+		void Unbind() const
+		{
+			// FunctionRefs don't own their binding - do nothing
+		}
+
+	private:
+		// A pointer to the callable object
+		void* Ptr;
+	};
 }
 
 /**
@@ -444,142 +654,36 @@ namespace UE4Function_Private
  * }
  */
 template <typename FuncType>
-class TFunctionRef : public UE4Function_Private::TFunctionRefBase<TFunctionRef<FuncType>, FuncType>
+class TFunctionRef : public UE4Function_Private::TFunctionRefBase<UE4Function_Private::FFunctionRefStoragePolicy, FuncType>
 {
-	friend struct UE4Function_Private::TFunctionRefBase<TFunctionRef<FuncType>, FuncType>;
-
-	typedef UE4Function_Private::TFunctionRefBase<TFunctionRef<FuncType>, FuncType> Super;
+	using Super = UE4Function_Private::TFunctionRefBase<UE4Function_Private::FFunctionRefStoragePolicy, FuncType>;
 
 public:
 	/**
-	 * Constructor which binds a TFunctionRef to a non-const lvalue function object.
+	 * Constructor which binds a TFunctionRef to a callable object.
 	 */
 	template <
 		typename FunctorType,
 		typename = typename TEnableIf<
 			TAnd<
-				TNot<TIsFunction<FunctorType>>,
-				TNot<TAreTypesEqual<TFunctionRef, FunctorType>>,
-				UE4Function_Private::TFuncCanBindToFunctor<FuncType, FunctorType>
+				TNot<TIsTFunctionRef<typename TDecay<FunctorType>::Type>>,
+				UE4Function_Private::TFuncCanBindToFunctor<FuncType, typename TDecay<FunctorType>::Type>
 			>::Value
 		>::Type
 	>
-	TFunctionRef(FunctorType& Functor)
-		: Super(NoInit)
+	TFunctionRef(FunctorType&& InFunc)
+		: Super(Forward<FunctorType>(InFunc))
 	{
-		// This constructor is disabled for function types because we want it to call the function pointer overload.
-		// It is also disabled for TFunctionRef types because VC is incorrectly treating it as a copy constructor.
-
-		checkf(UE4Function_Private::IsBound(Functor), TEXT("Cannot bind a null/unbound callable to a TFunctionRef"));
-
-		Set(&Functor);
+		// This constructor is disabled for TFunctionRef types so it isn't incorrectly selected as copy/move constructors.
 	}
 
-	/**
-	 * Constructor which binds a TFunctionRef to an rvalue or const lvalue function object.
-	 */
-	template <
-		typename FunctorType,
-		typename = typename TEnableIf<
-			TAnd<
-				TNot<TIsFunction<FunctorType>>,
-				TNot<TAreTypesEqual<TFunctionRef, FunctorType>>,
-				UE4Function_Private::TFuncCanBindToFunctor<FuncType, FunctorType>
-			>::Value
-		>::Type
-	>
-	TFunctionRef(const FunctorType& Functor)
-		: Super(NoInit)
-	{
-		// This constructor is disabled for function types because we want it to call the function pointer overload.
-		// It is also disabled for TFunctionRef types because VC is incorrectly treating it as a copy constructor.
-
-		checkf(UE4Function_Private::IsBound(Functor), TEXT("Cannot bind a null/unbound callable to a TFunctionRef"));
-
-		Set(&Functor);
-	}
-
-	/**
-	 * Constructor which binds a TFunctionRef to a function pointer.
-	 */
-	template <
-		typename FunctionType,
-		typename = typename TEnableIf<
-			TAnd<
-				TIsFunction<FunctionType>,
-				UE4Function_Private::TFuncCanBindToFunctor<FuncType, FunctionType>
-			>::Value
-		>::Type
-	>
-	TFunctionRef(FunctionType* Function)
-		: Super(NoInit)
-	{
-		// This constructor is enabled only for function types because we don't want weird errors from it being called with arbitrary pointers.
-
-		checkf(Function, TEXT("Cannot bind a null/unbound callable to a TFunctionRef"));
-
-		Set(Function);
-	}
-
-	#if ENABLE_TFUNCTIONREF_VISUALIZATION
-
-		/**
-		 * Copy constructor.
-		 */
-		TFunctionRef(const TFunctionRef& Other)
-			: Super(NoInit)
-		{
-			// If visualization is enabled, then we need to do an explicit copy
-			// to ensure that our hacky DebugPtrStorage's vptr is copied.
-			CopyAndReseat(Other, Other.Ptr);
-		}
-
-	#else
-
-		TFunctionRef(const TFunctionRef&) = default;
-
-	#endif
+	TFunctionRef(const TFunctionRef&) = default;
 
 	// We delete the assignment operators because we don't want it to be confused with being related to
 	// regular C++ reference assignment - i.e. calling the assignment operator of whatever the reference
 	// is bound to - because that's not what TFunctionRef does, nor is it even capable of doing that.
 	TFunctionRef& operator=(const TFunctionRef&) const = delete;
 	~TFunctionRef() = default;
-
-private:
-	/**
-	 * Sets the state of the TFunctionRef given a pointer to a callable thing.
-	 */
-	template <typename FunctorType>
-	void Set(FunctorType* Functor)
-	{
-		// We force a void* cast here because if FunctorType is an actual function then
-		// this won't compile.  We convert it back again before we use it anyway.
-
-		Ptr  = (void*)Functor;
-		Super::Set(Functor);
-	}
-
-	/**
-	 * Copies another TFunctionRef and rebinds it to another object of the same type which was originally bound.
-	 * Only intended to be used by TFunction's copy constructor/assignment operator.
-	 */
-	void CopyAndReseat(const TFunctionRef& Other, void* Functor)
-	{
-		Ptr = Functor;
-		Super::CopyAndReseat(Other, Functor);
-	}
-
-	/**
-	 * Returns a pointer to the callable object - needed by TFunctionRefBase.
-	 */
-	void* GetPtr() const
-	{
-		return Ptr;
-	}
-
-	// A pointer to the callable object
-	void* Ptr;
 };
 
 /**
@@ -620,20 +724,16 @@ private:
  * }
  */
 template <typename FuncType>
-class TFunction : public UE4Function_Private::TFunctionRefBase<TFunction<FuncType>, FuncType>
+class TFunction final : public UE4Function_Private::TFunctionRefBase<UE4Function_Private::FFunctionStorage, FuncType>
 {
-	friend struct UE4Function_Private::TFunctionRefBase<TFunction<FuncType>, FuncType>;
-
-	typedef UE4Function_Private::TFunctionRefBase<TFunction<FuncType>, FuncType> Super;
+	using Super = UE4Function_Private::TFunctionRefBase<UE4Function_Private::FFunctionStorage, FuncType>;
 
 public:
 	/**
 	 * Default constructor.
 	 */
 	TFunction(TYPE_OF_NULLPTR = nullptr)
-		: Super(NoInit)
 	{
-		Super::Unset();
 	}
 
 	/**
@@ -643,18 +743,15 @@ public:
 		typename FunctorType,
 		typename = typename TEnableIf<
 			TAnd<
-				TNot<TAreTypesEqual<TFunction, typename TDecay<FunctorType>::Type>>,
+				TNot<TIsTFunction<typename TDecay<FunctorType>::Type>>,
 				UE4Function_Private::TFuncCanBindToFunctor<FuncType, FunctorType>
 			>::Value
 		>::Type
 	>
 	TFunction(FunctorType&& InFunc)
-		: Super(NoInit)
+		: Super(Forward<FunctorType>(InFunc))
 	{
-		// This constructor is disabled for TFunction types because VC is incorrectly treating it as copy/move constructors.
-
-		typedef typename TDecay<FunctorType>::Type                             DecayedFunctorType;
-		typedef UE4Function_Private::TFunction_OwnedObject<DecayedFunctorType> OwnedType;
+		// This constructor is disabled for TFunction types so it isn't incorrectly selected as copy/move constructors.
 
 		// This is probably a mistake if you expect TFunction to take a copy of what
 		// TFunctionRef is bound to, because that's not possible.
@@ -663,110 +760,136 @@ public:
 		// wrap it in a lambda (and thus it's clear you're just binding to a call to another
 		// reference):
 		//
-		// TFunction<int32(float)> MyFunction = [=](float F) -> int32 { return MyFunctionRef(F); };
-		static_assert(!TIsTFunctionRef<DecayedFunctorType>::Value, "Cannot construct a TFunction from a TFunctionRef");
-
-		if (UE4Function_Private::IsBound(InFunc))
-		{
-			OwnedType* NewObj = new (Storage) OwnedType(Forward<FunctorType>(InFunc));
-			Super::Set(&NewObj->Obj);
-		}
-		else
-		{
-			Super::Unset();
-		}
+		// TFunction<int32(float)> MyFunction = [MyFunctionRef](float F) { return MyFunctionRef(F); };
+		static_assert(!TIsTFunctionRef<typename TDecay<FunctorType>::Type>::Value, "Cannot construct a TFunction from a TFunctionRef");
 	}
 
+	TFunction(TFunction&&) = default;
+	TFunction(const TFunction& Other) = default;
+	~TFunction() = default;
+
 	/**
-	 * Copy constructor.
+	 * Move assignment operator.
 	 */
-	TFunction(const TFunction& Other)
-		: Super(NoInit)
+	TFunction& operator=(TFunction&& Other)
 	{
-		if (UE4Function_Private::IFunction_OwnedObject* OtherFunc = Other.Storage.GetBoundObject())
-		{
-			UE4Function_Private::IFunction_OwnedObject* ThisFunc = OtherFunc->CopyToEmptyStorage(Storage);
-			Super::CopyAndReseat(Other, ThisFunc->GetAddress());
-		}
-		else
-		{
-			Super::Unset();
-		}
+		Swap(*this, Other);
+		return *this;
 	}
 
 	/**
-	 * Move constructor.
+	 * Copy assignment operator.
 	 */
-	TFunction(TFunction&& Other)
-		: Super  (NoInit)
-		, Storage(MoveTemp(Other.Storage))
+	TFunction& operator=(const TFunction& Other)
 	{
-		if (UE4Function_Private::IFunction_OwnedObject* Func = Storage.GetBoundObject())
-		{
-			Super::CopyAndReseat(Other, Func->GetAddress());
-		}
-		else
-		{
-			Super::Unset();
-		}
+		TFunction Temp = Other;
+		Swap(*this, Temp);
+		return *this;
+	}
 
-		Other.Unset();
+	/**
+	 * Tests if the TFunction is callable.
+	 */
+	FORCEINLINE explicit operator bool() const
+	{
+		return Super::IsSet();
 	}
+};
 
+/**
+ * TUniqueFunction<FuncType>
+ *
+ * Used like TFunction above, but is move-only.  This allows non-copyable functors to be bound to it.
+ *
+ * Example:
+ *
+ * TUniquePtr<FThing> Thing = MakeUnique<FThing>();
+ *
+ * TFunction      <void()> CopyableFunc = [Thing = MoveTemp(Thing)](){ Thing->DoSomething(); }; // error - lambda is not copyable
+ * TUniqueFunction<void()> MovableFunc  = [Thing = MoveTemp(Thing)](){ Thing->DoSomething(); }; // ok
+ *
+ * void Foo(TUniqueFunction<void()> Func);
+ * Foo(MovableFunc);           // error - TUniqueFunction is not copyable
+ * Foo(MoveTemp(MovableFunc)); // ok
+ */
+template <typename FuncType>
+class TUniqueFunction final : public UE4Function_Private::TFunctionRefBase<UE4Function_Private::FUniqueFunctionStorage, FuncType>
+{
+	using Super = UE4Function_Private::TFunctionRefBase<UE4Function_Private::FUniqueFunctionStorage, FuncType>;
+
+public:
 	/**
-	 * Copy/move assignment operator.
+	 * Default constructor.
 	 */
-	TFunction& operator=(TFunction Other)
+	TUniqueFunction(TYPE_OF_NULLPTR = nullptr)
 	{
-		FMemory::Memswap(&Other, this, sizeof(TFunction));
-		return *this;
 	}
 
 	/**
-	 * Nullptr assignment operator - unbinds any bound function.
+	 * Constructor which binds a TFunction to any function object.
 	 */
-	TFunction& operator=(TYPE_OF_NULLPTR)
+	template <
+		typename FunctorType,
+		typename = typename TEnableIf<
+			TAnd<
+				TNot<TOr<TIsTUniqueFunction<typename TDecay<FunctorType>::Type>, TIsTFunction<typename TDecay<FunctorType>::Type>>>,
+				UE4Function_Private::TFuncCanBindToFunctor<FuncType, FunctorType>
+			>::Value
+		>::Type
+	>
+	TUniqueFunction(FunctorType&& InFunc)
+		: Super(Forward<FunctorType>(InFunc))
 	{
-		if (UE4Function_Private::IFunction_OwnedObject* Obj = Storage.GetBoundObject())
-		{
-			Obj->~IFunction_OwnedObject();
-		}
-		Storage.Empty();
-		Super::Unset();
+		// This constructor is disabled for TUniqueFunction types so it isn't incorrectly selected as copy/move constructors.
 
-		return *this;
+		// This is probably a mistake if you expect TFunction to take a copy of what
+		// TFunctionRef is bound to, because that's not possible.
+		//
+		// If you really intended to bind a TFunction to a TFunctionRef, you can just
+		// wrap it in a lambda (and thus it's clear you're just binding to a call to another
+		// reference):
+		//
+		// TFunction<int32(float)> MyFunction = [MyFunctionRef](float F) { return MyFunctionRef(F); };
+		static_assert(!TIsTFunctionRef<typename TDecay<FunctorType>::Type>::Value, "Cannot construct a TUniqueFunction from a TFunctionRef");
 	}
 
 	/**
-	 * Destructor.
+	 * Constructor which takes ownership of a TFunction's functor.
 	 */
-	~TFunction()
+	TUniqueFunction(TFunction<FuncType>&& Other)
+		: Super(MoveTemp(*(UE4Function_Private::TFunctionRefBase<UE4Function_Private::FFunctionStorage, FuncType>*)&Other))
 	{
-		if (UE4Function_Private::IFunction_OwnedObject* Obj = Storage.GetBoundObject())
-		{
-			Obj->~IFunction_OwnedObject();
-		}
 	}
 
 	/**
-	 * Tests if the TFunction is callable.
+	 * Constructor which takes ownership of a TFunction's functor.
 	 */
-	FORCEINLINE explicit operator bool() const
+	TUniqueFunction(const TFunction<FuncType>& Other)
+		: Super(*(const UE4Function_Private::TFunctionRefBase<UE4Function_Private::FFunctionStorage, FuncType>*)&Other)
 	{
-		return !!Storage.GetBoundObject();
 	}
 
-private:
 	/**
-	 * Returns a pointer to the callable object - needed by TFunctionRefBase.
+	 * Copy/move assignment operator.
 	 */
-	void* GetPtr() const
+	TUniqueFunction& operator=(TUniqueFunction&& Other)
 	{
-		UE4Function_Private::IFunction_OwnedObject* Ptr = Storage.GetBoundObject();
-		return Ptr ? Ptr->GetAddress() : nullptr;
+		Swap(*this, Other);
+		return *this;
 	}
 
-	UE4Function_Private::FFunctionStorage Storage;
+	TUniqueFunction(TUniqueFunction&&) = default;
+	TUniqueFunction(const TUniqueFunction& Other) = delete;
+	TUniqueFunction& operator=(const TUniqueFunction& Other) = delete;
+	~TUniqueFunction() = default;
+
+	/**
+	 * Tests if the TUniqueFunction is callable.
+	 */
+	FORCEINLINE explicit operator bool() const
+	{
+		return Super::IsSet();
+	}
 };
 
 /**
