GitHub #4743 : modified VirtualAlloc function flag

https://blogs.msdn.microsoft.com/oldnewthing/[REDACTED_PHONE]/?p=91411

diff --git a/Engine/Source/Runtime/Core/Private/HAL/MallocStomp.cpp b/Engine/Source/Runtime/Core/Private/HAL/MallocStomp.cpp
--- a/Engine/Source/Runtime/Core/Private/HAL/MallocStomp.cpp
+++ b/Engine/Source/Runtime/Core/Private/HAL/MallocStomp.cpp
@@ -112,7 +112,7 @@ void* FMallocStomp::Malloc(SIZE_T Size, uint32 Alignment)
 
 #if PLATFORM_WINDOWS && MALLOC_STOMP_KEEP_VIRTUAL_MEMORY
 		// Commit physical pages to the used range, leaving the first page unmapped.
-		void* CommittedMemory = VirtualAlloc(AllocDataPointer, AllocFullPageSize, MEM_COMMIT, PAGE_READWRITE);
+		void* CommittedMemory = VirtualAlloc(AllocDataPointer, AllocFullPageSize, MEM_COMMIT | MEM_RESERVE, PAGE_READWRITE);
 		if (!CommittedMemory)
 		{
 			// Failed to allocate and commit physical memory pages. Report OOM.
@@ -131,7 +131,7 @@ void* FMallocStomp::Malloc(SIZE_T Size, uint32 Alignment)
 
 #if PLATFORM_WINDOWS && MALLOC_STOMP_KEEP_VIRTUAL_MEMORY
 		// Commit physical pages to the used range, leaving the last page unmapped.
-		void* CommittedMemory = VirtualAlloc(FullAllocationPointer, AllocFullPageSize, MEM_COMMIT, PAGE_READWRITE);
+		void* CommittedMemory = VirtualAlloc(FullAllocationPointer, AllocFullPageSize, MEM_COMMIT | MEM_RESERVE, PAGE_READWRITE);
 		if (!CommittedMemory)
 		{
 			// Failed to allocate and commit physical memory pages. Report OOM.
diff --git a/Engine/Source/Runtime/Core/Private/Misc/MemStack.cpp b/Engine/Source/Runtime/Core/Private/Misc/MemStack.cpp
--- a/Engine/Source/Runtime/Core/Private/Misc/MemStack.cpp
+++ b/Engine/Source/Runtime/Core/Private/Misc/MemStack.cpp
@@ -98,7 +98,7 @@ void *FPageAllocator::Alloc()
 		Result = TheAllocatorReady.Pop();
 		if (!Result)
 		{
-			Result = VirtualAlloc(NULL, PageSize, MEM_COMMIT, PAGE_READWRITE);
+			Result = VirtualAlloc(NULL, PageSize, MEM_COMMIT | MEM_RESERVE, PAGE_READWRITE);
 			BigBlocksOutstanding.Increment();
 		}
 		BigBlocksInUse.Increment();
@@ -146,7 +146,7 @@ void *FPageAllocator::AllocSmall()
 		Result = TheSmallAllocatorReady.Pop();
 		if (!Result)
 		{
-			Result = VirtualAlloc(NULL, SmallPageSize, MEM_COMMIT, PAGE_READWRITE);
+			Result = VirtualAlloc(NULL, SmallPageSize, MEM_COMMIT | MEM_RESERVE, PAGE_READWRITE);
 			SmallBlocksOutstanding.Increment();
 		}
 		SmallBlocksInUse.Increment();
diff --git a/Engine/Source/Runtime/Core/Private/Windows/WindowsPlatformMemory.cpp b/Engine/Source/Runtime/Core/Private/Windows/WindowsPlatformMemory.cpp
--- a/Engine/Source/Runtime/Core/Private/Windows/WindowsPlatformMemory.cpp
+++ b/Engine/Source/Runtime/Core/Private/Windows/WindowsPlatformMemory.cpp
@@ -261,7 +261,7 @@ bool FWindowsPlatformMemory::PageProtect(void* const Ptr, const SIZE_T Size, con
 }
 void* FWindowsPlatformMemory::BinnedAllocFromOS( SIZE_T Size )
 {
-	void* Ptr = VirtualAlloc( NULL, Size, MEM_COMMIT, PAGE_READWRITE );
+	void* Ptr = VirtualAlloc( NULL, Size, MEM_COMMIT | MEM_RESERVE, PAGE_READWRITE );
 	LLM(FLowLevelMemTracker::Get().OnLowLevelAlloc(ELLMTracker::Platform, Ptr, Size));
 	return Ptr;
 }
