Reverting change to FMallocStomp from CL# 4089771.

diff --git a/Engine/Source/Runtime/Core/Private/HAL/MallocStomp.cpp b/Engine/Source/Runtime/Core/Private/HAL/MallocStomp.cpp
--- a/Engine/Source/Runtime/Core/Private/HAL/MallocStomp.cpp
+++ b/Engine/Source/Runtime/Core/Private/HAL/MallocStomp.cpp
@@ -112,7 +112,7 @@ void* FMallocStomp::Malloc(SIZE_T Size, uint32 Alignment)
 
 #if PLATFORM_WINDOWS && MALLOC_STOMP_KEEP_VIRTUAL_MEMORY
 		// Commit physical pages to the used range, leaving the first page unmapped.
-		void* CommittedMemory = VirtualAlloc(AllocDataPointer, AllocFullPageSize, MEM_COMMIT | MEM_RESERVE, PAGE_READWRITE);
+		void* CommittedMemory = VirtualAlloc(AllocDataPointer, AllocFullPageSize, MEM_COMMIT, PAGE_READWRITE);
 		if (!CommittedMemory)
 		{
 			// Failed to allocate and commit physical memory pages. Report OOM.
@@ -131,7 +131,7 @@ void* FMallocStomp::Malloc(SIZE_T Size, uint32 Alignment)
 
 #if PLATFORM_WINDOWS && MALLOC_STOMP_KEEP_VIRTUAL_MEMORY
 		// Commit physical pages to the used range, leaving the last page unmapped.
-		void* CommittedMemory = VirtualAlloc(FullAllocationPointer, AllocFullPageSize, MEM_COMMIT | MEM_RESERVE, PAGE_READWRITE);
+		void* CommittedMemory = VirtualAlloc(FullAllocationPointer, AllocFullPageSize, MEM_COMMIT, PAGE_READWRITE);
 		if (!CommittedMemory)
 		{
 			// Failed to allocate and commit physical memory pages. Report OOM.
