Unification of all platforms' FMath::CountTrailingZeros() and FMath::CountLeadingZeros() for both 32-bit and 64-bit.

diff --git a/Engine/Source/Runtime/Core/Public/Android/AndroidMath.h b/Engine/Source/Runtime/Core/Public/Android/AndroidMath.h
--- a/Engine/Source/Runtime/Core/Public/Android/AndroidMath.h
+++ b/Engine/Source/Runtime/Core/Public/Android/AndroidMath.h
@@ -1,50 +1,10 @@
 // Copyright 1998-2018 Epic Games, Inc. All Rights Reserved.
 
 /*=============================================================================================
-	AndroidMath.h: KickStart platform Math functions
+	AndroidMath.h: Android platform Math functions
 ==============================================================================================*/
 
 #pragma once
-#include "GenericPlatform/GenericPlatformMath.h"
+#include "Clang/ClangPlatformMath.h"
 
-/**
- * Android implementation of the Math OS functions
- **/
-struct FAndroidPlatformMath : public FGenericPlatformMath
-{
-	/**
-	 * Counts the number of leading zeros in the bit representation of the value
-	 *
-	 * @param Value the value to determine the number of leading zeros for
-	 *
-	 * @return the number of zeros before the first "on" bit
-	 */
-	static FORCEINLINE uint32 CountLeadingZeros(uint32 Value)
-	{
-		if (Value == 0)
-		{
-			return 32;
-		}
-	
-		return __builtin_clz(Value);
-	}
-	
-	/**
-	 * Counts the number of trailing zeros in the bit representation of the value
-	 *
-	 * @param Value the value to determine the number of trailing zeros for
-	 *
-	 * @return the number of zeros after the last "on" bit
-	 */
-	static FORCEINLINE uint32 CountTrailingZeros(uint32 Value)
-	{
-		if (Value == 0)
-		{
-			return 32;
-		}
-	
-		return __builtin_ctz(Value);
-	}
-};
-
-typedef FAndroidPlatformMath FPlatformMath;
+typedef FClangPlatformMath FPlatformMath;
diff --git a/Engine/Source/Runtime/Core/Public/Clang/ClangPlatformMath.h b/Engine/Source/Runtime/Core/Public/Clang/ClangPlatformMath.h
new file mode 100644
--- /dev/null
+++ b/Engine/Source/Runtime/Core/Public/Clang/ClangPlatformMath.h
@@ -0,0 +1,82 @@
+// Copyright 1998-2018 Epic Games, Inc. All Rights Reserved.
+
+/*=============================================================================================
+	ClangPlatformMath.h: Clang intrinsic implementations of some platform Math functions
+==============================================================================================*/
+
+#pragma once
+#include "GenericPlatform/GenericPlatformMath.h"
+
+/**
+ * Clang implementation of math functions
+ **/
+struct FClangPlatformMath : public FGenericPlatformMath
+{
+	/**
+	 * Counts the number of leading zeros in the bit representation of the value
+	 *
+	 * @param Value the value to determine the number of leading zeros for
+	 *
+	 * @return the number of zeros before the first "on" bit
+	 */
+	static FORCEINLINE uint32 CountLeadingZeros(uint32 Value)
+	{
+		if (Value == 0)
+		{
+			return 32;
+		}
+	
+		return __builtin_clz(Value);
+	}
+
+	/**
+	 * Counts the number of leading zeros in the bit representation of the value
+	 *
+	 * @param Value the value to determine the number of leading zeros for
+	 *
+	 * @return the number of zeros before the first "on" bit
+	 */
+	static FORCEINLINE uint64 CountLeadingZeros64(uint64 Value)
+	{
+		if (Value == 0)
+		{
+			return 64;
+		}
+
+		return __builtin_clzll(Value);
+	}
+
+	/**
+	 * Counts the number of trailing zeros in the bit representation of the value
+	 *
+	 * @param Value the value to determine the number of trailing zeros for
+	 *
+	 * @return the number of zeros after the last "on" bit
+	 */
+	static FORCEINLINE uint32 CountTrailingZeros(uint32 Value)
+	{
+		if (Value == 0)
+		{
+			return 32;
+		}
+	
+		return __builtin_ctz(Value);
+	}
+	
+	/**
+	 * Counts the number of trailing zeros in the bit representation of the value
+	 *
+	 * @param Value the value to determine the number of trailing zeros for
+	 *
+	 * @return the number of zeros after the last "on" bit
+	 */
+	static FORCEINLINE uint64 CountTrailingZeros64(uint64 Value)
+	{
+		if (Value == 0)
+		{
+			return 64;
+		}
+	
+		return __builtin_ctzll(Value);
+	}
+};
diff --git a/Engine/Source/Runtime/Core/Public/GenericPlatform/GenericPlatformMath.h b/Engine/Source/Runtime/Core/Public/GenericPlatform/GenericPlatformMath.h
--- a/Engine/Source/Runtime/Core/Public/GenericPlatform/GenericPlatformMath.h
+++ b/Engine/Source/Runtime/Core/Public/GenericPlatform/GenericPlatformMath.h
@@ -405,6 +405,28 @@ struct FGenericPlatformMath
 		return Result;
 	}
 
+	/**
+	 * Counts the number of trailing zeros in the bit representation of the value
+	 *
+	 * @param Value the value to determine the number of trailing zeros for
+	 *
+	 * @return the number of zeros after the last "on" bit
+	 */
+	static FORCEINLINE uint64 CountTrailingZeros64(uint64 Value)
+	{
+		if (Value == 0)
+		{
+			return 64;
+		}
+		uint64 Result = 0;
+		while ((Value & 1) == 0)
+		{
+			Value >>= 1;
+			++Result;
+		}
+		return Result;
+	}
+
 	/**
 	 * Returns smallest N such that (1<<N)>=Arg.
 	 * Note: CeilLogTwo(0)=0 because (1<<0)=1 >= 0.
diff --git a/Engine/Source/Runtime/Core/Public/HTML5/HTML5PlatformMath.h b/Engine/Source/Runtime/Core/Public/HTML5/HTML5PlatformMath.h
--- a/Engine/Source/Runtime/Core/Public/HTML5/HTML5PlatformMath.h
+++ b/Engine/Source/Runtime/Core/Public/HTML5/HTML5PlatformMath.h
@@ -2,47 +2,10 @@
 
 
 /*=============================================================================================
-	PS4Math.h: PS4 platform Math functions
+	HTML5PlatformMath.h: HTML5 platform Math functions
 ==============================================================================================*/
 
 #pragma once
-#include "GenericPlatform/GenericPlatformMath.h"
+#include "Clang/ClangPlatformMath.h"
 
-struct FHTML5PlatformMath : public FGenericPlatformMath
-{
-	/**
-	 * Counts the number of leading zeros in the bit representation of the value
-	 *
-	 * @param Value the value to determine the number of leading zeros for
-	 *
-	 * @return the number of zeros before the first "on" bit
-	 */
-	static FORCEINLINE uint32 CountLeadingZeros(uint32 Value)
-	{
-		if (Value == 0)
-		{
-			return 32;
-		}
-	
-		return __builtin_clz(Value);
-	}
-	
-	/**
-	 * Counts the number of trailing zeros in the bit representation of the value
-	 *
-	 * @param Value the value to determine the number of trailing zeros for
-	 *
-	 * @return the number of zeros after the last "on" bit
-	 */
-	static FORCEINLINE uint32 CountTrailingZeros(uint32 Value)
-	{
-		if (Value == 0)
-		{
-			return 32;
-		}
-	
-		return __builtin_ctz(Value);
-	}
-};
-
-typedef FHTML5PlatformMath FPlatformMath;
+typedef FClangPlatformMath FPlatformMath;
diff --git a/Engine/Source/Runtime/Core/Public/IOS/IOSPlatformMath.h b/Engine/Source/Runtime/Core/Public/IOS/IOSPlatformMath.h
--- a/Engine/Source/Runtime/Core/Public/IOS/IOSPlatformMath.h
+++ b/Engine/Source/Runtime/Core/Public/IOS/IOSPlatformMath.h
@@ -5,46 +5,6 @@
 ==============================================================================================*/
 
 #pragma once
-#include "GenericPlatform/GenericPlatformMath.h"
+#include "Clang/ClangPlatformMath.h"
 
-/**
-* IOS implementation of the Math OS functions
-**/
-struct FIOSPlatformMath : public FGenericPlatformMath
-{
-	/**
-	 * Counts the number of leading zeros in the bit representation of the value
-	 *
-	 * @param Value the value to determine the number of leading zeros for
-	 *
-	 * @return the number of zeros before the first "on" bit
-	 */
-	static FORCEINLINE uint32 CountLeadingZeros(uint32 Value)
-	{
-		if (Value == 0)
-		{
-			return 32;
-		}
-	
-		return __builtin_clz(Value);
-	}
-	
-	/**
-	 * Counts the number of trailing zeros in the bit representation of the value
-	 *
-	 * @param Value the value to determine the number of trailing zeros for
-	 *
-	 * @return the number of zeros after the last "on" bit
-	 */
-	static FORCEINLINE uint32 CountTrailingZeros(uint32 Value)
-	{
-		if (Value == 0)
-		{
-			return 32;
-		}
-	
-		return __builtin_ctz(Value);
-	}
-};
-
-typedef FIOSPlatformMath FPlatformMath;
+typedef FClangPlatformMath FPlatformMath;
diff --git a/Engine/Source/Runtime/Core/Public/Mac/MacPlatformMath.h b/Engine/Source/Runtime/Core/Public/Mac/MacPlatformMath.h
--- a/Engine/Source/Runtime/Core/Public/Mac/MacPlatformMath.h
+++ b/Engine/Source/Runtime/Core/Public/Mac/MacPlatformMath.h
@@ -5,14 +5,14 @@
 ==============================================================================================*/
 
 #pragma once
-#include "GenericPlatform/GenericPlatformMath.h"
+#include "Clang/ClangPlatformMath.h"
 #include "Mac/MacSystemIncludes.h"
 #include "Math/UnrealPlatformMathSSE.h"
 
 /**
 * Mac implementation of the Math OS functions
 **/
-struct FMacPlatformMath : public FGenericPlatformMath
+struct FMacPlatformMath : public FClangPlatformMath
 {
 	static FORCEINLINE int32 TruncToInt(float F)
 	{
@@ -80,40 +80,6 @@ struct FMacPlatformMath : public FGenericPlatformMath
 		return UnrealPlatformMathSSE::InvSqrtEst(F);
 	}
 #endif
-
-	/**
-	 * Counts the number of leading zeros in the bit representation of the value
-	 *
-	 * @param Value the value to determine the number of leading zeros for
-	 *
-	 * @return the number of zeros before the first "on" bit
-	 */
-	static FORCEINLINE uint32 CountLeadingZeros(uint32 Value)
-	{
-		if (Value == 0)
-		{
-			return 32;
-		}
-	
-		return __builtin_clz(Value);
-	}
-	
-	/**
-	 * Counts the number of trailing zeros in the bit representation of the value
-	 *
-	 * @param Value the value to determine the number of trailing zeros for
-	 *
-	 * @return the number of zeros after the last "on" bit
-	 */
-	static FORCEINLINE uint32 CountTrailingZeros(uint32 Value)
-	{
-		if (Value == 0)
-		{
-			return 32;
-		}
-	
-		return __builtin_ctz(Value);
-	}
 };
 
 typedef FMacPlatformMath FPlatformMath;
diff --git a/Engine/Source/Runtime/Core/Public/Unix/UnixPlatformMath.h b/Engine/Source/Runtime/Core/Public/Unix/UnixPlatformMath.h
--- a/Engine/Source/Runtime/Core/Public/Unix/UnixPlatformMath.h
+++ b/Engine/Source/Runtime/Core/Public/Unix/UnixPlatformMath.h
@@ -7,7 +7,7 @@
 #pragma once
 
 #include "CoreTypes.h"
-#include "GenericPlatform/GenericPlatformMath.h"
+#include "Clang/ClangPlatformMath.h"
 #include "Linux/LinuxSystemIncludes.h"
 #if PLATFORM_ENABLE_VECTORINTRINSICS
 	#include "Math/UnrealPlatformMathSSE.h"
@@ -16,7 +16,7 @@
 /**
 * Linux implementation of the Math OS functions
 **/
-struct FLinuxPlatformMath : public FGenericPlatformMath
+struct FLinuxPlatformMath : public FClangPlatformMath
 {
 #if PLATFORM_ENABLE_VECTORINTRINSICS
 	static FORCEINLINE int32 TruncToInt(float F)
@@ -85,40 +85,6 @@ struct FLinuxPlatformMath : public FGenericPlatformMath
 	}
 #endif
 #endif
-
-	/**
-	 * Counts the number of leading zeros in the bit representation of the value
-	 *
-	 * @param Value the value to determine the number of leading zeros for
-	 *
-	 * @return the number of zeros before the first "on" bit
-	 */
-	static FORCEINLINE uint32 CountLeadingZeros(uint32 Value)
-	{
-		if (Value == 0)
-		{
-			return 32;
-		}
-
-		return __builtin_clz(Value);
-	}
-
-	/**
-	 * Counts the number of trailing zeros in the bit representation of the value
-	 *
-	 * @param Value the value to determine the number of trailing zeros for
-	 *
-	 * @return the number of zeros after the last "on" bit
-	 */
-	static FORCEINLINE uint32 CountTrailingZeros(uint32 Value)
-	{
-		if (Value == 0)
-		{
-			return 32;
-		}
-
-		return __builtin_ctz(Value);
-	}
 };
 
 typedef FLinuxPlatformMath FPlatformMath;
diff --git a/Engine/Source/Runtime/Core/Public/Windows/WindowsPlatformMath.h b/Engine/Source/Runtime/Core/Public/Windows/WindowsPlatformMath.h
--- a/Engine/Source/Runtime/Core/Public/Windows/WindowsPlatformMath.h
+++ b/Engine/Source/Runtime/Core/Public/Windows/WindowsPlatformMath.h
@@ -140,6 +140,17 @@ struct FWindowsPlatformMath : public FGenericPlatformMath
 
 		return 64;
 	}
+	static FORCEINLINE uint64 CountTrailingZeros64(uint64 Value)
+	{
+		if (Value == 0)
+		{
+			return 64;
+		}
+		unsigned long BitIndex;	// 0-based, where the LSB is 0 and MSB is 31
+		_BitScanForward64( &BitIndex, Value );	// Scans from LSB to MSB
+		return BitIndex;
+	}
+
 #endif
 
 #if PLATFORM_ENABLE_POPCNT_INTRINSIC
