Fix for TStaticBitArray::FindFirstSetBit() and TStaticBitArray::FindFirstClearBit().
Unused variables removed.

diff --git a/Engine/Source/Runtime/Core/Public/Containers/StaticBitArray.h b/Engine/Source/Runtime/Core/Public/Containers/StaticBitArray.h
--- a/Engine/Source/Runtime/Core/Public/Containers/StaticBitArray.h
+++ b/Engine/Source/Runtime/Core/Public/Containers/StaticBitArray.h
@@ -297,8 +297,7 @@ public:
 			const WordType Bits = ~(Words[WordIndex]);
 			ASSUME(Bits != 0);
 
-			const uint32 LowestBit = (Bits) & (-WordType(Bits));
-			const int32 LowestBitIndex = FMath::CountTrailingZeros(Bits) + (WordIndex << NumBitsPerWordLog2);
+			const int32 LowestBitIndex = FMath::CountTrailingZeros64(Bits) + (WordIndex << NumBitsPerWordLog2);
 			if (LowestBitIndex < LocalNumBits)
 			{
 				return LowestBitIndex;
@@ -330,8 +329,7 @@ public:
 			const WordType Bits = Words[WordIndex];
 			ASSUME(Bits != 0);
 
-			const uint32 LowestBit = (Bits) & (-WordType(Bits));
-			const int32 LowestBitIndex = FMath::CountTrailingZeros(Bits) + (WordIndex << NumBitsPerWordLog2);
+			const int32 LowestBitIndex = FMath::CountTrailingZeros64(Bits) + (WordIndex << NumBitsPerWordLog2);
 			if (LowestBitIndex < LocalNumBits)
 			{
 				return LowestBitIndex;
