CIS fixes.

diff --git a/Engine/Source/Runtime/Core/Private/Serialization/ArchiveFromStructuredArchive.cpp b/Engine/Source/Runtime/Core/Private/Serialization/ArchiveFromStructuredArchive.cpp
--- a/Engine/Source/Runtime/Core/Private/Serialization/ArchiveFromStructuredArchive.cpp
+++ b/Engine/Source/Runtime/Core/Private/Serialization/ArchiveFromStructuredArchive.cpp
@@ -13,7 +13,7 @@ FArchiveFromStructuredArchive::FArchiveFromStructuredArchive(FStructuredArchive:
 	// For some reason, the FArchive copy constructor will copy all the trivial members of the source archive, but then specifically set ArIsFilterEditorOnly to false, with a comment saying
 	// they don't know why it's doing this... make sure we inherit this flag here!
 	ArIsFilterEditorOnly = InnerArchive.ArIsFilterEditorOnly;
-	ArIsTextFormat = false;
+	SetIsTextFormat(false);
 }
 
 FArchiveFromStructuredArchive::~FArchiveFromStructuredArchive()
diff --git a/Engine/Source/Runtime/Core/Public/Containers/BitArray.h b/Engine/Source/Runtime/Core/Public/Containers/BitArray.h
--- a/Engine/Source/Runtime/Core/Public/Containers/BitArray.h
+++ b/Engine/Source/Runtime/Core/Public/Containers/BitArray.h
@@ -334,10 +334,10 @@ public:
 	}
 
 	/**
-	 * Adds bits to the array with the given value.
-	 * @return The index of the first added bit.
+	 * Adds a bit to the array with the given value.
+	 * @return The index of the added bit.
 	 */
-	int32 Add(const bool Value, const int32 Count = 1)
+	int32 Add(const bool Value)
 	{
 		const int32 Index = NumBits;
 
@@ -350,7 +350,7 @@ public:
 
 	/**
 	 * Adds multiple bits to the array with the given value.
-	 * @return The index of the added bit.
+	 * @return The index of the first added bit.
 	 */
 	int32 Add(const bool Value, int32 NumToAdd)
 	{
diff --git a/Engine/Source/Runtime/Core/Public/Templates/Function.h b/Engine/Source/Runtime/Core/Public/Templates/Function.h
--- a/Engine/Source/Runtime/Core/Public/Templates/Function.h
+++ b/Engine/Source/Runtime/Core/Public/Templates/Function.h
@@ -272,7 +272,7 @@ namespace UE4Function_Private
 			GetBoundObject()->~IFunction_OwnedObject();
 		}
 
-		FunctionAllocatorType::ForElementType<AlignedInlineFunctionType> Allocator;
+		FFunctionAllocatorType::ForElementType<FAlignedInlineFunctionType> Allocator;
 	};
 
 	struct FFunctionStorage : public FUniqueFunctionStorage
@@ -319,8 +319,8 @@ namespace UE4Function_Private
 
 inline void* operator new(size_t Size, UE4Function_Private::FUniqueFunctionStorage& Storage)
 {
-	int32 NewSize = FMath::DivideAndRoundUp(Size, sizeof(UE4Function_Private::AlignedInlineFunctionType));
-	Storage.Allocator.ResizeAllocation(0, NewSize, sizeof(UE4Function_Private::AlignedInlineFunctionType));
+	int32 NewSize = FMath::DivideAndRoundUp(Size, sizeof(UE4Function_Private::FAlignedInlineFunctionType));
+	Storage.Allocator.ResizeAllocation(0, NewSize, sizeof(UE4Function_Private::FAlignedInlineFunctionType));
 
 	return Storage.Allocator.GetAllocation();
 }
diff --git a/Engine/Source/Runtime/CoreUObject/Private/Serialization/Formatters/JsonArchiveOutputFormatter.cpp b/Engine/Source/Runtime/CoreUObject/Private/Serialization/Formatters/JsonArchiveOutputFormatter.cpp
--- a/Engine/Source/Runtime/CoreUObject/Private/Serialization/Formatters/JsonArchiveOutputFormatter.cpp
+++ b/Engine/Source/Runtime/CoreUObject/Private/Serialization/Formatters/JsonArchiveOutputFormatter.cpp
@@ -192,7 +192,7 @@ void FJsonArchiveOutputFormatter::Serialize(float& Value)
 		FString String = FString::Printf(TEXT("%.9g"), Value);
 #if DO_GUARD_SLOW
 		float RoundTripped;
-		Lex::FromString(RoundTripped, *String);
+		LexFromString(RoundTripped, *String);
 		check(RoundTripped == Value);
 #endif
 		WriteValue(String);
@@ -210,7 +210,7 @@ void FJsonArchiveOutputFormatter::Serialize(double& Value)
 		FString String = FString::Printf(TEXT("%.17g"), Value);
 #if DO_GUARD_SLOW
 		double RoundTripped;
-		Lex::FromString(RoundTripped, *String);
+		LexFromString(RoundTripped, *String);
 		check(RoundTripped == Value);
 #endif
 		WriteValue(FString::Printf(TEXT("%.17g"), Value));
