Fix for TBitArray::MaxBits in assignment.

diff --git a/Engine/Source/Runtime/Core/Public/Containers/BitArray.h b/Engine/Source/Runtime/Core/Public/Containers/BitArray.h
--- a/Engine/Source/Runtime/Core/Public/Containers/BitArray.h
+++ b/Engine/Source/Runtime/Core/Public/Containers/BitArray.h
@@ -226,7 +226,7 @@ public:
 		}
 
 		Empty(Copy.Num());
-		NumBits = MaxBits = Copy.NumBits;
+		NumBits = Copy.NumBits;
 		if(NumBits)
 		{
 			const int32 NumDWORDs = FMath::DivideAndRoundUp(MaxBits, NumBitsPerDWORD);
