GitHub #4762 : Improve wording and grammar of Math comments
Also includes improved accuracy in FMath::ComputeBoundingSphereForCone().

diff --git a/Engine/Source/Runtime/Core/Public/Math/InterpCurvePoint.h b/Engine/Source/Runtime/Core/Public/Math/InterpCurvePoint.h
--- a/Engine/Source/Runtime/Core/Public/Math/InterpCurvePoint.h
+++ b/Engine/Source/Runtime/Core/Public/Math/InterpCurvePoint.h
@@ -328,7 +328,7 @@ inline void ComputeCurveTangent( float PrevTime, const FTwoVectors& PrevPoint,
 }
 
 /**
- * Calculate bounds of float inervals
+ * Calculate bounds of float intervals
  *
  * @param Start interp curve point at Start
  * @param End interp curve point at End
diff --git a/Engine/Source/Runtime/Core/Public/Math/Matrix.h b/Engine/Source/Runtime/Core/Public/Math/Matrix.h
--- a/Engine/Source/Runtime/Core/Public/Math/Matrix.h
+++ b/Engine/Source/Runtime/Core/Public/Math/Matrix.h
@@ -99,7 +99,7 @@ public:
 	FORCEINLINE void operator+=(const FMatrix& Other);
 
 	/** 
-	  * This isn't applying SCALE, just multiplying float to all members - i.e. weighting
+	  * This isn't applying SCALE, just multiplying the value to all members - i.e. weighting
 	  */
 	FORCEINLINE FMatrix operator* (float Other) const;
 
diff --git a/Engine/Source/Runtime/Core/Public/Math/Sphere.h b/Engine/Source/Runtime/Core/Public/Math/Sphere.h
--- a/Engine/Source/Runtime/Core/Public/Math/Sphere.h
+++ b/Engine/Source/Runtime/Core/Public/Math/Sphere.h
@@ -195,7 +195,7 @@ FORCEINLINE bool FMath::SphereAABBIntersection(const FSphere& Sphere,const FBox&
 FORCEINLINE FSphere FMath::ComputeBoundingSphereForCone(FVector const& ConeOrigin, FVector const& ConeDirection, float ConeRadius, float CosConeAngle, float SinConeAngle)
 {
 	// Based on: https://bartwronski.com/2017/04/13/cull-that-cone/
-	const float COS_PI_OVER_4 = 0.707107f; // Cos(Pi/4);
+	const float COS_PI_OVER_4 = 0.70710678118f; // Cos(Pi/4);
 	if (CosConeAngle < COS_PI_OVER_4)
 	{
 		return FSphere(ConeOrigin + ConeDirection * ConeRadius * CosConeAngle, ConeRadius * SinConeAngle);
diff --git a/Engine/Source/Runtime/Core/Public/Math/TransformNonVectorized.h b/Engine/Source/Runtime/Core/Public/Math/TransformNonVectorized.h
--- a/Engine/Source/Runtime/Core/Public/Math/TransformNonVectorized.h
+++ b/Engine/Source/Runtime/Core/Public/Math/TransformNonVectorized.h
@@ -591,7 +591,7 @@ public:
 		return Rotation.Rotator();
 	}
 
-	/** Calculate the  */
+	/** Calculate the determinant of this transformation */
 	FORCEINLINE float GetDeterminant() const
 	{
 		return Scale3D.X * Scale3D.Y * Scale3D.Z;
diff --git a/Engine/Source/Runtime/Core/Public/Math/TransformVectorized.h b/Engine/Source/Runtime/Core/Public/Math/TransformVectorized.h
--- a/Engine/Source/Runtime/Core/Public/Math/TransformVectorized.h
+++ b/Engine/Source/Runtime/Core/Public/Math/TransformVectorized.h
@@ -1437,7 +1437,7 @@ FORCEINLINE void FTransform::ScaleTranslation(const FVector& InScale3D)
 	DiagnosticCheckNaN_Translate();
 }
 
-/** Scale the translation part of the Transform by the supplied float. */
+/** Scale the translation part of the Transform by the supplied value. */
 FORCEINLINE void FTransform::ScaleTranslation(const float& InScale)
 {
 	ScaleTranslation( FVector(InScale) );
diff --git a/Engine/Source/Runtime/Core/Public/Math/UnrealMathUtility.h b/Engine/Source/Runtime/Core/Public/Math/UnrealMathUtility.h
--- a/Engine/Source/Runtime/Core/Public/Math/UnrealMathUtility.h
+++ b/Engine/Source/Runtime/Core/Public/Math/UnrealMathUtility.h
@@ -324,7 +324,7 @@ struct FMath : public FPlatformMath
 	}
 
 	/**
-	* Computes the sine and cosine of a scalar float.
+	* Computes the sine and cosine of a scalar value.
 	*
 	* @param ScalarSin	Pointer to where the Sin result should be stored
 	* @param ScalarCos	Pointer to where the Cos result should be stored
@@ -379,7 +379,7 @@ struct FMath : public FPlatformMath
 	//		FASTASIN_HALF_PI	== 1.5707963050f  == 0x3fC90FDA
 #define FASTASIN_HALF_PI (1.5707963050f)
 	/**
-	* Computes the ASin of a scalar float.
+	* Computes the ASin of a scalar value.
 	*
 	* @param Value  input angle
 	* @return ASin of Value
diff --git a/Engine/Source/Runtime/Core/Public/Math/Vector2D.h b/Engine/Source/Runtime/Core/Public/Math/Vector2D.h
--- a/Engine/Source/Runtime/Core/Public/Math/Vector2D.h
+++ b/Engine/Source/Runtime/Core/Public/Math/Vector2D.h
@@ -107,18 +107,18 @@ public:
 	FVector2D operator/(float Scale) const;
 
 	/**
-	 * Gets the result of this vector + float A.
+	 * Gets the result of adding A to each component of the vector.
 	 *
 	 * @param A Float to add to each component.
-	 * @return The result of this vector + float A.
+	 * @return The result of adding A to each component.
 	 */
 	FORCEINLINE FVector2D operator+(float A) const;
 
 	/**
-	 * Gets the result of subtracting from each component of the vector.
+	 * Gets the result of subtracting A from each component of the vector.
 	 *
 	 * @param A Float to subtract from each component
-	 * @return The result of this vector - float A.
+	 * @return The result of subtracting A from each component.
 	 */
 	FORCEINLINE FVector2D operator-(float A) const;
 
