Don't call CreateDirectory redundantly when opening files for writing using FFileManagerGeneric::CreateFileWriter

This change avoids calling IPlatformFile::CreateDirectoryTree if possible since this is a very expensive function especially for deep hierarchies as it performs directory creation from the root directory onwards instead of from the leaf downwards. That function should also be fixed but this change improves performance in the meantime.

diff --git a/Engine/Source/Runtime/Core/Private/HAL/FileManagerGeneric.cpp b/Engine/Source/Runtime/Core/Private/HAL/FileManagerGeneric.cpp
--- a/Engine/Source/Runtime/Core/Private/HAL/FileManagerGeneric.cpp
+++ b/Engine/Source/Runtime/Core/Private/HAL/FileManagerGeneric.cpp
@@ -99,7 +99,6 @@ FArchive* FFileManagerGeneric::CreateFileWriter( const TCHAR* Filename, uint32 F
 		UE_LOG( LogFileManager, Log, TEXT( "Can't write to signed game file: %s" ),Filename );
 		return new FArchiveFileWriterDummy();
 	}
-	MakeDirectory( *FPaths::GetPath(Filename), true );
 
 	if( Flags & FILEWRITE_EvenIfReadOnly )
 	{
@@ -107,6 +106,14 @@ FArchive* FFileManagerGeneric::CreateFileWriter( const TCHAR* Filename, uint32 F
 	}
 
 	IFileHandle* Handle = GetLowLevel().OpenWrite( Filename, !!( Flags & FILEWRITE_Append ), !!( Flags & FILEWRITE_AllowRead ) );
+
+	if (!Handle)
+	{
+		MakeDirectory(*FPaths::GetPath(Filename), true);
+
+		Handle = GetLowLevel().OpenWrite(Filename, !!(Flags & FILEWRITE_Append), !!(Flags & FILEWRITE_AllowRead));
+	}
+
 	if( !Handle )
 	{
 		if( Flags & FILEWRITE_NoFail )
