CreateDirectoryTree now creates directories leaf-to-root instead of the other way around. This is much more efficient since we don't spend time on system API calls for directories which already exist. This accounted for a very large amount of CPU time in cooking as the full target file directory hierarchy would be "created" for every single output file.

diff --git a/Engine/Source/Runtime/Core/Private/GenericPlatform/GenericPlatformFile.cpp b/Engine/Source/Runtime/Core/Private/GenericPlatform/GenericPlatformFile.cpp
--- a/Engine/Source/Runtime/Core/Private/GenericPlatform/GenericPlatformFile.cpp
+++ b/Engine/Source/Runtime/Core/Private/GenericPlatform/GenericPlatformFile.cpp
@@ -737,34 +737,35 @@ FString IPlatformFile::ConvertToAbsolutePathForExternalAppForWrite( const TCHAR*
 	return FPaths::ConvertRelativePathToFull(Filename);
 }
 
-bool IPlatformFile::CreateDirectoryTree(const TCHAR* Directory)
+static bool InternalCreateDirectoryTree(IPlatformFile& Ipf, const FString& Directory)
 {
-	FString LocalFilename(Directory);
-	FPaths::NormalizeDirectoryName(LocalFilename);
-	const TCHAR* LocalPath = *LocalFilename;
-	int32 CreateCount = 0;
-	const int32 MaxCharacters = MAX_UNREAL_FILENAME_LENGTH - 1;
-	int32 Index = 0;
-	for (TCHAR Full[MaxCharacters + 1] = TEXT( "" ), *Ptr = Full; Index < MaxCharacters; *Ptr++ = *LocalPath++, Index++)
-	{
-		if (((*LocalPath) == TEXT('/')) || (*LocalPath== 0))
-		{
-			*Ptr = 0;
-			if ((Ptr != Full) && !FPaths::IsDrive( Full ))
-			{
-				if (!CreateDirectory(Full) && !DirectoryExists(Full))
-				{
-					break;
-				}
-				CreateCount++;
-			}
-		}
-		if (*LocalPath == 0)
+	// Just try creating directory first
+	if (Ipf.CreateDirectory(*Directory))
+		return true;
+
+	// If it fails, try creating parent(s), before attempting to create directory
+	// once again.
+	int32 SeparatorIndex = -1;
+	if (Directory.FindLastChar(TEXT('/'), /* out */ SeparatorIndex))
+	{
+		if (SeparatorIndex)
 		{
-			break;
+			if (!InternalCreateDirectoryTree(Ipf, Directory.Left(SeparatorIndex)))
+				return false;
+
+			return Ipf.CreateDirectory(*Directory) || Ipf.DirectoryExists(*Directory);
 		}
 	}
-	return DirectoryExists(*LocalFilename);
+
+	return Ipf.DirectoryExists(*Directory);
+}
+
+bool IPlatformFile::CreateDirectoryTree(const TCHAR* Directory)
+{
+	FString LocalDirname(Directory);
+	FPaths::NormalizeDirectoryName(LocalDirname);
+
+	return InternalCreateDirectoryTree(*this, LocalDirname);
 }
 
 bool IPhysicalPlatformFile::Initialize(IPlatformFile* Inner, const TCHAR* CmdLine)
