- Disable overlapped I/O in editor / cooker. Synchronous I/O reduces the number of syscalls and Windows performs prefetching on our behalf anyway for sequential reads
- Eliminated syscall which was issued for every write to update cached file size -- since we're the only writers to the file (file access allows read sharing at most) we can authoritatively update the file size on write completion

diff --git a/Engine/Source/Runtime/Core/Private/Windows/WindowsPlatformFile.cpp b/Engine/Source/Runtime/Core/Private/Windows/WindowsPlatformFile.cpp
--- a/Engine/Source/Runtime/Core/Private/Windows/WindowsPlatformFile.cpp
+++ b/Engine/Source/Runtime/Core/Private/Windows/WindowsPlatformFile.cpp
@@ -560,7 +560,7 @@ public:
 		// Update where we are in the file
 		FilePos += NumWritten;
 		UpdateOverlappedPos();
-		UpdateFileSize();
+		FileSize = FMath::Max(FilePos, FileSize);
 		return true;
 	}
 };
@@ -737,7 +737,7 @@ public:
 		uint32  Access    = GENERIC_READ;
 		uint32  WinFlags  = FILE_SHARE_READ | (bAllowWrite ? FILE_SHARE_WRITE : 0);
 		uint32  Create    = OPEN_EXISTING;
-#define USE_OVERLAPPED_IO 1
+#define USE_OVERLAPPED_IO (!IS_PROGRAM && !WITH_EDITOR)		// Use straightforward synchronous I/O in cooker/editor
 
 #if USE_OVERLAPPED_IO
 		HANDLE Handle    = CreateFileW(*NormalizeFilename(Filename), Access, WinFlags, NULL, Create, FILE_ATTRIBUTE_NORMAL | FILE_FLAG_OVERLAPPED, NULL);
