In IPlatformFile::DeleteDirectoryRecursively, attempt to delete file first and if it fails clear the readonly flag and try again

Previously there was a call to clear the readonly flag for every deleted file and this is a waste of resources 99% of the time. The SetFileAttributes call accounted for a significant amount of time during cooker sandbox directory deletion

diff --git a/Engine/Source/Runtime/Core/Private/GenericPlatform/GenericPlatformFile.cpp b/Engine/Source/Runtime/Core/Private/GenericPlatform/GenericPlatformFile.cpp
--- a/Engine/Source/Runtime/Core/Private/GenericPlatform/GenericPlatformFile.cpp
+++ b/Engine/Source/Runtime/Core/Private/GenericPlatform/GenericPlatformFile.cpp
@@ -599,6 +599,10 @@ bool IPlatformFile::DeleteDirectoryRecursively(const TCHAR* Directory)
 			}
 			else
 			{
+				if (PlatformFile.DeleteFile(FilenameOrDirectory))
+					return true;
+
+				// File delete failed -- unset readonly flag and try again
 				PlatformFile.SetReadOnly(FilenameOrDirectory, false);
 				PlatformFile.DeleteFile(FilenameOrDirectory);
 			}
