Some tweaks to FQueuedThreadPoolBase scheduling and memory management

- Explicitly pass in false for TArray::RemoveAt(..., bool bAllowShrinking) argument to prevent memory reallocation when arrays are drained and inevitably repopulated shortly afterwards
- Use a MRU strategy instead of LRU when picking a thread to wake up. The MRU thread is the most likely to have a 'hot' cache for the stack etc. Picking from the back of the array also happens to be cheaper since
no memory movement is necessary when RemoveAt is called. (This was the strategy in place before CL2600362 which seems to have changed it unintentionally)
- Release lock as soon as a thread has been chosen, before asking the worker thread to wake up and do the work

diff --git a/Engine/Source/Runtime/Core/Private/HAL/ThreadingBase.cpp b/Engine/Source/Runtime/Core/Private/HAL/ThreadingBase.cpp
--- a/Engine/Source/Runtime/Core/Private/HAL/ThreadingBase.cpp
+++ b/Engine/Source/Runtime/Core/Private/HAL/ThreadingBase.cpp
@@ -674,38 +674,48 @@ public:
 	}
 	void AddQueuedWork(IQueuedWork* InQueuedWork) override
 	{
+		check(InQueuedWork != nullptr);
+
 		if (TimeToDie)
 		{
 			InQueuedWork->Abandon();
 			return;
 		}
-		check(InQueuedWork != nullptr);
-		FQueuedThread* Thread = nullptr;
+
 		// Check to see if a thread is available. Make sure no other threads
 		// can manipulate the thread pool while we do this.
+		//
+		// We pick a thread from the back of the array since this will be the
+		// most recently used thread and therefore the most likely to have
+		// a 'hot' cache for the stack etc (similar to Windows IOCP scheduling
+		// strategy). Picking from the back also happens to be cheaper since
+		// no memory movement is necessary.
+
 		check(SynchQueue);
-		FScopeLock sl(SynchQueue);
-		if (QueuedThreads.Num() > 0)
+
+		FQueuedThread* Thread = nullptr;
+
 		{
-			// Cycle through all available threads to make sure that stats are up to date.
-			int32 Index = 0;
-			// Grab that thread to use
-			Thread = QueuedThreads[Index];
+			FScopeLock sl(SynchQueue);
+			const int32 AvailableThreadCount = QueuedThreads.Num();
+			if (AvailableThreadCount == 0)
+			{
+				// No thread available, queue the work to be done
+				// as soon as one does become available
+				QueuedWork.Add(InQueuedWork);
+
+				return;
+			}
+
+			const int32 ThreadIndex = AvailableThreadCount - 1;
+
+			Thread = QueuedThreads[ThreadIndex];
 			// Remove it from the list so no one else grabs it
-			QueuedThreads.RemoveAt(Index);
-		}
-		// Was there a thread ready?
-		if (Thread != nullptr)
-		{
-			// We have a thread, so tell it to do the work
-			Thread->DoWork(InQueuedWork);
-		}
-		else
-		{
-			// There were no threads available, queue the work to be done
-			// as soon as one does become available
-			QueuedWork.Add(InQueuedWork);
+			QueuedThreads.RemoveAt(ThreadIndex, 1, /* do not allow shrinking */ false);
 		}
+
+		// Tell our chosen thread to do the work
+		Thread->DoWork(InQueuedWork);
 	}
 
 	virtual bool RetractQueuedWork(IQueuedWork* InQueuedWork) override
@@ -737,7 +747,7 @@ public:
 			// queued and never done
 			Work = QueuedWork[0];
 			// Remove it from the list so no one else grabs it
-			QueuedWork.RemoveAt(0);
+			QueuedWork.RemoveAt(0, 1, /* do not allow shrinking */ false);
 		}
 		if (!Work)
 		{
