TBA: Add EnterRecord overload that allows outputting of available field names when loading.

diff --git a/Engine/Source/Runtime/Core/Private/Serialization/StructuredArchive.cpp b/Engine/Source/Runtime/Core/Private/Serialization/StructuredArchive.cpp
--- a/Engine/Source/Runtime/Core/Private/Serialization/StructuredArchive.cpp
+++ b/Engine/Source/Runtime/Core/Private/Serialization/StructuredArchive.cpp
@@ -224,6 +224,19 @@ FStructuredArchive::FRecord FStructuredArchive::FSlot::EnterRecord()
 	return FRecord(Ar, Depth + 1, ElementId);
 }
 
+FStructuredArchive::FRecord FStructuredArchive::FSlot::EnterRecord(TArray<FString>& OutFieldNamesWhenLoading)
+{
+	Ar.EnterSlot(ElementId, EElementType::Record);
+
+#if DO_GUARD_SLOW
+	Ar.CurrentContainer.Add(new FContainer(0));
+#endif
+
+	Ar.Formatter.EnterRecord(OutFieldNamesWhenLoading);
+
+	return FRecord(Ar, Depth + 1, ElementId);
+}
+
 FStructuredArchive::FArray FStructuredArchive::FSlot::EnterArray(int32& Num)
 {
 	Ar.EnterSlot(ElementId, EElementType::Array);
@@ -446,6 +459,11 @@ FStructuredArchive::FRecord FStructuredArchive::FRecord::EnterRecord(FArchiveFie
 	return EnterField(Name).EnterRecord();
 }
 
+FStructuredArchive::FRecord FStructuredArchive::FRecord::EnterRecord(FArchiveFieldName Name, TArray<FString>& OutFieldNamesWhenLoading)
+{
+	return EnterField(Name).EnterRecord(OutFieldNamesWhenLoading);
+}
+
 FStructuredArchive::FArray FStructuredArchive::FRecord::EnterArray(FArchiveFieldName Name, int32& Num)
 {
 	return EnterField(Name).EnterArray(Num);
diff --git a/Engine/Source/Runtime/Core/Public/Serialization/Formatters/BinaryArchiveFormatter.h b/Engine/Source/Runtime/Core/Public/Serialization/Formatters/BinaryArchiveFormatter.h
--- a/Engine/Source/Runtime/Core/Public/Serialization/Formatters/BinaryArchiveFormatter.h
+++ b/Engine/Source/Runtime/Core/Public/Serialization/Formatters/BinaryArchiveFormatter.h
@@ -16,6 +16,7 @@ public:
 	virtual FArchive& GetUnderlyingArchive() override;
 
 	virtual void EnterRecord() override;
+	virtual void EnterRecord(TArray<FString>& OutFieldNamesWhenLoading) override;
 	virtual void LeaveRecord() override;
 	virtual void EnterField(FArchiveFieldName Name) override;
 	virtual void LeaveField() override;
@@ -72,6 +73,10 @@ inline void FBinaryArchiveFormatter::EnterRecord()
 {
 }
 
+inline void FBinaryArchiveFormatter::EnterRecord(TArray<FString>& OutFieldNamesWhenLoading)
+{
+}
+
 inline void FBinaryArchiveFormatter::LeaveRecord()
 {
 }
diff --git a/Engine/Source/Runtime/Core/Public/Serialization/StructuredArchive.h b/Engine/Source/Runtime/Core/Public/Serialization/StructuredArchive.h
--- a/Engine/Source/Runtime/Core/Public/Serialization/StructuredArchive.h
+++ b/Engine/Source/Runtime/Core/Public/Serialization/StructuredArchive.h
@@ -106,6 +106,7 @@ public:
 	{
 	public:
 		FRecord EnterRecord();
+		FRecord EnterRecord(TArray<FString>& OutFieldNamesWhenLoading);
 		FArray EnterArray(int32& Num);
 		FStream EnterStream();
 		FMap EnterMap(int32& Num);
@@ -181,6 +182,7 @@ public:
 	public:
 		FSlot EnterField(FArchiveFieldName Name);
 		FRecord EnterRecord(FArchiveFieldName Name);
+		FRecord EnterRecord(FArchiveFieldName Name, TArray<FString>& OutFieldNamesWhenLoading);
 		FArray EnterArray(FArchiveFieldName Name, int32& Num);
 		FStream EnterStream(FArchiveFieldName Name);
 		FMap EnterMap(FArchiveFieldName Name, int32& Num);
@@ -518,6 +520,11 @@ private:
 		return FRecord(Ar);
 	}
 
+	FORCEINLINE FStructuredArchive::FRecord FStructuredArchive::FSlot::EnterRecord(TArray<FString>& OutFieldNamesWhenLoading)
+	{
+		return FRecord(Ar);
+	}
+
 	FORCEINLINE FStructuredArchive::FArray FStructuredArchive::FSlot::EnterArray(int32& Num)
 	{
 		Ar.Formatter.EnterArray(Num);
@@ -654,11 +661,6 @@ private:
 		return FSlot(Ar);
 	}
 
-	FORCEINLINE FStructuredArchive::FRecord FStructuredArchive::FRecord::EnterRecord(FArchiveFieldName Name)
-	{
-		return EnterField(Name).EnterRecord();
-	}
-
 	FORCEINLINE FStructuredArchive::FArray FStructuredArchive::FRecord::EnterArray(FArchiveFieldName Name, int32& Num)
 	{
 		return EnterField(Name).EnterArray(Num);
diff --git a/Engine/Source/Runtime/Core/Public/Serialization/StructuredArchiveFormatter.h b/Engine/Source/Runtime/Core/Public/Serialization/StructuredArchiveFormatter.h
--- a/Engine/Source/Runtime/Core/Public/Serialization/StructuredArchiveFormatter.h
+++ b/Engine/Source/Runtime/Core/Public/Serialization/StructuredArchiveFormatter.h
@@ -73,6 +73,7 @@ public:
 	virtual FStructuredArchiveFormatter* CreateSubtreeReader() { return this; }
 
 	virtual void EnterRecord() = 0;
+	virtual void EnterRecord(TArray<FString>& OutFieldNamesWhenLoading) = 0;
 	virtual void LeaveRecord() = 0;
 	virtual void EnterField(FArchiveFieldName Name) = 0;
 	virtual void LeaveField() = 0;
diff --git a/Engine/Source/Runtime/CoreUObject/Private/Serialization/Formatters/JsonArchiveInputFormatter.cpp b/Engine/Source/Runtime/CoreUObject/Private/Serialization/Formatters/JsonArchiveInputFormatter.cpp
--- a/Engine/Source/Runtime/CoreUObject/Private/Serialization/Formatters/JsonArchiveInputFormatter.cpp
+++ b/Engine/Source/Runtime/CoreUObject/Private/Serialization/Formatters/JsonArchiveInputFormatter.cpp
@@ -52,13 +52,13 @@ void FJsonArchiveInputFormatter::EnterRecord()
 	ObjectStack.Add(FObjectRecord(Value->AsObject(), ValueStack.Num()));
 }
 
-void FJsonArchiveInputFormatter::EnterRecord(TArray<FString>& OutKeys)
+void FJsonArchiveInputFormatter::EnterRecord(TArray<FString>& OutFieldNamesWhenLoading)
 {
 	EnterRecord();
-	ObjectStack.Top().JsonObject->Values.GetKeys(OutKeys);
-	for(int32 Idx = 0; Idx < OutKeys.Num(); Idx++)
+	ObjectStack.Top().JsonObject->Values.GetKeys(OutFieldNamesWhenLoading);
+	for(int32 Idx = 0; Idx < OutFieldNamesWhenLoading.Num(); Idx++)
 	{
-		OutKeys[Idx] = UnescapeFieldName(*OutKeys[Idx]);
+		OutFieldNamesWhenLoading[Idx] = UnescapeFieldName(*OutFieldNamesWhenLoading[Idx]);
 	}
 }
 
diff --git a/Engine/Source/Runtime/CoreUObject/Private/Serialization/Formatters/JsonArchiveOutputFormatter.cpp b/Engine/Source/Runtime/CoreUObject/Private/Serialization/Formatters/JsonArchiveOutputFormatter.cpp
--- a/Engine/Source/Runtime/CoreUObject/Private/Serialization/Formatters/JsonArchiveOutputFormatter.cpp
+++ b/Engine/Source/Runtime/CoreUObject/Private/Serialization/Formatters/JsonArchiveOutputFormatter.cpp
@@ -40,6 +40,11 @@ void FJsonArchiveOutputFormatter::EnterRecord()
 	bNeedsNewline = true;
 }
 
+void FJsonArchiveOutputFormatter::EnterRecord(TArray<FString>& OutFieldNamesWhenLoading)
+{
+	EnterRecord();
+}
+
 void FJsonArchiveOutputFormatter::LeaveRecord()
 {
 	Newline.Pop(false);
diff --git a/Engine/Source/Runtime/CoreUObject/Private/Serialization/Formatters/TaggedBinaryArchiveOutputFormatter.cpp b/Engine/Source/Runtime/CoreUObject/Private/Serialization/Formatters/TaggedBinaryArchiveOutputFormatter.cpp
--- a/Engine/Source/Runtime/CoreUObject/Private/Serialization/Formatters/TaggedBinaryArchiveOutputFormatter.cpp
+++ b/Engine/Source/Runtime/CoreUObject/Private/Serialization/Formatters/TaggedBinaryArchiveOutputFormatter.cpp
@@ -68,6 +68,11 @@ void FTaggedBinaryArchiveOutputFormatter::EnterRecord()
 	RecordStack.Push(NextRecordIdx++);
 }
 
+void FTaggedBinaryArchiveOutputFormatter::EnterRecord(TArray<FString>& OutFieldNames)
+{
+	EnterRecord();
+}
+
 void FTaggedBinaryArchiveOutputFormatter::LeaveRecord()
 {
 	FRecord& Record = Records[RecordStack.Top()];
diff --git a/Engine/Source/Runtime/CoreUObject/Public/Serialization/Formatters/JsonArchiveInputFormatter.h b/Engine/Source/Runtime/CoreUObject/Public/Serialization/Formatters/JsonArchiveInputFormatter.h
--- a/Engine/Source/Runtime/CoreUObject/Public/Serialization/Formatters/JsonArchiveInputFormatter.h
+++ b/Engine/Source/Runtime/CoreUObject/Public/Serialization/Formatters/JsonArchiveInputFormatter.h
@@ -20,7 +20,7 @@ public:
 	virtual FStructuredArchiveFormatter* CreateSubtreeReader() override;
 
 	virtual void EnterRecord() override;
-	virtual void EnterRecord(TArray<FString>& OutKeys) override;
+	virtual void EnterRecord(TArray<FString>& OutFieldNamesWhenLoading) override;
 	virtual void LeaveRecord() override;
 	virtual void EnterField(FArchiveFieldName Name) override;
 	virtual void EnterField(FArchiveFieldName Name, EArchiveValueType& OutType) override;
diff --git a/Engine/Source/Runtime/CoreUObject/Public/Serialization/Formatters/JsonArchiveOutputFormatter.h b/Engine/Source/Runtime/CoreUObject/Public/Serialization/Formatters/JsonArchiveOutputFormatter.h
--- a/Engine/Source/Runtime/CoreUObject/Public/Serialization/Formatters/JsonArchiveOutputFormatter.h
+++ b/Engine/Source/Runtime/CoreUObject/Public/Serialization/Formatters/JsonArchiveOutputFormatter.h
@@ -16,6 +16,7 @@ public:
 	virtual FArchive& GetUnderlyingArchive() override;
 
 	virtual void EnterRecord() override;
+	virtual void EnterRecord(TArray<FString>& OutFieldNamesWhenLoading) override;
 	virtual void LeaveRecord() override;
 	virtual void EnterField(FArchiveFieldName Name) override;
 	virtual void LeaveField() override;
diff --git a/Engine/Source/Runtime/CoreUObject/Public/Serialization/Formatters/TaggedBinaryArchiveOutputFormatter.h b/Engine/Source/Runtime/CoreUObject/Public/Serialization/Formatters/TaggedBinaryArchiveOutputFormatter.h
--- a/Engine/Source/Runtime/CoreUObject/Public/Serialization/Formatters/TaggedBinaryArchiveOutputFormatter.h
+++ b/Engine/Source/Runtime/CoreUObject/Public/Serialization/Formatters/TaggedBinaryArchiveOutputFormatter.h
@@ -16,6 +16,7 @@ public:
 	virtual FArchive& GetUnderlyingArchive() override;
 
 	virtual void EnterRecord() override;
+	virtual void EnterRecord(TArray<FString>& OutFieldNames) override;
 	virtual void LeaveRecord() override;
 	virtual void EnterField(FArchiveFieldName Name) override;
 	virtual void LeaveField() override;
