TBA: Build fix for textasset-less configs

diff --git a/Engine/Source/Runtime/Core/Public/Serialization/StructuredArchive.h b/Engine/Source/Runtime/Core/Public/Serialization/StructuredArchive.h
--- a/Engine/Source/Runtime/Core/Public/Serialization/StructuredArchive.h
+++ b/Engine/Source/Runtime/Core/Public/Serialization/StructuredArchive.h
@@ -661,6 +661,16 @@ private:
 		return FSlot(Ar);
 	}
 
+	FORCEINLINE FStructuredArchive::FRecord FStructuredArchive::FRecord::EnterRecord(FArchiveFieldName Name)
+	{
+		return EnterField(Name).EnterRecord();
+	}
+
+	FORCEINLINE FStructuredArchive::FRecord FStructuredArchive::FRecord::EnterRecord(FArchiveFieldName Name, TArray<FString>& OutFieldNamesWhenLoading)
+	{
+		return EnterField(Name).EnterRecord(OutFieldNamesWhenLoading);
+	}
+
 	FORCEINLINE FStructuredArchive::FArray FStructuredArchive::FRecord::EnterArray(FArchiveFieldName Name, int32& Num)
 	{
 		return EnterField(Name).EnterArray(Num);
