GitHub #3631 : Incorrect default argument in WeakObjectPtrTemplate

diff --git a/Engine/Source/Runtime/Core/Public/UObject/WeakObjectPtrTemplates.h b/Engine/Source/Runtime/Core/Public/UObject/WeakObjectPtrTemplates.h
--- a/Engine/Source/Runtime/Core/Public/UObject/WeakObjectPtrTemplates.h
+++ b/Engine/Source/Runtime/Core/Public/UObject/WeakObjectPtrTemplates.h
@@ -215,7 +215,7 @@ public:
 	 * @param bThreadsafeTest, set it to true when testing outside of Game Thread. Results in false if WeakObjPtr point to an existing object (no flags checked) 
 	 * @return true if this used to point at a real object but no longer does.
 	**/
-	FORCEINLINE bool IsStale(bool bIncludingIfPendingKill = false, bool bThreadsafeTest = false) const
+	FORCEINLINE bool IsStale(bool bIncludingIfPendingKill = true, bool bThreadsafeTest = false) const
 	{
 		return TWeakObjectPtrBase::IsStale(bIncludingIfPendingKill, bThreadsafeTest);
 	}
