PROJECTION and PROJECTION_MEMBER macros which provide the correct behavior when creating projections using functions which are overloaded or use default arguments.

diff --git a/Engine/Source/Runtime/Core/Public/Templates/Invoke.h b/Engine/Source/Runtime/Core/Public/Templates/Invoke.h
--- a/Engine/Source/Runtime/Core/Public/Templates/Invoke.h
+++ b/Engine/Source/Runtime/Core/Public/Templates/Invoke.h
@@ -65,3 +65,45 @@ FORCEINLINE auto Invoke(ReturnType (ObjType::*PtrMemFun)(PMFArgTypes...) const,
 {
 	return (UE4Invoke_Private::DereferenceIfNecessary<ObjType>(Forward<CallableType>(Callable)).*PtrMemFun)(Forward<ArgTypes>(Args)...);
 }
+
+
+/**
+ * Wraps up a named non-member function so that it can easily be passed as a callable.
+ * This allows functions with overloads or default arguments to be treated correctly.
+ *
+ * Example:
+ *
+ * TArray<FMyType> Array = ...;
+ *
+ * // Doesn't compile, because you can't take the address of an overloaded function when its type needs to be deduced. 
+ * Algo::SortBy(Array, &LexToString);
+ *
+ * // Works as expected
+ * Algo::SortBy(Array, PROJECTION(LexToString));
+ */
+#define PROJECTION(FuncName) \
+	[](auto&&... Args) \
+	{ \
+		return FuncName(Forward<decltype(Args)>(Args)...); \
+	}
+
+/**
+ * Wraps up a named member function so that it can easily be passed as a callable.
+ * This allows functions with overloads or default arguments to be treated correctly.
+ *
+ * Example:
+ *
+ * TArray<UObject*> Array = ...;
+ *
+ * // Doesn't compile, because &UObject::GetFullName loses the default argument and passes
+ * // FString (UObject::*)(const UObject*) to Algo::SortBy<>(), which is not a valid projection.
+ * Algo::SortBy(Array, &UObject::GetFullName);
+ *
+ * // Works as expected
+ * Algo::SortBy(Array, PROJECTION_MEMBER(UObject, GetFullName));
+ */
+#define PROJECTION_MEMBER(Type, FuncName) \
+	[](auto&& Obj, auto&&... Args) \
+	{ \
+		return UE4Invoke_Private::DereferenceIfNecessary<Type>(Forward<decltype(Obj)>(Obj)).FuncName(Forward<decltype(Args)>(Args)...); \
+	}
\ No newline at end of file
