Few more fixes for compiling the editor with Clang.

diff --git a/Engine/Source/Runtime/Core/Private/ProfilingDebugging/Windows/VSPerfExternalProfiler.cpp b/Engine/Source/Runtime/Core/Private/ProfilingDebugging/Windows/VSPerfExternalProfiler.cpp
--- a/Engine/Source/Runtime/Core/Private/ProfilingDebugging/Windows/VSPerfExternalProfiler.cpp
+++ b/Engine/Source/Runtime/Core/Private/ProfilingDebugging/Windows/VSPerfExternalProfiler.cpp
@@ -11,10 +11,18 @@
 // Not all versions of Visual Studio include the profiler SDK headers
 #if WITH_VS_PERF_PROFILER && UE_EXTERNAL_PROFILING_ENABLED
 
+#ifdef __clang__
+#pragma clang diagnostic push
+#pragma clang diagnostic ignored "-Wextra-tokens"
+#endif
+
 #define VSPERF_NO_DEFAULT_LIB	// Don't use #pragma lib to import the library, we'll handle this stuff ourselves
 #define PROFILERAPI				// We won't be statically importing anything (we're dynamically binding), so define PROFILERAPI to a empty value
 #include "VSPerf.h"				// NOTE: This header is in <Visual Studio install directory>/Team Tools/Performance Tools/x64/PerfSDK
 
+#ifdef __clang__
+#pragma clang diagnostic pop
+#endif
 
 /**
  * Visual Studio Profiler implementation of FExternalProfiler
