Improved readability and compliance with coding standards.

diff --git a/Engine/Source/Runtime/Core/Private/Misc/Parse.cpp b/Engine/Source/Runtime/Core/Private/Misc/Parse.cpp
--- a/Engine/Source/Runtime/Core/Private/Misc/Parse.cpp
+++ b/Engine/Source/Runtime/Core/Private/Misc/Parse.cpp
@@ -941,164 +941,186 @@ bool FParse::AlnumToken(const TCHAR*& Str, FString& Arg)
 // Get a line of Stream (everything up to, but not including, CR/LF.
 // Returns 0 if ok, nonzero if at end of stream and returned 0-length string.
 //
-bool FParse::Line
-(
-	const TCHAR**	Stream,
-	TCHAR*			Result,
-	int32				MaxLen,
-	bool			Exact
-)
+bool FParse::Line(const TCHAR** Stream, TCHAR* Result, int32 MaxLen, bool bExact)
 {
-	bool GotStream=0;
-	bool IsQuoted=0;
-	bool Ignore=0;
+	bool bGotStream = false;
+	bool bIsQuoted = false;
+	bool bIgnore = false;
 
 	*Result=0;
-	while( **Stream!=0 && **Stream!=10 && **Stream!=13 && --MaxLen>0 )
+	while (**Stream != TEXT('\0') && **Stream != TEXT('\n') && **Stream != TEXT('\r') && --MaxLen > 0)
 	{
 		// Start of comments.
-		if( !IsQuoted && !Exact && (*Stream)[0]=='/' && (*Stream)[1]=='/' )
-			Ignore = 1;
+		if (!bIsQuoted && !bExact && (*Stream)[0]=='/' && (*Stream)[1] == TEXT('/'))
+		{
+			bIgnore = true;
+		}
 		
 		// Command chaining.
-		if( !IsQuoted && !Exact && **Stream=='|' )
+		if (!bIsQuoted && !bExact && **Stream == TEXT('|'))
+		{
 			break;
+		}
 
 		// Check quoting.
-		IsQuoted = IsQuoted ^ (**Stream==34);
-		GotStream=1;
+		bIsQuoted = bIsQuoted ^ (**Stream == TEXT('\"'));
+		bGotStream = true;
 
 		// Got stuff.
-		if( !Ignore )
+		if (!bIgnore)
+		{
 			*(Result++) = *((*Stream)++);
+		}
 		else
+		{
 			(*Stream)++;
+		}
 	}
-	if( Exact )
+
+	if (bExact)
 	{
 		// Eat up exactly one CR/LF.
-		if( **Stream == 13 )
+		if (**Stream == TEXT('\r'))
+		{
 			(*Stream)++;
-		if( **Stream == 10 )
+		}
+
+		if (**Stream == TEXT('\n'))
+		{
 			(*Stream)++;
+		}
 	}
 	else
 	{
 		// Eat up all CR/LF's.
-		while( **Stream==10 || **Stream==13 || **Stream=='|' )
+		while (**Stream == TEXT('\n') || **Stream == TEXT('\r') || **Stream == TEXT('|'))
+		{
 			(*Stream)++;
+		}
 	}
-	*Result=0;
-	return **Stream!=0 || GotStream;
+
+	*Result = TEXT('\0');
+	return **Stream != TEXT('\0') || bGotStream;
 }
-bool FParse::Line
-(
-	const TCHAR**	Stream,
-	FString&		Result,
-	bool			Exact
-)
+
+bool FParse::Line(const TCHAR** Stream, FString& Result, bool bExact)
 {
-	bool GotStream=0;
-	bool IsQuoted=0;
-	bool Ignore=0;
+	bool bGotStream = false;
+	bool bIsQuoted = false;
+	bool bIgnore = false;
 
 	Result = TEXT("");
 
-	while( **Stream!=0 && **Stream!=10 && **Stream!=13 )
+	while (**Stream != TEXT('\0') && **Stream != TEXT('\n') && **Stream != TEXT('\r'))
 	{
 		// Start of comments.
-		if( !IsQuoted && !Exact && (*Stream)[0]=='/' && (*Stream)[1]=='/' )
-			Ignore = 1;
+		if (!bIsQuoted && !bExact && (*Stream)[0] == TEXT('/') && (*Stream)[1] == TEXT('/'))
+		{
+			bIgnore = true;
+		}
 
 		// Command chaining.
-		if( !IsQuoted && !Exact && **Stream=='|' )
+		if (!bIsQuoted && !bExact && **Stream == TEXT('|'))
+		{
 			break;
+		}
 
 		// Check quoting.
-		IsQuoted = IsQuoted ^ (**Stream==34);
-		GotStream=1;
+		bIsQuoted = bIsQuoted ^ (**Stream == TEXT('\"'));
+		bGotStream = true;
 
 		// Got stuff.
-		if( !Ignore )
+		if (!bIgnore)
 		{
-			Result.AppendChar( *((*Stream)++) );
+			Result.AppendChar(*((*Stream)++));
 		}
 		else
 		{
 			(*Stream)++;
 		}
 	}
-	if( Exact )
+
+	if (bExact)
 	{
 		// Eat up exactly one CR/LF.
-		if( **Stream == 13 )
+		if (**Stream == TEXT('\r'))
+		{
 			(*Stream)++;
-		if( **Stream == 10 )
+		}
+		if (**Stream == TEXT('\n'))
+		{
 			(*Stream)++;
+		}
 	}
 	else
 	{
 		// Eat up all CR/LF's.
-		while( **Stream==10 || **Stream==13 || **Stream=='|' )
+		while (**Stream == TEXT('\n') || **Stream == TEXT('\r') || **Stream == TEXT('|'))
+		{
 			(*Stream)++;
+		}
 	}
 
-	return **Stream!=0 || GotStream;
+	return **Stream != TEXT('\0') || bGotStream;
 }
 
-bool FParse::LineExtended(const TCHAR** Stream, FString& Result, int32& LinesConsumed, bool Exact)
+bool FParse::LineExtended(const TCHAR** Stream, FString& Result, int32& LinesConsumed, bool bExact)
 {
-	bool GotStream=0;
-	bool IsQuoted=0;
-	bool Ignore=0;
+	bool bGotStream = false;
+	bool bIsQuoted = false;
+	bool bIgnore = false;
 	int32 BracketDepth = 0;
 
 	Result = TEXT("");
 	LinesConsumed = 0;
 
-	while (**Stream != 0 && ((**Stream != 10 && **Stream != 13) || BracketDepth > 0))
+	while (**Stream != TEXT('\0') && ((**Stream != TEXT('\n') && **Stream != TEXT('\r')) || BracketDepth > 0))
 	{
 		// Start of comments.
-		if( !IsQuoted && !Exact && (*Stream)[0]=='/' && (*Stream)[1]=='/' )
-			Ignore = 1;
+		if (!bIsQuoted && !bExact && (*Stream)[0] == TEXT('/') && (*Stream)[1] == TEXT('/'))
+		{
+			bIgnore = true;
+		}
 
 		// Command chaining.
-		if( !IsQuoted && !Exact && **Stream=='|' )
+		if (!bIsQuoted && !bExact && **Stream == TEXT('|'))
+		{
 			break;
+		}
 
-		GotStream = 1;
+		bGotStream = true;
 
 		// bracketed line break
-		if (**Stream == 10 || **Stream == 13)
+		if (**Stream == TEXT('\n') || **Stream == TEXT('\r'))
 		{
 			checkSlow(BracketDepth > 0);
 
 			Result.AppendChar(TEXT(' '));
 			LinesConsumed++;
 			(*Stream)++;
-			if (**Stream == 10 || **Stream == 13)
+			if (**Stream == TEXT('\n') || **Stream == TEXT('\r'))
 			{
 				(*Stream)++;
 			}
 		}
 		// allow line break if the end of the line is a backslash
-		else if (!IsQuoted && (*Stream)[0] == '\\' && ((*Stream)[1] == 10 || (*Stream)[1] == 13))
+		else if (!bIsQuoted && (*Stream)[0] == '\\' && ((*Stream)[1] == TEXT('\n') || (*Stream)[1] == TEXT('\r')))
 		{
 			Result.AppendChar(TEXT(' '));
 			LinesConsumed++;
 			(*Stream) += 2;
-			if (**Stream == 10 || **Stream == 13)
+			if (**Stream == TEXT('\n') || **Stream == TEXT('\r'))
 			{
 				(*Stream)++;
 			}
 		}
 		// check for starting or ending brace
-		else if (!IsQuoted && **Stream == '{')
+		else if (!bIsQuoted && **Stream == TEXT('{'))
 		{
 			BracketDepth++;
 			(*Stream)++;
 		}
-		else if (!IsQuoted && **Stream == '}' && BracketDepth > 0)
+		else if (!bIsQuoted && **Stream == TEXT('}') && BracketDepth > 0)
 		{
 			BracketDepth--;
 			(*Stream)++;
@@ -1106,12 +1128,12 @@ bool FParse::LineExtended(const TCHAR** Stream, FString& Result, int32& LinesCon
 		else
 		{
 			// Check quoting.
-			IsQuoted = IsQuoted ^ (**Stream==34);
+			bIsQuoted = bIsQuoted ^ (**Stream == TEXT('\"'));
 
 			// Got stuff.
-			if( !Ignore )
+			if (!bIgnore)
 			{
-				Result.AppendChar( *((*Stream)++) );
+				Result.AppendChar(*((*Stream)++));
 			}
 			else
 			{
@@ -1121,22 +1143,22 @@ bool FParse::LineExtended(const TCHAR** Stream, FString& Result, int32& LinesCon
 	}
 	if (**Stream == 0)
 	{
-		if (GotStream)
+		if (bGotStream)
 		{
 			LinesConsumed++;
 		}
 	}
-	else if (Exact)
+	else if (bExact)
 	{
 		// Eat up exactly one CR/LF.
-		if (**Stream == 13 || **Stream == 10)
+		if (**Stream == TEXT('\r') || **Stream == TEXT('\n'))
 		{
 			LinesConsumed++;
-			if (**Stream == 13)
+			if (**Stream == TEXT('\r'))
 			{
 				(*Stream)++;
 			}
-			if( **Stream == 10 )
+			if (**Stream == TEXT('\n'))
 			{
 				(*Stream)++;
 			}
@@ -1145,13 +1167,13 @@ bool FParse::LineExtended(const TCHAR** Stream, FString& Result, int32& LinesCon
 	else
 	{
 		// Eat up all CR/LF's.
-		while (**Stream == 10 || **Stream == 13 || **Stream == '|')
+		while (**Stream == TEXT('\n') || **Stream == TEXT('\r') || **Stream == TEXT('|'))
 		{
-			if (**Stream != '|')
+			if (**Stream != TEXT('|'))
 			{
 				LinesConsumed++;
 			}
-			if (((*Stream)[0] == 10 && (*Stream)[1] == 13) || ((*Stream)[0] == 13 && (*Stream)[1] == 10))
+			if (((*Stream)[0] == TEXT('\n') && (*Stream)[1] == TEXT('\r')) || ((*Stream)[0] == TEXT('\r') && (*Stream)[1] == TEXT('\n')))
 			{
 				(*Stream)++;
 			}
@@ -1159,7 +1181,7 @@ bool FParse::LineExtended(const TCHAR** Stream, FString& Result, int32& LinesCon
 		}
 	}
 
-	return **Stream!=0 || GotStream;
+	return **Stream != TEXT('\0') || bGotStream;
 }
 
 uint32 FParse::HexNumber(const TCHAR* HexString)
