GitHub #4816 : UE-60771: Handle escaped double quote in FParse::LineExtended

diff --git a/Engine/Source/Runtime/Core/Private/Misc/Parse.cpp b/Engine/Source/Runtime/Core/Private/Misc/Parse.cpp
--- a/Engine/Source/Runtime/Core/Private/Misc/Parse.cpp
+++ b/Engine/Source/Runtime/Core/Private/Misc/Parse.cpp
@@ -1069,6 +1069,7 @@ bool FParse::LineExtended(const TCHAR** Stream, FString& Result, int32& LinesCon
 	bool bGotStream = false;
 	bool bIsQuoted = false;
 	bool bIgnore = false;
+	bool bIsEscapedDoubleQuote = false;
 	int32 BracketDepth = 0;
 
 	Result = TEXT("");
@@ -1104,7 +1105,7 @@ bool FParse::LineExtended(const TCHAR** Stream, FString& Result, int32& LinesCon
 			}
 		}
 		// allow line break if the end of the line is a backslash
-		else if (!bIsQuoted && (*Stream)[0] == '\\' && ((*Stream)[1] == TEXT('\n') || (*Stream)[1] == TEXT('\r')))
+		else if (!bIsQuoted && (*Stream)[0] == TEXT('\\') && ((*Stream)[1] == TEXT('\n') || (*Stream)[1] == TEXT('\r')))
 		{
 			Result.AppendChar(TEXT(' '));
 			LinesConsumed++;
@@ -1128,7 +1129,15 @@ bool FParse::LineExtended(const TCHAR** Stream, FString& Result, int32& LinesCon
 		else
 		{
 			// Check quoting.
-			bIsQuoted = bIsQuoted ^ (**Stream == TEXT('\"'));
+			bIsEscapedDoubleQuote |= ((*Stream)[0] == TEXT('\\') && (*Stream)[1] == TEXT('\"'));
+			if (!bIsEscapedDoubleQuote)
+			{
+				bIsQuoted = bIsQuoted ^ (**Stream == TEXT('\"'));
+			}
+			else if (**Stream == TEXT('\"'))
+			{
+				bIsEscapedDoubleQuote = false;
+			}
 
 			// Got stuff.
 			if (!bIgnore)
