Whitespace fixes caused by a bad merge.

diff --git a/Engine/Source/Runtime/Core/Public/Templates/Function.h b/Engine/Source/Runtime/Core/Public/Templates/Function.h
--- a/Engine/Source/Runtime/Core/Public/Templates/Function.h
+++ b/Engine/Source/Runtime/Core/Public/Templates/Function.h
@@ -157,18 +157,18 @@ namespace UE4Function_Private
 	{
 		/**
 		 * Constructor which creates its T by moving.
-	 */
+		 */
 		explicit TFunction_UniqueOwnedObject(T&& InObj)
 			: Obj(MoveTemp(InObj))
 		{
 		}
 
 		IFunction_OwnedObject* CopyToEmptyStorage(FUniqueFunctionStorage& Storage) const override
-	{
+		{
 			// Should never get here - copy functions are deleted for TUniqueFunction
 			check(false);
 			return nullptr;
-	}
+		}
 
 		virtual void* GetAddress() override
 		{
@@ -205,23 +205,23 @@ namespace UE4Function_Private
 
 	template <typename T>
 	FORCEINLINE typename TEnableIf<TIsNullableBinding<T>::Value, bool>::Type IsBound(const T& Func)
-		{
+	{
 		// Function pointers, data member pointers, member function pointers and TFunctions
 		// can all be null/unbound, so test them using their boolean state.
 		return !!Func;
-		}
+	}
 
 	template <typename T>
 	FORCEINLINE typename TEnableIf<!TIsNullableBinding<T>::Value, bool>::Type IsBound(const T& Func)
-		{
+	{
 		// We can't tell if any other generic callable can be invoked, so just assume they can be.
 		return true;
-		}
+	}
 
 	struct FFunctionStorage;
 
 	struct FUniqueFunctionStorage
-		{
+	{
 		FUniqueFunctionStorage() = default;
 
 		FUniqueFunctionStorage(FUniqueFunctionStorage&& Other)
@@ -256,9 +256,9 @@ namespace UE4Function_Private
 			return (IFunction_OwnedObject*)Allocator.GetAllocation();
 		}
 
-	/**
+		/**
 		 * Returns a pointer to the callable object - needed by TFunctionRefBase.
-	 */
+		 */
 		void* GetPtr() const
 		{
 			return GetBoundObject()->GetAddress();
@@ -288,9 +288,9 @@ namespace UE4Function_Private
 		typename TDecay<FunctorType>::Type* Bind(FunctorType&& InFunc)
 		{
 			if (!IsBound(InFunc))
-		{
+			{
 				return nullptr;
-		}
+			}
 
 			using OwnedType = TFunction_CopyableOwnedObject<typename TDecay<FunctorType>::Type>;
 
@@ -391,9 +391,9 @@ namespace UE4Function_Private
 		TFunctionRefBase(TFunctionRefBase&& Other)
 			: Callable(Other.Callable)
 			, Storage (MoveTemp(Other.Storage))
-	{
-			if (Callable)
 		{
+			if (Callable)
+			{
 				#if ENABLE_TFUNCTIONREF_VISUALIZATION
 					// Use Memcpy to copy the other DebugPtrStorage, including vptr (because we don't know the bound type
 					// here), and then reseat the underlying pointer.  Possibly even more evil than the Set code.
@@ -409,9 +409,9 @@ namespace UE4Function_Private
 		TFunctionRefBase(TFunctionRefBase<OtherStorage, Ret (ParamTypes...)>&& Other)
 			: Callable(Other.Callable)
 			, Storage (MoveTemp(Other.Storage))
-	{
-			if (Callable)
 		{
+			if (Callable)
+			{
 				#if ENABLE_TFUNCTIONREF_VISUALIZATION
 					// Use Memcpy to copy the other DebugPtrStorage, including vptr (because we don't know the bound type
 					// here), and then reseat the underlying pointer.  Possibly even more evil than the Set code.
@@ -426,9 +426,9 @@ namespace UE4Function_Private
 		template <typename OtherStorage>
 		TFunctionRefBase(const TFunctionRefBase<OtherStorage, Ret (ParamTypes...)>& Other)
 			: Callable(Other.Callable)
-	{
-			if (Callable)
 		{
+			if (Callable)
+			{
 				void* NewPtr = Storage.BindCopy(Other.Storage);
 
 				#if ENABLE_TFUNCTIONREF_VISUALIZATION
@@ -444,7 +444,7 @@ namespace UE4Function_Private
 			: Callable(Other.Callable)
 		{
 			if (Callable)
-		{
+			{
 				void* NewPtr = Storage.BindCopy(Other.Storage);
 
 				#if ENABLE_TFUNCTIONREF_VISUALIZATION
@@ -467,20 +467,20 @@ namespace UE4Function_Private
 		TFunctionRefBase(FunctorType&& InFunc)
 		{
 			if (auto* Binding = Storage.Bind(Forward<FunctorType>(InFunc)))
-		{
+			{
 				using DecayedFunctorType = typename TRemovePointer<decltype(Binding)>::Type;
 
 				Callable = &TFunctionRefCaller<DecayedFunctorType, Ret (ParamTypes...)>::Call;
 
-			#if ENABLE_TFUNCTIONREF_VISUALIZATION
-				// We placement new over the top of the same object each time.  This is illegal,
-				// but it ensures that the vptr is set correctly for the bound type, and so is
-				// visualizable.  We never depend on the state of this object at runtime, so it's
-				// ok.
+				#if ENABLE_TFUNCTIONREF_VISUALIZATION
+					// We placement new over the top of the same object each time.  This is illegal,
+					// but it ensures that the vptr is set correctly for the bound type, and so is
+					// visualizable.  We never depend on the state of this object at runtime, so it's
+					// ok.
 					new ((void*)&DebugPtrStorage) TDebugHelper<DecayedFunctorType>;
 					DebugPtrStorage.Ptr = (void*)Binding;
-			#endif
-		}
+				#endif
+			}
 		}
 
 		TFunctionRefBase& operator=(TFunctionRefBase&&) = delete;
@@ -681,7 +681,7 @@ public:
 		// This constructor is disabled for TFunctionRef types so it isn't incorrectly selected as copy/move constructors.
 	}
 
-		TFunctionRef(const TFunctionRef&) = default;
+	TFunctionRef(const TFunctionRef&) = default;
 
 	// We delete the assignment operators because we don't want it to be confused with being related to
 	// regular C++ reference assignment - i.e. calling the assignment operator of whatever the reference
@@ -779,13 +779,13 @@ public:
 	{
 		Swap(*this, Other);
 		return *this;
-		}
+	}
 
 	/**
 	 * Copy assignment operator.
 	 */
 	TFunction& operator=(const TFunction& Other)
-		{
+	{
 		TFunction Temp = Other;
 		Swap(*this, Temp);
 		return *this;
@@ -795,9 +795,9 @@ public:
 	 * Tests if the TFunction is callable.
 	 */
 	FORCEINLINE explicit operator bool() const
-		{
+	{
 		return Super::IsSet();
-		}
+	}
 };
 
 /**
@@ -818,7 +818,7 @@ public:
  */
 template <typename FuncType>
 class TUniqueFunction final : public UE4Function_Private::TFunctionRefBase<UE4Function_Private::FUniqueFunctionStorage, FuncType>
-		{
+{
 	using Super = UE4Function_Private::TFunctionRefBase<UE4Function_Private::FUniqueFunctionStorage, FuncType>;
 
 public:
@@ -870,7 +870,7 @@ public:
 	 */
 	TUniqueFunction(const TFunction<FuncType>& Other)
 		: Super(*(const UE4Function_Private::TFunctionRefBase<UE4Function_Private::FFunctionStorage, FuncType>*)&Other)
-		{
+	{
 	}
 
 	/**
