Ensure that strings formatted using FMicrosoftPlatformString::GetVarArgs() are always null terminated, whether we use the secure CRT or not.

diff --git a/Engine/Source/Runtime/Core/Public/GenericPlatform/MicrosoftPlatformString.h b/Engine/Source/Runtime/Core/Public/GenericPlatform/MicrosoftPlatformString.h
--- a/Engine/Source/Runtime/Core/Public/GenericPlatform/MicrosoftPlatformString.h
+++ b/Engine/Source/Runtime/Core/Public/GenericPlatform/MicrosoftPlatformString.h
@@ -142,7 +142,23 @@ struct FMicrosoftPlatformString : public FGenericPlatformString
 #if USE_SECURE_CRT
 		int32 Result = _vsntprintf_s( Dest, DestSize, Count, Fmt, ArgPtr );
 #else
-		int32 Result = _vsntprintf( Dest, Count, Fmt, ArgPtr );
+		// Get the max number of chars to write. It's not uncommon to pass DestSize = ARRAY_COUNT(...) and Count = ARRAY_COUNT(...) - 1, so don't just assume Count - 1.
+		int32 MaxChars;
+		if (Count < DestSize)
+		{
+			MaxChars = Count;
+		}
+		else
+		{
+			MaxChars = DestSize - 1;
+		}
+
+		// Format the string. If the output is truncated (Result = -1) or contains exactly the maximum number of characters available (Result == MaxChars), we need to insert a null terminator.
+		int32 Result = _vsntprintf( Dest, MaxChars, Fmt, ArgPtr );
+		if (Result < 0 || Result == MaxChars)
+		{
+			Dest[MaxChars] = 0;
+		}
 #endif // USE_SECURE_CRT
 		va_end( ArgPtr );
 		return Result;
