Fix CIS warnings.

#rnx

diff --git a/Engine/Source/Runtime/Core/Public/GenericPlatform/MicrosoftPlatformString.h b/Engine/Source/Runtime/Core/Public/GenericPlatform/MicrosoftPlatformString.h
--- a/Engine/Source/Runtime/Core/Public/GenericPlatform/MicrosoftPlatformString.h
+++ b/Engine/Source/Runtime/Core/Public/GenericPlatform/MicrosoftPlatformString.h
@@ -144,13 +144,13 @@ struct FMicrosoftPlatformString : public FGenericPlatformString
 #else
 		// Get the max number of chars to write. It's not uncommon to pass DestSize = ARRAY_COUNT(...) and Count = ARRAY_COUNT(...) - 1, so don't just assume Count - 1.
 		int32 MaxChars;
-		if (Count < DestSize)
+		if (Count < (int32)DestSize)
 		{
 			MaxChars = Count;
 		}
 		else
 		{
-			MaxChars = DestSize - 1;
+			MaxChars = (int32)DestSize - 1;
 		}
 
 		// Format the string. If the output is truncated (Result = -1) or contains exactly the maximum number of characters available (Result == MaxChars), we need to insert a null terminator.
