Replace custom code to ensure FMicrosoftPlatformString::GetVarArgs() null terminates its buffer with Microsoft's standards-compliant implementation.

diff --git a/Engine/Source/Runtime/Core/Public/GenericPlatform/MicrosoftPlatformString.h b/Engine/Source/Runtime/Core/Public/GenericPlatform/MicrosoftPlatformString.h
--- a/Engine/Source/Runtime/Core/Public/GenericPlatform/MicrosoftPlatformString.h
+++ b/Engine/Source/Runtime/Core/Public/GenericPlatform/MicrosoftPlatformString.h
@@ -142,23 +142,7 @@ struct FMicrosoftPlatformString : public FGenericPlatformString
 #if USE_SECURE_CRT
 		int32 Result = _vsntprintf_s( Dest, DestSize, Count, Fmt, ArgPtr );
 #else
-		// Get the max number of chars to write. It's not uncommon to pass DestSize = ARRAY_COUNT(...) and Count = ARRAY_COUNT(...) - 1, so don't just assume Count - 1.
-		int32 MaxChars;
-		if (Count < (int32)DestSize)
-		{
-			MaxChars = Count;
-		}
-		else
-		{
-			MaxChars = (int32)DestSize - 1;
-		}
-
-		// Format the string. If the output is truncated (Result = -1) or contains exactly the maximum number of characters available (Result == MaxChars), we need to insert a null terminator.
-		int32 Result = _vsntprintf( Dest, MaxChars, Fmt, ArgPtr );
-		if (Result < 0 || Result == MaxChars)
-		{
-			Dest[MaxChars] = 0;
-		}
+		int32 Result = vswprintf(Dest, Count, Fmt, ArgPtr);
 #endif // USE_SECURE_CRT
 		va_end( ArgPtr );
 		return Result;
@@ -320,7 +304,7 @@ struct FMicrosoftPlatformString : public FGenericPlatformString
 #if USE_SECURE_CRT
 		int32 Result = _vsnprintf_s( Dest, DestSize, Count, Fmt, ArgPtr );
 #else
-		int32 Result = _vsnprintf( Dest, Count, Fmt, ArgPtr );
+		int32 Result = vsnprintf( Dest, Count, Fmt, ArgPtr );
 #endif // USE_SECURE_CRT
 		va_end( ArgPtr );
 		return Result;
