Move text-only implementations of FOutputDeviceError back into Core, so we can build command-line applications that don't depend on ApplicationCore.

#rnx

diff --git a/Engine/Source/Runtime/ApplicationCore/Private/Android/AndroidErrorOutputDevice.cpp b/Engine/Source/Runtime/ApplicationCore/Private/Android/AndroidErrorOutputDevice.cpp
deleted file mode 100644
diff --git a/Engine/Source/Runtime/ApplicationCore/Private/GenericPlatform/GenericPlatformApplicationMisc.cpp b/Engine/Source/Runtime/ApplicationCore/Private/GenericPlatform/GenericPlatformApplicationMisc.cpp
--- a/Engine/Source/Runtime/ApplicationCore/Private/GenericPlatform/GenericPlatformApplicationMisc.cpp
+++ b/Engine/Source/Runtime/ApplicationCore/Private/GenericPlatform/GenericPlatformApplicationMisc.cpp
@@ -6,6 +6,7 @@
 #include "HAL/FeedbackContextAnsi.h"
 #include "Math/Color.h"
 
+#include "HAL/PlatformOutputDevices.h"
 #include "HAL/PlatformApplicationMisc.h"
 
 /** Hooks for moving ClipboardCopy and ClipboardPaste into FPlatformApplicationMisc */
@@ -52,14 +53,12 @@ class FOutputDeviceConsole* FGenericPlatformApplicationMisc::CreateConsoleOutput
 
 class FOutputDeviceError* FGenericPlatformApplicationMisc::GetErrorOutputDevice()
 {
-	static FOutputDeviceAnsiError Singleton;
-	return &Singleton;
+	return FPlatformOutputDevices::GetError();
 }
 
 class FFeedbackContext* FGenericPlatformApplicationMisc::GetFeedbackContext()
 {
-	static FFeedbackContextAnsi Singleton;
-	return &Singleton;
+	return FPlatformOutputDevices::GetFeedbackContext();
 }
 
 GenericApplication* FGenericPlatformApplicationMisc::CreateApplication()
diff --git a/Engine/Source/Runtime/ApplicationCore/Private/IOS/IOSErrorOutputDevice.cpp b/Engine/Source/Runtime/ApplicationCore/Private/IOS/IOSErrorOutputDevice.cpp
deleted file mode 100644
diff --git a/Engine/Source/Runtime/ApplicationCore/Private/Linux/LinuxPlatformApplicationMisc.cpp b/Engine/Source/Runtime/ApplicationCore/Private/Linux/LinuxPlatformApplicationMisc.cpp
--- a/Engine/Source/Runtime/ApplicationCore/Private/Linux/LinuxPlatformApplicationMisc.cpp
+++ b/Engine/Source/Runtime/ApplicationCore/Private/Linux/LinuxPlatformApplicationMisc.cpp
@@ -6,7 +6,7 @@
 #include "HAL/ThreadHeartBeat.h"
 #include "Modules/ModuleManager.h"
 #include "Linux/LinuxConsoleOutputDevice.h"
-#include "Unix/UnixErrorOutputDevice.h"
+#include "Unix/UnixApplicationErrorOutputDevice.h"
 #include "Unix/UnixFeedbackContext.h"
 #include "Linux/LinuxApplication.h"
 
@@ -389,7 +389,7 @@ class FOutputDeviceConsole* FLinuxPlatformApplicationMisc::CreateConsoleOutputDe
 
 class FOutputDeviceError* FLinuxPlatformApplicationMisc::GetErrorOutputDevice()
 {
-	static FUnixErrorOutputDevice Singleton;
+	static FUnixApplicationErrorOutputDevice Singleton;
 	return &Singleton;
 }
 
diff --git a/Engine/Source/Runtime/ApplicationCore/Private/Mac/MacApplicationErrorOutputDevice.cpp b/Engine/Source/Runtime/ApplicationCore/Private/Mac/MacApplicationErrorOutputDevice.cpp
new file mode 100644
--- /dev/null
+++ b/Engine/Source/Runtime/ApplicationCore/Private/Mac/MacApplicationErrorOutputDevice.cpp
@@ -0,0 +1,21 @@
+// Copyright 1998-2018 Epic Games, Inc. All Rights Reserved.
+
+#include "Mac/MacApplicationErrorOutputDevice.h"
+#include "HAL/PlatformApplicationMisc.h"
+
+void FMacApplicationErrorOutputDevice::HandleErrorRestoreUI()
+{
+	// Unhide the mouse.
+	// @TODO: Remove usage of deprecated CGCursorIsVisible function
+#pragma clang diagnostic push
+#pragma clang diagnostic ignored "-Wdeprecated-declarations"
+	while (!CGCursorIsVisible())
+	{
+		CGDisplayShowCursor(kCGDirectMainDisplay);
+	}
+#pragma clang diagnostic pop
+	// Release capture and allow mouse to freely roam around.
+	CGAssociateMouseAndMouseCursorPosition(true);
+
+	FPlatformApplicationMisc::ClipboardCopy(GErrorHist);
+}
diff --git a/Engine/Source/Runtime/ApplicationCore/Private/Mac/MacConsoleOutputDevice.cpp b/Engine/Source/Runtime/ApplicationCore/Private/Mac/MacConsoleOutputDevice.cpp
--- a/Engine/Source/Runtime/ApplicationCore/Private/Mac/MacConsoleOutputDevice.cpp
+++ b/Engine/Source/Runtime/ApplicationCore/Private/Mac/MacConsoleOutputDevice.cpp
@@ -274,3 +274,13 @@ void FMacConsoleOutputDevice::SetDefaultTextColor()
 		OutstandingTasks--;
 	}, NSDefaultRunLoopMode, false);
 }
+
+@implementation FMacConsoleWindow
+- (void)windowWillClose:(NSNotification*)Notification
+{
+	if (!MacApplication && [[NSApp orderedWindows] count] == 1)
+	{
+		_Exit(0);
+	}
+}
+@end
diff --git a/Engine/Source/Runtime/ApplicationCore/Private/Mac/MacErrorOutputDevice.cpp b/Engine/Source/Runtime/ApplicationCore/Private/Mac/MacErrorOutputDevice.cpp
deleted file mode 100644
diff --git a/Engine/Source/Runtime/ApplicationCore/Private/Mac/MacPlatformApplicationMisc.cpp b/Engine/Source/Runtime/ApplicationCore/Private/Mac/MacPlatformApplicationMisc.cpp
--- a/Engine/Source/Runtime/ApplicationCore/Private/Mac/MacPlatformApplicationMisc.cpp
+++ b/Engine/Source/Runtime/ApplicationCore/Private/Mac/MacPlatformApplicationMisc.cpp
@@ -8,7 +8,7 @@
 #include "Misc/App.h"
 #include "Modules/ModuleManager.h"
 #include "Mac/MacConsoleOutputDevice.h"
-#include "Mac/MacErrorOutputDevice.h"
+#include "Mac/MacApplicationErrorOutputDevice.h"
 #include "Mac/MacFeedbackContext.h"
 #include "Misc/OutputDeviceHelper.h"
 #include "Misc/OutputDeviceRedirector.h"
@@ -16,6 +16,7 @@
 #include "HAL/FeedbackContextAnsi.h"
 #include "Mac/CocoaMenu.h"
 #include "Mac/CocoaThread.h"
+#include "Mac/MacPlatformOutputDevices.h"
 
 #include <dlfcn.h>
 #include <IOKit/IOKitLib.h>
@@ -337,7 +338,7 @@ class FOutputDeviceConsole* FMacPlatformApplicationMisc::CreateConsoleOutputDevi
 
 class FOutputDeviceError* FMacPlatformApplicationMisc::GetErrorOutputDevice()
 {
-	static FMacErrorOutputDevice Singleton;
+	static FMacApplicationErrorOutputDevice Singleton;
 	return &Singleton;
 }
 
@@ -347,8 +348,7 @@ class FFeedbackContext* FMacPlatformApplicationMisc::GetFeedbackContext()
 	static FMacFeedbackContext Singleton;
 	return &Singleton;
 #else
-	static FFeedbackContextAnsi Singleton;
-	return &Singleton;
+	return FMacPlatformOutputDevices::GetFeedbackContext();
 #endif
 }
 
diff --git a/Engine/Source/Runtime/ApplicationCore/Private/Unix/UnixApplicationErrorOutputDevice.cpp b/Engine/Source/Runtime/ApplicationCore/Private/Unix/UnixApplicationErrorOutputDevice.cpp
new file mode 100644
--- /dev/null
+++ b/Engine/Source/Runtime/ApplicationCore/Private/Unix/UnixApplicationErrorOutputDevice.cpp
@@ -0,0 +1,14 @@
+// Copyright 1998-2018 Epic Games, Inc. All Rights Reserved.
+
+#include "Unix/UnixApplicationErrorOutputDevice.h"
+#include "HAL/PlatformApplicationMisc.h"
+#include "Misc/App.h"
+
+void FUnixApplicationErrorOutputDevice::HandleErrorRestoreUI()
+{
+	// do not copy if graphics have not been initialized or if we're on the wrong thread
+	if (FApp::CanEverRender() && IsInGameThread())
+	{
+		FPlatformApplicationMisc::ClipboardCopy(GErrorHist);
+	}
+}
diff --git a/Engine/Source/Runtime/ApplicationCore/Private/Unix/UnixErrorOutputDevice.cpp b/Engine/Source/Runtime/ApplicationCore/Private/Unix/UnixErrorOutputDevice.cpp
deleted file mode 100644
diff --git a/Engine/Source/Runtime/ApplicationCore/Private/Windows/WindowsApplicationErrorOutputDevice.cpp b/Engine/Source/Runtime/ApplicationCore/Private/Windows/WindowsApplicationErrorOutputDevice.cpp
new file mode 100644
--- /dev/null
+++ b/Engine/Source/Runtime/ApplicationCore/Private/Windows/WindowsApplicationErrorOutputDevice.cpp
@@ -0,0 +1,18 @@
+// Copyright 1998-2018 Epic Games, Inc. All Rights Reserved.
+
+#include "Windows/WindowsApplicationErrorOutputDevice.h"
+#include "Windows/WindowsHWrapper.h"
+#include "HAL/PlatformApplicationMisc.h"
+
+void FWindowsApplicationErrorOutputDevice::HandleErrorRestoreUI()
+{
+	// Unhide the mouse.
+	while (::ShowCursor(true) < 0);
+	// Release capture.
+	::ReleaseCapture();
+	// Allow mouse to freely roam around.
+	::ClipCursor(NULL);
+
+	// Copy to clipboard in non-cooked editor builds.
+	FPlatformApplicationMisc::ClipboardCopy(GErrorHist);
+}
diff --git a/Engine/Source/Runtime/ApplicationCore/Private/Windows/WindowsErrorOutputDevice.cpp b/Engine/Source/Runtime/ApplicationCore/Private/Windows/WindowsErrorOutputDevice.cpp
deleted file mode 100644
diff --git a/Engine/Source/Runtime/ApplicationCore/Private/Windows/WindowsPlatformApplicationMisc.cpp b/Engine/Source/Runtime/ApplicationCore/Private/Windows/WindowsPlatformApplicationMisc.cpp
--- a/Engine/Source/Runtime/ApplicationCore/Private/Windows/WindowsPlatformApplicationMisc.cpp
+++ b/Engine/Source/Runtime/ApplicationCore/Private/Windows/WindowsPlatformApplicationMisc.cpp
@@ -2,7 +2,7 @@
 
 #include "Windows/WindowsPlatformApplicationMisc.h"
 #include "Windows/WindowsApplication.h"
-#include "Windows/WindowsErrorOutputDevice.h"
+#include "Windows/WindowsApplicationErrorOutputDevice.h"
 #include "Windows/WindowsConsoleOutputDevice.h"
 #include "Windows/WindowsFeedbackContext.h"
 #include "HAL/FeedbackContextAnsi.h"
@@ -10,6 +10,8 @@
 #include "Math/Color.h"
 #include "Windows/WindowsHWrapper.h"
 #include "Modules/ModuleManager.h"
+#include "Misc/CoreDelegates.h"
+#include "Windows/WindowsPlatformOutputDevices.h"
 
 // Resource includes.
 #include "Runtime/Launch/Resources/Windows/Resource.h"
@@ -53,7 +55,7 @@ class FOutputDeviceConsole* FWindowsPlatformApplicationMisc::CreateConsoleOutput
 
 class FOutputDeviceError* FWindowsPlatformApplicationMisc::GetErrorOutputDevice()
 {
-	static FWindowsErrorOutputDevice Singleton;
+	static FWindowsApplicationErrorOutputDevice Singleton;
 	return &Singleton;
 }
 
@@ -63,8 +65,7 @@ class FFeedbackContext* FWindowsPlatformApplicationMisc::GetFeedbackContext()
 	static FWindowsFeedbackContext Singleton;
 	return &Singleton;
 #else
-	static FFeedbackContextAnsi Singleton;
-	return &Singleton;
+	return FPlatformOutputDevices::GetFeedbackContext();
 #endif
 }
 
diff --git a/Engine/Source/Runtime/ApplicationCore/Public/Android/AndroidErrorOutputDevice.h b/Engine/Source/Runtime/ApplicationCore/Public/Android/AndroidErrorOutputDevice.h
deleted file mode 100644
diff --git a/Engine/Source/Runtime/ApplicationCore/Public/IOS/IOSErrorOutputDevice.h b/Engine/Source/Runtime/ApplicationCore/Public/IOS/IOSErrorOutputDevice.h
deleted file mode 100644
diff --git a/Engine/Source/Runtime/ApplicationCore/Public/Mac/MacApplicationErrorOutputDevice.h b/Engine/Source/Runtime/ApplicationCore/Public/Mac/MacApplicationErrorOutputDevice.h
new file mode 100644
--- /dev/null
+++ b/Engine/Source/Runtime/ApplicationCore/Public/Mac/MacApplicationErrorOutputDevice.h
@@ -0,0 +1,13 @@
+// Copyright 1998-2018 Epic Games, Inc. All Rights Reserved.
+
+#pragma once
+
+#include "CoreTypes.h"
+#include "Mac/MacErrorOutputDevice.h"
+
+class FMacApplicationErrorOutputDevice : public FMacErrorOutputDevice
+{
+protected:
+	virtual void HandleErrorRestoreUI() override;
+};
+
diff --git a/Engine/Source/Runtime/ApplicationCore/Public/Mac/MacErrorOutputDevice.h b/Engine/Source/Runtime/ApplicationCore/Public/Mac/MacErrorOutputDevice.h
deleted file mode 100644
diff --git a/Engine/Source/Runtime/ApplicationCore/Public/Unix/UnixApplicationErrorOutputDevice.h b/Engine/Source/Runtime/ApplicationCore/Public/Unix/UnixApplicationErrorOutputDevice.h
new file mode 100644
--- /dev/null
+++ b/Engine/Source/Runtime/ApplicationCore/Public/Unix/UnixApplicationErrorOutputDevice.h
@@ -0,0 +1,12 @@
+// Copyright 1998-2018 Epic Games, Inc. All Rights Reserved.
+
+#pragma once
+
+#include "CoreTypes.h"
+#include "Unix/UnixErrorOutputDevice.h"
+
+struct FUnixApplicationErrorOutputDevice : public FUnixErrorOutputDevice
+{
+protected:
+	virtual void HandleErrorRestoreUI() override;
+};
diff --git a/Engine/Source/Runtime/ApplicationCore/Public/Unix/UnixErrorOutputDevice.h b/Engine/Source/Runtime/ApplicationCore/Public/Unix/UnixErrorOutputDevice.h
deleted file mode 100644
diff --git a/Engine/Source/Runtime/ApplicationCore/Public/Windows/WindowsApplicationErrorOutputDevice.h b/Engine/Source/Runtime/ApplicationCore/Public/Windows/WindowsApplicationErrorOutputDevice.h
new file mode 100644
--- /dev/null
+++ b/Engine/Source/Runtime/ApplicationCore/Public/Windows/WindowsApplicationErrorOutputDevice.h
@@ -0,0 +1,16 @@
+// Copyright 1998-2018 Epic Games, Inc. All Rights Reserved.
+
+#pragma once
+
+#include "CoreTypes.h"
+#include "Windows/WindowsErrorOutputDevice.h"
+
+/**
+ * Windows implementation of console log window, utilizing the Win32 console API
+ */
+class FWindowsApplicationErrorOutputDevice : public FWindowsErrorOutputDevice
+{
+protected:
+	virtual void HandleErrorRestoreUI() override;
+};
+
diff --git a/Engine/Source/Runtime/ApplicationCore/Public/Windows/WindowsErrorOutputDevice.h b/Engine/Source/Runtime/ApplicationCore/Public/Windows/WindowsErrorOutputDevice.h
deleted file mode 100644
diff --git a/Engine/Source/Runtime/Core/Private/Android/AndroidErrorOutputDevice.cpp b/Engine/Source/Runtime/Core/Private/Android/AndroidErrorOutputDevice.cpp
new file mode 100644
--- /dev/null
+++ b/Engine/Source/Runtime/Core/Private/Android/AndroidErrorOutputDevice.cpp
@@ -0,0 +1,52 @@
+// Copyright 1998-2018 Epic Games, Inc. All Rights Reserved.
+
+#include "Android/AndroidErrorOutputDevice.h"
+
+#include "CoreGlobals.h"
+#include "Misc/OutputDeviceHelper.h"
+#include "HAL/PlatformMisc.h"
+#include "Misc/OutputDeviceRedirector.h"
+#include "CoreGlobals.h"
+
+FAndroidErrorOutputDevice::FAndroidErrorOutputDevice()
+{
+}
+
+void FAndroidErrorOutputDevice::Serialize( const TCHAR* Msg, ELogVerbosity::Type Verbosity, const class FName& Category )
+{
+	FPlatformMisc::LowLevelOutputDebugString(*FOutputDeviceHelper::FormatLogLine(Verbosity, Category, Msg, GPrintLogTimes));
+	if (GIsGuarded)
+	{
+		UE_DEBUG_BREAK();
+	}
+	else
+	{
+		HandleError();
+		FPlatformMisc::RequestExit(true);
+	}
+}
+
+void FAndroidErrorOutputDevice::HandleError()
+{
+	static int32 CallCount = 0;
+	int32 NewCallCount = FPlatformAtomics::InterlockedIncrement(&CallCount);
+
+	if (NewCallCount != 1)
+	{
+		UE_LOG(LogAndroid, Error, TEXT("HandleError re-entered."));
+		return;
+	}
+	
+	GIsGuarded = 0;
+	GIsRunning = 0;
+	GIsCriticalError = 1;
+	GLogConsole = NULL;
+	GErrorHist[ARRAY_COUNT(GErrorHist) - 1] = 0;
+
+	// Dump the error and flush the log.
+#if !NO_LOGGING
+	FDebug::LogFormattedMessageWithCallstack(LogAndroid.GetCategoryName(), __FILE__, __LINE__, TEXT("=== Critical error: ==="), GErrorHist, ELogVerbosity::Error);
+#endif
+	
+	GLog->PanicFlushThreadedLogs();
+}
diff --git a/Engine/Source/Runtime/Core/Private/GenericPlatform/GenericPlatformOutputDevices.cpp b/Engine/Source/Runtime/Core/Private/GenericPlatform/GenericPlatformOutputDevices.cpp
--- a/Engine/Source/Runtime/Core/Private/GenericPlatform/GenericPlatformOutputDevices.cpp
+++ b/Engine/Source/Runtime/Core/Private/GenericPlatform/GenericPlatformOutputDevices.cpp
@@ -166,3 +166,15 @@ void FGenericPlatformOutputDevices::GetPerChannelFileOverrides(TArray<FOutputDev
 		}
 	}
 }
+
+FOutputDeviceError* FGenericPlatformOutputDevices::GetError()
+{
+	static FOutputDeviceAnsiError Singleton;
+	return &Singleton;
+}
+
+FFeedbackContext* FGenericPlatformOutputDevices::GetFeedbackContext()
+{
+	static FFeedbackContextAnsi Singleton;
+	return &Singleton;
+}
diff --git a/Engine/Source/Runtime/Core/Private/IOS/IOSErrorOutputDevice.cpp b/Engine/Source/Runtime/Core/Private/IOS/IOSErrorOutputDevice.cpp
new file mode 100644
--- /dev/null
+++ b/Engine/Source/Runtime/Core/Private/IOS/IOSErrorOutputDevice.cpp
@@ -0,0 +1,51 @@
+// Copyright 1998-2018 Epic Games, Inc. All Rights Reserved.
+
+/*=============================================================================
+	IOSPlatformOutputDevices.mm: iOS implementations of OutputDevices functions
+=============================================================================*/
+
+#include "IOS/IOSErrorOutputDevice.h"
+#include "HAL/PlatformMisc.h"
+#include "HAL/PlatformAtomics.h"
+#include "Misc/OutputDeviceRedirector.h"
+#include "CoreGlobals.h"
+
+FIOSErrorOutputDevice::FIOSErrorOutputDevice()
+:	ErrorPos(0)
+{
+}
+
+void FIOSErrorOutputDevice::Serialize( const TCHAR* Msg, ELogVerbosity::Type Verbosity, const class FName& Category )
+{
+	// @todo iosmerge: this was a big mess copied from Mac that we don't want at all. can we use default like other consoles?
+	if( GIsGuarded )
+	{
+//		FOutputDevice::Serialize(Msg, Verbosity, Category);
+		UE_DEBUG_BREAK();
+	}
+	else
+	{
+		// We crashed outside the guarded code (e.g. appExit).
+		HandleError();
+		FPlatformMisc::RequestExit( true );
+	}
+}
+
+void FIOSErrorOutputDevice::HandleError()
+{
+	// make sure we don't report errors twice
+	static int32 CallCount = 0;
+	int32 NewCallCount = FPlatformAtomics::InterlockedIncrement(&CallCount);
+	if (NewCallCount != 1)
+	{
+		UE_LOG(LogIOS, Error, TEXT("HandleError re-entered.") );
+		return;
+	}
+
+	GIsGuarded = 0;
+	GIsRunning = 0;
+	GIsCriticalError = 1;
+	GLogConsole = NULL;
+
+	GLog->PanicFlushThreadedLogs();
+}
diff --git a/Engine/Source/Runtime/Core/Private/Mac/MacErrorOutputDevice.cpp b/Engine/Source/Runtime/Core/Private/Mac/MacErrorOutputDevice.cpp
new file mode 100644
--- /dev/null
+++ b/Engine/Source/Runtime/Core/Private/Mac/MacErrorOutputDevice.cpp
@@ -0,0 +1,94 @@
+// Copyright 1998-2018 Epic Games, Inc. All Rights Reserved.
+
+#include "Mac/MacErrorOutputDevice.h"
+#include "Misc/App.h"
+#include "Misc/CoreDelegates.h"
+#include "Misc/ConfigCacheIni.h"
+#include "HAL/FeedbackContextAnsi.h"
+#include "Mac/CocoaThread.h"
+
+FMacErrorOutputDevice::FMacErrorOutputDevice()
+:	ErrorPos(0)
+{}
+
+void FMacErrorOutputDevice::Serialize( const TCHAR* Msg, ELogVerbosity::Type Verbosity, const class FName& Category )
+{
+	UE_DEBUG_BREAK();
+
+	if( !GIsCriticalError )
+	{   
+		// First appError.
+		GIsCriticalError = 1;
+		TCHAR ErrorBuffer[1024];
+		ErrorBuffer[0] = 0;
+		// pop up a crash window if we are not in unattended mode
+		if( FApp::IsUnattended() == false )
+		{ 
+			UE_LOG(LogMac, Error, TEXT("appError called: %s"), Msg );
+		} 
+		// log the warnings to the log
+		else
+		{
+			UE_LOG(LogMac, Error, TEXT("appError called: %s"), Msg );
+		}
+		FCString::Strncpy( GErrorHist, Msg, ARRAY_COUNT(GErrorHist) - 5 );
+		FCString::Strncat( GErrorHist, TEXT("\r\n\r\n"), ARRAY_COUNT(GErrorHist) - 1  );
+		ErrorPos = FCString::Strlen(GErrorHist);
+	}
+	else
+	{
+		UE_LOG(LogMac, Error, TEXT("Error reentered: %s"), Msg );
+	}
+
+	if( GIsGuarded )
+	{
+		// Propagate error so structured exception handler can perform necessary work.
+#if PLATFORM_EXCEPTIONS_DISABLED
+		UE_DEBUG_BREAK();
+#endif
+		FPlatformMisc::RaiseException( 1 );
+	}
+	else
+	{
+		// We crashed outside the guarded code (e.g. appExit).
+		HandleError();
+		FPlatformMisc::RequestExit( true );
+	}
+}
+
+void FMacErrorOutputDevice::HandleError()
+{
+	// make sure we don't report errors twice
+	static int32 CallCount = 0;
+	int32 NewCallCount = FPlatformAtomics::InterlockedIncrement(&CallCount);
+	if (NewCallCount != 1)
+	{
+		UE_LOG(LogMac, Error, TEXT("HandleError re-entered.") );
+		return;
+	}
+
+	// Trigger the OnSystemFailure hook if it exists
+	FCoreDelegates::OnHandleSystemError.Broadcast();
+
+	GIsGuarded				= 0;
+	GIsRunning				= 0;
+	GIsCriticalError		= 1;
+	GLogConsole				= NULL;
+	GErrorHist[ARRAY_COUNT(GErrorHist)-1]=0;
+
+	// Dump the error and flush the log.
+	UE_LOG(LogMac, Log, TEXT("=== Critical error: ===") LINE_TERMINATOR TEXT("%s") LINE_TERMINATOR, GErrorExceptionDescription);
+	UE_LOG(LogMac, Log, GErrorHist);
+
+	GLog->Flush();
+
+	HandleErrorRestoreUI();
+
+	FPlatformMisc::SubmitErrorReport( GErrorHist, EErrorReportMode::Interactive );
+
+	FCoreDelegates::OnShutdownAfterError.Broadcast();
+}
+
+void FMacErrorOutputDevice::HandleErrorRestoreUI()
+{
+}
diff --git a/Engine/Source/Runtime/Core/Private/Mac/MacPlatformOutputDevices.cpp b/Engine/Source/Runtime/Core/Private/Mac/MacPlatformOutputDevices.cpp
--- a/Engine/Source/Runtime/Core/Private/Mac/MacPlatformOutputDevices.cpp
+++ b/Engine/Source/Runtime/Core/Private/Mac/MacPlatformOutputDevices.cpp
@@ -1,8 +1,16 @@
 // Copyright 1998-2018 Epic Games, Inc. All Rights Reserved.
 
 #include "Mac/MacPlatformOutputDevices.h"
+#include "Mac/MacErrorOutputDevice.h"
 
 class FOutputDevice* FMacPlatformOutputDevices::GetEventLog()
 {
 	return NULL;
 }
+
+FOutputDeviceError* FMacPlatformOutputDevices::GetError()
+{
+	static FMacErrorOutputDevice Singleton;
+	return &Singleton;
+}
+
diff --git a/Engine/Source/Runtime/Core/Private/Unix/UnixErrorOutputDevice.cpp b/Engine/Source/Runtime/Core/Private/Unix/UnixErrorOutputDevice.cpp
new file mode 100644
--- /dev/null
+++ b/Engine/Source/Runtime/Core/Private/Unix/UnixErrorOutputDevice.cpp
@@ -0,0 +1,110 @@
+// Copyright 1998-2018 Epic Games, Inc. All Rights Reserved.
+
+#include "Unix/UnixErrorOutputDevice.h"
+#include "Containers/StringConv.h"
+#include "Logging/LogMacros.h"
+#include "CoreGlobals.h"
+#include "Misc/Parse.h"
+#include "Misc/CommandLine.h"
+#include "Misc/Paths.h"
+#include "Misc/OutputDeviceHelper.h"
+#include "Misc/OutputDeviceRedirector.h"
+#include "Misc/CoreDelegates.h"
+#include "Misc/App.h"
+
+FUnixErrorOutputDevice::FUnixErrorOutputDevice()
+:	ErrorPos(0)
+{
+}
+
+void FUnixErrorOutputDevice::Serialize(const TCHAR* Msg, ELogVerbosity::Type Verbosity, const class FName& Category)
+{
+	UE_DEBUG_BREAK();
+
+	if (!GIsCriticalError)
+	{
+		// First appError.
+		GIsCriticalError = 1;
+		TCHAR ErrorBuffer[1024];
+		ErrorBuffer[0] = 0;
+		// pop up a crash window if we are not in unattended mode; 
+		if (WITH_EDITOR && FApp::IsUnattended() == false)
+		{
+			// @TODO Not implemented
+			UE_LOG(LogCore, Error, TEXT("appError called: %s"), Msg );
+		}
+		// log the warnings to the log
+		else
+		{
+			UE_LOG(LogCore, Error, TEXT("appError called: %s"), Msg );
+		}
+		FCString::Strncpy( GErrorHist, Msg, ARRAY_COUNT(GErrorHist) - 5 );
+		FCString::Strncat( GErrorHist, TEXT("\r\n\r\n"), ARRAY_COUNT(GErrorHist) - 1 );
+		ErrorPos = FCString::Strlen(GErrorHist);
+	}
+	else
+	{
+		UE_LOG(LogCore, Error, TEXT("Error reentered: %s"), Msg );
+	}
+
+	if( GIsGuarded )
+	{
+		// Propagate error so structured exception handler can perform necessary work.
+#if PLATFORM_EXCEPTIONS_DISABLED
+		UE_DEBUG_BREAK();
+#endif
+		FPlatformMisc::RaiseException(1);
+	}
+	else
+	{
+		// We crashed outside the guarded code (e.g. appExit).
+		HandleError();
+		FPlatformMisc::RequestExit(true);
+	}
+}
+
+void FUnixErrorOutputDevice::HandleError()
+{
+	// make sure we don't report errors twice
+	static int32 CallCount = 0;
+	int32 NewCallCount = FPlatformAtomics::InterlockedIncrement(&CallCount);
+	if (NewCallCount != 1)
+	{
+		UE_LOG(LogCore, Error, TEXT("HandleError re-entered.") );
+		return;
+	}
+
+	// Trigger the OnSystemFailure hook if it exists
+	FCoreDelegates::OnHandleSystemError.Broadcast();
+
+#if !PLATFORM_EXCEPTIONS_DISABLED
+	try
+	{
+#endif // !PLATFORM_EXCEPTIONS_DISABLED
+		GIsGuarded = 0;
+		GIsRunning = 0;
+		GIsCriticalError = 1;
+		GLogConsole = NULL;
+		GErrorHist[ARRAY_COUNT(GErrorHist)-1] = 0;
+
+		// Dump the error and flush the log.
+		UE_LOG(LogCore, Log, TEXT("=== Critical error: ===") LINE_TERMINATOR TEXT("%s") LINE_TERMINATOR, GErrorExceptionDescription);
+		UE_LOG(LogCore, Log, GErrorHist);
+
+		GLog->Flush();
+
+		HandleErrorRestoreUI();
+
+		FPlatformMisc::SubmitErrorReport(GErrorHist, EErrorReportMode::Interactive);
+		FCoreDelegates::OnShutdownAfterError.Broadcast();
+#if !PLATFORM_EXCEPTIONS_DISABLED
+	}
+	catch(...)
+	{}
+#endif // !PLATFORM_EXCEPTIONS_DISABLED
+}
+
+void FUnixErrorOutputDevice::HandleErrorRestoreUI()
+{
+}
+
diff --git a/Engine/Source/Runtime/Core/Private/Unix/UnixPlatformOutputDevices.cpp b/Engine/Source/Runtime/Core/Private/Unix/UnixPlatformOutputDevices.cpp
--- a/Engine/Source/Runtime/Core/Private/Unix/UnixPlatformOutputDevices.cpp
+++ b/Engine/Source/Runtime/Core/Private/Unix/UnixPlatformOutputDevices.cpp
@@ -12,6 +12,7 @@
 #include "Misc/App.h"
 #include "Misc/OutputDeviceConsole.h"
 #include "Misc/OutputDeviceRedirector.h"
+#include "Unix/UnixErrorOutputDevice.h"
 
 void FUnixOutputDevices::SetupOutputDevices()
 {
@@ -44,3 +45,10 @@ class FOutputDevice* FUnixOutputDevices::GetEventLog()
 {
 	return NULL; // @TODO No event logging
 }
+
+FOutputDeviceError* FUnixOutputDevices::GetError()
+{
+	static FUnixErrorOutputDevice Singleton;
+	return &Singleton;
+}
+
diff --git a/Engine/Source/Runtime/Core/Private/Windows/WindowsErrorOutputDevice.cpp b/Engine/Source/Runtime/Core/Private/Windows/WindowsErrorOutputDevice.cpp
new file mode 100644
--- /dev/null
+++ b/Engine/Source/Runtime/Core/Private/Windows/WindowsErrorOutputDevice.cpp
@@ -0,0 +1,107 @@
+// Copyright 1998-2018 Epic Games, Inc. All Rights Reserved.
+
+#include "Windows/WindowsErrorOutputDevice.h"
+#include "Logging/LogMacros.h"
+#include "Misc/OutputDevice.h"
+#include "HAL/PlatformTime.h"
+#include "Math/UnrealMathUtility.h"
+#include "HAL/UnrealMemory.h"
+#include "Templates/UnrealTemplate.h"
+#include "CoreGlobals.h"
+#include "Misc/CString.h"
+#include "Misc/Parse.h"
+#include "Misc/CommandLine.h"
+#include "Internationalization/Internationalization.h"
+#include "Misc/ConfigCacheIni.h"
+#include "Misc/OutputDeviceHelper.h"
+#include "Misc/CoreDelegates.h"
+#include "Misc/App.h"
+#include "Misc/OutputDeviceRedirector.h"
+#include "HAL/ThreadHeartBeat.h"
+#include "Windows/WindowsHWrapper.h"
+
+FWindowsErrorOutputDevice::FWindowsErrorOutputDevice()
+{
+}
+
+void FWindowsErrorOutputDevice::Serialize( const TCHAR* Msg, ELogVerbosity::Type Verbosity, const class FName& Category )
+{
+	UE_DEBUG_BREAK();
+   
+	if( !GIsCriticalError )
+	{   
+		const int32 LastError = ::GetLastError();
+
+		// First appError.
+		GIsCriticalError = 1;
+		TCHAR ErrorBuffer[1024];
+		ErrorBuffer[0] = 0;
+
+		// Windows error.
+		if (LastError == 0)
+		{
+			UE_LOG(LogWindows, Log, TEXT("Windows GetLastError: %s (%i)"), FPlatformMisc::GetSystemErrorMessage(ErrorBuffer, 1024, LastError), LastError);
+		}
+		else
+		{
+			UE_LOG(LogWindows, Error, TEXT("Windows GetLastError: %s (%i)"), FPlatformMisc::GetSystemErrorMessage(ErrorBuffer, 1024, LastError), LastError);
+		}
+	}
+	else
+	{
+		UE_LOG(LogWindows, Error, TEXT("Error reentered: %s"), Msg );
+	}
+
+	if( GIsGuarded )
+	{
+		// Propagate error so structured exception handler can perform necessary work.
+#if PLATFORM_EXCEPTIONS_DISABLED
+		UE_DEBUG_BREAK();
+#endif
+		FPlatformMisc::RaiseException( 1 );
+	}
+	else
+	{
+		// We crashed outside the guarded code (e.g. appExit).
+		HandleError();
+		FPlatformMisc::RequestExit( true );
+	}
+}
+
+void FWindowsErrorOutputDevice::HandleError()
+{
+	// make sure we don't report errors twice
+	static int32 CallCount = 0;
+	int32 NewCallCount = FPlatformAtomics::InterlockedIncrement(&CallCount);
+	if (NewCallCount != 1)
+	{
+		UE_LOG(LogWindows, Error, TEXT("HandleError re-entered.") );
+		return;
+	}
+	
+	GIsGuarded				= 0;
+	GIsRunning				= 0;
+	GIsCriticalError		= 1;
+	GLogConsole				= NULL;
+	GErrorHist[ARRAY_COUNT(GErrorHist)-1]=0;
+
+	// Trigger the OnSystemFailure hook if it exists
+	// make sure it happens after GIsGuarded is set to 0 in case this hook crashes
+	FCoreDelegates::OnHandleSystemError.Broadcast();
+
+	// Dump the error and flush the log.
+#if !NO_LOGGING
+	FDebug::LogFormattedMessageWithCallstack(LogWindows.GetCategoryName(), __FILE__, __LINE__, TEXT("=== Critical error: ==="), GErrorHist, ELogVerbosity::Error);
+#endif
+	GLog->PanicFlushThreadedLogs();
+
+	HandleErrorRestoreUI();
+
+	FPlatformMisc::SubmitErrorReport( GErrorHist, EErrorReportMode::Interactive );
+
+	FCoreDelegates::OnShutdownAfterError.Broadcast();
+}
+
+void FWindowsErrorOutputDevice::HandleErrorRestoreUI()
+{
+}
diff --git a/Engine/Source/Runtime/Core/Private/Windows/WindowsPlatformOutputDevices.cpp b/Engine/Source/Runtime/Core/Private/Windows/WindowsPlatformOutputDevices.cpp
--- a/Engine/Source/Runtime/Core/Private/Windows/WindowsPlatformOutputDevices.cpp
+++ b/Engine/Source/Runtime/Core/Private/Windows/WindowsPlatformOutputDevices.cpp
@@ -17,6 +17,7 @@
 #include "Misc/CoreDelegates.h"
 #include "Misc/App.h"
 #include "HAL/FeedbackContextAnsi.h"
+#include "Windows/WindowsErrorOutputDevice.h"
 #include "Windows/WindowsEventLogOutputDevice.h"
 #include "HAL/ThreadHeartBeat.h"
 
@@ -41,3 +42,10 @@ class FOutputDevice* FWindowsPlatformOutputDevices::GetEventLog()
 	return NULL;
 #endif //WANTS_WINDOWS_EVENT_LOGGING
 }
+
+FOutputDeviceError* FWindowsPlatformOutputDevices::GetError()
+{
+	static FWindowsErrorOutputDevice Singleton;
+	return &Singleton;
+}
+
diff --git a/Engine/Source/Runtime/Core/Public/Android/AndroidErrorOutputDevice.h b/Engine/Source/Runtime/Core/Public/Android/AndroidErrorOutputDevice.h
new file mode 100644
--- /dev/null
+++ b/Engine/Source/Runtime/Core/Public/Android/AndroidErrorOutputDevice.h
@@ -0,0 +1,30 @@
+// Copyright 1998-2018 Epic Games, Inc. All Rights Reserved.
+
+#pragma once
+
+#include "Misc/OutputDevice.h"
+#include "Misc/OutputDeviceError.h"
+
+class CORE_API FAndroidErrorOutputDevice : public FOutputDeviceError
+{
+public:
+	/** Constructor, initializing member variables */
+	FAndroidErrorOutputDevice();
+
+	/**
+	 * Serializes the passed in data unless the current event is suppressed.
+	 *
+	 * @param	Data	Text to log
+	 * @param	Event	Event name used for suppression purposes
+	 */
+	virtual void Serialize( const TCHAR* Msg, ELogVerbosity::Type Verbosity, const class FName& Category ) override;
+
+	/**
+	 * Error handling function that is being called from within the system wide global
+	 * error handler, e.g. using structured exception handling on the PC.
+	 */
+	void HandleError();
+
+private:
+	int32		ErrorPos;
+};
diff --git a/Engine/Source/Runtime/Core/Public/GenericPlatform/GenericPlatformOutputDevices.h b/Engine/Source/Runtime/Core/Public/GenericPlatform/GenericPlatformOutputDevices.h
--- a/Engine/Source/Runtime/Core/Public/GenericPlatform/GenericPlatformOutputDevices.h
+++ b/Engine/Source/Runtime/Core/Public/GenericPlatform/GenericPlatformOutputDevices.h
@@ -23,6 +23,9 @@ struct CORE_API FGenericPlatformOutputDevices
 		return nullptr; // normally only used for dedicated servers
 	}
 
+	static FOutputDeviceError*			GetError();
+	static FFeedbackContext*            GetFeedbackContext();
+
 protected:
 	static TCHAR CachedAbsoluteFilename[1024];
 };
diff --git a/Engine/Source/Runtime/Core/Public/IOS/IOSErrorOutputDevice.h b/Engine/Source/Runtime/Core/Public/IOS/IOSErrorOutputDevice.h
new file mode 100644
--- /dev/null
+++ b/Engine/Source/Runtime/Core/Public/IOS/IOSErrorOutputDevice.h
@@ -0,0 +1,30 @@
+// Copyright 1998-2018 Epic Games, Inc. All Rights Reserved.
+
+#pragma once
+
+#include "Misc/OutputDevice.h"
+#include "Misc/OutputDeviceError.h"
+
+class CORE_API FIOSErrorOutputDevice : public FOutputDeviceError
+{
+public:
+	/** Constructor, initializing member variables */
+	FIOSErrorOutputDevice();
+
+	/**
+	 * Serializes the passed in data unless the current event is suppressed.
+	 *
+	 * @param	Data	Text to log
+	 * @param	Event	Event name used for suppression purposes
+	 */
+	virtual void Serialize( const TCHAR* Msg, ELogVerbosity::Type Verbosity, const class FName& Category ) override;
+
+	/**
+	 * Error handling function that is being called from within the system wide global
+	 * error handler, e.g. using structured exception handling on the PC.
+	 */
+	void HandleError();
+
+private:
+	int32		ErrorPos;
+};
diff --git a/Engine/Source/Runtime/Core/Public/Mac/MacErrorOutputDevice.h b/Engine/Source/Runtime/Core/Public/Mac/MacErrorOutputDevice.h
new file mode 100644
--- /dev/null
+++ b/Engine/Source/Runtime/Core/Public/Mac/MacErrorOutputDevice.h
@@ -0,0 +1,36 @@
+// Copyright 1998-2018 Epic Games, Inc. All Rights Reserved.
+
+#pragma once
+
+#include "CoreMinimal.h"
+#include "Misc/OutputDeviceError.h"
+
+class CORE_API FMacErrorOutputDevice : public FOutputDeviceError
+{
+public:
+	/** Constructor, initializing member variables */
+	FMacErrorOutputDevice();
+
+	/**
+	 * Serializes the passed in data unless the current event is suppressed.
+	 *
+	 * @param	Data	Text to log
+	 * @param	Event	Event name used for suppression purposes
+	 */
+	virtual void Serialize( const TCHAR* Msg, ELogVerbosity::Type Verbosity, const class FName& Category ) override;
+
+	/**
+	 * Error handling function that is being called from within the system wide global
+	 * error handler, e.g. using structured exception handling on the PC.
+	 */
+	virtual void HandleError() override;
+
+private:
+	int32		ErrorPos;
+
+protected:
+	/**
+	 * Callback to allow FMacApplicationErrorOutputDevice to restore the UI.
+	 */
+	virtual void HandleErrorRestoreUI();
+};
diff --git a/Engine/Source/Runtime/Core/Public/Mac/MacPlatformOutputDevices.h b/Engine/Source/Runtime/Core/Public/Mac/MacPlatformOutputDevices.h
--- a/Engine/Source/Runtime/Core/Public/Mac/MacPlatformOutputDevices.h
+++ b/Engine/Source/Runtime/Core/Public/Mac/MacPlatformOutputDevices.h
@@ -10,6 +10,7 @@
 struct CORE_API FMacPlatformOutputDevices : public FGenericPlatformOutputDevices
 {
 	static FOutputDevice*			GetEventLog();
+	static FOutputDeviceError*      GetError();
 };
 
 typedef FMacPlatformOutputDevices FPlatformOutputDevices;
diff --git a/Engine/Source/Runtime/Core/Public/Unix/UnixErrorOutputDevice.h b/Engine/Source/Runtime/Core/Public/Unix/UnixErrorOutputDevice.h
new file mode 100644
--- /dev/null
+++ b/Engine/Source/Runtime/Core/Public/Unix/UnixErrorOutputDevice.h
@@ -0,0 +1,38 @@
+// Copyright 1998-2018 Epic Games, Inc. All Rights Reserved.
+
+#pragma once
+
+#include "CoreTypes.h"
+#include "Misc/OutputDeviceError.h"
+
+class CORE_API FUnixErrorOutputDevice : public FOutputDeviceError
+{
+public:
+
+	/** Constructor, initializing member variables */
+	FUnixErrorOutputDevice();
+
+	/**
+	 * Serializes the passed in data unless the current event is suppressed.
+	 *
+	 * @param	Data	Text to log
+	 * @param	Event	Event name used for suppression purposes
+	 */
+	virtual void Serialize(const TCHAR* Msg, ELogVerbosity::Type Verbosity, const class FName& Category) override;
+
+	/**
+	 * Error handling function that is being called from within the system wide global
+	 * error handler, e.g. using structured exception handling on the PC.
+	 */
+	void HandleError() override;
+
+private:
+
+	int32	ErrorPos;
+
+protected:
+	/**
+	 * Callback to allow FUnixApplicationErrorOutputDevice to restore the UI.
+	 */
+	virtual void HandleErrorRestoreUI();
+};
diff --git a/Engine/Source/Runtime/Core/Public/Unix/UnixPlatformOutputDevices.h b/Engine/Source/Runtime/Core/Public/Unix/UnixPlatformOutputDevices.h
--- a/Engine/Source/Runtime/Core/Public/Unix/UnixPlatformOutputDevices.h
+++ b/Engine/Source/Runtime/Core/Public/Unix/UnixPlatformOutputDevices.h
@@ -15,6 +15,7 @@ struct CORE_API FUnixOutputDevices : public FGenericPlatformOutputDevices
 	static void SetupOutputDevices();
 	static FString GetAbsoluteLogFilename();
 	static FOutputDevice* GetEventLog();
+	static FOutputDeviceError* GetError();
 };
 
 typedef FUnixOutputDevices FPlatformOutputDevices;
diff --git a/Engine/Source/Runtime/Core/Public/Windows/WindowsErrorOutputDevice.h b/Engine/Source/Runtime/Core/Public/Windows/WindowsErrorOutputDevice.h
new file mode 100644
--- /dev/null
+++ b/Engine/Source/Runtime/Core/Public/Windows/WindowsErrorOutputDevice.h
@@ -0,0 +1,35 @@
+// Copyright 1998-2018 Epic Games, Inc. All Rights Reserved.
+
+#pragma once
+
+#include "CoreTypes.h"
+#include "Misc/OutputDeviceError.h"
+#include "Misc/OutputDeviceConsole.h"
+
+class CORE_API FWindowsErrorOutputDevice : public FOutputDeviceError
+{
+public:
+	/** Constructor, initializing member variables */
+	FWindowsErrorOutputDevice();
+
+	/**
+	 * Serializes the passed in data unless the current event is suppressed.
+	 *
+	 * @param	Data	Text to log
+	 * @param	Event	Event name used for suppression purposes
+	 */
+	virtual void Serialize( const TCHAR* Msg, ELogVerbosity::Type Verbosity, const class FName& Category ) override;
+
+	/**
+	 * Error handling function that is being called from within the system wide global
+	 * error handler, e.g. using structured exception handling on the PC.
+	 */
+	virtual void HandleError() override;
+
+protected:
+	/**
+	 * Callback to allow FWindowsApplicationErrorOutputDevice to restore the UI.
+	 */
+	virtual void HandleErrorRestoreUI();
+};
+
diff --git a/Engine/Source/Runtime/Core/Public/Windows/WindowsPlatformOutputDevices.h b/Engine/Source/Runtime/Core/Public/Windows/WindowsPlatformOutputDevices.h
--- a/Engine/Source/Runtime/Core/Public/Windows/WindowsPlatformOutputDevices.h
+++ b/Engine/Source/Runtime/Core/Public/Windows/WindowsPlatformOutputDevices.h
@@ -14,6 +14,7 @@ struct CORE_API FWindowsPlatformOutputDevices
 	: public FGenericPlatformOutputDevices
 {
 	static FOutputDevice*			GetEventLog();
+	static FOutputDeviceError*      GetError();
 };
 
 
