Fix non-unity compile error in Stats2.h.

#rnx

diff --git a/Engine/Source/Runtime/Core/Public/Stats/Stats.h b/Engine/Source/Runtime/Core/Public/Stats/Stats.h
--- a/Engine/Source/Runtime/Core/Public/Stats/Stats.h
+++ b/Engine/Source/Runtime/Core/Public/Stats/Stats.h
@@ -8,6 +8,9 @@
 #include "Delegates/Delegate.h"
 #include "HAL/PlatformTime.h"
 #include "HAL/ThreadSingleton.h"
+#include "StatsCommon.h"
+#include "ProfilingDebugging/UMemoryDefines.h"
+#include "Stats/Stats2.h"
 
 /** 
  *	Unreal Engine Stats system
@@ -289,16 +292,6 @@
  *				}
  */
 
-#define FORCEINLINE_STATS FORCEINLINE
-//#define FORCEINLINE_STATS FORCEINLINE_DEBUGGABLE
-#define checkStats(x)
-
-#if !defined(STATS)
-#error "STATS must be defined as either zero or one."
-#endif
-
-#include "ProfilingDebugging/UMemoryDefines.h"
-
 struct TStatId;
 
 // used by the profiler
@@ -323,8 +316,6 @@ enum EStatType
 	#define STAT(x)
 #endif
 
-#include "Stats/Stats2.h"
-
 #if STATS
 
 /**
diff --git a/Engine/Source/Runtime/Core/Public/Stats/Stats2.h b/Engine/Source/Runtime/Core/Public/Stats/Stats2.h
--- a/Engine/Source/Runtime/Core/Public/Stats/Stats2.h
+++ b/Engine/Source/Runtime/Core/Public/Stats/Stats2.h
@@ -15,6 +15,8 @@
 #include "Containers/ChunkedArray.h"
 #include "Delegates/Delegate.h"
 #include "Templates/Atomic.h"
+#include "Math/Color.h"
+#include "StatsCommon.h"
 
 class FScopeCycleCounter;
 class FThreadStats;
diff --git a/Engine/Source/Runtime/Core/Public/Stats/StatsCommon.h b/Engine/Source/Runtime/Core/Public/Stats/StatsCommon.h
new file mode 100644
--- /dev/null
+++ b/Engine/Source/Runtime/Core/Public/Stats/StatsCommon.h
@@ -0,0 +1,14 @@
+// Copyright 1998-2018 Epic Games, Inc. All Rights Reserved.
+
+#pragma once
+
+#include "CoreTypes.h"
+
+#define FORCEINLINE_STATS FORCEINLINE
+//#define FORCEINLINE_STATS FORCEINLINE_DEBUGGABLE
+#define checkStats(x)
+
+#if !defined(STATS)
+#error "STATS must be defined as either zero or one."
+#endif
+
