Remove CoreMinimal.h includes from Core.

#rnx

diff --git a/Engine/Source/Runtime/Core/Private/FramePro/FrameProProfiler.cpp b/Engine/Source/Runtime/Core/Private/FramePro/FrameProProfiler.cpp
--- a/Engine/Source/Runtime/Core/Private/FramePro/FrameProProfiler.cpp
+++ b/Engine/Source/Runtime/Core/Private/FramePro/FrameProProfiler.cpp
@@ -4,7 +4,7 @@
 
 #if FRAMEPRO_ENABLED
 
-#include "CoreMinimal.h"
+#include "CoreTypes.h"
 #include "HAL/ThreadSingleton.h"
 #include "HAL/IConsoleManager.h"
 #include "HAL/FileManager.h"
diff --git a/Engine/Source/Runtime/Core/Private/Internationalization/LocalizationResourceTextSource.h b/Engine/Source/Runtime/Core/Private/Internationalization/LocalizationResourceTextSource.h
--- a/Engine/Source/Runtime/Core/Private/Internationalization/LocalizationResourceTextSource.h
+++ b/Engine/Source/Runtime/Core/Private/Internationalization/LocalizationResourceTextSource.h
@@ -2,7 +2,7 @@
 
 #pragma once
 
-#include "CoreMinimal.h"
+#include "CoreTypes.h"
 #include "Internationalization/ILocalizedTextSource.h"
 
 /**
diff --git a/Engine/Source/Runtime/Core/Private/Internationalization/PolyglotTextSource.h b/Engine/Source/Runtime/Core/Private/Internationalization/PolyglotTextSource.h
--- a/Engine/Source/Runtime/Core/Private/Internationalization/PolyglotTextSource.h
+++ b/Engine/Source/Runtime/Core/Private/Internationalization/PolyglotTextSource.h
@@ -2,7 +2,7 @@
 
 #pragma once
 
-#include "CoreMinimal.h"
+#include "CoreTypes.h"
 #include "Internationalization/ILocalizedTextSource.h"
 #include "Internationalization/PolyglotTextData.h"
 #include "Internationalization/LocKeyFuncs.h"
diff --git a/Engine/Source/Runtime/Core/Private/Modules/SimpleParse.h b/Engine/Source/Runtime/Core/Private/Modules/SimpleParse.h
--- a/Engine/Source/Runtime/Core/Private/Modules/SimpleParse.h
+++ b/Engine/Source/Runtime/Core/Private/Modules/SimpleParse.h
@@ -2,7 +2,7 @@
 
 #pragma once
 
-#include "CoreMinimal.h"
+#include "CoreTypes.h"
 
 struct FSimpleParse
 {
diff --git a/Engine/Source/Runtime/Core/Public/Containers/BinaryHeap.h b/Engine/Source/Runtime/Core/Public/Containers/BinaryHeap.h
--- a/Engine/Source/Runtime/Core/Public/Containers/BinaryHeap.h
+++ b/Engine/Source/Runtime/Core/Public/Containers/BinaryHeap.h
@@ -3,7 +3,7 @@
 
 #pragma once
 
-#include "CoreMinimal.h"
+#include "CoreTypes.h"
 
 /*-----------------------------------------------------------------------------
 	Binary Heap, used to index another data structure.
diff --git a/Engine/Source/Runtime/Core/Public/Containers/HashTable.h b/Engine/Source/Runtime/Core/Public/Containers/HashTable.h
--- a/Engine/Source/Runtime/Core/Public/Containers/HashTable.h
+++ b/Engine/Source/Runtime/Core/Public/Containers/HashTable.h
@@ -2,7 +2,7 @@
 
 #pragma once
 
-#include "CoreMinimal.h"
+#include "CoreTypes.h"
 
 static FORCEINLINE uint32 Murmur32( std::initializer_list< uint32 > InitList )
 {
diff --git a/Engine/Source/Runtime/Core/Public/HAL/PlatformFileCommon.h b/Engine/Source/Runtime/Core/Public/HAL/PlatformFileCommon.h
--- a/Engine/Source/Runtime/Core/Public/HAL/PlatformFileCommon.h
+++ b/Engine/Source/Runtime/Core/Public/HAL/PlatformFileCommon.h
@@ -2,7 +2,7 @@
 
 #pragma once
 
-#include "CoreMinimal.h"
+#include "CoreTypes.h"
 #include "HAL/ThreadSafeCounter.h"
 #include "Misc/ScopeLock.h"
 
diff --git a/Engine/Source/Runtime/Core/Public/Hash/CityHash.h b/Engine/Source/Runtime/Core/Public/Hash/CityHash.h
--- a/Engine/Source/Runtime/Core/Public/Hash/CityHash.h
+++ b/Engine/Source/Runtime/Core/Public/Hash/CityHash.h
@@ -62,7 +62,7 @@
 #ifndef CITY_HASH_H_
 #define CITY_HASH_H_
 
-#include "CoreMinimal.h"
+#include "CoreTypes.h"
 
 struct Uint128_64
 {
diff --git a/Engine/Source/Runtime/Core/Public/Internationalization/ILocalizedTextSource.h b/Engine/Source/Runtime/Core/Public/Internationalization/ILocalizedTextSource.h
--- a/Engine/Source/Runtime/Core/Public/Internationalization/ILocalizedTextSource.h
+++ b/Engine/Source/Runtime/Core/Public/Internationalization/ILocalizedTextSource.h
@@ -2,7 +2,7 @@
 
 #pragma once
 
-#include "CoreMinimal.h"
+#include "CoreTypes.h"
 #include "Containers/ArrayView.h"
 #include "Internationalization/LocalizedTextSourceTypes.h"
 
diff --git a/Engine/Source/Runtime/Core/Public/Internationalization/PolyglotTextData.h b/Engine/Source/Runtime/Core/Public/Internationalization/PolyglotTextData.h
--- a/Engine/Source/Runtime/Core/Public/Internationalization/PolyglotTextData.h
+++ b/Engine/Source/Runtime/Core/Public/Internationalization/PolyglotTextData.h
@@ -2,7 +2,7 @@
 
 #pragma once
 
-#include "CoreMinimal.h"
+#include "CoreTypes.h"
 #include "Internationalization/LocalizedTextSourceTypes.h"
 
 /**
diff --git a/Engine/Source/Runtime/Core/Public/Internationalization/TextLocalizationResource.h b/Engine/Source/Runtime/Core/Public/Internationalization/TextLocalizationResource.h
--- a/Engine/Source/Runtime/Core/Public/Internationalization/TextLocalizationResource.h
+++ b/Engine/Source/Runtime/Core/Public/Internationalization/TextLocalizationResource.h
@@ -2,7 +2,7 @@
 
 #pragma once 
 
-#include "CoreMinimal.h"
+#include "CoreTypes.h"
 #include "Containers/SortedMap.h"
 #include "Internationalization/LocKeyFuncs.h"
 #include "Internationalization/LocalizedTextSourceTypes.h"
diff --git a/Engine/Source/Runtime/Core/Public/Internationalization/TextLocalizationResourceVersion.h b/Engine/Source/Runtime/Core/Public/Internationalization/TextLocalizationResourceVersion.h
--- a/Engine/Source/Runtime/Core/Public/Internationalization/TextLocalizationResourceVersion.h
+++ b/Engine/Source/Runtime/Core/Public/Internationalization/TextLocalizationResourceVersion.h
@@ -2,7 +2,7 @@
 
 #pragma once 
 
-#include "CoreMinimal.h"
+#include "CoreTypes.h"
 #include "Misc/Guid.h"
 
 /** Version information for compiled LocMeta (Localization MetaData Resource) and LocRes (Localization Resource) files */
diff --git a/Engine/Source/Runtime/Core/Public/Misc/AES.h b/Engine/Source/Runtime/Core/Public/Misc/AES.h
--- a/Engine/Source/Runtime/Core/Public/Misc/AES.h
+++ b/Engine/Source/Runtime/Core/Public/Misc/AES.h
@@ -2,7 +2,7 @@
 
 #pragma once
 
-#include "CoreMinimal.h"
+#include "CoreTypes.h"
 
 #define AES_BLOCK_SIZE 16
 
diff --git a/Engine/Source/Runtime/Core/Public/Misc/GenericGrowableAllocator.h b/Engine/Source/Runtime/Core/Public/Misc/GenericGrowableAllocator.h
--- a/Engine/Source/Runtime/Core/Public/Misc/GenericGrowableAllocator.h
+++ b/Engine/Source/Runtime/Core/Public/Misc/GenericGrowableAllocator.h
@@ -6,7 +6,7 @@
 
 #pragma once
 
-#include "CoreMinimal.h"
+#include "CoreTypes.h"
 #include "Misc/ScopeLock.h"
 #include "Stats/Stats.h"
 
diff --git a/Engine/Source/Runtime/Core/Public/Modules/BuildVersion.h b/Engine/Source/Runtime/Core/Public/Modules/BuildVersion.h
--- a/Engine/Source/Runtime/Core/Public/Modules/BuildVersion.h
+++ b/Engine/Source/Runtime/Core/Public/Modules/BuildVersion.h
@@ -2,7 +2,7 @@
 
 #pragma once
 
-#include "CoreMinimal.h"
+#include "CoreTypes.h"
 
 class FJsonObject;
 
diff --git a/Engine/Source/Runtime/Core/Public/Modules/ModuleManifest.h b/Engine/Source/Runtime/Core/Public/Modules/ModuleManifest.h
--- a/Engine/Source/Runtime/Core/Public/Modules/ModuleManifest.h
+++ b/Engine/Source/Runtime/Core/Public/Modules/ModuleManifest.h
@@ -2,7 +2,7 @@
 
 #pragma once
 
-#include "CoreMinimal.h"
+#include "CoreTypes.h"
 #include "Modules/BuildVersion.h"
 
 /**
diff --git a/Engine/Source/Runtime/Core/Public/PixelFormat.h b/Engine/Source/Runtime/Core/Public/PixelFormat.h
--- a/Engine/Source/Runtime/Core/Public/PixelFormat.h
+++ b/Engine/Source/Runtime/Core/Public/PixelFormat.h
@@ -3,7 +3,7 @@
 
 #pragma once
 
-#include "CoreMinimal.h"
+#include "CoreTypes.h"
 
 #if defined(PF_MAX)
 #undef PF_MAX
diff --git a/Engine/Source/Runtime/Core/Public/ProfilingDebugging/CsvProfiler.h b/Engine/Source/Runtime/Core/Public/ProfilingDebugging/CsvProfiler.h
--- a/Engine/Source/Runtime/Core/Public/ProfilingDebugging/CsvProfiler.h
+++ b/Engine/Source/Runtime/Core/Public/ProfilingDebugging/CsvProfiler.h
@@ -8,7 +8,6 @@
 #pragma once
 
 #include "CoreTypes.h"
-#include "CoreMinimal.h"
 #include "Containers/Queue.h"
 
 #if WITH_SERVER_CODE
