Fixing non-unity errors.

#rnx

diff --git a/Engine/Source/Runtime/Core/Private/Hash/CityHash.cpp b/Engine/Source/Runtime/Core/Private/Hash/CityHash.cpp
--- a/Engine/Source/Runtime/Core/Private/Hash/CityHash.cpp
+++ b/Engine/Source/Runtime/Core/Private/Hash/CityHash.cpp
@@ -29,6 +29,7 @@
 
 //#include "config.h"
 #include "Hash/CityHash.h"
+#include "HAL/UnrealMemory.h"
 
 
 static uint64 UNALIGNED_LOAD64(const char *p) {
diff --git a/Engine/Source/Runtime/Core/Private/Internationalization/PolyglotTextData.cpp b/Engine/Source/Runtime/Core/Private/Internationalization/PolyglotTextData.cpp
--- a/Engine/Source/Runtime/Core/Private/Internationalization/PolyglotTextData.cpp
+++ b/Engine/Source/Runtime/Core/Private/Internationalization/PolyglotTextData.cpp
@@ -3,6 +3,7 @@
 #include "Internationalization/PolyglotTextData.h"
 #include "Internationalization/TextLocalizationManager.h"
 #include "Internationalization/TextLocalizationResource.h"
+#include "Internationalization/Internationalization.h"
 
 #define LOCTEXT_NAMESPACE "PolyglotTextData"
 
diff --git a/Engine/Source/Runtime/Core/Private/Internationalization/TextLocalizationResource.cpp b/Engine/Source/Runtime/Core/Private/Internationalization/TextLocalizationResource.cpp
--- a/Engine/Source/Runtime/Core/Private/Internationalization/TextLocalizationResource.cpp
+++ b/Engine/Source/Runtime/Core/Private/Internationalization/TextLocalizationResource.cpp
@@ -6,6 +6,8 @@
 #include "HAL/FileManager.h"
 #include "Misc/Paths.h"
 #include "Misc/Optional.h"
+#include "Templates/UniquePtr.h"
+#include "Internationalization/Internationalization.h"
 
 DEFINE_LOG_CATEGORY_STATIC(LogTextLocalizationResource, Log, All);
 
diff --git a/Engine/Source/Runtime/Core/Private/Modules/SimpleParse.h b/Engine/Source/Runtime/Core/Private/Modules/SimpleParse.h
--- a/Engine/Source/Runtime/Core/Private/Modules/SimpleParse.h
+++ b/Engine/Source/Runtime/Core/Private/Modules/SimpleParse.h
@@ -3,6 +3,7 @@
 #pragma once
 
 #include "CoreTypes.h"
+#include "Containers/UnrealString.h"
 
 struct FSimpleParse
 {
diff --git a/Engine/Source/Runtime/Core/Private/ProfilingDebugging/CsvProfiler.cpp b/Engine/Source/Runtime/Core/Private/ProfilingDebugging/CsvProfiler.cpp
--- a/Engine/Source/Runtime/Core/Private/ProfilingDebugging/CsvProfiler.cpp
+++ b/Engine/Source/Runtime/Core/Private/ProfilingDebugging/CsvProfiler.cpp
@@ -1,7 +1,6 @@
 // Copyright 1998-2018 Epic Games, Inc. All Rights Reserved.
 
 /**
-*
 * A lightweight multi-threaded CSV profiler which can be used for profiling in Test/Shipping builds
 */
 
diff --git a/Engine/Source/Runtime/Core/Public/Containers/ContainersFwd.h b/Engine/Source/Runtime/Core/Public/Containers/ContainersFwd.h
--- a/Engine/Source/Runtime/Core/Public/Containers/ContainersFwd.h
+++ b/Engine/Source/Runtime/Core/Public/Containers/ContainersFwd.h
@@ -15,4 +15,6 @@ template<typename KeyType, typename ValueType, typename SetAllocator = FDefaultS
 template<typename KeyType, typename ValueType, typename SetAllocator = FDefaultSetAllocator, typename KeyFuncs = TDefaultMapHashableKeyFuncs<KeyType, ValueType, true > > class TMultiMap;
 template <typename T = void > struct TLess;
 template<typename KeyType, typename ValueType, typename ArrayAllocator = FDefaultAllocator, typename SortPredicate = TLess<KeyType> > class TSortedMap;
+template<typename ElementType,bool bInAllowDuplicateKeys = false> struct DefaultKeyFuncs;
+template<typename InElementType, typename KeyFuncs = DefaultKeyFuncs<InElementType>, typename Allocator = FDefaultSetAllocator> class TSet;
 /// @endcond
diff --git a/Engine/Source/Runtime/Core/Public/Containers/HashTable.h b/Engine/Source/Runtime/Core/Public/Containers/HashTable.h
--- a/Engine/Source/Runtime/Core/Public/Containers/HashTable.h
+++ b/Engine/Source/Runtime/Core/Public/Containers/HashTable.h
@@ -3,6 +3,10 @@
 #pragma once
 
 #include "CoreTypes.h"
+#include "Misc/AssertionMacros.h"
+#include "HAL/UnrealMemory.h"
+#include "Math/UnrealMathUtility.h"
+#include <initializer_list>
 
 static FORCEINLINE uint32 Murmur32( std::initializer_list< uint32 > InitList )
 {
diff --git a/Engine/Source/Runtime/Core/Public/Containers/Set.h b/Engine/Source/Runtime/Core/Public/Containers/Set.h
--- a/Engine/Source/Runtime/Core/Public/Containers/Set.h
+++ b/Engine/Source/Runtime/Core/Public/Containers/Set.h
@@ -18,6 +18,7 @@
 #include "Templates/AreTypesEqual.h"
 #include "Templates/Decay.h"
 #include "Serialization/StructuredArchive.h"
+#include "ContainersFwd.h"
 
 /**
  * The base KeyFuncs type with some useful definitions for all KeyFuncs; meant to be derived from instead of used directly.
@@ -37,7 +38,7 @@ struct BaseKeyFuncs
 /**
  * A default implementation of the KeyFuncs used by TSet which uses the element as a key.
  */
-template<typename ElementType,bool bInAllowDuplicateKeys = false>
+template<typename ElementType,bool bInAllowDuplicateKeys /*= false*/>
 struct DefaultKeyFuncs : BaseKeyFuncs<ElementType,ElementType,bInAllowDuplicateKeys>
 {
 	typedef typename TCallTraits<ElementType>::ParamType KeyInitType;
@@ -66,14 +67,6 @@ struct DefaultKeyFuncs : BaseKeyFuncs<ElementType,ElementType,bInAllowDuplicateK
 	}
 };
 
-// Forward declaration.
-template<
-	typename InElementType,
-	typename KeyFuncs = DefaultKeyFuncs<InElementType>,
-	typename Allocator = FDefaultSetAllocator
-	>
-class TSet;
-
 /** This is used to provide type specific behavior for a move which will destroy B. */
 /** Should be in UnrealTemplate but isn't for Clang build reasons - will move later */
 template<typename T>
diff --git a/Engine/Source/Runtime/Core/Public/Internationalization/ILocalizedTextSource.h b/Engine/Source/Runtime/Core/Public/Internationalization/ILocalizedTextSource.h
--- a/Engine/Source/Runtime/Core/Public/Internationalization/ILocalizedTextSource.h
+++ b/Engine/Source/Runtime/Core/Public/Internationalization/ILocalizedTextSource.h
@@ -3,8 +3,10 @@
 #pragma once
 
 #include "CoreTypes.h"
+#include "Containers/ContainersFwd.h"
 #include "Containers/ArrayView.h"
 #include "Internationalization/LocalizedTextSourceTypes.h"
+#include "CoreGlobals.h"
 
 class FTextLocalizationResource;
 class FTextLocalizationResources;
diff --git a/Engine/Source/Runtime/Core/Public/Internationalization/PolyglotTextData.h b/Engine/Source/Runtime/Core/Public/Internationalization/PolyglotTextData.h
--- a/Engine/Source/Runtime/Core/Public/Internationalization/PolyglotTextData.h
+++ b/Engine/Source/Runtime/Core/Public/Internationalization/PolyglotTextData.h
@@ -4,6 +4,9 @@
 
 #include "CoreTypes.h"
 #include "Internationalization/LocalizedTextSourceTypes.h"
+#include "Containers/UnrealString.h"
+#include "Containers/Map.h"
+#include "Internationalization/Text.h"
 
 /**
  * Polyglot data that may be registered to the text localization manager at runtime.
diff --git a/Engine/Source/Runtime/Core/Public/Modules/BuildVersion.h b/Engine/Source/Runtime/Core/Public/Modules/BuildVersion.h
--- a/Engine/Source/Runtime/Core/Public/Modules/BuildVersion.h
+++ b/Engine/Source/Runtime/Core/Public/Modules/BuildVersion.h
@@ -3,6 +3,7 @@
 #pragma once
 
 #include "CoreTypes.h"
+#include "Containers/UnrealString.h"
 
 class FJsonObject;
 
diff --git a/Engine/Source/Runtime/Core/Public/Modules/ModuleManifest.h b/Engine/Source/Runtime/Core/Public/Modules/ModuleManifest.h
--- a/Engine/Source/Runtime/Core/Public/Modules/ModuleManifest.h
+++ b/Engine/Source/Runtime/Core/Public/Modules/ModuleManifest.h
@@ -4,6 +4,8 @@
 
 #include "CoreTypes.h"
 #include "Modules/BuildVersion.h"
+#include "Containers/Map.h"
+#include "Containers/UnrealString.h"
 
 /**
  * Stores a record of a built target, with all metadata that other tools may need to know about the build.
diff --git a/Engine/Source/Runtime/Core/Public/ProfilingDebugging/CsvProfiler.h b/Engine/Source/Runtime/Core/Public/ProfilingDebugging/CsvProfiler.h
--- a/Engine/Source/Runtime/Core/Public/ProfilingDebugging/CsvProfiler.h
+++ b/Engine/Source/Runtime/Core/Public/ProfilingDebugging/CsvProfiler.h
@@ -9,6 +9,8 @@
 
 #include "CoreTypes.h"
 #include "Containers/Queue.h"
+#include "UObject/NameTypes.h"
+#include "Templates/UniquePtr.h"
 
 #if WITH_SERVER_CODE
   #define CSV_PROFILER (WITH_ENGINE && 1)
