Another non-unity fix.

#rnx

diff --git a/Engine/Source/Runtime/Core/Private/Misc/AES.cpp b/Engine/Source/Runtime/Core/Private/Misc/AES.cpp
--- a/Engine/Source/Runtime/Core/Private/Misc/AES.cpp
+++ b/Engine/Source/Runtime/Core/Private/Misc/AES.cpp
@@ -3,7 +3,6 @@
 #include "Misc/AES.h"
 #include "Misc/AssertionMacros.h"
 #include "Misc/CString.h"
-#include "HAL/UnrealMemory.h"
 #include "Misc/CoreDelegates.h"
 
 // This is using the reference implementation of rijndael encryption algorithm
diff --git a/Engine/Source/Runtime/Core/Public/Misc/AES.h b/Engine/Source/Runtime/Core/Public/Misc/AES.h
--- a/Engine/Source/Runtime/Core/Public/Misc/AES.h
+++ b/Engine/Source/Runtime/Core/Public/Misc/AES.h
@@ -3,6 +3,7 @@
 #pragma once
 
 #include "CoreTypes.h"
+#include "HAL/UnrealMemory.h"
 
 #define AES_BLOCK_SIZE 16
 
