Rewrite some Windows utility functions to support paths longer than MAX_PATH.

diff --git a/Engine/Source/Runtime/Core/Private/Windows/WindowsPlatformProcess.cpp b/Engine/Source/Runtime/Core/Private/Windows/WindowsPlatformProcess.cpp
--- a/Engine/Source/Runtime/Core/Private/Windows/WindowsPlatformProcess.cpp
+++ b/Engine/Source/Runtime/Core/Private/Windows/WindowsPlatformProcess.cpp
@@ -820,12 +820,16 @@ const TCHAR* FWindowsPlatformProcess::UserDir()
 	static FString WindowsUserDir;
 	if( !WindowsUserDir.Len() )
 	{
-		TCHAR UserPath[MAX_PATH];
-		// get the My Documents directory
-		HRESULT Ret = SHGetFolderPath(NULL, CSIDL_PERSONAL, NULL, SHGFP_TYPE_CURRENT, UserPath);
+		TCHAR* UserPath;
 
-		// make the base user dir path
-		WindowsUserDir = FString(UserPath).Replace(TEXT("\\"), TEXT("/")) + TEXT("/");
+		// get the My Documents directory
+		HRESULT Ret = SHGetKnownFolderPath(FOLDERID_Documents, 0, NULL, &UserPath);
+		if (SUCCEEDED(Ret))
+		{
+			// make the base user dir path
+			WindowsUserDir = FString(UserPath).Replace(TEXT("\\"), TEXT("/")) + TEXT("/");
+			CoTaskMemFree(UserPath);
+		}
 	}
 	return *WindowsUserDir;
 }
@@ -855,12 +859,16 @@ const TCHAR* FWindowsPlatformProcess::UserSettingsDir()
 	static FString WindowsUserSettingsDir;
 	if (!WindowsUserSettingsDir.Len())
 	{
-		TCHAR UserPath[MAX_PATH];
-		// get the My Documents directory
-		HRESULT Ret = SHGetFolderPath(NULL, CSIDL_LOCAL_APPDATA, NULL, SHGFP_TYPE_CURRENT, UserPath);
+		TCHAR* UserPath;
 
-		// make the base user dir path
-		WindowsUserSettingsDir = FString(UserPath).Replace(TEXT("\\"), TEXT("/")) + TEXT("/");
+		// get the local AppData directory
+		HRESULT Ret = SHGetKnownFolderPath(FOLDERID_LocalAppData, 0, NULL, &UserPath);
+		if (SUCCEEDED(Ret))
+		{
+			// make the base user dir path
+			WindowsUserSettingsDir = FString(UserPath).Replace(TEXT("\\"), TEXT("/")) + TEXT("/");
+			CoTaskMemFree(UserPath);
+		}
 	}
 	return *WindowsUserSettingsDir;
 }
@@ -870,13 +878,16 @@ const TCHAR* FWindowsPlatformProcess::ApplicationSettingsDir()
 	static FString WindowsApplicationSettingsDir;
 	if( !WindowsApplicationSettingsDir.Len() )
 	{
-		TCHAR ApplictionSettingsPath[MAX_PATH];
-		// get the My Documents directory
-		HRESULT Ret = SHGetFolderPath(NULL, CSIDL_COMMON_APPDATA, NULL, SHGFP_TYPE_CURRENT, ApplictionSettingsPath);
+		TCHAR* ApplictionSettingsPath;
 
-		// make the base user dir path
-		// @todo rocket this folder should be based on your company name, not just be hard coded to /Epic/
-		WindowsApplicationSettingsDir = FString(ApplictionSettingsPath).Replace(TEXT("\\"), TEXT("/")) + TEXT("/Epic/");
+		// get the local AppData directory
+		HRESULT Ret = SHGetKnownFolderPath(FOLDERID_ProgramData, 0, NULL, &ApplictionSettingsPath);
+		if (SUCCEEDED(Ret))
+		{
+			// make the base user dir path
+			WindowsApplicationSettingsDir = FString(ApplictionSettingsPath).Replace(TEXT("\\"), TEXT("/")) + TEXT("/Epic/");
+			CoTaskMemFree(ApplictionSettingsPath);
+		}
 	}
 	return *WindowsApplicationSettingsDir;
 }
@@ -930,10 +941,29 @@ void FWindowsPlatformProcess::SetCurrentWorkingDirectoryToBaseDir()
 /** Get the current working directory (only really makes sense on desktop platforms) */
 FString FWindowsPlatformProcess::GetCurrentWorkingDirectory()
 {
-	// get the current working directory (uncached)
-	TCHAR CurrentDirectory[MAX_PATH];
-	GetCurrentDirectoryW(MAX_PATH, CurrentDirectory);
-	return CurrentDirectory;
+	// Attempt to get the environment variable into a local buffer. If it succeeds, we can just copy the result into a string. If it fails, we have the length of the buffer we need to allocate.
+	TCHAR LocalBuffer[260];
+	uint32 Length = ::GetCurrentDirectoryW(ARRAY_COUNT(LocalBuffer), LocalBuffer);
+	if (Length < ARRAY_COUNT(LocalBuffer))
+	{
+		return LocalBuffer;
+	}
+
+	// Allocate the data for the string. Loop in case the variable happens to change while running.
+	FString Buffer;
+	for (;;)
+	{
+		TArray<TCHAR>& CharArray = Buffer.GetCharArray();
+		CharArray.SetNumUninitialized(Length);
+
+		Length = ::GetCurrentDirectoryW(CharArray.Num(), CharArray.GetData());
+		if (Length < (uint32)CharArray.Num())
+		{
+			CharArray.SetNum(Length + 1);
+			break;
+		}
+	}
+	return Buffer;
 }
 
 const FString FWindowsPlatformProcess::ShaderWorkingDir()
@@ -1070,11 +1100,9 @@ bool FWindowsPlatformProcess::ResolveNetworkPath( FString InUNCPath, FString& Ou
 			// NetShareGetInfo doesn't accept const TCHAR* as the share name so copy to temp array
 			SHARE_INFO_2* BufPtr = NULL;
 			::NET_API_STATUS res;
-			TCHAR ShareNamePtr[MAX_PATH];
-			FCString::Strcpy(ShareNamePtr, ShareName.Len() + 1, *ShareName);
 
 			// Call the NetShareGetInfo function, specifying level 2
-			if ( ( res = NetShareGetInfo( NULL, ShareNamePtr, 2, (LPBYTE*)&BufPtr ) ) == ERROR_SUCCESS )
+			if ( ( res = NetShareGetInfo( NULL, ShareName.GetCharArray().GetData(), 2, (LPBYTE*)&BufPtr ) ) == ERROR_SUCCESS )
 			{
 				// Construct the local path
 				OutPath = FString( BufPtr->shi2_path ) + InUNCPath.Mid( ComputerNameLen + 1 + ShareNameLen );
