Add const TCHAR* overloads of FString::RemoveFromStart() and FString::RemoveFromEnd().

diff --git a/Engine/Source/Runtime/Core/Private/Containers/String.cpp b/Engine/Source/Runtime/Core/Private/Containers/String.cpp
--- a/Engine/Source/Runtime/Core/Private/Containers/String.cpp
+++ b/Engine/Source/Runtime/Core/Private/Containers/String.cpp
@@ -325,6 +325,22 @@ bool FString::EndsWith(const FString& InSuffix, ESearchCase::Type SearchCase ) c
 	}
 }
 
+bool FString::RemoveFromStart( const TCHAR* InPrefix, ESearchCase::Type SearchCase )
+{
+	if ( *InPrefix == 0 )
+	{
+		return false;
+	}
+
+	if ( StartsWith( InPrefix, SearchCase ) )
+	{
+		RemoveAt( 0, FCString::Strlen(InPrefix) );
+		return true;
+	}
+
+	return false;
+}
+
 bool FString::RemoveFromStart( const FString& InPrefix, ESearchCase::Type SearchCase )
 {
 	if ( InPrefix.IsEmpty() )
@@ -341,6 +357,23 @@ bool FString::RemoveFromStart( const FString& InPrefix, ESearchCase::Type Search
 	return false;
 }
 
+bool FString::RemoveFromEnd( const TCHAR* InSuffix, ESearchCase::Type SearchCase )
+{
+	if ( *InSuffix == 0 )
+	{
+		return false;
+	}
+
+	if ( EndsWith( InSuffix, SearchCase ) )
+	{
+		int32 SuffixLen = FCString::Strlen(InSuffix);
+		RemoveAt( Len() - SuffixLen, SuffixLen );
+		return true;
+	}
+
+	return false;
+}
+
 bool FString::RemoveFromEnd( const FString& InSuffix, ESearchCase::Type SearchCase )
 {
 	if ( InSuffix.IsEmpty() )
diff --git a/Engine/Source/Runtime/Core/Public/Containers/UnrealString.h b/Engine/Source/Runtime/Core/Public/Containers/UnrealString.h
--- a/Engine/Source/Runtime/Core/Public/Containers/UnrealString.h
+++ b/Engine/Source/Runtime/Core/Public/Containers/UnrealString.h
@@ -511,6 +511,13 @@ public:
 			}
 		}
 	}
+	/**
+	 * Removes the text from the start of the string if it exists.
+	 *
+	 * @param InPrefix the prefix to search for at the start of the string to remove.
+	 * @return true if the prefix was removed, otherwise false.
+	 */
+	bool RemoveFromStart( const TCHAR* InPrefix, ESearchCase::Type SearchCase = ESearchCase::IgnoreCase );
 
 	/**
 	 * Removes the text from the start of the string if it exists.
@@ -520,6 +527,14 @@ public:
 	 */
 	bool RemoveFromStart( const FString& InPrefix, ESearchCase::Type SearchCase = ESearchCase::IgnoreCase );
 
+	/**
+	 * Removes the text from the end of the string if it exists.
+	 *
+	 * @param InSuffix the suffix to search for at the end of the string to remove.
+	 * @return true if the suffix was removed, otherwise false.
+	 */
+	bool RemoveFromEnd( const TCHAR* InSuffix, ESearchCase::Type SearchCase = ESearchCase::IgnoreCase );
+
 	/**
 	 * Removes the text from the end of the string if it exists.
 	 *
