Fix FWindowsPlatformFile::GetFilenameOnDisk() support for paths longer than MAX_PATH, and simplify some of the other long path functions to avoid copying string buffers.

diff --git a/Engine/Source/Runtime/Core/Private/Windows/WindowsPlatformFile.cpp b/Engine/Source/Runtime/Core/Private/Windows/WindowsPlatformFile.cpp
--- a/Engine/Source/Runtime/Core/Private/Windows/WindowsPlatformFile.cpp
+++ b/Engine/Source/Runtime/Core/Private/Windows/WindowsPlatformFile.cpp
@@ -688,30 +688,29 @@ public:
 		}
 		else
 		{
-			TCHAR NormalizedFileName[MAX_PATH + 1];
-			int Length = (int)GetFinalPathNameByHandle(hFile, NormalizedFileName, MAX_PATH, FILE_NAME_NORMALIZED);
-			NormalizedFileName[Length] = 0;
-			CloseHandle(hFile);
-
-			int SrcIdx = 0;
-			if(FCString::Strncmp(NormalizedFileName, TEXT("\\\\?\\"), 4) == 0)
+			FString NormalizedFileName;
+			for(uint32 Length = FCString::Strlen(Filename) + 10;;)
 			{
-				SrcIdx = 4;
-			}
+				TArray<TCHAR>& CharArray = NormalizedFileName.GetCharArray();
+				CharArray.SetNum(Length);
 
-			int DstIdx = 0;
-			for(; NormalizedFileName[SrcIdx] != 0; SrcIdx++, DstIdx++)
-			{
-				if(NormalizedFileName[SrcIdx] == '\\')
+				Length = GetFinalPathNameByHandle(hFile, CharArray.GetData(), CharArray.Num(), FILE_NAME_NORMALIZED);
+				if (Length == 0)
 				{
-					NormalizedFileName[DstIdx] = '/';
+					NormalizedFileName = Filename;
+					break;
 				}
-				else
+				if (Length < (uint32)CharArray.Num())
 				{
-					NormalizedFileName[DstIdx] = NormalizedFileName[SrcIdx];
+					CharArray.SetNum(Length + 1);
+					break;
 				}
 			}
-			NormalizedFileName[DstIdx] = 0;
+
+			CloseHandle(hFile);
+
+			NormalizedFileName.RemoveFromStart(TEXT("\\\\?\\"), ESearchCase::CaseSensitive);
+			NormalizedFileName.ReplaceInline(TEXT("\\"), TEXT("/"));
 
 			return NormalizedFileName;
 		}
diff --git a/Engine/Source/Runtime/Core/Private/Windows/WindowsPlatformMisc.cpp b/Engine/Source/Runtime/Core/Private/Windows/WindowsPlatformMisc.cpp
--- a/Engine/Source/Runtime/Core/Private/Windows/WindowsPlatformMisc.cpp
+++ b/Engine/Source/Runtime/Core/Private/Windows/WindowsPlatformMisc.cpp
@@ -616,18 +616,9 @@ void FWindowsPlatformMisc::GetEnvironmentVariable(const TCHAR* VariableName, TCH
 
 FString FWindowsPlatformMisc::GetEnvironmentVariable(const TCHAR* VariableName)
 {
-	// Attempt to get the environment variable into a local buffer. If it succeeds, we can just copy the result into a string. If it fails, we have the length of the buffer we need to allocate.
-	TCHAR LocalBuffer[128];
-	uint32 Length = ::GetEnvironmentVariableW(VariableName, LocalBuffer, ARRAY_COUNT(LocalBuffer));
-	if (Length < ARRAY_COUNT(LocalBuffer))
-	{
-		LocalBuffer[Length] = 0;
-		return LocalBuffer;
-	}
-
-	// Allocate the data for the string. Loop in case the variable happens to change while running.
+	// Allocate the data for the string. Loop in case the variable happens to change while running, or the buffer isn't large enough.
 	FString Buffer;
-	for (;;)
+	for(uint32 Length = 128;;)
 	{
 		TArray<TCHAR>& CharArray = Buffer.GetCharArray();
 		CharArray.SetNumUninitialized(Length);
@@ -640,7 +631,6 @@ FString FWindowsPlatformMisc::GetEnvironmentVariable(const TCHAR* VariableName)
 		}
 		else if (Length < (uint32)CharArray.Num())
 		{
-			CharArray[Length] = 0;
 			CharArray.SetNum(Length + 1);
 			break;
 		}
diff --git a/Engine/Source/Runtime/Core/Private/Windows/WindowsPlatformProcess.cpp b/Engine/Source/Runtime/Core/Private/Windows/WindowsPlatformProcess.cpp
--- a/Engine/Source/Runtime/Core/Private/Windows/WindowsPlatformProcess.cpp
+++ b/Engine/Source/Runtime/Core/Private/Windows/WindowsPlatformProcess.cpp
@@ -941,22 +941,19 @@ void FWindowsPlatformProcess::SetCurrentWorkingDirectoryToBaseDir()
 /** Get the current working directory (only really makes sense on desktop platforms) */
 FString FWindowsPlatformProcess::GetCurrentWorkingDirectory()
 {
-	// Attempt to get the environment variable into a local buffer. If it succeeds, we can just copy the result into a string. If it fails, we have the length of the buffer we need to allocate.
-	TCHAR LocalBuffer[260];
-	uint32 Length = ::GetCurrentDirectoryW(ARRAY_COUNT(LocalBuffer), LocalBuffer);
-	if (Length < ARRAY_COUNT(LocalBuffer))
-	{
-		return LocalBuffer;
-	}
-
-	// Allocate the data for the string. Loop in case the variable happens to change while running.
+	// Allocate the data for the string. Loop in case the variable happens to change while running, or the buffer isn't large enough.
 	FString Buffer;
-	for (;;)
+	for (uint32 Length = 128;;)
 	{
 		TArray<TCHAR>& CharArray = Buffer.GetCharArray();
 		CharArray.SetNumUninitialized(Length);
 
 		Length = ::GetCurrentDirectoryW(CharArray.Num(), CharArray.GetData());
+		if (Length == 0)
+		{
+			Buffer.Reset();
+			break;
+		}
 		if (Length < (uint32)CharArray.Num())
 		{
 			CharArray.SetNum(Length + 1);
