Linux CIS fix.

#rnx

diff --git a/Engine/Source/Runtime/Core/Private/Unix/UnixPlatformMisc.cpp b/Engine/Source/Runtime/Core/Private/Unix/UnixPlatformMisc.cpp
--- a/Engine/Source/Runtime/Core/Private/Unix/UnixPlatformMisc.cpp
+++ b/Engine/Source/Runtime/Core/Private/Unix/UnixPlatformMisc.cpp
@@ -33,6 +33,7 @@
 #include <ifaddrs.h>
 #include <net/if.h>
 #include <net/if_arp.h>
+#include <linux/limits.h>
 
 #include "Modules/ModuleManager.h"
 #include "HAL/ThreadHeartBeat.h"
@@ -181,6 +182,11 @@ void FUnixPlatformMisc::PlatformTearDown()
 	FPlatformProcess::CeaseBeingFirstInstance();
 }
 
+int32 FUnixPlatformMisc::GetMaxPathLength()
+{
+	return PATH_MAX;
+}
+
 void FUnixPlatformMisc::GetEnvironmentVariable(const TCHAR* InVariableName, TCHAR* Result, int32 ResultLength)
 {
 	FString VariableName = InVariableName;
diff --git a/Engine/Source/Runtime/Core/Public/Unix/UnixPlatformMisc.h b/Engine/Source/Runtime/Core/Public/Unix/UnixPlatformMisc.h
--- a/Engine/Source/Runtime/Core/Public/Unix/UnixPlatformMisc.h
+++ b/Engine/Source/Runtime/Core/Public/Unix/UnixPlatformMisc.h
@@ -29,11 +29,7 @@ struct CORE_API FUnixPlatformMisc : public FGenericPlatformMisc
 	static void PlatformTearDown();
 	static void SetGracefulTerminationHandler();
 	static void SetCrashHandler(void (* CrashHandler)(const FGenericCrashContext& Context));
-
-	FORCEINLINE static int32 GetMaxPathLength()
-	{
-		return PATH_MAX;
-	}
+	static int32 GetMaxPathLength();
 
 	DEPRECATED(4.21, "void FPlatformMisc::GetEnvironmentVariable(Name, Result, Length) is deprecated. Use FString FPlatformMisc::GetEnvironmentVariable(Name) instead.")
 	static void GetEnvironmentVariable(const TCHAR* VariableName, TCHAR* Result, int32 ResultLength);
