Static analysis CIS fix.

#rnx

diff --git a/Engine/Source/Runtime/Core/Private/Windows/WindowsPlatformMisc.cpp b/Engine/Source/Runtime/Core/Private/Windows/WindowsPlatformMisc.cpp
--- a/Engine/Source/Runtime/Core/Private/Windows/WindowsPlatformMisc.cpp
+++ b/Engine/Source/Runtime/Core/Private/Windows/WindowsPlatformMisc.cpp
@@ -614,11 +614,16 @@ int32 FWindowsPlatformMisc::GetMaxPathLength()
 		FLongPathsEnabled()
 		{
 			HMODULE Handle = GetModuleHandle(TEXT("ntdll.dll"));
-
-			typedef BOOLEAN(NTAPI *RtlAreLongPathsEnabledFunc)();
-			RtlAreLongPathsEnabledFunc RtlAreLongPathsEnabled = (RtlAreLongPathsEnabledFunc)(void*)GetProcAddress(Handle, "RtlAreLongPathsEnabled");
-
-			bValue = (RtlAreLongPathsEnabled != NULL && RtlAreLongPathsEnabled());
+			if (Handle == NULL)
+			{
+				bValue = false;
+			}
+			else
+			{
+				typedef BOOLEAN(NTAPI *RtlAreLongPathsEnabledFunc)();
+				RtlAreLongPathsEnabledFunc RtlAreLongPathsEnabled = (RtlAreLongPathsEnabledFunc)(void*)GetProcAddress(Handle, "RtlAreLongPathsEnabled");
+				bValue = (RtlAreLongPathsEnabled != NULL && RtlAreLongPathsEnabled());
+			}
 		}
 	};
 
