TBA: Merge FAnnotatedStructuredArchiveFormatter with FStructuredArchiveFormatter. Any functions that are only implemented for text archives now have a _TextOnly suffix, and are exposed through the FStructuredArchive interface.

diff --git a/Engine/Source/Runtime/Core/Private/Serialization/Formatters/BinaryArchiveFormatter.cpp b/Engine/Source/Runtime/Core/Private/Serialization/Formatters/BinaryArchiveFormatter.cpp
--- a/Engine/Source/Runtime/Core/Private/Serialization/Formatters/BinaryArchiveFormatter.cpp
+++ b/Engine/Source/Runtime/Core/Private/Serialization/Formatters/BinaryArchiveFormatter.cpp
@@ -12,7 +12,37 @@ FBinaryArchiveFormatter::~FBinaryArchiveFormatter()
 {
 }
 
-bool FBinaryArchiveFormatter::RequiresStructuralMetadata() const
+bool FBinaryArchiveFormatter::HasDocumentTree() const
 {
 	return false;
 }
+
+void FBinaryArchiveFormatter::EnterRecord_TextOnly(TArray<FString>& OutFieldNames)
+{
+	checkf(false, TEXT("FBinaryArchiveFormatter does not support functions with the _TextOnly() suffix."));
+}
+
+void FBinaryArchiveFormatter::EnterField_TextOnly(FArchiveFieldName Name, EArchiveValueType& OutType)
+{
+	checkf(false, TEXT("FBinaryArchiveFormatter does not support functions with the _TextOnly() suffix."));
+}
+
+void FBinaryArchiveFormatter::EnterArrayElement_TextOnly(EArchiveValueType& OutType)
+{
+	checkf(false, TEXT("FBinaryArchiveFormatter does not support functions with the _TextOnly() suffix."));
+}
+
+void FBinaryArchiveFormatter::EnterStream_TextOnly(int32& NumElements)
+{
+	checkf(false, TEXT("FBinaryArchiveFormatter does not support functions with the _TextOnly() suffix."));
+}
+
+void FBinaryArchiveFormatter::EnterStreamElement_TextOnly(EArchiveValueType& OutType)
+{
+	checkf(false, TEXT("FBinaryArchiveFormatter does not support functions with the _TextOnly() suffix."));
+}
+
+void FBinaryArchiveFormatter::EnterMapElement_TextOnly(FString& Name, EArchiveValueType& OutType)
+{
+	checkf(false, TEXT("FBinaryArchiveFormatter does not support functions with the _TextOnly() suffix."));
+}
diff --git a/Engine/Source/Runtime/Core/Private/Serialization/StructuredArchive.cpp b/Engine/Source/Runtime/Core/Private/Serialization/StructuredArchive.cpp
--- a/Engine/Source/Runtime/Core/Private/Serialization/StructuredArchive.cpp
+++ b/Engine/Source/Runtime/Core/Private/Serialization/StructuredArchive.cpp
@@ -66,7 +66,7 @@ FStructuredArchive::FStructuredArchive(FArchiveFormatterType& InFormatter)
 #if DO_GUARD_SLOW
 	, bRequiresStructuralMetadata(true)
 #else
-	, bRequiresStructuralMetadata(InFormatter.RequiresStructuralMetadata())
+	, bRequiresStructuralMetadata(InFormatter.HasDocumentTree())
 #endif
 	, NextElementId(RootElementId + 1)
 	, CurrentSlotElementId(INDEX_NONE)
@@ -224,7 +224,7 @@ FStructuredArchive::FRecord FStructuredArchive::FSlot::EnterRecord()
 	return FRecord(Ar, Depth + 1, ElementId);
 }
 
-FStructuredArchive::FRecord FStructuredArchive::FSlot::EnterRecord(TArray<FString>& OutFieldNamesWhenLoading)
+FStructuredArchive::FRecord FStructuredArchive::FSlot::EnterRecord_TextOnly(TArray<FString>& OutFieldNames)
 {
 	Ar.EnterSlot(ElementId, EElementType::Record);
 
@@ -232,7 +232,7 @@ FStructuredArchive::FRecord FStructuredArchive::FSlot::EnterRecord(TArray<FStrin
 	Ar.CurrentContainer.Add(new FContainer(0));
 #endif
 
-	Ar.Formatter.EnterRecord(OutFieldNamesWhenLoading);
+	Ar.Formatter.EnterRecord_TextOnly(OutFieldNames);
 
 	return FRecord(Ar, Depth + 1, ElementId);
 }
@@ -259,6 +259,15 @@ FStructuredArchive::FStream FStructuredArchive::FSlot::EnterStream()
 	return FStream(Ar, Depth + 1, ElementId);
 }
 
+FStructuredArchive::FStream FStructuredArchive::FSlot::EnterStream_TextOnly(int32& OutNumElements)
+{
+	Ar.EnterSlot(ElementId, EElementType::Stream);
+
+	Ar.Formatter.EnterStream_TextOnly(OutNumElements);
+
+	return FStream(Ar, Depth + 1, ElementId);
+}
+
 FStructuredArchive::FMap FStructuredArchive::FSlot::EnterMap(int32& Num)
 {
 	Ar.EnterSlot(ElementId, EElementType::Map);
@@ -459,9 +468,9 @@ FStructuredArchive::FRecord FStructuredArchive::FRecord::EnterRecord(FArchiveFie
 	return EnterField(Name).EnterRecord();
 }
 
-FStructuredArchive::FRecord FStructuredArchive::FRecord::EnterRecord(FArchiveFieldName Name, TArray<FString>& OutFieldNamesWhenLoading)
+FStructuredArchive::FRecord FStructuredArchive::FRecord::EnterRecord_TextOnly(FArchiveFieldName Name, TArray<FString>& OutFieldNames)
 {
-	return EnterField(Name).EnterRecord(OutFieldNamesWhenLoading);
+	return EnterField(Name).EnterRecord_TextOnly(OutFieldNames);
 }
 
 FStructuredArchive::FArray FStructuredArchive::FRecord::EnterArray(FArchiveFieldName Name, int32& Num)
@@ -474,6 +483,11 @@ FStructuredArchive::FStream FStructuredArchive::FRecord::EnterStream(FArchiveFie
 	return EnterField(Name).EnterStream();
 }
 
+FStructuredArchive::FStream FStructuredArchive::FRecord::EnterStream_TextOnly(FArchiveFieldName Name, int32& OutNumElements)
+{
+	return EnterField(Name).EnterStream_TextOnly(OutNumElements);
+}
+
 FStructuredArchive::FMap FStructuredArchive::FRecord::EnterMap(FArchiveFieldName Name, int32& Num)
 {
 	return EnterField(Name).EnterMap(Num);
@@ -522,6 +536,21 @@ FStructuredArchive::FSlot FStructuredArchive::FArray::EnterElement()
 	return FSlot(Ar, Depth, Ar.CurrentSlotElementId);
 }
 
+FStructuredArchive::FSlot FStructuredArchive::FArray::EnterElement_TextOnly(EArchiveValueType& OutType)
+{
+	Ar.SetScope(Depth, ElementId);
+
+#if DO_GUARD_SLOW
+	checkf(Ar.CurrentContainer.Top()->Index < Ar.CurrentContainer.Top()->Count, TEXT("Serialized too many array elements"));
+#endif
+
+	Ar.CurrentSlotElementId = Ar.NextElementId++;
+
+	Ar.Formatter.EnterArrayElement_TextOnly(OutType);
+
+	return FSlot(Ar, Depth, Ar.CurrentSlotElementId);
+}
+
 //////////// FStructuredArchive::FStream ////////////
 
 FStructuredArchive::FSlot FStructuredArchive::FStream::EnterElement()
@@ -535,6 +564,17 @@ FStructuredArchive::FSlot FStructuredArchive::FStream::EnterElement()
 	return FSlot(Ar, Depth, Ar.CurrentSlotElementId);
 }
 
+FStructuredArchive::FSlot FStructuredArchive::FStream::EnterElement_TextOnly(EArchiveValueType& OutType)
+{
+	Ar.SetScope(Depth, ElementId);
+
+	Ar.CurrentSlotElementId = Ar.NextElementId++;
+
+	Ar.Formatter.EnterStreamElement_TextOnly(OutType);
+
+	return FSlot(Ar, Depth, Ar.CurrentSlotElementId);
+}
+
 //////////// FStructuredArchive::FMap ////////////
 
 FStructuredArchive::FSlot FStructuredArchive::FMap::EnterElement(FString& Name)
@@ -574,4 +614,41 @@ FStructuredArchive::FSlot FStructuredArchive::FMap::EnterElement(FString& Name)
 	return FSlot(Ar, Depth, Ar.CurrentSlotElementId);
 }
 
+FStructuredArchive::FSlot FStructuredArchive::FMap::EnterElement_TextOnly(FString& Name, EArchiveValueType& OutType)
+{
+	Ar.SetScope(Depth, ElementId);
+
+#if DO_GUARD_SLOW
+	checkf(Ar.CurrentContainer.Top()->Index < Ar.CurrentContainer.Top()->Count, TEXT("Serialized too many map elements"));
+#endif
+
+	Ar.CurrentSlotElementId = Ar.NextElementId++;
+
+#if DO_GUARD_SLOW
+#if CHECK_UNIQUE_FIELD_NAMES
+	if(Ar.GetUnderlyingArchive().IsSaving())
+	{
+		FContainer& Container = *Ar.CurrentContainer.Top();
+		checkf(!Container.KeyNames.Contains(Name), TEXT("Multiple keys called '%s' serialized into record"), *Name);
+		Container.KeyNames.Add(Name);
+	}
+#endif
+#endif
+
+	Ar.Formatter.EnterMapElement_TextOnly(Name, OutType);
+
+#if DO_GUARD_SLOW
+#if CHECK_UNIQUE_FIELD_NAMES
+	if(Ar.GetUnderlyingArchive().IsLoading())
+	{
+		FContainer& Container = *Ar.CurrentContainer.Top();
+		checkf(!Container.KeyNames.Contains(Name), TEXT("Multiple keys called '%s' serialized into record"), *Name);
+		Container.KeyNames.Add(Name);
+	}
+#endif
+#endif
+
+	return FSlot(Ar, Depth, Ar.CurrentSlotElementId);
+}
+
 #endif
diff --git a/Engine/Source/Runtime/Core/Private/Serialization/StructuredArchiveFormatter.cpp b/Engine/Source/Runtime/Core/Private/Serialization/StructuredArchiveFormatter.cpp
--- a/Engine/Source/Runtime/Core/Private/Serialization/StructuredArchiveFormatter.cpp
+++ b/Engine/Source/Runtime/Core/Private/Serialization/StructuredArchiveFormatter.cpp
@@ -11,33 +11,28 @@ FStructuredArchiveFormatter::~FStructuredArchiveFormatter()
 {
 }
 
-bool FStructuredArchiveFormatter::RequiresStructuralMetadata() const
-{
-	return true;
-}
-
 //////////// Functions ////////////
 
 #if WITH_TEXT_ARCHIVE_SUPPORT
-	static void CopyFormattedSlot(FAnnotatedStructuredArchiveFormatter& InputFormatter, FStructuredArchiveFormatter& OutputFormatter, EArchiveValueType Type);
+	static void CopyFormattedSlot(FStructuredArchiveFormatter& InputFormatter, FStructuredArchiveFormatter& OutputFormatter, EArchiveValueType Type);
 
-	template<typename Type> static void CopyFormattedValue(FAnnotatedStructuredArchiveFormatter& InputFormatter, FStructuredArchiveFormatter& OutputFormatter)
+	template<typename Type> static void CopyFormattedValue(FStructuredArchiveFormatter& InputFormatter, FStructuredArchiveFormatter& OutputFormatter)
 	{
 		Type Value;
 		InputFormatter.Serialize(Value);
 		OutputFormatter.Serialize(Value);
 	}
 
-	static void CopyFormattedRecord(FAnnotatedStructuredArchiveFormatter& InputFormatter, FStructuredArchiveFormatter& OutputFormatter)
+	static void CopyFormattedRecord(FStructuredArchiveFormatter& InputFormatter, FStructuredArchiveFormatter& OutputFormatter)
 	{
 		TArray<FString> Fields;
-		InputFormatter.EnterRecord(Fields);
+		InputFormatter.EnterRecord_TextOnly(Fields);
 		OutputFormatter.EnterRecord();
 
 		for (const FString& Field : Fields)
 		{
 			EArchiveValueType Type;
-			InputFormatter.EnterField(FArchiveFieldName(*Field), Type);
+			InputFormatter.EnterField_TextOnly(FArchiveFieldName(*Field), Type);
 			OutputFormatter.EnterField(FArchiveFieldName(*Field));
 
 			CopyFormattedSlot(InputFormatter, OutputFormatter, Type);
@@ -50,7 +45,7 @@ bool FStructuredArchiveFormatter::RequiresStructuralMetadata() const
 		InputFormatter.LeaveRecord();
 	}
 
-	static void CopyFormattedArray(FAnnotatedStructuredArchiveFormatter& InputFormatter, FStructuredArchiveFormatter& OutputFormatter)
+	static void CopyFormattedArray(FStructuredArchiveFormatter& InputFormatter, FStructuredArchiveFormatter& OutputFormatter)
 	{
 		int32 NumElements = 0;
 		InputFormatter.EnterArray(NumElements);
@@ -59,7 +54,7 @@ bool FStructuredArchiveFormatter::RequiresStructuralMetadata() const
 		for (int32 Idx = 0; Idx < NumElements; Idx++)
 		{
 			EArchiveValueType Type;
-			InputFormatter.EnterArrayElement(Type);
+			InputFormatter.EnterArrayElement_TextOnly(Type);
 			OutputFormatter.EnterArrayElement();
 
 			CopyFormattedSlot(InputFormatter, OutputFormatter, Type);
@@ -72,16 +67,16 @@ bool FStructuredArchiveFormatter::RequiresStructuralMetadata() const
 		InputFormatter.LeaveArray();
 	}
 
-	static void CopyFormattedStream(FAnnotatedStructuredArchiveFormatter& InputFormatter, FStructuredArchiveFormatter& OutputFormatter)
+	static void CopyFormattedStream(FStructuredArchiveFormatter& InputFormatter, FStructuredArchiveFormatter& OutputFormatter)
 	{
 		int32 NumElements = 0;
-		InputFormatter.EnterStream(NumElements);
+		InputFormatter.EnterStream_TextOnly(NumElements);
 		OutputFormatter.EnterStream();
 
 		for (int32 Idx = 0; Idx < NumElements; Idx++)
 		{
 			EArchiveValueType Type;
-			InputFormatter.EnterStreamElement(Type);
+			InputFormatter.EnterStreamElement_TextOnly(Type);
 			OutputFormatter.EnterStreamElement();
 
 			CopyFormattedSlot(InputFormatter, OutputFormatter, Type);
@@ -94,7 +89,7 @@ bool FStructuredArchiveFormatter::RequiresStructuralMetadata() const
 		InputFormatter.LeaveStream();
 	}
 
-	static void CopyFormattedMap(FAnnotatedStructuredArchiveFormatter& InputFormatter, FStructuredArchiveFormatter& OutputFormatter)
+	static void CopyFormattedMap(FStructuredArchiveFormatter& InputFormatter, FStructuredArchiveFormatter& OutputFormatter)
 	{
 		int32 NumElements = 0;
 		InputFormatter.EnterMap(NumElements);
@@ -104,7 +99,7 @@ bool FStructuredArchiveFormatter::RequiresStructuralMetadata() const
 		{
 			FString Key;
 			EArchiveValueType Type;
-			InputFormatter.EnterMapElement(Key, Type);
+			InputFormatter.EnterMapElement_TextOnly(Key, Type);
 			OutputFormatter.EnterMapElement(Key);
 
 			CopyFormattedSlot(InputFormatter, OutputFormatter, Type);
@@ -117,7 +112,7 @@ bool FStructuredArchiveFormatter::RequiresStructuralMetadata() const
 		InputFormatter.LeaveMap();
 	}
 
-	static void CopyFormattedSlot(FAnnotatedStructuredArchiveFormatter& InputFormatter, FStructuredArchiveFormatter& OutputFormatter, EArchiveValueType Type)
+	static void CopyFormattedSlot(FStructuredArchiveFormatter& InputFormatter, FStructuredArchiveFormatter& OutputFormatter, EArchiveValueType Type)
 	{
 		switch (Type)
 		{
@@ -181,7 +176,7 @@ bool FStructuredArchiveFormatter::RequiresStructuralMetadata() const
 		}
 	}
 
-	CORE_API void CopyFormattedData(FAnnotatedStructuredArchiveFormatter& InputFormatter, FStructuredArchiveFormatter& OutputFormatter)
+	CORE_API void CopyFormattedData(FStructuredArchiveFormatter& InputFormatter, FStructuredArchiveFormatter& OutputFormatter)
 	{
 		CopyFormattedRecord(InputFormatter, OutputFormatter);
 	}
diff --git a/Engine/Source/Runtime/Core/Public/Serialization/Formatters/BinaryArchiveFormatter.h b/Engine/Source/Runtime/Core/Public/Serialization/Formatters/BinaryArchiveFormatter.h
--- a/Engine/Source/Runtime/Core/Public/Serialization/Formatters/BinaryArchiveFormatter.h
+++ b/Engine/Source/Runtime/Core/Public/Serialization/Formatters/BinaryArchiveFormatter.h
@@ -12,29 +12,34 @@ public:
 	FBinaryArchiveFormatter(FArchive& InInner);
 	virtual ~FBinaryArchiveFormatter();
 
-	virtual bool RequiresStructuralMetadata() const override;
+	virtual bool HasDocumentTree() const override;
 	virtual FArchive& GetUnderlyingArchive() override;
 
 	virtual void EnterRecord() override;
-	virtual void EnterRecord(TArray<FString>& OutFieldNamesWhenLoading) override;
+	virtual void EnterRecord_TextOnly(TArray<FString>& OutFieldNames) override;
 	virtual void LeaveRecord() override;
 	virtual void EnterField(FArchiveFieldName Name) override;
+	virtual void EnterField_TextOnly(FArchiveFieldName Name, EArchiveValueType& OutType) override;
 	virtual void LeaveField() override;
 	virtual bool TryEnterField(FArchiveFieldName Name, bool bEnterIfWriting);
 
 	virtual void EnterArray(int32& NumElements) override;
 	virtual void LeaveArray() override;
 	virtual void EnterArrayElement() override;
+	virtual void EnterArrayElement_TextOnly(EArchiveValueType& OutType) override;
 	virtual void LeaveArrayElement() override;
 
 	virtual void EnterStream() override;
+	virtual void EnterStream_TextOnly(int32& NumElements) override;
 	virtual void LeaveStream() override;
 	virtual void EnterStreamElement() override;
+	virtual void EnterStreamElement_TextOnly(EArchiveValueType& OutType) override;
 	virtual void LeaveStreamElement() override;
 
 	virtual void EnterMap(int32& NumElements) override;
 	virtual void LeaveMap() override;
 	virtual void EnterMapElement(FString& Name) override;
+	virtual void EnterMapElement_TextOnly(FString& Name, EArchiveValueType& OutType) override;
 	virtual void LeaveMapElement() override;
 
 	virtual void Serialize(uint8& Value) override;
@@ -73,10 +78,6 @@ inline void FBinaryArchiveFormatter::EnterRecord()
 {
 }
 
-inline void FBinaryArchiveFormatter::EnterRecord(TArray<FString>& OutFieldNamesWhenLoading)
-{
-}
-
 inline void FBinaryArchiveFormatter::LeaveRecord()
 {
 }
diff --git a/Engine/Source/Runtime/Core/Public/Serialization/StructuredArchive.h b/Engine/Source/Runtime/Core/Public/Serialization/StructuredArchive.h
--- a/Engine/Source/Runtime/Core/Public/Serialization/StructuredArchive.h
+++ b/Engine/Source/Runtime/Core/Public/Serialization/StructuredArchive.h
@@ -106,9 +106,10 @@ public:
 	{
 	public:
 		FRecord EnterRecord();
-		FRecord EnterRecord(TArray<FString>& OutFieldNamesWhenLoading);
+		FRecord EnterRecord_TextOnly(TArray<FString>& OutFieldNames);
 		FArray EnterArray(int32& Num);
 		FStream EnterStream();
+		FStream EnterStream_TextOnly(int32& OutNumElements);
 		FMap EnterMap(int32& Num);
 
 		// We don't support chaining writes to a single slot, so these return void.
@@ -181,10 +182,12 @@ public:
 	{
 	public:
 		FSlot EnterField(FArchiveFieldName Name);
+		FSlot EnterField_TextOnly(FArchiveFieldName Name, EArchiveValueType& OutType);
 		FRecord EnterRecord(FArchiveFieldName Name);
-		FRecord EnterRecord(FArchiveFieldName Name, TArray<FString>& OutFieldNamesWhenLoading);
+		FRecord EnterRecord_TextOnly(FArchiveFieldName Name, TArray<FString>& OutFieldNames);
 		FArray EnterArray(FArchiveFieldName Name, int32& Num);
 		FStream EnterStream(FArchiveFieldName Name);
+		FStream EnterStream_TextOnly(FArchiveFieldName Name, int32& OutNumElements);
 		FMap EnterMap(FArchiveFieldName Name, int32& Num);
 
 		TOptional<FSlot> TryEnterField(FArchiveFieldName Name, bool bEnterForSaving);
@@ -230,6 +233,7 @@ public:
 	{
 	public:
 		FSlot EnterElement();
+		FSlot EnterElement_TextOnly(EArchiveValueType& OutType);
 
 		template<typename T> FORCEINLINE FArray& operator<<(T& Item)
 		{
@@ -271,6 +275,7 @@ public:
 	{
 	public:
 		FSlot EnterElement();
+		FSlot EnterElement_TextOnly(EArchiveValueType& OutType);
 
 		template<typename T> FORCEINLINE FStream& operator<<(T& Item)
 		{
@@ -313,6 +318,7 @@ public:
 	{
 	public:
 		FSlot EnterElement(FString& Name);
+		FSlot EnterElement_TextOnly(FString& Name, EArchiveValueType& OutType);
 
 		FORCEINLINE FArchive& GetUnderlyingArchive()
 		{
@@ -520,8 +526,9 @@ private:
 		return FRecord(Ar);
 	}
 
-	FORCEINLINE FStructuredArchive::FRecord FStructuredArchive::FSlot::EnterRecord(TArray<FString>& OutFieldNamesWhenLoading)
+	FORCEINLINE FStructuredArchive::FRecord FStructuredArchive::FSlot::EnterRecord_TextOnly(TArray<FString>& OutFieldNames)
 	{
+		Ar.Formatter.EnterRecord_TextOnly(OutFieldNames);
 		return FRecord(Ar);
 	}
 
@@ -536,6 +543,12 @@ private:
 		return FStream(Ar);
 	}
 
+	FORCEINLINE FStructuredArchive::FStream FStructuredArchive::FSlot::EnterStream_TextOnly(int32& OutNumElements)
+	{
+		Ar.Formatter.EnterStream_TextOnly(OutNumElements);
+		return FStream(Ar);
+	}
+
 	FORCEINLINE FStructuredArchive::FMap FStructuredArchive::FSlot::EnterMap(int32& Num)
 	{
 		Ar.Formatter.EnterMap(Num);
@@ -654,21 +667,27 @@ private:
 		Ar.Formatter.Serialize(Data, DataSize);
 	}
 
-	//////////// FStructuredArchive::FObject ////////////
+	//////////// FStructuredArchive::FRecord ////////////
 
 	FORCEINLINE FStructuredArchive::FSlot FStructuredArchive::FRecord::EnterField(FArchiveFieldName Name)
 	{
 		return FSlot(Ar);
 	}
 
+	FORCEINLINE FStructuredArchive::FSlot FStructuredArchive::FRecord::EnterField_TextOnly(FArchiveFieldName Name, EArchiveValueType& OutType)
+	{
+		Ar.Formatter.EnterField_TextOnly(Name, OutType);
+		return FSlot(Ar);
+	}
+
 	FORCEINLINE FStructuredArchive::FRecord FStructuredArchive::FRecord::EnterRecord(FArchiveFieldName Name)
 	{
 		return EnterField(Name).EnterRecord();
 	}
 
-	FORCEINLINE FStructuredArchive::FRecord FStructuredArchive::FRecord::EnterRecord(FArchiveFieldName Name, TArray<FString>& OutFieldNamesWhenLoading)
+	FORCEINLINE FStructuredArchive::FRecord FStructuredArchive::FRecord::EnterRecord_TextOnly(FArchiveFieldName Name, TArray<FString>& OutFieldNames)
 	{
-		return EnterField(Name).EnterRecord(OutFieldNamesWhenLoading);
+		return EnterField(Name).EnterRecord_TextOnly(OutFieldNames);
 	}
 
 	FORCEINLINE FStructuredArchive::FArray FStructuredArchive::FRecord::EnterArray(FArchiveFieldName Name, int32& Num)
@@ -705,6 +724,12 @@ private:
 		return FSlot(Ar);
 	}
 
+	FORCEINLINE FStructuredArchive::FSlot FStructuredArchive::FArray::EnterElement_TextOnly(EArchiveValueType& OutType)
+	{
+		Ar.Formatter.EnterArrayElement_TextOnly(OutType);
+		return FSlot(Ar);
+	}
+
 	//////////// FStructuredArchive::FStream ////////////
 
 	FORCEINLINE FStructuredArchive::FSlot FStructuredArchive::FStream::EnterElement()
@@ -712,6 +737,12 @@ private:
 		return FSlot(Ar);
 	}
 
+	FORCEINLINE FStructuredArchive::FSlot FStructuredArchive::FStream::EnterElement_TextOnly(EArchiveValueType& OutType)
+	{
+		Ar.Formatter.EnterStreamElement_TextOnly(OutType);
+		return FSlot(Ar);
+	}
+
 	//////////// FStructuredArchive::FMap ////////////
 
 	FORCEINLINE FStructuredArchive::FSlot FStructuredArchive::FMap::EnterElement(FString& Name)
@@ -720,5 +751,11 @@ private:
 		return FSlot(Ar);
 	}
 
+	FORCEINLINE FStructuredArchive::FSlot FStructuredArchive::FMap::EnterElement_TextOnly(FString& Name, EArchiveValueType& OutType)
+	{
+		Ar.Formatter.EnterMapElement_TextOnly(Name, OutType);
+		return FSlot(Ar);
+	}
+
 #endif
 
diff --git a/Engine/Source/Runtime/Core/Public/Serialization/StructuredArchiveFormatter.h b/Engine/Source/Runtime/Core/Public/Serialization/StructuredArchiveFormatter.h
--- a/Engine/Source/Runtime/Core/Public/Serialization/StructuredArchiveFormatter.h
+++ b/Engine/Source/Runtime/Core/Public/Serialization/StructuredArchiveFormatter.h
@@ -61,37 +61,46 @@ enum class EArchiveValueType
  * Interface to format data to and from an underlying archive. Methods on this class are validated to be correct
  * with the current archive state (eg. EnterObject/LeaveObject calls are checked to be matching), and do not need
  * to be validated by implementations.
+ *
+ * Any functions with the _TextOnly suffix are intended to be implemented when reading text archives that have
+ * a fully defined document tree, and allow querying additional properties that aren't available when reading from
+ * a pure binary archive. These functions will assert if called on a binary archive.
  */
 class CORE_API FStructuredArchiveFormatter
 {
 public:
 	virtual ~FStructuredArchiveFormatter();
 
-	virtual bool RequiresStructuralMetadata() const;
-	
 	virtual FArchive& GetUnderlyingArchive() = 0;
 	virtual FStructuredArchiveFormatter* CreateSubtreeReader() { return this; }
 
+	virtual bool HasDocumentTree() const = 0;
+
 	virtual void EnterRecord() = 0;
-	virtual void EnterRecord(TArray<FString>& OutFieldNamesWhenLoading) = 0;
+	virtual void EnterRecord_TextOnly(TArray<FString>& OutFieldNames) = 0;
 	virtual void LeaveRecord() = 0;
 	virtual void EnterField(FArchiveFieldName Name) = 0;
+	virtual void EnterField_TextOnly(FArchiveFieldName Name, EArchiveValueType& OutType) = 0;
 	virtual void LeaveField() = 0;
 	virtual bool TryEnterField(FArchiveFieldName Name, bool bEnterWhenWriting) = 0;
 
 	virtual void EnterArray(int32& NumElements) = 0;
 	virtual void LeaveArray() = 0;
 	virtual void EnterArrayElement() = 0;
+	virtual void EnterArrayElement_TextOnly(EArchiveValueType& OutType) = 0;
 	virtual void LeaveArrayElement() = 0;
 
 	virtual void EnterStream() = 0;
+	virtual void EnterStream_TextOnly(int32& OutNumElements) = 0;
 	virtual void LeaveStream() = 0;
 	virtual void EnterStreamElement() = 0;
+	virtual void EnterStreamElement_TextOnly(EArchiveValueType& OutType) = 0;
 	virtual void LeaveStreamElement() = 0;
 
 	virtual void EnterMap(int32& NumElements) = 0;
 	virtual void LeaveMap() = 0;
 	virtual void EnterMapElement(FString& Name) = 0;
+	virtual void EnterMapElement_TextOnly(FString& Name, EArchiveValueType& OutType) = 0;
 	virtual void LeaveMapElement() = 0;
 
 	virtual void Serialize(uint8& Value) = 0;
@@ -117,35 +126,9 @@ public:
 	virtual void Serialize(void* Data, uint64 DataSize) = 0;
 };
 
-/**
- * Base class for formatters that include context-free annotations for typing while reading.
- * Text-based formats implement this to allow context-free conversions between archives without being dependent on C++ serialization code.
- */
-class CORE_API FAnnotatedStructuredArchiveFormatter : public FStructuredArchiveFormatter
-{
-public:
-	using FStructuredArchiveFormatter::EnterRecord;
-	virtual void EnterRecord(TArray<FString>& Keys) = 0;
-
-	using FStructuredArchiveFormatter::EnterStream;
-	virtual void EnterStream(int32& NumElements) = 0;
-
-	using FStructuredArchiveFormatter::EnterField;
-	virtual void EnterField(FArchiveFieldName Name, EArchiveValueType& OutType) = 0;
-
-	using FStructuredArchiveFormatter::EnterArrayElement;
-	virtual void EnterArrayElement(EArchiveValueType& OutType) = 0;
-
-	using FStructuredArchiveFormatter::EnterStreamElement;
-	virtual void EnterStreamElement(EArchiveValueType& OutType) = 0;
-
-	using FStructuredArchiveFormatter::EnterMapElement;
-	virtual void EnterMapElement(FString& OutName, EArchiveValueType& OutType) = 0;
-};
-
 #if WITH_TEXT_ARCHIVE_SUPPORT
 	/**
 	 * Copies formatted data from one place to another. Useful for conversion functions or visitor patterns.
 	 */
-	CORE_API void CopyFormattedData(FAnnotatedStructuredArchiveFormatter& InputFormatter, FStructuredArchiveFormatter& OutputFormatter);
+	CORE_API void CopyFormattedData(FStructuredArchiveFormatter& InputFormatter, FStructuredArchiveFormatter& OutputFormatter);
 #endif
diff --git a/Engine/Source/Runtime/CoreUObject/Private/Serialization/Formatters/JsonArchiveInputFormatter.cpp b/Engine/Source/Runtime/CoreUObject/Private/Serialization/Formatters/JsonArchiveInputFormatter.cpp
--- a/Engine/Source/Runtime/CoreUObject/Private/Serialization/Formatters/JsonArchiveInputFormatter.cpp
+++ b/Engine/Source/Runtime/CoreUObject/Private/Serialization/Formatters/JsonArchiveInputFormatter.cpp
@@ -46,19 +46,24 @@ FStructuredArchiveFormatter* FJsonArchiveInputFormatter::CreateSubtreeReader()
 	return Cloned;
 }
 
+bool FJsonArchiveInputFormatter::HasDocumentTree() const
+{
+	return true;
+}
+
 void FJsonArchiveInputFormatter::EnterRecord()
 {
 	TSharedPtr<FJsonValue>& Value = ValueStack.Top();
 	ObjectStack.Add(FObjectRecord(Value->AsObject(), ValueStack.Num()));
 }
 
-void FJsonArchiveInputFormatter::EnterRecord(TArray<FString>& OutFieldNamesWhenLoading)
+void FJsonArchiveInputFormatter::EnterRecord_TextOnly(TArray<FString>& OutFieldNames)
 {
 	EnterRecord();
-	ObjectStack.Top().JsonObject->Values.GetKeys(OutFieldNamesWhenLoading);
-	for(int32 Idx = 0; Idx < OutFieldNamesWhenLoading.Num(); Idx++)
+	ObjectStack.Top().JsonObject->Values.GetKeys(OutFieldNames);
+	for(int32 Idx = 0; Idx < OutFieldNames.Num(); Idx++)
 	{
-		OutFieldNamesWhenLoading[Idx] = UnescapeFieldName(*OutFieldNamesWhenLoading[Idx]);
+		OutFieldNames[Idx] = UnescapeFieldName(*OutFieldNames[Idx]);
 	}
 }
 
@@ -76,7 +81,7 @@ void FJsonArchiveInputFormatter::EnterField(FArchiveFieldName Name)
 	ValueStack.Add(Field);
 }
 
-void FJsonArchiveInputFormatter::EnterField(FArchiveFieldName Name, EArchiveValueType& OutType)
+void FJsonArchiveInputFormatter::EnterField_TextOnly(FArchiveFieldName Name, EArchiveValueType& OutType)
 {
 	EnterField(Name);
 	OutType = GetValueType(*ValueStack.Top());
@@ -120,7 +125,7 @@ void FJsonArchiveInputFormatter::EnterArrayElement()
 {
 }
 
-void FJsonArchiveInputFormatter::EnterArrayElement(EArchiveValueType& OutType)
+void FJsonArchiveInputFormatter::EnterArrayElement_TextOnly(EArchiveValueType& OutType)
 {
 	OutType = GetValueType(*ValueStack.Top());
 }
@@ -136,7 +141,7 @@ void FJsonArchiveInputFormatter::EnterStream()
 	EnterArray(NumElements);
 }
 
-void FJsonArchiveInputFormatter::EnterStream(int32& NumElements)
+void FJsonArchiveInputFormatter::EnterStream_TextOnly(int32& NumElements)
 {
 	EnterArray(NumElements);
 }
@@ -149,7 +154,7 @@ void FJsonArchiveInputFormatter::EnterStreamElement()
 {
 }
 
-void FJsonArchiveInputFormatter::EnterStreamElement(EArchiveValueType& OutType)
+void FJsonArchiveInputFormatter::EnterStreamElement_TextOnly(EArchiveValueType& OutType)
 {
 	OutType = GetValueType(*ValueStack.Top());
 }
@@ -178,7 +183,7 @@ void FJsonArchiveInputFormatter::EnterMapElement(FString& OutName)
 	ValueStack.Add(MapIteratorStack.Top()->Value);
 }
 
-void FJsonArchiveInputFormatter::EnterMapElement(FString& OutName, EArchiveValueType& OutType)
+void FJsonArchiveInputFormatter::EnterMapElement_TextOnly(FString& OutName, EArchiveValueType& OutType)
 {
 	EnterMapElement(OutName);
 	OutType = GetValueType(*ValueStack.Top());
diff --git a/Engine/Source/Runtime/CoreUObject/Private/Serialization/Formatters/JsonArchiveOutputFormatter.cpp b/Engine/Source/Runtime/CoreUObject/Private/Serialization/Formatters/JsonArchiveOutputFormatter.cpp
--- a/Engine/Source/Runtime/CoreUObject/Private/Serialization/Formatters/JsonArchiveOutputFormatter.cpp
+++ b/Engine/Source/Runtime/CoreUObject/Private/Serialization/Formatters/JsonArchiveOutputFormatter.cpp
@@ -31,6 +31,11 @@ FArchive& FJsonArchiveOutputFormatter::GetUnderlyingArchive()
 	return Inner;
 }
 
+bool FJsonArchiveOutputFormatter::HasDocumentTree() const
+{
+	return true;
+}
+
 void FJsonArchiveOutputFormatter::EnterRecord()
 {
 	WriteOptionalComma();
@@ -40,9 +45,10 @@ void FJsonArchiveOutputFormatter::EnterRecord()
 	bNeedsNewline = true;
 }
 
-void FJsonArchiveOutputFormatter::EnterRecord(TArray<FString>& OutFieldNamesWhenLoading)
+void FJsonArchiveOutputFormatter::EnterRecord_TextOnly(TArray<FString>& OutFieldNames)
 {
 	EnterRecord();
+	OutFieldNames.Reset();
 }
 
 void FJsonArchiveOutputFormatter::LeaveRecord()
@@ -61,6 +67,12 @@ void FJsonArchiveOutputFormatter::EnterField(FArchiveFieldName Name)
 	WriteFieldName(Name.Name);
 }
 
+void FJsonArchiveOutputFormatter::EnterField_TextOnly(FArchiveFieldName Name, EArchiveValueType& OutType)
+{
+	EnterField(Name);
+	OutType = EArchiveValueType::None;
+}
+
 void FJsonArchiveOutputFormatter::LeaveField()
 {
 	bNeedsComma = true;
@@ -91,6 +103,12 @@ void FJsonArchiveOutputFormatter::EnterArrayElement()
 	EnterStreamElement();
 }
 
+void FJsonArchiveOutputFormatter::EnterArrayElement_TextOnly(EArchiveValueType& OutType)
+{
+	EnterArrayElement();
+	OutType = EArchiveValueType::None;
+}
+
 void FJsonArchiveOutputFormatter::LeaveArrayElement()
 {
 	LeaveStreamElement();
@@ -105,6 +123,12 @@ void FJsonArchiveOutputFormatter::EnterStream()
 	bNeedsNewline = true;
 }
 
+void FJsonArchiveOutputFormatter::EnterStream_TextOnly(int32& OutNumElements)
+{
+	EnterStream();
+	OutNumElements = 0;
+}
+
 void FJsonArchiveOutputFormatter::LeaveStream()
 {
 	Newline.Pop(false);
@@ -120,6 +144,12 @@ void FJsonArchiveOutputFormatter::EnterStreamElement()
 	WriteOptionalNewline();
 }
 
+void FJsonArchiveOutputFormatter::EnterStreamElement_TextOnly(EArchiveValueType& OutType)
+{
+	EnterStreamElement();
+	OutType = EArchiveValueType::None;
+}
+
 void FJsonArchiveOutputFormatter::LeaveStreamElement()
 {
 	bNeedsComma = true;
@@ -141,6 +171,12 @@ void FJsonArchiveOutputFormatter::EnterMapElement(FString& Name)
 	EnterField(FArchiveFieldName(*Name));
 }
 
+void FJsonArchiveOutputFormatter::EnterMapElement_TextOnly(FString& Name, EArchiveValueType& OutType)
+{
+	EnterMapElement(Name);
+	OutType = EArchiveValueType::None;
+}
+
 void FJsonArchiveOutputFormatter::LeaveMapElement()
 {
 	LeaveField();
diff --git a/Engine/Source/Runtime/CoreUObject/Private/Serialization/Formatters/TaggedBinaryArchiveInputFormatter.cpp b/Engine/Source/Runtime/CoreUObject/Private/Serialization/Formatters/TaggedBinaryArchiveInputFormatter.cpp
--- a/Engine/Source/Runtime/CoreUObject/Private/Serialization/Formatters/TaggedBinaryArchiveInputFormatter.cpp
+++ b/Engine/Source/Runtime/CoreUObject/Private/Serialization/Formatters/TaggedBinaryArchiveInputFormatter.cpp
@@ -73,22 +73,27 @@ FArchive& FTaggedBinaryArchiveInputFormatter::GetUnderlyingArchive()
 	return Inner;
 }
 
+bool FTaggedBinaryArchiveInputFormatter::HasDocumentTree() const
+{
+	return true;
+}
+
 void FTaggedBinaryArchiveInputFormatter::EnterRecord()
 {
 	ExpectType(EArchiveValueType::Record);
 	RecordStack.Push(NextRecordIdx++);
 }
 
-void FTaggedBinaryArchiveInputFormatter::EnterRecord(TArray<FString>& OutKeys)
+void FTaggedBinaryArchiveInputFormatter::EnterRecord_TextOnly(TArray<FString>& OutFieldNames)
 {
 	EnterRecord();
 
 	FRecord& Record = Records[RecordStack.Top()];
 
-	OutKeys.Reset();
+	OutFieldNames.Reset();
 	for (const FField& Field : Record.Fields)
 	{
-		OutKeys.Add(Names[Field.NameIdx]);
+		OutFieldNames.Add(Names[Field.NameIdx]);
 	}
 }
 
@@ -109,7 +114,7 @@ void FTaggedBinaryArchiveInputFormatter::EnterField(FArchiveFieldName Name)
 	Inner.Seek(Field->Offset);
 }
 
-void FTaggedBinaryArchiveInputFormatter::EnterField(FArchiveFieldName Name, EArchiveValueType& OutType)
+void FTaggedBinaryArchiveInputFormatter::EnterField_TextOnly(FArchiveFieldName Name, EArchiveValueType& OutType)
 {
 	EnterField(Name);
 	OutType = PeekType();
@@ -147,7 +152,7 @@ void FTaggedBinaryArchiveInputFormatter::EnterArrayElement()
 {
 }
 
-void FTaggedBinaryArchiveInputFormatter::EnterArrayElement(EArchiveValueType& OutType)
+void FTaggedBinaryArchiveInputFormatter::EnterArrayElement_TextOnly(EArchiveValueType& OutType)
 {
 	OutType = PeekType();
 }
@@ -165,10 +170,10 @@ void FTaggedBinaryArchiveInputFormatter::EnterStream()
 	Inner << Length;
 }
 
-void FTaggedBinaryArchiveInputFormatter::EnterStream(int32& NumElements)
+void FTaggedBinaryArchiveInputFormatter::EnterStream_TextOnly(int32& OutNumElements)
 {
 	ExpectType(EArchiveValueType::Stream);
-	Inner << NumElements;
+	Inner << OutNumElements;
 }
 
 void FTaggedBinaryArchiveInputFormatter::LeaveStream()
@@ -179,7 +184,7 @@ void FTaggedBinaryArchiveInputFormatter::EnterStreamElement()
 {
 }
 
-void FTaggedBinaryArchiveInputFormatter::EnterStreamElement(EArchiveValueType& OutType)
+void FTaggedBinaryArchiveInputFormatter::EnterStreamElement_TextOnly(EArchiveValueType& OutType)
 {
 	OutType = PeekType();
 }
@@ -204,7 +209,7 @@ void FTaggedBinaryArchiveInputFormatter::EnterMapElement(FString& OutName)
 	Inner << OutName;
 }
 
-void FTaggedBinaryArchiveInputFormatter::EnterMapElement(FString& OutName, EArchiveValueType& OutType)
+void FTaggedBinaryArchiveInputFormatter::EnterMapElement_TextOnly(FString& OutName, EArchiveValueType& OutType)
 {
 	EnterMapElement(OutName);
 	OutType = PeekType();
diff --git a/Engine/Source/Runtime/CoreUObject/Private/Serialization/Formatters/TaggedBinaryArchiveOutputFormatter.cpp b/Engine/Source/Runtime/CoreUObject/Private/Serialization/Formatters/TaggedBinaryArchiveOutputFormatter.cpp
--- a/Engine/Source/Runtime/CoreUObject/Private/Serialization/Formatters/TaggedBinaryArchiveOutputFormatter.cpp
+++ b/Engine/Source/Runtime/CoreUObject/Private/Serialization/Formatters/TaggedBinaryArchiveOutputFormatter.cpp
@@ -58,6 +58,11 @@ FArchive& FTaggedBinaryArchiveOutputFormatter::GetUnderlyingArchive()
 	return Inner;
 }
 
+bool FTaggedBinaryArchiveOutputFormatter::HasDocumentTree() const
+{
+	return true;
+}
+
 void FTaggedBinaryArchiveOutputFormatter::EnterRecord()
 {
 	WriteType(EArchiveValueType::Record);
@@ -68,9 +73,10 @@ void FTaggedBinaryArchiveOutputFormatter::EnterRecord()
 	RecordStack.Push(NextRecordIdx++);
 }
 
-void FTaggedBinaryArchiveOutputFormatter::EnterRecord(TArray<FString>& OutFieldNames)
+void FTaggedBinaryArchiveOutputFormatter::EnterRecord_TextOnly(TArray<FString>& OutFieldNames)
 {
 	EnterRecord();
+	OutFieldNames.Reset();
 }
 
 void FTaggedBinaryArchiveOutputFormatter::LeaveRecord()
@@ -90,6 +96,12 @@ void FTaggedBinaryArchiveOutputFormatter::EnterField(FArchiveFieldName Name)
 	Field.Offset = Inner.Tell();
 }
 
+void FTaggedBinaryArchiveOutputFormatter::EnterField_TextOnly(FArchiveFieldName Name, EArchiveValueType& OutType)
+{
+	EnterField(Name);
+	OutType = EArchiveValueType::None;
+}
+
 void FTaggedBinaryArchiveOutputFormatter::LeaveField()
 {
 	FField& Field = Records[RecordStack.Top()].Fields.Top();
@@ -119,6 +131,11 @@ void FTaggedBinaryArchiveOutputFormatter::EnterArrayElement()
 {
 }
 
+void FTaggedBinaryArchiveOutputFormatter::EnterArrayElement_TextOnly(EArchiveValueType& OutType)
+{
+	OutType = EArchiveValueType::None;
+}
+
 void FTaggedBinaryArchiveOutputFormatter::LeaveArrayElement()
 {
 }
@@ -135,6 +152,11 @@ void FTaggedBinaryArchiveOutputFormatter::EnterStream()
 	Inner << Length;
 }
 
+void FTaggedBinaryArchiveOutputFormatter::EnterStream_TextOnly(int32& OutNumElements)
+{
+	OutNumElements = 0;
+}
+
 void FTaggedBinaryArchiveOutputFormatter::LeaveStream()
 {
 	FStream Stream = Streams.Pop();
@@ -151,6 +173,12 @@ void FTaggedBinaryArchiveOutputFormatter::EnterStreamElement()
 	Streams.Top().NumItems++;
 }
 
+void FTaggedBinaryArchiveOutputFormatter::EnterStreamElement_TextOnly(EArchiveValueType& OutType)
+{
+	EnterStreamElement();
+	OutType = EArchiveValueType::None;
+}
+
 void FTaggedBinaryArchiveOutputFormatter::LeaveStreamElement()
 {
 }
@@ -170,6 +198,12 @@ void FTaggedBinaryArchiveOutputFormatter::EnterMapElement(FString& Name)
 	Inner << Name;
 }
 
+void FTaggedBinaryArchiveOutputFormatter::EnterMapElement_TextOnly(FString& Name, EArchiveValueType& OutType)
+{
+	EnterMapElement(Name);
+	OutType = EArchiveValueType::None;
+}
+
 void FTaggedBinaryArchiveOutputFormatter::LeaveMapElement()
 {
 }
diff --git a/Engine/Source/Runtime/CoreUObject/Private/UObject/Class.cpp b/Engine/Source/Runtime/CoreUObject/Private/UObject/Class.cpp
--- a/Engine/Source/Runtime/CoreUObject/Private/UObject/Class.cpp
+++ b/Engine/Source/Runtime/CoreUObject/Private/UObject/Class.cpp
@@ -923,7 +923,7 @@ void UStruct::SerializeTaggedProperties(FStructuredArchive::FSlot Slot, uint8* D
 
 		if (UnderlyingArchive.IsTextFormat())
 		{
-			PropertiesRecord.Emplace(Slot.EnterRecord(FieldNames));
+			PropertiesRecord.Emplace(Slot.EnterRecord_TextOnly(FieldNames));
 		}
 		else
 		{
diff --git a/Engine/Source/Runtime/CoreUObject/Public/Serialization/Formatters/JsonArchiveInputFormatter.h b/Engine/Source/Runtime/CoreUObject/Public/Serialization/Formatters/JsonArchiveInputFormatter.h
--- a/Engine/Source/Runtime/CoreUObject/Public/Serialization/Formatters/JsonArchiveInputFormatter.h
+++ b/Engine/Source/Runtime/CoreUObject/Public/Serialization/Formatters/JsonArchiveInputFormatter.h
@@ -10,7 +10,7 @@ class FJsonValue;
 
 #if WITH_TEXT_ARCHIVE_SUPPORT
 
-class FJsonArchiveInputFormatter : public FAnnotatedStructuredArchiveFormatter
+class FJsonArchiveInputFormatter final : public FStructuredArchiveFormatter
 {
 public:
 	FJsonArchiveInputFormatter(FArchive& InInner, TFunction<UObject* (const FString&)> InResolveObjectName = nullptr);
@@ -19,31 +19,33 @@ public:
 	virtual FArchive& GetUnderlyingArchive() override;
 	virtual FStructuredArchiveFormatter* CreateSubtreeReader() override;
 
+	virtual bool HasDocumentTree() const override;
+
 	virtual void EnterRecord() override;
-	virtual void EnterRecord(TArray<FString>& OutFieldNamesWhenLoading) override;
+	virtual void EnterRecord_TextOnly(TArray<FString>& OutFieldNames) override;
 	virtual void LeaveRecord() override;
 	virtual void EnterField(FArchiveFieldName Name) override;
-	virtual void EnterField(FArchiveFieldName Name, EArchiveValueType& OutType) override;
+	virtual void EnterField_TextOnly(FArchiveFieldName Name, EArchiveValueType& OutType) override;
 	virtual void LeaveField() override;
 	virtual bool TryEnterField(FArchiveFieldName Name, bool bEnterWhenSaving) override;
 
 	virtual void EnterArray(int32& NumElements) override;
 	virtual void LeaveArray() override;
 	virtual void EnterArrayElement() override;
-	virtual void EnterArrayElement(EArchiveValueType& OutType) override;
+	virtual void EnterArrayElement_TextOnly(EArchiveValueType& OutType) override;
 	virtual void LeaveArrayElement() override;
 
 	virtual void EnterStream() override;
-	virtual void EnterStream(int32& NumElements) override;
+	virtual void EnterStream_TextOnly(int32& NumElements) override;
 	virtual void LeaveStream() override;
 	virtual void EnterStreamElement() override;
-	virtual void EnterStreamElement(EArchiveValueType& OutType) override;
+	virtual void EnterStreamElement_TextOnly(EArchiveValueType& OutType) override;
 	virtual void LeaveStreamElement() override;
 
 	virtual void EnterMap(int32& NumElements) override;
 	virtual void LeaveMap() override;
 	virtual void EnterMapElement(FString& Name) override;
-	virtual void EnterMapElement(FString& OutName, EArchiveValueType& OutType) override;
+	virtual void EnterMapElement_TextOnly(FString& OutName, EArchiveValueType& OutType) override;
 	virtual void LeaveMapElement() override;
 
 	virtual void Serialize(uint8& Value) override;
diff --git a/Engine/Source/Runtime/CoreUObject/Public/Serialization/Formatters/JsonArchiveOutputFormatter.h b/Engine/Source/Runtime/CoreUObject/Public/Serialization/Formatters/JsonArchiveOutputFormatter.h
--- a/Engine/Source/Runtime/CoreUObject/Public/Serialization/Formatters/JsonArchiveOutputFormatter.h
+++ b/Engine/Source/Runtime/CoreUObject/Public/Serialization/Formatters/JsonArchiveOutputFormatter.h
@@ -7,7 +7,7 @@
 
 #if WITH_TEXT_ARCHIVE_SUPPORT
 
-class COREUOBJECT_API FJsonArchiveOutputFormatter : public FStructuredArchiveFormatter
+class COREUOBJECT_API FJsonArchiveOutputFormatter final : public FStructuredArchiveFormatter
 {
 public:
 	FJsonArchiveOutputFormatter(FArchive& InInner);
@@ -15,26 +15,33 @@ public:
 
 	virtual FArchive& GetUnderlyingArchive() override;
 
+	virtual bool HasDocumentTree() const override;
+
 	virtual void EnterRecord() override;
-	virtual void EnterRecord(TArray<FString>& OutFieldNamesWhenLoading) override;
+	virtual void EnterRecord_TextOnly(TArray<FString>& OutFieldNames) override;
 	virtual void LeaveRecord() override;
 	virtual void EnterField(FArchiveFieldName Name) override;
+	virtual void EnterField_TextOnly(FArchiveFieldName Name, EArchiveValueType& OutType) override;
 	virtual void LeaveField() override;
 	virtual bool TryEnterField(FArchiveFieldName Name, bool bEnterWhenSaving) override;
 
 	virtual void EnterArray(int32& NumElements) override;
 	virtual void LeaveArray() override;
 	virtual void EnterArrayElement() override;
+	virtual void EnterArrayElement_TextOnly(EArchiveValueType& OutType) override;
 	virtual void LeaveArrayElement() override;
 
 	virtual void EnterStream() override;
+	virtual void EnterStream_TextOnly(int32& NumElements) override;
 	virtual void LeaveStream() override;
 	virtual void EnterStreamElement() override;
+	virtual void EnterStreamElement_TextOnly(EArchiveValueType& OutType) override;
 	virtual void LeaveStreamElement() override;
 
 	virtual void EnterMap(int32& NumElements) override;
 	virtual void LeaveMap() override;
 	virtual void EnterMapElement(FString& Name) override;
+	virtual void EnterMapElement_TextOnly(FString& Name, EArchiveValueType& OutType) override;
 	virtual void LeaveMapElement() override;
 
 	virtual void Serialize(uint8& Value) override;
diff --git a/Engine/Source/Runtime/CoreUObject/Public/Serialization/Formatters/TaggedBinaryArchiveInputFormatter.h b/Engine/Source/Runtime/CoreUObject/Public/Serialization/Formatters/TaggedBinaryArchiveInputFormatter.h
--- a/Engine/Source/Runtime/CoreUObject/Public/Serialization/Formatters/TaggedBinaryArchiveInputFormatter.h
+++ b/Engine/Source/Runtime/CoreUObject/Public/Serialization/Formatters/TaggedBinaryArchiveInputFormatter.h
@@ -7,7 +7,7 @@
 
 #if WITH_TEXT_ARCHIVE_SUPPORT
 
-class FTaggedBinaryArchiveInputFormatter : public FAnnotatedStructuredArchiveFormatter
+class FTaggedBinaryArchiveInputFormatter : public FStructuredArchiveFormatter
 {
 public:
 	FTaggedBinaryArchiveInputFormatter(FArchive& InInner, TFunction<void(FArchive&, UObject*&)> InSerializeObject);
@@ -15,31 +15,33 @@ public:
 
 	virtual FArchive& GetUnderlyingArchive() override;
 
+	virtual bool HasDocumentTree() const override;
+
 	virtual void EnterRecord() override;
-	virtual void EnterRecord(TArray<FString>& OutKeys) override;
+	virtual void EnterRecord_TextOnly(TArray<FString>& OutFieldNames) override;
 	virtual void LeaveRecord() override;
 	virtual void EnterField(FArchiveFieldName Name) override;
-	virtual void EnterField(FArchiveFieldName Name, EArchiveValueType& OutType) override;
+	virtual void EnterField_TextOnly(FArchiveFieldName Name, EArchiveValueType& OutType) override;
 	virtual void LeaveField() override;
 	virtual bool TryEnterField(FArchiveFieldName Name, bool bEnterWhenSaving) override;
 
 	virtual void EnterArray(int32& NumElements) override;
 	virtual void LeaveArray() override;
 	virtual void EnterArrayElement() override;
-	virtual void EnterArrayElement(EArchiveValueType& OutType) override;
+	virtual void EnterArrayElement_TextOnly(EArchiveValueType& OutType) override;
 	virtual void LeaveArrayElement() override;
 
 	virtual void EnterStream() override;
-	virtual void EnterStream(int32& NumElements) override;
+	virtual void EnterStream_TextOnly(int32& OutNumElements) override;
 	virtual void LeaveStream() override;
 	virtual void EnterStreamElement() override;
-	virtual void EnterStreamElement(EArchiveValueType& OutType) override;
+	virtual void EnterStreamElement_TextOnly(EArchiveValueType& OutType) override;
 	virtual void LeaveStreamElement() override;
 
 	virtual void EnterMap(int32& NumElements) override;
 	virtual void LeaveMap() override;
 	virtual void EnterMapElement(FString& OutName) override;
-	virtual void EnterMapElement(FString& OutName, EArchiveValueType& OutType) override;
+	virtual void EnterMapElement_TextOnly(FString& OutName, EArchiveValueType& OutType) override;
 	virtual void LeaveMapElement() override;
 
 	virtual void Serialize(uint8& Value) override;
diff --git a/Engine/Source/Runtime/CoreUObject/Public/Serialization/Formatters/TaggedBinaryArchiveOutputFormatter.h b/Engine/Source/Runtime/CoreUObject/Public/Serialization/Formatters/TaggedBinaryArchiveOutputFormatter.h
--- a/Engine/Source/Runtime/CoreUObject/Public/Serialization/Formatters/TaggedBinaryArchiveOutputFormatter.h
+++ b/Engine/Source/Runtime/CoreUObject/Public/Serialization/Formatters/TaggedBinaryArchiveOutputFormatter.h
@@ -15,26 +15,33 @@ public:
 
 	virtual FArchive& GetUnderlyingArchive() override;
 
+	virtual bool HasDocumentTree() const;
+
 	virtual void EnterRecord() override;
-	virtual void EnterRecord(TArray<FString>& OutFieldNames) override;
+	virtual void EnterRecord_TextOnly(TArray<FString>& OutFieldNames) override;
 	virtual void LeaveRecord() override;
 	virtual void EnterField(FArchiveFieldName Name) override;
+	virtual void EnterField_TextOnly(FArchiveFieldName Name, EArchiveValueType& OutType) override;
 	virtual void LeaveField() override;
 	virtual bool TryEnterField(FArchiveFieldName Name, bool bEnterWhenSaving) override;
 
 	virtual void EnterArray(int32& NumElements) override;
 	virtual void LeaveArray() override;
 	virtual void EnterArrayElement() override;
+	virtual void EnterArrayElement_TextOnly(EArchiveValueType& OutType) override;
 	virtual void LeaveArrayElement() override;
 
 	virtual void EnterStream() override;
+	virtual void EnterStream_TextOnly(int32& NumElements) override;
 	virtual void LeaveStream() override;
 	virtual void EnterStreamElement() override;
+	virtual void EnterStreamElement_TextOnly(EArchiveValueType& OutType) override;
 	virtual void LeaveStreamElement() override;
 
 	virtual void EnterMap(int32& NumElements) override;
 	virtual void LeaveMap() override;
 	virtual void EnterMapElement(FString& Name) override;
+	virtual void EnterMapElement_TextOnly(FString& Name, EArchiveValueType& OutType) override;
 	virtual void LeaveMapElement() override;
 
 	virtual void Serialize(uint8& Value) override;
