VectorLoadFloat2(Ptr) added, which loads { Ptr[0], Ptr[1], Ptr[0], Ptr[1] } into a VectorRegister.

diff --git a/Engine/Source/Runtime/Core/Public/Math/UnrealMathDirectX.h b/Engine/Source/Runtime/Core/Public/Math/UnrealMathDirectX.h
--- a/Engine/Source/Runtime/Core/Public/Math/UnrealMathDirectX.h
+++ b/Engine/Source/Runtime/Core/Public/Math/UnrealMathDirectX.h
@@ -142,6 +142,14 @@ FORCEINLINE VectorRegisterInt MakeVectorRegisterInt(int32 X, int32 Y, int32 Z, i
  */
 #define VectorLoadFloat1( Ptr )			DirectX::XMVectorReplicatePtr( (const float*)(Ptr) )
 
+/**
+ * Loads 2 floats from unaligned memory into X and Y and duplicates them in Z and W.
+ *
+ * @param Ptr	Unaligned memory pointer to the floats
+ * @return		VectorRegister(Ptr[0], Ptr[1], Ptr[0], Ptr[1])
+ */
+#define VectorLoadFloat2( Ptr )			MakeVectorRegister( ((const float*)(Ptr))[0], ((const float*)(Ptr))[1], ((const float*)(Ptr))[0], ((const float*)(Ptr))[1] )
+
 /**
  * Creates a vector out of three FLOATs and leaves W undefined.
  *
diff --git a/Engine/Source/Runtime/Core/Public/Math/UnrealMathFPU.h b/Engine/Source/Runtime/Core/Public/Math/UnrealMathFPU.h
--- a/Engine/Source/Runtime/Core/Public/Math/UnrealMathFPU.h
+++ b/Engine/Source/Runtime/Core/Public/Math/UnrealMathFPU.h
@@ -161,6 +161,14 @@ FORCEINLINE VectorRegisterInt MakeVectorRegisterInt(int32 X, int32 Y, int32 Z, i
  */
 #define VectorLoadFloat1( Ptr )			MakeVectorRegister( ((const float*)(Ptr))[0], ((const float*)(Ptr))[0], ((const float*)(Ptr))[0], ((const float*)(Ptr))[0] )
 
+/**
+ * Loads 2 floats from unaligned memory into X and Y and duplicates them in Z and W.
+ *
+ * @param Ptr	Unaligned memory pointer to the floats
+ * @return		VectorRegister(Ptr[0], Ptr[1], Ptr[0], Ptr[1])
+ */
+#define VectorLoadFloat2( Ptr )			MakeVectorRegister( ((const float*)(Ptr))[0], ((const float*)(Ptr))[1], ((const float*)(Ptr))[0], ((const float*)(Ptr))[1] )
+
 /**
  * Creates a vector out of three FLOATs and leaves W undefined.
  *
diff --git a/Engine/Source/Runtime/Core/Public/Math/UnrealMathNeon.h b/Engine/Source/Runtime/Core/Public/Math/UnrealMathNeon.h
--- a/Engine/Source/Runtime/Core/Public/Math/UnrealMathNeon.h
+++ b/Engine/Source/Runtime/Core/Public/Math/UnrealMathNeon.h
@@ -128,6 +128,14 @@ FORCEINLINE VectorRegister VectorLoad( const void* Ptr )
 }
 
 /**
+ * Loads 2 floats from unaligned memory into X and Y and duplicates them in Z and W.
+ *
+ * @param Ptr	Unaligned memory pointer to the floats
+ * @return		VectorRegister(Ptr[0], Ptr[1], Ptr[0], Ptr[1])
+ */
+#define VectorLoadFloat2( Ptr )			MakeVectorRegister( ((const float*)(Ptr))[0], ((const float*)(Ptr))[1], ((const float*)(Ptr))[0], ((const float*)(Ptr))[1] )
+
+ /**
  * Loads 3 floats from unaligned memory and leaves W undefined.
  *
  * @param Ptr	Unaligned memory pointer to the 3 floats
diff --git a/Engine/Source/Runtime/Core/Public/Math/UnrealMathSSE.h b/Engine/Source/Runtime/Core/Public/Math/UnrealMathSSE.h
--- a/Engine/Source/Runtime/Core/Public/Math/UnrealMathSSE.h
+++ b/Engine/Source/Runtime/Core/Public/Math/UnrealMathSSE.h
@@ -163,6 +163,14 @@ FORCEINLINE float VectorGetComponent(VectorRegister Vec, uint32 ComponentIndex)
  */
 #define VectorLoadFloat1( Ptr )			_mm_load1_ps( (float*)(Ptr) )
 
+/**
+ * Loads 2 floats from unaligned memory into X and Y and duplicates them in Z and W.
+ *
+ * @param Ptr	Unaligned memory pointer to the floats
+ * @return		VectorRegister(Ptr[0], Ptr[1], Ptr[0], Ptr[1])
+ */
+#define VectorLoadFloat2( Ptr )			_mm_castpd_ps(_mm_load1_pd((double*)(Ptr)))
+
 /**
  * Creates a vector out of three FLOATs and leaves W undefined.
  *
