Remove assert that GLogConsole is set; it won't be for command line utilities that don't depend on ApplicationCore.

#rnx

diff --git a/Engine/Source/Runtime/Core/Private/Unix/UnixPlatformOutputDevices.cpp b/Engine/Source/Runtime/Core/Private/Unix/UnixPlatformOutputDevices.cpp
--- a/Engine/Source/Runtime/Core/Private/Unix/UnixPlatformOutputDevices.cpp
+++ b/Engine/Source/Runtime/Core/Private/Unix/UnixPlatformOutputDevices.cpp
@@ -17,7 +17,6 @@
 void FUnixOutputDevices::SetupOutputDevices()
 {
 	check(GLog);
-	check(GLogConsole);
 
 	CachedAbsoluteFilename[0] = 0;
 
@@ -25,11 +24,13 @@ void FUnixOutputDevices::SetupOutputDevices()
 	GLog->AddOutputDevice(FPlatformOutputDevices::GetLog());
 
 	// @todo: set to false for minor utils?
-	bool bLogToConsole = !NO_LOGGING && !FParse::Param(FCommandLine::Get(), TEXT("NOCONSOLE"));
-
-	if (bLogToConsole)
+	if (GLogConsole != nullptr)
 	{
-		GLog->AddOutputDevice(GLogConsole);
+		bool bLogToConsole = !NO_LOGGING && !FParse::Param(FCommandLine::Get(), TEXT("NOCONSOLE"));
+		if (bLogToConsole)
+		{
+			GLog->AddOutputDevice(GLogConsole);
+		}
 	}
 
 	// debug and event logging is not really supported on Unix. 
