Fix for external profilers. The loop iterating over the available external profilers did actually not iterate over the available external profilers.

from UDN (Andrea Sansottera / Creative Assembly) https://udn.unrealengine.com/questions/450490/bug-fix-for-external-profilers.html

diff --git a/Engine/Source/Runtime/Core/Private/ProfilingDebugging/ExternalProfiler.cpp b/Engine/Source/Runtime/Core/Private/ProfilingDebugging/ExternalProfiler.cpp
--- a/Engine/Source/Runtime/Core/Private/ProfilingDebugging/ExternalProfiler.cpp
+++ b/Engine/Source/Runtime/Core/Private/ProfilingDebugging/ExternalProfiler.cpp
@@ -50,26 +50,26 @@ FExternalProfiler* FActiveExternalProfilerBase::GetActiveProfiler()
 	// Create profiler on demand.
 	if (ActiveProfiler == NULL && !bDidInitialize)
 	{
-		const int32 AvailableProfilerCount = IModularFeatures::Get().GetModularFeatureImplementationCount(FExternalProfiler::GetFeatureName());
-		for (int32 CurProfilerIndex = 0; CurProfilerIndex < AvailableProfilerCount; ++CurProfilerIndex)
-		{
-			FExternalProfiler& CurProfiler = IModularFeatures::Get().GetModularFeature<FExternalProfiler&>(FExternalProfiler::GetFeatureName());
+		const FName FeatureName = FExternalProfiler::GetFeatureName();
+		TArray<FExternalProfiler*> AvailbleProfilers = IModularFeatures::Get().GetModularFeatureImplementations<FExternalProfiler>(FeatureName);
 
+		for (FExternalProfiler* CurProfiler : AvailbleProfilers)
+		{
 #if 0
 			// Logging disabled here as it can cause a stack overflow whilst flushing logs during EnginePreInit
-			UE_LOG(LogExternalProfiler, Log, TEXT("Found external profiler: %s"), CurProfiler.GetProfilerName());
+			UE_LOG(LogExternalProfiler, Log, TEXT("Found external profiler: %s"), CurProfiler->GetProfilerName());
 #endif
 
 			// Default to the first profiler we have if none were specified on the command-line
 			if (ActiveProfiler == NULL)
 			{
-				ActiveProfiler = &CurProfiler;
+				ActiveProfiler = CurProfiler;
 			}
 
 			// Check to see if the profiler was specified on the command-line (e.g., "-VTune")
-			if (FParse::Param(FCommandLine::Get(), CurProfiler.GetProfilerName()))
+			if (FParse::Param(FCommandLine::Get(), CurProfiler->GetProfilerName()))
 			{
-				ActiveProfiler = &CurProfiler;
+				ActiveProfiler = CurProfiler;
 			}
 		}
 
@@ -88,6 +88,7 @@ FExternalProfiler* FActiveExternalProfilerBase::GetActiveProfiler()
 		// Don't try to initialize again this session
 		bDidInitialize = true;
 	}
+
 	return ActiveProfiler;
 }
 
