Added check in order to suppress CIS static analysis warning

diff --git a/Engine/Source/Runtime/Core/Private/ProfilingDebugging/ExternalProfiler.cpp b/Engine/Source/Runtime/Core/Private/ProfilingDebugging/ExternalProfiler.cpp
--- a/Engine/Source/Runtime/Core/Private/ProfilingDebugging/ExternalProfiler.cpp
+++ b/Engine/Source/Runtime/Core/Private/ProfilingDebugging/ExternalProfiler.cpp
@@ -55,6 +55,8 @@ FExternalProfiler* FActiveExternalProfilerBase::GetActiveProfiler()
 
 		for (FExternalProfiler* CurProfiler : AvailbleProfilers)
 		{
+			check(CurProfiler != nullptr);
+
 #if 0
 			// Logging disabled here as it can cause a stack overflow whilst flushing logs during EnginePreInit
 			UE_LOG(LogExternalProfiler, Log, TEXT("Found external profiler: %s"), CurProfiler->GetProfilerName());
