Fixed out of bounds read in FWindowsPlatformMisc::QueryRegKey.

String buffer size in bytes was erroneously used instead of length in characters to construct result FString.

Edigrated CL4318138 by [REDACTED_PERSON]

diff --git a/Engine/Source/Runtime/Core/Private/Windows/WindowsPlatformMisc.cpp b/Engine/Source/Runtime/Core/Private/Windows/WindowsPlatformMisc.cpp
--- a/Engine/Source/Runtime/Core/Private/Windows/WindowsPlatformMisc.cpp
+++ b/Engine/Source/Runtime/Core/Private/Windows/WindowsPlatformMisc.cpp
@@ -2602,8 +2602,8 @@ bool FWindowsPlatformMisc::QueryRegKey( const Windows::HKEY InKey, const TCHAR*
 				char *Buffer = new char[Size];
 				if (RegQueryValueEx( Key, InValueName, NULL, NULL, (LPBYTE)Buffer, &Size ) == ERROR_SUCCESS)
 				{
-					OutData = FString( Size-1, (TCHAR*)Buffer );
-					OutData.TrimToNullTerminator();
+					const uint32 Length = (Size / sizeof(TCHAR)) - 1;
+					OutData = FString( Length, (TCHAR*)Buffer );
 					bSuccess = true;
 				}
 				delete [] Buffer;
