Remove SecureCRT usage because it is not well supported and provides an inconsistency with the implementations on other platforms.

diff --git a/Engine/Source/Runtime/Core/Public/Android/AndroidString.h b/Engine/Source/Runtime/Core/Public/Android/AndroidString.h
--- a/Engine/Source/Runtime/Core/Public/Android/AndroidString.h
+++ b/Engine/Source/Runtime/Core/Public/Android/AndroidString.h
@@ -480,7 +480,7 @@ struct FAndroidPlatformString : public FGenericPlatformString
 			}
 		}
 		NewFormat[NewIndex] = 0;
-#endif // USE_SECURE_CRT
+#endif // PLATFORM_USE_LS_SPEC_FOR_WIDECHAR
 		int32 Result = vswprintf( Dest, Count, NewFormat, ArgPtr);
 		va_end( ArgPtr );
 		return Result;
diff --git a/Engine/Source/Runtime/Core/Public/GenericPlatform/MicrosoftPlatformString.h b/Engine/Source/Runtime/Core/Public/GenericPlatform/MicrosoftPlatformString.h
--- a/Engine/Source/Runtime/Core/Public/GenericPlatform/MicrosoftPlatformString.h
+++ b/Engine/Source/Runtime/Core/Public/GenericPlatform/MicrosoftPlatformString.h
@@ -10,11 +10,9 @@
 * Microsoft specific implementation 
 **/
 
-#if !USE_SECURE_CRT
 #pragma warning(push)
 #pragma warning(disable : 4996) // 'function' was declared deprecated  (needed for the secure string functions)
 #pragma warning(disable : 4995) // 'function' was declared deprecated  (needed for the secure string functions)
-#endif
 
 struct FMicrosoftPlatformString : public FGenericPlatformString
 {
@@ -23,43 +21,24 @@ struct FMicrosoftPlatformString : public FGenericPlatformString
 	 **/
 	static FORCEINLINE WIDECHAR* Strcpy(WIDECHAR* Dest, SIZE_T DestCount, const WIDECHAR* Src)
 	{
-#if USE_SECURE_CRT
-		_tcscpy_s(Dest, DestCount, Src);
-		return Dest;
-#else
 		return (WIDECHAR*)_tcscpy(Dest, Src);
-#endif // USE_SECURE_CRT
 	}
 
 	static FORCEINLINE WIDECHAR* Strncpy(WIDECHAR* Dest, const WIDECHAR* Src, SIZE_T MaxLen)
 	{
-#if USE_SECURE_CRT
-		_tcsncpy_s(Dest, MaxLen, Src, MaxLen-1);
-#else
 		_tcsncpy(Dest, Src, MaxLen-1);
 		Dest[MaxLen-1] = 0;
-#endif // USE_SECURE_CRT
 		return Dest;
 	}
 
 	static FORCEINLINE WIDECHAR* Strcat(WIDECHAR* Dest, SIZE_T DestCount, const WIDECHAR* Src)
 	{
-#if USE_SECURE_CRT
-		_tcscat_s(Dest, DestCount, Src);
-		return Dest;
-#else
 		return (WIDECHAR*)_tcscat(Dest, Src);
-#endif // USE_SECURE_CRT
 	}
 
 	static FORCEINLINE WIDECHAR* Strupr(WIDECHAR* Dest, SIZE_T DestCount)
 	{
-#if USE_SECURE_CRT
-		_tcsupr_s(Dest, DestCount);
-		return Dest;
-#else
 		return (WIDECHAR*)_tcsupr(Dest);
-#endif // USE_SECURE_CRT
 	}
 
 	static FORCEINLINE int32 Strcmp( const WIDECHAR* String1, const WIDECHAR* String2 )
@@ -139,11 +118,7 @@ struct FMicrosoftPlatformString : public FGenericPlatformString
 
 	static FORCEINLINE int32 GetVarArgs( WIDECHAR* Dest, SIZE_T DestSize, int32 Count, const WIDECHAR*& Fmt, va_list ArgPtr )
 	{
-#if USE_SECURE_CRT
-		int32 Result = _vsntprintf_s( Dest, DestSize, Count, Fmt, ArgPtr );
-#else
 		int32 Result = vswprintf(Dest, Count, Fmt, ArgPtr);
-#endif // USE_SECURE_CRT
 		va_end( ArgPtr );
 		return Result;
 	}
@@ -153,42 +128,23 @@ struct FMicrosoftPlatformString : public FGenericPlatformString
 	 **/
 	static FORCEINLINE ANSICHAR* Strcpy(ANSICHAR* Dest, SIZE_T DestCount, const ANSICHAR* Src)
 	{
-#if USE_SECURE_CRT
-		strcpy_s(Dest, DestCount, Src);
-		return Dest;
-#else
 		return (ANSICHAR*)strcpy(Dest, Src);
-#endif // USE_SECURE_CRT
 	}
 
 	static FORCEINLINE void Strncpy(ANSICHAR* Dest, const ANSICHAR* Src, SIZE_T MaxLen)
 	{
-#if USE_SECURE_CRT
-		strncpy_s(Dest, MaxLen, Src, MaxLen-1);
-#else
 		strncpy(Dest, Src, MaxLen);
 		Dest[MaxLen-1] = 0;
-#endif // USE_SECURE_CRT
 	}
 
 	static FORCEINLINE ANSICHAR* Strcat(ANSICHAR* Dest, SIZE_T DestCount, const ANSICHAR* Src)
 	{
-#if USE_SECURE_CRT
-		strcat_s( Dest, DestCount, Src );
-		return Dest;
-#else
 		return (ANSICHAR*)strcat( Dest, Src );
-#endif // USE_SECURE_CRT
 	}
 
 	static FORCEINLINE ANSICHAR* Strupr(ANSICHAR* Dest, SIZE_T DestCount)
 	{
-#if USE_SECURE_CRT
-		_strupr_s(Dest, DestCount);
-		return Dest;
-#else
 		return (ANSICHAR*)_strupr(Dest);
-#endif // USE_SECURE_CRT
 	}
 
 	static FORCEINLINE int32 Strcmp( const ANSICHAR* String1, const ANSICHAR* String2 )
@@ -301,11 +257,7 @@ struct FMicrosoftPlatformString : public FGenericPlatformString
 
 	static FORCEINLINE int32 GetVarArgs( ANSICHAR* Dest, SIZE_T DestSize, int32 Count, const ANSICHAR*& Fmt, va_list ArgPtr )
 	{
-#if USE_SECURE_CRT
-		int32 Result = _vsnprintf_s( Dest, DestSize, Count, Fmt, ArgPtr );
-#else
 		int32 Result = vsnprintf( Dest, Count, Fmt, ArgPtr );
-#endif // USE_SECURE_CRT
 		va_end( ArgPtr );
 		return Result;
 	}
@@ -320,6 +272,4 @@ struct FMicrosoftPlatformString : public FGenericPlatformString
 	}
 };
 
-#if !USE_SECURE_CRT
 #pragma warning(pop) // 'function' was was declared deprecated  (needed for the secure string functions)
-#endif
diff --git a/Engine/Source/Runtime/Core/Public/Windows/WIndowsPlatform.h b/Engine/Source/Runtime/Core/Public/Windows/WIndowsPlatform.h
--- a/Engine/Source/Runtime/Core/Public/Windows/WIndowsPlatform.h
+++ b/Engine/Source/Runtime/Core/Public/Windows/WIndowsPlatform.h
@@ -162,8 +162,3 @@ typedef FWindowsPlatformTypes FPlatformTypes;
 #ifndef ENABLE_WIN_ALLOC_TRACKING
 #define ENABLE_WIN_ALLOC_TRACKING 0
 #endif
-
-#ifndef USE_SECURE_CRT
-#define USE_SECURE_CRT 0
-#endif
-
