Removal of unused argument from FPlatformString::GetVarArgs() and TCString::GetVarArgs(), and deprecate old usage.
Removal of unused variable from FGenericWidePlatformString::GetVarArgs().

diff --git a/Engine/Plugins/Media/AjaMedia/Source/AjaMedia/Private/Aja/Aja.cpp b/Engine/Plugins/Media/AjaMedia/Source/AjaMedia/Private/Aja/Aja.cpp
--- a/Engine/Plugins/Media/AjaMedia/Source/AjaMedia/Private/Aja/Aja.cpp
+++ b/Engine/Plugins/Media/AjaMedia/Source/AjaMedia/Private/Aja/Aja.cpp
@@ -110,7 +110,7 @@ void FAja::LogInfo(const TCHAR* InFormat, ...)
 	va_list Args;
 
 	va_start(Args, InFormat );
-	FCString::GetVarArgs(TempString, ARRAY_COUNT(TempString), ARRAY_COUNT(TempString) - 1, InFormat, Args);
+	FCString::GetVarArgs(TempString, ARRAY_COUNT(TempString), InFormat, Args);
 	va_end(Args);
 
 	UE_LOG(LogAjaMedia, Log, TempString);
@@ -124,7 +124,7 @@ void FAja::LogWarning(const TCHAR* InFormat, ...)
 	va_list Args;
 
 	va_start(Args, InFormat );
-	FCString::GetVarArgs(TempString, ARRAY_COUNT(TempString), ARRAY_COUNT(TempString) - 1, InFormat, Args);
+	FCString::GetVarArgs(TempString, ARRAY_COUNT(TempString), InFormat, Args);
 	va_end(Args);
 
 	UE_LOG(LogAjaMedia, Warning, TempString);
@@ -138,7 +138,7 @@ void FAja::LogError(const TCHAR* InFormat, ...)
 	va_list Args;
 
 	va_start(Args, InFormat );
-	FCString::GetVarArgs(TempString, ARRAY_COUNT(TempString), ARRAY_COUNT(TempString) - 1, InFormat, Args);
+	FCString::GetVarArgs(TempString, ARRAY_COUNT(TempString), InFormat, Args);
 	va_end(Args);
 
 	UE_LOG(LogAjaMedia, Error, TempString);
diff --git a/Engine/Source/Programs/UnrealFrontend/Private/Commands/StatsConvertCommand.cpp b/Engine/Source/Programs/UnrealFrontend/Private/Commands/StatsConvertCommand.cpp
--- a/Engine/Source/Programs/UnrealFrontend/Private/Commands/StatsConvertCommand.cpp
+++ b/Engine/Source/Programs/UnrealFrontend/Private/Commands/StatsConvertCommand.cpp
@@ -119,7 +119,7 @@ void FCSVStatsProfiler::WriteString( const ANSICHAR* Format, ... )
 	va_list ArgPtr;
 	va_start( ArgPtr, Format );
 	// Build the string.
-	int32 Result = FCStringAnsi::GetVarArgs( Array, ARRAY_COUNT( Array ), ARRAY_COUNT( Array ) - 1, Format, ArgPtr );
+	int32 Result = FCStringAnsi::GetVarArgs( Array, ARRAY_COUNT( Array ), Format, ArgPtr );
 	// Now write that to the file.
 	CSVWriter->Serialize( (void*)Array, Result );
 }
diff --git a/Engine/Source/Runtime/Core/Private/GenericPlatform/GenericWidePlatformString.cpp b/Engine/Source/Runtime/Core/Private/GenericPlatform/GenericWidePlatformString.cpp
--- a/Engine/Source/Runtime/Core/Private/GenericPlatform/GenericWidePlatformString.cpp
+++ b/Engine/Source/Runtime/Core/Private/GenericPlatform/GenericWidePlatformString.cpp
@@ -179,60 +179,59 @@ int32 TestGetVarArgs(WIDECHAR* OutputString, const WIDECHAR* Format, ...)
 {
     va_list ArgPtr;
     va_start( ArgPtr, Format );
-    
-    return FGenericWidePlatformString::GetVarArgs(OutputString, OUTPUT_SIZE, OUTPUT_SIZE - 1, Format, ArgPtr);
+
+	return FGenericWidePlatformString::GetVarArgs(OutputString, OUTPUT_SIZE, Format, ArgPtr);
 }
 
 void RunGetVarArgsTests()
 {
-    WIDECHAR OutputString[OUTPUT_SIZE];
-
-    TestGetVarArgs(OutputString, TEXT("Test A|%-20s|%20s|%10.2f|%-10.2f|"), TEXT("LEFT"), TEXT("RIGHT"), 33.333333, 66.666666);
-    check(FString(OutputString) == FString(TEXT("Test A|LEFT                |               RIGHT|     33.33|66.67     |")));
-    
-    TestGetVarArgs(OutputString, TEXT("Test B|Percents:%%%%%%%d|"), 3);
-    check(FString(OutputString) == FString(TEXT("Test B|Percents:%%%3|")));
-    
-    TestGetVarArgs(OutputString, TEXT("Test C|%d|%i|%X|%x|%u|"), 12345, 54321, 0x123AbC, 15, 99);
-    check(FString(OutputString) == FString(TEXT("Test C|12345|54321|123ABC|f|99|")));
-    
-    TestGetVarArgs(OutputString, TEXT("Test D|%p|"), 0x12345);
-    check(FString(OutputString) == FString(TEXT("Test D|0x12345|")));
-
-    TestGetVarArgs(OutputString, TEXT("Test E|%lld|"), 12345678912345);
-    check(FString(OutputString) == FString(TEXT("Test E|12345678912345|")));
-
-    TestGetVarArgs(OutputString, TEXT("Test F|%f|%e|%g|"), 123.456, 123.456, 123.456);
-    check(FString(OutputString) == FString(TEXT("Test F|123.456000|1.234560e+02|123.456|")));
+	WIDECHAR OutputString[OUTPUT_SIZE];
+
+	TestGetVarArgs(OutputString, TEXT("Test A|%-20s|%20s|%10.2f|%-10.2f|"), TEXT("LEFT"), TEXT("RIGHT"), 33.333333, 66.666666);
+	check(FString(OutputString) == FString(TEXT("Test A|LEFT                |               RIGHT|     33.33|66.67     |")));
+
+	TestGetVarArgs(OutputString, TEXT("Test B|Percents:%%%%%%%d|"), 3);
+	check(FString(OutputString) == FString(TEXT("Test B|Percents:%%%3|")));
+
+	TestGetVarArgs(OutputString, TEXT("Test C|%d|%i|%X|%x|%u|"), 12345, 54321, 0x123AbC, 15, 99);
+	check(FString(OutputString) == FString(TEXT("Test C|12345|54321|123ABC|f|99|")));
+
+	TestGetVarArgs(OutputString, TEXT("Test D|%p|"), 0x12345);
+	check(FString(OutputString) == FString(TEXT("Test D|0x12345|")));
+
+	TestGetVarArgs(OutputString, TEXT("Test E|%lld|"), 12345678912345);
+	check(FString(OutputString) == FString(TEXT("Test E|12345678912345|")));
+
+	TestGetVarArgs(OutputString, TEXT("Test F|%f|%e|%g|"), 123.456, 123.456, 123.456);
+	check(FString(OutputString) == FString(TEXT("Test F|123.456000|1.234560e+02|123.456|")));
 }
 #endif
 
-int32 FGenericWidePlatformString::GetVarArgs( WIDECHAR* Dest, SIZE_T DestSize, int32 Count, const WIDECHAR*& Fmt, va_list ArgPtr )
+int32 FGenericWidePlatformString::GetVarArgs( WIDECHAR* Dest, SIZE_T DestSize, const WIDECHAR*& Fmt, va_list ArgPtr )
 {
 #if !(UE_BUILD_SHIPPING || UE_BUILD_TEST)
-    static bool bTested = false;
-    if(!bTested)
-    {
-        bTested = true;
-        RunGetVarArgsTests();
-    }
+	static bool bTested = false;
+	if(!bTested)
+	{
+		bTested = true;
+		RunGetVarArgsTests();
+	}
 #endif
 
 	if (Fmt == nullptr)
 	{
 		if ((DestSize > 0) && (Dest != nullptr))
-        {
+		{
 			*Dest = 0;
-        }
+		}
 		return 0;
 	}
-	
-	int FmtLen = FPlatformMath::Min(Count, Strlen(Fmt));
-    const TCHAR *Src = Fmt;
-	
+
+	const TCHAR *Src = Fmt;
+
 	TCHAR *Dst = Dest;
 	TCHAR *EndDst = Dst + (DestSize - 1);
-	
+
 	while ((*Src) && (Dst < EndDst))
 	{
 		if (*Src != '%')
@@ -242,58 +241,58 @@ int32 FGenericWidePlatformString::GetVarArgs( WIDECHAR* Dest, SIZE_T DestSize, i
 			Src++;
 			continue;
 		}
-		
+
 		const TCHAR *Percent = Src;
 		int FieldLen = 0;
 		int PrecisionLen = -1;
-		
+
 		Src++; // skip the '%' char...
-		
+
 		while (*Src == ' ')
 		{
 			*Dst = ' ';
 			Dst++;
 			Src++;
 		}
-		
+
 		// Skip modifier flags that don't need additional processing;
 		// they still get passed to snprintf() below based on the conversion.
 		if (*Src == '+')
 		{
 			Src++;
 		}
-		
+
 		// check for field width requests...
 		if ((*Src == '-') || ((*Src >= '0') && (*Src <= '9')))
 		{
 			const TCHAR *Cur = Src + 1;
 			while ((*Cur >= '0') && (*Cur <= '9'))
-            {
+			{
 				Cur++;
-            }
-			
+			}
+
 			FieldLen = Atoi(Src);			
 			Src = Cur;
 		}
-		
+
 		if (*Src == '.')
 		{
 			const TCHAR *Cur = Src + 1;
 			while ((*Cur >= '0') && (*Cur <= '9'))
-            {
+			{
 				Cur++;
-            }
-			
+			}
+
 			PrecisionLen = Atoi(Src + 1);
 			Src = Cur;
 		}
-		
+
 		// Check for 'ls' field, change to 's'
 		if ((Src[0] == 'l' && Src[1] == 's'))
 		{
 			Src++;
 		}
-		
+
 		switch (*Src)
 		{
 			case '%':
@@ -301,18 +300,18 @@ int32 FGenericWidePlatformString::GetVarArgs( WIDECHAR* Dest, SIZE_T DestSize, i
 				Src++;
 				*Dst = '%';
 				Dst++;
-                break;
+				break;
 			}
-				
+
 			case 'c':
 			{
 				TCHAR Val = (TCHAR) va_arg(ArgPtr, int);
 				Src++;
 				*Dst = Val;
 				Dst++;
-                break;
+				break;
 			}
-				
+
 			case 'd':
 			case 'i':
 			case 'X':
@@ -323,7 +322,7 @@ int32 FGenericWidePlatformString::GetVarArgs( WIDECHAR* Dest, SIZE_T DestSize, i
 				int Val = va_arg(ArgPtr, int);
 				ANSICHAR AnsiNum[30];
 				ANSICHAR FmtBuf[30];
-				
+
 				// Yes, this is lame.
 				int CpyIdx = 0;
 				while (Percent < Src)
@@ -333,27 +332,27 @@ int32 FGenericWidePlatformString::GetVarArgs( WIDECHAR* Dest, SIZE_T DestSize, i
 					CpyIdx++;
 				}
 				FmtBuf[CpyIdx] = 0;
-				
+
 				int RetCnt = snprintf(AnsiNum, sizeof(AnsiNum), FmtBuf, Val);
 				if ((Dst + RetCnt) > EndDst)
-                {
+				{
 					return -1;	// Fail - the app needs to create a larger buffer and try again
-                }
+				}
 				for (int i = 0; i < RetCnt; i++)
 				{
 					*Dst = (TCHAR)AnsiNum[i];
 					Dst++;
 				}
-                break;
+				break;
 			}
-				
+
 			case 'p':
 			{
 				Src++;
 				void* Val = va_arg(ArgPtr, void*);
 				ANSICHAR AnsiNum[30];
 				ANSICHAR FmtBuf[30];
-				
+
 				// Yes, this is lame.
 				int CpyIdx = 0;
 				while (Percent < Src)
@@ -363,20 +362,20 @@ int32 FGenericWidePlatformString::GetVarArgs( WIDECHAR* Dest, SIZE_T DestSize, i
 					CpyIdx++;
 				}
 				FmtBuf[CpyIdx] = 0;
-				
+
 				int RetCnt = snprintf(AnsiNum, sizeof (AnsiNum), FmtBuf, Val);
 				if ((Dst + RetCnt) > EndDst)
-                {
+				{
 					return -1;	// Fail - the app needs to create a larger buffer and try again
-                }
+				}
 				for (int i = 0; i < RetCnt; i++)
 				{
 					*Dst = (TCHAR) AnsiNum[i];
 					Dst++;
 				}
-                break;
+				break;
 			}
-				
+
 			case 'l':
 			case 'I':
 			{
@@ -387,7 +386,7 @@ int32 FGenericWidePlatformString::GetVarArgs( WIDECHAR* Dest, SIZE_T DestSize, i
 					int Val = va_arg(ArgPtr, int);
 					ANSICHAR AnsiNum[30];
 					ANSICHAR FmtBuf[30];
-					
+
 					// Yes, this is lame.
 					int CpyIdx = 0;
 					while (Percent < Src)
@@ -397,12 +396,12 @@ int32 FGenericWidePlatformString::GetVarArgs( WIDECHAR* Dest, SIZE_T DestSize, i
 						CpyIdx++;
 					}
 					FmtBuf[CpyIdx] = 0;
-					
+
 					int RetCnt = snprintf(AnsiNum, sizeof(AnsiNum), FmtBuf, Val);
 					if ((Dst + RetCnt) > EndDst)
-                    {
+					{
 						return -1;	// Fail - the app needs to create a larger buffer and try again
-                    }
+					}
 					for (int i = 0; i < RetCnt; i++)
 					{
 						*Dst = (TCHAR)AnsiNum[i];
@@ -410,8 +409,8 @@ int32 FGenericWidePlatformString::GetVarArgs( WIDECHAR* Dest, SIZE_T DestSize, i
 					}
 					break;
 				}
-				
-				
+
+
 				if ((Src[0] == 'l' && Src[1] != 'l') ||
 					(Src[0] == 'I' && (Src[1] != '6' || Src[2] != '4')))
 				{
@@ -419,7 +418,7 @@ int32 FGenericWidePlatformString::GetVarArgs( WIDECHAR* Dest, SIZE_T DestSize, i
 					Src++;  // skip it, I guess.
 					break;
 				}
-				
+
 				// Yes, this is lame.
 				int CpyIdx = 0;
 				unsigned long long Val = va_arg(ArgPtr, unsigned long long);
@@ -436,7 +435,7 @@ int32 FGenericWidePlatformString::GetVarArgs( WIDECHAR* Dest, SIZE_T DestSize, i
 					Percent += 4;
 					CpyIdx = 2;
 				}
-				
+
 				while (Percent < Src)
 				{
 					FmtBuf[CpyIdx] = (ANSICHAR) *Percent;
@@ -444,20 +443,20 @@ int32 FGenericWidePlatformString::GetVarArgs( WIDECHAR* Dest, SIZE_T DestSize, i
 					CpyIdx++;
 				}
 				FmtBuf[CpyIdx] = 0;
-				
+
 				int RetCnt = snprintf(AnsiNum, sizeof (AnsiNum), FmtBuf, Val);
 				if ((Dst + RetCnt) > EndDst)
-                {
+				{
 					return -1;	// Fail - the app needs to create a larger buffer and try again
-                }
+				}
 				for (int i = 0; i < RetCnt; i++)
 				{
 					*Dst = (TCHAR) AnsiNum[i];
 					Dst++;
 				}
-                break;
+				break;
 			}
-				
+
 			case 'f':
 			case 'e':
 			case 'g':
@@ -466,7 +465,7 @@ int32 FGenericWidePlatformString::GetVarArgs( WIDECHAR* Dest, SIZE_T DestSize, i
 				double Val = va_arg(ArgPtr, double);
 				ANSICHAR AnsiNum[30];
 				ANSICHAR FmtBuf[30];
-				
+
 				// Yes, this is lame.
 				int CpyIdx = 0;
 				while (Percent < Src)
@@ -476,36 +475,36 @@ int32 FGenericWidePlatformString::GetVarArgs( WIDECHAR* Dest, SIZE_T DestSize, i
 					CpyIdx++;
 				}
 				FmtBuf[CpyIdx] = 0;
-				
+
 				int RetCnt = snprintf(AnsiNum, sizeof (AnsiNum), FmtBuf, Val);
 				if ((Dst + RetCnt) > EndDst)
-                {
+				{
 					return -1;	// Fail - the app needs to create a larger buffer and try again
-                }
+				}
 				for (int i = 0; i < RetCnt; i++)
 				{
 					*Dst = (TCHAR) AnsiNum[i];
 					Dst++;
 				}
-                break;
+				break;
 			}
-				
+
 			case 's':
 			{
 				Src++;
 				static const TCHAR* Null = TEXT("(null)");
 				const TCHAR *Val = va_arg(ArgPtr, TCHAR *);
 				if (Val == nullptr)
-                {
+				{
 					Val = Null;
-                }
-				
+				}
+
 				int RetCnt = Strlen(Val);
 				int Spaces = FPlatformMath::Max(FPlatformMath::Abs(FieldLen) - RetCnt, 0);
 				if ((Dst + RetCnt + Spaces) > EndDst)
-                {
+				{
 					return -1;	// Fail - the app needs to create a larger buffer and try again
-                }
+				}
 				if (Spaces > 0 && FieldLen > 0)
 				{
 					for (int i = 0; i < Spaces; i++)
@@ -528,23 +527,23 @@ int32 FGenericWidePlatformString::GetVarArgs( WIDECHAR* Dest, SIZE_T DestSize, i
 						Dst++;
 					}
 				}
-                break;
+				break;
 			}
-				
+
 			default:
 				printf("Unknown percent [%%%c] in FGenericWidePlatformString::GetVarArgs().\n", *Src);
 				Src++;  // skip char, I guess.
 				break;
 		}
 	}
-	
+
 	// Check if we were able to finish the entire format string
 	// If not, the app needs to create a larger buffer and try again
 	if (*Src)
-    {
+	{
 		return -1;
-    }
-	
+	}
+
 	*Dst = 0;  // null terminate the new string.
 	return Dst - Dest;
 }
diff --git a/Engine/Source/Runtime/Core/Private/GenericPlatform/StandardPlatformString.cpp b/Engine/Source/Runtime/Core/Private/GenericPlatform/StandardPlatformString.cpp
--- a/Engine/Source/Runtime/Core/Private/GenericPlatform/StandardPlatformString.cpp
+++ b/Engine/Source/Runtime/Core/Private/GenericPlatform/StandardPlatformString.cpp
@@ -201,12 +201,13 @@ static const WIDECHAR* GetFormattedArgument(const FFormatInfo& Info, VA_LIST_REF
 	return Formatted;
 }
 
-int32 FStandardPlatformString::GetVarArgs( WIDECHAR* Dest, SIZE_T DestSize, int32 Count, const WIDECHAR*& Fmt, va_list ArgPtr )
+int32 FStandardPlatformString::GetVarArgs( WIDECHAR* Dest, SIZE_T DestSize, const WIDECHAR*& Fmt, va_list ArgPtr )
 {
-  	const WIDECHAR* Format = Fmt;
+	const WIDECHAR* Format = Fmt;
 	const WIDECHAR* DestStart = Dest;
 
-	while (Count > 1 &&  *Format)
+	--DestSize;
+	while (DestSize > 1 && *Format)
 	{
 		if (*Format == LITERAL(WIDECHAR, '%'))
 		{
@@ -214,7 +215,7 @@ int32 FStandardPlatformString::GetVarArgs( WIDECHAR* Dest, SIZE_T DestSize, int3
 			{
 				*Dest++ = *Format;
 				Format += 2;
-				Count--;
+				DestSize--;
 				continue;
 			}
 
@@ -226,24 +227,24 @@ int32 FStandardPlatformString::GetVarArgs( WIDECHAR* Dest, SIZE_T DestSize, int3
 			const WIDECHAR* FormattedArg = GetFormattedArgument(Info, ArgPtr, Formatted, Length);
 			if (FormattedArg && Length > 0)
 			{
-				if (Length < Count)
+				if (Length < DestSize)
 				{
 					FMemory::Memcpy(Dest, FormattedArg, Length * sizeof(WIDECHAR));
 					Dest += Length;
 				}
-				Count -= Length;
+				DestSize -= Length;
 			}
 		}
 		else
 		{
 			*Dest++ = *Format++;
-			Count--;
+			DestSize--;
 		}
 	}
 
 	*Dest = 0;
 
-	return Count > 1 ? Dest - DestStart : -1;
+	return DestSize > 1 ? Dest - DestStart : -1;
 }
 
 #endif // !PLATFORM_USE_SYSTEM_VSWPRINTF
diff --git a/Engine/Source/Runtime/Core/Private/Windows/WindowsPlatformAtomics.cpp b/Engine/Source/Runtime/Core/Private/Windows/WindowsPlatformAtomics.cpp
--- a/Engine/Source/Runtime/Core/Private/Windows/WindowsPlatformAtomics.cpp
+++ b/Engine/Source/Runtime/Core/Private/Windows/WindowsPlatformAtomics.cpp
@@ -13,7 +13,7 @@ void FWindowsPlatformAtomics::HandleAtomicsFailure( const TCHAR* InFormat, ... )
 	va_list Ptr;
 
 	va_start( Ptr, InFormat );	
-	FCString::GetVarArgs( TempStr, ARRAY_COUNT(TempStr), ARRAY_COUNT(TempStr) - 1, InFormat, Ptr );
+	FCString::GetVarArgs( TempStr, ARRAY_COUNT(TempStr), InFormat, Ptr );
 	va_end( Ptr );
 
 	UE_LOG(LogWindows, Log,  TempStr );
diff --git a/Engine/Source/Runtime/Core/Public/Android/AndroidString.h b/Engine/Source/Runtime/Core/Public/Android/AndroidString.h
--- a/Engine/Source/Runtime/Core/Public/Android/AndroidString.h
+++ b/Engine/Source/Runtime/Core/Public/Android/AndroidString.h
@@ -431,7 +431,13 @@ struct FAndroidPlatformString : public FGenericPlatformString
 		return StrToken + (Pos - AnsiStrToken);
 	}
 
+	DEPRECATED(4.22, "GetVarArgs with DestSize and Count arguments has been deprecated - only DestSize should be passed")
 	static FORCEINLINE int32 GetVarArgs( WIDECHAR* Dest, SIZE_T DestSize, int32 Count, const WIDECHAR*& Fmt, va_list ArgPtr )
+	{
+		return GetVarArgs(Dest, DestSize, Fmt, ArgPtr);
+	}
+
+	static FORCEINLINE int32 GetVarArgs( WIDECHAR* Dest, SIZE_T DestSize, const WIDECHAR*& Fmt, va_list ArgPtr )
 	{
 #if PLATFORM_USE_LS_SPEC_FOR_WIDECHAR
 		// fix up the Fmt string, as fast as possible, without using an FString
@@ -481,7 +487,7 @@ struct FAndroidPlatformString : public FGenericPlatformString
 		}
 		NewFormat[NewIndex] = 0;
 #endif // PLATFORM_USE_LS_SPEC_FOR_WIDECHAR
-		int32 Result = vswprintf( Dest, Count, NewFormat, ArgPtr);
+		int32 Result = vswprintf( Dest, DestSize, NewFormat, ArgPtr);
 		va_end( ArgPtr );
 		return Result;
 	}
@@ -576,9 +582,15 @@ struct FAndroidPlatformString : public FGenericPlatformString
 		return strtok(StrToken, Delim);
 	}
 
-	static FORCEINLINE int32 GetVarArgs( ANSICHAR* Dest, SIZE_T DestSize, int32 Count, const ANSICHAR*& Fmt, va_list ArgPtr )
+	DEPRECATED(4.22, "GetVarArgs with DestSize and Count arguments has been deprecated - only DestSize should be passed")
+	static FORCEINLINE int32 GetVarArgs(ANSICHAR* Dest, SIZE_T DestSize, int32 Count, const ANSICHAR*& Fmt, va_list ArgPtr)
+	{
+		return GetVarArgs(Dest, DestSize, Count, Fmt, ArgPtr);
+	}
+
+	static FORCEINLINE int32 GetVarArgs( ANSICHAR* Dest, SIZE_T DestSize, const ANSICHAR*& Fmt, va_list ArgPtr )
 	{
-		int32 Result = vsnprintf(Dest,Count,Fmt,ArgPtr);
+		int32 Result = vsnprintf(Dest, DestSize, Fmt, ArgPtr);
 		va_end( ArgPtr );
 		return Result;
 	}
diff --git a/Engine/Source/Runtime/Core/Public/GenericPlatform/GenericWidePlatformString.h b/Engine/Source/Runtime/Core/Public/GenericPlatform/GenericWidePlatformString.h
--- a/Engine/Source/Runtime/Core/Public/GenericPlatform/GenericWidePlatformString.h
+++ b/Engine/Source/Runtime/Core/Public/GenericPlatform/GenericWidePlatformString.h
@@ -184,7 +184,14 @@ public:
 	
 	
 	static WIDECHAR* Strtok( WIDECHAR* StrToken, const WIDECHAR* Delim, WIDECHAR** Context );
-	static int32 GetVarArgs( WIDECHAR* Dest, SIZE_T DestSize, int32 Count, const WIDECHAR*& Fmt, va_list ArgPtr );
+
+	DEPRECATED(4.22, "GetVarArgs with DestSize and Count arguments has been deprecated - only DestSize should be passed")
+	static int32 GetVarArgs( WIDECHAR* Dest, SIZE_T DestSize, int32 Count, const WIDECHAR*& Fmt, va_list ArgPtr )
+	{
+		return GetVarArgs(Dest, DestSize, Fmt, ArgPtr);
+	}
+
+	static int32 GetVarArgs( WIDECHAR* Dest, SIZE_T DestSize, const WIDECHAR*& Fmt, va_list ArgPtr );
 
 	/**
 	 * Ansi implementation
@@ -276,9 +283,15 @@ public:
 		return strtok(StrToken, Delim);
 	}
 
+	DEPRECATED(4.22, "GetVarArgs with DestSize and Count arguments has been deprecated - only DestSize should be passed")
+	static int32 GetVarArgs( ANSICHAR* Dest, SIZE_T DestSize, const ANSICHAR*& Fmt, va_list ArgPtr )
+	{
+		return GetVarArgs(Dest, DestSize, Fmt, ArgPtr);
+	}
+
 	static int32 GetVarArgs( ANSICHAR* Dest, SIZE_T DestSize, int32 Count, const ANSICHAR*& Fmt, va_list ArgPtr )
 	{
-		int32 Result = vsnprintf(Dest,Count,Fmt,ArgPtr);
+		int32 Result = vsnprintf(Dest, DestSize, Fmt, ArgPtr);
 		va_end( ArgPtr );
 		return Result;
 	}
diff --git a/Engine/Source/Runtime/Core/Public/GenericPlatform/MicrosoftPlatformString.h b/Engine/Source/Runtime/Core/Public/GenericPlatform/MicrosoftPlatformString.h
--- a/Engine/Source/Runtime/Core/Public/GenericPlatform/MicrosoftPlatformString.h
+++ b/Engine/Source/Runtime/Core/Public/GenericPlatform/MicrosoftPlatformString.h
@@ -116,9 +116,15 @@ struct FMicrosoftPlatformString : public FGenericPlatformString
 		return _tcstok_s(StrToken, Delim, Context);
 	}
 
+	DEPRECATED(4.22, "GetVarArgs with DestSize and Count arguments has been deprecated - only DestSize should be passed")
 	static FORCEINLINE int32 GetVarArgs( WIDECHAR* Dest, SIZE_T DestSize, int32 Count, const WIDECHAR*& Fmt, va_list ArgPtr )
 	{
-		int32 Result = vswprintf(Dest, Count, Fmt, ArgPtr);
+		return GetVarArgs(Dest, DestSize, Fmt, ArgPtr);
+	}
+
+	static FORCEINLINE int32 GetVarArgs( WIDECHAR* Dest, SIZE_T DestSize, const WIDECHAR*& Fmt, va_list ArgPtr )
+	{
+		int32 Result = vswprintf(Dest, DestSize, Fmt, ArgPtr);
 		va_end( ArgPtr );
 		return Result;
 	}
@@ -255,9 +261,15 @@ struct FMicrosoftPlatformString : public FGenericPlatformString
 		return strtok_s(StrToken, Delim, Context);
 	}
 
+	DEPRECATED(4.22, "GetVarArgs with DestSize and Count arguments has been deprecated - only DestSize should be passed")
 	static FORCEINLINE int32 GetVarArgs( ANSICHAR* Dest, SIZE_T DestSize, int32 Count, const ANSICHAR*& Fmt, va_list ArgPtr )
 	{
-		int32 Result = vsnprintf( Dest, Count, Fmt, ArgPtr );
+		return GetVarArgs(Dest, DestSize, Fmt, ArgPtr);
+	}
+
+	static FORCEINLINE int32 GetVarArgs( ANSICHAR* Dest, SIZE_T DestSize, const ANSICHAR*& Fmt, va_list ArgPtr )
+	{
+		int32 Result = vsnprintf( Dest, DestSize, Fmt, ArgPtr );
 		va_end( ArgPtr );
 		return Result;
 	}
diff --git a/Engine/Source/Runtime/Core/Public/GenericPlatform/StandardPlatformString.h b/Engine/Source/Runtime/Core/Public/GenericPlatform/StandardPlatformString.h
--- a/Engine/Source/Runtime/Core/Public/GenericPlatform/StandardPlatformString.h
+++ b/Engine/Source/Runtime/Core/Public/GenericPlatform/StandardPlatformString.h
@@ -150,8 +150,14 @@ public:
 		return wcstok(StrToken, Delim, Context);
 	}
 
+	DEPRECATED(4.22, "GetVarArgs with DestSize and Count arguments has been deprecated - only DestSize should be passed")
+	static int32 GetVarArgs(WIDECHAR* Dest, SIZE_T DestSize, int32 Count, const WIDECHAR*& Fmt, va_list ArgPtr)
+	{
+		return GetVarArgs(Dest, DestSize, Fmt, ArgPtr);
+	}
+
 #if PLATFORM_USE_SYSTEM_VSWPRINTF
-	static int32 GetVarArgs( WIDECHAR* Dest, SIZE_T DestSize, int32 Count, const WIDECHAR*& Fmt, va_list ArgPtr )
+	static int32 GetVarArgs( WIDECHAR* Dest, SIZE_T DestSize, const WIDECHAR*& Fmt, va_list ArgPtr )
 	{
 #if PLATFORM_USE_LS_SPEC_FOR_WIDECHAR
 		// fix up the Fmt string, as fast as possible, without using an FString
@@ -200,15 +206,15 @@ public:
 			}
 		}
 		NewFormat[NewIndex] = 0;
-		int32 Result = vswprintf( Dest, Count, NewFormat, ArgPtr);
+		int32 Result = vswprintf( Dest, DestSize, NewFormat, ArgPtr);
 #else
-		int32 Result = vswprintf( Dest, Count, Fmt, ArgPtr);
+		int32 Result = vswprintf( Dest, DestSize, Fmt, ArgPtr);
 #endif
 		va_end( ArgPtr );
 		return Result;
 	}
 #else // PLATFORM_USE_SYSTEM_VSWPRINTF
-	static int32 GetVarArgs( WIDECHAR* Dest, SIZE_T DestSize, int32 Count, const WIDECHAR*& Fmt, va_list ArgPtr );
+	static int32 GetVarArgs( WIDECHAR* Dest, SIZE_T DestSize, const WIDECHAR*& Fmt, va_list ArgPtr );
 #endif // PLATFORM_USE_SYSTEM_VSWPRINTF
 
 	/**
@@ -301,9 +307,15 @@ public:
 		return strtok(StrToken, Delim);
 	}
 
+	DEPRECATED(4.22, "GetVarArgs with DestSize and Count arguments has been deprecated - only DestSize should be passed")
 	static int32 GetVarArgs( ANSICHAR* Dest, SIZE_T DestSize, int32 Count, const ANSICHAR*& Fmt, va_list ArgPtr )
 	{
-		int32 Result = vsnprintf(Dest,Count,Fmt,ArgPtr);
+		return GetVarArgs(Dest, DestSize, Fmt, ArgPtr);
+	}
+
+	static int32 GetVarArgs( ANSICHAR* Dest, SIZE_T DestSize, const ANSICHAR*& Fmt, va_list ArgPtr )
+	{
+		int32 Result = vsnprintf(Dest, DestSize, Fmt, ArgPtr);
 		va_end( ArgPtr );
 		return Result;
 	}
diff --git a/Engine/Source/Runtime/Core/Public/Misc/CString.h b/Engine/Source/Runtime/Core/Public/Misc/CString.h
--- a/Engine/Source/Runtime/Core/Public/Misc/CString.h
+++ b/Engine/Source/Runtime/Core/Public/Misc/CString.h
@@ -360,16 +360,31 @@ public:
 	}
 
 	/**
-	* Helper function to write formatted output using an argument list
-	*
-	* @param Dest - destination string buffer
-	* @param DestSize - size of destination buffer
-	* @param Count - number of characters to write (not including null terminating character)
-	* @param Fmt - string to print
-	* @param Args - argument list
-	* @return number of characters written or -1 if truncated
-	*/
-	static FORCEINLINE int32 GetVarArgs( CharType* Dest, SIZE_T DestSize, int32 Count, const CharType*& Fmt, va_list ArgPtr );
+	 * Helper function to write formatted output using an argument list
+	 *
+	 * @param Dest - destination string buffer
+	 * @param DestSize - size of destination buffer
+	 * @param Count - number of characters to write (not including null terminating character)
+	 * @param Fmt - string to print
+	 * @param Args - argument list
+	 * @return number of characters written or -1 if truncated
+	 */
+	DEPRECATED(4.22, "GetVarArgs with DestSize and Count arguments has been deprecated - only DestSize should be passed")
+	static FORCEINLINE int32 GetVarArgs(CharType* Dest, SIZE_T DestSize, int32 Count, const CharType*& Fmt, va_list ArgPtr)
+	{
+		return GetVarArgs(Dest, DestSize, Fmt, ArgPtr);
+	}
+
+	/**
+	 * Helper function to write formatted output using an argument list
+	 *
+	 * @param Dest - destination string buffer
+	 * @param DestSize - size of destination buffer
+	 * @param Fmt - string to print
+	 * @param Args - argument list
+	 * @return number of characters written or -1 if truncated
+	 */
+	static FORCEINLINE int32 GetVarArgs( CharType* Dest, SIZE_T DestSize, const CharType*& Fmt, va_list ArgPtr );
 };
 
 typedef TCString<TCHAR>    FCString;
@@ -815,9 +830,9 @@ typename TCString<T>::CharType* TCString<T>::Strtok( CharType* TokenString, cons
 }
 
 template <typename T> FORCEINLINE
-int32 TCString<T>::GetVarArgs( CharType* Dest, SIZE_T DestSize, int32 Count, const CharType*& Fmt, va_list ArgPtr )
+int32 TCString<T>::GetVarArgs( CharType* Dest, SIZE_T DestSize, const CharType*& Fmt, va_list ArgPtr )
 {
-	return FPlatformString::GetVarArgs(Dest, DestSize, Count, Fmt, ArgPtr);
+	return FPlatformString::GetVarArgs(Dest, DestSize, Fmt, ArgPtr);
 }
 
 
diff --git a/Engine/Source/Runtime/Core/Public/Misc/VarArgs.h b/Engine/Source/Runtime/Core/Public/Misc/VarArgs.h
--- a/Engine/Source/Runtime/Core/Public/Misc/VarArgs.h
+++ b/Engine/Source/Runtime/Core/Public/Misc/VarArgs.h
@@ -19,25 +19,25 @@ VarArgs helper macros.
 	{ \
 		va_list ap; \
 		va_start(ap, lastarg); \
-		FCString::GetVarArgs(msg, msgsize, len, fmt, ap); \
+		FCString::GetVarArgs(msg, msgsize, fmt, ap); \
 	}
 #define GET_VARARGS_WIDE(msg, msgsize, len, lastarg, fmt) \
 	{ \
 		va_list ap; \
 		va_start(ap, lastarg); \
-		FCStringWide::GetVarArgs(msg, msgsize, len, fmt, ap); \
+		FCStringWide::GetVarArgs(msg, msgsize, fmt, ap); \
 	}
 #define GET_VARARGS_ANSI(msg, msgsize, len, lastarg, fmt) \
 	{ \
 		va_list ap; \
 		va_start(ap, lastarg); \
-		FCStringAnsi::GetVarArgs(msg, msgsize, len, fmt, ap); \
+		FCStringAnsi::GetVarArgs(msg, msgsize, fmt, ap); \
 	}
 #define GET_VARARGS_RESULT(msg, msgsize, len, lastarg, fmt, result) \
 	{ \
 		va_list ap; \
 		va_start(ap, lastarg); \
-		result = FCString::GetVarArgs(msg, msgsize, len, fmt, ap); \
+		result = FCString::GetVarArgs(msg, msgsize, fmt, ap); \
 		if (result >= msgsize) \
 		{ \
 			result = -1; \
@@ -47,7 +47,7 @@ VarArgs helper macros.
 	{ \
 		va_list ap; \
 		va_start(ap, lastarg); \
-		result = FCStringWide::GetVarArgs(msg, msgsize, len, fmt, ap); \
+		result = FCStringWide::GetVarArgs(msg, msgsize, fmt, ap); \
 		if (result >= msgsize) \
 		{ \
 			result = -1; \
@@ -57,7 +57,7 @@ VarArgs helper macros.
 	{ \
 		va_list ap; \
 		va_start(ap, lastarg); \
-		result = FCStringAnsi::GetVarArgs(msg, msgsize, len, fmt, ap); \
+		result = FCStringAnsi::GetVarArgs(msg, msgsize, fmt, ap); \
 		if (result >= msgsize) \
 		{ \
 			result = -1; \
diff --git a/Engine/Source/Runtime/Engine/Private/SceneUtils.cpp b/Engine/Source/Runtime/Engine/Private/SceneUtils.cpp
--- a/Engine/Source/Runtime/Engine/Private/SceneUtils.cpp
+++ b/Engine/Source/Runtime/Engine/Private/SceneUtils.cpp
@@ -60,7 +60,7 @@ void TDrawEvent<TRHICmdList>::Start(TRHICmdList& InRHICmdList, FColor Color, con
 		va_start(ptr, Fmt);
 		TCHAR TempStr[256];
 		// Build the string in the temp buffer
-		FCString::GetVarArgs(TempStr, ARRAY_COUNT(TempStr), ARRAY_COUNT(TempStr) - 1, Fmt, ptr);
+		FCString::GetVarArgs(TempStr, ARRAY_COUNT(TempStr), Fmt, ptr);
 		InRHICmdList.PushEvent(TempStr, Color);
 		RHICmdList = &InRHICmdList;
 	}
@@ -86,7 +86,7 @@ void FDrawEventRHIExecute::Start(IRHIComputeContext& InRHICommandContext, FColor
 		va_start(ptr, Fmt);
 		TCHAR TempStr[256];
 		// Build the string in the temp buffer
-		FCString::GetVarArgs(TempStr, ARRAY_COUNT(TempStr), ARRAY_COUNT(TempStr) - 1, Fmt, ptr);
+		FCString::GetVarArgs(TempStr, ARRAY_COUNT(TempStr), Fmt, ptr);
 		RHICommandContext = &InRHICommandContext;
 		RHICommandContext->RHIPushEvent(TempStr, Color);
 	}
diff --git a/Engine/Source/Runtime/Navmesh/Private/DebugUtils/RecastDump.cpp b/Engine/Source/Runtime/Navmesh/Private/DebugUtils/RecastDump.cpp
--- a/Engine/Source/Runtime/Navmesh/Private/DebugUtils/RecastDump.cpp
+++ b/Engine/Source/Runtime/Navmesh/Private/DebugUtils/RecastDump.cpp
@@ -35,7 +35,7 @@ static void ioprintf(duFileIO* io, const char* format, ...)
 	char line[256];
 	va_list ap;
 	va_start(ap, format);
-	const int n = FCStringAnsi::GetVarArgs(line, ARRAY_COUNT(line), ARRAY_COUNT(line) - 1, format, ap);
+	const int n = FCStringAnsi::GetVarArgs(line, ARRAY_COUNT(line), format, ap);
 	va_end(ap);
 	if (n > 0)
 		io->write(line, sizeof(char)*n);
diff --git a/Engine/Source/Runtime/Navmesh/Private/DetourTileCache/DetourTileCacheBuilder.cpp b/Engine/Source/Runtime/Navmesh/Private/DetourTileCache/DetourTileCacheBuilder.cpp
--- a/Engine/Source/Runtime/Navmesh/Private/DetourTileCache/DetourTileCacheBuilder.cpp
+++ b/Engine/Source/Runtime/Navmesh/Private/DetourTileCache/DetourTileCacheBuilder.cpp
@@ -2683,7 +2683,7 @@ void dtTileCacheLogContext::dtLog(const char* format, ...)
 	char msg[MSG_SIZE];
 	va_list ap;
 	va_start(ap, format);
-	int len = FCStringAnsi::GetVarArgs(msg, MSG_SIZE, MSG_SIZE - 1, format, ap);
+	int len = FCStringAnsi::GetVarArgs(msg, MSG_SIZE, format, ap);
 	if (len >= MSG_SIZE)
 	{
 		len = MSG_SIZE - 1;
diff --git a/Engine/Source/Runtime/Navmesh/Private/Recast/Recast.cpp b/Engine/Source/Runtime/Navmesh/Private/Recast/Recast.cpp
--- a/Engine/Source/Runtime/Navmesh/Private/Recast/Recast.cpp
+++ b/Engine/Source/Runtime/Navmesh/Private/Recast/Recast.cpp
@@ -57,7 +57,7 @@ void rcContext::log(const rcLogCategory category, const char* format, ...)
 	char msg[MSG_SIZE];
 	va_list ap;
 	va_start(ap, format);
-	int len = FCStringAnsi::GetVarArgs(msg, MSG_SIZE, MSG_SIZE - 1, format, ap);
+	int len = FCStringAnsi::GetVarArgs(msg, MSG_SIZE, format, ap);
 	if (len >= MSG_SIZE)
 	{
 		len = MSG_SIZE-1;
