CIS fix.

diff --git a/Engine/Source/Runtime/Core/Public/GenericPlatform/GenericWidePlatformString.h b/Engine/Source/Runtime/Core/Public/GenericPlatform/GenericWidePlatformString.h
--- a/Engine/Source/Runtime/Core/Public/GenericPlatform/GenericWidePlatformString.h
+++ b/Engine/Source/Runtime/Core/Public/GenericPlatform/GenericWidePlatformString.h
@@ -284,12 +284,12 @@ public:
 	}
 
 	DEPRECATED(4.22, "GetVarArgs with DestSize and Count arguments has been deprecated - only DestSize should be passed")
-	static int32 GetVarArgs( ANSICHAR* Dest, SIZE_T DestSize, const ANSICHAR*& Fmt, va_list ArgPtr )
+	static int32 GetVarArgs( ANSICHAR* Dest, SIZE_T DestSize, int32 Count, const ANSICHAR*& Fmt, va_list ArgPtr )
 	{
 		return GetVarArgs(Dest, DestSize, Fmt, ArgPtr);
 	}
 
-	static int32 GetVarArgs( ANSICHAR* Dest, SIZE_T DestSize, int32 Count, const ANSICHAR*& Fmt, va_list ArgPtr )
+	static int32 GetVarArgs( ANSICHAR* Dest, SIZE_T DestSize, const ANSICHAR*& Fmt, va_list ArgPtr )
 	{
 		int32 Result = vsnprintf(Dest, DestSize, Fmt, ArgPtr);
 		va_end( ArgPtr );
