Fixes for FGenericWidePlatformString::GetVarArgs() and FStandardWidePlatformString::GetVarArgs() to make them fill as much as the buffer as possible before returning -1.
Fixes for FPlatformString::GetVarArgs() to make them return -1 correctly in ANSI mode.

diff --git a/Engine/Source/Runtime/Core/Private/GenericPlatform/GenericWidePlatformString.cpp b/Engine/Source/Runtime/Core/Private/GenericPlatform/GenericWidePlatformString.cpp
--- a/Engine/Source/Runtime/Core/Private/GenericPlatform/GenericWidePlatformString.cpp
+++ b/Engine/Source/Runtime/Core/Private/GenericPlatform/GenericWidePlatformString.cpp
@@ -177,8 +177,8 @@ int iswspace(wint_t wc)
 static const int OUTPUT_SIZE = 256;
 int32 TestGetVarArgs(WIDECHAR* OutputString, const WIDECHAR* Format, ...)
 {
-    va_list ArgPtr;
-    va_start( ArgPtr, Format );
+	va_list ArgPtr;
+	va_start( ArgPtr, Format );
 
 	return FGenericWidePlatformString::GetVarArgs(OutputString, OUTPUT_SIZE, Format, ArgPtr);
 }
@@ -207,6 +207,94 @@ void RunGetVarArgsTests()
 }
 #endif
 
+namespace
+{
+	// Output iterator which will not overflow the destination buffer but will keep track of how many characters have been written.
+	// It will also null terminate the output when destructed.
+	struct FSafeDestIterator
+	{
+		// Non-copyable
+		FSafeDestIterator(const FSafeDestIterator&) = delete;
+		FSafeDestIterator& operator=(const FSafeDestIterator&) = delete;
+
+		explicit FSafeDestIterator(WIDECHAR* InPtr, SIZE_T InEnd)
+			: NumCharsWritten(0)
+			, Ptr            (InPtr)
+			, EndMinusOne    (InPtr + InEnd - 1)
+		{
+			check(InPtr);
+		}
+
+		FORCENOINLINE ~FSafeDestIterator()
+		{
+			// Auto-terminate buffer
+			*Ptr = 0;
+		}
+
+		// Writes Count instances of Ch to the destination buffer
+		FORCENOINLINE FSafeDestIterator& Write(TCHAR Ch, int32 Count = 1)
+		{
+			// Ensure we're not in the error state
+			if (EndMinusOne)
+			{
+				NumCharsWritten += Count;
+
+				int32 NumToWrite = FPlatformMath::Min(Count, (int32)(EndMinusOne - Ptr));
+				for (int32 Val = NumToWrite; Val; --Val)
+				{
+					*Ptr++ = Ch;
+				}
+
+				if (NumToWrite != Count)
+				{
+					EndMinusOne = 0;
+				}
+			}
+
+			return *this;
+		}
+
+		// Writes Count characters from Src to the destination buffer
+		template <typename CharType>
+		FORCENOINLINE FSafeDestIterator& Write(const CharType* Src, int32 Count)
+		{
+			// Ensure we're not in the error state
+			if (EndMinusOne)
+			{
+				NumCharsWritten += Count;
+
+				int32 NumToWrite = FPlatformMath::Min(Count, (int32)(EndMinusOne - Ptr));
+				for (int32 Val = NumToWrite; Val; --Val)
+				{
+					*Ptr++ = *Src++;
+				}
+
+				if (NumToWrite != Count)
+				{
+					EndMinusOne = 0;
+				}
+			}
+
+			return *this;
+		}
+
+		explicit operator bool() const
+		{
+			return !!EndMinusOne;
+		}
+
+		int32 GetNumCharsWritten() const
+		{
+			return NumCharsWritten;
+		}
+
+	private:
+		int32     NumCharsWritten;
+		WIDECHAR* Ptr;
+		WIDECHAR* EndMinusOne; // when null, this means the iterator has already moved past the writable area of the buffer
+	};
+}
+
 int32 FGenericWidePlatformString::GetVarArgs( WIDECHAR* Dest, SIZE_T DestSize, const WIDECHAR*& Fmt, va_list ArgPtr )
 {
 #if !(UE_BUILD_SHIPPING || UE_BUILD_TEST)
@@ -229,16 +317,15 @@ int32 FGenericWidePlatformString::GetVarArgs( WIDECHAR* Dest, SIZE_T DestSize, c
 
 	const TCHAR *Src = Fmt;
 
-	TCHAR *Dst = Dest;
-	TCHAR *EndDst = Dst + (DestSize - 1);
-
-	while ((*Src) && (Dst < EndDst))
+	FSafeDestIterator DestIter(Dest, DestSize);
+	while (*Src)
 	{
 		if (*Src != '%')
 		{
-			*Dst = *Src;
-			Dst++;
-			Src++;
+			if (!DestIter.Write(*Src++))
+			{
+				return -1;
+			}
 			continue;
 		}
 
@@ -250,8 +337,10 @@ int32 FGenericWidePlatformString::GetVarArgs( WIDECHAR* Dest, SIZE_T DestSize, c
 
 		while (*Src == ' ')
 		{
-			*Dst = ' ';
-			Dst++;
+			if (!DestIter.Write(' '))
+			{
+				return -1;
+			}
 			Src++;
 		}
 
@@ -271,7 +360,7 @@ int32 FGenericWidePlatformString::GetVarArgs( WIDECHAR* Dest, SIZE_T DestSize, c
 				Cur++;
 			}
 
-			FieldLen = Atoi(Src);			
+			FieldLen = Atoi(Src);
 			Src = Cur;
 		}
 
@@ -298,8 +387,10 @@ int32 FGenericWidePlatformString::GetVarArgs( WIDECHAR* Dest, SIZE_T DestSize, c
 			case '%':
 			{
 				Src++;
-				*Dst = '%';
-				Dst++;
+				if (!DestIter.Write('%'))
+				{
+					return -1;
+				}
 				break;
 			}
 
@@ -307,8 +398,10 @@ int32 FGenericWidePlatformString::GetVarArgs( WIDECHAR* Dest, SIZE_T DestSize, c
 			{
 				TCHAR Val = (TCHAR) va_arg(ArgPtr, int);
 				Src++;
-				*Dst = Val;
-				Dst++;
+				if (!DestIter.Write(Val))
+				{
+					return -1;
+				}
 				break;
 			}
 
@@ -334,14 +427,9 @@ int32 FGenericWidePlatformString::GetVarArgs( WIDECHAR* Dest, SIZE_T DestSize, c
 				FmtBuf[CpyIdx] = 0;
 
 				int RetCnt = snprintf(AnsiNum, sizeof(AnsiNum), FmtBuf, Val);
-				if ((Dst + RetCnt) > EndDst)
-				{
-					return -1;	// Fail - the app needs to create a larger buffer and try again
-				}
-				for (int i = 0; i < RetCnt; i++)
+				if (!DestIter.Write(AnsiNum, RetCnt))
 				{
-					*Dst = (TCHAR)AnsiNum[i];
-					Dst++;
+					return -1;
 				}
 				break;
 			}
@@ -357,21 +445,16 @@ int32 FGenericWidePlatformString::GetVarArgs( WIDECHAR* Dest, SIZE_T DestSize, c
 				int CpyIdx = 0;
 				while (Percent < Src)
 				{
-					FmtBuf[CpyIdx] = (ANSICHAR) *Percent;
+					FmtBuf[CpyIdx] = (ANSICHAR)*Percent;
 					Percent++;
 					CpyIdx++;
 				}
 				FmtBuf[CpyIdx] = 0;
 
-				int RetCnt = snprintf(AnsiNum, sizeof (AnsiNum), FmtBuf, Val);
-				if ((Dst + RetCnt) > EndDst)
-				{
-					return -1;	// Fail - the app needs to create a larger buffer and try again
-				}
-				for (int i = 0; i < RetCnt; i++)
+				int RetCnt = snprintf(AnsiNum, sizeof(AnsiNum), FmtBuf, Val);
+				if (!DestIter.Write(AnsiNum, RetCnt))
 				{
-					*Dst = (TCHAR) AnsiNum[i];
-					Dst++;
+					return -1;
 				}
 				break;
 			}
@@ -398,14 +481,9 @@ int32 FGenericWidePlatformString::GetVarArgs( WIDECHAR* Dest, SIZE_T DestSize, c
 					FmtBuf[CpyIdx] = 0;
 
 					int RetCnt = snprintf(AnsiNum, sizeof(AnsiNum), FmtBuf, Val);
-					if ((Dst + RetCnt) > EndDst)
+					if (!DestIter.Write(AnsiNum, RetCnt))
 					{
-						return -1;	// Fail - the app needs to create a larger buffer and try again
-					}
-					for (int i = 0; i < RetCnt; i++)
-					{
-						*Dst = (TCHAR)AnsiNum[i];
-						Dst++;
+						return -1;
 					}
 					break;
 				}
@@ -444,15 +522,10 @@ int32 FGenericWidePlatformString::GetVarArgs( WIDECHAR* Dest, SIZE_T DestSize, c
 				}
 				FmtBuf[CpyIdx] = 0;
 
-				int RetCnt = snprintf(AnsiNum, sizeof (AnsiNum), FmtBuf, Val);
-				if ((Dst + RetCnt) > EndDst)
-				{
-					return -1;	// Fail - the app needs to create a larger buffer and try again
-				}
-				for (int i = 0; i < RetCnt; i++)
+				int RetCnt = snprintf(AnsiNum, sizeof(AnsiNum), FmtBuf, Val);
+				if (!DestIter.Write(AnsiNum, RetCnt))
 				{
-					*Dst = (TCHAR) AnsiNum[i];
-					Dst++;
+					return -1;
 				}
 				break;
 			}
@@ -477,14 +550,9 @@ int32 FGenericWidePlatformString::GetVarArgs( WIDECHAR* Dest, SIZE_T DestSize, c
 				FmtBuf[CpyIdx] = 0;
 
 				int RetCnt = snprintf(AnsiNum, sizeof (AnsiNum), FmtBuf, Val);
-				if ((Dst + RetCnt) > EndDst)
+				if (!DestIter.Write(AnsiNum, RetCnt))
 				{
-					return -1;	// Fail - the app needs to create a larger buffer and try again
-				}
-				for (int i = 0; i < RetCnt; i++)
-				{
-					*Dst = (TCHAR) AnsiNum[i];
-					Dst++;
+					return -1;
 				}
 				break;
 			}
@@ -501,31 +569,18 @@ int32 FGenericWidePlatformString::GetVarArgs( WIDECHAR* Dest, SIZE_T DestSize, c
 
 				int RetCnt = Strlen(Val);
 				int Spaces = FPlatformMath::Max(FPlatformMath::Abs(FieldLen) - RetCnt, 0);
-				if ((Dst + RetCnt + Spaces) > EndDst)
-				{
-					return -1;	// Fail - the app needs to create a larger buffer and try again
-				}
 				if (Spaces > 0 && FieldLen > 0)
 				{
-					for (int i = 0; i < Spaces; i++)
-					{
-						*Dst = TEXT(' ');
-						Dst++;
-					}
+					DestIter.Write(TEXT(' '), Spaces);
 				}
-				for (int i = 0; i < RetCnt; i++)
+				DestIter.Write(Val, RetCnt);
+				if (Spaces > 0 && FieldLen < 0)
 				{
-					*Dst = *Val;
-					Dst++;
-					Val++;
+					DestIter.Write(TEXT(' '), Spaces);
 				}
-				if (Spaces > 0 && FieldLen < 0)
+				if (!DestIter)
 				{
-					for (int i = 0; i < Spaces; i++)
-					{
-						*Dst = TEXT(' ');
-						Dst++;
-					}
+					return -1;
 				}
 				break;
 			}
@@ -537,15 +592,8 @@ int32 FGenericWidePlatformString::GetVarArgs( WIDECHAR* Dest, SIZE_T DestSize, c
 		}
 	}
 
-	// Check if we were able to finish the entire format string
-	// If not, the app needs to create a larger buffer and try again
-	if (*Src)
-	{
-		return -1;
-	}
-
-	*Dst = 0;  // null terminate the new string.
-	return Dst - Dest;
+	int32 Result = DestIter.GetNumCharsWritten();
+	return Result;
 }
 
 #endif
diff --git a/Engine/Source/Runtime/Core/Private/GenericPlatform/StandardPlatformString.cpp b/Engine/Source/Runtime/Core/Private/GenericPlatform/StandardPlatformString.cpp
--- a/Engine/Source/Runtime/Core/Private/GenericPlatform/StandardPlatformString.cpp
+++ b/Engine/Source/Runtime/Core/Private/GenericPlatform/StandardPlatformString.cpp
@@ -206,13 +206,23 @@ int32 FStandardPlatformString::GetVarArgs( WIDECHAR* Dest, SIZE_T DestSize, cons
 	const WIDECHAR* Format = Fmt;
 	const WIDECHAR* DestStart = Dest;
 
+	if (DestSize == 0)
+	{
+		return -1;
+	}
+
 	--DestSize;
-	while (DestSize > 1 && *Format)
+	while (*Format)
 	{
 		if (*Format == LITERAL(WIDECHAR, '%'))
 		{
 			if (*(Format + 1) == LITERAL(WIDECHAR, '%'))
 			{
+				if (DestSize == 0)
+				{
+					*Dest = 0;
+					return -1;
+				}
 				*Dest++ = *Format;
 				Format += 2;
 				DestSize--;
@@ -231,12 +241,24 @@ int32 FStandardPlatformString::GetVarArgs( WIDECHAR* Dest, SIZE_T DestSize, cons
 				{
 					FMemory::Memcpy(Dest, FormattedArg, Length * sizeof(WIDECHAR));
 					Dest += Length;
+					DestSize -= Length;
+				}
+				else
+				{
+					FMemory::Memcpy(Dest, FormattedArg, DestSize * sizeof(WIDECHAR));
+					Dest += DestSize;
+					*Dest = 0;
+					return -1;
 				}
-				DestSize -= Length;
 			}
 		}
 		else
 		{
+			if (DestSize == 0)
+			{
+				*Dest = 0;
+				return -1;
+			}
 			*Dest++ = *Format++;
 			DestSize--;
 		}
@@ -244,7 +266,7 @@ int32 FStandardPlatformString::GetVarArgs( WIDECHAR* Dest, SIZE_T DestSize, cons
 
 	*Dest = 0;
 
-	return DestSize > 1 ? Dest - DestStart : -1;
+	return Dest - DestStart;
 }
 
 #endif // !PLATFORM_USE_SYSTEM_VSWPRINTF
diff --git a/Engine/Source/Runtime/Core/Public/Android/AndroidString.h b/Engine/Source/Runtime/Core/Public/Android/AndroidString.h
--- a/Engine/Source/Runtime/Core/Public/Android/AndroidString.h
+++ b/Engine/Source/Runtime/Core/Public/Android/AndroidString.h
@@ -592,7 +592,7 @@ struct FAndroidPlatformString : public FGenericPlatformString
 	{
 		int32 Result = vsnprintf(Dest, DestSize, Fmt, ArgPtr);
 		va_end( ArgPtr );
-		return Result;
+		return (Result != -1 && Result < DestSize) ? Result : -1;
 	}
 
 	/** 
diff --git a/Engine/Source/Runtime/Core/Public/GenericPlatform/GenericWidePlatformString.h b/Engine/Source/Runtime/Core/Public/GenericPlatform/GenericWidePlatformString.h
--- a/Engine/Source/Runtime/Core/Public/GenericPlatform/GenericWidePlatformString.h
+++ b/Engine/Source/Runtime/Core/Public/GenericPlatform/GenericWidePlatformString.h
@@ -191,7 +191,7 @@ public:
 		return GetVarArgs(Dest, DestSize, Fmt, ArgPtr);
 	}
 
-	static int32 GetVarArgs( WIDECHAR* Dest, SIZE_T DestSize, const WIDECHAR*& Fmt, va_list ArgPtr );
+	static CORE_API int32 GetVarArgs( WIDECHAR* Dest, SIZE_T DestSize, const WIDECHAR*& Fmt, va_list ArgPtr );
 
 	/**
 	 * Ansi implementation
@@ -293,7 +293,7 @@ public:
 	{
 		int32 Result = vsnprintf(Dest, DestSize, Fmt, ArgPtr);
 		va_end( ArgPtr );
-		return Result;
+		return (Result != -1 && Result < DestSize) ? Result : -1;
 	}
 
 	/**
diff --git a/Engine/Source/Runtime/Core/Public/GenericPlatform/MicrosoftPlatformString.h b/Engine/Source/Runtime/Core/Public/GenericPlatform/MicrosoftPlatformString.h
--- a/Engine/Source/Runtime/Core/Public/GenericPlatform/MicrosoftPlatformString.h
+++ b/Engine/Source/Runtime/Core/Public/GenericPlatform/MicrosoftPlatformString.h
@@ -271,7 +271,7 @@ struct FMicrosoftPlatformString : public FGenericPlatformString
 	{
 		int32 Result = vsnprintf( Dest, DestSize, Fmt, ArgPtr );
 		va_end( ArgPtr );
-		return Result;
+		return (Result != -1 && Result < DestSize) ? Result : -1;
 	}
 
 	/**
diff --git a/Engine/Source/Runtime/Core/Public/GenericPlatform/StandardPlatformString.h b/Engine/Source/Runtime/Core/Public/GenericPlatform/StandardPlatformString.h
--- a/Engine/Source/Runtime/Core/Public/GenericPlatform/StandardPlatformString.h
+++ b/Engine/Source/Runtime/Core/Public/GenericPlatform/StandardPlatformString.h
@@ -317,7 +317,7 @@ public:
 	{
 		int32 Result = vsnprintf(Dest, DestSize, Fmt, ArgPtr);
 		va_end( ArgPtr );
-		return Result;
+		return (Result != -1 && Result < DestSize) ? Result : -1;
 	}
 
 	/**
