Signed/unsigned mismatch fix.

diff --git a/Engine/Source/Runtime/Core/Public/Android/AndroidString.h b/Engine/Source/Runtime/Core/Public/Android/AndroidString.h
--- a/Engine/Source/Runtime/Core/Public/Android/AndroidString.h
+++ b/Engine/Source/Runtime/Core/Public/Android/AndroidString.h
@@ -592,7 +592,7 @@ struct FAndroidPlatformString : public FGenericPlatformString
 	{
 		int32 Result = vsnprintf(Dest, DestSize, Fmt, ArgPtr);
 		va_end( ArgPtr );
-		return (Result != -1 && Result < DestSize) ? Result : -1;
+		return (Result != -1 && Result < (int32)DestSize) ? Result : -1;
 	}
 
 	/** 
diff --git a/Engine/Source/Runtime/Core/Public/GenericPlatform/GenericWidePlatformString.h b/Engine/Source/Runtime/Core/Public/GenericPlatform/GenericWidePlatformString.h
--- a/Engine/Source/Runtime/Core/Public/GenericPlatform/GenericWidePlatformString.h
+++ b/Engine/Source/Runtime/Core/Public/GenericPlatform/GenericWidePlatformString.h
@@ -293,7 +293,7 @@ public:
 	{
 		int32 Result = vsnprintf(Dest, DestSize, Fmt, ArgPtr);
 		va_end( ArgPtr );
-		return (Result != -1 && Result < DestSize) ? Result : -1;
+		return (Result != -1 && Result < (int32)DestSize) ? Result : -1;
 	}
 
 	/**
diff --git a/Engine/Source/Runtime/Core/Public/GenericPlatform/MicrosoftPlatformString.h b/Engine/Source/Runtime/Core/Public/GenericPlatform/MicrosoftPlatformString.h
--- a/Engine/Source/Runtime/Core/Public/GenericPlatform/MicrosoftPlatformString.h
+++ b/Engine/Source/Runtime/Core/Public/GenericPlatform/MicrosoftPlatformString.h
@@ -271,7 +271,7 @@ struct FMicrosoftPlatformString : public FGenericPlatformString
 	{
 		int32 Result = vsnprintf( Dest, DestSize, Fmt, ArgPtr );
 		va_end( ArgPtr );
-		return (Result != -1 && Result < DestSize) ? Result : -1;
+		return (Result != -1 && Result < (int32)DestSize) ? Result : -1;
 	}
 
 	/**
diff --git a/Engine/Source/Runtime/Core/Public/GenericPlatform/StandardPlatformString.h b/Engine/Source/Runtime/Core/Public/GenericPlatform/StandardPlatformString.h
--- a/Engine/Source/Runtime/Core/Public/GenericPlatform/StandardPlatformString.h
+++ b/Engine/Source/Runtime/Core/Public/GenericPlatform/StandardPlatformString.h
@@ -317,7 +317,7 @@ public:
 	{
 		int32 Result = vsnprintf(Dest, DestSize, Fmt, ArgPtr);
 		va_end( ArgPtr );
-		return (Result != -1 && Result < DestSize) ? Result : -1;
+		return (Result != -1 && Result < (int32)DestSize) ? Result : -1;
 	}
 
 	/**
