The archive created by FOutputDeviceFile will no longer try to log errors as, in case of an error, it would redirect logging back to itself or, in some cases, cause threading deadlocks.

diff --git a/Engine/Source/Runtime/Core/Private/HAL/FileManagerGeneric.cpp b/Engine/Source/Runtime/Core/Private/HAL/FileManagerGeneric.cpp
--- a/Engine/Source/Runtime/Core/Private/HAL/FileManagerGeneric.cpp
+++ b/Engine/Source/Runtime/Core/Private/HAL/FileManagerGeneric.cpp
@@ -122,7 +122,7 @@ FArchive* FFileManagerGeneric::CreateFileWriter( const TCHAR* Filename, uint32 F
 		}
 		return NULL;
 	}
-	return new FArchiveFileWriterGeneric( Handle, Filename, Handle->Tell() );
+	return new FArchiveFileWriterGeneric(Handle, Filename, Handle->Tell(), Flags);
 }
 
 
@@ -761,8 +761,9 @@ void FArchiveFileReaderGeneric::Serialize( void* V, int64 Length )
 	}
 }
 
-FArchiveFileWriterGeneric::FArchiveFileWriterGeneric( IFileHandle* InHandle, const TCHAR* InFilename, int64 InPos )
+FArchiveFileWriterGeneric::FArchiveFileWriterGeneric( IFileHandle* InHandle, const TCHAR* InFilename, int64 InPos, uint32 InFlags )
 	: Filename( InFilename )
+	, Flags( InFlags )
 	, Pos( InPos )
 	, BufferCount( 0 )
 	, Handle( InHandle )
@@ -876,7 +877,7 @@ void FArchiveFileWriterGeneric::Flush()
 void FArchiveFileWriterGeneric::LogWriteError(const TCHAR* Message)
 {
 	// Prevent re-entry if logging causes another log error leading to a stack overflow
-	if (!bLoggingError)
+	if (!bLoggingError && !IsSilent())
 	{
 		bLoggingError = true;
 		TCHAR ErrorBuffer[1024];
diff --git a/Engine/Source/Runtime/Core/Private/Misc/OutputDeviceFile.cpp b/Engine/Source/Runtime/Core/Private/Misc/OutputDeviceFile.cpp
--- a/Engine/Source/Runtime/Core/Private/Misc/OutputDeviceFile.cpp
+++ b/Engine/Source/Runtime/Core/Private/Misc/OutputDeviceFile.cpp
@@ -405,7 +405,8 @@ void FOutputDeviceFile::WriteByteOrderMarkToArchive(EByteOrderMark ByteOrderMark
 
 bool FOutputDeviceFile::CreateWriter(uint32 MaxAttempts)
 {
-	uint32 WriteFlags = FILEWRITE_AllowRead | ((Opened || AppendIfExists) ? FILEWRITE_Append : 0);
+	// Create a silent filewriter so that it doesn't try to log any errors since it would redirect logging back to itself through this output device
+	uint32 WriteFlags = FILEWRITE_Silent | FILEWRITE_AllowRead | ((Opened || AppendIfExists) ? FILEWRITE_Append : 0);
 
 	// Open log file.
 	FArchive* Ar = IFileManager::Get().CreateFileWriter(Filename, WriteFlags);
diff --git a/Engine/Source/Runtime/Core/Public/HAL/FileManager.h b/Engine/Source/Runtime/Core/Public/HAL/FileManager.h
--- a/Engine/Source/Runtime/Core/Public/HAL/FileManager.h
+++ b/Engine/Source/Runtime/Core/Public/HAL/FileManager.h
@@ -18,7 +18,8 @@ enum EFileWrite
 	FILEWRITE_NoReplaceExisting = 0x02,
 	FILEWRITE_EvenIfReadOnly    = 0x04,
 	FILEWRITE_Append			= 0x08,
-	FILEWRITE_AllowRead			= 0x10
+	FILEWRITE_AllowRead			= 0x10,
+	FILEWRITE_Silent			= 0x20
 };
 
 
diff --git a/Engine/Source/Runtime/Core/Public/HAL/FileManagerGeneric.h b/Engine/Source/Runtime/Core/Public/HAL/FileManagerGeneric.h
--- a/Engine/Source/Runtime/Core/Public/HAL/FileManagerGeneric.h
+++ b/Engine/Source/Runtime/Core/Public/HAL/FileManagerGeneric.h
@@ -255,7 +255,7 @@ protected:
 class FArchiveFileWriterGeneric : public FArchive
 {
 public:
-	FArchiveFileWriterGeneric( IFileHandle* InHandle, const TCHAR* InFilename, int64 InPos );
+	FArchiveFileWriterGeneric( IFileHandle* InHandle, const TCHAR* InFilename, int64 InPos, uint32 InFlags );
 	~FArchiveFileWriterGeneric();
 
 	virtual void Seek( int64 InPos ) final;
@@ -301,8 +301,15 @@ protected:
 	 */
 	void LogWriteError(const TCHAR* Message);
 
+	/** Returns true if the archive should suppress logging in case of error */
+	bool IsSilent() const
+	{
+		return !!(Flags & FILEWRITE_Silent);
+	}
+
 	/** Filename for debugging purposes */
 	FString Filename;
+	uint32 Flags;
 	int64 Pos;
 	int64 BufferCount;
 	TUniquePtr<IFileHandle> Handle;
