Fix and improved implementation for TBitArray::FindAndSetLastZeroBit().  Same implementation for TBitArray::FindLast().

diff --git a/Engine/Source/Runtime/Core/Public/Containers/BitArray.h b/Engine/Source/Runtime/Core/Public/Containers/BitArray.h
--- a/Engine/Source/Runtime/Core/Public/Containers/BitArray.h
+++ b/Engine/Source/Runtime/Core/Public/Containers/BitArray.h
@@ -599,51 +599,38 @@ public:
 	*/
 	int32 FindLast(bool bValue) const 
 	{
-		// Iterate over the array until we see a word with a matching bit
-		const uint32 Test = bValue ? 0u : (uint32)-1;
-
-		const uint32* RESTRICT DwordArray = GetData();
 		const int32 LocalNumBits = NumBits;
-		const int32 DwordCount = FMath::DivideAndRoundUp(LocalNumBits, NumBitsPerDWORD);
-		int32 DwordIndex = DwordCount-1;
-		int32 SlackCount = (LocalNumBits % NumBitsPerDWORD);
 
-		if(SlackCount != 0)
-		{
-			uint32 Mask = (~0u) << SlackCount;
-			uint32 TailTest = bValue ? DwordArray[DwordCount - 1] & ~Mask : DwordArray[DwordCount - 1] | Mask;
+		// Get the correct mask for the last word
+		uint32 SlackIndex = ((LocalNumBits - 1) % NumBitsPerDWORD) + 1;
+		uint32 Mask = ~0u >> (NumBitsPerDWORD - SlackIndex);
 
-			// If we're looking for a false, then we flip the bits - then we only need to find the first one bit
-			const uint32 Bits = bValue ? TailTest : ~TailTest;
-			const int32 LocalBitIndex = (NumBitsPerDWORD - FMath::CountLeadingZeros(Bits) - 1);
-			if (LocalBitIndex >= 0)
+		// Iterate over the array until we see a word with a zero bit.
+		uint32 DwordIndex = FMath::DivideAndRoundUp(LocalNumBits, NumBitsPerDWORD);
+		const uint32* RESTRICT DwordArray = GetData();
+		const uint32 Test = bValue ? 0u : ~0u;
+		for (;;)
+		{
+			if (DwordIndex == 0)
 			{
-				const int32 LowestBitIndex = LocalBitIndex + (DwordIndex << NumBitsPerDWORDLogTwo);
-				checkSlow(LowestBitIndex < LocalNumBits);
-				return LowestBitIndex;
+				return INDEX_NONE;
 			}
-			DwordIndex--;
-		}
-
-		while (DwordIndex >= 0 && DwordArray[DwordIndex] == Test)
-		{
 			--DwordIndex;
-		}
-
-		if (DwordIndex >= 0 && DwordIndex < DwordCount)
-		{
-			// If we're looking for a false, then we flip the bits - then we only need to find the first one bit
-			const uint32 Bits = bValue ? (DwordArray[DwordIndex]) : ~(DwordArray[DwordIndex]);
-			ASSUME(Bits != 0);
-			const int32 LocalBitIndex = (NumBitsPerDWORD - FMath::CountLeadingZeros(Bits) - 1);
-			const int32 LowestBitIndex = LocalBitIndex + (DwordIndex << NumBitsPerDWORDLogTwo);
-			if (LowestBitIndex < LocalNumBits)
+			if ((DwordArray[DwordIndex] & Mask) != (Test & Mask))
 			{
-				return LowestBitIndex;
+				break;
 			}
+			Mask = ~0u;
 		}
 
-		return INDEX_NONE;
+		// Flip the bits, then we only need to find the first one bit -- easy.
+		const uint32 Bits = (bValue ? DwordArray[DwordIndex] : ~DwordArray[DwordIndex]) & Mask;
+		ASSUME(Bits != 0);
+
+		uint32 BitIndex = (NumBitsPerDWORD - 1) - FMath::CountLeadingZeros(Bits);
+
+		int32 Result = BitIndex + (DwordIndex << NumBitsPerDWORDLogTwo);
+		return Result;
 	}
 
 	FORCEINLINE bool Contains(bool bValue) const
@@ -685,46 +672,43 @@ public:
 	}
 
 	/**
-	* Finds the last zero bit in the array, sets it to true, and returns the bit index.
-	* If there is none, INDEX_NONE is returned.
-	*/
+	 * Finds the last zero bit in the array, sets it to true, and returns the bit index.
+	 * If there is none, INDEX_NONE is returned.
+	 */
 	int32 FindAndSetLastZeroBit()
 	{
-		// Iterate over the array until we see a word with a zero bit.
-		uint32* RESTRICT DwordArray = GetData();
 		const int32 LocalNumBits = NumBits;
-		const int32 DwordCount = FMath::DivideAndRoundUp(LocalNumBits, NumBitsPerDWORD);
-		int32 DwordIndex = DwordCount - 1;
 
-		int32 SlackIndex = (LocalNumBits % NumBitsPerDWORD);
-		uint32 Mask = (~0u) << SlackIndex;
-		uint32 Slack = DwordArray[DwordCount - 1] & Mask;
-		DwordArray[DwordCount - 1] = DwordArray[DwordCount - 1] | Mask;
+		// Get the correct mask for the last word
+		uint32 SlackIndex = ((LocalNumBits - 1) % NumBitsPerDWORD) + 1;
+		uint32 Mask = ~0u >> (NumBitsPerDWORD - SlackIndex);
 
-		while (DwordIndex >= 0 && DwordArray[DwordIndex] == (uint32)-1)
+		// Iterate over the array until we see a word with a zero bit.
+		uint32 DwordIndex = FMath::DivideAndRoundUp(LocalNumBits, NumBitsPerDWORD);
+		uint32* RESTRICT DwordArray = GetData();
+		for (;;)
 		{
+			if (DwordIndex == 0)
+			{
+				return INDEX_NONE;
+			}
 			--DwordIndex;
-		}
-
-		if (DwordIndex >= 0 && DwordIndex < DwordCount)
-		{
-			// Flip the bits, then we only need to find the first one bit -- easy.
-			const uint32 Bits = ~(DwordArray[DwordIndex]);
-			ASSUME(Bits != 0);
-			const int32 HighestLocalBitIndex = FMath::CountLeadingZeros(Bits);
-			const int32 BitIndex = (NumBitsPerDWORD - HighestLocalBitIndex - 1);
-			const uint32 HighestBit = 1u << BitIndex;
-			const int32 HighestBitIndex = BitIndex + (DwordIndex << NumBitsPerDWORDLogTwo);
-			if (HighestBitIndex < LocalNumBits)
+			if ((DwordArray[DwordIndex] & Mask) != Mask)
 			{
-				DwordArray[DwordIndex] |= HighestBit;
-				DwordArray[DwordCount - 1] = (DwordArray[DwordCount - 1] & ~Mask) | Slack;
-				return HighestBitIndex;
+				break;
 			}
+			Mask = ~0u;
 		}
 
-		DwordArray[DwordCount - 1] = (DwordArray[DwordCount - 1] & ~Mask) | Slack;
-		return INDEX_NONE;
+		// Flip the bits, then we only need to find the first one bit -- easy.
+		const uint32 Bits = ~DwordArray[DwordIndex] & Mask;
+		ASSUME(Bits != 0);
+
+		uint32 BitIndex = (NumBitsPerDWORD - 1) - FMath::CountLeadingZeros(Bits);
+		DwordArray[DwordIndex] |= 1u << BitIndex;
+
+		int32 Result = BitIndex + (DwordIndex << NumBitsPerDWORDLogTwo);
+		return Result;
 	}
 
 	// Accessors.
