Map properties which handle delta serialization within SerializeItem() can now be handled separately from entire-class delta serialization.
Set properties now support the flag.
Comment for ArNoDelta fixed.

diff --git a/Engine/Source/Runtime/Core/Public/Serialization/Archive.h b/Engine/Source/Runtime/Core/Public/Serialization/Archive.h
--- a/Engine/Source/Runtime/Core/Public/Serialization/Archive.h
+++ b/Engine/Source/Runtime/Core/Public/Serialization/Archive.h
@@ -986,6 +986,11 @@ PRAGMA_ENABLE_DEPRECATION_WARNINGS
 		return !ArNoDelta;
 	}
 
+	FORCEINLINE bool DoIntraPropertyDelta() const
+	{
+		return !ArNoIntraPropertyDelta;
+	}
+
 	FORCEINLINE bool IsIgnoringOuterRef() const
 	{
 		return ArIgnoreOuterRef;
@@ -1412,9 +1417,12 @@ public:
 	/** If true, we will not serialize the ObjectArchetype reference in UObject. */
 	uint8 ArIgnoreArchetypeRef : 1;
 
-	/** If true, we will not serialize the ObjectArchetype reference in UObject. */
+	/** If true, do not perform delta serialization of properties. */
 	uint8 ArNoDelta : 1;
 
+	/** If true, do not perform delta serialization within properties (e.g. TMaps and TSets). */
+	uint8 ArNoIntraPropertyDelta : 1;
+
 	/** If true, we will not serialize the Outer reference in UObject. */
 	uint8 ArIgnoreOuterRef : 1;
 
diff --git a/Engine/Source/Runtime/CoreUObject/Private/Serialization/DuplicateDataReader.cpp b/Engine/Source/Runtime/CoreUObject/Private/Serialization/DuplicateDataReader.cpp
--- a/Engine/Source/Runtime/CoreUObject/Private/Serialization/DuplicateDataReader.cpp
+++ b/Engine/Source/Runtime/CoreUObject/Private/Serialization/DuplicateDataReader.cpp
@@ -27,7 +27,7 @@ FDuplicateDataReader::FDuplicateDataReader( class FUObjectAnnotationSparse<FDupl
 {
 	this->SetIsLoading(true);
 	this->SetIsPersistent(true);
-	this->ArNoDelta = true;
+	this->ArNoIntraPropertyDelta = true;
 	ArPortFlags |= PPF_Duplicate | InPortFlags;
 
 #if USE_STABLE_LOCALIZATION_KEYS
diff --git a/Engine/Source/Runtime/CoreUObject/Private/Serialization/DuplicateDataWriter.cpp b/Engine/Source/Runtime/CoreUObject/Private/Serialization/DuplicateDataWriter.cpp
--- a/Engine/Source/Runtime/CoreUObject/Private/Serialization/DuplicateDataWriter.cpp
+++ b/Engine/Source/Runtime/CoreUObject/Private/Serialization/DuplicateDataWriter.cpp
@@ -33,7 +33,7 @@ FDuplicateDataWriter::FDuplicateDataWriter( FUObjectAnnotationSparse<FDuplicated
 {
 	this->SetIsSaving(true);
 	this->SetIsPersistent(true);
-	this->ArNoDelta = true;
+	this->ArNoIntraPropertyDelta = true;
 	ArAllowLazyLoading	= false;
 	ArPortFlags |= PPF_Duplicate | InPortFlags;
 
diff --git a/Engine/Source/Runtime/CoreUObject/Private/UObject/PropertyMap.cpp b/Engine/Source/Runtime/CoreUObject/Private/UObject/PropertyMap.cpp
--- a/Engine/Source/Runtime/CoreUObject/Private/UObject/PropertyMap.cpp
+++ b/Engine/Source/Runtime/CoreUObject/Private/UObject/PropertyMap.cpp
@@ -244,8 +244,8 @@ void UMapProperty::SerializeItem(FStructuredArchive::FSlot Slot, void* Value, co
 	FArchive& UnderlyingArchive = Slot.GetUnderlyingArchive();
 	FStructuredArchive::FRecord Record = Slot.EnterRecord();
 
-	// If we're doing delta serialization, act as if there are no defaults
-	if (!UnderlyingArchive.DoDelta())
+	// If we're doing delta serialization within this property, act as if there are no defaults
+	if (!UnderlyingArchive.DoIntraPropertyDelta())
 	{
 		Defaults = nullptr;
 	}
diff --git a/Engine/Source/Runtime/CoreUObject/Private/UObject/PropertySet.cpp b/Engine/Source/Runtime/CoreUObject/Private/UObject/PropertySet.cpp
--- a/Engine/Source/Runtime/CoreUObject/Private/UObject/PropertySet.cpp
+++ b/Engine/Source/Runtime/CoreUObject/Private/UObject/PropertySet.cpp
@@ -218,6 +218,12 @@ void USetProperty::SerializeItem(FStructuredArchive::FSlot Slot, void* Value, co
 	FArchive& UnderlyingArchive = Slot.GetUnderlyingArchive();
 	FStructuredArchive::FRecord Record = Slot.EnterRecord();
 
+	// If we're doing delta serialization within this property, act as if there are no defaults
+	if (!UnderlyingArchive.DoIntraPropertyDelta())
+	{
+		Defaults = nullptr;
+	}
+
 	// Ar related calls in this function must be mirrored in USetProperty::ConvertFromType
 	checkSlow(ElementProp);
 
