Fix for uninitialized FArchive::ArNoIntraPropertyDelta field.

diff --git a/Engine/Source/Runtime/Core/Private/Serialization/Archive.cpp b/Engine/Source/Runtime/Core/Private/Serialization/Archive.cpp
--- a/Engine/Source/Runtime/Core/Private/Serialization/Archive.cpp
+++ b/Engine/Source/Runtime/Core/Private/Serialization/Archive.cpp
@@ -147,6 +147,7 @@ void FArchive::Reset()
 	ArSerializingDefaults				= false;
 	ArIgnoreArchetypeRef				= false;
 	ArNoDelta							= false;
+	ArNoIntraPropertyDelta				= false;
 	ArIgnoreOuterRef					= false;
 	ArIgnoreClassGeneratedByRef			= false;
 	ArIgnoreClassRef					= false;
@@ -203,6 +204,7 @@ void FArchive::CopyTrivialFArchiveStatusMembers(const FArchive& ArchiveToCopy)
 	ArSerializingDefaults                = ArchiveToCopy.ArSerializingDefaults;
 	ArIgnoreArchetypeRef                 = ArchiveToCopy.ArIgnoreArchetypeRef;
 	ArNoDelta                            = ArchiveToCopy.ArNoDelta;
+	ArNoIntraPropertyDelta               = ArchiveToCopy.ArNoIntraPropertyDelta;
 	ArIgnoreOuterRef                     = ArchiveToCopy.ArIgnoreOuterRef;
 	ArIgnoreClassGeneratedByRef          = ArchiveToCopy.ArIgnoreClassGeneratedByRef;
 	ArIgnoreClassRef                     = ArchiveToCopy.ArIgnoreClassRef;
