Optimized wide stricmp for ASCII strings

Got 3.2x speedup in a crude benchmark. Avoids inlining code on purpose.

diff --git a/Engine/Source/Runtime/Core/Private/GenericPlatform/GenericPlatformStricmp.cpp b/Engine/Source/Runtime/Core/Private/GenericPlatform/GenericPlatformStricmp.cpp
new file mode 100644
--- /dev/null
+++ b/Engine/Source/Runtime/Core/Private/GenericPlatform/GenericPlatformStricmp.cpp
@@ -0,0 +1,46 @@
+// Copyright 1998-2018 Epic Games, Inc. All Rights Reserved.
+
+#include "GenericPlatform/GenericPlatformStricmp.h"
+
+static constexpr uint8 LowerAscii[128] = {
+	0x00, 0x01, 0x02, 0x03, 0x04, 0x05, 0x06, 0x07, 0x08, 0x09, 0x0A, 0x0B, 0x0C, 0x0D, 0x0E, 0x0F,
+	0x10, 0x11, 0x12, 0x13, 0x14, 0x15, 0x16, 0x17, 0x18, 0x19, 0x1A, 0x1B, 0x1C, 0x1D, 0x1E, 0x1F,
+	0x20, 0x21, 0x22, 0x23, 0x24, 0x25, 0x26, 0x27, 0x28, 0x29, 0x2A, 0x2B, 0x2C, 0x2D, 0x2E, 0x2F,
+	0x30, 0x31, 0x32, 0x33, 0x34, 0x35, 0x36, 0x37, 0x38, 0x39, 0x3A, 0x3B, 0x3C, 0x3D, 0x3E, 0x3F,
+	0x40, 'a',  'b',  'c',  'd',  'e',  'f',  'g',  'h',  'i',  'j',  'k',  'l',  'm',  'n',  'o',
+	'p',  'q',  'r',  's',  't',  'u',  'v',  'w',  'x',  'y',  'z',  0x5B, 0x5C, 0x5D, 0x5E, 0x5F,
+	0x60, 0x61, 0x62, 0x63, 0x64, 0x65, 0x66, 0x67, 0x68, 0x69, 0x6A, 0x6B, 0x6C, 0x6D, 0x6E, 0x6F,
+	0x70, 0x71, 0x72, 0x73, 0x74, 0x75, 0x76, 0x77, 0x78, 0x79, 0x7A, 0x7B, 0x7C, 0x7D, 0x7E, 0x7F
+};
+
+FORCEINLINE bool BothAscii(TCHAR a, TCHAR b)
+{
+	return (((uint32)a | (uint32)b) & 0xffffff80) == 0;
+}
+
+int32 FGenericPlatformStricmp::CompatibleCharTypesStricmp(const TCHAR* String1, const TCHAR* String2)
+{
+	while (true)
+	{
+		TCHAR C1 = *String1++;
+		TCHAR C2 = *String2++;
+
+		if (C1 == C2)
+		{
+			if (C1)
+				continue;
+
+			return 0;
+		}
+		else if (BothAscii(C1, C2))
+		{
+			if (SSIZE_T Diff = LowerAscii[C1] - LowerAscii[C2])
+				return Diff;
+		}
+		else
+		{
+			if (SSIZE_T Diff = TChar<TCHAR>::ToLower(C1) - TChar<TCHAR>::ToLower(C2))
+				return Diff;
+		}
+	}
+}
\ No newline at end of file
diff --git a/Engine/Source/Runtime/Core/Public/GenericPlatform/GenericPlatformStricmp.h b/Engine/Source/Runtime/Core/Public/GenericPlatform/GenericPlatformStricmp.h
--- a/Engine/Source/Runtime/Core/Public/GenericPlatform/GenericPlatformStricmp.h
+++ b/Engine/Source/Runtime/Core/Public/GenericPlatform/GenericPlatformStricmp.h
@@ -27,6 +27,8 @@ template<>							struct TIsComparisonCompatibleChar<UTF32CHAR, ANSICHAR>		{ enum
 struct FGenericPlatformStricmp
 {
 private:
+	CORE_API static int32 CompatibleCharTypesStricmp(const TCHAR* String1, const TCHAR* String2);
+
 	/**
 	 * Compares two strings case-insensitive assuming that both char types are
 	 * compatible.
