Avoid needless sprintf calls when there are no variadic arguments

Shaves off ~1.5% of UHT null builds.

diff --git a/Engine/Source/Runtime/Core/Public/Misc/OutputDevice.h b/Engine/Source/Runtime/Core/Public/Misc/OutputDevice.h
--- a/Engine/Source/Runtime/Core/Public/Misc/OutputDevice.h
+++ b/Engine/Source/Runtime/Core/Public/Misc/OutputDevice.h
@@ -210,6 +210,13 @@ private:
 	void VARARGS CategorizedLogfImpl(const FName& Category, ELogVerbosity::Type Verbosity, const TCHAR* Fmt, ...);
 
 public:
+	template <typename FmtType>
+	void Logf(const FmtType& Fmt)
+	{
+		static_assert(TIsArrayOrRefOfType<FmtType, TCHAR>::Value, "Formatting string must be a TCHAR array.");
+		return Log(Fmt);
+	}
+
 	template <typename FmtType, typename... Types>
 	void Logf(const FmtType& Fmt, Types... Args)
 	{
