Added braces

diff --git a/Engine/Source/Runtime/Core/Private/GenericPlatform/GenericPlatformStricmp.cpp b/Engine/Source/Runtime/Core/Private/GenericPlatform/GenericPlatformStricmp.cpp
--- a/Engine/Source/Runtime/Core/Private/GenericPlatform/GenericPlatformStricmp.cpp
+++ b/Engine/Source/Runtime/Core/Private/GenericPlatform/GenericPlatformStricmp.cpp
@@ -28,19 +28,25 @@ int32 FGenericPlatformStricmp::CompatibleCharTypesStricmp(const TCHAR* String1,
 		if (C1 == C2)
 		{
 			if (C1)
+			{
 				continue;
+			}
 
 			return 0;
 		}
 		else if (BothAscii(C1, C2))
 		{
 			if (SSIZE_T Diff = LowerAscii[C1] - LowerAscii[C2])
+			{
 				return Diff;
+			}
 		}
 		else
 		{
 			if (SSIZE_T Diff = TChar<TCHAR>::ToLower(C1) - TChar<TCHAR>::ToLower(C2))
+			{
 				return Diff;
+			}
 		}
 	}
 }
\ No newline at end of file
