TFunction optimizations.

diff --git a/Engine/Source/Runtime/Core/Public/Templates/Function.h b/Engine/Source/Runtime/Core/Public/Templates/Function.h
--- a/Engine/Source/Runtime/Core/Public/Templates/Function.h
+++ b/Engine/Source/Runtime/Core/Public/Templates/Function.h
@@ -6,6 +6,7 @@
 #include "Misc/AssertionMacros.h"
 #include "HAL/UnrealMemory.h"
 #include "Templates/AndOrNot.h"
+#include "Templates/ChooseClass.h"
 #include "Templates/UnrealTypeTraits.h"
 #include "Templates/RemoveReference.h"
 #include "Templates/Decay.h"
@@ -15,16 +16,26 @@
 #include "Templates/IsMemberPointer.h"
 #include "Templates/IsPointer.h"
 #include "Templates/UnrealTemplate.h"
-#include "Containers/ContainerAllocationPolicies.h"
 #include "Math/UnrealMathUtility.h"
 #include <new>
 
 
 // Disable visualization hack for shipping or test builds.
-#if !(UE_BUILD_SHIPPING || UE_BUILD_TEST)
-	#define ENABLE_TFUNCTIONREF_VISUALIZATION 1
-#else
+//#if !(UE_BUILD_SHIPPING || UE_BUILD_TEST)
+//	#define ENABLE_TFUNCTIONREF_VISUALIZATION 1
+//#else
 	#define ENABLE_TFUNCTIONREF_VISUALIZATION 0
+//#endif
+
+#if defined(_WIN32) && !defined(_WIN64)
+	// Don't use inline storage on Win32, because that will affect the alignment of TFunction, and we can't pass extra-aligned types by value on Win32.
+	#define TFUNCTION_USES_INLINE_STORAGE 0
+#elif USE_SMALL_TFUNCTIONS
+	#define TFUNCTION_USES_INLINE_STORAGE 0
+#else
+	#define TFUNCTION_USES_INLINE_STORAGE 1
+	#define TFUNCTION_INLINE_SIZE         32
+	#define TFUNCTION_INLINE_ALIGNMENT    16
 #endif
 
 template <typename FuncType> class TFunction;
@@ -89,7 +100,8 @@ template <typename T> struct TIsTFunctionRef<const volatile T> { enum { Value =
  */
 namespace UE4Function_Private
 {
-	struct FUniqueFunctionStorage;
+	template <typename T, bool bOnHeap>
+	struct TFunction_CopyableOwnedObject;
 
 	/**
 	 * Common interface to a callable object owned by TFunction.
@@ -97,9 +109,9 @@ namespace UE4Function_Private
 	struct IFunction_OwnedObject
 	{
 		/**
-		 * Creates a copy of the object in the allocator and returns a pointer to it.
+		 * Creates a copy of itself into the storage and returns a pointer to the new object within it.
 		 */
-		virtual IFunction_OwnedObject* CopyToEmptyStorage(FUniqueFunctionStorage& Storage) const = 0;
+		virtual void* CloneToEmptyStorage(void* Storage) const = 0;
 
 		/**
 		 * Returns the address of the object.
@@ -109,16 +121,61 @@ namespace UE4Function_Private
 		/**
 		 * Destructor.
 		 */
-		virtual ~IFunction_OwnedObject()
+		virtual void Destroy() = 0;
+
+		/**
+		 * Destructor.
+		 */
+		virtual ~IFunction_OwnedObject() = default;
+	};
+
+	/**
+	 * Common interface to a callable object owned by TFunction.
+	 */
+	template <typename T>
+	struct IFunction_OwnedObject_OnHeap : public IFunction_OwnedObject
+	{
+		/**
+		 * Destructor.
+		 */
+		virtual void Destroy() override
 		{
+			using DerivedType = TFunction_CopyableOwnedObject<T, true>;
+
+			auto* Derived = static_cast<DerivedType*>(this);
+			Derived->DerivedType::~DerivedType();
+			FMemory::Free(Derived);
 		}
 	};
 
 	/**
-	 * Implementation of IFunction_OwnedObject for a given copyable T.
+	 * Common interface to a callable object owned by TFunction.
 	 */
 	template <typename T>
-	struct TFunction_CopyableOwnedObject final : public IFunction_OwnedObject
+	struct IFunction_OwnedObject_Inline : public IFunction_OwnedObject
+	{
+		/**
+		 * Destructor.
+		 */
+		virtual void Destroy() override
+		{
+			using DerivedType = TFunction_CopyableOwnedObject<T, false>;
+
+			auto* Derived = static_cast<DerivedType*>(this);
+			Derived->DerivedType::~DerivedType();
+		}
+	};
+
+	/**
+	 * Implementation of IFunction_OwnedObject for a given copyable T.
+	 */
+	template <typename T, bool bOnHeap>
+	struct TFunction_CopyableOwnedObject final : public
+#if TFUNCTION_USES_INLINE_STORAGE
+		TChooseClass<bOnHeap, IFunction_OwnedObject_OnHeap<T>, IFunction_OwnedObject_Inline<T>>::Result
+#else
+		IFunction_OwnedObject_OnHeap<T>
+#endif
 	{
 		/**
 		 * Constructor which creates its T by copying.
@@ -136,9 +193,27 @@ namespace UE4Function_Private
 		{
 		}
 
-		IFunction_OwnedObject* CopyToEmptyStorage(FUniqueFunctionStorage& Storage) const override
+		void* CloneToEmptyStorage(void* UntypedStorage) const override
 		{
-			return new (Storage) TFunction_CopyableOwnedObject(Obj);
+			TFunctionStorage<false>& Storage = *(TFunctionStorage<false>*)UntypedStorage;
+
+			void* NewAlloc;
+			#if TFUNCTION_USES_INLINE_STORAGE
+			if /* constexpr */ (!bOnHeap)
+			{
+				NewAlloc = &Storage.InlineAllocation;
+			}
+			else
+			#endif
+			{
+				NewAlloc = FMemory::Malloc(sizeof(TFunction_CopyableOwnedObject), alignof(TFunction_CopyableOwnedObject));
+				Storage.HeapAllocation = NewAlloc;
+				CA_ASSUME(NewAlloc);
+			}
+
+			auto* NewOwned = new (NewAlloc) TFunction_CopyableOwnedObject(Obj);
+
+			return &NewOwned->Obj;
 		}
 
 		virtual void* GetAddress() override
@@ -152,8 +227,13 @@ namespace UE4Function_Private
 	/**
 	 * Implementation of IFunction_OwnedObject for a given non-copyable T.
 	 */
-	template <typename T>
-	struct TFunction_UniqueOwnedObject final : public IFunction_OwnedObject
+	template <typename T, bool bOnHeap>
+	struct TFunction_UniqueOwnedObject final : public
+#if TFUNCTION_USES_INLINE_STORAGE
+		TChooseClass<bOnHeap, IFunction_OwnedObject_OnHeap<T>, IFunction_OwnedObject_Inline<T>>::Result
+#else
+		IFunction_OwnedObject_OnHeap<T>
+#endif
 	{
 		/**
 		 * Constructor which creates its T by moving.
@@ -163,7 +243,7 @@ namespace UE4Function_Private
 		{
 		}
 
-		IFunction_OwnedObject* CopyToEmptyStorage(FUniqueFunctionStorage& Storage) const override
+		void* CloneToEmptyStorage(void* Storage) const override
 		{
 			// Should never get here - copy functions are deleted for TUniqueFunction
 			check(false);
@@ -178,21 +258,6 @@ namespace UE4Function_Private
 		T Obj;
 	};
 
-	#if !defined(_WIN32) || defined(_WIN64)
-		// Let TFunction store up to 32 bytes which are 16-byte aligned before we heap allocate
-		typedef TAlignedBytes<16, 16> FAlignedInlineFunctionType;
-		#if USE_SMALL_TFUNCTIONS
-			typedef FHeapAllocator FFunctionAllocatorType;
-		#else
-			typedef TInlineAllocator<2> FFunctionAllocatorType;
-		#endif
-	#else
-		// ... except on Win32, because we can't pass 16-byte aligned types by value, as some TFunctions are.
-		// So we'll just keep it heap-allocated, which is always sufficiently aligned.
-		typedef TAlignedBytes<16, 8> FAlignedInlineFunctionType;
-		typedef FHeapAllocator FFunctionAllocatorType;
-	#endif
-
 	template <typename T>
 	struct TIsNullableBinding :
 		TOr<
@@ -218,42 +283,61 @@ namespace UE4Function_Private
 		return true;
 	}
 
-	struct FFunctionStorage;
+	template <typename FunctorType, bool bUnique, bool bInline>
+	struct TStorageOwnerType;
 
-	struct FUniqueFunctionStorage
+	template <typename FunctorType, bool bInline>
+	struct TStorageOwnerType<FunctorType, true, bInline>
 	{
-		FUniqueFunctionStorage() = default;
+		using Type = TFunction_UniqueOwnedObject<typename TDecay<FunctorType>::Type, bInline>;
+	};
 
-		FUniqueFunctionStorage(FUniqueFunctionStorage&& Other)
-		{
-			Allocator.MoveToEmpty(Other.Allocator);
-		}
+	template <typename FunctorType, bool bInline>
+	struct TStorageOwnerType<FunctorType, false, bInline>
+	{
+		using Type = TFunction_CopyableOwnedObject<typename TDecay<FunctorType>::Type, bInline>;
+	};
 
-		// Allow moving from a TFunction's storage
-		FUniqueFunctionStorage(FFunctionStorage&& Other);
+	template <typename FunctorType, bool bUnique, bool bInline>
+	using TStorageOwnerTypeT = typename TStorageOwnerType<FunctorType, bUnique, bInline>::Type;
 
-		FUniqueFunctionStorage& operator=(FUniqueFunctionStorage&& Other) = delete;
-		FUniqueFunctionStorage& operator=(const FUniqueFunctionStorage& Other) = delete;
+	struct FFunctionStorage
+	{
+		FFunctionStorage()
+			: HeapAllocation(nullptr)
+		{
+		}
 
-		template <typename FunctorType>
-		typename TDecay<FunctorType>::Type* Bind(FunctorType&& InFunc)
+		FFunctionStorage(FFunctionStorage&& Other)
+			: HeapAllocation(Other.HeapAllocation)
 		{
-			if (!IsBound(InFunc))
-			{
-				return nullptr;
-			}
+			Other.HeapAllocation = nullptr;
+			#if TFUNCTION_USES_INLINE_STORAGE
+				FMemory::Memcpy(&InlineAllocation, &Other.InlineAllocation, sizeof(InlineAllocation));
+			#endif
+		}
 
-			using OwnedType = TFunction_UniqueOwnedObject<typename TDecay<FunctorType>::Type>;
+		FFunctionStorage(const FFunctionStorage& Other) = delete;
+		FFunctionStorage& operator=(FFunctionStorage&& Other) = delete;
+		FFunctionStorage& operator=(const FFunctionStorage& Other) = delete;
 
-			OwnedType* NewObj = new (*this) OwnedType(Forward<FunctorType>(InFunc));
-			return &NewObj->Obj;
+		void* BindCopy(const FFunctionStorage& Other)
+		{
+			void* NewObj = Other.GetBoundObject()->CloneToEmptyStorage(this);
+			return NewObj;
 		}
 
-		void* BindCopy(const FFunctionStorage& Other);
-
 		IFunction_OwnedObject* GetBoundObject() const
 		{
-			return (IFunction_OwnedObject*)Allocator.GetAllocation();
+			IFunction_OwnedObject* Result = (IFunction_OwnedObject*)HeapAllocation;
+			#if TFUNCTION_USES_INLINE_STORAGE
+				if (!Result)
+				{
+					Result = (IFunction_OwnedObject*)&InlineAllocation;
+				}
+			#endif
+
+			return Result;
 		}
 
 		/**
@@ -261,7 +345,17 @@ namespace UE4Function_Private
 		 */
 		void* GetPtr() const
 		{
-			return GetBoundObject()->GetAddress();
+		#if TFUNCTION_USES_INLINE_STORAGE
+			IFunction_OwnedObject* Owned = (IFunction_OwnedObject*)HeapAllocation;
+			if (!Owned)
+			{
+				Owned = (IFunction_OwnedObject*)&InlineAllocation;
+			}
+
+			return Owned->GetAddress();
+		#else
+			return ((IFunction_OwnedObject*)HeapAllocation)->GetAddress();
+		#endif
 		}
 
 		/**
@@ -269,64 +363,75 @@ namespace UE4Function_Private
 		 */
 		void Unbind()
 		{
-			GetBoundObject()->~IFunction_OwnedObject();
+			IFunction_OwnedObject* Owned = GetBoundObject();
+			Owned->Destroy();
 		}
 
-		FFunctionAllocatorType::ForElementType<FAlignedInlineFunctionType> Allocator;
+		void* HeapAllocation;
+		#if TFUNCTION_USES_INLINE_STORAGE
+			// Inline storage for an owned object
+			TAlignedBytes<TFUNCTION_INLINE_SIZE, TFUNCTION_INLINE_ALIGNMENT> InlineAllocation;
+		#endif
 	};
 
-	struct FFunctionStorage : public FUniqueFunctionStorage
+	template <bool bUnique>
+	struct TFunctionStorage : FFunctionStorage
 	{
-		FFunctionStorage() = default;
+		TFunctionStorage() = default;
 
-		FFunctionStorage(FFunctionStorage&& Other)
-			: FUniqueFunctionStorage(MoveTemp((FUniqueFunctionStorage&)Other))
+		TFunctionStorage(FFunctionStorage&& Other)
+			: FFunctionStorage(MoveTemp(Other))
 		{
 		}
 
+		#if TFUNCTION_USES_INLINE_STORAGE
 		template <typename FunctorType>
-		typename TDecay<FunctorType>::Type* Bind(FunctorType&& InFunc)
+		typename TEnableIf<
+			(sizeof(TStorageOwnerTypeT<FunctorType, bUnique, false>) <= TFUNCTION_INLINE_SIZE),
+			typename TDecay<FunctorType>::Type*
+		>::Type Bind(FunctorType&& InFunc)
 		{
 			if (!IsBound(InFunc))
 			{
 				return nullptr;
 			}
 
-			using OwnedType = TFunction_CopyableOwnedObject<typename TDecay<FunctorType>::Type>;
+			using OwnedType = TStorageOwnerTypeT<FunctorType, bUnique, false>;
+
+			void* NewAlloc = &InlineAllocation;
+			auto* NewOwned = new (NewAlloc) OwnedType(Forward<FunctorType>(InFunc));
 
-			OwnedType* NewObj = new (*this) OwnedType(Forward<FunctorType>(InFunc));
-			return &NewObj->Obj;
+			return &NewOwned->Obj;
 		}
+		#endif
 
-		void* BindCopy(const FFunctionStorage& Other)
+		template <typename FunctorType>
+		#if TFUNCTION_USES_INLINE_STORAGE
+		typename TEnableIf<
+			(sizeof(TStorageOwnerTypeT<FunctorType, bUnique, true>) > TFUNCTION_INLINE_SIZE),
+			typename TDecay<FunctorType>::Type*
+		>::Type Bind(FunctorType&& InFunc)
+		#else
+		typename TDecay<FunctorType>::Type* Bind(FunctorType&& InFunc)
+		#endif
 		{
-			IFunction_OwnedObject* ThisFunc = Other.GetBoundObject()->CopyToEmptyStorage(*this);
-			return ThisFunc->GetAddress();
-		}
-	};
+			if (!IsBound(InFunc))
+			{
+				return nullptr;
+			}
 
-	inline FUniqueFunctionStorage::FUniqueFunctionStorage(FFunctionStorage&& Other)
-	{
-		Allocator.MoveToEmpty(Other.Allocator);
-	}
+			using OwnedType = TStorageOwnerTypeT<FunctorType, bUnique, true>;
 
-	inline void* FUniqueFunctionStorage::BindCopy(const FFunctionStorage& Other)
-	{
-		IFunction_OwnedObject* ThisFunc = Other.GetBoundObject()->CopyToEmptyStorage(*this);
-		return ThisFunc->GetAddress();
-	}
-}
+			void* NewAlloc = FMemory::Malloc(sizeof(OwnedType), alignof(OwnedType));
+			CA_ASSUME(NewAlloc);
+			auto* NewOwned = new (NewAlloc) OwnedType(Forward<FunctorType>(InFunc));
 
-inline void* operator new(size_t Size, UE4Function_Private::FUniqueFunctionStorage& Storage)
-{
-	int32 NewSize = FMath::DivideAndRoundUp(Size, sizeof(UE4Function_Private::FAlignedInlineFunctionType));
-	Storage.Allocator.ResizeAllocation(0, NewSize, sizeof(UE4Function_Private::FAlignedInlineFunctionType));
+			HeapAllocation = NewAlloc;
 
-	return Storage.Allocator.GetAllocation();
-}
+			return &NewOwned->Obj;
+		}
+	};
 
-namespace UE4Function_Private
-{
 	#if ENABLE_TFUNCTIONREF_VISUALIZATION
 		/**
 		 * Helper classes to help debugger visualization.
@@ -486,9 +591,15 @@ namespace UE4Function_Private
 		TFunctionRefBase& operator=(TFunctionRefBase&&) = delete;
 		TFunctionRefBase& operator=(const TFunctionRefBase&) = delete;
 
+		// Move all of the assert code out of line
+		FORCENOINLINE void CheckCallable() const
+		{
+//			checkf(Callable, TEXT("Attempting to call an unbound TFunction!"));
+		}
+
 		Ret operator()(ParamTypes... Params) const
 		{
-			checkf(Callable, TEXT("Attempting to call an unbound TFunction!"));
+			CheckCallable();
 			return Callable(Storage.GetPtr(), Params...);
 		}
 
@@ -728,9 +839,9 @@ public:
  * }
  */
 template <typename FuncType>
-class TFunction final : public UE4Function_Private::TFunctionRefBase<UE4Function_Private::FFunctionStorage, FuncType>
+class TFunction final : public UE4Function_Private::TFunctionRefBase<UE4Function_Private::TFunctionStorage<false>, FuncType>
 {
-	using Super = UE4Function_Private::TFunctionRefBase<UE4Function_Private::FFunctionStorage, FuncType>;
+	using Super = UE4Function_Private::TFunctionRefBase<UE4Function_Private::TFunctionStorage<false>, FuncType>;
 
 public:
 	/**
@@ -817,9 +928,9 @@ public:
  * Foo(MoveTemp(MovableFunc)); // ok
  */
 template <typename FuncType>
-class TUniqueFunction final : public UE4Function_Private::TFunctionRefBase<UE4Function_Private::FUniqueFunctionStorage, FuncType>
+class TUniqueFunction final : public UE4Function_Private::TFunctionRefBase<UE4Function_Private::TFunctionStorage<true>, FuncType>
 {
-	using Super = UE4Function_Private::TFunctionRefBase<UE4Function_Private::FUniqueFunctionStorage, FuncType>;
+	using Super = UE4Function_Private::TFunctionRefBase<UE4Function_Private::TFunctionStorage<true>, FuncType>;
 
 public:
 	/**
@@ -861,7 +972,7 @@ public:
 	 * Constructor which takes ownership of a TFunction's functor.
 	 */
 	TUniqueFunction(TFunction<FuncType>&& Other)
-		: Super(MoveTemp(*(UE4Function_Private::TFunctionRefBase<UE4Function_Private::FFunctionStorage, FuncType>*)&Other))
+		: Super(MoveTemp(*(UE4Function_Private::TFunctionRefBase<UE4Function_Private::TFunctionStorage<false>, FuncType>*)&Other))
 	{
 	}
 
@@ -869,7 +980,7 @@ public:
 	 * Constructor which takes ownership of a TFunction's functor.
 	 */
 	TUniqueFunction(const TFunction<FuncType>& Other)
-		: Super(*(const UE4Function_Private::TFunctionRefBase<UE4Function_Private::FFunctionStorage, FuncType>*)&Other)
+		: Super(*(const UE4Function_Private::TFunctionRefBase<UE4Function_Private::TFunctionStorage<false>, FuncType>*)&Other)
 	{
 	}
 
