Re-enabling of disabled debug stuff.

diff --git a/Engine/Source/Runtime/Core/Public/Templates/Function.h b/Engine/Source/Runtime/Core/Public/Templates/Function.h
--- a/Engine/Source/Runtime/Core/Public/Templates/Function.h
+++ b/Engine/Source/Runtime/Core/Public/Templates/Function.h
@@ -21,11 +21,11 @@
 
 
 // Disable visualization hack for shipping or test builds.
-//#if !(UE_BUILD_SHIPPING || UE_BUILD_TEST)
-//	#define ENABLE_TFUNCTIONREF_VISUALIZATION 1
-//#else
+#if !(UE_BUILD_SHIPPING || UE_BUILD_TEST)
+	#define ENABLE_TFUNCTIONREF_VISUALIZATION 1
+#else
 	#define ENABLE_TFUNCTIONREF_VISUALIZATION 0
-//#endif
+#endif
 
 #if defined(_WIN32) && !defined(_WIN64)
 	// Don't use inline storage on Win32, because that will affect the alignment of TFunction, and we can't pass extra-aligned types by value on Win32.
@@ -594,7 +594,7 @@ namespace UE4Function_Private
 		// Move all of the assert code out of line
 		FORCENOINLINE void CheckCallable() const
 		{
-//			checkf(Callable, TEXT("Attempting to call an unbound TFunction!"));
+			checkf(Callable, TEXT("Attempting to call an unbound TFunction!"));
 		}
 
 		Ret operator()(ParamTypes... Params) const
