Missing forward declaration added.

diff --git a/Engine/Source/Runtime/Core/Public/Templates/Function.h b/Engine/Source/Runtime/Core/Public/Templates/Function.h
--- a/Engine/Source/Runtime/Core/Public/Templates/Function.h
+++ b/Engine/Source/Runtime/Core/Public/Templates/Function.h
@@ -103,6 +103,9 @@ namespace UE4Function_Private
 	template <typename T, bool bOnHeap>
 	struct TFunction_CopyableOwnedObject;
 
+	template <bool bUnique>
+	struct TFunctionStorage;
+
 	/**
 	 * Common interface to a callable object owned by TFunction.
 	 */
