Avoid instantiation of undefined class.

diff --git a/Engine/Source/Runtime/Core/Public/Templates/Function.h b/Engine/Source/Runtime/Core/Public/Templates/Function.h
--- a/Engine/Source/Runtime/Core/Public/Templates/Function.h
+++ b/Engine/Source/Runtime/Core/Public/Templates/Function.h
@@ -196,28 +196,7 @@ namespace UE4Function_Private
 		{
 		}
 
-		void* CloneToEmptyStorage(void* UntypedStorage) const override
-		{
-			TFunctionStorage<false>& Storage = *(TFunctionStorage<false>*)UntypedStorage;
-
-			void* NewAlloc;
-			#if TFUNCTION_USES_INLINE_STORAGE
-			if /* constexpr */ (!bOnHeap)
-			{
-				NewAlloc = &Storage.InlineAllocation;
-			}
-			else
-			#endif
-			{
-				NewAlloc = FMemory::Malloc(sizeof(TFunction_CopyableOwnedObject), alignof(TFunction_CopyableOwnedObject));
-				Storage.HeapAllocation = NewAlloc;
-				CA_ASSUME(NewAlloc);
-			}
-
-			auto* NewOwned = new (NewAlloc) TFunction_CopyableOwnedObject(Obj);
-
-			return &NewOwned->Obj;
-		}
+		void* CloneToEmptyStorage(void* UntypedStorage) const override;
 
 		virtual void* GetAddress() override
 		{
@@ -435,6 +414,30 @@ namespace UE4Function_Private
 		}
 	};
 
+	template <typename T, bool bOnHeap>
+	void* TFunction_CopyableOwnedObject<T, bOnHeap>::CloneToEmptyStorage(void* UntypedStorage) const
+	{
+		TFunctionStorage<false>& Storage = *(TFunctionStorage<false>*)UntypedStorage;
+
+		void* NewAlloc;
+		#if TFUNCTION_USES_INLINE_STORAGE
+		if /* constexpr */ (!bOnHeap)
+		{
+			NewAlloc = &Storage.InlineAllocation;
+		}
+		else
+		#endif
+		{
+			NewAlloc = FMemory::Malloc(sizeof(TFunction_CopyableOwnedObject), alignof(TFunction_CopyableOwnedObject));
+			Storage.HeapAllocation = NewAlloc;
+			CA_ASSUME(NewAlloc);
+		}
+
+		auto* NewOwned = new (NewAlloc) TFunction_CopyableOwnedObject(Obj);
+
+		return &NewOwned->Obj;
+	}
+
 	#if ENABLE_TFUNCTIONREF_VISUALIZATION
 		/**
 		 * Helper classes to help debugger visualization.
