PR #4975: Provide the file path of the ICU DLL that failed to load (Contributed by [REDACTED_PERSON])

diff --git a/Engine/Source/Runtime/Core/Private/Internationalization/ICUInternationalization.cpp b/Engine/Source/Runtime/Core/Private/Internationalization/ICUInternationalization.cpp
--- a/Engine/Source/Runtime/Core/Private/Internationalization/ICUInternationalization.cpp
+++ b/Engine/Source/Runtime/Core/Private/Internationalization/ICUInternationalization.cpp
@@ -253,8 +253,9 @@ void FICUInternationalization::LoadDLLs()
 		FString LibraryName = "lib" "icu" + Stem + ".53.1" + LibraryNamePostfix + "." "dylib";
 #endif //PLATFORM_*
 
-		void* DLLHandle = FPlatformProcess::GetDllHandle(*(TargetSpecificPath / LibraryName));
-		check(DLLHandle != nullptr);
+		const FString DllPath = TargetSpecificPath / LibraryName;
+		void* DLLHandle = FPlatformProcess::GetDllHandle(*DllPath);
+		checkf(DLLHandle != nullptr, TEXT("GetDllHandle failed to load: '%s'"), *DllPath);
 		DLLHandles.Add(DLLHandle);
 	}
 }
