Added support for SetThreadDescription thread naming API

This API is better than the old SEH style thread naming since the names also show up in crash dumps and ETW traces

diff --git a/Engine/Source/Runtime/Core/Private/Windows/WindowsRunnableThread.cpp b/Engine/Source/Runtime/Core/Private/Windows/WindowsRunnableThread.cpp
--- a/Engine/Source/Runtime/Core/Private/Windows/WindowsRunnableThread.cpp
+++ b/Engine/Source/Runtime/Core/Private/Windows/WindowsRunnableThread.cpp
@@ -6,9 +6,28 @@
 #include "Misc/FeedbackContext.h"
 #include "HAL/ExceptionHandling.h"
 
-
 DEFINE_LOG_CATEGORY_STATIC(LogThreadingWindows, Log, All);
 
+void FRunnableThreadWin::SetThreadDescription(HANDLE hThread, PCWSTR lpThreadDescription)
+{
+	// SetThreadDescription is only available from Windows 10 version 1607 / Windows Server 2016
+	//
+	// So in order to be compatible with older Windows versions we probe for the API at runtime
+	// and call it only if available.
+
+	typedef HRESULT(*SetThreadDescriptionFnPtr)(HANDLE hThread, PCWSTR lpThreadDescription);
+
+#pragma warning( push )
+#pragma warning( disable: 4191 )	// unsafe conversion from 'type of expression' to 'type required'
+	static SetThreadDescriptionFnPtr RealSetThreadDescription = (SetThreadDescriptionFnPtr) GetProcAddress(GetModuleHandle(TEXT("kernel32.dll")), "SetThreadDescription");
+#pragma warning( pop )
+
+	if (RealSetThreadDescription)
+	{
+		RealSetThreadDescription(hThread, lpThreadDescription);
+	}
+}
+
 
 uint32 FRunnableThreadWin::GuardedRun()
 {
diff --git a/Engine/Source/Runtime/Core/Private/Windows/WindowsRunnableThread.h b/Engine/Source/Runtime/Core/Private/Windows/WindowsRunnableThread.h
--- a/Engine/Source/Runtime/Core/Private/Windows/WindowsRunnableThread.h
+++ b/Engine/Source/Runtime/Core/Private/Windows/WindowsRunnableThread.h
@@ -66,6 +66,8 @@ class FRunnableThreadWin
 #endif
 	}
 
+	static void SetThreadDescription(HANDLE hThread, PCWSTR lpThreadDescription);
+
 	/**
 	 * The thread entry point. Simply forwards the call on to the right
 	 * thread main function
@@ -190,6 +192,8 @@ protected:
 		// Create a sync event to guarantee the Init() function is called first
 		ThreadInitSyncEvent	= FPlatformProcess::GetSynchEventFromPool(true);
 
+		ThreadName = InThreadName ? InThreadName : TEXT("Unnamed UE4");
+
 		// Create the new thread
 		{
 			LLM_SCOPE(ELLMTag::ThreadStack);
@@ -209,12 +213,17 @@ protected:
 		}
 		else
 		{
+#if !PLATFORM_XBOXONE		
+			// We try to use the SetThreadDescription API where possible since this
+			// enables thread names in crashdumps and ETW traces
+			SetThreadDescription(Thread, *ThreadName);
+#endif
+
 			FThreadManager::Get().AddThread(ThreadID, this);
 			ResumeThread(Thread);
 
 			// Let the thread start up, then set the name for debug purposes.
 			ThreadInitSyncEvent->Wait(INFINITE);
-			ThreadName = InThreadName ? InThreadName : TEXT("Unnamed UE4");
 #if !PLATFORM_XBOXONE
 			SetThreadName( ThreadID, TCHAR_TO_ANSI( *ThreadName ) );
 #elif !UE_BUILD_SHIPPING
