Adding TChar unit test in preparation of replacing C stdlib implementation with handwritten ones.

It verifies that the C locale is set. It also checks that TChar implementations match the standard C functions, which is useless with the current TChar implementation that calls the same functions internally.

diff --git a/Engine/Source/Runtime/Core/Private/Tests/Misc/CharTest.cpp b/Engine/Source/Runtime/Core/Private/Tests/Misc/CharTest.cpp
new file mode 100644
--- /dev/null
+++ b/Engine/Source/Runtime/Core/Private/Tests/Misc/CharTest.cpp
@@ -0,0 +1,66 @@
+// Copyright 1998-2018 Epic Games, Inc. All Rights Reserved.
+
+#include "Misc/Char.h"
+#include "Misc/AutomationTest.h"
+#include <locale.h>
+#include <ctype.h>
+#include <wctype.h>
+
+#if WITH_DEV_AUTOMATION_TESTS 
+
+IMPLEMENT_SIMPLE_AUTOMATION_TEST(TCharTest, "System.Core.Misc.Char", EAutomationTestFlags::ApplicationContextMask | EAutomationTestFlags::SmokeFilter)
+
+namespace cstd
+{
+	int tolower(char c) { return ::tolower(c); }
+	int toupper(char c) { return ::toupper(c); }
+	int islower(char c) { return ::islower(c); }
+	int isupper(char c) { return ::isupper(c); }
+	int isalpha(char c) { return ::isalpha(c); }
+	int isgraph(char c) { return ::isgraph(c); }
+	int isprint(char c) { return ::isprint(c); }
+	int ispunct(char c) { return ::ispunct(c); }
+	int isalnum(char c) { return ::isalnum(c); }
+	int isdigit(char c) { return ::isdigit(c); }
+	int isxdigit(char c) { return ::isxdigit(c); }
+	int isspace(char c) { return ::isspace(c); }
+
+	int tolower(wchar_t c) { return ::towlower(c); }
+	int toupper(wchar_t c) { return ::towupper(c); }
+	int islower(wchar_t c) { return ::iswlower(c); }
+	int isupper(wchar_t c) { return ::iswupper(c); }
+	int isalpha(wchar_t c) { return ::iswalpha(c); }
+	int isgraph(wchar_t c)	{ return ::iswgraph(c); }
+	int isprint(wchar_t c) { return ::iswprint(c); }
+	int ispunct(wchar_t c) { return ::iswpunct(c); }
+	int isalnum(wchar_t c) { return ::iswalnum(c); }
+	int isdigit(wchar_t c) { return ::iswdigit(c); }
+	int isxdigit(wchar_t c) { return ::iswxdigit(c); }
+	int isspace(wchar_t c) { return ::iswspace(c); }
+}
+
+bool TCharTest::RunTest(const FString& Parameters)
+{
+	TestTrue(TEXT("C locale not used"), strcmp("C", setlocale(LC_CTYPE, nullptr)) == 0);
+	
+	for (int32 i = 0; i < 0x10000; ++i)
+	{
+		TCHAR C = static_cast<TCHAR>(i);
+		TestEqual("FChar::ToLower()", FChar::ToLower(C), cstd::tolower(C));
+		TestEqual("FChar::ToUpper()", FChar::ToUpper(C), cstd::toupper(C));
+		TestEqual("FChar::IsLower()", !!FChar::IsLower(C), !!cstd::islower(C));
+		TestEqual("FChar::IsUpper()", !!FChar::IsUpper(C), !!cstd::isupper(C));
+		TestEqual("FChar::IsAlpha()", !!FChar::IsAlpha(C), !!cstd::isalpha(C));
+		TestEqual("FChar::IsGraph()", !!FChar::IsGraph(C), !!cstd::isgraph(C));
+		TestEqual("FChar::IsPrint()", !!FChar::IsPrint(C), !!cstd::isprint(C));
+		TestEqual("FChar::IsPunct()", !!FChar::IsPunct(C), !!cstd::ispunct(C));
+		TestEqual("FChar::IsAlnum()", !!FChar::IsAlnum(C), !!cstd::isalnum(C));
+		TestEqual("FChar::IsDigit()", !!FChar::IsDigit(C), !!cstd::isdigit(C));
+		TestEqual("FChar::IsHexDigit()", !!FChar::IsHexDigit(C), !!cstd::isxdigit(C));
+		TestEqual("FChar::IsWhitespace()", !!FChar::IsWhitespace(C), !!cstd::isspace(C));
+	}
+
+	return true;
+}
+
+#endif //WITH_DEV_AUTOMATION_TESTS
\ No newline at end of file
