Fixed Linux compilation for new CharTest

diff --git a/Engine/Source/Runtime/Core/Private/Tests/Misc/CharTest.cpp b/Engine/Source/Runtime/Core/Private/Tests/Misc/CharTest.cpp
--- a/Engine/Source/Runtime/Core/Private/Tests/Misc/CharTest.cpp
+++ b/Engine/Source/Runtime/Core/Private/Tests/Misc/CharTest.cpp
@@ -9,40 +9,40 @@
 #if WITH_DEV_AUTOMATION_TESTS 
 
 IMPLEMENT_SIMPLE_AUTOMATION_TEST(TCharTest, "System.Core.Misc.Char", EAutomationTestFlags::ApplicationContextMask | EAutomationTestFlags::SmokeFilter)
-/*
+
 namespace cstd
 {
-	int tolower(char c) { return ::tolower(c); }
-	int toupper(char c) { return ::toupper(c); }
-	int islower(char c) { return ::islower(c); }
-	int isupper(char c) { return ::isupper(c); }
-	int isalpha(char c) { return ::isalpha(c); }
-	int isgraph(char c) { return ::isgraph(c); }
-	int isprint(char c) { return ::isprint(c); }
-	int ispunct(char c) { return ::ispunct(c); }
-	int isalnum(char c) { return ::isalnum(c); }
-	int isdigit(char c) { return ::isdigit(c); }
-	int isxdigit(char c) { return ::isxdigit(c); }
-	int isspace(char c) { return ::isspace(c); }
+	int tolower(ANSICHAR c) { return ::tolower(c); }
+	int toupper(ANSICHAR c) { return ::toupper(c); }
+	int islower(ANSICHAR c) { return ::islower(c); }
+	int isupper(ANSICHAR c) { return ::isupper(c); }
+	int isalpha(ANSICHAR c) { return ::isalpha(c); }
+	int isgraph(ANSICHAR c) { return ::isgraph(c); }
+	int isprint(ANSICHAR c) { return ::isprint(c); }
+	int ispunct(ANSICHAR c) { return ::ispunct(c); }
+	int isalnum(ANSICHAR c) { return ::isalnum(c); }
+	int isdigit(ANSICHAR c) { return ::isdigit(c); }
+	int isxdigit(ANSICHAR c) { return ::isxdigit(c); }
+	int isspace(ANSICHAR c) { return ::isspace(c); }
 
-	int tolower(wchar_t c) { return ::towlower(c); }
-	int toupper(wchar_t c) { return ::towupper(c); }
-	int islower(wchar_t c) { return ::iswlower(c); }
-	int isupper(wchar_t c) { return ::iswupper(c); }
-	int isalpha(wchar_t c) { return ::iswalpha(c); }
-	int isgraph(wchar_t c)	{ return ::iswgraph(c); }
-	int isprint(wchar_t c) { return ::iswprint(c); }
-	int ispunct(wchar_t c) { return ::iswpunct(c); }
-	int isalnum(wchar_t c) { return ::iswalnum(c); }
-	int isdigit(wchar_t c) { return ::iswdigit(c); }
-	int isxdigit(wchar_t c) { return ::iswxdigit(c); }
-	int isspace(wchar_t c) { return ::iswspace(c); }
+	int tolower(WIDECHAR c) { return ::towlower(c); }
+	int toupper(WIDECHAR c) { return ::towupper(c); }
+	int islower(WIDECHAR c) { return ::iswlower(c); }
+	int isupper(WIDECHAR c) { return ::iswupper(c); }
+	int isalpha(WIDECHAR c) { return ::iswalpha(c); }
+	int isgraph(WIDECHAR c)	{ return ::iswgraph(c); }
+	int isprint(WIDECHAR c) { return ::iswprint(c); }
+	int ispunct(WIDECHAR c) { return ::iswpunct(c); }
+	int isalnum(WIDECHAR c) { return ::iswalnum(c); }
+	int isdigit(WIDECHAR c) { return ::iswdigit(c); }
+	int isxdigit(WIDECHAR c) { return ::iswxdigit(c); }
+	int isspace(WIDECHAR c) { return ::iswspace(c); }
 }
-*/
+
 bool TCharTest::RunTest(const FString& Parameters)
 {
 	TestTrue(TEXT("C locale not used"), strcmp("C", setlocale(LC_CTYPE, nullptr)) == 0);
-	/*
+	
 	for (int32 i = 0; i < 0x10000; ++i)
 	{
 		TCHAR C = static_cast<TCHAR>(i);
@@ -59,7 +59,7 @@ bool TCharTest::RunTest(const FString& Parameters)
 		TestEqual("FChar::IsHexDigit()", !!FChar::IsHexDigit(C), !!cstd::isxdigit(C));
 		TestEqual("FChar::IsWhitespace()", !!FChar::IsWhitespace(C), !!cstd::isspace(C));
 	}
-*/
+
 	return true;
 }
 
