Silence bogus static compilation warning for Win64

Mimicking unsigned char cast that Char.h does before calling ::islower etc. IIUC this cast should not be needed for the actual function declarations in ctype.h. The static analysis seems to have it's own declarations with the wrong signature.

diff --git a/Engine/Source/Runtime/Core/Private/Tests/Misc/CharTest.cpp b/Engine/Source/Runtime/Core/Private/Tests/Misc/CharTest.cpp
--- a/Engine/Source/Runtime/Core/Private/Tests/Misc/CharTest.cpp
+++ b/Engine/Source/Runtime/Core/Private/Tests/Misc/CharTest.cpp
@@ -14,16 +14,16 @@ namespace cstd
 {
 	int tolower(ANSICHAR c) { return ::tolower(c); }
 	int toupper(ANSICHAR c) { return ::toupper(c); }
-	int islower(ANSICHAR c) { return ::islower(c); }
-	int isupper(ANSICHAR c) { return ::isupper(c); }
-	int isalpha(ANSICHAR c) { return ::isalpha(c); }
-	int isgraph(ANSICHAR c) { return ::isgraph(c); }
-	int isprint(ANSICHAR c) { return ::isprint(c); }
-	int ispunct(ANSICHAR c) { return ::ispunct(c); }
-	int isalnum(ANSICHAR c) { return ::isalnum(c); }
-	int isdigit(ANSICHAR c) { return ::isdigit(c); }
-	int isxdigit(ANSICHAR c) { return ::isxdigit(c); }
-	int isspace(ANSICHAR c) { return ::isspace(c); }
+	int islower(ANSICHAR c) { return ::islower((unsigned char)c); }
+	int isupper(ANSICHAR c) { return ::isupper((unsigned char)c); }
+	int isalpha(ANSICHAR c) { return ::isalpha((unsigned char)c); }
+	int isgraph(ANSICHAR c) { return ::isgraph((unsigned char)c); }
+	int isprint(ANSICHAR c) { return ::isprint((unsigned char)c); }
+	int ispunct(ANSICHAR c) { return ::ispunct((unsigned char)c); }
+	int isalnum(ANSICHAR c) { return ::isalnum((unsigned char)c); }
+	int isdigit(ANSICHAR c) { return ::isdigit((unsigned char)c); }
+	int isxdigit(ANSICHAR c) { return ::isxdigit((unsigned char)c); }
+	int isspace(ANSICHAR c) { return ::isspace((unsigned char)c); }
 
 	int tolower(WIDECHAR c) { return ::towlower(c); }
 	int toupper(WIDECHAR c) { return ::towupper(c); }
