Use the correct delegate when broadcasting an OOM condition.

diff --git a/Engine/Source/Runtime/Core/Private/GenericPlatform/GenericPlatformMemory.cpp b/Engine/Source/Runtime/Core/Private/GenericPlatform/GenericPlatformMemory.cpp
--- a/Engine/Source/Runtime/Core/Private/GenericPlatform/GenericPlatformMemory.cpp
+++ b/Engine/Source/Runtime/Core/Private/GenericPlatform/GenericPlatformMemory.cpp
@@ -195,7 +195,7 @@ void FGenericPlatformMemory::OnOutOfMemory(uint64 Size, uint32 Alignment)
 	}
 
 	// let any registered handlers go
-	FCoreDelegates::GetMemoryTrimDelegate().Broadcast();
+	FCoreDelegates::GetOutOfMemoryDelegate().Broadcast();
 
 	UE_LOG(LogMemory, Fatal, TEXT("Ran out of memory allocating %llu bytes with alignment %u"), Size, Alignment);
 }
diff --git a/Engine/Source/Runtime/Core/Private/Unix/UnixPlatformMemory.cpp b/Engine/Source/Runtime/Core/Private/Unix/UnixPlatformMemory.cpp
--- a/Engine/Source/Runtime/Core/Private/Unix/UnixPlatformMemory.cpp
+++ b/Engine/Source/Runtime/Core/Private/Unix/UnixPlatformMemory.cpp
@@ -733,7 +733,7 @@ void FUnixPlatformMemory::OnOutOfMemory(uint64 Size, uint32 Alignment)
 	}
 
 	// let any registered handlers go
-	FCoreDelegates::GetMemoryTrimDelegate().Broadcast();
+	FCoreDelegates::GetOutOfMemoryDelegate().Broadcast();
 
 	UE_LOG(LogMemory, Fatal, TEXT("Ran out of memory allocating %llu bytes with alignment %u"), Size, Alignment);
 	// unreachable
