Consolidated same or similar implementations of the check/verify and ensure assertion macros and expanded out the code to multiple lines to ease readability. Note that it would be possible to further consolidate check/check_f but this is a strategic commit aimed at improving code generation at assert macros, and this will require a formatted/plain strings separation.

diff --git a/Engine/Source/Runtime/Core/Public/Misc/AssertionMacros.h b/Engine/Source/Runtime/Core/Public/Misc/AssertionMacros.h
--- a/Engine/Source/Runtime/Core/Public/Misc/AssertionMacros.h
+++ b/Engine/Source/Runtime/Core/Public/Misc/AssertionMacros.h
@@ -139,15 +139,37 @@ public:
 #endif // !UE_BUILD_SHIPPING
 #if DO_CHECK
 	#define checkCode( Code )		do { Code; } while ( false );
-	#define verify(expr)			{ if(UNLIKELY(!(expr))) { FDebug::LogAssertFailedMessage( #expr, __FILE__, __LINE__, TEXT("") ); _DebugBreakAndPromptForRemote(); FDebug::AssertFailed( #expr, __FILE__, __LINE__ ); CA_ASSUME(false); } }
-	#define check(expr)				{ if(UNLIKELY(!(expr))) { FDebug::LogAssertFailedMessage( #expr, __FILE__, __LINE__, TEXT("") ); _DebugBreakAndPromptForRemote(); FDebug::AssertFailed( #expr, __FILE__, __LINE__ ); CA_ASSUME(false); } }
+	#define verify(expr)			UE_CHECK_IMPL(expr)
+	#define check(expr)				UE_CHECK_IMPL(expr)
+
+	#define UE_CHECK_IMPL(expr) \
+		{ \
+			if (UNLIKELY(!(expr))) \
+			{ \
+				FDebug::LogAssertFailedMessage( #expr, __FILE__, __LINE__, TEXT("") ); \
+				_DebugBreakAndPromptForRemote(); \
+				FDebug::AssertFailed( #expr, __FILE__, __LINE__ ); \
+				CA_ASSUME(false); \
+			} \
+		}
 	
 	/**
 	 * verifyf, checkf: Same as verify, check but with printf style additional parameters
 	 * Read about __VA_ARGS__ (variadic macros) on http://gcc.gnu.org/onlinedocs/gcc-3.4.4/cpp.pdf.
 	 */
-	#define verifyf(expr, format,  ...)		{ if(UNLIKELY(!(expr))) { FDebug::LogAssertFailedMessage( #expr, __FILE__, __LINE__, format, ##__VA_ARGS__ ); _DebugBreakAndPromptForRemote(); FDebug::AssertFailed(#expr, __FILE__, __LINE__, format, ##__VA_ARGS__); CA_ASSUME(false); } }
-	#define checkf(expr, format,  ...)		{ if(UNLIKELY(!(expr))) { FDebug::LogAssertFailedMessage( #expr, __FILE__, __LINE__, format, ##__VA_ARGS__ ); _DebugBreakAndPromptForRemote(); FDebug::AssertFailed(#expr, __FILE__, __LINE__, format, ##__VA_ARGS__); CA_ASSUME(false); } }
+	#define verifyf(expr, format,  ...)		UE_CHECK_F_IMPL(expr, format, ##__VA_ARGS__)
+	#define checkf(expr, format,  ...)		UE_CHECK_F_IMPL(expr, format, ##__VA_ARGS__)
+
+	#define UE_CHECK_F_IMPL(expr, format, ...) \
+		{ \
+			if (UNLIKELY(!(expr))) \
+			{ \
+				FDebug::LogAssertFailedMessage( #expr, __FILE__, __LINE__, format, ##__VA_ARGS__ ); \
+				_DebugBreakAndPromptForRemote(); \
+				FDebug::AssertFailed(#expr, __FILE__, __LINE__, format, ##__VA_ARGS__); \
+				CA_ASSUME(false); \
+			} \
+		}
 	/**
 	 * Denotes code paths that should never be reached.
 	 */
@@ -194,9 +216,9 @@ public:
 // Check for development only.
 //
 #if DO_GUARD_SLOW
-	#define checkSlow(expr)					{ if(UNLIKELY(!(expr))) { FDebug::LogAssertFailedMessage( #expr, __FILE__, __LINE__, TEXT("") ); _DebugBreakAndPromptForRemote(); FDebug::AssertFailed(#expr, __FILE__, __LINE__); CA_ASSUME(false); } }
-	#define checkfSlow(expr, format, ...)	{ if(UNLIKELY(!(expr))) { FDebug::LogAssertFailedMessage( #expr, __FILE__, __LINE__, format, ##__VA_ARGS__ ); _DebugBreakAndPromptForRemote(); FDebug::AssertFailed( #expr, __FILE__, __LINE__, format, ##__VA_ARGS__ ); CA_ASSUME(false); } }
-	#define verifySlow(expr)				{ if(UNLIKELY(!(expr))) { FDebug::LogAssertFailedMessage( #expr, __FILE__, __LINE__, TEXT("") ); _DebugBreakAndPromptForRemote(); FDebug::AssertFailed(#expr, __FILE__, __LINE__); CA_ASSUME(false); } }
+	#define checkSlow(expr)					check(expr)
+	#define checkfSlow(expr, format, ...)	checkf(expr, format, ##__VA_ARGS__)
+	#define verifySlow(expr)				check(expr)
 #else
 	#define checkSlow(expr)					{ CA_ASSUME(expr); }
 	#define checkfSlow(expr, format, ...)	{ CA_ASSUME(expr); }
@@ -246,17 +268,20 @@ public:
 	}
 
 	#if UE_BUILD_SHIPPING
-		#define UE_ENSURE_BREAK() false
-		#define UE_ENSURE_BREAK_ONCE() false
+		#define UE_ENSURE_BREAK_IMPL() false
 	#else
-		#define UE_ENSURE_BREAK() ((FPlatformMisc::IsDebuggerPresent() || (FPlatformMisc::PromptForRemoteDebugging(true), true)), UE_DEBUG_BREAK(), false)
-		#define UE_ENSURE_BREAK_ONCE() (UE4Asserts_Private::TrueOnFirstCallOnly([]{}) && UE_ENSURE_BREAK())
+		#define UE_ENSURE_BREAK_IMPL() ((FPlatformMisc::IsDebuggerPresent() || (FPlatformMisc::PromptForRemoteDebugging(true), true)), UE_DEBUG_BREAK(), false)
 	#endif
 
-	#define ensure(           InExpression                ) (LIKELY(!!(InExpression)) || FDebug::OptionallyLogFormattedEnsureMessageReturningFalse(UE4Asserts_Private::TrueOnFirstCallOnly([]{}), #InExpression, __FILE__, __LINE__, TEXT("")               ) || UE_ENSURE_BREAK_ONCE())
-	#define ensureMsgf(       InExpression, InFormat, ... ) (LIKELY(!!(InExpression)) || FDebug::OptionallyLogFormattedEnsureMessageReturningFalse(UE4Asserts_Private::TrueOnFirstCallOnly([]{}), #InExpression, __FILE__, __LINE__, InFormat, ##__VA_ARGS__) || UE_ENSURE_BREAK_ONCE())
-	#define ensureAlways(     InExpression                ) (LIKELY(!!(InExpression)) || FDebug::OptionallyLogFormattedEnsureMessageReturningFalse(true,                                          #InExpression, __FILE__, __LINE__, TEXT("")               ) || UE_ENSURE_BREAK())
-	#define ensureAlwaysMsgf( InExpression, InFormat, ... ) (LIKELY(!!(InExpression)) || FDebug::OptionallyLogFormattedEnsureMessageReturningFalse(true,                                          #InExpression, __FILE__, __LINE__, InFormat, ##__VA_ARGS__) || UE_ENSURE_BREAK())
+	#define UE_ENSURE_IMPL(Always, InExpression, InFormat, ...) \
+		(LIKELY(!!(InExpression)) \
+		|| FDebug::OptionallyLogFormattedEnsureMessageReturningFalse(Always, #InExpression, __FILE__, __LINE__, InFormat, ##__VA_ARGS__) \
+		|| (Always && UE_ENSURE_BREAK_IMPL()))
+
+	#define ensure(           InExpression                ) UE_ENSURE_IMPL(UE4Asserts_Private::TrueOnFirstCallOnly([]{}), InExpression, TEXT(""))
+	#define ensureMsgf(       InExpression, InFormat, ... ) UE_ENSURE_IMPL(UE4Asserts_Private::TrueOnFirstCallOnly([]{}), InExpression, InFormat, ##__VA_ARGS__)
+	#define ensureAlways(     InExpression                ) UE_ENSURE_IMPL(true,                                          InExpression, TEXT(""))
+	#define ensureAlwaysMsgf( InExpression, InFormat, ... ) UE_ENSURE_IMPL(true,                                          InExpression, InFormat, ##__VA_ARGS__)
 
 #else	// DO_CHECK
 
