Fixed Ctrl-C handling on Windows so that it always ends execution when there's a valid console.

diff --git a/Engine/Source/Runtime/ApplicationCore/Private/Windows/WindowsConsoleOutputDevice.cpp b/Engine/Source/Runtime/ApplicationCore/Private/Windows/WindowsConsoleOutputDevice.cpp
--- a/Engine/Source/Runtime/ApplicationCore/Private/Windows/WindowsConsoleOutputDevice.cpp
+++ b/Engine/Source/Runtime/ApplicationCore/Private/Windows/WindowsConsoleOutputDevice.cpp
@@ -82,6 +82,10 @@ void FWindowsConsoleOutputDevice::Show( bool ShowWindow )
 			{
 				bIsAttached = true;
 			}
+			else
+			{
+				FWindowsPlatformMisc::SetGracefulTerminationHandler();
+			}
 			ConsoleHandle = GetStdHandle(OutputDeviceConstants::WIN_STD_OUTPUT_HANDLE);
 
 			if( ConsoleHandle != INVALID_HANDLE_VALUE )
@@ -150,9 +154,6 @@ void FWindowsConsoleOutputDevice::Show( bool ShowWindow )
 				ConsolePosY = FMath::Min(FMath::Max(ConsolePosY, DisplayMetrics.VirtualDisplayRect.Top), DisplayMetrics.VirtualDisplayRect.Bottom - BottomPadding);
 
 				::SetWindowPos( GetConsoleWindow(), HWND_TOP, ConsolePosX, ConsolePosY, 0, 0, SWP_NOSIZE | SWP_NOSENDCHANGING | SWP_NOZORDER );
-				
-				// set the control-c, etc handler
-				FPlatformMisc::SetGracefulTerminationHandler();
 			}
 		}
 	}
diff --git a/Engine/Source/Runtime/Core/Private/Windows/WindowsPlatformMisc.cpp b/Engine/Source/Runtime/Core/Private/Windows/WindowsPlatformMisc.cpp
--- a/Engine/Source/Runtime/Core/Private/Windows/WindowsPlatformMisc.cpp
+++ b/Engine/Source/Runtime/Core/Private/Windows/WindowsPlatformMisc.cpp
@@ -590,14 +590,57 @@ static BOOL WINAPI ConsoleCtrlHandler( ::DWORD /*Type*/ )
 		GError->Flush();
 	}
 
-	// Windows will now terminate this process with ExitCode = 0xC000013A
-	return true;
+	// There's no guarantee that the process is paying attention to GIsRequestingExit
+	// (e.g. some long-running commandlets). Using that for shutdown is therefore not
+	// reliable. Deferring to the default CtrlHandler is also no good as that calls
+	// ExitProcess() which will terminate all threads and detach all DLLS. This can
+	// result in deadlocks and/or asserts as each DLL's atexit() is processed. So
+	// let's hard terminate the process.
+	TerminateProcess(GetCurrentProcess(), 1);
+	return false;
 }
 
 void FWindowsPlatformMisc::SetGracefulTerminationHandler()
 {
+	if (GetConsoleWindow() == nullptr)
+		return;
+
 	// Set console control handler so we can exit if requested.
 	SetConsoleCtrlHandler(ConsoleCtrlHandler, true);
+
+#if !UE_BUILD_SHIPPING && PLATFORM_CPU_X86_FAMILY
+	// There are many places that can register a Ctrl-C handler, some of which
+	// we are not in control of (third party Perforce libraries for example). This
+	// can result in Ctrl-C doing nothing, or an abrupt call to ExitProcess() which
+	// will cause asserts and/or deadlocks. So we're going to apply a patch so no
+	// one else can register a handler.
+	auto* SetCtrlCProc = (uint8*)SetConsoleCtrlHandler;
+	if (SetCtrlCProc[0] == 0xff && SetCtrlCProc[1] == 0x25)
+	{
+#if PLATFORM_64BITS
+		// Follow a possible "jmp [eip] + disp32" instruction to get to the actual
+		// implementation of the SetConsoleCtrlHandler function.
+		SetCtrlCProc = *(uint8**)(SetCtrlCProc + 6 + *(uint32*)(SetCtrlCProc + 2));
+#else
+		// Follow "jmp [disp32]"
+		uintptr_t Disp32 = *(uintptr_t*)(SetCtrlCProc + 2);
+		SetCtrlCProc = *(uint8**)(Disp32);
+#endif
+	}
+
+	// Patch the start of the SetConsoleCtrlHandler function.
+	uint8 Patch[] = {
+		0xb8, 0x01, 0x00, 0x00, 0x00,	// mov eax, 1
+		0xc3							// ret
+	};
+	DWORD PrevProtection;
+	if (VirtualProtect(SetCtrlCProc, sizeof(Patch), PAGE_READWRITE, &PrevProtection))
+	{
+		memcpy(SetCtrlCProc, Patch, sizeof(Patch));
+		VirtualProtect(SetCtrlCProc, sizeof(Patch), PrevProtection, &PrevProtection);
+		FlushInstructionCache(GetCurrentProcess(), nullptr, 0);
+	}
+#endif // !UE_BUILD_SHIPPING && PLATFORM_CPU_X86_FAMILY
 }
 
 int32 FWindowsPlatformMisc::GetMaxPathLength()
diff --git a/Engine/Source/Runtime/Launch/Private/Windows/LaunchWindows.cpp b/Engine/Source/Runtime/Launch/Private/Windows/LaunchWindows.cpp
--- a/Engine/Source/Runtime/Launch/Private/Windows/LaunchWindows.cpp
+++ b/Engine/Source/Runtime/Launch/Private/Windows/LaunchWindows.cpp
@@ -106,6 +106,9 @@ void SetupWindowsEnvironment( void )
 	// don't fill buffers with 0xfd as we make assumptions for FNames st we only use a fraction of the entire buffer
 	_CrtSetDebugFillThreshold( 0 );
 #endif
+
+	// Register a handler for Ctrl-C so we've effective signal handling from the outset.
+	FWindowsPlatformMisc::SetGracefulTerminationHandler();
 }
 
 /**
