PR #5177: Fixed issue of type mismatch in StringToArray(). (Contributed by [REDACTED_PERSON])

diff --git a/Engine/Source/Runtime/Core/Public/Containers/StringConv.h b/Engine/Source/Runtime/Core/Public/Containers/StringConv.h
--- a/Engine/Source/Runtime/Core/Public/Containers/StringConv.h
+++ b/Engine/Source/Runtime/Core/Public/Containers/StringConv.h
@@ -938,7 +938,7 @@ FORCEINLINE typename TEnableIf<!FPlatformString::TAreEncodingsCompatible<To, Fro
 template <typename ToType, typename FromType>
 FORCEINLINE TArray<ToType> StringToArray(const FromType* Src, int32 SrcLen)
 {
-	int32 DestLen = FPlatformString::ConvertedLength<TCHAR>(Src, SrcLen);
+	int32 DestLen = FPlatformString::ConvertedLength<ToType>(Src, SrcLen);
 
 	TArray<ToType> Result;
 	Result.AddUninitialized(DestLen);
