PR #4750: Exposed working directory in FMonitoredProcess (Contributed by [REDACTED_PERSON])

diff --git a/Engine/Source/Runtime/Core/Private/Misc/MonitoredProcess.cpp b/Engine/Source/Runtime/Core/Private/Misc/MonitoredProcess.cpp
--- a/Engine/Source/Runtime/Core/Private/Misc/MonitoredProcess.cpp
+++ b/Engine/Source/Runtime/Core/Private/Misc/MonitoredProcess.cpp
@@ -8,6 +8,10 @@
  *****************************************************************************/
 
 FMonitoredProcess::FMonitoredProcess( const FString& InURL, const FString& InParams, bool InHidden, bool InCreatePipes )
+	: FMonitoredProcess(InURL, InParams, FPaths::RootDir(), InHidden, InCreatePipes)
+{ }
+
+FMonitoredProcess::FMonitoredProcess( const FString& InURL, const FString& InParams, const FString& InWorkingDir, bool InHidden, bool InCreatePipes )
 	: Canceling(false)
 	, EndTime(0)
 	, Hidden(InHidden)
@@ -19,6 +23,7 @@ FMonitoredProcess::FMonitoredProcess( const FString& InURL, const FString& InPar
 	, Thread(nullptr)
 	, bIsRunning(false)
 	, URL(InURL)
+	, WorkingDir(InWorkingDir)
 	, WritePipe(nullptr)
 	, bCreatePipes(InCreatePipes)
 	, SleepInterval(0.0f)
@@ -68,7 +73,7 @@ bool FMonitoredProcess::Launch()
 		return false;
 	}
 
-	ProcessHandle = FPlatformProcess::CreateProc(*URL, *Params, false, Hidden, Hidden, nullptr, 0, *FPaths::RootDir(), WritePipe);
+	ProcessHandle = FPlatformProcess::CreateProc(*URL, *Params, false, Hidden, Hidden, nullptr, 0, *WorkingDir, WritePipe);
 
 	if (!ProcessHandle.IsValid())
 	{
diff --git a/Engine/Source/Runtime/Core/Public/Misc/MonitoredProcess.h b/Engine/Source/Runtime/Core/Public/Misc/MonitoredProcess.h
--- a/Engine/Source/Runtime/Core/Public/Misc/MonitoredProcess.h
+++ b/Engine/Source/Runtime/Core/Public/Misc/MonitoredProcess.h
@@ -40,9 +40,21 @@ public:
 	 * @param InURL The URL of the executable to launch.
 	 * @param InParams The command line parameters.
 	 * @param InHidden Whether the window of the process should be hidden.
+	 * @param InCreatePipes Whether the output should be redirected to the caller.
 	 */
 	FMonitoredProcess( const FString& InURL, const FString& InParams, bool InHidden, bool InCreatePipes = true );
 
+	/**
+	* Creates a new monitored process.
+	*
+	* @param InURL The URL of the executable to launch.
+	* @param InParams The command line parameters.
+	* @param InHidden Whether the window of the process should be hidden.
+	* @param InWorkingDir The URL of the working dir where the executable should launch.
+	* @param InCreatePipes Whether the output should be redirected to the caller.
+	*/
+	FMonitoredProcess( const FString& InURL, const FString& InParams, const FString& InWorkingDir, bool InHidden, bool InCreatePipes = true );
+
 	/** Destructor. */
 	~FMonitoredProcess();
 
@@ -216,6 +228,9 @@ private:
 	// Holds the URL of the executable to launch. */
 	FString URL;
 
+	// Holds the URL of the working dir for the process. */
+	FString WorkingDir;
+
 	// Holds the write pipe. */
 	void* WritePipe;
 
