Removed deprecated DebugBreak*() functions.

diff --git a/Engine/Source/Runtime/Core/Public/Android/AndroidMisc.h b/Engine/Source/Runtime/Core/Public/Android/AndroidMisc.h
--- a/Engine/Source/Runtime/Core/Public/Android/AndroidMisc.h
+++ b/Engine/Source/Runtime/Core/Public/Android/AndroidMisc.h
@@ -211,40 +211,7 @@ public:
 #endif
 		}
 	}
-
-	DEPRECATED(4.19, "FPlatformMisc::DebugBreak is deprecated. Use the UE_DEBUG_BREAK() macro instead.")
-	FORCEINLINE static void DebugBreak()
-	{
-		DebugBreakInternal();
-	}
-#endif
-
-	/** Break into debugger. Returning false allows this function to be used in conditionals. */
-	DEPRECATED(4.19, "FPlatformMisc::DebugBreakReturningFalse is deprecated. Use the (UE_DEBUG_BREAK(), false) expression instead.")
-	FORCEINLINE static bool DebugBreakReturningFalse()
-	{
-#if !UE_BUILD_SHIPPING
-		UE_DEBUG_BREAK();
-#endif
-		return false;
-	}
-
-	/** Prompts for remote debugging if debugger is not attached. Regardless of result, breaks into debugger afterwards. Returns false for use in conditionals. */
-	DEPRECATED(4.19, "FPlatformMisc::DebugBreakAndPromptForRemoteReturningFalse() is deprecated.")
-	static FORCEINLINE bool DebugBreakAndPromptForRemoteReturningFalse(bool bIsEnsure = false)
-	{
-#if !UE_BUILD_SHIPPING
-		if (!IsDebuggerPresent())
-		{
-			PromptForRemoteDebugging(bIsEnsure);
-		}
-
-		UE_DEBUG_BREAK();
 #endif
-
-		return false;
-	}
-
 	FORCEINLINE static void MemoryBarrier()
 	{
 		__sync_synchronize();
diff --git a/Engine/Source/Runtime/Core/Public/Apple/ApplePlatformMisc.h b/Engine/Source/Runtime/Core/Public/Apple/ApplePlatformMisc.h
--- a/Engine/Source/Runtime/Core/Public/Apple/ApplePlatformMisc.h
+++ b/Engine/Source/Runtime/Core/Public/Apple/ApplePlatformMisc.h
@@ -96,36 +96,6 @@ struct CORE_API FApplePlatformMisc : public FGenericPlatformMisc
 	}
 #endif
 
-	DEPRECATED(4.19, "FPlatformMisc::DebugBreak is deprecated. Use the UE_DEBUG_BREAK() macro instead.")
-	FORCEINLINE static void DebugBreak()
-	{
-		UE_DEBUG_BREAK();
-	}
-
-	/** Break into debugger. Returning false allows this function to be used in conditionals. */
-	DEPRECATED(4.19, "FPlatformMisc::DebugBreakReturningFalse is deprecated. Use the (UE_DEBUG_BREAK(), false) expression instead.")
-	FORCEINLINE static bool DebugBreakReturningFalse()
-	{
-		UE_DEBUG_BREAK();
-		return false;
-	}
-
-	/** Prompts for remote debugging if debugger is not attached. Regardless of result, breaks into debugger afterwards. Returns false for use in conditionals. */
-	DEPRECATED(4.19, "FPlatformMisc::DebugBreakAndPromptForRemoteReturningFalse() is deprecated.")
-	FORCEINLINE static bool DebugBreakAndPromptForRemoteReturningFalse(bool bIsEnsure = false)
-	{
-#if !UE_BUILD_SHIPPING
-		if (!IsDebuggerPresent())
-		{
-			PromptForRemoteDebugging(bIsEnsure);
-		}
-
-		UE_DEBUG_BREAK();
-#endif
-
-		return false;
-	}
-
 	FORCEINLINE static void MemoryBarrier()
 	{
 		__sync_synchronize();
diff --git a/Engine/Source/Runtime/Core/Public/HTML5/HTML5PlatformMisc.h b/Engine/Source/Runtime/Core/Public/HTML5/HTML5PlatformMisc.h
--- a/Engine/Source/Runtime/Core/Public/HTML5/HTML5PlatformMisc.h
+++ b/Engine/Source/Runtime/Core/Public/HTML5/HTML5PlatformMisc.h
@@ -74,36 +74,6 @@ struct CORE_API FHTML5Misc : public FGenericPlatformMisc
 		}
 	}
 
-	DEPRECATED(4.19, "FPlatformMisc::DebugBreak is deprecated. Use the UE_DEBUG_BREAK() macro instead.")
-	FORCEINLINE static void DebugBreak()
-	{
-		UE_DEBUG_BREAK();
-	}
-
-	/** Break into debugger. Returning false allows this function to be used in conditionals. */
-	DEPRECATED(4.19, "FPlatformMisc::DebugBreakReturningFalse is deprecated. Use the (UE_DEBUG_BREAK(), false) expression instead.")
-	FORCEINLINE static bool DebugBreakReturningFalse()
-	{
-		UE_DEBUG_BREAK();
-		return false;
-	}
-
-	/** Prompts for remote debugging if debugger is not attached. Regardless of result, breaks into debugger afterwards. Returns false for use in conditionals. */
-	DEPRECATED(4.19, "FPlatformMisc::DebugBreakAndPromptForRemoteReturningFalse() is deprecated.")
-	static FORCEINLINE bool DebugBreakAndPromptForRemoteReturningFalse(bool bIsEnsure = false)
-	{
-#if !UE_BUILD_SHIPPING
-		if (!IsDebuggerPresent())
-		{
-			PromptForRemoteDebugging(bIsEnsure);
-		}
-
-		UE_DEBUG_BREAK();
-#endif
-
-		return false;
-	}
-
 	static void LocalPrint(const TCHAR* Str);
 };
 
diff --git a/Engine/Source/Runtime/Core/Public/Unix/UnixPlatformMisc.h b/Engine/Source/Runtime/Core/Public/Unix/UnixPlatformMisc.h
--- a/Engine/Source/Runtime/Core/Public/Unix/UnixPlatformMisc.h
+++ b/Engine/Source/Runtime/Core/Public/Unix/UnixPlatformMisc.h
@@ -43,40 +43,8 @@ struct CORE_API FUnixPlatformMisc : public FGenericPlatformMisc
 #if !UE_BUILD_SHIPPING
 	static bool IsDebuggerPresent();
 	static void DebugBreakInternal();
-
-	DEPRECATED(4.19, "FPlatformMisc::DebugBreak is deprecated. Use the UE_DEBUG_BREAK() macro instead.")
-	FORCEINLINE static void DebugBreak()
-	{
-		UE_DEBUG_BREAK();
-	}
 #endif // !UE_BUILD_SHIPPING
 
-	/** Break into debugger. Returning false allows this function to be used in conditionals. */
-	DEPRECATED(4.19, "FPlatformMisc::DebugBreakReturningFalse is deprecated. Use the (UE_DEBUG_BREAK(), false) expression instead.")
-	FORCEINLINE static bool DebugBreakReturningFalse()
-	{
-#if !UE_BUILD_SHIPPING
-		UE_DEBUG_BREAK();
-#endif
-		return false;
-	}
-
-	/** Prompts for remote debugging if debugger is not attached. Regardless of result, breaks into debugger afterwards. Returns false for use in conditionals. */
-	DEPRECATED(4.19, "FPlatformMisc::DebugBreakAndPromptForRemoteReturningFalse() is deprecated.")
-	static FORCEINLINE bool DebugBreakAndPromptForRemoteReturningFalse(bool bIsEnsure = false)
-	{
-#if !UE_BUILD_SHIPPING
-		if (!IsDebuggerPresent())
-		{
-			PromptForRemoteDebugging(bIsEnsure);
-		}
-
-		UE_DEBUG_BREAK();
-#endif
-
-		return false;
-	}
-
 	static void LowLevelOutputDebugString(const TCHAR *Message);
 
 	static void RequestExit(bool Force);
diff --git a/Engine/Source/Runtime/Core/Public/Windows/WindowsPlatformMisc.h b/Engine/Source/Runtime/Core/Public/Windows/WindowsPlatformMisc.h
--- a/Engine/Source/Runtime/Core/Public/Windows/WindowsPlatformMisc.h
+++ b/Engine/Source/Runtime/Core/Public/Windows/WindowsPlatformMisc.h
@@ -56,41 +56,8 @@ struct CORE_API FWindowsPlatformMisc
 
 #if !UE_BUILD_SHIPPING
 	static bool IsDebuggerPresent();
-
-	DEPRECATED(4.19, "FPlatformMisc::DebugBreak is deprecated. Use the UE_DEBUG_BREAK() macro instead.")
-	FORCEINLINE static void DebugBreak()
-	{
-		UE_DEBUG_BREAK();
-	}
-#endif
-
-
-	/** Break into debugger. Returning false allows this function to be used in conditionals. */
-	DEPRECATED(4.19, "FPlatformMisc::DebugBreakReturningFalse is deprecated. Use the (UE_DEBUG_BREAK(), false) expression instead.")
-	FORCEINLINE static bool DebugBreakReturningFalse()
-	{
-#if !UE_BUILD_SHIPPING
-		UE_DEBUG_BREAK();
-#endif
-		return false;
-	}
-
-	/** Prompts for remote debugging if debugger is not attached. Regardless of result, breaks into debugger afterwards. Returns false for use in conditionals. */
-	DEPRECATED(4.19, "FPlatformMisc::DebugBreakAndPromptForRemoteReturningFalse() is deprecated.")
-	static FORCEINLINE bool DebugBreakAndPromptForRemoteReturningFalse(bool bIsEnsure = false)
-	{
-#if !UE_BUILD_SHIPPING
-		if (!IsDebuggerPresent())
-		{
-			PromptForRemoteDebugging(bIsEnsure);
-		}
-
-		UE_DEBUG_BREAK();
 #endif
 
-		return false;
-	}
-
 #if !UE_BUILD_SHIPPING
 	static void BeginNamedEventFrame();
 	static void BeginNamedEvent(const struct FColor& Color, const TCHAR* Text);
