- Updated all platforms to check the presence of a debugger and issue breakpoints at UE_DEBUG_BREAK sites instead of from inside child calls. Windows already operated this was as of 3877545, however Windows-VS2017 is not the only platform that will include inlined functions when displaying callstacks.
- Exposed trap instructions in a separate PLATFORM_BREAK macro that does not include a debugger check.

diff --git a/Engine/Source/Runtime/Core/Private/Unix/UnixPlatformMisc.cpp b/Engine/Source/Runtime/Core/Private/Unix/UnixPlatformMisc.cpp
--- a/Engine/Source/Runtime/Core/Private/Unix/UnixPlatformMisc.cpp
+++ b/Engine/Source/Runtime/Core/Private/Unix/UnixPlatformMisc.cpp
@@ -927,19 +927,3 @@ void FUnixPlatformMisc::UngrabAllInput()
 		UngrabAllInputCallback();
 	}
 }
-
-#if !UE_BUILD_SHIPPING
-void FUnixPlatformMisc::DebugBreakInternal()
-{
-	if( IsDebuggerPresent() )
-	{
-		UngrabAllInput();
-
-#if PLATFORM_CPU_X86_FAMILY
-		__asm__ volatile("int $0x03");
-#else
-		raise(SIGTRAP);
-#endif
-	}
-}
-#endif // !UE_BUILD_SHIPPING
diff --git a/Engine/Source/Runtime/Core/Public/Android/AndroidMisc.h b/Engine/Source/Runtime/Core/Public/Android/AndroidMisc.h
--- a/Engine/Source/Runtime/Core/Public/Android/AndroidMisc.h
+++ b/Engine/Source/Runtime/Core/Public/Android/AndroidMisc.h
@@ -14,12 +14,16 @@
 template <typename FuncType>
 class TFunction;
 
-#if UE_BUILD_SHIPPING
-#define UE_DEBUG_BREAK() ((void)0)
+#if PLATFORM_ANDROID_ARM64
+	#define PLATFORM_BREAK()	__asm__(".inst 0xd4200000")
+#elif PLATFORM_ANDROID_ARM
+	#define PLATFORM_BREAK()	__asm__("trap")
 #else
-#define UE_DEBUG_BREAK() (FAndroidMisc::DebugBreakInternal())
+	#define PLATFORM_BREAK()	__asm__("int $3")
 #endif
 
+#define UE_DEBUG_BREAK_IMPL()	PLATFORM_BREAK()
+
 /**
  * Android implementation of the misc OS functions
  */
@@ -198,20 +202,6 @@ public:
 #if !UE_BUILD_SHIPPING
 	static bool IsDebuggerPresent();
 
-	FORCEINLINE static void DebugBreakInternal()
-	{
-		if( IsDebuggerPresent() )
-		{
-#if PLATFORM_ANDROID_ARM64
-			__asm__(".inst 0xd4200000");
-#elif PLATFORM_ANDROID_ARM
-			__asm__("trap");
-#else
-			__asm__("int $3");
-#endif
-		}
-	}
-#endif
 	FORCEINLINE static void MemoryBarrier()
 	{
 		__sync_synchronize();
diff --git a/Engine/Source/Runtime/Core/Public/Apple/ApplePlatformMisc.h b/Engine/Source/Runtime/Core/Public/Apple/ApplePlatformMisc.h
--- a/Engine/Source/Runtime/Core/Public/Apple/ApplePlatformMisc.h
+++ b/Engine/Source/Runtime/Core/Public/Apple/ApplePlatformMisc.h
@@ -20,11 +20,7 @@
 #define WITH_SIMULATOR 0
 #endif
 
-#if UE_BUILD_SHIPPING
-#define UE_DEBUG_BREAK() ((void)0)
-#else
-#define UE_DEBUG_BREAK() (FApplePlatformMisc::DebugBreakInternal())
-#endif
+#define UE_DEBUG_BREAK_IMPL() PLATFORM_BREAK()
 
 #ifdef __OBJC__
 
@@ -76,24 +72,6 @@ struct CORE_API FApplePlatformMisc : public FGenericPlatformMisc
 
 		return ( Info.kp_proc.p_flag & P_TRACED ) != 0;
 	}
-	FORCEINLINE static void DebugBreakInternal()
-	{
-		if( IsDebuggerPresent() )
-		{
-			//Signal interupt to our process, this is caught by the main thread, this is not immediate but you can continue
-			//when triggered by check macros you will need to inspect other threads for the appFailAssert call.
-			//kill( getpid(), SIGINT );
-
-			//If you want an immediate break use the trap instruction, continued execuction is halted
-#if WITH_SIMULATOR || PLATFORM_MAC
-            __asm__ ( "int $3" );
-#elif PLATFORM_64BITS
-			__asm__ ( "svc 0" );
-#else
-            __asm__ ( "trap" );
-#endif
-		}
-	}
 #endif
 
 	FORCEINLINE static void MemoryBarrier()
diff --git a/Engine/Source/Runtime/Core/Public/GenericPlatform/GenericPlatformMisc.h b/Engine/Source/Runtime/Core/Public/GenericPlatform/GenericPlatformMisc.h
--- a/Engine/Source/Runtime/Core/Public/GenericPlatform/GenericPlatformMisc.h
+++ b/Engine/Source/Runtime/Core/Public/GenericPlatform/GenericPlatformMisc.h
@@ -18,6 +18,12 @@ struct FChunkTagID;
 template <typename FuncType>
 class TFunction;
 
+#if UE_BUILD_SHIPPING
+#define UE_DEBUG_BREAK() ((void)0)
+#else
+#define UE_DEBUG_BREAK() ((void)(FPlatformMisc::IsDebuggerPresent() && ([] () { UE_DEBUG_BREAK_IMPL(); } (), 1)))
+#endif
+
 namespace EBuildConfigurations
 {
 	/**
diff --git a/Engine/Source/Runtime/Core/Public/HTML5/HTML5Platform.h b/Engine/Source/Runtime/Core/Public/HTML5/HTML5Platform.h
--- a/Engine/Source/Runtime/Core/Public/HTML5/HTML5Platform.h
+++ b/Engine/Source/Runtime/Core/Public/HTML5/HTML5Platform.h
@@ -42,6 +42,16 @@ typedef FHTML5Types FPlatformTypes;
 #define PLATFORM_BUILTIN_VERTEX_HALF_FLOAT			0
 #define PLATFORM_SUPPORTS_STACK_SYMBOLS				1
 
+#define PLATFORM_BREAK() \
+	[] () -> void \
+	{ \
+			emscripten_log(255, "Debug breakpoint!"); \
+			EM_ASM( \
+				var callstack = new Error; \
+				throw callstack.stack; \
+			); \
+	} ()
+
 #if __has_feature(cxx_decltype_auto)
 	#define PLATFORM_COMPILER_HAS_DECLTYPE_AUTO 1
 #else
diff --git a/Engine/Source/Runtime/Core/Public/HTML5/HTML5PlatformMisc.h b/Engine/Source/Runtime/Core/Public/HTML5/HTML5PlatformMisc.h
--- a/Engine/Source/Runtime/Core/Public/HTML5/HTML5PlatformMisc.h
+++ b/Engine/Source/Runtime/Core/Public/HTML5/HTML5PlatformMisc.h
@@ -12,11 +12,7 @@
 #include "HTML5/HTML5SystemIncludes.h"
 #include <emscripten/emscripten.h>
 
-#if UE_BUILD_SHIPPING
-#define UE_DEBUG_BREAK() ((void)0)
-#else
-#define UE_DEBUG_BREAK() (FHTML5Misc::DebugBreakInternal())
-#endif
+#define UE_DEBUG_BREAK_IMPL() PLATFORM_BREAK()
 
 /**
  * HTML5 implementation of the misc OS functions
@@ -61,19 +57,6 @@ struct CORE_API FHTML5Misc : public FGenericPlatformMisc
 		return true;
 	}
 
-	/** Break into the debugger, if IsDebuggerPresent returns true, otherwise do nothing  */
-	FORCEINLINE static void DebugBreakInternal()
-	{
-		if (IsDebuggerPresent())
-		{
-			emscripten_log(255, "DebugBreak() called!");
-			EM_ASM(
-				var callstack = new Error;
-				throw callstack.stack;
-			);
-		}
-	}
-
 	static void LocalPrint(const TCHAR* Str);
 };
 
diff --git a/Engine/Source/Runtime/Core/Public/IOS/IOSPlatform.h b/Engine/Source/Runtime/Core/Public/IOS/IOSPlatform.h
--- a/Engine/Source/Runtime/Core/Public/IOS/IOSPlatform.h
+++ b/Engine/Source/Runtime/Core/Public/IOS/IOSPlatform.h
@@ -63,6 +63,14 @@ typedef FIOSPlatformTypes FPlatformTypes;
 #define PLATFORM_UI_NEEDS_TOOLTIPS						0
 #define PLATFORM_UI_NEEDS_FOCUS_OUTLINES				0
 
+#if WITH_SIMULATOR
+	#define PLATFORM_BREAK()							__asm__("int $3")
+#elif PLATFORM_64BITS
+	#define PLATFORM_BREAK()							__asm__("svc 0")
+#else
+	#define PLATFORM_BREAK()							__asm__("trap")
+#endif
+
 #if __has_feature(cxx_decltype_auto)
 	#define PLATFORM_COMPILER_HAS_DECLTYPE_AUTO 1
 #else
diff --git a/Engine/Source/Runtime/Core/Public/Mac/MacPlatform.h b/Engine/Source/Runtime/Core/Public/Mac/MacPlatform.h
--- a/Engine/Source/Runtime/Core/Public/Mac/MacPlatform.h
+++ b/Engine/Source/Runtime/Core/Public/Mac/MacPlatform.h
@@ -47,6 +47,8 @@ typedef FMacPlatformTypes FPlatformTypes;
 
 #define PLATFORM_ENABLE_POPCNT_INTRINSIC 1
 
+#define PLATFORM_BREAK()							__asm__("int $3")
+
 #if __has_feature(cxx_decltype_auto)
 	#define PLATFORM_COMPILER_HAS_DECLTYPE_AUTO 1
 #else
diff --git a/Engine/Source/Runtime/Core/Public/Unix/UnixPlatform.h b/Engine/Source/Runtime/Core/Public/Unix/UnixPlatform.h
--- a/Engine/Source/Runtime/Core/Public/Unix/UnixPlatform.h
+++ b/Engine/Source/Runtime/Core/Public/Unix/UnixPlatform.h
@@ -7,6 +7,7 @@
 #pragma once
 
 #include <linux/version.h>
+#include <signal.h>
 
 struct FGenericPlatformTypes;
 
@@ -57,6 +58,12 @@ typedef FUnixPlatformTypes FPlatformTypes;
 #define PLATFORM_HAS_BSD_SOCKET_FEATURE_MSG_DONTWAIT	1
 #define PLATFORM_SUPPORTS_STACK_SYMBOLS					1
 
+#if PLATFORM_CPU_X86_FAMILY
+	#define PLATFORM_BREAK()							__asm__ volatile("int $0x03")
+#else
+	#define PLATFORM_BREAK()							raise(SIGTRAP)
+#endif // PLATFORM_CPU_X86_FAMILY
+
 #define PLATFORM_ENABLE_POPCNT_INTRINSIC				1
 
 #if __has_feature(cxx_decltype_auto)
diff --git a/Engine/Source/Runtime/Core/Public/Unix/UnixPlatformMisc.h b/Engine/Source/Runtime/Core/Public/Unix/UnixPlatformMisc.h
--- a/Engine/Source/Runtime/Core/Public/Unix/UnixPlatformMisc.h
+++ b/Engine/Source/Runtime/Core/Public/Unix/UnixPlatformMisc.h
@@ -11,11 +11,7 @@
 #include "GenericPlatform/GenericPlatformMisc.h"
 #include "Misc/Build.h"
 
-#if UE_BUILD_SHIPPING
-#define UE_DEBUG_BREAK() ((void)0)
-#else
-#define UE_DEBUG_BREAK() (FUnixPlatformMisc::DebugBreakInternal())
-#endif
+#define UE_DEBUG_BREAK_IMPL() FUnixPlatformMisc::UngrabAllInput(); PLATFORM_BREAK()
 
 class Error;
 struct FGenericCrashContext;
@@ -42,7 +38,6 @@ struct CORE_API FUnixPlatformMisc : public FGenericPlatformMisc
 
 #if !UE_BUILD_SHIPPING
 	static bool IsDebuggerPresent();
-	static void DebugBreakInternal();
 #endif // !UE_BUILD_SHIPPING
 
 	static void LowLevelOutputDebugString(const TCHAR *Message);
diff --git a/Engine/Source/Runtime/Core/Public/Windows/WIndowsPlatform.h b/Engine/Source/Runtime/Core/Public/Windows/WIndowsPlatform.h
--- a/Engine/Source/Runtime/Core/Public/Windows/WIndowsPlatform.h
+++ b/Engine/Source/Runtime/Core/Public/Windows/WIndowsPlatform.h
@@ -60,6 +60,7 @@ typedef FWindowsPlatformTypes FPlatformTypes;
 #define PLATFORM_COMPILER_HAS_TCHAR_WMAIN					1
 #define PLATFORM_SUPPORTS_EARLY_MOVIE_PLAYBACK				(!WITH_EDITOR) // movies will start before engine is initalized
 #define PLATFORM_RHITHREAD_DEFAULT_BYPASS					0
+#define PLATFORM_BREAK()									__debugbreak()
 
 #define PLATFORM_SUPPORTS_STACK_SYMBOLS						1
 
diff --git a/Engine/Source/Runtime/Core/Public/Windows/WindowsPlatformMisc.h b/Engine/Source/Runtime/Core/Public/Windows/WindowsPlatformMisc.h
--- a/Engine/Source/Runtime/Core/Public/Windows/WindowsPlatformMisc.h
+++ b/Engine/Source/Runtime/Core/Public/Windows/WindowsPlatformMisc.h
@@ -6,11 +6,7 @@
 #include "HAL/PlatformMemory.h"
 #include "GenericPlatform/GenericPlatformMisc.h"
 
-#if UE_BUILD_SHIPPING
-#define UE_DEBUG_BREAK() ((void)0)
-#else
-#define UE_DEBUG_BREAK() ((void)(FWindowsPlatformMisc::IsDebuggerPresent() && (__debugbreak(), 1)))
-#endif
+#define UE_DEBUG_BREAK_IMPL() PLATFORM_BREAK()
 
 class GenericApplication;
 struct FGuid;
