Changed Ctrl-C process return code to match what Windows would set (extracted from kernel32!DefaultHandler).

diff --git a/Engine/Source/Runtime/Core/Private/Windows/WindowsPlatformMisc.cpp b/Engine/Source/Runtime/Core/Private/Windows/WindowsPlatformMisc.cpp
--- a/Engine/Source/Runtime/Core/Private/Windows/WindowsPlatformMisc.cpp
+++ b/Engine/Source/Runtime/Core/Private/Windows/WindowsPlatformMisc.cpp
@@ -595,8 +595,9 @@ static BOOL WINAPI ConsoleCtrlHandler( ::DWORD /*Type*/ )
 	// reliable. Deferring to the default CtrlHandler is also no good as that calls
 	// ExitProcess() which will terminate all threads and detach all DLLS. This can
 	// result in deadlocks and/or asserts as each DLL's atexit() is processed. So
-	// let's hard terminate the process.
-	TerminateProcess(GetCurrentProcess(), 1);
+	// let's hard terminate the process. 0xc000013a is what Windows' default handler
+	// would normally pass to ExitProcess().
+	TerminateProcess(GetCurrentProcess(), 0xc000013au);
 	return false;
 }
 
