Changed implementation of check/verify/ensure marcos to generate much less code on the hot path.

diff --git a/Engine/Source/Runtime/Core/Private/Misc/AssertionMacros.cpp b/Engine/Source/Runtime/Core/Private/Misc/AssertionMacros.cpp
--- a/Engine/Source/Runtime/Core/Private/Misc/AssertionMacros.cpp
+++ b/Engine/Source/Runtime/Core/Private/Misc/AssertionMacros.cpp
@@ -218,12 +218,20 @@ void FDebug::LogFormattedMessageWithCallstack(const FName& LogName, const ANSICH
 //
 
 void FDebug::LogAssertFailedMessageImpl(const ANSICHAR* Expr, const ANSICHAR* File, int32 Line, const TCHAR* Fmt, ...)
+{
+	va_list Args;
+	va_start(Args, Fmt);
+	LogAssertFailedMessageImplV(Expr, File, Line, Fmt, Args);
+	va_end(Args);
+}
+
+void FDebug::LogAssertFailedMessageImplV(const ANSICHAR* Expr, const ANSICHAR* File, int32 Line, const TCHAR* Fmt, va_list Args)
 {
 	// Ignore this assert if we're already forcibly shutting down because of a critical error.
 	if( !GIsCriticalError )
 	{
 		TCHAR DescriptionString[4096];
-		GET_VARARGS( DescriptionString, ARRAY_COUNT( DescriptionString ), ARRAY_COUNT( DescriptionString ) - 1, Fmt, Fmt );
+		FCString::GetVarArgs( DescriptionString, ARRAY_COUNT(DescriptionString), Fmt, Args );
 
 		TCHAR ErrorString[MAX_SPRINTF];
 		FCString::Sprintf( ErrorString, TEXT( "Assertion failed: %s" ), ANSI_TO_TCHAR( Expr ) );
@@ -399,6 +407,25 @@ void FDebug::EnsureFailed(const ANSICHAR* Expr, const ANSICHAR* File, int32 Line
 	FPlatformAtomics::InterlockedDecrement(&ActiveEnsureCount);
 }
 
+void FORCENOINLINE FDebug::CheckVerifyFailedImpl(
+	const ANSICHAR* Expr,
+	const ANSICHAR* File,
+	const int Line,
+	const TCHAR* Format,
+	...)
+{
+	va_list Args;
+	va_start(Args, Format);
+	FDebug::LogAssertFailedMessageImplV(Expr, File, Line, Format, Args);
+	va_end(Args);
+
+	if (!FPlatformMisc::IsDebuggerPresent())
+	{
+		FPlatformMisc::PromptForRemoteDebugging(false);
+		FDebug::AssertFailed(Expr, File, Line);
+	}
+}
+
 #endif // DO_CHECK || DO_GUARD_SLOW
 
 void VARARGS FDebug::AssertFailed(const ANSICHAR* Expr, const ANSICHAR* File, int32 Line, const TCHAR* Format/* = TEXT("")*/, ...)
diff --git a/Engine/Source/Runtime/Core/Public/Misc/AssertionMacros.h b/Engine/Source/Runtime/Core/Public/Misc/AssertionMacros.h
--- a/Engine/Source/Runtime/Core/Public/Misc/AssertionMacros.h
+++ b/Engine/Source/Runtime/Core/Public/Misc/AssertionMacros.h
@@ -41,7 +41,9 @@ struct CORE_API FDebug
 
 #if DO_CHECK || DO_GUARD_SLOW
 private:
+	static void VARARGS CheckVerifyFailedImpl(const ANSICHAR* Expr, const char* File, int32 Line, const TCHAR* Format, ...);
 	static void VARARGS LogAssertFailedMessageImpl(const ANSICHAR* Expr, const ANSICHAR* File, int32 Line, const TCHAR* Fmt, ...);
+	static void LogAssertFailedMessageImplV(const ANSICHAR* Expr, const ANSICHAR* File, int32 Line, const TCHAR* Fmt, va_list Args);
 
 public:
 	/** Failed assertion handler.  Warning: May be called at library startup time. */
@@ -53,6 +55,12 @@ public:
 
 		LogAssertFailedMessageImpl(Expr, File, Line, Fmt, Args...);
 	}
+
+	/**
+	 * Called when a 'check/verify' assertion fails.
+	 */
+	template <typename FmtType, typename... Types>
+	static void CheckVerifyFailed(const ANSICHAR* Expr, const char* File, int32 Line, const FmtType& Format, Types... Args);
 	
 	/**
 	 * Called when an 'ensure' assertion fails; gathers stack data and generates and error report.
@@ -131,24 +139,58 @@ public:
 // "verify" expressions are always evaluated, but only cause an error if enabled.
 //
 
+#if DO_CHECK || DO_GUARD_SLOW
+	template <typename FmtType, typename... Types>
+	void FORCENOINLINE FDebug::CheckVerifyFailed(
+		const ANSICHAR* Expr,
+		const ANSICHAR* File,
+		const int Line,
+		const FmtType& Format,
+		Types... Args)
+	{
+		static_assert(TIsArrayOrRefOfType<FmtType, TCHAR>::Value, "Formatting string must be a TCHAR array.");
+		static_assert(TAnd<TIsValidVariadicFunctionArg<Types>...>::Value, "Invalid argument(s) passed to CheckVerifyFailed()");
+		return CheckVerifyFailedImpl(Expr, File, Line, Format, Args...);
+	}
+
+	// MSVC (v19.00.24215.1 at time of writing) ignores no-inline attributes on
+	// lambdas. This can be worked around by calling the lambda from inside this
+	// templated (and correctly non-inlined) function.
+	template <typename RetType=void, class InnerType>
+	RetType FORCENOINLINE DispatchCheckVerify(InnerType&& Inner)
+	{
+		return Inner();
+	}
+#endif
+
 #if !UE_BUILD_SHIPPING
 #define _DebugBreakAndPromptForRemote() \
 	if (!FPlatformMisc::IsDebuggerPresent()) { FPlatformMisc::PromptForRemoteDebugging(false); } UE_DEBUG_BREAK();
 #else
 	#define _DebugBreakAndPromptForRemote()
 #endif // !UE_BUILD_SHIPPING
+
 #if DO_CHECK
 	#define checkCode( Code )		do { Code; } while ( false );
 	#define verify(expr)			UE_CHECK_IMPL(expr)
 	#define check(expr)				UE_CHECK_IMPL(expr)
 
+	// Technically we could use just the _F version (lambda-based) for asserts
+	// both with and without formatted messages. However MSVC emits extra
+	// unnecessary instructions when using a lambda; hence the Exec() impl.
 	#define UE_CHECK_IMPL(expr) \
 		{ \
-			if (UNLIKELY(!(expr))) \
+			if(UNLIKELY(!(expr))) \
 			{ \
-				FDebug::LogAssertFailedMessage( #expr, __FILE__, __LINE__, TEXT("") ); \
-				_DebugBreakAndPromptForRemote(); \
-				FDebug::AssertFailed( #expr, __FILE__, __LINE__ ); \
+				struct Impl \
+				{ \
+					static void FORCENOINLINE Exec() \
+					{ \
+						FDebug::CheckVerifyFailed(#expr, __FILE__, __LINE__, TEXT("")); \
+					} \
+				}; \
+				Impl::Exec(); \
+				PLATFORM_BREAK(); \
 				CA_ASSUME(false); \
 			} \
 		}
@@ -162,14 +204,17 @@ public:
 
 	#define UE_CHECK_F_IMPL(expr, format, ...) \
 		{ \
-			if (UNLIKELY(!(expr))) \
+			if(UNLIKELY(!(expr))) \
 			{ \
-				FDebug::LogAssertFailedMessage( #expr, __FILE__, __LINE__, format, ##__VA_ARGS__ ); \
-				_DebugBreakAndPromptForRemote(); \
-				FDebug::AssertFailed(#expr, __FILE__, __LINE__, format, ##__VA_ARGS__); \
+				DispatchCheckVerify([&] () FORCENOINLINE \
+				{ \
+					FDebug::CheckVerifyFailed(#expr, __FILE__, __LINE__, format, ##__VA_ARGS__); \
+				}); \
+				PLATFORM_BREAK(); \
 				CA_ASSUME(false); \
 			} \
 		}
+
 	/**
 	 * Denotes code paths that should never be reached.
 	 */
@@ -273,15 +318,27 @@ public:
 		#define UE_ENSURE_BREAK_IMPL() ((FPlatformMisc::IsDebuggerPresent() || (FPlatformMisc::PromptForRemoteDebugging(true), true)), UE_DEBUG_BREAK(), false)
 	#endif
 
-	#define UE_ENSURE_IMPL(Always, InExpression, InFormat, ...) \
-		(LIKELY(!!(InExpression)) \
-		|| FDebug::OptionallyLogFormattedEnsureMessageReturningFalse(Always, #InExpression, __FILE__, __LINE__, InFormat, ##__VA_ARGS__) \
-		|| (Always && UE_ENSURE_BREAK_IMPL()))
+	#define UE_ENSURE_IMPL(Capture, Always, InExpression, ...) \
+		(LIKELY(!!(InExpression)) || DispatchCheckVerify<bool>([Capture] () FORCENOINLINE \
+		{ \
+			static bool bExecuted = false; \
+			if (!bExecuted || Always) \
+			{ \
+				bExecuted = true; \
+				FDebug::OptionallyLogFormattedEnsureMessageReturningFalse(true, #InExpression, __FILE__, __LINE__, ##__VA_ARGS__); \
+				if (!FPlatformMisc::IsDebuggerPresent()) \
+				{ \
+					FPlatformMisc::PromptForRemoteDebugging(true); \
+					return true; \
+				} \
+			} \
+			return !Always; \
+		}) || ([] () { PLATFORM_BREAK(); } (), true))
 
-	#define ensure(           InExpression                ) UE_ENSURE_IMPL(UE4Asserts_Private::TrueOnFirstCallOnly([]{}), InExpression, TEXT(""))
-	#define ensureMsgf(       InExpression, InFormat, ... ) UE_ENSURE_IMPL(UE4Asserts_Private::TrueOnFirstCallOnly([]{}), InExpression, InFormat, ##__VA_ARGS__)
-	#define ensureAlways(     InExpression                ) UE_ENSURE_IMPL(true,                                          InExpression, TEXT(""))
-	#define ensureAlwaysMsgf( InExpression, InFormat, ... ) UE_ENSURE_IMPL(true,                                          InExpression, InFormat, ##__VA_ARGS__)
+	#define ensure(           InExpression                ) UE_ENSURE_IMPL( , false, InExpression, TEXT(""))
+	#define ensureMsgf(       InExpression, InFormat, ... ) UE_ENSURE_IMPL(&, false, InExpression, InFormat, ##__VA_ARGS__)
+	#define ensureAlways(     InExpression                ) UE_ENSURE_IMPL( , true,  InExpression, TEXT(""))
+	#define ensureAlwaysMsgf( InExpression, InFormat, ... ) UE_ENSURE_IMPL(&, true,  InExpression, InFormat, ##__VA_ARGS__)
 
 #else	// DO_CHECK
 
