Moved over more assert macros to use the newer check/verify implementations, prefering the check() version as that generates better code when that the formatted print version.

diff --git a/Engine/Source/Runtime/Core/Public/Misc/AssertionMacros.h b/Engine/Source/Runtime/Core/Public/Misc/AssertionMacros.h
--- a/Engine/Source/Runtime/Core/Public/Misc/AssertionMacros.h
+++ b/Engine/Source/Runtime/Core/Public/Misc/AssertionMacros.h
@@ -218,13 +218,13 @@ public:
 	/**
 	 * Denotes code paths that should never be reached.
 	 */
-	#define checkNoEntry()       { FDebug::LogAssertFailedMessage( "Enclosing block should never be called", __FILE__, __LINE__, TEXT("") ); _DebugBreakAndPromptForRemote(); FDebug::AssertFailed("Enclosing block should never be called", __FILE__, __LINE__ ); CA_ASSUME(false); }
+	#define checkNoEntry()       check(!"Enclosing block should never be called")
 
 	/**
 	 * Denotes code paths that should not be executed more than once.
 	 */
 	#define checkNoReentry()     { static bool s_beenHere##__LINE__ = false;                                         \
-	                               checkf( !s_beenHere##__LINE__, TEXT("Enclosing block was called more than once") );   \
+	                               check( !"Enclosing block was called more than once" || !s_beenHere##__LINE__ );   \
 								   s_beenHere##__LINE__ = true; }
 
 	class FRecursionScopeMarker
@@ -240,10 +240,10 @@ public:
 	 * Denotes code paths that should never be called recursively.
 	 */
 	#define checkNoRecursion()  static uint16 RecursionCounter##__LINE__ = 0;                                            \
-	                            checkf( RecursionCounter##__LINE__ == 0, TEXT("Enclosing block was entered recursively") );  \
+	                            check( !"Enclosing block was entered recursively" || RecursionCounter##__LINE__ == 0 );  \
 	                            const FRecursionScopeMarker ScopeMarker##__LINE__( RecursionCounter##__LINE__ )
 
-	#define unimplemented()       { FDebug::LogAssertFailedMessage( "Unimplemented function called", __FILE__, __LINE__, TEXT("") ); _DebugBreakAndPromptForRemote(); FDebug::AssertFailed("Unimplemented function called", __FILE__, __LINE__); CA_ASSUME(false); }
+	#define unimplemented()		check(!"Unimplemented function called")
 
 #else
 	#define checkCode(...)
