Removed unused code.

diff --git a/Engine/Source/Runtime/Core/Private/Misc/AssertionMacros.cpp b/Engine/Source/Runtime/Core/Private/Misc/AssertionMacros.cpp
--- a/Engine/Source/Runtime/Core/Private/Misc/AssertionMacros.cpp
+++ b/Engine/Source/Runtime/Core/Private/Misc/AssertionMacros.cpp
@@ -260,7 +260,7 @@ void FDebug::EnsureFailed(const ANSICHAR* Expr, const ANSICHAR* File, int32 Line
 	if( bShouldCrash )
 	{
 		// Just trigger a regular assertion which will crash via GError->Logf()
-		FDebug::LogAssertFailedMessage( Expr, File, Line, TEXT("%s"), Msg );
+		FDebug::LogAssertFailedMessageImpl( Expr, File, Line, TEXT("%s"), Msg );
 		return;
 	}
 
diff --git a/Engine/Source/Runtime/Core/Public/Logging/LogMacros.h b/Engine/Source/Runtime/Core/Public/Logging/LogMacros.h
--- a/Engine/Source/Runtime/Core/Public/Logging/LogMacros.h
+++ b/Engine/Source/Runtime/Core/Public/Logging/LogMacros.h
@@ -371,9 +371,11 @@ namespace UE4Asserts_Private
 #define logOrEnsureNanError(_FormatString_, ...) \
 	if (!GEnsureOnNANDiagnostic)\
 	{\
-		if (UE4Asserts_Private::TrueOnFirstCallOnly([]{}))\
+		static bool OnceOnly = false;\
+		if (!OnceOnly)\
 		{\
 			UE4Asserts_Private::InternalLogNANDiagnosticMessage(_FormatString_, ##__VA_ARGS__); \
+			OnceOnly = true;\
 		}\
 	}\
 	else\
diff --git a/Engine/Source/Runtime/Core/Public/Misc/AssertionMacros.h b/Engine/Source/Runtime/Core/Public/Misc/AssertionMacros.h
--- a/Engine/Source/Runtime/Core/Public/Misc/AssertionMacros.h
+++ b/Engine/Source/Runtime/Core/Public/Misc/AssertionMacros.h
@@ -46,16 +46,6 @@ private:
 	static void LogAssertFailedMessageImplV(const ANSICHAR* Expr, const ANSICHAR* File, int32 Line, const TCHAR* Fmt, va_list Args);
 
 public:
-	/** Failed assertion handler.  Warning: May be called at library startup time. */
-	template <typename FmtType, typename... Types>
-	static void LogAssertFailedMessage(const ANSICHAR* Expr, const ANSICHAR* File, int32 Line, const FmtType& Fmt, Types... Args)
-	{
-		static_assert(TIsArrayOrRefOfType<FmtType, TCHAR>::Value, "Formatting string must be a TCHAR array.");
-		static_assert(TAnd<TIsValidVariadicFunctionArg<Types>...>::Value, "Invalid argument(s) passed to FDebug::LogAssertFailedMessage");
-
-		LogAssertFailedMessageImpl(Expr, File, Line, Fmt, Args...);
-	}
-
 	/**
 	 * Called when a 'check/verify' assertion fails.
 	 */
@@ -298,19 +288,6 @@ public:
 
 #if DO_CHECK && !USING_CODE_ANALYSIS // The Visual Studio 2013 analyzer doesn't understand these complex conditionals
 
-	namespace UE4Asserts_Private
-	{
-		// This is used by ensure to generate a bool per instance
-		// by passing a lambda which will uniquely instantiate the template.
-		template <typename Type>
-		bool TrueOnFirstCallOnly(const Type&)
-		{
-			static bool bValue = true;
-			bool Result = bValue;
-			bValue = false;
-			return Result;
-		}
-	}
 
 	#if UE_BUILD_SHIPPING
 		#define UE_ENSURE_BREAK_IMPL() false
