Removed another bit of unused code.

diff --git a/Engine/Source/Runtime/Core/Public/Misc/AssertionMacros.h b/Engine/Source/Runtime/Core/Public/Misc/AssertionMacros.h
--- a/Engine/Source/Runtime/Core/Public/Misc/AssertionMacros.h
+++ b/Engine/Source/Runtime/Core/Public/Misc/AssertionMacros.h
@@ -288,13 +288,6 @@ public:
 
 #if DO_CHECK && !USING_CODE_ANALYSIS // The Visual Studio 2013 analyzer doesn't understand these complex conditionals
 
-
-	#if UE_BUILD_SHIPPING
-		#define UE_ENSURE_BREAK_IMPL() false
-	#else
-		#define UE_ENSURE_BREAK_IMPL() ((FPlatformMisc::IsDebuggerPresent() || (FPlatformMisc::PromptForRemoteDebugging(true), true)), UE_DEBUG_BREAK(), false)
-	#endif
-
 	#define UE_ENSURE_IMPL(Capture, Always, InExpression, ...) \
 		(LIKELY(!!(InExpression)) || DispatchCheckVerify<bool>([Capture] () FORCENOINLINE \
 		{ \
