Missing "#endif".

diff --git a/Engine/Source/Runtime/Core/Public/Android/AndroidMisc.h b/Engine/Source/Runtime/Core/Public/Android/AndroidMisc.h
--- a/Engine/Source/Runtime/Core/Public/Android/AndroidMisc.h
+++ b/Engine/Source/Runtime/Core/Public/Android/AndroidMisc.h
@@ -201,6 +201,7 @@ public:
 
 #if !UE_BUILD_SHIPPING
 	static bool IsDebuggerPresent();
+#endif
 
 	FORCEINLINE static void MemoryBarrier()
 	{
