Removed deprecated DebugBreak*() functions from Lumin platform.

diff --git a/Engine/Source/Runtime/Core/Public/Lumin/LuminPlatformMisc.h b/Engine/Source/Runtime/Core/Public/Lumin/LuminPlatformMisc.h
--- a/Engine/Source/Runtime/Core/Public/Lumin/LuminPlatformMisc.h
+++ b/Engine/Source/Runtime/Core/Public/Lumin/LuminPlatformMisc.h
@@ -28,28 +28,6 @@ struct CORE_API FLuminPlatformMisc : public FAndroidMisc
 	static bool IsDebuggerPresent();
 #endif
 
-	/** Break into debugger. Returning false allows this function to be used in conditionals. */
-	FORCEINLINE static bool DebugBreakReturningFalse()
-	{
-#if !UE_BUILD_SHIPPING
-		UE_DEBUG_BREAK();
-#endif
-		return false;
-	}
-
-	/** Prompts for remote debugging if debugger is not attached. Regardless of result, breaks into debugger afterwards. Returns false for use in conditionals. */
-	static FORCEINLINE bool DebugBreakAndPromptForRemoteReturningFalse(bool bIsEnsure = false)
-	{
-#if !UE_BUILD_SHIPPING
-		if (!IsDebuggerPresent())
-		{
-			PromptForRemoteDebugging(bIsEnsure);
-		}
-		UE_DEBUG_BREAK();
-#endif
-		return false;
-	}
-
 	FORCEINLINE static void MemoryBarrier()
 	{
 		__sync_synchronize();
