GitHub #5225 : Fix several issues with weak lambda delegates

diff --git a/Engine/Source/Runtime/Core/Public/Delegates/DelegateInstancesImpl.h b/Engine/Source/Runtime/Core/Public/Delegates/DelegateInstancesImpl.h
--- a/Engine/Source/Runtime/Core/Public/Delegates/DelegateInstancesImpl.h
+++ b/Engine/Source/Runtime/Core/Public/Delegates/DelegateInstancesImpl.h
@@ -973,7 +973,7 @@ private:
 	static_assert(TAreTypesEqual<FunctorType, typename TRemoveReference<FunctorType>::Type>::Value, "FunctorType cannot be a reference");
 
 	typedef IBaseDelegateInstance<typename TUnwrapType<WrappedRetValType>::Type(ParamTypes...)> Super;
-	typedef TWeakBaseFunctorDelegateInstance<RetValType(ParamTypes...), FunctorType, VarTypes...>   UnwrappedThisType;
+	typedef TWeakBaseFunctorDelegateInstance<UserClass, RetValType(ParamTypes...), FunctorType, VarTypes...>   UnwrappedThisType;
 
 public:
 	TWeakBaseFunctorDelegateInstance(UserClass* InContextObject, const FunctorType& InFunctor, VarTypes... Vars)
@@ -1008,12 +1008,22 @@ public:
 		return ContextObject.Get();
 	}
 
+	virtual const void* GetObjectForTimerManager() const override final
+	{
+		return ContextObject.Get();
+	}
+
 	// Deprecated
 	virtual bool HasSameObject(const void* InContextObject) const override final
 	{
 		return GetUObject() == InContextObject;
 	}
 
+	virtual bool IsCompactable() const override final
+	{
+		return !ContextObject.Get(true);
+	}
+
 	virtual bool IsSafeToExecute() const override final
 	{
 		return ContextObject.IsValid();
@@ -1072,7 +1082,7 @@ private:
 template <typename UserClass, typename FunctorType, typename... ParamTypes, typename... VarTypes>
 class TWeakBaseFunctorDelegateInstance<UserClass, void(ParamTypes...), FunctorType, VarTypes...> : public TWeakBaseFunctorDelegateInstance<UserClass, TTypeWrapper<void>(ParamTypes...), FunctorType, VarTypes...>
 {
-	typedef TWeakBaseFunctorDelegateInstance<TTypeWrapper<void>(ParamTypes...), FunctorType, VarTypes...> Super;
+	typedef TWeakBaseFunctorDelegateInstance<UserClass, TTypeWrapper<void>(ParamTypes...), FunctorType, VarTypes...> Super;
 
 public:
 	/**
