Fix for memory leaks in stats system.

diff --git a/Engine/Source/Runtime/Core/Private/Stats/Stats2.cpp b/Engine/Source/Runtime/Core/Private/Stats/Stats2.cpp
--- a/Engine/Source/Runtime/Core/Private/Stats/Stats2.cpp
+++ b/Engine/Source/Runtime/Core/Private/Stats/Stats2.cpp
@@ -392,10 +392,7 @@ class FStatGroupEnableManager : public IStatGroupEnableManager
 	FCriticalSection SynchronizationObject;
 
 	/** Pointer to the long name in the names block. */
-	TStatIdData* PendingStatIds;
-
-	/** Pending count of the name in the names block. */
-	int32 PendingCount;
+	TArray<TArray<TStatIdData>> StatIDs;
 
 	/** Holds the amount of memory allocated for the stats descriptions. */
 	FThreadSafeCounter MemoryCounter;
@@ -417,9 +414,7 @@ class FStatGroupEnableManager : public IStatGroupEnableManager
 
 public:
 	FStatGroupEnableManager()
-		: PendingStatIds(NULL)
-		, PendingCount(0)
-		, EnableForNewGroups(false)
+		: EnableForNewGroups(false)
 		, UseEnableForNewGroups(false)
 	{
 		check(IsInGameThread());
@@ -555,33 +550,42 @@ public:
 				Found->CurrentEnable = EnableForNewGroups;
 			}
 		}
-		if (PendingCount < 1)
+		if (StatIDs.Num() == 0 || StatIDs.Last().Num() == NUM_PER_BLOCK)
 		{
-			PendingStatIds = new TStatIdData[NUM_PER_BLOCK];
-			FMemory::Memzero( PendingStatIds, NUM_PER_BLOCK * sizeof( TStatIdData ) );
-			PendingCount = NUM_PER_BLOCK;
+			TArray<TStatIdData>& NewBlock = StatIDs.AddDefaulted_GetRef();
+			NewBlock.Reserve(NUM_PER_BLOCK);
 		}
-		--PendingCount;
-		TStatIdData* Result = PendingStatIds;
+		TStatIdData* Result = &StatIDs.Last().AddDefaulted_GetRef();
 
 		const FString StatDescription = InDescription ? InDescription : StatShortName.GetPlainNameString();
 
-		// Get the wide stat description.
 		const int32 StatDescLen = StatDescription.Len() + 1;
-		// We are leaking this. @see STAT_StatDescMemory
-		WIDECHAR* StatDescWide = new WIDECHAR[StatDescLen];
-		TCString<WIDECHAR>::Strcpy( StatDescWide, StatDescLen, StringCast<WIDECHAR>( *StatDescription ).Get() );
-		Result->WideString = reinterpret_cast<uint64>(StatDescWide);
+
+		int32 MemoryAllocated = 0;
+
+		// Get the wide stat description.
+		{
+			auto StatDescriptionWide = StringCast<WIDECHAR>(*StatDescription, StatDescLen);
+			int32 StatDescriptionWideLength = StatDescriptionWide.Length();
+			TUniquePtr<WIDECHAR[]> StatDescWide = MakeUnique<WIDECHAR[]>(StatDescriptionWideLength + 1);
+			TCString<WIDECHAR>::Strcpy(StatDescWide.Get(), StatDescriptionWideLength, StatDescriptionWide.Get());
+			Result->StatDescriptionWide = MoveTemp(StatDescWide);
+
+			MemoryAllocated += StatDescriptionWideLength * sizeof(WIDECHAR);
+		}
 
 		// Get the ansi stat description.
-		// We are leaking this. @see STAT_StatDescMemory
-		ANSICHAR* StatDescAnsi = new ANSICHAR[StatDescLen];
-		TCString<ANSICHAR>::Strcpy( StatDescAnsi, StatDescLen, StringCast<ANSICHAR>( *StatDescription ).Get() );
-		Result->AnsiString = reinterpret_cast<uint64>(StatDescAnsi);
+		{
+			auto StatDescriptionAnsi = StringCast<ANSICHAR>(*StatDescription, StatDescLen);
+			int32 StatDescriptionAnsiLength = StatDescriptionAnsi.Length();
+			TUniquePtr<ANSICHAR[]> StatDescAnsi = MakeUnique<ANSICHAR[]>(StatDescriptionAnsiLength + 1);
+			TCString<ANSICHAR>::Strcpy(StatDescAnsi.Get(), StatDescriptionAnsiLength, StatDescriptionAnsi.Get());
+			Result->StatDescriptionAnsi = MoveTemp(StatDescAnsi);
 
-		MemoryCounter.Add( StatDescLen*(sizeof( ANSICHAR ) + sizeof( WIDECHAR )) );
+			MemoryAllocated += StatDescriptionAnsiLength * sizeof(ANSICHAR);
+		}
 
-		++PendingStatIds;
+		MemoryCounter.Add( MemoryAllocated );
 
 		if( Found->CurrentEnable )
 		{
@@ -693,12 +697,8 @@ public:
 
 IStatGroupEnableManager& IStatGroupEnableManager::Get()
 {
-	static IStatGroupEnableManager* Singleton = NULL;
-	if (!Singleton)
-	{
-		verify(!FPlatformAtomics::InterlockedCompareExchangePointer((void**)&Singleton, new FStatGroupEnableManager, NULL));
-	}
-	return *Singleton;
+	static FStatGroupEnableManager Singleton;
+	return Singleton;
 }
 
 /*-----------------------------------------------------------------------------
diff --git a/Engine/Source/Runtime/Core/Public/Containers/LockFreeList.h b/Engine/Source/Runtime/Core/Public/Containers/LockFreeList.h
--- a/Engine/Source/Runtime/Core/Public/Containers/LockFreeList.h
+++ b/Engine/Source/Runtime/Core/Public/Containers/LockFreeList.h
@@ -10,6 +10,7 @@
 #include "Misc/NoopCounter.h"
 #include "Containers/ContainersFwd.h"
 #include "HAL/PlatformProcess.h"
+#include "Templates/AlignmentTemplates.h"
 
 CORE_API DECLARE_LOG_CATEGORY_EXTERN(LogLockFreeList, Log, All);
 
diff --git a/Engine/Source/Runtime/Core/Public/Stats/Stats2.h b/Engine/Source/Runtime/Core/Public/Stats/Stats2.h
--- a/Engine/Source/Runtime/Core/Public/Stats/Stats2.h
+++ b/Engine/Source/Runtime/Core/Public/Stats/Stats2.h
@@ -17,6 +17,7 @@
 #include "Templates/Atomic.h"
 #include "Math/Color.h"
 #include "StatsCommon.h"
+#include "Templates/UniquePtr.h"
 
 class FScopeCycleCounter;
 class FThreadStats;
@@ -104,12 +105,6 @@ struct CORE_API FStats
 
 struct TStatIdData
 {
-	FORCEINLINE TStatIdData()
-		: AnsiString(0)
-		, WideString(0)
-	{
-	}
-
 	bool IsNone() const
 	{
 		FMinimalName LocalName = Name.Load(EMemoryOrder::Relaxed);
@@ -119,11 +114,11 @@ struct TStatIdData
 	/** Name of the active stat; stored as a minimal name to minimize the data size */
 	TAtomic<FMinimalName> Name;
 
-	/** const ANSICHAR* pointer to a string; stored as a uint64 so it doesn't change size and affect TStatIdData alignment between 32 and 64-bit builds) */
-	uint64 AnsiString;
+	/** const ANSICHAR* pointer to a string describing the stat */
+	TUniquePtr<WIDECHAR[]> StatDescriptionWide;
 
-	/** const WIDECHAR* pointer to a string; stored as a uint64 so it doesn't change size and affect TStatIdData alignment between 32 and 64-bit builds) */
-	uint64 WideString;
+	/** const WIDECHAR* pointer to a string describing the stat */
+	TUniquePtr<ANSICHAR[]> StatDescriptionAnsi;
 };
 
 struct TStatId
@@ -171,7 +166,7 @@ struct TStatId
 	 */
 	FORCEINLINE const ANSICHAR* GetStatDescriptionANSI() const
 	{
-		return reinterpret_cast<const ANSICHAR*>(StatIdPtr->AnsiString);
+		return StatIdPtr->StatDescriptionAnsi.Get();
 	}
 
 	/**
@@ -181,7 +176,7 @@ struct TStatId
 	 */
 	FORCEINLINE const WIDECHAR* GetStatDescriptionWIDE() const
 	{
-		return reinterpret_cast<const WIDECHAR*>(StatIdPtr->WideString);
+		return StatIdPtr->StatDescriptionWide.Get();
 	}
 
 private:
