Fix for off-by-one error.

diff --git a/Engine/Source/Runtime/Core/Private/Stats/Stats2.cpp b/Engine/Source/Runtime/Core/Private/Stats/Stats2.cpp
--- a/Engine/Source/Runtime/Core/Private/Stats/Stats2.cpp
+++ b/Engine/Source/Runtime/Core/Private/Stats/Stats2.cpp
@@ -568,7 +568,7 @@ public:
 			auto StatDescriptionWide = StringCast<WIDECHAR>(*StatDescription, StatDescLen);
 			int32 StatDescriptionWideLength = StatDescriptionWide.Length();
 			TUniquePtr<WIDECHAR[]> StatDescWide = MakeUnique<WIDECHAR[]>(StatDescriptionWideLength + 1);
-			TCString<WIDECHAR>::Strcpy(StatDescWide.Get(), StatDescriptionWideLength, StatDescriptionWide.Get());
+			TCString<WIDECHAR>::Strcpy(StatDescWide.Get(), StatDescriptionWideLength + 1, StatDescriptionWide.Get());
 			Result->StatDescriptionWide = MoveTemp(StatDescWide);
 
 			MemoryAllocated += StatDescriptionWideLength * sizeof(WIDECHAR);
@@ -579,7 +579,7 @@ public:
 			auto StatDescriptionAnsi = StringCast<ANSICHAR>(*StatDescription, StatDescLen);
 			int32 StatDescriptionAnsiLength = StatDescriptionAnsi.Length();
 			TUniquePtr<ANSICHAR[]> StatDescAnsi = MakeUnique<ANSICHAR[]>(StatDescriptionAnsiLength + 1);
-			TCString<ANSICHAR>::Strcpy(StatDescAnsi.Get(), StatDescriptionAnsiLength, StatDescriptionAnsi.Get());
+			TCString<ANSICHAR>::Strcpy(StatDescAnsi.Get(), StatDescriptionAnsiLength + 1, StatDescriptionAnsi.Get());
 			Result->StatDescriptionAnsi = MoveTemp(StatDescAnsi);
 
 			MemoryAllocated += StatDescriptionAnsiLength * sizeof(ANSICHAR);
