Typo fix.

diff --git a/Engine/Source/Runtime/Core/Private/HAL/MallocBinned2.cpp b/Engine/Source/Runtime/Core/Private/HAL/MallocBinned2.cpp
--- a/Engine/Source/Runtime/Core/Private/HAL/MallocBinned2.cpp
+++ b/Engine/Source/Runtime/Core/Private/HAL/MallocBinned2.cpp
@@ -137,11 +137,11 @@ public:
 			UE_LOG(LogMemory, Fatal, TEXT("MallocBinned2 Corruption Canary was 0x%x, should be 0x%x"), int32(Canary), int32(ShouldBe));
 		}
 	}
-	void SetCanary(ECanary ShouldBe, bool bPreexisting, bool bGuarnteedToBeNew)
+	void SetCanary(ECanary ShouldBe, bool bPreexisting, bool bGuaranteedToBeNew)
 	{
 		if (bPreexisting)
 		{
-			if (bGuarnteedToBeNew)
+			if (bGuaranteedToBeNew)
 			{
 				UE_LOG(LogMemory, Fatal, TEXT("MallocBinned2 Corruption Canary was 0x%x, should be 0x%x. This block is both preexisting and guaranteed to be new; which makes no sense."), int32(Canary), int32(ShouldBe));
 			}
@@ -159,7 +159,7 @@ public:
 		}
 		else
 		{
-			if (bGuarnteedToBeNew)
+			if (bGuaranteedToBeNew)
 			{
 				if (Canary != ECanary::Unassigned)
 				{
