Allow MallocBinned2 to allocate up to 2^36-1 bytes without crashing.
Asserts added to validate that.
Fix for ReallocExternal() of an OS allocation not storing the aligned size, as MallocExternal() does.
Redundant checks moved to OS allocation block in ResizeExternal().
Asserts added to TArray::Reserve() and FString::Reserve() to ensure the requested sizes are legal.

diff --git a/Engine/Source/Runtime/Core/Private/Containers/String.cpp b/Engine/Source/Runtime/Core/Private/Containers/String.cpp
--- a/Engine/Source/Runtime/Core/Private/Containers/String.cpp
+++ b/Engine/Source/Runtime/Core/Private/Containers/String.cpp
@@ -422,7 +422,7 @@ void FString::PathAppend(const TCHAR* Str, int32 StrLength)
 			}
 		}
 
-		Data.Reserve(DataNum + StrLength + 1);
+		Reserve(DataNum + StrLength);
 		Data.Append(Str, StrLength);
 		Data.Add(TEXT('\0'));
 	}
diff --git a/Engine/Source/Runtime/Core/Private/HAL/MallocBinned2.cpp b/Engine/Source/Runtime/Core/Private/HAL/MallocBinned2.cpp
--- a/Engine/Source/Runtime/Core/Private/HAL/MallocBinned2.cpp
+++ b/Engine/Source/Runtime/Core/Private/HAL/MallocBinned2.cpp
@@ -101,6 +101,13 @@ FMallocBinned2::FPoolTable::FPoolTable()
 
 struct FMallocBinned2::FPoolInfo
 {
+	static FORCEINLINE bool IsSupportedSize(SIZE_T Size)
+	{
+		bool bResult = IsAligned(Size, BINNED2_MINIMUM_ALIGNMENT);
+		bResult = bResult && (Size >> BINNED2_MINIMUM_ALIGNMENT_SHIFT) <= SIZE_T(MAX_uint32);
+		return bResult;
+	}
+
 	enum class ECanary : uint16
 	{
 		Unassigned = 0x3941,
@@ -188,9 +195,9 @@ public:
 		return Result;
 	}
 
-	uint32 GetOSRequestedBytes() const
+	SIZE_T GetOSRequestedBytes() const
 	{
-		return AllocSize;
+		return SIZE_T(AllocSize) << BINNED2_MINIMUM_ALIGNMENT_SHIFT;
 	}
 
 	UPTRINT GetOsAllocatedBytes() const
@@ -199,13 +206,14 @@ public:
 		return (UPTRINT)FirstFreeBlock;
 	}
 
-	void SetOSAllocationSizes(uint32 InRequestedBytes, UPTRINT InAllocatedBytes)
+	void SetOSAllocationSizes(SIZE_T InRequestedBytes, UPTRINT InAllocatedBytes)
 	{
 		CheckCanary(ECanary::FirstFreeBlockIsOSAllocSize);
-		check(InRequestedBytes != 0);                // Shouldn't be pooling zero byte allocations
-		check(InAllocatedBytes >= InRequestedBytes); // We must be allocating at least as much as we requested
+		checkSlow(InRequestedBytes != 0);                // Shouldn't be pooling zero byte allocations
+		checkSlow(InAllocatedBytes >= InRequestedBytes); // We must be allocating at least as much as we requested
+		checkSlow(IsSupportedSize(InRequestedBytes));    // We must be allocating a size we can store
 
-		AllocSize      = InRequestedBytes;
+		AllocSize      = uint32(InRequestedBytes >> BINNED2_MINIMUM_ALIGNMENT_SHIFT);
 		FirstFreeBlock = (FFreeBlock*)InAllocatedBytes;
 	}
 
@@ -768,11 +776,13 @@ void* FMallocBinned2::MallocExternal(SIZE_T Size, uint32 Alignment)
 
 		return Result;
 	}
+
 	Alignment = FMath::Max<uint32>(Alignment, BINNED2_MINIMUM_ALIGNMENT);
 	Size = Align(FMath::Max((SIZE_T)1, Size), Alignment);
 
 	check(FMath::IsPowerOfTwo(Alignment));
 	check(Alignment <= PageSize);
+	check(FMallocBinned2::FPoolInfo::IsSupportedSize(Size));
 
 	FScopeLock Lock(&Mutex);
 
@@ -809,9 +819,8 @@ void* FMallocBinned2::ReallocExternal(void* Ptr, SIZE_T NewSize, uint32 Alignmen
 		FMallocBinned2::FreeExternal(Ptr);
 		return nullptr;
 	}
+
 	static_assert(DEFAULT_ALIGNMENT <= BINNED2_MINIMUM_ALIGNMENT, "DEFAULT_ALIGNMENT is assumed to be zero"); // used below
-	check(FMath::IsPowerOfTwo(Alignment));
-	check(Alignment <= PageSize);
 
 	if (!IsOSAllocation(Ptr))
 	{
@@ -849,7 +858,7 @@ void* FMallocBinned2::ReallocExternal(void* Ptr, SIZE_T NewSize, uint32 Alignmen
 		UE_LOG(LogMemory, Fatal, TEXT("FMallocBinned2 Attempt to realloc an unrecognized block %p"), Ptr);
 	}
 	UPTRINT PoolOsBytes = Pool->GetOsAllocatedBytes();
-	uint32 PoolOSRequestedBytes = Pool->GetOSRequestedBytes();
+	SIZE_T PoolOSRequestedBytes = Pool->GetOSRequestedBytes();
 	checkf(PoolOSRequestedBytes <= PoolOsBytes, TEXT("FMallocBinned2::ReallocExternal %d %d"), int32(PoolOSRequestedBytes), int32(PoolOsBytes));
 	if (NewSize > PoolOsBytes || // can't fit in the old block
 		(NewSize <= BINNED2_MAX_SMALL_POOL_SIZE && Alignment <= BINNED2_MINIMUM_ALIGNMENT) || // can switch to the small block allocator
@@ -857,16 +866,23 @@ void* FMallocBinned2::ReallocExternal(void* Ptr, SIZE_T NewSize, uint32 Alignmen
 	{
 		// Grow or shrink.
 		void* Result = FMallocBinned2::MallocExternal(NewSize, Alignment);
-		FMemory::Memcpy(Result, Ptr, FMath::Min<SIZE_T>(NewSize, PoolOSRequestedBytes));
+		FMemory::Memcpy(Result, Ptr, FMath::Min(NewSize, PoolOSRequestedBytes));
 		FMallocBinned2::FreeExternal(Ptr);
 		return Result;
 	}
 
+	Alignment = FMath::Max<uint32>(Alignment, BINNED2_MINIMUM_ALIGNMENT);
+	NewSize = Align(FMath::Max((SIZE_T)1, NewSize), Alignment);
+
+	check(FMath::IsPowerOfTwo(Alignment));
+	check(Alignment <= PageSize);
+	check(FMallocBinned2::FPoolInfo::IsSupportedSize(NewSize));
+
 #if BINNED2_ALLOCATOR_STATS
-	AllocatedLargePoolMemory += ((int64)NewSize) - ((int64)Pool->GetOSRequestedBytes());
+	AllocatedLargePoolMemory += ((int64)NewSize) - ((int64)PoolOSRequestedBytes);
 	// don't need to change the AllocatedLargePoolMemoryWAlignment because we didn't reallocate so it's the same size
 #endif
-	
+
 	Pool->SetOSAllocationSizes(NewSize, PoolOsBytes);
 
 	return Ptr;
@@ -921,7 +937,7 @@ void FMallocBinned2::FreeExternal(void* Ptr)
 			UE_LOG(LogMemory, Fatal, TEXT("FMallocBinned2 Attempt to free an unrecognized block %p"), Ptr);
 		}
 		UPTRINT PoolOsBytes = Pool->GetOsAllocatedBytes();
-		uint32 PoolOSRequestedBytes = Pool->GetOSRequestedBytes();
+		SIZE_T PoolOSRequestedBytes = Pool->GetOSRequestedBytes();
 
 #if BINNED2_ALLOCATOR_STATS
 		AllocatedLargePoolMemory -= ((int64)PoolOSRequestedBytes);
@@ -957,7 +973,7 @@ bool FMallocBinned2::GetAllocationSizeExternal(void* Ptr, SIZE_T& SizeOut)
 		UE_LOG(LogMemory, Fatal, TEXT("FMallocBinned2 Attempt to GetAllocationSizeExternal an unrecognized block %p"), Ptr);
 	}
 	UPTRINT PoolOsBytes = Pool->GetOsAllocatedBytes();
-	uint32 PoolOSRequestedBytes = Pool->GetOSRequestedBytes();
+	SIZE_T PoolOSRequestedBytes = Pool->GetOSRequestedBytes();
 	checkf(PoolOSRequestedBytes <= PoolOsBytes, TEXT("FMallocBinned2::GetAllocationSizeExternal %d %d"), int32(PoolOSRequestedBytes), int32(PoolOsBytes));
 	SizeOut = PoolOsBytes;
 	return true;
diff --git a/Engine/Source/Runtime/Core/Public/Containers/Array.h b/Engine/Source/Runtime/Core/Public/Containers/Array.h
--- a/Engine/Source/Runtime/Core/Public/Containers/Array.h
+++ b/Engine/Source/Runtime/Core/Public/Containers/Array.h
@@ -2063,6 +2063,7 @@ public:
 	 */
 	FORCEINLINE void Reserve(int32 Number)
 	{
+		checkSlow(Number >= 0);
 		if (Number > ArrayMax)
 		{
 			ResizeTo(Number);
diff --git a/Engine/Source/Runtime/Core/Public/Containers/UnrealString.h b/Engine/Source/Runtime/Core/Public/Containers/UnrealString.h
--- a/Engine/Source/Runtime/Core/Public/Containers/UnrealString.h
+++ b/Engine/Source/Runtime/Core/Public/Containers/UnrealString.h
@@ -1684,8 +1684,9 @@ public:
 	static FString FormatAsNumber( int32 InNumber );
 
 	// To allow more efficient memory handling, automatically adds one for the string termination.
-	FORCEINLINE void Reserve(const uint32 CharacterCount)
+	FORCEINLINE void Reserve(int32 CharacterCount)
 	{
+		checkSlow(CharacterCount >= 0 && CharacterCount < MAX_int32);
 		Data.Reserve(CharacterCount + 1);
 	}
 
