Added macros to decorate functions/methods with the executable section they should be linked into.

diff --git a/Engine/Source/Runtime/Core/Public/Android/AndroidPlatform.h b/Engine/Source/Runtime/Core/Public/Android/AndroidPlatform.h
--- a/Engine/Source/Runtime/Core/Public/Android/AndroidPlatform.h
+++ b/Engine/Source/Runtime/Core/Public/Android/AndroidPlatform.h
@@ -57,6 +57,8 @@ typedef FAndroidTypes FPlatformTypes;
 #define PLATFORM_UI_NEEDS_FOCUS_OUTLINES			0
 #define PLATFORM_SUPPORTS_EARLY_MOVIE_PLAYBACK		1 // movies will start before engine is initalized
 
+#define PLATFORM_CODE_SECTION(Name)					__attribute__((section(Name)))
+
 #if defined(EXPERIMENTAL_OPENGL_RHITHREAD) && EXPERIMENTAL_OPENGL_RHITHREAD
 	#define PLATFORM_RHITHREAD_DEFAULT_BYPASS			0
 #else
diff --git a/Engine/Source/Runtime/Core/Public/HAL/Platform.h b/Engine/Source/Runtime/Core/Public/HAL/Platform.h
--- a/Engine/Source/Runtime/Core/Public/HAL/Platform.h
+++ b/Engine/Source/Runtime/Core/Public/HAL/Platform.h
@@ -624,6 +624,11 @@
 	#define COMPILE_ERROR(x) GCC_DIAGNOSTIC_HELPER(GCC error x)
 #endif
 
+// Tells the compiler to put the decorated function in a certain section (aka. segment) of the executable.
+#ifndef PLATFORM_CODE_SECTION
+	#define PLATFORM_CODE_SECTION(Name)
+#endif
+
 // These have to be forced inline on some OSes so the dynamic loader will not
 // resolve to our allocators for the system libraries.
 #ifndef OPERATOR_NEW_INLINE
diff --git a/Engine/Source/Runtime/Core/Public/IOS/IOSPlatform.h b/Engine/Source/Runtime/Core/Public/IOS/IOSPlatform.h
--- a/Engine/Source/Runtime/Core/Public/IOS/IOSPlatform.h
+++ b/Engine/Source/Runtime/Core/Public/IOS/IOSPlatform.h
@@ -71,6 +71,8 @@ typedef FIOSPlatformTypes FPlatformTypes;
 	#define PLATFORM_BREAK()							__asm__("trap")
 #endif
 
+#define PLATFORM_CODE_SECTION(Name)						__attribute__((section("__TEXT,__" Name ",regular,pure_instructions")))
+
 #if __has_feature(cxx_decltype_auto)
 	#define PLATFORM_COMPILER_HAS_DECLTYPE_AUTO 1
 #else
diff --git a/Engine/Source/Runtime/Core/Public/Mac/MacPlatform.h b/Engine/Source/Runtime/Core/Public/Mac/MacPlatform.h
--- a/Engine/Source/Runtime/Core/Public/Mac/MacPlatform.h
+++ b/Engine/Source/Runtime/Core/Public/Mac/MacPlatform.h
@@ -49,6 +49,8 @@ typedef FMacPlatformTypes FPlatformTypes;
 
 #define PLATFORM_BREAK()							__asm__("int $3")
 
+#define PLATFORM_CODE_SECTION(Name)					__attribute__((section("__TEXT,__" Name ",regular,pure_instructions")))
+
 #if __has_feature(cxx_decltype_auto)
 	#define PLATFORM_COMPILER_HAS_DECLTYPE_AUTO 1
 #else
diff --git a/Engine/Source/Runtime/Core/Public/Unix/UnixPlatform.h b/Engine/Source/Runtime/Core/Public/Unix/UnixPlatform.h
--- a/Engine/Source/Runtime/Core/Public/Unix/UnixPlatform.h
+++ b/Engine/Source/Runtime/Core/Public/Unix/UnixPlatform.h
@@ -64,6 +64,8 @@ typedef FUnixPlatformTypes FPlatformTypes;
 	#define PLATFORM_BREAK()							raise(SIGTRAP)
 #endif // PLATFORM_CPU_X86_FAMILY
 
+#define PLATFORM_CODE_SECTION(Name)						__attribute__((section(Name)))
+
 #define PLATFORM_ENABLE_POPCNT_INTRINSIC				1
 
 #if __has_feature(cxx_decltype_auto)
diff --git a/Engine/Source/Runtime/Core/Public/Windows/WIndowsPlatform.h b/Engine/Source/Runtime/Core/Public/Windows/WIndowsPlatform.h
--- a/Engine/Source/Runtime/Core/Public/Windows/WIndowsPlatform.h
+++ b/Engine/Source/Runtime/Core/Public/Windows/WIndowsPlatform.h
@@ -103,6 +103,9 @@ typedef FWindowsPlatformTypes FPlatformTypes;
 	#define PRAGMA_ENABLE_OPTIMIZATION_ACTUAL  __pragma(clang optimize on)
 #endif
 
+// Tells the compiler to put the decorated function in a certain section (aka. segment) of the executable.
+#define PLATFORM_CODE_SECTION(Name) __declspec(code_seg(Name)) 
+
 // Backwater of the spec. All compilers support this except microsoft, and they will soon
 #if !defined(__clang__)		// Clang expects typename outside template
 	#define TYPENAME_OUTSIDE_TEMPLATE
