Moved assertion macros' implementations into a separate section of linked executables to keep the code well off the hot path.

diff --git a/Engine/Source/Runtime/Core/Public/Misc/AssertionMacros.h b/Engine/Source/Runtime/Core/Public/Misc/AssertionMacros.h
--- a/Engine/Source/Runtime/Core/Public/Misc/AssertionMacros.h
+++ b/Engine/Source/Runtime/Core/Public/Misc/AssertionMacros.h
@@ -131,8 +131,14 @@ public:
 //
 
 #if DO_CHECK || DO_GUARD_SLOW
+	// We'll put all assert implementation code into a separate section in the linked
+	// executable. This code should never execute so using a separate section keeps
+	// it well off the hot path and hopefully out of the instruction cache. It also
+	// facilitates reasoning about the makeup of a compiled/linked binary.
+	#define UE_DEBUG_SECTION PLATFORM_CODE_SECTION(".uedbg")
+
 	template <typename FmtType, typename... Types>
-	void FORCENOINLINE FDebug::CheckVerifyFailed(
+	void FORCENOINLINE UE_DEBUG_SECTION FDebug::CheckVerifyFailed(
 		const ANSICHAR* Expr,
 		const ANSICHAR* File,
 		const int Line,
@@ -148,7 +154,7 @@ public:
 	// lambdas. This can be worked around by calling the lambda from inside this
 	// templated (and correctly non-inlined) function.
 	template <typename RetType=void, class InnerType>
-	RetType FORCENOINLINE DispatchCheckVerify(InnerType&& Inner)
+	RetType FORCENOINLINE UE_DEBUG_SECTION DispatchCheckVerify(InnerType&& Inner)
 	{
 		return Inner();
 	}
@@ -175,7 +181,7 @@ public:
 			{ \
 				struct Impl \
 				{ \
-					static void FORCENOINLINE Exec() \
+					static void FORCENOINLINE UE_DEBUG_SECTION Exec() \
 					{ \
 						FDebug::CheckVerifyFailed(#expr, __FILE__, __LINE__, TEXT("")); \
 					} \
@@ -197,7 +203,7 @@ public:
 		{ \
 			if(UNLIKELY(!(expr))) \
 			{ \
-				DispatchCheckVerify([&] () FORCENOINLINE \
+				DispatchCheckVerify([&] () FORCENOINLINE UE_DEBUG_SECTION \
 				{ \
 					FDebug::CheckVerifyFailed(#expr, __FILE__, __LINE__, format, ##__VA_ARGS__); \
 				}); \
@@ -290,7 +296,7 @@ public:
 #if DO_CHECK && !USING_CODE_ANALYSIS // The Visual Studio 2013 analyzer doesn't understand these complex conditionals
 
 	#define UE_ENSURE_IMPL(Capture, Always, InExpression, ...) \
-		(LIKELY(!!(InExpression)) || DispatchCheckVerify<bool>([Capture] () FORCENOINLINE \
+		(LIKELY(!!(InExpression)) || DispatchCheckVerify<bool>([Capture] () FORCENOINLINE UE_DEBUG_SECTION \
 		{ \
 			static bool bExecuted = false; \
 			if (!bExecuted || Always) \
