Fixed UHT compile error.

diff --git a/Engine/Source/Runtime/Core/Public/Misc/AssertionMacros.h b/Engine/Source/Runtime/Core/Public/Misc/AssertionMacros.h
--- a/Engine/Source/Runtime/Core/Public/Misc/AssertionMacros.h
+++ b/Engine/Source/Runtime/Core/Public/Misc/AssertionMacros.h
@@ -10,6 +10,14 @@
 #include "Templates/IsValidVariadicFunctionArg.h"
 #include "Misc/VarArgs.h"
 
+#if DO_CHECK || DO_GUARD_SLOW
+	// We'll put all assert implementation code into a separate section in the linked
+	// executable. This code should never execute so using a separate section keeps
+	// it well off the hot path and hopefully out of the instruction cache. It also
+	// facilitates reasoning about the makeup of a compiled/linked binary.
+	#define UE_DEBUG_SECTION PLATFORM_CODE_SECTION(".uedbg")
+#endif // DO_CHECK || DO_GUARD_SLOW
+
 namespace ELogVerbosity
 {
 	enum Type : uint8;
@@ -50,7 +58,7 @@ public:
 	 * Called when a 'check/verify' assertion fails.
 	 */
 	template <typename FmtType, typename... Types>
-	static void CheckVerifyFailed(const ANSICHAR* Expr, const char* File, int32 Line, const FmtType& Format, Types... Args);
+	static void UE_DEBUG_SECTION CheckVerifyFailed(const ANSICHAR* Expr, const char* File, int32 Line, const FmtType& Format, Types... Args);
 	
 	/**
 	 * Called when an 'ensure' assertion fails; gathers stack data and generates and error report.
@@ -131,12 +139,6 @@ public:
 //
 
 #if DO_CHECK || DO_GUARD_SLOW
-	// We'll put all assert implementation code into a separate section in the linked
-	// executable. This code should never execute so using a separate section keeps
-	// it well off the hot path and hopefully out of the instruction cache. It also
-	// facilitates reasoning about the makeup of a compiled/linked binary.
-	#define UE_DEBUG_SECTION PLATFORM_CODE_SECTION(".uedbg")
-
 	template <typename FmtType, typename... Types>
 	void FORCENOINLINE UE_DEBUG_SECTION FDebug::CheckVerifyFailed(
 		const ANSICHAR* Expr,
