Reduce time locks are held during allocation and reallocation.

diff --git a/Engine/Source/Runtime/Core/Private/HAL/MallocBinned2.cpp b/Engine/Source/Runtime/Core/Private/HAL/MallocBinned2.cpp
--- a/Engine/Source/Runtime/Core/Private/HAL/MallocBinned2.cpp
+++ b/Engine/Source/Runtime/Core/Private/HAL/MallocBinned2.cpp
@@ -54,11 +54,11 @@ static FAutoConsoleVariableRef GMallocBinned2AllocExtraCVar(
 #endif
 
 #if BINNED2_ALLOCATOR_STATS
-int64 AllocatedSmallPoolMemory = 0; // memory that's requested to be allocated by the game
-int64 AllocatedOSSmallPoolMemory = 0;
+TAtomic<int64> AllocatedSmallPoolMemory(0); // memory that's requested to be allocated by the game
+TAtomic<int64> AllocatedOSSmallPoolMemory(0);
 
-int64 AllocatedLargePoolMemory = 0; // memory requests to the OS which don't fit in the small pool
-int64 AllocatedLargePoolMemoryWAlignment = 0; // when we allocate at OS level we need to align to a size
+TAtomic<int64> AllocatedLargePoolMemory(0); // memory requests to the OS which don't fit in the small pool
+TAtomic<int64> AllocatedLargePoolMemoryWAlignment(0); // when we allocate at OS level we need to align to a size
 #endif
 
 #if BINNED2_ALLOCATOR_STATS_VALIDATION
@@ -710,72 +710,70 @@ bool FMallocBinned2::IsInternallyThreadSafe() const
 	return true;
 }
 
-void* FMallocBinned2::MallocExternal(SIZE_T Size, uint32 Alignment)
+void* FMallocBinned2::MallocExternalSmall(SIZE_T Size, uint32 Alignment)
 {
-	static_assert(DEFAULT_ALIGNMENT <= BINNED2_MINIMUM_ALIGNMENT, "DEFAULT_ALIGNMENT is assumed to be zero"); // used below
-
-	// Only allocate from the small pools if the size is small enough and the alignment isn't crazy large.
-	// With large alignments, we'll waste a lot of memory allocating an entire page, but such alignments are highly unlikely in practice.
-	if ((Size <= BINNED2_MAX_SMALL_POOL_SIZE) & (Alignment <= BINNED2_MINIMUM_ALIGNMENT)) // one branch, not two
+	uint32 PoolIndex = BoundSizeToPoolIndex(Size);
+	FPerThreadFreeBlockLists* Lists = GMallocBinned2PerThreadCaches ? FPerThreadFreeBlockLists::Get() : nullptr;
+	if (Lists)
 	{
-		uint32 PoolIndex = BoundSizeToPoolIndex(Size);
-		FPerThreadFreeBlockLists* Lists = GMallocBinned2PerThreadCaches ? FPerThreadFreeBlockLists::Get() : nullptr;
-		if (Lists)
+		if (Lists->ObtainRecycledPartial(PoolIndex))
 		{
-			if (Lists->ObtainRecycledPartial(PoolIndex))
+			if (void* Result = Lists->Malloc(PoolIndex))
 			{
-				if (void* Result = Lists->Malloc(PoolIndex))
-				{
 #if BINNED2_ALLOCATOR_STATS
-					uint32 BlockSize = PoolIndexToBlockSize(PoolIndex);
-					Lists->AllocatedMemory += BlockSize;
+				uint32 BlockSize = PoolIndexToBlockSize(PoolIndex);
+				Lists->AllocatedMemory += BlockSize;
 #endif
-					return Result;
-				}
+				return Result;
 			}
 		}
+	}
 
-		FScopeLock Lock(&Mutex);
+	FScopeLock Lock(&Mutex);
 
-		// Allocate from small object pool.
-		FPoolTable& Table = SmallPoolTables[PoolIndex];
+	// Allocate from small object pool.
+	FPoolTable& Table = SmallPoolTables[PoolIndex];
 
-		FPoolInfo* Pool;
-		if (!Table.ActivePools.IsEmpty())
-		{
-			Pool = &Table.ActivePools.GetFrontPool();
-		}
-		else
-		{
-			Pool = &Table.ActivePools.PushNewPoolToFront(*this, Table.BlockSize, PoolIndex);
-		}
+	FPoolInfo* Pool;
+	if (!Table.ActivePools.IsEmpty())
+	{
+		Pool = &Table.ActivePools.GetFrontPool();
+	}
+	else
+	{
+		Pool = &Table.ActivePools.PushNewPoolToFront(*this, Table.BlockSize, PoolIndex);
+	}
 
-		void* Result = Pool->AllocateRegularBlock();
+	void* Result = Pool->AllocateRegularBlock();
 #if BINNED2_ALLOCATOR_STATS
-		AllocatedSmallPoolMemory += PoolIndexToBlockSize(PoolIndex);
+	AllocatedSmallPoolMemory += PoolIndexToBlockSize(PoolIndex);
 #endif // BINNED2_ALLOCATOR_STATS
-		if (GMallocBinned2AllocExtra)
+	if (GMallocBinned2AllocExtra)
+	{
+		if (Lists)
 		{
-			if (Lists)
+			// prefill the free list with some allocations so we are less likely to hit this slow path with the mutex 
+			for (int32 Index = 0; Index < GMallocBinned2AllocExtra && Pool->HasFreeRegularBlock(); Index++)
 			{
-				// prefill the free list with some allocations so we are less likely to hit this slow path with the mutex 
-				for (int32 Index = 0; Index < GMallocBinned2AllocExtra && Pool->HasFreeRegularBlock(); Index++)
+				if (!Lists->Free(Result, PoolIndex, Table.BlockSize))
 				{
-					if (!Lists->Free(Result, PoolIndex, Table.BlockSize))
-					{
-						break;
-					}
-					Result = Pool->AllocateRegularBlock();
+					break;
 				}
+				Result = Pool->AllocateRegularBlock();
 			}
 		}
-		if (!Pool->HasFreeRegularBlock())
-		{
-			Table.ExhaustedPools.LinkToFront(Pool);
-		}
-
-		return Result;
 	}
+	if (!Pool->HasFreeRegularBlock())
+	{
+		Table.ExhaustedPools.LinkToFront(Pool);
+	}
+
+	return Result;
+}
+
+void* FMallocBinned2::MallocExternalLarge(SIZE_T Size, uint32 Alignment)
+{
+	static_assert(DEFAULT_ALIGNMENT <= BINNED2_MINIMUM_ALIGNMENT, "DEFAULT_ALIGNMENT is assumed to be zero"); // used below
 
 	Alignment = FMath::Max<uint32>(Alignment, BINNED2_MINIMUM_ALIGNMENT);
 	Size = Align(FMath::Max((SIZE_T)1, Size), Alignment);
@@ -784,18 +782,24 @@ void* FMallocBinned2::MallocExternal(SIZE_T Size, uint32 Alignment)
 	check(Alignment <= PageSize);
 	check(FMallocBinned2::FPoolInfo::IsSupportedSize(Size));
 
-	FScopeLock Lock(&Mutex);
-
-	// Use OS for non-pooled allocations.
 	UPTRINT AlignedSize = Align(Size, OsAllocationGranularity);
-	void* Result = CachedOSPageAllocator.Allocate(AlignedSize);
-
-	UE_CLOG(!IsAligned(Result, Alignment) ,LogMemory, Fatal, TEXT("FMallocBinned2 alignment was too large for OS. Alignment=%d   Ptr=%p"), Alignment, Result);
 
-	if (!Result)
+	FPoolInfo* Pool;
+	void*      Result;
 	{
-		Private::OutOfMemory(AlignedSize);
+		FScopeLock Lock(&Mutex);
+
+		// Use OS for non-pooled allocations.
+		Result = CachedOSPageAllocator.Allocate(AlignedSize);
+		if (!Result)
+		{
+			Private::OutOfMemory(AlignedSize);
+		}
+
+		Pool = Private::GetOrCreatePoolInfo(*this, Result, FPoolInfo::ECanary::FirstFreeBlockIsOSAllocSize, false);
 	}
+
+	UE_CLOG(!IsAligned(Result, Alignment) ,LogMemory, Fatal, TEXT("FMallocBinned2 alignment was too large for OS. Alignment=%d   Ptr=%p"), Alignment, Result);
 	check(IsAligned(Result, PageSize) && IsOSAllocation(Result));
 
 #if BINNED2_ALLOCATOR_STATS
@@ -804,14 +808,12 @@ void* FMallocBinned2::MallocExternal(SIZE_T Size, uint32 Alignment)
 #endif
 
 	// Create pool.
-	FPoolInfo* Pool = Private::GetOrCreatePoolInfo(*this, Result, FPoolInfo::ECanary::FirstFreeBlockIsOSAllocSize, false);
 	check(Size > 0 && Size <= AlignedSize && AlignedSize >= OsAllocationGranularity);
 	Pool->SetOSAllocationSizes(Size, AlignedSize);
 
 	return Result;
 }
 
-
 void* FMallocBinned2::ReallocExternal(void* Ptr, SIZE_T NewSize, uint32 Alignment)
 {
 	if (NewSize == 0)
@@ -838,39 +840,52 @@ void* FMallocBinned2::ReallocExternal(void* Ptr, SIZE_T NewSize, uint32 Alignmen
 		}
 
 		// Reallocate and copy the data across
-		void* Result = FMallocBinned2::MallocExternal(NewSize, Alignment);
+		void* Result = FMallocBinned2::MallocSelect(NewSize, Alignment);
 		FMemory::Memcpy(Result, Ptr, FMath::Min<SIZE_T>(NewSize, BlockSize));
 		FMallocBinned2::FreeExternal(Ptr);
 		return Result;
 	}
 	if (!Ptr)
 	{
-		void* Result = FMallocBinned2::MallocExternal(NewSize, Alignment);
+		void* Result = FMallocBinned2::MallocSelect(NewSize, Alignment);
 		return Result;
 	}
 
-	FScopeLock Lock(&Mutex);
-
 	// Allocated from OS.
+	Mutex.Lock();
 	FPoolInfo* Pool = Private::FindPoolInfo(*this, Ptr);
-	if (!Pool)
-	{
-		UE_LOG(LogMemory, Fatal, TEXT("FMallocBinned2 Attempt to realloc an unrecognized block %p"), Ptr);
-	}
+	UE_CLOG(!Pool, LogMemory, Fatal, TEXT("FMallocBinned2 Attempt to realloc an unrecognized block %p"), Ptr);
+
 	UPTRINT PoolOsBytes = Pool->GetOsAllocatedBytes();
 	SIZE_T PoolOSRequestedBytes = Pool->GetOSRequestedBytes();
 	checkf(PoolOSRequestedBytes <= PoolOsBytes, TEXT("FMallocBinned2::ReallocExternal %d %d"), int32(PoolOSRequestedBytes), int32(PoolOsBytes));
+	bool bUseSmallMalloc = UseSmallAlloc(NewSize, Alignment);
 	if (NewSize > PoolOsBytes || // can't fit in the old block
-		(NewSize <= BINNED2_MAX_SMALL_POOL_SIZE && Alignment <= BINNED2_MINIMUM_ALIGNMENT) || // can switch to the small block allocator
+		bUseSmallMalloc || // can switch to the small block allocator
 		Align(NewSize, OsAllocationGranularity) < PoolOsBytes) // we can get some pages back
 	{
 		// Grow or shrink.
-		void* Result = FMallocBinned2::MallocExternal(NewSize, Alignment);
+		void* Result;
+		if (bUseSmallMalloc)
+		{
+			// Unlock before a small alloc, which rarely takes a lock
+			Mutex.Unlock();
+			Result = MallocExternalSmall(NewSize, Alignment);
+		}
+		else
+		{
+			// Unlock after a large alloc, which does take a lock, to save unlocking and re-locking unnecessarily
+			Result = MallocExternalLarge(NewSize, Alignment);
+			Mutex.Unlock();
+		}
+
 		FMemory::Memcpy(Result, Ptr, FMath::Min(NewSize, PoolOSRequestedBytes));
 		FMallocBinned2::FreeExternal(Ptr);
 		return Result;
 	}
 
+	Mutex.Unlock();
+
 	Alignment = FMath::Max<uint32>(Alignment, BINNED2_MINIMUM_ALIGNMENT);
 	NewSize = Align(FMath::Max((SIZE_T)1, NewSize), Alignment);
 
@@ -966,8 +981,13 @@ bool FMallocBinned2::GetAllocationSizeExternal(void* Ptr, SIZE_T& SizeOut)
 	{
 		return false;
 	}
-	FScopeLock Lock(&Mutex);
-	FPoolInfo* Pool = Private::FindPoolInfo(*this, Ptr);
+
+	FPoolInfo* Pool;
+	{
+		FScopeLock Lock(&Mutex);
+		Pool = Private::FindPoolInfo(*this, Ptr);
+	}
+
 	if (!Pool)
 	{
 		UE_LOG(LogMemory, Fatal, TEXT("FMallocBinned2 Attempt to GetAllocationSizeExternal an unrecognized block %p"), Ptr);
@@ -1207,7 +1227,7 @@ int64 FMallocBinned2::GetTotalAllocatedSmallPoolMemory() const
 		FreeBlockAllocatedMemory += FPerThreadFreeBlockLists::ConsolidatedMemory;
 	}
 
-	return AllocatedSmallPoolMemory + FreeBlockAllocatedMemory;
+	return AllocatedSmallPoolMemory.Load(EMemoryOrder::Relaxed) + FreeBlockAllocatedMemory;
 }
 #endif
 
@@ -1215,16 +1235,20 @@ void FMallocBinned2::GetAllocatorStats( FGenericMemoryStats& OutStats )
 {
 #if BINNED2_ALLOCATOR_STATS
 
-	int64 TotalAllocatedSmallPoolMemory = GetTotalAllocatedSmallPoolMemory();
+	int64  TotalAllocatedSmallPoolMemory           = GetTotalAllocatedSmallPoolMemory();
+	int64  LocalAllocatedOSSmallPoolMemory         = AllocatedOSSmallPoolMemory.Load(EMemoryOrder::Relaxed);
+	int64  LocalAllocatedLargePoolMemory           = AllocatedLargePoolMemory.Load(EMemoryOrder::Relaxed);
+	int64  LocalAllocatedLargePoolMemoryWAlignment = AllocatedLargePoolMemoryWAlignment.Load(EMemoryOrder::Relaxed);
+	uint64 OSPageAllocatorCachedFreeSize           = CachedOSPageAllocator.GetCachedFreeTotal();
 
 	OutStats.Add(TEXT("AllocatedSmallPoolMemory"), TotalAllocatedSmallPoolMemory);
-	OutStats.Add(TEXT("AllocatedOSSmallPoolMemory"), AllocatedOSSmallPoolMemory);
-	OutStats.Add(TEXT("AllocatedLargePoolMemory"), AllocatedLargePoolMemory);
-	OutStats.Add(TEXT("AllocatedLargePoolMemoryWAlignment"), AllocatedLargePoolMemoryWAlignment);
-	OutStats.Add(TEXT("PageAllocatorFreeCacheSize"), CachedOSPageAllocator.GetCachedFreeTotal());
+	OutStats.Add(TEXT("AllocatedOSSmallPoolMemory"), LocalAllocatedOSSmallPoolMemory);
+	OutStats.Add(TEXT("AllocatedLargePoolMemory"), LocalAllocatedLargePoolMemory);
+	OutStats.Add(TEXT("AllocatedLargePoolMemoryWAlignment"), LocalAllocatedLargePoolMemoryWAlignment);
+	OutStats.Add(TEXT("PageAllocatorFreeCacheSize"), OSPageAllocatorCachedFreeSize);
 
-	uint64 TotalAllocated = TotalAllocatedSmallPoolMemory + AllocatedLargePoolMemory;
-	uint64 TotalOSAllocated = AllocatedOSSmallPoolMemory + AllocatedLargePoolMemoryWAlignment + CachedOSPageAllocator.GetCachedFreeTotal();
+	uint64 TotalAllocated = TotalAllocatedSmallPoolMemory + LocalAllocatedLargePoolMemory;
+	uint64 TotalOSAllocated = LocalAllocatedOSSmallPoolMemory + LocalAllocatedLargePoolMemoryWAlignment + OSPageAllocatorCachedFreeSize;
 
 	OutStats.Add(TEXT("TotalAllocated"), TotalAllocated);
 	OutStats.Add(TEXT("TotalOSAllocated"), TotalOSAllocated);
@@ -1236,17 +1260,18 @@ void FMallocBinned2::DumpAllocatorStats(class FOutputDevice& Ar)
 {
 #if BINNED2_ALLOCATOR_STATS
 
-	int64 TotalAllocatedSmallPoolMemory = GetTotalAllocatedSmallPoolMemory();
+	int64  TotalAllocatedSmallPoolMemory           = GetTotalAllocatedSmallPoolMemory();
+	int64  LocalAllocatedLargePoolMemory           = AllocatedLargePoolMemory.Load(EMemoryOrder::Relaxed);
+	int64  LocalAllocatedLargePoolMemoryWAlignment = AllocatedLargePoolMemoryWAlignment.Load(EMemoryOrder::Relaxed);
+	uint64 OSPageAllocatorCachedFreeSize           = CachedOSPageAllocator.GetCachedFreeTotal();
 
 	Ar.Logf(TEXT("FMallocBinned2 Mem report"));
 	Ar.Logf(TEXT("Small Pool"));
 	Ar.Logf(TEXT("Requested Allocations: %fmb  (including block size padding)"), ((double)TotalAllocatedSmallPoolMemory) / (1024.0f * 1024.0f) );
 	Ar.Logf(TEXT("OS Allocated: %fmb"), ((double)TotalAllocatedSmallPoolMemory) / (1024.0f * 1024.0f));
 	Ar.Logf(TEXT("Large Pool"));
-	Ar.Logf(TEXT("Requested Allocations: %fmb"), ((double)AllocatedLargePoolMemory) / (1024.0f * 1024.0f));
-	Ar.Logf(TEXT("OS Allocated: %fmb"), ((double)AllocatedLargePoolMemoryWAlignment) / (1024.0f * 1024.0f));
-
-	uint64 OSPageAllocatorCachedFreeSize = CachedOSPageAllocator.GetCachedFreeTotal();
+	Ar.Logf(TEXT("Requested Allocations: %fmb"), ((double)LocalAllocatedLargePoolMemory) / (1024.0f * 1024.0f));
+	Ar.Logf(TEXT("OS Allocated: %fmb"), ((double)LocalAllocatedLargePoolMemoryWAlignment) / (1024.0f * 1024.0f));
 	Ar.Logf(TEXT("OS Page Allocator"));
 	Ar.Logf(TEXT("Cached free pages: %fmb"), ((double)OSPageAllocatorCachedFreeSize) / (1024.0f * 1024.0f));
 #else
diff --git a/Engine/Source/Runtime/Core/Public/HAL/MallocBinned2.h b/Engine/Source/Runtime/Core/Public/HAL/MallocBinned2.h
--- a/Engine/Source/Runtime/Core/Public/HAL/MallocBinned2.h
+++ b/Engine/Source/Runtime/Core/Public/HAL/MallocBinned2.h
@@ -14,6 +14,7 @@
 #include "HAL/Allocators/PooledVirtualMemoryAllocator.h"
 #include "HAL/PlatformMath.h"
 #include "HAL/LowLevelMemTracker.h"
+#include "Templates/Atomic.h"
 
 #define BINNED2_MAX_CACHED_OS_FREES (64)
 #if PLATFORM_64BITS
@@ -66,10 +67,10 @@
 #if BINNED2_ALLOCATOR_STATS
 //////////////////////////////////////////////////////////////////////////
 // the following don't need a critical section because they are covered by the critical section called Mutex
-extern int64 AllocatedSmallPoolMemory; // memory that's requested to be allocated by the game
-extern int64 AllocatedOSSmallPoolMemory;
-extern int64 AllocatedLargePoolMemory; // memory requests to the OS which don't fit in the small pool
-extern int64 AllocatedLargePoolMemoryWAlignment; // when we allocate at OS level we need to align to a size
+extern TAtomic<int64> AllocatedSmallPoolMemory; // memory that's requested to be allocated by the game
+extern TAtomic<int64> AllocatedOSSmallPoolMemory;
+extern TAtomic<int64> AllocatedLargePoolMemory; // memory requests to the OS which don't fit in the small pool
+extern TAtomic<int64> AllocatedLargePoolMemoryWAlignment; // when we allocate at OS level we need to align to a size
 #endif
 #if BINNED2_ALLOCATOR_STATS_VALIDATION
 #include "Misc/ScopeLock.h"
@@ -461,7 +462,27 @@ public:
 		}
 		if (Result == nullptr)
 		{
-			Result = MallocExternal(Size, Alignment);
+			Result = MallocSelect(Size, Alignment);
+		}
+
+		return Result;
+	}
+	FORCEINLINE static bool UseSmallAlloc(SIZE_T Size, uint32 Alignment)
+	{
+		bool bResult = ((Size <= BINNED2_MAX_SMALL_POOL_SIZE) & (Alignment <= BINNED2_MINIMUM_ALIGNMENT)); // one branch, not two
+		return bResult;
+	}
+	FORCEINLINE void* MallocSelect(SIZE_T Size, uint32 Alignment)
+	{
+		void* Result;
+
+		if (UseSmallAlloc(Size, Alignment))
+		{
+			Result = MallocExternalSmall(Size, Alignment);
+		}
+		else
+		{
+			Result = MallocExternalLarge(Size, Alignment);
 		}
 
 		return Result;
@@ -648,7 +669,8 @@ public:
 	// End FMalloc interface.
 
 	void FlushCurrentThreadCache();
-	void* MallocExternal(SIZE_T Size, uint32 Alignment);
+	void* MallocExternalSmall(SIZE_T Size, uint32 Alignment);
+	void* MallocExternalLarge(SIZE_T Size, uint32 Alignment);
 	void* ReallocExternal(void* Ptr, SIZE_T NewSize, uint32 Alignment);
 	void FreeExternal(void *Ptr);
 	bool GetAllocationSizeExternal(void* Ptr, SIZE_T& SizeOut);
