Added a compile time assert to Mac and iOS PlatformCompilerSetup.h to display an error message when attempting to build with Xcode veriosn older than 9.4

diff --git a/Engine/Source/Runtime/Core/Public/IOS/IOSPlatformCompilerSetup.h b/Engine/Source/Runtime/Core/Public/IOS/IOSPlatformCompilerSetup.h
--- a/Engine/Source/Runtime/Core/Public/IOS/IOSPlatformCompilerSetup.h
+++ b/Engine/Source/Runtime/Core/Public/IOS/IOSPlatformCompilerSetup.h
@@ -5,3 +5,8 @@
 ==============================================================================================*/
 
 #pragma once
+
+/**
+ * We require at least Xcode 9.4 to compile
+ */
+static_assert((__clang_major__ > 9) || (__clang_major__ == 9 && __clang_minor__ >= 1), "Xcode 9.4 or newer is required to compile on Mac. Please make sure it's installed and selected as default using xcode-select tool.");
diff --git a/Engine/Source/Runtime/Core/Public/Mac/MacPlatformCompilerSetup.h b/Engine/Source/Runtime/Core/Public/Mac/MacPlatformCompilerSetup.h
--- a/Engine/Source/Runtime/Core/Public/Mac/MacPlatformCompilerSetup.h
+++ b/Engine/Source/Runtime/Core/Public/Mac/MacPlatformCompilerSetup.h
@@ -8,3 +8,8 @@
 
 // In OS X 10.8 SDK gl3.h complains if gl.h was also included. Unfortunately, gl.h is included by Cocoa in CoreVideo.framework, so we need to disable this warning
 #define GL_DO_NOT_WARN_IF_MULTI_GL_VERSION_HEADERS_INCLUDED
+
+/**
+ * We require at least Xcode 9.4 to compile
+ */
+static_assert((__clang_major__ > 9) || (__clang_major__ == 9 && __clang_minor__ >= 1), "Xcode 9.4 or newer is required to compile on Mac. Please make sure it's installed and selected as default using xcode-select tool.");
