RemoveIf and StableRemoveIf algorithms.

diff --git a/Engine/Source/Runtime/Core/Public/Algo/RemoveIf.h b/Engine/Source/Runtime/Core/Public/Algo/RemoveIf.h
new file mode 100644
--- /dev/null
+++ b/Engine/Source/Runtime/Core/Public/Algo/RemoveIf.h
@@ -0,0 +1,120 @@
+// Copyright 1998-2018 Epic Games, Inc. All Rights Reserved.
+
+#pragma once
+
+#include "Templates/Invoke.h"
+
+namespace Algo
+{
+	/**
+	 * Moves all elements which do not match the predicate to the front of the range, while leaving all
+	 * other elements is a constructed but unspecified state.  The elements which were not removed are
+	 * not guaranteed to be kept in order (unstable).
+	 *
+	 * @param  Range  The range of elements to manipulate.
+	 * @param  Pred   A callable which maps elements to truthy values, specifying elements to be removed.
+	 *
+	 * @return The index of the first element after those which were not removed.
+	 */
+	template <typename RangeType, typename Predicate>
+	int32 RemoveIf(RangeType& Range, Predicate Pred)
+	{
+		auto* First = GetData(Range);
+		auto* Last  = First + GetNum(Range);
+
+		auto* IterStart = First;
+		auto* IterEnd   = Last;
+		for (;;)
+		{
+			// Skip non-removed elements at the start
+			for (;;)
+			{
+				if (IterStart == IterEnd)
+				{
+					return IterStart - First;
+				}
+
+				if (!Invoke(Pred, *IterStart))
+				{
+					break;
+				}
+
+				++IterStart;
+			}
+
+			// Skip removed elements at the end
+			for (;;)
+			{
+				if (Invoke(Pred, *(IterEnd - 1)))
+				{
+					break;
+				}
+
+				--IterEnd;
+
+				if (IterStart == IterEnd)
+				{
+					return IterStart - First;
+				}
+			}
+
+			*IterStart = MoveTemp(*(IterEnd - 1));
+
+			++IterStart;
+			--IterEnd;
+		}
+	}
+
+	/**
+	 * Moves all elements which do not match the predicate to the front of the range, while leaving all
+	 * other elements is a constructed but unspecified state.  The elements which were not removed are
+	 * guaranteed to be kept in order (stable).
+	 *
+	 * @param  Range  The range of elements to manipulate.
+	 * @param  Pred   A callable which maps elements to truthy values, specifying elements to be removed.
+	 *
+	 * @return The index of the first element after those which were not removed.
+	 */
+	template <typename RangeType, typename Predicate>
+	int32 StableRemoveIf(RangeType& Range, Predicate Pred)
+	{
+		auto* First = GetData(Range);
+		auto* Last  = First + GetNum(Range);
+
+		auto* IterStart = First;
+
+		// Skip non-removed elements at the start
+		for (;;)
+		{
+			if (IterStart == Last)
+			{
+				return IterStart - First;
+			}
+
+			if (!Invoke(Pred, *IterStart))
+			{
+				break;
+			}
+
+			++IterStart;
+		}
+
+		auto* IterKeep = IterStart;
+		++IterKeep;
+
+		for (;;)
+		{
+			if (IterKeep == Last)
+			{
+				return IterStart - First;
+			}
+
+			if (Invoke(Pred, *IterKeep))
+			{
+				*IterStart++ = MoveTemp(*IterKeep);
+			}
+
+			++IterKeep;
+		}
+	}
+}
