Change compile time PLATFORM_CACHE_LINE_SIZE to 64 on all platforms

Both Gen4 consoles already set this to 64 in their platform headers. Windows was set to 128. All other platforms use the Platform.h fallback which was 128.

diff --git a/Engine/Source/Runtime/Core/Public/HAL/Platform.h b/Engine/Source/Runtime/Core/Public/HAL/Platform.h
--- a/Engine/Source/Runtime/Core/Public/HAL/Platform.h
+++ b/Engine/Source/Runtime/Core/Public/HAL/Platform.h
@@ -609,7 +609,7 @@
 
 // Prefetch
 #ifndef PLATFORM_CACHE_LINE_SIZE
-	#define PLATFORM_CACHE_LINE_SIZE	128
+	#define PLATFORM_CACHE_LINE_SIZE	64
 #endif
 
 // Compile-time warnings and errors. Use these as "#pragma COMPILER_WARNING("XYZ")". GCC does not expand macro parameters to _Pragma, so we can't wrap the #pragma part.
diff --git a/Engine/Source/Runtime/Core/Public/Windows/WIndowsPlatform.h b/Engine/Source/Runtime/Core/Public/Windows/WIndowsPlatform.h
--- a/Engine/Source/Runtime/Core/Public/Windows/WIndowsPlatform.h
+++ b/Engine/Source/Runtime/Core/Public/Windows/WIndowsPlatform.h
@@ -141,7 +141,7 @@ typedef FWindowsPlatformTypes FPlatformTypes;
 #endif
 
 // Prefetch
-#define PLATFORM_CACHE_LINE_SIZE	128
+#define PLATFORM_CACHE_LINE_SIZE	64
 
 // DLL export and import definitions
 #define DLLEXPORT __declspec(dllexport)
