Fix for VS' debug engine swallowing explicit breakpoint interrupts when single-stepping.

diff --git a/Engine/Source/Runtime/Core/Public/Windows/WIndowsPlatform.h b/Engine/Source/Runtime/Core/Public/Windows/WIndowsPlatform.h
--- a/Engine/Source/Runtime/Core/Public/Windows/WIndowsPlatform.h
+++ b/Engine/Source/Runtime/Core/Public/Windows/WIndowsPlatform.h
@@ -60,10 +60,16 @@ typedef FWindowsPlatformTypes FPlatformTypes;
 #define PLATFORM_COMPILER_HAS_TCHAR_WMAIN					1
 #define PLATFORM_SUPPORTS_EARLY_MOVIE_PLAYBACK				(!WITH_EDITOR) // movies will start before engine is initalized
 #define PLATFORM_RHITHREAD_DEFAULT_BYPASS					0
-#define PLATFORM_BREAK()									__debugbreak()
-
 #define PLATFORM_SUPPORTS_STACK_SYMBOLS						1
 
+// Q: Why is there a __nop() before __debugbreak()?
+// A: VS' debug engine has a bug where it will silently swallow explicit
+// breakpoint interrupts when single-step debugging either line-by-line or
+// over call instructions. This can hide legitimate reasons to trap. Asserts
+// for example, which can appear as if the did not fire, leaving a programmer
+// unknowingly debugging an undefined process.
+#define PLATFORM_BREAK() (__nop(), __debugbreak())
+
 #if defined(__INTEL_COMPILER) || _MSC_VER > 1900
 	#define PLATFORM_COMPILER_HAS_DECLTYPE_AUTO 1
 #else
