Fixed warnings on IOS regarding the alignment of functions that are in a section/segment other than __text.

diff --git a/Engine/Source/Runtime/Core/Public/IOS/IOSPlatform.h b/Engine/Source/Runtime/Core/Public/IOS/IOSPlatform.h
--- a/Engine/Source/Runtime/Core/Public/IOS/IOSPlatform.h
+++ b/Engine/Source/Runtime/Core/Public/IOS/IOSPlatform.h
@@ -71,7 +71,8 @@ typedef FIOSPlatformTypes FPlatformTypes;
 	#define PLATFORM_BREAK()							__asm__("trap")
 #endif
 
-#define PLATFORM_CODE_SECTION(Name)						__attribute__((section("__TEXT,__" Name ",regular,pure_instructions")))
+#define PLATFORM_CODE_SECTION(Name)						__attribute__((section("__TEXT,__" Name ",regular,pure_instructions"))) \
+														__attribute__((aligned(4)))
 
 #if __has_feature(cxx_decltype_auto)
 	#define PLATFORM_COMPILER_HAS_DECLTYPE_AUTO 1
