Added missing brackets following a conditional.

diff --git a/Engine/Source/Runtime/Core/Private/Windows/WindowsPlatformMisc.cpp b/Engine/Source/Runtime/Core/Private/Windows/WindowsPlatformMisc.cpp
--- a/Engine/Source/Runtime/Core/Private/Windows/WindowsPlatformMisc.cpp
+++ b/Engine/Source/Runtime/Core/Private/Windows/WindowsPlatformMisc.cpp
@@ -604,7 +604,9 @@ static BOOL WINAPI ConsoleCtrlHandler( ::DWORD /*Type*/ )
 void FWindowsPlatformMisc::SetGracefulTerminationHandler()
 {
 	if (GetConsoleWindow() == nullptr)
+	{
 		return;
+	}
 
 	// Set console control handler so we can exit if requested.
 	SetConsoleCtrlHandler(ConsoleCtrlHandler, true);
