Matched Ctrl-C handling on Windows to follow the same two-step pattern that is used on Linux and Mac.

diff --git a/Engine/Source/Runtime/Core/Private/Windows/WindowsPlatformMisc.cpp b/Engine/Source/Runtime/Core/Private/Windows/WindowsPlatformMisc.cpp
--- a/Engine/Source/Runtime/Core/Private/Windows/WindowsPlatformMisc.cpp
+++ b/Engine/Source/Runtime/Core/Private/Windows/WindowsPlatformMisc.cpp
@@ -566,15 +566,20 @@ void FWindowsPlatformMisc::PlatformInit()
  *
  * @param Type Ctrl-C, Ctrl-Break, Close Console Log Window, Logoff, or Shutdown
  */
-static BOOL WINAPI ConsoleCtrlHandler( ::DWORD /*Type*/ )
+static BOOL WINAPI ConsoleCtrlHandler(DWORD CtrlType)
 {
-	// Once this function is called, Windows gives us about 5 seconds to clean up and exit.
-	// There's no way to cancel. Since console is shutting down, logging might not be reliable.
-
-	GIsRequestingExit = true;
+	// Only "two-step Ctrl-C" if the termination event is Ctrl-C and the process
+	// is considered interactive. Hard-terminate on all other cases.
+	if (CtrlType != CTRL_C_EVENT || FApp::IsUnattended())
+	{
+		GIsRequestingExit = true;
+	}
 
-	// Notify anyone listening that we're about to terminate
-	FCoreDelegates::ApplicationWillTerminateDelegate.Broadcast();
+	if (!GIsRequestingExit)
+	{
+		UE_LOG(LogCore, Warning, TEXT("*** INTERRUPTED *** : SHUTTING DOWN"));
+		UE_LOG(LogCore, Warning, TEXT("*** INTERRUPTED *** : CTRL-C TO FORCE QUIT"));
+	}
 
 	// make sure as much data is written to disk as possible
 	if (GLog)
@@ -590,13 +595,21 @@ static BOOL WINAPI ConsoleCtrlHandler( ::DWORD /*Type*/ )
 		GError->Flush();
 	}
 
+	if (!GIsRequestingExit)
+	{
+		// Notify anyone listening that we're about to terminate
+		GIsRequestingExit = true;
+		FCoreDelegates::ApplicationWillTerminateDelegate.Broadcast();
+		return true;
+	}
+
 	// There's no guarantee that the process is paying attention to GIsRequestingExit
 	// (e.g. some long-running commandlets). Using that for shutdown is therefore not
 	// reliable. Deferring to the default CtrlHandler is also no good as that calls
 	// ExitProcess() which will terminate all threads and detach all DLLS. This can
 	// result in deadlocks and/or asserts as each DLL's atexit() is processed. So
 	// let's hard terminate the process. 0xc000013a is what Windows' default handler
-	// would normally pass to ExitProcess().
+	// would normally pass to ExitProcess() and how ExitProc() terminates threads.
 	TerminateProcess(GetCurrentProcess(), 0xc000013au);
 	return false;
 }
