GitHub #3697 : Bugfix - repeated .module file opening by FindModulePaths()

This has had some work applied to it after more recent changes were made, to remove statics and to accomodate hot reload.

diff --git a/Engine/Source/Runtime/Core/Private/Modules/ModuleManager.cpp b/Engine/Source/Runtime/Core/Private/Modules/ModuleManager.cpp
--- a/Engine/Source/Runtime/Core/Private/Modules/ModuleManager.cpp
+++ b/Engine/Source/Runtime/Core/Private/Modules/ModuleManager.cpp
@@ -921,87 +921,86 @@ void FModuleManager::SetModuleFilename(FName ModuleName, const FString& Filename
 void FModuleManager::ResetModulePathsCache()
 {
 	ModulePathsCache.Reset();
+
+	// Set all folders as pending again
+	PendingEngineBinariesDirectories.Append(MoveTemp(EngineBinariesDirectories));
+	PendingGameBinariesDirectories  .Append(MoveTemp(GameBinariesDirectories));
 }
 
 #if !IS_MONOLITHIC
-void FModuleManager::FindModulePaths(const TCHAR* NamePattern, TMap<FName, FString> &OutModulePaths, bool bCanUseCache /*= true*/) const
+void FModuleManager::FindModulePaths(const TCHAR* NamePattern, TMap<FName, FString> &OutModulePaths) const
 {
-	if (!ModulePathsCache)
+	if (ModulePathsCache.Num() == 0)
 	{
-		ModulePathsCache.Emplace();
-		const bool bCanUseCacheWhileGeneratingIt = false;
-		FindModulePaths(TEXT("*"), ModulePathsCache.GetValue(), bCanUseCacheWhileGeneratingIt);
-	}
-
-	if (bCanUseCache)
-	{
-		// Try to use cache first
-		if (const FString* ModulePathPtr = ModulePathsCache->Find(NamePattern))
+		// Figure out the BuildId if it's not already set.
+		if (!BuildId.IsSet())
 		{
-			OutModulePaths.Add(FName(NamePattern), *ModulePathPtr);
-			return;
-		}
+			FString FileName = FModuleManifest::GetFileName(FPlatformProcess::GetModulesDirectory(), false);
 
-		// Wildcard for all items
-		if (FCString::Strcmp(NamePattern, TEXT("*")) == 0)
-		{
-			OutModulePaths = ModulePathsCache.GetValue();
-			return;
-		}
-		
-		// Wildcard search
-		if (FCString::Strchr(NamePattern, TEXT('*')) || FCString::Strchr(NamePattern, TEXT('?')))
-		{
-			bool bFoundItems = false;
-			FString NamePatternString(NamePattern);
-			for (const TPair<FName, FString>& CacheIt : ModulePathsCache.GetValue())
+			FModuleManifest Manifest;
+			if (!FModuleManifest::TryRead(FileName, Manifest))
 			{
-				if (CacheIt.Key.ToString().MatchesWildcard(NamePatternString))
-				{
-					OutModulePaths.Add(CacheIt.Key, *CacheIt.Value);
-					bFoundItems = true;
-				}
+				UE_LOG(LogModuleManager, Fatal, TEXT("Unable to read module manifest from '%s'. Module manifests are generated at build time, and must be present to locate modules at runtime."), *FileName)
 			}
 
-			if (bFoundItems)
-			{
-				return;
-			}
+			BuildId = Manifest.BuildId;
 		}
+
+		// Add the engine directory to the cache - only needs to be cached once as the contents do not change at runtime
+		FindModulePathsInDirectory(FPlatformProcess::GetModulesDirectory(), false, ModulePathsCache);
 	}
 
-	// Search through the engine directory
-	FindModulePathsInDirectory(FPlatformProcess::GetModulesDirectory(), false, NamePattern, OutModulePaths);
+	// If any entries have been added to the PendingEngineBinariesDirectories or PendingGameBinariesDirectories arrays, add any
+	// new valid modules within them to the cache.
+	// Static iterators used as once we've cached a BinariesDirectories array entry we don't want to do it again (ModuleManager is a Singleton)
 
 	// Search any engine directories
-	for (int Idx = 0; Idx < EngineBinariesDirectories.Num(); Idx++)
+	if (PendingEngineBinariesDirectories.Num() > 0)
 	{
-		FindModulePathsInDirectory(EngineBinariesDirectories[Idx], false, NamePattern, OutModulePaths);
+		TArray<FString> LocalPendingEngineBinariesDirectories = MoveTemp(PendingEngineBinariesDirectories);
+		check(PendingEngineBinariesDirectories.Num() == 0);
+
+		for (const FString& EngineBinaryDirectory : LocalPendingEngineBinariesDirectories)
+		{
+			FindModulePathsInDirectory(EngineBinaryDirectory, false, ModulePathsCache);
+		}
+
+		EngineBinariesDirectories.Append(MoveTemp(LocalPendingEngineBinariesDirectories));
 	}
 
 	// Search any game directories
-	for (int Idx = 0; Idx < GameBinariesDirectories.Num(); Idx++)
+	if (PendingGameBinariesDirectories.Num() > 0)
 	{
-		FindModulePathsInDirectory(GameBinariesDirectories[Idx], true, NamePattern, OutModulePaths);
+		TArray<FString> LocalPendingGameBinariesDirectories = MoveTemp(PendingGameBinariesDirectories);
+		check(PendingGameBinariesDirectories.Num() == 0);
+
+		for (const FString& GameBinaryDirectory : LocalPendingGameBinariesDirectories)
+		{
+			FindModulePathsInDirectory(GameBinaryDirectory, true, ModulePathsCache);
+		}
+
+		GameBinariesDirectories.Append(MoveTemp(LocalPendingGameBinariesDirectories));
 	}
-}
 
-void FModuleManager::FindModulePathsInDirectory(const FString& InDirectoryName, bool bIsGameDirectory, const TCHAR* NamePattern, TMap<FName, FString> &OutModulePaths) const
-{
-	// Figure out the BuildId if it's not already set.
-	if (!BuildId.IsSet())
+	// Wildcard for all items
+	if (FCString::Strcmp(NamePattern, TEXT("*")) == 0)
 	{
-		FString FileName = FModuleManifest::GetFileName(FPlatformProcess::GetModulesDirectory(), false);
+		OutModulePaths = ModulePathsCache;
+		return;
+	}
 
-		FModuleManifest Manifest;
-		if (!FModuleManifest::TryRead(FileName, Manifest))
+	// Search the cache
+	for (const TPair<FName, FString>& Pair : ModulePathsCache)
+	{
+		if (Pair.Key.ToString().MatchesWildcard(NamePattern))
 		{
-			UE_LOG(LogModuleManager, Fatal, TEXT("Unable to read module manifest from '%s'. Module manifests are generated at build time, and must be present to locate modules at runtime."), *FileName)
+			OutModulePaths.Add(Pair.Key, Pair.Value);
 		}
-
-		BuildId = Manifest.BuildId;
 	}
+}
 
+void FModuleManager::FindModulePathsInDirectory(const FString& InDirectoryName, bool bIsGameDirectory, TMap<FName, FString>& OutModulePaths) const
+{
 	// Find all the directories to search through, including the base directory
 	TArray<FString> SearchDirectoryNames;
 	IFileManager::Get().FindFilesRecursive(SearchDirectoryNames, *InDirectoryName, TEXT("*"), false, true);
@@ -1015,10 +1014,7 @@ void FModuleManager::FindModulePathsInDirectory(const FString& InDirectoryName,
 		{
 			for (const TPair<FString, FString>& Pair : Manifest.ModuleNameToFileName)
 			{
-				if (Pair.Key.MatchesWildcard(NamePattern))
-				{
-					OutModulePaths.Add(FName(*Pair.Key), *FPaths::Combine(*SearchDirectoryName, *Pair.Value));
-				}
+				OutModulePaths.Add(FName(*Pair.Key), *FPaths::Combine(*SearchDirectoryName, *Pair.Value));
 			}
 		}
 	}
@@ -1121,11 +1117,11 @@ void FModuleManager::AddBinariesDirectory(const TCHAR *InDirectory, bool bIsGame
 {
 	if (bIsGameDirectory)
 	{
-		GameBinariesDirectories.Add(InDirectory);
+		PendingGameBinariesDirectories.Add(InDirectory);
 	}
 	else
 	{
-		EngineBinariesDirectories.Add(InDirectory);
+		PendingEngineBinariesDirectories.Add(InDirectory);
 	}
 
 	FPlatformProcess::AddDllDirectory(InDirectory);
@@ -1140,8 +1136,6 @@ void FModuleManager::AddBinariesDirectory(const TCHAR *InDirectory, bool bIsGame
 			AddBinariesDirectory(*RestrictedFolder, bIsGameDirectory);
 		}
 	}
-
-	ResetModulePathsCache();
 }
 
 
@@ -1154,9 +1148,7 @@ void FModuleManager::SetGameBinariesDirectory(const TCHAR* InDirectory)
 	FPlatformProcess::PushDllDirectory(InDirectory);
 
 	// Add it to the list of game directories to search
-	GameBinariesDirectories.Add(InDirectory);
-
-	ResetModulePathsCache();
+	PendingGameBinariesDirectories.Add(InDirectory);
 #endif
 }
 
@@ -1166,6 +1158,10 @@ FString FModuleManager::GetGameBinariesDirectory() const
 	{
 		return GameBinariesDirectories[0];
 	}
+	if (PendingGameBinariesDirectories.Num())
+	{
+		return PendingGameBinariesDirectories[0];
+	}
 	return FString();
 }
 
diff --git a/Engine/Source/Runtime/Core/Public/Modules/ModuleManager.h b/Engine/Source/Runtime/Core/Public/Modules/ModuleManager.h
--- a/Engine/Source/Runtime/Core/Public/Modules/ModuleManager.h
+++ b/Engine/Source/Runtime/Core/Public/Modules/ModuleManager.h
@@ -555,10 +555,10 @@ private:
 
 #if !IS_MONOLITHIC
 	/** Finds modules matching a given name wildcard. */
-	void FindModulePaths(const TCHAR *NamePattern, TMap<FName, FString> &OutModulePaths, bool bCanUseCache = true) const;
+	void FindModulePaths(const TCHAR *NamePattern, TMap<FName, FString> &OutModulePaths) const;
 
-	/** Finds modules matching a given name wildcard within a given directory. */
-	void FindModulePathsInDirectory(const FString &DirectoryName, bool bIsGameDirectory, const TCHAR *NamePattern, TMap<FName, FString> &OutModulePaths) const;
+	/** Finds modules within a given directory. */
+	void FindModulePathsInDirectory(const FString &DirectoryName, bool bIsGameDirectory, TMap<FName, FString> &OutModulePaths) const;
 #endif
 
 private:
@@ -576,7 +576,7 @@ private:
 	bool bCanProcessNewlyLoadedObjects;
 
 	/** Cache of known module paths. Used for performance. Can increase editor startup times by up to 30% */
-	mutable TOptional<TMap<FName, FString>> ModulePathsCache;
+	mutable TMap<FName, FString> ModulePathsCache;
 
 	/** Multicast delegate that will broadcast a notification when modules are loaded, unloaded, or
 		our set of known modules changes */
@@ -590,10 +590,12 @@ private:
 	FIsPackageLoadedCallback IsPackageLoaded;
 
 	/** Array of engine binaries directories. */
-	TArray<FString> EngineBinariesDirectories;
+	mutable TArray<FString> EngineBinariesDirectories;
+	mutable TArray<FString> PendingEngineBinariesDirectories;
 
 	/** Array of game binaries directories. */
-	TArray<FString> GameBinariesDirectories;
+	mutable TArray<FString> GameBinariesDirectories;
+	mutable TArray<FString> PendingGameBinariesDirectories;
 
 	/** ID used to validate module manifests. Read from the module manifest in the engine directory on first query to load a new module; unset until then. */
 	mutable TOptional<FString> BuildId;
