Optimized Str[n]icmp() and TChar::ToLower/ToUpper()

Replaced expensive CRT to[w]lower() / to[w]upper() calls with locale-independent implementations.

The semantics match standard C locale that all UE platforms use - only ASCII characters are considered.

* Simplified TChar IsLinebreak()
* New TChar::ToLower() and ToUpper() implementation
* Replaced all Str[n]icmp() implementations with a unified implementation that all platforms use
* Moved implementation inside a cpp file to reduce code bloat
* Added unit tests
* Fixed sign extension bug in the replaced Str[n]icmp()

The bug would only manifest for 1 and 2 byte signed character types. Also MicrosoftPlatformString forwarded Str[n]icmp() to correct but non-standard CRT calls. AFAICT only narrow strings containing non-ASCII characters on non-MS platforms was affected. See TChar::ToUnsigned() for more details.

diff --git a/Engine/Source/Runtime/Core/Private/GenericPlatform/GenericPlatformStricmp.cpp b/Engine/Source/Runtime/Core/Private/GenericPlatform/GenericPlatformStricmp.cpp
--- a/Engine/Source/Runtime/Core/Private/GenericPlatform/GenericPlatformStricmp.cpp
+++ b/Engine/Source/Runtime/Core/Private/GenericPlatform/GenericPlatformStricmp.cpp
@@ -1,6 +1,7 @@
 // Copyright 1998-2018 Epic Games, Inc. All Rights Reserved.
 
 #include "GenericPlatform/GenericPlatformStricmp.h"
+#include "Misc/Char.h"
 
 static constexpr uint8 LowerAscii[128] = {
 	0x00, 0x01, 0x02, 0x03, 0x04, 0x05, 0x06, 0x07, 0x08, 0x09, 0x0A, 0x0B, 0x0C, 0x0D, 0x0E, 0x0F,
@@ -13,18 +14,22 @@ static constexpr uint8 LowerAscii[128] = {
 	0x70, 0x71, 0x72, 0x73, 0x74, 0x75, 0x76, 0x77, 0x78, 0x79, 0x7A, 0x7B, 0x7C, 0x7D, 0x7E, 0x7F
 };
 
-FORCEINLINE bool BothAscii(TCHAR a, TCHAR b)
+template<typename CharType1, typename CharType2>
+FORCEINLINE bool BothAscii(CharType1 C1, CharType2 C2)
 {
-	return (((uint32)a | (uint32)b) & 0xffffff80) == 0;
+	return (((uint32)C1 | (uint32)C2) & 0xffffff80) == 0;
 }
 
-int32 FGenericPlatformStricmp::CompatibleCharTypesStricmp(const TCHAR* String1, const TCHAR* String2)
+template<typename CharType1, typename CharType2>
+int32 StricmpImpl(const CharType1* String1, const CharType2* String2)
 {
 	while (true)
 	{
-		TCHAR C1 = *String1++;
-		TCHAR C2 = *String2++;
+		CharType1 C1 = *String1++;
+		CharType2 C2 = *String2++;
 
+		// Quickly move on if characters are identical but
+		// return equals if we found two null terminators
 		if (C1 == C2)
 		{
 			if (C1)
@@ -36,17 +41,140 @@ int32 FGenericPlatformStricmp::CompatibleCharTypesStricmp(const TCHAR* String1,
 		}
 		else if (BothAscii(C1, C2))
 		{
-			if (SSIZE_T Diff = LowerAscii[C1] - LowerAscii[C2])
+			if (int32 Diff = LowerAscii[C1] - LowerAscii[C2])
 			{
 				return Diff;
 			}
 		}
 		else
 		{
-			if (SSIZE_T Diff = TChar<TCHAR>::ToLower(C1) - TChar<TCHAR>::ToLower(C2))
+			return TChar<CharType1>::ToUnsigned(C1) - TChar<CharType2>::ToUnsigned(C2);
+		}
+	}
+}
+
+template<typename CharType>
+int32 StrnicmpImpl(const CharType* String1, const CharType* String2, SIZE_T Count)
+{
+	for (; Count > 0; --Count)
+	{
+		CharType C1 = *String1++;
+		CharType C2 = *String2++;
+
+		// Quickly move on if characters are identical but
+		// return equals if we found two null terminators
+		if (C1 == C2)
+		{
+			if (C1)
+			{
+				continue;
+			}
+
+			return 0;
+		}
+		else if (BothAscii(C1, C2))
+		{
+			if (int32 Diff = LowerAscii[C1] - LowerAscii[C2])
 			{
 				return Diff;
 			}
 		}
+		else
+		{
+			return TChar<CharType>::ToUnsigned(C1) - TChar<CharType>::ToUnsigned(C2);
+		}
+	}
+
+	return 0;
+}
+
+int32 FGenericPlatformStricmp::Stricmp(const ANSICHAR* Str1, const ANSICHAR* Str2) { return StricmpImpl(Str1, Str2); }
+int32 FGenericPlatformStricmp::Stricmp(const WIDECHAR* Str1, const WIDECHAR* Str2) { return StricmpImpl(Str1, Str2); }
+int32 FGenericPlatformStricmp::Stricmp(const UTF8CHAR* Str1, const UTF8CHAR* Str2) { return StricmpImpl(Str1, Str2); }
+int32 FGenericPlatformStricmp::Stricmp(const UTF16CHAR* Str1, const UTF16CHAR* Str2) { return StricmpImpl(Str1, Str2); }
+int32 FGenericPlatformStricmp::Stricmp(const UTF32CHAR* Str1, const UTF32CHAR* Str2) { return StricmpImpl(Str1, Str2); }
+int32 FGenericPlatformStricmp::Stricmp(const ANSICHAR* Str1, const WIDECHAR* Str2) { return StricmpImpl(Str1, Str2); }
+int32 FGenericPlatformStricmp::Stricmp(const ANSICHAR* Str1, const UTF8CHAR* Str2) { return StricmpImpl(Str1, Str2); }
+int32 FGenericPlatformStricmp::Stricmp(const ANSICHAR* Str1, const UTF16CHAR* Str2) { return StricmpImpl(Str1, Str2); }
+int32 FGenericPlatformStricmp::Stricmp(const ANSICHAR* Str1, const UTF32CHAR* Str2) { return StricmpImpl(Str1, Str2); }
+int32 FGenericPlatformStricmp::Stricmp(const WIDECHAR* Str1, const ANSICHAR* Str2) { return StricmpImpl(Str1, Str2); }
+int32 FGenericPlatformStricmp::Stricmp(const UTF8CHAR* Str1, const ANSICHAR* Str2) { return StricmpImpl(Str1, Str2); }
+int32 FGenericPlatformStricmp::Stricmp(const UTF16CHAR* Str1, const ANSICHAR* Str2) { return StricmpImpl(Str1, Str2); }
+int32 FGenericPlatformStricmp::Stricmp(const UTF32CHAR* Str1, const ANSICHAR* Str2) { return StricmpImpl(Str1, Str2); }
+int32 FGenericPlatformStricmp::Strnicmp(const ANSICHAR* Str1, const ANSICHAR* Str2, SIZE_T Count) { return StrnicmpImpl(Str1, Str2, Count); }
+int32 FGenericPlatformStricmp::Strnicmp(const WIDECHAR* Str1, const WIDECHAR* Str2, SIZE_T Count) { return StrnicmpImpl(Str1, Str2, Count); }
+
+//////////////////////////////////////////////////////////////////////////
+
+#if WITH_DEV_AUTOMATION_TESTS 
+
+#include "Misc/AutomationTest.h"
+
+IMPLEMENT_SIMPLE_AUTOMATION_TEST(FGenericPlatformStricmpTest, "System.Core.GenericPlatform.Stricmp", EAutomationTestFlags::ApplicationContextMask | EAutomationTestFlags::SmokeFilter)
+
+// Simpler reference implementation based on ToLower() instead of a lookup table.
+// Used to verify correctness on non-Windows platforms
+template <typename CharType>
+int32 StricmpExpected(const CharType* Str1, const CharType* Str2)
+{
+	for (; *Str1 || *Str2; ++Str1, ++Str2)
+	{
+		CharType Char1 = TChar<CharType>::ToLower(*Str1);
+		CharType Char2 = TChar<CharType>::ToLower(*Str2);
+
+		if (Char1 != Char2)
+		{
+			return TChar<CharType>::ToUnsigned(Char1) - TChar<CharType>::ToUnsigned(Char2);
+		}
+	}
+
+	return 0;
+}
+
+template<typename CharType>
+void TestStricmp(const CharType* Str1, const CharType* Str2, FAutomationTestBase& Test)
+{
+	Test.TestEqual("Stricmp()", FMath::Sign(StricmpImpl(Str1, Str2)), FMath::Sign(StricmpExpected(Str1, Str2)));
+}
+
+template<typename CharType>
+void RunStricmpTests(FAutomationTestBase& Test)
+{
+	// Test a range of single character strings starting with a few large ones 
+	const CharType Empty[1] = { '\0' };
+	for (int32 Char = -2; Char < 256; ++Char)
+	{
+		const CharType Current[2] = { (CharType)Char, '\0' };
+		const CharType Next[2] = { (CharType)Char + 1, '\0' };
+		const CharType CurrentPlusCasingDistance[2] = { (CharType)Char + ('a' - 'A'), '\0' };
+
+		TestStricmp(Current, Current, Test);
+		TestStricmp(Current, Empty, Test);
+		TestStricmp(Current, Next, Test);
+		TestStricmp(Next, Current, Test);
+		TestStricmp(Current, CurrentPlusCasingDistance, Test);
 	}
-}
\ No newline at end of file
+	
+	// Test various ASCII casings
+	const CharType HelloLower[] = { 'h', 'e', 'l', 'l', 'o', '\0' };
+	const CharType HelloUpper[] = { 'H', 'E', 'L', 'L', 'O', '\0' };
+	const CharType HelloMixed1[] = { 'H', 'e', 'L', 'L', 'o', '\0' };
+	const CharType HelloMixed2[] = { 'h', 'E', 'l', 'l', 'O', '\0' };
+	const CharType Hell0[] = { 'h', 'e', 'l', 'l', '0', '\0' };
+
+	TestStricmp(HelloLower, HelloLower, Test);
+	TestStricmp(HelloLower, HelloUpper, Test);
+	TestStricmp(HelloLower, HelloMixed1, Test);
+	TestStricmp(HelloLower, HelloMixed2, Test);
+	TestStricmp(HelloLower, Hell0, Test);
+}
+
+bool FGenericPlatformStricmpTest::RunTest(const FString& Parameters)
+{
+	RunStricmpTests<ANSICHAR>(*this);
+	RunStricmpTests<WIDECHAR>(*this);
+
+	return true;
+}
+
+#endif //WITH_DEV_AUTOMATION_TESTS
\ No newline at end of file
diff --git a/Engine/Source/Runtime/Core/Private/Tests/Misc/CharTest.cpp b/Engine/Source/Runtime/Core/Private/Tests/Misc/CharTest.cpp
--- a/Engine/Source/Runtime/Core/Private/Tests/Misc/CharTest.cpp
+++ b/Engine/Source/Runtime/Core/Private/Tests/Misc/CharTest.cpp
@@ -1,64 +1,41 @@
 // Copyright 1998-2018 Epic Games, Inc. All Rights Reserved.
 
+#if WITH_DEV_AUTOMATION_TESTS 
+
 #include "Misc/Char.h"
 #include "Misc/AutomationTest.h"
 #include <locale.h>
 #include <ctype.h>
 #include <wctype.h>
 
-#if WITH_DEV_AUTOMATION_TESTS 
-
 IMPLEMENT_SIMPLE_AUTOMATION_TEST(TCharTest, "System.Core.Misc.Char", EAutomationTestFlags::ApplicationContextMask | EAutomationTestFlags::SmokeFilter)
 
-namespace cstd
+namespace crt
 {
 	int tolower(ANSICHAR c) { return ::tolower(c); }
 	int toupper(ANSICHAR c) { return ::toupper(c); }
-	int islower(ANSICHAR c) { return ::islower((unsigned char)c); }
-	int isupper(ANSICHAR c) { return ::isupper((unsigned char)c); }
-	int isalpha(ANSICHAR c) { return ::isalpha((unsigned char)c); }
-	int isgraph(ANSICHAR c) { return ::isgraph((unsigned char)c); }
-	int isprint(ANSICHAR c) { return ::isprint((unsigned char)c); }
-	int ispunct(ANSICHAR c) { return ::ispunct((unsigned char)c); }
-	int isalnum(ANSICHAR c) { return ::isalnum((unsigned char)c); }
-	int isdigit(ANSICHAR c) { return ::isdigit((unsigned char)c); }
-	int isxdigit(ANSICHAR c) { return ::isxdigit((unsigned char)c); }
-	int isspace(ANSICHAR c) { return ::isspace((unsigned char)c); }
 
 	int tolower(WIDECHAR c) { return ::towlower(c); }
 	int toupper(WIDECHAR c) { return ::towupper(c); }
-	int islower(WIDECHAR c) { return ::iswlower(c); }
-	int isupper(WIDECHAR c) { return ::iswupper(c); }
-	int isalpha(WIDECHAR c) { return ::iswalpha(c); }
-	int isgraph(WIDECHAR c)	{ return ::iswgraph(c); }
-	int isprint(WIDECHAR c) { return ::iswprint(c); }
-	int ispunct(WIDECHAR c) { return ::iswpunct(c); }
-	int isalnum(WIDECHAR c) { return ::iswalnum(c); }
-	int isdigit(WIDECHAR c) { return ::iswdigit(c); }
-	int isxdigit(WIDECHAR c) { return ::iswxdigit(c); }
-	int isspace(WIDECHAR c) { return ::iswspace(c); }
 }
 
-bool TCharTest::RunTest(const FString& Parameters)
+template<typename CharType>
+void RunCharTests(FAutomationTestBase& Test, uint32 MaxChar)
 {
-	TestTrue(TEXT("C locale not used"), strcmp("C", setlocale(LC_CTYPE, nullptr)) == 0);
-	
-	for (int32 i = 0; i < 0x10000; ++i)
+	for (uint32 I = 0; I < MaxChar; ++I)
 	{
-		TCHAR C = static_cast<TCHAR>(i);
-		TestEqual("FChar::ToLower()", FChar::ToLower(C), cstd::tolower(C));
-		TestEqual("FChar::ToUpper()", FChar::ToUpper(C), cstd::toupper(C));
-		TestEqual("FChar::IsLower()", !!FChar::IsLower(C), !!cstd::islower(C));
-		TestEqual("FChar::IsUpper()", !!FChar::IsUpper(C), !!cstd::isupper(C));
-		TestEqual("FChar::IsAlpha()", !!FChar::IsAlpha(C), !!cstd::isalpha(C));
-		TestEqual("FChar::IsGraph()", !!FChar::IsGraph(C), !!cstd::isgraph(C));
-		TestEqual("FChar::IsPrint()", !!FChar::IsPrint(C), !!cstd::isprint(C));
-		TestEqual("FChar::IsPunct()", !!FChar::IsPunct(C), !!cstd::ispunct(C));
-		TestEqual("FChar::IsAlnum()", !!FChar::IsAlnum(C), !!cstd::isalnum(C));
-		TestEqual("FChar::IsDigit()", !!FChar::IsDigit(C), !!cstd::isdigit(C));
-		TestEqual("FChar::IsHexDigit()", !!FChar::IsHexDigit(C), !!cstd::isxdigit(C));
-		TestEqual("FChar::IsWhitespace()", !!FChar::IsWhitespace(C), !!cstd::isspace(C));
+		CharType C(I);
+		Test.TestEqual("TChar::ToLower()", TChar<CharType>::ToLower(C), crt::tolower(C));
+		Test.TestEqual("TChar::ToUpper()", TChar<CharType>::ToUpper(C), crt::toupper(C));
 	}
+}
+
+bool TCharTest::RunTest(const FString& Parameters)
+{
+	TestTrue(TEXT("C locale not used"), strcmp("C", setlocale(LC_CTYPE, nullptr)) == 0);
+
+	RunCharTests<ANSICHAR>(*this, 256);
+	RunCharTests<WIDECHAR>(*this, 65536);
 
 	return true;
 }
diff --git a/Engine/Source/Runtime/Core/Public/Android/AndroidString.h b/Engine/Source/Runtime/Core/Public/Android/AndroidString.h
--- a/Engine/Source/Runtime/Core/Public/Android/AndroidString.h
+++ b/Engine/Source/Runtime/Core/Public/Android/AndroidString.h
@@ -32,7 +32,6 @@ typedef FGenericWidePlatformString FAndroidPlatformString;
 
 #include "Misc/Char.h"
 #include "GenericPlatform/GenericPlatformMemory.h"
-#include "GenericPlatform/GenericPlatformStricmp.h"
 #include "GenericPlatform/GenericPlatformString.h"
 
 struct FAndroidPlatformString : public FGenericPlatformString
@@ -47,37 +46,6 @@ struct FAndroidPlatformString : public FGenericPlatformString
 		return Dest;
 	}
 
-	/**
-	 * Compares two strings case-insensitive.
-	 *
-	 * @param String1 First string to compare.
-	 * @param String2 Second string to compare.
-	 *
-	 * @returns Zero if both strings are equal. Greater than zero if first
-	 *          string is greater than the second one. Less than zero
-	 *          otherwise.
-	 */
-	template <typename CharType1, typename CharType2>
-	static inline int32 Stricmp(const CharType1* String1, const CharType2* String2)
-	{
-		return FGenericPlatformStricmp::Stricmp(String1, String2);
-	}
-
-	template <typename CharType>
-	static inline int32 Strnicmp( const CharType* String1, const CharType* String2, SIZE_T Count )
-	{
-		// walk the strings, comparing them case insensitively, up to a max size
-		for (; (*String1 || *String2) && Count > 0; String1++, String2++, Count--)
-		{
-			CharType Char1 = TChar<CharType>::ToUpper(*String1), Char2 = TChar<CharType>::ToUpper(*String2);
-			if (Char1 != Char2)
-			{
-				return Char1 - Char2;
-			}
-		}
-		return 0;
-	}
-
 	/** 
 	 * Widechar implementation 
 	 **/
diff --git a/Engine/Source/Runtime/Core/Public/GenericPlatform/GenericPlatformStricmp.h b/Engine/Source/Runtime/Core/Public/GenericPlatform/GenericPlatformStricmp.h
--- a/Engine/Source/Runtime/Core/Public/GenericPlatform/GenericPlatformStricmp.h
+++ b/Engine/Source/Runtime/Core/Public/GenericPlatform/GenericPlatformStricmp.h
@@ -3,106 +3,32 @@
 #pragma once
 
 #include "CoreTypes.h"
-#include "Templates/EnableIf.h"
-#include "Templates/UnrealTypeTraits.h"
-#include "Misc/Char.h"
-
 
 /**
- * This trait tells if given CharTypeA is comparison compatible with CharTypeB, i.e.
- * if CharTypeA contains whole CharTypeB character set on the same positions.
+ * Optimized locale  and CRT independent case-insensitive string comparisons
  *
- * This is true e.g. for WIDECHAR and ANSICHAR because first 256 characters of WIDECHAR is
- * basically ANSICHAR table.
- */
-template<typename A, typename B>	struct TIsComparisonCompatibleChar							{ enum { Value = false }; };
-template<>							struct TIsComparisonCompatibleChar<WIDECHAR, ANSICHAR>		{ enum { Value = true }; };
-template<>							struct TIsComparisonCompatibleChar<UTF8CHAR, ANSICHAR>		{ enum { Value = true }; };
-template<>							struct TIsComparisonCompatibleChar<UTF16CHAR, ANSICHAR>		{ enum { Value = true }; };
-template<>							struct TIsComparisonCompatibleChar<UTF32CHAR, ANSICHAR>		{ enum { Value = true }; };
-
-/**
- * Static struct that implements generic stricmp functionality.
+ * Only considers ASCII character casing, i.e. C locale semantics
+ *
+ * @return	Zero if strings are equal, greater than zero if first string is
+ *			greater than the second one and less than zero otherwise.
  */
 struct FGenericPlatformStricmp
 {
-private:
-	CORE_API static int32 CompatibleCharTypesStricmp(const TCHAR* String1, const TCHAR* String2);
-
-	/**
-	 * Compares two strings case-insensitive assuming that both char types are
-	 * compatible.
-	 *
-	 * @param String1 First string to compare.
-	 * @param String2 Second string to compare.
-	 * @returns Zero if both strings are equal. Greater than zero if first
-	 *          string is greater than the second one. Less than zero
-	 *          otherwise.
-	 */
-	template <typename CompatibleCharType1, typename CompatibleCharType2>
-	static inline int32 CompatibleCharTypesStricmp(const CompatibleCharType1* String1, const CompatibleCharType2* String2)
-	{
-		// Walk the strings, comparing them case insensitively.
-		for (; *String1 || *String2; String1++, String2++)
-		{
-			if(*String1 != *String2)
-			{
-				CompatibleCharType1 Char1 = TChar<CompatibleCharType1>::ToLower(*String1);
-				CompatibleCharType2 Char2 = TChar<CompatibleCharType2>::ToLower(*String2);
-
-				if (Char1 != Char2)
-				{
-					return Char1 - Char2;
-				}
-			}
-		}
-		return 0;
-	}
-
-	/**
-	 * Helper function wrapper to verify compatibility. This overload checks if CharType1 == CharType2.
-	 */
-	template <typename CharType1, typename CharType2>
-	static inline int32 StricmpImpl(const CharType1* String1, const CharType2* String2, typename TEnableIf<TIsSame<CharType1, CharType2>::Value>::Type* Dummy1 = nullptr)
-	{
-		return CompatibleCharTypesStricmp(String1, String2);
-	}
-
-	/**
-	 * Helper function wrapper to verify compatibility. This overload checks if
-	 * CharType1 is compatible with CharType2.
-	 */
-	template <typename CharType1, typename CharType2>
-	static inline int32 StricmpImpl(const CharType1* String1, const CharType2* String2, typename TEnableIf<TIsComparisonCompatibleChar<CharType1, CharType2>::Value>::Type* Dummy = nullptr)
-	{
-		return CompatibleCharTypesStricmp(String1, String2);
-	}
-
-	/**
-	 * Helper function wrapper to verify compatibility. This overload checks if
-	 * CharType2 is compatible with CharType1. This is the mirrored version of
-	 * above function, so we don't have to declare compatibility trait for both
-	 * ways.
-	 */
-	template <typename CharType1, typename CharType2>
-	static inline int32 StricmpImpl(const CharType1* String1, const CharType2* String2, typename TEnableIf<TIsComparisonCompatibleChar<CharType2, CharType1>::Value>::Type* Dummy = nullptr)
-	{
-		return CompatibleCharTypesStricmp(String1, String2);
-	}
-
-public:
-	/**
-	 * Compares two strings case-insensitive.
-	 *
-	 * @param String1 First string to compare.
-	 * @param String2 Second string to compare.
-	 * @returns Zero if both strings are equal. Greater than zero if first
-	 *          string is greater than the second one. Less than zero
-	 *          otherwise.
-	 */
-	template <typename CharType1, typename CharType2>
-	static inline int32 Stricmp(const CharType1* String1, const CharType2* String2)
-	{
-		return FGenericPlatformStricmp::StricmpImpl(String1, String2);
-	}
+	CORE_API static int32 Stricmp(const ANSICHAR* String1, const ANSICHAR* String2);
+	CORE_API static int32 Stricmp(const WIDECHAR* String1, const WIDECHAR* String2);
+	CORE_API static int32 Stricmp(const UTF8CHAR* String1, const UTF8CHAR* String2);
+	CORE_API static int32 Stricmp(const UTF16CHAR* String1, const UTF16CHAR* String2);
+	CORE_API static int32 Stricmp(const UTF32CHAR* String1, const UTF32CHAR* String2);
+
+	CORE_API static int32 Stricmp(const ANSICHAR* String1, const WIDECHAR* String2);
+	CORE_API static int32 Stricmp(const ANSICHAR* String1, const UTF8CHAR* String2);
+	CORE_API static int32 Stricmp(const ANSICHAR* String1, const UTF16CHAR* String2);
+	CORE_API static int32 Stricmp(const ANSICHAR* String1, const UTF32CHAR* String2);
+	CORE_API static int32 Stricmp(const WIDECHAR* String1, const ANSICHAR* String2);
+	CORE_API static int32 Stricmp(const UTF8CHAR* String1, const ANSICHAR* String2);
+	CORE_API static int32 Stricmp(const UTF16CHAR* String1, const ANSICHAR* String2);
+	CORE_API static int32 Stricmp(const UTF32CHAR* String1, const ANSICHAR* String2);
+
+	CORE_API static int32 Strnicmp(const ANSICHAR* String1, const ANSICHAR* String2, SIZE_T Count);
+	CORE_API static int32 Strnicmp(const WIDECHAR* String1, const WIDECHAR* String2, SIZE_T Count);
 };
diff --git a/Engine/Source/Runtime/Core/Public/GenericPlatform/GenericPlatformString.h b/Engine/Source/Runtime/Core/Public/GenericPlatform/GenericPlatformString.h
--- a/Engine/Source/Runtime/Core/Public/GenericPlatform/GenericPlatformString.h
+++ b/Engine/Source/Runtime/Core/Public/GenericPlatform/GenericPlatformString.h
@@ -4,12 +4,13 @@
 
 #include "CoreTypes.h"
 #include "Templates/EnableIf.h"
+#include "GenericPlatform/GenericPlatformStricmp.h"
 
 
 /**
  * Generic string implementation for most platforms
  */
-struct FGenericPlatformString
+struct FGenericPlatformString : public FGenericPlatformStricmp
 {
 	/**
 	 * Tests whether a particular character is a valid member of its encoding.
diff --git a/Engine/Source/Runtime/Core/Public/GenericPlatform/GenericWidePlatformString.h b/Engine/Source/Runtime/Core/Public/GenericPlatform/GenericWidePlatformString.h
--- a/Engine/Source/Runtime/Core/Public/GenericPlatform/GenericWidePlatformString.h
+++ b/Engine/Source/Runtime/Core/Public/GenericPlatform/GenericWidePlatformString.h
@@ -25,39 +25,6 @@ struct FGenericWidePlatformString : public FGenericPlatformString
 	}
 
 public:
-	/**
-	 * Compares two strings case-insensitive.
-	 *
-	 * @param String1 First string to compare.
-	 * @param String2 Second string to compare.
-	 *
-	 * @returns Zero if both strings are equal. Greater than zero if first
-	 *          string is greater than the second one. Less than zero
-	 *          otherwise.
-	 */
-	template <typename CharType1, typename CharType2>
-	static inline int32 Stricmp(const CharType1* String1, const CharType2* String2)
-	{
-		return FGenericPlatformStricmp::Stricmp(String1, String2);
-	}
-
-	template <typename CharType>
-	static inline int32 Strnicmp( const CharType* String1, const CharType* String2, SIZE_T Count )
-	{
-		// walk the strings, comparing them case insensitively, up to a max size
-		for (; (*String1 || *String2) && Count > 0; String1++, String2++, Count--)
-		{
-			if(*String1 != *String2)
-			{
-				CharType Char1 = TChar<CharType>::ToUpper(*String1), Char2 = TChar<CharType>::ToUpper(*String2);
-				if (Char1 != Char2)
-				{
-					return Char1 - Char2;
-				}
-			}
-		}
-		return 0;
-	}
 
 	/**
 	 * Unicode implementation
diff --git a/Engine/Source/Runtime/Core/Public/GenericPlatform/MicrosoftPlatformString.h b/Engine/Source/Runtime/Core/Public/GenericPlatform/MicrosoftPlatformString.h
--- a/Engine/Source/Runtime/Core/Public/GenericPlatform/MicrosoftPlatformString.h
+++ b/Engine/Source/Runtime/Core/Public/GenericPlatform/MicrosoftPlatformString.h
@@ -51,11 +51,6 @@ struct FMicrosoftPlatformString : public FGenericPlatformString
 		return _tcsncmp( String1, String2, Count );
 	}
 
-	static FORCEINLINE int32 Strnicmp( const WIDECHAR* String1, const WIDECHAR* String2, SIZE_T Count )
-	{
-		return _tcsnicmp( String1, String2, Count );
-	}
-
 	static FORCEINLINE int32 Strlen( const WIDECHAR* String )
 	{
 		return _tcslen( String );
@@ -162,45 +157,7 @@ struct FMicrosoftPlatformString : public FGenericPlatformString
 	{
 		return strncmp( String1, String2, Count );
 	}
-
-	/**
-	 * Compares two strings case-insensitive.
-	 *
-	 * Specialized version for ANSICHAR types.
-	 *
-	 * @param String1 First string to compare.
-	 * @param String2 Second string to compare.
-	 *
-	 * @returns Zero if both strings are equal. Greater than zero if first
-	 *          string is greater than the second one. Less than zero
-	 *          otherwise.
-	 */
-	static FORCEINLINE int32 Stricmp(const ANSICHAR* String1, const ANSICHAR* String2)
-	{
-		return _stricmp(String1, String2);
-	}
-
-	/**
-	 * Compares two strings case-insensitive.
-	 *
-	 * @param String1 First string to compare.
-	 * @param String2 Second string to compare.
-	 *
-	 * @returns Zero if both strings are equal. Greater than zero if first
-	 *          string is greater than the second one. Less than zero
-	 *          otherwise.
-	 */
-	template <typename CharType1, typename CharType2>
-	static FORCEINLINE int32 Stricmp(const CharType1* String1, const CharType2* String2)
-	{
-		return FGenericPlatformStricmp::Stricmp(String1, String2);
-	}
-
-	static FORCEINLINE int32 Strnicmp( const ANSICHAR* String1, const ANSICHAR* String2, SIZE_T Count )
-	{
-		return _strnicmp( String1, String2, Count );
-	}
-
+	
 	static FORCEINLINE int32 Strlen( const ANSICHAR* String )
 	{
 		return strlen( String ); 
diff --git a/Engine/Source/Runtime/Core/Public/GenericPlatform/StandardPlatformString.h b/Engine/Source/Runtime/Core/Public/GenericPlatform/StandardPlatformString.h
--- a/Engine/Source/Runtime/Core/Public/GenericPlatform/StandardPlatformString.h
+++ b/Engine/Source/Runtime/Core/Public/GenericPlatform/StandardPlatformString.h
@@ -26,39 +26,6 @@ struct FStandardPlatformString : public FGenericPlatformString
 	}
 
 public:
-	/**
-	 * Compares two strings case-insensitive.
-	 *
-	 * @param String1 First string to compare.
-	 * @param String2 Second string to compare.
-	 *
-	 * @returns Zero if both strings are equal. Greater than zero if first
-	 *          string is greater than the second one. Less than zero
-	 *          otherwise.
-	 */
-	template <typename CharType1, typename CharType2>
-	static inline int32 Stricmp(const CharType1* String1, const CharType2* String2)
-	{
-		return FGenericPlatformStricmp::Stricmp(String1, String2);
-	}
-
-	template <typename CharType>
-	static inline int32 Strnicmp( const CharType* String1, const CharType* String2, SIZE_T Count )
-	{
-		// walk the strings, comparing them case insensitively, up to a max size
-		for (; (*String1 || *String2) && Count > 0; String1++, String2++, Count--)
-		{
-			if(*String1 != *String2)
-			{
-				CharType Char1 = TChar<CharType>::ToUpper(*String1), Char2 = TChar<CharType>::ToUpper(*String2);
-				if (Char1 != Char2)
-				{
-					return Char1 - Char2;
-				}
-			}
-		}
-		return 0;
-	}
 
 	/**
 	 * Unicode implementation
diff --git a/Engine/Source/Runtime/Core/Public/Misc/Char.h b/Engine/Source/Runtime/Core/Public/Misc/Char.h
--- a/Engine/Source/Runtime/Core/Public/Misc/Char.h
+++ b/Engine/Source/Runtime/Core/Public/Misc/Char.h
@@ -3,6 +3,7 @@
 #pragma once
 
 #include "CoreTypes.h"
+#include "Traits/IntType.h"
 #include <ctype.h>
 #include <wctype.h>
 
@@ -28,103 +29,109 @@ template <> struct TLiteral<WIDECHAR>
 
 #define LITERAL(CharType, StringLiteral) TLiteral<CharType>::Select(StringLiteral, L##StringLiteral)
 
-/**
- * TChar
- * Set of utility functions operating on a single character. The functions
- * are specialized for ANSICHAR and TCHAR character types. You can use the
- * typedefs FChar and FCharAnsi for convenience.
- */
 
-template <typename T, const unsigned int Size>
+template <typename CharType, const unsigned int Size>
 struct TCharBase
 {
-	typedef T CharType;
-
-	static const CharType LineFeed				= L'\x000A';
-	static const CharType VerticalTab			= L'\x000B';
-	static const CharType FormFeed				= L'\x000C';
-	static const CharType CarriageReturn		= L'\x000D';
-	static const CharType NextLine				= L'\x0085';
-	static const CharType LineSeparator			= L'\x2028';
-	static const CharType ParagraphSeparator	= L'\x2029';
+	static constexpr CharType LineFeed = 0xa;
+	static constexpr CharType VerticalTab = 0xb;
+	static constexpr CharType FormFeed = 0xc;
+	static constexpr CharType CarriageReturn = 0xd;
+	static constexpr CharType NextLine = 0x85;
+	static constexpr CharType LineSeparator = 0x2028;
+	static constexpr CharType ParagraphSeparator = 0x2029;
+
+	static bool IsLinebreak(CharType Char)
+	{
+		return ((uint32(Char) - LineFeed) <= uint32(CarriageReturn - LineFeed)) |
+			(Char == NextLine) | (Char == LineSeparator) | (Char == ParagraphSeparator);
+	}
+
 };
 
-template <typename T>
-struct TCharBase<T, 1>
+template <typename CharType>
+struct TCharBase<CharType, 1>
 {
-	typedef T CharType;
+	static constexpr CharType LineFeed = 0xa;
+	static constexpr CharType VerticalTab = 0xb;
+	static constexpr CharType FormFeed = 0xc;
+	static constexpr CharType CarriageReturn = 0xd;
 
-	static const CharType LineFeed			= '\x000A';
-	static const CharType VerticalTab		= '\x000B';
-	static const CharType FormFeed			= '\x000C';
-	static const CharType CarriageReturn	= '\x000D';
-	static const CharType NextLine			= '\x0085';
+	static bool IsLinebreak(CharType Char)
+	{
+		return ((uint32(Char) - LineFeed) <= uint32(CarriageReturn - LineFeed));
+	}
 };
 
-template <typename T, const unsigned int Size>
-struct LineBreakImplementation
+
+/**
+ * TChar
+ * Set of utility functions operating on a single character. The functions
+ * are specialized for ANSICHAR and TCHAR character types. You can use the
+ * typedefs FChar and FCharAnsi for convenience.
+ */
+template <typename CharType>
+struct TChar : TCharBase<CharType, sizeof(CharType)>
 {
-	typedef T CharType;
-	static inline bool IsLinebreak( CharType c )
+	/**
+	* Only converts ASCII characters, same as CRT to[w]upper() with standard C locale
+	*/
+	static CharType ToUpper(CharType Char)
 	{
-		return	c == TCharBase<CharType, Size>::LineFeed			 
-			 || c == TCharBase<CharType, Size>::VerticalTab		 
-			 || c == TCharBase<CharType, Size>::FormFeed			 
-			 || c == TCharBase<CharType, Size>::CarriageReturn	 
-			 || c == TCharBase<CharType, Size>::NextLine			 
-			 || c == TCharBase<CharType, Size>::LineSeparator		 
-			 || c == TCharBase<CharType, Size>::ParagraphSeparator;
+		return ToUnsigned(Char) - ((uint32(Char) - 'a' < 26u) << 5);
 	}
-};
 
-template <typename T>
-struct LineBreakImplementation<T, 1>
-{
-	typedef T CharType;
-	static inline bool IsLinebreak( CharType c )
+	/**
+	* Only converts ASCII characters, same as CRT to[w]upper() with standard C locale
+	*/
+	static CharType ToLower(CharType Char)
 	{
-		return	c == TCharBase<CharType, 1>::LineFeed		 
-			 || c == TCharBase<CharType, 1>::VerticalTab	 
-			 || c == TCharBase<CharType, 1>::FormFeed		 
-			 || c == TCharBase<CharType, 1>::CarriageReturn 
-			 || c == TCharBase<CharType, 1>::NextLine	   ;
+		return ToUnsigned(Char) + ((uint32(Char) - 'A' < 26u) << 5);
 	}
-};
 
-template <typename T>
-struct TChar : public TCharBase<T, sizeof(T)>
-{
-	typedef T CharType;
-public:
-	static inline CharType ToUpper(CharType Char);
-	static inline CharType ToLower(CharType Char);
-	static inline bool IsUpper(CharType Char);
-	static inline bool IsLower(CharType Char);
-	static inline bool IsAlpha(CharType Char);
-	static inline bool IsGraph(CharType Char);
-	static inline bool IsPrint(CharType Char);
-	static inline bool IsPunct(CharType Char);
-
-	static inline bool IsAlnum(CharType Char);
-	static inline bool IsDigit(CharType Char);
-	static inline bool IsOctDigit(CharType Char)
+	static bool IsUpper(CharType Char);
+	static bool IsLower(CharType Char);
+	static bool IsAlpha(CharType Char);
+	static bool IsGraph(CharType Char);
+	static bool IsPrint(CharType Char);
+	static bool IsPunct(CharType Char);
+	static bool IsAlnum(CharType Char);
+	static bool IsDigit(CharType Char);
+	static bool IsHexDigit(CharType Char);
+	static bool IsWhitespace(CharType Char);
+
+	static bool IsOctDigit(CharType Char)
 	{
-		return Char >= '0' && Char <= '7';
+		return uint32(Char) - '0' < 8u;
 	}
-	static inline bool IsHexDigit(CharType Char);
-	static inline int32 ConvertCharDigitToInt(CharType Char)
+
+	static int32 ConvertCharDigitToInt(CharType Char)
 	{
 		return static_cast<int32>(Char) - static_cast<int32>('0');
 	}
-	static inline bool IsWhitespace(CharType Char);
-	static inline bool IsIdentifier(CharType Char)
+
+	static bool IsIdentifier(CharType Char)
 	{
 		return IsAlnum(Char) || IsUnderscore(Char);
 	}
-	static inline bool IsUnderscore(CharType Char)		{ return Char == LITERAL(CharType, '_'); }
 
-public:
-	static inline bool IsLinebreak(CharType Char)	{ return LineBreakImplementation<CharType, sizeof(CharType)>::IsLinebreak(Char); }
+	static bool IsUnderscore(CharType Char)
+	{
+		return Char == LITERAL(CharType, '_');
+	}
+
+	/**
+	* Avoid sign extension problems with signed characters smaller than int
+	*
+	* E.g. 'Ö' - 'A' is negative since the char 'Ö' (0xD6) is negative and gets
+	* sign-extended to the 32-bit int 0xFFFFFFD6 before subtraction happens.
+	*
+	* Mainly needed for subtraction and addition.
+	*/
+	static uint32 ToUnsigned(CharType Char)
+	{
+		return (TUnsignedIntType<sizeof(CharType)>::Type)Char;
+	}
 };
 
 typedef TChar<TCHAR>    FChar;
@@ -132,45 +139,30 @@ typedef TChar<WIDECHAR> FCharWide;
 typedef TChar<ANSICHAR> FCharAnsi;
 
 /*-----------------------------------------------------------------------------
-	TCHAR specialized functions
+	WIDECHAR specialized functions
 -----------------------------------------------------------------------------*/
 
-template <> inline TChar<WIDECHAR>::CharType TChar<WIDECHAR>::ToUpper(CharType Char)
-{
-#if !PLATFORM_TCHAR_IS_4_BYTES
-	// for 2 byte wide chars use ansi toupper() if it is not an actual wide char (value less than 128)
-	// towupper() is slow on some platforms.
-	if (Char < 128)
-	{
-		return ::toupper((ANSICHAR)Char);
-	}
-#endif
-	return ::towupper(Char);
-}
-template <> inline TChar<WIDECHAR>::CharType TChar<WIDECHAR>::ToLower(CharType Char)	{ return ::towlower(Char); }
-template <> inline bool TChar<WIDECHAR>::IsUpper(CharType Char)							{ return ::iswupper(Char) != 0; }
-template <> inline bool TChar<WIDECHAR>::IsLower(CharType Char)							{ return ::iswlower(Char) != 0; }
-template <> inline bool TChar<WIDECHAR>::IsAlpha(CharType Char)							{ return ::iswalpha(Char) != 0; }
-template <> inline bool TChar<WIDECHAR>::IsGraph(CharType Char)							{ return ::iswgraph(Char) != 0; }
-template <> inline bool TChar<WIDECHAR>::IsPrint(CharType Char)							{ return ::iswprint(Char) != 0; }
-template <> inline bool TChar<WIDECHAR>::IsPunct(CharType Char)							{ return ::iswpunct(Char) != 0; }
-template <> inline bool TChar<WIDECHAR>::IsAlnum(CharType Char)							{ return ::iswalnum(Char) != 0; }
-template <> inline bool TChar<WIDECHAR>::IsDigit(CharType Char)							{ return ::iswdigit(Char) != 0; }
-template <> inline bool TChar<WIDECHAR>::IsHexDigit(CharType Char)						{ return ::iswxdigit(Char) != 0; }
-template <> inline bool TChar<WIDECHAR>::IsWhitespace(CharType Char)					{ return ::iswspace(Char) != 0; }
+template <> inline bool TChar<WIDECHAR>::IsUpper(WIDECHAR Char) { return ::iswupper(Char) != 0; }
+template <> inline bool TChar<WIDECHAR>::IsLower(WIDECHAR Char) { return ::iswlower(Char) != 0; }
+template <> inline bool TChar<WIDECHAR>::IsAlpha(WIDECHAR Char) { return ::iswalpha(Char) != 0; }
+template <> inline bool TChar<WIDECHAR>::IsGraph(WIDECHAR Char) { return ::iswgraph(Char) != 0; }
+template <> inline bool TChar<WIDECHAR>::IsPrint(WIDECHAR Char) { return ::iswprint(Char) != 0; }
+template <> inline bool TChar<WIDECHAR>::IsPunct(WIDECHAR Char) { return ::iswpunct(Char) != 0; }
+template <> inline bool TChar<WIDECHAR>::IsAlnum(WIDECHAR Char) { return ::iswalnum(Char) != 0; }
+template <> inline bool TChar<WIDECHAR>::IsDigit(WIDECHAR Char) { return ::iswdigit(Char) != 0; }
+template <> inline bool TChar<WIDECHAR>::IsHexDigit(WIDECHAR Char) { return ::iswxdigit(Char) != 0; }
+template <> inline bool TChar<WIDECHAR>::IsWhitespace(WIDECHAR Char) { return ::iswspace(Char) != 0; }
 
 /*-----------------------------------------------------------------------------
 	ANSICHAR specialized functions
 -----------------------------------------------------------------------------*/
-template <> inline TChar<ANSICHAR>::CharType TChar<ANSICHAR>::ToUpper(CharType Char)	{ return ::toupper(Char); }
-template <> inline TChar<ANSICHAR>::CharType TChar<ANSICHAR>::ToLower(CharType Char)	{ return ::tolower(Char); }
-template <> inline bool TChar<ANSICHAR>::IsUpper(CharType Char)							{ return ::isupper((unsigned char)Char) != 0; }
-template <> inline bool TChar<ANSICHAR>::IsLower(CharType Char)							{ return ::islower((unsigned char)Char) != 0; }
-template <> inline bool TChar<ANSICHAR>::IsAlpha(CharType Char)							{ return ::isalpha((unsigned char)Char) != 0; }
-template <> inline bool TChar<ANSICHAR>::IsGraph(CharType Char)							{ return ::isgraph((unsigned char)Char) != 0; }
-template <> inline bool TChar<ANSICHAR>::IsPrint(CharType Char)							{ return ::isprint((unsigned char)Char) != 0; }
-template <> inline bool TChar<ANSICHAR>::IsPunct(CharType Char)							{ return ::ispunct((unsigned char)Char) != 0; }
-template <> inline bool TChar<ANSICHAR>::IsAlnum(CharType Char)							{ return ::isalnum((unsigned char)Char) != 0; }
-template <> inline bool TChar<ANSICHAR>::IsDigit(CharType Char)							{ return ::isdigit((unsigned char)Char) != 0; }
-template <> inline bool TChar<ANSICHAR>::IsHexDigit(CharType Char)						{ return ::isxdigit((unsigned char)Char) != 0; }
-template <> inline bool TChar<ANSICHAR>::IsWhitespace(CharType Char)					{ return ::isspace((unsigned char)Char) != 0; }
+template <> inline bool TChar<ANSICHAR>::IsUpper(ANSICHAR Char) { return ::isupper((unsigned char)Char) != 0; }
+template <> inline bool TChar<ANSICHAR>::IsLower(ANSICHAR Char) { return ::islower((unsigned char)Char) != 0; }
+template <> inline bool TChar<ANSICHAR>::IsAlpha(ANSICHAR Char) { return ::isalpha((unsigned char)Char) != 0; }
+template <> inline bool TChar<ANSICHAR>::IsGraph(ANSICHAR Char) { return ::isgraph((unsigned char)Char) != 0; }
+template <> inline bool TChar<ANSICHAR>::IsPrint(ANSICHAR Char) { return ::isprint((unsigned char)Char) != 0; }
+template <> inline bool TChar<ANSICHAR>::IsPunct(ANSICHAR Char) { return ::ispunct((unsigned char)Char) != 0; }
+template <> inline bool TChar<ANSICHAR>::IsAlnum(ANSICHAR Char) { return ::isalnum((unsigned char)Char) != 0; }
+template <> inline bool TChar<ANSICHAR>::IsDigit(ANSICHAR Char) { return ::isdigit((unsigned char)Char) != 0; }
+template <> inline bool TChar<ANSICHAR>::IsHexDigit(ANSICHAR Char) { return ::isxdigit((unsigned char)Char) != 0; }
+template <> inline bool TChar<ANSICHAR>::IsWhitespace(ANSICHAR Char) { return ::isspace((unsigned char)Char) != 0; }
