Copy constructor and assignment operator of TEnumAsByte removed to improve error message.
Fixes for warnings now generated by this change.

diff --git a/Engine/Source/Runtime/Core/Public/Containers/EnumAsByte.h b/Engine/Source/Runtime/Core/Public/Containers/EnumAsByte.h
--- a/Engine/Source/Runtime/Core/Public/Containers/EnumAsByte.h
+++ b/Engine/Source/Runtime/Core/Public/Containers/EnumAsByte.h
@@ -21,17 +21,9 @@ class TEnumAsByte
 public:
 	typedef TEnum EnumType;
 
-	/** Default Constructor (no initialization). */
-	FORCEINLINE TEnumAsByte() { }
-
-	/**
-	 * Copy constructor.
-	 *
-	 * @param InValue value to construct with.
-	 */
-	FORCEINLINE TEnumAsByte( const TEnumAsByte &InValue )
-		: Value(InValue.Value)
-	{ }
+	TEnumAsByte() = default;
+	TEnumAsByte(const TEnumAsByte&) = default;
+	TEnumAsByte& operator=(const TEnumAsByte&) = default;
 
 	/**
 	 * Constructor, initialize to the enum value.
@@ -61,30 +53,6 @@ public:
 	{ }
 
 public:
-
-	/**
-	 * Assignment operator.
-	 *
-	 * @param InValue value to set.
-	 * @return This instance.
-	 */
-	FORCEINLINE TEnumAsByte& operator=( TEnumAsByte InValue )
-	{
-		Value = InValue.Value;
-		return *this;
-	}
-	/**
-	 * Assignment operator.
-	 *
-	 * @param InValue value to set.
-	 * @return This instance.
-	 */
-	FORCEINLINE TEnumAsByte& operator=( TEnum InValue )
-	{
-		Value = static_cast<uint8>(InValue);
-		return *this;
-	}
-
 	/**
 	 * Compares two enumeration values for equality.
 	 *
diff --git a/Engine/Source/Runtime/NavigationSystem/Private/NavigationData.cpp b/Engine/Source/Runtime/NavigationSystem/Private/NavigationData.cpp
--- a/Engine/Source/Runtime/NavigationSystem/Private/NavigationData.cpp
+++ b/Engine/Source/Runtime/NavigationSystem/Private/NavigationData.cpp
@@ -86,6 +86,7 @@ FAsyncPathFindingQuery::FAsyncPathFindingQuery(const UObject* InOwner, const ANa
 : FPathFindingQuery(InOwner, InNavData, Start, End, SourceQueryFilter)
 , QueryID(GetUniqueID())
 , OnDoneDelegate(Delegate)
+, Mode(EPathFindingMode::Regular)
 {
 
 }
diff --git a/Engine/Source/Runtime/NavigationSystem/Public/NavigationData.h b/Engine/Source/Runtime/NavigationSystem/Public/NavigationData.h
--- a/Engine/Source/Runtime/NavigationSystem/Public/NavigationData.h
+++ b/Engine/Source/Runtime/NavigationSystem/Public/NavigationData.h
@@ -969,6 +969,7 @@ struct FAsyncPathFindingQuery : public FPathFindingQuery
 
 	FAsyncPathFindingQuery()
 		: QueryID(INVALID_NAVQUERYID)
+		, Mode(EPathFindingMode::Regular)
 	{ }
 
 	FAsyncPathFindingQuery(const UObject* InOwner, const ANavigationData& InNavData, const FVector& Start, const FVector& End, const FNavPathQueryDelegate& Delegate, FSharedConstNavQueryFilter SourceQueryFilter);
