FScriptSetLayout::ElementOffset and FScriptMapLayout::KeyOffset removed, since they're always zero and applied inconsistently.

diff --git a/Engine/Source/Editor/UnrealEd/Private/Commandlets/FixConflictingLocalizationKeys.cpp b/Engine/Source/Editor/UnrealEd/Private/Commandlets/FixConflictingLocalizationKeys.cpp
--- a/Engine/Source/Editor/UnrealEd/Private/Commandlets/FixConflictingLocalizationKeys.cpp
+++ b/Engine/Source/Editor/UnrealEd/Private/Commandlets/FixConflictingLocalizationKeys.cpp
@@ -231,7 +231,7 @@ bool ReKeyTextProperty(UStruct* InOuterType, void* InAddrToUpdate, const TArray<
 
 						if (ScriptSetHelper.IsValidIndex(SparseIndex))
 						{
-							AddrToUpdate = ScriptSetHelper.GetElementPtr(SparseIndex) + SetProp->SetLayout.ElementOffset;
+							AddrToUpdate = ScriptSetHelper.GetElementPtr(SparseIndex);
 
 							// Is this a complex property? If so, we need to recurse into it
 							if (UStructProperty* StructProp = Cast<UStructProperty>(SetProp->ElementProp))
@@ -253,7 +253,7 @@ bool ReKeyTextProperty(UStruct* InOuterType, void* InAddrToUpdate, const TArray<
 						FScriptMapHelper ScriptMapHelper(MapProp, AddrToUpdate);
 						if (ScriptMapHelper.IsValidIndex(ContainerIndex))
 						{
-							AddrToUpdate = ScriptMapHelper.GetPairPtr(ContainerIndex) + MapProp->MapLayout.KeyOffset;
+							AddrToUpdate = ScriptMapHelper.GetPairPtr(ContainerIndex);
 
 							// Is this a complex property? If so, we need to recurse into it
 							if (UStructProperty* StructProp = Cast<UStructProperty>(MapProp->KeyProp))
diff --git a/Engine/Source/Runtime/Core/Public/Containers/Map.h b/Engine/Source/Runtime/Core/Public/Containers/Map.h
--- a/Engine/Source/Runtime/Core/Public/Containers/Map.h
+++ b/Engine/Source/Runtime/Core/Public/Containers/Map.h
@@ -1362,7 +1362,7 @@ public:
 
 struct FScriptMapLayout
 {
-	int32 KeyOffset;
+	// int32 KeyOffset; // is always at zero offset from the TPair - not stored here
 	int32 ValueOffset;
 
 	FScriptSetLayout SetLayout;
@@ -1380,10 +1380,12 @@ public:
 
 		// TPair<Key, Value>
 		FStructBuilder PairStruct;
-		Result.KeyOffset   = PairStruct.AddMember(KeySize,   KeyAlignment);
+		int32 KeyOffset    = PairStruct.AddMember(KeySize,   KeyAlignment);
 		Result.ValueOffset = PairStruct.AddMember(ValueSize, ValueAlignment);
 		Result.SetLayout   = FScriptSet::GetScriptLayout(PairStruct.GetSize(),  PairStruct.GetAlignment());
 
+		checkf(KeyOffset == 0, TEXT("The key inside the TPair is expected to be at the start of the struct"));
+
 		return Result;
 	}
 
@@ -1456,7 +1458,7 @@ public:
 				GetKeyHash, // We 'know' that the implementation of Find doesn't call GetKeyHash on anything except Key
 				[KeyEqualityFn, MapLayout](const void* InKey, const void* InPair )
 				{
-					return KeyEqualityFn(InKey, (uint8*)InPair + MapLayout.KeyOffset);
+					return KeyEqualityFn(InKey, (uint8*)InPair);
 				}
 			);
 		}
@@ -1497,13 +1499,13 @@ public:
 			KeyEqualityFn,
 			[KeyConstructAndAssignFn, ValueConstructAndAssignFn, Layout](void* NewPair)
 			{
-				KeyConstructAndAssignFn((uint8*)NewPair + Layout.KeyOffset);
+				KeyConstructAndAssignFn((uint8*)NewPair);
 				ValueConstructAndAssignFn((uint8*)NewPair + Layout.ValueOffset);
 			},
 			[DestructKeyFn, DestructValueFn, Layout](void* NewPair)
 			{
 				DestructValueFn((uint8*)NewPair + Layout.ValueOffset);
-				DestructKeyFn((uint8*)NewPair + Layout.KeyOffset);
+				DestructKeyFn((uint8*)NewPair);
 			}
 		);
 	}
diff --git a/Engine/Source/Runtime/Core/Public/Containers/Set.h b/Engine/Source/Runtime/Core/Public/Containers/Set.h
--- a/Engine/Source/Runtime/Core/Public/Containers/Set.h
+++ b/Engine/Source/Runtime/Core/Public/Containers/Set.h
@@ -1310,7 +1310,7 @@ struct TContainerTraits<TSet<ElementType, KeyFuncs, Allocator> > : public TConta
 
 struct FScriptSetLayout
 {
-	int32 ElementOffset;
+	// int32 ElementOffset = 0; // always at zero offset from the TSetElement - not stored here
 	int32 HashNextIdOffset;
 	int32 HashIndexOffset;
 	int32 Size;
@@ -1329,12 +1329,14 @@ public:
 
 		// TSetElement<TPair<Key, Value>>
 		FStructBuilder SetElementStruct;
-		Result.ElementOffset     = SetElementStruct.AddMember(ElementSize,           ElementAlignment);
+		int32 ElementOffset      = SetElementStruct.AddMember(ElementSize,           ElementAlignment);
 		Result.HashNextIdOffset  = SetElementStruct.AddMember(sizeof(FSetElementId), alignof(FSetElementId));
 		Result.HashIndexOffset   = SetElementStruct.AddMember(sizeof(int32),         alignof(int32));
 		Result.Size              = SetElementStruct.GetSize();
 		Result.SparseArrayLayout = FScriptSparseArray::GetScriptLayout(SetElementStruct.GetSize(), SetElementStruct.GetAlignment());
 
+		checkf(ElementOffset == 0, TEXT("The element inside the TSetElement is expected to be at the start of the struct"));
+
 		return Result;
 	}
 
diff --git a/Engine/Source/Runtime/Core/Public/Containers/SparseArray.h b/Engine/Source/Runtime/Core/Public/Containers/SparseArray.h
--- a/Engine/Source/Runtime/Core/Public/Containers/SparseArray.h
+++ b/Engine/Source/Runtime/Core/Public/Containers/SparseArray.h
@@ -999,7 +999,7 @@ template <typename T, typename Allocator> void* operator new( size_t Size, TSpar
 
 struct FScriptSparseArrayLayout
 {
-	int32 ElementOffset;
+	// ElementOffset is at zero offset from the TSparseArrayElementOrFreeListLink - not stored here
 	int32 Alignment;
 	int32 Size;
 };
@@ -1012,7 +1012,6 @@ public:
 	static FScriptSparseArrayLayout GetScriptLayout(int32 ElementSize, int32 ElementAlignment)
 	{
 		FScriptSparseArrayLayout Result;
-		Result.ElementOffset = 0;
 		Result.Alignment     = FMath::Max(ElementAlignment, (int32)alignof(FFreeListLink));
 		Result.Size          = FMath::Max(ElementSize,      (int32)sizeof (FFreeListLink));
 
diff --git a/Engine/Source/Runtime/CoreUObject/Private/Serialization/PropertyLocalizationDataGathering.cpp b/Engine/Source/Runtime/CoreUObject/Private/Serialization/PropertyLocalizationDataGathering.cpp
--- a/Engine/Source/Runtime/CoreUObject/Private/Serialization/PropertyLocalizationDataGathering.cpp
+++ b/Engine/Source/Runtime/CoreUObject/Private/Serialization/PropertyLocalizationDataGathering.cpp
@@ -285,7 +285,7 @@ void FPropertyLocalizationDataGatherer::GatherLocalizationDataFromChildTextPrope
 				const uint8* MapPairPtr = ScriptMapHelper.GetPairPtr(j);
 				if (CanGatherFromInnerProperty(MapProperty->KeyProp))
 				{
-				GatherLocalizationDataFromChildTextProperties(PathToElement + FString::Printf(TEXT("(%d - Key)"), ElementIndex), MapProperty->KeyProp, MapPairPtr + MapProperty->MapLayout.KeyOffset, nullptr, ChildPropertyGatherTextFlags);
+				GatherLocalizationDataFromChildTextProperties(PathToElement + FString::Printf(TEXT("(%d - Key)"), ElementIndex), MapProperty->KeyProp, MapPairPtr, nullptr, ChildPropertyGatherTextFlags);
 				}
 				if (CanGatherFromInnerProperty(MapProperty->ValueProp))
 				{
@@ -314,7 +314,7 @@ void FPropertyLocalizationDataGatherer::GatherLocalizationDataFromChildTextPrope
 				}
 
 				const uint8* ElementPtr = ScriptSetHelper.GetElementPtr(j);
-				GatherLocalizationDataFromChildTextProperties(PathToElement + FString::Printf(TEXT("(%d)"), ElementIndex), SetProperty->ElementProp, ElementPtr + SetProperty->SetLayout.ElementOffset, nullptr, ChildPropertyGatherTextFlags);
+				GatherLocalizationDataFromChildTextProperties(PathToElement + FString::Printf(TEXT("(%d)"), ElementIndex), SetProperty->ElementProp, ElementPtr, nullptr, ChildPropertyGatherTextFlags);
 				++ElementIndex;
 			}
 		}
diff --git a/Engine/Source/Runtime/CoreUObject/Private/UObject/PropertySet.cpp b/Engine/Source/Runtime/CoreUObject/Private/UObject/PropertySet.cpp
--- a/Engine/Source/Runtime/CoreUObject/Private/UObject/PropertySet.cpp
+++ b/Engine/Source/Runtime/CoreUObject/Private/UObject/PropertySet.cpp
@@ -489,14 +489,14 @@ void USetProperty::ExportTextItem(FString& ValueStr, const void* PropertyValue,
 		checkSlow(StructElementProp->Struct);
 
 		StructDefaults = (uint8*)FMemory::Malloc(SetLayout.Size);
-		ElementProp->InitializeValue(StructDefaults + SetLayout.ElementOffset);
+		ElementProp->InitializeValue(StructDefaults);
 	}
 
 	ON_SCOPE_EXIT
 	{
 		if (StructDefaults)
 		{
-			ElementProp->DestroyValue(StructDefaults + SetLayout.ElementOffset);
+			ElementProp->DestroyValue(StructDefaults);
 			FMemory::Free(StructDefaults);
 		}
 	};
diff --git a/Engine/Source/Runtime/CoreUObject/Public/UObject/UnrealType.h b/Engine/Source/Runtime/CoreUObject/Public/UObject/UnrealType.h
--- a/Engine/Source/Runtime/CoreUObject/Public/UObject/UnrealType.h
+++ b/Engine/Source/Runtime/CoreUObject/Public/UObject/UnrealType.h
@@ -3090,7 +3090,7 @@ public:
 		}
 		
 		checkSlow(IsValidIndex(Index));
-		return (uint8*)Map->GetData(Index, MapLayout) + MapLayout.KeyOffset;
+		return (uint8*)Map->GetData(Index, MapLayout);
 	}
 
 	/**
diff --git a/Engine/Source/Runtime/Serialization/Private/StructDeserializer.cpp b/Engine/Source/Runtime/Serialization/Private/StructDeserializer.cpp
--- a/Engine/Source/Runtime/Serialization/Private/StructDeserializer.cpp
+++ b/Engine/Source/Runtime/Serialization/Private/StructDeserializer.cpp
@@ -204,7 +204,7 @@ bool FStructDeserializer::Deserialize( void* OutStruct, UStruct& TypeInfo, IStru
 					int32 PairIndex = MapHelper.AddDefaultValue_Invalid_NeedsRehash();
 					uint8* PairPtr = MapHelper.GetPairPtr(PairIndex);
 
-					MapProperty->KeyProp->ImportText(*PropertyName, PairPtr + MapProperty->MapLayout.KeyOffset, PPF_None, nullptr);
+					MapProperty->KeyProp->ImportText(*PropertyName, PairPtr, PPF_None, nullptr);
 
 					if (!Backend.ReadProperty(Property, CurrentState.Property, PairPtr, CurrentState.ArrayIndex))
 					{
@@ -326,7 +326,7 @@ bool FStructDeserializer::Deserialize( void* OutStruct, UStruct& TypeInfo, IStru
 					NewState.Data = PairPtr + MapHelper.MapLayout.ValueOffset;
 					NewState.Property = MapProperty->ValueProp;
 
-					MapProperty->KeyProp->ImportText(*PropertyName, PairPtr + MapProperty->MapLayout.KeyOffset, PPF_None, nullptr);
+					MapProperty->KeyProp->ImportText(*PropertyName, PairPtr, PPF_None, nullptr);
 				}
 				else if ((CurrentState.Property != nullptr) && (CurrentState.Property->GetClass() == USetProperty::StaticClass()))
 				{
@@ -336,7 +336,7 @@ bool FStructDeserializer::Deserialize( void* OutStruct, UStruct& TypeInfo, IStru
 					const int32 ElementIndex = SetHelper.AddDefaultValue_Invalid_NeedsRehash();
 					uint8* ElementPtr = SetHelper.GetElementPtr(ElementIndex);
 
-					NewState.Data = ElementPtr + SetHelper.SetLayout.ElementOffset;
+					NewState.Data = ElementPtr;
 					NewState.Property = SetProperty->ElementProp;
 				}
 				else
diff --git a/Engine/Source/Runtime/Serialization/Private/StructSerializer.cpp b/Engine/Source/Runtime/Serialization/Private/StructSerializer.cpp
--- a/Engine/Source/Runtime/Serialization/Private/StructSerializer.cpp
+++ b/Engine/Source/Runtime/Serialization/Private/StructSerializer.cpp
@@ -214,7 +214,7 @@ void FStructSerializer::Serialize( const void* Struct, UStruct& TypeInfo, IStruc
 						FStructSerializerState NewState;
 						{
 							NewState.HasBeenProcessed = false;
-							NewState.KeyData = PairPtr + MapProperty->MapLayout.KeyOffset;
+							NewState.KeyData = PairPtr;
 							NewState.KeyProperty = MapProperty->KeyProp;
 							NewState.ValueData = PairPtr;
 							NewState.ValueProperty = ValueProperty;
