Compile fix - missing typename before dependent type

diff --git a/Engine/Source/Runtime/Core/Public/Misc/Char.h b/Engine/Source/Runtime/Core/Public/Misc/Char.h
--- a/Engine/Source/Runtime/Core/Public/Misc/Char.h
+++ b/Engine/Source/Runtime/Core/Public/Misc/Char.h
@@ -130,7 +130,7 @@ struct TChar : TCharBase<CharType, sizeof(CharType)>
 	*/
 	static uint32 ToUnsigned(CharType Char)
 	{
-		return (TUnsignedIntType<sizeof(CharType)>::Type)Char;
+		return (typename TUnsignedIntType<sizeof(CharType)>::Type)Char;
 	}
 };
 
