Fix Linux compile error - array subscript is of type 'char'

diff --git a/Engine/Source/Runtime/Core/Private/GenericPlatform/GenericPlatformStricmp.cpp b/Engine/Source/Runtime/Core/Private/GenericPlatform/GenericPlatformStricmp.cpp
--- a/Engine/Source/Runtime/Core/Private/GenericPlatform/GenericPlatformStricmp.cpp
+++ b/Engine/Source/Runtime/Core/Private/GenericPlatform/GenericPlatformStricmp.cpp
@@ -41,7 +41,7 @@ int32 StricmpImpl(const CharType1* String1, const CharType2* String2)
 		}
 		else if (BothAscii(C1, C2))
 		{
-			if (int32 Diff = LowerAscii[C1] - LowerAscii[C2])
+			if (int32 Diff = LowerAscii[TChar<CharType1>::ToUnsigned(C1)] - LowerAscii[TChar<CharType2>::ToUnsigned(C2)])
 			{
 				return Diff;
 			}
@@ -74,7 +74,7 @@ int32 StrnicmpImpl(const CharType* String1, const CharType* String2, SIZE_T Coun
 		}
 		else if (BothAscii(C1, C2))
 		{
-			if (int32 Diff = LowerAscii[C1] - LowerAscii[C2])
+			if (int32 Diff = LowerAscii[TChar<CharType>::ToUnsigned(C1)] - LowerAscii[TChar<CharType>::ToUnsigned(C2)])
 			{
 				return Diff;
 			}
