Mac compile fix

diff --git a/Engine/Source/Runtime/Core/Private/GenericPlatform/GenericPlatformStricmp.cpp b/Engine/Source/Runtime/Core/Private/GenericPlatform/GenericPlatformStricmp.cpp
--- a/Engine/Source/Runtime/Core/Private/GenericPlatform/GenericPlatformStricmp.cpp
+++ b/Engine/Source/Runtime/Core/Private/GenericPlatform/GenericPlatformStricmp.cpp
@@ -144,9 +144,9 @@ void RunStricmpTests(FAutomationTestBase& Test)
 	const CharType Empty[1] = { '\0' };
 	for (int32 Char = -2; Char < 256; ++Char)
 	{
-		const CharType Current[2] = { (CharType)Char, '\0' };
-		const CharType Next[2] = { (CharType)Char + 1, '\0' };
-		const CharType CurrentPlusCasingDistance[2] = { (CharType)Char + ('a' - 'A'), '\0' };
+		const CharType Current[2] = { (CharType)(Char), '\0' };
+		const CharType Next[2] = { (CharType)(Char + 1), '\0' };
+		const CharType CurrentPlusCasingDistance[2] = { (CharType)(Char + ('a' - 'A')), '\0' };
 
 		TestStricmp(Current, Current, Test);
 		TestStricmp(Current, Empty, Test);
