Speculative Linux compile fix

diff --git a/Engine/Source/Runtime/Core/Public/GenericPlatform/GenericWidePlatformString.h b/Engine/Source/Runtime/Core/Public/GenericPlatform/GenericWidePlatformString.h
--- a/Engine/Source/Runtime/Core/Public/GenericPlatform/GenericWidePlatformString.h
+++ b/Engine/Source/Runtime/Core/Public/GenericPlatform/GenericWidePlatformString.h
@@ -6,6 +6,7 @@
 #include "Misc/Char.h"
 #include "GenericPlatform/GenericPlatformStricmp.h"
 #include "GenericPlatform/GenericPlatformString.h"
+#include "HAL/PlatformCrt.h"
 
 #if PLATFORM_APPLE || PLATFORM_UNIX || PLATFORM_HTML5 || PLATFORM_PS4 || PLATFORM_SWITCH || PLATFORM_ANDROID
 
