GitHub #5317 : Fix VARARGS family of macro which would not close the va_list scope as per the standard

diff --git a/Engine/Source/Programs/PixelStreaming/WebRTCProxy/src/Logging.cpp b/Engine/Source/Programs/PixelStreaming/WebRTCProxy/src/Logging.cpp
--- a/Engine/Source/Programs/PixelStreaming/WebRTCProxy/src/Logging.cpp
+++ b/Engine/Source/Programs/PixelStreaming/WebRTCProxy/src/Logging.cpp
@@ -103,4 +103,6 @@ void ILogOutput::LogToAll(
 	{
 		Out->Log(File, Line, Category, Verbosity, Buf);
 	}
+
+	va_end(Args);
 }
diff --git a/Engine/Source/Programs/PixelStreaming/WebRTCProxy/src/StringUtils.cpp b/Engine/Source/Programs/PixelStreaming/WebRTCProxy/src/StringUtils.cpp
--- a/Engine/Source/Programs/PixelStreaming/WebRTCProxy/src/StringUtils.cpp
+++ b/Engine/Source/Programs/PixelStreaming/WebRTCProxy/src/StringUtils.cpp
@@ -29,6 +29,8 @@ void SNPrintf(char* OutBuffer, int BufSize, _Printf_format_string_ const char* F
 	va_list Args;
 	va_start(Args, Fmt);
 	VSNPrintf(OutBuffer, BufSize, Fmt, Args);
+	va_end(Args);
+
 }
 
 char* GetTemporaryString()
diff --git a/Engine/Source/Programs/UnrealFrontend/Private/Commands/StatsConvertCommand.cpp b/Engine/Source/Programs/UnrealFrontend/Private/Commands/StatsConvertCommand.cpp
--- a/Engine/Source/Programs/UnrealFrontend/Private/Commands/StatsConvertCommand.cpp
+++ b/Engine/Source/Programs/UnrealFrontend/Private/Commands/StatsConvertCommand.cpp
@@ -120,6 +120,8 @@ void FCSVStatsProfiler::WriteString( const ANSICHAR* Format, ... )
 	va_start( ArgPtr, Format );
 	// Build the string.
 	int32 Result = FCStringAnsi::GetVarArgs( Array, ARRAY_COUNT( Array ), Format, ArgPtr );
+	va_end(ArgPtr);
+
 	// Now write that to the file.
 	CSVWriter->Serialize( (void*)Array, Result );
 }
diff --git a/Engine/Source/Runtime/Core/Private/GenericPlatform/GenericWidePlatformString.cpp b/Engine/Source/Runtime/Core/Private/GenericPlatform/GenericWidePlatformString.cpp
--- a/Engine/Source/Runtime/Core/Private/GenericPlatform/GenericWidePlatformString.cpp
+++ b/Engine/Source/Runtime/Core/Private/GenericPlatform/GenericWidePlatformString.cpp
@@ -214,8 +214,9 @@ int32 TestGetVarArgs(WIDECHAR* OutputString, const WIDECHAR* Format, ...)
 {
 	va_list ArgPtr;
 	va_start( ArgPtr, Format );
-
-	return FGenericWidePlatformString::GetVarArgs(OutputString, OUTPUT_SIZE, Format, ArgPtr);
+	int32 Result = FGenericWidePlatformString::GetVarArgs(OutputString, OUTPUT_SIZE, Format, ArgPtr);
+	va_end(ArgPtr);
+	return Result;
 }
 
 void RunGetVarArgsTests()
diff --git a/Engine/Source/Runtime/Core/Public/Misc/VarArgs.h b/Engine/Source/Runtime/Core/Public/Misc/VarArgs.h
--- a/Engine/Source/Runtime/Core/Public/Misc/VarArgs.h
+++ b/Engine/Source/Runtime/Core/Public/Misc/VarArgs.h
@@ -20,18 +20,21 @@ VarArgs helper macros.
 		va_list ap; \
 		va_start(ap, lastarg); \
 		FCString::GetVarArgs(msg, msgsize, fmt, ap); \
+		va_end(ap); \
 	}
 #define GET_VARARGS_WIDE(msg, msgsize, len, lastarg, fmt) \
 	{ \
 		va_list ap; \
 		va_start(ap, lastarg); \
 		FCStringWide::GetVarArgs(msg, msgsize, fmt, ap); \
+		va_end(ap); \
 	}
 #define GET_VARARGS_ANSI(msg, msgsize, len, lastarg, fmt) \
 	{ \
 		va_list ap; \
 		va_start(ap, lastarg); \
 		FCStringAnsi::GetVarArgs(msg, msgsize, fmt, ap); \
+		va_end(ap); \
 	}
 #define GET_VARARGS_RESULT(msg, msgsize, len, lastarg, fmt, result) \
 	{ \
@@ -42,6 +45,7 @@ VarArgs helper macros.
 		{ \
 			result = -1; \
 		} \
+		va_end(ap); \
 	}
 #define GET_VARARGS_RESULT_WIDE(msg, msgsize, len, lastarg, fmt, result) \
 	{ \
@@ -52,6 +56,7 @@ VarArgs helper macros.
 		{ \
 			result = -1; \
 		} \
+		va_end(ap); \
 	}
 #define GET_VARARGS_RESULT_ANSI(msg, msgsize, len, lastarg, fmt, result) \
 	{ \
@@ -62,6 +67,7 @@ VarArgs helper macros.
 		{ \
 			result = -1; \
 		} \
+		va_end(ap); \
 	}
 
 /*-----------------------------------------------------------------------------
diff --git a/Engine/Source/Runtime/Engine/Private/SceneUtils.cpp b/Engine/Source/Runtime/Engine/Private/SceneUtils.cpp
--- a/Engine/Source/Runtime/Engine/Private/SceneUtils.cpp
+++ b/Engine/Source/Runtime/Engine/Private/SceneUtils.cpp
@@ -63,6 +63,7 @@ void TDrawEvent<TRHICmdList>::Start(TRHICmdList& InRHICmdList, FColor Color, con
 		FCString::GetVarArgs(TempStr, ARRAY_COUNT(TempStr), Fmt, ptr);
 		InRHICmdList.PushEvent(TempStr, Color);
 		RHICmdList = &InRHICmdList;
+		va_end(ptr);
 	}
 }
 
@@ -89,6 +90,7 @@ void FDrawEventRHIExecute::Start(IRHIComputeContext& InRHICommandContext, FColor
 		FCString::GetVarArgs(TempStr, ARRAY_COUNT(TempStr), Fmt, ptr);
 		RHICommandContext = &InRHICommandContext;
 		RHICommandContext->RHIPushEvent(TempStr, Color);
+		va_end(ptr);
 	}
 }
 
