Fixed race condition in Async testcase (FAsyncCompletionCallbackTest)

diff --git a/Engine/Source/Runtime/Core/Private/Tests/Async/AsyncTest.cpp b/Engine/Source/Runtime/Core/Private/Tests/Async/AsyncTest.cpp
--- a/Engine/Source/Runtime/Core/Private/Tests/Async/AsyncTest.cpp
+++ b/Engine/Source/Runtime/Core/Private/Tests/Async/AsyncTest.cpp
@@ -17,13 +17,13 @@ IMPLEMENT_SIMPLE_AUTOMATION_TEST(FAsyncCompletionCallbackTest, "System.Core.Asyn
 /** Helper methods used in the test cases. */
 namespace AsyncTestUtils
 {
-	TFunction<int()> Task = []() {
+	TFunction<int()> Task = [] {
 		return 123;
 	};
 
 	bool bHasVoidTaskFinished = false;
 
-	TFunction<void()> VoidTask = []() {
+	TFunction<void()> VoidTask = [] {
 		bHasVoidTaskFinished = true;
 	};
 }
@@ -84,13 +84,23 @@ bool FAsyncVoidTaskTest::RunTest(const FString& Parameters)
 bool FAsyncCompletionCallbackTest::RunTest(const FString& Parameters)
 {
 	bool Completed = false;
-	auto Future = Async(EAsyncExecution::TaskGraph, AsyncTestUtils::Task, [&]() {
+	FEvent* CompletedEvent = FPlatformProcess::GetSynchEventFromPool(true);
+
+	auto Future = Async(EAsyncExecution::TaskGraph, AsyncTestUtils::Task, [&] {
 		Completed = true;
+		CompletedEvent->Trigger();
 	});
+
 	int Result = Future.Get();
 
+	// We need an additional synchronization point here since the future get above will return after
+	// the task is done but before the completion callback has returned!
+
+	bool CompletedEventTriggered = CompletedEvent->Wait(FTimespan(0 /* hours */, 0 /* minutes */, 5 /* seconds */));
+	FPlatformProcess::ReturnSynchEventToPool(CompletedEvent);
+
 	TestEqual(TEXT("Async Result"), Result, 123);
-	TestTrue(TEXT("Completion callack to be called"), Completed);
+	TestTrue(TEXT("Completion callback to be called"), CompletedEventTriggered && Completed);
 
 	return true;
 }
