SetThreadDescription wrapper is only used on Windows so let's use the same conditionals for the implementation of the API wrapper

diff --git a/Engine/Source/Runtime/Core/Private/Windows/WindowsRunnableThread.cpp b/Engine/Source/Runtime/Core/Private/Windows/WindowsRunnableThread.cpp
--- a/Engine/Source/Runtime/Core/Private/Windows/WindowsRunnableThread.cpp
+++ b/Engine/Source/Runtime/Core/Private/Windows/WindowsRunnableThread.cpp
@@ -8,6 +8,7 @@
 
 DEFINE_LOG_CATEGORY_STATIC(LogThreadingWindows, Log, All);
 
+#if !PLATFORM_XBOXONE		
 void FRunnableThreadWin::SetThreadDescription(HANDLE hThread, PCWSTR lpThreadDescription)
 {
 	// SetThreadDescription is only available from Windows 10 version 1607 / Windows Server 2016
@@ -27,7 +28,7 @@ void FRunnableThreadWin::SetThreadDescription(HANDLE hThread, PCWSTR lpThreadDes
 		RealSetThreadDescription(hThread, lpThreadDescription);
 	}
 }
-
+#endif
 
 uint32 FRunnableThreadWin::GuardedRun()
 {
diff --git a/Engine/Source/Runtime/Core/Private/Windows/WindowsRunnableThread.h b/Engine/Source/Runtime/Core/Private/Windows/WindowsRunnableThread.h
--- a/Engine/Source/Runtime/Core/Private/Windows/WindowsRunnableThread.h
+++ b/Engine/Source/Runtime/Core/Private/Windows/WindowsRunnableThread.h
@@ -66,7 +66,9 @@ class FRunnableThreadWin
 #endif
 	}
 
+#if !PLATFORM_XBOXONE		
 	static void SetThreadDescription(HANDLE hThread, PCWSTR lpThreadDescription);
+#endif
 
 	/**
 	 * The thread entry point. Simply forwards the call on to the right
