Prevent compiler from avoiding a security check caused by a negated integer overflow.
Similar security check added to FName's deserializer.
Fix to UE_LOG categories in FString.

diff --git a/Engine/Source/Runtime/Core/Private/Containers/String.cpp b/Engine/Source/Runtime/Core/Private/Containers/String.cpp
--- a/Engine/Source/Runtime/Core/Private/Containers/String.cpp
+++ b/Engine/Source/Runtime/Core/Private/Containers/String.cpp
@@ -1380,16 +1380,16 @@ FArchive& operator<<( FArchive& Ar, FString& A )
 		bool LoadUCS2Char = SaveNum < 0;
 		if (LoadUCS2Char)
 		{
-			SaveNum = -SaveNum;
-		}
+			// If SaveNum cannot be negated due to integer overflow, Ar is corrupted.
+			if (SaveNum == MIN_int32)
+			{
+				Ar.ArIsError = 1;
+				Ar.ArIsCriticalError = 1;
+				UE_LOG(LogCore, Error, TEXT("Archive is corrupted"));
+				return Ar;
+			}
 
-		// If SaveNum is still less than 0, they must have passed in MIN_INT. Archive is corrupted.
-		if (SaveNum < 0)
-		{
-			Ar.ArIsError = 1;
-			Ar.ArIsCriticalError = 1;
-			UE_LOG(LogNetSerialization, Error, TEXT("Archive is corrupted"));
-			return Ar;
+			SaveNum = -SaveNum;
 		}
 
 		auto MaxSerializeSize = Ar.GetMaxSerializeSize();
@@ -1398,7 +1398,7 @@ FArchive& operator<<( FArchive& Ar, FString& A )
 		{
 			Ar.ArIsError         = 1;
 			Ar.ArIsCriticalError = 1;
-			UE_LOG( LogNetSerialization, Error, TEXT( "String is too large" ) );
+			UE_LOG(LogCore, Error, TEXT("String is too large"));
 			return Ar;
 		}
 
diff --git a/Engine/Source/Runtime/Core/Private/UObject/UnrealNames.cpp b/Engine/Source/Runtime/Core/Private/UObject/UnrealNames.cpp
--- a/Engine/Source/Runtime/Core/Private/UObject/UnrealNames.cpp
+++ b/Engine/Source/Runtime/Core/Private/UObject/UnrealNames.cpp
@@ -1256,6 +1256,15 @@ FArchive& operator<<(FArchive& Ar, FNameEntrySerialized& E)
 		// negative stringlen means it's a wide string
 		if (StringLen < 0)
 		{
+			// If StringLen cannot be negated due to integer overflow, Ar is corrupted.
+			if (StringLen == MIN_int32)
+			{
+				Ar.ArIsError = 1;
+				Ar.ArIsCriticalError = 1;
+				UE_LOG(LogUnrealNames, Error, TEXT("Archive is corrupted"));
+				return Ar;
+			}
+
 			StringLen = -StringLen;
 
 			// mark the name will be wide
