Fixing DEPRECATED macro warnings

diff --git a/Engine/Source/Runtime/Core/Public/Async/AsyncWork.h b/Engine/Source/Runtime/Core/Public/Async/AsyncWork.h
--- a/Engine/Source/Runtime/Core/Public/Async/AsyncWork.h
+++ b/Engine/Source/Runtime/Core/Public/Async/AsyncWork.h
@@ -582,7 +582,7 @@ public:
 	{
 	}
 
-	DEPRECATED(4.20, "Use FAsyncCompress with FName for Format")
+	UE_DEPRECATED(4.20, "Use FAsyncCompress with FName for Format")
 		FAsyncUncompress(
 			ECompressionFlags InFlags,
 			void* InUncompressedBuffer,
diff --git a/Engine/Source/Runtime/Core/Public/Containers/IndirectArray.h b/Engine/Source/Runtime/Core/Public/Containers/IndirectArray.h
--- a/Engine/Source/Runtime/Core/Public/Containers/IndirectArray.h
+++ b/Engine/Source/Runtime/Core/Public/Containers/IndirectArray.h
@@ -448,7 +448,7 @@ struct TContainerTraits<TIndirectArray<T, Allocator> >
 };
 
 template <typename T,typename Allocator>
-DEPRECATED(4.22, "Placement new on TIndirectArray has been deprecated - users should call Add() passing a pointer to an object created with new.")
+UE_DEPRECATED(4.22, "Placement new on TIndirectArray has been deprecated - users should call Add() passing a pointer to an object created with new.")
 void* operator new( size_t Size, TIndirectArray<T,Allocator>& Array )
 {
 	check(Size == sizeof(T));
@@ -457,7 +457,7 @@ void* operator new( size_t Size, TIndirectArray<T,Allocator>& Array )
 }
 
 template <typename T,typename Allocator>
-DEPRECATED(4.22, "Placement new on TIndirectArray has been deprecated - users should call Insert() passing a pointer to an object created with new.")
+UE_DEPRECATED(4.22, "Placement new on TIndirectArray has been deprecated - users should call Insert() passing a pointer to an object created with new.")
 void* operator new( size_t Size, TIndirectArray<T,Allocator>& Array, int32 Index )
 {
 	check(Size == sizeof(T));
diff --git a/Engine/Source/Runtime/Core/Public/GenericPlatform/MicrosoftPlatformString.h b/Engine/Source/Runtime/Core/Public/GenericPlatform/MicrosoftPlatformString.h
--- a/Engine/Source/Runtime/Core/Public/GenericPlatform/MicrosoftPlatformString.h
+++ b/Engine/Source/Runtime/Core/Public/GenericPlatform/MicrosoftPlatformString.h
@@ -111,7 +111,7 @@ struct FMicrosoftPlatformString : public FGenericPlatformString
 		return _tcstok_s(StrToken, Delim, Context);
 	}
 
-	DEPRECATED(4.22, "GetVarArgs with DestSize and Count arguments has been deprecated - only DestSize should be passed")
+	UE_DEPRECATED(4.22, "GetVarArgs with DestSize and Count arguments has been deprecated - only DestSize should be passed")
 	static FORCEINLINE int32 GetVarArgs( WIDECHAR* Dest, SIZE_T DestSize, int32 Count, const WIDECHAR*& Fmt, va_list ArgPtr )
 	{
 		return GetVarArgs(Dest, DestSize, Fmt, ArgPtr);
@@ -218,7 +218,7 @@ struct FMicrosoftPlatformString : public FGenericPlatformString
 		return strtok_s(StrToken, Delim, Context);
 	}
 
-	DEPRECATED(4.22, "GetVarArgs with DestSize and Count arguments has been deprecated - only DestSize should be passed")
+	UE_DEPRECATED(4.22, "GetVarArgs with DestSize and Count arguments has been deprecated - only DestSize should be passed")
 	static FORCEINLINE int32 GetVarArgs( ANSICHAR* Dest, SIZE_T DestSize, int32 Count, const ANSICHAR*& Fmt, va_list ArgPtr )
 	{
 		return GetVarArgs(Dest, DestSize, Fmt, ArgPtr);
diff --git a/Engine/Source/Runtime/Core/Public/Misc/CString.h b/Engine/Source/Runtime/Core/Public/Misc/CString.h
--- a/Engine/Source/Runtime/Core/Public/Misc/CString.h
+++ b/Engine/Source/Runtime/Core/Public/Misc/CString.h
@@ -369,7 +369,7 @@ public:
 	 * @param Args - argument list
 	 * @return number of characters written or -1 if truncated
 	 */
-	DEPRECATED(4.22, "GetVarArgs with DestSize and Count arguments has been deprecated - only DestSize should be passed")
+	UE_DEPRECATED(4.22, "GetVarArgs with DestSize and Count arguments has been deprecated - only DestSize should be passed")
 	static FORCEINLINE int32 GetVarArgs(CharType* Dest, SIZE_T DestSize, int32 Count, const CharType*& Fmt, va_list ArgPtr)
 	{
 		return GetVarArgs(Dest, DestSize, Fmt, ArgPtr);
diff --git a/Engine/Source/Runtime/Core/Public/Misc/Compression.h b/Engine/Source/Runtime/Core/Public/Misc/Compression.h
--- a/Engine/Source/Runtime/Core/Public/Misc/Compression.h
+++ b/Engine/Source/Runtime/Core/Public/Misc/Compression.h
@@ -81,11 +81,11 @@ struct FCompression
 
 
 
-	DEPRECATED(4.20, "Use the FName based version of CompressMemoryBound")
+	UE_DEPRECATED(4.20, "Use the FName based version of CompressMemoryBound")
 	CORE_API static int32 CompressMemoryBound(ECompressionFlags Flags, int32 UncompressedSize, int32 BitWindow = DEFAULT_ZLIB_BIT_WINDOW);
-	DEPRECATED(4.20, "Use the FName based version of CompressMemory")
+	UE_DEPRECATED(4.20, "Use the FName based version of CompressMemory")
 	CORE_API static bool CompressMemory(ECompressionFlags Flags, void* CompressedBuffer, int32& CompressedSize, const void* UncompressedBuffer, int32 UncompressedSize, int32 BitWindow = DEFAULT_ZLIB_BIT_WINDOW);
-	DEPRECATED(4.20, "Use the FName based version of UncompressMemory")
+	UE_DEPRECATED(4.20, "Use the FName based version of UncompressMemory")
 	CORE_API static bool UncompressMemory(ECompressionFlags Flags, void* UncompressedBuffer, int32 UncompressedSize, const void* CompressedBuffer, int32 CompressedSize, bool bIsSourcePadded = false, int32 BitWindow = DEFAULT_ZLIB_BIT_WINDOW);
 
 	CORE_API static FName GetCompressionFormatFromDeprecatedFlags(ECompressionFlags DeprecatedFlags);
diff --git a/Engine/Source/Runtime/Core/Public/Serialization/Archive.h b/Engine/Source/Runtime/Core/Public/Serialization/Archive.h
--- a/Engine/Source/Runtime/Core/Public/Serialization/Archive.h
+++ b/Engine/Source/Runtime/Core/Public/Serialization/Archive.h
@@ -746,7 +746,7 @@ public:
 	 * @param	bTreatBufferAsFileReader true if V is actually an FArchive, which is used when saving to read data - helps to avoid single huge allocations of source data
 	 * @param	bUsePlatformBitWindow use a platform specific bitwindow setting
 	 */
-	DEPRECATED(4.20, "Use the FName based version of SerializeCompressed (which also removes the basically-unused bUsePlatformBitWindow)")
+	UE_DEPRECATED(4.20, "Use the FName based version of SerializeCompressed (which also removes the basically-unused bUsePlatformBitWindow)")
 	void SerializeCompressed(void* V, int64 Length, ECompressionFlags Flags = COMPRESS_NoFlags, bool bTreatBufferAsFileReader = false, bool bUsePlatformBitWindow = false);
 	void SerializeCompressed(void* V, int64 Length, FName CompressionFormat, ECompressionFlags Flags=COMPRESS_NoFlags, bool bTreatBufferAsFileReader=false);
 
diff --git a/Engine/Source/Runtime/CoreUObject/Public/Serialization/BulkData.h b/Engine/Source/Runtime/CoreUObject/Public/Serialization/BulkData.h
--- a/Engine/Source/Runtime/CoreUObject/Public/Serialization/BulkData.h
+++ b/Engine/Source/Runtime/CoreUObject/Public/Serialization/BulkData.h
@@ -373,7 +373,7 @@ public:
 	 *
 	 * @param CompressionFlags	Flags to use for compressing the data. Use COMPRESS_NONE for no compression, or something like COMPRESS_ZLIB to compress the data
 	 */
-	DEPRECATED(4.21, "Use the FName version of StoreCompressedOnDisk")
+	UE_DEPRECATED(4.21, "Use the FName version of StoreCompressedOnDisk")
 	void StoreCompressedOnDisk( ECompressionFlags CompressionFlags );
 	void StoreCompressedOnDisk( FName CompressionFormat );
 
