Fixing the remaining DEPRECATED macro depracation warnings (workaround for VC compiler bug)

diff --git a/Engine/Source/Runtime/Core/Private/Misc/Compression.cpp b/Engine/Source/Runtime/Core/Private/Misc/Compression.cpp
--- a/Engine/Source/Runtime/Core/Private/Misc/Compression.cpp
+++ b/Engine/Source/Runtime/Core/Private/Misc/Compression.cpp
@@ -437,7 +437,7 @@ bool FCompression::UncompressMemory(FName FormatName, void* UncompressedBuffer,
  * @param	InMaxPendingBufferSize	Max chunk size to compress in uncompressed bytes
  * @param	InCompressionFlags		Compression flags to compress memory with
  */
-FCompressedGrowableBuffer::FCompressedGrowableBuffer(int32 InMaxPendingBufferSize, ECompressionFlags InCompressionFlags)
+FCompressedGrowableBuffer::FCompressedGrowableBuffer(FCompressedGrowableBuffer::EVS2015Redirector, int32 InMaxPendingBufferSize, ECompressionFlags InCompressionFlags)
 	: FCompressedGrowableBuffer(InMaxPendingBufferSize, FCompression::GetCompressionFormatFromDeprecatedFlags(InCompressionFlags), InCompressionFlags)
 {
 
diff --git a/Engine/Source/Runtime/Core/Private/Serialization/ArchiveLoadCompressedProxy.cpp b/Engine/Source/Runtime/Core/Private/Serialization/ArchiveLoadCompressedProxy.cpp
--- a/Engine/Source/Runtime/Core/Private/Serialization/ArchiveLoadCompressedProxy.cpp
+++ b/Engine/Source/Runtime/Core/Private/Serialization/ArchiveLoadCompressedProxy.cpp
@@ -9,7 +9,7 @@
 	FArchiveLoadCompressedProxy
 ----------------------------------------------------------------------------*/
 
-FArchiveLoadCompressedProxy::FArchiveLoadCompressedProxy(const TArray<uint8>& InCompressedData, ECompressionFlags InCompressionFlags)
+FArchiveLoadCompressedProxy::FArchiveLoadCompressedProxy(FArchiveLoadCompressedProxy::EVS2015Redirector, const TArray<uint8>& InCompressedData, ECompressionFlags InCompressionFlags)
 	: FArchiveLoadCompressedProxy(InCompressedData, FCompression::GetCompressionFormatFromDeprecatedFlags(InCompressionFlags))
 {
 }
diff --git a/Engine/Source/Runtime/Core/Private/Serialization/ArchiveSaveCompressedProxy.cpp b/Engine/Source/Runtime/Core/Private/Serialization/ArchiveSaveCompressedProxy.cpp
--- a/Engine/Source/Runtime/Core/Private/Serialization/ArchiveSaveCompressedProxy.cpp
+++ b/Engine/Source/Runtime/Core/Private/Serialization/ArchiveSaveCompressedProxy.cpp
@@ -9,7 +9,7 @@
 /*----------------------------------------------------------------------------
 	FArchiveSaveCompressedProxy
 ----------------------------------------------------------------------------*/
-FArchiveSaveCompressedProxy::FArchiveSaveCompressedProxy(TArray<uint8>& InCompressedData, ECompressionFlags InCompressionFlags)
+FArchiveSaveCompressedProxy::FArchiveSaveCompressedProxy(FArchiveSaveCompressedProxy::EVS2015Redirector, TArray<uint8>& InCompressedData, ECompressionFlags InCompressionFlags)
 	: FArchiveSaveCompressedProxy(InCompressedData, FCompression::GetCompressionFormatFromDeprecatedFlags(InCompressionFlags), InCompressionFlags)
 {
 }
diff --git a/Engine/Source/Runtime/Core/Public/Misc/CompressedGrowableBuffer.h b/Engine/Source/Runtime/Core/Public/Misc/CompressedGrowableBuffer.h
--- a/Engine/Source/Runtime/Core/Public/Misc/CompressedGrowableBuffer.h
+++ b/Engine/Source/Runtime/Core/Public/Misc/CompressedGrowableBuffer.h
@@ -11,14 +11,28 @@
  */
 struct CORE_API FCompressedGrowableBuffer
 {
+private:
+	/**
+	 * This enum and the following constructor is a workaround for VC compiler bug that prevents using attributes
+	 * on constructors without inline implementation. This should be removed when the deprecated ctor is removed.
+	 */
+	enum EVS2015Redirector
+	{
+		Redirect
+	};
+	FCompressedGrowableBuffer(EVS2015Redirector, int32 MaxPendingBufferSize, ECompressionFlags CompressionFlags);
+public:
 	/**
 	 * Constructor
 	 *
 	 * @param	MaxPendingBufferSize	Max chunk size to compress in uncompressed bytes
 	 * @param	CompressionFlags		Compression flags to compress memory with
 	 */
-	DEPRECATED(4.21, "Use FName version of FCompressedGrowableBuffer constructor")
-	FCompressedGrowableBuffer( int32 MaxPendingBufferSize, ECompressionFlags CompressionFlags );
+	UE_DEPRECATED(4.21, "Use FName version of FCompressedGrowableBuffer constructor")
+	FCompressedGrowableBuffer(int32 MaxPendingBufferSize, ECompressionFlags CompressionFlags)
+		// Make sure to remove the EVS2015Redirector constructor when this constructor is removed
+		: FCompressedGrowableBuffer(EVS2015Redirector::Redirect, MaxPendingBufferSize, CompressionFlags)
+	{}
 	FCompressedGrowableBuffer(int32 MaxPendingBufferSize, FName COmpressionFormat, ECompressionFlags CompressionFlags=COMPRESS_None);
 
 	/**
diff --git a/Engine/Source/Runtime/Core/Public/Serialization/ArchiveLoadCompressedProxy.h b/Engine/Source/Runtime/Core/Public/Serialization/ArchiveLoadCompressedProxy.h
--- a/Engine/Source/Runtime/Core/Public/Serialization/ArchiveLoadCompressedProxy.h
+++ b/Engine/Source/Runtime/Core/Public/Serialization/ArchiveLoadCompressedProxy.h
@@ -16,6 +16,17 @@
  */
 class CORE_API FArchiveLoadCompressedProxy : public FArchive
 {
+private:
+	/**
+	 * This enum and the following constructor is a workaround for VC compiler bug that prevents using attributes
+	 * on constructors without inline implementation. This should be removed when the deprecated ctor is removed.
+	 */
+	enum EVS2015Redirector
+	{
+		Redirect
+	};
+	FArchiveLoadCompressedProxy(EVS2015Redirector InUnused, const TArray<uint8>& InCompressedData, ECompressionFlags InCompressionFlags);
+
 public:
 	/** 
 	 * Constructor, initializing all member variables and allocating temp memory.
@@ -23,8 +34,11 @@ public:
 	 * @param	InCompressedData	Array of bytes that is holding compressed data
 	 * @param	InCompressionFormat	What format to compress with
 	 */
-	DEPRECATED(4.21, "Use the FName version of FArchiveLoadCompressedProxy constructor")
-	FArchiveLoadCompressedProxy( const TArray<uint8>& InCompressedData, ECompressionFlags InCompressionFlags );
+	UE_DEPRECATED(4.21, "Use the FName version of FArchiveLoadCompressedProxy constructor")
+	FArchiveLoadCompressedProxy(const TArray<uint8>& InCompressedData, ECompressionFlags InCompressionFlags)
+		// Make sure to remove the EVS2015Redirector constructor when this constructor is removed
+		: FArchiveLoadCompressedProxy(EVS2015Redirector::Redirect, InCompressedData, InCompressionFlags)
+	{}
 	FArchiveLoadCompressedProxy(const TArray<uint8>& InCompressedData, FName CompressionFormat, ECompressionFlags InCompressionFlags=COMPRESS_NoFlags);
 
 	/** Destructor, freeing temporary memory. */
diff --git a/Engine/Source/Runtime/Core/Public/Serialization/ArchiveSaveCompressedProxy.h b/Engine/Source/Runtime/Core/Public/Serialization/ArchiveSaveCompressedProxy.h
--- a/Engine/Source/Runtime/Core/Public/Serialization/ArchiveSaveCompressedProxy.h
+++ b/Engine/Source/Runtime/Core/Public/Serialization/ArchiveSaveCompressedProxy.h
@@ -16,6 +16,16 @@
  */
 class CORE_API FArchiveSaveCompressedProxy : public FArchive
 {
+	/**
+	 * This enum and the following constructor is a workaround for VC compiler bug that prevents using attributes
+	 * on constructors without inline implementation. This should be removed when the deprecated ctor is removed.
+	 */
+	enum EVS2015Redirector
+	{
+		Redirect
+	};
+	FArchiveSaveCompressedProxy(EVS2015Redirector, TArray<uint8>& InCompressedData, ECompressionFlags InCompressionFlags);
+
 public:
 	/** 
 	 * Constructor, initializing all member variables and allocating temp memory.
@@ -23,8 +33,11 @@ public:
 	 * @param	InCompressedData [ref]	Array of bytes that is going to hold compressed data
 	 * @param	InCompressionFormat		Compression format to use
 	 */
-	DEPRECATED(4.21, "Use the FName version of FArchiveSaveCompressedProxy constructor")
-	FArchiveSaveCompressedProxy(TArray<uint8>& InCompressedData, ECompressionFlags InCompressionFlags);
+	UE_DEPRECATED(4.21, "Use the FName version of FArchiveSaveCompressedProxy constructor")
+	FArchiveSaveCompressedProxy(TArray<uint8>& InCompressedData, ECompressionFlags InCompressionFlags)
+		// Make sure to remove the EVS2015Redirector constructor when this constructor is removed
+		: FArchiveSaveCompressedProxy(EVS2015Redirector::Redirect, InCompressedData, InCompressionFlags)
+	{}
 	FArchiveSaveCompressedProxy( TArray<uint8>& InCompressedData, FName InCompressionFormat, ECompressionFlags InCompressionFlags=COMPRESS_None);
 
 	/** Destructor, flushing array if needed. Also frees temporary memory. */
