More DEPRACATED -> UE_DEPRECATED warnings CIS fixes

diff --git a/Engine/Source/Runtime/Core/Public/GenericPlatform/GenericWidePlatformString.h b/Engine/Source/Runtime/Core/Public/GenericPlatform/GenericWidePlatformString.h
--- a/Engine/Source/Runtime/Core/Public/GenericPlatform/GenericWidePlatformString.h
+++ b/Engine/Source/Runtime/Core/Public/GenericPlatform/GenericWidePlatformString.h
@@ -153,7 +153,7 @@ public:
 	
 	static WIDECHAR* Strtok( WIDECHAR* StrToken, const WIDECHAR* Delim, WIDECHAR** Context );
 
-	DEPRECATED(4.22, "GetVarArgs with DestSize and Count arguments has been deprecated - only DestSize should be passed")
+	UE_DEPRECATED(4.22, "GetVarArgs with DestSize and Count arguments has been deprecated - only DestSize should be passed")
 	static int32 GetVarArgs( WIDECHAR* Dest, SIZE_T DestSize, int32 Count, const WIDECHAR*& Fmt, va_list ArgPtr )
 	{
 		return GetVarArgs(Dest, DestSize, Fmt, ArgPtr);
