Optimization of TSet::operator= to not rehash all the elements.

diff --git a/Engine/Source/Runtime/Core/Public/Containers/Set.h b/Engine/Source/Runtime/Core/Public/Containers/Set.h
--- a/Engine/Source/Runtime/Core/Public/Containers/Set.h
+++ b/Engine/Source/Runtime/Core/Public/Containers/Set.h
@@ -256,20 +256,14 @@ public:
 	{
 		if(this != &Copy)
 		{
-			if (Copy.Elements.IsCompact())
-			{
-				Elements = Copy.Elements;
-				HashSize = Copy.HashSize;
-				Rehash();
-			}
-			else
-			{
-				Empty(Copy.Num());
-				for (TConstIterator CopyIt(Copy); CopyIt; ++CopyIt)
-				{
-					Add(*CopyIt);
-				}
-			}
+			int32 CopyHashSize = Copy.HashSize;
+
+			DestructItems((FSetElementId*)Hash.GetAllocation(), HashSize);
+			Hash.ResizeAllocation(0, CopyHashSize, sizeof(FSetElementId));
+			ConstructItems<FSetElementId>(Hash.GetAllocation(), (FSetElementId*)Copy.Hash.GetAllocation(), CopyHashSize);
+			HashSize = CopyHashSize;
+
+			Elements = Copy.Elements;
 		}
 		return *this;
 	}
