Added Strnicmp() overloads for ANSICHAR vs WIDECHAR comparisons

diff --git a/Engine/Source/Runtime/Core/Private/GenericPlatform/GenericPlatformStricmp.cpp b/Engine/Source/Runtime/Core/Private/GenericPlatform/GenericPlatformStricmp.cpp
--- a/Engine/Source/Runtime/Core/Private/GenericPlatform/GenericPlatformStricmp.cpp
+++ b/Engine/Source/Runtime/Core/Private/GenericPlatform/GenericPlatformStricmp.cpp
@@ -53,13 +53,13 @@ int32 StricmpImpl(const CharType1* String1, const CharType2* String2)
 	}
 }
 
-template<typename CharType>
-int32 StrnicmpImpl(const CharType* String1, const CharType* String2, SIZE_T Count)
+template<typename CharType1, typename CharType2>
+int32 StrnicmpImpl(const CharType1* String1, const CharType2* String2, SIZE_T Count)
 {
 	for (; Count > 0; --Count)
 	{
-		CharType C1 = *String1++;
-		CharType C2 = *String2++;
+		CharType1 C1 = *String1++;
+		CharType2 C2 = *String2++;
 
 		// Quickly move on if characters are identical but
 		// return equals if we found two null terminators
@@ -74,14 +74,14 @@ int32 StrnicmpImpl(const CharType* String1, const CharType* String2, SIZE_T Coun
 		}
 		else if (BothAscii(C1, C2))
 		{
-			if (int32 Diff = LowerAscii[TChar<CharType>::ToUnsigned(C1)] - LowerAscii[TChar<CharType>::ToUnsigned(C2)])
+			if (int32 Diff = LowerAscii[TChar<CharType1>::ToUnsigned(C1)] - LowerAscii[TChar<CharType2>::ToUnsigned(C2)])
 			{
 				return Diff;
 			}
 		}
 		else
 		{
-			return TChar<CharType>::ToUnsigned(C1) - TChar<CharType>::ToUnsigned(C2);
+			return TChar<CharType1>::ToUnsigned(C1) - TChar<CharType2>::ToUnsigned(C2);
 		}
 	}
 
@@ -103,6 +103,8 @@ int32 FGenericPlatformStricmp::Stricmp(const UTF16CHAR* Str1, const ANSICHAR* St
 int32 FGenericPlatformStricmp::Stricmp(const UTF32CHAR* Str1, const ANSICHAR* Str2) { return StricmpImpl(Str1, Str2); }
 int32 FGenericPlatformStricmp::Strnicmp(const ANSICHAR* Str1, const ANSICHAR* Str2, SIZE_T Count) { return StrnicmpImpl(Str1, Str2, Count); }
 int32 FGenericPlatformStricmp::Strnicmp(const WIDECHAR* Str1, const WIDECHAR* Str2, SIZE_T Count) { return StrnicmpImpl(Str1, Str2, Count); }
+int32 FGenericPlatformStricmp::Strnicmp(const ANSICHAR* Str1, const WIDECHAR* Str2, SIZE_T Count) { return StrnicmpImpl(Str1, Str2, Count); }
+int32 FGenericPlatformStricmp::Strnicmp(const WIDECHAR* Str1, const ANSICHAR* Str2, SIZE_T Count) { return StrnicmpImpl(Str1, Str2, Count); }
 
 //////////////////////////////////////////////////////////////////////////
 
diff --git a/Engine/Source/Runtime/Core/Public/GenericPlatform/GenericPlatformStricmp.h b/Engine/Source/Runtime/Core/Public/GenericPlatform/GenericPlatformStricmp.h
--- a/Engine/Source/Runtime/Core/Public/GenericPlatform/GenericPlatformStricmp.h
+++ b/Engine/Source/Runtime/Core/Public/GenericPlatform/GenericPlatformStricmp.h
@@ -31,4 +31,6 @@ struct FGenericPlatformStricmp
 
 	CORE_API static int32 Strnicmp(const ANSICHAR* String1, const ANSICHAR* String2, SIZE_T Count);
 	CORE_API static int32 Strnicmp(const WIDECHAR* String1, const WIDECHAR* String2, SIZE_T Count);
+	CORE_API static int32 Strnicmp(const ANSICHAR* String1, const WIDECHAR* String2, SIZE_T Count);
+	CORE_API static int32 Strnicmp(const WIDECHAR* String1, const ANSICHAR* String2, SIZE_T Count);
 };
