Removed invalid postcondition check in TLockFreeFixedSizeAllocator::Trim(). Another thread could add an entry to the FreeList before the check executes.

diff --git a/Engine/Source/Runtime/Core/Public/Containers/LockFreeFixedSizeAllocator.h b/Engine/Source/Runtime/Core/Public/Containers/LockFreeFixedSizeAllocator.h
--- a/Engine/Source/Runtime/Core/Public/Containers/LockFreeFixedSizeAllocator.h
+++ b/Engine/Source/Runtime/Core/Public/Containers/LockFreeFixedSizeAllocator.h
@@ -242,7 +242,6 @@ public:
 			FMemory::Free(Mem);
 			NumFree.Decrement();
 		}
-		check(!NumFree.GetValue());
 	}
 
 	/**
