Speedup and remove code bloat from automation tests

Smoke tests now run ~8x faster, 17ms instead of 137ms.

* Avoid needless FString allocation when passing in a wide string literal to TestEquals() and other TestXxx() functions.
* Narrow literals still create a temporary FString.
* Removed AddInfo() calls for successful testing that about half of the TestXxx() functions did.
* Cleaned up variable naming and error messages a bit.

diff --git a/Engine/Source/Runtime/Core/Private/Misc/AutomationTest.cpp b/Engine/Source/Runtime/Core/Private/Misc/AutomationTest.cpp
--- a/Engine/Source/Runtime/Core/Private/Misc/AutomationTest.cpp
+++ b/Engine/Source/Runtime/Core/Private/Misc/AutomationTest.cpp
@@ -1062,88 +1062,76 @@ void FAutomationTestBase::GenerateTestNames(TArray<FAutomationTestInfo>& TestInf
 
 // --------------------------------------------------------------------------------------
 
-void FAutomationTestBase::TestEqual(const FString& What, const int32 Actual, const int32 Expected)
+void FAutomationTestBase::TestEqual(const TCHAR* What, const int32 Actual, const int32 Expected)
 {
-	if ( Actual != Expected )
+	if (Actual != Expected)
 	{
-		AddError(FString::Printf(TEXT("Expected '%s' to be %d, but it was %d."), *What, Expected, Actual), 1);
-	}
-	else
-	{
-		AddInfo(FString::Printf(TEXT("Expected '%s' to be %d."), *What, Expected), 1);
+		AddError(FString::Printf(TEXT("Expected '%s' to be %d, but it was %d."), What, Expected, Actual), 1);
 	}
 }
 
-void FAutomationTestBase::TestEqual(const FString& What, const float Actual, const float Expected, float Tolerance)
+void FAutomationTestBase::TestEqual(const TCHAR* What, const float Actual, const float Expected, float Tolerance)
 {
-	if ( !FMath::IsNearlyEqual(Actual, Expected, Tolerance) )
-	{
-		AddError(FString::Printf(TEXT("Expected '%s' to be %f, but it was %f within tolerance %f."), *What, Expected, Actual, Tolerance), 1);
-	}
-	else
+	if (!FMath::IsNearlyEqual(Actual, Expected, Tolerance))
 	{
-		AddInfo(FString::Printf(TEXT("Expected '%s' to be %f within tolerance %f."), *What, Expected, Tolerance), 1);
+		AddError(FString::Printf(TEXT("Expected '%s' to be %f, but it was %f within tolerance %f."), What, Expected, Actual, Tolerance), 1);
 	}
 }
 
-void FAutomationTestBase::TestEqual(const FString& What, const FVector Actual, const FVector Expected, float Tolerance)
+void FAutomationTestBase::TestEqual(const TCHAR* What, const FVector Actual, const FVector Expected, float Tolerance)
 {
-	if ( !Expected.Equals(Actual, Tolerance) )
+	if (!Expected.Equals(Actual, Tolerance))
 	{
-		AddError(FString::Printf(TEXT("Expected '%s' to be %s, but it was %s within tolerance %f."), *What, *Expected.ToString(), *Actual.ToString(), Tolerance), 1);
-	}
-	else
-	{
-		AddInfo(FString::Printf(TEXT("Expected '%s' to be %s within tolerance %f."), *What, *Expected.ToString(), Tolerance), 1);
+		AddError(FString::Printf(TEXT("Expected '%s' to be %s, but it was %s within tolerance %f."), What, *Expected.ToString(), *Actual.ToString(), Tolerance), 1);
 	}
 }
 
-void FAutomationTestBase::TestEqual(const FString& What, const FColor Actual, const FColor Expected)
+void FAutomationTestBase::TestEqual(const TCHAR* What, const FColor Actual, const FColor Expected)
 {
-	if ( Expected != Actual )
-	{
-		AddError(FString::Printf(TEXT("Expected '%s' to be %s, but it was %s."), *What, *Expected.ToString(), *Actual.ToString()), 1);
-	}
-	else
+	if (Expected != Actual)
 	{
-		AddInfo(FString::Printf(TEXT("Expected '%s' to be %s."), *What, *Expected.ToString()), 1);
+		AddError(FString::Printf(TEXT("Expected '%s' to be %s, but it was %s."), What, *Expected.ToString(), *Actual.ToString()), 1);
 	}
 }
 
-void FAutomationTestBase::TestFalse(const FString& What, bool Value)
+void FAutomationTestBase::TestEqual(const TCHAR* What, const TCHAR* Actual, const TCHAR* Expected)
 {
-	if ( Value )
+	if (FCString::Strcmp(Actual, Expected) != 0)
 	{
-		AddError(FString::Printf(TEXT("Expected '%s' to be false."), *What), 1);
+		AddError(FString::Printf(TEXT("Expected '%s' to be %s, but it was %s."), What, Expected, Actual), 1);
 	}
-	else
+}
+
+void FAutomationTestBase::TestEqualInsensitive(const TCHAR* What, const TCHAR* Actual, const TCHAR* Expected)
+{
+	if (FCString::Stricmp(Actual, Expected) != 0)
 	{
-		AddInfo(FString::Printf(TEXT("Expected '%s' to be false."), *What), 1);
+		AddError(FString::Printf(TEXT("Expected '%s' to be %s, but it was %s."), What, Expected, Actual), 1);
 	}
 }
 
-void FAutomationTestBase::TestTrue(const FString& What, bool Value)
+void FAutomationTestBase::TestFalse(const TCHAR* What, bool Value)
 {
-	if ( !Value )
+	if (Value)
 	{
-		AddError(FString::Printf(TEXT("Expected '%s' to be true."), *What), 1);
+		AddError(FString::Printf(TEXT("Expected '%s' to be false."), What), 1);
 	}
-	else
+}
+
+void FAutomationTestBase::TestTrue(const TCHAR* What, bool Value)
+{
+	if (!Value)
 	{
-		AddInfo(FString::Printf(TEXT("Expected '%s' to be true."), *What), 1);
+		AddError(FString::Printf(TEXT("Expected '%s' to be true."), What), 1);
 	}
 }
 
-void FAutomationTestBase::TestNull(const FString& What, const void* Pointer)
+void FAutomationTestBase::TestNull(const TCHAR* What, const void* Pointer)
 {
 	if ( Pointer != nullptr )
 	{
 		AddError(FString::Printf(TEXT("Expected '%s' to be null."), *What), 1);
 	}
-	else
-	{
-		AddInfo(FString::Printf(TEXT("Expected '%s' to be null."), *What), 1);
-	}
 }
 
 bool FAutomationTestBase::IsExpectedError(const FString& Error)
diff --git a/Engine/Source/Runtime/Core/Public/Misc/AutomationTest.h b/Engine/Source/Runtime/Core/Public/Misc/AutomationTest.h
--- a/Engine/Source/Runtime/Core/Public/Misc/AutomationTest.h
+++ b/Engine/Source/Runtime/Core/Public/Misc/AutomationTest.h
@@ -1222,62 +1222,107 @@ public:
 
 public:
 
-	void TestEqual(const FString& What, const int32 Actual, const int32 Expected);
-	void TestEqual(const FString& What, const float Actual, const float Expected, float Tolerance = KINDA_SMALL_NUMBER);
-	void TestEqual(const FString& What, const FVector Actual, const FVector Expected, float Tolerance = KINDA_SMALL_NUMBER);
-	void TestEqual(const FString& What, const FColor Actual, const FColor Expected);
+	void TestEqual(const TCHAR* What, int32 Actual, int32 Expected);
+	void TestEqual(const TCHAR* What, float Actual, float Expected, float Tolerance = KINDA_SMALL_NUMBER);
+	void TestEqual(const TCHAR* What, FVector Actual, FVector Expected, float Tolerance = KINDA_SMALL_NUMBER);
+	void TestEqual(const TCHAR* What, FColor Actual, FColor Expected);
+	void TestEqual(const TCHAR* What, const TCHAR* A, const TCHAR* B);
+	void TestEqualInsensitive(const TCHAR* What, const TCHAR* A, const TCHAR* B);
+
+	void TestEqual(const FString& What, int32 Actual, int32 Expected)
+	{
+		TestEqual(*What, Actual, Expected);
+	}
+
+	void TestEqual(const FString& What, float Actual, float Expected, float Tolerance = KINDA_SMALL_NUMBER)
+	{
+		TestEqual(*What, Actual, Expected, Tolerance);
+	}
+
+	void TestEqual(const FString& What, FVector Actual, FVector Expected, float Tolerance = KINDA_SMALL_NUMBER)
+	{
+		TestEqual(*What, Actual, Expected, Tolerance);
+	}
+
+	void TestEqual(const FString& What, FColor Actual, FColor Expected)
+	{
+		TestEqual(*What, Actual, Expected);
+	}
+
+	void TestEqual(const FString& What, const TCHAR* Actual, const TCHAR* Expected)
+	{
+		TestEqual(*What, Actual, Expected);
+	}
+
+	void TestEqual(const TCHAR* What, const FString& Actual, const TCHAR* Expected)
+	{
+		TestEqualInsensitive(What, *Actual, Expected);
+	}
+
+	void TestEqual(const FString& What, const FString& Actual, const TCHAR* Expected)
+	{
+		TestEqualInsensitive(*What, *Actual, Expected);
+	}
+
+	void TestEqual(const TCHAR* What, const TCHAR* Actual, const FString& Expected)
+	{
+		TestEqualInsensitive(What, Actual, *Expected);
+	}
+
+	void TestEqual(const FString& What, const TCHAR* Actual, const FString& Expected)
+	{
+		TestEqualInsensitive(*What, Actual, *Expected);
+	}
+
+	void TestEqual(const TCHAR* What, const FString& Actual, const FString& Expected)
+	{
+		TestEqualInsensitive(What, *Actual, *Expected);
+	}
+
+	void TestEqual(const FString& What, const FString& Actual, const FString& Expected)
+	{
+		TestEqualInsensitive(*What, *Actual, *Expected);
+	}
 
 	/**
 	 * Logs an error if the two values are not equal.
 	 *
-	 * @param Description - Description text for the test.
+	 * @param What - Description text for the test.
 	 * @param A - The first value.
 	 * @param B - The second value.
 	 *
 	 * @see TestNotEqual
 	 */
 	template<typename ValueType> 
-	void TestEqual(const FString& Description, const ValueType& A, const ValueType& B)
+	void TestEqual(const TCHAR* What, const ValueType& A, const ValueType& B)
 	{
 		if (A != B)
 		{
-			AddError(FString::Printf(TEXT("%s: The two values are not equal."), *Description), 1);
+			AddError(FString::Printf(TEXT("%s: The two values are not equal."), What), 1);
 		}
 	}
 
-	void TestEqual(const FString& Description, const FString& A, const TCHAR* B)
+	template<typename ValueType>
+	void TestEqual(const FString& What, const ValueType& A, const ValueType& B)
 	{
-		if (A != B)
-		{
-			AddError(FString::Printf(TEXT("%s: The two values are not equal.\n  Actual: '%s'\nExpected: '%s'"), *Description, *A, B), 1);
-		}
+		TestEqual(*What, A, B);
 	}
 
-	void TestEqual(const FString& Description, const TCHAR* A, const TCHAR* B)
-	{
-		if (FCString::Strcmp(A,B) != 0)
-		{
-			AddError(FString::Printf(TEXT("%s: The two values are not equal.\n  Actual: '%s'\nExpected: '%s'"), *Description, A, B), 1);
-		}
-	}
-
-	void TestEqual(const FString& Description, const TCHAR* A, const FString& B)
-	{
-		if (FCString::Strcmp(A, *B) != 0)
-		{
-			AddError(FString::Printf(TEXT("%s: The two values are not equal.\n  Actual: '%s'\nExpected: '%s'"), *Description, A, *B), 1);
-		}
-	}
 
 	/**
 	 * Logs an error if the specified Boolean value is not false.
 	 *
-	 * @param Description - Description text for the test.
+	 * @param What - Description text for the test.
 	 * @param Value - The value to test.
 	 *
 	 * @see TestFalse
 	 */
-	void TestFalse(const FString& What, bool Value);
+	void TestFalse(const TCHAR* What, bool Value);
+
+	void TestFalse(const FString& What, bool Value)
+	{
+		TestFalse(*What, Value);
+	}
 
 	/**
 	 * Logs an error if the given shared pointer is valid.
@@ -1287,14 +1332,19 @@ public:
 	 *
 	 * @see TestValid
 	 */
-	template<typename ValueType> void TestInvalid(const FString& Description, const TSharedPtr<ValueType>& SharedPointer)
+	template<typename ValueType> void TestInvalid(const TCHAR* Description, const TSharedPtr<ValueType>& SharedPointer)
 	{
 		if (SharedPointer.IsValid())
 		{
-			AddError(FString::Printf(TEXT("%s: The shared pointer is valid."), *Description), 1);
+			AddError(FString::Printf(TEXT("%s: The shared pointer is valid."), Description), 1);
 		}
 	}
 
+	template<typename ValueType> void TestInvalid(const FString& Description, const TSharedPtr<ValueType>& SharedPointer)
+	{
+		TestInvalid(*Description, SharedPointer);
+	}
+
 	/**
 	 * Logs an error if the two values are equal.
 	 *
@@ -1304,34 +1354,40 @@ public:
 	 *
 	 * @see TestEqual
 	 */
-	template<typename ValueType> void TestNotEqual(const FString& Description, const ValueType& A, const ValueType& B)
+	template<typename ValueType> void TestNotEqual(const TCHAR* Description, const ValueType& A, const ValueType& B)
 	{
 		if (A == B)
 		{
-			AddError(FString::Printf(TEXT("%s: The two values are equal."), *Description), 1);
+			AddError(FString::Printf(TEXT("%s: The two values are equal."), Description), 1);
 		}
 	}
 
+	template<typename ValueType> void TestNotEqual(const FString& Description, const ValueType& A, const ValueType& B)
+	{
+		TestNotEqual(*Description, A, B);
+	}
+
 	/**
 	 * Logs an error if the specified pointer is NULL.
 	 *
-	 * @param Description - Description text for the test.
+	 * @param What - Description text for the test.
 	 * @param Pointer - The pointer to test.
 	 *
 	 * @see TestNull
 	 */
-	template<typename ValueType> void TestNotNull(const FString& What, ValueType* Pointer)
+	template<typename ValueType> void TestNotNull(const TCHAR* What, ValueType* Pointer)
 	{
 		if (Pointer == nullptr)
 		{
-			AddError(FString::Printf(TEXT("Expected '%s' to be not null."), *What), 1);
-		}
-		else
-		{
-			AddInfo(FString::Printf(TEXT("Expected '%s' to be not null."), *What), 1);
+			AddError(FString::Printf(TEXT("Expected '%s' to be not null."), What), 1);
 		}
 	}
 
+	template<typename ValueType> void TestNotNull(const FString& What, ValueType* Pointer)
+	{
+		TestNotNull(*What, Pointer);
+	}
+
 	/**
 	 * Logs an error if the two values are the same object in memory.
 	 *
@@ -1341,14 +1397,19 @@ public:
 	 *
 	 * @see TestSame
 	 */
-	template<typename ValueType> void TestNotSame(const FString& Description, const ValueType& A, const ValueType& B)
+	template<typename ValueType> void TestNotSame(const TCHAR* Description, const ValueType& A, const ValueType& B)
 	{
 		if (&A == &B)
 		{
-			AddError(FString::Printf(TEXT("%s: The two values are the same."), *Description), 1);
+			AddError(FString::Printf(TEXT("%s: The two values are the same."), Description), 1);
 		}
 	}
 
+	template<typename ValueType> void TestNotSame(const FString& Description, const ValueType& A, const ValueType& B)
+	{
+		TestNotSame(*Description, A, B);
+	}
+
 	/**
 	 * Logs an error if the specified pointer is not NULL.
 	 *
@@ -1357,7 +1418,12 @@ public:
 	 *
 	 * @see TestNotNull
 	 */
-	void TestNull(const FString& What, const void* Pointer);
+	void TestNull(const TCHAR* What, const void* Pointer);
+
+	void TestNull(const FString& What, const void* Pointer)
+	{
+		TestNull(*What, Pointer);
+	}
 
 	/**
 	 * Logs an error if the two values are not the same object in memory.
@@ -1368,14 +1434,19 @@ public:
 	 *
 	 * @see TestNotSame
 	 */
-	template<typename ValueType> void TestSame(const FString& Description, const ValueType& A, const ValueType& B)
+	template<typename ValueType> void TestSame(const TCHAR* Description, const ValueType& A, const ValueType& B)
 	{
 		if (&A != &B)
 		{
-			AddError(FString::Printf(TEXT("%s: The two values are not the same."), *Description), 1);
+			AddError(FString::Printf(TEXT("%s: The two values are not the same."), Description), 1);
 		}
 	}
 
+	template<typename ValueType> void TestSame(const FString& Description, const ValueType& A, const ValueType& B)
+	{
+		TestSame(*Description, A, B);
+	}
+
 	/**
 	 * Logs an error if the specified Boolean value is not true.
 	 *
@@ -1384,7 +1455,12 @@ public:
 	 *
 	 * @see TestFalse
 	 */
-	void TestTrue(const FString& What, bool Value);
+	void TestTrue(const TCHAR* What, bool Value);
+
+	void TestTrue(const FString& What, bool Value)
+	{
+		TestTrue(*What, Value);
+	}
 
 	/** Macro version of above, uses the passed in expression as the description as well */
 	#define TestTrueExpr(Expression) TestTrue(TEXT(#Expression), Expression)
@@ -1397,14 +1473,19 @@ public:
 	 *
 	 * @see TestInvalid
 	 */
-	template<typename ValueType> void TestValid(const FString& Description, const TSharedPtr<ValueType>& SharedPointer)
+	template<typename ValueType> void TestValid(const TCHAR* Description, const TSharedPtr<ValueType>& SharedPointer)
 	{
 		if (!SharedPointer.IsValid())
 		{
-			AddError(FString::Printf(TEXT("%s: The shared pointer is not valid."), *Description), 1);
+			AddError(FString::Printf(TEXT("%s: The shared pointer is not valid."), Description), 1);
 		}
 	}
 
+	template<typename ValueType> void TestValid(const FString& Description, const TSharedPtr<ValueType>& SharedPointer)
+	{
+		TestValid(*Description, SharedPointer);
+	}
+
 protected:
 	/**
 	 * Asks the test to enumerate variants that will all go through the "RunTest" function with different parameters (for load all maps, this should enumerate all maps to load)\
