GitHub #5366 : Make TWeakObjectPtr's assignment operators return a reference to self

diff --git a/Engine/Source/Runtime/Core/Public/UObject/WeakObjectPtrTemplates.h b/Engine/Source/Runtime/Core/Public/UObject/WeakObjectPtrTemplates.h
--- a/Engine/Source/Runtime/Core/Public/UObject/WeakObjectPtrTemplates.h
+++ b/Engine/Source/Runtime/Core/Public/UObject/WeakObjectPtrTemplates.h
@@ -111,17 +111,19 @@ public:
 	 * @param Object object to create a weak pointer to
 	**/
 	template<class U>
-	FORCEINLINE typename TEnableIf<!TLosesQualifiersFromTo<U, T>::Value>::Type operator=(U* Object)
+	FORCEINLINE typename TEnableIf<!TLosesQualifiersFromTo<U, T>::Value, TWeakObjectPtr&>::Type operator=(U* Object)
 	{
 		T* TempObject = Object;
 		TWeakObjectPtrBase::operator=(TempObject);
+		return *this;
 	}
 	template<class U>
 	UE_DEPRECATED(4.19, "Implicit conversions from const pointers to non-const TWeakObjectPtrs has been deprecated. Please const-correct this usage.")
-	FORCEINLINE typename TEnableIf<TLosesQualifiersFromTo<U, T>::Value>::Type operator=(U* Object)
+	FORCEINLINE typename TEnableIf<TLosesQualifiersFromTo<U, T>::Value, TWeakObjectPtr&>::Type operator=(U* Object)
 	{
 		const T* TempObject = Object;
 		TWeakObjectPtrBase::operator=(TempObject);
+		return *this;
 	}
 
 	/**  
@@ -129,13 +131,15 @@ public:
 	 * @param Other weak pointer to copy from
 	**/
 	template <typename OtherT>
-	FORCEINLINE void operator=(const TWeakObjectPtr<OtherT, TWeakObjectPtrBase>& Other)
+	FORCEINLINE TWeakObjectPtr& operator=(const TWeakObjectPtr<OtherT, TWeakObjectPtrBase>& Other)
 	{
 		// It's also possible that this static_assert may fail for valid conversions because
 		// one or both of the types have only been forward-declared.
 		static_assert(TPointerIsConvertibleFromTo<OtherT, T>::Value, "Unable to convert TWeakObjectPtr - types are incompatible");
 
 		*(TWeakObjectPtrBase*)this = *(TWeakObjectPtrBase*)&Other; // we do a C-style cast to private base here to avoid clang 3.6.0 compilation problems with friend declarations
+
+		return *this;
 	}
 
 	/**  
