GitHub #5414 : Fix compiler error in Atomic.h caused by VS2017's recent changes in template handling.

diff --git a/Engine/Source/Runtime/Core/Public/Templates/Atomic.h b/Engine/Source/Runtime/Core/Public/Templates/Atomic.h
--- a/Engine/Source/Runtime/Core/Public/Templates/Atomic.h
+++ b/Engine/Source/Runtime/Core/Public/Templates/Atomic.h
@@ -185,7 +185,7 @@ namespace UE4Atomic_Private
 
 // Basic storage and implementation - only allows getting and setting via platform atomics.
 template <typename T>
-struct alignas(UE4Atomic_Private::TIsSupportedSize<sizeof(T)>::Value ? alignof(UE4Atomic_Private::TUnderlyingIntegerType_T<T>) : alignof(T)) TAtomicBase_Basic
+struct alignas((UE4Atomic_Private::TIsSupportedSize<sizeof(T)>::Value) ? alignof(UE4Atomic_Private::TUnderlyingIntegerType_T<T>) : alignof(T)) TAtomicBase_Basic
 {
 public:
 	/**
