PR #5400: Replace TrackFree with IncrTag (Contributed by [REDACTED_PERSON])

diff --git a/Engine/Source/Runtime/Core/Private/HAL/LowLevelMemTracker.cpp b/Engine/Source/Runtime/Core/Private/HAL/LowLevelMemTracker.cpp
--- a/Engine/Source/Runtime/Core/Private/HAL/LowLevelMemTracker.cpp
+++ b/Engine/Source/Runtime/Core/Private/HAL/LowLevelMemTracker.cpp
@@ -550,9 +550,9 @@ void FLowLevelMemTracker::UpdateStatsPerFrame(const TCHAR* LogName)
 	}
 
 	// calculate FMalloc unused stat and set it in the Default tracker
-	int32 FMallocAmount = Trackers[(int32)ELLMTracker::Default]->GetAllocTypeAmount(ELLMAllocType::FMalloc);
-	int32 FMallocPlatformAmount = Trackers[(int32)ELLMTracker::Platform]->GetTagAmount(ELLMTag::FMalloc);
-	int32 FMallocUnused = FMallocPlatformAmount - FMallocAmount;
+	int64 FMallocAmount = Trackers[(int32)ELLMTracker::Default]->GetAllocTypeAmount(ELLMAllocType::FMalloc);
+	int64 FMallocPlatformAmount = Trackers[(int32)ELLMTracker::Platform]->GetTagAmount(ELLMTag::FMalloc);
+	int64 FMallocUnused = FMallocPlatformAmount - FMallocAmount;
 	Trackers[(int32)ELLMTracker::Default]->SetTagAmount(ELLMTag::FMallocUnused, FMallocUnused, true);
 
 	// update totals for all trackers
@@ -1159,7 +1159,7 @@ void FLLMTracker::TrackFree(const void* Ptr, ELLMTracker Tracker, ELLMAllocType
 #if LLM_USE_ALLOC_INFO_STRUCT
 	State->TrackFree(Ptr, AllocInfo.Tag, Size, true, Tracker, AllocType);
 	#if LLM_ALLOW_ASSETS_TAGS
-		State->TrackFree(nullptr, AllocInfo.AssetTag, Size, false, Tracker, AllocType);
+		State->IncrTag(AllocInfo.AssetTag, 0 - Size, false);
 	#endif
 #else
 	State->TrackFree(Ptr, (int64)AllocInfo, Size, true, Tracker, AllocType);
diff --git a/Engine/Source/Runtime/Core/Public/HAL/LowLevelMemTracker.h b/Engine/Source/Runtime/Core/Public/HAL/LowLevelMemTracker.h
--- a/Engine/Source/Runtime/Core/Public/HAL/LowLevelMemTracker.h
+++ b/Engine/Source/Runtime/Core/Public/HAL/LowLevelMemTracker.h
@@ -29,7 +29,8 @@
 #endif
 
 	// using asset tagging requires a significantly higher number of per-thread tags, so make it optional
-	// even if this is on, we still need to run with -llmtagset=assets because of the shear number of stat ids it makes
+	// even if this is on, we still need to run with -llmtagsets=assets because of the shear number of stat ids it makes
+	// LLM Assets can be viewed in game using 'Stat LLMAssets'
 	#define LLM_ALLOW_ASSETS_TAGS 0
 
 	#define LLM_STAT_TAGS_ENABLED (LLM_ALLOW_ASSETS_TAGS || 0)
