Fixed unit parsing bugs.

* Fix bugs,  always used first matching unit (e.g always used "c"/Celcius instead of "cd").
* ParseWhitespace now correctly uses the Acculumate argument.
* Added candela per meters squared to supported units.
* Added unit test for unit conversion.

PR #5431: UE-68074: Proper detection of unit strings (Contributed by [REDACTED_PERSON])

diff --git a/Engine/Source/Runtime/Core/Private/Math/UnitConversion.cpp b/Engine/Source/Runtime/Core/Private/Math/UnitConversion.cpp
--- a/Engine/Source/Runtime/Core/Private/Math/UnitConversion.cpp
+++ b/Engine/Source/Runtime/Core/Private/Math/UnitConversion.cpp
@@ -67,7 +67,7 @@ FParseCandidate ParseCandidates[] = {
 	{ TEXT("Lumens"),				EUnit::Lumens },				{ TEXT("lm"),		EUnit::Lumens },
 	{ TEXT("Candela"),				EUnit::Candela },				{ TEXT("cd"),		EUnit::Candela },
 	{ TEXT("Lux"),					EUnit::Lux },					{ TEXT("lx"),		EUnit::Lux },
-	{ TEXT("CandelaPerMeterSquared"), EUnit::CandelaPerMeter2 },	{ TEXT("CandelaPerMeter2"), EUnit::CandelaPerMeter2 },
+	{ TEXT("CandelaPerMeterSquared"), EUnit::CandelaPerMeter2 },	{ TEXT("cd/m2"),	EUnit::CandelaPerMeter2 },		{ TEXT("CandelaPerMeter2"),		EUnit::CandelaPerMeter2 },
 
 	{ TEXT("Milliseconds"),			EUnit::Milliseconds },			{ TEXT("ms"),		EUnit::Milliseconds },
 	{ TEXT("Seconds"),				EUnit::Seconds },				{ TEXT("s"),		EUnit::Seconds },
@@ -288,24 +288,33 @@ struct FUnitExpressionParser
 		
 		if (NumberToken.IsSet())
 		{
-			// Skip over whitespace
-			Stream.ParseToken([](TCHAR InC){ return FChar::IsWhitespace(InC) ? EParseState::Continue : EParseState::StopBefore; }, &NumberToken.GetValue());
+			// Skip past any additional whitespace between number and unit token, e.g "1.0 cm"
+			const TOptional<FStringToken> WhiteSpaceToken = Stream.ParseWhitespace(&NumberToken.GetValue());
 
+			size_t LongestMatch = 0;
 			TOptional<EUnit> Unit;
-			TOptional<FStringToken> UnitToken;
+			FStringToken UnitToken;
+			const FStringToken StartingToken = NumberToken.GetValue();
+
 			for (int32 Index = 0; Index < ARRAY_COUNT(ParseCandidates); ++Index)
 			{
-				UnitToken = Stream.ParseTokenIgnoreCase(ParseCandidates[Index].String, &NumberToken.GetValue());
-				if (UnitToken.IsSet())
+				TOptional<FStringToken> CandidateUnitToken = Stream.ParseTokenIgnoreCase(ParseCandidates[Index].String, &NumberToken.GetValue());
+				const size_t MatchedCharCount = NumberToken.GetValue().GetTokenEndPos() - StartingToken.GetTokenEndPos();
+				if (CandidateUnitToken.IsSet() && MatchedCharCount > LongestMatch)
 				{
 					Unit = ParseCandidates[Index].Unit;
-					break;
+					UnitToken = CandidateUnitToken.GetValue();
+					LongestMatch = MatchedCharCount;
+				}
+				else
+				{
+					NumberToken = StartingToken;
 				}
 			}
 
 			if (Unit.IsSet())
 			{
-				Consumer.Add(NumberToken.GetValue(), FNumericUnit<double>(Value, Unit.GetValue()));
+				Consumer.Add(UnitToken, FNumericUnit<double>(Value, Unit.GetValue()));
 			}
 			else
 			{
diff --git a/Engine/Source/Runtime/Core/Private/Misc/ExpressionParser.cpp b/Engine/Source/Runtime/Core/Private/Misc/ExpressionParser.cpp
--- a/Engine/Source/Runtime/Core/Private/Misc/ExpressionParser.cpp
+++ b/Engine/Source/Runtime/Core/Private/Misc/ExpressionParser.cpp
@@ -228,7 +228,7 @@ TOptional<FStringToken> FTokenStream::ParseWhitespace(FStringToken* Accumulate)
 
 	if (IsReadPosValid(OptReadPos))
 	{
-		return ParseToken([](TCHAR InC){ return FChar::IsWhitespace(InC) ? EParseState::Continue : EParseState::StopBefore; });
+		return ParseToken([](TCHAR InC){ return FChar::IsWhitespace(InC) ? EParseState::Continue : EParseState::StopBefore; }, Accumulate);
 	}
 
 	return TOptional<FStringToken>();
diff --git a/Engine/Source/Runtime/Core/Private/Tests/Math/UnitConversionTests.cpp b/Engine/Source/Runtime/Core/Private/Tests/Math/UnitConversionTests.cpp
--- a/Engine/Source/Runtime/Core/Private/Tests/Math/UnitConversionTests.cpp
+++ b/Engine/Source/Runtime/Core/Private/Tests/Math/UnitConversionTests.cpp
@@ -9,6 +9,7 @@
 #if WITH_DEV_AUTOMATION_TESTS
 
 IMPLEMENT_SIMPLE_AUTOMATION_TEST(FUnitUnitTests, "System.Core.Math.Unit Conversion", EAutomationTestFlags::ApplicationContextMask | EAutomationTestFlags::SmokeFilter)
+IMPLEMENT_SIMPLE_AUTOMATION_TEST(FParsingUnitTests, "System.Core.Math.Unit Parsing", EAutomationTestFlags::ApplicationContextMask | EAutomationTestFlags::SmokeFilter)
 
 bool IsRoughlyEqual(double One, double Two, float Epsilon)
 {
@@ -106,4 +107,56 @@ bool FUnitUnitTests::RunTest(const FString& Parameters)
 	return bSuccess;
 }
 
+bool FParsingUnitTests::RunTest(const FString& Parameters)
+{
+	struct FTestCases
+	{
+		TCHAR*							Expression;
+		double							ExpectedValue;
+		EUnit							UnderlyingUnit;
+		TOptional<FNumericUnit<double>>	ExistingValue;
+	};
+
+	FTestCases Tests[] = {
+		{TEXT("10.7cm"),						10.7,		EUnit::Centimeters,			TOptional<FNumericUnit<double>>()},
+		{TEXT("0.7 m"),							70.0,		EUnit::Centimeters,			TOptional<FNumericUnit<double>>()},
+		{TEXT("2m - 1m"),						100.0,		EUnit::Centimeters,			TOptional<FNumericUnit<double>>()},
+		{TEXT("10.4865 MetersPerSecond"),		10.4865,	EUnit::MetersPerSecond,		TOptional<FNumericUnit<double>>()},
+		{TEXT("10.8 cd"),						10.8,		EUnit::Candela,				TOptional<FNumericUnit<double>>()},
+		{TEXT("4.8 cd + 1.2cd"),				6.0,		EUnit::Candela,				TOptional<FNumericUnit<double>>()},
+		{TEXT("1cd/m2 + 0.5 CandelaPerMeter2"),	1.5,		EUnit::CandelaPerMeter2,	TOptional<FNumericUnit<double>>()},
+		{TEXT("+=0.7 m"),						140.0,		EUnit::Centimeters,			TOptional<FNumericUnit<double>>(70.0)},
+	};
+
+	bool bSuccess = true;
+	for (FTestCases& Test : Tests)
+	{
+		const FNumericUnit<double> ExistingValue = Test.ExistingValue.IsSet() ? Test.ExistingValue.GetValue() : FNumericUnit<double>(0.0, EUnit::Unspecified);
+		TValueOrError<FNumericUnit<double>, FText> Result = UnitConversion::TryParseExpression(Test.Expression, Test.UnderlyingUnit, ExistingValue);
+		if (Result.IsValid())
+		{
+			const bool IsEqual = IsRoughlyEqual(Result.GetValue().ConvertTo(Test.UnderlyingUnit).GetValue().Value, Test.ExpectedValue, 1e-6);
+			if (!IsEqual)
+			{
+				AddError(FString::Printf(TEXT("Parsing of expression \"%s\" failed. Expected %d but got %d."),
+					Test.Expression,
+					Test.ExpectedValue,
+					Result.GetValue().Value
+				));
+			}
+			bSuccess &= IsEqual;
+		}
+		else
+		{
+			AddError(FString::Printf( TEXT("Parsing of expression \"%s\" was incorrect (%s). Expected %d."), 
+				Test.Expression, 
+				*(Result.GetError().ToString()), 
+				Test.ExpectedValue
+			));
+			bSuccess = false;
+		}
+	}
+	return bSuccess;
+}
+
 #endif //WITH_DEV_AUTOMATION_TESTS
