Forgot const specifier for expression member.

diff --git a/Engine/Source/Runtime/Core/Private/Tests/Math/UnitConversionTests.cpp b/Engine/Source/Runtime/Core/Private/Tests/Math/UnitConversionTests.cpp
--- a/Engine/Source/Runtime/Core/Private/Tests/Math/UnitConversionTests.cpp
+++ b/Engine/Source/Runtime/Core/Private/Tests/Math/UnitConversionTests.cpp
@@ -111,7 +111,7 @@ bool FParsingUnitTests::RunTest(const FString& Parameters)
 {
 	struct FTestCases
 	{
-		TCHAR*							Expression;
+		const TCHAR*					Expression;
 		double							ExpectedValue;
 		EUnit							UnderlyingUnit;
 		TOptional<FNumericUnit<double>>	ExistingValue;
