MD5 and SHA1 Update functions now handle input lengths of up to 61 bits

diff --git a/Engine/Source/Runtime/Core/Private/Misc/SecureHash.cpp b/Engine/Source/Runtime/Core/Private/Misc/SecureHash.cpp
--- a/Engine/Source/Runtime/Core/Private/Misc/SecureHash.cpp
+++ b/Engine/Source/Runtime/Core/Private/Misc/SecureHash.cpp
@@ -96,19 +96,18 @@ FMD5::~FMD5()
 
 }
 
-void FMD5::Update( const uint8* input, int32 inputLen )
+void FMD5::Update( const uint8* input, uint64 inputLen )
 {
-	int32 i, index, partLen;
+	uint64 i, index, partLen;
 
 	// Compute number of bytes mod 64.
 	index = (int32)((Context.count[0] >> 3) & 0x3F);
 
 	// Update number of bits.
-	if ((Context.count[0] += ((uint32)inputLen << 3)) < ((uint32)inputLen << 3))
-	{
-		Context.count[1]++;
-	}
-	Context.count[1] += ((uint32)inputLen >> 29);
+	uint64 inputBitCount = inputLen << 3;
+	uint64 newBitCount = (uint64(Context.count[0]) | (uint64(Context.count[1]) << 32)) + inputBitCount;
+	Context.count[0] = uint32(newBitCount);
+	Context.count[1] = uint32(newBitCount >> 32);
 
 	partLen = 64 - index;
 
@@ -681,15 +680,16 @@ void FSHA1::Transform(uint32 *state, const uint8 *buffer)
 #undef _R4
 
 // Use this function to hash in binary data
-void FSHA1::Update(const uint8 *data, uint32 len)
+void FSHA1::Update(const uint8 *data, uint64 len)
 {
-	uint32 i, j;
+	uint64 i, j;
 
 	j = (m_count[0] >> 3) & 63;
 
-	if((m_count[0] += len << 3) < (len << 3)) m_count[1]++;
-
-	m_count[1] += (len >> 29);
+	uint64 inputBitCount = len << 3;
+	uint64 newBitCount = (uint64(m_count[0]) | (uint64(m_count[1]) << 32)) + inputBitCount;
+	m_count[0] = uint32(newBitCount);
+	m_count[1] = uint32(newBitCount >> 32);
 
 	if((j + len) > 63)
 	{
@@ -750,7 +750,7 @@ void FSHA1::GetHash(uint8 *puDest)
 * @param DataSize Size of the Data block
 * @param OutHash Resulting hash value (20 byte buffer)
 */
-void FSHA1::HashBuffer(const void* Data, uint32 DataSize, uint8* OutHash)
+void FSHA1::HashBuffer(const void* Data, uint64 DataSize, uint8* OutHash)
 {
 	// do an atomic hash operation
 	FSHA1 Sha;
@@ -759,7 +759,7 @@ void FSHA1::HashBuffer(const void* Data, uint32 DataSize, uint8* OutHash)
 	Sha.GetHash(OutHash);
 }
 
-void FSHA1::HMACBuffer(const void* Key, uint32 KeySize, const void* Data, uint32 DataSize, uint8* OutHash)
+void FSHA1::HMACBuffer(const void* Key, uint32 KeySize, const void* Data, uint64 DataSize, uint8* OutHash)
 {
 	const uint8 BlockSize = 64;
 	const uint8 HashSize = 20;
diff --git a/Engine/Source/Runtime/Core/Public/Misc/SecureHash.h b/Engine/Source/Runtime/Core/Public/Misc/SecureHash.h
--- a/Engine/Source/Runtime/Core/Public/Misc/SecureHash.h
+++ b/Engine/Source/Runtime/Core/Public/Misc/SecureHash.h
@@ -49,7 +49,7 @@ public:
 	 * @param input		input data
 	 * @param inputLen	length of the input data in bytes
 	 **/
-	void Update(const uint8* input, int32 inputLen);
+	void Update(const uint8* input, uint64 inputLen);
 
 	/**
 	 * MD5 finalization. Ends an MD5 message-digest operation, writing the
@@ -233,7 +233,7 @@ public:
 	void Reset();
 
 	// Update the hash value
-	void Update(const uint8 *data, uint32 len);
+	void Update(const uint8 *data, uint64 len);
 
 	// Update the hash value with string
 	void UpdateWithString(const TCHAR *data, uint32 len);
@@ -251,7 +251,7 @@ public:
 	 * @param DataSize Size of the Data block
 	 * @param OutHash Resulting hash value (20 byte buffer)
 	 */
-	static void HashBuffer(const void* Data, uint32 DataSize, uint8* OutHash);
+	static void HashBuffer(const void* Data, uint64 DataSize, uint8* OutHash);
 
 	/**
 	 * Generate the HMAC (Hash-based Message Authentication Code) for a block of data.
@@ -263,7 +263,7 @@ public:
 	 * @param DataSize	Size of the Data block
 	 * @param OutHash	Resulting hash value (20 byte buffer)
 	 */
-	static void HMACBuffer(const void* Key, uint32 KeySize, const void* Data, uint32 DataSize, uint8* OutHash);
+	static void HMACBuffer(const void* Key, uint32 KeySize, const void* Data, uint64 DataSize, uint8* OutHash);
 
 	/**
 	 * Shared hashes.sha reading code (each platform gets a buffer to the data,
