Fix for accidental copying of non-copyable TUniqueFunction bindings.

diff --git a/Engine/Source/Runtime/Core/Public/Templates/Function.h b/Engine/Source/Runtime/Core/Public/Templates/Function.h
--- a/Engine/Source/Runtime/Core/Public/Templates/Function.h
+++ b/Engine/Source/Runtime/Core/Public/Templates/Function.h
@@ -101,7 +101,7 @@ template <typename T> struct TIsTFunctionRef<const volatile T> { enum { Value =
 namespace UE4Function_Private
 {
 	template <typename T, bool bOnHeap>
-	struct TFunction_CopyableOwnedObject;
+	struct TFunction_OwnedObject;
 
 	template <bool bUnique>
 	struct TFunctionStorage;
@@ -143,7 +143,7 @@ namespace UE4Function_Private
 		 */
 		virtual void Destroy() override
 		{
-			using DerivedType = TFunction_CopyableOwnedObject<T, true>;
+			using DerivedType = TFunction_OwnedObject<T, true>;
 
 			auto* Derived = static_cast<DerivedType*>(this);
 			Derived->DerivedType::~DerivedType();
@@ -162,29 +162,46 @@ namespace UE4Function_Private
 		 */
 		virtual void Destroy() override
 		{
-			using DerivedType = TFunction_CopyableOwnedObject<T, false>;
+			using DerivedType = TFunction_OwnedObject<T, false>;
 
 			auto* Derived = static_cast<DerivedType*>(this);
 			Derived->DerivedType::~DerivedType();
 		}
 	};
 
-	/**
-	 * Implementation of IFunction_OwnedObject for a given copyable T.
-	 */
 	template <typename T, bool bOnHeap>
-	struct TFunction_CopyableOwnedObject final : public
+	struct TFunction_OwnedObject : public
 #if TFUNCTION_USES_INLINE_STORAGE
 		TChooseClass<bOnHeap, IFunction_OwnedObject_OnHeap<T>, IFunction_OwnedObject_Inline<T>>::Result
 #else
 		IFunction_OwnedObject_OnHeap<T>
 #endif
+	{
+		template <typename... ArgTypes>
+		explicit TFunction_OwnedObject(ArgTypes&&... Args)
+			: Obj(Forward<ArgTypes>(Args)...)
+		{
+		}
+
+		virtual void* GetAddress() override
+		{
+			return &Obj;
+		}
+
+		T Obj;
+	};
+
+	/**
+	 * Implementation of IFunction_OwnedObject for a given copyable T.
+	 */
+	template <typename T, bool bOnHeap>
+	struct TFunction_CopyableOwnedObject final : public TFunction_OwnedObject<T, bOnHeap>
 	{
 		/**
 		 * Constructor which creates its T by copying.
 		 */
 		explicit TFunction_CopyableOwnedObject(const T& InObj)
-			: Obj(InObj)
+			: TFunction_OwnedObject<T, bOnHeap>(InObj)
 		{
 		}
 
@@ -192,36 +209,24 @@ namespace UE4Function_Private
 		 * Constructor which creates its T by moving.
 		 */
 		explicit TFunction_CopyableOwnedObject(T&& InObj)
-			: Obj(MoveTemp(InObj))
+			: TFunction_OwnedObject<T, bOnHeap>(MoveTemp(InObj))
 		{
 		}
 
 		void* CloneToEmptyStorage(void* UntypedStorage) const override;
-
-		virtual void* GetAddress() override
-		{
-			return &Obj;
-		}
-
-		T Obj;
 	};
 
 	/**
 	 * Implementation of IFunction_OwnedObject for a given non-copyable T.
 	 */
 	template <typename T, bool bOnHeap>
-	struct TFunction_UniqueOwnedObject final : public
-#if TFUNCTION_USES_INLINE_STORAGE
-		TChooseClass<bOnHeap, IFunction_OwnedObject_OnHeap<T>, IFunction_OwnedObject_Inline<T>>::Result
-#else
-		IFunction_OwnedObject_OnHeap<T>
-#endif
+	struct TFunction_UniqueOwnedObject final : public TFunction_OwnedObject<T, bOnHeap>
 	{
 		/**
 		 * Constructor which creates its T by moving.
 		 */
 		explicit TFunction_UniqueOwnedObject(T&& InObj)
-			: Obj(MoveTemp(InObj))
+			: TFunction_OwnedObject<T, bOnHeap>(MoveTemp(InObj))
 		{
 		}
 
@@ -231,13 +236,6 @@ namespace UE4Function_Private
 			check(false);
 			return nullptr;
 		}
-
-		virtual void* GetAddress() override
-		{
-			return &Obj;
-		}
-
-		T Obj;
 	};
 
 	template <typename T>
