CIS fix.

diff --git a/Engine/Source/Runtime/Core/Public/Templates/Function.h b/Engine/Source/Runtime/Core/Public/Templates/Function.h
--- a/Engine/Source/Runtime/Core/Public/Templates/Function.h
+++ b/Engine/Source/Runtime/Core/Public/Templates/Function.h
@@ -431,7 +431,7 @@ namespace UE4Function_Private
 			CA_ASSUME(NewAlloc);
 		}
 
-		auto* NewOwned = new (NewAlloc) TFunction_CopyableOwnedObject(Obj);
+		auto* NewOwned = new (NewAlloc) TFunction_CopyableOwnedObject(this->Obj);
 
 		return &NewOwned->Obj;
 	}
