Fixing more CIS issues

diff --git a/Engine/Source/Runtime/Core/Public/Misc/AssertionMacros.h b/Engine/Source/Runtime/Core/Public/Misc/AssertionMacros.h
--- a/Engine/Source/Runtime/Core/Public/Misc/AssertionMacros.h
+++ b/Engine/Source/Runtime/Core/Public/Misc/AssertionMacros.h
@@ -183,12 +183,12 @@ public:
 			{ \
 				struct Impl \
 				{ \
-					static void FORCENOINLINE UE_DEBUG_SECTION Exec() \
+					static void FORCENOINLINE UE_DEBUG_SECTION ExecCheckImplInternal() \
 					{ \
 						FDebug::CheckVerifyFailed(#expr, __FILE__, __LINE__, TEXT("")); \
 					} \
 				}; \
-				Impl::Exec(); \
+				Impl::ExecCheckImplInternal(); \
 				PLATFORM_BREAK(); \
 				CA_ASSUME(false); \
 			} \
