Moved logic to decide if a failed ensure should be handled into the non-inlined lambda to keep generated code clean.

diff --git a/Engine/Source/Runtime/Core/Public/Misc/AssertionMacros.h b/Engine/Source/Runtime/Core/Public/Misc/AssertionMacros.h
--- a/Engine/Source/Runtime/Core/Public/Misc/AssertionMacros.h
+++ b/Engine/Source/Runtime/Core/Public/Misc/AssertionMacros.h
@@ -298,10 +298,10 @@ public:
 #if DO_CHECK && !USING_CODE_ANALYSIS // The Visual Studio 2013 analyzer doesn't understand these complex conditionals
 
 	#define UE_ENSURE_IMPL(Capture, Always, InExpression, ...) \
-		(LIKELY(!!(InExpression)) || (FPlatformMisc::IsEnsureAllowed() && DispatchCheckVerify<bool>([Capture] () FORCENOINLINE UE_DEBUG_SECTION \
+		(LIKELY(!!(InExpression)) || DispatchCheckVerify<bool>([Capture] () FORCENOINLINE UE_DEBUG_SECTION \
 		{ \
 			static bool bExecuted = false; \
-			if (!bExecuted || Always) \
+			if ((!bExecuted || Always) && FPlatformMisc::IsEnsureAllowed()) \
 			{ \
 				bExecuted = true; \
 				FDebug::OptionallyLogFormattedEnsureMessageReturningFalse(true, #InExpression, __FILE__, __LINE__, ##__VA_ARGS__); \
@@ -312,7 +312,7 @@ public:
 				} \
 			} \
 			return !Always; \
-		})) || ([] () { PLATFORM_BREAK(); } (), true))
+		}) || ([] () { PLATFORM_BREAK(); } (), true))
 
 	#define ensure(           InExpression                ) UE_ENSURE_IMPL( , false, InExpression, TEXT(""))
 	#define ensureMsgf(       InExpression, InFormat, ... ) UE_ENSURE_IMPL(&, false, InExpression, InFormat, ##__VA_ARGS__)
