Hopeful attempt for crash on Clang builds.

diff --git a/Engine/Source/Runtime/Core/Public/Templates/Function.h b/Engine/Source/Runtime/Core/Public/Templates/Function.h
--- a/Engine/Source/Runtime/Core/Public/Templates/Function.h
+++ b/Engine/Source/Runtime/Core/Public/Templates/Function.h
@@ -143,11 +143,9 @@ namespace UE4Function_Private
 		 */
 		virtual void Destroy() override
 		{
-			using DerivedType = TFunction_OwnedObject<T, true>;
-
-			auto* Derived = static_cast<DerivedType*>(this);
-			Derived->DerivedType::~DerivedType();
-			FMemory::Free(Derived);
+			void* This = this;
+			this->~IFunction_OwnedObject_OnHeap();
+			FMemory::Free(This);
 		}
 	};
 
@@ -162,10 +160,7 @@ namespace UE4Function_Private
 		 */
 		virtual void Destroy() override
 		{
-			using DerivedType = TFunction_OwnedObject<T, false>;
-
-			auto* Derived = static_cast<DerivedType*>(this);
-			Derived->DerivedType::~DerivedType();
+			this->~IFunction_OwnedObject_Inline();
 		}
 	};
 
