Fixed warning from MSVC.

diff --git a/Engine/Source/Runtime/Core/Public/Templates/Function.h b/Engine/Source/Runtime/Core/Public/Templates/Function.h
--- a/Engine/Source/Runtime/Core/Public/Templates/Function.h
+++ b/Engine/Source/Runtime/Core/Public/Templates/Function.h
@@ -147,6 +147,12 @@ namespace UE4Function_Private
 			this->~IFunction_OwnedObject_OnHeap();
 			FMemory::Free(This);
 		}
+
+		~IFunction_OwnedObject_OnHeap() override
+		{
+			// It is not necessary to define this destructor but MSVC will
+			// erroneously issue warning C5046 without it.
+		}
 	};
 
 	/**
@@ -162,6 +168,12 @@ namespace UE4Function_Private
 		{
 			this->~IFunction_OwnedObject_Inline();
 		}
+
+		~IFunction_OwnedObject_Inline() override
+		{
+			// It is not necessary to define this destructor but MSVC will
+			// erroneously issue warning C5046 without it.
+		}
 	};
 
 	template <typename T, bool bOnHeap>
