Fix incorrect calling convention causing stack corruption on Win32.

diff --git a/Engine/Source/Runtime/Core/Private/Windows/WindowsRunnableThread.cpp b/Engine/Source/Runtime/Core/Private/Windows/WindowsRunnableThread.cpp
--- a/Engine/Source/Runtime/Core/Private/Windows/WindowsRunnableThread.cpp
+++ b/Engine/Source/Runtime/Core/Private/Windows/WindowsRunnableThread.cpp
@@ -16,7 +16,7 @@ void FRunnableThreadWin::SetThreadDescription(HANDLE hThread, PCWSTR lpThreadDes
 	// So in order to be compatible with older Windows versions we probe for the API at runtime
 	// and call it only if available.
 
-	typedef HRESULT(*SetThreadDescriptionFnPtr)(HANDLE hThread, PCWSTR lpThreadDescription);
+	typedef HRESULT(WINAPI *SetThreadDescriptionFnPtr)(HANDLE hThread, PCWSTR lpThreadDescription);
 
 #pragma warning( push )
 #pragma warning( disable: 4191 )	// unsafe conversion from 'type of expression' to 'type required'
