Fix range for ASCII char tolower/toupper test (TCharTest::RunTest)

diff --git a/Engine/Source/Runtime/Core/Private/Tests/Misc/CharTest.cpp b/Engine/Source/Runtime/Core/Private/Tests/Misc/CharTest.cpp
--- a/Engine/Source/Runtime/Core/Private/Tests/Misc/CharTest.cpp
+++ b/Engine/Source/Runtime/Core/Private/Tests/Misc/CharTest.cpp
@@ -34,7 +34,7 @@ bool TCharTest::RunTest(const FString& Parameters)
 {
 	TestTrue(TEXT("C locale not used"), strcmp("C", setlocale(LC_CTYPE, nullptr)) == 0);
 
-	RunCharTests<ANSICHAR>(*this, 256);
+	RunCharTests<ANSICHAR>(*this, 128);
 	RunCharTests<WIDECHAR>(*this, 65536);
 
 	return true;
