Fixing non-unity CIS error on the Mac

diff --git a/Engine/Source/Runtime/Core/Public/GenericPlatform/StandardPlatformString.h b/Engine/Source/Runtime/Core/Public/GenericPlatform/StandardPlatformString.h
--- a/Engine/Source/Runtime/Core/Public/GenericPlatform/StandardPlatformString.h
+++ b/Engine/Source/Runtime/Core/Public/GenericPlatform/StandardPlatformString.h
@@ -7,6 +7,7 @@
 #include "GenericPlatform/GenericPlatformStricmp.h"
 #include "GenericPlatform/GenericPlatformString.h"
 #include <wchar.h>
+#include <string.h>
 
 #if (PLATFORM_APPLE || PLATFORM_UNIX || PLATFORM_HTML5 || PLATFORM_PS4 || PLATFORM_SWITCH) && !PLATFORM_TCHAR_IS_CHAR16
 
