Override SerializeIntPacked in FBitReader and FBitWriter. This avoids calling appBitsCpy.

Already reviewed.

diff --git a/Engine/Source/Runtime/Core/Private/Serialization/BitReader.cpp b/Engine/Source/Runtime/Core/Private/Serialization/BitReader.cpp
--- a/Engine/Source/Runtime/Core/Private/Serialization/BitReader.cpp
+++ b/Engine/Source/Runtime/Core/Private/Serialization/BitReader.cpp
@@ -209,6 +209,48 @@ void FBitReader::SerializeBitsWithOffset( void* Dest, int32 DestBit, int64 Lengt
 	}
 }
 
+// This version is bit compatible with FArchive::SerializeIntPacked. See notes in FBitWriter::SerializeIntPacked for more info.
+void FBitReader::SerializeIntPacked(uint32& OutValue)
+{
+	if (IsError())
+	{
+		OutValue = 0;
+		return;
+	}
+
+	const uint8* Src = Buffer.GetData() + (Pos >> 3U);
+	const uint32 BitCountUsedInByte = Pos & 7;
+	const uint32 BitCountLeftInByte = 8 - (Pos & 7);
+	const uint8 SrcMaskByte0 = uint8((1U << BitCountLeftInByte) - 1U);
+	const uint8 SrcMaskByte1 = uint8((1U << BitCountUsedInByte) - 1U);
+	const uint32 NextSrcIndex = (BitCountUsedInByte != 0);
+
+	uint32 Value = 0;
+	for (unsigned It = 0, ShiftCount = 0; It < 5; ++It, ShiftCount += 7)
+	{
+		if (Pos + 8 > Num)
+		{
+			SetOverflowed(8);
+			break;
+		}
+
+		Pos += 8;
+
+		const uint8 Byte = ((Src[0] >> BitCountUsedInByte) & SrcMaskByte0) | ((Src[NextSrcIndex] & SrcMaskByte1) << (BitCountLeftInByte & 7));
+		const uint8 NextByteIndicator = Byte & 1;
+		const uint32 ByteAsWord = Byte >> 1U;
+		Value = (ByteAsWord << ShiftCount) | Value;
+		++Src;
+
+		if (!NextByteIndicator)
+		{
+			break;
+		}
+	}
+
+	OutValue = Value;
+}
+
 /*-----------------------------------------------------------------------------
 	FBitReader.
 -----------------------------------------------------------------------------*/
diff --git a/Engine/Source/Runtime/Core/Private/Serialization/BitWriter.cpp b/Engine/Source/Runtime/Core/Private/Serialization/BitWriter.cpp
--- a/Engine/Source/Runtime/Core/Private/Serialization/BitWriter.cpp
+++ b/Engine/Source/Runtime/Core/Private/Serialization/BitWriter.cpp
@@ -27,7 +27,7 @@ FBitWriter::FBitWriter( int64 InMaxBits, bool InAllowResize /*=false*/ )
 {
 	Buffer.AddUninitialized( (InMaxBits+7)>>3 );
 
-	AllowResize = InAllowResize;
+	bAllowResize = InAllowResize;
 	FMemory::Memzero(Buffer.GetData(), Buffer.Num());
 	this->SetIsSaving(true);
 	this->SetIsPersistent(true);
@@ -42,7 +42,7 @@ FBitWriter::FBitWriter( int64 InMaxBits, bool InAllowResize /*=false*/ )
 FBitWriter::FBitWriter(void)
 	: Num(0)
 	, Max(0)
-	, AllowResize(false)
+	, bAllowResize(false)
 	, bAllowOverflow(false)
 {
 	this->SetIsSaving(true);
@@ -207,6 +207,53 @@ void FBitWriter::SetOverflowed(int32 LengthBits)
 	ArIsError = 1;
 }
 
+/**
+ * This function is bit compatible with FArchive::SerializeIntPacked. It is more efficient
+ * as only a few bytes are written and the base version is best suited for writing many bytes.
+ * This version can be made more efficient and take less bits when we can break backward compatibility.
+ * The last byte will only need 4 bits so we're currently wasting 4 bits. Another way to pack
+ * could be to store 2 bits first in order to indicate how many bytes are needed. That
+ * would eliminate all shifting and masking to reconstruct the bytes. The downside is that
+ * values less than 2^14 will waste 1-2 bits compared to the below algorithm.
+ */
+void FBitWriter::SerializeIntPacked(uint32& InValue)
+{
+	uint32 Value = InValue;
+	uint32 BytesAsWords[5];
+	uint32 ByteCount = 0;
+	for (unsigned It = 0; (It == 0) | (Value != 0); ++It, Value = Value >> 7U)
+	{
+		const uint32 NextByteIndicator = (Value & ~0x7FU) != 0;
+		const uint32 ByteAsWord = ((Value & 0x7FU) << 1U) | NextByteIndicator;
+		BytesAsWords[ByteCount++] = ByteAsWord;
+	}
+
+	const int64 LengthBits = ByteCount * 8;
+	if (!AllowAppend(LengthBits))
+	{
+		SetOverflowed(LengthBits);
+		return;
+	}
+
+	const uint32 BitCountUsedInByte = Num & 7;
+	const uint32 BitCountLeftInByte = 8 - (Num & 7);
+	const uint8 DestMaskByte0 = uint8((1U << BitCountUsedInByte) - 1U);
+	const uint8 DestMaskByte1 = 0xFFU ^ DestMaskByte0;
+	const bool bStraddlesTwoBytes = (BitCountUsedInByte != 0);
+	uint8* Dest = Buffer.GetData() + (Num >> 3U);
+
+	Num += LengthBits;
+	for (uint32 ByteIt = 0; ByteIt != ByteCount; ++ByteIt)
+	{
+		const uint32 ByteAsWord = BytesAsWords[ByteIt];
+
+		*Dest = (*Dest & DestMaskByte0) | uint8(ByteAsWord << BitCountUsedInByte);
+		++Dest;
+		if (bStraddlesTwoBytes)
+			*Dest = (*Dest & DestMaskByte1) | uint8(ByteAsWord >> BitCountLeftInByte);
+	}
+}
+
 /*-----------------------------------------------------------------------------
 	FBitWriterMark.
 -----------------------------------------------------------------------------*/
diff --git a/Engine/Source/Runtime/Core/Public/Serialization/BitReader.h b/Engine/Source/Runtime/Core/Public/Serialization/BitReader.h
--- a/Engine/Source/Runtime/Core/Public/Serialization/BitReader.h
+++ b/Engine/Source/Runtime/Core/Public/Serialization/BitReader.h
@@ -87,6 +87,8 @@ public:
 		}
 	}
 
+	virtual void SerializeIntPacked(uint32& Value) override;
+
 	FORCEINLINE_DEBUGGABLE uint32 ReadInt(uint32 Max)
 	{
 		uint32 Value = 0;
diff --git a/Engine/Source/Runtime/Core/Public/Serialization/BitWriter.h b/Engine/Source/Runtime/Core/Public/Serialization/BitWriter.h
--- a/Engine/Source/Runtime/Core/Public/Serialization/BitWriter.h
+++ b/Engine/Source/Runtime/Core/Public/Serialization/BitWriter.h
@@ -28,7 +28,7 @@ public:
 	 */
 	FBitWriter( int64 InMaxBits, bool AllowResize = false );
 
-	void SerializeBits( void* Src, int64 LengthBits );
+	virtual void SerializeBits( void* Src, int64 LengthBits ) override;
 
 	virtual void SerializeBitsWithOffset( void* Src, int32 SourceBit, int64 LengthBits ) override;
 
@@ -38,7 +38,9 @@ public:
 	 * @param Value		The value to serialize, must be < Max
 	 * @param Max		The maximum allowed value - good to aim for power-of-two
 	 */
-	void SerializeInt(uint32& Value, uint32 Max);
+	virtual void SerializeInt(uint32& Value, uint32 Max) override;
+
+	virtual void SerializeIntPacked(uint32& Value) override;
 
 	/**
 	 * Serializes the specified Value, but does not bounds check against ValueMax;
@@ -50,7 +52,7 @@ public:
 	void WriteIntWrapped(uint32 Value, uint32 ValueMax);
 
 	void WriteBit( uint8 In );
-	void Serialize( void* Src, int64 LengthBytes );
+	virtual void Serialize( void* Src, int64 LengthBytes ) override;
 
 	/**
 	 * Returns a pointer to our internal buffer.
@@ -128,7 +130,7 @@ public:
 	{
 		if (Num+LengthBits > Max)
 		{
-			if (AllowResize)
+			if (bAllowResize)
 			{
 				// Resize our buffer. The common case for resizing bitwriters is hitting the max and continuing to add a lot of small segments of data
 				// Though we could just allow the TArray buffer to handle the slack and resizing, we would still constantly hit the FBitWriter's max
@@ -151,7 +153,7 @@ public:
 
 	FORCEINLINE void SetAllowResize(bool NewResize)
 	{
-		AllowResize = NewResize;
+		bAllowResize = NewResize;
 	}
 
 	/**
@@ -165,11 +167,10 @@ public:
 	}
 
 private:
-
 	TArray<uint8> Buffer;
 	int64   Num;
 	int64   Max;
-	bool	AllowResize;
+	bool	bAllowResize;
 
 	/** Whether or not overflowing is allowed (overflows silently) */
 	bool bAllowOverflow;
