Test platform specific math implementations as well as the generic implementations.
Allow platform specific math tests.

diff --git a/Engine/Source/Runtime/Core/Private/GenericPlatform/GenericPlatformMath.cpp b/Engine/Source/Runtime/Core/Private/GenericPlatform/GenericPlatformMath.cpp
--- a/Engine/Source/Runtime/Core/Private/GenericPlatform/GenericPlatformMath.cpp
+++ b/Engine/Source/Runtime/Core/Private/GenericPlatform/GenericPlatformMath.cpp
@@ -1,6 +1,7 @@
 // Copyright 1998-2018 Epic Games, Inc. All Rights Reserved.
 
 #include "GenericPlatform/GenericPlatformMath.h"
+#include "HAL/PlatformMath.h"
 #include "Misc/AssertionMacros.h"
 #include "Math/UnrealMathUtility.h"
 #include "Math/BigInt.h"
@@ -98,31 +99,42 @@ namespace CompilerHiddenConstants
 	volatile float FloatMax = MAX_FLT;
 }
 
-
-void FGenericPlatformMath::AutoTest() 
+template<class MathPlatform>
+class FPlatformMathTest
 {
+public:
+	// Tests for functions that should be implemented in FGenericPlatformMath and can have a platform specific implementation. 
+	static void AutoTest()
 	{
 		using namespace CompilerHiddenConstants;
 
-		check(IsNaN(sqrtf(MinusOne)));
-		check(!IsFinite(sqrtf(MinusOne)));
-		check(!IsFinite(-1.0f / Zero));
-		check(!IsFinite(1.0f / Zero));
-		check(!IsNaN(-1.0f / Zero));
-		check(!IsNaN(1.0f / Zero));
-		check(!IsNaN(FloatMax));
-		check(IsFinite(FloatMax));
-		check(!IsNaN(Zero));
-		check(IsFinite(Zero));
-		check(!IsNaN(One));
-		check(IsFinite(One));
-		check(!IsNaN(MinusOneE37));
-		check(IsFinite(MinusOneE37));
-		check(FloorLog2(Zero) == 0);
-		check(FloorLog2(One) == 0);
-		check(FloorLog2(Two) == 1);
-		check(FloorLog2(Twelve) == 3);
-		check(FloorLog2(Sixteen) == 4);
+		check(MathPlatform::IsNaN(sqrtf(MinusOne)));
+		check(!MathPlatform::IsFinite(sqrtf(MinusOne)));
+		check(!MathPlatform::IsFinite(-1.0f / Zero));
+		check(!MathPlatform::IsFinite(1.0f / Zero));
+		check(!MathPlatform::IsNaN(-1.0f / Zero));
+		check(!MathPlatform::IsNaN(1.0f / Zero));
+		check(!MathPlatform::IsNaN(FloatMax));
+		check(MathPlatform::IsFinite(FloatMax));
+		check(!MathPlatform::IsNaN(Zero));
+		check(MathPlatform::IsFinite(Zero));
+		check(!MathPlatform::IsNaN(One));
+		check(MathPlatform::IsFinite(One));
+		check(!MathPlatform::IsNaN(MinusOneE37));
+		check(MathPlatform::IsFinite(MinusOneE37));
+		check(MathPlatform::FloorLog2(Zero) == 0);
+		check(MathPlatform::FloorLog2(One) == 0);
+		check(MathPlatform::FloorLog2(Two) == 1);
+		check(MathPlatform::FloorLog2(Twelve) == 3);
+		check(MathPlatform::FloorLog2(Sixteen) == 4);
+	}
+};
+
+void FGenericPlatformMath::AutoTest() 
+{
+	{
+		FPlatformMathTest<FPlatformMath>::AutoTest();
+		FPlatformMathTest<FGenericPlatformMath>::AutoTest();
 	}
 
 	{
diff --git a/Engine/Source/Runtime/Core/Private/Tests/HAL/PlatformTest.cpp b/Engine/Source/Runtime/Core/Private/Tests/HAL/PlatformTest.cpp
--- a/Engine/Source/Runtime/Core/Private/Tests/HAL/PlatformTest.cpp
+++ b/Engine/Source/Runtime/Core/Private/Tests/HAL/PlatformTest.cpp
@@ -53,7 +53,7 @@ bool FPlatformVerificationTest::RunTest (const FString& Parameters)
 	check(*(uint8*)&Test == 0x12);
 #endif
 
-	FGenericPlatformMath::AutoTest();
+	FPlatformMath::AutoTest();
 
 #if WITH_EDITORONLY_DATA
 	check(FPlatformProperties::HasEditorOnlyData());
