- Framepro captures taken with the framepro.startrec/stoprec commands now have the Pid added to the filename
- Exposed the csv filename to a public getter

diff --git a/Engine/Source/Runtime/Core/Private/FramePro/FrameProProfiler.cpp b/Engine/Source/Runtime/Core/Private/FramePro/FrameProProfiler.cpp
--- a/Engine/Source/Runtime/Core/Private/FramePro/FrameProProfiler.cpp
+++ b/Engine/Source/Runtime/Core/Private/FramePro/FrameProProfiler.cpp
@@ -333,7 +333,7 @@ void FFrameProProfiler::PopEvent(const ANSICHAR* Override)
 
 void FFrameProProfiler::StartFrameProRecordingFromCommand(const TArray< FString >& Args)
 {
-	FString FilenameRoot = TEXT("Profile");
+	FString FilenameRoot = FString::Printf(TEXT("ProfilePid%d"), FPlatformProcess::GetCurrentProcessId());
 	if (Args.Num() > 0 && Args[0].Len() > 0)
 	{
 		FilenameRoot = Args[0];
@@ -350,7 +350,7 @@ void FFrameProProfiler::StartFrameProRecordingScopeOverrideFromCommand(const TAr
 		MinScopeTime = FCString::Atoi(*Args[0]);
 	}
 
-	StartFrameProRecording(TEXT("Profile"), MinScopeTime);
+	StartFrameProRecording(FString::Printf(TEXT("ProfilePid%d"), FPlatformProcess::GetCurrentProcessId()), MinScopeTime);
 }
 
 FString FFrameProProfiler::StartFrameProRecording(const FString& FilenameRoot, int32 MinScopeTime)
diff --git a/Engine/Source/Runtime/Core/Public/ProfilingDebugging/CsvProfiler.h b/Engine/Source/Runtime/Core/Public/ProfilingDebugging/CsvProfiler.h
--- a/Engine/Source/Runtime/Core/Public/ProfilingDebugging/CsvProfiler.h
+++ b/Engine/Source/Runtime/Core/Public/ProfilingDebugging/CsvProfiler.h
@@ -208,6 +208,8 @@ public:
 
 	CORE_API void SetDeviceProfileName(FString InDeviceProfileName);
 
+	CORE_API FString GetOutputFilename() const { return OutputFilename; }
+
 private:
 	CORE_API static void VARARGS RecordEventfInternal(int32 CategoryIndex, const TCHAR* Fmt, ...);
 
