Fixed bad merge from dev CL 4707991

diff --git a/Engine/Source/Runtime/Core/Private/Misc/NetworkVersion.cpp b/Engine/Source/Runtime/Core/Private/Misc/NetworkVersion.cpp
--- a/Engine/Source/Runtime/Core/Private/Misc/NetworkVersion.cpp
+++ b/Engine/Source/Runtime/Core/Private/Misc/NetworkVersion.cpp
@@ -18,13 +18,7 @@ FString FNetworkVersion::ProjectVersion = TEXT("1.0.0");
 bool FNetworkVersion::bHasCachedNetworkChecksum			= false;
 uint32 FNetworkVersion::CachedNetworkChecksum			= 0;
 
->>>> ORIGINAL //UE4/Main/Engine/Source/Runtime/Core/Private/Misc/NetworkVersion.cpp#19
-uint32 FNetworkVersion::EngineNetworkProtocolVersion	= HISTORY_CHANNEL_NAMES;
-==== THEIRS //UE4/Main/Engine/Source/Runtime/Core/Private/Misc/NetworkVersion.cpp#20
-uint32 FNetworkVersion::EngineNetworkProtocolVersion	= HISTORY_CHANNEL_CLOSE_REASON;
-==== YOURS //James.Hopkin_RDU-WD-2786_Dev-Networking/Engine/Source/Runtime/Core/Private/Misc/NetworkVersion.cpp
 uint32 FNetworkVersion::EngineNetworkProtocolVersion	= HISTORY_ACKS_INCLUDED_IN_HEADER;
-<<<<
 uint32 FNetworkVersion::GameNetworkProtocolVersion		= 0;
 
 uint32 FNetworkVersion::EngineCompatibleNetworkProtocolVersion		= HISTORY_REPLAY_BACKWARDS_COMPAT;
diff --git a/Engine/Source/Runtime/Core/Public/Misc/NetworkVersion.h b/Engine/Source/Runtime/Core/Public/Misc/NetworkVersion.h
--- a/Engine/Source/Runtime/Core/Public/Misc/NetworkVersion.h
+++ b/Engine/Source/Runtime/Core/Public/Misc/NetworkVersion.h
@@ -38,12 +38,8 @@ enum EEngineNetworkVersionHistory
 	HISTORY_REPCMD_CHECKSUM_REMOVE_PRINTF = 4,		// Bump version since the way FRepLayoutCmd::CompatibleChecksum was calculated changed due to an optimization
 	HISTORY_NEW_ACTOR_OVERRIDE_LEVEL = 5,			// Bump version since a level reference was added to the new actor information
 	HISTORY_CHANNEL_NAMES = 6,						// Bump version since channel type is now an fname
->>>> ORIGINAL //UE4/Main/Engine/Source/Runtime/Core/Public/Misc/NetworkVersion.h#11
-==== THEIRS //UE4/Main/Engine/Source/Runtime/Core/Public/Misc/NetworkVersion.h#12
 	HISTORY_CHANNEL_CLOSE_REASON = 7,				// Bump version to serialize a channel close reason in bunches instead of bDormant
-==== YOURS //James.Hopkin_RDU-WD-2786_Dev-Networking/Engine/Source/Runtime/Core/Public/Misc/NetworkVersion.h
-	HISTORY_ACKS_INCLUDED_IN_HEADER = 7,			// Bump version since acks are now sent as part of the header
-<<<<
+	HISTORY_ACKS_INCLUDED_IN_HEADER = 8,			// Bump version since acks are now sent as part of the header
 };
 
 struct CORE_API FNetworkVersion
