UE-66436 Log from Android device shows incorrect non-unicode character.

#4.22
#Android

Unicode support in FAndroidMisc::LocalPrint (Android implementation of UE_LOG)

diff --git a/Engine/Source/Runtime/Core/Private/Android/AndroidMisc.cpp b/Engine/Source/Runtime/Core/Private/Android/AndroidMisc.cpp
--- a/Engine/Source/Runtime/Core/Private/Android/AndroidMisc.cpp
+++ b/Engine/Source/Runtime/Core/Private/Android/AndroidMisc.cpp
@@ -106,12 +106,12 @@ void FAndroidMisc::LocalPrint(const TCHAR *Message)
 #if !UE_BUILD_SHIPPING
 	const int MAX_LOG_LENGTH = 4096;
 	// not static since may be called by different threads
-	ANSICHAR MessageBuffer[MAX_LOG_LENGTH];
+	wchar_t MessageBuffer[MAX_LOG_LENGTH];
 
 	const TCHAR* SourcePtr = Message;
 	while (*SourcePtr)
 	{
-		ANSICHAR* WritePtr = MessageBuffer;
+		wchar_t* WritePtr = MessageBuffer;
 		int32 RemainingSpace = MAX_LOG_LENGTH;
 		while (*SourcePtr && --RemainingSpace > 0)
 		{
@@ -128,13 +128,12 @@ void FAndroidMisc::LocalPrint(const TCHAR *Message)
 				break;
 			}
 			else {
-				*WritePtr++ = static_cast<ANSICHAR>(*SourcePtr++);
+				*WritePtr++ = static_cast<wchar_t>(*SourcePtr++);
 			}
 		}
 		*WritePtr = '\0';
-		__android_log_write(ANDROID_LOG_DEBUG, "UE4", MessageBuffer);
+		__android_log_print(ANDROID_LOG_DEBUG, "UE4", "%ls", MessageBuffer);
 	}
-	//	__android_log_print(ANDROID_LOG_DEBUG, "UE4", "%s", TCHAR_TO_ANSI(Message));
 #endif
 }
 
