Fixed incorrect VectorShuffle implementation on ARM Neon

https://clang.llvm.org/docs/LanguageExtensions.html#builtin-shufflevector
Contributed by licensee in https://udn.unrealengine.com/questions/476857/broken-vectorshuffle-on-arm-neon.html

diff --git a/Engine/Source/Runtime/Core/Public/Math/UnrealMathNeon.h b/Engine/Source/Runtime/Core/Public/Math/UnrealMathNeon.h
--- a/Engine/Source/Runtime/Core/Public/Math/UnrealMathNeon.h
+++ b/Engine/Source/Runtime/Core/Public/Math/UnrealMathNeon.h
@@ -552,12 +552,12 @@ FORCEINLINE VectorRegister VectorBitwiseXor(const VectorRegister& Vec1, const Ve
 * @param Vec1		Source vector1
 * @param Vec2		Source vector2
 * @param X			Index for which component of Vector1 to use for X (literal 0-3)
-* @param Y			Index for which component to Vector1 to use for Y (literal 0-3)
-* @param Z			Index for which component to Vector2 to use for Z (literal 0-3)
-* @param W			Index for which component to Vector2 to use for W (literal 0-3)
+* @param Y			Index for which component of Vector1 to use for Y (literal 0-3)
+* @param Z			Index for which component of Vector2 to use for Z (literal 0-3)
+* @param W			Index for which component of Vector2 to use for W (literal 0-3)
 * @return			The swizzled vector
 */
-#define VectorShuffle( Vec1, Vec2, X, Y, Z, W )	__builtin_shufflevector(Vec1, Vec2, X, Y, Z, W)
+#define VectorShuffle( Vec1, Vec2, X, Y, Z, W )	__builtin_shufflevector(Vec1, Vec2, X, Y, Z + 4, W + 4)
 
 
 /**
