Added FMacPlatformMisc::GetPlatformChunkInstall() (identical to Windows)

diff --git a/Engine/Source/Runtime/Core/Private/Mac/MacPlatformMisc.cpp b/Engine/Source/Runtime/Core/Private/Mac/MacPlatformMisc.cpp
--- a/Engine/Source/Runtime/Core/Private/Mac/MacPlatformMisc.cpp
+++ b/Engine/Source/Runtime/Core/Private/Mac/MacPlatformMisc.cpp
@@ -27,6 +27,7 @@
 #include "Internationalization/Internationalization.h"
 #include "Internationalization/Culture.h"
 #include "Modules/ModuleManager.h"
+#include "GenericPlatform/GenericPlatformChunkInstall.h"
 
 #include "Apple/PreAppleSystemHeaders.h"
 #include <dlfcn.h>
@@ -2369,3 +2370,37 @@ int FMacPlatformMisc::GetDefaultStackSize()
 	return 4 * 1024 * 1024;
 #endif
 }
+
+IPlatformChunkInstall* FMacPlatformMisc::GetPlatformChunkInstall()
+{
+	static IPlatformChunkInstall* ChunkInstall = nullptr;
+	static bool bIniChecked = false;
+	if (!ChunkInstall || !bIniChecked)
+	{
+		IPlatformChunkInstallModule* PlatformChunkInstallModule = nullptr;
+		if (!GEngineIni.IsEmpty())
+		{
+			FString InstallModule;
+			GConfig->GetString(TEXT("StreamingInstall"), TEXT("DefaultProviderName"), InstallModule, GEngineIni);
+			FModuleStatus Status;
+			if (FModuleManager::Get().QueryModule(*InstallModule, Status))
+			{
+				PlatformChunkInstallModule = FModuleManager::LoadModulePtr<IPlatformChunkInstallModule>(*InstallModule);
+				if (PlatformChunkInstallModule != nullptr)
+				{
+					// Attempt to grab the platform installer
+					ChunkInstall = PlatformChunkInstallModule->GetPlatformChunkInstall();
+				}
+			}
+			bIniChecked = true;
+		}
+
+		if (PlatformChunkInstallModule == nullptr)
+		{
+			// Placeholder instance
+			ChunkInstall = FGenericPlatformMisc::GetPlatformChunkInstall();
+		}
+	}
+
+	return ChunkInstall;
+}
diff --git a/Engine/Source/Runtime/Core/Public/Mac/MacPlatformMisc.h b/Engine/Source/Runtime/Core/Public/Mac/MacPlatformMisc.h
--- a/Engine/Source/Runtime/Core/Public/Mac/MacPlatformMisc.h
+++ b/Engine/Source/Runtime/Core/Public/Mac/MacPlatformMisc.h
@@ -94,6 +94,8 @@ struct CORE_API FMacPlatformMisc : public FApplePlatformMisc
 	 */
 	static bool IsRunningOnBattery();
 
+	static IPlatformChunkInstall* GetPlatformChunkInstall();
+
 	/**
 	 * Returns whether the Mac OS X version is 10.9.x or not.
 	 */
