#DEV-VR - Speculative fix for unix platform tick stat being undeclared

diff --git a/Engine/Source/Runtime/Core/Private/Unix/UnixPlatformCrashContext.cpp b/Engine/Source/Runtime/Core/Private/Unix/UnixPlatformCrashContext.cpp
--- a/Engine/Source/Runtime/Core/Private/Unix/UnixPlatformCrashContext.cpp
+++ b/Engine/Source/Runtime/Core/Private/Unix/UnixPlatformCrashContext.cpp
@@ -296,6 +296,8 @@ void FUnixCrashContext::CaptureStackTrace()
 	}
 }
 
+DECLARE_CYCLE_STAT(TEXT("ImgMedia Loader Release Cache"), STAT_UnixCrashReporterTracker_Tick, STATGROUP_Engine);
+
 namespace UnixCrashReporterTracker
 {
 	FProcHandle CurrentlyRunningCrashReporter;
@@ -303,7 +305,7 @@ namespace UnixCrashReporterTracker
 
 	bool Tick(float DeltaTime)
 	{
-        QUICK_SCOPE_CYCLE_COUNTER(STAT_UnixCrashReporterTracker_Tick);
+		SCOPE_CYCLE_COUNTER(STAT_UnixCrashReporterTracker_Tick);
 
 		if (!FPlatformProcess::IsProcRunning(CurrentlyRunningCrashReporter))
 		{
