#DEV-VR - Merge from //UE4/Dev-MagicLeap/... @ CL 4136411

diff --git a/Engine/Source/Runtime/Core/Private/Lumin/LuminPlatformFile.cpp b/Engine/Source/Runtime/Core/Private/Lumin/LuminPlatformFile.cpp
--- a/Engine/Source/Runtime/Core/Private/Lumin/LuminPlatformFile.cpp
+++ b/Engine/Source/Runtime/Core/Private/Lumin/LuminPlatformFile.cpp
@@ -648,11 +648,14 @@ bool FLuminPlatformFile::FileExists(const TCHAR* Filename, FString& OutLuminPath
 
 int64 FLuminPlatformFile::FileSize(const TCHAR* Filename)
 {
-	FString NormalizedFilename = NormalizeFilename(Filename);
-	// Check the read path first.
-	int64 result = FileSizeCaseInsensitive(ConvertToLuminPath(NormalizedFilename, false));
-	// If it doesnt exist, check for the write path instead.
-	return (result != -1) ? result : FileSizeCaseInsensitive(ConvertToLuminPath(NormalizedFilename, true));
+	// Checking that the file exists will also give us the true location of the file.
+	// Which can be either in the read-only or the read-write areas of the application.
+	FString LuminPath;
+	if (FileExists(Filename, LuminPath))
+	{
+		return FileSizeCaseInsensitive(LuminPath);
+	}
+	return -1;
 }
 
 bool FLuminPlatformFile::DeleteFile(const TCHAR* Filename)
@@ -676,17 +679,21 @@ bool FLuminPlatformFile::DeleteFile(const TCHAR* Filename)
 
 bool FLuminPlatformFile::IsReadOnly(const TCHAR* Filename)
 {
-	FString NormalizedFilename = NormalizeFilename(Filename);
-
-	// Check the read path first, if it doesnt exist, check for the write path instead.
-	return (IsReadOnlyCaseInsensitive(ConvertToLuminPath(NormalizedFilename, false))) ? true : IsReadOnlyCaseInsensitive(ConvertToLuminPath(NormalizedFilename, true));
+	// Checking that the file exists will also give us the true location of the file.
+	// Which can be either in the read-only or the read-write areas of the application.
+	FString LuminPath;
+	if (FileExists(Filename, LuminPath))
+	{
+		return IsReadOnlyCaseInsensitive(LuminPath);
+	}
+	return false;
 }
 
 bool FLuminPlatformFile::MoveFile(const TCHAR* To, const TCHAR* From)
 {
 	// Move to write path only.
 	FString ToLuminFilename = ConvertToLuminPath(NormalizeFilename(To), true);
-	FString FromLuminFilename = ConvertToLuminPath(NormalizeFilename(From), false);
+	FString FromLuminFilename = ConvertToLuminPath(NormalizeFilename(From), true);
 	FString CaseSensitiveFilename;
 	if (!GCaseInsensMapper.MapCaseInsensitiveFile(FromLuminFilename, CaseSensitiveFilename))
 	{
@@ -694,7 +701,7 @@ bool FLuminPlatformFile::MoveFile(const TCHAR* To, const TCHAR* From)
 		return false;
 	}
 
-	return rename(TCHAR_TO_UTF8(*CaseSensitiveFilename), TCHAR_TO_UTF8(*ToLuminFilename)) != -1;
+	return rename(TCHAR_TO_UTF8(*CaseSensitiveFilename), TCHAR_TO_UTF8(*ToLuminFilename)) == 0;
 }
 
 bool FLuminPlatformFile::SetReadOnly(const TCHAR* Filename, bool bNewReadOnlyValue)
@@ -708,7 +715,7 @@ bool FLuminPlatformFile::SetReadOnly(const TCHAR* Filename, bool bNewReadOnlyVal
 	}
 
 	struct stat FileInfo;
-	if (stat(TCHAR_TO_UTF8(*CaseSensitiveFilename), &FileInfo) != -1)
+	if (stat(TCHAR_TO_UTF8(*CaseSensitiveFilename), &FileInfo) == 0)
 	{
 		if (bNewReadOnlyValue)
 		{
@@ -718,18 +725,21 @@ bool FLuminPlatformFile::SetReadOnly(const TCHAR* Filename, bool bNewReadOnlyVal
 		{
 			FileInfo.st_mode |= S_IWUSR;
 		}
-		return chmod(TCHAR_TO_UTF8(*CaseSensitiveFilename), FileInfo.st_mode);
+		return chmod(TCHAR_TO_UTF8(*CaseSensitiveFilename), FileInfo.st_mode) == 0;
 	}
 	return false;
 }
 
 FDateTime FLuminPlatformFile::GetTimeStamp(const TCHAR* Filename)
 {
-	FString NormalizedFilename = NormalizeFilename(Filename);
-
-	// Check the read path first, if it doesnt exist, check for the write path instead.
-	FDateTime result = GetTimeStampCaseInsensitive(ConvertToLuminPath(NormalizedFilename, false));
-	return (result == FDateTime::MinValue()) ? GetTimeStampCaseInsensitive(ConvertToLuminPath(NormalizedFilename, true)) : result;
+	// Checking that the file exists will also give us the true location of the file.
+	// Which can be either in the read-only or the read-write areas of the application.
+	FString LuminPath;
+	if (FileExists(Filename, LuminPath))
+	{
+		return GetTimeStampCaseInsensitive(LuminPath);
+	}
+	return FDateTime::MinValue();
 }
 
 void FLuminPlatformFile::SetTimeStamp(const TCHAR* Filename, const FDateTime DateTime)
@@ -744,7 +754,7 @@ void FLuminPlatformFile::SetTimeStamp(const TCHAR* Filename, const FDateTime Dat
 
 	// get file times
 	struct stat FileInfo;
-	if(stat(TCHAR_TO_UTF8(*CaseSensitiveFilename), &FileInfo) == -1)
+	if(stat(TCHAR_TO_UTF8(*CaseSensitiveFilename), &FileInfo) != 0)
 	{
 		return;
 	}
@@ -758,11 +768,14 @@ void FLuminPlatformFile::SetTimeStamp(const TCHAR* Filename, const FDateTime Dat
 
 FDateTime FLuminPlatformFile::GetAccessTimeStamp(const TCHAR* Filename)
 {
-	FString NormalizedFilename = NormalizeFilename(Filename);
-
-	// Check the read path first, if it doesnt exist, check for the write path instead.
-	FDateTime result = GetAccessTimeStampCaseInsensitive(ConvertToLuminPath(NormalizedFilename, false));
-	return (result == FDateTime::MinValue()) ? GetAccessTimeStampCaseInsensitive(ConvertToLuminPath(NormalizedFilename, true)) : result;
+	// Checking that the file exists will also give us the true location of the file.
+	// Which can be either in the read-only or the read-write areas of the application.
+	FString LuminPath;
+	if (FileExists(Filename, LuminPath))
+	{
+		return GetAccessTimeStampCaseInsensitive(LuminPath);
+	}
+	return FDateTime::MinValue();
 }
 
 FString FLuminPlatformFile::GetFilenameOnDisk(const TCHAR* Filename)
diff --git a/Engine/Source/Runtime/Core/Private/Lumin/LuminPlatformMisc.cpp b/Engine/Source/Runtime/Core/Private/Lumin/LuminPlatformMisc.cpp
--- a/Engine/Source/Runtime/Core/Private/Lumin/LuminPlatformMisc.cpp
+++ b/Engine/Source/Runtime/Core/Private/Lumin/LuminPlatformMisc.cpp
@@ -35,6 +35,19 @@ void FLuminPlatformMisc::InitLifecycle()
 	FLuminLifecycle::Initialize();
 }
 
+void FLuminPlatformMisc::RequestExit(bool Force)
+{
+	UE_LOG(LogWindows, Log, TEXT("FLuminPlatformMisc::RequestExit(%i)"), Force);
+	if (Force)
+	{
+		_exit(1);
+	}
+	else
+	{
+		GIsRequestingExit = 1;
+	}
+}
+
 void FLuminPlatformMisc::PlatformPreInit()
 {
 	// base class version
@@ -281,8 +294,9 @@ void FLuminPlatformMisc::InitApplicationPaths()
 		}
 	}
 
-	MLLifecycleSelfInfo* SelfInfo = MLLifecycleGetSelfInfo();
-	if (SelfInfo == nullptr)
+	MLLifecycleSelfInfo* SelfInfo = nullptr;
+	MLResult Result = MLLifecycleGetSelfInfo(&SelfInfo);
+	if (Result != MLResult_Ok || SelfInfo == nullptr)
 	{
 		FLuminPlatformMisc::LowLevelOutputDebugStringWithVerbosity(TEXT("Could not get self info for the application. The application paths will be incorrect"), ELogVerbosity::Error);
 		return;
diff --git a/Engine/Source/Runtime/Core/Private/Lumin/LuminPlatformProcess.cpp b/Engine/Source/Runtime/Core/Private/Lumin/LuminPlatformProcess.cpp
--- a/Engine/Source/Runtime/Core/Private/Lumin/LuminPlatformProcess.cpp
+++ b/Engine/Source/Runtime/Core/Private/Lumin/LuminPlatformProcess.cpp
@@ -80,40 +80,32 @@ void FLuminPlatformProcess::LaunchURL(const TCHAR* URL, const TCHAR* Parms, FStr
 	check(URL);
 	const FString URLWithParams = FString::Printf(TEXT("%s %s"), URL, Parms ? Parms : TEXT("")).TrimEnd();
 
-	MLDispatchPacket* Packet = MLDispatchAllocateEmptyPacket();
+	MLDispatchPacket* Packet = nullptr;
+	MLResult Result = MLDispatchAllocateEmptyPacket(&Packet);
 	if (Packet != nullptr)
 	{
-		bool bSetUri = MLDispatchSetUri(Packet, TCHAR_TO_UTF8(*URLWithParams));
-		if (bSetUri)
+		Result = MLDispatchSetUri(Packet, TCHAR_TO_UTF8(*URLWithParams));
+		if (Result == MLResult_Ok)
 		{
-			MLDispatchErrorCode ErrorCode = MLDispatchTryOpenApplication(Packet);
-			switch(ErrorCode)
+			Result = MLDispatchTryOpenApplication(Packet);
+			if (Result != MLResult_Ok)
 			{
-				case MLDispatchErrorCode_CannotStartApp:
-					UE_LOG(LogCore, Error, TEXT("Failed to launch URL %s - Cannot start app"), *URLWithParams);
-					break;
-				case MLDispatchErrorCode_InvalidPacket:
-					UE_LOG(LogCore, Error, TEXT("Failed to launch URL %s - Invalid packet"), *URLWithParams);
-					break;
-				case MLDispatchErrorCode_NoAppFound:
-					UE_LOG(LogCore, Error, TEXT("Failed to launch URL %s - No app found"), *URLWithParams);
-					break;
-				case MLDispatchErrorCode_Internal:
-					UE_LOG(LogCore, Error, TEXT("Failed to launch URL %s - Internal error"), *URLWithParams);
-					break;
-				default:
-					UE_LOG(LogCore, Display, TEXT("Launching URL %s"), *URLWithParams);					
+				UE_LOG(LogCore, Error, TEXT("Failed to launch URL %s: MLDispatchTryOpenApplication failure: %d"), *URLWithParams, static_cast<int32>(Result));
 			}
 		}
 		else
 		{
-			UE_LOG(LogCore, Error, TEXT("Failed to launch URL %s - could not set URI in dispatch packet"), *URLWithParams);
+			UE_LOG(LogCore, Error, TEXT("Failed to launch URL %s: MLDispatchSetUri failure: %s"), *URLWithParams, static_cast<int32>(Result));
+		}
+		Result = MLDispatchReleasePacket(&Packet, true, false);
+		if (Result != MLResult_Ok)
+		{
+			UE_LOG(LogCore, Error, TEXT("MLDispatchReleasePacket failed: %s"), static_cast<int32>(Result));
 		}
-		MLDispatchReleasePacket(&Packet, true, false);
 	}
 	else
 	{
-		UE_LOG(LogCore, Error, TEXT("Failed to launch URL %s - could not allocate dispatch packet"), *URLWithParams);
+		UE_LOG(LogCore, Error, TEXT("Failed to launch URL %s: MLDispatchAllocateEmptyPacket failure: %s"), *URLWithParams, static_cast<int32>(Result));
 	}
 }
 
diff --git a/Engine/Source/Runtime/Core/Public/Lumin/LuminPlatformMisc.h b/Engine/Source/Runtime/Core/Public/Lumin/LuminPlatformMisc.h
--- a/Engine/Source/Runtime/Core/Public/Lumin/LuminPlatformMisc.h
+++ b/Engine/Source/Runtime/Core/Public/Lumin/LuminPlatformMisc.h
@@ -11,6 +11,7 @@ struct CORE_API FLuminPlatformMisc : public FAndroidMisc
 	/** 
 	 * Platform overrides 
 	 */
+	static void RequestExit(bool Force);
 	static void PlatformPreInit();
 	static bool AllowRenderThread();
 	static bool SupportsLocalCaching();
diff --git a/Engine/Source/Runtime/Launch/Private/Lumin/LaunchLumin.cpp b/Engine/Source/Runtime/Launch/Private/Lumin/LaunchLumin.cpp
--- a/Engine/Source/Runtime/Launch/Private/Lumin/LaunchLumin.cpp
+++ b/Engine/Source/Runtime/Launch/Private/Lumin/LaunchLumin.cpp
@@ -62,33 +62,6 @@ static void InitCommandLine()
 		FCommandLine::Append(UTF8_TO_TCHAR(CommandLine));
 	}
 
-	// To use lifecycle init args, launch the app using -
-	// mldb launch -i "-arg1=value1 -arg2=value2" <package_name>
-	MLLifecycleInitArgList* InitArgList = MLLifecycleGetInitArgList();
-	if (InitArgList != nullptr)
-	{
-		int64 InitArgCount = MLLifecycleGetInitArgListLength(InitArgList);
-		if (InitArgCount > 0)
-		{
-			for (int64 i = 0; i < InitArgCount; ++i)
-			{
-				const MLLifecycleInitArg* InitArg = MLLifecycleGetInitArgByIndex(InitArgList, i);
-				if (InitArg != nullptr)
-				{
-					const char* Arg = MLLifecycleGetInitArgUri(InitArg);
-					if (Arg != nullptr)
-					{
-						// Start with a space because the command line already in place does not have any trailing spaces.
-						FString ArgStr = TEXT(" ");
-						ArgStr.Append(UTF8_TO_TCHAR(Arg));
-						ArgStr.TrimEndInline();
-						FCommandLine::Append(*ArgStr);
-					}
-				}
-			}
-		}
-	}
-
 	// This prevents duplicate logs in logcat as indicated in -
 	// Engine/Source/Runtime/Core/Public/HAL/FeedbackContextAnsi.h
 	FCommandLine::Append(TEXT(" -stdout "));
diff --git a/Engine/Source/Runtime/Lumin/LuminRuntimeSettings/Classes/LuminRuntimeSettings.h b/Engine/Source/Runtime/Lumin/LuminRuntimeSettings/Classes/LuminRuntimeSettings.h
--- a/Engine/Source/Runtime/Lumin/LuminRuntimeSettings/Classes/LuminRuntimeSettings.h
+++ b/Engine/Source/Runtime/Lumin/LuminRuntimeSettings/Classes/LuminRuntimeSettings.h
@@ -87,18 +87,14 @@ public:
 	UPROPERTY(GlobalConfig, EditAnywhere, Category = "Advanced MPK Packaging", Meta = (DisplayName = "Version Code", ClampMin = 0))
 	int32 VersionCode;
 
+	/** Name of the minimum Lumin OS Version required */
+	UPROPERTY(GlobalConfig, EditAnywhere, Category = "Advanced MPK Packaging", Meta = (DisplayName = "Minimum OS Version Required"))
+	FString MinimumOSVersion;
+
 	/** Any privileges your app needs. */
 	UPROPERTY(GlobalConfig, EditAnywhere, Category = "Advanced MPK Packaging", Meta = (DisplayName = "App Privileges (e.g. 'DenseMap')"))
 	TArray<FString> AppPrivileges;
 
-	/** 
-	 * The Store will display to users if this app requires internet access in order to function properly.
-	 * This is different than the Internet app privilege which is autogranted to all installed apps.
-	 * This will let users know if the experience can be used offline. 
-	 */
-	UPROPERTY(GlobalConfig, EditAnywhere, Category = "Advanced MPK Packaging", Meta = (DisplayName = "Internet Required"))
-	bool bInternetRequired;
-
 	/** Extra nodes under the <application> node. */
 	UPROPERTY(GlobalConfig, EditAnywhere, Category = "Advanced MPK Packaging", Meta = (DisplayName = "Extra nodes under the <application> node"))
 	TArray<FString> ExtraApplicationNodes;
diff --git a/Engine/Source/Runtime/VulkanRHI/Private/VulkanUtil.cpp b/Engine/Source/Runtime/VulkanRHI/Private/VulkanUtil.cpp
--- a/Engine/Source/Runtime/VulkanRHI/Private/VulkanUtil.cpp
+++ b/Engine/Source/Runtime/VulkanRHI/Private/VulkanUtil.cpp
@@ -499,6 +499,11 @@ void FVulkanGPUProfiler::DumpCrashMarkers(void* BufferData)
 #include "VulkanRHIBridge.h"
 namespace VulkanRHIBridge
 {
+	uint64 GetInstance(FVulkanDynamicRHI* RHI)
+	{
+		return (uint64)RHI->GetInstance();
+	}
+
 	FVulkanDevice* GetDevice(FVulkanDynamicRHI* RHI)
 	{
 		return RHI->GetDevice();
diff --git a/Engine/Source/Runtime/VulkanRHI/Public/VulkanRHIBridge.h b/Engine/Source/Runtime/VulkanRHI/Public/VulkanRHIBridge.h
--- a/Engine/Source/Runtime/VulkanRHI/Public/VulkanRHIBridge.h
+++ b/Engine/Source/Runtime/VulkanRHI/Public/VulkanRHIBridge.h
@@ -11,6 +11,8 @@ class FVulkanDevice;
 
 namespace VulkanRHIBridge
 {
+	VULKANRHI_API uint64 GetInstance(FVulkanDynamicRHI*);
+
 	VULKANRHI_API FVulkanDevice* GetDevice(FVulkanDynamicRHI*);
 
 	// Returns a VkDevice
