Trigger a link error whenever a module is missing a valid IMPLEMENT_MODULE macro on Windows.

The IMPLEMENT_MODULE macro now defines an empty function called IMPLEMENT_MODULE_{ModuleName}, to which we force a reference to via the linker command line. This behavior can be disabled on a per-module basis by setting bRequiresImplementModule = false from a .build.cs file.

diff --git a/Engine/Plugins/AI/HTNPlanner/Source/HTNPlanner/Private/HTNPlannerModule.cpp b/Engine/Plugins/AI/HTNPlanner/Source/HTNPlanner/Private/HTNPlannerModule.cpp
--- a/Engine/Plugins/AI/HTNPlanner/Source/HTNPlanner/Private/HTNPlannerModule.cpp
+++ b/Engine/Plugins/AI/HTNPlanner/Source/HTNPlanner/Private/HTNPlannerModule.cpp
@@ -20,7 +20,7 @@ class FHTNPlannerModule : public IHTNPlannerModule
 	virtual void ShutdownModule() override;
 };
 
-IMPLEMENT_MODULE(FHTNPlannerModule, HTNPlannerModule)
+IMPLEMENT_MODULE(FHTNPlannerModule, HTNPlanner)
 
 void FHTNPlannerModule::StartupModule()
 {
diff --git a/Engine/Plugins/AI/HTNPlanner/Source/HTNTestSuite/Private/HTNTestSuiteModule.cpp b/Engine/Plugins/AI/HTNPlanner/Source/HTNTestSuite/Private/HTNTestSuiteModule.cpp
--- a/Engine/Plugins/AI/HTNPlanner/Source/HTNTestSuite/Private/HTNTestSuiteModule.cpp
+++ b/Engine/Plugins/AI/HTNPlanner/Source/HTNTestSuite/Private/HTNTestSuiteModule.cpp
@@ -8,6 +8,6 @@ class FHTNTestSuiteModule : public IHTNTestSuiteModule
 {
 };
 
-IMPLEMENT_MODULE(FHTNTestSuiteModule, HTNTestSuiteModule)
+IMPLEMENT_MODULE(FHTNTestSuiteModule, HTNTestSuite)
 
 #undef LOCTEXT_NAMESPACE
diff --git a/Engine/Plugins/Experimental/AppleVision/Source/AppleVisionBlueprintSupport/Private/AppleVisionBlueprintSupport.cpp b/Engine/Plugins/Experimental/AppleVision/Source/AppleVisionBlueprintSupport/Private/AppleVisionBlueprintSupport.cpp
--- a/Engine/Plugins/Experimental/AppleVision/Source/AppleVisionBlueprintSupport/Private/AppleVisionBlueprintSupport.cpp
+++ b/Engine/Plugins/Experimental/AppleVision/Source/AppleVisionBlueprintSupport/Private/AppleVisionBlueprintSupport.cpp
@@ -2,6 +2,9 @@
 
 #include "AppleVisionBlueprintSupport.h"
 #include "AppleVisionBlueprintProxy.h"
+#include "Modules/ModuleManager.h"
+
+IMPLEMENT_MODULE(FDefaultModuleImpl, AppleVisionBlueprintSupport)
 
 UK2Node_DetectFaces::UK2Node_DetectFaces(const FObjectInitializer& ObjectInitializer)
 	: Super(ObjectInitializer)
diff --git a/Engine/Plugins/Runtime/Analytics/Adjust/Source/AdjustEditor/Private/AdjustEditor.cpp b/Engine/Plugins/Runtime/Analytics/Adjust/Source/AdjustEditor/Private/AdjustEditor.cpp
--- a/Engine/Plugins/Runtime/Analytics/Adjust/Source/AdjustEditor/Private/AdjustEditor.cpp
+++ b/Engine/Plugins/Runtime/Analytics/Adjust/Source/AdjustEditor/Private/AdjustEditor.cpp
@@ -7,7 +7,7 @@
 #include "AdjustSettings.h"
 #include "Analytics.h"
 
-IMPLEMENT_MODULE( FAdjustEditorModule, AdjustEditorModule );
+IMPLEMENT_MODULE( FAdjustEditorModule, AdjustEditor );
 
 #define LOCTEXT_NAMESPACE "Adjust"
 
diff --git a/Engine/Plugins/Runtime/Analytics/Apsalar/Source/ApsalarEditor/Private/ApsalarEditor.cpp b/Engine/Plugins/Runtime/Analytics/Apsalar/Source/ApsalarEditor/Private/ApsalarEditor.cpp
--- a/Engine/Plugins/Runtime/Analytics/Apsalar/Source/ApsalarEditor/Private/ApsalarEditor.cpp
+++ b/Engine/Plugins/Runtime/Analytics/Apsalar/Source/ApsalarEditor/Private/ApsalarEditor.cpp
@@ -6,7 +6,7 @@
 #include "ApsalarSettings.h"
 #include "Analytics.h"
 
-IMPLEMENT_MODULE( FApsalarEditorModule, ApsalarEditorModule );
+IMPLEMENT_MODULE( FApsalarEditorModule, ApsalarEditor );
 
 #define LOCTEXT_NAMESPACE "Apsalar"
 
diff --git a/Engine/Plugins/Runtime/Analytics/Flurry/Source/FlurryEditor/Private/FlurryEditor.cpp b/Engine/Plugins/Runtime/Analytics/Flurry/Source/FlurryEditor/Private/FlurryEditor.cpp
--- a/Engine/Plugins/Runtime/Analytics/Flurry/Source/FlurryEditor/Private/FlurryEditor.cpp
+++ b/Engine/Plugins/Runtime/Analytics/Flurry/Source/FlurryEditor/Private/FlurryEditor.cpp
@@ -6,7 +6,7 @@
 #include "FlurrySettings.h"
 #include "Analytics.h"
 
-IMPLEMENT_MODULE( FFlurryEditorModule, FlurryEditorModule );
+IMPLEMENT_MODULE( FFlurryEditorModule, FlurryEditor );
 
 #define LOCTEXT_NAMESPACE "Flurry"
 
diff --git a/Engine/Plugins/Runtime/LocationServicesAndroidImpl/Source/LocationServicesAndroidEditor/Private/LocationServicesAndroidEditor.cpp b/Engine/Plugins/Runtime/LocationServicesAndroidImpl/Source/LocationServicesAndroidEditor/Private/LocationServicesAndroidEditor.cpp
--- a/Engine/Plugins/Runtime/LocationServicesAndroidImpl/Source/LocationServicesAndroidEditor/Private/LocationServicesAndroidEditor.cpp
+++ b/Engine/Plugins/Runtime/LocationServicesAndroidImpl/Source/LocationServicesAndroidEditor/Private/LocationServicesAndroidEditor.cpp
@@ -12,7 +12,7 @@
 #endif
 
 
-IMPLEMENT_MODULE( FLocationServicesAndroidEditorModule, LocationServicesAndroidEditorModule );
+IMPLEMENT_MODULE( FLocationServicesAndroidEditorModule, LocationServicesAndroidEditor );
 
 #define LOCTEXT_NAMESPACE "FLocationServicesAndroidEditorModule"
 
diff --git a/Engine/Plugins/Runtime/LocationServicesIOSImpl/Source/LocationServicesIOSEditor/Private/LocationServicesIOSEditor.cpp b/Engine/Plugins/Runtime/LocationServicesIOSImpl/Source/LocationServicesIOSEditor/Private/LocationServicesIOSEditor.cpp
--- a/Engine/Plugins/Runtime/LocationServicesIOSImpl/Source/LocationServicesIOSEditor/Private/LocationServicesIOSEditor.cpp
+++ b/Engine/Plugins/Runtime/LocationServicesIOSImpl/Source/LocationServicesIOSEditor/Private/LocationServicesIOSEditor.cpp
@@ -12,7 +12,7 @@
 #endif
 
 
-IMPLEMENT_MODULE( FLocationServicesIOSEditorModule, LocationServicesIOSEditorModule );
+IMPLEMENT_MODULE( FLocationServicesIOSEditorModule, LocationServicesIOSEditor );
 
 #define LOCTEXT_NAMESPACE "FLocationServicesIOSEditorModule"
 
diff --git a/Engine/Plugins/Runtime/Spatialization/Source/SpatializationEditor/Private/SpatializationEditorModule.cpp b/Engine/Plugins/Runtime/Spatialization/Source/SpatializationEditor/Private/SpatializationEditorModule.cpp
--- a/Engine/Plugins/Runtime/Spatialization/Source/SpatializationEditor/Private/SpatializationEditorModule.cpp
+++ b/Engine/Plugins/Runtime/Spatialization/Source/SpatializationEditor/Private/SpatializationEditorModule.cpp
@@ -13,4 +13,4 @@ void FSpatializationEditorModule::ShutdownModule()
 	// Nothing done here for now.
 }
 
-IMPLEMENT_MODULE(FSpatializationEditorModule, NativeSpatializationEditor)
\ No newline at end of file
+IMPLEMENT_MODULE(FSpatializationEditorModule, SpatializationEditor)
\ No newline at end of file
diff --git a/Engine/Plugins/Runtime/Synthesis/Source/SynthesisEditor/Private/SynthesisEditorModule.cpp b/Engine/Plugins/Runtime/Synthesis/Source/SynthesisEditor/Private/SynthesisEditorModule.cpp
--- a/Engine/Plugins/Runtime/Synthesis/Source/SynthesisEditor/Private/SynthesisEditorModule.cpp
+++ b/Engine/Plugins/Runtime/Synthesis/Source/SynthesisEditor/Private/SynthesisEditorModule.cpp
@@ -12,7 +12,7 @@
 #include "AudioEditorModule.h"
 #include "EpicSynth1PresetBank.h"
 
-IMPLEMENT_MODULE(FSynthesisEditorModule, FSynthesisEditor)
+IMPLEMENT_MODULE(FSynthesisEditorModule, SynthesisEditor)
 
 void FSynthesisEditorModule::StartupModule()
 {
diff --git a/Engine/Plugins/Runtime/TimeSynth/Source/TimeSynthEditor/Private/TimeSynthEditorModule.cpp b/Engine/Plugins/Runtime/TimeSynth/Source/TimeSynthEditor/Private/TimeSynthEditorModule.cpp
--- a/Engine/Plugins/Runtime/TimeSynth/Source/TimeSynthEditor/Private/TimeSynthEditorModule.cpp
+++ b/Engine/Plugins/Runtime/TimeSynth/Source/TimeSynthEditor/Private/TimeSynthEditorModule.cpp
@@ -12,7 +12,7 @@
 #include "TimeSynthVolumeGroup.h"
 #include "AudioEditorModule.h"
 
-IMPLEMENT_MODULE(FTimeSynthEditorModule, FTimeSynthEditor)
+IMPLEMENT_MODULE(FTimeSynthEditorModule, TimeSynthEditor)
 
 void FTimeSynthEditorModule::StartupModule()
 {
diff --git a/Engine/Source/Developer/MeshDescriptionOperations/Private/MeshDescriptionOperations.cpp b/Engine/Source/Developer/MeshDescriptionOperations/Private/MeshDescriptionOperations.cpp
--- a/Engine/Source/Developer/MeshDescriptionOperations/Private/MeshDescriptionOperations.cpp
+++ b/Engine/Source/Developer/MeshDescriptionOperations/Private/MeshDescriptionOperations.cpp
@@ -10,6 +10,9 @@
 #include "RenderUtils.h"
 #include "mikktspace.h"
 #include "UVMapSettings.h"
+#include "Modules/ModuleManager.h"
+
+IMPLEMENT_MODULE(FDefaultModuleImpl, MeshDescriptionOperations)
 
 DEFINE_LOG_CATEGORY(LogMeshDescriptionOperations);
 
diff --git a/Engine/Source/Developer/MeshUtilitiesCommon/Private/LayoutUV.cpp b/Engine/Source/Developer/MeshUtilitiesCommon/Private/LayoutUV.cpp
--- a/Engine/Source/Developer/MeshUtilitiesCommon/Private/LayoutUV.cpp
+++ b/Engine/Source/Developer/MeshUtilitiesCommon/Private/LayoutUV.cpp
@@ -5,6 +5,9 @@
 #include "OverlappingCorners.h"
 #include "Algo/IntroSort.h"
 #include "HAL/PlatformTime.h"
+#include "Modules/ModuleManager.h"
+
+IMPLEMENT_MODULE(FDefaultModuleImpl, MeshUtilitiesCommon)
 
 DEFINE_LOG_CATEGORY_STATIC(LogLayoutUV, Warning, All);
 
diff --git a/Engine/Source/Developer/PakFileUtilities/Private/PakFileUtilities.cpp b/Engine/Source/Developer/PakFileUtilities/Private/PakFileUtilities.cpp
--- a/Engine/Source/Developer/PakFileUtilities/Private/PakFileUtilities.cpp
+++ b/Engine/Source/Developer/PakFileUtilities/Private/PakFileUtilities.cpp
@@ -19,7 +19,9 @@
 #include "HAL/PlatformFilemanager.h"
 #include "Async/ParallelFor.h"
 #include "Async/AsyncWork.h"
+#include "Modules/ModuleManager.h"
 
+IMPLEMENT_MODULE(FDefaultModuleImpl, PakFileUtilities);
 
 /**
  * Encryption keys: public and private
diff --git a/Engine/Source/Developer/StandaloneRenderer/Private/StandaloneRenderer.cpp b/Engine/Source/Developer/StandaloneRenderer/Private/StandaloneRenderer.cpp
--- a/Engine/Source/Developer/StandaloneRenderer/Private/StandaloneRenderer.cpp
+++ b/Engine/Source/Developer/StandaloneRenderer/Private/StandaloneRenderer.cpp
@@ -45,4 +45,4 @@ class FStandaloneRenderer : public IModuleInterface
 {
 };
 
-IMPLEMENT_MODULE( FStandaloneRenderer, BlankPlugin )
+IMPLEMENT_MODULE( FStandaloneRenderer, StandaloneRenderer )
diff --git a/Engine/Source/Editor/TimeManagementEditor/Private/TimeManagementEditorModule.cpp b/Engine/Source/Editor/TimeManagementEditor/Private/TimeManagementEditorModule.cpp
--- a/Engine/Source/Editor/TimeManagementEditor/Private/TimeManagementEditorModule.cpp
+++ b/Engine/Source/Editor/TimeManagementEditor/Private/TimeManagementEditorModule.cpp
@@ -24,4 +24,4 @@ public:
 	}
 };
 
-IMPLEMENT_MODULE(FTimeManagementEditorModule, TimeManagementEditorModule);
\ No newline at end of file
+IMPLEMENT_MODULE(FTimeManagementEditorModule, TimeManagementEditor);
\ No newline at end of file
diff --git a/Engine/Source/Programs/UnrealBuildTool/Configuration/ModuleRules.cs b/Engine/Source/Programs/UnrealBuildTool/Configuration/ModuleRules.cs
--- a/Engine/Source/Programs/UnrealBuildTool/Configuration/ModuleRules.cs
+++ b/Engine/Source/Programs/UnrealBuildTool/Configuration/ModuleRules.cs
@@ -687,6 +687,12 @@ namespace UnrealBuildTool
 		/// </summary>
 		public List<string> ExternalDependencies = new List<string>();
 
+		/// <summary>
+		/// Whether this module requires the IMPLEMENT_MODULE macro to be implemented. Most UE4 modules require this, since we use the IMPLEMENT_MODULE macro
+		/// to do other global overloads (eg. operator new/delete forwarding to GMalloc).
+		/// </summary>
+		public bool? bRequiresImplementModule;
+
 		/// <summary>
 		/// Whether this module qualifies included headers from other modules relative to the root of their 'Public' folder. This reduces the number
 		/// of search paths that have to be passed to the compiler, improving performance and reducing the length of the compiler command line.
diff --git a/Engine/Source/Programs/UnrealBuildTool/Configuration/UEBuildBinary.cs b/Engine/Source/Programs/UnrealBuildTool/Configuration/UEBuildBinary.cs
--- a/Engine/Source/Programs/UnrealBuildTool/Configuration/UEBuildBinary.cs
+++ b/Engine/Source/Programs/UnrealBuildTool/Configuration/UEBuildBinary.cs
@@ -692,6 +692,12 @@ namespace UnrealBuildTool
 							BinaryLinkEnvironment.InputFiles.Add(LinkInputFile);
 						}
 					}
+
+					// Force a reference to initialize module for this binary
+					if(Module.Rules.bRequiresImplementModule.Value)
+					{
+						BinaryLinkEnvironment.IncludeFunctions.Add(String.Format("IMPLEMENT_MODULE_{0}", Module.Name));
+					}
 				}
 				else
 				{
diff --git a/Engine/Source/Programs/UnrealBuildTool/Configuration/UEBuildTarget.cs b/Engine/Source/Programs/UnrealBuildTool/Configuration/UEBuildTarget.cs
--- a/Engine/Source/Programs/UnrealBuildTool/Configuration/UEBuildTarget.cs
+++ b/Engine/Source/Programs/UnrealBuildTool/Configuration/UEBuildTarget.cs
@@ -3150,6 +3150,12 @@ namespace UnrealBuildTool
 			// Create the rules from the assembly
 			ModuleRules RulesObject = RulesAssembly.CreateModuleRules(ModuleName, Rules, ReferenceChain);
 
+			// Set whether the module requires an IMPLEMENT_MODULE macro
+			if(!RulesObject.bRequiresImplementModule.HasValue)
+			{
+				RulesObject.bRequiresImplementModule = (RulesObject.Type == ModuleRules.ModuleType.CPlusPlus && RulesObject.Name != Rules.LaunchModuleName);
+			}
+
 			// Reads additional dependencies array for project module from project file and fills PrivateDependencyModuleNames. 
 			if (ProjectDescriptor != null && ProjectDescriptor.Modules != null)
 			{
diff --git a/Engine/Source/Programs/UnrealBuildTool/Platform/Windows/VCToolChain.cs b/Engine/Source/Programs/UnrealBuildTool/Platform/Windows/VCToolChain.cs
--- a/Engine/Source/Programs/UnrealBuildTool/Platform/Windows/VCToolChain.cs
+++ b/Engine/Source/Programs/UnrealBuildTool/Platform/Windows/VCToolChain.cs
@@ -1502,6 +1502,19 @@ namespace UnrealBuildTool
 				}
 			}
 
+			// Add any forced references to functions
+			foreach(string IncludeFunction in LinkEnvironment.IncludeFunctions)
+			{
+				if(CppPlatform == CppPlatform.Win32)
+				{
+					Arguments.Add(String.Format("/INCLUDE:_{0}", IncludeFunction)); // Assume decorated cdecl name
+				}
+				else
+				{
+					Arguments.Add(String.Format("/INCLUDE:{0}", IncludeFunction));
+				}
+			}
+
 			// Create a response file for the linker, unless we're generating IntelliSense data
 			FileReference ResponseFileName = GetResponseFileName(LinkEnvironment, OutputFile);
 			if (!ProjectFileGenerator.bGenerateProjectFiles)
diff --git a/Engine/Source/Programs/UnrealBuildTool/System/LinkEnvironment.cs b/Engine/Source/Programs/UnrealBuildTool/System/LinkEnvironment.cs
--- a/Engine/Source/Programs/UnrealBuildTool/System/LinkEnvironment.cs
+++ b/Engine/Source/Programs/UnrealBuildTool/System/LinkEnvironment.cs
@@ -274,6 +274,11 @@ namespace UnrealBuildTool
 		/// </summary>
 		public List<FileItem> CommonResourceFiles = new List<FileItem>();
 
+		/// <summary>
+		/// List of functions that should be exported from this module
+		/// </summary>
+		public List<string> IncludeFunctions = new List<string>();
+
 		/// <summary>
 		/// Provides a Module Definition File (.def) to the linker to describe various attributes of a DLL.
 		/// Necessary when exporting functions by ordinal values instead of by name.
@@ -348,6 +353,7 @@ namespace UnrealBuildTool
 			InputLibraries.AddRange(Other.InputLibraries);
 			DefaultResourceFiles.AddRange(Other.DefaultResourceFiles);
 			CommonResourceFiles.AddRange(Other.CommonResourceFiles);
+			IncludeFunctions.AddRange(Other.IncludeFunctions);
 			ModuleDefinitionFile = Other.ModuleDefinitionFile;
 			AdditionalProperties.AddRange(Other.AdditionalProperties);
         }
diff --git a/Engine/Source/Runtime/BuildSettings/BuildSettings.Build.cs b/Engine/Source/Runtime/BuildSettings/BuildSettings.Build.cs
--- a/Engine/Source/Runtime/BuildSettings/BuildSettings.Build.cs
+++ b/Engine/Source/Runtime/BuildSettings/BuildSettings.Build.cs
@@ -9,6 +9,8 @@ public class BuildSettings : ModuleRules
 	{
 		PrivateIncludePathModuleNames.Add("Core");
 
+		bRequiresImplementModule = false;
+
 		PrivateDefinitions.Add(string.Format("ENGINE_VERSION_MAJOR={0}", Target.Version.MajorVersion));
 		PrivateDefinitions.Add(string.Format("ENGINE_VERSION_MINOR={0}", Target.Version.MinorVersion));
 		PrivateDefinitions.Add(string.Format("ENGINE_VERSION_HOTFIX={0}", Target.Version.PatchVersion));
diff --git a/Engine/Source/Runtime/Core/Public/Modules/ModuleManager.h b/Engine/Source/Runtime/Core/Public/Modules/ModuleManager.h
--- a/Engine/Source/Runtime/Core/Public/Modules/ModuleManager.h
+++ b/Engine/Source/Runtime/Core/Public/Modules/ModuleManager.h
@@ -699,9 +699,8 @@ class FDefaultGameModuleImpl
 	#define IMPLEMENT_MODULE( ModuleImplClass, ModuleName ) \
 		/** Global registrant object for this module when linked statically */ \
 		static FStaticallyLinkedModuleRegistrant< ModuleImplClass > ModuleRegistrant##ModuleName( #ModuleName ); \
-		/** Implement an empty function so that if this module is built as a statically linked lib, */ \
-		/** static initialization for this lib can be forced by referencing this symbol */ \
-		void EmptyLinkFunctionForStaticInitialization##ModuleName(){} \
+		/* Forced reference to this function is added by the linker to check that each module uses IMPLEMENT_MODULE */ \
+		extern "C" void IMPLEMENT_MODULE_##ModuleName() { } \
 		PER_MODULE_BOILERPLATE_ANYLINK(ModuleImplClass, ModuleName)
 
 #else
@@ -717,6 +716,8 @@ class FDefaultGameModuleImpl
 		{ \
 			return new ModuleImplClass(); \
 		} \
+		/* Forced reference to this function is added by the linker to check that each module uses IMPLEMENT_MODULE */ \
+		extern "C" void IMPLEMENT_MODULE_##ModuleName() { } \
 		PER_MODULE_BOILERPLATE \
 		PER_MODULE_BOILERPLATE_ANYLINK(ModuleImplClass, ModuleName)
 
diff --git a/Engine/Source/Runtime/Experimental/Chaos/Private/ChaosModule.cpp b/Engine/Source/Runtime/Experimental/Chaos/Private/ChaosModule.cpp
--- a/Engine/Source/Runtime/Experimental/Chaos/Private/ChaosModule.cpp
+++ b/Engine/Source/Runtime/Experimental/Chaos/Private/ChaosModule.cpp
@@ -1,10 +1,11 @@
 // Copyright 1998-2019 Epic Games, Inc. All Rights Reserved.
 
-#if INCLUDE_CHAOS
-
 #include "ChaosModule.h"
 #include "CoreMinimal.h"
 #include "Modules/ModuleManager.h"
+
+#if INCLUDE_CHAOS
+
 #include "ChaosLog.h"
 
 DEFINE_LOG_CATEGORY(LogChaosGeneral);
@@ -14,4 +15,8 @@ DEFINE_LOG_CATEGORY(LogChaosDebug);
 
 IMPLEMENT_MODULE(FChaosEngineModule, Chaos);
 
+#else
+
+IMPLEMENT_MODULE(FDefaultModuleImpl, Chaos);
+
 #endif
diff --git a/Engine/Source/Runtime/Experimental/ChaosSolvers/Private/ChaosSolversModule.cpp b/Engine/Source/Runtime/Experimental/ChaosSolvers/Private/ChaosSolversModule.cpp
--- a/Engine/Source/Runtime/Experimental/ChaosSolvers/Private/ChaosSolversModule.cpp
+++ b/Engine/Source/Runtime/Experimental/ChaosSolvers/Private/ChaosSolversModule.cpp
@@ -2,6 +2,7 @@
 
 #include "ChaosSolversModule.h"
 #include "CoreMinimal.h"
+#include "Modules/ModuleManager.h"
 
 #if INCLUDE_CHAOS
 
@@ -308,4 +309,6 @@ IMPLEMENT_MODULE(FChaosSolversModule, ChaosSolvers);
 // Workaround for module not having any exported symbols
 CHAOSSOLVERS_API int ChaosSolversExportedSymbol = 0;
 
+IMPLEMENT_MODULE(FDefaultModuleImpl, ChaosSolvers);
+
 #endif
diff --git a/Engine/Source/Runtime/Experimental/GeometryCollectionCore/Private/GeometryCollection/GeometryCollectionPlugin.cpp b/Engine/Source/Runtime/Experimental/GeometryCollectionCore/Private/GeometryCollection/GeometryCollectionPlugin.cpp
--- a/Engine/Source/Runtime/Experimental/GeometryCollectionCore/Private/GeometryCollection/GeometryCollectionPlugin.cpp
+++ b/Engine/Source/Runtime/Experimental/GeometryCollectionCore/Private/GeometryCollection/GeometryCollectionPlugin.cpp
@@ -15,7 +15,7 @@ class FGeometryCollectionPlugin : public IGeometryCollectionPlugin
 	virtual void ShutdownModule() override;
 };
 
-IMPLEMENT_MODULE( FGeometryCollectionPlugin, GeometryCollection )
+IMPLEMENT_MODULE( FGeometryCollectionPlugin, GeometryCollectionCore )
 
 
 
diff --git a/Engine/Source/Runtime/Experimental/GeometryCollectionEngine/Private/GeometryCollection/GeometryCollectionComponentPlugin.cpp b/Engine/Source/Runtime/Experimental/GeometryCollectionEngine/Private/GeometryCollection/GeometryCollectionComponentPlugin.cpp
--- a/Engine/Source/Runtime/Experimental/GeometryCollectionEngine/Private/GeometryCollection/GeometryCollectionComponentPlugin.cpp
+++ b/Engine/Source/Runtime/Experimental/GeometryCollectionEngine/Private/GeometryCollection/GeometryCollectionComponentPlugin.cpp
@@ -15,7 +15,7 @@ class FGeometryCollectionComponentPlugin : public IGeometryCollectionComponentPl
 	virtual void ShutdownModule() override;
 };
 
-IMPLEMENT_MODULE( FGeometryCollectionComponentPlugin, GeometryCollectionComponent )
+IMPLEMENT_MODULE( FGeometryCollectionComponentPlugin, GeometryCollectionEngine )
 
 
 
