Integrate PART of CL 4378335.
Integrate the parts involving the FPlatformMisc::LoadTextFileFromPlatformPackage.

diff --git a/Engine/Source/Runtime/Core/Private/Android/AndroidMisc.cpp b/Engine/Source/Runtime/Core/Private/Android/AndroidMisc.cpp
--- a/Engine/Source/Runtime/Core/Private/Android/AndroidMisc.cpp
+++ b/Engine/Source/Runtime/Core/Private/Android/AndroidMisc.cpp
@@ -25,6 +25,7 @@
 #include "Misc/Parse.h"
 #include "Internationalization/Regex.h"
 
+#include <android/asset_manager.h>
 #include <android/log.h>
 #if USE_ANDROID_INPUT
 #include <android/keycodes.h>
@@ -41,6 +42,10 @@
 
 #include "FramePro/FrameProProfiler.h"
 
+#if USE_ANDROID_JNI
+extern AAssetManager * AndroidThunkCpp_GetAssetManager();
+#endif
+
 static int32 GAndroidTraceMarkersEnabled = 0;
 static FAutoConsoleVariableRef CAndroidTraceMarkersEnabled(
 	TEXT("android.tracemarkers"),
@@ -1065,6 +1070,28 @@ void FAndroidMisc::ShareURL(const FString& URL, const FText& Description, int32
 #endif
 }
 
+FString FAndroidMisc::LoadTextFileFromPlatformPackage(const FString& RelativePath)
+{
+	AAssetManager* AssetMgr = AndroidThunkCpp_GetAssetManager();
+	AAsset* asset = AAssetManager_open(AssetMgr, TCHAR_TO_UTF8(*RelativePath), AASSET_MODE_BUFFER);
+
+	if (asset)
+	{
+		const void* FileContents = (const ANSICHAR*)AAsset_getBuffer(asset);
+		int32 FileLength = AAsset_getLength(asset);
+
+		TArray<ANSICHAR> TextContents;
+		TextContents.AddUninitialized(FileLength + 1);
+		FMemory::Memcpy(TextContents.GetData(), FileContents, FileLength);
+		TextContents[FileLength] = 0;
+
+		AAsset_close(asset);
+
+		return FString(ANSI_TO_TCHAR(TextContents.GetData()));
+	}
+
+	return FString();
+}
 
 void FAndroidMisc::SetVersionInfo( FString InAndroidVersion, FString InDeviceMake, FString InDeviceModel, FString InDeviceBuildNumber, FString InOSLanguage )
 {
diff --git a/Engine/Source/Runtime/Core/Private/GenericPlatform/GenericPlatformMisc.cpp b/Engine/Source/Runtime/Core/Private/GenericPlatform/GenericPlatformMisc.cpp
--- a/Engine/Source/Runtime/Core/Private/GenericPlatform/GenericPlatformMisc.cpp
+++ b/Engine/Source/Runtime/Core/Private/GenericPlatform/GenericPlatformMisc.cpp
@@ -12,6 +12,7 @@
 #include "Misc/Parse.h"
 #include "Misc/CommandLine.h"
 #include "Misc/Paths.h"
+#include "Misc/FileHelper.h"
 #include "Internationalization/Text.h"
 #include "Internationalization/Internationalization.h"
 #include "Misc/Guid.h"
@@ -1284,4 +1285,12 @@ bool FGenericPlatformMisc::RequestDeviceCheckToken(TFunction<void(const TArray<u
 TArray<FChunkTagID> FGenericPlatformMisc::GetOnDemandChunkTagIDs()
 {
 	return TArray<FChunkTagID>();
-}
\ No newline at end of file
+}
+
+FString FGenericPlatformMisc::LoadTextFileFromPlatformPackage(const FString& RelativePath)
+{
+	FString Path = RootDir() / RelativePath;
+	FString Result;
+	FFileHelper::LoadFileToString(Result, *Path);
+	return Result;
+}
diff --git a/Engine/Source/Runtime/Core/Private/IOS/IOSPlatformMisc.cpp b/Engine/Source/Runtime/Core/Private/IOS/IOSPlatformMisc.cpp
--- a/Engine/Source/Runtime/Core/Private/IOS/IOSPlatformMisc.cpp
+++ b/Engine/Source/Runtime/Core/Private/IOS/IOSPlatformMisc.cpp
@@ -18,6 +18,7 @@
 #include "Misc/CommandLine.h"
 #include "Misc/ConfigCacheIni.h"
 #include "Misc/CoreDelegates.h"
+#include "Misc/ScopeExit.h"
 #include "Apple/ApplePlatformCrashContext.h"
 #include "IOS/IOSPlatformCrashContext.h"
 #if !PLATFORM_TVOS
@@ -1001,6 +1002,63 @@ void FIOSPlatformMisc::ShareURL(const FString& URL, const FText& Description, in
 }
 
 
+FString FIOSPlatformMisc::LoadTextFileFromPlatformPackage(const FString& RelativePath)
+{
+	FString FilePath = FString([[NSBundle mainBundle] bundlePath]) / RelativePath;
+
+	// read in the command line text file (coming from UnrealFrontend) if it exists
+	int32 File = open(TCHAR_TO_UTF8(*FilePath), O_RDONLY);
+	if (File == -1)
+	{
+		LowLevelOutputDebugStringf(TEXT("No file found at %s") LINE_TERMINATOR, *FilePath);
+		return FString();
+	}
+
+	ON_SCOPE_EXIT
+	{
+		close(File);
+	};
+
+	struct stat FileInfo;
+	FileInfo.st_size = -1;
+
+	if (fstat(File, &FileInfo))
+	{
+		LowLevelOutputDebugStringf(TEXT("Failed to determine file size of %s") LINE_TERMINATOR, *FilePath);
+		return FString();
+	}
+
+	if (FileInfo.st_size > MAX_int32 - 1)
+	{
+		LowLevelOutputDebugStringf(TEXT("File too big %s") LINE_TERMINATOR, *FilePath);
+		return FString();
+	}
+
+	LowLevelOutputDebugStringf(TEXT("Found %s file") LINE_TERMINATOR, *RelativePath);
+
+	int32 FileSize = static_cast<int32>(FileInfo.st_size);
+	TArray<char> FileContents;
+	FileContents.AddUninitialized(FileSize + 1);
+	FileContents[FileSize] = 0;
+
+	int32 NumRead = read(File, FileContents.GetData(), FileSize);
+	if (NumRead != FileSize)
+	{
+		LowLevelOutputDebugStringf(TEXT("Failed to read %s") LINE_TERMINATOR, *FilePath);
+		return FString();
+	}
+
+	// chop off trailing spaces
+	int32 Last = FileSize - 1;
+	while (FileContents[0] && isspace(FileContents[Last]))
+	{
+		FileContents[Last] = 0;
+		--Last;
+	}
+
+	return FString(UTF8_TO_TCHAR(FileContents.GetData()));
+}
+
 void FIOSPlatformMisc::EnableVoiceChat(bool bEnable)
 {
 	return [[IOSAppDelegate GetDelegate] EnableVoiceChat:bEnable];
diff --git a/Engine/Source/Runtime/Core/Public/Android/AndroidMisc.h b/Engine/Source/Runtime/Core/Public/Android/AndroidMisc.h
--- a/Engine/Source/Runtime/Core/Public/Android/AndroidMisc.h
+++ b/Engine/Source/Runtime/Core/Public/Android/AndroidMisc.h
@@ -158,6 +158,8 @@ public:
 	static void ReleaseMobileHaptics();
 	static void ShareURL(const FString& URL, const FText& Description, int32 LocationHintX, int32 LocationHintY);
 
+	static FString LoadTextFileFromPlatformPackage(const FString& RelativePath);
+
 	// ANDROID ONLY:
 	static void SetVersionInfo(FString AndroidVersion, FString DeviceMake, FString DeviceModel, FString DeviceBuildNumber, FString OSLanguage);
 	static const FString GetAndroidVersion();
diff --git a/Engine/Source/Runtime/Core/Public/GenericPlatform/GenericPlatformMisc.h b/Engine/Source/Runtime/Core/Public/GenericPlatform/GenericPlatformMisc.h
--- a/Engine/Source/Runtime/Core/Public/GenericPlatform/GenericPlatformMisc.h
+++ b/Engine/Source/Runtime/Core/Public/GenericPlatform/GenericPlatformMisc.h
@@ -1257,7 +1257,13 @@ public:
 	static bool RequestDeviceCheckToken(TFunction<void(const TArray<uint8>&)> QuerySucceededFunc, TFunction<void(const FString&, const FString&)> QueryFailedFunc);
 
 	static TArray<FChunkTagID> GetOnDemandChunkTagIDs();
-	
+
+	/*
+	 * Loads a text file relative to the package root on platforms that distribute apps in package formats.
+	 * For other platforms, the path is relative to the root directory.
+	*/
+	static FString LoadTextFileFromPlatformPackage(const FString& RelativePath);
+
 #if !UE_BUILD_SHIPPING
 	/** 
 	 * Returns any platform specific warning messages we want printed on screen
diff --git a/Engine/Source/Runtime/Core/Public/IOS/IOSPlatformMisc.h b/Engine/Source/Runtime/Core/Public/IOS/IOSPlatformMisc.h
--- a/Engine/Source/Runtime/Core/Public/IOS/IOSPlatformMisc.h
+++ b/Engine/Source/Runtime/Core/Public/IOS/IOSPlatformMisc.h
@@ -75,7 +75,9 @@ struct CORE_API FIOSPlatformMisc : public FApplePlatformMisc
 	static void ReleaseMobileHaptics();
 
 	static void ShareURL(const FString& URL, const FText& Description, int32 LocationHintX, int32 LocationHintY);
-    
+
+	static FString LoadTextFileFromPlatformPackage(const FString& RelativePath);
+
 	static void EnableVoiceChat(bool bEnable);
 	static bool IsVoiceChatEnabled();
 
