Enforce 2 worker threads as the minimum spawn count on Windows. This resolves threading issues with dual-logical-core users that infinitely stall when multiple initialization processes and loading effectively deadlock.

diff --git a/Engine/Source/Runtime/Core/Private/Windows/WindowsPlatformMisc.cpp b/Engine/Source/Runtime/Core/Private/Windows/WindowsPlatformMisc.cpp
--- a/Engine/Source/Runtime/Core/Private/Windows/WindowsPlatformMisc.cpp
+++ b/Engine/Source/Runtime/Core/Private/Windows/WindowsPlatformMisc.cpp
@@ -1678,7 +1678,7 @@ int32 FWindowsPlatformMisc::NumberOfWorkerThreadsToSpawn()
 
 	int32 MaxWorkerThreadsWanted = IsRunningDedicatedServer() ? MaxServerWorkerThreads : MaxWorkerThreads;
 	// need to spawn at least one worker thread (see FTaskGraphImplementation)
-	return FMath::Max(FMath::Min(NumberOfThreads, MaxWorkerThreadsWanted), 1);
+	return FMath::Max(FMath::Min(NumberOfThreads, MaxWorkerThreadsWanted), 2);
 }
 
 bool FWindowsPlatformMisc::OsExecute(const TCHAR* CommandType, const TCHAR* Command, const TCHAR* CommandLine)
