Set async loading thread affinity to AllCores instead of LittleCores, improves boot time (1:45 -> 1:05 on S9+ Adreno)

diff --git a/Engine/Source/Runtime/Core/Public/Android/AndroidAffinity.h b/Engine/Source/Runtime/Core/Public/Android/AndroidAffinity.h
--- a/Engine/Source/Runtime/Core/Public/Android/AndroidAffinity.h
+++ b/Engine/Source/Runtime/Core/Public/Android/AndroidAffinity.h
@@ -66,7 +66,7 @@ public:
 
 	static const CORE_API uint64 GetAsyncLoadingThreadMask()
 	{
-		return GetLittleCoreMask();
+		return AllCores;
 	}
 
 	static EThreadPriority GetRenderingThreadPriority()
