Add a guid to the generic crash report which is specific to that particular execution of the project. Should allow multiple crashes in the same run to be correlated.

diff --git a/Engine/Source/Runtime/Core/Private/GenericPlatform/GenericPlatformCrashContext.cpp b/Engine/Source/Runtime/Core/Private/GenericPlatform/GenericPlatformCrashContext.cpp
--- a/Engine/Source/Runtime/Core/Private/GenericPlatform/GenericPlatformCrashContext.cpp
+++ b/Engine/Source/Runtime/Core/Private/GenericPlatform/GenericPlatformCrashContext.cpp
@@ -60,6 +60,8 @@ bool FGenericCrashContext::bIsInitialized = false;
 FPlatformMemoryStats FGenericCrashContext::CrashMemoryStats = FPlatformMemoryStats();
 int32 FGenericCrashContext::StaticCrashContextIndex = 0;
 
+const FGuid FGenericCrashContext::ExecutionGuid = FGuid::NewGuid();
+
 namespace NCachedCrashContextProperties
 {
 	static bool bIsInternalBuild;
@@ -275,6 +277,7 @@ void FGenericCrashContext::SerializeContentToBuffer() const
 
 	BeginSection( *RuntimePropertiesTag );
 	AddCrashProperty( TEXT( "CrashVersion" ), (int32)ECrashDescVersions::VER_3_CrashContext );
+	AddCrashProperty( TEXT( "ExecutionGuid" ), *ExecutionGuid.ToString() );
 	AddCrashProperty( TEXT( "CrashGUID" ), (const TCHAR*)CrashGUID);
 	AddCrashProperty( TEXT( "ProcessId" ), FPlatformProcess::GetCurrentProcessId() );
 	AddCrashProperty( TEXT( "IsInternalBuild" ), NCachedCrashContextProperties::bIsInternalBuild );
diff --git a/Engine/Source/Runtime/Core/Public/GenericPlatform/GenericPlatformCrashContext.h b/Engine/Source/Runtime/Core/Public/GenericPlatform/GenericPlatformCrashContext.h
--- a/Engine/Source/Runtime/Core/Public/GenericPlatform/GenericPlatformCrashContext.h
+++ b/Engine/Source/Runtime/Core/Public/GenericPlatform/GenericPlatformCrashContext.h
@@ -133,6 +133,9 @@ public:
 	static const FString EngineModeExDirty;
 	static const FString EngineModeExVanilla;
 
+	// A guid that identifies this particular execution. Allows multiple crash reports from the same run of the project to be tied together
+	static const FGuid ExecutionGuid;
+
 	/** Initializes crash context related platform specific data that can be impossible to obtain after a crash. */
 	static void Initialize();
 
