[Gil.Gribb and maybe others did most of the work, I just did the IOS parts, but I am checking it in, since IOS was the last stop on the code train]
- Added memory mapped IO concept to the engine
- Implemented and tested on Windows and IOS
- Not in use by anything yet, however, as the use cases are narrow and special (bulk data that does not need massaging after load - so it can be swapped out and back in silently by OS)

diff --git a/Engine/Source/Runtime/Core/Private/GenericPlatform/GenericPlatformFile.cpp b/Engine/Source/Runtime/Core/Private/GenericPlatform/GenericPlatformFile.cpp
--- a/Engine/Source/Runtime/Core/Private/GenericPlatform/GenericPlatformFile.cpp
+++ b/Engine/Source/Runtime/Core/Private/GenericPlatform/GenericPlatformFile.cpp
@@ -12,6 +12,7 @@
 #include "HAL/LowLevelMemTracker.h"
 
 #include "Async/AsyncFileHandle.h"
+#include "Async/MappedFileHandle.h"
 
 class FGenericBaseRequest;
 class FGenericAsyncReadFileHandle;
@@ -444,6 +445,10 @@ DEFINE_STAT(STAT_AsyncFileMemory);
 DEFINE_STAT(STAT_AsyncFileHandles);
 DEFINE_STAT(STAT_AsyncFileRequests);
 
+DEFINE_STAT(STAT_MappedFileMemory);
+DEFINE_STAT(STAT_MappedFileHandles);
+DEFINE_STAT(STAT_MappedFileRegions);
+
 int64 IFileHandle::Size()
 {
 	int64 Current = Tell();
diff --git a/Engine/Source/Runtime/Core/Private/IOS/IOSPlatformFile.cpp b/Engine/Source/Runtime/Core/Private/IOS/IOSPlatformFile.cpp
--- a/Engine/Source/Runtime/Core/Private/IOS/IOSPlatformFile.cpp
+++ b/Engine/Source/Runtime/Core/Private/IOS/IOSPlatformFile.cpp
@@ -11,6 +11,9 @@
 #include "Misc/CommandLine.h"
 #include "Misc/App.h"
 
+#include "Async/MappedFileHandle.h"
+#include <sys/mman.h>
+
 //#if PLATFORM_IOS
 //#include "IPlatformFileSandboxWrapper.h"
 //#endif
@@ -392,6 +395,119 @@ int32 FIOSFileHandle::ManagedFilesTlsSlot = FPlatformTLS::AllocTlsSlot();
 #endif
 
 
+class FIOSMappedFileRegion final : public IMappedFileRegion
+{
+public:
+	class FIOSMappedFileHandle* Parent;
+	const uint8* AlignedPtr;
+	uint64 AlignedSize;
+	
+	FIOSMappedFileRegion(const uint8* InMappedPtr, const uint8* InAlignedPtr, size_t InMappedSize, uint64 InAlignedSize, const FString& InDebugFilename, size_t InDebugOffsetIntoFile, class FIOSMappedFileHandle* InParent)
+		: IMappedFileRegion(InMappedPtr, InMappedSize, InDebugFilename, InDebugOffsetIntoFile)
+		, Parent(InParent)
+		, AlignedPtr(InAlignedPtr)
+		, AlignedSize(InAlignedSize)
+	{
+	}
+	
+	~FIOSMappedFileRegion();
+	
+	virtual void PreloadHint(int64 PreloadOffset = 0, int64 BytesToPreload = MAX_int64) override
+	{
+		int64 Size = GetMappedSize();
+		const uint8* Ptr = GetMappedPtr();
+		int32 FoolTheOptimizer = 0;
+		while (Size > 0)
+		{
+			FoolTheOptimizer += Ptr[0];
+			Size -= 4096;
+			Ptr += 4096;
+		}
+		if (FoolTheOptimizer == 0xbadf00d)
+		{
+			FPlatformProcess::Sleep(0.0f); // this will more or less never happen, but we can't let the optimizer strip these reads
+		}
+	}
+	
+};
+
+class FIOSMappedFileHandle final : public IMappedFileHandle
+{
+	const uint8* MappedPtr;
+	FString Filename;
+	int32 NumOutstandingRegions;
+	int32 Alignment;
+	int FileHandle;
+	
+public:
+	
+	FIOSMappedFileHandle(int InFileHandle, int64 FileSize, const FString& InFilename)
+		: IMappedFileHandle(FileSize)
+		, MappedPtr(nullptr)
+#if !UE_BUILD_SHIPPING
+		, Filename(InFilename)
+#endif
+		, NumOutstandingRegions(0)
+		, FileHandle(InFileHandle)
+	{
+		Alignment = sysconf(_SC_PAGE_SIZE);
+	}
+
+	~FIOSMappedFileHandle()
+	{
+		check(!NumOutstandingRegions); // can't delete the file before you delete all outstanding regions
+		close(FileHandle);
+	}
+	
+	virtual IMappedFileRegion* MapRegion(int64 Offset = 0, int64 BytesToMap = MAX_int64, bool bPreloadHint = false) override
+	{
+		check(Offset < GetFileSize()); // don't map zero bytes and don't map off the end of the file
+		BytesToMap = FMath::Min<int64>(BytesToMap, GetFileSize() - Offset);
+		check(BytesToMap > 0); // don't map zero bytes
+
+		
+		// const uint8* MapPtr = (const uint8 *)mmap(NULL, BytesToMap, PROT_READ, MAP_PRIVATE | (bPreloadHint ? MAP_POPULATE : 0), FileHandle, Offset);
+		// const uint8* MapPtr = (const uint8 *)mmap(NULL, BytesToMap, PROT_READ, MAP_PRIVATE, FileHandle, Offset);
+		//		const uint8* MapPtr = (const uint8 *)mmap(NULL, BytesToMap, PROT_READ, MAP_SHARED, FileHandle, Offset);
+		
+		int64 AlignedOffset = AlignDown(Offset, Alignment);
+		int64 AlignedSize = Align(BytesToMap + Offset - AlignedOffset, Alignment);
+		
+		// if we are about to go off the end, let's not
+		if (AlignedOffset + AlignedSize > GetFileSize())
+		{
+			UE_LOG(LogIOS, Warning, TEXT("Mapping fell off the end, did we need to actually abort? [%lld + %lld > %lld]"), AlignedOffset, AlignedSize, GetFileSize());
+			return nullptr;
+		}
+		
+		const uint8* AlignedMapPtr = (const uint8 *)mmap(NULL, AlignedSize, PROT_READ, MAP_PRIVATE, FileHandle, AlignedOffset);
+		if (AlignedMapPtr == (const uint8*)-1 || AlignedMapPtr == nullptr)
+		{
+			UE_LOG(LogIOS, Warning, TEXT("Failed to map memory %s, error is %d"), *Filename, errno);
+			return nullptr;
+		}
+
+		// create a mapping for this range
+		const uint8* MapPtr = AlignedMapPtr + Offset - AlignedOffset;
+		FIOSMappedFileRegion* Result = new FIOSMappedFileRegion(MapPtr, AlignedMapPtr, BytesToMap, AlignedSize, Filename, Offset, this);
+		NumOutstandingRegions++;
+		return Result;
+	}
+	
+	void UnMap(FIOSMappedFileRegion* Region)
+	{
+		check(NumOutstandingRegions > 0);
+		NumOutstandingRegions--;
+		
+		int Res = munmap((void*)Region->AlignedPtr, Region->AlignedSize);
+		checkf(Res == 0, TEXT("Failed to unmap, error is %d, errno is %d [params: %x, %d]"), Res, errno, MappedPtr, GetFileSize());
+	}
+};
+
+FIOSMappedFileRegion::~FIOSMappedFileRegion()
+{
+	Parent->UnMap(this);
+}
 
 
 
@@ -734,6 +850,42 @@ IFileHandle* FIOSPlatformFile::OpenWrite(const TCHAR* Filename, bool bAppend, bo
 	return NULL;
 }
 
+IMappedFileHandle* FIOSPlatformFile::OpenMapped(const TCHAR* Filename)
+{
+	FString NormalizedFilename = NormalizeFilename(Filename);
+	
+	// check the read path
+	FString FinalPath = ConvertToIOSPath(NormalizedFilename, false, false);
+	FILE* FP;
+	FP = fopen(TCHAR_TO_UTF8(*FinalPath), "r");
+//	if(Handle == -1)
+//	{
+//		// if not in the read path, check the private write path
+//		FinalPath = ConvertToIOSPath(NormalizedFilename, true, false);
+//		Handle = open(TCHAR_TO_UTF8(*FinalPath), O_RDONLY);
+//
+//		if(Handle == -1)
+//		{
+//			// if not in the private write path, check the public write path
+//			FinalPath = ConvertToIOSPath(NormalizedFilename, true, true);
+//			Handle = open(TCHAR_TO_UTF8(*FinalPath), O_RDONLY);
+//		}
+//	}
+	int32 Handle = fileno(FP);
+	
+	if (Handle != -1)
+	{
+		struct stat FileInfo;
+		FileInfo.st_size = -1;
+		// check the read path
+		fstat(Handle, &FileInfo);
+		uint64 FileSize = FileInfo.st_size;
+
+		return new FIOSMappedFileHandle(Handle, FileSize, FinalPath);
+	}
+	return NULL;
+}
+
 bool FIOSPlatformFile::DirectoryExists(const TCHAR* Directory)
 {
 	struct stat FileInfo;
diff --git a/Engine/Source/Runtime/Core/Private/IOS/IOSPlatformFile.h b/Engine/Source/Runtime/Core/Private/IOS/IOSPlatformFile.h
--- a/Engine/Source/Runtime/Core/Private/IOS/IOSPlatformFile.h
+++ b/Engine/Source/Runtime/Core/Private/IOS/IOSPlatformFile.h
@@ -42,6 +42,7 @@ public:
 
 	virtual IFileHandle* OpenRead(const TCHAR* Filename, bool bAllowWrite = false) override;
 	virtual IFileHandle* OpenWrite(const TCHAR* Filename, bool bAppend = false, bool bAllowRead = false) override;
+	virtual IMappedFileHandle* OpenMapped(const TCHAR* Filename) override;
 
 	virtual bool DirectoryExists(const TCHAR* Directory) override;
 	virtual bool CreateDirectory(const TCHAR* Directory) override;
diff --git a/Engine/Source/Runtime/Core/Private/Windows/WindowsPlatformFile.cpp b/Engine/Source/Runtime/Core/Private/Windows/WindowsPlatformFile.cpp
--- a/Engine/Source/Runtime/Core/Private/Windows/WindowsPlatformFile.cpp
+++ b/Engine/Source/Runtime/Core/Private/Windows/WindowsPlatformFile.cpp
@@ -23,6 +23,7 @@
 #include "Windows/AllowWindowsPlatformTypes.h"
 
 #include "WindowsAsyncIO.h"
+#include "Async/MappedFileHandle.h"
 
 TLockFreePointerListUnordered<void, PLATFORM_CACHE_LINE_SIZE> WindowsAsyncIOEventPool;
 bool GTriggerFailedWindowsRead = false;
@@ -600,6 +601,102 @@ public:
 	}
 };
 
+class FMappedFileRegionWindows final : public IMappedFileRegion
+{
+	class FMappedFileHandleWindows* Parent;
+	const uint8* AlignedMappedPtr;
+	size_t AlignedMappedSize;
+public:
+	FMappedFileRegionWindows(const uint8* InMappedPtr, const uint8* InAlignedMappedPtr, size_t InMappedSize, size_t InAlignedMappedSize, const FString& InDebugFilename, size_t InDebugOffsetRelativeToFile, class FMappedFileHandleWindows* InParent)
+		: IMappedFileRegion(InMappedPtr, InMappedSize, InDebugFilename, InDebugOffsetRelativeToFile)
+		, Parent(InParent)
+		, AlignedMappedPtr(InAlignedMappedPtr)
+		, AlignedMappedSize(InAlignedMappedSize)
+	{
+	}
+
+	~FMappedFileRegionWindows();
+
+	virtual void PreloadHint(int64 PreloadOffset = 0, int64 BytesToPreload = MAX_int64) override
+	{
+		// perhaps this could be done with a commit instead
+		int64 Size = GetMappedSize();
+		const uint8* Ptr = GetMappedPtr();
+		int32 FoolTheOptimizer = 0;
+		while (Size > 0)
+		{
+			FoolTheOptimizer += Ptr[0];
+			Size -= 4096;
+			Ptr += 4096;
+		}
+		if (FoolTheOptimizer == 0xbadf00d)
+		{
+			FPlatformProcess::Sleep(0.0f); // this will more or less never happen, but we can't let the optimizer strip these reads
+		}
+	}
+};
+
+class FMappedFileHandleWindows final : public IMappedFileHandle
+{
+	HANDLE Handle;
+	HANDLE MappingHandle;
+	FString DebugFilename;
+	int32 NumOutstandingRegions;
+public:
+	FMappedFileHandleWindows(HANDLE InHandle, HANDLE InMappingHandle, int64 Size, const TCHAR* InDebugFilename)
+		: IMappedFileHandle(Size)
+		, Handle(InHandle)
+		, MappingHandle(InMappingHandle)
+		, DebugFilename(InDebugFilename)
+		, NumOutstandingRegions(0)
+	{
+		check(Size >= 0);
+		check(Handle != INVALID_HANDLE_VALUE);
+	}
+	~FMappedFileHandleWindows()
+	{
+		check(!NumOutstandingRegions); // can't delete the file before you delete all outstanding regions
+		CloseHandle(MappingHandle);
+		CloseHandle(Handle);
+	}
+	virtual IMappedFileRegion* MapRegion(int64 Offset = 0, int64 BytesToMap = MAX_int64, bool bPreloadHint = false) override
+	{
+		check(Offset < GetFileSize()); // don't map zero bytes and don't map off the end of the file
+		BytesToMap = FMath::Min<int64>(BytesToMap, GetFileSize() - Offset);
+		check(BytesToMap > 0); // don't map zero bytes
+
+		DWORD OpenMappingAccess = FILE_MAP_READ;
+
+		int64 AlignedOffset = AlignDown(Offset, 65536);
+		int64 AlignedSize = Align(BytesToMap + Offset - AlignedOffset, 65536);
+
+		ULARGE_INTEGER LI;
+		LI.QuadPart = AlignedOffset;
+
+		const uint8* AlignedMapPtr = (const uint8*)MapViewOfFile(MappingHandle, FILE_MAP_READ, LI.HighPart, LI.LowPart, AlignedSize + AlignedOffset > GetFileSize() ? 0 : AlignedSize);
+		if (!AlignedMapPtr)
+		{
+			return nullptr;
+		}
+		const uint8* MapPtr = AlignedMapPtr + Offset - AlignedOffset;
+		FMappedFileRegionWindows* Result = new FMappedFileRegionWindows(MapPtr, AlignedMapPtr, BytesToMap, AlignedSize, DebugFilename, Offset, this);
+		NumOutstandingRegions++;
+		return Result;
+	}
+
+	void UnMap(FMappedFileRegionWindows* Region)
+	{
+		check(NumOutstandingRegions > 0);
+		NumOutstandingRegions--;
+	}
+};
+
+FMappedFileRegionWindows::~FMappedFileRegionWindows()
+{
+	UnmapViewOfFile((LPVOID)AlignedMappedPtr);
+	Parent->UnMap(this);
+}
+
 /**
  * Windows File I/O implementation
 **/
@@ -826,6 +923,30 @@ public:
 		return NULL;
 	}
 
+
+	virtual IMappedFileHandle* OpenMapped(const TCHAR* Filename) override
+	{
+		int64 Size = FileSize(Filename);
+		if (Size < 1)
+		{
+			return nullptr;
+		}
+		uint32  Access = GENERIC_READ;
+		uint32  WinFlags = FILE_SHARE_READ;
+		uint32  Create = OPEN_EXISTING;
+		HANDLE Handle = CreateFileW(*NormalizeFilename(Filename), Access, WinFlags, NULL, Create, FILE_ATTRIBUTE_NORMAL, NULL);
+		if (Handle == INVALID_HANDLE_VALUE)
+		{
+			return nullptr;
+		}
+		HANDLE MappingHandle = CreateFileMapping(Handle, NULL, PAGE_READONLY, 0, 0, NULL);
+		if (MappingHandle == INVALID_HANDLE_VALUE)
+		{
+			CloseHandle(Handle);
+			return nullptr;
+		}
+		return new FMappedFileHandleWindows(Handle, MappingHandle, Size, Filename);
+	}
 	virtual bool DirectoryExists(const TCHAR* Directory) override
 	{
 		// Empty Directory is the current directory so assume it always exists.
diff --git a/Engine/Source/Runtime/Core/Public/Async/MappedFileHandle.h b/Engine/Source/Runtime/Core/Public/Async/MappedFileHandle.h
new file mode 100644
--- /dev/null
+++ b/Engine/Source/Runtime/Core/Public/Async/MappedFileHandle.h
@@ -0,0 +1,121 @@
+// Copyright 1998-2018 Epic Games, Inc. All Rights Reserved.
+
+#pragma once
+
+#include "CoreTypes.h"
+#include "Misc/AssertionMacros.h"
+#include "Stats/Stats.h"
+
+
+DECLARE_MEMORY_STAT_EXTERN(TEXT("Mapped File Handle Memory"), STAT_MappedFileMemory, STATGROUP_Memory, CORE_API);
+DECLARE_DWORD_ACCUMULATOR_STAT_EXTERN(TEXT("Num Mapped File Handles"), STAT_MappedFileHandles, STATGROUP_Memory, CORE_API);
+DECLARE_DWORD_ACCUMULATOR_STAT_EXTERN(TEXT("Num Mapped File Regions"), STAT_MappedFileRegions, STATGROUP_Memory, CORE_API);
+
+// Note on threading. Like the rest of the filesystem platform abstraction, these methods are threadsafe, but it is expected you are not concurrently _using_ these data structures. 
+
+class CORE_API IMappedFileRegion
+{
+	const uint8* MappedPtr;
+	size_t MappedSize;
+	FString DebugFilename;
+	size_t DebugOffsetRelativeToFile;
+
+	FORCEINLINE void CheckInvariants()
+	{
+		check(MappedPtr && MappedSize);
+	}
+
+public:
+
+	FORCEINLINE IMappedFileRegion(const uint8* InMappedPtr, size_t InMappedSize, const FString& InDebugFilename, size_t InDebugOffsetRelativeToFile)
+		: MappedPtr(InMappedPtr)
+		, MappedSize(InMappedSize)
+		, DebugFilename(InDebugFilename)
+		, DebugOffsetRelativeToFile(InDebugOffsetRelativeToFile)
+	{
+		CheckInvariants();
+		INC_DWORD_STAT(STAT_MappedFileRegions);
+		INC_MEMORY_STAT_BY(STAT_MappedFileMemory, MappedSize);
+	}
+
+	virtual ~IMappedFileRegion()
+	{
+		check(MappedSize);  
+		DEC_MEMORY_STAT_BY(STAT_MappedFileMemory, MappedSize);
+		DEC_DWORD_STAT(STAT_MappedFileRegions);
+	}
+
+
+	/**
+	* Return the pointer to the mapped region. 
+	* @return Returned size of the file or -1 if the file was not found, the request was canceled or other error.
+	**/
+	FORCEINLINE const uint8* GetMappedPtr()
+	{
+		CheckInvariants();
+		return MappedPtr;
+	}
+
+	/**
+	* Return the size of the mapped region.
+	**/
+	FORCEINLINE int64 GetMappedSize()
+	{
+		CheckInvariants();
+		return MappedSize;
+	}
+
+	/**
+	* Synchronously preload part or all of the mapped region. Typically this is done by reading a byte from each CPU page.
+	* This is only a hint, some platforms might ignore it.
+	* There are not guarantees how long this data will stay in memory.
+	* @param PreloadOffset		Offset into this region to preload
+	* @param BytesToPreload		number of bytes to preload. This is clamped to the size of the mapped region
+	**/
+	virtual void PreloadHint(int64 PreloadOffset = 0, int64 BytesToPreload = MAX_int64)
+	{
+	}
+
+	// Non-copyable
+	IMappedFileRegion(const IMappedFileRegion&) = delete;
+	IMappedFileRegion& operator=(const IMappedFileRegion&) = delete;
+};
+
+class CORE_API IMappedFileHandle
+{
+	size_t MappedFileSize;
+
+public:
+	IMappedFileHandle(size_t InFileSize)
+		: MappedFileSize(InFileSize)
+	{
+		check(MappedFileSize >= 0);
+		INC_DWORD_STAT(STAT_MappedFileHandles);
+	}
+	/** Destructor, also the only way to close the file handle. It is not legal to delete an async file with outstanding requests. You must always call WaitCompletion before deleting a request. **/
+	virtual ~IMappedFileHandle()
+	{
+		DEC_DWORD_STAT(STAT_MappedFileHandles);
+	}
+
+	/**
+	* Return the size of the mapped file.
+	**/
+	FORCEINLINE int64 GetFileSize()
+	{
+		return MappedFileSize;
+	}
+
+	/**
+	* Map a region of the file.
+	* @param Offset				Offset into the file to start mapping.
+	* @param BytesToMap			Number of bytes to map. Clamped to the size of the file.
+	* @param bPreloadHint		If true, preload the data. This is only a hint and might be ignored, see IMappedFileRegion::PreloadHint
+	* @return the mapped region interface. This method cannot fail.
+	**/
+	virtual IMappedFileRegion* MapRegion(int64 Offset = 0, int64 BytesToMap = MAX_int64, bool bPreloadHint = false) = 0;
+
+	// Non-copyable
+	IMappedFileHandle(const IMappedFileHandle&) = delete;
+	IMappedFileHandle& operator=(const IMappedFileHandle&) = delete;
+};
diff --git a/Engine/Source/Runtime/Core/Public/GenericPlatform/GenericPlatformFile.h b/Engine/Source/Runtime/Core/Public/GenericPlatform/GenericPlatformFile.h
--- a/Engine/Source/Runtime/Core/Public/GenericPlatform/GenericPlatformFile.h
+++ b/Engine/Source/Runtime/Core/Public/GenericPlatform/GenericPlatformFile.h
@@ -14,6 +14,7 @@
 #include "Misc/EnumClassFlags.h"
 
 class IAsyncReadFileHandle;
+class IMappedFileHandle;
 
 /**
 * Enum for async IO priorities. 
@@ -386,6 +387,19 @@ public:
 	{
 	}
 
+	/** Open a file for async reading. This call does hit the disk; it is synchronous open. 
+	*
+	* @param Filename file to be mapped. This doesn't actually map anything, just opens the file.
+	* @return Close the file by delete'ing the handle. A non-null return value does mean the file exists. 
+	* Null can be returned for many reasons even if the file exists. Perhaps this platform does not support mapped files, or this file is compressed in a pak file.
+	* Generally you attempt to open mapped, and if that fails, then use other file operations instead.
+	*/
+	virtual IMappedFileHandle* OpenMapped(const TCHAR* Filename)
+	{
+		return nullptr;
+	}
+
+
 	virtual void GetTimeStampPair(const TCHAR* PathA, const TCHAR* PathB, FDateTime& OutTimeStampA, FDateTime& OutTimeStampB);
 
 	/** Return the modification time of a file in the local time of the calling code (GetTimeStamp returns UTC). Returns FDateTime::MinValue() on failure **/
diff --git a/Engine/Source/Runtime/Core/Public/HAL/IPlatformFileCachedWrapper.h b/Engine/Source/Runtime/Core/Public/HAL/IPlatformFileCachedWrapper.h
--- a/Engine/Source/Runtime/Core/Public/HAL/IPlatformFileCachedWrapper.h
+++ b/Engine/Source/Runtime/Core/Public/HAL/IPlatformFileCachedWrapper.h
@@ -433,6 +433,10 @@ public:
 	{
 		return LowerLevel->OpenAsyncRead(Filename);
 	}
+	virtual IMappedFileHandle* OpenMapped(const TCHAR* Filename) override
+	{
+		return LowerLevel->OpenMapped(Filename);
+	}
 	virtual void SetAsyncMinimumPriority(EAsyncIOPriorityAndFlags MinPriority) override
 	{
 		LowerLevel->SetAsyncMinimumPriority(MinPriority);
diff --git a/Engine/Source/Runtime/Core/Public/HAL/IPlatformFileLogWrapper.h b/Engine/Source/Runtime/Core/Public/HAL/IPlatformFileLogWrapper.h
--- a/Engine/Source/Runtime/Core/Public/HAL/IPlatformFileLogWrapper.h
+++ b/Engine/Source/Runtime/Core/Public/HAL/IPlatformFileLogWrapper.h
@@ -482,6 +482,10 @@ public:
 		//@todo no wrapped logging for async file handles (yet)
 		return Result;
 	}
+	virtual IMappedFileHandle* OpenMapped(const TCHAR* Filename) override
+	{
+		return LowerLevel->OpenMapped(Filename);
+	}
 	virtual void SetAsyncMinimumPriority(EAsyncIOPriorityAndFlags Priority) override
 	{
 		LowerLevel->SetAsyncMinimumPriority(Priority);
diff --git a/Engine/Source/Runtime/Core/Public/HAL/IPlatformFileOpenLogWrapper.h b/Engine/Source/Runtime/Core/Public/HAL/IPlatformFileOpenLogWrapper.h
--- a/Engine/Source/Runtime/Core/Public/HAL/IPlatformFileOpenLogWrapper.h
+++ b/Engine/Source/Runtime/Core/Public/HAL/IPlatformFileOpenLogWrapper.h
@@ -258,6 +258,10 @@ public:
 		// we must not record the "open" here...what matters is when we start reading the file!
 		return new FLoggingAsyncReadFileHandle(this, Filename, LowerLevel->OpenAsyncRead(Filename));
 	}
+	virtual IMappedFileHandle* OpenMapped(const TCHAR* Filename) override
+	{
+		return LowerLevel->OpenMapped(Filename);
+	}
 
 	void AddToOpenLog(const TCHAR* Filename)
 	{
diff --git a/Engine/Source/Runtime/Core/Public/HAL/IPlatformFileProfilerWrapper.h b/Engine/Source/Runtime/Core/Public/HAL/IPlatformFileProfilerWrapper.h
--- a/Engine/Source/Runtime/Core/Public/HAL/IPlatformFileProfilerWrapper.h
+++ b/Engine/Source/Runtime/Core/Public/HAL/IPlatformFileProfilerWrapper.h
@@ -514,6 +514,10 @@ public:
 	{
 		return LowerLevel->OpenAsyncRead(Filename);
 	}
+	virtual IMappedFileHandle* OpenMapped(const TCHAR* Filename) override
+	{
+		return LowerLevel->OpenMapped(Filename);
+	}
 	virtual void SetAsyncMinimumPriority(EAsyncIOPriorityAndFlags Priority) override
 	{
 		LowerLevel->SetAsyncMinimumPriority(Priority);
diff --git a/Engine/Source/Runtime/CoreUObject/Private/Serialization/BulkData.cpp b/Engine/Source/Runtime/CoreUObject/Private/Serialization/BulkData.cpp
--- a/Engine/Source/Runtime/CoreUObject/Private/Serialization/BulkData.cpp
+++ b/Engine/Source/Runtime/CoreUObject/Private/Serialization/BulkData.cpp
@@ -16,6 +16,8 @@
 #include "Interfaces/ITargetPlatform.h"
 #include "UObject/DebugSerializationFlags.h"
 #include "Serialization/AsyncLoadingPrivate.h"
+#include "Async/MappedFileHandle.h"
+#include "HAL/PlatformFilemanager.h"
 
 /*-----------------------------------------------------------------------------
 	Constructors and operators
@@ -92,6 +94,42 @@ uint32 GetTypeHash( const FUntypedBulkData* BulkData )
 #endif
 
 
+bool FUntypedBulkData::FAllocatedPtr::MapFile(const TCHAR *InFilename, int64 Offset, int64 Size)
+{
+	check(!MappedHandle && !MappedRegion); // It doesn't make sense to do this twice, but if need be, not hard to do
+
+	MappedHandle = FPlatformFileManager::Get().GetPlatformFile().OpenMapped(InFilename);
+
+	if (!MappedHandle)
+	{
+		return false;
+	}
+	MappedRegion = MappedHandle->MapRegion(Offset, Size, true); //@todo we really don't want to hit the disk here to bring it into memory
+	if (!MappedRegion)
+	{
+		delete MappedHandle;
+		MappedHandle = nullptr;
+		return false;
+	}
+
+	check(Size == MappedRegion->GetMappedSize());
+	Ptr = (void*)(MappedRegion->GetMappedPtr()); //@todo mapped files should probably be const-correct
+	bAllocated = true;
+	return true;
+}
+
+void FUntypedBulkData::FAllocatedPtr::UnmapFile()
+{
+	if (MappedRegion || MappedHandle)
+	{
+		delete MappedRegion;
+		delete MappedHandle;
+		MappedRegion = nullptr;
+		MappedHandle = nullptr;
+		Ptr = nullptr; // make sure we don't try to free this pointer
+	}
+}
+
 /**
  * Constructor, initializing all member variables.
  */
@@ -789,14 +827,7 @@ bool FUntypedBulkData::ShouldStreamBulkData()
 		GMinimumBulkDataSizeForAsyncLoading >= 0);
 }
 
-/**
-* Serialize function used to serialize this bulk data structure.
-*
-* @param Ar	Archive to serialize with
-* @param Owner	Object owning the bulk data
-* @param Idx	Index of bulk data item being serialized
-*/
-void FUntypedBulkData::Serialize( FArchive& Ar, UObject* Owner, int32 Idx )
+void FUntypedBulkData::Serialize( FArchive& Ar, UObject* Owner, int32 Idx, bool bAttemptFileMapping)
 {
 	DECLARE_SCOPE_CYCLE_COUNTER(TEXT("FUntypedBulkData::Serialize"), STAT_UBD_Serialize, STATGROUP_Memory);
 
@@ -924,10 +955,30 @@ void FUntypedBulkData::Serialize( FArchive& Ar, UObject* Owner, int32 Idx )
 					}
 					else
 					{
-						// Force non-lazy loading of inline bulk data to prevent PostLoad spikes.
-						BulkData.Reallocate(GetBulkDataSize(), BulkDataAlignment);
-						// if the payload is stored inline, just serialize it
-						SerializeBulkData(Ar, BulkData.Get());
+						bool bWasMapped = false;
+						if (bAttemptFileMapping)
+						{
+							if (GEventDrivenLoaderEnabled && (Filename.EndsWith(TEXT(".uasset")) || Filename.EndsWith(TEXT(".umap"))))
+							{
+								BulkDataOffsetInFile -= IFileManager::Get().FileSize(*Filename);
+								check(BulkDataOffsetInFile >= 0);
+								Filename = FPaths::GetBaseFilename(Filename, false) + TEXT(".uexp");
+							}
+							bWasMapped = BulkData.MapFile(*Filename, BulkDataOffsetInFile, GetBulkDataSize());
+						}
+						if (bWasMapped)
+						{
+							// we need to seek past the inline bulk data 
+							// @todo, we really don't want to do this with inline data
+							Ar.Seek(Ar.Tell() + GetBulkDataSize());
+						}
+						else
+						{
+							// Force non-lazy loading of inline bulk data to prevent PostLoad spikes.
+							BulkData.Reallocate(GetBulkDataSize(), BulkDataAlignment);
+							// if the payload is stored inline, just serialize it
+							SerializeBulkData(Ar, BulkData.Get());
+						}
 					}
 				}
 				else if (BulkDataFlags & BULKDATA_OptionalPayload)
@@ -1661,4 +1712,17 @@ void FFormatContainer::Serialize(FArchive& Ar, UObject* Owner, const TArray<FNam
 	}
 }
 
+void FFormatContainer::SerializeAttemptMappedLoad(FArchive& Ar, UObject* Owner)
+{
+	check(Ar.IsLoading());
+	int32 NumFormats = 0;
+	Ar << NumFormats;
+	for (int32 Index = 0; Index < NumFormats; Index++)
+	{
+		FName Name;
+		Ar << Name;
+		FByteBulkData& Bulk = GetFormat(Name);
+		Bulk.Serialize(Ar, Owner, -1, true);
+	}
+}
 
diff --git a/Engine/Source/Runtime/CoreUObject/Public/Serialization/BulkData.h b/Engine/Source/Runtime/CoreUObject/Public/Serialization/BulkData.h
--- a/Engine/Source/Runtime/CoreUObject/Public/Serialization/BulkData.h
+++ b/Engine/Source/Runtime/CoreUObject/Public/Serialization/BulkData.h
@@ -62,6 +62,9 @@ enum EBulkDataLockFlags
 	LOCK_READ_WRITE								= 2,
 };
 
+class IMappedFileHandle;
+class IMappedFileRegion;
+
 /*-----------------------------------------------------------------------------
 	Base version of untyped bulk data.
 -----------------------------------------------------------------------------*/
@@ -77,15 +80,21 @@ private:
 	{
 		FAllocatedPtr()
 			: Ptr       (nullptr)
+			, MappedHandle(nullptr)
+			, MappedRegion(nullptr)
 			, bAllocated(false)
 		{
 		}
 
 		FAllocatedPtr(FAllocatedPtr&& Other)
 			: Ptr       (Other.Ptr)
+			, MappedHandle(Other.MappedHandle)
+			, MappedRegion(Other.MappedRegion)
 			, bAllocated(Other.bAllocated)
 		{
-			Other.Ptr        = nullptr;
+			Other.Ptr = nullptr;
+			Other.MappedHandle = nullptr;
+			Other.MappedRegion = nullptr;
 			Other.bAllocated = false;
 		}
 
@@ -99,7 +108,7 @@ private:
 
 		~FAllocatedPtr()
 		{
-			FMemory::Free(Ptr);
+			Deallocate();
 		}
 
 		void* Get() const
@@ -114,6 +123,7 @@ private:
 
 		void Reallocate(int32 Count, int32 Alignment = DEFAULT_ALIGNMENT)
 		{
+			check(!MappedHandle && !MappedRegion); // not legal for mapped bulk data
 			if (Count)
 			{
 				Ptr = FMemory::Realloc(Ptr, Count, Alignment);
@@ -129,6 +139,11 @@ private:
 
 		void* ReleaseWithoutDeallocating()
 		{
+			if (MappedHandle || MappedRegion)
+			{
+				// Super scary, we returned a pointer to a mapped file but we have no guarantees that this outlives the pointer we let out into the engine
+				// @todo transfer ownership properly by making FAllocatedPtr a public thing that can be used by people to take ownership of the entire mapping. 
+			}
 			void* Result = Ptr;
 			Ptr = nullptr;
 			bAllocated = false;
@@ -137,16 +152,25 @@ private:
 
 		void Deallocate()
 		{
+			if (MappedHandle || MappedRegion)
+			{
+				UnmapFile();
+			}
 			FMemory::Free(Ptr);
 			Ptr = nullptr;
 			bAllocated = false;
 		}
 
+		COREUOBJECT_API bool MapFile(const TCHAR *Filename, int64 Offset, int64 Size);
+		COREUOBJECT_API void UnmapFile();
+
 	private:
 		FAllocatedPtr(const FAllocatedPtr&);
 		FAllocatedPtr& operator=(const FAllocatedPtr&);
 
 		void* Ptr;
+		IMappedFileHandle* MappedHandle;
+		IMappedFileRegion* MappedRegion;
 		bool  bAllocated;
 	};
 
@@ -395,8 +419,9 @@ public:
 	 * @param Ar	Archive to serialize with
 	 * @param Owner	Object owning the bulk data
 	 * @param Idx	Index of bulk data item being serialized
+	 * @param bAttemptFileMapping	If true, attempt to map this instead of loading it into malloc'ed memory
 	 */
-	void Serialize( FArchive& Ar, UObject* Owner, int32 Idx=INDEX_NONE );
+	void Serialize( FArchive& Ar, UObject* Owner, int32 Idx=INDEX_NONE, bool bAttemptFileMapping = false);
 
 	/**
 	 * Serialize just the bulk data portion to/ from the passed in memory.
@@ -662,5 +687,6 @@ public:
 		Formats.Empty();
 	}
 	COREUOBJECT_API void Serialize(FArchive& Ar, UObject* Owner, const TArray<FName>* FormatsToSave = nullptr, bool bSingleUse = true, uint32 InAlignment = DEFAULT_ALIGNMENT);
+	COREUOBJECT_API void SerializeAttemptMappedLoad(FArchive& Ar, UObject* Owner);
 };
 
diff --git a/Engine/Source/Runtime/PakFile/Private/IPlatformFilePak.cpp b/Engine/Source/Runtime/PakFile/Private/IPlatformFilePak.cpp
--- a/Engine/Source/Runtime/PakFile/Private/IPlatformFilePak.cpp
+++ b/Engine/Source/Runtime/PakFile/Private/IPlatformFilePak.cpp
@@ -26,6 +26,8 @@
 #endif
 #include "ProfilingDebugging/CsvProfiler.h"
 
+#include "Async/MappedFileHandle.h"
+
 DEFINE_LOG_CATEGORY(LogPakFile);
 
 DEFINE_STAT(STAT_PakFile_Read);
@@ -3733,6 +3735,101 @@ void FPakPlatformFile::Tick()
 #endif
 }
 
+class FMappedFilePakProxy final : public IMappedFileHandle
+{
+	IMappedFileHandle* LowerLevel;
+	int64 OffsetInPak;
+	int64 PakSize;
+	FString DebugFilename;
+public:
+	FMappedFilePakProxy(IMappedFileHandle* InLowerLevel, int64 InOffset, int64 InSize, int64 InPakSize, const TCHAR* InDebugFilename)
+		: IMappedFileHandle(InSize)
+		, LowerLevel(InLowerLevel)
+		, OffsetInPak(InOffset)
+		, PakSize(InPakSize)
+		, DebugFilename(InDebugFilename)
+	{
+		check(PakSize >= 0);
+	}
+	virtual ~FMappedFilePakProxy()
+	{
+		// we don't own lower level, it is shared
+	}
+	virtual IMappedFileRegion* MapRegion(int64 Offset = 0, int64 BytesToMap = MAX_int64, bool bPreloadHint = false) override
+	{
+		check(Offset + OffsetInPak < PakSize); // don't map zero bytes and don't map off the end of the (real) file
+		check(Offset < GetFileSize()); // don't map zero bytes and don't map off the end of the (virtual) file
+		BytesToMap = FMath::Min<int64>(BytesToMap, GetFileSize() - Offset);
+		check(BytesToMap > 0); // don't map zero bytes
+		check(Offset + BytesToMap <= GetFileSize()); // don't map zero bytes and don't map off the end of the (virtual) file
+		check(Offset + OffsetInPak + BytesToMap <= PakSize); // don't map zero bytes and don't map off the end of the (real) file
+		return LowerLevel->MapRegion(Offset + OffsetInPak, BytesToMap, bPreloadHint);
+	}
+};
+
+
+#if !UE_BUILD_SHIPPING
+
+static void MappedFileTest(const TArray<FString>& Args)
+{
+	FString TestFile(TEXT("../../../Engine/Config/BaseDeviceProfiles.ini"));
+	if (Args.Num() > 0)
+	{
+		TestFile = Args[0];
+	}
+
+	while (true)
+	{
+		IMappedFileHandle* Handle = FPlatformFileManager::Get().GetPlatformFile().OpenMapped(*TestFile);
+		IMappedFileRegion *Region = Handle->MapRegion();
+
+		int64 Size = Region->GetMappedSize();
+		const char* Data = (const char *)Region->GetMappedPtr();
+
+		delete Region;
+		delete Handle;
+	}
+
+
+}
+
+static FAutoConsoleCommand MappedFileTestCmd(
+	TEXT("MappedFileTest"),
+	TEXT("Tests the file mappings through the low level."),
+	FConsoleCommandWithArgsDelegate::CreateStatic(&MappedFileTest)
+);
+#endif
+
+IMappedFileHandle* FPakPlatformFile::OpenMapped(const TCHAR* Filename)
+{
+	// Check pak files first
+	FPakEntry FileEntry;
+	FPakFile* PakEntry = nullptr;
+	if (FindFileInPakFiles(Filename, &PakEntry, &FileEntry) && PakEntry)
+	{
+		if (FileEntry.CompressionMethod != COMPRESS_None)
+		{
+			// can't map compressed files
+			return nullptr;
+		}
+		FScopeLock Lock(&PakEntry->MappedFileHandleCriticalSection);
+		if (!PakEntry->MappedFileHandle)
+		{
+			PakEntry->MappedFileHandle = LowerLevel->OpenMapped(*PakEntry->GetFilename());
+		}
+		if (!PakEntry->MappedFileHandle)
+		{
+			return nullptr;
+		}
+		return new FMappedFilePakProxy(PakEntry->MappedFileHandle, FileEntry.Offset + FileEntry.GetSerializedSize(PakEntry->GetInfo().Version), FileEntry.UncompressedSize, PakEntry->TotalSize(), Filename);
+	}
+	if (IsNonPakFilenameAllowed(Filename))
+	{
+		return LowerLevel->OpenMapped(Filename);
+	}
+	return nullptr;
+}
+
 
 /**
  * Class to handle correctly reading from a compressed file within a compressed package
@@ -4020,6 +4117,7 @@ FPakFile::FPakFile(const TCHAR* Filename, bool bIsSigned)
 	, bIsValid(false)
 	, bFilenamesRemoved(false)
 	, ChunkID(ParseChunkIDFromFilename(Filename))
+	, MappedFileHandle(nullptr)
 {
 	FArchive* Reader = GetSharedReader(NULL);
 	if (Reader)
@@ -4044,6 +4142,7 @@ FPakFile::FPakFile(IPlatformFile* LowerLevel, const TCHAR* Filename, bool bIsSig
 	, bIsValid(false)
 	, bFilenamesRemoved(false)
 	, ChunkID(ParseChunkIDFromFilename(Filename))
+	, MappedFileHandle(nullptr)
 {
 	FArchive* Reader = GetSharedReader(LowerLevel);
 	if (Reader)
@@ -4065,6 +4164,7 @@ FPakFile::FPakFile(FArchive* Archive)
 	, bIsValid(false)
 	, bFilenamesRemoved(false)
 	, ChunkID(INDEX_NONE)
+	, MappedFileHandle(nullptr)
 {
 	Initialize(Archive);
 }
@@ -4072,6 +4172,7 @@ FPakFile::FPakFile(FArchive* Archive)
 
 FPakFile::~FPakFile()
 {
+	delete MappedFileHandle;
 	delete[] MiniPakEntries;
 	delete[] MiniPakEntriesOffsets;
 	delete[] FilenameHashes;
diff --git a/Engine/Source/Runtime/PakFile/Public/IPlatformFilePak.h b/Engine/Source/Runtime/PakFile/Public/IPlatformFilePak.h
--- a/Engine/Source/Runtime/PakFile/Public/IPlatformFilePak.h
+++ b/Engine/Source/Runtime/PakFile/Public/IPlatformFilePak.h
@@ -530,6 +530,10 @@ class PAKFILE_API FPakFile : FNoncopyable
 	/** ID for the chunk this pakfile is part of. INDEX_NONE if this isn't a pak chunk (derived from filename) */
 	int32 ChunkID;
 
+	class IMappedFileHandle* MappedFileHandle;
+	FCriticalSection MappedFileHandleCriticalSection;
+
+
 	static inline int32 CDECL CompareFilenameHashes(const void* Left, const void* Right)
 	{
 		const uint32* LeftHash = (const uint32*)Left;
@@ -2242,6 +2246,7 @@ public:
 	virtual IAsyncReadFileHandle* OpenAsyncRead(const TCHAR* Filename) override;
 	virtual void SetAsyncMinimumPriority(EAsyncIOPriorityAndFlags Priority) override;
 
+	virtual IMappedFileHandle* OpenMapped(const TCHAR* Filename) override;
 	/**
 	 * Converts a filename to a path inside pak file.
 	 *
diff --git a/Engine/Source/Runtime/SandboxFile/Public/IPlatformFileSandboxWrapper.h b/Engine/Source/Runtime/SandboxFile/Public/IPlatformFileSandboxWrapper.h
--- a/Engine/Source/Runtime/SandboxFile/Public/IPlatformFileSandboxWrapper.h
+++ b/Engine/Source/Runtime/SandboxFile/Public/IPlatformFileSandboxWrapper.h
@@ -653,6 +653,15 @@ public:
 	{
 		LowerLevel->SetAsyncMinimumPriority(Priority);
 	}
+	virtual IMappedFileHandle* OpenMapped(const TCHAR* Filename) override
+	{
+		FString UserFilename(*ConvertToSandboxPath(Filename));
+		if (!OkForInnerAccess(Filename) || LowerLevel->FileExists(*UserFilename))
+		{
+			return LowerLevel->OpenMapped(*UserFilename);
+		}
+		return LowerLevel->OpenMapped(Filename);
+	}
 
 };
 
