[NO NON-PII COMMIT MESSAGE]

diff --git a/Engine/Source/Runtime/Core/Public/Math/UnrealMathSSE.h b/Engine/Source/Runtime/Core/Public/Math/UnrealMathSSE.h
--- a/Engine/Source/Runtime/Core/Public/Math/UnrealMathSSE.h
+++ b/Engine/Source/Runtime/Core/Public/Math/UnrealMathSSE.h
@@ -12,6 +12,8 @@ struct FMath;
 // We require SSE2
 #include <emmintrin.h>
 
+// We suppress static analysis warnings for the cast from (double*) to (float*) in VectorLoadFloat2 below:
+// -V:VectorLoadFloat2:615
 
 /*=============================================================================
  *	Helpers:
