- Depreceate FGetHSV, replace with new function MakeFromHSV8 which wraps existing method HSVToLinearRGB
- FGetHSV behavior is confusing/wrong
- Behavior of HSVToLinearRGB is what we want, matches with LinearRGBToHSV

diff --git a/Engine/Plugins/Experimental/PythonScriptPlugin/Source/PythonScriptPlugin/Private/PyWrapperMath.cpp b/Engine/Plugins/Experimental/PythonScriptPlugin/Source/PythonScriptPlugin/Private/PyWrapperMath.cpp
--- a/Engine/Plugins/Experimental/PythonScriptPlugin/Source/PythonScriptPlugin/Private/PyWrapperMath.cpp
+++ b/Engine/Plugins/Experimental/PythonScriptPlugin/Source/PythonScriptPlugin/Private/PyWrapperMath.cpp
@@ -1540,7 +1540,7 @@ PyTypeObject InitializePyWrapperMathType_LinearColor()
 				return nullptr;
 			}
 
-			FPyWrapperLinearColor::GetTypedStruct(InSelf) = FLinearColor::FGetHSV(H, S, V);
+			FPyWrapperLinearColor::GetTypedStruct(InSelf) = FLinearColor::MakeFromHSV8(H, S, V);
 			Py_RETURN_NONE;
 		}
 
diff --git a/Engine/Source/Developer/LogVisualizer/Private/VisualLoggerCanvasRenderer.cpp b/Engine/Source/Developer/LogVisualizer/Private/VisualLoggerCanvasRenderer.cpp
--- a/Engine/Source/Developer/LogVisualizer/Private/VisualLoggerCanvasRenderer.cpp
+++ b/Engine/Source/Developer/LogVisualizer/Private/VisualLoggerCanvasRenderer.cpp
@@ -269,7 +269,7 @@ void FVisualLoggerCanvasRenderer::DrawHistogramGraphs(class UCanvas* Canvas, cla
 					Hue -= FMath::FloorToFloat(Hue);
 				}
 
-				HistogramGraph->GetGraphLine(LineIndex)->Color = FLinearColor::FGetHSV(Hue * 255, 0, 244);
+				HistogramGraph->GetGraphLine(LineIndex)->Color = FLinearColor::MakeFromHSV8(Hue * 255, 0, 244);
 				HistogramGraph->GetGraphLine(LineIndex)->LineName = DataName;
 				HistogramGraph->GetGraphLine(LineIndex)->Data.Append(LinesIt->Value.Samples);
 				HistogramGraph->GetGraphLine(LineIndex)->LeftExtreme = LinesIt->Value.LeftExtreme;
diff --git a/Engine/Source/Runtime/Core/Private/Math/Color.cpp b/Engine/Source/Runtime/Core/Private/Math/Color.cpp
--- a/Engine/Source/Runtime/Core/Private/Math/Color.cpp
+++ b/Engine/Source/Runtime/Core/Private/Math/Color.cpp
@@ -228,7 +228,7 @@ FLinearColor FColor::FromRGBE() const
 }
 
 /**
- * Converts byte hue-saturation-brightness to floating point red-green-blue.
+ * Old version of MakeFromHSV8, formula is not correct
  */
 FLinearColor FLinearColor::FGetHSV( uint8 H, uint8 S, uint8 V )
 {
@@ -240,6 +240,15 @@ FLinearColor FLinearColor::FGetHSV( uint8 H, uint8 S, uint8 V )
 	return FLinearColor(ColorVector.X,ColorVector.Y,ColorVector.Z,1);
 }
 
+/**
+ * Converts byte hue-saturation-brightness to floating point red-green-blue.
+ */
+FLinearColor FLinearColor::MakeFromHSV8(uint8 H, uint8 S, uint8 V)
+{
+	// want a given H value of 255 to map to just below 360 degrees
+	const FLinearColor HSVColor((float)H * (360.0f / 256.0f), (float)S / 255.0f, (float)V / 255.0f);
+	return HSVColor.HSVToLinearRGB();
+}
 
 /** Converts a linear space RGB color to an HSV color */
 FLinearColor FLinearColor::LinearRGBToHSV() const
@@ -344,7 +353,7 @@ FLinearColor FLinearColor::LerpUsingHSV( const FLinearColor& From, const FLinear
 FLinearColor FLinearColor::MakeRandomColor()
 {
 	const uint8 Hue = (uint8)(FMath::FRand()*255.f);
-	return FLinearColor::FGetHSV(Hue, 0, 255);
+	return FLinearColor::MakeFromHSV8(Hue, 0, 255);
 }
 
 FColor FColor::MakeRandomColor()
diff --git a/Engine/Source/Runtime/Core/Public/Math/Color.h b/Engine/Source/Runtime/Core/Public/Math/Color.h
--- a/Engine/Source/Runtime/Core/Public/Math/Color.h
+++ b/Engine/Source/Runtime/Core/Public/Math/Color.h
@@ -258,10 +258,13 @@ struct FLinearColor
 		return NewCopy;
 	}
 
+	DEPRECATED(4.22, "FGetHSV doesn't perform a valid HSV conversion, use MakeFromHSV8 instead")
+	static CORE_API FLinearColor FGetHSV(uint8 H,uint8 S,uint8 V);
+
 	/**
 	 * Converts byte hue-saturation-brightness to floating point red-green-blue.
 	 */
-	static CORE_API FLinearColor FGetHSV(uint8 H,uint8 S,uint8 V);
+	static CORE_API FLinearColor MakeFromHSV8(uint8 H, uint8 S, uint8 V);
 
 	/**
 	* Makes a random but quite nice color.
diff --git a/Engine/Source/Runtime/Engine/Private/SkeletalMesh.cpp b/Engine/Source/Runtime/Engine/Private/SkeletalMesh.cpp
--- a/Engine/Source/Runtime/Engine/Private/SkeletalMesh.cpp
+++ b/Engine/Source/Runtime/Engine/Private/SkeletalMesh.cpp
@@ -3916,7 +3916,7 @@ void FSkeletalMeshSceneProxy::DebugDrawSkeleton(int32 ViewIndex, FMeshElementCol
 	{
 		FRandomStream Stream((int32)InUID);
 		const uint8 Hue = (uint8)(Stream.FRand()*255.f);
-		return FLinearColor::FGetHSV(Hue, 0, 255);
+		return FLinearColor::MakeFromHSV8(Hue, 0, 255);
 	};
 
 	FPrimitiveDrawInterface* PDI = Collector.GetPDI(ViewIndex);
diff --git a/Engine/Source/Runtime/Renderer/Private/ShadowRendering.cpp b/Engine/Source/Runtime/Renderer/Private/ShadowRendering.cpp
--- a/Engine/Source/Runtime/Renderer/Private/ShadowRendering.cpp
+++ b/Engine/Source/Runtime/Renderer/Private/ShadowRendering.cpp
@@ -1132,7 +1132,7 @@ void FProjectedShadowInfo::RenderFrustumWireframe(FPrimitiveDrawInterface* PDI)
 	}
 	else
 	{
-		Color = FLinearColor::FGetHSV(( ( SubjectPrimitiveId + LightSceneInfo->Id ) * 31 ) & 255, 0, 255).ToFColor(true);
+		Color = FLinearColor::MakeFromHSV8(( ( SubjectPrimitiveId + LightSceneInfo->Id ) * 31 ) & 255, 0, 255).ToFColor(true);
 	}
 
 	// Render the wireframe for the frustum derived from ReceiverMatrix.
