DR - Added minimum aggregation in stat system.
More reliable & useful for timing GPU

diff --git a/Engine/Source/Runtime/Core/Private/Stats/StatsCommand.cpp b/Engine/Source/Runtime/Core/Private/Stats/StatsCommand.cpp
--- a/Engine/Source/Runtime/Core/Private/Stats/StatsCommand.cpp
+++ b/Engine/Source/Runtime/Core/Private/Stats/StatsCommand.cpp
@@ -1383,17 +1383,17 @@ struct FHUDGroupManager
 				SCOPE_CYCLE_COUNTER(STAT_Accumulate);
 				const FHudFrame& Frame = FrameIt.Value();
 
-				AggregatedHierarchyHistory.MergeAddAndMax( Frame.HierarchyInclusive );
+				AggregatedHierarchyHistory.MergeAddAndMinMax( Frame.HierarchyInclusive );
 
-				FComplexStatUtils::MergeAddAndMaxArray( AggregatedFlatHistory, Frame.InclusiveAggregate, EComplexStatField::IncSum, EComplexStatField::IncMax );
+				FComplexStatUtils::MergeAddAndMinMaxArray( AggregatedFlatHistory, Frame.InclusiveAggregate, EComplexStatField::IncSum, EComplexStatField::IncMax, EComplexStatField::IncMin );
 
 				for (TMap<FName, TArray<FStatMessage>>::TConstIterator It(Frame.InclusiveAggregateThreadBreakdown); It; ++It)
 				{
-					FComplexStatUtils::MergeAddAndMaxArray( AggregatedFlatHistoryThreadBreakdown.FindChecked(It.Key()), It.Value(), EComplexStatField::IncSum, EComplexStatField::IncMax );
+					FComplexStatUtils::MergeAddAndMinMaxArray( AggregatedFlatHistoryThreadBreakdown.FindChecked(It.Key()), It.Value(), EComplexStatField::IncSum, EComplexStatField::IncMax, EComplexStatField::IncMin);
 				}
 
-				FComplexStatUtils::MergeAddAndMaxArray( AggregatedFlatHistory, Frame.ExclusiveAggregate, EComplexStatField::ExcSum, EComplexStatField::ExcMax );
-				FComplexStatUtils::MergeAddAndMaxArray( AggregatedNonStackStatsHistory, Frame.NonStackStats, EComplexStatField::IncSum, EComplexStatField::IncMax );
+				FComplexStatUtils::MergeAddAndMinMaxArray( AggregatedFlatHistory, Frame.ExclusiveAggregate, EComplexStatField::ExcSum, EComplexStatField::ExcMax, EComplexStatField::ExcMin);
+				FComplexStatUtils::MergeAddAndMinMaxArray( AggregatedNonStackStatsHistory, Frame.NonStackStats, EComplexStatField::IncSum, EComplexStatField::IncMax, EComplexStatField::IncMin);
 			}
 
 			// Divide stats to get average values.
diff --git a/Engine/Source/Runtime/Core/Private/Stats/StatsData.cpp b/Engine/Source/Runtime/Core/Private/Stats/StatsData.cpp
--- a/Engine/Source/Runtime/Core/Private/Stats/StatsData.cpp
+++ b/Engine/Source/Runtime/Core/Private/Stats/StatsData.cpp
@@ -500,16 +500,16 @@ FComplexRawStatStackNode::FComplexRawStatStackNode( const FRawStatStackNode& Oth
 	}
 }
 
-void FComplexRawStatStackNode::MergeAddAndMax( const FRawStatStackNode& Other )
+void FComplexRawStatStackNode::MergeAddAndMinMax( const FRawStatStackNode& Other )
 {
-	FComplexStatUtils::AddAndMax( ComplexStat, Other.Meta, EComplexStatField::IncSum, EComplexStatField::IncMax );
+	FComplexStatUtils::AddAndMinMax( ComplexStat, Other.Meta, EComplexStatField::IncSum, EComplexStatField::IncMax, EComplexStatField::IncMin );
 
 	for (auto It = Other.Children.CreateConstIterator(); It; ++It)
 	{
 		FComplexRawStatStackNode* Child = Children.FindRef(It.Key());
 		if (Child)
 		{
-			Child->MergeAddAndMax(*It.Value());
+			Child->MergeAddAndMinMax(*It.Value());
 		}
 		else
 		{
@@ -2133,14 +2133,23 @@ FString FStatsUtils::ToEscapedFString(const TCHAR* Source)
 }
 
 
-void FComplexStatUtils::AddAndMax( FComplexStatMessage& Dest, const FStatMessage& Item, EComplexStatField::Type SumIndex, EComplexStatField::Type MaxIndex )
+void FComplexStatUtils::AddAndMinMax( FComplexStatMessage& Dest, const FStatMessage& Item, EComplexStatField::Type SumIndex, EComplexStatField::Type MaxIndex, EComplexStatField::Type MinIndex )
 {
 	check(Dest.NameAndInfo.GetRawName() == Item.NameAndInfo.GetRawName());
 
 	// Copy the data type from the other stack node.
 	if( Dest.NameAndInfo.GetField<EStatDataType>() == EStatDataType::ST_None )
 	{
-		Dest.NameAndInfo.SetField<EStatDataType>( Item.NameAndInfo.GetField<EStatDataType>() );
+		const EStatDataType::Type ItemStatDataType = Item.NameAndInfo.GetField<EStatDataType>();
+		Dest.NameAndInfo.SetField<EStatDataType>(ItemStatDataType);
+		if( ItemStatDataType == EStatDataType::ST_int64 )
+		{
+			Dest.GetValue_int64(MinIndex) = INT64_MAX; 
+		}
+		else if( ItemStatDataType == EStatDataType::ST_double )
+		{
+			Dest.GetValue_double(MinIndex) = DBL_MAX;
+		}
 	}
 
 	const EStatDataType::Type StatDataType = Dest.NameAndInfo.GetField<EStatDataType>();
@@ -2158,16 +2167,18 @@ void FComplexStatUtils::AddAndMax( FComplexStatMessage& Dest, const FStatMessage
 		}
 	}
 
-	// Maximum time.
+	// Maximum/Minimum time.
 	if( StatDataType != EStatDataType::ST_None && StatDataType != EStatDataType::ST_FName)
 	{
 		if( StatDataType == EStatDataType::ST_int64 )
 		{
 			FStatsUtils::StatOpMaxVal_Int64( Dest.NameAndInfo, Dest.GetValue_int64(MaxIndex), Item.GetValue_int64() );
+			FStatsUtils::StatOpMinVal_Int64( Dest.NameAndInfo, Dest.GetValue_int64(MinIndex), Item.GetValue_int64() );
 		}
 		else if( StatDataType == EStatDataType::ST_double )
 		{
 			Dest.GetValue_double(MaxIndex) = FMath::Max<double>(Dest.GetValue_double(MaxIndex), Item.GetValue_double());
+			Dest.GetValue_double(MinIndex) = FMath::Min<double>(Dest.GetValue_double(MinIndex), Item.GetValue_double());
 		}
 	}
 }
@@ -2205,7 +2216,7 @@ void FComplexStatUtils::DivideStat( FComplexStatMessage& Dest, uint32 Div, EComp
 	}
 }
 
-void FComplexStatUtils::MergeAddAndMaxArray( TArray<FComplexStatMessage>& Dest, const TArray<FStatMessage>& Source, EComplexStatField::Type SumIndex, EComplexStatField::Type MaxIndex )
+void FComplexStatUtils::MergeAddAndMinMaxArray( TArray<FComplexStatMessage>& Dest, const TArray<FStatMessage>& Source, EComplexStatField::Type SumIndex, EComplexStatField::Type MaxIndex, EComplexStatField::Type MinIndex )
 {
 	TMap<FName, int32> NameToIndex;
 	for( int32 Index = 0; Index < Dest.Num(); Index++ )
@@ -2217,7 +2228,7 @@ void FComplexStatUtils::MergeAddAndMaxArray( TArray<FComplexStatMessage>& Dest,
 	for( int32 Index = 0; Index < Source.Num(); Index++ )
 	{
 		const int32& DestIndex = NameToIndex.FindChecked( Source[Index].NameAndInfo.GetRawName() );
-		FComplexStatUtils::AddAndMax( Dest[DestIndex], Source[Index], SumIndex, MaxIndex );
+		FComplexStatUtils::AddAndMinMax( Dest[DestIndex], Source[Index], SumIndex, MaxIndex, MinIndex );
 	}
 }
 
diff --git a/Engine/Source/Runtime/Core/Public/Stats/Stats2.h b/Engine/Source/Runtime/Core/Public/Stats/Stats2.h
--- a/Engine/Source/Runtime/Core/Public/Stats/Stats2.h
+++ b/Engine/Source/Runtime/Core/Public/Stats/Stats2.h
@@ -1072,12 +1072,17 @@ struct EComplexStatField
 		IncAve,
 		/** Maximum inclusive time. */
 		IncMax,
+		/** Minimum inclusive time. */
+		IncMin,
 		/** Summed exclusive time. */
 		ExcSum,
 		/** Average exclusive time. */
 		ExcAve,
 		/** Maximum exclusive time. */
 		ExcMax,
+		/** Minimum exclusive time. */
+		ExcMin,
+
 		/** Number of enumerates. */
 		Num,
 	};
diff --git a/Engine/Source/Runtime/Core/Public/Stats/StatsData.h b/Engine/Source/Runtime/Core/Public/Stats/StatsData.h
--- a/Engine/Source/Runtime/Core/Public/Stats/StatsData.h
+++ b/Engine/Source/Runtime/Core/Public/Stats/StatsData.h
@@ -407,7 +407,7 @@ public:
 	}
 
 	/** Merges this stack with the other. */
-	void MergeAddAndMax( const FRawStatStackNode& Other );
+	void MergeAddAndMinMax( const FRawStatStackNode& Other );
 
 	/** Divides this stack by the specified value. */
 	void Divide(uint32 Div);
@@ -767,6 +767,20 @@ struct CORE_API FStatsUtils
 			Dest = FMath::Max<int64>(Dest,Other);
 		}
 	}
+	/** Finds a minimum for int64 based stat data. */
+	static void StatOpMinVal_Int64(const FStatNameAndInfo& DestNameAndInfo, int64& Dest, const int64& Other)
+	{
+		if (DestNameAndInfo.GetFlag(EStatMetaFlags::IsPackedCCAndDuration))
+		{
+			Dest = ToPackedCallCountDuration(
+				FMath::Min<uint32>(FromPackedCallCountDuration_CallCount(Dest), FromPackedCallCountDuration_CallCount(Other)),
+				FMath::Min<uint32>(FromPackedCallCountDuration_Duration(Dest), FromPackedCallCountDuration_Duration(Other)));
+		}
+		else
+		{
+			Dest = FMath::Min<int64>(Dest, Other);
+		}
+	}
 
 	/** Internal use, converts arbitrary string to and from an escaped notation for storage in an FName. **/
 	static FString ToEscapedFString(const TCHAR* Source);
@@ -804,13 +818,13 @@ struct CORE_API FStatsUtils
 struct FComplexStatUtils
 {
 	/** Accumulates a stat message into a complex stat message. */
-	static void AddAndMax( FComplexStatMessage& Dest, const FStatMessage& Item, EComplexStatField::Type SumIndex, EComplexStatField::Type MaxIndex );
+	static void AddAndMinMax( FComplexStatMessage& Dest, const FStatMessage& Item, EComplexStatField::Type SumIndex, EComplexStatField::Type MaxIndex, EComplexStatField::Type MinIndex);
 
 	/** Divides a complex stat message by the specified value. */
 	static void DivideStat( FComplexStatMessage& Dest, uint32 Div, EComplexStatField::Type SumIndex, EComplexStatField::Type DestIndex );
 
 	/** Merges two complex stat message arrays. */
-	static void MergeAddAndMaxArray( TArray<FComplexStatMessage>& Dest, const TArray<FStatMessage>& Source, EComplexStatField::Type SumIndex, EComplexStatField::Type MaxIndex );
+	static void MergeAddAndMinMaxArray( TArray<FComplexStatMessage>& Dest, const TArray<FStatMessage>& Source, EComplexStatField::Type SumIndex, EComplexStatField::Type MaxIndex, EComplexStatField::Type MinIndex);
 
 	/** Divides a complex stat array by the specified value. */
 	static void DiviveStatArray( TArray<FComplexStatMessage>& Dest, uint32 Div, EComplexStatField::Type SumIndex, EComplexStatField::Type DestIndex );
diff --git a/Engine/Source/Runtime/Engine/Private/StatsRender2.cpp b/Engine/Source/Runtime/Engine/Private/StatsRender2.cpp
--- a/Engine/Source/Runtime/Engine/Private/StatsRender2.cpp
+++ b/Engine/Source/Runtime/Engine/Private/StatsRender2.cpp
@@ -488,6 +488,11 @@ static int32 RenderCounterHeadings( class FCanvas* Canvas, const int32 X, const
 
 	// Draw the max column label.
 	RightJustify(Canvas, CurrX, Y, TEXT("Max"), Globals.HeadingColor);
+	CurrX += Globals.InterColumnOffset;
+
+	// Draw the max column label.
+	RightJustify(Canvas, CurrX, Y, TEXT("Min"), Globals.HeadingColor);
+
 	return Globals.GetFontHeight() + (Globals.GetFontHeight() / 3);
 }
 
@@ -628,6 +633,22 @@ static int32 RenderCounter( const FGameThreadStatsData& ViewData, const FComplex
 		const FString ValueFormatted = FormatStatValueInt64( All.GetValue_int64( EComplexStatField::IncMax ) );
 		RightJustify(Canvas, CurrX, Y, *ValueFormatted, Globals.StatColor);
 	}
+
+	CurrX += Globals.InterColumnOffset;
+
+	// Append the minimum.
+	if (All.NameAndInfo.GetField<EStatDataType>() == EStatDataType::ST_double)
+	{
+		const FString ValueFormatted = FormatStatValueFloat(All.GetValue_double(EComplexStatField::IncMin));
+		RightJustify(Canvas, CurrX, Y, *ValueFormatted, Globals.StatColor);
+	}
+	else if (All.NameAndInfo.GetField<EStatDataType>() == EStatDataType::ST_int64)
+	{
+		const FString ValueFormatted = FormatStatValueInt64(All.GetValue_int64(EComplexStatField::IncMin));
+		RightJustify(Canvas, CurrX, Y, *ValueFormatted, Globals.StatColor);
+	}
+
+
 	return Globals.GetFontHeight();
 }
 
