Fix a clang compile error

diff --git a/Engine/Source/Runtime/Core/Public/Templates/IsValidVariadicFunctionArg.h b/Engine/Source/Runtime/Core/Public/Templates/IsValidVariadicFunctionArg.h
--- a/Engine/Source/Runtime/Core/Public/Templates/IsValidVariadicFunctionArg.h
+++ b/Engine/Source/Runtime/Core/Public/Templates/IsValidVariadicFunctionArg.h
@@ -3,11 +3,12 @@
 #pragma once
 
 #include "CoreTypes.h"
+#include "IsEnum.h"
 
 /**
  * Tests if a type is a valid argument to a variadic function, e.g. printf.
  */
-template <typename T>
+template <typename T, bool = TIsEnum<T>::Value>
 struct TIsValidVariadicFunctionArg
 {
 private:
@@ -30,3 +31,8 @@ public:
 	enum { Value = sizeof(Tester(DeclValT())) == sizeof(uint32) };
 };
 
+template <typename T>
+struct TIsValidVariadicFunctionArg<T, true>
+{
+	enum { Value = true };
+};
