DR - Missing enums in PixelFormat
#rnx

diff --git a/Engine/Source/Runtime/Core/Public/PixelFormat.h b/Engine/Source/Runtime/Core/Public/PixelFormat.h
--- a/Engine/Source/Runtime/Core/Public/PixelFormat.h
+++ b/Engine/Source/Runtime/Core/Public/PixelFormat.h
@@ -140,9 +140,12 @@ enum EPixelFormat
 	op(PF_BC6H) \
 	op(PF_BC7) \
 	op(PF_R8_UINT) \
+	op(PF_L8) \
+	op(PF_XGXR8) \
+	op(PF_R8G8B8A8_UINT) \
+	op(PF_R8G8B8A8_SNORM) \
 	op(PF_R16G16B16A16_UNORM) \
 	op(PF_R16G16B16A16_SNORM) \
-	op(PF_L8) \
 	op(PF_PLATFORM_HDR_0) \
 	op(PF_PLATFORM_HDR_1) \
 	op(PF_PLATFORM_HDR_2) \
