Updated GPU class codes for IOKit.

diff --git a/Engine/Source/Runtime/Core/Private/Mac/MacPlatformMisc.cpp b/Engine/Source/Runtime/Core/Private/Mac/MacPlatformMisc.cpp
--- a/Engine/Source/Runtime/Core/Private/Mac/MacPlatformMisc.cpp
+++ b/Engine/Source/Runtime/Core/Private/Mac/MacPlatformMisc.cpp
@@ -911,13 +911,13 @@ public:
 				CFMutableDictionaryRef ServiceInfo;
 				if(IORegistryEntryCreateCFProperties(ServiceEntry, &ServiceInfo, kCFAllocatorDefault, kNilOptions) == kIOReturnSuccess)
 				{
-					// GPUs are class-code 0x30000
+					// GPUs are class-code 0x30000 || 0x38000
 					static CFStringRef ClassCodeRef = CFSTR("class-code");
 					const CFDataRef ClassCode = (const CFDataRef)CFDictionaryGetValue(ServiceInfo, ClassCodeRef);
 					if(ClassCode && CFGetTypeID(ClassCode) == CFDataGetTypeID())
 					{
 						const uint32* ClassCodeValue = reinterpret_cast<const uint32*>(CFDataGetBytePtr(ClassCode));
-						if(ClassCodeValue && *ClassCodeValue == 0x30000)
+						if(ClassCodeValue && (*ClassCodeValue == 0x30000 || *ClassCodeValue == 0x38000))
 						{
 							FMacPlatformMisc::FGPUDescriptor Desc;
 							
@@ -972,13 +972,13 @@ public:
 								CFMutableDictionaryRef ServiceInfo;
 								if(IORegistryEntryCreateCFProperties(ParentEntry, &ServiceInfo, kCFAllocatorDefault, kNilOptions) == kIOReturnSuccess)
 								{
-									// GPUs are class-code 0x30000
+									// GPUs are class-code 0x30000 || 0x38000
 									static CFStringRef ClassCodeRef = CFSTR("class-code");
 									const CFDataRef ClassCode = (const CFDataRef)CFDictionaryGetValue(ServiceInfo, ClassCodeRef);
 									if(ClassCode && CFGetTypeID(ClassCode) == CFDataGetTypeID())
 									{
 										const uint32* ClassCodeValue = reinterpret_cast<const uint32*>(CFDataGetBytePtr(ClassCode));
-										if(ClassCodeValue && *ClassCodeValue == 0x30000)
+										if(ClassCodeValue && (*ClassCodeValue == 0x30000 || *ClassCodeValue == 0x38000))
 										{
 											FScopeLock Lock(&Mutex);
 											
