UE-69300 //UE4/Main - Compile UE4Game Lumin - fatal error: 'android/asset_manager.h' file not found
-Wrapped a couple new things in USE_ANDROID_JNI

diff --git a/Engine/Source/Runtime/Core/Private/Android/AndroidMisc.cpp b/Engine/Source/Runtime/Core/Private/Android/AndroidMisc.cpp
--- a/Engine/Source/Runtime/Core/Private/Android/AndroidMisc.cpp
+++ b/Engine/Source/Runtime/Core/Private/Android/AndroidMisc.cpp
@@ -25,12 +25,12 @@
 #include "Misc/Parse.h"
 #include "Internationalization/Regex.h"
 
-#include <android/asset_manager.h>
 #include <android/log.h>
 #if USE_ANDROID_INPUT
 #include <android/keycodes.h>
 #endif
 #if USE_ANDROID_JNI
+#include <android/asset_manager.h>
 #include <cpu-features.h>
 #include <android_native_app_glue.h>
 #include "Templates/Function.h"
@@ -1072,6 +1072,7 @@ void FAndroidMisc::ShareURL(const FString& URL, const FText& Description, int32
 
 FString FAndroidMisc::LoadTextFileFromPlatformPackage(const FString& RelativePath)
 {
+#if USE_ANDROID_JNI
 	AAssetManager* AssetMgr = AndroidThunkCpp_GetAssetManager();
 	AAsset* asset = AAssetManager_open(AssetMgr, TCHAR_TO_UTF8(*RelativePath), AASSET_MODE_BUFFER);
 
@@ -1089,7 +1090,7 @@ FString FAndroidMisc::LoadTextFileFromPlatformPackage(const FString& RelativePat
 
 		return FString(ANSI_TO_TCHAR(TextContents.GetData()));
 	}
-
+#endif
 	return FString();
 }
 
