Trimming out unused FramePro function

diff --git a/Engine/Source/Runtime/Core/Private/FramePro/FramePro.cpp b/Engine/Source/Runtime/Core/Private/FramePro/FramePro.cpp
--- a/Engine/Source/Runtime/Core/Private/FramePro/FramePro.cpp
+++ b/Engine/Source/Runtime/Core/Private/FramePro/FramePro.cpp
@@ -6945,94 +6945,12 @@ using namespace FramePro;
 namespace FramePro
 {
 	//------------------------------------------------------------------------
+	// START EPIC
 	bool GetStackTrace(void** stack, int& stack_size, unsigned int& hash)
 	{
-		#if FRAMEPRO_PLATFORM_WIN
-			#if FRAMEPRO_USE_STACKWALK64
-
-				// get the context
-				CONTEXT context;
-				memset(&context, 0, sizeof(context));
-				RtlCaptureContext(&context);
-
-				// setup the stack frame
-				STACKFRAME64 stack_frame;
-				memset(&stack_frame, 0, sizeof(stack_frame));
-				stack_frame.AddrPC.Mode = AddrModeFlat;
-				stack_frame.AddrFrame.Mode = AddrModeFlat;
-				stack_frame.AddrStack.Mode = AddrModeFlat;
-				DWORD machine = IMAGE_FILE_MACHINE_IA64;
-				stack_frame.AddrPC.Offset = context.Rip;
-				stack_frame.AddrFrame.Offset = context.Rsp;
-				stack_frame.AddrStack.Offset = context.Rbp;
-				HANDLE thread = GetCurrentThread();
-
-				static HANDLE process = GetCurrentProcess();
-
-				stack_size = 0;
-				while (StackWalk64(
-					machine,
-					process,
-					thread,
-					&stack_frame,
-					&context,
-					NULL,
-					SymFunctionTableAccess64,
-					SymGetModuleBase64,
-					NULL) && stack_size < FRAMEPRO_STACK_TRACE_SIZE)
-				{
-					void* p = (void*)(stack_frame.AddrPC.Offset);
-					stack[stack_size++] = p;
-				}
-				hash = GetHash(stack, stack_size);
-			#elif FRAMEPRO_USE_RTLVIRTUALUNWIND
-				MemPro::VirtualUnwindStackWalk(stack, FRAMEPRO_STACK_TRACE_SIZE);
-				hash = GetHashAndStackSize(stack, stack_size);
-			#elif FRAMEPRO_USE_RTLCAPTURESTACKBACKTRACE
-				MemPro::RTLCaptureStackBackTrace(stack, FRAMEPRO_STACK_TRACE_SIZE, hash, stack_size);
-			#else
-				CaptureStackBackTrace(0, FRAMEPRO_STACK_TRACE_SIZE, stack, (PDWORD)&hash);
-				for (stack_size = 0; stack_size<FRAMEPRO_STACK_TRACE_SIZE; ++stack_size)
-					if (!stack[stack_size])
-						break;
-			#endif
-			return true;
-
-		#elif FRAMEPRO_PLATFORM_XBOX360
-			DmCaptureStackBackTrace(FRAMEPRO_STACK_TRACE_SIZE, stack);
-			hash = GetHashAndStackSize(stack, stack_size);
-			return true;
-
-		#elif FRAMEPRO_PLATFORM_XBOXONE
-			#if FRAMEPRO_USE_RTLCAPTURESTACKBACKTRACE
-					MemPro::RTLCapruteStackBackTrace(stack, FRAMEPRO_STACK_TRACE_SIZE, hash, stack_size);
-			#elif FRAMEPRO_USE_RTLVIRTUALUNWIND
-					MemPro::VirtualUnwindStackWalk(stack, FRAMEPRO_STACK_TRACE_SIZE);
-					hash = GetHashAndStackSize(stack, stack_size);
-			#else
-				#error
-			#endif
-			return true;
-
-		#elif MEMPRO_PLATFORM_PS4
-			SceDbgCallFrame frames[FRAMEPRO_STACK_TRACE_SIZE];
-			unsigned int frame_count = 0;
-			sceDbgBacktraceSelf(frames, sizeof(frames), &frame_count, SCE_DBG_BACKTRACE_MODE_DONT_EXCEED);
-
-			stack_size = frame_count;
-			for (unsigned int i = 0; i < frame_count; ++i)
-				stack[i] = (void*)frames[i].pc;
-			return true;
-
-		#elif MEMPRO_UNIX_BASED_PLATFORM
-			stack_size = backtrace(stack, FRAMEPRO_STACK_TRACE_SIZE);
-			hash = GetHashAndStackSize(stack, stack_size);
-			return true;
-
-		#else
-			return false;
-		#endif
+		return false;
 	}
+	// END EPIC
 }
 
 //------------------------------------------------------------------------
