Keep the page executable when patching SetConsoleCtrlCHandler to prevent access vioations on other threads whose instruction pointers are in kernel32.dll's text segment.

This fixes intermittent ShaderCompileWorker crash at startup due to a race condition betwen SetConsoleCtrlCHandler performing the patching and FCrashReportingThread code trying to execute code on the same page.

(peafour-eddy of //UE4/Dev-Core/[REDACTED_HANDLE] by [REDACTED_PERSON])

diff --git a/Engine/Source/Runtime/Core/Private/Windows/WindowsPlatformMisc.cpp b/Engine/Source/Runtime/Core/Private/Windows/WindowsPlatformMisc.cpp
--- a/Engine/Source/Runtime/Core/Private/Windows/WindowsPlatformMisc.cpp
+++ b/Engine/Source/Runtime/Core/Private/Windows/WindowsPlatformMisc.cpp
@@ -659,7 +659,7 @@ void FWindowsPlatformMisc::SetGracefulTerminationHandler()
 		0xc3							// ret
 	};
 	DWORD PrevProtection;
-	if (VirtualProtect(SetCtrlCProc, sizeof(Patch), PAGE_READWRITE, &PrevProtection))
+	if (VirtualProtect(SetCtrlCProc, sizeof(Patch), PAGE_EXECUTE_READWRITE, &PrevProtection))
 	{
 		memcpy(SetCtrlCProc, Patch, sizeof(Patch));
 		VirtualProtect(SetCtrlCProc, sizeof(Patch), PrevProtection, &PrevProtection);
@@ -2661,4 +2661,4 @@ IPlatformChunkInstall* FWindowsPlatformMisc::GetPlatformChunkInstall()
 	}
 
 	return ChunkInstall;
-}
\ No newline at end of file
+}
