Linux: removed an obsolete cmline arg and added logs for build settings.

#rnx

diff --git a/Engine/Source/Runtime/Core/Private/Unix/UnixPlatformMemory.cpp b/Engine/Source/Runtime/Core/Private/Unix/UnixPlatformMemory.cpp
--- a/Engine/Source/Runtime/Core/Private/Unix/UnixPlatformMemory.cpp
+++ b/Engine/Source/Runtime/Core/Private/Unix/UnixPlatformMemory.cpp
@@ -183,14 +183,6 @@ class FMalloc* FUnixPlatformMemory::BaseAllocator()
 					int32 Max = FCStringAnsi::Atoi(Cmd + sizeof(FileMapCacheCmd) - 1);
 					GMaxNumberFileMappingCache = FMath::Clamp(Max, 0, MaximumAllowedMaxNumFileMappingCache);
 				}
-#if UE_SERVER
-				const char MemPoolTableScaleCmd[] = "-mempoolscale=";
-				if (const char* Cmd = FCStringAnsi::Stristr(Arg, MemPoolTableScaleCmd))
-				{
-					float MemPoolScale = FCStringAnsi::Atof(Cmd + sizeof(MemPoolTableScaleCmd) - 1);
-					GPoolTableScale = FMath::Max(MemPoolScale, 0.0f);
-				}
-#endif
 
 #if UE_USE_MALLOC_REPLAY_PROXY
 				if (FCStringAnsi::Stricmp(Arg, "-mallocsavereplay") == 0)
diff --git a/Engine/Source/Runtime/Core/Private/Unix/UnixPlatformMisc.cpp b/Engine/Source/Runtime/Core/Private/Unix/UnixPlatformMisc.cpp
--- a/Engine/Source/Runtime/Core/Private/Unix/UnixPlatformMisc.cpp
+++ b/Engine/Source/Runtime/Core/Private/Unix/UnixPlatformMisc.cpp
@@ -39,6 +39,7 @@
 #include "HAL/ThreadHeartBeat.h"
 
 #include "FramePro/FrameProProfiler.h"
+#include "BuildSettings.h"
 
 // define for glibc 2.12.2 and lower (which is shipped with CentOS 6.x and which we target by default)
 #define __secure_getenv getenv
@@ -199,6 +200,7 @@ void FUnixPlatformMisc::PlatformInit()
 		PLATFORM_COMPILER_OPTIMIZATION_PG ? TEXT("yes") : TEXT("no"),
 		PLATFORM_COMPILER_OPTIMIZATION_PG_PROFILING ? TEXT("yes") : TEXT("no")
 		);
+	UE_LOG(LogInit, Log, TEXT(" - This is %s build."), BuildSettings::IsLicenseeVersion() ? TEXT("a licensee") : TEXT("an internal"));
 
 	FPlatformTime::PrintCalibrationLog();
 
@@ -211,13 +213,6 @@ void FUnixPlatformMisc::PlatformInit()
 	UE_LOG(LogInit, Log, TEXT(" -ksmmergeall - marks all mmap'd memory pages suitable for KSM (%s)"), GKSMMergeAllPages ? TEXT("ON") : TEXT("OFF"));
 	UE_LOG(LogInit, Log, TEXT(" -preloadmodulesymbols - Loads the main module symbols file into memory (%s)"), bPreloadedModuleSymbolFile ? TEXT("ON") : TEXT("OFF"));
 
-#if UE_SERVER
-	// Defined in UnixPlatformMemory, allows changing the number of 64k buckets used for the memory pool
-	extern float GPoolTableScale;
-
-	UE_LOG(LogInit, Log, TEXT(" -mempoolscale=SCALE - Scale the memory pool by (%lf)"), GPoolTableScale);
-#endif
-
 	// [RCL] FIXME: this should be printed in specific modules, if at all
 	UE_LOG(LogInit, Log, TEXT(" -httpproxy=ADDRESS:PORT - redirects HTTP requests to a proxy (only supported if compiled with libcurl)"));
 	UE_LOG(LogInit, Log, TEXT(" -reuseconn - allow libcurl to reuse HTTP connections (only matters if compiled with libcurl)"));
