RemoveAll should accept const objects
#rnx

diff --git a/Engine/Source/Runtime/Core/Public/UObject/ScriptDelegates.h b/Engine/Source/Runtime/Core/Public/UObject/ScriptDelegates.h
--- a/Engine/Source/Runtime/Core/Public/UObject/ScriptDelegates.h
+++ b/Engine/Source/Runtime/Core/Public/UObject/ScriptDelegates.h
@@ -381,7 +381,7 @@ public:
 	 *
 	 * @param InObject The object to remove bindings for.
 	 */
-	void RemoveAll( UObject* Object )
+	void RemoveAll( const UObject* Object )
 	{
 		for (int32 BindingIndex = InvocationList.Num() - 1; BindingIndex >= 0; --BindingIndex)
 		{
