Don't send ensures in non-editor builds (UE-66271).

(Edigrating CL 4861209 to 4.22)

diff --git a/Engine/Source/Runtime/Core/Private/Mac/MacPlatformMisc.cpp b/Engine/Source/Runtime/Core/Private/Mac/MacPlatformMisc.cpp
--- a/Engine/Source/Runtime/Core/Private/Mac/MacPlatformMisc.cpp
+++ b/Engine/Source/Runtime/Core/Private/Mac/MacPlatformMisc.cpp
@@ -28,6 +28,7 @@
 #include "Internationalization/Culture.h"
 #include "Modules/ModuleManager.h"
 #include "GenericPlatform/GenericPlatformChunkInstall.h"
+#include "BuildSettings.h"
 
 #include "Apple/PreAppleSystemHeaders.h"
 #include <dlfcn.h>
@@ -1777,6 +1778,12 @@ void FMacCrashContext::GenerateCrashInfoAndLaunchReporter() const
 	bool bSendUnattendedBugReports = true;
 	GConfig->GetBool(TEXT("/Script/UnrealEd.CrashReportsPrivacySettings"), TEXT("bSendUnattendedBugReports"), bSendUnattendedBugReports, GEditorSettingsIni);
 
+	if (BuildSettings::IsLicenseeVersion() && !UE_EDITOR)
+	{
+		// do not send unattended reports in licensees' builds except for the editor, where it is governed by the above setting
+		bSendUnattendedBugReports = false;
+	}
+
 	if (GMacAppInfo.bIsUnattended && !bSendUnattendedBugReports)
 	{
 		bCanRunCrashReportClient = false;
@@ -1852,6 +1859,12 @@ void FMacCrashContext::GenerateEnsureInfoAndLaunchReporter() const
 	bool bSendUnattendedBugReports = true;
 	GConfig->GetBool(TEXT("/Script/UnrealEd.CrashReportsPrivacySettings"), TEXT("bSendUnattendedBugReports"), bSendUnattendedBugReports, GEditorSettingsIni);
 
+	if (BuildSettings::IsLicenseeVersion() && !UE_EDITOR)
+	{
+		// do not send unattended reports in licensees' builds except for the editor, where it is governed by the above setting
+		bSendUnattendedBugReports = false;
+	}
+
 	if(GMacAppInfo.bIsUnattended && !bSendUnattendedBugReports)
 	{
 		bCanRunCrashReportClient = false;
diff --git a/Engine/Source/Runtime/Core/Private/Unix/UnixPlatformCrashContext.cpp b/Engine/Source/Runtime/Core/Private/Unix/UnixPlatformCrashContext.cpp
--- a/Engine/Source/Runtime/Core/Private/Unix/UnixPlatformCrashContext.cpp
+++ b/Engine/Source/Runtime/Core/Private/Unix/UnixPlatformCrashContext.cpp
@@ -28,6 +28,7 @@
 #include "HAL/ExceptionHandling.h"
 #include "Stats/Stats.h"
 #include "HAL/ThreadHeartBeat.h"
+#include "BuildSettings.h"
 
 extern CORE_API bool GIsGPUCrashed;
 
@@ -418,6 +419,12 @@ void FUnixCrashContext::GenerateCrashInfoAndLaunchReporter(bool bReportingNonCra
 	bool bSendUnattendedBugReports = false;
 	GConfig->GetBool(TEXT("/Script/UnrealEd.CrashReportsPrivacySettings"), TEXT("bSendUnattendedBugReports"), bSendUnattendedBugReports, GEditorSettingsIni);
 
+	if (BuildSettings::IsLicenseeVersion() && !UE_EDITOR)
+	{
+		// do not send unattended reports in licensees' builds except for the editor, where it is governed by the above setting
+		bSendUnattendedBugReports = false;
+	}
+
 	bool bSkipCRC = bUnattended && !bSendUnattendedBugReports;
 
 	if (!bSkipCRC)
diff --git a/Engine/Source/Runtime/Core/Private/Windows/WindowsPlatformCrashContext.cpp b/Engine/Source/Runtime/Core/Private/Windows/WindowsPlatformCrashContext.cpp
--- a/Engine/Source/Runtime/Core/Private/Windows/WindowsPlatformCrashContext.cpp
+++ b/Engine/Source/Runtime/Core/Private/Windows/WindowsPlatformCrashContext.cpp
@@ -26,7 +26,7 @@
 #include "Templates/UniquePtr.h"
 #include "Misc/OutputDeviceArchiveWrapper.h"
 #include "HAL/ThreadManager.h"
-
+#include "BuildSettings.h"
 #include <strsafe.h>
 #include <dbghelp.h>
 #include <Shlwapi.h>
@@ -314,6 +314,12 @@ int32 ReportCrashUsingCrashReportClient(FWindowsPlatformCrashContext& InContext,
 	bool bSendUnattendedBugReports = true;
 	GConfig->GetBool(TEXT("/Script/UnrealEd.CrashReportsPrivacySettings"), TEXT("bSendUnattendedBugReports"), bSendUnattendedBugReports, GEditorSettingsIni);
 
+	if (BuildSettings::IsLicenseeVersion() && !UE_EDITOR)
+	{
+		// do not send unattended reports in licensees' builds except for the editor, where it is governed by the above setting
+		bSendUnattendedBugReports = false;
+	}
+
 	if (bNoDialog && !bSendUnattendedBugReports)
 	{
 		bCanRunCrashReportClient = false;
