Adding native android trace support.

diff --git a/Engine/Source/Runtime/Core/Private/Android/AndroidMisc.cpp b/Engine/Source/Runtime/Core/Private/Android/AndroidMisc.cpp
--- a/Engine/Source/Runtime/Core/Private/Android/AndroidMisc.cpp
+++ b/Engine/Source/Runtime/Core/Private/Android/AndroidMisc.cpp
@@ -64,6 +64,17 @@ static FAutoConsoleVariableRef CAndroidLowPowerBatteryThreshold(
 
 #if STATS || ENABLE_STATNAMEDEVENTS
 int32 FAndroidMisc::TraceMarkerFileDescriptor = -1;
+
+typedef void (*ATrace_beginSection_Type) (const char* sectionName);
+typedef void (*ATrace_endSection_Type) (void);
+typedef bool (*ATrace_isEnabled_Type) (void);
+
+static ATrace_beginSection_Type ATrace_beginSection = NULL;
+static ATrace_endSection_Type ATrace_endSection = NULL;
+static ATrace_isEnabled_Type ATrace_isEnabled = NULL;
+
+static bool bUseNativeSystrace = false;
+
 #endif
 
 // run time compatibility information
@@ -361,18 +372,40 @@ void FAndroidMisc::PlatformInit()
 	AndroidSetupDefaultThreadAffinity();
 
 #if (STATS || ENABLE_STATNAMEDEVENTS)
-	if (FParse::Param(FCommandLine::Get(), TEXT("enablesystrace")))
+	//Loading NDK libandroid.so atrace functions, available in the android libraries way before NDK headers.
+	void* const LibAndroid = dlopen("libandroid.so", RTLD_NOW | RTLD_LOCAL);
+	if (Lib != nullptr)
 	{
-		GAndroidTraceMarkersEnabled = 1;
+		// Retrieve function pointers from shared object.
+		ATrace_beginSection = reinterpret_cast<ATrace_beginSection_Type>(dlsym(LibAndroid, "ATrace_beginSection"));
+		ATrace_endSection = reinterpret_cast<ATrace_endSection_Type>(dlsym(LibAndroid, "ATrace_endSection"));
+		ATrace_isEnabled = 	reinterpret_cast<ATrace_isEnabled_Type>(dlsym(LibAndroid, "ATrace_isEnabled"));
 	}
 
-	if (GAndroidTraceMarkersEnabled)
+	if (!ATrace_beginSection || !ATrace_endSection || !ATrace_isEnabled)
 	{
-		StartTraceMarkers();
-	}
+		UE_LOG(LogAndroid, Warning, TEXT("Failed to use native systrace functionality."));
+		ATrace_beginSection = nullptr;
+		ATrace_endSection = nullptr;
+		ATrace_isEnabled = nullptr;
 
-	// Watch for CVar update
-	CAndroidTraceMarkersEnabled->SetOnChangedCallback(FConsoleVariableDelegate::CreateStatic(&UpdateTraceMarkersEnable));
+		if (FParse::Param(FCommandLine::Get(), TEXT("enablesystrace")))
+		{
+			GAndroidTraceMarkersEnabled = 1;
+		}
+
+		if (GAndroidTraceMarkersEnabled)
+		{
+			StartTraceMarkers();
+		}
+
+		// Watch for CVar update
+		CAndroidTraceMarkersEnabled->SetOnChangedCallback(FConsoleVariableDelegate::CreateStatic(&UpdateTraceMarkersEnable));
+	}
+	else
+	{
+		bUseNativeSystrace = true;
+	}
 #endif
 
 #if USE_ANDROID_JNI
@@ -1879,11 +1912,18 @@ void FAndroidMisc::BeginNamedEventFrame()
 
 static void WriteTraceMarkerEvent(const ANSICHAR* Text, int32 TraceMarkerFileDescriptor)
 {
-	const int MAX_TRACE_EVENT_LENGTH = 256;
+	if (bUseNativeSystrace)
+	{
+		ATrace_beginSection(Text);
+	}
+	else
+	{
+		const int MAX_TRACE_EVENT_LENGTH = 256;
+		ANSICHAR EventBuffer[MAX_TRACE_EVENT_LENGTH];
+		int EventLength = snprintf(EventBuffer, MAX_TRACE_EVENT_LENGTH, "B|%d|%s", getpid(), Text);
 
-	ANSICHAR EventBuffer[MAX_TRACE_EVENT_LENGTH];
-	int EventLength = snprintf(EventBuffer, MAX_TRACE_EVENT_LENGTH, "B|%d|%s", getpid(), Text);
-	write(TraceMarkerFileDescriptor, EventBuffer, EventLength);
+		write(TraceMarkerFileDescriptor, EventBuffer, EventLength);
+	}
 }
 
 void FAndroidMisc::BeginNamedEvent(const struct FColor& Color, const TCHAR* Text)
@@ -1891,7 +1931,7 @@ void FAndroidMisc::BeginNamedEvent(const struct FColor& Color, const TCHAR* Text
 #if FRAMEPRO_ENABLED
 	FFrameProProfiler::PushEvent(Text);
 #endif // FRAMEPRO_ENABLED
-	if (TraceMarkerFileDescriptor == -1)
+	if (bUseNativeSystrace ? !ATrace_isEnabled() : TraceMarkerFileDescriptor == -1)
 	{
 		return;
 	}
@@ -1918,7 +1958,7 @@ void FAndroidMisc::BeginNamedEvent(const struct FColor& Color, const ANSICHAR* T
 #if FRAMEPRO_ENABLED
 	FFrameProProfiler::PushEvent(Text);
 #endif // FRAMEPRO_ENABLED
-	if (TraceMarkerFileDescriptor == -1)
+	if (bUseNativeSystrace ? !ATrace_isEnabled() : TraceMarkerFileDescriptor == -1)
 	{
 		return;
 	}
@@ -1931,13 +1971,20 @@ void FAndroidMisc::EndNamedEvent()
 #if FRAMEPRO_ENABLED
 	FFrameProProfiler::PopEvent();
 #endif // FRAMEPRO_ENABLED
-	if (TraceMarkerFileDescriptor == -1)
+	if (bUseNativeSystrace ? !ATrace_isEnabled() : TraceMarkerFileDescriptor == -1)
 	{
 		return;
 	}
 
-	const ANSICHAR EventTerminatorChar = 'E';
-	write(TraceMarkerFileDescriptor, &EventTerminatorChar, 1);
+	if (bUseNativeSystrace)
+	{
+		ATrace_endSection();
+	}
+	else
+	{
+		const ANSICHAR EventTerminatorChar = 'E';
+		write(TraceMarkerFileDescriptor, &EventTerminatorChar, 1);
+	}
 }
 
 void FAndroidMisc::CustomNamedStat(const TCHAR* Text, float Value, const TCHAR* Graph, const TCHAR* Unit)
