Linux: update to v13 (llvm 7) toolchain in hopes to fix the linker crash.

#rnx

diff --git a/Engine/Source/Runtime/Core/Public/HAL/Platform.h b/Engine/Source/Runtime/Core/Public/HAL/Platform.h
--- a/Engine/Source/Runtime/Core/Public/HAL/Platform.h
+++ b/Engine/Source/Runtime/Core/Public/HAL/Platform.h
@@ -822,7 +822,7 @@ namespace TypeTests
 	static_assert((!TAreTypesEqual<ANSICHAR, WIDECHAR>::Value), "ANSICHAR and WIDECHAR should be different types.");
 	static_assert((!TAreTypesEqual<ANSICHAR, UCS2CHAR>::Value), "ANSICHAR and CHAR16 should be different types.");
 	static_assert((!TAreTypesEqual<WIDECHAR, UCS2CHAR>::Value), "WIDECHAR and CHAR16 should be different types.");
-	static_assert((TAreTypesEqual<TCHAR, ANSICHAR>::Value || TAreTypesEqual<TCHAR, WIDECHAR>::Value), "TCHAR should either be ANSICHAR or WIDECHAR.");
+	static_assert((TAreTypesEqual<TCHAR, ANSICHAR>::Value == true || TAreTypesEqual<TCHAR, WIDECHAR>::Value == true), "TCHAR should either be ANSICHAR or WIDECHAR.");
 
 	static_assert(sizeof(uint8) == 1, "BYTE type size test failed.");
 	static_assert(int32(uint8(-1)) == 0xFF, "BYTE type sign test failed.");
