Fix compile error.

diff --git a/Engine/Source/Runtime/Core/Private/Android/AndroidMisc.cpp b/Engine/Source/Runtime/Core/Private/Android/AndroidMisc.cpp
--- a/Engine/Source/Runtime/Core/Private/Android/AndroidMisc.cpp
+++ b/Engine/Source/Runtime/Core/Private/Android/AndroidMisc.cpp
@@ -374,7 +374,7 @@ void FAndroidMisc::PlatformInit()
 #if (STATS || ENABLE_STATNAMEDEVENTS)
 	//Loading NDK libandroid.so atrace functions, available in the android libraries way before NDK headers.
 	void* const LibAndroid = dlopen("libandroid.so", RTLD_NOW | RTLD_LOCAL);
-	if (Lib != nullptr)
+	if (LibAndroid != nullptr)
 	{
 		// Retrieve function pointers from shared object.
 		ATrace_beginSection = reinterpret_cast<ATrace_beginSection_Type>(dlsym(LibAndroid, "ATrace_beginSection"));
