Fix non-unity compile error.

diff --git a/Engine/Source/Runtime/Core/Public/GenericPlatform/StandardPlatformString.h b/Engine/Source/Runtime/Core/Public/GenericPlatform/StandardPlatformString.h
--- a/Engine/Source/Runtime/Core/Public/GenericPlatform/StandardPlatformString.h
+++ b/Engine/Source/Runtime/Core/Public/GenericPlatform/StandardPlatformString.h
@@ -9,6 +9,7 @@
 #include <wchar.h>
 #include <string.h>
 #include <stdlib.h>
+#include <stdio.h>
 
 #if (PLATFORM_APPLE || PLATFORM_UNIX || PLATFORM_HTML5 || PLATFORM_PS4 || PLATFORM_SWITCH) && !PLATFORM_TCHAR_IS_CHAR16
 
