Ensure that TUniquePtr::Reset doesn't destroy the object being set if it was already set to that object.

diff --git a/Engine/Source/Runtime/Core/Public/Templates/UniquePtr.h b/Engine/Source/Runtime/Core/Public/Templates/UniquePtr.h
--- a/Engine/Source/Runtime/Core/Public/Templates/UniquePtr.h
+++ b/Engine/Source/Runtime/Core/Public/Templates/UniquePtr.h
@@ -274,10 +274,13 @@ public:
 	 */
 	FORCEINLINE void Reset(T* InPtr = nullptr)
 	{
-		// We delete last, because we don't want odd side effects if the destructor of T relies on the state of this
-		T* OldPtr = Ptr;
-		Ptr = InPtr;
-		GetDeleter()(OldPtr);
+		if (Ptr != InPtr)
+		{
+			// We delete last, because we don't want odd side effects if the destructor of T relies on the state of this
+			T* OldPtr = Ptr;
+			Ptr = InPtr;
+			GetDeleter()(OldPtr);
+		}
 	}
 
 	/**
