#UE4 Large compile speed improvement. Moved ranged-based for implementations from global overloads to members of their respective classes. Compilers were spending a lot of time resolving these overloads in global space but do not have the same problem now that they are in the scope of the class. Measured to save 20-40% compile time on files in the engine that use range-based for.

diff --git a/Engine/Source/Runtime/AssetRegistry/Public/AssetDataTagMap.h b/Engine/Source/Runtime/AssetRegistry/Public/AssetDataTagMap.h
--- a/Engine/Source/Runtime/AssetRegistry/Public/AssetDataTagMap.h
+++ b/Engine/Source/Runtime/AssetRegistry/Public/AssetDataTagMap.h
@@ -716,18 +716,21 @@ private:
 		return Ar;
 	}
 
+public:
 	/** Range for iterator access - DO NOT USE DIRECTLY */
-	friend FAssetDataTagMap::RangedForConstIteratorType begin(const FAssetDataTagMapSharedView& View)
+	FAssetDataTagMap::RangedForConstIteratorType begin() const
 	{
-		return begin(View.GetMap());
+		return GetMap().begin();
 	}
 
 	/** Range for iterator access - DO NOT USE DIRECTLY */
-	friend FAssetDataTagMap::RangedForConstIteratorType end(const FAssetDataTagMapSharedView& View)
+	FAssetDataTagMap::RangedForConstIteratorType end() const
 	{
-		return end(View.GetMap());
+		return GetMap().end();
 	}
 
+private:
+
 	/** Pointer to map being wrapped, it is created on demand */
 	TSharedPtr<FAssetDataTagMap> Map;
 };
diff --git a/Engine/Source/Runtime/Core/Public/Algo/Impl/RangePointerType.h b/Engine/Source/Runtime/Core/Public/Algo/Impl/RangePointerType.h
--- a/Engine/Source/Runtime/Core/Public/Algo/Impl/RangePointerType.h
+++ b/Engine/Source/Runtime/Core/Public/Algo/Impl/RangePointerType.h
@@ -12,7 +12,7 @@ namespace AlgoImpl
 	template <typename RangeType>
 	struct TRangePointerType
 	{
-		using Type = decltype(&*begin(DeclVal<RangeType&>()));
+		using Type = decltype(&*DeclVal<RangeType&>().begin());
 	};
 
 	template <typename T, unsigned int N>
diff --git a/Engine/Source/Runtime/Core/Public/Containers/Array.h b/Engine/Source/Runtime/Core/Public/Containers/Array.h
--- a/Engine/Source/Runtime/Core/Public/Containers/Array.h
+++ b/Engine/Source/Runtime/Core/Public/Containers/Array.h
@@ -2352,22 +2352,22 @@ public:
 		typedef const ElementType* RangedForConstIteratorType;
 	#endif
 
-private:
+public:
 
 	/**
 	 * DO NOT USE DIRECTLY
 	 * STL-like iterators to enable range-based for loop support.
 	 */
 	#if TARRAY_RANGED_FOR_CHECKS
-		FORCEINLINE friend RangedForIteratorType      begin(      TArray& Array) { return RangedForIteratorType     (Array.ArrayNum, Array.GetData()); }
-		FORCEINLINE friend RangedForConstIteratorType begin(const TArray& Array) { return RangedForConstIteratorType(Array.ArrayNum, Array.GetData()); }
-		FORCEINLINE friend RangedForIteratorType      end  (      TArray& Array) { return RangedForIteratorType     (Array.ArrayNum, Array.GetData() + Array.Num()); }
-		FORCEINLINE friend RangedForConstIteratorType end  (const TArray& Array) { return RangedForConstIteratorType(Array.ArrayNum, Array.GetData() + Array.Num()); }
+		FORCEINLINE RangedForIteratorType      begin()       { return RangedForIteratorType     (ArrayNum, GetData()); }
+		FORCEINLINE RangedForConstIteratorType begin() const { return RangedForConstIteratorType(ArrayNum, GetData()); }
+		FORCEINLINE RangedForIteratorType      end  ()       { return RangedForIteratorType     (ArrayNum, GetData() + Num()); }
+		FORCEINLINE RangedForConstIteratorType end  () const { return RangedForConstIteratorType(ArrayNum, GetData() + Num()); }
 	#else
-		FORCEINLINE friend RangedForIteratorType      begin(      TArray& Array) { return Array.GetData(); }
-		FORCEINLINE friend RangedForConstIteratorType begin(const TArray& Array) { return Array.GetData(); }
-		FORCEINLINE friend RangedForIteratorType      end  (      TArray& Array) { return Array.GetData() + Array.Num(); }
-		FORCEINLINE friend RangedForConstIteratorType end  (const TArray& Array) { return Array.GetData() + Array.Num(); }
+		FORCEINLINE RangedForIteratorType      begin()       { return GetData(); }
+		FORCEINLINE RangedForConstIteratorType begin() const { return GetData(); }
+		FORCEINLINE RangedForIteratorType      end()         { return GetData() + Num(); }
+		FORCEINLINE RangedForConstIteratorType end() const   { return GetData() + Num(); }
 	#endif
 
 public:
diff --git a/Engine/Source/Runtime/Core/Public/Containers/ArrayView.h b/Engine/Source/Runtime/Core/Public/Containers/ArrayView.h
--- a/Engine/Source/Runtime/Core/Public/Containers/ArrayView.h
+++ b/Engine/Source/Runtime/Core/Public/Containers/ArrayView.h
@@ -508,13 +508,13 @@ public:
 		return FindByPredicate(Pred) != nullptr;
 	}
 
-private:
+public:
 	/**
 	 * DO NOT USE DIRECTLY
 	 * STL-like iterators to enable range-based for loop support.
 	 */
-	FORCEINLINE friend ElementType* begin(const TArrayView& Array) { return Array.GetData(); }
-	FORCEINLINE friend ElementType* end  (const TArrayView& Array) { return Array.GetData() + Array.Num(); }
+	FORCEINLINE ElementType* begin() const { return GetData(); }
+	FORCEINLINE ElementType* end  () const { return GetData() + Num(); }
 
 public:
 	/**
diff --git a/Engine/Source/Runtime/Core/Public/Containers/ChunkedArray.h b/Engine/Source/Runtime/Core/Public/Containers/ChunkedArray.h
--- a/Engine/Source/Runtime/Core/Public/Containers/ChunkedArray.h
+++ b/Engine/Source/Runtime/Core/Public/Containers/ChunkedArray.h
@@ -280,37 +280,38 @@ private:
 	typedef UE4ChunkedArray_Private::TChunkedArrayIterator<      FChunk,       ElementType, NumElementsPerChunk> FIterType;
 	typedef UE4ChunkedArray_Private::TChunkedArrayIterator<const FChunk, const ElementType, NumElementsPerChunk> FConstIterType;
 
-	friend FIterType begin(TChunkedArray& Array)
+public:
+	FIterType begin()
 	{
-		int32 Num = Array.NumElements;
-		FChunk** ChunkPtr = Array.Chunks.GetData();
-		FChunk** LastChunkPtr = Array.Chunks.GetData() + (Num ? Num - 1 : 0) / NumElementsPerChunk;
+		int32 Num = NumElements;
+		FChunk** ChunkPtr = Chunks.GetData();
+		FChunk** LastChunkPtr = Chunks.GetData() + (Num ? Num - 1 : 0) / NumElementsPerChunk;
 		return FIterType(ChunkPtr, LastChunkPtr, ChunkPtr ? (*ChunkPtr)->Elements : nullptr);
 	}
 
-	friend FConstIterType begin(const TChunkedArray& Array)
+	FConstIterType begin() const
 	{
-		int32 Num = Array.NumElements;
-		const FChunk** ChunkPtr = Array.Chunks.GetData();
-		const FChunk** LastChunkPtr = Array.Chunks.GetData() + (Num ? Num - 1 : 0) / NumElementsPerChunk;
+		int32 Num = NumElements;
+		const FChunk** ChunkPtr = Chunks.GetData();
+		const FChunk** LastChunkPtr = Chunks.GetData() + (Num ? Num - 1 : 0) / NumElementsPerChunk;
 		return FConstIterType(ChunkPtr, LastChunkPtr, ChunkPtr ? (*ChunkPtr)->Elements : nullptr);
 	}
 
-	friend FIterType end(TChunkedArray& Array)
+	FIterType end()
 	{
-		int32 Num = Array.NumElements;
+		int32 Num = NumElements;
 		bool bBeyondLastChunk = Num && (Num % NumElementsPerChunk) == 0;
-		FChunk** ChunkPtr = Array.Chunks.GetData() + (Num / NumElementsPerChunk) + (bBeyondLastChunk ? -1 : 0); // do not read off the end of the chunk array!
-		FChunk** LastChunkPtr = Array.Chunks.GetData() + (Num ? Num - 1 : 0) / NumElementsPerChunk;
+		FChunk** ChunkPtr = Chunks.GetData() + (Num / NumElementsPerChunk) + (bBeyondLastChunk ? -1 : 0); // do not read off the end of the chunk array!
+		FChunk** LastChunkPtr = Chunks.GetData() + (Num ? Num - 1 : 0) / NumElementsPerChunk;
 		return FIterType(ChunkPtr, LastChunkPtr, ChunkPtr ? (*ChunkPtr)->Elements + (bBeyondLastChunk ? NumElementsPerChunk : (Num % NumElementsPerChunk))  : nullptr);
 	}
 
-	friend FConstIterType end(const TChunkedArray& Array)
+	FConstIterType end() const
 	{
-		int32 Num = Array.NumElements;
+		int32 Num = NumElements;
 		bool bBeyondLastChunk = Num && Num % NumElementsPerChunk == 0;
-		const FChunk** ChunkPtr = Array.Chunks.GetData() + (Num / NumElementsPerChunk) + (bBeyondLastChunk ? -1 : 0); // do not read off the end of the chunk array!
-		const FChunk** LastChunkPtr = Array.Chunks.GetData() + (Num ? Num - 1 : 0) / NumElementsPerChunk;
+		const FChunk** ChunkPtr = Chunks.GetData() + (Num / NumElementsPerChunk) + (bBeyondLastChunk ? -1 : 0); // do not read off the end of the chunk array!
+		const FChunk** LastChunkPtr = Chunks.GetData() + (Num ? Num - 1 : 0) / NumElementsPerChunk;
 		return FConstIterType(ChunkPtr, LastChunkPtr, ChunkPtr ? (*ChunkPtr)->Elements + (bBeyondLastChunk ? NumElementsPerChunk : (Num % NumElementsPerChunk)) : nullptr);
 	}
 };
diff --git a/Engine/Source/Runtime/Core/Public/Containers/IndirectArray.h b/Engine/Source/Runtime/Core/Public/Containers/IndirectArray.h
--- a/Engine/Source/Runtime/Core/Public/Containers/IndirectArray.h
+++ b/Engine/Source/Runtime/Core/Public/Containers/IndirectArray.h
@@ -427,14 +427,17 @@ private:
 		}
 	}
 
+public:
 	/**
 	 * DO NOT USE DIRECTLY
 	 * STL-like iterators to enable range-based for loop support.
 	 */
-	FORCEINLINE friend TDereferencingIterator<      ElementType, typename InternalArrayType::RangedForIteratorType     > begin(      TIndirectArray& IndirectArray) { return TDereferencingIterator<      ElementType, typename InternalArrayType::RangedForIteratorType     >(begin(IndirectArray.Array)); }
-	FORCEINLINE friend TDereferencingIterator<const ElementType, typename InternalArrayType::RangedForConstIteratorType> begin(const TIndirectArray& IndirectArray) { return TDereferencingIterator<const ElementType, typename InternalArrayType::RangedForConstIteratorType>(begin(IndirectArray.Array)); }
-	FORCEINLINE friend TDereferencingIterator<      ElementType, typename InternalArrayType::RangedForIteratorType     > end  (      TIndirectArray& IndirectArray) { return TDereferencingIterator<      ElementType, typename InternalArrayType::RangedForIteratorType     >(end  (IndirectArray.Array)); }
-	FORCEINLINE friend TDereferencingIterator<const ElementType, typename InternalArrayType::RangedForConstIteratorType> end  (const TIndirectArray& IndirectArray) { return TDereferencingIterator<const ElementType, typename InternalArrayType::RangedForConstIteratorType>(end  (IndirectArray.Array)); }
+	FORCEINLINE TDereferencingIterator<      ElementType, typename InternalArrayType::RangedForIteratorType     > begin()       { return TDereferencingIterator<      ElementType, typename InternalArrayType::RangedForIteratorType     >(Array.begin()); }
+	FORCEINLINE TDereferencingIterator<const ElementType, typename InternalArrayType::RangedForConstIteratorType> begin() const { return TDereferencingIterator<const ElementType, typename InternalArrayType::RangedForConstIteratorType>(Array.begin()); }
+	FORCEINLINE TDereferencingIterator<      ElementType, typename InternalArrayType::RangedForIteratorType     > end  ()       { return TDereferencingIterator<      ElementType, typename InternalArrayType::RangedForIteratorType     >(Array.end()); }
+	FORCEINLINE TDereferencingIterator<const ElementType, typename InternalArrayType::RangedForConstIteratorType> end  () const { return TDereferencingIterator<const ElementType, typename InternalArrayType::RangedForConstIteratorType>(Array.end()); }
+
+private:
 
 	InternalArrayType Array;
 };
diff --git a/Engine/Source/Runtime/Core/Public/Containers/LruCache.h b/Engine/Source/Runtime/Core/Public/Containers/LruCache.h
--- a/Engine/Source/Runtime/Core/Public/Containers/LruCache.h
+++ b/Engine/Source/Runtime/Core/Public/Containers/LruCache.h
@@ -688,12 +688,12 @@ protected:
 		delete Entry;
 	}
 
-private:
+public:
 
-	friend TIterator begin(TLruCache& Cache) { return TIterator(Cache); }
-	friend TConstIterator begin(const TLruCache& Cache) { return TConstIterator(Cache); }
-	friend TIterator end(TLruCache& Cache) { return TIterator(); }
-	friend TConstIterator end(const TLruCache& Cache) { return TConstIterator(); }
+	TIterator begin() { return TIterator(*this); }
+	TConstIterator begin() const { return TConstIterator(*this); }
+	TIterator end() { return TIterator(); }
+	TConstIterator end() const { return TConstIterator(); }
 
 private:
 
diff --git a/Engine/Source/Runtime/Core/Public/Containers/Map.h b/Engine/Source/Runtime/Core/Public/Containers/Map.h
--- a/Engine/Source/Runtime/Core/Public/Containers/Map.h
+++ b/Engine/Source/Runtime/Core/Public/Containers/Map.h
@@ -784,15 +784,15 @@ public:
 		return TConstKeyIterator(*this, InKey);
 	}
 
-private:
+public:
 	/**
 	 * DO NOT USE DIRECTLY
 	 * STL-like iterators to enable range-based for loop support.
 	 */
-	FORCEINLINE friend TRangedForIterator      begin(      TMapBase& MapBase) { return TRangedForIterator     (begin(MapBase.Pairs)); }
-	FORCEINLINE friend TRangedForConstIterator begin(const TMapBase& MapBase) { return TRangedForConstIterator(begin(MapBase.Pairs)); }
-	FORCEINLINE friend TRangedForIterator      end  (      TMapBase& MapBase) { return TRangedForIterator     (end  (MapBase.Pairs)); }
-	FORCEINLINE friend TRangedForConstIterator end  (const TMapBase& MapBase) { return TRangedForConstIterator(end  (MapBase.Pairs)); }
+	FORCEINLINE TRangedForIterator      begin()       { return TRangedForIterator     (Pairs.begin()); }
+	FORCEINLINE TRangedForConstIterator begin() const { return TRangedForConstIterator(Pairs.begin()); }
+	FORCEINLINE TRangedForIterator      end  ()       { return TRangedForIterator     (Pairs.end());   }
+	FORCEINLINE TRangedForConstIterator end  () const { return TRangedForConstIterator(Pairs.end());   }
 };
 
 
diff --git a/Engine/Source/Runtime/Core/Public/Containers/Set.h b/Engine/Source/Runtime/Core/Public/Containers/Set.h
--- a/Engine/Source/Runtime/Core/Public/Containers/Set.h
+++ b/Engine/Source/Runtime/Core/Public/Containers/Set.h
@@ -1221,7 +1221,7 @@ public:
 
 	public:
 		FORCEINLINE TConstIterator(const TSet& InSet)
-			: TBaseIterator<true>(begin(InSet.Elements))
+			: TBaseIterator<true>(InSet.Elements.begin())
 		{
 		}
 	};
@@ -1233,7 +1233,7 @@ public:
 
 	public:
 		FORCEINLINE TIterator(TSet& InSet)
-			: TBaseIterator<false>(begin(InSet.Elements))
+			: TBaseIterator<false>(InSet.Elements.begin())
 			, Set                 (InSet)
 		{
 		}
@@ -1291,15 +1291,15 @@ public:
 		return TConstIterator(*this);
 	}
 
-private:
+public:
 	/**
 	 * DO NOT USE DIRECTLY
 	 * STL-like iterators to enable range-based for loop support.
 	 */
-	FORCEINLINE friend TRangedForIterator      begin(      TSet& Set) { return TRangedForIterator     (begin(Set.Elements)); }
-	FORCEINLINE friend TRangedForConstIterator begin(const TSet& Set) { return TRangedForConstIterator(begin(Set.Elements)); }
-	FORCEINLINE friend TRangedForIterator      end  (      TSet& Set) { return TRangedForIterator     (end  (Set.Elements)); }
-	FORCEINLINE friend TRangedForConstIterator end  (const TSet& Set) { return TRangedForConstIterator(end  (Set.Elements)); }
+	FORCEINLINE TRangedForIterator      begin()       { return TRangedForIterator     (Elements.begin()); }
+	FORCEINLINE TRangedForConstIterator begin() const { return TRangedForConstIterator(Elements.begin()); }
+	FORCEINLINE TRangedForIterator      end()         { return TRangedForIterator     (Elements.end());   }
+	FORCEINLINE TRangedForConstIterator end() const   { return TRangedForConstIterator(Elements.end());   }
 };
 
 template<typename ElementType, typename KeyFuncs, typename Allocator>
diff --git a/Engine/Source/Runtime/Core/Public/Containers/SortedMap.h b/Engine/Source/Runtime/Core/Public/Containers/SortedMap.h
--- a/Engine/Source/Runtime/Core/Public/Containers/SortedMap.h
+++ b/Engine/Source/Runtime/Core/Public/Containers/SortedMap.h
@@ -666,15 +666,15 @@ public:
 	typedef typename ElementArrayType::RangedForIteratorType RangedForIteratorType;
 	typedef typename ElementArrayType::RangedForConstIteratorType RangedForConstIteratorType;
 
-private:
+public:
 	/**
 	 * DO NOT USE DIRECTLY
 	 * STL-like iterators to enable range-based for loop support.
 	 */
-	FORCEINLINE friend RangedForIteratorType		begin(TSortedMap& MapBase)		{ return begin(MapBase.Pairs); }
-	FORCEINLINE friend RangedForConstIteratorType	begin(const TSortedMap& MapBase){ return begin(MapBase.Pairs); }
-	FORCEINLINE friend RangedForIteratorType		end(TSortedMap& MapBase)		{ return end(MapBase.Pairs); }
-	FORCEINLINE friend RangedForConstIteratorType	end(const TSortedMap& MapBase)	{ return end(MapBase.Pairs); }
+	FORCEINLINE RangedForIteratorType		begin()	      { return Pairs.begin(); }
+	FORCEINLINE RangedForConstIteratorType	begin() const { return Pairs.begin(); }
+	FORCEINLINE RangedForIteratorType		end()         { return Pairs.end(); }
+	FORCEINLINE RangedForConstIteratorType	end() const   { return Pairs.end(); }
 };
 
 
diff --git a/Engine/Source/Runtime/Core/Public/Containers/SparseArray.h b/Engine/Source/Runtime/Core/Public/Containers/SparseArray.h
--- a/Engine/Source/Runtime/Core/Public/Containers/SparseArray.h
+++ b/Engine/Source/Runtime/Core/Public/Containers/SparseArray.h
@@ -873,15 +873,15 @@ public:
 		return TConstIterator(*this);
 	}
 
-private:
+public:
 	/**
 	 * DO NOT USE DIRECTLY
 	 * STL-like iterators to enable range-based for loop support.
 	 */
-	FORCEINLINE friend TRangedForIterator      begin(      TSparseArray& Array) { return TRangedForIterator     (Array, TConstSetBitIterator<typename Allocator::BitArrayAllocator>(Array.AllocationFlags)); }
-	FORCEINLINE friend TRangedForConstIterator begin(const TSparseArray& Array) { return TRangedForConstIterator(Array, TConstSetBitIterator<typename Allocator::BitArrayAllocator>(Array.AllocationFlags)); }
-	FORCEINLINE friend TRangedForIterator      end  (      TSparseArray& Array) { return TRangedForIterator     (Array, TConstSetBitIterator<typename Allocator::BitArrayAllocator>(Array.AllocationFlags, Array.AllocationFlags.Num())); }
-	FORCEINLINE friend TRangedForConstIterator end  (const TSparseArray& Array) { return TRangedForConstIterator(Array, TConstSetBitIterator<typename Allocator::BitArrayAllocator>(Array.AllocationFlags, Array.AllocationFlags.Num())); }
+	FORCEINLINE TRangedForIterator      begin()       { return TRangedForIterator     (*this, TConstSetBitIterator<typename Allocator::BitArrayAllocator>(AllocationFlags)); }
+	FORCEINLINE TRangedForConstIterator begin() const { return TRangedForConstIterator(*this, TConstSetBitIterator<typename Allocator::BitArrayAllocator>(AllocationFlags)); }
+	FORCEINLINE TRangedForIterator      end  ()       { return TRangedForIterator     (*this, TConstSetBitIterator<typename Allocator::BitArrayAllocator>(AllocationFlags, AllocationFlags.Num())); }
+	FORCEINLINE TRangedForConstIterator end  () const { return TRangedForConstIterator(*this, TConstSetBitIterator<typename Allocator::BitArrayAllocator>(AllocationFlags, AllocationFlags.Num())); }
 
 public:
 	/** An iterator which only iterates over the elements of the array which correspond to set bits in a separate bit array. */
diff --git a/Engine/Source/Runtime/Core/Public/Containers/UnrealString.h b/Engine/Source/Runtime/Core/Public/Containers/UnrealString.h
--- a/Engine/Source/Runtime/Core/Public/Containers/UnrealString.h
+++ b/Engine/Source/Runtime/Core/Public/Containers/UnrealString.h
@@ -232,15 +232,15 @@ public:
 		return Data.CreateConstIterator();
 	}
 
-private:
+public:
 	/**
 	 * DO NOT USE DIRECTLY
 	 * STL-like iterators to enable range-based for loop support.
 	 */
-	FORCEINLINE friend DataType::RangedForIteratorType      begin(      FString& Str) { auto Result = begin(Str.Data);                                   return Result; }
-	FORCEINLINE friend DataType::RangedForConstIteratorType begin(const FString& Str) { auto Result = begin(Str.Data);                                   return Result; }
-	FORCEINLINE friend DataType::RangedForIteratorType      end  (      FString& Str) { auto Result = end  (Str.Data); if (Str.Data.Num()) { --Result; } return Result; }
-	FORCEINLINE friend DataType::RangedForConstIteratorType end  (const FString& Str) { auto Result = end  (Str.Data); if (Str.Data.Num()) { --Result; } return Result; }
+	FORCEINLINE DataType::RangedForIteratorType      begin()       { auto Result = Data.begin();                                   return Result; }
+	FORCEINLINE DataType::RangedForConstIteratorType begin() const { auto Result = Data.begin();                                   return Result; }
+	FORCEINLINE DataType::RangedForIteratorType      end  ()       { auto Result = Data.end();   if (Data.Num()) { --Result; }     return Result; }
+	FORCEINLINE DataType::RangedForConstIteratorType end  () const { auto Result = Data.end();   if (Data.Num()) { --Result; }     return Result; }
 
 public:
 	FORCEINLINE uint32 GetAllocatedSize() const
diff --git a/Engine/Source/Runtime/Core/Public/Misc/EnumRange.h b/Engine/Source/Runtime/Core/Public/Misc/EnumRange.h
--- a/Engine/Source/Runtime/Core/Public/Misc/EnumRange.h
+++ b/Engine/Source/Runtime/Core/Public/Misc/EnumRange.h
@@ -159,15 +159,15 @@ namespace NEnumRangePrivate
 	template <typename EnumType>
 	struct TEnumRange_Impl<EnumType, 0>
 	{
-		friend TEnumContiguousIterator<EnumType> begin(const TEnumRange_Impl&) { return TEnumContiguousIterator<EnumType>(TEnumRangeTraits<EnumType>::Begin); }
-		friend TEnumContiguousIterator<EnumType> end  (const TEnumRange_Impl&) { return TEnumContiguousIterator<EnumType>(TEnumRangeTraits<EnumType>::End  ); }
+		TEnumContiguousIterator<EnumType> begin() const { return TEnumContiguousIterator<EnumType>(TEnumRangeTraits<EnumType>::Begin); }
+		TEnumContiguousIterator<EnumType> end()   const { return TEnumContiguousIterator<EnumType>(TEnumRangeTraits<EnumType>::End  ); }
 	};
 
 	template <typename EnumType>
 	struct TEnumRange_Impl<EnumType, 1>
 	{
-		friend TEnumValueArrayIterator<EnumType> begin(const TEnumRange_Impl&) { return TEnumValueArrayIterator<EnumType>(TEnumRangeTraits<EnumType>::template GetPointer<void>(false)); }
-		friend TEnumValueArrayIterator<EnumType> end  (const TEnumRange_Impl&) { return TEnumValueArrayIterator<EnumType>(TEnumRangeTraits<EnumType>::template GetPointer<void>(true )); }
+		TEnumValueArrayIterator<EnumType> begin() const { return TEnumValueArrayIterator<EnumType>(TEnumRangeTraits<EnumType>::template GetPointer<void>(false)); }
+		TEnumValueArrayIterator<EnumType> end  () const { return TEnumValueArrayIterator<EnumType>(TEnumRangeTraits<EnumType>::template GetPointer<void>(true )); }
 	};
 }
 
diff --git a/Engine/Source/Runtime/LevelSequence/Public/LevelSequenceLegacyObjectReference.h b/Engine/Source/Runtime/LevelSequence/Public/LevelSequenceLegacyObjectReference.h
--- a/Engine/Source/Runtime/LevelSequence/Public/LevelSequenceLegacyObjectReference.h
+++ b/Engine/Source/Runtime/LevelSequence/Public/LevelSequenceLegacyObjectReference.h
@@ -75,10 +75,10 @@ public:
 	bool Serialize(FArchive& Ar);
 
 	typedef TMap<FGuid, FLevelSequenceLegacyObjectReference> MapType;
-	FORCEINLINE friend MapType::TRangedForIterator      begin(      FLevelSequenceObjectReferenceMap& Impl) { return begin(Impl.Map); }
-	FORCEINLINE friend MapType::TRangedForConstIterator begin(const FLevelSequenceObjectReferenceMap& Impl) { return begin(Impl.Map); }
-	FORCEINLINE friend MapType::TRangedForIterator      end  (      FLevelSequenceObjectReferenceMap& Impl) { return end(Impl.Map); }
-	FORCEINLINE friend MapType::TRangedForConstIterator end  (const FLevelSequenceObjectReferenceMap& Impl) { return end(Impl.Map); }
+	FORCEINLINE MapType::TRangedForIterator      begin()       { return Map.begin(); }
+	FORCEINLINE MapType::TRangedForConstIterator begin() const { return Map.begin(); }
+	FORCEINLINE MapType::TRangedForIterator      end  ()       { return Map.end();   }
+	FORCEINLINE MapType::TRangedForConstIterator end  () const { return Map.end();   }
 
 	TMap<FGuid, FLevelSequenceLegacyObjectReference> Map;
 };
diff --git a/Engine/Source/Runtime/MeshDescription/Public/MeshElementArray.h b/Engine/Source/Runtime/MeshDescription/Public/MeshElementArray.h
--- a/Engine/Source/Runtime/MeshDescription/Public/MeshElementArray.h
+++ b/Engine/Source/Runtime/MeshDescription/Public/MeshElementArray.h
@@ -236,18 +236,20 @@ public:
 			return TConstIterator( typename TSparseArray<ElementType>::TConstIterator( Array ) );
 		}
 
-	private:
+	public:
 
-		friend FORCEINLINE TConstIterator begin( const TElementIDs& ArrayIDs )
+		FORCEINLINE TConstIterator begin() const
 		{
-			return TConstIterator( begin( ArrayIDs.Array ) );
+			return TConstIterator( Array.begin() );
 		}
 
-		friend FORCEINLINE TConstIterator end( const TElementIDs& ArrayIDs )
+		FORCEINLINE TConstIterator end() const
 		{
-			return TConstIterator( end( ArrayIDs.Array ) );
+			return TConstIterator( Array.end() );
 		}
 
+	private:
+
 		const TSparseArray<ElementType>& Array;
 	};
 
