Allow per device memory bucket rounding from configrules with MemoryBucketRoundingAddition variable (defaults to 384 if not set)

#android

diff --git a/Engine/Source/Runtime/Core/Private/GenericPlatform/GenericPlatformMemory.cpp b/Engine/Source/Runtime/Core/Private/GenericPlatform/GenericPlatformMemory.cpp
--- a/Engine/Source/Runtime/Core/Private/GenericPlatform/GenericPlatformMemory.cpp
+++ b/Engine/Source/Runtime/Core/Private/GenericPlatform/GenericPlatformMemory.cpp
@@ -600,7 +600,12 @@ EPlatformMemorySizeBucket FGenericPlatformMemory::GetMemorySizeBucket()
 
 #if PLATFORM_ANDROID
 		// we don't exactly want to round up on android
-		uint32 TotalPhysicalGB = (Stats.TotalPhysical + 384 * 1024 * 1024 - 1) / 1024 / 1024 / 1024;
+		uint64 MemoryBucketRoundingAddition = 384;
+		if (FString* MemoryBucketRoundingAdditionVar = FAndroidMisc::GetConfigRulesVariable(TEXT("MemoryBucketRoundingAddition")))
+		{
+			MemoryBucketRoundingAddition = FCString::Atoi64(**MemoryBucketRoundingAdditionVar);
+		}
+		uint32 TotalPhysicalGB = (Stats.TotalPhysical + MemoryBucketRoundingAddition * 1024 * 1024 - 1) / 1024 / 1024 / 1024;
 #else
 		uint32 TotalPhysicalGB = (Stats.TotalPhysical + 1024 * 1024 * 1024 - 1) / 1024 / 1024 / 1024;
 #endif
