Reimplementing CL 4670945 due to robomerge issues.

Original description:

Fix GetVarArgs format specifier off by 1 error.  These would only show up if you had the format specifier at the end of the format string and only with 'l' and not 'll'.  In that case the code would walk off the end of the format string, usually into another static string.

Fix GetVarArgs to never walk off the end of the format string.

diff --git a/Engine/Source/Runtime/Core/Private/GenericPlatform/GenericWidePlatformString.cpp b/Engine/Source/Runtime/Core/Private/GenericPlatform/GenericWidePlatformString.cpp
--- a/Engine/Source/Runtime/Core/Private/GenericPlatform/GenericWidePlatformString.cpp
+++ b/Engine/Source/Runtime/Core/Private/GenericPlatform/GenericWidePlatformString.cpp
@@ -494,6 +494,13 @@ int32 FGenericWidePlatformString::GetVarArgs( WIDECHAR* Dest, SIZE_T DestSize, c
 			case 'z':
 			case 'Z':
 			{
+				if (Src[1] == 0)
+				{
+					printf("Unknown percent [%lc] in FGenericWidePlatformString::GetVarArgs() [%s]\n.", Src[0], TCHAR_TO_ANSI(Fmt));
+					Src++;  // skip it, I guess.
+					break;
+				}
+
 				Src += 2;
 
 				size_t Val = va_arg(ArgPtr, size_t);
@@ -564,8 +571,15 @@ int32 FGenericWidePlatformString::GetVarArgs( WIDECHAR* Dest, SIZE_T DestSize, c
 			{
 				int RemainingSize = Strlen(Src);
 
+				if (RemainingSize < 2)
+				{
+					printf("Unknown percent [%lc] in FGenericWidePlatformString::GetVarArgs() [%s]\n.", Src[0], TCHAR_TO_ANSI(Fmt));
+					Src++;  // skip it, I guess.
+					break;
+				}
+
 				// treat %ld as %d. Also shorts for %h will be promoted to ints
-				if (RemainingSize >= 2 && ((Src[0] == 'l' && Src[1] == 'd') || Src[0] == 'h'))
+				if ((Src[0] == 'l' && Src[1] == 'd') || Src[0] == 'h')
 				{
 					Src+=2;
 					int Val = va_arg(ArgPtr, int);
@@ -590,7 +604,7 @@ int32 FGenericWidePlatformString::GetVarArgs( WIDECHAR* Dest, SIZE_T DestSize, c
 					break;
 				}
 				// Treat %lf as a %f
-				else if (RemainingSize >= 2 && Src[0] == 'l' && Src[1] == 'f')
+				else if (Src[0] == 'l' && Src[1] == 'f')
 				{
 					Src += 2;
 					double Val = va_arg(ArgPtr, double);
@@ -615,17 +629,41 @@ int32 FGenericWidePlatformString::GetVarArgs( WIDECHAR* Dest, SIZE_T DestSize, c
 					break;
 				}
 
-				if (RemainingSize >= 2 && (Src[0] == 'l' && Src[1] != 'l' && Src[1] != 'u' && Src[1] != 'x'))
+				if (Src[0] == 'l')
 				{
-					printf("Unknown percent [%lc%lc] in FGenericWidePlatformString::GetVarArgs() [%s]\n.", Src[0], Src[1], TCHAR_TO_ANSI(Fmt));
-					Src++;  // skip it, I guess.
-					break;
+					if (Src[1] != 'l' && Src[1] != 'u' && Src[1] != 'x')
+					{
+						printf("Unknown percent [%lc%lc] in FGenericWidePlatformString::GetVarArgs() [%s]\n.", Src[0], Src[1], TCHAR_TO_ANSI(Fmt));
+						Src++;  // skip it, I guess.
+						break;
+					}
+					else if (Src[1] == 'l' && RemainingSize == 2)
+					{
+						printf("Unknown percent [%lc%lc] in FGenericWidePlatformString::GetVarArgs() [%s]\n.", Src[0], Src[1], TCHAR_TO_ANSI(Fmt));
+						Src++;  // skip it, I guess.
+						break;
+					}
 				}
-				else if (RemainingSize >= 3 && Src[0] == 'I' && (Src[1] != '6' || Src[2] != '4'))
+				else if (Src[0] == 'I')
 				{
-					printf("Unknown percent [%lc%lc%lc] in FGenericWidePlatformString::GetVarArgs() [%s]\n.", Src[0], Src[1], Src[2], TCHAR_TO_ANSI(Fmt));
-					Src++;  // skip it, I guess.
-					break;
+					if (RemainingSize == 2)
+					{
+						printf("Unknown percent [%lc%lc] in FGenericWidePlatformString::GetVarArgs() [%s]\n.", Src[0], Src[1], TCHAR_TO_ANSI(Fmt));
+						Src++;  // skip it, I guess.
+						break;
+					}
+					else if (RemainingSize == 3)
+					{
+						printf("Unknown percent [%lc%lc%lc] in FGenericWidePlatformString::GetVarArgs() [%s]\n.", Src[0], Src[1], Src[2], TCHAR_TO_ANSI(Fmt));
+						Src++;  // skip it, I guess.
+						break;
+					}
+					else if (Src[1] != '6' || Src[2] != '4')
+					{
+						printf("Unknown percent [%lc%lc%lc] in FGenericWidePlatformString::GetVarArgs() [%s]\n.", Src[0], Src[1], Src[2], TCHAR_TO_ANSI(Fmt));
+						Src++;  // skip it, I guess.
+						break;
+					}
 				}
 
 				// Yes, this is lame.
@@ -633,10 +671,14 @@ int32 FGenericWidePlatformString::GetVarArgs( WIDECHAR* Dest, SIZE_T DestSize, c
 				unsigned long long Val = va_arg(ArgPtr, unsigned long long);
 				ANSICHAR AnsiNum[60];
 				ANSICHAR FmtBuf[30];
-				if (Src[0] == 'l')
+				if (Src[0] == 'l' && Src[1] == 'l')
 				{
 					Src += 3;
 				}
+				else if (Src[0] == 'l')
+				{
+					Src += 2;
+				}
 				else
 				{
 					Src += 4;
