Fix merge error.

diff --git a/Engine/Source/Runtime/Core/Private/GenericPlatform/GenericPlatformMisc.cpp b/Engine/Source/Runtime/Core/Private/GenericPlatform/GenericPlatformMisc.cpp
--- a/Engine/Source/Runtime/Core/Private/GenericPlatform/GenericPlatformMisc.cpp
+++ b/Engine/Source/Runtime/Core/Private/GenericPlatform/GenericPlatformMisc.cpp
@@ -848,12 +848,6 @@ IPlatformInstallBundleManager* FGenericPlatformMisc::GetPlatformInstallBundleMan
 	return Manager;
 }
 
-IPlatformCompression* FGenericPlatformMisc::GetPlatformCompression()
-{
-	static FGenericPlatformCompression Singleton;
-	return &Singleton;
-}
-
 void GenericPlatformMisc_GetProjectFilePathProjectDir(FString& OutGameDir)
 {
 	// Here we derive the game path from the project file location.
