#ue4 Removed incorrect semicolon from UE_GET_LOG_VERBOSITY (NO_LOGGING version doesn't have one)

diff --git a/Engine/Source/Runtime/Core/Public/Logging/LogMacros.h b/Engine/Source/Runtime/Core/Public/Logging/LogMacros.h
--- a/Engine/Source/Runtime/Core/Public/Logging/LogMacros.h
+++ b/Engine/Source/Runtime/Core/Public/Logging/LogMacros.h
@@ -175,7 +175,7 @@ private:
 	#define UE_LOG_ACTIVE(CategoryName, Verbosity) (::UE4Asserts_Private::IsLogActive<(int32)ELogVerbosity::Verbosity>(CategoryName))
 
 	#define UE_GET_LOG_VERBOSITY(CategoryName) \
-		CategoryName.GetVerbosity();
+		CategoryName.GetVerbosity()
 
 	#define UE_SET_LOG_VERBOSITY(CategoryName, Verbosity) \
 		CategoryName.SetVerbosity(ELogVerbosity::Verbosity);
