Ensures: fixed discrepancy between the UI and the actual setting (UE-69764)

- Also hardened the code.

diff --git a/Engine/Config/BaseEditorSettings.ini b/Engine/Config/BaseEditorSettings.ini
--- a/Engine/Config/BaseEditorSettings.ini
+++ b/Engine/Config/BaseEditorSettings.ini
@@ -6,6 +6,9 @@ ShouldUseLocalizedNodeAndPinNames=True
 [/Script/UnrealEd.EditorSettings]
 bCopyStarterContentPreference = true
 
+[/Script/UnrealEd.CrashReportsPrivacySettings]
+bSendUnattendedBugReports=True
+
 [/Script/NewsFeed.NewsFeedSettings]
 CdnSourceUrl="http://cdn.unrealengine.com/newsfeed/[[Culture]]/index.json"
 LocalSourcePath="Editor/NewsFeed/[[Culture]]/"
diff --git a/Engine/Source/Editor/UnrealEd/Private/Analytics/CrashReportsPrivacySettings.cpp b/Engine/Source/Editor/UnrealEd/Private/Analytics/CrashReportsPrivacySettings.cpp
--- a/Engine/Source/Editor/UnrealEd/Private/Analytics/CrashReportsPrivacySettings.cpp
+++ b/Engine/Source/Editor/UnrealEd/Private/Analytics/CrashReportsPrivacySettings.cpp
@@ -4,13 +4,18 @@
 #include "UObject/UnrealType.h"
 #include "Interfaces/IAnalyticsProvider.h"
 #include "EngineAnalytics.h"
+#include "Misc/ConfigCacheIni.h"
 
 #define LOCTEXT_NAMESPACE "CrashReportsPrivacySettings"
 
 UCrashReportsPrivacySettings::UCrashReportsPrivacySettings(const FObjectInitializer& ObjectInitializer)
 	: Super(ObjectInitializer)
-	, bSendUnattendedBugReports(true)
+	, bSendUnattendedBugReports(false)
 {
+	if (GConfig)
+	{
+		GConfig->GetBool(TEXT("/Script/UnrealEd.CrashReportsPrivacySettings"), TEXT("bSendUnattendedBugReports"), bSendUnattendedBugReports, GEditorSettingsIni);
+	}
 }
 
 void UCrashReportsPrivacySettings::GetToogleCategoryAndPropertyNames(FName& OutCategory, FName& OutProperty) const
diff --git a/Engine/Source/Runtime/Core/Private/Mac/MacPlatformMisc.cpp b/Engine/Source/Runtime/Core/Private/Mac/MacPlatformMisc.cpp
--- a/Engine/Source/Runtime/Core/Private/Mac/MacPlatformMisc.cpp
+++ b/Engine/Source/Runtime/Core/Private/Mac/MacPlatformMisc.cpp
@@ -1776,7 +1776,10 @@ void FMacCrashContext::GenerateCrashInfoAndLaunchReporter() const
 	bool bCanRunCrashReportClient = FCString::Stristr( *(GMacAppInfo.ExecutableName), TEXT( "CrashReportClient" ) ) == nullptr;
 
 	bool bSendUnattendedBugReports = true;
-	GConfig->GetBool(TEXT("/Script/UnrealEd.CrashReportsPrivacySettings"), TEXT("bSendUnattendedBugReports"), bSendUnattendedBugReports, GEditorSettingsIni);
+	if (GConfig)
+	{
+		GConfig->GetBool(TEXT("/Script/UnrealEd.CrashReportsPrivacySettings"), TEXT("bSendUnattendedBugReports"), bSendUnattendedBugReports, GEditorSettingsIni);
+	}
 
 	if (BuildSettings::IsLicenseeVersion() && !UE_EDITOR)
 	{
@@ -1857,7 +1860,10 @@ void FMacCrashContext::GenerateEnsureInfoAndLaunchReporter() const
 	bool bCanRunCrashReportClient = FCString::Stristr( *(GMacAppInfo.ExecutableName), TEXT( "CrashReportClient" ) ) == nullptr;
 	
 	bool bSendUnattendedBugReports = true;
-	GConfig->GetBool(TEXT("/Script/UnrealEd.CrashReportsPrivacySettings"), TEXT("bSendUnattendedBugReports"), bSendUnattendedBugReports, GEditorSettingsIni);
+	if (GConfig)
+	{
+		GConfig->GetBool(TEXT("/Script/UnrealEd.CrashReportsPrivacySettings"), TEXT("bSendUnattendedBugReports"), bSendUnattendedBugReports, GEditorSettingsIni);
+	}
 
 	if (BuildSettings::IsLicenseeVersion() && !UE_EDITOR)
 	{
diff --git a/Engine/Source/Runtime/Core/Private/Unix/UnixPlatformCrashContext.cpp b/Engine/Source/Runtime/Core/Private/Unix/UnixPlatformCrashContext.cpp
--- a/Engine/Source/Runtime/Core/Private/Unix/UnixPlatformCrashContext.cpp
+++ b/Engine/Source/Runtime/Core/Private/Unix/UnixPlatformCrashContext.cpp
@@ -417,7 +417,10 @@ void FUnixCrashContext::GenerateCrashInfoAndLaunchReporter(bool bReportingNonCra
 
 	// By default we wont upload unless the *.ini has set this to true
 	bool bSendUnattendedBugReports = false;
-	GConfig->GetBool(TEXT("/Script/UnrealEd.CrashReportsPrivacySettings"), TEXT("bSendUnattendedBugReports"), bSendUnattendedBugReports, GEditorSettingsIni);
+	if (GConfig)
+	{
+		GConfig->GetBool(TEXT("/Script/UnrealEd.CrashReportsPrivacySettings"), TEXT("bSendUnattendedBugReports"), bSendUnattendedBugReports, GEditorSettingsIni);
+	}
 
 	if (BuildSettings::IsLicenseeVersion() && !UE_EDITOR)
 	{
diff --git a/Engine/Source/Runtime/Core/Private/Windows/WindowsPlatformCrashContext.cpp b/Engine/Source/Runtime/Core/Private/Windows/WindowsPlatformCrashContext.cpp
--- a/Engine/Source/Runtime/Core/Private/Windows/WindowsPlatformCrashContext.cpp
+++ b/Engine/Source/Runtime/Core/Private/Windows/WindowsPlatformCrashContext.cpp
@@ -312,7 +312,10 @@ int32 ReportCrashUsingCrashReportClient(FWindowsPlatformCrashContext& InContext,
 	bool bNoDialog = FApp::IsUnattended() || ReportUI == EErrorReportUI::ReportInUnattendedMode || IsRunningDedicatedServer();
 
 	bool bSendUnattendedBugReports = true;
-	GConfig->GetBool(TEXT("/Script/UnrealEd.CrashReportsPrivacySettings"), TEXT("bSendUnattendedBugReports"), bSendUnattendedBugReports, GEditorSettingsIni);
+	if (GConfig)
+	{
+		GConfig->GetBool(TEXT("/Script/UnrealEd.CrashReportsPrivacySettings"), TEXT("bSendUnattendedBugReports"), bSendUnattendedBugReports, GEditorSettingsIni);
+	}
 
 	if (BuildSettings::IsLicenseeVersion() && !UE_EDITOR)
 	{
