Enable fast path for unique name generation for all cooked builds by default

diff --git a/Engine/Source/Runtime/Core/Private/Misc/CoreGlobals.cpp b/Engine/Source/Runtime/Core/Private/Misc/CoreGlobals.cpp
--- a/Engine/Source/Runtime/Core/Private/Misc/CoreGlobals.cpp
+++ b/Engine/Source/Runtime/Core/Private/Misc/CoreGlobals.cpp
@@ -11,7 +11,7 @@
 
 
 #ifndef FAST_PATH_UNIQUE_NAME_GENERATION
-#define FAST_PATH_UNIQUE_NAME_GENERATION 0
+#define FAST_PATH_UNIQUE_NAME_GENERATION (!WITH_EDITORONLY_DATA)
 #endif
 
 #define LOCTEXT_NAMESPACE "Core"
