Clamp FTimespan value inside the wait function of the Event class.
Timespans longer than 4294967295 milliseconds were wrapping around to zero.

diff --git a/Engine/Source/Runtime/Core/Public/HAL/Event.h b/Engine/Source/Runtime/Core/Public/HAL/Event.h
--- a/Engine/Source/Runtime/Core/Public/HAL/Event.h
+++ b/Engine/Source/Runtime/Core/Public/HAL/Event.h
@@ -82,7 +82,7 @@ public:
 	bool Wait( const FTimespan& WaitTime, const bool bIgnoreThreadIdleStats = false )
 	{
 		check(WaitTime.GetTicks() >= 0);
-		return Wait(WaitTime.GetTotalMilliseconds(), bIgnoreThreadIdleStats);
+		return Wait((uint32)FMath::Clamp<int64>(WaitTime.GetTicks() / ETimespan::TicksPerMillisecond, 0, MAX_uint32), bIgnoreThreadIdleStats);
 	}
 
 	/** Default constructor. */
