Back out all the coded added for stats memory leak investigation.  It is no longer needed and is causing crashes.

diff --git a/Engine/Source/Runtime/Core/Private/HAL/LowLevelMemTracker.cpp b/Engine/Source/Runtime/Core/Private/HAL/LowLevelMemTracker.cpp
--- a/Engine/Source/Runtime/Core/Private/HAL/LowLevelMemTracker.cpp
+++ b/Engine/Source/Runtime/Core/Private/HAL/LowLevelMemTracker.cpp
@@ -7,7 +7,6 @@
 #include "Misc/Paths.h"
 #include "Misc/CString.h"
 #include "HAL/IConsoleManager.h"
-#include "Stats/StatsData.h"
 
 #if ENABLE_LOW_LEVEL_MEM_TRACKER
 
@@ -18,8 +17,6 @@
 // Disable if you need a little more memory or speed
 #define LLM_TRACK_PEAK_MEMORY 0		// currently disabled because there was a problem with tracking peaks from multiple threads.
 
-#define ENABLE_CATCH_RUNAWAY_TAGS 1
-
 #ifdef ENABLE_MEMPRO
 namespace MemProProfiler
 {
@@ -243,10 +240,6 @@ public:
 	void SetStat(int64 Tag, int64 Value);
 #endif
 
-	int64 GetStat(int64 Tag);
-
-	void CatchRunawayTags(FLLMPlatformTag* PlatformTags);
-
 	void Update(FLLMPlatformTag* PlatformTags, const int32* ParentTags);
 
 	void SetEnabled(bool value) { Enabled = value; }
@@ -348,7 +341,6 @@ public:
 	void SetTotalTags(ELLMTag Untagged, ELLMTag Tracked);
 	void Update(FLLMPlatformTag* PlatformTags, const int32* ParentTags);
 	void UpdateTotals();
-	void CatchRunawayTags(FLLMPlatformTag* PlatformTags);
 
 	int64 GetTagAmount(ELLMTag Tag) const;
 	void SetTagAmount(ELLMTag Tag, int64 Amount, bool AddToTotal);
@@ -585,9 +577,6 @@ void FLowLevelMemTracker::UpdateStatsPerFrame(const TCHAR* LogName)
 	for (int32 TrackerIndex = 0; TrackerIndex < (int32)ELLMTracker::Max; TrackerIndex++)
 	{
 		GetTracker((ELLMTracker)TrackerIndex)->Update(PlatformTags,ParentTags);
-#if ENABLE_CATCH_RUNAWAY_TAGS
-		GetTracker((ELLMTracker)TrackerIndex)->CatchRunawayTags(PlatformTags);
-#endif
 	}
 
 	// calculate FMalloc unused stat and set it in the Default tracker
@@ -1324,11 +1313,6 @@ void FLLMTracker::UpdateTotals()
 	}
 }
 
-void FLLMTracker::CatchRunawayTags(FLLMPlatformTag* PlatformTags)
-{
-	CsvWriter.CatchRunawayTags(PlatformTags);
-}
-
 void FLLMTracker::WriteCsv(FLLMPlatformTag* PlatformTags, const int32* ParentTags)
 {
 	CsvWriter.Update(PlatformTags,ParentTags);
@@ -1799,44 +1783,6 @@ void FLLMCsvWriter::SetStat(int64 Tag, int64 Value)
 	StatValues.Add(NewStatValue);
 }
 
-int64 FLLMCsvWriter::GetStat(int64 Tag)
-{
-	int StatValueCount = StatValues.Num();
-	for (int32 i = 0; i < StatValueCount; ++i)
-	{
-		if (StatValues[i].Tag == Tag)
-		{
-			return StatValues[i].Value;
-		}
-	}
-
-	return 0;
-}
-
-void FLLMCsvWriter::CatchRunawayTags(FLLMPlatformTag* PlatformTags)
-{
-#if STATS
-	const FString StatsTagName(TEXT("Stats"));
-	int64 StatsTag = -1;
-
-	for (uint32 i = 0; i < StatValues.Num(); i++)
-	{
-		if (GetTagName(StatValues[i].Tag, PlatformTags,nullptr) == StatsTagName)
-		{
-			StatsTag = StatValues[i].Tag;
-		}
-	}
-
-	static bool bDumpStatsEnabled = false;
-	const int64 StatsDumpThreshold = 500 * 1024 * 1024;
-	if (!bDumpStatsEnabled && StatsTag >= 0 && GetStat(StatsTag) > StatsDumpThreshold)
-	{
-		FStatPacket::bDumpStatPacket = true;
-		bDumpStatsEnabled = true;
-	}
-#endif
-}
-
 /*
 * memory can be allocated in this function
 */
diff --git a/Engine/Source/Runtime/Core/Private/Stats/Stats2.cpp b/Engine/Source/Runtime/Core/Private/Stats/Stats2.cpp
--- a/Engine/Source/Runtime/Core/Private/Stats/Stats2.cpp
+++ b/Engine/Source/Runtime/Core/Private/Stats/Stats2.cpp
@@ -1444,6 +1444,4 @@ void FThreadStats::WaitForStats()
 	}
 }
 
-bool FStatPacket::bDumpStatPacket = false;
-
 #endif
diff --git a/Engine/Source/Runtime/Core/Private/Stats/StatsData.cpp b/Engine/Source/Runtime/Core/Private/Stats/StatsData.cpp
--- a/Engine/Source/Runtime/Core/Private/Stats/StatsData.cpp
+++ b/Engine/Source/Runtime/Core/Private/Stats/StatsData.cpp
@@ -811,18 +811,6 @@ void FStatsThreadState::ToggleFindMemoryExtensiveStats()
 	UE_LOG(LogStats, Log, TEXT("bFindMemoryExtensiveStats is %s now"), bFindMemoryExtensiveStats?TEXT("enabled"):TEXT("disabled"));
 }
 
-void FStatsThreadState::EnableFindMemoryExtensiveStats()
-{
-	bFindMemoryExtensiveStats = true;
-	UE_LOG(LogStats, Log, TEXT("bFindMemoryExtensiveStats is enabled now"));
-}
-
-void FStatsThreadState::DisableFindMemoryExtensiveStats()
-{
-	bFindMemoryExtensiveStats = false;
-	UE_LOG(LogStats, Log, TEXT("bFindMemoryExtensiveStats is disabled now"));
-}
-
 void FStatsThreadState::ProcessNonFrameStats(FStatMessagesArray& Data, TSet<FName>* NonFrameStatsFound)
 {
 	for (FStatMessage& Item : Data)
diff --git a/Engine/Source/Runtime/Core/Public/Stats/Stats2.h b/Engine/Source/Runtime/Core/Public/Stats/Stats2.h
--- a/Engine/Source/Runtime/Core/Public/Stats/Stats2.h
+++ b/Engine/Source/Runtime/Core/Public/Stats/Stats2.h
@@ -17,7 +17,6 @@
 #include "Templates/Atomic.h"
 #include "Math/Color.h"
 #include "StatsCommon.h"
-#include "HAL/ThreadManager.h"
 #include "Templates/UniquePtr.h"
 
 class FScopeCycleCounter;
@@ -1123,9 +1122,6 @@ struct FStatPacket
 	/** Size we presize the message buffer to, currently the max of what we have seen for the last PRESIZE_MAX_NUM_ENTRIES. **/
 	TArray<int32> StatMessagesPresize;
 
-	/** If true, we dump ThreadedStats to track number of statmessage in each thread **/
-	static bool bDumpStatPacket;
-
 	/** constructor **/
 	FStatPacket()
 		: Frame(1)
@@ -1167,11 +1163,6 @@ struct FStatPacket
 	void AssignFrame( int64 InFrame )
 	{
 		Frame = InFrame;
-
-		if (bDumpStatPacket)
-		{
-			UE_LOG(LogStats, Display, TEXT("Frame %d, Size of StatMessages for Thread %d(%s) is %d"), Frame, ThreadId, *FThreadManager::Get().GetThreadName(ThreadId), StatMessages.Num());
-		}
 	}
 };
 
diff --git a/Engine/Source/Runtime/Core/Public/Stats/StatsData.h b/Engine/Source/Runtime/Core/Public/Stats/StatsData.h
--- a/Engine/Source/Runtime/Core/Public/Stats/StatsData.h
+++ b/Engine/Source/Runtime/Core/Public/Stats/StatsData.h
@@ -511,12 +511,6 @@ public:
 	/** Toggles tracking the most memory expensive stats. */
 	void ToggleFindMemoryExtensiveStats();
 
-	/** Enables tracking the most memory expensive stats. */
-	void EnableFindMemoryExtensiveStats();
-
-	/** Disables tracking the most memory expensive stats. */
-	void DisableFindMemoryExtensiveStats();
-
 	/** Resets stats for raw stats. */
 	void ResetStatsForRawStats()
 	{
