Split out some new helpful LLM tags

diff --git a/Engine/Source/Runtime/Core/Private/HAL/LowLevelMemTracker.cpp b/Engine/Source/Runtime/Core/Private/HAL/LowLevelMemTracker.cpp
--- a/Engine/Source/Runtime/Core/Private/HAL/LowLevelMemTracker.cpp
+++ b/Engine/Source/Runtime/Core/Private/HAL/LowLevelMemTracker.cpp
@@ -104,6 +104,9 @@ DECLARE_LLM_MEMORY_STAT(TEXT("ConfigSystem"), STAT_ConfigSystemLLM, STATGROUP_LL
 DECLARE_LLM_MEMORY_STAT(TEXT("InitUObject"), STAT_InitUObjectLLM, STATGROUP_LLMFULL);
 DECLARE_LLM_MEMORY_STAT(TEXT("VideoRecording"), STAT_VideoRecordingLLM, STATGROUP_LLMFULL);
 DECLARE_LLM_MEMORY_STAT(TEXT("CsvProfiler"), STAT_CsvProfilerLLM, STATGROUP_LLMFULL);
+DECLARE_LLM_MEMORY_STAT(TEXT("MaterialInstance"), STAT_MaterialInstanceLLM, STATGROUP_LLMFULL);
+DECLARE_LLM_MEMORY_STAT(TEXT("SkeletalMesh"), STAT_SkeletalMeshLLM, STATGROUP_LLMFULL);
+DECLARE_LLM_MEMORY_STAT(TEXT("Landscape"), STAT_LandscapeLLM, STATGROUP_LLMFULL);
 
 /*
 * LLM Summary stats referenced by ELLMTagNames
@@ -1161,9 +1164,9 @@ void FLLMTracker::TrackAllocation(const void* Ptr, uint64 Size, ELLMTag DefaultT
 
 		FLLMTracker::FLowLevelAllocInfo AllocInfo;
 		#if LLM_USE_ALLOC_INFO_STRUCT
-			AllocInfo.Tag = tag;
+		AllocInfo.Tag = tag;
 			#if LLM_ALLOW_ASSETS_TAGS
-				AllocInfo.AssetTag = State->GetTopAssetTag();
+		AllocInfo.AssetTag = State->GetTopAssetTag();
 			#endif
 		#else
 		LLMCheck(tag >= 0 && tag <= (int64)ELLMTag::PlatformTagEnd);
diff --git a/Engine/Source/Runtime/Core/Public/HAL/LowLevelMemTracker.h b/Engine/Source/Runtime/Core/Public/HAL/LowLevelMemTracker.h
--- a/Engine/Source/Runtime/Core/Public/HAL/LowLevelMemTracker.h
+++ b/Engine/Source/Runtime/Core/Public/HAL/LowLevelMemTracker.h
@@ -229,6 +229,9 @@ enum class ELLMTagSet : uint8
 	macro(ConfigSystem,							"ConfigSystem",					GET_STATFNAME(STAT_ConfigSystemLLM),						GET_STATFNAME(STAT_EngineSummaryLLM),			-1)\
 	macro(InitUObject,							"InitUObject",					GET_STATFNAME(STAT_InitUObjectLLM),							GET_STATFNAME(STAT_EngineSummaryLLM),			-1)\
 	macro(VideoRecording,						"VideoRecording",				GET_STATFNAME(STAT_VideoRecordingLLM),						GET_STATFNAME(STAT_EngineSummaryLLM),			-1)\
+	macro(MaterialInstance,						"MaterialInstance",				GET_STATFNAME(STAT_MaterialInstanceLLM),					GET_STATFNAME(STAT_EngineSummaryLLM),			-1)\
+	macro(SkeletalMesh,							"SkeletalMesh",					GET_STATFNAME(STAT_SkeletalMeshLLM),						GET_STATFNAME(STAT_EngineSummaryLLM),			-1)\
+	macro(Landscape,							"Landscape",					GET_STATFNAME(STAT_LandscapeLLM),							GET_STATFNAME(STAT_EngineSummaryLLM),			-1)\
 	macro(CsvProfiler,							"CsvProfiler",					GET_STATFNAME(STAT_CsvProfilerLLM),							GET_STATFNAME(STAT_EngineSummaryLLM),			-1)
 
 /*
diff --git a/Engine/Source/Runtime/Engine/Private/Materials/MaterialInstance.cpp b/Engine/Source/Runtime/Engine/Private/Materials/MaterialInstance.cpp
--- a/Engine/Source/Runtime/Engine/Private/Materials/MaterialInstance.cpp
+++ b/Engine/Source/Runtime/Engine/Private/Materials/MaterialInstance.cpp
@@ -407,6 +407,7 @@ UMaterialInstance::UMaterialInstance(const FObjectInitializer& ObjectInitializer
 
 void UMaterialInstance::PostInitProperties()	
 {
+	LLM_SCOPE(ELLMTag::MaterialInstance);
 	Super::PostInitProperties();
 
 	if(!HasAnyFlags(RF_ClassDefaultObject))
@@ -2956,7 +2957,7 @@ void UMaterialInstance::ClearAllCachedCookedPlatformData()
 
 void UMaterialInstance::Serialize(FArchive& Ar)
 {
-	LLM_SCOPE(ELLMTag::Materials);
+	LLM_SCOPE(ELLMTag::MaterialInstance);
 	SCOPED_LOADTIMER(MaterialInstanceSerializeTime);
 	SCOPE_CYCLE_COUNTER(STAT_MaterialInstance_Serialize);
 
@@ -3082,6 +3083,7 @@ void UMaterialInstance::Serialize(FArchive& Ar)
 
 void UMaterialInstance::PostLoad()
 {
+	LLM_SCOPE(ELLMTag::MaterialInstance);
 	SCOPED_LOADTIMER(MaterialInstancePostLoad);
 
 	Super::PostLoad();
diff --git a/Engine/Source/Runtime/Engine/Private/SkeletalMesh.cpp b/Engine/Source/Runtime/Engine/Private/SkeletalMesh.cpp
--- a/Engine/Source/Runtime/Engine/Private/SkeletalMesh.cpp
+++ b/Engine/Source/Runtime/Engine/Private/SkeletalMesh.cpp
@@ -990,6 +990,7 @@ bool USkeletalMesh::IsReadyForFinishDestroy()
 
 void USkeletalMesh::Serialize( FArchive& Ar )
 {
+	LLM_SCOPE(ELLMTag::SkeletalMesh);
 	DECLARE_SCOPE_CYCLE_COUNTER( TEXT("USkeletalMesh::Serialize"), STAT_SkeletalMesh_Serialize, STATGROUP_LoadTime );
 
 	Super::Serialize(Ar);
diff --git a/Engine/Source/Runtime/Landscape/Private/Landscape.cpp b/Engine/Source/Runtime/Landscape/Private/Landscape.cpp
--- a/Engine/Source/Runtime/Landscape/Private/Landscape.cpp
+++ b/Engine/Source/Runtime/Landscape/Private/Landscape.cpp
@@ -300,6 +300,7 @@ void ULandscapeComponent::CheckGenerateLandscapePlatformData(bool bIsCooking, co
 
 void ULandscapeComponent::Serialize(FArchive& Ar)
 {
+	LLM_SCOPE(ELLMTag::Landscape);
 	Ar.UsingCustomVersion(FRenderingObjectVersion::GUID);
 	Ar.UsingCustomVersion(FFortniteMainBranchObjectVersion::GUID);
 
