Linux: respecting NumStackFramesToIgnore when triggering Asserts

diff --git a/Engine/Source/Runtime/Core/Private/Unix/UnixErrorOutputDevice.cpp b/Engine/Source/Runtime/Core/Private/Unix/UnixErrorOutputDevice.cpp
--- a/Engine/Source/Runtime/Core/Private/Unix/UnixErrorOutputDevice.cpp
+++ b/Engine/Source/Runtime/Core/Private/Unix/UnixErrorOutputDevice.cpp
@@ -54,7 +54,13 @@ void FUnixErrorOutputDevice::Serialize(const TCHAR* Msg, ELogVerbosity::Type Ver
 #if PLATFORM_EXCEPTIONS_DISABLED
 		UE_DEBUG_BREAK();
 #endif
-		ReportAssert(Msg, 0);
+		// Generate the portable callstack. For asserts, we ignore the following frames:
+		//     FDebug::AssertFailed()
+		//   [ FOutputDevice::Logf() ] - force-inlined; ignored
+		//     FOutputDevice::LogfImpl()
+		//     FWindowsErrorOutputDevice::Serialize()
+		const int32 NumStackFramesToIgnore = 3;
+		ReportAssert(Msg, NumStackFramesToIgnore);
 	}
 	else
 	{
diff --git a/Engine/Source/Runtime/Core/Private/Unix/UnixPlatformCrashContext.cpp b/Engine/Source/Runtime/Core/Private/Unix/UnixPlatformCrashContext.cpp
--- a/Engine/Source/Runtime/Core/Private/Unix/UnixPlatformCrashContext.cpp
+++ b/Engine/Source/Runtime/Core/Private/Unix/UnixPlatformCrashContext.cpp
@@ -101,6 +101,12 @@ void FUnixCrashContext::InitFromSignal(int32 InSignal, siginfo_t* InInfo, void*
 	Context = reinterpret_cast< ucontext_t* >( InContext );
 
 	FCString::Strcat(SignalDescription, ARRAY_COUNT( SignalDescription ) - 1, *DescribeSignal(Signal, Info, Context));
+
+	// Retrieve NumStackFramesToIgnore from signal data
+	if (Info && Info->si_value.sival_int != 0)
+	{
+		SetNumMinidumpFramesToIgnore(Info->si_value.sival_int);
+	}
 }
 
 void FUnixCrashContext::InitFromEnsureHandler(const TCHAR* EnsureMessage, const void* CrashAddress)
@@ -292,7 +298,7 @@ void FUnixCrashContext::CaptureStackTrace()
 		ANSICHAR* StackTrace = (ANSICHAR*) FMemory::Malloc( StackTraceSize );
 		StackTrace[0] = 0;
 
-		int32 IgnoreCount = 0;
+		int32 IgnoreCount = NumMinidumpFramesToIgnore;
 		CapturePortableCallStack(IgnoreCount, this);
 
 		// Walk the stack and dump it to the allocated memory (do not ignore any stack frames to be consistent with check()/ensure() handling)
diff --git a/Engine/Source/Runtime/Core/Private/Unix/UnixPlatformStackWalk.cpp b/Engine/Source/Runtime/Core/Private/Unix/UnixPlatformStackWalk.cpp
--- a/Engine/Source/Runtime/Core/Private/Unix/UnixPlatformStackWalk.cpp
+++ b/Engine/Source/Runtime/Core/Private/Unix/UnixPlatformStackWalk.cpp
@@ -787,7 +787,17 @@ thread_local const TCHAR* GAssertErrorMessage = nullptr;
 void ReportAssert(const TCHAR* ErrorMessage, int NumStackFramesToIgnore)
 {
 	GAssertErrorMessage = ErrorMessage;
-	FPlatformMisc::RaiseException(1);
+
+	// Store NumStackFramesToIgnore in signal data	
+	sigval UserData;
+	UserData.sival_int = NumStackFramesToIgnore + 2; // +2 for this function and sigqueue()
+	sigqueue(getpid(),  SIGSEGV, UserData);
+	
+	// Make sure we never return
+	for (;;)
+	{
+		FPlatformProcess::Sleep(60.0f);
+	}
 }
 
 static FCriticalSection EnsureLock;
