Fix FNetFieldExport serialization for replays

diff --git a/Engine/Source/Runtime/Core/Private/Misc/NetworkVersion.cpp b/Engine/Source/Runtime/Core/Private/Misc/NetworkVersion.cpp
--- a/Engine/Source/Runtime/Core/Private/Misc/NetworkVersion.cpp
+++ b/Engine/Source/Runtime/Core/Private/Misc/NetworkVersion.cpp
@@ -18,7 +18,7 @@ FString FNetworkVersion::ProjectVersion = TEXT("1.0.0");
 bool FNetworkVersion::bHasCachedNetworkChecksum			= false;
 uint32 FNetworkVersion::CachedNetworkChecksum			= 0;
 
-uint32 FNetworkVersion::EngineNetworkProtocolVersion	= HISTORY_NETEXPORT_SERIALIZATION;
+uint32 FNetworkVersion::EngineNetworkProtocolVersion	= HISTORY_NETEXPORT_SERIALIZE_FIX;
 uint32 FNetworkVersion::GameNetworkProtocolVersion		= 0;
 
 uint32 FNetworkVersion::EngineCompatibleNetworkProtocolVersion		= HISTORY_REPLAY_BACKWARDS_COMPAT;
diff --git a/Engine/Source/Runtime/Core/Public/Misc/NetworkVersion.h b/Engine/Source/Runtime/Core/Public/Misc/NetworkVersion.h
--- a/Engine/Source/Runtime/Core/Public/Misc/NetworkVersion.h
+++ b/Engine/Source/Runtime/Core/Public/Misc/NetworkVersion.h
@@ -41,6 +41,7 @@ enum EEngineNetworkVersionHistory
 	HISTORY_CHANNEL_CLOSE_REASON = 7,				// Bump version to serialize a channel close reason in bunches instead of bDormant
 	HISTORY_ACKS_INCLUDED_IN_HEADER = 8,			// Bump version since acks are now sent as part of the header
 	HISTORY_NETEXPORT_SERIALIZATION = 9,			// Bump version due to serialization change to FNetFieldExport
+	HISTORY_NETEXPORT_SERIALIZE_FIX = 10,			// Bump version to fix net field export name serialization 
 };
 
 struct CORE_API FNetworkVersion
diff --git a/Engine/Source/Runtime/Engine/Classes/Engine/PackageMapClient.h b/Engine/Source/Runtime/Engine/Classes/Engine/PackageMapClient.h
--- a/Engine/Source/Runtime/Engine/Classes/Engine/PackageMapClient.h
+++ b/Engine/Source/Runtime/Engine/Classes/Engine/PackageMapClient.h
@@ -88,9 +88,19 @@ public:
 			}
 			else
 			{
-				Ar << C.ExportName;
-
-				C.Name = C.ExportName.ToString();
+				if (Ar.IsLoading() && Ar.EngineNetVer() < HISTORY_NETEXPORT_SERIALIZE_FIX)
+				{
+					Ar << C.ExportName;
+				}
+				else
+				{
+					UPackageMap::StaticSerializeName(Ar, C.ExportName);
+				}
+
+				if (Ar.IsLoading())
+				{
+					C.Name = C.ExportName.ToString();
+				}
 			}
 		}
 
