Allow games to pass arbitrary key/value pairs via crash reports.

Calling FGenericCrashContext::AddGameData(TEXT("Foo"), TEXT("Bar")) will add a key/value, which will also be sent as an analytics event from CrashReportClient (with the key name "GameData.Foo").

diff --git a/Engine/Source/Programs/CrashReportClient/Private/CrashDescription.cpp b/Engine/Source/Programs/CrashReportClient/Private/CrashDescription.cpp
--- a/Engine/Source/Programs/CrashReportClient/Private/CrashDescription.cpp
+++ b/Engine/Source/Programs/CrashReportClient/Private/CrashDescription.cpp
@@ -387,6 +387,17 @@ void FPrimaryCrashProperties::MakeCrashEventAttributes(TArray<FAnalyticsEventAtt
 	OutCrashAttributes.Add(FAnalyticsEventAttribute(TEXT("UserActivityHint"), UserActivityHint.AsString()));
 	OutCrashAttributes.Add(FAnalyticsEventAttribute(TEXT("GameSessionID"), GameSessionID.AsString()));
 	OutCrashAttributes.Add(FAnalyticsEventAttribute(TEXT("DeploymentName"), DeploymentName));
+
+	// Add arbitrary game data
+	const FXmlNode* MainNode = XmlFile->GetRootNode()->FindChildNode( FGenericCrashContext::GameDataTag );
+	if (MainNode)
+	{
+		for (const FXmlNode* ChildNode : MainNode->GetChildrenNodes())
+		{
+			FString KeyName = FString(TEXT("GameData.")) + ChildNode->GetTag();
+			OutCrashAttributes.Add(FAnalyticsEventAttribute(*KeyName, ChildNode->GetContent()));
+		}
+	}
 }
 
 void FPrimaryCrashProperties::Save()
diff --git a/Engine/Source/Runtime/Core/Private/GenericPlatform/GenericPlatformCrashContext.cpp b/Engine/Source/Runtime/Core/Private/GenericPlatform/GenericPlatformCrashContext.cpp
--- a/Engine/Source/Runtime/Core/Private/GenericPlatform/GenericPlatformCrashContext.cpp
+++ b/Engine/Source/Runtime/Core/Private/GenericPlatform/GenericPlatformCrashContext.cpp
@@ -43,6 +43,7 @@ const FString FGenericCrashContext::CrashGUIDRootPrefix = TEXT("UE4CC-");
 const FString FGenericCrashContext::CrashContextExtension = TEXT(".runtime-xml");
 const FString FGenericCrashContext::RuntimePropertiesTag = TEXT( "RuntimeProperties" );
 const FString FGenericCrashContext::PlatformPropertiesTag = TEXT( "PlatformProperties" );
+const FString FGenericCrashContext::GameDataTag = TEXT( "GameData" );
 const FString FGenericCrashContext::EnabledPluginsTag = TEXT("EnabledPlugins");
 const FString FGenericCrashContext::UE4MinidumpName = TEXT( "UE4Minidump.dmp" );
 const FString FGenericCrashContext::NewLineTag = TEXT( "&nl;" );
@@ -97,6 +98,7 @@ namespace NCachedCrashContextProperties
 	static FString CrashReportClientRichText;
 	static FString GameStateName;
 	static TArray<FString> EnabledPluginsList;
+	static TMap<FString, FString> GameData;
 }
 
 void FGenericCrashContext::Initialize()
@@ -403,6 +405,14 @@ void FGenericCrashContext::SerializeContentToBuffer() const
 	AddPlatformSpecificProperties();
 	EndSection( *PlatformPropertiesTag );
 
+	// Add the game data
+	BeginSection( *GameDataTag );
+	for (const TPair<FString, FString>& Pair : NCachedCrashContextProperties::GameData)
+	{
+		AddCrashProperty(*Pair.Key, *Pair.Value);
+	}
+	EndSection( *GameDataTag );
+
 	// Writing out the list of plugin JSON descriptors causes us to run out of memory
 	// in GMallocCrash on console, so enable this only for desktop platforms.
 #if PLATFORM_DESKTOP
@@ -653,6 +663,23 @@ void FGenericCrashContext::PurgeOldCrashConfig()
 	}
 }
 
+void FGenericCrashContext::ResetGameData()
+{
+	NCachedCrashContextProperties::GameData.Reset();
+}
+
+void FGenericCrashContext::AddGameData(const FString& Key, const FString& Value)
+{
+	if (Value.Len() == 0)
+	{
+		NCachedCrashContextProperties::GameData.Remove(Key);
+	}
+	else
+	{
+		NCachedCrashContextProperties::GameData.Add(Key, Value);
+	}
+}
+
 void FGenericCrashContext::AddPlugin(const FString& PluginDesc)
 {
 	NCachedCrashContextProperties::EnabledPluginsList.Add(PluginDesc);
diff --git a/Engine/Source/Runtime/Core/Public/GenericPlatform/GenericPlatformCrashContext.h b/Engine/Source/Runtime/Core/Public/GenericPlatform/GenericPlatformCrashContext.h
--- a/Engine/Source/Runtime/Core/Public/GenericPlatform/GenericPlatformCrashContext.h
+++ b/Engine/Source/Runtime/Core/Public/GenericPlatform/GenericPlatformCrashContext.h
@@ -119,6 +119,7 @@ public:
 	static const FString CrashContextExtension;
 	static const FString RuntimePropertiesTag;
 	static const FString PlatformPropertiesTag;
+	static const FString GameDataTag;
 	static const FString EnabledPluginsTag;
 	static const FString UE4MinidumpName;
 	static const FString NewLineTag;
@@ -210,6 +211,12 @@ public:
 	/** Helper to clean out old files in the crash report client config folder. */
 	static void PurgeOldCrashConfig();
 
+	/** Clears the game data dictionary */
+	static void ResetGameData();
+
+	/** Adds arbitrary game data to the crash context */
+	static void AddGameData(const FString& Key, const FString& Value);
+
 	/** Adds a plugin descriptor string to the enabled plugins list in the crash context */
 	static void AddPlugin(const FString& PluginDesc);
 
