Fix TNonRelocatableInlineAllocator not switching to heap allocator correctly, due to CalculateSlackGrow() not increasing to contain the requested number of elements.

diff --git a/Engine/Source/Runtime/Core/Public/Containers/ContainerAllocationPolicies.h b/Engine/Source/Runtime/Core/Public/Containers/ContainerAllocationPolicies.h
--- a/Engine/Source/Runtime/Core/Public/Containers/ContainerAllocationPolicies.h
+++ b/Engine/Source/Runtime/Core/Public/Containers/ContainerAllocationPolicies.h
@@ -673,19 +673,19 @@ public:
 		FORCEINLINE int32 CalculateSlackReserve(int32 NumElements, SIZE_T NumBytesPerElement) const
 		{
 			// If the elements use less space than the inline allocation, only use the inline allocation as slack.
-			return HasAllocation() ? DefaultCalculateSlackReserve(NumElements, NumBytesPerElement, true) : NumInlineElements;
+			return (NumElements <= NumInlineElements) ? NumInlineElements : DefaultCalculateSlackReserve(NumElements, NumBytesPerElement, true);
 		}
 
 		FORCEINLINE int32 CalculateSlackShrink(int32 NumElements, int32 NumAllocatedElements, int32 NumBytesPerElement) const
 		{
 			// If the elements use less space than the inline allocation, only use the inline allocation as slack.
-			return HasAllocation() ? DefaultCalculateSlackShrink(NumElements, NumAllocatedElements, NumBytesPerElement, true) : NumInlineElements;
+			return (NumElements <= NumInlineElements) ? NumInlineElements : DefaultCalculateSlackShrink(NumElements, NumAllocatedElements, NumBytesPerElement, true);
 		}
 
 		FORCEINLINE int32 CalculateSlackGrow(int32 NumElements, int32 NumAllocatedElements, int32 NumBytesPerElement) const
 		{
 			// If the elements use less space than the inline allocation, only use the inline allocation as slack.
-			return HasAllocation() ? DefaultCalculateSlackGrow(NumElements, NumAllocatedElements, NumBytesPerElement, true) : NumInlineElements;
+			return (NumElements <= NumInlineElements) ? NumInlineElements : DefaultCalculateSlackGrow(NumElements, NumAllocatedElements, NumBytesPerElement, true);
 		}
 
 		SIZE_T GetAllocatedSize(int32 NumAllocatedElements, SIZE_T NumBytesPerElement) const
