Fixed SupportsQuit

diff --git a/Engine/Source/Runtime/Core/Public/Android/AndroidProperties.h b/Engine/Source/Runtime/Core/Public/Android/AndroidProperties.h
--- a/Engine/Source/Runtime/Core/Public/Android/AndroidProperties.h
+++ b/Engine/Source/Runtime/Core/Public/Android/AndroidProperties.h
@@ -88,11 +88,6 @@ struct FAndroidPlatformProperties
 		return true;
 	}
 
-	static FORCEINLINE bool SupportsQuit() 
-	{
-		return true;
-	}
-
 	static FORCEINLINE bool HasFixedResolution()
 	{
 		return true;
diff --git a/Engine/Source/Runtime/Core/Public/IOS/IOSPlatformProperties.h b/Engine/Source/Runtime/Core/Public/IOS/IOSPlatformProperties.h
--- a/Engine/Source/Runtime/Core/Public/IOS/IOSPlatformProperties.h
+++ b/Engine/Source/Runtime/Core/Public/IOS/IOSPlatformProperties.h
@@ -42,11 +42,6 @@ struct FIOSPlatformProperties
 	{
 		return true;
 	}
-
-    static FORCEINLINE bool SupportsQuit()
-    {
-        return true;
-    }
     
 	static FORCEINLINE bool SupportsBuildTarget( EBuildTargets::Type BuildTarget )
 	{
