Original CLs 4985309, 4986255

Fixed a bug as uncovered by the new unit test (to be checked in in a separate submit).

* Added a new automation test file setup for animation
* Added a unit test for the calculate velocity node.

diff --git a/Engine/Source/Runtime/AnimGraphRuntime/Private/KismetAnimationLibrary.cpp b/Engine/Source/Runtime/AnimGraphRuntime/Private/KismetAnimationLibrary.cpp
--- a/Engine/Source/Runtime/AnimGraphRuntime/Private/KismetAnimationLibrary.cpp
+++ b/Engine/Source/Runtime/AnimGraphRuntime/Private/KismetAnimationLibrary.cpp
@@ -153,7 +153,7 @@ float UKismetAnimationLibrary::K2_CalculateVelocityFromPositionHistory(
 				NextIndex = 0;
 			}
 			History.Positions[NextIndex] = Position;
-			History.Velocities[History.LastIndex] = LengthOfV;
+			History.Velocities[NextIndex] = LengthOfV;
 			History.LastIndex = NextIndex;
 		}
 		else
diff --git a/Engine/Source/Runtime/Core/Private/Misc/AutomationTest.cpp b/Engine/Source/Runtime/Core/Private/Misc/AutomationTest.cpp
--- a/Engine/Source/Runtime/Core/Private/Misc/AutomationTest.cpp
+++ b/Engine/Source/Runtime/Core/Private/Misc/AutomationTest.cpp
@@ -882,6 +882,14 @@ void FAutomationTestBase::AddError(const FString& InError, int32 StackOffset)
 	}
 }
 
+void FAutomationTestBase::AddErrorIfFalse(bool bCondition, const FString& InError, int32 StackOffset)
+{
+	if (!bCondition)
+	{
+		AddError(InError, StackOffset);
+	}
+}
+
 void FAutomationTestBase::AddErrorS(const FString& InError, const FString& InFilename, int32 InLineNumber)
 {
 	if ( !bSuppressLogs && !IsExpectedError(InError))
diff --git a/Engine/Source/Runtime/Core/Public/Misc/AutomationTest.h b/Engine/Source/Runtime/Core/Public/Misc/AutomationTest.h
--- a/Engine/Source/Runtime/Core/Public/Misc/AutomationTest.h
+++ b/Engine/Source/Runtime/Core/Public/Misc/AutomationTest.h
@@ -1044,6 +1044,14 @@ public:
 	 */
 	virtual void AddError( const FString& InError, int32 StackOffset = 0 );
 
+	/**
+	 * Adds an error message to this test if the condition is false
+	 *
+	 * @param   bCondition The condition to validate.
+	 * @param   InError	   Error message to add to this test
+	 */
+	virtual void AddErrorIfFalse( bool bCondition, const FString& InError, int32 StackOffset = 0 );
+
 	/**
 	 * Adds an error message to this test
 	 *
