Core: Added support for using FString with RegisterConsoleVariableRef

diff --git a/Engine/Source/Runtime/Core/Private/HAL/ConsoleManager.cpp b/Engine/Source/Runtime/Core/Private/HAL/ConsoleManager.cpp
--- a/Engine/Source/Runtime/Core/Private/HAL/ConsoleManager.cpp
+++ b/Engine/Source/Runtime/Core/Private/HAL/ConsoleManager.cpp
@@ -414,7 +414,7 @@ template<> TConsoleVariableData<FString>* FConsoleVariable<FString>::AsVariableS
 
 // ----
 
-// T: int32, float
+// T: int32, float, bool
 template <class T>
 class FConsoleVariableRef : public FConsoleVariableBase
 {
@@ -485,7 +485,73 @@ FString FConsoleVariableRef<float>::GetString() const
 	return FString::SanitizeFloat(RefValue);
 }
 
-// ----
+// string version
+
+class FConsoleVariableStringRef : public FConsoleVariableBase
+{
+public:
+	FConsoleVariableStringRef(FString& InRefValue, const TCHAR* Help, EConsoleVariableFlags Flags)
+		: FConsoleVariableBase(Help, Flags)
+		, RefValue(InRefValue)
+		, MainValue(InRefValue)
+	{
+	}
+
+	// interface IConsoleVariable -----------------------------------
+
+	virtual void Release()
+	{
+		delete this;
+	}
+	virtual void Set(const TCHAR* InValue, EConsoleVariableFlags SetBy)
+	{
+		if (CanChange(SetBy))
+		{
+			MainValue = InValue;
+			OnChanged(SetBy);
+		}
+	}
+	virtual int32 GetInt() const
+	{
+		int32 Result;
+		TTypeFromString<int32>::FromString(Result, *MainValue);
+		return Result;
+	}
+	virtual float GetFloat() const
+	{
+		float Result = 0.0f;
+		TTypeFromString<float>::FromString(Result, *MainValue);
+		return Result;
+	}
+	virtual FString GetString() const
+	{
+		return MainValue;
+	}
+
+private: // ----------------------------------------------------
+
+	// reference the the value (should not be changed from outside), if ECVF_RenderThreadSafe this is the render thread version, otherwise same as MainValue
+	FString& RefValue;
+	// main thread version 
+	FString MainValue;
+
+	const FString& Value() const
+	{
+		uint32 Index = GetShadowIndex();
+		checkSlow(Index < 2);
+		return (Index == 0) ? MainValue : RefValue;
+	}
+
+	void OnChanged(EConsoleVariableFlags SetBy)
+	{
+		if (CanChange(SetBy))
+		{
+			// propagate from main thread to render thread or to reference
+			OnCVarChange(RefValue, MainValue, Flags);
+			FConsoleVariableBase::OnChanged(SetBy);
+		}
+	}
+};
 
 class FConsoleVariableBitRef : public FConsoleVariableBase
 {
@@ -858,6 +924,11 @@ IConsoleVariable* FConsoleManager::RegisterConsoleVariableRef(const TCHAR* Name,
 	return AddConsoleObject(Name, new FConsoleVariableRef<bool>(RefValue, Help, (EConsoleVariableFlags)Flags))->AsVariable();
 }
 
+IConsoleVariable* FConsoleManager::RegisterConsoleVariableRef(const TCHAR* Name, FString& RefValue, const TCHAR* Help, uint32 Flags)
+{
+	return AddConsoleObject(Name, new FConsoleVariableStringRef(RefValue, Help, (EConsoleVariableFlags)Flags))->AsVariable();
+}
+
 IConsoleCommand* FConsoleManager::RegisterConsoleCommand(const TCHAR* Name, const TCHAR* Help, const FConsoleCommandDelegate& Command, uint32 Flags)
 {
 	return AddConsoleObject(Name, new FConsoleCommand(Command, Help, (EConsoleVariableFlags)Flags))->AsCommand();
diff --git a/Engine/Source/Runtime/Core/Public/HAL/ConsoleManager.h b/Engine/Source/Runtime/Core/Public/HAL/ConsoleManager.h
--- a/Engine/Source/Runtime/Core/Public/HAL/ConsoleManager.h
+++ b/Engine/Source/Runtime/Core/Public/HAL/ConsoleManager.h
@@ -52,6 +52,7 @@ public:
 	virtual IConsoleVariable* RegisterConsoleVariableRef(const TCHAR* Name, int32& RefValue, const TCHAR* Help, uint32 Flags) override;
 	virtual IConsoleVariable* RegisterConsoleVariableRef(const TCHAR* Name, float& RefValue, const TCHAR* Help, uint32 Flags) override;
 	virtual IConsoleVariable* RegisterConsoleVariableRef(const TCHAR* Name, bool& RefValue, const TCHAR* Help, uint32 Flags) override;
+	virtual IConsoleVariable* RegisterConsoleVariableRef(const TCHAR* Name, FString& RefValue, const TCHAR* Help, uint32 Flags) override;
 	virtual IConsoleVariable* RegisterConsoleVariableBitRef(const TCHAR* CVarName, const TCHAR* FlagName, uint32 BitNumber, uint8* Force0MaskPtr, uint8* Force1MaskPtr, const TCHAR* Help, uint32 Flags) override;
 	
 	virtual void CallAllConsoleVariableSinks() override;
diff --git a/Engine/Source/Runtime/Core/Public/HAL/IConsoleManager.h b/Engine/Source/Runtime/Core/Public/HAL/IConsoleManager.h
--- a/Engine/Source/Runtime/Core/Public/HAL/IConsoleManager.h
+++ b/Engine/Source/Runtime/Core/Public/HAL/IConsoleManager.h
@@ -495,6 +495,13 @@ struct CORE_API IConsoleManager
 	* @param Flags bitmask combined from EConsoleVariableFlags
 	*/
 	virtual IConsoleVariable* RegisterConsoleVariableRef(const TCHAR* Name, bool& RefValue, const TCHAR* Help, uint32 Flags = ECVF_Default) = 0;
+	/**
+	* Create a reference to a string console variable
+	* @param Name must not be 0
+	* @param Help must not be 0
+	* @param Flags bitmask combined from EConsoleVariableFlags
+	*/
+	virtual IConsoleVariable* RegisterConsoleVariableRef(const TCHAR* Name, FString& RefValue, const TCHAR* Help, uint32 Flags = ECVF_Default) = 0;
 	/**
 	 * Create a reference to a show flag variable
 	 * @param CVarName must not be 0, e.g. "Show.PostProcessing"
