[LLM] task graph tasks inherit the LLM tag scope of the thread that created them.

diff --git a/Engine/Source/Runtime/Core/Public/Async/TaskGraphInterfaces.h b/Engine/Source/Runtime/Core/Public/Async/TaskGraphInterfaces.h
--- a/Engine/Source/Runtime/Core/Public/Async/TaskGraphInterfaces.h
+++ b/Engine/Source/Runtime/Core/Public/Async/TaskGraphInterfaces.h
@@ -18,6 +18,7 @@
 #include "Stats/Stats.h"
 #include "HAL/IConsoleManager.h"
 #include "HAL/Event.h"
+#include "HAL/LowLevelMemTracker.h"
 #include "Templates/RefCounting.h"
 #include "Containers/LockFreeFixedSizeAllocator.h"
 #include "Misc/MemStack.h"
@@ -432,6 +433,7 @@ protected:
 		, NumberOfPrerequistitesOutstanding(InNumberOfPrerequistitesOutstanding + 1) // + 1 is not a prerequisite, it is a lock to prevent it from executing while it is getting prerequisites, one it is safe to execute, call PrerequisitesComplete
 	{
 		checkThreadGraph(LifeStage.Increment() == int32(LS_Contructed));
+		LLM(InheritedLLMTag = FLowLevelMemTracker::bIsDisabled ? ELLMTag::Untagged : (ELLMTag)FLowLevelMemTracker::Get().GetActiveTag(ELLMTracker::Default));
 	}
 	/** 
 	 *	Sets the desired execution thread. This is not part of the constructor because this information may not be known quite yet duiring construction.
@@ -507,6 +509,7 @@ private:
 	 **/
 	FORCEINLINE void Execute(TArray<FBaseGraphTask*>& NewTasks, ENamedThreads::Type CurrentThread)
 	{
+		LLM_SCOPE(InheritedLLMTag);
 		checkThreadGraph(LifeStage.Increment() == int32(LS_Executing));
 		ExecuteTask(NewTasks, CurrentThread);
 	}
@@ -545,6 +548,8 @@ private:
 	FThreadSafeCounter			LifeStage;
 
 #endif
+
+	LLM(ELLMTag InheritedLLMTag);
 };
 
 /** 
