FGenericCrashContext::AddGameData -> SetGameData, and allow it to alter keys that already exist. Tweak header docs to clarify this behavior.
#rnx

diff --git a/Engine/Source/Runtime/Core/Private/GenericPlatform/GenericPlatformCrashContext.cpp b/Engine/Source/Runtime/Core/Private/GenericPlatform/GenericPlatformCrashContext.cpp
--- a/Engine/Source/Runtime/Core/Private/GenericPlatform/GenericPlatformCrashContext.cpp
+++ b/Engine/Source/Runtime/Core/Private/GenericPlatform/GenericPlatformCrashContext.cpp
@@ -19,11 +19,14 @@
 #include "Misc/EngineBuildSettings.h"
 #include "Stats/Stats.h"
 #include "Internationalization/TextLocalizationManager.h"
+#include "Logging/LogScopedCategoryAndVerbosityOverride.h"
 
 #ifndef NOINITCRASHREPORTER
 #define NOINITCRASHREPORTER 0
 #endif
 
+DEFINE_LOG_CATEGORY_STATIC(LogCrashContext, Display, All);
+
 extern CORE_API bool GIsGPUCrashed;
 
 /*-----------------------------------------------------------------------------
@@ -668,15 +671,31 @@ void FGenericCrashContext::ResetGameData()
 	NCachedCrashContextProperties::GameData.Reset();
 }
 
-void FGenericCrashContext::AddGameData(const FString& Key, const FString& Value)
+void FGenericCrashContext::SetGameData(const FString& Key, const FString& Value)
 {
 	if (Value.Len() == 0)
 	{
+		// for testing purposes, only log values when they change, but don't pay the lookup price normally.
+		UE_SUPPRESS(LogCrashContext, VeryVerbose, 
+		{
+			if (NCachedCrashContextProperties::GameData.Find(Key))
+			{
+				UE_LOG(LogCrashContext, VeryVerbose, TEXT("FGenericCrashContext::SetGameData(%s, <RemoveKey>)"), *Key);
+			}
+		});
 		NCachedCrashContextProperties::GameData.Remove(Key);
 	}
 	else
 	{
-		NCachedCrashContextProperties::GameData.Add(Key, Value);
+		FString& OldVal = NCachedCrashContextProperties::GameData.FindOrAdd(Key);
+		UE_SUPPRESS(LogCrashContext, VeryVerbose, 
+		{
+			if (OldVal != Value)
+			{
+				UE_LOG(LogCrashContext, VeryVerbose, TEXT("FGenericCrashContext::SetGameData(%s, %s)"), *Key, *Value);
+			}
+		});
+		OldVal = Value;
 	}
 }
 
diff --git a/Engine/Source/Runtime/Core/Public/GenericPlatform/GenericPlatformCrashContext.h b/Engine/Source/Runtime/Core/Public/GenericPlatform/GenericPlatformCrashContext.h
--- a/Engine/Source/Runtime/Core/Public/GenericPlatform/GenericPlatformCrashContext.h
+++ b/Engine/Source/Runtime/Core/Public/GenericPlatform/GenericPlatformCrashContext.h
@@ -214,8 +214,8 @@ public:
 	/** Clears the game data dictionary */
 	static void ResetGameData();
 
-	/** Adds arbitrary game data to the crash context */
-	static void AddGameData(const FString& Key, const FString& Value);
+	/** Updates (or adds if not already present) arbitrary game data to the crash context (will remove the key if passed an empty string) */
+	static void SetGameData(const FString& Key, const FString& Value);
 
 	/** Adds a plugin descriptor string to the enabled plugins list in the crash context */
 	static void AddPlugin(const FString& PluginDesc);
