[LLM] FMemory::Realloc can retain the original pointer's tag via a scoping helper.

manually merged for robomerge.

diff --git a/Engine/Source/Runtime/Core/Private/HAL/LowLevelMemTracker.cpp b/Engine/Source/Runtime/Core/Private/HAL/LowLevelMemTracker.cpp
--- a/Engine/Source/Runtime/Core/Private/HAL/LowLevelMemTracker.cpp
+++ b/Engine/Source/Runtime/Core/Private/HAL/LowLevelMemTracker.cpp
@@ -348,6 +348,7 @@ public:
 	int64 GetTagAmount(ELLMTag Tag) const;
 	void SetTagAmount(ELLMTag Tag, int64 Amount, bool AddToTotal);
     int64 GetActiveTag();
+	int64 FindTagForPtr( void* Ptr );
 
 	int64 GetAllocTypeAmount(ELLMAllocType AllocType);
 
@@ -1032,6 +1033,47 @@ FLLMPauseScope::~FLLMPauseScope()
 }
 
 
+FLLMScopeFromPtr::FLLMScopeFromPtr(void* Ptr, ELLMTracker Tracker )
+	: TrackerSet(Tracker)
+	, Enabled(false)
+{
+	if(GIsRequestingExit || Ptr == nullptr)
+	{
+		return;
+	}
+
+	FLowLevelMemTracker& LLM = FLowLevelMemTracker::Get();
+	if (!LLM.IsEnabled())
+	{
+		return;
+	}
+
+	int64 Tag = LLM.GetTracker(TrackerSet)->FindTagForPtr( Ptr );
+	if( Tag != (int64)ELLMTag::Untagged )
+	{
+		LLM.GetTracker(TrackerSet)->PushTag(Tag);
+		Enabled = true;
+	}
+}
+
+FLLMScopeFromPtr::~FLLMScopeFromPtr()
+{
+	if (!Enabled)
+	{
+		return;
+	}
+
+	FLowLevelMemTracker& LLM = FLowLevelMemTracker::Get();
+	if (!LLM.IsEnabled())
+	{
+		return;
+	}
+
+	LLM.GetTracker(TrackerSet)->PopTag();
+}
+
+
+
 
 
 
@@ -1327,6 +1369,16 @@ int64 FLLMTracker::GetActiveTag()
     return State->GetTopTag();
 }
 
+int64 FLLMTracker::FindTagForPtr( void* Ptr )
+{
+#if LLM_USE_ALLOC_INFO_STRUCT
+	return GetAllocationMap().GetValue(Ptr).Value2.Tag;
+#else
+	return (int64)GetAllocationMap().GetValue(Ptr).Value2;
+#endif
+}
+
+
 FLLMTracker::FLLMThreadState::FLLMThreadState()
 	: UntaggedAllocs(0)
 {
diff --git a/Engine/Source/Runtime/Core/Public/HAL/FMemory.inl b/Engine/Source/Runtime/Core/Public/HAL/FMemory.inl
--- a/Engine/Source/Runtime/Core/Public/HAL/FMemory.inl
+++ b/Engine/Source/Runtime/Core/Public/HAL/FMemory.inl
@@ -34,6 +34,7 @@ FMEMORY_INLINE_FUNCTION_DECORATOR void* FMemory::Malloc(SIZE_T Count, uint32 Ali
 FMEMORY_INLINE_FUNCTION_DECORATOR void* FMemory::Realloc(void* Original, SIZE_T Count, uint32 Alignment)
 {
 	// optional tracking of every allocation
+	LLM_REALLOC_SCOPE(Original);
 	LLM_IF_ENABLED(FLowLevelMemTracker::Get().OnLowLevelFree(ELLMTracker::Default, Original, ELLMAllocType::FMalloc));
 
 	void* Ptr;
diff --git a/Engine/Source/Runtime/Core/Public/HAL/LowLevelMemTracker.h b/Engine/Source/Runtime/Core/Public/HAL/LowLevelMemTracker.h
--- a/Engine/Source/Runtime/Core/Public/HAL/LowLevelMemTracker.h
+++ b/Engine/Source/Runtime/Core/Public/HAL/LowLevelMemTracker.h
@@ -29,7 +29,7 @@
 #endif
 
 	// using asset tagging requires a significantly higher number of per-thread tags, so make it optional
-	// even if this is on, we still need to run with -llmtagsets=assets because of the shear number of stat ids it makes
+	// even if this is on, we still need to run with -llmtagsets=assets because of the sheer number of stat ids it makes
 	// LLM Assets can be viewed in game using 'Stat LLMAssets'
 	#define LLM_ALLOW_ASSETS_TAGS 0
 
@@ -297,6 +297,12 @@ extern FName LLMGetTagStat(ELLMTag Tag);
 #define LLM_SCOPED_PAUSE_TRACKING_FOR_TRACKER(Tracker, AllocType) FLLMPauseScope SCOPE_NAME(NAME_None, 0, Tracker, AllocType);
 #define LLM_SCOPED_PAUSE_TRACKING_WITH_ENUM_AND_AMOUNT(Tag, Amount, Tracker, AllocType) FLLMPauseScope SCOPE_NAME(Tag, Amount, Tracker, AllocType);
 
+/**
+ * LLM realloc scope macros. Used when reallocating a pointer and you wish to retain the tagging from the source pointer
+ */
+#define LLM_REALLOC_SCOPE(Ptr) FLLMScopeFromPtr SCOPE_NAME(Ptr, ELLMTracker::Default);
+#define LLM_REALLOC_PLATFORM_SCOPE(Ptr) FLLMScopeFromPtr SCOPE_NAME(Ptr, ELLMTracker::Platform);
+
 /**
  * LLM Stat scope macros (if stats system is enabled)
  */
@@ -467,6 +473,7 @@ private:
 
 	friend class FLLMScope;
 	friend class FLLMPauseScope;
+	friend class FLLMScopeFromPtr;
 
 	FLLMAllocator Allocator;
 	
@@ -525,11 +532,27 @@ protected:
 	ELLMAllocType AllocType;
 };
 
+/*
+* LLM scope for inheriting tag from the given address
+*/
+class CORE_API FLLMScopeFromPtr
+{
+public:
+	FLLMScopeFromPtr(void* Ptr, ELLMTracker Tracker);
+	~FLLMScopeFromPtr();
+protected:
+	ELLMTracker TrackerSet;
+	bool Enabled;
+};
+
+
 #else
 	#define LLM(...)
 	#define LLM_IF_ENABLED(...)
 	#define LLM_SCOPE(...)
 	#define LLM_PLATFORM_SCOPE(...)
+	#define LLM_REALLOC_SCOPE(...)
+	#define LLM_REALLOC_PLATFORM_SCOPE(...)
 	#define LLM_SCOPED_TAG_WITH_STAT(...)
 	#define LLM_SCOPED_TAG_WITH_STAT_IN_SET(...)
 	#define LLM_SCOPED_TAG_WITH_STAT_NAME(...)
