- Fixed WCHAR_TO_TCHAR macro

diff --git a/Engine/Source/Runtime/Core/Public/GenericPlatform/GenericWidePlatformString.h b/Engine/Source/Runtime/Core/Public/GenericPlatform/GenericWidePlatformString.h
--- a/Engine/Source/Runtime/Core/Public/GenericPlatform/GenericWidePlatformString.h
+++ b/Engine/Source/Runtime/Core/Public/GenericPlatform/GenericWidePlatformString.h
@@ -75,6 +75,21 @@ public:
 		return Length;
 	}
 
+#if PLATFORM_TCHAR_IS_CHAR16
+	static int32 Strlen( const wchar_t* String )
+	{
+		int32 Length = -1;
+		
+		do
+		{
+			Length++;
+		}
+		while (*String++);
+		
+		return Length;
+	}
+#endif
+	
 	static const WIDECHAR* Strstr( const WIDECHAR* String, const WIDECHAR* Find)
 	{
 		WIDECHAR Char1, Char2;
