Added handling of 64-bit Mach-O segments in the portable call stack generation code for Apple platforms.

diff --git a/Engine/Source/Runtime/Core/Private/Apple/ApplePlatformStackWalk.cpp b/Engine/Source/Runtime/Core/Private/Apple/ApplePlatformStackWalk.cpp
--- a/Engine/Source/Runtime/Core/Private/Apple/ApplePlatformStackWalk.cpp
+++ b/Engine/Source/Runtime/Core/Private/Apple/ApplePlatformStackWalk.cpp
@@ -45,7 +45,12 @@ int32 GetModuleImageSize( const struct mach_header* Header )
 			struct segment_command *SegmentCommand = (struct segment_command *) CurrentCommand;
 			ModuleSize += SegmentCommand->vmsize;
 		}
-		
+		else if( CurrentCommand->cmd == LC_SEGMENT_64 )
+		{
+			struct segment_command_64 *SegmentCommand = (struct segment_command_64 *) CurrentCommand;
+			ModuleSize += SegmentCommand->vmsize;
+		}
+
 		CurrentCommand = (struct load_command *)( (char *)CurrentCommand + CurrentCommand->cmdsize );
 	}
 	
