Add a FPlatformProcess::ExecutablePath() function, and use it to figure out which executable to launch for commandlets, PIE, etc...

Saves having to attempt to reconstruct the executable path using BaseDir(), which may not actually be the executable directory we want (eg. app bundle on Mac).

diff --git a/Engine/Source/Editor/UnrealEd/Private/PlayLevel.cpp b/Engine/Source/Editor/UnrealEd/Private/PlayLevel.cpp
--- a/Engine/Source/Editor/UnrealEd/Private/PlayLevel.cpp
+++ b/Engine/Source/Editor/UnrealEd/Private/PlayLevel.cpp
@@ -1470,7 +1470,7 @@ void UEditorEngine::PlayStandaloneLocalPc(FString MapNameOverride, FIntPoint* Wi
 	}
 
 	// launch the game process
-	FString GamePath = FString(FPlatformProcess::BaseDir()) / FString(FPlatformProcess::ExecutableName(false));
+	FString GamePath = FPlatformProcess::ExecutablePath();
 	FPlayOnPCInfo *NewSession = new (PlayOnLocalPCSessions) FPlayOnPCInfo();
 
 	uint32 ProcessID = 0;
diff --git a/Engine/Source/Editor/UnrealEd/Private/UnrealEdMisc.cpp b/Engine/Source/Editor/UnrealEd/Private/UnrealEdMisc.cpp
--- a/Engine/Source/Editor/UnrealEd/Private/UnrealEdMisc.cpp
+++ b/Engine/Source/Editor/UnrealEd/Private/UnrealEdMisc.cpp
@@ -1790,7 +1790,7 @@ bool FUnrealEdMisc::GetURL( const TCHAR* InKey, FString& OutURL, const bool bChe
 
 FString FUnrealEdMisc::GetExecutableForCommandlets() const
 {
-	FString ExecutableName = FString(FPlatformProcess::BaseDir()) / FString(FPlatformProcess::ExecutableName(false));
+	FString ExecutableName = FString(FPlatformProcess::ExecutablePath());
 #if PLATFORM_WINDOWS
 	// turn UE4editor into UE4editor-cmd
 	if(ExecutableName.EndsWith(".exe", ESearchCase::IgnoreCase) && !FPaths::GetBaseFilename(ExecutableName).EndsWith("-cmd", ESearchCase::IgnoreCase))
diff --git a/Engine/Source/Runtime/Core/Private/GenericPlatform/GenericPlatformProcess.cpp b/Engine/Source/Runtime/Core/Private/GenericPlatform/GenericPlatformProcess.cpp
--- a/Engine/Source/Runtime/Core/Private/GenericPlatform/GenericPlatformProcess.cpp
+++ b/Engine/Source/Runtime/Core/Private/GenericPlatform/GenericPlatformProcess.cpp
@@ -154,6 +154,12 @@ void FGenericPlatformProcess::CleanShaderWorkingDir()
 	IFileManager::Get().DeleteDirectory(*LegacyShaderWorkingDirectory, false, true);
 }
 
+const TCHAR* FGenericPlatformProcess::ExecutablePath()
+{
+	UE_LOG(LogHAL, Fatal, TEXT("FGenericPlatformProcess::ExecutablePath not implemented on this platform"));
+	return NULL;
+}
+
 const TCHAR* FGenericPlatformProcess::ExecutableName(bool bRemoveExtension)
 {
 	UE_LOG(LogHAL, Fatal, TEXT("FGenericPlatformProcess::ExecutableName not implemented on this platform"));
diff --git a/Engine/Source/Runtime/Core/Private/Mac/MacPlatformProcess.cpp b/Engine/Source/Runtime/Core/Private/Mac/MacPlatformProcess.cpp
--- a/Engine/Source/Runtime/Core/Private/Mac/MacPlatformProcess.cpp
+++ b/Engine/Source/Runtime/Core/Private/Mac/MacPlatformProcess.cpp
@@ -920,6 +920,18 @@ FString FMacPlatformProcess::GetCurrentWorkingDirectory()
 	return UTF8_TO_TCHAR(CurrentDir);
 }
 
+const TCHAR* FMacPlatformProcess::ExecutablePath()
+{
+	static TCHAR Result[512]=TEXT("");
+	if( !Result[0] )
+	{
+		SCOPED_AUTORELEASE_POOL;
+		NSString *NSExeName = [[NSBundle mainBundle] executablePath];
+		FPlatformString::CFStringToTCHAR( ( CFStringRef )NSExeName, Result );
+	}
+	return Result;
+}
+
 const TCHAR* FMacPlatformProcess::ExecutableName(bool bRemoveExtension)
 {
 	static TCHAR Result[512]=TEXT("");
diff --git a/Engine/Source/Runtime/Core/Private/Unix/UnixPlatformProcess.cpp b/Engine/Source/Runtime/Core/Private/Unix/UnixPlatformProcess.cpp
--- a/Engine/Source/Runtime/Core/Private/Unix/UnixPlatformProcess.cpp
+++ b/Engine/Source/Runtime/Core/Private/Unix/UnixPlatformProcess.cpp
@@ -319,6 +319,30 @@ bool FUnixPlatformProcess::SetProcessLimits(EProcessResource::Type Resource, uin
 	return true;
 }
 
+
+const TCHAR* FUnixPlatformProcess::ExecutablePath()
+{
+	static bool bHaveResult = false;
+	static TCHAR CachedResult[ PlatformProcessLimits::MaxBaseDirLength ];
+	if (!bHaveResult)
+	{
+		char SelfPath[ PlatformProcessLimits::MaxBaseDirLength ] = {0};
+		if (readlink( "/proc/self/exe", SelfPath, ARRAY_COUNT(SelfPath) - 1) == -1)
+		{
+			int ErrNo = errno;
+			UE_LOG(LogHAL, Fatal, TEXT("readlink() failed with errno = %d (%s)"), ErrNo,
+				StringCast< TCHAR >(strerror(ErrNo)).Get());
+			return CachedResult;
+		}
+		SelfPath[ARRAY_COUNT(SelfPath) - 1] = 0;
+
+		FCString::Strcpy(CachedResult, ARRAY_COUNT(CachedResult) - 1, UTF8_TO_TCHAR(SelfPath));
+		CachedResult[ARRAY_COUNT(CachedResult) - 1] = 0;
+		bHaveResult = true;
+	}
+	return CachedResult;
+}
+
 const TCHAR* FUnixPlatformProcess::ExecutableName(bool bRemoveExtension)
 {
 	static bool bHaveResult = false;
diff --git a/Engine/Source/Runtime/Core/Private/Windows/WindowsPlatformProcess.cpp b/Engine/Source/Runtime/Core/Private/Windows/WindowsPlatformProcess.cpp
--- a/Engine/Source/Runtime/Core/Private/Windows/WindowsPlatformProcess.cpp
+++ b/Engine/Source/Runtime/Core/Private/Windows/WindowsPlatformProcess.cpp
@@ -953,6 +953,19 @@ const FString FWindowsPlatformProcess::ShaderWorkingDir()
 }
 
 
+const TCHAR* FWindowsPlatformProcess::ExecutablePath()
+{
+	static TCHAR Result[512]=TEXT("");
+	if( !Result[0] )
+	{
+		if ( !GetModuleFileName( hInstance, Result, ARRAY_COUNT(Result) ) )
+		{
+			Result[0] = 0;
+		}
+	}
+	return Result;
+}
+
 const TCHAR* FWindowsPlatformProcess::ExecutableName(bool bRemoveExtension)
 {
 	static TCHAR Result[512]=TEXT("");
diff --git a/Engine/Source/Runtime/Core/Public/GenericPlatform/GenericPlatformProcess.h b/Engine/Source/Runtime/Core/Public/GenericPlatform/GenericPlatformProcess.h
--- a/Engine/Source/Runtime/Core/Public/GenericPlatform/GenericPlatformProcess.h
+++ b/Engine/Source/Runtime/Core/Public/GenericPlatform/GenericPlatformProcess.h
@@ -242,6 +242,13 @@ struct CORE_API FGenericPlatformProcess
 	/**	Clean the shader working directory. */
 	static void CleanShaderWorkingDir();
 
+	/**
+	 * Return the path to the currently running executable
+	 *
+	 * @return 	Path of the currently running executable
+	 */
+	static const TCHAR* ExecutablePath();
+
 	/**
 	 * Return the name of the currently running executable
 	 *
diff --git a/Engine/Source/Runtime/Core/Public/Mac/MacPlatformProcess.h b/Engine/Source/Runtime/Core/Public/Mac/MacPlatformProcess.h
--- a/Engine/Source/Runtime/Core/Public/Mac/MacPlatformProcess.h
+++ b/Engine/Source/Runtime/Core/Public/Mac/MacPlatformProcess.h
@@ -110,6 +110,7 @@ struct CORE_API FMacPlatformProcess : public FGenericPlatformProcess
 	static const TCHAR* UserName(bool bOnlyAlphaNumeric = true);
 	static void SetCurrentWorkingDirectoryToBaseDir();
 	static FString GetCurrentWorkingDirectory();
+	static const TCHAR* ExecutablePath();
 	static const TCHAR* ExecutableName(bool bRemoveExtension = true);
 	static FString GenerateApplicationPath( const FString& AppName, EBuildConfigurations::Type BuildConfiguration);
 	static const TCHAR* GetModuleExtension();
diff --git a/Engine/Source/Runtime/Core/Public/Unix/UnixPlatformProcess.h b/Engine/Source/Runtime/Core/Public/Unix/UnixPlatformProcess.h
--- a/Engine/Source/Runtime/Core/Public/Unix/UnixPlatformProcess.h
+++ b/Engine/Source/Runtime/Core/Public/Unix/UnixPlatformProcess.h
@@ -250,6 +250,7 @@ struct CORE_API FUnixPlatformProcess : public FGenericPlatformProcess
 	static FString GenerateApplicationPath(const FString& AppName, EBuildConfigurations::Type BuildConfiguration);
 	static FString GetApplicationName( uint32 ProcessId );
 	static bool SetProcessLimits(EProcessResource::Type Resource, uint64 Limit);
+	static const TCHAR* ExecutablePath();
 	static const TCHAR* ExecutableName(bool bRemoveExtension = true);
 	static const TCHAR* GetModulePrefix();
 	static const TCHAR* GetModuleExtension();
diff --git a/Engine/Source/Runtime/Core/Public/Windows/WindowsPlatformProcess.h b/Engine/Source/Runtime/Core/Public/Windows/WindowsPlatformProcess.h
--- a/Engine/Source/Runtime/Core/Public/Windows/WindowsPlatformProcess.h
+++ b/Engine/Source/Runtime/Core/Public/Windows/WindowsPlatformProcess.h
@@ -140,6 +140,7 @@ public:
 	static void SetCurrentWorkingDirectoryToBaseDir();
 	static FString GetCurrentWorkingDirectory();
 	static const FString ShaderWorkingDir();
+	static const TCHAR* ExecutablePath();
 	static const TCHAR* ExecutableName(bool bRemoveExtension = true);
 	static FString GenerateApplicationPath( const FString& AppName, EBuildConfigurations::Type BuildConfiguration);
 	static const TCHAR* GetModuleExtension();
