In timecode, use of round for double instead of round for float.

diff --git a/Engine/Source/Runtime/Core/Private/Tests/Misc/TimecodeTest.cpp b/Engine/Source/Runtime/Core/Private/Tests/Misc/TimecodeTest.cpp
--- a/Engine/Source/Runtime/Core/Private/Tests/Misc/TimecodeTest.cpp
+++ b/Engine/Source/Runtime/Core/Private/Tests/Misc/TimecodeTest.cpp
@@ -194,7 +194,7 @@ bool FTimecodeTest::RunTest(const FString& Parameters)
 	for (const TFuture<bool>& Future : Futures)
 	{
 		Future.Wait();
-		bSuccessfully == bSuccessfully && Future.Get();
+		bSuccessfully = bSuccessfully && Future.Get();
 	}
 
 	return bSuccessfully;
diff --git a/Engine/Source/Runtime/Core/Public/Misc/Timecode.h b/Engine/Source/Runtime/Core/Public/Misc/Timecode.h
--- a/Engine/Source/Runtime/Core/Public/Misc/Timecode.h
+++ b/Engine/Source/Runtime/Core/Public/Misc/Timecode.h
@@ -47,7 +47,7 @@ struct FTimecode
 	 */
 	explicit FTimecode(double InSeconds, const FFrameRate& InFrameRate, bool InbDropFrame)
 	{
-		int32 NumberOfFrames = InbDropFrame ? (int32)FMath::RoundToInt(InSeconds * InFrameRate.AsDecimal()) : (int32)FMath::RoundToZero(InSeconds * FMath::RoundToInt(InFrameRate.AsDecimal()));
+		int32 NumberOfFrames = InbDropFrame ? (int32)FMath::RoundToDouble(InSeconds * InFrameRate.AsDecimal()) : (int32)FMath::RoundToDouble(InSeconds * FMath::RoundToDouble(InFrameRate.AsDecimal()));
 		*this = FromFrameNumber(FFrameNumber(NumberOfFrames), InFrameRate, InbDropFrame);
 	}
 
@@ -188,7 +188,7 @@ public:
 		const FFrameNumber ConvertedFrameNumber = ToFrameNumber(InFrameRate);
 		const double NumberOfSeconds = bDropFrameFormat
 				? ConvertedFrameNumber.Value * InFrameRate.AsInterval()
-				: (double)ConvertedFrameNumber.Value / FMath::RoundToInt(InFrameRate.AsDecimal());
+				: (double)ConvertedFrameNumber.Value / FMath::RoundToDouble(InFrameRate.AsDecimal());
 		return FTimespan::FromSeconds(NumberOfSeconds);
 	}
 
@@ -266,7 +266,7 @@ public:
 
 inline bool operator==(const FTimecode& A, const FTimecode& B)
 {
-	return A.Hours == B.Hours && A.Minutes == B.Minutes && A.Seconds == B.Seconds && A.Frames == A.Frames;
+	return A.Hours == B.Hours && A.Minutes == B.Minutes && A.Seconds == B.Seconds && A.Frames == B.Frames;
 }
 
 inline bool operator!=(const FTimecode& A, const FTimecode& B)
