Core - Adding constructors to allow going from a FVector4 to an FVector2D, like we have for FVector.

diff --git a/Engine/Source/Runtime/Core/Public/Math/Vector2D.h b/Engine/Source/Runtime/Core/Public/Math/Vector2D.h
--- a/Engine/Source/Runtime/Core/Public/Math/Vector2D.h
+++ b/Engine/Source/Runtime/Core/Public/Math/Vector2D.h
@@ -83,6 +83,14 @@ public:
 	 */
 	explicit FORCEINLINE FVector2D(const FVector& V);
 
+	/**
+	 * Constructs a vector from an FVector4.
+	 * Copies the X and Y components from the FVector4.
+	 *
+	 * @param V Vector to copy from.
+	 */
+	explicit FORCEINLINE FVector2D(const FVector4& V);
+
 public:
 
 	/**
diff --git a/Engine/Source/Runtime/Core/Public/Math/Vector4.h b/Engine/Source/Runtime/Core/Public/Math/Vector4.h
--- a/Engine/Source/Runtime/Core/Public/Math/Vector4.h
+++ b/Engine/Source/Runtime/Core/Public/Math/Vector4.h
@@ -756,3 +756,11 @@ FORCEINLINE FVector::FVector( const FVector4& V )
 {
 	DiagnosticCheckNaN();
 }
+
+/* FVector2D inline functions
+ *****************************************************************************/
+
+FORCEINLINE FVector2D::FVector2D(const FVector4& V)
+	: X(V.X), Y(V.Y)
+{
+}
\ No newline at end of file
