Allow a command line to set a signal to use the default signal handler vs ignoring it

diff --git a/Engine/Source/Runtime/Core/Private/Unix/UnixPlatformCrashContext.cpp b/Engine/Source/Runtime/Core/Private/Unix/UnixPlatformCrashContext.cpp
--- a/Engine/Source/Runtime/Core/Private/Unix/UnixPlatformCrashContext.cpp
+++ b/Engine/Source/Runtime/Core/Private/Unix/UnixPlatformCrashContext.cpp
@@ -735,6 +735,9 @@ void FUnixPlatformMisc::SetGracefulTerminationHandler()
 // reserve stack for the main thread in BSS
 char FRunnableThreadUnix::MainThreadSignalHandlerStack[FRunnableThreadUnix::EConstants::CrashHandlerStackSize];
 
+// Defined in UnixPlatformMemory.cpp. Allows settings a specific signal to maintain its default handler rather then ignoring it
+extern int32 GSignalToDefault;
+
 void FUnixPlatformMisc::SetCrashHandler(void (* CrashHandler)(const FGenericCrashContext & Context))
 {
 	GCrashHandlerPointer = CrashHandler;
@@ -796,6 +799,11 @@ void FUnixPlatformMisc::SetCrashHandler(void (* CrashHandler)(const FGenericCras
 			}
 		}
 
+		if (GSignalToDefault && Signal == GSignalToDefault)
+		{
+			bSignalShouldBeIgnored = false;
+		}
+
 		if (bSignalShouldBeIgnored)
 		{
 			sigaction(Signal, &Action, nullptr);
diff --git a/Engine/Source/Runtime/Core/Private/Unix/UnixPlatformMemory.cpp b/Engine/Source/Runtime/Core/Private/Unix/UnixPlatformMemory.cpp
--- a/Engine/Source/Runtime/Core/Private/Unix/UnixPlatformMemory.cpp
+++ b/Engine/Source/Runtime/Core/Private/Unix/UnixPlatformMemory.cpp
@@ -52,6 +52,9 @@ bool CORE_API GKSMMergeAllPages = false;
 // Used to enable or disable timing of ensures. Enabled by default
 bool CORE_API GTimeEnsures = true;
 
+// Allows settings a specific signal to maintain its default handler rather then ignoring the signal
+int32 CORE_API GSignalToDefault = 0;
+
 #if UE_SERVER
 // Scale factor for how much we would like to increase or decrease the memory pool size
 float CORE_API GPoolTableScale = 1.0f;
@@ -177,6 +180,20 @@ class FMalloc* FUnixPlatformMemory::BaseAllocator()
 					GTimeEnsures = false;
 				}
 
+				const char SignalToDefaultCmd[] = "-sigdfl=";
+				if (const char* Cmd = FCStringAnsi::Stristr(Arg, SignalToDefaultCmd))
+				{
+					int32 SignalToDefault = FCStringAnsi::Atoi(Cmd + sizeof(SignalToDefaultCmd) - 1);
+
+					// Valid signals are only from 1 -> SIGRTMAX
+					if (SignalToDefault > SIGRTMAX)
+					{
+						SignalToDefault = 0;
+					}
+
+					GSignalToDefault = FMath::Max(SignalToDefault, 0);
+				}
+
 				const char FileMapCacheCmd[] = "-filemapcachesize=";
 				if (const char* Cmd = FCStringAnsi::Stristr(Arg, FileMapCacheCmd))
 				{
diff --git a/Engine/Source/Runtime/Core/Private/Unix/UnixPlatformMisc.cpp b/Engine/Source/Runtime/Core/Private/Unix/UnixPlatformMisc.cpp
--- a/Engine/Source/Runtime/Core/Private/Unix/UnixPlatformMisc.cpp
+++ b/Engine/Source/Runtime/Core/Private/Unix/UnixPlatformMisc.cpp
@@ -212,6 +212,7 @@ void FUnixPlatformMisc::PlatformInit()
 	UE_LOG(LogInit, Log, TEXT(" -useksm - uses kernel same-page mapping (KSM) for mapped memory (%s)"), GUseKSM ? TEXT("ON") : TEXT("OFF"));
 	UE_LOG(LogInit, Log, TEXT(" -ksmmergeall - marks all mmap'd memory pages suitable for KSM (%s)"), GKSMMergeAllPages ? TEXT("ON") : TEXT("OFF"));
 	UE_LOG(LogInit, Log, TEXT(" -preloadmodulesymbols - Loads the main module symbols file into memory (%s)"), bPreloadedModuleSymbolFile ? TEXT("ON") : TEXT("OFF"));
+	UE_LOG(LogInit, Log, TEXT(" -sigdfl=SIGNAL - Allows a specific signal to be set to its default handler rather then ignoring the signal"));
 
 	// [RCL] FIXME: this should be printed in specific modules, if at all
 	UE_LOG(LogInit, Log, TEXT(" -httpproxy=ADDRESS:PORT - redirects HTTP requests to a proxy (only supported if compiled with libcurl)"));
