[LLM] remove dependency on Stats.h from LowLevelMemTracker.h, allowing proper tracking of Stats memory:  separated out llm stat definitions.

diff --git a/Engine/Source/Runtime/Core/Public/HAL/LowLevelMemStats.h b/Engine/Source/Runtime/Core/Public/HAL/LowLevelMemStats.h
new file mode 100644
--- /dev/null
+++ b/Engine/Source/Runtime/Core/Public/HAL/LowLevelMemStats.h
@@ -0,0 +1,31 @@
+// Copyright 1998-2018 Epic Games, Inc. All Rights Reserved.
+
+#include "HAL/LowLevelMemTracker.h"
+#include "Stats/Stats.h"
+
+#if ENABLE_LOW_LEVEL_MEM_TRACKER && !defined(DECLARE_LLM_MEMORY_STAT)
+#if STATS
+
+	#define DECLARE_LLM_MEMORY_STAT(CounterName,StatId,GroupId) \
+		DECLARE_STAT(CounterName,StatId,GroupId,EStatDataType::ST_int64, false, false, FPlatformMemory::MCR_PhysicalLLM); \
+		static DEFINE_STAT(StatId)
+	#define DECLARE_LLM_MEMORY_STAT_EXTERN(CounterName,StatId,GroupId, API) \
+		DECLARE_STAT(CounterName,StatId,GroupId,EStatDataType::ST_int64, false, false, FPlatformMemory::MCR_PhysicalLLM); \
+		extern API DEFINE_STAT(StatId);
+
+
+	DECLARE_STATS_GROUP(TEXT("LLM FULL"), STATGROUP_LLMFULL, STATCAT_Advanced);
+	DECLARE_STATS_GROUP(TEXT("LLM Platform"), STATGROUP_LLMPlatform, STATCAT_Advanced);
+	DECLARE_STATS_GROUP(TEXT("LLM Summary"), STATGROUP_LLM, STATCAT_Advanced);
+	DECLARE_STATS_GROUP(TEXT("LLM Overhead"), STATGROUP_LLMOverhead, STATCAT_Advanced);
+	DECLARE_STATS_GROUP(TEXT("LLM Assets"), STATGROUP_LLMAssets, STATCAT_Advanced);
+
+	DECLARE_LLM_MEMORY_STAT_EXTERN(TEXT("Engine"), STAT_EngineSummaryLLM, STATGROUP_LLM, CORE_API);
+
+#else
+	#define DECLARE_LLM_MEMORY_STAT(CounterName,StatId,GroupId)
+	#define DECLARE_LLM_MEMORY_STAT_EXTERN(CounterName,StatId,GroupId, API)
+
+#endif //STATS
+#endif //ENABLE_LOW_LEVEL_MEM_TRACKER
+
