[LLM] remove dependency on Stats.h from LowLevelMemTracker.h, allowing proper tracking of Stats memory.

diff --git a/Engine/Source/Runtime/Apple/MetalRHI/Private/MetalLLM.cpp b/Engine/Source/Runtime/Apple/MetalRHI/Private/MetalLLM.cpp
--- a/Engine/Source/Runtime/Apple/MetalRHI/Private/MetalLLM.cpp
+++ b/Engine/Source/Runtime/Apple/MetalRHI/Private/MetalLLM.cpp
@@ -3,6 +3,7 @@
 #include "MetalLLM.h"
 #include "MetalProfiler.h"
 #include "RenderUtils.h"
+#include "HAL/LowLevelMemStats.h"
 
 #include <objc/runtime.h>
 
diff --git a/Engine/Source/Runtime/Core/Private/Apple/AppleLLM.cpp b/Engine/Source/Runtime/Core/Private/Apple/AppleLLM.cpp
--- a/Engine/Source/Runtime/Core/Private/Apple/AppleLLM.cpp
+++ b/Engine/Source/Runtime/Core/Private/Apple/AppleLLM.cpp
@@ -1,6 +1,7 @@
 // Copyright 1998-2019 Epic Games, Inc. All Rights Reserved.
 
 #include "Apple/AppleLLM.h"
+#include "HAL/LowLevelMemStats.h"
 
 #if ENABLE_LOW_LEVEL_MEM_TRACKER
 
diff --git a/Engine/Source/Runtime/Core/Private/HAL/LowLevelMemTracker.cpp b/Engine/Source/Runtime/Core/Private/HAL/LowLevelMemTracker.cpp
--- a/Engine/Source/Runtime/Core/Private/HAL/LowLevelMemTracker.cpp
+++ b/Engine/Source/Runtime/Core/Private/HAL/LowLevelMemTracker.cpp
@@ -1,6 +1,7 @@
 // Copyright 1998-2019 Epic Games, Inc. All Rights Reserved.
 
 #include "HAL/LowLevelMemTracker.h"
+#include "HAL/LowLevelMemStats.h"
 #include "Misc/ScopeLock.h"
 #include "LowLevelMemoryUtils.h"
 #include "HAL/FileManager.h"
diff --git a/Engine/Source/Runtime/Core/Private/Stats/Stats2.cpp b/Engine/Source/Runtime/Core/Private/Stats/Stats2.cpp
--- a/Engine/Source/Runtime/Core/Private/Stats/Stats2.cpp
+++ b/Engine/Source/Runtime/Core/Private/Stats/Stats2.cpp
@@ -321,6 +321,7 @@ void FStartupMessages::AddThreadMetadata( const FName InThreadName, uint32 InThr
 
 void FStartupMessages::AddMetadata( FName InStatName, const TCHAR* InStatDesc, const char* InGroupName, const char* InGroupCategory, const TCHAR* InGroupDesc, bool bShouldClearEveryFrame, EStatDataType::Type InStatType, bool bCycleStat, bool bSortByName, FPlatformMemory::EMemoryCounterRegion InMemoryRegion /*= FPlatformMemory::MCR_Invalid*/ )
 {
+	LLM_SCOPE(ELLMTag::Stats);
 	FScopeLock Lock( &CriticalSection );
 
 	new (DelayedMessages)FStatMessage( InGroupName, EStatDataType::ST_None, "Groups", InGroupCategory, InGroupDesc, false, false, bSortByName );
@@ -333,6 +334,7 @@ FStartupMessages& FStartupMessages::Get()
 	static FStartupMessages* Messages = NULL;
 	if( !Messages )
 	{
+		LLM_SCOPE(ELLMTag::Stats);
 		check( IsInGameThread() );
 		Messages = new FStartupMessages;
 	}
@@ -1056,6 +1058,7 @@ public:
 
 FThreadStatsPool::FThreadStatsPool()
 {
+	LLM_SCOPE(ELLMTag::Stats);
 	for( int32 Index = 0; Index < NUM_ELEMENTS_IN_POOL; ++Index )
 	{
 		Pool.Push( new FThreadStats(EConstructor::FOR_POOL) );
@@ -1064,6 +1067,7 @@ FThreadStatsPool::FThreadStatsPool()
 
 FThreadStats* FThreadStatsPool::GetFromPool()
 {
+	LLM_SCOPE(ELLMTag::Stats);
 	FPlatformMisc::MemoryBarrier();
 	FThreadStats* Address = Pool.Pop();
 	while (!Address)
diff --git a/Engine/Source/Runtime/Core/Public/HAL/LowLevelMemTracker.h b/Engine/Source/Runtime/Core/Public/HAL/LowLevelMemTracker.h
--- a/Engine/Source/Runtime/Core/Public/HAL/LowLevelMemTracker.h
+++ b/Engine/Source/Runtime/Core/Public/HAL/LowLevelMemTracker.h
@@ -3,9 +3,6 @@
 #pragma once
 
 #include "CoreTypes.h"
-#include "Misc/ScopeLock.h"
-#include "Templates/Atomic.h"
-
 
 // this is currently incompatible with PLATFORM_USES_FIXED_GMalloc_CLASS, because this ends up being included way too early
 // and currently, PLATFORM_USES_FIXED_GMalloc_CLASS is only used in Test/Shipping builds, where we don't have STATS anyway,
@@ -15,7 +12,6 @@
 	#define DECLARE_LLM_MEMORY_STAT(CounterName,StatId,GroupId)
 	#define DECLARE_LLM_MEMORY_STAT_EXTERN(CounterName,StatId,GroupId, API)
 #else
-	#include "Stats/Stats.h"
 
 	#define LLM_SUPPORTED_PLATFORM (PLATFORM_XBOXONE || PLATFORM_PS4 || PLATFORM_WINDOWS || PLATFORM_IOS || PLATFORM_MAC || PLATFORM_ANDROID || PLATFORM_SWITCH || PLATFORM_UNIX)
 
@@ -48,32 +44,17 @@
 	#define LLM_AUTO_ENABLE 0
 #endif
 
-	#if STATS
-		#define DECLARE_LLM_MEMORY_STAT(CounterName,StatId,GroupId) \
-			DECLARE_STAT(CounterName,StatId,GroupId,EStatDataType::ST_int64, false, false, FPlatformMemory::MCR_PhysicalLLM); \
-			static DEFINE_STAT(StatId)
-		#define DECLARE_LLM_MEMORY_STAT_EXTERN(CounterName,StatId,GroupId, API) \
-			DECLARE_STAT(CounterName,StatId,GroupId,EStatDataType::ST_int64, false, false, FPlatformMemory::MCR_PhysicalLLM); \
-			extern API DEFINE_STAT(StatId);
-	#else
-		#define DECLARE_LLM_MEMORY_STAT(CounterName,StatId,GroupId)
-		#define DECLARE_LLM_MEMORY_STAT_EXTERN(CounterName,StatId,GroupId, API)
-	#endif
 #endif	// #if PLATFORM_USES_FIXED_GMalloc_CLASS
 
 
-#if STATS
-	DECLARE_STATS_GROUP(TEXT("LLM FULL"), STATGROUP_LLMFULL, STATCAT_Advanced);
-	DECLARE_STATS_GROUP(TEXT("LLM Platform"), STATGROUP_LLMPlatform, STATCAT_Advanced);
-	DECLARE_STATS_GROUP(TEXT("LLM Summary"), STATGROUP_LLM, STATCAT_Advanced);
-	DECLARE_STATS_GROUP(TEXT("LLM Overhead"), STATGROUP_LLMOverhead, STATCAT_Advanced);
-	DECLARE_STATS_GROUP(TEXT("LLM Assets"), STATGROUP_LLMAssets, STATCAT_Advanced);
-
-	DECLARE_LLM_MEMORY_STAT_EXTERN(TEXT("Engine"), STAT_EngineSummaryLLM, STATGROUP_LLM, CORE_API);
-#endif
-
-
 #if ENABLE_LOW_LEVEL_MEM_TRACKER
+#include "Misc/ScopeLock.h"
+#include "Templates/Atomic.h"
+#include "Templates/UnrealTemplate.h"
+#include "Templates/AlignmentTemplates.h"
+#include "Misc/CString.h"
+#include "Misc/VarArgs.h"
+#include "UObject/NameTypes.h"
 
 #if DO_CHECK
 
diff --git a/Engine/Source/Runtime/Core/Public/Stats/Stats2.h b/Engine/Source/Runtime/Core/Public/Stats/Stats2.h
--- a/Engine/Source/Runtime/Core/Public/Stats/Stats2.h
+++ b/Engine/Source/Runtime/Core/Public/Stats/Stats2.h
@@ -7,6 +7,7 @@
 #include "HAL/ThreadSafeCounter.h"
 #include "Math/NumericLimits.h"
 #include "HAL/ThreadSingleton.h"
+#include "HAL/LowLevelMemTracker.h"
 #include "Containers/Array.h"
 #include "Containers/UnrealString.h"
 #include "HAL/PlatformTime.h"
@@ -1347,6 +1348,7 @@ public:
 
 	FORCEINLINE_STATS void AddStatMessage( const FStatMessage& StatMessage )
 	{
+		LLM_SCOPE(ELLMTag::Stats);
 		FStatMessageLock MessageLock(MemoryMessageScope);
 		Packet.StatMessages.AddElement(StatMessage);
 	}
diff --git a/Engine/Source/Runtime/VulkanRHI/Private/VulkanLLM.cpp b/Engine/Source/Runtime/VulkanRHI/Private/VulkanLLM.cpp
--- a/Engine/Source/Runtime/VulkanRHI/Private/VulkanLLM.cpp
+++ b/Engine/Source/Runtime/VulkanRHI/Private/VulkanLLM.cpp
@@ -7,6 +7,7 @@
 
 #include "VulkanRHIPrivate.h"
 #include "VulkanLLM.h"
+#include "HAL/LowLevelMemStats.h"
 
 #if ENABLE_LOW_LEVEL_MEM_TRACKER
 
