CIS fix for UE-71344 - non-unity build failure

diff --git a/Engine/Source/Runtime/Core/Private/HAL/Allocators/PooledVirtualMemoryAllocator.cpp b/Engine/Source/Runtime/Core/Private/HAL/Allocators/PooledVirtualMemoryAllocator.cpp
--- a/Engine/Source/Runtime/Core/Private/HAL/Allocators/PooledVirtualMemoryAllocator.cpp
+++ b/Engine/Source/Runtime/Core/Private/HAL/Allocators/PooledVirtualMemoryAllocator.cpp
@@ -5,7 +5,9 @@
 #include "Logging/LogMacros.h"
 #include "CoreGlobals.h"
 #include "HAL/LowLevelMemTracker.h"
-#include "Core/Public/Misc/AssertionMacros.h"
+#include "Misc/AssertionMacros.h"
+#include "Misc/ScopeLock.h"
+#include "Templates/AlignmentTemplates.h"
 #include "GenericPlatform/OSAllocationPool.h"
 
 using T64KBAlignedPool = TMemoryPool< FPlatformMemory::MemoryRangeCommit, FPlatformMemory::MemoryRangeDecommit, 65536 >;
