Added CPU usage statistics to CSV profiling charts; Active process usage, idle usage and assumed external process usage.

diff --git a/Engine/Source/Runtime/Core/Private/Windows/WindowsPlatformProcess.cpp b/Engine/Source/Runtime/Core/Private/Windows/WindowsPlatformProcess.cpp
--- a/Engine/Source/Runtime/Core/Private/Windows/WindowsPlatformProcess.cpp
+++ b/Engine/Source/Runtime/Core/Private/Windows/WindowsPlatformProcess.cpp
@@ -490,6 +490,101 @@ bool FWindowsPlatformProcess::GetApplicationMemoryUsage(uint32 ProcessId, SIZE_T
 	return bSuccess;
 }
 
+bool FWindowsPlatformProcess::GetPerFrameProcessorUsage(uint32 ProcessId, float& ProcessUsageFraction, float& OtherUsageFraction, float& IdleUsageFraction)
+{
+	bool bSuccess = true;
+
+	static float LastProcessTime = 0.f;
+	static float LastIdleTime = 0.f;
+	static uint32 LastFrameNumber = 0;
+
+	if (LastFrameNumber != GFrameNumber)
+	{
+		LastFrameNumber = GFrameNumber;
+
+		// Get queryable process handle
+		HANDLE ProcessHandle = ::OpenProcess(PROCESS_QUERY_INFORMATION | PROCESS_VM_READ, false, ProcessId);
+
+		if (ProcessHandle != nullptr)
+		{
+			const uint64 NumCores = FPlatformMisc::NumberOfCoresIncludingHyperthreads();
+			const uint32 CurrFrameIndex = LastFrameNumber % 2;
+			const uint32 PrevFrameIndex = 1 - CurrFrameIndex;
+
+			// Get total processor cycles per second
+			static double DeltaCyclesPerSecond = 0.0;
+			if (DeltaCyclesPerSecond == 0.0)
+			{
+				LARGE_INTEGER Frequency;
+				QueryPerformanceFrequency(&Frequency);
+				DeltaCyclesPerSecond = (double)(Frequency.QuadPart) * 1000.0;
+				DeltaCyclesPerSecond *= NumCores;
+			}
+
+			// Calculate total number of cycles that have passed this frame
+			static double PrevTotalSeconds = 0.0;
+			double TotalSeconds = FPlatformTime::Seconds();
+			double DeltaSecondsPerFrame = (double)(TotalSeconds - PrevTotalSeconds);
+			PrevTotalSeconds = TotalSeconds;
+
+			double DeltaCyclesPerFrame = DeltaSecondsPerFrame * DeltaCyclesPerSecond;
+
+			// Grab cycle time for this process as fraction of total processor time
+			static uint64 ProcessCycleTimeBuffers[2] = {0};
+			uint64& ProcessCycleTime = ProcessCycleTimeBuffers[CurrFrameIndex];
+			uint64& PrevProcessCycleTime = ProcessCycleTimeBuffers[PrevFrameIndex];
+
+			if (!QueryProcessCycleTime(ProcessHandle, (PULONG64)&ProcessCycleTime))
+			{
+				bSuccess = false;
+			}
+			uint64 DeltaProcessCycleTime = ProcessCycleTime - PrevProcessCycleTime;
+			LastProcessTime = (double)DeltaProcessCycleTime / DeltaCyclesPerFrame;
+
+			// Idle cycles are stored per core and flipped to allow per-frame calculation
+			const uint32 BufferLength = 512;
+			check(BufferLength >= NumCores * 8);
+
+			static uint64 IdleCycleTimeBuffers[2][BufferLength] = {{0}};
+			uint64* IdleCycleTime = IdleCycleTimeBuffers[CurrFrameIndex];
+			uint64* PrevIdleCycleTime = IdleCycleTimeBuffers[PrevFrameIndex];
+
+			// Grab idle cycle time as percentage of total processor time
+			// Note: Idle processes are specified per core and accumulated
+			if (!QueryIdleProcessorCycleTime((PULONG)&BufferLength, (PULONG64)IdleCycleTime))
+			{
+				bSuccess = false;
+			}
+
+			uint64 DeltaIdleTime = 0;
+			for (int Core = 0; Core < NumCores; ++Core)
+			{
+				DeltaIdleTime += IdleCycleTime[Core] - PrevIdleCycleTime[Core];
+			}
+			LastIdleTime = (double)DeltaIdleTime / DeltaCyclesPerFrame;
+
+			::CloseHandle(ProcessHandle);
+		}
+		else
+		{
+			bSuccess = false;
+		}
+	}
+
+	if (bSuccess)
+	{
+		ProcessUsageFraction = LastProcessTime;
+		IdleUsageFraction = LastIdleTime;
+		OtherUsageFraction = FMath::Clamp(1.f - (LastProcessTime + LastIdleTime), 0.f, 1.f);
+	}
+	else
+	{
+		ProcessUsageFraction = OtherUsageFraction = IdleUsageFraction = 0.f;
+	}
+
+	return bSuccess;
+}
+
 bool FWindowsPlatformProcess::IsApplicationRunning( uint32 ProcessId )
 {
 	bool bApplicationRunning = true;
diff --git a/Engine/Source/Runtime/Core/Public/Windows/WindowsPlatformProcess.h b/Engine/Source/Runtime/Core/Public/Windows/WindowsPlatformProcess.h
--- a/Engine/Source/Runtime/Core/Public/Windows/WindowsPlatformProcess.h
+++ b/Engine/Source/Runtime/Core/Public/Windows/WindowsPlatformProcess.h
@@ -156,6 +156,7 @@ public:
 	static void TerminateProc( FProcHandle & ProcessHandle, bool KillTree = false );
 	static bool GetProcReturnCode( FProcHandle & ProcHandle, int32* ReturnCode );
 	static bool GetApplicationMemoryUsage(uint32 ProcessId, SIZE_T* OutMemoryUsage);
+	static bool GetPerFrameProcessorUsage(uint32 ProcessId, float& ProcessUsageFraction, float& OtherUsageFraction, float& IdleUsageFraction);
 	static bool IsApplicationRunning( uint32 ProcessId );
 	static bool IsApplicationRunning( const TCHAR* ProcName );
 	static FString GetApplicationName( uint32 ProcessId );	
diff --git a/Engine/Source/Runtime/Launch/Private/LaunchEngineLoop.cpp b/Engine/Source/Runtime/Launch/Private/LaunchEngineLoop.cpp
--- a/Engine/Source/Runtime/Launch/Private/LaunchEngineLoop.cpp
+++ b/Engine/Source/Runtime/Launch/Private/LaunchEngineLoop.cpp
@@ -986,7 +986,23 @@ bool IsServerDelegateForOSS(FName WorldContextHandle)
 #endif
 
 #if WITH_ENGINE && CSV_PROFILER
-static void UpdateCoreCsvStats()
+static void UpdateCoreCsvStats_BeginFrame()
+{
+#if PLATFORM_WINDOWS
+	if (FCsvProfiler::Get()->IsCapturing())
+	{
+		const uint32 ProcessId = (uint32)GetCurrentProcessId();
+		float ProcessUsageFraction = 0.f, OtherUsageFraction = 0.f, IdleUsageFraction = 0.f;
+		FWindowsPlatformProcess::GetPerFrameProcessorUsage(ProcessId, ProcessUsageFraction, OtherUsageFraction, IdleUsageFraction);
+
+		CSV_CUSTOM_STAT_GLOBAL(CPUUsage_Process, ProcessUsageFraction, ECsvCustomStatOp::Set);
+		CSV_CUSTOM_STAT_GLOBAL(CPUUsage_Other, OtherUsageFraction, ECsvCustomStatOp::Set);
+		CSV_CUSTOM_STAT_GLOBAL(CPUUsage_Idle, IdleUsageFraction, ECsvCustomStatOp::Set);
+	}
+#endif
+}
+
+static void UpdateCoreCsvStats_EndFrame()
 {
 	CSV_CUSTOM_STAT_GLOBAL(RenderThreadTime, FPlatformTime::ToMilliseconds(GRenderThreadTime), ECsvCustomStatOp::Set);
 	CSV_CUSTOM_STAT_GLOBAL(GameThreadTime, FPlatformTime::ToMilliseconds(GGameThreadTime), ECsvCustomStatOp::Set);
@@ -1676,7 +1692,8 @@ int32 FEngineLoop::PreInit(const TCHAR* CmdLine)
 #if WITH_ENGINE && CSV_PROFILER
 	if (!IsRunningDedicatedServer())
 	{
-		FCoreDelegates::OnEndFrame.AddStatic(UpdateCoreCsvStats);
+		FCoreDelegates::OnBeginFrame.AddStatic(UpdateCoreCsvStats_BeginFrame);
+		FCoreDelegates::OnEndFrame.AddStatic(UpdateCoreCsvStats_EndFrame);
 	}
 	FCsvProfiler::Get()->Init();
 #endif
