Add rollover for timecode created from time

diff --git a/Engine/Plugins/Runtime/AR/Apple/AppleARKit/Source/AppleARKit/Private/AppleARKitTimecodeProvider.cpp b/Engine/Plugins/Runtime/AR/Apple/AppleARKit/Source/AppleARKit/Private/AppleARKitTimecodeProvider.cpp
--- a/Engine/Plugins/Runtime/AR/Apple/AppleARKit/Source/AppleARKit/Private/AppleARKitTimecodeProvider.cpp
+++ b/Engine/Plugins/Runtime/AR/Apple/AppleARKit/Source/AppleARKit/Private/AppleARKitTimecodeProvider.cpp
@@ -20,7 +20,7 @@ void UAppleARKitTimecodeProvider::Shutdown(class UEngine*)
 
 FTimecode UAppleARKitTimecodeProvider::GetTimecode() const
 {
-	return FTimecode(FPlatformTime::Seconds(), FrameRate, FTimecode::IsDropFormatTimecodeSupported(FrameRate));
+	return FTimecode(FPlatformTime::Seconds(), FrameRate, FTimecode::IsDropFormatTimecodeSupported(FrameRate), true);
 }
 
 FFrameRate UAppleARKitTimecodeProvider::GetFrameRate() const
diff --git a/Engine/Source/Runtime/Core/Private/Tests/Misc/TimecodeTest.cpp b/Engine/Source/Runtime/Core/Private/Tests/Misc/TimecodeTest.cpp
--- a/Engine/Source/Runtime/Core/Private/Tests/Misc/TimecodeTest.cpp
+++ b/Engine/Source/Runtime/Core/Private/Tests/Misc/TimecodeTest.cpp
@@ -79,11 +79,22 @@ bool FTimecodeTest::RunTest(const FString& Parameters)
 			if (bDoTest)
 			{
 				const FTimespan TimespanFromTimecode = TimecodeValue.ToTimespan(FrameRate);
-				const FTimecode TimecodeFromTimespan = FTimecode::FromTimespan(TimespanFromTimecode, FrameRate, bIsDropFrame);
+				const FTimecode TimecodeFromTimespanWithRollover = FTimecode::FromTimespan(TimespanFromTimecode, FrameRate, bIsDropFrame, true);
+				const FTimecode TimecodeFromTimespanWithoutRollover = FTimecode::FromTimespan(TimespanFromTimecode, FrameRate, bIsDropFrame, false);
 
-				if (TimecodeFromTimespan != TimecodeValue)
+				if (TimecodeFromTimespanWithoutRollover != TimecodeValue)
 				{
-					AddError(FString::Printf(TEXT("Timecode '%s' didn't convert properly from Timespan '%f' for frame rate '%s'.")
+					AddError(FString::Printf(TEXT("Timecode '%s' didn't convert properly from Timespan '%f' with rollover for frame rate '%s'.")
+						, *TimecodeValue.ToString()
+						, TimespanFromTimecode.GetTotalSeconds()
+						, *FrameRate.ToPrettyText().ToString()
+					));
+					bDoTest = false;
+					++NumberOfErrors;
+				}
+				else if (TimecodeFromTimespanWithoutRollover.Minutes != TimecodeValue.Minutes || TimecodeFromTimespanWithoutRollover.Seconds != TimecodeValue.Seconds || TimecodeFromTimespanWithoutRollover.Frames != TimecodeValue.Frames)
+				{
+					AddError(FString::Printf(TEXT("Timecode '%s' didn't convert properly from Timespan '%f' without rollover for frame rate '%s'.")
 						, *TimecodeValue.ToString()
 						, TimespanFromTimecode.GetTotalSeconds()
 						, *FrameRate.ToPrettyText().ToString()
@@ -94,10 +105,11 @@ bool FTimecodeTest::RunTest(const FString& Parameters)
 				else if (!bIsDropFrame)
 				{
 					// Do they have the same hours, minutes, seconds
-					bool bHoursAreValid = TimespanFromTimecode.GetHours() == TimecodeValue.Hours || ((TimecodeValue.Hours%24) == TimespanFromTimecode.GetHours() && (TimecodeValue.Hours/24) == TimespanFromTimecode.GetDays());
+					bool bRolloverHoursAreValid = TimespanFromTimecode.GetHours() == TimecodeFromTimespanWithRollover.Hours;
+					bool bHoursAreValid = (TimecodeValue.Hours % 24) == TimespanFromTimecode.GetHours() && (TimecodeValue.Hours / 24) == TimespanFromTimecode.GetDays();
 					bool bMinutesAreValid = TimespanFromTimecode.GetMinutes() == TimecodeValue.Minutes;
 					bool bSecondsAreValid = TimespanFromTimecode.GetSeconds() == TimecodeValue.Seconds;
-					if (!bHoursAreValid || !bMinutesAreValid || !bSecondsAreValid)
+					if (!bRolloverHoursAreValid || !bHoursAreValid || !bMinutesAreValid || !bSecondsAreValid)
 					{
 						AddError(FString::Printf(TEXT("Timecode '%s' hours/minutes/seconds doesn't matches with Timespan '%s' from frame rate '%s'.")
 							, *TimecodeValue.ToString()
@@ -153,7 +165,6 @@ bool FTimecodeTest::RunTest(const FString& Parameters)
 					AddError(FString::Printf(TEXT("Timecode '%s' didn't convert properly from FrameNumber '%d' when the frame rate is tripled.")
 						, *TimecodeValue.ToString()
 						, FrameNumber.Value
-						, *FrameRate.ToPrettyText().ToString()
 					));
 					bDoTest = false;
 					++NumberOfErrors;
@@ -178,6 +189,44 @@ bool FTimecodeTest::RunTest(const FString& Parameters)
 			}
 		}
 
+		// Conversion from current time to Timecode
+		if (NumberOfErrors == 0)
+		{
+			const FTimespan CurrentTimespan = FTimespan(11694029893428);
+			const double CurrentSeconds = 1169402.9893428; //from FPlatformTime::Seconds()
+
+			const FTimecode FromTimespanTimecodeValueWithRollover = FTimecode::FromTimespan(CurrentTimespan, FrameRate, bIsDropFrame, true);
+			const FTimecode FromTimespanTimecodeValueWithoutRollover = FTimecode::FromTimespan(CurrentTimespan, FrameRate, bIsDropFrame, false);
+			const FTimecode FromSecondsTimecodeValueWithRollover = FTimecode(CurrentSeconds, FrameRate, bIsDropFrame, true);
+			const FTimecode FromSecondsTimecodeValueWithoutRollover = FTimecode(CurrentSeconds, FrameRate, bIsDropFrame, false);
+
+			if (FromTimespanTimecodeValueWithRollover != FromSecondsTimecodeValueWithRollover)
+			{
+				AddError(FString::Printf(TEXT("The timecode '%s' do not match timecode '%s' when converted from the computer clock's time and the frame rate is '%s'")
+					, *FromTimespanTimecodeValueWithRollover.ToString()
+					, *FromSecondsTimecodeValueWithRollover.ToString()
+					, *FrameRate.ToPrettyText().ToString()
+				));
+				++NumberOfErrors;
+			}
+			else if (FromTimespanTimecodeValueWithoutRollover != FromSecondsTimecodeValueWithoutRollover)
+			{
+				AddError(FString::Printf(TEXT("The timecode '%s' do not match timecode '%s' when converted from the computer clock's time and the frame rate is '%s'")
+					, *FromTimespanTimecodeValueWithoutRollover.ToString()
+					, *FromSecondsTimecodeValueWithoutRollover.ToString()
+					, *FrameRate.ToPrettyText().ToString()
+				));
+				++NumberOfErrors;
+			}
+			else if (!bIsDropFrame && FromTimespanTimecodeValueWithRollover.Frames != FromTimespanTimecodeValueWithoutRollover.Frames)
+			{
+				AddError(FString::Printf(TEXT("The timecode didn't convert properly from the computer clock's time when the frame rate is '%s'")
+					, *FrameRate.ToPrettyText().ToString()
+				));
+				++NumberOfErrors;
+			}
+		}
+
 		AddInfo(FString::Printf(TEXT("Timecode test was completed with frame rate '%s'"), *FrameRate.ToPrettyText().ToString()));
 
 		return NumberOfErrors == 0;
diff --git a/Engine/Source/Runtime/Core/Public/Misc/Timecode.h b/Engine/Source/Runtime/Core/Public/Misc/Timecode.h
--- a/Engine/Source/Runtime/Core/Public/Misc/Timecode.h
+++ b/Engine/Source/Runtime/Core/Public/Misc/Timecode.h
@@ -41,12 +41,22 @@ struct FTimecode
 
 	/**
 	 * User construction from a time in seconds
-	 * @param InbDropFrame - If true, this Timecode represents a "Drop Frame Timecode" format which
+	 * @param InbDropFrame	- If true, this Timecode represents a "Drop Frame Timecode" format which
 							skips the first frames of every minute (except those ending in multiples of 10)
 							to account for drift when using a fractional NTSC framerate.
+	 * @param InbRollover	- If true, the hours will be the modulo of 24.
+	 * @note Be aware that the Cycles may not correspond to the System Time. See FDateTime and "leap seconds".
 	 */
-	explicit FTimecode(double InSeconds, const FFrameRate& InFrameRate, bool InbDropFrame)
+	explicit FTimecode(double InSeconds, const FFrameRate& InFrameRate, bool InbDropFrame, bool InbRollover)
 	{
+		if (InbRollover)
+		{
+			const int32 NumberOfSecondsPerDay = 60 * 60 * 24;
+			double IntegralPart = 0.0;
+			double Fractional = FMath::Modf(InSeconds, &IntegralPart);
+			const int32 CurrentRolloverSeconds = (int32)IntegralPart % NumberOfSecondsPerDay;
+			InSeconds = (double)CurrentRolloverSeconds + Fractional;
+		}
 		int32 NumberOfFrames = InbDropFrame ? (int32)FMath::RoundToDouble(InSeconds * InFrameRate.AsDecimal()) : (int32)FMath::RoundToDouble(InSeconds * FMath::RoundToDouble(InFrameRate.AsDecimal()));
 		*this = FromFrameNumber(FFrameNumber(NumberOfFrames), InFrameRate, InbDropFrame);
 	}
@@ -196,16 +206,17 @@ public:
 	 * Create a FTimecode from a timespan at the given frame rate. Optionally supports creating a drop frame timecode,
 	 * which drops certain timecode display numbers to help account for NTSC frame rates which are fractional.
 	 *
-	 * @param InFrameNumber - The timespan to convert into a timecode.
+	 * @param InFrameNumber	- The timespan to convert into a timecode.
 	 * @param InFrameRate	- The framerate that this timecode is based in. This should be the playback framerate as it is used to determine
 	 *						  when the Frame value wraps over.
 	 * @param InbDropFrame	- If true, the returned timecode will drop the first two frames on every minute (except when Minute % 10 == 0)
 	 *						  This is only valid for NTSC framerates (29.97, 59.94) and will assert if you try to create a drop-frame format
 	 *						  from a non-valid framerate. All framerates can be represented when in non-drop frame format.
+	 * @param InbRollover	- If true, the hours will be the modulo of 24.
 	 */
-	static FTimecode FromTimespan(const FTimespan& InTimespan, const FFrameRate& InFrameRate, bool InbDropFrame)
+	static FTimecode FromTimespan(const FTimespan& InTimespan, const FFrameRate& InFrameRate, bool InbDropFrame, bool InbRollover)
 	{
-		return FTimecode(InTimespan.GetTotalSeconds(), InFrameRate, InbDropFrame);
+		return FTimecode(InTimespan.GetTotalSeconds(), InFrameRate, InbDropFrame, InbRollover);
 	}
 
 	static bool IsDropFormatTimecodeSupported(const FFrameRate& InFrameRate)
diff --git a/Engine/Source/Runtime/Engine/Private/Engine/SystemTimeTimecodeProvider.cpp b/Engine/Source/Runtime/Engine/Private/Engine/SystemTimeTimecodeProvider.cpp
--- a/Engine/Source/Runtime/Engine/Private/Engine/SystemTimeTimecodeProvider.cpp
+++ b/Engine/Source/Runtime/Engine/Private/Engine/SystemTimeTimecodeProvider.cpp
@@ -11,5 +11,5 @@ FTimecode USystemTimeTimecodeProvider::GetTimecode() const
 	const FDateTime DateTime = FDateTime::Now();
 	const FTimespan Timespan = DateTime.GetTimeOfDay();
 
-	return FTimecode::FromTimespan(Timespan, FrameRate, FTimecode::IsDropFormatTimecodeSupported(FrameRate));
+	return FTimecode::FromTimespan(Timespan, FrameRate, FTimecode::IsDropFormatTimecodeSupported(FrameRate), false);
 }
