Do not trim assert callstacks.

- Linux and Windows used to trim that, but this optimization is brittle.
- The risk to cut a valid rame outweighs the readability benefits.

diff --git a/Engine/Source/Runtime/Core/Private/Unix/UnixErrorOutputDevice.cpp b/Engine/Source/Runtime/Core/Private/Unix/UnixErrorOutputDevice.cpp
--- a/Engine/Source/Runtime/Core/Private/Unix/UnixErrorOutputDevice.cpp
+++ b/Engine/Source/Runtime/Core/Private/Unix/UnixErrorOutputDevice.cpp
@@ -52,16 +52,14 @@ void FUnixErrorOutputDevice::Serialize(const TCHAR* Msg, ELogVerbosity::Type Ver
 
 	if( GIsGuarded )
 	{
-		// Propagate error so structured exception handler can perform necessary work.
 #if PLATFORM_EXCEPTIONS_DISABLED
 		UE_DEBUG_BREAK();
 #endif
-		// Generate the portable callstack. For asserts, we ignore the following frames:
-		//     FDebug::AssertFailed()
-		//   [ FOutputDevice::Logf() ] - force-inlined; ignored
-		//     FOutputDevice::LogfImpl()
-		//     FWindowsErrorOutputDevice::Serialize()
-		const int32 NumStackFramesToIgnore = 3;
+		// Generate the callstack.
+		// We do not ignore any stack frames since the optimization is
+		// brittle and the risk of trimming the valid frames is too high.
+		// The common frames will be instead filtered out in the web UI
+		const int32 NumStackFramesToIgnore = 0;
 
 		if (GIsGPUCrashed)
 		{
diff --git a/Engine/Source/Runtime/Core/Private/Windows/WindowsErrorOutputDevice.cpp b/Engine/Source/Runtime/Core/Private/Windows/WindowsErrorOutputDevice.cpp
--- a/Engine/Source/Runtime/Core/Private/Windows/WindowsErrorOutputDevice.cpp
+++ b/Engine/Source/Runtime/Core/Private/Windows/WindowsErrorOutputDevice.cpp
@@ -62,11 +62,10 @@ void FWindowsErrorOutputDevice::Serialize( const TCHAR* Msg, ELogVerbosity::Type
 		UE_DEBUG_BREAK();
 #endif
 		// Generate the portable callstack. For asserts, we ignore the following frames:
-		//     FDebug::AssertFailed()
-		//   [ FOutputDevice::Logf() ] - force-inlined; ignored
-		//     FOutputDevice::LogfImpl()
-		//     FWindowsErrorOutputDevice::Serialize()
-		const int32 NumStackFramesToIgnore = 3;
+		// We do not ignore any stack frames since the optimization is
+		// brittle and the risk of trimming the valid frames is too high.
+		// The common frames will be instead filtered out in the web UI
+		const int32 NumStackFramesToIgnore = 0;
 
 		if (GIsGPUCrashed)
 		{
