Fix FString::RemoveSpacesInline not returning the correct length after modification

diff --git a/Engine/Source/Runtime/Core/Private/Containers/String.cpp b/Engine/Source/Runtime/Core/Private/Containers/String.cpp
--- a/Engine/Source/Runtime/Core/Private/Containers/String.cpp
+++ b/Engine/Source/Runtime/Core/Private/Containers/String.cpp
@@ -257,6 +257,7 @@ void FString::RemoveSpacesInline()
 	if (CopyToIndex <= StringLength)
 	{
 		RawData[CopyToIndex] = '\0';
+		Data.SetNum(CopyToIndex + 1, false);
 	}
 }
 
