config hierarchy: added NoRedist and NotForLicensees versions of game's dedicated server
#rnx

diff --git a/Engine/Source/Runtime/Core/Private/Misc/ConfigCacheIni.cpp b/Engine/Source/Runtime/Core/Private/Misc/ConfigCacheIni.cpp
--- a/Engine/Source/Runtime/Core/Private/Misc/ConfigCacheIni.cpp
+++ b/Engine/Source/Runtime/Core/Private/Misc/ConfigCacheIni.cpp
@@ -3168,6 +3168,11 @@ static void GetSourceIniHierarchyFilenames(const TCHAR* InBaseIniName, const TCH
 	if (IsRunningDedicatedServer())
 	{
 		OutHierarchy.Add(EConfigFileHierarchy::GameDirDedicatedServer, FIniFilename(FString::Printf(TEXT("%s/DedicatedServer%s.ini"), SourceConfigDir, InBaseIniName), false));
+
+		// Game/Config/NotForLicensees/DedicatedServer*.ini
+		OutHierarchy.Add(EConfigFileHierarchy::GameDirDedicatedServer_NotForLicensees, FIniFilename(FString::Printf(TEXT("%sNotForLicensees/DedicatedServer%s.ini"), SourceConfigDir, InBaseIniName), false));
+		// Game/Config/NoRedist/DedicatedServer*.ini
+		OutHierarchy.Add(EConfigFileHierarchy::GameDirDedicatedServer_NoRedist, FIniFilename(FString::Printf(TEXT("%sNoRedist/DedicatedServer%s.ini"), SourceConfigDir, InBaseIniName), false));
 	}
 	// Game/Config/NotForLicensees/Default* ini
 	OutHierarchy.Add(EConfigFileHierarchy::GameDirDefault_NotForLicensees, FIniFilename(FString::Printf(TEXT("%sNotForLicensees/Default%s.ini"), SourceConfigDir, InBaseIniName), false));
diff --git a/Engine/Source/Runtime/Core/Public/Misc/ConfigCacheIni.h b/Engine/Source/Runtime/Core/Public/Misc/ConfigCacheIni.h
--- a/Engine/Source/Runtime/Core/Public/Misc/ConfigCacheIni.h
+++ b/Engine/Source/Runtime/Core/Public/Misc/ConfigCacheIni.h
@@ -300,6 +300,10 @@ enum class EConfigFileHierarchy : uint8
 	GameDirDefault,
 	// Game/Config/DedicatedServer*.ini
 	GameDirDedicatedServer,
+	// Game/Config/NotForLicensees/DedicatedServer*.ini
+	GameDirDedicatedServer_NotForLicensees,
+	// Game/Config/NoRedist/DedicatedServer*.ini
+	GameDirDedicatedServer_NoRedist,
 	// Game/Config/NotForLicensees/*.ini
 	GameDirDefault_NotForLicensees,
 	// Game/Config/NoRedist*.ini
