make it so the debug output device is not created on IOS when not connected to a debugger

diff --git a/Engine/Source/Runtime/Core/Private/IOS/IOSPlatformMisc.cpp b/Engine/Source/Runtime/Core/Private/IOS/IOSPlatformMisc.cpp
--- a/Engine/Source/Runtime/Core/Private/IOS/IOSPlatformMisc.cpp
+++ b/Engine/Source/Runtime/Core/Private/IOS/IOSPlatformMisc.cpp
@@ -1665,6 +1665,11 @@ void FIOSPlatformMisc::SetCrashHandler(void (* CrashHandler)(const FGenericCrash
 #endif
 }
 
+bool FIOSPlatformMisc::HasSeparateChannelForDebugOutput()
+{
+    return FPlatformMisc::IsDebuggerPresent();
+}
+
 FIOSCrashContext::FIOSCrashContext(ECrashContextType InType, const TCHAR* InErrorMessage)
 	: FApplePlatformCrashContext(InType, InErrorMessage)
 {
diff --git a/Engine/Source/Runtime/Core/Public/IOS/IOSPlatformMisc.h b/Engine/Source/Runtime/Core/Public/IOS/IOSPlatformMisc.h
--- a/Engine/Source/Runtime/Core/Public/IOS/IOSPlatformMisc.h
+++ b/Engine/Source/Runtime/Core/Public/IOS/IOSPlatformMisc.h
@@ -42,6 +42,7 @@ struct CORE_API FIOSPlatformMisc : public FApplePlatformMisc
 	static ENetworkConnectionType GetNetworkConnectionType();
 	static bool HasActiveWiFiConnection();
 	static const TCHAR* GamePersistentDownloadDir();
+    static bool HasSeparateChannelForDebugOutput();
 
 	UE_DEPRECATED(4.21, "Use GetDeviceVolume, it is now callable on all platforms.")
 	static int GetAudioVolume();
