Improvements to RSA key generation and pak signing
- CryptoKeys plugin can now be told how many bits to use for keys. Defaults to 2048 as a happy medium between security and performance
- After creating the OpenSSL RSA key, use the exp/mod values directly rather than recomputing from primes
- Added RSA.h with some TBigInt based helpers for RSA keys
- Added FPakSignatureFile to wrap the contents of the pak .sig file, including determining that it is of a new format that supports versioning. Wraps the encryption and decryption of he master table hash
- Changed pak master signature hash from crc to SHA1
- Modified access to executable embedded public key data so that it is just an array of bytes. Will allow custom key sizes in the future.

diff --git a/Engine/Plugins/Editor/CryptoKeys/Source/CryptoKeys/Private/CryptoKeysHelpers.cpp b/Engine/Plugins/Editor/CryptoKeys/Source/CryptoKeys/Private/CryptoKeysHelpers.cpp
--- a/Engine/Plugins/Editor/CryptoKeys/Source/CryptoKeys/Private/CryptoKeysHelpers.cpp
+++ b/Engine/Plugins/Editor/CryptoKeys/Source/CryptoKeys/Private/CryptoKeysHelpers.cpp
@@ -26,12 +26,12 @@ namespace CryptoKeysHelpers
 		return bResult;
 	}
 
-	bool GenerateSigningKey(FString& OutPublicExponent, FString& OutPrivateExponent, FString& OutModulus)
+	bool GenerateSigningKey(FString& OutPublicExponent, FString& OutPrivateExponent, FString& OutModulus, int32 NumKeyBits)
 	{
 		bool bResult = false;
 
 		TArray<uint8> PublicExponent, PrivateExponent, Modulus;
-		if (CryptoKeysOpenSSL::GenerateNewSigningKey(PublicExponent, PrivateExponent, Modulus))
+		if (CryptoKeysOpenSSL::GenerateNewSigningKey(PublicExponent, PrivateExponent, Modulus, NumKeyBits))
 		{
 			OutPublicExponent = FBase64::Encode(PublicExponent);
 			OutPrivateExponent = FBase64::Encode(PrivateExponent);
diff --git a/Engine/Plugins/Editor/CryptoKeys/Source/CryptoKeys/Private/CryptoKeysHelpers.h b/Engine/Plugins/Editor/CryptoKeys/Source/CryptoKeys/Private/CryptoKeysHelpers.h
--- a/Engine/Plugins/Editor/CryptoKeys/Source/CryptoKeys/Private/CryptoKeysHelpers.h
+++ b/Engine/Plugins/Editor/CryptoKeys/Source/CryptoKeys/Private/CryptoKeysHelpers.h
@@ -2,6 +2,8 @@
 
 #pragma once
 
+#include "CoreTypes.h"
+
 class FString;
 
 namespace CryptoKeysHelpers
@@ -17,8 +19,9 @@ namespace CryptoKeysHelpers
 	Generates a new RSA signing key
 	@param OutPublicExponent 	- Assigned the base64 encoded representation of the RSA public exponent
 	@param OutPrivateExponent 	- Assigned the base64 encoded representation of the RSA private exponent
-	@param OutModulus 		- Assigned the base64 encoded representation of the RSA modulus
+	@param OutModulus		 Assigned the base64 encoded representation of the RSA modulus
+	@param InNumKeyBits 		- How many bits to use for the RSA key
 	@returns true if key generation succeeded, false otherwise
 	*/
-	bool GenerateSigningKey(FString& OutPublicExponent, FString& OutPrivateExponent, FString& OutModulus);
+	bool GenerateSigningKey(FString& OutPublicExponent, FString& OutPrivateExponent, FString& OutModulus, int32 NumKeyBits = 2048);
 }
\ No newline at end of file
diff --git a/Engine/Plugins/Editor/CryptoKeys/Source/CryptoKeysOpenSSL/Private/CryptoKeysOpenSSL.cpp b/Engine/Plugins/Editor/CryptoKeys/Source/CryptoKeysOpenSSL/Private/CryptoKeysOpenSSL.cpp
--- a/Engine/Plugins/Editor/CryptoKeys/Source/CryptoKeysOpenSSL/Private/CryptoKeysOpenSSL.cpp
+++ b/Engine/Plugins/Editor/CryptoKeys/Source/CryptoKeysOpenSSL/Private/CryptoKeysOpenSSL.cpp
@@ -9,13 +9,12 @@
 #include <openssl/rsa.h>
 #include <openssl/rand.h>
 #include <openssl/pem.h>
+#include <openssl/bn.h>
 
 #if PLATFORM_WINDOWS
 #include "Windows/HideWindowsPlatformTypes.h"
 #endif
 
-#include "Math/BigInt.h"
-
 DEFINE_LOG_CATEGORY_STATIC(LogCryptoKeys, Log, All);
 
 namespace CryptoKeysOpenSSL
@@ -32,99 +31,41 @@ namespace CryptoKeysOpenSSL
 		return bResult;
 	}
 
-	bool TestKeys(FEncryptionKey& InPublicKey, FEncryptionKey& InPrivateKey)
+	void BigNumToArray(const BIGNUM* InNum, TArray<uint8>& OutBytes, int32 InKeySize)
 	{
-		UE_LOG(LogCryptoKeys, Display, TEXT("Testing signature keys."));
-
-		// Just some random values
-		static TEncryptionInt TestData[] =
-		{
-			11,
-			253,
-			128,
-			234,
-			56,
-			89,
-			34,
-			179,
-			29,
-			1024,
-			(int64)(MAX_int32),
-			(int64)(MAX_uint32)-1
-		};
-
-		for (int32 TestIndex = 0; TestIndex < ARRAY_COUNT(TestData); ++TestIndex)
-		{
-			TEncryptionInt EncryptedData = FEncryption::ModularPow(TestData[TestIndex], InPrivateKey.Exponent, InPrivateKey.Modulus);
-			TEncryptionInt DecryptedData = FEncryption::ModularPow(EncryptedData, InPublicKey.Exponent, InPublicKey.Modulus);
-			if (TestData[TestIndex] != DecryptedData)
-			{
-				UE_LOG(LogCryptoKeys, Error, TEXT("Keys do not properly encrypt/decrypt data (failed test with %lld)"), TestData[TestIndex].ToInt());
-				return false;
-			}
-		}
-
-		UE_LOG(LogCryptoKeys, Display, TEXT("Signature keys check completed successfuly."));
-
-		return true;
+		int NumBytes = BN_num_bytes(InNum);
+		check(NumBytes <= InKeySize);
+		OutBytes.SetNumZeroed(InKeySize);
+		BN_bn2bin(InNum, OutBytes.GetData());
+		Algo::Reverse(OutBytes); // bn2bin produces big endian data
 	}
 
-	bool GenerateNewSigningKey(TArray<uint8>& OutPublicExponent, TArray<uint8>& OutPrivateExponent, TArray<uint8>& OutModulus)
+	bool GenerateNewSigningKey(TArray<uint8>& OutPublicExponent, TArray<uint8>& OutPrivateExponent, TArray<uint8>& OutModulus, int32 InNumKeyBits)
 	{
-		TArray<uint8> NewP, NewQ;
+		int32 KeySize = InNumKeyBits;
+		int32 KeySizeInBytes = InNumKeyBits / 8;
 
 		RSA* RSAKey = RSA_new();
 		BIGNUM* E = BN_new();
 		BN_set_word(E, RSA_F4);
-		RSA_generate_key_ex(RSAKey, 255, E, nullptr);
+		RSA_generate_key_ex(RSAKey, KeySize, E, nullptr);
 
 #if !defined(OPENSSL_VERSION_NUMBER) || OPENSSL_VERSION_NUMBER < 0x10100000L
-		const BIGNUM* P = RSAKey->p;
-		const BIGNUM* Q = RSAKey->q;
+		const BIGNUM* PublicModulus = RSAKey->n;
+		const BIGNUM* PublicExponent = RSAKey->e;
+		const BIGNUM* PrivateExponent = RSAKey->d;
 #else
-		const BIGNUM* P = nullptr;
-		const BIGNUM* Q = nullptr;
-		// This does not increment the refcounts of P or Q
-		RSA_get0_factors(RSAKey, &P, &Q);
+		const BIGNUM* PublicModulus = RSA_get0_n(RSAKey);
+		const BIGNUM* PublicExponent = RSA_get0_e(RSAKey);
+		const BIGNUM* PrivateExponent = RSA_get0_d(RSAKey);
 #endif
 
-		uint32 NumBytes = BN_num_bytes(P);
-
-		NewP.Empty(NumBytes);
-		NewP.AddUninitialized(NumBytes);
-		BN_bn2bin(P, NewP.GetData());
-		P = nullptr;
-
-		NumBytes = BN_num_bytes(Q);
-
-		NewQ.Empty(NumBytes);
-		NewQ.AddUninitialized(NumBytes);
-		BN_bn2bin(Q, NewP.GetData());
-		Q = nullptr;
-
+		BigNumToArray(PublicModulus, OutModulus, KeySizeInBytes);
+		BigNumToArray(PublicExponent, OutPublicExponent, KeySizeInBytes);
+		BigNumToArray(PrivateExponent, OutPrivateExponent, KeySizeInBytes);
+		
 		RSA_free(RSAKey);
 
-		NewP.AddZeroed(sizeof(TEncryptionInt) - NewP.Num());
-		NewQ.AddZeroed(sizeof(TEncryptionInt) - NewQ.Num());
-
-		check(NewP.Num() == sizeof(TEncryptionInt));
-		check(NewQ.Num() == sizeof(TEncryptionInt));
-
-		TEncryptionInt OurP((uint32*)&NewP[0]);
-		TEncryptionInt OurQ((uint32*)&NewQ[0]);
-
-		FEncryptionKey PublicKey, PrivateKey;
-		FEncryption::GenerateKeyPair(OurP, OurQ, PublicKey, PrivateKey);
-
-		TestKeys(PublicKey, PrivateKey);
-
-		OutPublicExponent.AddUninitialized(sizeof(TEncryptionInt));
-		OutPrivateExponent.AddUninitialized(sizeof(TEncryptionInt));
-		OutModulus.AddUninitialized(sizeof(TEncryptionInt));
-		FMemory::Memcpy(&OutPublicExponent[0], (const uint8*)PublicKey.Exponent.GetBits(), sizeof(TEncryptionInt));
-		FMemory::Memcpy(&OutPrivateExponent[0], (const uint8*)PrivateKey.Exponent.GetBits(), sizeof(TEncryptionInt));
-		FMemory::Memcpy(&OutModulus[0], (const uint8*)PublicKey.Modulus.GetBits(), sizeof(TEncryptionInt));
-
 		return true;
 	}
 }
diff --git a/Engine/Plugins/Editor/CryptoKeys/Source/CryptoKeysOpenSSL/Public/CryptoKeysOpenSSL.h b/Engine/Plugins/Editor/CryptoKeys/Source/CryptoKeysOpenSSL/Public/CryptoKeysOpenSSL.h
--- a/Engine/Plugins/Editor/CryptoKeys/Source/CryptoKeysOpenSSL/Public/CryptoKeysOpenSSL.h
+++ b/Engine/Plugins/Editor/CryptoKeys/Source/CryptoKeysOpenSSL/Public/CryptoKeysOpenSSL.h
@@ -6,5 +6,5 @@
 namespace CryptoKeysOpenSSL
 {
 	bool CRYPTOKEYSOPENSSL_API GenerateNewEncryptionKey(TArray<uint8>& OutKey);
-	bool CRYPTOKEYSOPENSSL_API GenerateNewSigningKey(TArray<uint8>& OutPublicExponent, TArray<uint8>& OutPrivateExponent, TArray<uint8>& OutModulus);
+	bool CRYPTOKEYSOPENSSL_API GenerateNewSigningKey(TArray<uint8>& OutPublicExponent, TArray<uint8>& OutPrivateExponent, TArray<uint8>& OutModulus, int32 InNumKeyBits = 2048);
 }
\ No newline at end of file
diff --git a/Engine/Source/Developer/PakFileUtilities/Private/PakFileUtilities.cpp b/Engine/Source/Developer/PakFileUtilities/Private/PakFileUtilities.cpp
--- a/Engine/Source/Developer/PakFileUtilities/Private/PakFileUtilities.cpp
+++ b/Engine/Source/Developer/PakFileUtilities/Private/PakFileUtilities.cpp
@@ -29,62 +29,55 @@ IMPLEMENT_MODULE(FDefaultModuleImpl, PakFileUtilities);
 struct FKeyPair
 {
 	/** Public decryption key */
-	FEncryptionKey PublicKey;
+	TPakRSAKey PublicKey;
 	/** Private encryption key */
-	FEncryptionKey PrivateKey;
+	TPakRSAKey PrivateKey;
 
-	friend FArchive& operator<<(FArchive& Ar, FKeyPair& Pair)
+	bool IsValid() const
 	{
-		Ar << Pair.PublicKey.Exponent;
-		Ar << Pair.PublicKey.Modulus;
-		Ar << Pair.PrivateKey.Exponent;
-		Ar << Pair.PrivateKey.Modulus;
-		return Ar;
+		return !PrivateKey.Exponent.IsZero() &&
+		!PrivateKey.Modulus.IsZero() &&
+		!PublicKey.Exponent.IsZero() &&
+		!PublicKey.Modulus.IsZero();
 	}
+};
 
-	bool IsValid() const
+bool TestKeys(FKeyPair& Pair)
+{
+	UE_LOG(LogPakFile, Display, TEXT("Testing signature keys."));
+
+	// Just some random values
+	static TPakRSAKey::TIntType TestData[] =
+	{
+		11,
+		253,
+		128,
+		234,
+		56,
+		89,
+		34,
+		179,
+		29,
+		1024,
+		(int64)(MAX_int32),
+		(int64)(MAX_uint32)-1
+	};
+
+	for (int32 TestIndex = 0; TestIndex < ARRAY_COUNT(TestData); ++TestIndex)
 	{
-		bool bAllKeysValid = !PrivateKey.Exponent.IsZero()
-			&& !PrivateKey.Modulus.IsZero()
-			&& !PublicKey.Exponent.IsZero()
-			&& !PublicKey.Modulus.IsZero();
-		
-		bool bIsValid = true;
-
-		if (bAllKeysValid)
-		{
-			// Just some random values
-			static TEncryptionInt TestData[] =
-			{
-				11,
-				253,
-				128,
-				234,
-				56,
-				89,
-				34,
-				179,
-				29,
-				1024,
-				(int64)(MAX_int32),
-				(int64)(MAX_uint32)-1
-			};
-
-			for (int32 TestIndex = 0; TestIndex < ARRAY_COUNT(TestData); ++TestIndex)
-			{
-				TEncryptionInt EncryptedData = FEncryption::ModularPow(TestData[TestIndex], PrivateKey.Exponent, PrivateKey.Modulus);
-				TEncryptionInt DecryptedData = FEncryption::ModularPow(EncryptedData, PublicKey.Exponent, PublicKey.Modulus);
-				if (TestData[TestIndex] != DecryptedData)
-				{
-					bIsValid = false;
-					break;
-				}
-			}
+		TPakRSAKey::TIntType EncryptedData = FEncryption::ModularPow(TestData[TestIndex], Pair.PrivateKey.Exponent, Pair.PrivateKey.Modulus);
+		TPakRSAKey::TIntType DecryptedData = FEncryption::ModularPow(EncryptedData, Pair.PublicKey.Exponent, Pair.PublicKey.Modulus);
+		if (TestData[TestIndex] != DecryptedData)
+		{
+			UE_LOG(LogPakFile, Error, TEXT("Keys do not properly encrypt/decrypt data (failed test with %lld)"), TestData[TestIndex].ToInt());
+			return false;
 		}
-
-		return bIsValid;
 	}
-};
+
+	UE_LOG(LogPakFile, Display, TEXT("Signature keys check completed successfuly."));
+
+	return true;
+}
 
 struct FNamedAESKey
 {
@@ -1163,6 +1156,26 @@ TEncryptionInt ParseEncryptionIntFromJson(TSharedPtr<FJsonObject> InObj, const T
 	}
 }
 
+void ParseRSAKeyFromJson(TSharedPtr<FJsonObject> InObj, TPakRSAKey& OutKey)
+{
+	FString ExponentBase64, ModulusBase64;
+	TArray<uint8> Exponent, Modulus;
+	if (InObj->TryGetStringField("Exponent", ExponentBase64)
+		&& InObj->TryGetStringField("Modulus", ModulusBase64))
+	{
+		FBase64::Decode(ExponentBase64, Exponent);
+		FBase64::Decode(ModulusBase64, Modulus);
+
+		check(Exponent.Num() > 0 && Exponent.Num() <= TPakRSAKey::KeySizeInBytes);
+		check(Modulus.Num() > 0 && Modulus.Num() <= TPakRSAKey::KeySizeInBytes);
+		OutKey = TPakRSAKey(Exponent, Modulus);
+	}
+	else
+	{
+		OutKey = TPakRSAKey();
+	}
+}
+
 void LoadKeyChainFromFile(const FString& InFilename, FKeyChain& OutCryptoSettings)
 {
 	FArchive* File = IFileManager::Get().CreateFileReader(*InFilename);
@@ -1196,10 +1209,8 @@ void LoadKeyChainFromFile(const FString& InFilename, FKeyChain& OutCryptoSetting
 		{
 			TSharedPtr<FJsonObject> PublicKey = (*SigningKey)->GetObjectField(TEXT("PublicKey"));
 			TSharedPtr<FJsonObject> PrivateKey = (*SigningKey)->GetObjectField(TEXT("PrivateKey"));
-			OutCryptoSettings.SigningKey.PublicKey.Exponent = ParseEncryptionIntFromJson(PublicKey, TEXT("Exponent"));
-			OutCryptoSettings.SigningKey.PublicKey.Modulus = ParseEncryptionIntFromJson(PublicKey, TEXT("Modulus"));
-			OutCryptoSettings.SigningKey.PrivateKey.Exponent = ParseEncryptionIntFromJson(PrivateKey, TEXT("Exponent"));
-			OutCryptoSettings.SigningKey.PrivateKey.Modulus = ParseEncryptionIntFromJson(PrivateKey, TEXT("Modulus"));
+			ParseRSAKeyFromJson(PublicKey, OutCryptoSettings.SigningKey.PublicKey);
+			ParseRSAKeyFromJson(PrivateKey, OutCryptoSettings.SigningKey.PrivateKey);
 			check(OutCryptoSettings.SigningKey.PublicKey.Modulus == OutCryptoSettings.SigningKey.PrivateKey.Modulus);
 		}
 
@@ -1293,10 +1304,8 @@ void LoadKeyChain(const TCHAR* CmdLine, FKeyChain& OutCryptoSettings)
 					FBase64::Decode(PrivateExpBase64, PrivateExp);
 					FBase64::Decode(ModulusBase64, Modulus);
 
-					OutCryptoSettings.SigningKey.PrivateKey.Exponent = TEncryptionInt((uint32*)&PrivateExp[0]);
-					OutCryptoSettings.SigningKey.PrivateKey.Modulus = TEncryptionInt((uint32*)&Modulus[0]);
-					OutCryptoSettings.SigningKey.PublicKey.Exponent = TEncryptionInt((uint32*)&PublicExp[0]);
-					OutCryptoSettings.SigningKey.PublicKey.Modulus = OutCryptoSettings.SigningKey.PrivateKey.Modulus;
+					OutCryptoSettings.SigningKey.PrivateKey = TPakRSAKey(PrivateExp, Modulus);
+					OutCryptoSettings.SigningKey.PublicKey = TPakRSAKey(PublicExp, Modulus);
 
 					UE_LOG(LogPakFile, Display, TEXT("Parsed signature keys from config files."));
 				}
@@ -1404,11 +1413,6 @@ void LoadKeyChain(const TCHAR* CmdLine, FKeyChain& OutCryptoSettings)
 		}
 	}
 
-	if (!OutCryptoSettings.SigningKey.IsValid())
-	{
-		UE_LOG(LogPakFile, Fatal, TEXT("Pak signing keys are invalid"));
-	}
-
 	FString EncryptionKeyOverrideGuidString;
 	FGuid EncryptionKeyOverrideGuid;
 	if (FParse::Value(CmdLine, TEXT("EncryptionKeyOverrideGuid="), EncryptionKeyOverrideGuidString))
diff --git a/Engine/Source/Developer/PakFileUtilities/Private/SignedArchiveWriter.cpp b/Engine/Source/Developer/PakFileUtilities/Private/SignedArchiveWriter.cpp
--- a/Engine/Source/Developer/PakFileUtilities/Private/SignedArchiveWriter.cpp
+++ b/Engine/Source/Developer/PakFileUtilities/Private/SignedArchiveWriter.cpp
@@ -5,7 +5,7 @@
 #include "Misc/SecureHash.h"
 #include "HAL/FileManager.h"
 
-FSignedArchiveWriter::FSignedArchiveWriter(FArchive& InPak, const FString& InPakFilename, const FEncryptionKey& InPublicKey, const FEncryptionKey& InPrivateKey)
+FSignedArchiveWriter::FSignedArchiveWriter(FArchive& InPak, const FString& InPakFilename, const TPakRSAKey& InPublicKey, const TPakRSAKey& InPrivateKey)
 : BufferArchive(Buffer)
 	, PakWriter(InPak)
 	, PakSignaturesFilename(FPaths::ChangeExtension(InPakFilename, TEXT("sig")))
@@ -44,16 +44,21 @@ bool FSignedArchiveWriter::Close()
 		SerializeBufferAndSign();
 	}
 
-	FEncryptedSignature EncryptedMasterHash;
-	FDecryptedSignature DecryptedMasterHash;
-	DecryptedMasterHash.Data = ComputePakChunkHash((const uint8*)&ChunkHashes[0], ChunkHashes.Num() * sizeof(TPakChunkHash));
-	FEncryption::EncryptSignature(DecryptedMasterHash, EncryptedMasterHash, PrivateKey);
-
 	FArchive* SignatureWriter = IFileManager::Get().CreateFileWriter(*PakSignaturesFilename);
-	*SignatureWriter << EncryptedMasterHash;
-	*SignatureWriter << ChunkHashes;
+	FPakSignatureFile SignatureFile;
+	SignatureFile.Initialize(ChunkHashes, PrivateKey);
+	SignatureFile.Serialize(*SignatureWriter);
 	delete SignatureWriter;
 
+	// Test reading back in
+	{
+		FArchive* SignatureReader = IFileManager::Get().CreateFileReader(*PakSignaturesFilename);
+		FPakSignatureFile SigFileTest;
+		SigFileTest.Serialize(*SignatureReader);
+		UE_CLOG(!SigFileTest.DecryptSignatureAndValidate(PublicKey, FString()), LogPakFile, Fatal, TEXT("Failed to validate RSA encrypted pak hash"));
+		delete SignatureReader;
+	}
+
 	return FArchive::Close();
 }
 
diff --git a/Engine/Source/Developer/PakFileUtilities/Private/SignedArchiveWriter.h b/Engine/Source/Developer/PakFileUtilities/Private/SignedArchiveWriter.h
--- a/Engine/Source/Developer/PakFileUtilities/Private/SignedArchiveWriter.h
+++ b/Engine/Source/Developer/PakFileUtilities/Private/SignedArchiveWriter.h
@@ -25,9 +25,9 @@ class FSignedArchiveWriter : public FArchive
 	/** Data size (excluding signatures) */
 	int64 PakSize;
 	/** Decryption key */
-	FEncryptionKey PublicKey;
+	TPakRSAKey PublicKey;
 	/** Encryption key */
-	FEncryptionKey PrivateKey;
+	TPakRSAKey PrivateKey;
 	/** Hashes */
 	TArray<TPakChunkHash> ChunkHashes;
 
@@ -38,7 +38,7 @@ class FSignedArchiveWriter : public FArchive
 
 public:
 
-	FSignedArchiveWriter(FArchive& InPak, const FString& InPakFilename, const FEncryptionKey& InPublicKey, const FEncryptionKey& InPrivateKey);
+	FSignedArchiveWriter(FArchive& InPak, const FString& InPakFilename, const TPakRSAKey& InPublicKey, const TPakRSAKey& InPrivateKey);
 	virtual ~FSignedArchiveWriter();
 
 	// FArchive interface
diff --git a/Engine/Source/Runtime/Core/Private/Misc/CoreDelegates.cpp b/Engine/Source/Runtime/Core/Private/Misc/CoreDelegates.cpp
--- a/Engine/Source/Runtime/Core/Private/Misc/CoreDelegates.cpp
+++ b/Engine/Source/Runtime/Core/Private/Misc/CoreDelegates.cpp
@@ -182,12 +182,12 @@ FSimpleMulticastDelegate& FCoreDelegates::GetOutOfMemoryDelegate()
 
 FCoreDelegates::FGetOnScreenMessagesDelegate FCoreDelegates::OnGetOnScreenMessages;
 
-typedef void(*TSigningKeyFunc)(uint8[64], uint8[64]);
+typedef void(*TSigningKeyFunc)(TArray<uint8>&, TArray<uint8>&);
 typedef void(*TEncryptionKeyFunc)(unsigned char[32]);
 
 void RegisterSigningKeyCallback(TSigningKeyFunc InCallback)
 {
-	FCoreDelegates::GetPakSigningKeysDelegate().BindLambda([InCallback](uint8 OutExponent[64], uint8 OutModulus[64])
+	FCoreDelegates::GetPakSigningKeysDelegate().BindLambda([InCallback](TArray<uint8>& OutExponent, TArray<uint8>& OutModulus)
 	{
 		InCallback(OutExponent, OutModulus);
 	});
diff --git a/Engine/Source/Runtime/Core/Private/Misc/SecureHash.cpp b/Engine/Source/Runtime/Core/Private/Misc/SecureHash.cpp
--- a/Engine/Source/Runtime/Core/Private/Misc/SecureHash.cpp
+++ b/Engine/Source/Runtime/Core/Private/Misc/SecureHash.cpp
@@ -610,6 +610,16 @@ uint32 GetTypeHash(FSHAHash const& InKey)
 	return FCrc::MemCrc32(InKey.Hash, sizeof(InKey.Hash));
 }
 
+FString LexToString(const FSHAHash& InHash)
+{
+	return InHash.ToString();
+}
+
+void LexFromString(FSHAHash& InHash, const TCHAR* InString)
+{
+	InHash.FromString(InString);
+}
+
 FSHA1::FSHA1()
 {
 	m_block = (SHA1_WORKSPACE_BLOCK *)m_workspace;
diff --git a/Engine/Source/Runtime/Core/Public/Math/BigInt.h b/Engine/Source/Runtime/Core/Public/Math/BigInt.h
--- a/Engine/Source/Runtime/Core/Public/Math/BigInt.h
+++ b/Engine/Source/Runtime/Core/Public/Math/BigInt.h
@@ -84,6 +84,11 @@ public:
 		return Bits;
 	}
 
+	FORCEINLINE const uint32* GetBits() const
+	{
+		return Bits;
+	}
+
 	/** Sets this integer to 0. */
 	FORCEINLINE void Zero()
 	{
@@ -126,6 +131,21 @@ public:
 		FMemory::Memcpy(Bits, InBits, sizeof(Bits));
 	}
 
+	/**
+	 * Constructor. Initializes this big int with an array of bytes.
+	 */
+	explicit TBigInt(const uint8* InData, uint32 InNumBytes)
+	{
+		// If we weren't given enough data to fill the int, zero first
+		// TODO: Only zero the bits we aren't going to copy over
+		if (InNumBytes < (ARRAY_COUNT(Bits) * sizeof(uint32)))
+		{
+			Zero();
+		}
+
+		FMemory::Memcpy(Bits, InData, InNumBytes);
+	}
+
 	/**
 	 * Constructor. Initializes this big int with a string representing a hex value.
 	 */
@@ -139,7 +159,7 @@ public:
 	 *
 	 * @param BitCount the number of bits to shift.
 	 */
-	FORCEINLINE void ShiftLeftInternal(const int32 BitCount)
+	void ShiftLeftInternal(const int32 BitCount)
 	{
 		checkSlow(BitCount > 0);
 
@@ -172,7 +192,7 @@ public:
 	/**
 	 * Shift left by 1 bit.
 	 */
-	FORCEINLINE void ShiftLeftByOneInternal()
+	void ShiftLeftByOneInternal()
 	{
 		const int32 HiWordShift = BitsPerWord - 1;
 		Bits[NumWords - 1] = Bits[NumWords - 1] << 1;
@@ -189,7 +209,7 @@ public:
 	 *
 	 * @param BitCount the number of bits to shift.
 	 */
-	FORCEINLINE void ShiftRightInternal(const int32 BitCount)
+	void ShiftRightInternal(const int32 BitCount)
 	{
 		checkSlow(BitCount > 0);
 
@@ -222,7 +242,7 @@ public:
 	/**
 	 * Shift right by 1 bit.
 	 */
-	FORCEINLINE void ShiftRightByOneInternal()
+	void ShiftRightByOneInternal()
 	{
 		const int32 LoWordShift = BitsPerWord - 1;
 		Bits[0] = Bits[0] >> 1;
@@ -237,7 +257,7 @@ public:
 	/**
 	 * Adds two integers.
 	 */
-	FORCEINLINE void Add(const BigInt& Other)
+	void Add(const BigInt& Other)
 	{
 		int64 CarryOver = 0;
 		for (int32 WordIndex = 0; WordIndex < NumWords; ++WordIndex)
@@ -252,7 +272,7 @@ public:
 	/**
 	 * Subtracts two integers.
 	 */
-	FORCEINLINE void Subtract(const BigInt& Other)
+	void Subtract(const BigInt& Other)
 	{
 		BigInt NegativeOther(Other);
 		NegativeOther.Negate();
@@ -262,7 +282,7 @@ public:
 	/**
 	 * Checks if this integer is negative.
 	 */
-	FORCEINLINE bool IsNegative() const
+	bool IsNegative() const
 	{
 		if (bSigned)
 		{
@@ -277,7 +297,7 @@ public:
 	/**
 	 * Returns the sign of this integer.
 	 */
-	FORCEINLINE int32 Sign() const
+	int32 Sign() const
 	{
 		return IsNegative() ? -1 : 1;
 	}
@@ -297,7 +317,7 @@ public:
 	/**
 	* Returns the index of the highest word that is not zero. -1 if no such word exists.
 	*/
-	FORCEINLINE int32 GetHighestNonZeroWord() const
+	int32 GetHighestNonZeroWord() const
 	{
 		int32 WordIndex;
 		for (WordIndex = NumWords - 1; WordIndex >= 0 && Bits[WordIndex] == 0; --WordIndex);
@@ -307,7 +327,7 @@ public:
 	/**
 	* Multiplies two positive integers.
 	 */
-	FORCEINLINE void MultiplyFast(const BigInt& Factor)
+	void MultiplyFast(const BigInt& Factor)
 	{
 		const int64 Base = 2LL << BitsPerWord;
 		TBigInt<NumBits * 2, bSigned> Result; // Twice as many bits for overflow protection
@@ -340,7 +360,7 @@ public:
 	/**
 	 * Multiplies two integers.
 	 */
-	FORCEINLINE void Multiply(const BigInt& Factor)
+	void Multiply(const BigInt& Factor)
 	{
 		BigInt Result = *this;
 		BigInt Other = Factor;
@@ -443,7 +463,7 @@ public:
 	/**
 	 * Performs modulo operation on this integer.
 	 */
-	FORCEINLINE void Modulo(const BigInt& Modulus)
+	void Modulo(const BigInt& Modulus)
 	{
 		// a mod b = a - floor(a/b)*b
 		check(!IsNegative());
@@ -497,7 +517,7 @@ public:
 	/**
 	 * Assignment operator for int64 values.
 	 */
-	FORCEINLINE TBigInt& operator=(int64 Other)
+	TBigInt& operator=(int64 Other)
 	{
 		Set(Other);
 		return *this;
@@ -506,7 +526,7 @@ public:
 	/**
 	 * Returns the index of the highest non-zero bit. -1 if no such bit exists.
 	 */
-	FORCEINLINE int32 GetHighestNonZeroBit() const
+	int32 GetHighestNonZeroBit() const
 	{
 		for (int32 WordIndex = NumWords - 1; WordIndex >= 0; --WordIndex)
 		{
@@ -523,7 +543,7 @@ public:
 	/**
 	 * Returns a bit value as an integer value (0 or 1).
 	 */
-	FORCEINLINE int32 GetBit(int32 BitIndex) const
+	int32 GetBit(int32 BitIndex) const
 	{
 		const int32 WordIndex = BitIndex / BitsPerWord;
 		BitIndex -= WordIndex * BitsPerWord;
@@ -533,7 +553,7 @@ public:
 	/**
 	 * Sets a bit value.
 	 */
-	FORCEINLINE void SetBit(int32 BitIndex, int32 Value)
+	void SetBit(int32 BitIndex, int32 Value)
 	{
 		const int32 WordIndex = BitIndex / BitsPerWord;
 		BitIndex -= WordIndex * BitsPerWord;
@@ -600,7 +620,7 @@ public:
 	/**
 	 * Bitwise 'or'
 	 */
-	FORCEINLINE void BitwiseOr(const BigInt& Other)
+	void BitwiseOr(const BigInt& Other)
 	{
 		for (int32 WordIndex = 0; WordIndex < NumWords; ++WordIndex)
 		{
@@ -611,7 +631,7 @@ public:
 	/**
 	 * Bitwise 'and'
 	 */
-	FORCEINLINE void BitwiseAnd(const BigInt& Other)
+	void BitwiseAnd(const BigInt& Other)
 	{
 		for (int32 WordIndex = 0; WordIndex < NumWords; ++WordIndex)
 		{
@@ -622,7 +642,7 @@ public:
 	/**
 	 * Bitwise 'not'
 	 */
-	FORCEINLINE void BitwiseNot()
+	void BitwiseNot()
 	{
 		for (int32 WordIndex = 0; WordIndex < NumWords; ++WordIndex)
 		{
@@ -633,7 +653,7 @@ public:
 	/**
 	 * Checks if two integers are equal.
 	 */
-	FORCEINLINE bool IsEqual(const BigInt& Other) const
+	bool IsEqual(const BigInt& Other) const
 	{
 		for (int32 WordIndex = 0; WordIndex < NumWords; ++WordIndex)
 		{
@@ -649,7 +669,7 @@ public:
 	/**
 	 * this < Other
 	 */
-	FORCEINLINE bool IsLess(const BigInt& Other) const
+	bool IsLess(const BigInt& Other) const
 	{
 		if (IsNegative())
 		{
@@ -670,7 +690,7 @@ public:
 	/**
 	 * this <= Other
 	 */
-	FORCEINLINE bool IsLessOrEqual(const BigInt& Other) const
+	bool IsLessOrEqual(const BigInt& Other) const
 	{
 		if (IsNegative())
 		{
@@ -691,7 +711,7 @@ public:
 	/**
 	 * this > Other
 	 */
-	FORCEINLINE bool IsGreater(const BigInt& Other) const
+	bool IsGreater(const BigInt& Other) const
 	{
 		if (IsNegative())
 		{
@@ -712,7 +732,7 @@ public:
 	/**
 	 * this >= Other
 	 */
-	FORCEINLINE bool IsGreaterOrEqual(const BigInt& Other) const
+	bool IsGreaterOrEqual(const BigInt& Other) const
 	{
 		if (IsNegative())
 		{
@@ -733,7 +753,7 @@ public:
 	/**
 	 * this == 0
 	 */
-	FORCEINLINE bool IsZero() const
+	bool IsZero() const
 	{
 		int32 WordIndex;
 		for (WordIndex = NumWords - 1; WordIndex >= 0 && !Bits[WordIndex]; --WordIndex);
@@ -743,7 +763,7 @@ public:
 	/**
 	 * this > 0
 	 */
-	FORCEINLINE bool IsGreaterThanZero() const
+	bool IsGreaterThanZero() const
 	{
 		return !IsNegative() && !IsZero();
 	}
@@ -751,184 +771,184 @@ public:
 	/**
 	 * this < 0
 	 */
-	FORCEINLINE bool IsLessThanZero() const
+	bool IsLessThanZero() const
 	{
 		return IsNegative() && !IsZero();
 	}
 
-	FORCEINLINE bool IsFirstBitSet() const
+	bool IsFirstBitSet() const
 	{
 		return !!(Bits[0] & 1);
 	}
 	/**
 	 * Bit indexing operator
 	 */
-	FORCEINLINE bool operator[](int32 BitIndex) const
+	bool operator[](int32 BitIndex) const
 	{
 		return GetBit(BitIndex);
 	}
 
 	// Begin operator overloads
-	FORCEINLINE BigInt operator>>(int32 Count) const
+	BigInt operator>>(int32 Count) const
 	{
 		BigInt Result = *this;
 		Result.ShiftRight(Count);
 		return Result;
 	}
 
-	FORCEINLINE BigInt& operator>>=(int32 Count)
+	BigInt& operator>>=(int32 Count)
 	{
 		ShiftRight(Count);
 		return *this;
 	}
 
-	FORCEINLINE BigInt operator<<(int32 Count) const
+	BigInt operator<<(int32 Count) const
 	{
 		BigInt Result = *this;
 		Result.ShiftLeft(Count);
 		return Result;
 	}
 
-	FORCEINLINE BigInt& operator<<=(int32 Count)
+	BigInt& operator<<=(int32 Count)
 	{
 		ShiftLeft(Count);
 		return *this;
 	}
 
-	FORCEINLINE BigInt operator+(const BigInt& Other) const
+	BigInt operator+(const BigInt& Other) const
 	{
 		BigInt Result(*this);
 		Result.Add(Other);
 		return Result;
 	}
 
-	FORCEINLINE BigInt& operator++()
+	BigInt& operator++()
 	{
 		Add(One);
 		return *this;
 	}
 
-	FORCEINLINE BigInt& operator+=(const BigInt& Other)
+	BigInt& operator+=(const BigInt& Other)
 	{
 		Add(Other);
 		return *this;
 	}
 
-	FORCEINLINE BigInt operator-(const BigInt& Other) const
+	BigInt operator-(const BigInt& Other) const
 	{
 		BigInt Result(*this);
 		Result.Subtract(Other);
 		return Result;
 	}
 
-	FORCEINLINE BigInt& operator--()
+	BigInt& operator--()
 	{
 		Subtract(One);
 		return *this;
 	}
 
-	FORCEINLINE BigInt& operator-=(const BigInt& Other)
+	BigInt& operator-=(const BigInt& Other)
 	{
 		Subtract(Other);
 		return *this;
 	}
 
-	FORCEINLINE BigInt operator*(const BigInt& Other) const
+	BigInt operator*(const BigInt& Other) const
 	{
 		BigInt Result(*this);
 		Result.Multiply(Other);
 		return Result;
 	}
 
-	FORCEINLINE BigInt& operator*=(const BigInt& Other)
+	BigInt& operator*=(const BigInt& Other)
 	{
 		Multiply(Other);
 		return *this;
 	}
 
-	FORCEINLINE BigInt operator/(const BigInt& Divider) const
+	BigInt operator/(const BigInt& Divider) const
 	{
 		BigInt Result(*this);
 		Result.Divide(Divider);
 		return Result;
 	}
 
-	FORCEINLINE BigInt& operator/=(const BigInt& Divider)
+	BigInt& operator/=(const BigInt& Divider)
 	{
 		Divide(Divider);
 		return *this;
 	}
 
-	FORCEINLINE BigInt operator%(const BigInt& Modulus) const
+	BigInt operator%(const BigInt& Modulus) const
 	{
 		BigInt Result(*this);
 		Result.Modulo(Modulus);
 		return Result;
 	}
 
-	FORCEINLINE BigInt& operator%=(const BigInt& Modulus)
+	BigInt& operator%=(const BigInt& Modulus)
 	{
 		Modulo(Modulus);
 		return *this;
 	}
 
-	FORCEINLINE bool operator<(const BigInt& Other) const
+	bool operator<(const BigInt& Other) const
 	{
 		return IsLess(Other);
 	}
 
-	FORCEINLINE bool operator<=(const BigInt& Other) const
+	bool operator<=(const BigInt& Other) const
 	{
 		return IsLessOrEqual(Other);
 	}
 
-	FORCEINLINE bool operator>(const BigInt& Other) const
+	bool operator>(const BigInt& Other) const
 	{
 		return IsGreater(Other);
 	}
 
-	FORCEINLINE bool operator>=(const BigInt& Other) const
+	bool operator>=(const BigInt& Other) const
 	{
 		return IsGreaterOrEqual(Other);
 	}
 
-	FORCEINLINE bool operator==(const BigInt& Other) const
+	bool operator==(const BigInt& Other) const
 	{
 		return IsEqual(Other);
 	}
 
-	FORCEINLINE bool operator!=(const BigInt& Other) const
+	bool operator!=(const BigInt& Other) const
 	{
 		return !IsEqual(Other);
 	}
 
-	FORCEINLINE BigInt operator&(const BigInt& Other) const
+	BigInt operator&(const BigInt& Other) const
 	{
 		BigInt Result(*this);
 		Result.BitwiseAnd(Other);
 		return Result;
 	}
 
-	FORCEINLINE BigInt& operator&=(const BigInt& Other)
+	BigInt& operator&=(const BigInt& Other)
 	{
 		BitwiseAnd(Other);
 		return *this;
 	}
 
-	FORCEINLINE BigInt operator|(const BigInt& Other) const
+	BigInt operator|(const BigInt& Other) const
 	{
 		BigInt Result(*this);
 		Result.BitwiseOr(Other);
 		return Result;
 	}
 
-	FORCEINLINE BigInt& operator|=(const BigInt& Other)
+	BigInt& operator|=(const BigInt& Other)
 	{
 		BitwiseOr(Other);
 		return *this;
 	}
 
-	FORCEINLINE BigInt operator~() const
+	BigInt operator~() const
 	{
 		BigInt Result(*this);
 		Result.BitwiseNot();
@@ -1022,7 +1042,7 @@ typedef TBigInt<512> int512;
 typedef TBigInt<512> TEncryptionInt;
 
 /**
- * Encyption key - exponent and modulus pair
+ * Encryption key - exponent and modulus pair
  */
 template <typename IntType>
 struct TEncryptionKey
@@ -1159,7 +1179,7 @@ namespace FEncryption
 	 * Raise Base to power of Exponent in mod Modulus.
 	 */
 	template <typename IntType>
-	FORCEINLINE static IntType ModularPow(IntType Base, IntType Exponent, IntType Modulus)
+	static IntType ModularPow(IntType Base, IntType Exponent, IntType Modulus)
 	{
 		const IntType One(1);
 		IntType Result(1);
@@ -1179,7 +1199,7 @@ namespace FEncryption
 	* Encrypts a stream of bytes
 	*/
 	template <typename IntType>
-	FORCEINLINE static void EncryptBytes(IntType* EncryptedData, const uint8* Data, const int64 DataLength, const FEncryptionKey& EncryptionKey)
+	static void EncryptBytes(IntType* EncryptedData, const uint8* Data, const int64 DataLength, const FEncryptionKey& EncryptionKey)
 	{
 		for (int Index = 0; Index < DataLength; ++Index)
 		{
@@ -1191,7 +1211,7 @@ namespace FEncryption
 	* Decrypts a stream of bytes
 	*/
 	template <typename IntType>
-	FORCEINLINE static void DecryptBytes(uint8* DecryptedData, const IntType* Data, const int64 DataLength, const FEncryptionKey& DecryptionKey)
+	static void DecryptBytes(uint8* DecryptedData, const IntType* Data, const int64 DataLength, const FEncryptionKey& DecryptionKey)
 	{
 		for (int64 Index = 0; Index < DataLength; ++Index)
 		{
@@ -1207,7 +1227,7 @@ namespace FEncryption
  * Specialization for int type used in encryption (performance). Avoids using temporary results and most of the operations are inplace.
  */
 template <>
-FORCEINLINE TEncryptionInt FEncryption::ModularPow(TEncryptionInt Base, TEncryptionInt Exponent, TEncryptionInt Modulus)
+TEncryptionInt FEncryption::ModularPow(TEncryptionInt Base, TEncryptionInt Exponent, TEncryptionInt Modulus)
 {
 	TEncryptionInt Result(1LL);
 	while (Exponent.IsGreaterThanZero())
@@ -1283,13 +1303,14 @@ struct FDecryptedSignature : public FSignatureBase<uint32>
 
 namespace FEncryption
 {
-	FORCEINLINE static void EncryptSignature(const FDecryptedSignature& InUnencryptedSignature, FEncryptedSignature& OutEncryptedSignature, const FEncryptionKey& EncryptionKey)
+	static void EncryptSignature(const FDecryptedSignature& InUnencryptedSignature, FEncryptedSignature& OutEncryptedSignature, const FEncryptionKey& EncryptionKey)
 	{
 		OutEncryptedSignature.Data = FEncryption::ModularPow(TEncryptionInt(InUnencryptedSignature.Data), EncryptionKey.Exponent, EncryptionKey.Modulus);
 	}
 
-	FORCEINLINE static void DecryptSignature(const FEncryptedSignature& InEncryptedSignature, FDecryptedSignature& OutUnencryptedSignature, const FEncryptionKey& EncryptionKey)
+	static void DecryptSignature(const FEncryptedSignature& InEncryptedSignature, FDecryptedSignature& OutUnencryptedSignature, const FEncryptionKey& EncryptionKey)
 	{
 		OutUnencryptedSignature.Data = FEncryption::ModularPow(TEncryptionInt(InEncryptedSignature.Data), EncryptionKey.Exponent, EncryptionKey.Modulus).ToInt();
 	}
 }
+
diff --git a/Engine/Source/Runtime/Core/Public/Misc/CoreDelegates.h b/Engine/Source/Runtime/Core/Public/Misc/CoreDelegates.h
--- a/Engine/Source/Runtime/Core/Public/Misc/CoreDelegates.h
+++ b/Engine/Source/Runtime/Core/Public/Misc/CoreDelegates.h
@@ -101,7 +101,7 @@ public:
 	DECLARE_DELEGATE_OneParam(FPakEncryptionKeyDelegate, uint8[32]);
 
 	// Callback for gathering pak signing keys, if they exist
-	DECLARE_DELEGATE_TwoParams(FPakSigningKeysDelegate, uint8[64], uint8[64]);
+	DECLARE_DELEGATE_TwoParams(FPakSigningKeysDelegate, TArray<uint8>&, TArray<uint8>&);
 
 	// Callback for handling the Controller connection / disconnection
 	// first param is true for a connection, false for a disconnection.
diff --git a/Engine/Source/Runtime/Core/Public/Misc/RSA.h b/Engine/Source/Runtime/Core/Public/Misc/RSA.h
new file mode 100644
--- /dev/null
+++ b/Engine/Source/Runtime/Core/Public/Misc/RSA.h
@@ -0,0 +1,84 @@
+// Copyright 1998-2019 Epic Games, Inc. All Rights Reserved.
+#pragma once
+
+#include "CoreTypes.h"
+#include "Math/BigInt.h"
+
+template<int KEY_SIZE>
+struct FRSAKey
+{
+	enum { KeySize = KEY_SIZE };
+	enum { KeySizeInBytes = KeySize / 8 };
+	enum { MaxDataSize = KeySizeInBytes - 7 };
+	static_assert(KEY_SIZE % 8 == 0, "Key sizes must be multiples of 8");
+	typedef TBigInt<KeySize * 2, false> TIntType;
+	TIntType Exponent;
+	TIntType Modulus;
+
+	FRSAKey()
+	{
+
+	}
+
+	FRSAKey(const TArray<uint8>& InExponent, const TArray<uint8>& InModulus)
+		: Exponent(InExponent.GetData(), InExponent.Num())
+		, Modulus(InModulus.GetData(), InModulus.Num())
+	{
+	}
+
+	bool IsValid() const
+	{
+		return !Exponent.IsZero() && !Modulus.IsZero();
+	}
+
+	static void Encrypt(const TIntType& InSource, TIntType& OutEncrypted, const FRSAKey& InKey)
+	{
+		OutEncrypted = FEncryption::ModularPow(InSource, InKey.Exponent, InKey.Modulus);
+	}
+
+	static void Encrypt(const uint8* InSourceBuffer, int64 InSourceBufferSizeInBytes, TArray<uint8>& OutEncrypted, const FRSAKey& InKey)
+	{
+		if (InSourceBufferSizeInBytes <= MaxDataSize)
+		{
+			TIntType Source(InSourceBuffer, InSourceBufferSizeInBytes);
+			// TODO: Pad Source data!
+			TIntType Encrypted = FEncryption::ModularPow(Source, InKey.Exponent, InKey.Modulus);
+			OutEncrypted.SetNum(KeySizeInBytes);
+			FMemory::Memcpy(OutEncrypted.GetData(), Encrypted.GetBits(), KeySizeInBytes);
+		}
+		else
+		{
+			OutEncrypted.Empty();
+		}
+	}
+
+	static void Encrypt(const TArray<uint8> InSource, TArray<uint8>& OutEncrypted, const FRSAKey& InKey)
+	{
+		Encrypt(InSource.GetData(), InSource.Num(), OutEncrypted, InKey);
+	}
+
+	static void Decrypt(const TIntType& InEncrypted, TIntType& OutDecrypted, const FRSAKey& InKey)
+	{
+		OutDecrypted = FEncryption::ModularPow(InEncrypted, InKey.Exponent, InKey.Modulus);
+	}
+
+	static void Decrypt(const TArray<uint8>& InEncrypted, uint8* InOutputBuffer, int32 InOutputBufferSizeInBytes, const FRSAKey& InKey)
+	{
+		if (InEncrypted.Num() <= KeySizeInBytes)
+		{
+			TIntType Source(InEncrypted.GetData(), InEncrypted.Num());
+			TIntType Decrypted = FEncryption::ModularPow(Source, InKey.Exponent, InKey.Modulus);
+			FMemory::Memcpy(InOutputBuffer, Decrypted.GetBits(), FMath::Min((int32)MaxDataSize, InOutputBufferSizeInBytes));
+		}
+		else
+		{
+			FMemory::Memset(InOutputBuffer, 0, InOutputBufferSizeInBytes);
+		}
+	}
+	
+	static void Decrypt(const TArray<uint8>& InEncrypted, TArray<uint8>& OutDecrypted, const FRSAKey& InKey)
+	{
+		OutDecrypted.SetNumUninitialized(MaxDataSize);
+		Decrypt(InEncrypted, OutDecrypted.GetData(), MaxDataSize, InKey);
+	}
+};
\ No newline at end of file
diff --git a/Engine/Source/Runtime/Core/Public/Misc/SecureHash.h b/Engine/Source/Runtime/Core/Public/Misc/SecureHash.h
--- a/Engine/Source/Runtime/Core/Public/Misc/SecureHash.h
+++ b/Engine/Source/Runtime/Core/Public/Misc/SecureHash.h
@@ -223,6 +223,9 @@ public:
 	friend CORE_API FArchive& operator<<( FArchive& Ar, FSHAHash& G );
 	
 	friend CORE_API uint32 GetTypeHash(FSHAHash const& InKey);
+
+	friend CORE_API FString LexToString(const FSHAHash&);
+	friend CORE_API void LexFromString(FSHAHash& Hash, const TCHAR*);
 };
 
 class CORE_API FSHA1
diff --git a/Engine/Source/Runtime/Core/Public/Modules/ModuleManager.h b/Engine/Source/Runtime/Core/Public/Modules/ModuleManager.h
--- a/Engine/Source/Runtime/Core/Public/Modules/ModuleManager.h
+++ b/Engine/Source/Runtime/Core/Public/Modules/ModuleManager.h
@@ -765,16 +765,21 @@ class FDefaultGameModuleImpl
 	{ \
 		FSigningKeyRegistration() \
 		{ \
-			extern void RegisterSigningKeyCallback(void (*)(unsigned char OutExponent[64], unsigned char OutModulus[64])); \
+			extern void RegisterSigningKeyCallback(void (*)(TArray<uint8>&, TArray<uint8>&)); \
 			RegisterSigningKeyCallback(&Callback); \
 		} \
-		static void Callback(unsigned char OutExponent[64], unsigned char OutModulus[64]) \
+		static void Callback(TArray<uint8>& OutExponent, TArray<uint8>& OutModulus) \
 		{ \
-			const unsigned char Exponent[64] = { ExponentValue }; \
-			const unsigned char Modulus[64] = { ModulusValue }; \
-			for(int ByteIdx = 0; ByteIdx < 64; ByteIdx++) \
+			const uint8 Exponent[] = { ExponentValue }; \
+			const uint8 Modulus[] = { ModulusValue }; \
+			OutExponent.SetNum(ARRAY_COUNT(Exponent)); \
+			OutModulus.SetNum(ARRAY_COUNT(Modulus)); \
+			for(int ByteIdx = 0; ByteIdx < ARRAY_COUNT(Exponent); ByteIdx++) \
 			{ \
 				OutExponent[ByteIdx] = Exponent[ByteIdx]; \
+			} \
+			for(int ByteIdx = 0; ByteIdx < ARRAY_COUNT(Modulus); ByteIdx++) \
+			{ \
 				OutModulus[ByteIdx] = Modulus[ByteIdx]; \
 			} \
 		} \
diff --git a/Engine/Source/Runtime/PakFile/Private/IPlatformFilePak.cpp b/Engine/Source/Runtime/PakFile/Private/IPlatformFilePak.cpp
--- a/Engine/Source/Runtime/PakFile/Private/IPlatformFilePak.cpp
+++ b/Engine/Source/Runtime/PakFile/Private/IPlatformFilePak.cpp
@@ -152,7 +152,7 @@ TPakChunkHash ComputePakChunkHash(const void* InData, int64 InDataSizeInBytes)
 	return FCrc::MemCrc32(InData, InDataSizeInBytes);
 #else
 	FSHAHash Hash;
-	FSHA1::HashBuffer(InData, InDataSizeInBytes, Hash);
+	FSHA1::HashBuffer(InData, InDataSizeInBytes, Hash.Hash);
 	return Hash;
 #endif
 }
@@ -213,16 +213,15 @@ void FPakPlatformFile::GetPakEncryptionKey(FAES::FAESKey& OutKey, const FGuid& I
 	}
 }
 
-void FPakPlatformFile::GetPakSigningKeys(FEncryptionKey& OutKey)
+void FPakPlatformFile::GetPakSigningKeys(TPakRSAKey& OutKey)
 {
 	FCoreDelegates::FPakSigningKeysDelegate& Delegate = FCoreDelegates::GetPakSigningKeysDelegate();
 	if (Delegate.IsBound())
 	{
-		uint8 Exponent[sizeof(TEncryptionInt)];
-		uint8 Modulus[sizeof(TEncryptionInt)];
+		TArray<uint8> Exponent;
+		TArray<uint8> Modulus;
 		Delegate.Execute(Exponent, Modulus);
-		OutKey.Exponent = TEncryptionInt((uint32*)Exponent);
-		OutKey.Modulus = TEncryptionInt((uint32*)Modulus);
+		OutKey = TPakRSAKey(Exponent, Modulus);
 	}
 }
 
@@ -1038,8 +1037,7 @@ class FPakPrecacher
 		TIntervalTreeIndex InRequests[AIOP_NUM][(int32)EInRequestStatus::Num];
 		TIntervalTreeIndex CacheBlocks[(int32)EBlockStatus::Num];
 
-		TArray<TPakChunkHash> ChunkHashes;
-		TPakChunkHash OriginalSignatureFileHash;
+		FPakSignatureFile Signatures;
 
 		FPakData(IAsyncReadFileHandle* InHandle, FName InName, int64 InTotalSize)
 			: Handle(InHandle)
@@ -1048,7 +1046,6 @@ class FPakPrecacher
 			, MaxShift(0)
 			, BytesToBitsShift(0)
 			, Name(InName)
-			, OriginalSignatureFileHash(0)
 		{
 			check(Handle && TotalSize > 0 && Name != NAME_None);
 			for (int32 Index = 0; Index < AIOP_NUM; Index++)
@@ -1129,13 +1126,13 @@ class FPakPrecacher
 	uint32 Loads;
 	uint32 Frees;
 	uint64 LoadSize;
-	FEncryptionKey EncryptionKey;
+	TPakRSAKey PublicKey;
 	bool bSigned;
 	EAsyncIOPriorityAndFlags AsyncMinPriority;
 	FCriticalSection SetAsyncMinimumPriorityScopeLock;
 public:
 
-	static void Init(IPlatformFile* InLowerLevel, const FEncryptionKey& InEncryptionKey)
+	static void Init(IPlatformFile* InLowerLevel, const TPakRSAKey& InEncryptionKey)
 	{
 		if (!PakPrecacherSingleton)
 		{
@@ -1175,7 +1172,7 @@ public:
 		return *PakPrecacherSingleton;
 	}
 
-	FPakPrecacher(IPlatformFile* InLowerLevel, const FEncryptionKey& InEncryptionKey)
+	FPakPrecacher(IPlatformFile* InLowerLevel, const TPakRSAKey& InPublicKey)
 		: LowerLevel(InLowerLevel)
 		, LastReadRequest(0)
 		, NextUniqueID(1)
@@ -1185,8 +1182,8 @@ public:
 		, Loads(0)
 		, Frees(0)
 		, LoadSize(0)
-		, EncryptionKey(InEncryptionKey)
-		, bSigned(!InEncryptionKey.Exponent.IsZero() && !InEncryptionKey.Modulus.IsZero())
+		, PublicKey(InPublicKey)
+		, bSigned(!InPublicKey.Exponent.IsZero() && !InPublicKey.Modulus.IsZero())
 		, AsyncMinPriority(AIOP_MIN)
 	{
 		check(LowerLevel && FPlatformProcess::SupportsMultithreading());
@@ -1213,42 +1210,33 @@ public:
 		uint16* PakIndexPtr = CachedPaks.Find(File);
 		if (!PakIndexPtr)
 		{
+			FString PakFilename = File.ToString();
 			check(CachedPakData.Num() < MAX_uint16);
-			IAsyncReadFileHandle* Handle = LowerLevel->OpenAsyncRead(*File.ToString());
+			IAsyncReadFileHandle* Handle = LowerLevel->OpenAsyncRead(*PakFilename);
 			if (!Handle)
 			{
 				return nullptr;
 			}
 			CachedPakData.Add(FPakData(Handle, File, PakFileSize));
 			PakIndexPtr = &CachedPaks.Add(File, CachedPakData.Num() - 1);
-			UE_LOG(LogPakFile, Log, TEXT("New pak file %s added to pak precacher."), *File.ToString());
+			UE_LOG(LogPakFile, Log, TEXT("New pak file %s added to pak precacher."), *PakFilename);
 
 			FPakData& Pak = CachedPakData[*PakIndexPtr];
 
 			if (bSigned)
 			{
 				// Load signature data
-				FString SignaturesFilename = FPaths::ChangeExtension(File.ToString(), TEXT("sig"));
+				FString SignaturesFilename = FPaths::ChangeExtension(*PakFilename, TEXT("sig"));
 				IFileHandle* SignaturesFile = LowerLevel->OpenRead(*SignaturesFilename);
 				ensure(SignaturesFile);
-
 				FArchiveFileReaderGeneric* Reader = new FArchiveFileReaderGeneric(SignaturesFile, *SignaturesFilename, SignaturesFile->Size());
-				FEncryptedSignature MasterSignature;
-				*Reader << MasterSignature;
-				*Reader << Pak.ChunkHashes;
+				Pak.Signatures.Serialize(*Reader);
 				delete Reader;
+				Pak.Signatures.DecryptSignatureAndValidate(PublicKey, PakFilename);
 
 				// Check that we have the correct match between signature and pre-cache granularity
 				int64 NumPakChunks = Align(PakFileSize, FPakInfo::MaxChunkDataSize) / FPakInfo::MaxChunkDataSize;
-				ensure(NumPakChunks == Pak.ChunkHashes.Num());
-
-				// Decrypt signature hash
-				FDecryptedSignature DecryptedSignature;
-				FEncryption::DecryptSignature(MasterSignature, DecryptedSignature, EncryptionKey);
-
-				// Check the signatures are still as we expected them
-				Pak.OriginalSignatureFileHash = ComputePakChunkHash(&Pak.ChunkHashes[0], Pak.ChunkHashes.Num() * sizeof(TPakChunkHash));
-				ensure(Pak.OriginalSignatureFileHash == DecryptedSignature.Data);
+				ensure(NumPakChunks == Pak.Signatures.ChunkHashes.Num());
 			}
 		}
 		return PakIndexPtr;
@@ -1261,10 +1249,9 @@ public:
 
 		for (FPakData& PakData : CachedPakData)
 		{
-			for (TPakChunkHash& Hash : PakData.ChunkHashes)
+			for (TPakChunkHash& Hash : PakData.Signatures.ChunkHashes)
 			{
-				Hash |= (uint32)FMath::Rand();
-				Hash &= (uint32)FMath::Rand();
+				*((uint8*)&Hash) |= 0x1;
 			}
 		}
 	}
@@ -3117,9 +3104,9 @@ void FPakPrecacher::DoSignatureCheck(bool bWasCanceled, IAsyncReadRequest* Reque
 	int64 RequestSize = 0;
 	int64 RequestOffset = 0;
 	uint16 PakIndex;
-	TPakChunkHash MasterSignatureHash = 0;
+	FSHAHash MasterSignatureHash;
 	static const int64 MaxHashesToCache = 16;
-	TPakChunkHash HashCache[MaxHashesToCache] = { 0 };
+	TPakChunkHash HashCache[MaxHashesToCache];
 
 	{
 		// Try and keep lock for as short a time as possible. Find our request and copy out the data we need
@@ -3140,11 +3127,11 @@ void FPakPrecacher::DoSignatureCheck(bool bWasCanceled, IAsyncReadRequest* Reque
 		SignatureIndex = RequestOffset / FPakInfo::MaxChunkDataSize;
 
 		FPakData& PakData = CachedPakData[PakIndex];
-		MasterSignatureHash = PakData.OriginalSignatureFileHash;
+		MasterSignatureHash = PakData.Signatures.DecryptedHash;
 
 		for (int32 CacheIndex = 0; CacheIndex < FMath::Min(NumSignaturesToCheck, MaxHashesToCache); ++CacheIndex)
 		{
-			HashCache[CacheIndex] = PakData.ChunkHashes[SignatureIndex + CacheIndex];
+			HashCache[CacheIndex] = PakData.Signatures.ChunkHashes[SignatureIndex + CacheIndex];
 		}
 	}
 
@@ -3164,7 +3151,7 @@ void FPakPrecacher::DoSignatureCheck(bool bWasCanceled, IAsyncReadRequest* Reque
 			FPakData& PakData = CachedPakData[PakIndex];
 			for (int32 CacheIndex = 0; (CacheIndex < MaxHashesToCache) && ((SignedChunkIndex + CacheIndex) < NumSignaturesToCheck); ++CacheIndex)
 			{
-				HashCache[CacheIndex] = PakData.ChunkHashes[SignatureIndex + CacheIndex];
+				HashCache[CacheIndex] = PakData.Signatures.ChunkHashes[SignatureIndex + CacheIndex];
 			}
 		}
 
@@ -3179,11 +3166,10 @@ void FPakPrecacher::DoSignatureCheck(bool bWasCanceled, IAsyncReadRequest* Reque
 				FScopeLock Lock(&CachedFilesScopeLock);
 				FPakData* PakData = &CachedPakData[PakIndex];
 
-				UE_LOG(LogPakFile, Warning, TEXT("Pak chunk signing mismatch on chunk [%i/%i]! Expected 0x%8X, Received 0x%8X"), SignatureIndex, PakData->ChunkHashes.Num(), PakData->OriginalSignatureFileHash, ThisHash);
+				UE_LOG(LogPakFile, Warning, TEXT("Pak chunk signing mismatch on chunk [%i/%i]! Expected 0x%8X, Received 0x%8X"), SignatureIndex, PakData->Signatures.ChunkHashes.Num(), *LexToString(PakData->Signatures.ChunkHashes[SignatureIndex]), *LexToString(ThisHash));
 
 				// Check the signatures are still as we expected them
-				TPakChunkHash CurrentSignatureHash = ComputePakChunkHash(&PakData->ChunkHashes[0], PakData->ChunkHashes.Num() * sizeof(TPakChunkHash));
-				if (PakData->OriginalSignatureFileHash != CurrentSignatureHash)
+				if (PakData->Signatures.DecryptedHash != PakData->Signatures.ComputeCurrentMasterHash())
 				{
 					UE_LOG(LogPakFile, Warning, TEXT("Master signature table has changed since initialization!"));
 				}
@@ -5395,7 +5381,7 @@ bool FPakPlatformFile::Initialize(IPlatformFile* Inner, const TCHAR* CmdLine)
 	GameUserSettingsIniFilename = TEXT("GameUserSettings.ini");
 #endif
 
-	FEncryptionKey DecryptionKey;
+	TPakRSAKey DecryptionKey;
 	GetPakSigningKeys(DecryptionKey);
 
 	// signed if we have keys, and are not running with fileopenlog (currently results in a deadlock).
@@ -5456,7 +5442,7 @@ void FPakPlatformFile::InitializeNewAsyncIO()
 #if !WITH_EDITOR
 	if (FPlatformProcess::SupportsMultithreading() && !FParse::Param(FCommandLine::Get(), TEXT("FileOpenLog")))
 	{
-		FEncryptionKey DecryptionKey;
+		TPakRSAKey DecryptionKey;
 		GetPakSigningKeys(DecryptionKey);
 
 		FPakPrecacher::Init(LowerLevel, DecryptionKey);
diff --git a/Engine/Source/Runtime/PakFile/Private/SignedArchiveReader.cpp b/Engine/Source/Runtime/PakFile/Private/SignedArchiveReader.cpp
--- a/Engine/Source/Runtime/PakFile/Private/SignedArchiveReader.cpp
+++ b/Engine/Source/Runtime/PakFile/Private/SignedArchiveReader.cpp
@@ -36,8 +36,6 @@ FChunkCacheWorker::FChunkCacheWorker(FArchive* InReader, const TCHAR* Filename)
 	, QueuedRequestsEvent(nullptr)
 	, ChunkRequestAvailable(nullptr)
 {
-	SetupDecryptionKey();
-
 	FString SigFileFilename = FPaths::ChangeExtension(Filename, TEXT("sig"));
 	FArchive* SigFileReader = IFileManager::Get().CreateFileReader(*SigFileFilename);
 
@@ -46,18 +44,9 @@ FChunkCacheWorker::FChunkCacheWorker(FArchive* InReader, const TCHAR* Filename)
 		UE_LOG(LogPakFile, Fatal, TEXT("Couldn't find pak signature file '%s'"), *SigFileFilename);
 	}
 
-	FEncryptedSignature MasterSignature;
-	*SigFileReader << MasterSignature;
-	*SigFileReader << ChunkHashes;
+	Signatures.Serialize(*SigFileReader);
 	delete SigFileReader;
-
-	OriginalSignatureFileHash = ComputePakChunkHash(&ChunkHashes[0], ChunkHashes.Num() * sizeof(TPakChunkHash));
-	FDecryptedSignature DecryptedMasterSignature;
-	FEncryption::DecryptSignature(MasterSignature, DecryptedMasterSignature, DecryptionKey);
-	if (DecryptedMasterSignature.Data != OriginalSignatureFileHash)
-	{
-		FPakPlatformFile::GetPakMasterSignatureTableCheckFailureHandler().Broadcast(Filename);
-	}
+	Signatures.DecryptSignatureAndValidate(Filename);
 
 	const bool bEnableMultithreading = FPlatformProcess::SupportsMultithreading();
 	if (bEnableMultithreading)
@@ -89,31 +78,6 @@ bool FChunkCacheWorker::Init()
 	return true;
 }
 
-void FChunkCacheWorker::SetupDecryptionKey()
-{
-	FPakPlatformFile::GetPakSigningKeys(DecryptionKey);
-
-	// Public key should never be zero at this point.
-	UE_CLOG(DecryptionKey.Exponent.IsZero() || DecryptionKey.Modulus.IsZero(), LogPakFile, Fatal, TEXT("Invalid decryption key detected"));
-	// Public key should produce decrypted results - check for identity keys
-	static TEncryptionInt TestValues[] = 
-	{
-		11,
-		23,
-		67,
-		121,
-		180,
-		211
-	};
-	bool bIdentical = true;
-	for (int32 Index = 0; bIdentical && Index < ARRAY_COUNT(TestValues); ++Index)
-	{
-		TEncryptionInt DecryptedValue = FEncryption::ModularPow(TestValues[Index], DecryptionKey.Exponent, DecryptionKey.Modulus);
-		bIdentical = (DecryptedValue == TestValues[Index]);
-	}	
-	UE_CLOG(bIdentical, LogPakFile, Fatal, TEXT("Decryption key produces identical results to source data."));
-}
-
 uint32 FChunkCacheWorker::Run()
 {	
 	check(QueuedRequestsEvent);
@@ -265,7 +229,7 @@ bool FChunkCacheWorker::CheckSignature(const FChunkRequest& ChunkInfo)
 {
 	SCOPE_SECONDS_ACCUMULATOR(STAT_FChunkCacheWorker_CheckSignature);
 
-	TPakChunkHash ChunkHash = 0;
+	TPakChunkHash ChunkHash;
 
 	{
 		SCOPE_SECONDS_ACCUMULATOR(STAT_FChunkCacheWorker_Serialize);
@@ -277,18 +241,17 @@ bool FChunkCacheWorker::CheckSignature(const FChunkRequest& ChunkInfo)
 		ChunkHash = ComputePakChunkHash(ChunkInfo.Buffer->Data, ChunkInfo.Size);
 	}
 
-	bool bChunkHashesMatch = (ChunkHash == ChunkHashes[ChunkInfo.Index]);
+	bool bChunkHashesMatch = (ChunkHash == Signatures.ChunkHashes[ChunkInfo.Index]);
 	if (!bChunkHashesMatch)
 	{
-		UE_LOG(LogPakFile, Warning, TEXT("Pak chunk signing mismatch on chunk [%i/%i]! Expected 0x%8X, Received 0x%8X"), ChunkInfo.Index, ChunkHashes.Num(), ChunkHashes[ChunkInfo.Index], ChunkHash);
+		UE_LOG(LogPakFile, Warning, TEXT("Pak chunk signing mismatch on chunk [%i/%i]! Expected 0x%8X, Received 0x%8X"), ChunkInfo.Index, Signatures.ChunkHashes.Num(), *LexToString(Signatures.ChunkHashes[ChunkInfo.Index]), *LexToString(ChunkHash));
 
-		TPakChunkHash CurrentSignatureHash = ComputePakChunkHash(&ChunkHashes[0], ChunkHashes.Num() * sizeof(TPakChunkHash));
-		if (OriginalSignatureFileHash != CurrentSignatureHash)
+		if (Signatures.DecryptedHash != Signatures.ComputeCurrentMasterHash())
 		{
 			UE_LOG(LogPakFile, Warning, TEXT("Master signature table has changed since initialization!"));
 		}
 
-		const FPakChunkSignatureCheckFailedData Data(Reader->GetArchiveName(), ChunkHashes[ChunkInfo.Index], ChunkHash, ChunkInfo.Index);
+		const FPakChunkSignatureCheckFailedData Data(Reader->GetArchiveName(), Signatures.ChunkHashes[ChunkInfo.Index], ChunkHash, ChunkInfo.Index);
 		FPakPlatformFile::GetPakChunkSignatureCheckFailedHandler().Broadcast(Data);
 	}
 	
diff --git a/Engine/Source/Runtime/PakFile/Private/SignedArchiveReader.h b/Engine/Source/Runtime/PakFile/Private/SignedArchiveReader.h
--- a/Engine/Source/Runtime/PakFile/Private/SignedArchiveReader.h
+++ b/Engine/Source/Runtime/PakFile/Private/SignedArchiveReader.h
@@ -105,7 +105,7 @@ class FChunkCacheWorker : public FRunnable
 	};
 
 	/** Reference hashes */
-	TArray<TPakChunkHash> ChunkHashes;
+	FPakSignatureFile Signatures;
 	/** Hash of the sig file data. Used to check that nothing was corrupted when a signature check fails */
 	TPakChunkHash OriginalSignatureFileHash;
 	/** Thread to run the worker FRunnable on */
@@ -130,8 +130,6 @@ class FChunkCacheWorker : public FRunnable
 	FThreadSafeCounter StopTaskCounter;
 	/** Available chunk requests */
 	TLockFreePointerListUnordered<FChunkRequest, PLATFORM_CACHE_LINE_SIZE> FreeChunkRequests;
-	/** Public decryption key */
-	FEncryptionKey DecryptionKey;
 
 	/** 
 	 * Process requested chunks 
diff --git a/Engine/Source/Runtime/PakFile/Public/IPlatformFilePak.h b/Engine/Source/Runtime/PakFile/Public/IPlatformFilePak.h
--- a/Engine/Source/Runtime/PakFile/Public/IPlatformFilePak.h
+++ b/Engine/Source/Runtime/PakFile/Public/IPlatformFilePak.h
@@ -11,6 +11,8 @@
 #include "Templates/UniquePtr.h"
 #include "Math/BigInt.h"
 #include "Misc/AES.h"
+#include "Misc/RSA.h"
+#include "Misc/SecureHash.h"
 #include "GenericPlatform/GenericPlatformChunkInstall.h"
 
 class FChunkCacheWorker;
@@ -21,15 +23,19 @@ DECLARE_FLOAT_ACCUMULATOR_STAT_EXTERN(TEXT("Total pak file read time"), STAT_Pak
 
 DECLARE_DWORD_ACCUMULATOR_STAT_EXTERN(TEXT("Num open pak file handles"), STAT_PakFile_NumOpenHandles, STATGROUP_PakFile, PAKFILE_API);
 
-#define PAKHASH_USE_CRC	1
 #define PAK_TRACKER 0
 
+// Define the type of a chunk hash. Currently selectable between SHA1 and CRC32.
+#define PAKHASH_USE_CRC	0
 #if PAKHASH_USE_CRC
 typedef uint32 TPakChunkHash;
 #else
 typedef FSHAHash TPakChunkHash;
 #endif
 
+// Define the key size that we use for pak signing purposes. Must match with the key size created in the cryptokeys editor plugin
+typedef FRSAKey<2048> TPakRSAKey;
+
 PAKFILE_API TPakChunkHash ComputePakChunkHash(const void* InData, int64 InDataSizeInBytes);
 
 struct FPakChunkSignatureCheckFailedData
@@ -1513,7 +1519,7 @@ public:
 	/**
 	* Helper function for accessing pak signing keys
 	*/
-	static void GetPakSigningKeys(FEncryptionKey& OutKey);
+	static void GetPakSigningKeys(TPakRSAKey& OutKey);
 
 	/**
 	 * Constructor.
@@ -2326,3 +2332,105 @@ public:
 #endif
 };
 
+/**
+ * Structure which describes the content of the pak .sig files
+ */
+struct FPakSignatureFile
+{
+	// Magic number that tells us we're dealing with the new format sig files
+	static const uint32 Magic = 0x73832DAA;
+
+	enum class EVersion
+	{
+		Legacy,
+		First,
+
+		Last,
+		Latest = Last - 1
+	};
+
+	// Sig file version. Set to Legacy if the sig file is of an old version
+	EVersion Version = EVersion::Latest;
+
+	// RSA encrypted hash
+	TArray<uint8> EncryptedHash;
+
+	// SHA1 hash of the chunk CRC data. Only valid after calling DecryptSignatureAndValidate
+	FSHAHash DecryptedHash;
+
+	// CRCs of each contiguous 64kb block of the pak file
+	TArray<TPakChunkHash> ChunkHashes;
+	
+	/**
+	 * Initialize and hash the CRC list then use the provided private key to encrypt the hash
+	 */
+	void Initialize(const TArray<TPakChunkHash>& InChunkHashes, const TPakRSAKey& InPrivateKey)
+	{
+		ChunkHashes = InChunkHashes;
+		DecryptedHash = ComputeCurrentMasterHash();
+		TPakRSAKey::Encrypt(DecryptedHash.Hash, ARRAY_COUNT(FSHAHash::Hash), EncryptedHash, InPrivateKey);
+	}
+
+	/**
+	 * Serialize/deserialize this object to/from an FArchive
+	 */
+	void Serialize(FArchive& Ar)
+	{
+		int64 StartPos = Ar.Tell();
+		uint32 FileMagic = Magic;
+		Ar << FileMagic;
+
+		if (Ar.IsLoading() && FileMagic != Magic)
+		{
+			// Old format with no versioning! Go back to where we were and mark our version as legacy
+			Ar.Seek(StartPos);
+			Version = EVersion::Legacy;
+			TEncryptionInt LegacyEncryptedCRC;
+			Ar << LegacyEncryptedCRC;
+			Ar << ChunkHashes;
+			// Note that we don't do any actual signature checking here because this is old data that won't have been
+			// encrypted with the new larger keys.
+			return;
+		}
+
+		Ar << Version;
+		Ar << EncryptedHash;
+		Ar << ChunkHashes;
+	}
+
+	/**
+	 * Decrypt the chunk CRCs hash and validate that it matches the current one
+	 */
+	bool DecryptSignatureAndValidate(const FString& InFilename)
+	{
+		TPakRSAKey PublicKey;
+		FPakPlatformFile::GetPakSigningKeys(PublicKey);
+		return DecryptSignatureAndValidate(PublicKey, InFilename);
+	}
+
+	/**
+	 * Decrypt the chunk CRCs hash and validate that it matches the current one
+	 */
+	bool DecryptSignatureAndValidate(TPakRSAKey& InPublicKey, const FString& InFilename)
+	{
+		TPakRSAKey::Decrypt(EncryptedHash, DecryptedHash.Hash, ARRAY_COUNT(FSHAHash::Hash), InPublicKey);
+
+		if (DecryptedHash != ComputeCurrentMasterHash())
+		{
+			FPakPlatformFile::GetPakMasterSignatureTableCheckFailureHandler().Broadcast(InFilename);
+			return false;
+		}
+
+		return true;
+	}
+
+	/**
+	 * Helper function for computing the SHA1 hash of the current chunk CRC array
+	 */
+	FSHAHash ComputeCurrentMasterHash() const
+	{
+		FSHAHash CurrentHash;
+		FSHA1::HashBuffer(ChunkHashes.GetData(), ChunkHashes.Num() * sizeof(TPakChunkHash), CurrentHash.Hash);
+		return CurrentHash;
+	}
+};
\ No newline at end of file
