disable hang detector in PGO profile builds

diff --git a/Engine/Source/Runtime/Core/Public/Misc/Build.h b/Engine/Source/Runtime/Core/Public/Misc/Build.h
--- a/Engine/Source/Runtime/Core/Public/Misc/Build.h
+++ b/Engine/Source/Runtime/Core/Public/Misc/Build.h
@@ -311,7 +311,7 @@
 #ifndef ALLOW_HANG_DETECTION
 	#define ALLOW_HANG_DETECTION 1
 #endif
-#define USE_HANG_DETECTION (ALLOW_HANG_DETECTION && !WITH_EDITORONLY_DATA && !IS_PROGRAM && !UE_BUILD_DEBUG)
+#define USE_HANG_DETECTION (ALLOW_HANG_DETECTION && !WITH_EDITORONLY_DATA && !IS_PROGRAM && !UE_BUILD_DEBUG && !ENABLE_PGO_PROFILE)
 
 // Controls the creation of a thread for detecting hitches (FGameThreadHitchHeartBeat). This is subject to other criteria, USE_HITCH_DETECTION
 #ifndef ALLOW_HITCH_DETECTION
