Add support for passing csv metadata via the commandline, e.g -csvmetadata="testname=DisableParallelRendering,myOtherMetadata=myOtherValue"

diff --git a/Engine/Source/Runtime/Core/Private/ProfilingDebugging/CsvProfiler.cpp b/Engine/Source/Runtime/Core/Private/ProfilingDebugging/CsvProfiler.cpp
--- a/Engine/Source/Runtime/Core/Private/ProfilingDebugging/CsvProfiler.cpp
+++ b/Engine/Source/Runtime/Core/Private/ProfilingDebugging/CsvProfiler.cpp
@@ -3055,6 +3055,23 @@ void FCsvProfiler::Init()
 		}
 	}
 
+	FString CsvMetadataStr;
+	if (FParse::Value(FCommandLine::Get(), TEXT("csvMetadata="), CsvMetadataStr))
+	{ 
+		TArray<FString> CsvMetadataList;
+		CsvMetadataStr.ParseIntoArray(CsvMetadataList, TEXT(","), true);
+		for (int i = 0; i < CsvMetadataList.Num(); i++)
+		{
+			const FString& Metadata = CsvMetadataList[i];
+			FString Key;
+			FString Value;
+			if (Metadata.Split(TEXT("="), &Key, &Value))
+			{
+				SetMetadata(*Key, *Value);
+			}
+		}
+	}
+
 	if (FParse::Param(FCommandLine::Get(), TEXT("csvNoProcessingThread")))
 	{
 		GCsvUseProcessingThread = false;
