[EmbeddedCommunication] Add some extra logging around the GameThread sleep

diff --git a/Engine/Source/Runtime/Core/Private/Misc/EmbeddedCommunication.cpp b/Engine/Source/Runtime/Core/Private/Misc/EmbeddedCommunication.cpp
--- a/Engine/Source/Runtime/Core/Private/Misc/EmbeddedCommunication.cpp
+++ b/Engine/Source/Runtime/Core/Private/Misc/EmbeddedCommunication.cpp
@@ -327,7 +327,9 @@ bool FEmbeddedCommunication::TickGameThread(float DeltaTime)
 		&& !GTickAnotherFrame) // Don't sleep if we have been asked to tick another frame
  	{
  		// wake up every 5 seconds even if nothing to do
- 		GSleepEvent->Wait(5000);
+		UE_LOG(LogInit, VeryVerbose, TEXT("FEmbeddedCommunication Sleeping GameThread..."));
+ 		bool bWasTriggered = GSleepEvent->Wait(5000);
+		UE_LOG(LogInit, VeryVerbose, TEXT("FEmbeddedCommunication Woke up. Reason=[%s]"), bWasTriggered ? TEXT("Triggered") : TEXT("TimedOut"));
  	}
 	GTickAnotherFrame = false;
 #endif
