Add support for FString and bool to FAutoConsoleVariableRef

diff --git a/Engine/Source/Runtime/Core/Public/HAL/IConsoleManager.h b/Engine/Source/Runtime/Core/Public/HAL/IConsoleManager.h
--- a/Engine/Source/Runtime/Core/Public/HAL/IConsoleManager.h
+++ b/Engine/Source/Runtime/Core/Public/HAL/IConsoleManager.h
@@ -853,7 +853,7 @@ public:
 };
 
 /**
- * Autoregistering float, int REF variable class...this changes that value when the console variable is changed. 
+ * Autoregistering float, int, bool, FString REF variable class...this changes that value when the console variable is changed. 
  */
 class CORE_API FAutoConsoleVariableRef : private FAutoConsoleObject
 {
@@ -878,6 +878,26 @@ public:
 		: FAutoConsoleObject(IConsoleManager::Get().RegisterConsoleVariableRef(Name, RefValue, Help, Flags))
 	{
 	}
+	/**
+	 * Create a reference to a bool console variable
+	 * @param Name must not be 0
+	 * @param Help must not be 0
+	 * @param Flags bitmask combined from EConsoleVariableFlags
+	 */
+	FAutoConsoleVariableRef(const TCHAR* Name, bool& RefValue, const TCHAR* Help, uint32 Flags = ECVF_Default)
+		: FAutoConsoleObject(IConsoleManager::Get().RegisterConsoleVariableRef(Name, RefValue, Help, Flags))
+	{
+	}
+	/**
+	 * Create a reference to a FString console variable
+	 * @param Name must not be 0
+	 * @param Help must not be 0
+	 * @param Flags bitmask combined from EConsoleVariableFlags
+	 */
+	FAutoConsoleVariableRef(const TCHAR* Name, FString& RefValue, const TCHAR* Help, uint32 Flags = ECVF_Default)
+		: FAutoConsoleObject(IConsoleManager::Get().RegisterConsoleVariableRef(Name, RefValue, Help, Flags))
+	{
+	}
 
 	/**
 	 * Create a reference to a int console variable
@@ -905,6 +925,32 @@ public:
 		AsVariable()->SetOnChangedCallback(Callback);
 	}
 
+	/**
+	 * Create a reference to a bool console variable
+	 * @param Name must not be 0
+	 * @param Help must not be 0
+	 * @param Callback Delegate called when the variable changes. @see IConsoleVariable::SetOnChangedCallback
+	 * @param Flags bitmask combined from EConsoleVariableFlags
+	 */
+	FAutoConsoleVariableRef(const TCHAR* Name, bool& RefValue, const TCHAR* Help, const FConsoleVariableDelegate& Callback, uint32 Flags = ECVF_Default)
+		: FAutoConsoleObject(IConsoleManager::Get().RegisterConsoleVariableRef(Name, RefValue, Help, Flags))
+	{
+		AsVariable()->SetOnChangedCallback(Callback);
+	}
+
+	/**
+	 * Create a reference to a FString console variable
+	 * @param Name must not be 0
+	 * @param Help must not be 0
+	 * @param Callback Delegate called when the variable changes. @see IConsoleVariable::SetOnChangedCallback
+	 * @param Flags bitmask combined from EConsoleVariableFlags
+	 */
+	FAutoConsoleVariableRef(const TCHAR* Name, FString& RefValue, const TCHAR* Help, const FConsoleVariableDelegate& Callback, uint32 Flags = ECVF_Default)
+		: FAutoConsoleObject(IConsoleManager::Get().RegisterConsoleVariableRef(Name, RefValue, Help, Flags))
+	{
+		AsVariable()->SetOnChangedCallback(Callback);
+	}
+
 	virtual ~FAutoConsoleVariableRef()
 	{
 	}
