Install Bundle Manger
- Fix user settings for uninstall not being persisted in shipping builds.
- Make failure to write manifest archive an error condition.
- Remove legacy install migration code and just clean up and legacy install data if it exists.  Those biulds are now too old to attempt patching.
- Make FFileHelper utilities return false on write errors.

diff --git a/Engine/Source/Runtime/Core/Private/Misc/FileHelper.cpp b/Engine/Source/Runtime/Core/Private/Misc/FileHelper.cpp
--- a/Engine/Source/Runtime/Core/Private/Misc/FileHelper.cpp
+++ b/Engine/Source/Runtime/Core/Private/Misc/FileHelper.cpp
@@ -188,14 +188,13 @@ bool FFileHelper::LoadFileToStringArray( TArray<FString>& Result, const TCHAR* F
  */
 bool FFileHelper::SaveArrayToFile(TArrayView<const uint8> Array, const TCHAR* Filename, IFileManager* FileManager /*= &IFileManager::Get()*/, uint32 WriteFlags)
 {
-	FArchive* Ar = FileManager->CreateFileWriter( Filename, WriteFlags );
+	TUniquePtr<FArchive> Ar = TUniquePtr<FArchive>( FileManager->CreateFileWriter( Filename, WriteFlags ) );
 	if( !Ar )
 	{
 		return false;
 	}
-	Ar->Serialize(const_cast<uint8*>(Array.GetData()), Array.Num());
-	delete Ar;
-	return true;
+	Ar->Serialize(const_cast<uint8*>(Array.GetData()), Array.Num());	
+	return !Ar->IsError() && !Ar->IsCriticalError();
 }
 
 /**
@@ -205,7 +204,7 @@ bool FFileHelper::SaveArrayToFile(TArrayView<const uint8> Array, const TCHAR* Fi
 bool FFileHelper::SaveStringToFile( const FString& String, const TCHAR* Filename,  EEncodingOptions EncodingOptions, IFileManager* FileManager /*= &IFileManager::Get()*/, uint32 WriteFlags )
 {
 	// max size of the string is a UCS2CHAR for each character and some UNICODE magic 
-	auto Ar = TUniquePtr<FArchive>( FileManager->CreateFileWriter( Filename, WriteFlags ) );
+	TUniquePtr<FArchive> Ar = TUniquePtr<FArchive>( FileManager->CreateFileWriter( Filename, WriteFlags ) );
 	if( !Ar )
 		return false;
 
@@ -242,7 +241,7 @@ bool FFileHelper::SaveStringToFile( const FString& String, const TCHAR* Filename
 		Ar->Serialize( (ANSICHAR*)Src.Get(), Src.Length() * sizeof(ANSICHAR) );
 	}
 
-	return true;
+	return !Ar->IsError() && !Ar->IsCriticalError();
 }
 
 bool FFileHelper::SaveStringArrayToFile( const TArray<FString>& Lines, const TCHAR* Filename, EEncodingOptions EncodingOptions, IFileManager* FileManager, uint32 WriteFlags )
diff --git a/Engine/Source/Runtime/Core/Public/GenericPlatform/GenericPlatformInstallBundleManager.h b/Engine/Source/Runtime/Core/Public/GenericPlatform/GenericPlatformInstallBundleManager.h
--- a/Engine/Source/Runtime/Core/Public/GenericPlatform/GenericPlatformInstallBundleManager.h
+++ b/Engine/Source/Runtime/Core/Public/GenericPlatform/GenericPlatformInstallBundleManager.h
@@ -16,6 +16,7 @@ enum class EInstallBundleModuleInitResult : int
 	BuildMetaDataParsingError,
 	DistributionRootParseError,
 	DistributionRootDownloadError,
+	ManifestArchiveError,
 	ManifestCreationError,
 	ManifestDownloadError,
 	BackgroundDownloadsIniDownloadError,
@@ -33,6 +34,7 @@ inline const TCHAR* LexToString(EInstallBundleModuleInitResult Result)
 		TEXT("BuildMetaDataParsingError"),
 		TEXT("DistributionRootParseError"),
 		TEXT("DistributionRootDownloadError"),
+		TEXT("ManifestArchiveError"),
 		TEXT("ManifestCreationError"),
 		TEXT("ManifestDownloadError"),
 		TEXT("BackgroundDownloadsIniDownloadError"),
@@ -49,6 +51,7 @@ enum class EInstallBundleResult : int
 	FailedPrereqRequiresLatestClient,
 	InstallError,
 	InstallerOutOfDiskSpaceError,
+	ManifestArchiveError,
 	UserCancelledError,
 	InitializationError,
 	Count,
@@ -63,6 +66,7 @@ inline const TCHAR* LexToString(EInstallBundleResult Result)
 		TEXT("FailedPrereqRequiresLatestClient"),
 		TEXT("InstallError"),
 		TEXT("InstallerOutOfDiskSpaceError"),
+		TEXT("ManifestArchiveError"),
 		TEXT("UserCancelledError"),
 		TEXT("InitializationError"),
 	};
@@ -307,9 +311,6 @@ public:
 	virtual bool IsNullInterface() const = 0;
 
 	virtual void SetErrorSimulationCommands(const FString& CommandLine) {}
-
-	// return true if we actually cleaned up a migration directory
-	virtual bool CleanupMigrationDirectory() { return false;  }
 };
 
 class IPlatformInstallBundleManagerModule : public IModuleInterface
