Add a function to cancel an uninstall on reboot.

diff --git a/Engine/Source/Runtime/Core/Public/GenericPlatform/GenericPlatformInstallBundleManager.h b/Engine/Source/Runtime/Core/Public/GenericPlatform/GenericPlatformInstallBundleManager.h
--- a/Engine/Source/Runtime/Core/Public/GenericPlatform/GenericPlatformInstallBundleManager.h
+++ b/Engine/Source/Runtime/Core/Public/GenericPlatform/GenericPlatformInstallBundleManager.h
@@ -294,6 +294,8 @@ public:
 
 	virtual void RequestRemoveBundleOnNextInit(FName BundleName) = 0;
 
+	virtual void CancelRequestRemoveBundleOnNextInit(FName BundleName) = 0;
+
 	virtual void CancelBundle(FName BundleName, EInstallBundleCancelFlags Flags) = 0;
 
 	virtual void CancelAllBundles(EInstallBundleCancelFlags Flags) = 0;
diff --git a/Engine/Source/Runtime/NullInstallBundleManager/Source/NullInstallBundleManager.cpp b/Engine/Source/Runtime/NullInstallBundleManager/Source/NullInstallBundleManager.cpp
--- a/Engine/Source/Runtime/NullInstallBundleManager/Source/NullInstallBundleManager.cpp
+++ b/Engine/Source/Runtime/NullInstallBundleManager/Source/NullInstallBundleManager.cpp
@@ -64,6 +64,11 @@ class FNullInstallBundleManager : public IPlatformInstallBundleManager
 
 	}
 
+	virtual void CancelRequestRemoveBundleOnNextInit(FName BundleName) override
+	{
+
+	}
+
 	virtual void CancelBundle(FName BundleName, EInstallBundleCancelFlags Flags) override
 	{
 
