fix bug in ensure macros
- it was incorrectly returning true for false expressions causing outer if-statements to incorrectly execute their body

(via shelf by PJKack)

diff --git a/Engine/Source/Runtime/Core/Public/Misc/AssertionMacros.h b/Engine/Source/Runtime/Core/Public/Misc/AssertionMacros.h
--- a/Engine/Source/Runtime/Core/Public/Misc/AssertionMacros.h
+++ b/Engine/Source/Runtime/Core/Public/Misc/AssertionMacros.h
@@ -298,7 +298,7 @@ public:
 #if DO_CHECK && !USING_CODE_ANALYSIS // The Visual Studio 2013 analyzer doesn't understand these complex conditionals
 
 	#define UE_ENSURE_IMPL(Capture, Always, InExpression, ...) \
-		(LIKELY(!!(InExpression)) || DispatchCheckVerify<bool>([Capture] () FORCENOINLINE UE_DEBUG_SECTION \
+		(LIKELY(!!(InExpression)) || (DispatchCheckVerify<bool>([Capture] () FORCENOINLINE UE_DEBUG_SECTION \
 		{ \
 			static bool bExecuted = false; \
 			if ((!bExecuted || Always) && FPlatformMisc::IsEnsureAllowed()) \
@@ -308,11 +308,12 @@ public:
 				if (!FPlatformMisc::IsDebuggerPresent()) \
 				{ \
 					FPlatformMisc::PromptForRemoteDebugging(true); \
-					return true; \
+					return false; \
 				} \
+				return true; \
 			} \
-			return !Always; \
-		}) || ([] () { PLATFORM_BREAK(); } (), true))
+			return false; \
+		}) && ([] () { PLATFORM_BREAK(); } (), false)))
 
 	#define ensure(           InExpression                ) UE_ENSURE_IMPL( , false, InExpression, TEXT(""))
 	#define ensureMsgf(       InExpression, InFormat, ... ) UE_ENSURE_IMPL(&, false, InExpression, InFormat, ##__VA_ARGS__)
