Fixed MallocVerify compile error.

diff --git a/Engine/Source/Runtime/Core/Private/HAL/MallocVerify.h b/Engine/Source/Runtime/Core/Private/HAL/MallocVerify.h
--- a/Engine/Source/Runtime/Core/Private/HAL/MallocVerify.h
+++ b/Engine/Source/Runtime/Core/Private/HAL/MallocVerify.h
@@ -117,9 +117,9 @@ public:
 	{
 		return UsedMalloc->QuantizeSize(Count, Alignment);
 	}
-	virtual void Trim() override
+	virtual void Trim(bool bTrimThreadCaches) override
 	{
-		return UsedMalloc->Trim();
+		return UsedMalloc->Trim(bTrimThreadCaches);
 	}
 	virtual void SetupTLSCachesOnCurrentThread() override
 	{
