Improved unit test error message when wrong locale is set and skip further tests to avoid error spam

diff --git a/Engine/Source/Runtime/Core/Private/Tests/Misc/CharTest.cpp b/Engine/Source/Runtime/Core/Private/Tests/Misc/CharTest.cpp
--- a/Engine/Source/Runtime/Core/Private/Tests/Misc/CharTest.cpp
+++ b/Engine/Source/Runtime/Core/Private/Tests/Misc/CharTest.cpp
@@ -32,10 +32,16 @@ void RunCharTests(FAutomationTestBase& Test, uint32 MaxChar)
 
 bool TCharTest::RunTest(const FString& Parameters)
 {
-	TestTrue(TEXT("C locale not used"), strcmp("C", setlocale(LC_CTYPE, nullptr)) == 0);
-
-	RunCharTests<ANSICHAR>(*this, 128);
-	RunCharTests<WIDECHAR>(*this, 65536);
+	const char* CurrentLocale = setlocale(LC_CTYPE, nullptr);
+	if (strcmp("C", CurrentLocale))
+	{
+		AddError(FString::Printf(TEXT("Locale is \"%s\" but should be \"C\". Did something call setlocale()?"), CurrentLocale));
+	}
+	else
+	{
+		RunCharTests<ANSICHAR>(*this, 128);
+		RunCharTests<WIDECHAR>(*this, 65536);
+	}
 
 	return true;
 }
