Removal of some VS2015-specific code, which is no longer supported.

diff --git a/Engine/Source/Runtime/Core/Public/Templates/Tuple.h b/Engine/Source/Runtime/Core/Public/Templates/Tuple.h
--- a/Engine/Source/Runtime/Core/Public/Templates/Tuple.h
+++ b/Engine/Source/Runtime/Core/Public/Templates/Tuple.h
@@ -10,16 +10,6 @@
 #include "Templates/Invoke.h"
 #include "Serialization/StructuredArchive.h"
 
-// VS2015 Update 2 (and seemingly earlier) erroneously complains about multiple versions
-// of special member functions, so we disable the use of defaulting in that case.
-// 
-// http://stackoverflow.com/questions/36657243/c-multiple-versions-of-a-defaulted-special-member-functions-error-in-msvc-2
-#if defined(_MSC_VER) && _MSC_FULL_VER <= 190023918
-	#define TUPLES_USE_DEFAULTED_FUNCTIONS 0
-#else
-	#define TUPLES_USE_DEFAULTED_FUNCTIONS 1
-#endif
-
 class FArchive;
 
 template <typename... Types>
@@ -73,36 +63,10 @@ namespace UE4Tuple_Private
 		{
 		}
 
-		#if TUPLES_USE_DEFAULTED_FUNCTIONS
-
-			TTupleElement(TTupleElement&&) = default;
-			TTupleElement(const TTupleElement&) = default;
-			TTupleElement& operator=(TTupleElement&&) = default;
-			TTupleElement& operator=(const TTupleElement&) = default;
-
-		#else
-
-			TTupleElement(TTupleElement&& Other)
-				: Value(MoveTemp(Other.Value))
-			{
-			}
-
-			TTupleElement(const TTupleElement& Other)
-				: Value(Other.Value)
-			{
-			}
-
-			void operator=(TTupleElement&& Other)
-			{
-				Value = MoveTemp(Other.Value);
-			}
-
-			void operator=(const TTupleElement& Other)
-			{
-				Value = Other.Value;
-			}
-
-		#endif
+		TTupleElement(TTupleElement&&) = default;
+		TTupleElement(const TTupleElement&) = default;
+		TTupleElement& operator=(TTupleElement&&) = default;
+		TTupleElement& operator=(const TTupleElement&) = default;
 
 		T Value;
 	};
@@ -211,44 +175,11 @@ namespace UE4Tuple_Private
 		{
 		}
 
-		#if TUPLES_USE_DEFAULTED_FUNCTIONS
-
-			TTupleStorage() = default;
-			TTupleStorage(TTupleStorage&&) = default;
-			TTupleStorage(const TTupleStorage&) = default;
-			TTupleStorage& operator=(TTupleStorage&&) = default;
-			TTupleStorage& operator=(const TTupleStorage&) = default;
-
-		#else
-
-			TTupleStorage()
-				: TTupleElement<Types, Indices>()...
-			{
-			}
-
-			TTupleStorage(TTupleStorage&& Other)
-				: TTupleElement<Types, Indices>(MoveTemp(*(TTupleElement<Types, Indices>*)&Other))...
-			{
-			}
-
-			TTupleStorage(const TTupleStorage& Other)
-				: TTupleElement<Types, Indices>(*(const TTupleElement<Types, Indices>*)&Other)...
-			{
-			}
-
-			void operator=(TTupleStorage&& Other)
-			{
-				int Temp[] = { 0, (*(TTupleElement<Types, Indices>*)this = MoveTemp(*(TTupleElement<Types, Indices>*)&Other), 0)... };
-				(void)Temp;
-			}
-
-			void operator=(const TTupleStorage& Other)
-			{
-				int Temp[] = { 0, (*(TTupleElement<Types, Indices>*)this = *(const TTupleElement<Types, Indices>*)&Other, 0)... };
-				(void)Temp;
-			}
-
-		#endif
+		TTupleStorage() = default;
+		TTupleStorage(TTupleStorage&&) = default;
+		TTupleStorage(const TTupleStorage&) = default;
+		TTupleStorage& operator=(TTupleStorage&&) = default;
+		TTupleStorage& operator=(const TTupleStorage&) = default;
 
 		template <uint32 Index> FORCEINLINE const typename TTupleElementHelper<Index, Types...>::Type& Get() const { return TTupleElementHelper<Index, Types...>::Get(*this); }
 		template <uint32 Index> FORCEINLINE       typename TTupleElementHelper<Index, Types...>::Type& Get()       { return TTupleElementHelper<Index, Types...>::Get(*this); }
@@ -297,40 +228,10 @@ namespace UE4Tuple_Private
 		{
 		}
 
-		#if TUPLES_USE_DEFAULTED_FUNCTIONS
-
-			TTupleStorage(TTupleStorage&&) = default;
-			TTupleStorage(const TTupleStorage&) = default;
-			TTupleStorage& operator=(TTupleStorage&&) = default;
-			TTupleStorage& operator=(const TTupleStorage&) = default;
-
-		#else
-
-			TTupleStorage(TTupleStorage&& Other)
-				: Key  (MoveTemp(Other.Key))
-				, Value(MoveTemp(Other.Value))
-			{
-			}
-
-			TTupleStorage(const TTupleStorage& Other)
-				: Key  (Other.Key)
-				, Value(Other.Value)
-			{
-			}
-
-			void operator=(TTupleStorage&& Other)
-			{
-				Key   = MoveTemp(Other.Key);
-				Value = MoveTemp(Other.Value);
-			}
-
-			void operator=(const TTupleStorage& Other)
-			{
-				Key   = Other.Key;
-				Value = Other.Value;
-			}
-
-		#endif
+		TTupleStorage(TTupleStorage&&) = default;
+		TTupleStorage(const TTupleStorage&) = default;
+		TTupleStorage& operator=(TTupleStorage&&) = default;
+		TTupleStorage& operator=(const TTupleStorage&) = default;
 
 		template <uint32 Index> FORCEINLINE const typename TGetHelper<Index, void>::ResultType& Get() const { return TGetHelper<Index, void>::Get(*this); }
 		template <uint32 Index> FORCEINLINE       typename TGetHelper<Index, void>::ResultType& Get()       { return TGetHelper<Index, void>::Get(*this); }
@@ -368,42 +269,11 @@ namespace UE4Tuple_Private
 		{
 		}
 
-		#if TUPLES_USE_DEFAULTED_FUNCTIONS
-
-			TTupleImpl() = default;
-			TTupleImpl(TTupleImpl&& Other) = default;
-			TTupleImpl(const TTupleImpl& Other) = default;
-			TTupleImpl& operator=(TTupleImpl&& Other) = default;
-			TTupleImpl& operator=(const TTupleImpl& Other) = default;
-
-		#else
-
-			TTupleImpl()
-				: Super()
-			{
-			}
-
-			TTupleImpl(TTupleImpl&& Other)
-				: Super(MoveTemp(*(Super*)&Other))
-			{
-			}
-
-			TTupleImpl(const TTupleImpl& Other)
-				: Super(*(const Super*)&Other)
-			{
-			}
-
-			void operator=(TTupleImpl&& Other)
-			{
-				*(Super*)this = MoveTemp(*(Super*)&Other);
-			}
-
-			void operator=(const TTupleImpl& Other)
-			{
-				*(Super*)this = *(const Super*)&Other;
-			}
-
-		#endif
+		TTupleImpl() = default;
+		TTupleImpl(TTupleImpl&& Other) = default;
+		TTupleImpl(const TTupleImpl& Other) = default;
+		TTupleImpl& operator=(TTupleImpl&& Other) = default;
+		TTupleImpl& operator=(const TTupleImpl& Other) = default;
 
 		template <typename FuncType, typename... ArgTypes>
 		#if PLATFORM_COMPILER_HAS_DECLTYPE_AUTO
@@ -473,79 +343,6 @@ namespace UE4Tuple_Private
 		}
 	};
 
-	#ifdef _MSC_VER
-
-		// Not strictly necessary, but some VC versions give a 'syntax error: <fake-expression>' error
-		// for empty tuples.
-		template <>
-		struct TTupleImpl<TIntegerSequence<uint32>>
-		{
-			explicit TTupleImpl()
-			{
-			}
-
-			// Doesn't matter what these return, or even have a function body, but they need to be declared
-			template <uint32 Index> FORCEINLINE const int32& Get() const;
-			template <uint32 Index> FORCEINLINE       int32& Get();
-
-			template <typename FuncType, typename... ArgTypes>
-			#if PLATFORM_COMPILER_HAS_DECLTYPE_AUTO
-				decltype(auto) ApplyAfter(FuncType&& Func, ArgTypes&&... Args) const
-			#else
-				auto ApplyAfter(FuncType&& Func, ArgTypes&&... Args) const -> decltype(Func(Forward<ArgTypes>(Args)...))
-			#endif
-			{
-				return Func(Forward<ArgTypes>(Args)...);
-			}
-
-			template <typename FuncType, typename... ArgTypes>
-			#if PLATFORM_COMPILER_HAS_DECLTYPE_AUTO
-				decltype(auto) ApplyBefore(FuncType&& Func, ArgTypes&&... Args) const
-			#else
-				auto ApplyBefore(FuncType&& Func, ArgTypes&&... Args) const -> decltype(Func(Forward<ArgTypes>(Args)...))
-			#endif
-			{
-				return Func(Forward<ArgTypes>(Args)...);
-			}
-
-			FORCEINLINE friend FArchive& operator<<(FArchive& Ar, TTupleImpl& Tuple)
-			{
-				return Ar;
-			}
-
-			FORCEINLINE friend bool operator==(const TTupleImpl& Lhs, const TTupleImpl& Rhs)
-			{
-				return true;
-			}
-
-			FORCEINLINE friend bool operator!=(const TTupleImpl& Lhs, const TTupleImpl& Rhs)
-			{
-				return false;
-			}
-
-			FORCEINLINE friend bool operator<(const TTupleImpl& Lhs, const TTupleImpl& Rhs)
-			{
-				return false;
-			}
-
-			FORCEINLINE friend bool operator<=(const TTupleImpl& Lhs, const TTupleImpl& Rhs)
-			{
-				return true;
-			}
-
-			FORCEINLINE friend bool operator>(const TTupleImpl& Lhs, const TTupleImpl& Rhs)
-			{
-				return false;
-			}
-
-			FORCEINLINE friend bool operator>=(const TTupleImpl& Lhs, const TTupleImpl& Rhs)
-			{
-				return true;
-			}
-		};
-
-	#endif
-
 
 	template <typename... Types>
 	FORCEINLINE TTuple<typename TDecay<Types>::Type...> MakeTupleImpl(Types&&... Args)
@@ -629,43 +426,11 @@ public:
 		// This constructor is disabled for TTuple and zero parameters because VC is incorrectly instantiating it as a move/copy/default constructor.
 	}
 
-	#if TUPLES_USE_DEFAULTED_FUNCTIONS
-
-		TTuple() = default;
-		TTuple(TTuple&&) = default;
-		TTuple(const TTuple&) = default;
-		TTuple& operator=(TTuple&&) = default;
-		TTuple& operator=(const TTuple&) = default;
-
-	#else
-
-		TTuple()
-		{
-		}
-
-		TTuple(TTuple&& Other)
-			: Super(MoveTemp(*(Super*)&Other))
-		{
-		}
-
-		TTuple(const TTuple& Other)
-			: Super(*(const Super*)&Other)
-		{
-		}
-
-		TTuple& operator=(TTuple&& Other)
-		{
-			*(Super*)this = MoveTemp(*(Super*)&Other);
-			return *this;
-		}
-
-		TTuple& operator=(const TTuple& Other)
-		{
-			*(Super*)this = *(const Super*)&Other;
-			return *this;
-		}
-
-	#endif
+	TTuple() = default;
+	TTuple(TTuple&&) = default;
+	TTuple(const TTuple&) = default;
+	TTuple& operator=(TTuple&&) = default;
+	TTuple& operator=(const TTuple&) = default;
 };
 
 
diff --git a/Engine/Source/Runtime/Core/Public/Windows/WIndowsPlatform.h b/Engine/Source/Runtime/Core/Public/Windows/WIndowsPlatform.h
--- a/Engine/Source/Runtime/Core/Public/Windows/WIndowsPlatform.h
+++ b/Engine/Source/Runtime/Core/Public/Windows/WIndowsPlatform.h
@@ -61,6 +61,7 @@ typedef FWindowsPlatformTypes FPlatformTypes;
 #define PLATFORM_SUPPORTS_EARLY_MOVIE_PLAYBACK				(!WITH_EDITOR) // movies will start before engine is initalized
 #define PLATFORM_RHITHREAD_DEFAULT_BYPASS					0
 #define PLATFORM_SUPPORTS_STACK_SYMBOLS						1
+#define PLATFORM_COMPILER_HAS_DECLTYPE_AUTO					1
 
 // Q: Why is there a __nop() before __debugbreak()?
 // A: VS' debug engine has a bug where it will silently swallow explicit
@@ -70,15 +71,6 @@ typedef FWindowsPlatformTypes FPlatformTypes;
 // unknowingly debugging an undefined process.
 #define PLATFORM_BREAK() (__nop(), __debugbreak())
 
-#if defined(__INTEL_COMPILER) || _MSC_VER > 1900
-	#define PLATFORM_COMPILER_HAS_DECLTYPE_AUTO 1
-#else
-	// Static analysis causes internal compiler errors with auto-deduced return types,
-	// but some older VC versions still have return type deduction failures inside the delegate code
-	// when they are enabled.  So we currently only enable them for static analysis builds.
-	#define PLATFORM_COMPILER_HAS_DECLTYPE_AUTO USING_CODE_ANALYSIS
-#endif
-
 // Intrinsics for 128-bit atomics on Windows platform requires Windows 8 or higher (WINVER>=0x0602)
 // http://msdn.microsoft.com/en-us/library/windows/desktop/hh972640.aspx
 #define PLATFORM_HAS_128BIT_ATOMICS							(!HACK_HEADER_GENERATOR && PLATFORM_64BITS && (WINVER >= 0x602))
