Multicast delegates' Remove() returns if it successfully removed something and RemoveAll() returns how many delegates were removed.
Multicast delegate compaction is only performed on Remove() if a delegate was found.

diff --git a/Engine/Source/Runtime/Core/Public/Delegates/DelegateSignatureImpl.inl b/Engine/Source/Runtime/Core/Public/Delegates/DelegateSignatureImpl.inl
--- a/Engine/Source/Runtime/Core/Public/Delegates/DelegateSignatureImpl.inl
+++ b/Engine/Source/Runtime/Core/Public/Delegates/DelegateSignatureImpl.inl
@@ -887,13 +887,16 @@ public:
 	 * Note that the order of the delegate instances may not be preserved!
 	 *
 	 * @param Handle The handle of the delegate instance to remove.
+	 * @return  true if the delegate was successfully removed.
 	 */
-	void Remove( FDelegateHandle Handle )
+	bool Remove( FDelegateHandle Handle )
 	{
+		bool bResult = false;
 		if (Handle.IsValid())
 		{
-			RemoveDelegateInstance(Handle);
+			bResult = RemoveDelegateInstance(Handle);
 		}
+		return bResult;
 	}
 
 protected:
@@ -996,8 +999,9 @@ protected:
 	 * It will be removed next time the invocation list is compacted within Broadcast().
 	 *
 	 * @param Handle The handle of the delegate instance to remove.
+	 * @return  true if the delegate was successfully removed.
 	 */
-	void RemoveDelegateInstance( FDelegateHandle Handle )
+	bool RemoveDelegateInstance( FDelegateHandle Handle )
 	{
 		const TInvocationList& LocalInvocationList = Super::GetInvocationList();
 
@@ -1011,12 +1015,12 @@ protected:
 			if ((DelegateInstanceInterface != nullptr) && DelegateInstanceInterface->GetHandle() == Handle)
 			{
 				DelegateBase.Unbind();
-
-				break; // each delegate binding has a unique handle, so once we find it, we can stop
+				Super::CompactInvocationList();
+				return true; // each delegate binding has a unique handle, so once we find it, we can stop
 			}
 		}
 
-		Super::CompactInvocationList();
+		return false;
 	}
 };
 
diff --git a/Engine/Source/Runtime/Core/Public/Delegates/MulticastDelegateBase.h b/Engine/Source/Runtime/Core/Public/Delegates/MulticastDelegateBase.h
--- a/Engine/Source/Runtime/Core/Public/Delegates/MulticastDelegateBase.h
+++ b/Engine/Source/Runtime/Core/Public/Delegates/MulticastDelegateBase.h
@@ -81,12 +81,13 @@ public:
 	 * Note that the order of the delegates may not be preserved!
 	 *
 	 * @param InUserObject The object to remove all delegates for.
+	 * @return  The number of delegates successfully removed.
 	 */
-	void RemoveAll( const void* InUserObject )
+	int32 RemoveAll( const void* InUserObject )
 	{
+		int32 Result = 0;
 		if (InvocationListLockCount > 0)
 		{
-			bool NeedsCompacted = false;
 			for (FDelegateBase& DelegateBaseRef : InvocationList)
 			{
 				IDelegateInstance* DelegateInstance = DelegateBaseRef.GetDelegateInstanceProtected();
@@ -94,12 +95,12 @@ public:
 				{
 					// Manually unbind the delegate here so the compaction will find and remove it.
 					DelegateBaseRef.Unbind();
-					NeedsCompacted = true;
+					++Result;
 				}
 			}
 
 			// can't compact at the moment, but set out threshold to zero so the next add will do it
-			if (NeedsCompacted)
+			if (Result > 0)
 			{
 				CompactionThreshold = 0;
 			}
@@ -117,6 +118,7 @@ public:
 					|| DelegateInstance->IsCompactable())
 				{
 					InvocationList.RemoveAtSwap(InvocationListIndex, 1, false);
+					++Result;
 				}
 				else
 				{
@@ -128,6 +130,8 @@ public:
 
 			InvocationList.Shrink();
 		}
+
+		return Result;
 	}
 
 protected:
