Cleaned up LLM behaviour with Realloc

* Made FMemory::Realloc not bother calling into LLM to track free if the original pointer was null
* Made FMemory::Realloc not bother calling into LLM to track alloc if the new size was zero (indicating a free)

diff --git a/Engine/Source/Runtime/Core/Public/HAL/FMemory.inl b/Engine/Source/Runtime/Core/Public/HAL/FMemory.inl
--- a/Engine/Source/Runtime/Core/Public/HAL/FMemory.inl
+++ b/Engine/Source/Runtime/Core/Public/HAL/FMemory.inl
@@ -33,8 +33,9 @@ FMEMORY_INLINE_FUNCTION_DECORATOR void* FMemory::Malloc(SIZE_T Count, uint32 Ali
 
 FMEMORY_INLINE_FUNCTION_DECORATOR void* FMemory::Realloc(void* Original, SIZE_T Count, uint32 Alignment)
 {
-	// optional tracking of every allocation
-	LLM_IF_ENABLED(FLowLevelMemTracker::Get().OnLowLevelFree(ELLMTracker::Default, Original, ELLMAllocType::FMalloc));
+	// optional tracking -- a realloc with an Original pointer of null is equivalent
+	// to malloc() so there's nothing to free
+	LLM_IF_ENABLED(if (Original != nullptr) FLowLevelMemTracker::Get().OnLowLevelFree(ELLMTracker::Default, Original, ELLMAllocType::FMalloc));
 
 	void* Ptr;
 	if (!FMEMORY_INLINE_GMalloc)
@@ -48,8 +49,10 @@ FMEMORY_INLINE_FUNCTION_DECORATOR void* FMemory::Realloc(void* Original, SIZE_T
 		Ptr = FMEMORY_INLINE_GMalloc->Realloc(Original, Count, Alignment);
 	}
 
-	// optional tracking of every allocation
-	LLM_IF_ENABLED(FLowLevelMemTracker::Get().OnLowLevelAlloc(ELLMTracker::Default, Ptr, Count, ELLMTag::Untagged, ELLMAllocType::FMalloc));
+	// optional tracking of every allocation - a realloc with a Count of zero is equivalent to a call 
+	// to free() and will return a null pointer which does not require tracking. If realloc returns null
+	// for some other reason (like failure to allocate) there's also no reason to track it
+	LLM_IF_ENABLED(if (Ptr != nullptr) FLowLevelMemTracker::Get().OnLowLevelAlloc(ELLMTracker::Default, Ptr, Count, ELLMTag::Untagged, ELLMAllocType::FMalloc));
 
 	return Ptr;
 }
