Changed some SecureHash related buffer sizes from 32 to 64 bits to keep the API consistent

diff --git a/Engine/Source/Runtime/Core/Private/Misc/SecureHash.cpp b/Engine/Source/Runtime/Core/Private/Misc/SecureHash.cpp
--- a/Engine/Source/Runtime/Core/Private/Misc/SecureHash.cpp
+++ b/Engine/Source/Runtime/Core/Private/Misc/SecureHash.cpp
@@ -822,12 +822,12 @@ void FSHA1::HMACBuffer(const void* Key, uint32 KeySize, const void* Data, uint64
  * Shared hashes.sha reading code (each platform gets a buffer to the data,
  * then passes it to this function for processing)
  */
-void FSHA1::InitializeFileHashesFromBuffer(uint8* Buffer, int32 BufferSize, bool bDuplicateKeyMemory)
+void FSHA1::InitializeFileHashesFromBuffer(uint8* Buffer, uint64 BufferSize, bool bDuplicateKeyMemory)
 {
 	// the start of the file is full file hashes
 	bool bIsDoingFullFileHashes = true;
 	// if it exists, parse it
-	int32 Offset = 0;
+	uint64 Offset = 0;
 	while (Offset < BufferSize)
 	{
 		// format is null terminated string followed by hash
diff --git a/Engine/Source/Runtime/Core/Public/Misc/SecureHash.h b/Engine/Source/Runtime/Core/Public/Misc/SecureHash.h
--- a/Engine/Source/Runtime/Core/Public/Misc/SecureHash.h
+++ b/Engine/Source/Runtime/Core/Public/Misc/SecureHash.h
@@ -75,7 +75,7 @@ public:
 	 *
 	 * @param String	hex representation of the hash (32 lower-case hex digits)
 	 **/
-	static FString HashBytes(const uint8* input, int32 inputLen)
+	static FString HashBytes(const uint8* input, uint64 inputLen)
 	{
 		uint8 Digest[16];
 
@@ -284,7 +284,7 @@ public:
 	 * @param BufferSize Size of Buffer
 	 * @param bDuplicateKeyMemory If Buffer is not always loaded, pass true so that the 20 byte hashes are duplicated 
 	 */
-	static void InitializeFileHashesFromBuffer(uint8* Buffer, int32 BufferSize, bool bDuplicateKeyMemory=false);
+	static void InitializeFileHashesFromBuffer(uint8* Buffer, uint64 BufferSize, bool bDuplicateKeyMemory=false);
 
 	/**
 	 * Gets the stored SHA hash from the platform, if it exists. This function
@@ -324,7 +324,7 @@ protected:
 	void* Buffer;
 
 	/** Size of Buffer */
-	int32 BufferSize;
+	uint64 BufferSize;
 
 	/** Hash to compare against */
 	uint8 Hash[20];
@@ -352,7 +352,7 @@ public:
 	 */
 	FAsyncSHAVerify(
 		void* InBuffer, 
-		int32 InBufferSize, 
+		uint64 InBufferSize, 
 		bool bInShouldDeleteBuffer, 
 		const TCHAR* InPathname, 
 		bool bInIsUnfoundHashAnError)
@@ -425,7 +425,7 @@ public:
 	 */
 	FBufferReaderWithSHA( 
 		void* Data, 
-		int32 Size, 
+		int64 Size, 
 		bool bInFreeOnClose, 
 		const TCHAR* SHASourcePathname, 
 		bool bIsPersistent=false, 
