Support log file paths which contains '(' characters.

diff --git a/Engine/Source/Runtime/Core/Private/GenericPlatform/GenericPlatformOutputDevices.cpp b/Engine/Source/Runtime/Core/Private/GenericPlatform/GenericPlatformOutputDevices.cpp
--- a/Engine/Source/Runtime/Core/Private/GenericPlatform/GenericPlatformOutputDevices.cpp
+++ b/Engine/Source/Runtime/Core/Private/GenericPlatform/GenericPlatformOutputDevices.cpp
@@ -16,7 +16,7 @@
 #include "Misc/OutputDeviceConsole.h"
 #include "Templates/UniquePtr.h"
 
-TCHAR FGenericPlatformOutputDevices::CachedAbsoluteFilename[1024] = { 0 };
+TCHAR FGenericPlatformOutputDevices::CachedAbsoluteFilename[FGenericPlatformOutputDevices::AbsoluteFileNameMaxLength] = { 0 };
 
 void FGenericPlatformOutputDevices::SetupOutputDevices()
 {
@@ -59,9 +59,10 @@ FString FGenericPlatformOutputDevices::GetAbsoluteLogFilename()
 	{
 		FCString::Strcpy(CachedAbsoluteFilename, ARRAY_COUNT(CachedAbsoluteFilename), *FPaths::ProjectLogDir());
 		FString LogFilename;
-		if (!FParse::Value(FCommandLine::Get(), TEXT("LOG="), LogFilename))
+		const bool bShouldStopOnSeparator = false;
+		if (!FParse::Value(FCommandLine::Get(), TEXT("LOG="), LogFilename, bShouldStopOnSeparator))
 		{
-			if (FParse::Value(FCommandLine::Get(), TEXT("ABSLOG="), LogFilename))
+			if (FParse::Value(FCommandLine::Get(), TEXT("ABSLOG="), LogFilename, bShouldStopOnSeparator))
 			{
 				CachedAbsoluteFilename[0] = 0;
 			}
diff --git a/Engine/Source/Runtime/Core/Public/GenericPlatform/GenericPlatformOutputDevices.h b/Engine/Source/Runtime/Core/Public/GenericPlatform/GenericPlatformOutputDevices.h
--- a/Engine/Source/Runtime/Core/Public/GenericPlatform/GenericPlatformOutputDevices.h
+++ b/Engine/Source/Runtime/Core/Public/GenericPlatform/GenericPlatformOutputDevices.h
@@ -27,5 +27,6 @@ struct CORE_API FGenericPlatformOutputDevices
 	static FFeedbackContext*            GetFeedbackContext();
 
 protected:
-	static TCHAR CachedAbsoluteFilename[1024];
+	static const SIZE_T AbsoluteFileNameMaxLength = 1024;
+	static TCHAR CachedAbsoluteFilename[AbsoluteFileNameMaxLength];
 };
