Fixes to call sorting and reverse predicates via Invoke.

diff --git a/Engine/Source/Runtime/Core/Public/Algo/IntroSort.h b/Engine/Source/Runtime/Core/Public/Algo/IntroSort.h
--- a/Engine/Source/Runtime/Core/Public/Algo/IntroSort.h
+++ b/Engine/Source/Runtime/Core/Public/Algo/IntroSort.h
@@ -60,7 +60,7 @@ namespace AlgoImpl
 					T *Max, *Item;
 					for( Max=Current.Min, Item=Current.Min+1; Item<=Current.Max; Item++ )
 					{
-						if( Predicate( Invoke( Projection, *Max ), Invoke( Projection, *Item ) ) )
+						if( Invoke( Predicate, Invoke( Projection, *Max ), Invoke( Projection, *Item ) ) )
 						{
 							Max = Item;
 						}
@@ -78,8 +78,8 @@ namespace AlgoImpl
 				Inner.Max = Current.Max+1;
 				for( ; ; )
 				{
-					while( ++Inner.Min<=Current.Max && !Predicate( Invoke( Projection, *Current.Min ), Invoke( Projection, *Inner.Min ) ) );
-					while( --Inner.Max> Current.Min && !Predicate( Invoke( Projection, *Inner.Max ), Invoke( Projection, *Current.Min ) ) );
+					while( ++Inner.Min<=Current.Max && !Invoke( Predicate, Invoke( Projection, *Current.Min ), Invoke( Projection, *Inner.Min ) ) );
+					while( --Inner.Max> Current.Min && !Invoke( Predicate, Invoke( Projection, *Inner.Max ), Invoke( Projection, *Current.Min ) ) );
 					if( Inner.Min>Inner.Max )
 					{
 						break;
diff --git a/Engine/Source/Runtime/Core/Public/Templates/ReversePredicate.h b/Engine/Source/Runtime/Core/Public/Templates/ReversePredicate.h
--- a/Engine/Source/Runtime/Core/Public/Templates/ReversePredicate.h
+++ b/Engine/Source/Runtime/Core/Public/Templates/ReversePredicate.h
@@ -2,6 +2,7 @@
 
 #pragma once
 
+#include "Templates/Invoke.h"
 #include "Templates/UnrealTemplate.h"
 
 /**
@@ -16,8 +17,12 @@ class TReversePredicate
 public:
 	TReversePredicate( const PredicateType& InPredicate )
 		: Predicate( InPredicate )
-	{}
+	{
+	}
 
 	template <typename T>
-	FORCEINLINE bool operator()( T&& A, T&& B ) const { return Predicate( Forward<T>(B),  Forward<T>(A) ); }
-};
\ No newline at end of file
+	FORCEINLINE bool operator()( T&& A, T&& B ) const
+	{
+		return Invoke( Predicate, Forward<T>(B), Forward<T>(A) );
+	}
+};
