[UE-71305]  Fix rounding big numbers

diff --git a/Engine/Source/Runtime/Core/Private/Math/UnrealMath.cpp b/Engine/Source/Runtime/Core/Private/Math/UnrealMath.cpp
--- a/Engine/Source/Runtime/Core/Private/Math/UnrealMath.cpp
+++ b/Engine/Source/Runtime/Core/Private/Math/UnrealMath.cpp
@@ -2857,26 +2857,62 @@ double FMath::RoundHalfToEven(double F)
 
 float FMath::RoundHalfFromZero(float F)
 {
-	F = MathRoundingUtil::TruncateToHalfIfClose(F);
-	return (F < 0.0f) ? CeilToFloat(F - 0.5f) : FloorToFloat(F + 0.5f);
+	float ValueToFudgeIntegralPart = 0.0f;
+	float ValueToFudgeFractionalPart = FMath::Modf(F, &ValueToFudgeIntegralPart);
+
+	if (F < 0.0f)
+	{
+		return ValueToFudgeFractionalPart > -0.5f ? ValueToFudgeIntegralPart : ValueToFudgeIntegralPart - 1.0f;
+	}
+	else
+	{
+		return ValueToFudgeFractionalPart < 0.5f ? ValueToFudgeIntegralPart : ValueToFudgeIntegralPart + 1.0f;
+	}
 }
 
 double FMath::RoundHalfFromZero(double F)
 {
-	F = MathRoundingUtil::TruncateToHalfIfClose(F);
-	return (F < 0.0) ? CeilToDouble(F - 0.5) : FloorToDouble(F + 0.5);
+	double ValueToFudgeIntegralPart = 0.0;
+	double ValueToFudgeFractionalPart = FMath::Modf(F, &ValueToFudgeIntegralPart);
+
+	if (F < 0.0)
+	{
+		return ValueToFudgeFractionalPart > -0.5 ? ValueToFudgeIntegralPart : ValueToFudgeIntegralPart - 1.0;
+	}
+	else
+	{
+		return ValueToFudgeFractionalPart < 0.5 ? ValueToFudgeIntegralPart : ValueToFudgeIntegralPart + 1.0f;
+	}
 }
 
 float FMath::RoundHalfToZero(float F)
 {
-	F = MathRoundingUtil::TruncateToHalfIfClose(F);
-	return (F < 0.0f) ? FloorToFloat(F + 0.5f) : CeilToFloat(F - 0.5f);
+	float ValueToFudgeIntegralPart = 0.0f;
+	float ValueToFudgeFractionalPart = FMath::Modf(F, &ValueToFudgeIntegralPart);
+
+	if (F < 0.0f)
+	{
+		return ValueToFudgeFractionalPart < -0.5f ? ValueToFudgeIntegralPart - 1.0f : ValueToFudgeIntegralPart;
+	}
+	else
+	{
+		return ValueToFudgeFractionalPart > 0.5f ? ValueToFudgeIntegralPart + 1.0f : ValueToFudgeIntegralPart;
+	}
 }
 
 double FMath::RoundHalfToZero(double F)
 {
-	F = MathRoundingUtil::TruncateToHalfIfClose(F);
-	return (F < 0.0) ? FloorToDouble(F + 0.5) : CeilToDouble(F - 0.5);
+	double ValueToFudgeIntegralPart = 0.0;
+	double ValueToFudgeFractionalPart = FMath::Modf(F, &ValueToFudgeIntegralPart);
+
+	if (F < 0.0)
+	{
+		return ValueToFudgeFractionalPart < -0.5 ? ValueToFudgeIntegralPart - 1.0 : ValueToFudgeIntegralPart;
+	}
+	else
+	{
+		return ValueToFudgeFractionalPart > 0.5 ? ValueToFudgeIntegralPart + 1.0 : ValueToFudgeIntegralPart;
+	}
 }
 
 FString FMath::FormatIntToHumanReadable(int32 Val)
@@ -3405,4 +3441,4 @@ float FMath::PerlinNoise1D(const float Value)
 	const float v = rx1 * g[p[bx1]];
 
 	return 2.0f * (u + sx * (v - u));
-}
+}
\ No newline at end of file
diff --git a/Engine/Source/Runtime/Core/Private/Misc/AutomationTest.cpp b/Engine/Source/Runtime/Core/Private/Misc/AutomationTest.cpp
--- a/Engine/Source/Runtime/Core/Private/Misc/AutomationTest.cpp
+++ b/Engine/Source/Runtime/Core/Private/Misc/AutomationTest.cpp
@@ -1086,6 +1086,14 @@ void FAutomationTestBase::TestEqual(const TCHAR* What, const float Actual, const
 	}
 }
 
+void FAutomationTestBase::TestEqual(const TCHAR* What, const double Actual, const double Expected, double Tolerance)
+{
+	if (!FMath::IsNearlyEqual(Actual, Expected, Tolerance))
+	{
+		AddError(FString::Printf(TEXT("Expected '%s' to be %f, but it was %f within tolerance %f."), What, Expected, Actual, Tolerance), 1);
+	}
+}
+
 void FAutomationTestBase::TestEqual(const TCHAR* What, const FVector Actual, const FVector Expected, float Tolerance)
 {
 	if (!Expected.Equals(Actual, Tolerance))
diff --git a/Engine/Source/Runtime/Core/Private/Tests/Math/UnrealMathTest.cpp b/Engine/Source/Runtime/Core/Private/Tests/Math/UnrealMathTest.cpp
--- a/Engine/Source/Runtime/Core/Private/Tests/Math/UnrealMathTest.cpp
+++ b/Engine/Source/Runtime/Core/Private/Tests/Math/UnrealMathTest.cpp
@@ -1600,4 +1600,152 @@ bool FInterpolationFunctionTests::RunTest(const FString&)
 	return true;
 }
 
+IMPLEMENT_SIMPLE_AUTOMATION_TEST(FMathRoundHalfToZeroTests, "System.Core.Math.Round HalfToZero", EAutomationTestFlags::ApplicationContextMask | EAutomationTestFlags::SmokeFilter)
+bool FMathRoundHalfToZeroTests::RunTest(const FString& Parameters)
+{
+	TestEqual(TEXT("RoundHalfToZero32-Zero"), FMath::RoundHalfToZero(0.0f), 0.0f);
+	TestEqual(TEXT("RoundHalfToZero32-One"), FMath::RoundHalfToZero(1.0f), 1.0f);
+	TestEqual(TEXT("RoundHalfToZero32-LessHalf"), FMath::RoundHalfToZero(1.4f), 1.0f);
+	TestEqual(TEXT("RoundHalfToZero32-NegGreaterHalf"), FMath::RoundHalfToZero(-1.4f), -1.0f);
+	TestEqual(TEXT("RoundHalfToZero32-LessNearHalf"), FMath::RoundHalfToZero(1.4999999f), 1.0f);
+	TestEqual(TEXT("RoundHalfToZero32-NegGreaterNearHalf"), FMath::RoundHalfToZero(-1.4999999f), -1.0f);
+	TestEqual(TEXT("RoundHalfToZero32-Half"), FMath::RoundHalfToZero(1.5f), 1.0f);
+	TestEqual(TEXT("RoundHalfToZero32-NegHalf"), FMath::RoundHalfToZero(-1.5f), -1.0f);
+	TestEqual(TEXT("RoundHalfToZero32-GreaterNearHalf"), FMath::RoundHalfToZero(1.5000001f), 2.0f);
+	TestEqual(TEXT("RoundHalfToZero32-NegLesserNearHalf"), FMath::RoundHalfToZero(-1.5000001f), -2.0f);
+	TestEqual(TEXT("RoundHalfToZero32-GreaterThanHalf"), FMath::RoundHalfToZero(1.6f), 2.0f);
+	TestEqual(TEXT("RoundHalfToZero32-NegLesserThanHalf"), FMath::RoundHalfToZero(-1.6f), -2.0f);
+
+	TestEqual(TEXT("RoundHalfToZero32-TwoToOneBitPrecision"), FMath::RoundHalfToZero(4194303.25f), 4194303.0f);
+	TestEqual(TEXT("RoundHalfToZero32-TwoToOneBitPrecision"), FMath::RoundHalfToZero(4194303.5f), 4194303.0f);
+	TestEqual(TEXT("RoundHalfToZero32-TwoToOneBitPrecision"), FMath::RoundHalfToZero(4194303.75f), 4194304.0f);
+	TestEqual(TEXT("RoundHalfToZero32-TwoToOneBitPrecision"), FMath::RoundHalfToZero(4194304.0f), 4194304.0f);
+	TestEqual(TEXT("RoundHalfToZero32-TwoToOneBitPrecision"), FMath::RoundHalfToZero(4194304.5f), 4194304.0f);
+	TestEqual(TEXT("RoundHalfToZero32-NegTwoToOneBitPrecision"), FMath::RoundHalfToZero(-4194303.25f), -4194303.0f);
+	TestEqual(TEXT("RoundHalfToZero32-NegTwoToOneBitPrecision"), FMath::RoundHalfToZero(-4194303.5f), -4194303.0f);
+	TestEqual(TEXT("RoundHalfToZero32-NegTwoToOneBitPrecision"), FMath::RoundHalfToZero(-4194303.75f), -4194304.0f);
+	TestEqual(TEXT("RoundHalfToZero32-NegTwoToOneBitPrecision"), FMath::RoundHalfToZero(-4194304.0f), -4194304.0f);
+	TestEqual(TEXT("RoundHalfToZero32-NegTwoToOneBitPrecision"), FMath::RoundHalfToZero(-4194304.5f), -4194304.0f);
+
+	TestEqual(TEXT("RoundHalfToZero32-OneToZeroBitPrecision"), FMath::RoundHalfToZero(8388607.0f), 8388607.0f);
+	TestEqual(TEXT("RoundHalfToZero32-OneToZeroBitPrecision"), FMath::RoundHalfToZero(8388607.5f), 8388607.0f);
+	TestEqual(TEXT("RoundHalfToZero32-OneToZeroBitPrecision"), FMath::RoundHalfToZero(8388608.0f), 8388608.0f);
+	TestEqual(TEXT("RoundHalfToZero32-NegOneToZeroBitPrecision"), FMath::RoundHalfToZero(-8388607.0f), -8388607.0f);
+	TestEqual(TEXT("RoundHalfToZero32-NegOneToZeroBitPrecision"), FMath::RoundHalfToZero(-8388607.5f), -8388607.0f);
+	TestEqual(TEXT("RoundHalfToZero32-NegOneToZeroBitPrecision"), FMath::RoundHalfToZero(-8388608.0f), -8388608.0f);
+
+	TestEqual(TEXT("RoundHalfToZero32-ZeroBitPrecision"), FMath::RoundHalfToZero(16777215.0f), 16777215.0f);
+	TestEqual(TEXT("RoundHalfToZero32-NegZeroBitPrecision"), FMath::RoundHalfToZero(-16777215.0f), -16777215.0f);
+
+	TestEqual(TEXT("RoundHalfToZero64-Zero"), FMath::RoundHalfToZero(0.0), 0.0);
+	TestEqual(TEXT("RoundHalfToZero64-One"), FMath::RoundHalfToZero(1.0), 1.0);
+	TestEqual(TEXT("RoundHalfToZero64-LessHalf"), FMath::RoundHalfToZero(1.4), 1.0);
+	TestEqual(TEXT("RoundHalfToZero64-NegGreaterHalf"), FMath::RoundHalfToZero(-1.4), -1.0);
+	TestEqual(TEXT("RoundHalfToZero64-LessNearHalf"), FMath::RoundHalfToZero(1.4999999999999997), 1.0);
+	TestEqual(TEXT("RoundHalfToZero64-NegGreaterNearHalf"), FMath::RoundHalfToZero(-1.4999999999999997), -1.0);
+	TestEqual(TEXT("RoundHalfToZero64-Half"), FMath::RoundHalfToZero(1.5), 1.0);
+	TestEqual(TEXT("RoundHalfToZero64-NegHalf"), FMath::RoundHalfToZero(-1.5), -1.0);
+	TestEqual(TEXT("RoundHalfToZero64-GreaterNearHalf"), FMath::RoundHalfToZero(1.5000000000000002), 2.0);
+	TestEqual(TEXT("RoundHalfToZero64-NegLesserNearHalf"), FMath::RoundHalfToZero(-1.5000000000000002), -2.0);
+	TestEqual(TEXT("RoundHalfToZero64-GreaterThanHalf"), FMath::RoundHalfToZero(1.6), 2.0);
+	TestEqual(TEXT("RoundHalfToZero64-NegLesserThanHalf"), FMath::RoundHalfToZero(-1.6), -2.0);
+	
+	TestEqual(TEXT("RoundHalfToZero64-TwoToOneBitPrecision"), FMath::RoundHalfToZero(2251799813685247.25), 2251799813685247.0);
+	TestEqual(TEXT("RoundHalfToZero64-TwoToOneBitPrecision"), FMath::RoundHalfToZero(2251799813685247.5), 2251799813685247.0);
+	TestEqual(TEXT("RoundHalfToZero64-TwoToOneBitPrecision"), FMath::RoundHalfToZero(2251799813685247.75), 2251799813685248.0);
+	TestEqual(TEXT("RoundHalfToZero64-TwoToOneBitPrecision"), FMath::RoundHalfToZero(2251799813685248.0), 2251799813685248.0);
+	TestEqual(TEXT("RoundHalfToZero64-TwoToOneBitPrecision"), FMath::RoundHalfToZero(2251799813685248.5), 2251799813685248.0);
+	TestEqual(TEXT("RoundHalfToZero64-NegTwoToOneBitPrecision"), FMath::RoundHalfToZero(-2251799813685247.25), -2251799813685247.0);
+	TestEqual(TEXT("RoundHalfToZero64-NegTwoToOneBitPrecision"), FMath::RoundHalfToZero(-2251799813685247.5), -2251799813685247.0);
+	TestEqual(TEXT("RoundHalfToZero64-NegTwoToOneBitPrecision"), FMath::RoundHalfToZero(-2251799813685247.75), -2251799813685248.0);
+	TestEqual(TEXT("RoundHalfToZero64-NegTwoToOneBitPrecision"), FMath::RoundHalfToZero(-2251799813685248.0), -2251799813685248.0);
+	TestEqual(TEXT("RoundHalfToZero64-NegTwoToOneBitPrecision"), FMath::RoundHalfToZero(-2251799813685248.5), -2251799813685248.0);
+
+	TestEqual(TEXT("RoundHalfToZero64-OneToZeroBitPrecision"), FMath::RoundHalfToZero(4503599627370495.0), 4503599627370495.0);
+	TestEqual(TEXT("RoundHalfToZero64-OneToZeroBitPrecision"), FMath::RoundHalfToZero(4503599627370495.5), 4503599627370495.0);
+	TestEqual(TEXT("RoundHalfToZero64-OneToZeroBitPrecision"), FMath::RoundHalfToZero(4503599627370496.0), 4503599627370496.0);
+	TestEqual(TEXT("RoundHalfToZero64-NegOneToZeroBitPrecision"), FMath::RoundHalfToZero(-4503599627370495.0), -4503599627370495.0);
+	TestEqual(TEXT("RoundHalfToZero64-NegOneToZeroBitPrecision"), FMath::RoundHalfToZero(-4503599627370495.5), -4503599627370495.0);
+	TestEqual(TEXT("RoundHalfToZero64-NegOneToZeroBitPrecision"), FMath::RoundHalfToZero(-4503599627370496.0), -4503599627370496.0);
+
+	TestEqual(TEXT("RoundHalfToZero64-ZeroBitPrecision"), FMath::RoundHalfToZero(9007199254740991.0), 9007199254740991.0);
+	TestEqual(TEXT("RoundHalfToZero64-NegZeroBitPrecision"), FMath::RoundHalfToZero(-9007199254740991.0), -9007199254740991.0);
+
+	return true;
+}
+
+IMPLEMENT_SIMPLE_AUTOMATION_TEST(FMathRoundHalfFromZeroTests, "System.Core.Math.Round HalfFromZero", EAutomationTestFlags::ApplicationContextMask | EAutomationTestFlags::SmokeFilter)
+bool FMathRoundHalfFromZeroTests::RunTest(const FString& Parameters)
+{
+	TestEqual(TEXT("RoundHalfFromZero32-Zero"), FMath::RoundHalfFromZero(0.0f), 0.0f);
+	TestEqual(TEXT("RoundHalfFromZero32-One"), FMath::RoundHalfFromZero(1.0f), 1.0f);
+	TestEqual(TEXT("RoundHalfFromZero32-LessHalf"), FMath::RoundHalfFromZero(1.4f), 1.0f);
+	TestEqual(TEXT("RoundHalfFromZero32-NegGreaterHalf"), FMath::RoundHalfFromZero(-1.4f), -1.0f);
+	TestEqual(TEXT("RoundHalfFromZero32-LessNearHalf"), FMath::RoundHalfFromZero(1.4999999f), 1.0f);
+	TestEqual(TEXT("RoundHalfFromZero32-NegGreaterNearHalf"), FMath::RoundHalfFromZero(-1.4999999f), -1.0f);
+	TestEqual(TEXT("RoundHalfFromZero32-Half"), FMath::RoundHalfFromZero(1.5f), 2.0f);
+	TestEqual(TEXT("RoundHalfFromZero32-NegHalf"), FMath::RoundHalfFromZero(-1.5f), -2.0f);
+	TestEqual(TEXT("RoundHalfFromZero32-LessGreaterNearHalf"), FMath::RoundHalfFromZero(1.5000001f), 2.0f);
+	TestEqual(TEXT("RoundHalfFromZero32-NegLesserNearHalf"), FMath::RoundHalfFromZero(-1.5000001f), -2.0f);
+	TestEqual(TEXT("RoundHalfFromZero32-GreaterThanHalf"), FMath::RoundHalfFromZero(1.6f), 2.0f);
+	TestEqual(TEXT("RoundHalfFromZero32-NegLesserThanHalf"), FMath::RoundHalfFromZero(-1.6f), -2.0f);
+
+	TestEqual(TEXT("RoundHalfFromZero32-TwoToOneBitPrecision"), FMath::RoundHalfFromZero(4194303.25f), 4194303.0f);
+	TestEqual(TEXT("RoundHalfFromZero32-TwoToOneBitPrecision"), FMath::RoundHalfFromZero(4194303.5f), 4194304.0f);
+	TestEqual(TEXT("RoundHalfFromZero32-TwoToOneBitPrecision"), FMath::RoundHalfFromZero(4194303.75f), 4194304.0f);
+	TestEqual(TEXT("RoundHalfFromZero32-TwoToOneBitPrecision"), FMath::RoundHalfFromZero(4194304.0f), 4194304.0f);
+	TestEqual(TEXT("RoundHalfFromZero32-TwoToOneBitPrecision"), FMath::RoundHalfFromZero(4194304.5f), 4194305.0f);
+	TestEqual(TEXT("RoundHalfFromZero32-NegTwoToOneBitPrecision"), FMath::RoundHalfFromZero(-4194303.25f), -4194303.0f);
+	TestEqual(TEXT("RoundHalfFromZero32-NegTwoToOneBitPrecision"), FMath::RoundHalfFromZero(-4194303.5f), -4194304.0f);
+	TestEqual(TEXT("RoundHalfFromZero32-NegTwoToOneBitPrecision"), FMath::RoundHalfFromZero(-4194303.75f), -4194304.0f);
+	TestEqual(TEXT("RoundHalfFromZero32-NegTwoToOneBitPrecision"), FMath::RoundHalfFromZero(-4194304.0f), -4194304.0f);
+	TestEqual(TEXT("RoundHalfFromZero32-NegTwoToOneBitPrecision"), FMath::RoundHalfFromZero(-4194304.5f), -4194305.0f);
+
+	TestEqual(TEXT("RoundHalfFromZero32-OneToZeroBitPrecision"), FMath::RoundHalfFromZero(8388607.0f), 8388607.0f);
+	TestEqual(TEXT("RoundHalfFromZero32-OneToZeroBitPrecision"), FMath::RoundHalfFromZero(8388607.5f), 8388608.0f);
+	TestEqual(TEXT("RoundHalfFromZero32-OneToZeroBitPrecision"), FMath::RoundHalfFromZero(8388608.0f), 8388608.0f);
+	TestEqual(TEXT("RoundHalfFromZero32-NegOneToZeroBitPrecision"), FMath::RoundHalfFromZero(-8388607.0f), -8388607.0f);
+	TestEqual(TEXT("RoundHalfFromZero32-NegOneToZeroBitPrecision"), FMath::RoundHalfFromZero(-8388607.5f), -8388608.0f);
+	TestEqual(TEXT("RoundHalfFromZero32-NegOneToZeroBitPrecision"), FMath::RoundHalfFromZero(-8388608.0f), -8388608.0f);
+
+	TestEqual(TEXT("RoundHalfFromZero32-ZeroBitPrecision"), FMath::RoundHalfToZero(16777215.0f), 16777215.0f);
+	TestEqual(TEXT("RoundHalfFromZero32-NegZeroBitPrecision"), FMath::RoundHalfToZero(-16777215.0f), -16777215.0f);
+
+	TestEqual(TEXT("RoundHalfFromZero64-Zero"), FMath::RoundHalfFromZero(0.0), 0.0);
+	TestEqual(TEXT("RoundHalfFromZero64-One"), FMath::RoundHalfFromZero(1.0), 1.0);
+	TestEqual(TEXT("RoundHalfFromZero64-LessHalf"), FMath::RoundHalfFromZero(1.4), 1.0);
+	TestEqual(TEXT("RoundHalfFromZero64-NegGreaterHalf"), FMath::RoundHalfFromZero(-1.4), -1.0);
+	TestEqual(TEXT("RoundHalfFromZero64-LessNearHalf"), FMath::RoundHalfFromZero(1.4999999999999997), 1.0);
+	TestEqual(TEXT("RoundHalfFromZero64-NegGreaterNearHalf"), FMath::RoundHalfFromZero(-1.4999999999999997), -1.0);
+	TestEqual(TEXT("RoundHalfFromZero64-Half"), FMath::RoundHalfFromZero(1.5), 2.0);
+	TestEqual(TEXT("RoundHalfFromZero64-NegHalf"), FMath::RoundHalfFromZero(-1.5), -2.0);
+	TestEqual(TEXT("RoundHalfFromZero64-LessGreaterNearHalf"), FMath::RoundHalfFromZero(1.5000000000000002), 2.0);
+	TestEqual(TEXT("RoundHalfFromZero64-NegLesserNearHalf"), FMath::RoundHalfFromZero(-1.5000000000000002), -2.0);
+	TestEqual(TEXT("RoundHalfFromZero64-GreaterThanHalf"), FMath::RoundHalfFromZero(1.6), 2.0);
+	TestEqual(TEXT("RoundHalfFromZero64-NegLesserThanHalf"), FMath::RoundHalfFromZero(-1.6), -2.0);
+
+	TestEqual(TEXT("RoundHalfFromZero64-TwoToOneBitPrecision"), FMath::RoundHalfFromZero(2251799813685247.25), 2251799813685247.0);
+	TestEqual(TEXT("RoundHalfFromZero64-TwoToOneBitPrecision"), FMath::RoundHalfFromZero(2251799813685247.5), 2251799813685248.0);
+	TestEqual(TEXT("RoundHalfFromZero64-TwoToOneBitPrecision"), FMath::RoundHalfFromZero(2251799813685247.75), 2251799813685248.0);
+	TestEqual(TEXT("RoundHalfFromZero64-TwoToOneBitPrecision"), FMath::RoundHalfFromZero(2251799813685248.0), 2251799813685248.0);
+	TestEqual(TEXT("RoundHalfFromZero64-TwoToOneBitPrecision"), FMath::RoundHalfFromZero(2251799813685248.5), 2251799813685249.0);
+	TestEqual(TEXT("RoundHalfFromZero64-NegTwoToOneBitPrecision"), FMath::RoundHalfFromZero(-2251799813685247.25), -2251799813685247.0);
+	TestEqual(TEXT("RoundHalfFromZero64-NegTwoToOneBitPrecision"), FMath::RoundHalfFromZero(-2251799813685247.5), -2251799813685248.0);
+	TestEqual(TEXT("RoundHalfFromZero64-NegTwoToOneBitPrecision"), FMath::RoundHalfFromZero(-2251799813685247.75), -2251799813685248.0);
+	TestEqual(TEXT("RoundHalfFromZero64-NegTwoToOneBitPrecision"), FMath::RoundHalfFromZero(-2251799813685248.0), -2251799813685248.0);
+	TestEqual(TEXT("RoundHalfFromZero64-NegTwoToOneBitPrecision"), FMath::RoundHalfFromZero(-2251799813685248.5), -2251799813685249.0);
+
+	TestEqual(TEXT("RoundHalfFromZero64-OneToZeroBitPrecision"), FMath::RoundHalfFromZero(4503599627370495.0), 4503599627370495.0);
+	TestEqual(TEXT("RoundHalfFromZero64-OneToZeroBitPrecision"), FMath::RoundHalfFromZero(4503599627370495.5), 4503599627370496.0);
+	TestEqual(TEXT("RoundHalfFromZero64-OneToZeroBitPrecision"), FMath::RoundHalfFromZero(4503599627370496.0), 4503599627370496.0);
+	TestEqual(TEXT("RoundHalfFromZero64-NegOneToZeroBitPrecision"), FMath::RoundHalfFromZero(-4503599627370495.0), -4503599627370495.0);
+	TestEqual(TEXT("RoundHalfFromZero64-NegOneToZeroBitPrecision"), FMath::RoundHalfFromZero(-4503599627370495.5), -4503599627370496.0);
+	TestEqual(TEXT("RoundHalfFromZero64-NegOneToZeroBitPrecision"), FMath::RoundHalfFromZero(-4503599627370496.0), -4503599627370496.0);
+
+	TestEqual(TEXT("RoundHalfFromZero64-ZeroBitPrecision"), FMath::RoundHalfToZero(9007199254740991.0), 9007199254740991.0);
+	TestEqual(TEXT("RoundHalfFromZero64-NegZeroBitPrecision"), FMath::RoundHalfToZero(-9007199254740991.0), -9007199254740991.0);
+
+	return true;
+}
+
 #endif //WITH_DEV_AUTOMATION_TESTS
diff --git a/Engine/Source/Runtime/Core/Public/Misc/AutomationTest.h b/Engine/Source/Runtime/Core/Public/Misc/AutomationTest.h
--- a/Engine/Source/Runtime/Core/Public/Misc/AutomationTest.h
+++ b/Engine/Source/Runtime/Core/Public/Misc/AutomationTest.h
@@ -1232,6 +1232,7 @@ public:
 
 	void TestEqual(const TCHAR* What, int32 Actual, int32 Expected);
 	void TestEqual(const TCHAR* What, float Actual, float Expected, float Tolerance = KINDA_SMALL_NUMBER);
+	void TestEqual(const TCHAR* What, double Actual, double Expected, double Tolerance = KINDA_SMALL_NUMBER);
 	void TestEqual(const TCHAR* What, FVector Actual, FVector Expected, float Tolerance = KINDA_SMALL_NUMBER);
 	void TestEqual(const TCHAR* What, FColor Actual, FColor Expected);
 	void TestEqual(const TCHAR* What, const TCHAR* A, const TCHAR* B);
@@ -1247,6 +1248,11 @@ public:
 		TestEqual(*What, Actual, Expected, Tolerance);
 	}
 
+	void TestEqual(const FString& What, double Actual, double Expected, double Tolerance = KINDA_SMALL_NUMBER)
+	{
+		TestEqual(*What, Actual, Expected, Tolerance);
+	}
+
 	void TestEqual(const FString& What, FVector Actual, FVector Expected, float Tolerance = KINDA_SMALL_NUMBER)
 	{
 		TestEqual(*What, Actual, Expected, Tolerance);
diff --git a/Engine/Source/Runtime/Json/Private/Dom/JsonValue.cpp b/Engine/Source/Runtime/Json/Private/Dom/JsonValue.cpp
--- a/Engine/Source/Runtime/Json/Private/Dom/JsonValue.cpp
+++ b/Engine/Source/Runtime/Json/Private/Dom/JsonValue.cpp
@@ -79,14 +79,7 @@ bool FJsonValue::TryGetNumber( int32& OutNumber ) const
 
 	if (TryGetNumber(Double) && (Double >= INT_MIN) && (Double <= INT_MAX))
 	{
-		if (Double >= 0.0)
-		{
-			OutNumber = (int32)(Double + 0.5);
-		}
-		else
-		{
-			OutNumber = (int32)(Double - 0.5);
-		}
+		OutNumber = static_cast<int32>(FMath::RoundHalfFromZero(Double));
 		
 		return true;
 	}
@@ -101,7 +94,7 @@ bool FJsonValue::TryGetNumber( uint32& OutNumber ) const
 
 	if (TryGetNumber(Double) && (Double >= 0.0) && (Double <= UINT_MAX))
 	{
-		OutNumber = (uint32)(Double + 0.5);
+		OutNumber = static_cast<uint32>(FMath::RoundHalfFromZero(Double));
 
 		return true;
 	}
@@ -116,14 +109,7 @@ bool FJsonValue::TryGetNumber( int64& OutNumber ) const
 
 	if (TryGetNumber(Double) && (Double >= INT64_MIN) && (Double <= INT64_MAX))
 	{
-		if (Double >= 0.0)
-		{
-			OutNumber = (int64)(Double + 0.5);
-		}
-		else
-		{
-			OutNumber = (int64)(Double - 0.5);
-		}
+		OutNumber = static_cast<int64>(FMath::RoundHalfFromZero(Double));
 		
 		return true;
 	}
diff --git a/Engine/Source/Runtime/Json/Private/Tests/JsonTests.cpp b/Engine/Source/Runtime/Json/Private/Tests/JsonTests.cpp
--- a/Engine/Source/Runtime/Json/Private/Tests/JsonTests.cpp
+++ b/Engine/Source/Runtime/Json/Private/Tests/JsonTests.cpp
@@ -688,6 +688,180 @@ bool FJsonAutomationTest::RunTest(const FString& Parameters)
 		check( !Object.IsValid() );
 	}
 
+	// TryGetNumber tests
+	{
+		auto JsonNumberToInt64 = [](double Val, int64& OutVal) -> bool
+		{
+			FJsonValueNumber JsonVal(Val);
+			return ((FJsonValue&)JsonVal).TryGetNumber(OutVal);
+		};
+
+		auto JsonNumberToInt32 = [](double Val, int32& OutVal) -> bool
+		{
+			FJsonValueNumber JsonVal(Val);
+			return ((FJsonValue&)JsonVal).TryGetNumber(OutVal);
+		};
+
+		auto JsonNumberToUInt32 = [](double Val, uint32& OutVal) -> bool
+		{
+			FJsonValueNumber JsonVal(Val);
+			return ((FJsonValue&)JsonVal).TryGetNumber(OutVal);
+		};
+		
+		// TryGetNumber-Int64 tests
+		{
+			int64 IntVal;
+			bool bOk = JsonNumberToInt64(9007199254740991.0, IntVal);
+			TestTrue(TEXT("TryGetNumber-Int64 Big Float64 succeeds"), bOk);
+			TestEqual(TEXT("TryGetNumber-Int64 Big Float64"), IntVal, 9007199254740991L);
+		}
+
+		{
+			int64 IntVal;
+			bool bOk = JsonNumberToInt64(-9007199254740991.0, IntVal);
+			TestTrue(TEXT("TryGetNumber-Int64 Small Float64 succeeds"), bOk);
+			TestEqual(TEXT("TryGetNumber-Int64 Small Float64"), IntVal, -9007199254740991L);
+		}
+
+		{
+			int64 IntVal;
+			bool bOk = JsonNumberToInt64(0.4999999999999997, IntVal);
+			TestTrue(TEXT("TryGetNumber-Int64 Lesser than near half succeeds"), bOk);
+			TestEqual(TEXT("TryGetNumber-Int64 Lesser than near half rounds to zero"), IntVal, 0);
+		}
+
+		{
+			int64 IntVal;
+			bool bOk = JsonNumberToInt64(-0.4999999999999997, IntVal);
+			TestTrue(TEXT("TryGetNumber-Int64 Greater than near negative half succeeds"), bOk);
+			TestEqual(TEXT("TryGetNumber-Int64 Greater than near negative half rounds to zero"), IntVal, 0);
+		}
+
+		{
+			int64 IntVal;
+			bool bOk = JsonNumberToInt64(0.5, IntVal);
+			TestTrue(TEXT("TryGetNumber-Int64 Half rounds to next integer succeeds"), bOk);
+			TestEqual(TEXT("TryGetNumber-Int64 Half rounds to next integer"), IntVal, 1LL);
+		}
+
+		{
+			int64 IntVal;
+			bool bOk = JsonNumberToInt64(-0.5, IntVal);
+			TestTrue(TEXT("TryGetNumber-Int64 Negative half rounds to next negative integer succeeds"), bOk);
+			TestEqual(TEXT("TryGetNumber-Int64 Negative half rounds to next negative integer succeeds"), IntVal, -1LL);
+		}
+
+		// TryGetNumber-Int32 tests
+		{
+			int32 IntVal;
+			bool bOk = JsonNumberToInt32(2147483647.000001, IntVal);
+			TestFalse(TEXT("TryGetNumber-Int32 Number greater than max Int32 fails"), bOk);
+		}
+
+		{
+			int32 IntVal;
+			bool bOk = JsonNumberToInt32(-2147483648.000001, IntVal);
+			TestFalse(TEXT("TryGetNumber-Int32 Number lesser than min Int32 fails"), bOk);
+		}
+
+		{
+			int32 IntVal;
+			bool bOk = JsonNumberToInt32(2147483647.0, IntVal);
+			TestTrue(TEXT("TryGetNumber-Int32 Max Int32 succeeds"), bOk);
+			TestEqual(TEXT("TryGetNumber-Int32 Max Int32"), IntVal, INT_MAX);
+		}
+
+		{
+			int32 IntVal;
+			bool bOk = JsonNumberToInt32(2147483646.5, IntVal);
+			TestTrue(TEXT("TryGetNumber-Int32 Round up to max Int32 succeeds"), bOk);
+			TestEqual(TEXT("TryGetNumber-Int32 Round up to max Int32"), IntVal, INT_MAX);
+		}
+
+		{
+			int32 IntVal;
+			bool bOk = JsonNumberToInt32(-2147483648.0, IntVal);
+			TestTrue(TEXT("TryGetNumber-Int32 Min Int32 succeeds"), bOk);
+			TestEqual(TEXT("TryGetNumber-Int32 Min Int32"), IntVal, INT_MIN);
+		}
+
+		{
+			int32 IntVal;
+			bool bOk = JsonNumberToInt32(-2147483647.5, IntVal);
+			TestTrue(TEXT("TryGetNumber-Int32 Round down to min Int32 succeeds"), bOk);
+			TestEqual(TEXT("TryGetNumber-Int32 Round down to min Int32"), IntVal, INT_MIN);
+		}
+
+		{
+			int32 IntVal;
+			bool bOk = JsonNumberToInt32(0.4999999999999997, IntVal);
+			TestTrue(TEXT("TryGetNumber-Int32 Lesser than near half succeeds"), bOk);
+			TestEqual(TEXT("TryGetNumber-Int32 Lesser than near half rounds to zero"), IntVal, 0);
+		}
+
+		{
+			int32 IntVal;
+			bool bOk = JsonNumberToInt32(-0.4999999999999997, IntVal);
+			TestTrue(TEXT("TryGetNumber-Int32 Greater than near negative half succeeds"), bOk);
+			TestEqual(TEXT("TryGetNumber-Int32 Greater than near negative half rounds to zero"), IntVal, 0);
+		}
+
+		{
+			int32 IntVal;
+			bool bOk = JsonNumberToInt32(0.5, IntVal);
+			TestTrue(TEXT("TryGetNumber-Int32 Half rounds to next integer succeeds"), bOk);
+			TestEqual(TEXT("TryGetNumber-Int32 Half rounds to next integer"), IntVal, 1);
+		}
+
+		{
+			int32 IntVal;
+			bool bOk = JsonNumberToInt32(-0.5, IntVal);
+			TestTrue(TEXT("TryGetNumber-Int32 Negative half rounds to next negative integer succeeds"), bOk);
+			TestEqual(TEXT("TryGetNumber-Int32 Negative half rounds to next negative integer succeeds"), IntVal, -1);
+		}
+
+		// TryGetNumber-UInt32 tests
+		{
+			uint32 IntVal;
+			bool bOk = JsonNumberToUInt32(4294967295.000001, IntVal);
+			TestFalse(TEXT("TryGetNumber-UInt32 Number greater than max Uint32 fails"), bOk);
+		}
+
+		{
+			uint32 IntVal;
+			bool bOk = JsonNumberToUInt32(-0.000000000000001, IntVal);
+			TestFalse(TEXT("TryGetNumber-UInt32 Negative number fails"), bOk);
+		}
+
+		{
+			uint32 IntVal;
+			bool bOk = JsonNumberToUInt32(4294967295.0, IntVal);
+			TestTrue(TEXT("TryGetNumber-UInt32 Max UInt32 succeeds"), bOk);
+			TestEqual(TEXT("TryGetNumber-UInt32  Max UInt32"), IntVal, UINT_MAX);
+		}
+
+		{
+			uint32 IntVal;
+			bool bOk = JsonNumberToUInt32(4294967294.5, IntVal);
+			TestTrue(TEXT("TryGetNumber-UInt32 Round up to max UInt32 succeeds"), bOk);
+			TestEqual(TEXT("TryGetNumber-UInt32 Round up to max UInt32"), IntVal, UINT_MAX);
+		}
+
+		{
+			uint32 IntVal;
+			bool bOk = JsonNumberToUInt32(0.4999999999999997, IntVal);
+			TestTrue(TEXT("TryGetNumber-UInt32 Lesser than near half succeeds"), bOk);
+			TestEqual(TEXT("TryGetNumber-UInt32 Lesser than near half rounds to zero"), IntVal, 0U);
+		}
+
+		{
+			uint32 IntVal;
+			bool bOk = JsonNumberToUInt32(0.5, IntVal);
+			TestTrue(TEXT("TryGetNumber-UInt32 Half rounds to next integer succeeds"), bOk);
+			TestEqual(TEXT("TryGetNumber-UInt32 Half rounds to next integer"), IntVal, 1U);
+		}
+	}
+
 	return true;
 }
 
