Fix for wrongly-named template parameter.

diff --git a/Engine/Source/Runtime/Core/Public/Templates/Function.h b/Engine/Source/Runtime/Core/Public/Templates/Function.h
--- a/Engine/Source/Runtime/Core/Public/Templates/Function.h
+++ b/Engine/Source/Runtime/Core/Public/Templates/Function.h
@@ -270,23 +270,23 @@ namespace UE4Function_Private
 		return true;
 	}
 
-	template <typename FunctorType, bool bUnique, bool bInline>
+	template <typename FunctorType, bool bUnique, bool bOnHeap>
 	struct TStorageOwnerType;
 
-	template <typename FunctorType, bool bInline>
-	struct TStorageOwnerType<FunctorType, true, bInline>
+	template <typename FunctorType, bool bOnHeap>
+	struct TStorageOwnerType<FunctorType, true, bOnHeap>
 	{
-		using Type = TFunction_UniqueOwnedObject<typename TDecay<FunctorType>::Type, bInline>;
+		using Type = TFunction_UniqueOwnedObject<typename TDecay<FunctorType>::Type, bOnHeap>;
 	};
 
-	template <typename FunctorType, bool bInline>
-	struct TStorageOwnerType<FunctorType, false, bInline>
+	template <typename FunctorType, bool bOnHeap>
+	struct TStorageOwnerType<FunctorType, false, bOnHeap>
 	{
-		using Type = TFunction_CopyableOwnedObject<typename TDecay<FunctorType>::Type, bInline>;
+		using Type = TFunction_CopyableOwnedObject<typename TDecay<FunctorType>::Type, bOnHeap>;
 	};
 
-	template <typename FunctorType, bool bUnique, bool bInline>
-	using TStorageOwnerTypeT = typename TStorageOwnerType<FunctorType, bUnique, bInline>::Type;
+	template <typename FunctorType, bool bUnique, bool bOnHeap>
+	using TStorageOwnerTypeT = typename TStorageOwnerType<FunctorType, bUnique, bOnHeap>::Type;
 
 	struct FFunctionStorage
 	{
