Fix clang warning inline function 'FAsciiSet::Set<char>' is not defined

diff --git a/Engine/Source/Runtime/Core/Public/Misc/AsciiSet.h b/Engine/Source/Runtime/Core/Public/Misc/AsciiSet.h
--- a/Engine/Source/Runtime/Core/Public/Misc/AsciiSet.h
+++ b/Engine/Source/Runtime/Core/Public/Misc/AsciiSet.h
@@ -37,7 +37,7 @@ public:
 	constexpr FORCEINLINE FAsciiSet operator+(char Char) const
 	{
 		InitData Bitset = { LoMask, HiMask };
-		Set(Bitset, Char);
+		SetImpl(Bitset, TChar<char>::ToUnsigned(Char));
 		return FAsciiSet(Bitset);
 	}
 
@@ -91,8 +91,7 @@ private:
 	struct InitData { uint64 Lo, Hi; };
 	static constexpr uint64 NilMask = uint64(1) << '\0';
 
-	template<typename CharType>
-	static constexpr FORCEINLINE void Set(InitData& Bitset, CharType Char)
+	static constexpr FORCEINLINE void SetImpl(InitData& Bitset, uint32 Char)
 	{
 		uint64 IsLo = uint64(0) - (Char >> 6 == 0);
 		uint64 IsHi = uint64(0) - (Char >> 6 == 1);
@@ -117,7 +116,7 @@ private:
 		InitData Bitset = { 0, 0 };
 		for (int I = 0; I < N - 1; ++I)
 		{
-			Set(Bitset, Chars[I]);
+			SetImpl(Bitset, TChar<CharType>::ToUnsigned(Chars[I]));
 		}
 
 		return Bitset;
diff --git a/Engine/Source/Runtime/Core/Public/Misc/Char.h b/Engine/Source/Runtime/Core/Public/Misc/Char.h
--- a/Engine/Source/Runtime/Core/Public/Misc/Char.h
+++ b/Engine/Source/Runtime/Core/Public/Misc/Char.h
@@ -128,7 +128,7 @@ struct TChar : TCharBase<CharType, sizeof(CharType)>
 	*
 	* Mainly needed for subtraction and addition.
 	*/
-	static uint32 ToUnsigned(CharType Char)
+	static constexpr FORCEINLINE uint32 ToUnsigned(CharType Char)
 	{
 		return (typename TUnsignedIntType<sizeof(CharType)>::Type)Char;
 	}
