FName optimizations and improvements

Encoding improvements
* New hash table implementation
--- Move away from 16-bit hashes since we need more than 64k buckets to hold 2M entries efficiently
--- Change to CityHash64, which is faster and stronger
--- Remove hardcoded max limit
* NAME_INDEX changed from contigouos int to monotonically increasing int
--- Opens up for future deduplication schemes that are better than number suffix dedup.
--- Saves some memory since we don't need to maintain a contiguous array
--- Typed up to cause compile errors when used as integer directly
* Avoid touching data repeatedly, normal path only does single hash of data
* New constructor that allows supplying string length up front
* Avoid dynamic allocations in string conversions >= 128 characters
* Avoid extra copying when splitting numbers
* More efficient IsPureAnsi check
* Only do one global lazy initialization call instead of multiple
* Switch to faster RW locks
* Switch from single lock to sharded hash map with separate locks
* Memory optimizations that reduces per entry overhead: 4 bytes slots, 2 byte headers and don't store null terminator

Improved API & documentation
* Document that IsValid() rarely makes sense
* Hide global state such as GetNames()
* Reduce amount of implementation details visible in header
* NameTypes.h size down by ~1/3 while adding documentation, stronger type safety and new APIs

Future possibilities
* Memory savings: Removing public NAME_INDEX and global FName array allows using the 32-bit FNameEntryId for arbitrary deduplication schemes. This can save both actual stored strings memory by deduplication and half the size of FName instances from 8B to 4B in shipping / test configs by removing the number part.
* Implementation can be tweaked further, for instance could persist 32-bit slot index hash inside slot to increase encoding performance in development / debug at the cost of memory.

Perf & mem results for internal project:
* Editor startup wall time: 12% speedup, 22.0s -> 19.5s
--- Warm disk and asset registry cache
--- Some gains from nametable serialization improvements
--- Some of these gains might be from orthogonal asset discovery optimizations
* Win64 Test Client memory usage with ~400k names: 30.3Mb -> 22.2Mb

diff --git a/Engine/Extras/NatvisHelpers/EntryPoint.cpp b/Engine/Extras/NatvisHelpers/EntryPoint.cpp
--- a/Engine/Extras/NatvisHelpers/EntryPoint.cpp
+++ b/Engine/Extras/NatvisHelpers/EntryPoint.cpp
@@ -2,8 +2,8 @@
 
 #include "Globals.h"
 
-extern "C" __declspec(dllexport) void InitNatvisHelpers(FNameEntry*** NameTable, FChunkedFixedUObjectArray* ObjectArray)
+extern "C" __declspec(dllexport) void InitNatvisHelpers(uint8** NameTable, FChunkedFixedUObjectArray* ObjectArray)
 {
-	GFNameTableForDebuggerVisualizers_MT = NameTable;
+	GNameBlocksDebug = NameTable;
 	GObjectArrayForDebugVisualizers = ObjectArray;
 }
diff --git a/Engine/Extras/NatvisHelpers/Globals.cpp b/Engine/Extras/NatvisHelpers/Globals.cpp
--- a/Engine/Extras/NatvisHelpers/Globals.cpp
+++ b/Engine/Extras/NatvisHelpers/Globals.cpp
@@ -3,5 +3,5 @@
 struct FNameEntry;
 class FChunkedFixedUObjectArray;
 
-FNameEntry*** GFNameTableForDebuggerVisualizers_MT = nullptr;
+uint8** GNameBlocksDebug = nullptr;
 FChunkedFixedUObjectArray* GObjectArrayForDebugVisualizers = nullptr;
diff --git a/Engine/Extras/NatvisHelpers/Globals.h b/Engine/Extras/NatvisHelpers/Globals.h
--- a/Engine/Extras/NatvisHelpers/Globals.h
+++ b/Engine/Extras/NatvisHelpers/Globals.h
@@ -5,5 +5,5 @@
 struct FNameEntry;
 class FChunkedFixedUObjectArray;
 
-extern FNameEntry*** GFNameTableForDebuggerVisualizers_MT;
+extern uint8** GNameBlocksDebug;
 extern FChunkedFixedUObjectArray* GObjectArrayForDebugVisualizers;
diff --git a/Engine/Plugins/Developer/Concert/ConcertMain/Source/ConcertTransport/Private/IdentifierTable/ConcertTransportArchives.cpp b/Engine/Plugins/Developer/Concert/ConcertMain/Source/ConcertTransport/Private/IdentifierTable/ConcertTransportArchives.cpp
--- a/Engine/Plugins/Developer/Concert/ConcertMain/Source/ConcertTransport/Private/IdentifierTable/ConcertTransportArchives.cpp
+++ b/Engine/Plugins/Developer/Concert/ConcertMain/Source/ConcertTransport/Private/IdentifierTable/ConcertTransportArchives.cpp
@@ -33,12 +33,12 @@ FArchive& FConcertIdentifierWriter::operator<<(FName& Name)
 		SerializeIntPacked(UnsignedIndex);
 	};
 
-	int32 HardcodedIndex = Name.GetComparisonIndex();
+	const EName* Ename = Name.ToEName();
 	int32 IdentifierTableIndex = INDEX_NONE;
-	if (HardcodedIndex <= MAX_NETWORKED_HARDCODED_NAME)
+	if (Ename && ShouldReplicateAsInteger(*Ename))
 	{
 		SerializeConcertIdentifierSource(EConcertIdentifierSource::HardcodedIndex);
-		SerializeIndexValue(HardcodedIndex);
+		SerializeIndexValue(*Ename);
 	}
 	else if (LocalIdentifierTable)
 	{
diff --git a/Engine/Plugins/Developer/Concert/ConcertMain/Source/ConcertTransport/Public/IdentifierTable/ConcertIdentifierTable.h b/Engine/Plugins/Developer/Concert/ConcertMain/Source/ConcertTransport/Public/IdentifierTable/ConcertIdentifierTable.h
--- a/Engine/Plugins/Developer/Concert/ConcertMain/Source/ConcertTransport/Public/IdentifierTable/ConcertIdentifierTable.h
+++ b/Engine/Plugins/Developer/Concert/ConcertMain/Source/ConcertTransport/Public/IdentifierTable/ConcertIdentifierTable.h
@@ -19,7 +19,7 @@ struct TConcertIdentifierTable_CaseSensitivePlainNameKeyFuncs : BaseKeyFuncs<Val
 	}
 	static FORCEINLINE uint32 GetKeyHash(const FName& Key)
 	{
-		return Key.GetDisplayIndex();
+		return GetTypeHash(Key.GetDisplayIndex());
 	}
 };
 
diff --git a/Engine/Source/Developer/CollectionManager/Private/Collection.cpp b/Engine/Source/Developer/CollectionManager/Private/Collection.cpp
--- a/Engine/Source/Developer/CollectionManager/Private/Collection.cpp
+++ b/Engine/Source/Developer/CollectionManager/Private/Collection.cpp
@@ -122,13 +122,13 @@ bool FCollection::Load(FText& OutError)
 	if (StorageMode == ECollectionStorageMode::Static)
 	{
 		// Static collection, a flat list of asset paths
-		for (FString Line : FileContents)
+		for (FString& Line : FileContents)
 		{
 			Line.TrimStartAndEndInline();
 
-			if ( Line.Len() )
+			if ( int32 Len = Line.Len() )
 			{
-				AddObjectToCollection(FName(*Line));
+				AddObjectToCollection(FName(Len, *Line));
 			}
 		}
 	}
@@ -398,10 +398,11 @@ void FCollection::Empty()
 
 bool FCollection::AddObjectToCollection(FName ObjectPath)
 {
-	if (StorageMode == ECollectionStorageMode::Static && !ObjectSet.Contains(ObjectPath))
+	if (StorageMode == ECollectionStorageMode::Static)
 	{
-		ObjectSet.Add(ObjectPath);
-		return true;
+		bool bAlreadyInSet = false;
+		ObjectSet.Add(ObjectPath, &bAlreadyInSet);
+		return !bAlreadyInSet;
 	}
 
 	return false;
diff --git a/Engine/Source/Developer/Profiler/Private/ProfilerRawStatsForMemory.cpp b/Engine/Source/Developer/Profiler/Private/ProfilerRawStatsForMemory.cpp
--- a/Engine/Source/Developer/Profiler/Private/ProfilerRawStatsForMemory.cpp
+++ b/Engine/Source/Developer/Profiler/Private/ProfilerRawStatsForMemory.cpp
@@ -61,7 +61,7 @@ struct FStatsCallstack
 		FString Result;
 		for (const auto& Name : Callstack)
 		{
-			Result += TTypeToString<int32>::ToString( (int32)Name.GetComparisonIndex() );
+			Result += TTypeToString<uint32>::ToString(Name.GetComparisonIndex().ToUnstableInt() );
 			Result += CallstackSeparator;
 		}
 		return Result;
@@ -76,9 +76,10 @@ struct FStatsCallstack
 		// Convert back to FNames
 		for (const auto& It : DecodedCallstack)
 		{
-			NAME_INDEX NameIndex = 0;
-			TTypeFromString<NAME_INDEX>::FromString( NameIndex, *It );
-			const FName LongName = FName( NameIndex, NameIndex, 0 );
+			uint32 NameInt = 0;
+			TTypeFromString<uint32>::FromString(NameInt, *It );
+			FNameEntryId Id = FNameEntryId::FromUnstableInt(NameInt);
+			const FName LongName = FName(Id, Id, 0 );
 
 			out_DecodedCallstack.Add( LongName );
 		}
diff --git a/Engine/Source/Developer/Profiler/Private/ProfilerRawStatsForThreadView.cpp b/Engine/Source/Developer/Profiler/Private/ProfilerRawStatsForThreadView.cpp
--- a/Engine/Source/Developer/Profiler/Private/ProfilerRawStatsForThreadView.cpp
+++ b/Engine/Source/Developer/Profiler/Private/ProfilerRawStatsForThreadView.cpp
@@ -169,11 +169,6 @@ void FRawProfilerSession::PrepareLoading()
 		Stream.ReadFNamesAndMetadataMessages( *FileReader, MetadataMessages );
 		StatsThreadStats.ProcessMetaDataOnly( MetadataMessages );
 
-		const FName F00245 = FName(245, 245, 0);
-		
-		const FName F11602 = FName(11602, 11602, 0);
-		const FName F06394 = FName(6394, 6394, 0);
-
 		const int64 CurrentFilePos = FileReader->Tell();
 
 		// Update profiler's metadata.
diff --git a/Engine/Source/Developer/Windows/LiveCoding/Private/External/LC_ClientCommandActions.cpp b/Engine/Source/Developer/Windows/LiveCoding/Private/External/LC_ClientCommandActions.cpp
--- a/Engine/Source/Developer/Windows/LiveCoding/Private/External/LC_ClientCommandActions.cpp
+++ b/Engine/Source/Developer/Windows/LiveCoding/Private/External/LC_ClientCommandActions.cpp
@@ -128,8 +128,7 @@ bool actions::GetModule::Execute(CommandType* command, const DuplexPipe* pipe, v
 
 
 // BEGIN EPIC MOD - Support for UE4 debug visualizers
-struct FNameEntry;
-extern FNameEntry*** GFNameTableForDebuggerVisualizers_MT;
+extern uint8** GNameBlocksDebug;
 
 class FChunkedFixedUObjectArray;
 extern FChunkedFixedUObjectArray*& GObjectArrayForDebugVisualizers;
@@ -144,12 +143,12 @@ bool actions::LoadPatch::Execute(CommandType* command, const DuplexPipe* pipe, v
 	// BEGIN EPIC MOD - Support for UE4 debug visualizers
 	if (module != nullptr)
 	{
-		typedef void InitNatvisHelpersFunc(FNameEntry*** NameTable, FChunkedFixedUObjectArray* ObjectArray);
+		typedef void InitNatvisHelpersFunc(uint8** NameTable, FChunkedFixedUObjectArray* ObjectArray);
 
 		InitNatvisHelpersFunc* InitNatvisHelpers = (InitNatvisHelpersFunc*)(void*)GetProcAddress(module, "InitNatvisHelpers");
 		if (InitNatvisHelpers != nullptr)
 		{
-			(*InitNatvisHelpers)(GFNameTableForDebuggerVisualizers_MT, GObjectArrayForDebugVisualizers);
+			(*InitNatvisHelpers)(GNameBlocksDebug, GObjectArrayForDebugVisualizers);
 		}
 	}
 	// END EPIC MOD
diff --git a/Engine/Source/Editor/PinnedCommandList/Private/UICommandList_Pinnable.cpp b/Engine/Source/Editor/PinnedCommandList/Private/UICommandList_Pinnable.cpp
--- a/Engine/Source/Editor/PinnedCommandList/Private/UICommandList_Pinnable.cpp
+++ b/Engine/Source/Editor/PinnedCommandList/Private/UICommandList_Pinnable.cpp
@@ -26,7 +26,7 @@ void FUICommandList_Pinnable::MapAction( const TSharedPtr< const FUICommandInfo
 	FUICommandList::MapAction(InUICommandInfo, InUIAction);
 
 	// Map the index and group
-	int32 ActionIndex = InUICommandInfo->GetBindingContext().GetComparisonIndex() + CurrentActionIndex;
+	int32 ActionIndex = InUICommandInfo->GetBindingContext().GetComparisonIndex().ToUnstableInt() + CurrentActionIndex;
 	CurrentActionIndex++;
 
 	CommandIndexMap.Add(InUICommandInfo->GetCommandName(), ActionIndex);
diff --git a/Engine/Source/Editor/UnrealEd/Private/Commandlets/PackageUtilities.cpp b/Engine/Source/Editor/UnrealEd/Private/Commandlets/PackageUtilities.cpp
--- a/Engine/Source/Editor/UnrealEd/Private/Commandlets/PackageUtilities.cpp
+++ b/Engine/Source/Editor/UnrealEd/Private/Commandlets/PackageUtilities.cpp
@@ -976,8 +976,8 @@ void FPkgInfoReporter_Log::GeneratePackageReport( FLinkerLoad* InLinker /*=nullp
 		Out.Logf(ELogVerbosity::Display, TEXT("========"));
 		for( int32 i = 0; i < Linker->NameMap.Num(); ++i )
 		{
-			FName& name = Linker->NameMap[ i ];
-			Out.Logf(ELogVerbosity::Display, TEXT("\t%d: Name '%s' Comparison Index %d Display Index %d [Internal: %s, %d]"), i, *name.ToString(), name.GetComparisonIndex(), name.GetDisplayIndex(), *name.GetPlainNameString(), name.GetNumber() );
+			FName name = FName::CreateFromDisplayId(Linker->NameMap[ i ], 0);
+			Out.Logf(ELogVerbosity::Display, TEXT("\t%d: Name '%s' Comparison Index %d Display Index %d [Internal: %s, %d]"), i, *name.ToString(), name.GetComparisonIndex().ToUnstableInt(), name.GetDisplayIndex().ToUnstableInt(), *name.GetPlainNameString(), name.GetNumber() );
 		}
 	}
 
diff --git a/Engine/Source/Runtime/AssetRegistry/Private/NameTableArchive.cpp b/Engine/Source/Runtime/AssetRegistry/Private/NameTableArchive.cpp
--- a/Engine/Source/Runtime/AssetRegistry/Private/NameTableArchive.cpp
+++ b/Engine/Source/Runtime/AssetRegistry/Private/NameTableArchive.cpp
@@ -85,6 +85,7 @@ bool FNameTableArchiveReader::SerializeNameMap()
 			return false;
 		}
 
+		NameMap.Reserve(NameCount);
 		for ( int32 NameMapIdx = 0; NameMapIdx < NameCount; ++NameMapIdx )
 		{
 			// Read the name entry from the file.
@@ -96,7 +97,7 @@ bool FNameTableArchiveReader::SerializeNameMap()
 				return false;
 			}
 
-			NameMap.Add(FName(NameEntry));
+			NameMap.Add(FName(NameEntry).GetDisplayIndex());
 		}
 
 		Seek( OriginalOffset );
@@ -193,19 +194,18 @@ FArchive& FNameTableArchiveReader::operator<<( FName& Name )
 		SetError();
 	}
 
-	const FName& MappedName = NameMap.IsValidIndex(NameIndex) ? NameMap[NameIndex] : NAME_None;
-	if (MappedName.IsNone())
+	int32 Number;
+	Ar << Number;
+
+	FNameEntryId MappedName = NameMap.IsValidIndex(NameIndex) ? NameMap[NameIndex] : FNameEntryId();
+	if (!MappedName)
 	{
-		int32 TempNumber;
-		Ar << TempNumber;
-		Name = NAME_None;
+		Name = FName();
 	}
 	else
 	{
-		int32 Number;
-		Ar << Number;
 		// simply create the name from the NameMap's name and the serialized instance number
-		Name = FName(MappedName, Number);
+		Name = FName::CreateFromDisplayId(MappedName, Number);
 	}
 
 	return *this;
@@ -290,7 +290,7 @@ void FNameTableArchiveWriter::SerializeNameMap()
 		for (auto& Pair : NameMap)
 		{
 			check(NameMapIdx == Pair.Value);
-			Pair.Key.GetDisplayNameEntry()->Write(*this);
+			FName::GetEntry(Pair.Key)->Write(*this);
 			NameMapIdx++;
 		}
 	}
@@ -375,15 +375,12 @@ void FNameTableArchiveWriter::ResetCustomVersions()
 
 FArchive& FNameTableArchiveWriter::operator<<( FName& Name )
 {
-	int32* NameIndexPtr = NameMap.Find(Name);
+	int32* NameIndexPtr = NameMap.Find(Name.GetDisplayIndex());
 	int32 NameIndex = NameIndexPtr ? *NameIndexPtr : INDEX_NONE;
 	if ( NameIndex == INDEX_NONE )
 	{
-		// We need to store the FName without the number, as the number is stored separately and we don't 
-		// want duplicate entries in the name table just because of the number
-		const FName NameNoNumber(Name, 0);
 		NameIndex = NameMap.Num();
-		NameMap.Add(NameNoNumber, NameIndex);
+		NameMap.Add(Name.GetDisplayIndex(), NameIndex);
 	}
 
 	FArchive& Ar = *this;
diff --git a/Engine/Source/Runtime/AssetRegistry/Private/NameTableArchive.h b/Engine/Source/Runtime/AssetRegistry/Private/NameTableArchive.h
--- a/Engine/Source/Runtime/AssetRegistry/Private/NameTableArchive.h
+++ b/Engine/Source/Runtime/AssetRegistry/Private/NameTableArchive.h
@@ -51,7 +51,7 @@ private:
 
 	FArchive* ProxyAr;
 	FArchive* FileAr;
-	TArray<FName> NameMap;
+	TArray<FNameEntryId> NameMap;
 };
 
 class FNameTableArchiveWriter : public FArchive
@@ -84,6 +84,6 @@ private:
 	FArchive* FileAr;
 	FString FinalFilename;
 	FString TempFilename;
-	TMap<FName, int32, FDefaultSetAllocator, TLinkerNameMapKeyFuncs<int32>> NameMap;
+	TMap<FNameEntryId, int32> NameMap;
 	int64 NameOffsetLoc;
 };
diff --git a/Engine/Source/Runtime/AssetRegistry/Private/PackageReader.cpp b/Engine/Source/Runtime/AssetRegistry/Private/PackageReader.cpp
--- a/Engine/Source/Runtime/AssetRegistry/Private/PackageReader.cpp
+++ b/Engine/Source/Runtime/AssetRegistry/Private/PackageReader.cpp
@@ -463,7 +463,7 @@ FArchive& FPackageReader::operator<<( FName& Name )
 {
 	check(Loader);
 
-	NAME_INDEX NameIndex;
+	int32 NameIndex;
 	FArchive& Ar = *this;
 	Ar << NameIndex;
 
diff --git a/Engine/Source/Runtime/AssetRegistry/Public/AssetDataTagMap.h b/Engine/Source/Runtime/AssetRegistry/Public/AssetDataTagMap.h
--- a/Engine/Source/Runtime/AssetRegistry/Public/AssetDataTagMap.h
+++ b/Engine/Source/Runtime/AssetRegistry/Public/AssetDataTagMap.h
@@ -113,18 +113,18 @@ private:
 		FString ToString() const
 		{
 			FString Result;
-			if (Class != 0)
+			if (Class)
 			{
 				FName(Class, Class, 0).AppendString(Result);
 				Result += TCHAR('\'');
 			}
 			FName(Package, Package, 0).AppendString(Result);
-			if (Object != 0)
+			if (Object)
 			{
 				Result += TCHAR('.');
 				FName(Object, Object, 0).AppendString(Result);
 			}
-			if (Class != 0)
+			if (Class)
 			{
 				Result += TCHAR('\'');
 			}
@@ -153,14 +153,13 @@ private:
 		FStorageID Result;
 		Result.Fields.IsFName = true;
 		FName ValueName(*Value);
-		if (ValueName.GetNumber() || ValueName.GetComparisonIndex() != ValueName.GetDisplayIndex() || ValueName.GetComparisonIndex() > FStorageID::MaxNoNumberFNameIndex)
+		if (ValueName.GetNumber() || ValueName.GetComparisonIndex() != ValueName.GetDisplayIndex() || ValueName.GetComparisonIndex().ToUnstableInt() > FStorageID::MaxNoNumberFNameIndex)
 		{
 			Result.Fields.Index = FNames.Add(ValueName);
 		}
 		else
 		{
-			check(ValueName.GetComparisonIndex() >= 0);
-			Result.Fields.Index = ValueName.GetComparisonIndex();
+			Result.Fields.Index = ValueName.GetComparisonIndex().ToUnstableInt();
 			Result.Fields.NoNumbers = true;
 		}
 		// there are cases where the results do not match on case
@@ -297,7 +296,8 @@ public:
 			{
 				return FNames.IsAllocated(Id.Fields.Index);
 			}
-			return FName(Id.Fields.Index, Id.Fields.Index, 0).IsValid();
+			FNameEntryId EntryId = FNameEntryId::FromUnstableInt(Id.Fields.Index);
+			return FName(EntryId, EntryId, 0).IsValid();
 		}
 		if (Id.Fields.IsFNameExportText)
 		{
@@ -332,7 +332,9 @@ public:
 			{
 				return TEXT("False");
 			}
-			return FName(Id.Fields.Index, Id.Fields.Index, 0).ToString();
+
+			FNameEntryId EntryId = FNameEntryId::FromUnstableInt(Id.Fields.Index);
+			return FName(EntryId, EntryId, 0).ToString();
 		}
 		if (Id.Fields.IsFNameExportText)
 		{
diff --git a/Engine/Source/Runtime/Core/Private/HAL/LowLevelMemTracker.cpp b/Engine/Source/Runtime/Core/Private/HAL/LowLevelMemTracker.cpp
--- a/Engine/Source/Runtime/Core/Private/HAL/LowLevelMemTracker.cpp
+++ b/Engine/Source/Runtime/Core/Private/HAL/LowLevelMemTracker.cpp
@@ -448,7 +448,7 @@ static int64 FNameToTag(FName Name)
 	}
 
 	// get the bits out of the FName we need
-	int64 NameIndex = Name.GetComparisonIndex();
+	int64 NameIndex = Name.GetComparisonIndex().ToUnstableInt();
 	int64 NameNumber = Name.GetNumber();
 	int64 tag = (NameNumber << 32) | NameIndex;
 	LLMCheckf(tag > (int64)ELLMTag::PlatformTagEnd, TEXT("Passed with a name index [%d - %s] that was less than MemTracker_MaxUserAllocation"), NameIndex, *Name.ToString());
@@ -460,7 +460,7 @@ static int64 FNameToTag(FName Name)
 static FName TagToFName(int64 Tag)
 {
 	// pull the bits back out of the tag
-	int32 NameIndex = (int32)(Tag & 0xFFFFFFFF);
+	FNameEntryId NameIndex = FNameEntryId::FromUnstableInt((int32)(Tag & 0xFFFFFFFF));
 	int32 NameNumber = (int32)(Tag >> 32);
 	return FName(NameIndex, NameIndex, NameNumber);
 }
diff --git a/Engine/Source/Runtime/Core/Private/Misc/CoreGlobals.cpp b/Engine/Source/Runtime/Core/Private/Misc/CoreGlobals.cpp
--- a/Engine/Source/Runtime/Core/Private/Misc/CoreGlobals.cpp
+++ b/Engine/Source/Runtime/Core/Private/Misc/CoreGlobals.cpp
@@ -181,7 +181,7 @@ bool					GExitPurge						= false;
 
 FChunkedFixedUObjectArray* GCoreObjectArrayForDebugVisualizers = nullptr;
 #if PLATFORM_UNIX
-FNameEntry*** CORE_API GFNameTableForDebuggerVisualizers_MT = FName::GetNameTableForDebuggerVisualizers_MT();
+uint8** CORE_API GNameBlocksDebug = FName::GetBlocksForDebugVisualizer();
 FChunkedFixedUObjectArray*& CORE_API GObjectArrayForDebugVisualizers = GCoreObjectArrayForDebugVisualizers;
 #endif
 
diff --git a/Engine/Source/Runtime/Core/Private/ProfilingDebugging/CsvProfiler.cpp b/Engine/Source/Runtime/Core/Private/ProfilingDebugging/CsvProfiler.cpp
--- a/Engine/Source/Runtime/Core/Private/ProfilingDebugging/CsvProfiler.cpp
+++ b/Engine/Source/Runtime/Core/Private/ProfilingDebugging/CsvProfiler.cpp
@@ -720,7 +720,7 @@ public:
 			if (bInIsFName)
 			{
 				check((InStatIDRaw & FNameOrIndexMask) == InStatIDRaw);
-				const FNameEntry* NameEntry = FName::GetEntry(UniqueID.Fields.FNameOrIndex);
+				const FNameEntry* NameEntry = FName::GetEntry(FNameEntryId::FromUnstableInt(UniqueID.Fields.FNameOrIndex));
 				NameStr = NameEntry->GetPlainNameString();
 			}
 			else
@@ -1425,7 +1425,7 @@ public:
 
 private:
 	static CSV_PROFILER_INLINE uint64 GetStatID(const char* StatName) { return uint64(StatName); }
-	static CSV_PROFILER_INLINE uint64 GetStatID(const FName& StatId) { return uint64(StatId.GetComparisonIndex()); }
+	static CSV_PROFILER_INLINE uint64 GetStatID(const FName& StatId) { return StatId.GetComparisonIndex().ToUnstableInt(); }
 
 	static inline FString GenerateThreadName(uint32 ThreadId)
 	{
diff --git a/Engine/Source/Runtime/Core/Private/Stats/StatsFile.cpp b/Engine/Source/Runtime/Core/Private/Stats/StatsFile.cpp
--- a/Engine/Source/Runtime/Core/Private/Stats/StatsFile.cpp
+++ b/Engine/Source/Runtime/Core/Private/Stats/StatsFile.cpp
@@ -330,15 +330,15 @@ void IStatsWriteFile::Finalize()
 	}
 
 	// Create a copy of names.
-	TSet<int32> FNamesToSent = FNamesSent;
+	TSet<FNameEntryId> FNamesToSent = FNamesSent;
 	FNamesSent.Empty( FNamesSent.Num() );
 
 	// Serialize FNames.
 	Header.FNameTableOffset = Ar.Tell();
 	Header.NumFNames = FNamesToSent.Num();
-	for (const int32 It : FNamesToSent)
+	for (FNameEntryId Id : FNamesToSent)
 	{
-		WriteFName( Ar, FStatNameAndInfo( FName( It, It, 0 ), false ) );
+		WriteFName( Ar, FStatNameAndInfo( FName( Id, Id, 0 ), false ) );
 	}
 
 	// Serialize metadata messages.
@@ -346,22 +346,6 @@ void IStatsWriteFile::Finalize()
 	Header.NumMetadataMessages = Stats.ShortNameToLongName.Num();
 	WriteMetadata( Ar );
 
-	// Verify data.
-	TSet<int32> BMinA = FNamesSent.Difference( FNamesToSent );
-	struct FLocal
-	{
-		static TArray<FName> GetFNameArray( const TSet<int32>& NameIndices )
-		{
-			TArray<FName> Result;
-			for (const int32 NameIndex : NameIndices)
-			{
-				new(Result)FName( NameIndex, NameIndex, 0 );
-			}
-			return Result;
-		}
-	};
-	TArray<FName> BMinANames = FLocal::GetFNameArray( BMinA );
-
 	// Seek to the position just after a magic value of the file and write out proper header.
 	Ar.Seek( sizeof( uint32 ) );
 	Ar << Header;
diff --git a/Engine/Source/Runtime/Core/Private/UObject/UnrealNames.cpp b/Engine/Source/Runtime/Core/Private/UObject/UnrealNames.cpp
--- a/Engine/Source/Runtime/Core/Private/UObject/UnrealNames.cpp
+++ b/Engine/Source/Runtime/Core/Private/UObject/UnrealNames.cpp
@@ -15,7 +15,7 @@
 #include "Misc/ByteSwap.h"
 #include "UObject/ObjectVersion.h"
 #include "HAL/ThreadSafeCounter.h"
-#include "Misc/ScopeLock.h"
+#include "Misc/ScopeRWLock.h"
 #include "Containers/Set.h"
 #include "Internationalization/Text.h"
 #include "Internationalization/Internationalization.h"
@@ -23,13 +23,878 @@
 #include "HAL/IConsoleManager.h"
 #include "HAL/LowLevelMemTracker.h"
 #include "Serialization/ArchiveFromStructuredArchive.h"
+#include "Hash/CityHash.h"
 
+DEFINE_LOG_CATEGORY_STATIC(LogUnrealNames, Log, All);
+
+const TCHAR* LexToString(EName Ename)
+{
+	switch (Ename)
+	{
+#define REGISTER_NAME(num,namestr) case num: return TEXT(#namestr);
+#include "UObject/UnrealNames.inl"
+#undef REGISTER_NAME
+		default:
+			return TEXT("*INVALID*");
+	}
+}
+
+int32 FNameEntry::GetDataOffset()
+{
+	const FNameEntry* Dummy = nullptr;
+	return (intptr_t)Dummy->AnsiName - (intptr_t)Dummy;
+}
+
+/*-----------------------------------------------------------------------------
+	FName helpers. 
+-----------------------------------------------------------------------------*/
+
+static bool operator==(FNameEntryHeader A, FNameEntryHeader B)
+{
+	static_assert(sizeof(FNameEntryHeader) == 2, "");
+	return (uint16&)A == (uint16&)B;
+}
+
+union FNameBuffer
+{
+	ANSICHAR AnsiName[NAME_SIZE];
+	WIDECHAR WideName[NAME_SIZE];
+
+	template<typename FromCharType, typename ToCharType>
+	ToCharType* ConvertInPlace(uint32 Len);
+};
+
+template<> ANSICHAR* FNameBuffer::ConvertInPlace<ANSICHAR, ANSICHAR>(uint32 Len)
+{
+	return AnsiName;
+}
+
+template<> WIDECHAR* FNameBuffer::ConvertInPlace<WIDECHAR, WIDECHAR>(uint32 Len)
+{
+	return WideName;
+}
+
+template<> WIDECHAR* FNameBuffer::ConvertInPlace<ANSICHAR, WIDECHAR>(uint32 Len)
+{
+	for (uint32 Index = Len; Index--; )
+	{
+		WideName[Index] = AnsiName[Index];
+	}
+
+	return WideName;
+}
+
+template<> ANSICHAR* FNameBuffer::ConvertInPlace<WIDECHAR, ANSICHAR>(uint32 Len)
+{
+	for (uint32 Index = 0; Index < Len; ++Index)
+	{
+		AnsiName[Index] = WideName[Index];
+	}
+
+	return AnsiName;
+}
+
+struct FNameStringView
+{
+	FNameStringView() : Data(nullptr), Len(0), bIsWide(false) {}
+	FNameStringView(const ANSICHAR* Str, uint32 Len_) : Ansi(Str), Len(Len_), bIsWide(false) {}
+	FNameStringView(const WIDECHAR* Str, uint32 Len_) : Wide(Str), Len(Len_), bIsWide(true) {}
+
+	union
+	{
+		const void* Data;
+		const ANSICHAR* Ansi;
+		const WIDECHAR* Wide;
+	};
+
+	uint32 Len;
+	bool bIsWide;
+
+	bool IsAnsi() const { return !bIsWide; }
+
+	int32 BytesWithTerminator() const
+	{
+		return (Len + 1) * (bIsWide ? sizeof(WIDECHAR) : sizeof(ANSICHAR));
+	}
+
+	int32 BytesWithoutTerminator() const
+	{
+		return Len * (bIsWide ? sizeof(WIDECHAR) : sizeof(ANSICHAR));
+	}
+};
+
+template<ENameCase Sensitivity>
+FORCEINLINE bool EqualsSameDimensions(FNameStringView A, FNameStringView B)
+{
+	checkSlow(A.Len == B.Len && A.IsAnsi() == B.IsAnsi());
+
+	int32 Len = A.Len;
+
+	if (Sensitivity == ENameCase::CaseSensitive)
+	{
+		return B.IsAnsi() ? !FPlatformString::Strncmp(A.Ansi, B.Ansi, Len) : !FPlatformString::Strncmp(A.Wide, B.Wide, Len);
+	}
+	else
+	{
+		return B.IsAnsi() ? !FPlatformString::Strnicmp(A.Ansi, B.Ansi, Len) : !FPlatformString::Strnicmp(A.Wide, B.Wide, Len);
+	}
+
+}
+
+template<ENameCase Sensitivity>
+FORCEINLINE bool Equals(FNameStringView A, FNameStringView B)
+{
+	return (A.Len == B.Len & A.IsAnsi() == B.IsAnsi()) && EqualsSameDimensions<Sensitivity>(A, B);
+}
+
+// Minimize stack lifetime of large decode buffers
+#ifdef WITH_CUSTOM_NAME_ENCODING
+#define OUTLINE_DECODE_BUFFER FORCENOINLINE
+#else
+#define OUTLINE_DECODE_BUFFER
+#endif
+
+template<ENameCase Sensitivity>
+OUTLINE_DECODE_BUFFER bool EqualsSameDimensions(const FNameEntry& Entry, FNameStringView Name)
+{
+	FNameBuffer DecodeBuffer;
+	return EqualsSameDimensions<Sensitivity>(Entry.MakeView(DecodeBuffer), Name);
+}
+
+/** Remember to update natvis if you change these */
+enum { FNameMaxBlockBits = 13 }; // Limit block array a bit, still allowing 8k * block size = 1GB - 2G of FName entry data
+enum { FNameBlockOffsetBits = 16 };
+enum { FNameMaxBlocks = 1 << FNameMaxBlockBits };
+enum { FNameBlockOffsets = 1 << FNameBlockOffsetBits };
+
+/** An unpacked FNameEntryId */
+struct FNameEntryHandle
+{
+	uint32 Block = 0;
+	uint32 Offset = 0;
+
+	FNameEntryHandle(uint32 InBlock, uint32 InOffset)
+		: Block(InBlock)
+		, Offset(InOffset)
+	{}
+
+	FNameEntryHandle(FNameEntryId Id)
+		: Block(Id.ToUnstableInt() >> FNameBlockOffsetBits)
+		, Offset(Id.ToUnstableInt() & (FNameBlockOffsets - 1))
+	{}
+
+	operator FNameEntryId() const
+	{
+		return FNameEntryId::FromUnstableInt(Block << FNameBlockOffsetBits | Offset);
+	}
+
+	explicit operator bool() const { return Block | Offset; }
+};
+
+static uint32 GetTypeHash(FNameEntryHandle Handle)
+{
+	return (Handle.Block << (32 - FNameMaxBlockBits)) + Handle.Block // Let block index impact most hash bits
+		+ (Handle.Offset << FNameBlockOffsetBits) + Handle.Offset // Let offset impact most hash bits
+		+ (Handle.Offset >> 4); // Reduce impact of non-uniformly distributed entry name lengths 
+}
+
+uint32 GetTypeHash(FNameEntryId Id)
+{
+	return GetTypeHash(FNameEntryHandle(Id));
+}
+
+FArchive& operator<<(FArchive& Ar, FNameEntryId& Id)
+{
+	if (Ar.IsLoading())
+	{
+		uint32 UnstableInt = 0;
+		Ar << UnstableInt;
+		Id = FNameEntryId::FromUnstableInt(UnstableInt);
+	}
+	else
+	{
+		uint32 UnstableInt = Id.ToUnstableInt();
+		Ar << UnstableInt;
+	}
+
+	return Ar;
+}
+
+FNameEntryId FNameEntryId::FromUnstableInt(uint32 Value)
+{
+	FNameEntryId Id;
+	Id.Value = Value;
+	return Id;
+}
+
+struct FNameSlot
+{
+	// Use the remaining few bits to store a hash that can determine inequality
+	// during probing without touching entry data
+	static constexpr uint32 EntryIdBits = FNameMaxBlockBits + FNameBlockOffsetBits;
+	static constexpr uint32 EntryIdMask = (1 << EntryIdBits) - 1;
+	static constexpr uint32 ProbeHashShift = EntryIdBits;
+	static constexpr uint32 ProbeHashMask = ~EntryIdMask;
+	
+	FNameSlot() {}
+	FNameSlot(FNameEntryId Value, uint32 ProbeHash)
+		: IdAndHash(Value.ToUnstableInt() | ProbeHash)
+	{
+		check(!(Value.ToUnstableInt() & ProbeHashMask) && !(ProbeHash & EntryIdMask) && Used());
+	}
+
+	FNameEntryId GetId() const { return FNameEntryId::FromUnstableInt(IdAndHash & EntryIdMask); }
+	uint32 GetProbeHash() const { return IdAndHash & ProbeHashMask; }
+	
+	bool operator==(FNameSlot Rhs) const { return IdAndHash == Rhs.IdAndHash; }
+
+	bool Used() const { return IdAndHash;  }
+private:
+	uint32 IdAndHash = 0;
+};
+
+/**
+ * Thread-safe paged FNameEntry allocator
+ */
+class FNameEntryAllocator
+{
+public:
+	enum { Stride = alignof(FNameEntry) };
+	enum { BlockSizeBytes = Stride * FNameBlockOffsets };
+
+	/** Initializes all member variables. */
+	FNameEntryAllocator()
+	{
+		Blocks[0] = (uint8*)FMemory::Malloc(BlockSizeBytes);
+	}
+
+	/**
+	 * Allocates the requested amount of bytes and returns an id that can be used to access them
+	 *
+	 * @param   Size  Size in bytes to allocate, 
+	 * @return  Allocation of passed in size cast to a FNameEntry pointer.
+	 */
+	FNameEntryHandle Allocate(uint32 Bytes)
+	{
+		Bytes = Align(Bytes, alignof(FNameEntry));
+		check(Bytes <= BlockSizeBytes);
+
+		FRWScopeLock _(Lock, FRWScopeLockType::SLT_Write);
+
+		// Allocate a new pool if current one is exhausted. We don't worry about a little bit
+		// of waste at the end given the relative size of pool to average and max allocation.
+		if (BlockSizeBytes - CurrentByteCursor < Bytes)
+		{
+			AllocateNewBlock();
+		}
+
+		// Use current cursor position for this allocation and increment cursor for next allocation
+		uint32 ByteOffset = CurrentByteCursor;
+		CurrentByteCursor += Bytes;
+		
+		check(ByteOffset % Stride == 0 && ByteOffset / Stride < FNameBlockOffsets);
+
+		return FNameEntryHandle(CurrentBlock, ByteOffset / Stride);
+	}
+
+	FNameEntryHandle Create(FNameStringView Name, FNameEntryId ComparisonId, FNameEntryHeader Header)
+	{
+		FNameEntryHandle Handle = Allocate(FNameEntry::GetDataOffset() + Name.BytesWithoutTerminator());
+		FNameEntry& Entry = Resolve(Handle);
+
+#if WITH_CASE_PRESERVING_NAME
+		Entry.ComparisonId = ComparisonId ? ComparisonId : FNameEntryId(Handle);
+#endif
+
+		Entry.Header = Header;
+		
+		if (Name.bIsWide)
+		{
+			Entry.StoreName(Name.Wide, Name.Len);
+		}
+		else
+		{
+			Entry.StoreName(Name.Ansi, Name.Len);
+		}
+
+		return Handle;
+	}
+
+	FNameEntry& Resolve(FNameEntryHandle Handle) const
+	{
+		// Lock not needed
+		return *reinterpret_cast<FNameEntry*>(Blocks[Handle.Block] + Stride * Handle.Offset);
+	}
+
+	/** Returns the number of blocks that have been allocated so far for names. */
+	uint32 NumBlocks() const
+	{
+		return CurrentBlock + 1;
+	}
+	
+	uint8** GetBlocksForDebugVisualizer() { return Blocks; }
+
+	void DebugDump(TArray<const FNameEntry*>& Out) const
+	{
+		FRWScopeLock _(Lock, FRWScopeLockType::SLT_ReadOnly);
+
+		for (uint32 BlockIdx = 0; BlockIdx < CurrentBlock; ++BlockIdx)
+		{
+			DebugDumpBlock(Blocks[BlockIdx], BlockSizeBytes, Out);
+		}
+
+		DebugDumpBlock(Blocks[CurrentBlock], CurrentByteCursor, Out);
+	}
+
+private:
+	static void DebugDumpBlock(const uint8* It, uint32 BlockSize, TArray<const FNameEntry*>& Out)
+	{
+		const uint8* End = It + BlockSize - FNameEntry::GetDataOffset();
+		while (It < End)
+		{
+			const FNameEntry* Entry = (const FNameEntry*)It;
+			if (uint32 Len = Entry->Header.Len)
+			{
+				Out.Add(Entry);
+				It += FNameEntry::GetSize(Len, !Entry->IsWide());
+			}
+			else // Null-terminator entry found
+			{
+				break;
+			}
+		}
+	}
+
+	/** Allocates a new pool. */
+	void AllocateNewBlock()
+	{
+		// Null-terminate final entry to allow DebugDump() entry iteration
+		if (CurrentByteCursor + FNameEntry::GetDataOffset() <= BlockSizeBytes)
+		{
+			FNameEntry* Terminator = (FNameEntry*)(Blocks[CurrentBlock] + CurrentByteCursor);
+			Terminator->Header.Len = 0;
+		}
+
+		++CurrentBlock;
+		CurrentByteCursor = 0;
+
+		check(CurrentBlock < FNameMaxBlocks);
+		check(Blocks[CurrentBlock] == nullptr);
+
+		Blocks[CurrentBlock] = (uint8*)FMemory::Malloc(BlockSizeBytes);
+	}
+
+	mutable FRWLock Lock;
+	uint32 CurrentBlock = 0;
+	uint32 CurrentByteCursor = 0;
+	uint8* Blocks[FNameMaxBlocks] = {};
+};
+
+enum { FNamePoolShardBits = 4 };
+enum { FNamePoolShards = 1 << FNamePoolShardBits };
+enum { FNamePoolInitialSlotBits = 8 };
+enum { FNamePoolInitialSlots = 1 << FNamePoolInitialSlotBits };
+
+/** Hashes name into 64 bits that determines shard and slot index.
+ *	
+ *	Small parts	of the hash is also stored in unused bits of the slot and entry. 
+ *	The former optimizes linear probing by accessing less entry data.
+ *	The latter optimizes linear probing by avoiding copying and deobfuscating entry data.
+ *
+ *	The slot index could be stored in the slot, at least in non shipping / test configs.
+ *	This costs memory by doubling slot size but would essentially never touch entry data
+ *	nor copy and deobfuscate a name needlessy. It also allows growing the hash table
+ *	without rehashing the strings, since the unmasked slot index would be known.
+ */
+struct FNameHash
+{
+	uint32 ShardIndex;
+	uint32 UnmaskedSlotIndex; // Determines at what slot index to start probing
+	uint32 SlotProbeHash; // Helps cull equality checks (decode + strnicmp) when probing slots
+	FNameEntryHeader EntryProbeHeader; // Helps cull equality checks when probing inspects entries
+
+
+	template<class CharType>
+	FNameHash(const CharType* Str, int32 Len)
+	{
+		uint64 Hash = CityHash64(reinterpret_cast<const char*>(Str), Len * sizeof(CharType));
+		uint32 Hi = static_cast<uint32>(Hash >> 32);
+		uint32 Lo = static_cast<uint32>(Hash);
+
+		// "None" has FNameEntryId with a value of zero
+		// Always set a bit in SlotProbeHash for "None" to distinguish unused slot values from None
+		// @see FNameSlot::Used()
+		uint32 IsNoneBit = IsAnsiNone(Str, Len) << FNameSlot::ProbeHashShift;
+
+		static constexpr uint32 ShardMask = FNamePoolShards - 1;
+		static_assert((ShardMask & FNameSlot::ProbeHashMask) == 0, "Masks overlap");
+
+		ShardIndex = Hi & ShardMask;
+		UnmaskedSlotIndex = Lo;
+		SlotProbeHash = (Hi & FNameSlot::ProbeHashMask) | IsNoneBit;
+		EntryProbeHeader.Len = Len;
+		EntryProbeHeader.bIsWide = sizeof(CharType) == sizeof(WIDECHAR);
+
+		// When we always use lowercase hashing, we can store parts of the hash in the entry
+		// to avoid copying and decoding entries needlessly. WITH_CUSTOM_NAME_ENCODING
+		// that makes this important is normally on when WITH_CASE_PRESERVING_NAME is off.
+#if !WITH_CASE_PRESERVING_NAME		
+		static constexpr uint32 EntryProbeMask = (1u << FNameEntryHeader::ProbeHashBits) - 1; 
+		EntryProbeHeader.LowercaseProbeHash = static_cast<uint16>((Hi >> FNamePoolShardBits) & EntryProbeMask);
+#endif
+	}
+	
+	uint32 GetProbeStart(uint32 SlotMask) const
+	{
+		return UnmaskedSlotIndex & SlotMask;
+	}
+
+	static uint32 GetProbeStart(uint32 UnmaskedSlotIndex, uint32 SlotMask)
+	{
+		return UnmaskedSlotIndex & SlotMask;
+	}
+
+	static uint32 IsAnsiNone(const WIDECHAR* Str, int32 Len)
+	{
+		return 0;
+	}
+
+	static uint32 IsAnsiNone(const ANSICHAR* Str, int32 Len)
+	{
+		if (Len != 4)
+		{
+			return 0;
+		}
+
+#if PLATFORM_LITTLE_ENDIAN
+		static constexpr uint32 NoneAsInt = 0x454e4f4e;
+#else
+		static constexpr uint32 NoneAsInt = 0x4e4f4e45;
+#endif
+		static constexpr uint32 ToUpperMask = 0xdfdfdfdf;
+
+		uint32 FourChars = FPlatformMemory::ReadUnaligned<uint32>(Str);
+		return (FourChars & ToUpperMask) == NoneAsInt;
+	}
+};
+
+
+template<class CharType>
+FORCENOINLINE FNameHash HashLowerCase(const CharType* Str, uint32 Len)
+{
+	CharType LowerStr[NAME_SIZE];
+	for (uint32 I = 0; I < Len; ++I)
+	{
+		LowerStr[I] = TChar<CharType>::ToLower(Str[I]);
+	}
+	return FNameHash(LowerStr, Len);
+}
+
+template<ENameCase Sensitivity>
+FNameHash HashName(FNameStringView Name);
+
+template<>
+FNameHash HashName<ENameCase::IgnoreCase>(FNameStringView Name)
+{
+	return Name.IsAnsi() ? HashLowerCase(Name.Ansi, Name.Len) : HashLowerCase(Name.Wide, Name.Len);
+}
+template<>
+FNameHash HashName<ENameCase::CaseSensitive>(FNameStringView Name)
+{
+	return Name.IsAnsi() ? FNameHash(Name.Ansi, Name.Len) : FNameHash(Name.Wide, Name.Len);
+}
+
+template<ENameCase Sensitivity>
+struct FNameValue
+{
+	FNameValue(FNameStringView InName)
+		: Name(InName)
+		, Hash(HashName<Sensitivity>(InName))
+	{}
+
+	FNameStringView Name;
+	FNameHash Hash;
+	FNameEntryId ComparisonId;
+};
+
+using FNameComparisonValue = FNameValue<ENameCase::IgnoreCase>;
+#if WITH_CASE_PRESERVING_NAME
+using FNameDisplayValue = FNameValue<ENameCase::CaseSensitive>;
+#endif
+
+
+class alignas(PLATFORM_CACHE_LINE_SIZE) FNamePoolShardBase : FNoncopyable
+{
+public:
+	void Initialize(FNameEntryAllocator& InEntries)
+	{
+		Entries = &InEntries;
+
+		Slots = (FNameSlot*)FMemory::Malloc(FNamePoolInitialSlots * sizeof(FNameSlot), alignof(FNameSlot));
+		memset(Slots, 0, FNamePoolInitialSlots * sizeof(FNameSlot));
+		CapacityMask = FNamePoolInitialSlots - 1;
+	}
+
+	uint32 Capacity() const { return CapacityMask + 1; }
+
+protected:
+	enum { LoadFactorQuotient = 9, LoadFactorDivisor = 10 }; // I.e. realloc slots when 90% full
+
+	mutable FRWLock Lock;
+	uint32 UsedSlots = 0;
+	uint32 CapacityMask = 0;
+	FNameSlot* Slots = nullptr;
+	FNameEntryAllocator* Entries = nullptr;
+
+
+	template<ENameCase Sensitivity>
+	FORCEINLINE static bool EntryEqualsValue(const FNameEntry& Entry, const FNameValue<Sensitivity>& Value)
+	{
+		return Entry.Header == Value.Hash.EntryProbeHeader && EqualsSameDimensions<Sensitivity>(Entry, Value.Name);
+	}
+};
+
+template<ENameCase Sensitivity>
+class FNamePoolShard : public FNamePoolShardBase
+{
+public:
+	FNameEntryId Find(const FNameValue<Sensitivity>& Value) const
+	{
+		FRWScopeLock _(Lock, FRWScopeLockType::SLT_ReadOnly);
+
+		return Probe(Value).GetId();
+	}
+
+	FNameEntryId Insert(const FNameValue<Sensitivity>& Value, bool& bCreatedNewEntry)
+	{
+		FRWScopeLock _(Lock, FRWScopeLockType::SLT_Write);
+
+		FNameSlot& Slot = Probe(Value);
+
+		if (Slot.Used())
+		{
+			return Slot.GetId();
+		}
+
+		FNameEntryId NewEntryId = Entries->Create(Value.Name, Value.ComparisonId, Value.Hash.EntryProbeHeader);
+
+		ClaimSlot(Slot, FNameSlot(NewEntryId, Value.Hash.SlotProbeHash));
+
+		bCreatedNewEntry = true;
+
+		return NewEntryId;
+	}
+
+	void InsertExistingEntry(FNameHash Hash, FNameEntryId ExistingId)
+	{
+		FNameSlot NewLookup(ExistingId, Hash.SlotProbeHash);
+
+		FRWScopeLock _(Lock, FRWScopeLockType::SLT_Write);
+		 
+		FNameSlot& Slot = Probe(Hash.UnmaskedSlotIndex, [=](FNameSlot Old) { return Old == NewLookup; });
+		if (!Slot.Used())
+		{
+			ClaimSlot(Slot, NewLookup);
+		}
+	}
+
+private:
+	void ClaimSlot(FNameSlot& UnusedSlot, FNameSlot NewValue)
+	{
+		UnusedSlot = NewValue;
+
+		++UsedSlots;
+		if (UsedSlots * LoadFactorDivisor >= LoadFactorQuotient * Capacity())
+		{
+			Grow();
+		}
+	}
+
+	void Grow()
+	{
+		FNameSlot* const OldSlots = Slots;
+		const uint32 OldUsedSlots = UsedSlots;
+		const uint32 OldCapacity = Capacity();
+		const uint32 NewCapacity = OldCapacity * 2;
+
+		Slots = (FNameSlot*)FMemory::Malloc(NewCapacity * sizeof(FNameSlot), alignof(FNameSlot));
+		memset(Slots, 0, NewCapacity * sizeof(FNameSlot));
+		UsedSlots = 0;
+		CapacityMask = NewCapacity - 1;
+
+
+		for (uint32 OldIdx = 0; OldIdx < OldCapacity; ++OldIdx)
+		{
+			const FNameSlot& OldSlot = OldSlots[OldIdx];
+			if (OldSlot.Used())
+			{
+				FNameHash Hash = Rehash(OldSlot.GetId());
+				FNameSlot& NewSlot = Probe(Hash.UnmaskedSlotIndex, [](FNameSlot Slot) { return false; });
+				NewSlot = OldSlot;
+				++UsedSlots;
+			}
+		}
+
+		check(OldUsedSlots == UsedSlots);
+
+		FMemory::Free(OldSlots);
+	}
+
+	/** Find slot containing value or the first free slot that should be used to store it  */
+	FNameSlot& Probe(const FNameValue<Sensitivity>& Value) const
+	{
+		return Probe(Value.Hash.UnmaskedSlotIndex, 
+			[&](FNameSlot Slot)	{ return Slot.GetProbeHash() == Value.Hash.SlotProbeHash && 
+									EntryEqualsValue<Sensitivity>(Entries->Resolve(Slot.GetId()), Value); });
+	}
+
+	/** Find slot that fulfills predicate or the first free slot  */
+	template<class PredicateFn>
+	FNameSlot& Probe(uint32 UnmaskedSlotIndex, PredicateFn Predicate) const
+	{
+		const uint32 Mask = CapacityMask;
+		for (uint32 I = FNameHash::GetProbeStart(UnmaskedSlotIndex, Mask); true; I = (I + 1) & Mask)
+		{
+			FNameSlot& Slot = Slots[I];
+			if (!Slot.Used() || Predicate(Slot))
+			{
+				return Slot;
+			}
+		}
+	}
+
+	OUTLINE_DECODE_BUFFER FNameHash Rehash(FNameEntryId EntryId)
+	{
+		const FNameEntry& Entry = Entries->Resolve(EntryId);
+		FNameBuffer DecodeBuffer;
+		return HashName<Sensitivity>(Entry.MakeView(DecodeBuffer));
+	}
+};
+
+
+class FNamePool
+{
+public:
+	FNamePool();
+
+	FNameEntryId	Store(FNameStringView View);
+	FNameEntryId	Find(FNameStringView View) const;
+	FNameEntryId	Find(EName Ename) const;
+	const EName*	FindEName(FNameEntryId Id) const;
+
+	/** @pre !!Handle */
+	FNameEntry&		Resolve(FNameEntryHandle Handle) const { return Entries.Resolve(Handle); }
+
+	bool			IsValid(FNameEntryHandle Handle) const;
+		
+	/// Stats and debug related functions ///
+
+	uint32			NumAnsiEntries() const { return AnsiCount; }
+	uint32			NumWideEntries() const { return WideCount; }
+	uint32			NumBlocks() const { return Entries.NumBlocks(); }
+	uint32			NumSlots() const;
+	void			LogStats(FOutputDevice& Ar) const;
+	uint8**			GetBlocksForDebugVisualizer() { return Entries.GetBlocksForDebugVisualizer(); }
+	TArray<const FNameEntry*> DebugDump() const;
+
+private:
+	enum { MaxENames = 512 };
+
+	FNameEntryAllocator Entries;
+	TAtomic<uint32> AnsiCount;
+	TAtomic<uint32> WideCount;
+
+#if WITH_CASE_PRESERVING_NAME
+	FNamePoolShard<ENameCase::CaseSensitive> DisplayShards[FNamePoolShards];
+#endif
+	FNamePoolShard<ENameCase::IgnoreCase> ComparisonShards[FNamePoolShards];
+
+	// Put constant lookup on separate cache line to avoid it being constantly invalidated by insertion
+	alignas(PLATFORM_CACHE_LINE_SIZE) FNameEntryId ENameToEntry[NAME_MaxHardcodedNameIndex] = {};
+	TMap<FNameEntryId, EName, TInlineSetAllocator<MaxENames>> EntryToEName;
+};
+
+FNamePool::FNamePool()
+	: AnsiCount(0)
+	, WideCount(0)
+{
+	for (FNamePoolShardBase& Shard : ComparisonShards)
+	{
+		Shard.Initialize(Entries);
+	}
+
+#if WITH_CASE_PRESERVING_NAME
+	for (FNamePoolShardBase& Shard : DisplayShards)
+	{
+		Shard.Initialize(Entries);
+	}
+#endif
+
+	// Register all hardcoded names
+#define REGISTER_NAME(num, name) ENameToEntry[num] = Store(FNameStringView(#name, FCStringAnsi::Strlen(#name)));
+#include "UObject/UnrealNames.inl"
+#undef REGISTER_NAME
+
+	// Make reverse mapping
+	for (uint32 ENameIndex = 0; ENameIndex < NAME_MaxHardcodedNameIndex; ++ENameIndex)
+	{
+		if (ENameIndex == NAME_None || ENameToEntry[ENameIndex])
+		{
+			EntryToEName.Add(ENameToEntry[ENameIndex], (EName)ENameIndex);
+		}
+	}
+
+	// Verify all ENames are unique
+	if (NumAnsiEntries() != EntryToEName.Num())
+	{
+		// we can't print out here because there may be no log yet if this happens before main starts
+		if (FPlatformMisc::IsDebuggerPresent())
+		{
+			UE_DEBUG_BREAK();
+		}
+		else
+		{
+			FPlatformMisc::PromptForRemoteDebugging(false);
+			FMessageDialog::Open(EAppMsgType::Ok, NSLOCTEXT("UnrealEd", "DuplicatedHardcodedName", "Duplicate hardcoded name"));
+			FPlatformMisc::RequestExit(false);
+		}
+	}
+}
+
+static bool IsPureAnsi(const WIDECHAR* Str, const int32 Len)
+{
+	// Consider SSE version if this function takes significant amount of time
+	uint32 Result = 0;
+	for (int32 I = 0; I < Len; ++I)
+	{
+		Result |= TChar<WIDECHAR>::ToUnsigned(Str[I]);
+	}
+	return !(Result & 0xffffff80u);
+}
+
+FNameEntryId FNamePool::Find(EName Ename) const
+{
+	check(Ename < NAME_MaxHardcodedNameIndex);
+	return ENameToEntry[Ename];
+}
+
+FNameEntryId FNamePool::Find(FNameStringView Name) const
+{
+#if WITH_CASE_PRESERVING_NAME
+	FNameDisplayValue DisplayValue(Name);
+	if (FNameEntryId Existing = DisplayShards[DisplayValue.Hash.ShardIndex].Find(DisplayValue))
+	{
+		return Existing;
+	}
+#endif
+
+	FNameComparisonValue ComparisonValue(Name);
+	return ComparisonShards[ComparisonValue.Hash.ShardIndex].Find(ComparisonValue);
+}
+
+FNameEntryId FNamePool::Store(FNameStringView Name)
+{
+#if WITH_CASE_PRESERVING_NAME
+	FNameDisplayValue DisplayValue(Name);
+	FNamePoolShard<ENameCase::CaseSensitive>& DisplayShard = DisplayShards[DisplayValue.Hash.ShardIndex];
+	if (FNameEntryId Existing = DisplayShard.Find(DisplayValue))
+	{
+		return Existing;
+	}
+#endif
+
+	TAtomic<uint32>& EntryCount = Name.IsAnsi() ? AnsiCount : WideCount;
+	bool bAdded = false;
+
+	// Insert comparison name first since display value must contain comparison name
+	FNameComparisonValue ComparisonValue(Name);
+	FNameEntryId ComparisonId = ComparisonShards[ComparisonValue.Hash.ShardIndex].Insert(ComparisonValue, bAdded);
+	EntryCount += bAdded;
+
+#if WITH_CASE_PRESERVING_NAME
+	// Check if ComparisonId can be used as DisplayId
+	if (bAdded || EqualsSameDimensions<ENameCase::CaseSensitive>(Resolve(ComparisonId), Name))
+	{
+		DisplayShard.InsertExistingEntry(DisplayValue.Hash, ComparisonId);
+		return ComparisonId;
+	}
+	else
+	{
+		bAdded = false;
+		DisplayValue.ComparisonId = ComparisonId;
+		FNameEntryId DisplayId = DisplayShard.Insert(DisplayValue, bAdded);
+		EntryCount += bAdded;
+
+		return DisplayId;
+	}
+#else
+	return ComparisonId;
+#endif
+}
+
+uint32 FNamePool::NumSlots() const
+{
+	uint32 SlotCapacity = 0;
+#if WITH_CASE_PRESERVING_NAME
+	for (const FNamePoolShardBase& Shard : DisplayShards)
+	{
+		SlotCapacity += Shard.Capacity();
+	}
+#endif
+	for (const FNamePoolShardBase& Shard : ComparisonShards)
+	{
+		SlotCapacity += Shard.Capacity();
+	}
+
+	return SlotCapacity;
+}
+
+void FNamePool::LogStats(FOutputDevice& Ar) const
+{
+	Ar.Logf(TEXT("%i FNames using in %ikB + %ikB"), WideCount + AnsiCount, sizeof(FNamePool), Entries.NumBlocks() * FNameEntryAllocator::BlockSizeBytes / 1024);
+}
+
+TArray<const FNameEntry*> FNamePool::DebugDump() const
+{
+	TArray<const FNameEntry*> Out;
+	Out.Reserve(WideCount + AnsiCount);
+	Entries.DebugDump(Out);
+	return Out;
+}
+
+bool FNamePool::IsValid(FNameEntryHandle Handle) const
+{
+	return Handle.Block < Entries.NumBlocks();
+}
 
-DEFINE_LOG_CATEGORY_STATIC(LogUnrealNames, Log, All);
+const EName* FNamePool::FindEName(FNameEntryId Id) const
+{
+	return EntryToEName.Find(Id);
+}
+
+static bool bNamePoolInitialized;
+alignas(FNamePool) static uint8 NamePoolData[sizeof(FNamePool)];
+
+// Only call this once per public FName function called
+//
+// Not using magic statics to run as little code as possible
+static FNamePool& GetNamePool()
+{
+	if (bNamePoolInitialized)
+	{
+		return *(FNamePool*)NamePoolData;
+	}
+
+	FNamePool* Singleton = new (NamePoolData) FNamePool;
+	bNamePoolInitialized = true;
+	return *Singleton;
+}
+
+bool operator==(FNameEntryId Id, EName Ename)
+{
+	return Id == GetNamePool().Find(Ename);
+}
 
-/*-----------------------------------------------------------------------------
-	FName helpers. 
------------------------------------------------------------------------------*/
 #if !UE_BUILD_SHIPPING && !UE_BUILD_TEST
 	void CallNameCreationHook();
 #else
@@ -38,11 +903,7 @@ DEFINE_LOG_CATEGORY_STATIC(LogUnrealNames, Log, All);
 	}
 #endif
 
-// Get the size of a FNameEntry without the union buffer included
-static const int32 NameEntryWithoutUnionSize = sizeof(FNameEntry) - (NAME_SIZE * sizeof(TCHAR));
-
-template<typename TCharType>
-FNameEntry* AllocateNameEntry( const TCharType* Name, NAME_INDEX Index);
+static FNameEntryId DebugCastNameEntryId(int32 Id) { return (FNameEntryId&)(Id); }
 
 /**
 * Helper function that can be used inside the debuggers watch window. E.g. "DebugFName(Class->Name.Index)". 
@@ -50,10 +911,10 @@ FNameEntry* AllocateNameEntry( const TCharType* Name, NAME_INDEX Index);
 * @param	Index	Name index to look up string for
 * @return			Associated name
 */
-const TCHAR* DebugFName(int32 Index)
+const TCHAR* DebugFName(NAME_INDEX Index)
 {
 	// Hardcoded static array. This function is only used inside the debugger so it should be fine to return it.
-	static TCHAR TempName[256];
+	static TCHAR TempName[NAME_SIZE];
 	FCString::Strcpy(TempName, *FName::SafeString(Index));
 	return TempName;
 }
@@ -68,8 +929,8 @@ const TCHAR* DebugFName(int32 Index)
 const TCHAR* DebugFName(int32 Index, int32 Number)
 {
 	// Hardcoded static array. This function is only used inside the debugger so it should be fine to return it.
-	static TCHAR TempName[256];
-	FCString::Strcpy(TempName, *FName::SafeString(Index, Number));
+	static TCHAR TempName[NAME_SIZE];
+	FCString::Strcpy(TempName, *FName::SafeString(DebugCastNameEntryId(Index), Number));
 	return TempName;
 }
 
@@ -82,7 +943,7 @@ const TCHAR* DebugFName(int32 Index, int32 Number)
 const TCHAR* DebugFName(FName& Name)
 {
 	// Hardcoded static array. This function is only used inside the debugger so it should be fine to return it.
-	static TCHAR TempName[256];
+	static TCHAR TempName[NAME_SIZE];
 	FCString::Strcpy(TempName, *FName::SafeString(Name.GetDisplayIndex(), Name.GetNumber()));
 	return TempName;
 }
@@ -99,162 +960,134 @@ static uint16 GetRawNonCasePreservingHash(const TCharType* Source)
 	return FCrc::Strihash_DEPRECATED(Source) & 0xFFFF;
 }
 
-/*----------------------------------------------------------------------------
-	FNameBuffer.
-----------------------------------------------------------------------------*/
-
-union FNameBuffer
-{
-	ANSICHAR AnsiName[NAME_SIZE];
-	WIDECHAR WideName[NAME_SIZE];
-};
-
 /*-----------------------------------------------------------------------------
 	FNameEntry
 -----------------------------------------------------------------------------*/
 
-/**
- * @return FString of name portion minus number.
- */
-FString FNameEntry::GetPlainNameString() const
+void FNameEntry::StoreName(const ANSICHAR* InName, uint32 Len)
 {
-	FNameBuffer TempBuffer;
-	if( IsWide() )
-	{
-		return FString(GetWideNamePtr(TempBuffer.WideName));
-	}
-	else
-	{
-		return FString(GetAnsiNamePtr(TempBuffer.AnsiName));
-	}
+	memcpy(AnsiName, InName, sizeof(ANSICHAR) * Len);
+	Encode(AnsiName, Len);
 }
 
-/**
- * Appends this name entry to the passed in string.
- *
- * @param	String	String to append this name to
- */
-void FNameEntry::AppendNameToString( FString& String ) const
+void FNameEntry::StoreName(const WIDECHAR* InName, uint32 Len)
 {
-	FNameBuffer TempBuffer;
-	if( IsWide() )
-	{
-		String += GetWideNamePtr(TempBuffer.WideName);
-	}
-	else
-	{
-		String += GetAnsiNamePtr(TempBuffer.AnsiName);
-	}
+	memcpy(WideName, InName, sizeof(WIDECHAR) * Len);
+	Encode(WideName, Len);
 }
 
-void FNameEntry::AppendNameToPathString(FString& String) const
+void FNameEntry::CopyUnterminatedName(ANSICHAR* Out) const
 {
-	FNameBuffer TempBuffer;
-	if (IsWide())
-	{
-		String /= GetWideNamePtr(TempBuffer.WideName);
-	}
-	else
-	{
-		String /= GetAnsiNamePtr(TempBuffer.AnsiName);
-	}
+	memcpy(Out, AnsiName, sizeof(ANSICHAR) * Header.Len);
+	Decode(Out, Header.Len);
 }
 
-/**
- * @return length of name
- */
-int32 FNameEntry::GetNameLength() const
+void FNameEntry::CopyUnterminatedName(WIDECHAR* Out) const
 {
-	FNameBuffer TempBuffer;
-	if( IsWide() )
-	{
-		return FCStringWide::Strlen( GetWideNamePtr(TempBuffer.WideName) );
-	}
-	else
-	{
-		return FCStringAnsi::Strlen( GetAnsiNamePtr(TempBuffer.AnsiName) );
-	}
+	memcpy(Out, WideName, sizeof(WIDECHAR) * Header.Len);
+	Decode(Out, Header.Len);
 }
 
-/**
- * Compares name using the compare method provided.
- *
- * @param	InName	Name to compare to
- * @return	true if equal, false otherwise
- */
-bool FNameEntry::IsEqual( const ANSICHAR* InName, const ENameCase CompareMethod ) const
+FORCEINLINE const WIDECHAR* FNameEntry::GetUnterminatedName(WIDECHAR(&OptionalDecodeBuffer)[NAME_SIZE]) const
 {
-	if( IsWide() )
-	{
-		// Matching wide-ness means strings are not equal.
-		return false;
-	}
-	else
-	{
-		ANSICHAR TempBuffer[NAME_SIZE];
-		const ANSICHAR* CompareName = GetAnsiNamePtr(TempBuffer);
-		return ( (CompareMethod == ENameCase::CaseSensitive) ? FCStringAnsi::Strcmp( CompareName, InName ) : FCStringAnsi::Stricmp( CompareName, InName ) ) == 0;
-	}
+#ifdef WITH_CUSTOM_NAME_ENCODING
+	CopyUnterminatedName(OptionalDecodeBuffer);
+	return OptionalDecodeBuffer;
+#else
+	return WideName;
+#endif
 }
 
-/**
- * Compares name using the compare method provided.
- *
- * @param	InName	Name to compare to
- * @return	true if equal, false otherwise
- */
-bool FNameEntry::IsEqual( const WIDECHAR* InName, const ENameCase CompareMethod ) const
+FORCEINLINE ANSICHAR const* FNameEntry::GetUnterminatedName(ANSICHAR(&OptionalDecodeBuffer)[NAME_SIZE]) const
+{
+#ifdef WITH_CUSTOM_NAME_ENCODING
+	CopyUnterminatedName(OptionalDecodeBuffer);
+	return OptionalDecodeBuffer;
+#else
+	return AnsiName;
+#endif
+}
+
+FORCEINLINE FNameStringView FNameEntry::MakeView(FNameBuffer& OptionalDecodeBuffer) const
+{
+	return IsWide()	? FNameStringView(GetUnterminatedName(OptionalDecodeBuffer.WideName), GetNameLength())
+					: FNameStringView(GetUnterminatedName(OptionalDecodeBuffer.AnsiName), GetNameLength());
+}
+
+void FNameEntry::GetAnsiName(ANSICHAR(&Out)[NAME_SIZE]) const
+{
+	check(!IsWide());
+	CopyUnterminatedName(Out);
+	Out[Header.Len] = '\0';
+}
+
+void FNameEntry::GetWideName(WIDECHAR(&Out)[NAME_SIZE]) const
 {
-	if( !IsWide() )
+	check(IsWide());
+	CopyUnterminatedName(Out);
+	Out[Header.Len] = '\0';
+}
+
+/** @return null-terminated string */
+static const TCHAR* EntryToCString(const FNameEntry& Entry, FNameBuffer& Temp)
+{
+	if (Entry.IsWide())
 	{
-		// Matching wide-ness means strings are not equal.
-		return false;
+		Entry.GetWideName(Temp.WideName);
+		return Temp.ConvertInPlace<WIDECHAR, TCHAR>(Entry.GetNameLength() + 1);
 	}
 	else
 	{
-		WIDECHAR TempBuffer[NAME_SIZE];
-		const WIDECHAR* CompareName = GetWideNamePtr(TempBuffer);
-		return ( (CompareMethod == ENameCase::CaseSensitive) ? FCStringWide::Strcmp( CompareName, InName ) : FCStringWide::Stricmp( CompareName, InName ) ) == 0;
+		Entry.GetAnsiName(Temp.AnsiName);
+		return Temp.ConvertInPlace<ANSICHAR, TCHAR>(Entry.GetNameLength() + 1);
 	}
 }
 
-/**
- * Returns the size in bytes for FNameEntry structure. This is != sizeof(FNameEntry) as we only allocated as needed.
- *
- * @param	Name	Name to determine size for
- * @return	required size of FNameEntry structure to hold this string (might be wide or ansi)
- */
-int32 FNameEntry::GetSize( const TCHAR* Name )
+FString FNameEntry::GetPlainNameString() const
 {
-	return FNameEntry::GetSize( FCString::Strlen( Name ), FCString::IsPureAnsi( Name ) );
+	FNameBuffer Temp;
+	return FString(EntryToCString(*this, Temp));
 }
 
-/**
- * Returns the size in bytes for FNameEntry structure. This is != sizeof(FNameEntry) as we only allocated as needed.
- *
- * @param	Length			Length of name
- * @param	bIsPureAnsi		Whether name is pure ANSI or not
- * @return	required size of FNameEntry structure to hold this string (might be wide or ansi)
- */
-int32 FNameEntry::GetSize( int32 Length, bool bIsPureAnsi )
+void FNameEntry::AppendNameToString(FString& Out) const
+{
+	FNameBuffer Temp;
+	Out.Append(EntryToCString(*this, Temp), Header.Len);
+}
+
+void FNameEntry::AppendNameToPathString(FString& Out) const
+{
+	FNameBuffer Temp;
+	Out.PathAppend(EntryToCString(*this, Temp), Header.Len);
+}
+
+int32 FNameEntry::GetSize(const TCHAR* Name)
+{
+	return FNameEntry::GetSize(FCString::Strlen(Name), FCString::IsPureAnsi(Name));
+}
+
+int32 FNameEntry::GetSize(int32 Length, bool bIsPureAnsi)
 {
-	// Add size required for string to the base size used by the FNameEntry
-	int32 Size = NameEntryWithoutUnionSize + (Length + 1) * (bIsPureAnsi ? sizeof(ANSICHAR) : sizeof(TCHAR));
-	return Size;
+	int32 Bytes = GetDataOffset() + (Length + 1) * (bIsPureAnsi ? sizeof(ANSICHAR) : sizeof(TCHAR));
+	return Align(Bytes, alignof(FNameEntry));
+}
+
+int32 FNameEntry::GetSizeInBytes() const
+{
+	return GetSize(GetNameLength(), !IsWide());
 }
 
 FNameEntrySerialized::FNameEntrySerialized(const FNameEntry& NameEntry)
 {
-	if (NameEntry.IsWide())
+	bIsWide = NameEntry.IsWide();
+	if (bIsWide)
 	{
-		PreSetIsWideForSerialization(true);
 		NameEntry.GetWideName(WideName);
 		NonCasePreservingHash = GetRawNonCasePreservingHash(WideName);
 		CasePreservingHash = GetRawCasePreservingHash(WideName);
 	}
 	else
 	{
-		PreSetIsWideForSerialization(false);
 		NameEntry.GetAnsiName(AnsiName);
 		NonCasePreservingHash = GetRawNonCasePreservingHash(AnsiName);
 		CasePreservingHash = GetRawCasePreservingHash(AnsiName);
@@ -266,7 +1099,7 @@ FNameEntrySerialized::FNameEntrySerialized(const FNameEntry& NameEntry)
  */
 FString FNameEntrySerialized::GetPlainNameString() const
 {
-	if( IsWide() )
+	if (bIsWide)
 	{
 		return FString(WideName);
 	}
@@ -280,37 +1113,45 @@ FString FNameEntrySerialized::GetPlainNameString() const
 	FName statics.
 -----------------------------------------------------------------------------*/
 
-TNameEntryArray& FName::GetNames()
+int32 FName::GetNameEntryMemorySize()
 {
-	// NOTE: The reason we initialize to NULL here is due to an issue with static initialization of variables with
-	// constructors/destructors across DLL boundaries, where a function called from a statically constructed object
-	// calls a function in another module (such as this function) that creates a static variable.  A crash can occur
-	// because the static initialization of this DLL has not yet happened, and the CRT's list of static destructors
-	// cannot be written to because it has not yet been initialized fully.	(@todo UE4 DLL)
-	static TNameEntryArray*	Names = NULL;
-	if( Names == NULL )
-	{
-		check(IsInGameThread());
-		Names = new TNameEntryArray();
-	}
-	return *Names;
+	return GetNamePool().NumBlocks() * FNameEntryAllocator::BlockSizeBytes;
 }
 
-FNameEntry*** FName::GetNameTableForDebuggerVisualizers_MT()
+int32 FName::GetNameTableMemorySize()
 {
-	return GetNames().GetRootBlockForDebuggerVisualizers();
+	return GetNameEntryMemorySize() + sizeof(FNamePool) + GetNamePool().NumSlots() * sizeof(FNameSlot);
 }
 
+int32 FName::GetNumAnsiNames()
+{
+	return GetNamePool().NumAnsiEntries();
+}
 
-FCriticalSection* FName::GetCriticalSection()
+int32 FName::GetNumWideNames()
 {
-	static FCriticalSection*	CriticalSection = NULL;
-	if( CriticalSection == NULL )
-	{
-		check(IsInGameThread());
-		CriticalSection = new FCriticalSection();
-	}
-	return CriticalSection;
+	return GetNamePool().NumWideEntries();
+}
+
+TArray<const FNameEntry*> FName::DebugDump()
+{
+	return GetNamePool().DebugDump();
+}
+
+uint8** FName::GetBlocksForDebugVisualizer()
+{
+	return ((FNamePool*)(NamePoolData))->GetBlocksForDebugVisualizer();
+}
+
+FNameEntry const* FName::GetEntry(EName Ename)
+{
+	FNamePool& Pool = GetNamePool();
+	return &Pool.Resolve(Pool.Find(Ename));
+}
+
+FNameEntry const* FName::GetEntry(FNameEntryId Id)
+{
+	return &GetNamePool().Resolve(Id);
 }
 
 FString FName::NameToDisplayString( const FString& InDisplayName, const bool bIsBool )
@@ -442,501 +1283,413 @@ FString FName::NameToDisplayString( const FString& InDisplayName, const bool bIs
 	return OutDisplayName;
 }
 
+const EName* FName::ToEName() const
+{
+	return GetNamePool().FindEName(ComparisonIndex);
+}
 
-// Static variables.
-TAtomic<FNameEntry*> FName::NameHashHead[FNameDefs::NameHashBucketCount];
-TAtomic<FNameEntry*> FName::NameHashTail[FNameDefs::NameHashBucketCount];
-int32                FName::NameEntryMemorySize;
-int32                FName::NumAnsiNames;
-int32                FName::NumWideNames;
-
+bool FName::IsWithinBounds(FNameEntryId Id)
+{
+	return GetNamePool().IsValid(Id);
+}
 
 /*-----------------------------------------------------------------------------
 	FName implementation.
 -----------------------------------------------------------------------------*/
 
-/**
- * Create an FName. If FindType is FNAME_Find, and the string part of the name 
- * doesn't already exist, then the name will be NAME_None
- *
- * @param Name Value for the string portion of the name
- * @param FindType Action to take (see EFindName)
- */
-FName::FName(const WIDECHAR* Name, EFindName FindType)
+template<class CharType>
+static bool NumberEqualsString(int32 Number, const CharType* Str)
 {
-	if (Name)
-	{
-		Init(Name, NAME_NO_NUMBER_INTERNAL, FindType);
-	}
-	else
-	{
-		*this = FName(NAME_None);
-	}
+	CharType* End = nullptr;
+	return TCString<CharType>::Strtoi64(Str, &End, 10) == Number && End && *End == '\0';
 }
 
-FName::FName(const ANSICHAR* Name, EFindName FindType)
+template<class CharType1, class CharType2>
+static bool StringAndNumberEqualsString(const CharType1* Name, uint32 NameLen, int32 InternalNumber, const CharType2* Str)
 {
-	if (Name)
+	if (FPlatformString::Strnicmp(Name, Str, NameLen))
 	{
-		Init(Name, NAME_NO_NUMBER_INTERNAL, FindType);
+		return false;
 	}
-	else
+
+	if (InternalNumber == NAME_NO_NUMBER_INTERNAL)
 	{
-		*this = FName(NAME_None);
+		return Str[NameLen] == '\0';
 	}
-}
 
-/**
- * Create an FName. If FindType is FNAME_Find, and the string part of the name 
- * doesn't already exist, then the name will be NAME_None
- *
- * @param Name Value for the string portion of the name
- * @param Number Value for the number portion of the name
- * @param FindType Action to take (see EFindName)
- */
-FName::FName( const TCHAR* Name, int32 InNumber, EFindName FindType, const bool bSplitName )
-{
-	Init(Name, InNumber, FindType, bSplitName);
+	uint32 Number = NAME_INTERNAL_TO_EXTERNAL(InternalNumber);
+	return Str[NameLen] == '_' && NumberEqualsString(Number, Str + NameLen + 1);
 }
 
-FName::FName(const FNameEntrySerialized& LoadedEntry)
+struct FAnsiStringView
 {
-	if (LoadedEntry.bWereHashesLoaded)
-	{
-		// Since the name table can change sizes we need to mask the raw hash to the current size so we don't access out of bounds
-		const uint16 NonCasePreservingHash = LoadedEntry.NonCasePreservingHash & (FNameDefs::NameHashBucketCount - 1);
-		const uint16 CasePreservingHash = LoadedEntry.CasePreservingHash & (FNameDefs::NameHashBucketCount - 1);
-		if (LoadedEntry.IsWide())
-		{
-			Init(LoadedEntry.GetWideName(), NAME_NO_NUMBER_INTERNAL, FNAME_Add, NonCasePreservingHash, CasePreservingHash);
-		}
-		else
-		{
-			Init(LoadedEntry.GetAnsiName(), NAME_NO_NUMBER_INTERNAL, FNAME_Add, NonCasePreservingHash, CasePreservingHash);
-		}
-	}
-	else
-	{
-		if (LoadedEntry.IsWide())
-		{
-			Init(LoadedEntry.GetWideName(), NAME_NO_NUMBER_INTERNAL, FNAME_Add, false);
-		}
-		else
-		{
-			Init(LoadedEntry.GetAnsiName(), NAME_NO_NUMBER_INTERNAL, FNAME_Add, false);
-		}
-	}
-}
+	using CharType = ANSICHAR;
 
-FName::FName( EName HardcodedIndex, const TCHAR* Name )
-{
-	check(HardcodedIndex >= 0);
-	Init(Name, NAME_NO_NUMBER_INTERNAL, FNAME_Add, false, HardcodedIndex);
-}
+	const ANSICHAR* Str;
+	int32 Len;
+};
 
-/**
- * Compares name to passed in one. Sort is alphabetical ascending.
- *
- * @param	Other	Name to compare this against
- * @return	< 0 is this < Other, 0 if this == Other, > 0 if this > Other
- */
-int32 FName::Compare( const FName& Other ) const
+struct FWideStringViewWithWidth
 {
-	// Names match, check whether numbers match.
-	if( GetComparisonIndexFast() == Other.GetComparisonIndexFast() )
-	{
-		return GetNumber() - Other.GetNumber();
-	}
-	// Names don't match. This means we don't even need to check numbers.
-	else
-	{
-		TNameEntryArray& Names = GetNames();
-		const FNameEntry* const ThisEntry = GetComparisonNameEntry();
-		const FNameEntry* const OtherEntry = Other.GetComparisonNameEntry();
+	using CharType = WIDECHAR;
 
-		FNameBuffer TempBuffer1;
-		FNameBuffer TempBuffer2;
-
-		// If one or both entries return an invalid name entry, the comparison fails - fallback to comparing the index
-		if (ThisEntry == nullptr || OtherEntry == nullptr)
-		{
-			return GetComparisonIndexFast() - Other.GetComparisonIndexFast();
-		}
-		// Ansi/Wide mismatch, convert to wide
-		else if( ThisEntry->IsWide() != OtherEntry->IsWide() )
-		{
-			return FCStringWide::Stricmp(
-				ThisEntry->IsWide() ?  ThisEntry->GetWideNamePtr(TempBuffer1.WideName)  : StringCast<WIDECHAR>(ThisEntry->GetAnsiNamePtr(TempBuffer1.AnsiName)).Get(),
-				OtherEntry->IsWide() ? OtherEntry->GetWideNamePtr(TempBuffer2.WideName) : StringCast<WIDECHAR>(OtherEntry->GetAnsiNamePtr(TempBuffer2.AnsiName)).Get() );
-		}
-		// Both are wide.
-		else if( ThisEntry->IsWide() )
-		{
-			return FCStringWide::Stricmp( ThisEntry->GetWideNamePtr(TempBuffer1.WideName), OtherEntry->GetWideNamePtr(TempBuffer2.WideName) );
-		}
-		// Both are ansi.
-		else
-		{
-			return FCStringAnsi::Stricmp( ThisEntry->GetAnsiNamePtr(TempBuffer1.AnsiName), OtherEntry->GetAnsiNamePtr(TempBuffer2.AnsiName) );
-		}		
-	}
-}
+	const WIDECHAR* Str;
+	int32 Len;
+	bool bIsWide;
+};
 
-template <typename TCharType>
-uint16 FName::GetCasePreservingHash(const TCharType* Source)
+static FAnsiStringView MakeUnconvertedView(const ANSICHAR* Str, int32 Len)
 {
-	return GetRawCasePreservingHash(Source) & (FNameDefs::NameHashBucketCount - 1);
+	return { Str, Len };
 }
 
-template <typename TCharType>
-uint16 FName::GetNonCasePreservingHash(const TCharType* Source)
+static FAnsiStringView MakeUnconvertedView(const ANSICHAR* Str)
 {
-	return GetRawNonCasePreservingHash(Source) & (FNameDefs::NameHashBucketCount - 1);
+	return { Str, Str ? FCStringAnsi::Strlen(Str) : 0 };
 }
 
-void FName::Init(const WIDECHAR* InName, int32 InNumber, EFindName FindType, bool bSplitName, int32 HardcodeIndex)
+static bool IsWide(const WIDECHAR* Str, const int32 Len)
 {
-	LLM_SCOPE(ELLMTag::FName);
-
-	const bool bIsPureAnsi = TCString<WIDECHAR>::IsPureAnsi(InName);
-	// Switch to ANSI if possible to save memory
-	if (bIsPureAnsi)
-	{
-		InitInternal_HashSplit<ANSICHAR>(StringCast<ANSICHAR>(InName).Get(), InNumber, FindType, bSplitName, HardcodeIndex);
-	}
-	else
+	uint32 UserCharBits = 0;
+	for (int32 I = 0; I < Len; ++I)
 	{
-		InitInternal_HashSplit<WIDECHAR>(InName, InNumber, FindType, bSplitName, HardcodeIndex);
+		UserCharBits |= TChar<WIDECHAR>::ToUnsigned(Str[I]);
 	}
+	return UserCharBits & 0xffffff80u;
 }
 
-void FName::Init(const WIDECHAR* InName, int32 InNumber, EFindName FindType, const uint16 NonCasePreservingHash, const uint16 CasePreservingHash)
+static int32 GetLengthAndWidth(const WIDECHAR* Str, bool& bOutIsWide)
 {
-	// Since this comes from the FLinkerLoad we know that it is not pure ansi
-	InitInternal<WIDECHAR>(InName, InNumber, FindType, -1, NonCasePreservingHash, CasePreservingHash);
+	uint32 UserCharBits = 0;
+	const WIDECHAR* It = Str;
+	while (*It)
+	{
+		UserCharBits |= TChar<WIDECHAR>::ToUnsigned(*It);
+		++It;
+	}
+
+	bOutIsWide = UserCharBits & 0xffffff80u;
+
+	return It - Str;
 }
 
-void FName::Init(const ANSICHAR* InName, int32 InNumber, EFindName FindType, bool bSplitName, int32 HardcodeIndex)
+static FWideStringViewWithWidth MakeUnconvertedView(const WIDECHAR* Str, int32 Len)
 {
-	InitInternal_HashSplit<ANSICHAR>(InName, InNumber, FindType, bSplitName, HardcodeIndex);
+	return { Str, Len, IsWide(Str, Len) };
 }
 
-void FName::Init(const ANSICHAR* InName, int32 InNumber, EFindName FindType, const uint16 NonCasePreservingHash, const uint16 CasePreservingHash)
+static FWideStringViewWithWidth MakeUnconvertedView(const WIDECHAR* Str)
 {
-	InitInternal<ANSICHAR>(InName, InNumber, FindType, -1, NonCasePreservingHash, CasePreservingHash);
+	FWideStringViewWithWidth View;
+	View.Str = Str;
+	View.Len = GetLengthAndWidth(Str, View.bIsWide);
+	return View;
 }
 
-template <typename TCharType>
-void FName::InitInternal_HashSplit(const TCharType* InName, int32 InNumber, const EFindName FindType, bool bSplitName, const int32 HardcodeIndex)
+/** Templated implementations of non-templated member functions, helps keep header clean */
+struct FNameHelper
 {
-	TCharType TempBuffer[NAME_SIZE];
-	int32 TempNumber;
-	// if we were passed in a number, we can't split again, other wise, a_1_2_3_4 would change every time
-	// it was loaded in
-	if (InNumber == NAME_NO_NUMBER_INTERNAL && bSplitName == true )
+	template<typename ViewType>
+	static FName MakeDetectNumber(ViewType View, EFindName FindType)
 	{
-		if (SplitNameWithCheckImpl<TCharType>(InName, TempBuffer, ARRAY_COUNT(TempBuffer), TempNumber))
+		if (View.Len == 0)
 		{
-			InName = TempBuffer;
-			InNumber = NAME_EXTERNAL_TO_INTERNAL(TempNumber);
+			return FName();
 		}
-	}
-	// Hash value of string after splitting
-	const uint16 NonCasePreservingHash = GetNonCasePreservingHash(InName);
-#if WITH_CASE_PRESERVING_NAME
-	const uint16 CasePreservingHash = GetCasePreservingHash(InName);
-#else
-	const uint16 CasePreservingHash = 0;
-#endif
-	InitInternal<TCharType>(InName, InNumber, FindType, HardcodeIndex, NonCasePreservingHash, CasePreservingHash);
-}
 
-template <typename TCharType>
-void FName::InitInternal(const TCharType* InName, int32 InNumber, const EFindName FindType, const int32 HardcodeIndex, const uint16 NonCasePreservingHash, const uint16 CasePreservingHash)
-{
-	check(TCString<TCharType>::Strlen(InName)<=NAME_SIZE);
-
-	// initialize the name subsystem if necessary
-	if (!GetIsInitialized())
-	{
-		StaticInit();
-	}
-
-	check(InName);
-
-	// If empty or invalid name was specified, return NAME_None.
-	if( !InName[0] )
-	{
-		check(HardcodeIndex < 1); // if this is hardcoded, it better be zero 
-		ComparisonIndex = NAME_None;
-#if WITH_CASE_PRESERVING_NAME
-		DisplayIndex = NAME_None;
-#endif
-		Number = NAME_NO_NUMBER_INTERNAL;
-		return;
-	}
-
-
-	//!!!! Caution, since these are set by static initializers from multiple threads, we must use local variables for this stuff until just before we return.
+		using CharType = typename ViewType::CharType;
+		const CharType* Name = View.Str;
+		const int32 Len = View.Len;
+		
+		int32 Digits = 0;
+		for (const CharType* It = Name + Len - 1; It >= Name && *It >= '0' && *It <= '9'; --It)
+		{
+			++Digits;
+		}
 
-	bool bWasFoundOrAdded = true;
-	int32 OutComparisonIndex = HardcodeIndex;
-	int32 OutDisplayIndex = HardcodeIndex;
+		const CharType* FirstDigit = Name + Len - Digits;
+		if (Digits && Digits < Len && *(FirstDigit - 1) == '_')
+		{
+			// check for the case where there are multiple digits after the _ and the first one
+			// is a 0 ("Rocket_04"). Can't split this case. (So, we check if the first char
+			// is not 0 or the length of the number is 1 (since ROcket_0 is valid)
+			if (Digits == 1 || *FirstDigit != '0')
+			{
+				// Attempt to convert what's following it to a number
+				// This relies on Name being null-terminated
+				uint64 Number = TCString<CharType>::Atoi64(Name + Len - Digits);
+				if (Number < MAX_uint32)
+				{
+					View.Len -= 1 + Digits;
+					return MakeWithNumber(View, FindType, static_cast<uint32>(NAME_EXTERNAL_TO_INTERNAL(Number)));
+				}
+			}
+		}
 
-	const bool bIsPureAnsi = TCString<TCharType>::IsPureAnsi(InName);
-	if(bIsPureAnsi)
-	{
-		bWasFoundOrAdded = InitInternal_FindOrAdd<ANSICHAR>(StringCast<ANSICHAR>(InName).Get(), FindType, HardcodeIndex, NonCasePreservingHash, CasePreservingHash, OutComparisonIndex, OutDisplayIndex);
-	}
-	else
-	{
-		bWasFoundOrAdded = InitInternal_FindOrAdd<WIDECHAR>(StringCast<WIDECHAR>(InName).Get(), FindType, HardcodeIndex, NonCasePreservingHash, CasePreservingHash, OutComparisonIndex, OutDisplayIndex);
+		return MakeWithNumber(View, FindType, NAME_NO_NUMBER_INTERNAL);
 	}
 
-	if(bWasFoundOrAdded)
-	{
-		ComparisonIndex = OutComparisonIndex;
-#if WITH_CASE_PRESERVING_NAME
-		DisplayIndex = OutDisplayIndex;
-#endif
-		Number = InNumber;
-	}
-	else
+	static FName MakeWithNumber(FAnsiStringView View, EFindName FindType, int32 InternalNumber)
 	{
-		ComparisonIndex = NAME_None;
-#if WITH_CASE_PRESERVING_NAME
-		DisplayIndex = NAME_None;
-#endif
-		Number = NAME_NO_NUMBER_INTERNAL;
-	}
-}
-
-template <typename TCharType> void IncrementNameCount();
-template <> void IncrementNameCount<ANSICHAR>()
-{
-	FName::NumAnsiNames++;
-}
-template <> void IncrementNameCount<WIDECHAR>()
-{
-	FName::NumWideNames++;
-}
-
-template <typename TCharType>
-struct FNameInitHelper
-{
-};
-
-template <>
-struct FNameInitHelper<ANSICHAR>
-{
-	static const bool IsAnsi = true;
+		// Ignore the supplied number if the name string is empty
+		// to keep the semantics of the old FName implementation
+		if (View.Len == 0)
+		{
+			return FName();
+		}
 
-	static const ANSICHAR* GetNameString(const FNameEntry* const NameEntry, ANSICHAR(&OptionalTempBuffer)[NAME_SIZE])
-	{
-		return NameEntry->GetAnsiNamePtr(OptionalTempBuffer);
+		return Make(FNameStringView(View.Str, View.Len), FindType, InternalNumber);
 	}
 
-	static void SetNameString(FNameEntry* DestNameEntry, const ANSICHAR* SrcName, const int32 NameLength)
+	static FName MakeWithNumber(const FWideStringViewWithWidth View, EFindName FindType, int32 InternalNumber)
 	{
-		DestNameEntry->SetAnsiName(SrcName, NameLength);
-	}
+		// Ignore the supplied number if the name string is empty
+		// to keep the semantics of the old FName implementation
+		if (View.Len == 0)
+		{
+			return FName();
+		}
 
-	static int32 GetIndexShiftValue()
-	{
-		return 0;
+		// Convert to narrow if possible
+		if (!View.bIsWide)
+		{
+			// Consider _mm_packus_epi16 or similar if this proves too slow
+			ANSICHAR AnsiName[NAME_SIZE];
+			for (int32 I = 0; I < View.Len; ++I)
+			{
+				AnsiName[I] = View.Str[I];
+			}
+			return Make(FNameStringView(AnsiName, View.Len), FindType, InternalNumber);
+		}
+		else
+		{
+			return Make(FNameStringView(View.Str, View.Len), FindType, InternalNumber);
+		}
 	}
 
-	static int32 GetSize(int32 Length)
+	static FName Make(FNameStringView View, EFindName FindType, int32 InternalNumber)
 	{
-		// Add size required for string to the base size used by the FNameEntry
-		int32 Size = NameEntryWithoutUnionSize + (Length + 1) * sizeof(ANSICHAR);
-		return Size;
-	}
-};
+		FNamePool& Pool = GetNamePool();
 
-template <>
-struct FNameInitHelper<WIDECHAR>
-{
-	static const bool IsAnsi = false;
+		FNameEntryId DisplayId, ComparisonId;
+		if (FindType == FNAME_Add)
+		{
+			DisplayId = Pool.Store(View);
+#if WITH_CASE_PRESERVING_NAME
+			ComparisonId = Pool.Resolve(DisplayId).ComparisonId;
+#else
+			ComparisonId = DisplayId;
+#endif
+		}
+		else if (FindType == FNAME_Find)
+		{
+			DisplayId = Pool.Find(View);
+#if WITH_CASE_PRESERVING_NAME
+			ComparisonId = DisplayId ? Pool.Resolve(DisplayId).ComparisonId : DisplayId;
+#else
+			ComparisonId = DisplayId;
+#endif
+		}
+		else
+		{
+			check(FindType == FNAME_Replace_Not_Safe_For_Threading);
 
-	static const WIDECHAR* GetNameString(const FNameEntry* const NameEntry, WIDECHAR(&OptionalTempBuffer)[NAME_SIZE])
-	{
-		return NameEntry->GetWideNamePtr(OptionalTempBuffer);
-	}
+			DisplayId = Pool.Store(View);
+#if WITH_CASE_PRESERVING_NAME
+			ComparisonId = Pool.Resolve(DisplayId).ComparisonId;
+#else
+			ComparisonId = DisplayId;
+#endif
+			ReplaceName(Pool.Resolve(ComparisonId), View);
+		}
 
-	static void SetNameString(FNameEntry* DestNameEntry, const WIDECHAR* SrcName, const int32 NameLength)
-	{
-		DestNameEntry->SetWideName(SrcName, NameLength);
+		return FName(ComparisonId, DisplayId, InternalNumber);
 	}
 
-	static int32 GetIndexShiftValue()
+	static FName MakeFromLoaded(const FNameEntrySerialized& LoadedEntry)
 	{
-		return 1;
+		FNameStringView View = LoadedEntry.bIsWide
+			? FNameStringView(LoadedEntry.WideName, FCStringWide::Strlen(LoadedEntry.WideName))
+			: FNameStringView(LoadedEntry.AnsiName, FCStringAnsi::Strlen(LoadedEntry.AnsiName));
+
+		return Make(View, FNAME_Add, NAME_NO_NUMBER_INTERNAL);
+
 	}
 
-	static int32 GetSize(int32 Length)
+	template<class CharType>
+	static bool EqualsString(FName Name, const CharType* Str)
 	{
-		// Add size required for string to the base size used by the FNameEntry
-		int32 Size = NameEntryWithoutUnionSize + (Length + 1) * sizeof(TCHAR);
-		return Size;
+		// Make NAME_None == TEXT("") or nullptr consistent with NAME_None == FName(TEXT("")) or FName(nullptr)
+		if (Str == nullptr || Str[0] == '\0')
+		{
+			return Name == NAME_None;
+		}
+
+		const FNameEntry& Entry = *Name.GetComparisonNameEntry();
+
+		uint32 NameLen = Entry.Header.Len;
+		FNameBuffer Temp;
+		return Entry.IsWide()
+			? StringAndNumberEqualsString(Entry.GetUnterminatedName(Temp.WideName), NameLen, Name.GetNumber(), Str)
+			: StringAndNumberEqualsString(Entry.GetUnterminatedName(Temp.AnsiName), NameLen, Name.GetNumber(), Str);
 	}
-};
 
-template <typename TCharType>
-bool FName::InitInternal_FindOrAdd(const TCharType* InName, const EFindName FindType, const int32 HardcodeIndex, const uint16 NonCasePreservingHash, const uint16 CasePreservingHash, int32& OutComparisonIndex, int32& OutDisplayIndex)
-{
-	const bool bWasFoundOrAdded = InitInternal_FindOrAddNameEntry<TCharType>(InName, FindType, ENameCase::IgnoreCase, NonCasePreservingHash, OutComparisonIndex);
-	
-#if WITH_CASE_PRESERVING_NAME
-	if(bWasFoundOrAdded && HardcodeIndex < 0)
+	static void ReplaceName(FNameEntry& Existing, FNameStringView Updated)
 	{
-		TNameEntryArray& Names = GetNames();
-		const FNameEntry* const NameEntry = Names[OutComparisonIndex];
+		check(Existing.Header.bIsWide == Updated.bIsWide);
+		check(Existing.Header.Len == Updated.Len);
 
-		// If the string we got back doesn't match the case of the string we provided, also add a case variant version for display purposes
-		TCharType TempBuffer[NAME_SIZE];
-		if(TCString<TCharType>::Strcmp(InName, FNameInitHelper<TCharType>::GetNameString(NameEntry, TempBuffer)) != 0)
+		if (Updated.bIsWide)
 		{
-			if(!InitInternal_FindOrAddNameEntry<TCharType>(InName, FindType, ENameCase::CaseSensitive, CasePreservingHash, OutDisplayIndex))
-			{
-				// We don't consider failing to find/add the case variant a full failure
-				OutDisplayIndex = OutComparisonIndex;
-			}
+			Existing.StoreName(Updated.Wide, Updated.Len);
 		}
 		else
 		{
-			OutDisplayIndex = OutComparisonIndex;
+			Existing.StoreName(Updated.Ansi, Updated.Len);
 		}
 	}
-	else
-#endif
-	{
-		OutDisplayIndex = OutComparisonIndex;
-	}
+};
 
-	return bWasFoundOrAdded;
-}
 
-template <typename TCharType>
-bool FName::InitInternal_FindOrAddNameEntry(const TCharType* InName, const EFindName FindType, const ENameCase ComparisonMode, const uint16 iHash, int32& OutIndex)
+FNameEntryId FName::GetComparisonIdFromDisplayId(FNameEntryId DisplayId)
 {
-	CallNameCreationHook();
-	if (OutIndex < 0)
-	{
-		// Try to find the name in the hash.
-		FNameEntry* Hash = NameHashHead[iHash].Load(EMemoryOrder::Relaxed);
-		while (Hash)
-		{
-			FNameEntry* NextHash = Hash->HashNext.Load(EMemoryOrder::Relaxed);
+#if WITH_CASE_PRESERVING_NAME
+	return GetEntry(DisplayId)->ComparisonId;
+#else
+	return DisplayId;
+#endif
+}
 
-			FPlatformMisc::Prefetch( NextHash );
+FName::FName(const WIDECHAR* Name, EFindName FindType)
+	: FName(FNameHelper::MakeDetectNumber(MakeUnconvertedView(Name), FindType))
+{}
 
-			// Compare the passed in string
-			if( Hash->IsEqual( InName, ComparisonMode ) )
-			{
-				// Found it in the hash.
-				OutIndex = Hash->GetIndex();
-
-				// Check to see if the caller wants to replace the contents of the
-				// FName with the specified value. This is useful for compiling
-				// script classes where the file name is lower case but the class
-				// was intended to be uppercase.
-				if (FindType == FNAME_Replace_Not_Safe_For_Threading)
-				{
-					check(IsInGameThread());
+FName::FName(const ANSICHAR* Name, EFindName FindType)
+	: FName(FNameHelper::MakeDetectNumber(MakeUnconvertedView(Name), FindType))
+{}
 
-					// This *must* be true, or we'll overwrite memory when the
-					// copy happens if it is longer
-					int32 NewLength = TCString<TCharType>::Strlen(InName);
-					check(NewLength == Hash->GetNameLength());
+FName::FName(int32 Len, const WIDECHAR* Name, EFindName FindType)
+	: FName(FNameHelper::MakeDetectNumber(MakeUnconvertedView(Name, Len), FindType))
+{}
 
-					FNameInitHelper<TCharType>::SetNameString(Hash, InName, NewLength);
-				}
-				check(OutIndex >= 0);
-				return true;
-			}
+FName::FName(int32 Len, const ANSICHAR* Name, EFindName FindType)
+	: FName(FNameHelper::MakeDetectNumber(MakeUnconvertedView(Name, Len), FindType))
+{}
 
-			Hash = NextHash;
-		}
+FName::FName(const WIDECHAR* Name, int32 InNumber, EFindName FindType)
+	: FName(FNameHelper::MakeWithNumber(MakeUnconvertedView(Name), FindType, InNumber))
+{}
 
-		// Didn't find name.
-		if( FindType==FNAME_Find )
-		{
-			// Not found.
-			return false;
-		}
+FName::FName(const ANSICHAR* Name, int32 InNumber, EFindName FindType)
+	: FName(FNameHelper::MakeWithNumber(MakeUnconvertedView(Name), FindType, InNumber))
+{}
+
+FName::FName(int32 Len, const WIDECHAR* Name, int32 InNumber, EFindName FindType)
+	: FName(InNumber != NAME_NO_NUMBER_INTERNAL ? FNameHelper::MakeWithNumber(MakeUnconvertedView(Name, Len), FindType, InNumber)
+												: FNameHelper::MakeDetectNumber(MakeUnconvertedView(Name, Len), FindType))
+{}
+
+FName::FName(int32 Len, const ANSICHAR* Name, int32 InNumber, EFindName FindType)
+	: FName(InNumber != NAME_NO_NUMBER_INTERNAL ? FNameHelper::MakeWithNumber(MakeUnconvertedView(Name, Len), FindType, InNumber)
+												: FNameHelper::MakeDetectNumber(MakeUnconvertedView(Name, Len), FindType))
+{}
+
+FName::FName(const TCHAR* Name, int32 InNumber, EFindName FindType, bool bSplitName)
+	: FName(InNumber == NAME_NO_NUMBER_INTERNAL && bSplitName 
+			? FNameHelper::MakeDetectNumber(MakeUnconvertedView(Name), FindType)
+			: FNameHelper::MakeWithNumber(MakeUnconvertedView(Name), FindType, InNumber))
+{}
+
+FName::FName(const FNameEntrySerialized& LoadedEntry)
+	: FName(FNameHelper::MakeFromLoaded(LoadedEntry))
+{}
+
+FName::FName(EName Ename, int32 InNumber)
+	: ComparisonIndex(GetNamePool().Find(Ename))
+#if WITH_CASE_PRESERVING_NAME
+	, DisplayIndex(ComparisonIndex)
+#endif
+	, Number(InNumber)
+{
+	check(Ename < NAME_MaxHardcodedNameIndex);
+}
+
+FName::FName(EName Ename)
+	: FName(Ename, NAME_NO_NUMBER_INTERNAL)
+{}
+
+bool FName::operator==(const ANSICHAR* Str) const
+{
+	return FNameHelper::EqualsString(*this, Str);
+}
+
+bool FName::operator==(const WIDECHAR* Str) const
+{
+	return FNameHelper::EqualsString(*this, Str);
+}
+
+int32 FName::Compare( const FName& Other ) const
+{
+	// Names match, check whether numbers match.
+	if (ComparisonIndex == Other.ComparisonIndex)
+	{
+		return GetNumber() - Other.GetNumber();
 	}
-	// acquire the lock
-	FScopeLock ScopeLock(GetCriticalSection());
-	FNameEntry* OldHashHead = NameHashHead[iHash].Load(EMemoryOrder::Relaxed);
-	if (OutIndex < 0)
+	else // Names don't match. This means we don't even need to check numbers.
 	{
-		// Try to find the name in the hash. AGAIN...we might have been adding from a different thread and we just missed it
-		for (FNameEntry* Hash = OldHashHead; Hash; Hash = Hash->HashNext.Load(EMemoryOrder::Relaxed))
+		const FNameEntry* const ThisEntry = GetComparisonNameEntry();
+		const FNameEntry* const OtherEntry = Other.GetComparisonNameEntry();
+		
+		// If one or both entries return an invalid name entry, the comparison fails - fallback to comparing the index
+		if (ThisEntry == nullptr || OtherEntry == nullptr)
 		{
-			// Compare the passed in string
-			if( Hash->IsEqual( InName, ComparisonMode ) )
-			{
-				// Found it in the hash.
-				OutIndex = Hash->GetIndex();
-				check(FindType == FNAME_Add);  // if this was a replace, well it isn't safe for threading. Find should have already been handled
-				return true;
-			}
+			return ComparisonIndex.CompareCreationOrder(Other.ComparisonIndex);
 		}
-	}
-	FNameEntry* OldHashTail = NameHashTail[iHash].Load(EMemoryOrder::Relaxed);
-	TNameEntryArray& Names = GetNames();
-	if (OutIndex < 0)
-	{
-		OutIndex = Names.AddZeroed(1);
-	}
-	else
-	{
-		check(OutIndex < Names.Num());
-	}
-	FNameEntry* NewEntry = AllocateNameEntry(InName, OutIndex);
-	if (FPlatformAtomics::InterlockedCompareExchangePointer((void**)&Names[OutIndex], NewEntry, nullptr) != nullptr) // we use an atomic operation to check for unexpected concurrency, verify alignment, etc
-	{
-		UE_LOG(LogUnrealNames, Fatal, TEXT("Hardcoded name '%s' at index %i was duplicated (or unexpected concurrency). Existing entry is '%s'."), *NewEntry->GetPlainNameString(), NewEntry->GetIndex(), *Names[OutIndex]->GetPlainNameString() );
-	}
-	if (!OldHashHead)
-	{
-		checkSlow(!OldHashTail);
-
-		// atomically assign the new head as other threads may be reading it
-		if (!NameHashHead[iHash].CompareExchange(OldHashHead, NewEntry)) // we use an atomic operation to check for unexpected concurrency, verify alignment, etc
+		else if (ThisEntry->IsWide() == OtherEntry->IsWide())
 		{
-			check(0); // someone changed this while we were changing it
+			return ThisEntry->IsWide() ?	FCStringWide::Stricmp(ThisEntry->WideName, OtherEntry->WideName) :
+											FCStringAnsi::Stricmp(ThisEntry->AnsiName, OtherEntry->AnsiName);
 		}
-	}
-	else
-	{
-		checkSlow(OldHashTail);
-
-		// atomically update the linked list as other threads may be reading it
-		FNameEntry* Expected = nullptr;
-		if (!OldHashTail->HashNext.CompareExchange(Expected, NewEntry)) // we use an atomic operation to check for unexpected concurrency, verify alignment, etc
+		else // Ansi/Wide mismatch, convert one string to wide
 		{
-			check(0); // someone changed this while we were changing it
+			WIDECHAR TempBuffer[NAME_SIZE];
+			const FNameEntry* AnsiEntry = ThisEntry->IsWide() ? OtherEntry : ThisEntry;
+			FPlatformString::Convert(TempBuffer, NAME_SIZE, AnsiEntry->AnsiName, FCStringAnsi::Strlen(AnsiEntry->AnsiName));
+
+			return ThisEntry->IsWide() ? FCStringWide::Stricmp(ThisEntry->WideName, TempBuffer)
+										: FCStringWide::Stricmp(TempBuffer, OtherEntry->WideName);
 		}
 	}
-	NameHashTail[iHash].Store(NewEntry, EMemoryOrder::Relaxed); // We can non-atomically assign the tail since it's only ever read while locked
-	check(OutIndex >= 0);
-	return true;
+}
+
+FString FName::GetPlainNameString() const
+{
+	if (const FNameEntry* CurEntry = GetDisplayNameEntry())
+	{
+		return CurEntry->GetPlainNameString();
+	}
+
+	return TEXT("*INVALID*");
+}
+
+void FName::GetPlainANSIString(ANSICHAR(&AnsiName)[NAME_SIZE]) const
+{
+	GetDisplayNameEntry()->GetAnsiName(AnsiName);
+}
+
+void FName::GetPlainWIDEString(WIDECHAR(&WideName)[NAME_SIZE]) const
+{
+	GetDisplayNameEntry()->GetWideName(WideName);
 }
 
 const FNameEntry* FName::GetComparisonNameEntry() const
 {
-	TNameEntryArray& Names = GetNames();
-	const NAME_INDEX Index = GetComparisonIndex();
-	return Names[Index];
+	return &GetNamePool().Resolve(GetComparisonIndex());
 }
 
 const FNameEntry* FName::GetDisplayNameEntry() const
 {
-	TNameEntryArray& Names = GetNames();
-	const NAME_INDEX Index = GetDisplayIndex();
-	return Names[Index];
+	return &GetNamePool().Resolve(GetDisplayIndex());
 }
 
 FString FName::ToString() const
@@ -998,164 +1751,14 @@ void FName::AppendString(FString& Out) const
 	}
 }
 
-
-/*-----------------------------------------------------------------------------
-	FName subsystem.
------------------------------------------------------------------------------*/
-
-void FName::StaticInit()
-{
-	check(IsInGameThread());
-	/** Global instance used to initialize the GCRCTable. It used to be initialized in appInit() */
-	//@todo: Massive workaround for static init order without needing to use a function call for every use of GCRCTable
-	// This ASSUMES that fname::StaticINit is going to be called BEFORE ANY use of GCRCTable
-	FCrc::Init();
-
-	check(GetIsInitialized() == false);
-	check((FNameDefs::NameHashBucketCount&(FNameDefs::NameHashBucketCount-1)) == 0);
-	GetIsInitialized() = true;
-
-	{
-		FScopeLock ScopeLock(GetCriticalSection());
-
-		TNameEntryArray& Names = GetNames();
-		Names.AddZeroed(NAME_MaxHardcodedNameIndex + 1);
-	}
-
-	{
-		// Register all hardcoded names.
-		#define REGISTER_NAME(num,namestr) FName Temp_##namestr(EName(num), TEXT(#namestr));
-		#include "UObject/UnrealNames.inl"
-		#undef REGISTER_NAME
-	}
-
-#if DO_CHECK
-	// Verify no duplicate names.
-	for (int32 HashIndex = 0; HashIndex < FNameDefs::NameHashBucketCount; HashIndex++)
-	{
-		FNameEntry* Hash = NameHashHead[HashIndex].Load(EMemoryOrder::Relaxed);
-		while (Hash)
-		{
-			FNameEntry* NextHash = Hash->HashNext.Load(EMemoryOrder::Relaxed);
-			for (FNameEntry* Other = NextHash; Other; Other = Other->HashNext.Load(EMemoryOrder::Relaxed))
-			{
-				if (FCString::Stricmp(*Hash->GetPlainNameString(), *Other->GetPlainNameString()) == 0)
-				{
-					// we can't print out here because there may be no log yet if this happens before main starts
-					if (FPlatformMisc::IsDebuggerPresent())
-					{
-						UE_DEBUG_BREAK();
-					}
-					else
-					{
-						FPlatformMisc::PromptForRemoteDebugging(false);
-						FMessageDialog::Open(EAppMsgType::Ok, FText::Format( NSLOCTEXT("UnrealEd", "DuplicatedHardcodedName", "Duplicate hardcoded name: {0}"), FText::FromString( Hash->GetPlainNameString() ) ) );
-						FPlatformMisc::RequestExit(false);
-					}
-				}
-			}
-
-			Hash = NextHash;
-		}
-	}
-	// check that the MAX_NETWORKED_HARDCODED_NAME define is correctly set
-	if (GetMaxNames() <= MAX_NETWORKED_HARDCODED_NAME)
-	{
-		if (FPlatformMisc::IsDebuggerPresent())
-		{
-			UE_DEBUG_BREAK();
-		}
-		else
-		{
-			FPlatformMisc::PromptForRemoteDebugging(false);
-			// can't use normal check()/UE_LOG(LogUnrealNames, Fatal,) here
-			FMessageDialog::Open( EAppMsgType::Ok, FText::Format( NSLOCTEXT("UnrealEd", "MAX_NETWORKED_HARDCODED_NAME Incorrect", "MAX_NETWORKED_HARDCODED_NAME is incorrectly set! (Currently {0}, must be no greater than {1}"), FText::AsNumber( MAX_NETWORKED_HARDCODED_NAME ), FText::AsNumber( GetMaxNames() - 1 ) ) );
-			FPlatformMisc::RequestExit(false);
-		}
-	}
-#endif
-}
-
-bool& FName::GetIsInitialized()
-{
-	static bool bIsInitialized = false;
-	return bIsInitialized;
-}
-
-void FName::DisplayHash( FOutputDevice& Ar )
-{
-	int32 UsedBins=0, NameCount=0, MemUsed = 0;
-	for( int32 i=0; i<FNameDefs::NameHashBucketCount; i++ )
-	{
-		if (FNameEntry* Hash = NameHashHead[i].Load(EMemoryOrder::Relaxed))
-		{
-			++UsedBins;
-
-			for (; Hash; Hash = Hash->HashNext.Load(EMemoryOrder::Relaxed))
-			{
-				NameCount++;
-				// Count how much memory this entry is using
-				MemUsed += FNameEntry::GetSize( Hash->GetNameLength(), Hash->IsWide() );
-			}
-		}
-	}
-	Ar.Logf( TEXT("Hash: %i names, %i/%i hash bins, Mem in bytes %i"), NameCount, UsedBins, FNameDefs::NameHashBucketCount, MemUsed);
-}
-
-bool FName::SplitNameWithCheck(const WIDECHAR* OldName, WIDECHAR* NewName, int32 NewNameLen, int32& NewNumber)
+void FName::DisplayHash(FOutputDevice& Ar)
 {
-	return SplitNameWithCheckImpl<WIDECHAR>(OldName, NewName, NewNameLen, NewNumber);
+	GetNamePool().LogStats(Ar);
 }
 
-template <typename TCharType>
-bool FName::SplitNameWithCheckImpl(const TCharType* OldName, TCharType* NewName, int32 NewNameLen, int32& NewNumber)
+FString FName::SafeString(NAME_INDEX InDisplayIndex, int32 InstanceNumber)
 {
-	bool bSucceeded = false;
-	const int32 OldNameLength = TCString<TCharType>::Strlen(OldName);
-
-	if(OldNameLength > 0)
-	{
-		// get string length
-		const TCharType* LastChar = OldName + (OldNameLength - 1);
-		
-		// if the last char is a number, then we will try to split
-		const TCharType* Ch = LastChar;
-		if (*Ch >= '0' && *Ch <= '9')
-		{
-			// go backwards, looking an underscore or the start of the string
-			// (we don't look at first char because '_9' won't split well)
-			while (*Ch >= '0' && *Ch <= '9' && Ch > OldName)
-			{
-				Ch--;
-			}
-
-			// if the first non-number was an underscore (as opposed to a letter),
-			// we can split
-			if (*Ch == '_')
-			{
-				// check for the case where there are multiple digits after the _ and the first one
-				// is a 0 ("Rocket_04"). Can't split this case. (So, we check if the first char
-				// is not 0 or the length of the number is 1 (since ROcket_0 is valid)
-				if (Ch[1] != '0' || LastChar - Ch == 1)
-				{
-					// attempt to convert what's following it to a number
-					uint64 TempConvert = TCString<TCharType>::Atoi64(Ch + 1);
-					if (TempConvert <= MAX_int32)
-					{
-						NewNumber = (int32)TempConvert;
-						// copy the name portion into the buffer
-						TCString<TCharType>::Strncpy(NewName, OldName, FMath::Min<int32>(Ch - OldName + 1, NewNameLen));
-
-						// mark successful
-						bSucceeded = true;
-					}
-				}
-			}
-		}
-	}
-
-	// return success code
-	return bSucceeded;
+	return FName(InDisplayIndex, InDisplayIndex, InstanceNumber).ToString();
 }
 
 bool FName::IsValidXName(const FString& InName, const FString& InInvalidChars, FText* OutReason, const FText* InErrorCtx)
@@ -1194,13 +1797,20 @@ bool FName::IsValidXName(const FString& InName, const FString& InInvalidChars, F
 
 void FName::AutoTest()
 {
+	check(FNameHash::IsAnsiNone("None", 4) == 1);
+	check(FNameHash::IsAnsiNone("none", 4) == 1);
+	check(FNameHash::IsAnsiNone("NONE", 4) == 1);
+	check(FNameHash::IsAnsiNone("nOnE", 4) == 1);
+	check(FNameHash::IsAnsiNone("None", 5) == 0);
+	check(FNameHash::IsAnsiNone(TEXT("None"), 4) == 0);
+	check(FNameHash::IsAnsiNone("nono", 4) == 0);
+	check(FNameHash::IsAnsiNone("enon", 4) == 0);
+
 	const FName AutoTest_1("AutoTest_1");
 	const FName autoTest_1("autoTest_1");
 	const FName autoTeSt_1("autoTeSt_1");
-	const FName AutoTest1Find("autoTEST_1", EFindName::FNAME_Find);
 	const FName AutoTest_2(TEXT("AutoTest_2"));
 	const FName AutoTestB_2(TEXT("AutoTestB_2"));
-	const FName NullName(static_cast<ANSICHAR*>(nullptr));
 
 	check(AutoTest_1 != AutoTest_2);
 	check(AutoTest_1 == autoTest_1);
@@ -1217,7 +1827,119 @@ void FName::AutoTest()
 	check(*AutoTestB_2.GetPlainNameString() != *AutoTest_2.GetPlainNameString());
 	check(AutoTestB_2.GetNumber() == AutoTest_2.GetNumber());
 	check(autoTest_1.GetNumber() != AutoTest_2.GetNumber());
+
+	const FName NullName(static_cast<ANSICHAR*>(nullptr));
 	check(NullName.IsNone());
+	check(NullName == FName(NAME_None));
+	check(NullName == FName());
+	check(NullName == FName(""));
+	check(NullName == FName(TEXT("")));
+	check(NullName == FName("None"));
+	check(NullName == FName(TEXT("None")));
+	check(FName().ToEName());
+	check(*FName().ToEName() == NAME_None);
+	check(NullName.GetComparisonIndex().ToUnstableInt() == 0);
+
+	const FName Cylinder(NAME_Cylinder);
+	check(Cylinder == FName("Cylinder"));
+	check(Cylinder.ToEName());
+	check(*Cylinder.ToEName() == NAME_Cylinder);
+	check(Cylinder.GetPlainNameString() == TEXT("Cylinder"));
+
+	check(FName("") == FName(0, "Unused"));
+	check(FName("Used") == FName(4, "UsedUnused"));
+	check(FName("Used") == FName(4, "Used"));
+
+	// Test wide strings
+	FString Wide("Wide ");
+	Wide[4] = 60000;
+	FName WideName(*Wide);
+	check(WideName.GetPlainNameString() == Wide);
+	FString WideLong = FString::ChrN(1000, 60000);
+	check(FName(*WideLong).GetPlainNameString() == WideLong);
+
+
+	// Check that FNAME_Find doesn't add entries
+	static bool Once = true;
+	if (Once)
+	{
+		check(FName("UniqueUnicorn!!", FNAME_Find) == FName());
+		Once = false;
+	}
+
+	// Check that FNAME_Find can find entries
+	const FName UniqueName("UniqueUnicorn!!", FNAME_Add);
+	check(FName("UniqueUnicorn!!", FNAME_Find) == UniqueName);
+	check(FName(TEXT("UniqueUnicorn!!"), FNAME_Find) == UniqueName);
+	check(FName("UNIQUEUNICORN!!", FNAME_Find) == UniqueName);
+	check(FName(TEXT("UNIQUEUNICORN!!"), FNAME_Find) == UniqueName);
+	check(FName("uniqueunicorn!!", FNAME_Find) == UniqueName);
+
+	// Check FNAME_Replace_Not_Safe_For_Threading updates casing
+	check(0 != UniqueName.GetPlainNameString().Compare("UNIQUEunicorn!!", ESearchCase::CaseSensitive));
+	const FName UniqueNameReplaced("UNIQUEunicorn!!", FNAME_Replace_Not_Safe_For_Threading);
+	check(0 == UniqueName.GetPlainNameString().Compare("UNIQUEunicorn!!", ESearchCase::CaseSensitive));
+	check(UniqueNameReplaced == UniqueName);
+
+	// Check FNAME_Replace_Not_Safe_For_Threading works with wide string
+	check(0 != UniqueName.GetPlainNameString().Compare("uniqueunicorn!!", ESearchCase::CaseSensitive));
+	const FName UpdatedCasing(TEXT("uniqueunicorn!!"), FNAME_Replace_Not_Safe_For_Threading);
+	check(0 == UniqueName.GetPlainNameString().Compare("uniqueunicorn!!", ESearchCase::CaseSensitive));
+
+	// Check FNAME_Replace_Not_Safe_For_Threading adds entries that do not exist
+	const FName AddedByReplace("WasAdded!!", FNAME_Replace_Not_Safe_For_Threading);
+	check(FName("WasAdded!!", FNAME_Find) == AddedByReplace);
+
+	check(NumberEqualsString(0, "0"));
+	check(NumberEqualsString(11, "11"));
+	check(NumberEqualsString(2147483647, "2147483647"));
+
+	check(!NumberEqualsString(0, "1"));
+	check(!NumberEqualsString(1, "0"));
+	check(!NumberEqualsString(11, "12"));
+	check(!NumberEqualsString(12, "11"));
+	check(!NumberEqualsString(2147483647, "2147483646"));
+	check(!NumberEqualsString(2147483646, "2147483647"));
+
+	check(StringAndNumberEqualsString("abc", 3, NAME_EXTERNAL_TO_INTERNAL(10), "abc_10"));
+	check(!StringAndNumberEqualsString("aba", 3, NAME_EXTERNAL_TO_INTERNAL(10), "abc_10"));
+	check(!StringAndNumberEqualsString("abc", 2, NAME_EXTERNAL_TO_INTERNAL(10), "abc_10"));
+	check(!StringAndNumberEqualsString("abc", 2, NAME_EXTERNAL_TO_INTERNAL(11), "abc_10"));
+	check(!StringAndNumberEqualsString("abc", 3, NAME_EXTERNAL_TO_INTERNAL(10), "aba_10"));
+	check(!StringAndNumberEqualsString("abc", 3, NAME_EXTERNAL_TO_INTERNAL(10), "abc_11"));
+	check(!StringAndNumberEqualsString("abc", 3, NAME_EXTERNAL_TO_INTERNAL(10), "abc_100"));
+
+	check(StringAndNumberEqualsString("abc", 3, NAME_EXTERNAL_TO_INTERNAL(0), "abc_0"));
+	check(!StringAndNumberEqualsString("abc", 3, NAME_EXTERNAL_TO_INTERNAL(0), "abc_1"));
+
+	check(StringAndNumberEqualsString("abc", 3, NAME_NO_NUMBER_INTERNAL, "abc"));
+	check(!StringAndNumberEqualsString("abc", 2, NAME_NO_NUMBER_INTERNAL, "abc"));
+	check(!StringAndNumberEqualsString("abc", 3, NAME_NO_NUMBER_INTERNAL, "abcd"));
+	check(!StringAndNumberEqualsString("abc", 3, NAME_NO_NUMBER_INTERNAL, "abc_0"));
+	check(!StringAndNumberEqualsString("abc", 3, NAME_NO_NUMBER_INTERNAL, "abc_"));
+
+#if 0
+	// Check hash table growth still yields the same unique FName ids
+	static int32 OverflowAtLeastTwiceCount = 4 * FNamePoolInitialSlots * FNamePoolShards;
+	TArray<FNameEntryId> Ids;
+	for (int I = 0; I < OverflowAtLeastTwiceCount; ++I)
+	{
+		FNameEntryId Id = FName(*FString::Printf(TEXT("%d"), I)).GetComparisonIndex();
+		Ids.Add(Id);
+	}
+
+	for (int I = 0; I < OverflowAtLeastTwiceCount; ++I)
+	{
+		FNameEntryId Id = FName(*FString::Printf(TEXT("%d"), I)).GetComparisonIndex();
+		FNameEntryId OldId = Ids[I];
+
+		while (Id != OldId)
+		{
+			Id = FName(*FString::Printf(TEXT("%d"), I)).GetComparisonIndex();
+		}
+		check(Id == OldId);
+	}
+#endif
 }
 
 /*-----------------------------------------------------------------------------
@@ -1278,7 +2000,7 @@ FArchive& operator<<(FArchive& Ar, FNameEntrySerialized& E)
 			}
 
 			// mark the name will be wide
-			E.PreSetIsWideForSerialization(true);
+			E.bIsWide = true;
 
 			// get the pointer to the wide array 
 			WIDECHAR* WideName = const_cast<WIDECHAR*>(E.GetWideName());
@@ -1303,7 +2025,7 @@ FArchive& operator<<(FArchive& Ar, FNameEntrySerialized& E)
 			}
 
 			// mark the name will be ansi
-			E.PreSetIsWideForSerialization(false);
+			E.bIsWide = false;
 
 			// ansi strings can go right into the AnsiBuffer
 			ANSICHAR* AnsiName = const_cast<ANSICHAR*>(E.GetAnsiName());
@@ -1311,16 +2033,11 @@ FArchive& operator<<(FArchive& Ar, FNameEntrySerialized& E)
 		}
 		if (Ar.UE4Ver() >= VER_UE4_NAME_HASHES_SERIALIZED)
 		{
-			// Read the save time calculated hashes to save load time perf
+			// These hashes are no longer used. They're only kept to maintain serialization format.
+			// Please remove them if you ever change serialization format.
 			Ar << E.NonCasePreservingHash;
 			Ar << E.CasePreservingHash;
 			E.bWereHashesLoaded = true;
-
-			if (Ar.UE4Ver() < VER_UE4_FIX_WIDE_STRING_CRC && E.IsWide())
-			{
-				// Regenerate hashes for packages saved before a fix to how hashes are generated
-				E.NonCasePreservingHash = GetRawNonCasePreservingHash(E.GetWideName());
-			}
 		}
 	}
 	else
@@ -1344,7 +2061,7 @@ void operator<<(FStructuredArchive::FSlot Slot, FNameEntrySerialized& E)
 		if (Slot.GetUnderlyingArchive().IsLoading())
 		{
 			// mark the name will be wide
-			E.PreSetIsWideForSerialization(true);
+			E.bIsWide = true;
 
 			// get the pointer to the wide array 
 			WIDECHAR* WideName = const_cast<WIDECHAR*>(E.GetWideName());
@@ -1360,110 +2077,6 @@ void operator<<(FStructuredArchive::FSlot Slot, FNameEntrySerialized& E)
 	}
 }
 
-/**
- * Pooled allocator for FNameEntry structures. Doesn't have to worry about freeing memory as those
- * never go away. It simply uses 64K chunks and allocates new ones as space runs out. This reduces
- * allocation overhead significantly (only minor waste on 64k boundaries) and also greatly helps
- * with fragmentation as 50-100k allocations turn into tens of allocations.
- */
-class FNameEntryPoolAllocator
-{
-public:
-	/** Initializes all member variables. */
-	FNameEntryPoolAllocator()
-	{
-		TotalAllocatedPages	= 0;
-		CurrentPoolStart	= NULL;
-		CurrentPoolEnd		= NULL;
-		check(ThreadGuard.GetValue() == 0);
-	}
-
-	/**
-	 * Allocates the requested amount of bytes and casts them to a FNameEntry pointer.
-	 *
-	 * @param   Size  Size in bytes to allocate
-	 * @return  Allocation of passed in size cast to a FNameEntry pointer.
-	 */
-	FNameEntry* Allocate( int32 Size )
-	{
-		check(ThreadGuard.Increment() == 1);
-		// Some platforms need all of the name entries to be aligned to 4 bytes, so by
-		// aligning the size here the next allocation will be aligned to 4
-		Size = Align( Size, alignof(FNameEntry) );
-
-		// Allocate a new pool if current one is exhausted. We don't worry about a little bit
-		// of waste at the end given the relative size of pool to average and max allocation.
-		if( CurrentPoolEnd - CurrentPoolStart < Size )
-		{
-			AllocateNewPool();
-		}
-		check( CurrentPoolEnd - CurrentPoolStart >= Size );
-		// Return current pool start as allocation and increment by size.
-		FNameEntry* NameEntry = (FNameEntry*) CurrentPoolStart;
-		CurrentPoolStart += Size;
-		check(ThreadGuard.Decrement() == 0);
-		return NameEntry;
-	}
-
-	/**
-	 * Returns the amount of memory to allocate for each page pool.
-	 *
-	 * @return  Page pool size.
-	 */
-	FORCEINLINE int32 PoolSize()
-	{
-		// Allocate in 64k chunks as it's ideal for page size.
-		return 256 * 1024;
-	}
-
-	/**
-	 * Returns the number of pages that have been allocated so far for names.
-	 *
-	 * @return  The number of pages allocated.
-	 */
-	FORCEINLINE int32 PageCount()
-	{
-		return TotalAllocatedPages;
-	}
-
-private:
-	/** Allocates a new pool. */
-	void AllocateNewPool()
-	{
-		TotalAllocatedPages++;
-		CurrentPoolStart = (uint8*) FMemory::Malloc(PoolSize());
-		CurrentPoolEnd = CurrentPoolStart + PoolSize();
-	}
-
-	/** Beginning of pool. Allocated by AllocateNewPool, incremented by Allocate.	*/
-	uint8* CurrentPoolStart;
-	/** End of current pool. Set by AllocateNewPool and checked by Allocate.		*/
-	uint8* CurrentPoolEnd;
-	/** Total number of pages that have been allocated.								*/
-	int32 TotalAllocatedPages;
-	/** Threadsafe counter to test for unwanted concurrency							*/
-	FThreadSafeCounter ThreadGuard;
-};
-
-/** Global allocator for name entries. */
-FNameEntryPoolAllocator GNameEntryPoolAllocator;
-
-template<typename TCharType>
-FNameEntry* AllocateNameEntry(const TCharType* Name, NAME_INDEX Index)
-{
-	LLM_SCOPE(ELLMTag::FName);
-
-	const SIZE_T NameLen  = TCString<TCharType>::Strlen(Name);
-	int32 NameEntrySize	  = FNameInitHelper<TCharType>::GetSize( NameLen );
-	FNameEntry* NameEntry = GNameEntryPoolAllocator.Allocate( NameEntrySize );
-	FName::NameEntryMemorySize += NameEntrySize;
-	FPlatformAtomics::InterlockedExchange(&NameEntry->Index, (Index << NAME_INDEX_SHIFT) | (FNameInitHelper<TCharType>::GetIndexShiftValue()));
-	NameEntry->HashNext.Store(nullptr, EMemoryOrder::Relaxed);
-	FNameInitHelper<TCharType>::SetNameString(NameEntry, Name, NameLen);
-	IncrementNameCount<TCharType>();
-	return NameEntry;
-}
-
 
 #if !UE_BUILD_SHIPPING && !UE_BUILD_TEST
 
@@ -1581,4 +2194,3 @@ void CallNameCreationHook()
 }
 
 #endif
-
diff --git a/Engine/Source/Runtime/Core/Public/Modules/Boilerplate/ModuleBoilerplate.h b/Engine/Source/Runtime/Core/Public/Modules/Boilerplate/ModuleBoilerplate.h
--- a/Engine/Source/Runtime/Core/Public/Modules/Boilerplate/ModuleBoilerplate.h
+++ b/Engine/Source/Runtime/Core/Public/Modules/Boilerplate/ModuleBoilerplate.h
@@ -49,7 +49,7 @@ class FChunkedFixedUObjectArray;
 #define UE4_VISUALIZERS_HELPERS
 #else
 #define UE4_VISUALIZERS_HELPERS \
-	FNameEntry*** GFNameTableForDebuggerVisualizers_MT = FName::GetNameTableForDebuggerVisualizers_MT(); \
+	uint8** GNameBlocksDebug = FName::GetBlocksForDebugVisualizer(); \
 	FChunkedFixedUObjectArray*& GObjectArrayForDebugVisualizers = GCoreObjectArrayForDebugVisualizers;
 #endif
 
diff --git a/Engine/Source/Runtime/Core/Public/Stats/Stats2.h b/Engine/Source/Runtime/Core/Public/Stats/Stats2.h
--- a/Engine/Source/Runtime/Core/Public/Stats/Stats2.h
+++ b/Engine/Source/Runtime/Core/Public/Stats/Stats2.h
@@ -651,7 +651,7 @@ private:
 	/** For FName. */
 	CORE_API const FString GetName() const
 	{
-		return FName::SafeString( (int32)Cycles );
+		return FName::SafeString(FNameEntryId::FromUnstableInt(static_cast<uint32>(Cycles)));
 	}
 };
 
diff --git a/Engine/Source/Runtime/Core/Public/Stats/StatsFile.h b/Engine/Source/Runtime/Core/Public/Stats/StatsFile.h
--- a/Engine/Source/Runtime/Core/Public/Stats/StatsFile.h
+++ b/Engine/Source/Runtime/Core/Public/Stats/StatsFile.h
@@ -363,8 +363,10 @@ protected:
 	{
 		FName RawName = NameAndInfo.GetRawName();
 		bool bSendFName = !FNamesSent.Contains( RawName.GetComparisonIndex() );
-		int32 Index = RawName.GetComparisonIndex();
-		Ar << Index;
+
+		int32 ComparisonInt = RawName.GetComparisonIndex().ToUnstableInt();
+		Ar << ComparisonInt;
+
 		int32 Number = NameAndInfo.GetRawNumber();
 		if (bSendFName)
 		{
@@ -415,7 +417,7 @@ protected:
 	}
 
 	/** Set of names already sent. */
-	TSet<int32> FNamesSent;
+	TSet<FNameEntryId> FNamesSent;
 
 	/** Data to write. */
 	TArray<uint8> OutData;
@@ -612,7 +614,7 @@ public:
 	FStatsStreamHeader Header;
 
 	/** FNames have a different index on each machine, so we translate via this map. **/
-	TMap<int32, int32> FNamesIndexMap;
+	TMap<int32, FNameEntryId> FNamesIndexMap;
 
 	/** Array of stats frame info. Empty for the raw stats. */
 	TArray<FStatsFrameInfo> FramesInfo;
@@ -702,7 +704,7 @@ public:
 			{
 				if( FNamesIndexMap.Contains( Index ) )
 				{
-					const int32 MyIndex = FNamesIndexMap.FindChecked( Index );
+					const FNameEntryId MyIndex = FNamesIndexMap.FindChecked( Index );
 					TheFName = FName( MyIndex, MyIndex, 0 );
 				}
 				else
@@ -723,7 +725,7 @@ public:
 			}
 			if( FNamesIndexMap.Contains( Index ) )
 			{
-				const int32 MyIndex = FNamesIndexMap.FindChecked( Index );
+				const FNameEntryId MyIndex = FNamesIndexMap.FindChecked( Index );
 				TheFName = FName( MyIndex, MyIndex, 0 );
 			}
 			else
diff --git a/Engine/Source/Runtime/Core/Public/UObject/NameTypes.h b/Engine/Source/Runtime/Core/Public/UObject/NameTypes.h
--- a/Engine/Source/Runtime/Core/Public/UObject/NameTypes.h
+++ b/Engine/Source/Runtime/Core/Public/UObject/NameTypes.h
@@ -33,8 +33,50 @@ class FText;
 /** Maximum size of name. */
 enum {NAME_SIZE	= 1024};
 
-/** Name index. */
-typedef int32 NAME_INDEX;
+/** Opaque id to a deduplicated name */
+struct FNameEntryId
+{
+	FNameEntryId() : Value(0) {}
+	FNameEntryId(ENoInit) {}
+	
+	int32 CompareCreationOrder(FNameEntryId Rhs) const { return Value - Rhs.Value; };
+
+	/** Creation order*/
+	bool operator<(FNameEntryId Rhs) const { return Value < Rhs.Value; }
+	bool operator>(FNameEntryId Rhs) const { return Rhs.Value < Value; }
+	bool operator==(FNameEntryId Rhs) const { return Value == Rhs.Value; }
+	bool operator!=(FNameEntryId Rhs) const { return Value != Rhs.Value; }
+
+	explicit operator bool() const { return Value != 0; }
+
+	UE_DEPRECATED(4.23, "NAME_INDEX is replaced by FNameEntryId, which is no longer a contiguous integer. "
+						"Please use 'GetTypeHash(MyId)' instead of 'MyId' for hash functions. "
+						"ToUnstableInt() can be used in other cases if you *really* know what you are doing.")
+	operator int32() const;
+
+	/** Get process specific integer */
+	uint32 ToUnstableInt() const { return Value; }
+
+	/** Create from unstable int produced by this process */
+	CORE_API static FNameEntryId FromUnstableInt(uint32 UnstableInt);
+private:
+	uint32 Value;
+};
+
+CORE_API uint32 GetTypeHash(FNameEntryId Id);
+CORE_API bool operator==(FNameEntryId Id, EName Ename);
+inline bool operator==(EName Ename, FNameEntryId Id) { return Id == Ename; }
+
+/** Serialize as process specific unstable int */
+CORE_API FArchive& operator<<(FArchive& Ar, FNameEntryId& InId);
+
+/**
+ * Legacy typedef - this is no longer an index
+ *
+ * Use GetTypeHash(FName) or GetTypeHash(FNameEntryId) for hashing
+ * To compare with ENames use FName(EName) or FName::ToEName() instead
+ */
+typedef FNameEntryId NAME_INDEX;
 
 #define checkName checkSlow
 
@@ -78,18 +120,6 @@ enum class ENameCase : uint8
 	IgnoreCase,
 };
 
-namespace FNameDefs
-{
-#if !WITH_EDITORONLY_DATA
-	// Use a modest bucket count on consoles
-	static const uint32 NameHashBucketCount = 65536;
-#else
-	// On PC platform we use a large number of name hash buckets to accommodate the editor's
-	// use of FNames to store asset path and content tags
-	static const uint32 NameHashBucketCount = 65536;
-#endif
-}
-
 enum ELinkerNameTableConstructor    {ENAME_LinkerConstructor};
 
 /** Enumeration for finding name. */
@@ -111,40 +141,34 @@ enum EFindName
 	FNameEntry.
 ----------------------------------------------------------------------------*/
 
-/** 
- * Mask for index bit used to determine whether string is encoded as TCHAR or ANSICHAR. We don't
- * add an extra bool in order to keep the name size to a minimum and 2 billion names is impractical
- * so there are a few bits left in the index.
- */
-#define NAME_WIDE_MASK 0x1
-#define NAME_INDEX_SHIFT 1
+/** Implementation detail exposed for debug visualizers */
+struct FNameEntryHeader
+{
+	uint16 bIsWide : 1;
+#if WITH_CASE_PRESERVING_NAME
+	uint16 Len : 15;
+#else
+	static constexpr uint32 ProbeHashBits = 5;
+	uint16 LowercaseProbeHash : ProbeHashBits;
+	uint16 Len : 10;
+#endif
+};
 
 /**
  * A global name, as stored in the global name table.
  */
 struct FNameEntry
 {
-public:
-	/** Pointer to the next entry in this hash bin's linked list. */
-	TAtomic<FNameEntry*>		HashNext;
-
 private:
-	/** Index of name in hash. */
-	NAME_INDEX		Index;
-
-protected:
-	/** Name, variable-sized - note that AllocateNameEntry only allocates memory as needed. */
+#if WITH_CASE_PRESERVING_NAME
+	FNameEntryId ComparisonId;
+#endif
+	FNameEntryHeader Header;
 	union
 	{
 		ANSICHAR	AnsiName[NAME_SIZE];
 		WIDECHAR	WideName[NAME_SIZE];
 	};
-	// DO NOT ADD VARIABLES BELOW!
-
-private:
-	/** Default constructor doesn't do anything. AllocateNameEntry is responsible for work. */
-	FNameEntry()
-	{}
 
 	FNameEntry(const FNameEntry&) = delete;
 	FNameEntry(FNameEntry&&) = delete;
@@ -152,34 +176,12 @@ private:
 	FNameEntry& operator=(FNameEntry&&) = delete;
 
 public:
-
-	/** 
-	 * Sets whether or not the NameEntry will have a wide string, or an ansi string
-	 *
-	 * @param bIsWide true if we are going to serialize a wide string
-	 */
-	FORCEINLINE void PreSetIsWideForSerialization(bool bIsWide)
-	{
-		Index = bIsWide ? NAME_WIDE_MASK : 0;
-	}
-
-	/**
-	 * Returns index of name in hash passed to FNameEntry via AllocateNameEntry. The lower bits
-	 * are used for internal state, which is why we need to shift.
-	 *
-	 * @return Index of name in hash
-	 */
-	FORCEINLINE int32 GetIndex() const
-	{
-		return Index >> NAME_INDEX_SHIFT;
-	}
-
 	/**
-	 * Returns whether this name entry is represented via TCHAR or ANSICHAR
+	 * Returns whether this name entry is represented via WIDECHAR or ANSICHAR
 	 */
 	FORCEINLINE bool IsWide() const
 	{
-		return (Index & NAME_WIDE_MASK);
+		return Header.bIsWide;
 	}
 
 	/**
@@ -201,36 +203,20 @@ public:
 	*/
 	CORE_API void AppendNameToPathString( FString& String ) const;
 
-	/**
-	 * @return length of name
-	 */
-	CORE_API int32 GetNameLength() const;
-
-	/**
-	 * Compares name using the compare method provided.
-	 *
-	 * @param	InName	Name to compare to
-	 * @return	true if equal, false otherwise
-	 */
-	bool IsEqual( const ANSICHAR* InName, const ENameCase CompareMethod ) const;
-
-	/**
-	 * Compares name using the compare method provided.
-	 *
-	 * @param	InName	Name to compare to
-	 * @return	true if equal, false otherwise
-	 */
-	bool IsEqual( const WIDECHAR* InName, const ENameCase CompareMethod ) const;
+	int32 GetNameLength() const
+	{
+		return Header.Len;
+	}
 
 	/**
-	 * @return create a copy of the ANSI name
+	 * Copy name to a buffer to avoid allocation
 	 */
-	void GetAnsiName(ANSICHAR(&OutName)[NAME_SIZE]) const;
+	CORE_API void GetAnsiName(ANSICHAR(&OutName)[NAME_SIZE]) const;
 
 	/**
-	 * @return direct access to WIDE name if stored in widechars
+	 * Copy name to a buffer to avoid allocation
 	 */
-	void GetWideName(WIDECHAR(&OutName)[NAME_SIZE]) const;
+	CORE_API void GetWideName(WIDECHAR(&OutName)[NAME_SIZE]) const;
 
 	static CORE_API int32 GetSize( const TCHAR* Name );
 
@@ -241,43 +227,33 @@ public:
 	 * @param	bIsPureAnsi		Whether name is pure ANSI or not
 	 * @return	required size of FNameEntry structure to hold this string (might be wide or ansi)
 	 */
-	static int32 GetSize( int32 Length, bool bIsPureAnsi );
+	static int32 GetSize(int32 Length, bool bIsPureAnsi);
+
+	CORE_API int32 GetSizeInBytes() const;
 
 	// Functions.
 	CORE_API void Write(FArchive& Ar) const;
 	CORE_API void Write(FStructuredArchive::FSlot Slot) const;
 
-	// Friend for access to Flags.
-	template<typename TCharType>
-	friend FNameEntry* AllocateNameEntry(const TCharType* Name, NAME_INDEX Index);
-
+	static int32 GetDataOffset();
+	struct FNameStringView MakeView(union FNameBuffer& OptionalDecodeBuffer) const;
 private:
 	friend class FName;
-
-	friend struct FNameEntrySerialized;
-
-	template <typename TCharType>
-	friend struct FNameInitHelper;
-
-	/**
-	 * @return direct access to ANSI name if stored in ANSI
-	 */
-	ANSICHAR const* GetAnsiNamePtr(ANSICHAR(&OptionalTempBuffer)[NAME_SIZE]) const;
-
-	/**
-	 * @return direct access to wide name if stored in widechars
-	 */
-	WIDECHAR const* GetWideNamePtr(WIDECHAR(&OptionalTempBuffer)[NAME_SIZE]) const;
-
-	/**
-	 * updates the stored name if ANSI
-	 */
-	void SetAnsiName(const ANSICHAR* SrcName, int32 SrcNameLen);
-
-	/**
-	 * updates the stored name in widechars
-	 */
-	void SetWideName(const WIDECHAR* SrcName, int32 SrcNameLen);
+	friend struct FNameHelper;
+	friend class FNameEntryAllocator;
+	friend class FNamePoolShardBase;
+
+	static void Encode(ANSICHAR* Name, uint32 Len);
+	static void Encode(WIDECHAR* Name, uint32 Len);
+	static void Decode(ANSICHAR* Name, uint32 Len);
+	static void Decode(WIDECHAR* Name, uint32 Len);
+
+	void StoreName(const ANSICHAR* InName, uint32 Len);
+	void StoreName(const WIDECHAR* InName, uint32 Len);
+	void CopyUnterminatedName(ANSICHAR* OutName) const;
+	void CopyUnterminatedName(WIDECHAR* OutName) const;
+	const ANSICHAR* GetUnterminatedName(ANSICHAR(&OptionalDecodeBuffer)[NAME_SIZE]) const;
+	const WIDECHAR* GetUnterminatedName(WIDECHAR(&OptionalDecodeBuffer)[NAME_SIZE]) const;
 };
 
 /**
@@ -285,7 +261,8 @@ private:
  */
 struct FNameEntrySerialized
 {
-	NAME_INDEX Index;
+	FNameEntryId Index;
+	bool bIsWide = false;
 
 	union
 	{
@@ -293,248 +270,61 @@ struct FNameEntrySerialized
 		WIDECHAR	WideName[NAME_SIZE];
 	};
 
-	uint16 NonCasePreservingHash;
-	uint16 CasePreservingHash;
-	bool bWereHashesLoaded;
+	// These are not used anymore but recalculated on save to maintain serialization format
+	uint16 NonCasePreservingHash = 0;
+	uint16 CasePreservingHash = 0;
+	bool bWereHashesLoaded = false;
 
 	FNameEntrySerialized(const FNameEntry& NameEntry);
-	FNameEntrySerialized(enum ELinkerNameTableConstructor) :
-		NonCasePreservingHash(0),
-		CasePreservingHash(0),
-		bWereHashesLoaded(false)
-	{
-	}
-
-	/** 
-	 * Sets whether or not the NameEntry will have a wide string, or an ansi string
-	 *
-	 * @param bIsWide true if we are going to serialize a wide string
-	 */
-	FORCEINLINE void PreSetIsWideForSerialization(bool bIsWide)
-	{
-		Index = bIsWide ? NAME_WIDE_MASK : 0;
-	}
+	FNameEntrySerialized(enum ELinkerNameTableConstructor) {}
 
 	/**
-	 * Returns whether this name entry is represented via TCHAR or ANSICHAR
+	 * Returns direct access to null-terminated name if narrow
 	 */
-	FORCEINLINE bool IsWide() const
+	ANSICHAR const* GetAnsiName() const
 	{
-		return (Index & NAME_WIDE_MASK);
-	}
-
-	/**
-	 * @return direct access to ANSI name if stored in ANSI
-	 */
-	inline ANSICHAR const* GetAnsiName() const
-	{
-		check(!IsWide());
+		check(!bIsWide);
 		return AnsiName;
 	}
 
 	/**
-	 * @return direct access to wide name if stored in widechars
+	 * Returns direct access to null-terminated name if wide
 	 */
-	inline WIDECHAR const* GetWideName() const
+	WIDECHAR const* GetWideName() const
 	{
-		check(IsWide());
+		check(bIsWide);
 		return WideName;
 	}
 
 	/**
-	 * @return FString of name portion minus number.
+	 * Returns FString of name portion minus number.
 	 */
 	CORE_API FString GetPlainNameString() const;	
 
 	friend CORE_API FArchive& operator<<(FArchive& Ar, FNameEntrySerialized& E);
-	friend CORE_API FArchive& operator<<(FArchive& Ar, FNameEntrySerialized* E)
+	friend FArchive& operator<<(FArchive& Ar, FNameEntrySerialized* E)
 	{
 		return Ar << *E;
 	}
 
 	friend CORE_API void operator<<(FStructuredArchive::FSlot Slot, FNameEntrySerialized& E);
-	friend CORE_API void operator<<(FStructuredArchive::FSlot Slot, FNameEntrySerialized* E)
+	friend void operator<<(FStructuredArchive::FSlot Slot, FNameEntrySerialized* E)
 	{
 		Slot << *E;
 	}
 };
 
-/**
- * Simple array type that can be expanded without invalidating existing entries.
- * This is critical to thread safe FNames.
- * @param ElementType Type of the pointer we are storing in the array
- * @param MaxTotalElements absolute maximum number of elements this array can ever hold
- * @param ElementsPerChunk how many elements to allocate in a chunk
- **/
- template<typename ElementType, int32 MaxTotalElements, int32 ElementsPerChunk>
-class TStaticIndirectArrayThreadSafeRead
-{
-	enum
-	{
-		// figure out how many elements we need in the master table
-		ChunkTableSize = (MaxTotalElements + ElementsPerChunk - 1) / ElementsPerChunk
-	};
-	/** Static master table to chunks of pointers **/
-	ElementType** Chunks[ChunkTableSize];
-	/** Number of elements we currently have **/
-	TAtomic<int32> NumElements;
-	/** Number of chunks we currently have **/
-	int32 NumChunks;
-
-	/**
-	 * Expands the array so that Element[Index] is allocated. New pointers are all zero.
-	 * @param Index The Index of an element we want to be sure is allocated
-	 **/
-	void ExpandChunksToIndex(int32 Index)
-	{
-		check(Index >= 0 && Index < MaxTotalElements);
-		int32 ChunkIndex = Index / ElementsPerChunk;
-		while (1)
-		{
-			if (ChunkIndex < NumChunks)
-			{
-				break;
-			}
-			// add a chunk, and make sure nobody else tries
-			ElementType*** Chunk = &Chunks[ChunkIndex];
-			ElementType** NewChunk = (ElementType**)FMemory::Malloc(sizeof(ElementType*) * ElementsPerChunk);
-			FMemory::Memzero(NewChunk, sizeof(ElementType*) * ElementsPerChunk);
-			if (FPlatformAtomics::InterlockedCompareExchangePointer((void**)Chunk, NewChunk, nullptr))
-			{
-				// someone else beat us to the add, we don't support multiple concurrent adds
-				check(0)
-			}
-			else
-			{
-				NumChunks++;
-			}
-		}
-		check(ChunkIndex < NumChunks && Chunks[ChunkIndex]); // should have a valid pointer now
-	}
-
-	/**
-	 * Return a pointer to the pointer to a given element
-	 * @param Index The Index of an element we want to retrieve the pointer-to-pointer for
-	 **/
-	FORCEINLINE_DEBUGGABLE ElementType const* const* GetItemPtr(int32 Index) const
-	{
-		int32 ChunkIndex = Index / ElementsPerChunk;
-		int32 WithinChunkIndex = Index % ElementsPerChunk;
-		checkf(IsValidIndex(Index), TEXT("IsValidIndex(%d)"), Index);
-		checkf(ChunkIndex < NumChunks, TEXT("ChunkIndex (%d) < NumChunks (%d)"), ChunkIndex, NumChunks);
-		checkf(Index < MaxTotalElements, TEXT("Index (%d) < MaxTotalElements (%d)"), Index, MaxTotalElements);
-		ElementType** Chunk = Chunks[ChunkIndex];
-		check(Chunk);
-		return Chunk + WithinChunkIndex;
-	}
-
-public:
-	/** Constructor : Probably not thread safe **/
-	TStaticIndirectArrayThreadSafeRead()
-		: NumElements(0)
-		, NumChunks(0)
-	{
-		FMemory::Memzero(Chunks);
-	}
-	/** 
-	 * Return the number of elements in the array 
-	 * Thread safe, but you know, someone might have added more elements before this even returns
-	 * @return	the number of elements in the array
-	**/
-	FORCEINLINE int32 Num() const
-	{
-		return NumElements.Load(EMemoryOrder::Relaxed);
-	}
-	/** 
-	 * Return if this index is valid
-	 * Thread safe, if it is valid now, it is valid forever. Other threads might be adding during this call.
-	 * @param	Index	Index to test
-	 * @return	true, if this is a valid
-	**/
-	FORCEINLINE bool IsValidIndex(int32 Index) const
-	{
-		return Index < Num() && Index >= 0;
-	}
-	/** 
-	 * Return a reference to an element
-	 * @param	Index	Index to return
-	 * @return	a reference to the pointer to the element
-	 * Thread safe, if it is valid now, it is valid forever. This might return nullptr, but by then, some other thread might have made it non-nullptr.
-	**/
-	FORCEINLINE ElementType const* const& operator[](int32 Index) const
-	{
-		ElementType const* const* ItemPtr = GetItemPtr(Index);
-		check(ItemPtr);
-		return *ItemPtr;
-	}
-	/** 
-	 * Add more elements to the array
-	 * @param	NumToAdd	Number of elements to add
-	 * @return	the number of elements in the container before we did the add. In other words, the add index.
-	 * Not thread safe. This should only be called by one thread, but the other methods can be called while this is going on.
-	**/
-	int32 AddZeroed(int32 NumToAdd)
-	{
-		int32 Result = NumElements;
-		check(Result + NumToAdd <= MaxTotalElements);
-		ExpandChunksToIndex(Result + NumToAdd - 1);
-		NumElements += NumToAdd;
-		return Result;
-	}
-	/** 
-	 * Return a naked pointer to the fundamental data structure for debug visualizers.
-	**/
-	ElementType*** GetRootBlockForDebuggerVisualizers()
-	{
-		return Chunks;
-	}
-	/**
-	* Make sure chunks are allocated to hold the specified capacity of items. This is NOT thread safe.
-	**/
-	void Reserve(int32 Capacity)
-	{
-		check(Capacity >= 0 && Capacity <= MaxTotalElements);
-		if (Capacity > NumElements.Load(EMemoryOrder::Relaxed))
-		{			
-			int32 MaxChunks = (Capacity + ElementsPerChunk - 1) / ElementsPerChunk;
-			check(MaxChunks >= NumChunks);
-			for (int32 ChunkIndex = 0; ChunkIndex < MaxChunks; ++ChunkIndex)
-			{
-				if (!Chunks[ChunkIndex])
-				{
-					ElementType** NewChunk = (ElementType**)FMemory::Malloc(sizeof(ElementType*) * ElementsPerChunk);
-					FMemory::Memzero(NewChunk, sizeof(ElementType*) * ElementsPerChunk);
-					Chunks[ChunkIndex] = NewChunk;
-				}
-			}
-			NumChunks = MaxChunks;
-		}
-	}
-};
-
-// Typedef for the threadsafe master name table. 
-// CAUTION: If you change those constants, you probably need to update the debug visualizers.
-typedef TStaticIndirectArrayThreadSafeRead<FNameEntry, 4 * 1024 * 1024 /* 4M unique FNames */, 16384 /* allocated in 64K/128K chunks */ > TNameEntryArray;
-
 /**
  * The minimum amount of data required to reconstruct a name
  * This is smaller than FName, but you lose the case-preserving behavior
  */
-struct CORE_API FMinimalName
+struct FMinimalName
 {
-	FORCEINLINE FMinimalName()
-		: Index(0)
-		, Number(NAME_NO_NUMBER_INTERNAL)
-	{
-	}
+	FMinimalName() {}
 
-	FORCEINLINE FMinimalName(const EName N)
-		: Index(N)
-		, Number(NAME_NO_NUMBER_INTERNAL)
-	{
-	}
+	CORE_API FMinimalName(EName N);
 
-	FORCEINLINE FMinimalName(const NAME_INDEX InIndex, const int32 InNumber)
+	FMinimalName(FNameEntryId InIndex, int32 InNumber)
 		: Index(InIndex)
 		, Number(InNumber)
 	{
@@ -542,13 +332,13 @@ struct CORE_API FMinimalName
 
 	FORCEINLINE bool IsNone() const
 	{
-		return Index == 0 && Number == 0;
+		return Index && Number == NAME_NO_NUMBER_INTERNAL;
 	}
 
 	/** Index into the Names array (used to find String portion of the string/number pair) */
-	NAME_INDEX		Index;
+	FNameEntryId	Index;
 	/** Number portion of the string/number pair (stored internally as 1 more than actual, so zero'd memory will be the default, no-instance case) */
-	int32			Number;
+	int32			Number = NAME_NO_NUMBER_INTERNAL;
 };
 
 /**
@@ -556,35 +346,25 @@ struct CORE_API FMinimalName
  * This will be the same size as FName when WITH_CASE_PRESERVING_NAME is 1, and is used to store an FName in cases where 
  * the size of FName must be constant between build configurations (eg, blueprint bytecode)
  */
-struct CORE_API FScriptName
+struct FScriptName
 {
-	FORCEINLINE FScriptName()
-		: ComparisonIndex(0)
-		, DisplayIndex(0)
-		, Number(NAME_NO_NUMBER_INTERNAL)
-	{
-	}
+	FScriptName() {}
 
-	FORCEINLINE FScriptName(const EName N)
-		: ComparisonIndex(N)
-		, DisplayIndex(N)
-		, Number(NAME_NO_NUMBER_INTERNAL)
-	{
-	}
+	CORE_API FScriptName(EName N);
 
-	FORCEINLINE FScriptName(const NAME_INDEX InComparisonIndex, const NAME_INDEX InDisplayIndex, const int32 InNumber)
-		: ComparisonIndex(InComparisonIndex)
-		, DisplayIndex(InDisplayIndex)
-		, Number(InNumber)
+	FScriptName(FNameEntryId InComparisonIndex, FNameEntryId InDisplayIndex, int32 InNumber)
+	: ComparisonIndex(InComparisonIndex)
+	, DisplayIndex(InDisplayIndex)
+	, Number(InNumber)
 	{
 	}
 
 	/** Index into the Names array (used to find String portion of the string/number pair used for comparison) */
-	NAME_INDEX		ComparisonIndex;
+	FNameEntryId	ComparisonIndex;
 	/** Index into the Names array (used to find String portion of the string/number pair used for display) */
-	NAME_INDEX		DisplayIndex;
+	FNameEntryId	DisplayIndex;
 	/** Number portion of the string/number pair (stored internally as 1 more than actual, so zero'd memory will be the default, no-instance case) */
-	uint32			Number;
+	uint32			Number = NAME_NO_NUMBER_INTERNAL;
 };
 
 /**
@@ -595,20 +375,16 @@ struct CORE_API FScriptName
 class CORE_API FName
 {
 public:
-
 	FORCEINLINE NAME_INDEX GetComparisonIndex() const
 	{
-		const NAME_INDEX Index = GetComparisonIndexFast();
-		checkName(Index >= 0 && Index < GetNames().Num());
-		checkName(GetNames()[Index]);
-		return Index;
+		checkName(IsWithinBounds(ComparisonIndex));
+		return ComparisonIndex;
 	}
 
 	FORCEINLINE NAME_INDEX GetDisplayIndex() const
 	{
 		const NAME_INDEX Index = GetDisplayIndexFast();
-		checkName(Index >= 0 && Index < GetNames().Num());
-		checkName(GetNames()[Index]);
+		checkName(IsWithinBounds(Index));
 		return Index;
 	}
 
@@ -623,31 +399,17 @@ public:
 	}
 	
 	/** Returns the pure name string without any trailing numbers */
-	FString GetPlainNameString() const
-	{
-		if (const FNameEntry* CurEntry = GetDisplayNameEntry())
-		{
-			return CurEntry->GetPlainNameString();
-		}
-
-		return TEXT("*INVALID*");
-	}
+	FString GetPlainNameString() const;
 
 	/**
 	 * Returns the underlying ANSI string pointer.  No allocations.  Will fail if this is actually a wide name.
 	 */
-	FORCEINLINE void GetPlainANSIString(ANSICHAR (&AnsiName)[NAME_SIZE]) const
-	{
-		GetDisplayNameEntry()->GetAnsiName(AnsiName);
-	}
+	void GetPlainANSIString(ANSICHAR(&AnsiName)[NAME_SIZE]) const;
 
 	/**
 	 * Returns the underlying WIDE string pointer.  No allocations.  Will fail if this is actually an ANSI name.
 	 */
-	FORCEINLINE void GetPlainWIDEString(WIDECHAR (&WideName)[NAME_SIZE]) const
-	{
-		GetDisplayNameEntry()->GetWideName(WideName);
-	}
+	void GetPlainWIDEString(WIDECHAR(&WideName)[NAME_SIZE]) const;
 
 	const FNameEntry* GetComparisonNameEntry() const;
 	const FNameEntry* GetDisplayNameEntry() const;
@@ -678,13 +440,13 @@ public:
 	 */
 	FORCEINLINE bool IsEqual(const FName& Other, const ENameCase CompareMethod = ENameCase::IgnoreCase, const bool bCompareNumber = true ) const
 	{
-		return ((CompareMethod == ENameCase::IgnoreCase) ? GetComparisonIndexFast() == Other.GetComparisonIndexFast() : GetDisplayIndexFast() == Other.GetDisplayIndexFast())
+		return ((CompareMethod == ENameCase::IgnoreCase) ? ComparisonIndex == Other.ComparisonIndex : GetDisplayIndexFast() == Other.GetDisplayIndexFast())
 			&& (!bCompareNumber || GetNumber() == Other.GetNumber());
 	}
 
 	FORCEINLINE bool operator==(const FName& Other) const
 	{
-		return (GetComparisonIndexFast() == Other.GetComparisonIndexFast()) & (GetNumber() == Other.GetNumber());
+		return (ComparisonIndex == Other.ComparisonIndex) & (GetNumber() == Other.GetNumber());
 	}
 	FORCEINLINE bool operator!=(const FName& Other) const
 	{
@@ -709,24 +471,20 @@ public:
 
 	FORCEINLINE bool IsNone() const
 	{
-		return GetComparisonIndexFast() == 0 && GetNumber() == 0;
-	}
-
-	FORCEINLINE bool IsValid() const
-	{
-		TNameEntryArray& Names = GetNames();
-		return GetComparisonIndexFast() >= 0 && GetComparisonIndexFast() < Names.Num() && Names[GetComparisonIndexFast()] != nullptr
-			&& GetDisplayIndexFast() >= 0 && GetDisplayIndexFast() < Names.Num() && Names[GetDisplayIndexFast()] != nullptr;
+		return !ComparisonIndex && GetNumber() == NAME_NO_NUMBER_INTERNAL;
 	}
 
 	/**
-	 * Helper function to check if the index is valid. Does not check if the name itself is valid.
-	 */
-	FORCEINLINE bool IsValidIndexFast() const
-	{
-		return GetComparisonIndexFast() >= 0 && GetComparisonIndexFast() < GetNames().Num()
-			&& GetDisplayIndexFast() >= 0 && GetDisplayIndexFast() < GetNames().Num();
-	}
+	* Paranoid sanity check
+	*
+	* All FNames are valid except for stomped memory, dangling pointers, etc.
+	* Should only be used to  investigate such bugs and not in production code.
+	*/
+	bool IsValid() const { return IsWithinBounds(ComparisonIndex); }
+
+	/** Paranoid sanity check, same as IsValid() */
+	bool IsValidIndexFast() const { return IsValid(); }
+
 
 	/**
 	 * Checks to see that a given name-like string follows the rules that Unreal requires.
@@ -792,20 +550,6 @@ public:
 		return IsValidXName(ToString(), INVALID_LONGPACKAGE_CHARACTERS, &OutReason);
 	}
 
-#ifdef IMPLEMENT_ASSIGNMENT_OPERATOR_MANUALLY
-	// Assignment operator
-	FORCEINLINE FName& operator=(const FName& Other)
-	{
-		this->ComparisonIndex = Other.ComparisonIndex;
-#if WITH_CASE_PRESERVING_NAME
-		this->DisplayIndex = Other.DisplayIndex;
-#endif
-		this->Number = Other.Number;
-
-		return *this;
-	}
-#endif
-
 	/**
 	 * Compares name to passed in one. Sort is alphabetical ascending.
 	 *
@@ -822,13 +566,12 @@ public:
 	 */
 	FORCEINLINE int32 CompareIndexes(const FName& Other) const
 	{
-		int32 ComparisonDiff = GetComparisonIndexFast() - Other.GetComparisonIndexFast();
-
-		if (ComparisonDiff == 0)
+		if (int32 ComparisonDiff = ComparisonIndex.CompareCreationOrder(Other.ComparisonIndex))
 		{
-			return GetNumber() - Other.GetNumber();
+			return ComparisonDiff;
 		}
-		return ComparisonDiff;
+
+		return GetNumber() - Other.GetNumber();
 	}
 
 	/**
@@ -836,15 +579,7 @@ public:
 	 *
 	 * @param N The hardcoded value the string portion of the name will have. The number portion will be NAME_NO_NUMBER
 	 */
-	FORCEINLINE FName( EName N )
-		: ComparisonIndex( N )
-#if WITH_CASE_PRESERVING_NAME
-		, DisplayIndex( N )
-#endif
-		, Number( NAME_NO_NUMBER_INTERNAL )
-	{
-		check(N < NAME_MaxHardcodedNameIndex);
-	}
+	FName(EName N);
 
 	/**
 	 * Create an FName with a hardcoded string index and (instance).
@@ -852,18 +587,7 @@ public:
 	 * @param N The hardcoded value the string portion of the name will have
 	 * @param InNumber The hardcoded value for the number portion of the name
 	 */
-	FORCEINLINE FName( EName N, int32 InNumber )
-		: ComparisonIndex( N )
-#if WITH_CASE_PRESERVING_NAME
-		, DisplayIndex( N )
-#endif
-		, Number( InNumber )
-	{
-		// If this fires the enum was out of bounds - did you pass an index instead?
-		// If you want to clone an FName with a new number, FName( const FName& Other, int32 InNumber ) is the thing you want
-		check(N < NAME_MaxHardcodedNameIndex);
-		check(InNumber >= 0 && InNumber <= 0xffffff);
-	}
+	FName(EName N, int32 InNumber);
 
 	/**
 	 * Create an FName from an existing string, but with a different instance.
@@ -893,15 +617,21 @@ public:
 	{
 	}
 
+	static FNameEntryId GetComparisonIdFromDisplayId(FNameEntryId DisplayId);
+
+	/**
+	 * Only call this if you *really* know what you're doing
+	 */
+	static FName CreateFromDisplayId(FNameEntryId DisplayId, int32 Number)
+	{
+		return FName(GetComparisonIdFromDisplayId(DisplayId), DisplayId, Number);
+	}
+
 	/**
 	 * Default constructor, initialized to None
 	 */
 	FORCEINLINE FName()
-		: ComparisonIndex( 0 )
-#if WITH_CASE_PRESERVING_NAME
-		, DisplayIndex( 0 )
-#endif
-		, Number( NAME_NO_NUMBER_INTERNAL )
+		: Number(NAME_NO_NUMBER_INTERNAL)
 	{
 	}
 
@@ -909,8 +639,11 @@ public:
 	 * Scary no init constructor, used for something obscure in UObjectBase
 	 */
 	explicit FName(ENoInit)
-	{
-	}
+		: ComparisonIndex(NoInit)
+#if WITH_CASE_PRESERVING_NAME
+		, DisplayIndex(NoInit)
+#endif
+	{}
 
 	/**
 	 * Create an FName. If FindType is FNAME_Find, and the string part of the name 
@@ -922,9 +655,23 @@ public:
 	FName(const WIDECHAR* Name, EFindName FindType=FNAME_Add);
 	FName(const ANSICHAR* Name, EFindName FindType=FNAME_Add);
 
-	// Deprecated bUnused
-	UE_DEPRECATED(4.12, "Removed bUnused from FName") FName(const WIDECHAR* Name, EFindName FindType, bool bUnused) : FName(Name, FindType) { }
-	UE_DEPRECATED(4.12, "Removed bUnused from FName") FName(const ANSICHAR* Name, EFindName FindType, bool bUnused) : FName(Name, FindType) { }
+	/** Create FName from non-null string with known length  */
+	FName(int32 Len, const WIDECHAR* Name, EFindName FindType=FNAME_Add);
+	FName(int32 Len, const ANSICHAR* Name, EFindName FindType=FNAME_Add);
+
+	/**
+	 * Create an FName. If FindType is FNAME_Find, and the string part of the name 
+	 * doesn't already exist, then the name will be NAME_None
+	 *
+	 * @param Name Value for the string portion of the name
+	 * @param Number Value for the number portion of the name
+	 * @param FindType Action to take (see EFindName)
+	 * @param bSplitName true if the trailing number should be split from the name when Number == NAME_NO_NUMBER_INTERNAL, or false to always use the name as-is
+	 */
+	FName(const WIDECHAR* Name, int32 InNumber, EFindName FindType = FNAME_Add);
+	FName(const ANSICHAR* Name, int32 InNumber, EFindName FindType = FNAME_Add);
+	FName(int32 Len, const WIDECHAR* Name, int32 Number, EFindName FindType = FNAME_Add);
+	FName(int32 Len, const ANSICHAR* Name, int32 InNumber, EFindName FindType = FNAME_Add);
 
 	/**
 	 * Create an FName. If FindType is FNAME_Find, and the string part of the name 
@@ -935,7 +682,7 @@ public:
 	 * @param FindType Action to take (see EFindName)
 	 * @param bSplitName true if the trailing number should be split from the name when Number == NAME_NO_NUMBER_INTERNAL, or false to always use the name as-is
 	 */
-	FName( const TCHAR* Name, int32 InNumber, EFindName FindType=FNAME_Add, const bool bSplitName=true );
+	FName( const TCHAR* Name, int32 InNumber, EFindName FindType, bool bSplitName);
 
 	/**
 	 * Constructor used by FLinkerLoad when loading its name table; Creates an FName with an instance
@@ -943,68 +690,15 @@ public:
 	 * this version skips calculating the hashes of the names if possible
 	 */
 	FName(const FNameEntrySerialized& LoadedEntry);
-
-	/**
-	 * Create an FName with a hardcoded string index.
-	 *
-	 * @param HardcodedIndex	The hardcoded value the string portion of the name will have. 
-	 * @param Name				The hardcoded name to initialize
-	 */
-	explicit FName( EName HardcodedIndex, const TCHAR* Name );
-
+	
 	/**
 	 * Equality operator.
 	 *
 	 * @param	Other	String to compare this name to
 	 * @return true if name matches the string, false otherwise
 	 */
-	template <typename CharType>
-	bool operator==(const CharType* Other) const
-	{
-		// Make NAME_None == TEXT("") or nullptr consistent with NAME_None == FName(TEXT("")) or FName(nullptr)
-		if (Other == nullptr || Other[0] == 0)
-		{
-			return (*this == NAME_None);
-		}
-
-		// Find name entry associated with this FName.
-		bool bAreNamesMatching = false;
-		const FNameEntry* const Entry = GetComparisonNameEntry();
-
-		if (Entry != nullptr)
-		{
-			// Temporary buffer to hold split name in case passed in name is of Name_Number format.
-			WIDECHAR TempBuffer[NAME_SIZE];
-			int32 InNumber = NAME_NO_NUMBER_INTERNAL;
-			int32 TempNumber = NAME_NO_NUMBER_INTERNAL;
-
-			// Check whether we need to split the passed in string into name and number portion.
-			auto WideOther = StringCast<WIDECHAR>(Other);
-			const WIDECHAR* WideOtherPtr = WideOther.Get();
-			if (SplitNameWithCheck(WideOtherPtr, TempBuffer, ARRAY_COUNT(TempBuffer), TempNumber))
-			{
-				WideOtherPtr = TempBuffer;
-				InNumber = NAME_EXTERNAL_TO_INTERNAL(TempNumber);
-			}
-
-			// Report a match if both the number and string portion match.
-			if (InNumber == GetNumber())
-			{
-				if (Entry->IsWide())
-				{
-					WIDECHAR EntryTempBuffer[NAME_SIZE];
-					bAreNamesMatching = !FPlatformString::Stricmp(WideOtherPtr, Entry->GetWideNamePtr(EntryTempBuffer));
-				}
-				else
-				{
-					ANSICHAR EntryTempBuffer[NAME_SIZE];
-					bAreNamesMatching = !FPlatformString::Stricmp(WideOtherPtr, Entry->GetAnsiNamePtr(EntryTempBuffer));
-				}
-			}
-		}
-
-		return bAreNamesMatching;
-	}
+	bool operator==(const ANSICHAR* Other) const;
+	bool operator==(const WIDECHAR* Other) const;
 
 	/**
 	 * Inequality operator.
@@ -1018,76 +712,38 @@ public:
 		return !operator==(Other);
 	}
 
-	template <typename TCharType>
-	static uint16 GetCasePreservingHash(const TCharType* Source);
-	template <typename TCharType>
-	static uint16 GetNonCasePreservingHash(const TCharType* Source);
-
 	static void DisplayHash( class FOutputDevice& Ar );
-	static FString SafeString( int32 InDisplayIndex, int32 InstanceNumber=NAME_NO_NUMBER_INTERNAL )
-	{
-		TNameEntryArray& Names = GetNames();
-		return GetIsInitialized()
-			? (Names.IsValidIndex(InDisplayIndex) && Names[InDisplayIndex])
-				? FName(InDisplayIndex, InDisplayIndex, InstanceNumber).ToString()
-				: FString(TEXT("*INVALID*"))
-			: FString(TEXT("*UNINITIALIZED*"));
-	}
-	static int32 GetMaxNames()
-	{
-		return GetNames().Num();
-	}
+	static FString SafeString(NAME_INDEX InDisplayIndex, int32 InstanceNumber = NAME_NO_NUMBER_INTERNAL);
+
 	/**
 	 * @return Size of all name entries.
 	 */
-	static int32 GetNameEntryMemorySize()
-	{
-		return NameEntryMemorySize;
-	}
+	static int32 GetNameEntryMemorySize();
+
 	/**
 	* @return Size of Name Table object as a whole
 	*/
-	static int32 GetNameTableMemorySize()
-	{
-		return GetNameEntryMemorySize() + (GetMaxNames() * sizeof(FNameEntry*)) + sizeof(NameHashHead) + sizeof(NameHashTail);
-	}
+	static int32 GetNameTableMemorySize();
 
 	/**
 	 * @return number of ansi names in name table
 	 */
-	static int32 GetNumAnsiNames()
-	{
-		return NumAnsiNames;
-	}
+	static int32 GetNumAnsiNames();
+
 	/**
 	 * @return number of wide names in name table
 	 */
-	static int32 GetNumWideNames()
-	{
-		return NumWideNames;
-	}
-	static FNameEntry const* GetEntry( int i )
-	{
-		return GetNames()[i];
-	}
-	//@}
+	static int32 GetNumWideNames();
 
-	/**
-	 * Helper function to split an old-style name (Class_Number, ie Rocket_17) into
-	 * the component parts usable by new-style FNames. Only use results if this function
-	 * returns true.
-	 *
-	 * @param OldName		Old-style name
-	 * @param NewName		Output string portion of the name/number pair
-	 * @param NewNameLen	Size of NewName buffer (in TCHAR units)
-	 * @param NewNumber		Number portion of the name/number pair
-	 *
-	 * @return true if the name was split, only then will NewName/NewNumber have valid values
-	 */
-	static bool SplitNameWithCheck(const WIDECHAR* OldName, WIDECHAR* NewName, int32 NewNameLen, int32& NewNumber);
+	static TArray<const FNameEntry*> DebugDump();
 
-	/** Singleton to retrieve a table of all names (multithreaded) for debug visualizers. */
-	static FNameEntry*** GetNameTableForDebuggerVisualizers_MT();
+	static FNameEntry const* GetEntry(EName Ename);
+	static FNameEntry const* GetEntry(FNameEntryId Id);
+
+	//@}
+
+	static uint8** GetBlocksForDebugVisualizer();
+	
 	/** Run autotest on FNames. */
 	static void AutoTest();
 	
@@ -1102,103 +758,23 @@ public:
 	 */
 	static FString NameToDisplayString( const FString& InDisplayName, const bool bIsBool );
 
-private:
-
-	union
-	{
-		struct
-		{
-			/** Index into the Names array (used to find String portion of the string/number pair used for comparison) */
-			NAME_INDEX		ComparisonIndex;
-		#if WITH_CASE_PRESERVING_NAME
-			/** Index into the Names array (used to find String portion of the string/number pair used for display) */
-			NAME_INDEX		DisplayIndex;
-		#endif
-			/** Number portion of the string/number pair (stored internally as 1 more than actual, so zero'd memory will be the default, no-instance case) */
-			uint32			Number;
-		};
-	};
+	/** Get the EName that this FName represents or nullptr */
+	const EName* ToEName() const;
 
-	/** Name hash head - used to iterate the single-linked list.		*/
-	static TAtomic<FNameEntry*>				NameHashHead[FNameDefs::NameHashBucketCount];
-	/** Name hash tail - insert new entries after this - NON ATOMIC!	*/
-	static TAtomic<FNameEntry*>				NameHashTail[FNameDefs::NameHashBucketCount];
-	/** Size of all name entries.								*/
-	static int32							NameEntryMemorySize;	
-	/** Number of ANSI names in name table.						*/
-	static int32							NumAnsiNames;			
-	/** Number of wide names in name table.					*/
-	static int32							NumWideNames;
-
-	/** Singleton to retrieve a table of all names. */
-	static TNameEntryArray& GetNames();
-	/**
-	 * Return the static initialized flag. Must be in a function like this so we don't have problems with 
-	 * different initialization order of static variables across the codebase. Use this function to get or set the variable.
-	 */
-	static bool& GetIsInitialized();
-	static void StaticInit();
+private:
 
+	/** Index into the Names array (used to find String portion of the string/number pair used for comparison) */
+	FNameEntryId	ComparisonIndex;
+#if WITH_CASE_PRESERVING_NAME
+	/** Index into the Names array (used to find String portion of the string/number pair used for display) */
+	FNameEntryId	DisplayIndex;
+#endif
+	/** Number portion of the string/number pair (stored internally as 1 more than actual, so zero'd memory will be the default, no-instance case) */
+	uint32			Number;
+	
 	friend const TCHAR* DebugFName(int32);
 	friend const TCHAR* DebugFName(int32, int32);
 	friend const TCHAR* DebugFName(FName&);
-	template<typename TCharType>
-	friend FNameEntry* AllocateNameEntry(const TCharType* Name, NAME_INDEX Index);
-	/** Used to increment the correct counter based upon TCharType */
-	template <typename TCharType> friend void IncrementNameCount();
-
-	/**
-	 * Initialization from a wide string
-	 * 
-	 * @param InName String name of the name/number pair
-	 * @param InNumber Number part of the name/number pair
-	 * @param FindType Operation to perform on names
-	 * @param bSplitName If true, this function will attempt to split a number off of the string portion (turning Rocket_17 to Rocket and number 17)
-	 * @param HardcodeIndex If >= 0, this represents a hardcoded FName and so automatically gets this index
-	 */
-	void Init(const WIDECHAR* InName, int32 InNumber, EFindName FindType, bool bSplitName=true, int32 HardcodeIndex = -1);
-	/**
-	 * Version that takes the hashes as part of the constructor (serialized previously)
-	 */
-	void Init(const WIDECHAR* InName, int32 InNumber, EFindName FindType, const uint16 NonCasePreservingHash, const uint16 CasePreservingHash);
-
-	/**
-	 * Initialization from an ANSI string
-	 * 
-	 * @param InName		String name of the name/number pair
-	 * @param InNumber		Number part of the name/number pair
-	 * @param FindType		Operation to perform on names
-	 * @param bSplitName	If true, this function will attempt to split a number off of the string portion (turning Rocket_17 to Rocket and number 17)
-	 * @param HardcodeIndex If >= 0, this represents a hardcoded FName and so automatically gets this index
-	 */
-	void Init(const ANSICHAR* InName, int32 InNumber, EFindName FindType, bool bSplitName=true, int32 HardcodeIndex = -1);
-	/**
-	 * Version that takes the hashes as part of the constructor (serialized previously). Skips InitInternal_HashSplit for loading perf reasons
-	 */
-	void Init(const ANSICHAR* InName, int32 InNumber, EFindName FindType, const uint16 NonCasePreservingHash, const uint16 CasePreservingHash);
-
-	template <typename TCharType>
-	void InitInternal(const TCharType* InName, int32 InNumber, const EFindName FindType, const int32 HardcodeIndex, const uint16 NonCasePreservingHash, const uint16 CasePreservingHash);
-
-	/**
-	 * Version of InitInternal that calculates the hash after splitting the string. Used by runtime FName construction
-	 */
-	template <typename TCharType>
-	void InitInternal_HashSplit(const TCharType* InName, int32 InNumber, const EFindName FindType, bool bSplitName, const int32 HardcodeIndex);
-
-	template <typename TCharType>
-	static bool InitInternal_FindOrAdd(const TCharType* InName, const EFindName FindType, const int32 HardcodeIndex, const uint16 NonCasePreservingHash, const uint16 CasePreservingHash, int32& OutComparisonIndex, int32& OutDisplayIndex);
-
-	template <typename TCharType>
-	static bool InitInternal_FindOrAddNameEntry(const TCharType* InName, const EFindName FindType, const ENameCase ComparisonMode, const uint16 iHash, int32& OutIndex);
-
-	template <typename TCharType>
-	static bool SplitNameWithCheckImpl(const TCharType* OldName, TCharType* NewName, int32 NewNameLen, int32& NewNumber);
-
-	FORCEINLINE NAME_INDEX GetComparisonIndexFast() const
-	{
-		return ComparisonIndex;
-	}
 
 	FORCEINLINE NAME_INDEX GetDisplayIndexFast() const
 	{
@@ -1209,17 +785,16 @@ private:
 #endif
 	}
 
-	/** Singleton to retrieve the critical section. */
-	static FCriticalSection* GetCriticalSection();
+	static bool IsWithinBounds(FNameEntryId Id);
 };
 
 template<> struct TIsZeroConstructType<class FName> { enum { Value = true }; };
 Expose_TNameOf(FName)
 
 
-inline uint32 GetTypeHash( const FName N )
+FORCEINLINE uint32 GetTypeHash(FName Name)
 {
-	return N.GetComparisonIndex() + N.GetNumber();
+	return GetTypeHash(Name.GetComparisonIndex()) + Name.GetNumber();
 }
 
 FORCEINLINE FString LexToString(const FName& Name)
@@ -1242,7 +817,6 @@ FORCEINLINE FName MinimalNameToName(const FMinimalName& InName)
 	return FName(InName.Index, InName.Index, InName.Number);
 }
 
-
 FORCEINLINE FScriptName NameToScriptName(const FName& InName)
 {
 	return FScriptName(InName.GetComparisonIndex(), InName.GetDisplayIndex(), InName.GetNumber());
@@ -1294,41 +868,8 @@ struct FNameSortIndexes
 };
 
 #ifndef WITH_CUSTOM_NAME_ENCODING
-
-	inline void FNameEntry::GetAnsiName(ANSICHAR(&OutName)[NAME_SIZE]) const
-	{
-		check(!IsWide());
-		FCStringAnsi::Strcpy(OutName, AnsiName);
-	}
-
-	inline void FNameEntry::GetWideName(WIDECHAR(&OutName)[NAME_SIZE]) const
-	{
-		check(IsWide());
-		FCStringWide::Strcpy(OutName, WideName);
-	}
-
-	inline ANSICHAR const* FNameEntry::GetAnsiNamePtr(ANSICHAR(&OptionalTempBuffer)[NAME_SIZE]) const
-	{
-		check(!IsWide());
-		return AnsiName;
-	}
-
-	inline WIDECHAR const* FNameEntry::GetWideNamePtr(WIDECHAR(&OptionalTempBuffer)[NAME_SIZE]) const
-	{
-		check(IsWide());
-		return WideName;
-	}
-
-	inline void FNameEntry::SetAnsiName(const ANSICHAR* SrcName, int32 SrcNameLen)
-	{
-		check(!IsWide());
-		FCStringAnsi::Strcpy(AnsiName, SrcNameLen + 1, SrcName);
-	}
-
-	inline void FNameEntry::SetWideName(const WIDECHAR* SrcName, int32 SrcNameLen)
-	{
-		check(IsWide());
-		FCStringWide::Strcpy(WideName, SrcNameLen + 1, SrcName);
-	}
-
+inline void FNameEntry::Encode(ANSICHAR*, uint32) {}
+inline void FNameEntry::Encode(WIDECHAR*, uint32) {}
+inline void FNameEntry::Decode(ANSICHAR*, uint32) {}
+inline void FNameEntry::Decode(WIDECHAR*, uint32) {}
 #endif
diff --git a/Engine/Source/Runtime/Core/Public/UObject/UnrealNames.h b/Engine/Source/Runtime/Core/Public/UObject/UnrealNames.h
--- a/Engine/Source/Runtime/Core/Public/UObject/UnrealNames.h
+++ b/Engine/Source/Runtime/Core/Public/UObject/UnrealNames.h
@@ -4,20 +4,10 @@
 
 #include "CoreTypes.h"
 
-//
-// Macros.
-//
-
-/** Index of highest hardcoded name to be replicated by index by the networking code
- * @warning: changing this number or making any change to the list of hardcoded names with index
- * less than this value breaks network compatibility, which by default checks for the same changelist
- * @note: names with a greater value than this can still be replicated, but they are sent as
- * strings instead of an efficient index
- */
-#define MAX_NETWORKED_HARDCODED_NAME 410
+struct FNameEntryId;
 
 //
-// Hardcoded names.
+// Macros.
 //
 
 // Define a message as an enumeration.
@@ -31,3 +21,20 @@ enum EName
 	NAME_MaxHardcodedNameIndex,
 };
 #undef REGISTER_NAME
+
+CORE_API const TCHAR* LexToString(EName Ename);
+
+/** Index of highest hardcoded name to be replicated by index by the networking code
+ * @warning: changing this number or making any change to the list of hardcoded names with index
+ * less than this value breaks network compatibility, which by default checks for the same changelist
+ * @note: names with a greater value than this can still be replicated, but they are sent as
+ * strings instead of an efficient index
+ *
+ * @see ShouldReplicateENameAsInteger()
+ */
+#define MAX_NETWORKED_HARDCODED_NAME 410
+
+inline bool ShouldReplicateAsInteger(EName Ename)
+{
+	return Ename <= MAX_NETWORKED_HARDCODED_NAME;
+}
diff --git a/Engine/Source/Runtime/CoreUObject/Private/Serialization/ArchiveStackTrace.cpp b/Engine/Source/Runtime/CoreUObject/Private/Serialization/ArchiveStackTrace.cpp
--- a/Engine/Source/Runtime/CoreUObject/Private/Serialization/ArchiveStackTrace.cpp
+++ b/Engine/Source/Runtime/CoreUObject/Private/Serialization/ArchiveStackTrace.cpp
@@ -969,6 +969,11 @@ static inline FString GetTableKey(const FLinkerLoad* Linker, const FName& Name)
 	return *Name.ToString();
 }
 
+static inline FString GetTableKey(const FLinkerLoad* Linker, FNameEntryId Id)
+{
+	return FName::GetEntry(Id)->GetPlainNameString();
+}
+
 static inline FString GetTableKeyForIndex(const FLinkerLoad* Linker, FPackageIndex Index)
 {
 	if (Index.IsNull())
@@ -992,11 +997,21 @@ bool CompareTableItem(FLinkerLoad* SourceLinker, FLinkerLoad* DestLinker, const
 	return SourceName == DestName;
 }
 
+bool CompareTableItem(FLinkerLoad* SourceLinker, FLinkerLoad* DestLinker, FNameEntryId SourceName, FNameEntryId DestName)
+{
+	return SourceName == DestName;
+}
+
 FString ConvertItemToText(const FName& Name, FLinkerLoad* Linker)
 {
 	return Name.ToString();
 }
 
+FString ConvertItemToText(FNameEntryId Id, FLinkerLoad* Linker)
+{
+	return FName::GetEntry(Id)->GetPlainNameString();
+}
+
 bool CompareTableItem(FLinkerLoad* SourceLinker, FLinkerLoad* DestLinker, const FObjectImport& SourceImport, const FObjectImport& DestImport)
 {
 	if (SourceImport.ObjectName != DestImport.ObjectName ||
@@ -1320,7 +1335,7 @@ void FArchiveStackTrace::DumpPackageHeaderDiffs(const FPackageData& SourcePackag
 	{
 		if (SourceLinker->NameMap != DestLinker->NameMap)
 		{
-			DumpTableDifferences<FName>(SourceLinker, DestLinker, SourceLinker->NameMap, DestLinker->NameMap, *AssetFilename, TEXT("Name"), MaxDiffsToLog);
+			DumpTableDifferences<FNameEntryId>(SourceLinker, DestLinker, SourceLinker->NameMap, DestLinker->NameMap, *AssetFilename, TEXT("Name"), MaxDiffsToLog);
 		}
 
 		if (!IsImportMapIdentical(SourceLinker, DestLinker))
diff --git a/Engine/Source/Runtime/CoreUObject/Private/UObject/CoreNet.cpp b/Engine/Source/Runtime/CoreUObject/Private/UObject/CoreNet.cpp
--- a/Engine/Source/Runtime/CoreUObject/Private/UObject/CoreNet.cpp
+++ b/Engine/Source/Runtime/CoreUObject/Private/UObject/CoreNet.cpp
@@ -311,13 +311,14 @@ bool UPackageMap::StaticSerializeName(FArchive& Ar, FName& InName)
 	}
 	else if (Ar.IsSaving())
 	{
-		uint8 bHardcoded = InName.GetComparisonIndex() <= MAX_NETWORKED_HARDCODED_NAME;
+		const EName* InEName = InName.ToEName();
+		uint8 bHardcoded = InEName && ShouldReplicateAsInteger(*InEName);
 		Ar.SerializeBits(&bHardcoded, 1);
 		if (bHardcoded)
 		{
 			// send by hardcoded index
 			checkSlow(InName.GetNumber() <= 0); // hardcoded names should never have a Number
-			uint32 NameIndex = uint32(InName.GetComparisonIndex());
+			uint32 NameIndex = *InEName;
 			Ar.SerializeIntPacked(NameIndex);
 		}
 		else
diff --git a/Engine/Source/Runtime/CoreUObject/Private/UObject/LinkerLoad.cpp b/Engine/Source/Runtime/CoreUObject/Private/UObject/LinkerLoad.cpp
--- a/Engine/Source/Runtime/CoreUObject/Private/UObject/LinkerLoad.cpp
+++ b/Engine/Source/Runtime/CoreUObject/Private/UObject/LinkerLoad.cpp
@@ -383,7 +383,9 @@ static FTexture2DResourceMem* CreateResourceMem(int32 SizeX, int32 SizeY, int32
 
 static inline int32 HashNames(FName Object, FName Class, FName Package)
 {
-	return Object.GetComparisonIndex() + 7 * Class.GetComparisonIndex() + 31 * FPackageName::GetShortFName(Package).GetComparisonIndex();
+	return GetTypeHash(Object.GetComparisonIndex())
+		+ 7 * GetTypeHash(Class.GetComparisonIndex())
+		+ 31 * GetTypeHash(FPackageName::GetShortFName(Package).GetComparisonIndex());
 }
 
 static FORCEINLINE bool IsCoreUObjectPackage(const FName& PackageName)
@@ -1441,7 +1443,7 @@ FLinkerLoad::ELinkerStatus FLinkerLoad::SerializeNameMap()
 		}
 
 		// Add it to the name table with no splitting and no hash calculations
-		NameMap.Emplace(NameEntry);
+		NameMap.Emplace(FName(NameEntry).GetDisplayIndex());
 
 		NameMapIndex++;
 	}
@@ -4885,7 +4887,7 @@ FArchive& FLinkerLoad::operator<<( UObject*& Object )
 	return *this;
 }
 
-void FLinkerLoad::BadNameIndexError(NAME_INDEX NameIndex)
+void FLinkerLoad::BadNameIndexError(int32 NameIndex)
 {
 	UE_LOG(LogLinker, Error, TEXT("Bad name index %i/%i"), NameIndex, NameMap.Num());
 }
diff --git a/Engine/Source/Runtime/CoreUObject/Private/UObject/LinkerManager.cpp b/Engine/Source/Runtime/CoreUObject/Private/UObject/LinkerManager.cpp
--- a/Engine/Source/Runtime/CoreUObject/Private/UObject/LinkerManager.cpp
+++ b/Engine/Source/Runtime/CoreUObject/Private/UObject/LinkerManager.cpp
@@ -49,9 +49,9 @@ bool FLinkerManager::Exec(class UWorld* InWorld, const TCHAR* Cmd, FOutputDevice
 			int32 NameSize = 0;
 			for (int32 j = 0; j < Linker->NameMap.Num(); j++)
 			{
-				if (Linker->NameMap[j] != NAME_None)
+				if (FNameEntryId Id = Linker->NameMap[j])
 				{
-					NameSize += FNameEntry::GetSize(*Linker->NameMap[j].ToString());
+					NameSize += FName::GetEntry(Id)->GetSizeInBytes();
 				}
 			}
 			Ar.Logf
diff --git a/Engine/Source/Runtime/CoreUObject/Private/UObject/LinkerSave.cpp b/Engine/Source/Runtime/CoreUObject/Private/UObject/LinkerSave.cpp
--- a/Engine/Source/Runtime/CoreUObject/Private/UObject/LinkerSave.cpp
+++ b/Engine/Source/Runtime/CoreUObject/Private/UObject/LinkerSave.cpp
@@ -189,9 +189,9 @@ FLinkerSave::~FLinkerSave()
 	Detach();
 }
 
-int32 FLinkerSave::MapName(const FName& Name) const
+int32 FLinkerSave::MapName(FNameEntryId Id) const
 {
-	const int32* IndexPtr = NameIndices.Find(Name);
+	const int32* IndexPtr = NameIndices.Find(Id);
 
 	if (IndexPtr)
 	{
@@ -275,7 +275,7 @@ FString FLinkerSave::GetArchiveName() const
 
 FArchive& FLinkerSave::operator<<( FName& InName )
 {
-	int32 Save = MapName(InName);
+	int32 Save = MapName(InName.GetDisplayIndex());
 
 	check(GetSerializeContext());
 	ensureMsgf(Save != INDEX_NONE, TEXT("Name \"%s\" is not mapped when saving %s (object: %s, property: %s)"), 
diff --git a/Engine/Source/Runtime/CoreUObject/Private/UObject/SavePackage.cpp b/Engine/Source/Runtime/CoreUObject/Private/UObject/SavePackage.cpp
--- a/Engine/Source/Runtime/CoreUObject/Private/UObject/SavePackage.cpp
+++ b/Engine/Source/Runtime/CoreUObject/Private/UObject/SavePackage.cpp
@@ -722,10 +722,12 @@ public:
 	 */
 	void MarkNameAsReferenced(const FName& Name)
 	{
-		// We need to store the FName without the number, as the number is stored separately by FLinkerSave 
-		// and we don't want duplicate entries in the name table just because of the number
-		const FName NameNoNumber(Name, 0);
-		ReferencedNames.Add(NameNoNumber);
+		ReferencedNames.Add(Name.GetDisplayIndex());
+	}
+
+	void MarkNameAsReferenced(FNameEntryId Name)
+	{
+		ReferencedNames.Add(Name);
 	}
 
 #if WITH_EDITOR
@@ -762,14 +764,14 @@ public:
 	void UpdateLinkerWithMarkedNames(FLinkerSave* Linker)
 	{
 		Linker->NameMap.Reserve(Linker->NameMap.Num() + ReferencedNames.Num());
-		for (const FName& Name : ReferencedNames)
+		for (FNameEntryId Name : ReferencedNames)
 		{
 			Linker->NameMap.Add(Name);
 		}
 	}
 
 private:
-	TSet<FName, FLinkerNamePairKeyFuncs> ReferencedNames;
+	TSet<FNameEntryId> ReferencedNames;
 };
 
 bool IsEditorOnlyObject(const UObject* InObject, bool bCheckRecursive, bool bCheckMarks)
@@ -1509,7 +1511,7 @@ struct FObjectNameSortHelper
 {
 private:
 	/** the linker that we're sorting names for */
-	friend struct TDereferenceWrapper<FName, FObjectNameSortHelper>;
+	friend struct TDereferenceWrapper<FNameEntryId, FObjectNameSortHelper>;
 
 	/** Comparison function used by Sort */
 	FORCEINLINE bool operator()( const FName& A, const FName& B ) const
@@ -1517,6 +1519,13 @@ private:
 		return A.Compare(B) < 0;
 	}
 
+	/** Comparison function used by Sort */
+	FORCEINLINE bool operator()(FNameEntryId A, FNameEntryId B) const
+	{
+		// Could be implemented without constructing FName but would new FNameEntry comparison API
+		return A != B && operator()(FName::CreateFromDisplayId(A, 0), FName::CreateFromDisplayId(B, 0));
+	}
+
 public:
 
 
@@ -1534,10 +1543,10 @@ public:
 		if ( LinkerToConformTo != nullptr )
 		{
 			SortStartPosition = LinkerToConformTo->NameMap.Num();
-			TArray<FName> ConformedNameMap = LinkerToConformTo->NameMap;
+			TArray<FNameEntryId> ConformedNameMap = LinkerToConformTo->NameMap;
 			for ( int32 NameIndex = 0; NameIndex < Linker->NameMap.Num(); NameIndex++ )
 			{
-				FName& CurrentName = Linker->NameMap[NameIndex];
+				FNameEntryId CurrentName = Linker->NameMap[NameIndex];
 				if ( !ConformedNameMap.Contains(CurrentName) )
 				{
 					ConformedNameMap.Add(CurrentName);
@@ -1547,7 +1556,7 @@ public:
 			Linker->NameMap = ConformedNameMap;
 			for ( int32 NameIndex = 0; NameIndex < Linker->NameMap.Num(); NameIndex++ )
 			{
-				FName& CurrentName = Linker->NameMap[NameIndex];
+				FNameEntryId CurrentName = Linker->NameMap[NameIndex];
 				SavePackageState.MarkNameAsReferenced(CurrentName);
 			}
 		}
@@ -4526,17 +4535,17 @@ FSavePackageResultStruct UPackage::Save(UPackage* InOuter, UObject* Base, EObjec
 				if ( GOutputCookingWarnings )
 				{
 					// check the name list for uniqueobjectnamefor cooking
-					static FName NAME_UniqueObjectNameForCooking(TEXT("UniqueObjectNameForCooking"));
+					static FNameEntryId NAME_UniqueObjectNameForCooking = FName("UniqueObjectNameForCooking").GetComparisonIndex();
 
-					for (const auto& NameInUse : Linker->NameMap)
+					for (FNameEntryId NameInUse : Linker->NameMap)
 					{
-						if (NameInUse.GetComparisonIndex() == NAME_UniqueObjectNameForCooking.GetComparisonIndex())
+						if (FName::GetComparisonIdFromDisplayId(NameInUse) == NAME_UniqueObjectNameForCooking)
 						{
 							//UObject *Object = FindObject<UObject>( ANY_PACKAGE, *NameInUse.ToString());
 
 							// error
 							// check(Object);
-							UE_LOG(LogSavePackage, Warning, TEXT("Saving object into cooked package %s which was created at cook time, Object Name %s"), Filename, *NameInUse.ToString());
+							UE_LOG(LogSavePackage, Warning, TEXT("Saving object into cooked package %s which was created at cook time"), Filename);
 						}
 					}
 				}
@@ -4568,7 +4577,7 @@ FSavePackageResultStruct UPackage::Save(UPackage* InOuter, UObject* Base, EObjec
 					Linker->Summary.NameCount = Linker->NameMap.Num();
 					for (int32 i = 0; i < Linker->NameMap.Num(); i++)
 					{
-						Linker->NameMap[i].GetDisplayNameEntry()->Write(NameStream.EnterElement());
+						FName::GetEntry(Linker->NameMap[i])->Write(NameStream.EnterElement());
 						Linker->NameIndices.Add(Linker->NameMap[i], i);
 					}
 				}
diff --git a/Engine/Source/Runtime/CoreUObject/Private/UObject/UObjectGlobals.cpp b/Engine/Source/Runtime/CoreUObject/Private/UObject/UObjectGlobals.cpp
--- a/Engine/Source/Runtime/CoreUObject/Private/UObject/UObjectGlobals.cpp
+++ b/Engine/Source/Runtime/CoreUObject/Private/UObject/UObjectGlobals.cpp
@@ -467,9 +467,9 @@ void StaticTick( float DeltaTime, bool bUseFullTimeLimit, float AsyncLoadingTime
 
 #if STATS
 	// Set name table stats.
-	int32 NameTableEntries = FName::GetMaxNames();
 	int32 NameTableAnsiEntries = FName::GetNumAnsiNames();
 	int32 NameTableWideEntries = FName::GetNumWideNames();
+	int32 NameTableEntries = NameTableAnsiEntries + NameTableWideEntries;
 	int32 NameTableMemorySize = FName::GetNameTableMemorySize();
 	SET_DWORD_STAT( STAT_NameTableEntries, NameTableEntries );
 	SET_DWORD_STAT( STAT_NameTableAnsiEntries, NameTableAnsiEntries );
@@ -1762,12 +1762,13 @@ FName MakeUniqueObjectName( UObject* Parent, UClass* Class, FName InBaseName/*=N
 		do
 		{
 			// create the next name in the sequence for this class
-			if (BaseName.GetComparisonIndex() == NAME_Package)
+			static const FName NamePackage(NAME_Package);
+			if (BaseName == NamePackage)
 			{
 				if (Parent == NULL)
 				{
 					//package names should default to "/Temp/Untitled" when their parent is NULL. Otherwise they are a group.
-					TestName = FName(*FString::Printf(TEXT("/Temp/%s"), *FName(NAME_Untitled).ToString()), ++Class->ClassUnique);
+					TestName = FName(*FString::Printf(TEXT("/Temp/%s"), LexToString(NAME_Untitled)), ++Class->ClassUnique);
 				}
 				else
 				{
diff --git a/Engine/Source/Runtime/CoreUObject/Private/UObject/UObjectHash.cpp b/Engine/Source/Runtime/CoreUObject/Private/UObject/UObjectHash.cpp
--- a/Engine/Source/Runtime/CoreUObject/Private/UObject/UObjectHash.cpp
+++ b/Engine/Source/Runtime/CoreUObject/Private/UObject/UObjectHash.cpp
@@ -409,7 +409,7 @@ public:
  */
 static FORCEINLINE int32 GetObjectHash(FName ObjName)
 {
-	return (ObjName.GetComparisonIndex() ^ ObjName.GetNumber());
+	return GetTypeHash(ObjName);
 }
 
 /**
@@ -423,7 +423,7 @@ static FORCEINLINE int32 GetObjectHash(FName ObjName)
  */
 static FORCEINLINE int32 GetObjectOuterHash(FName ObjName,PTRINT Outer)
 {
-	return ((ObjName.GetComparisonIndex() ^ ObjName.GetNumber()) ^ (Outer >> 6));
+	return GetTypeHash(ObjName) + (Outer >> 6);
 }
 
 UObject* StaticFindObjectFastExplicitThreadSafe(FUObjectHashTables& ThreadHash, UClass* ObjectClass, FName ObjectName, const FString& ObjectPathName, bool bExactClass, EObjectFlags ExcludeFlags/*=0*/)
diff --git a/Engine/Source/Runtime/CoreUObject/Public/UObject/Linker.h b/Engine/Source/Runtime/CoreUObject/Public/UObject/Linker.h
--- a/Engine/Source/Runtime/CoreUObject/Public/UObject/Linker.h
+++ b/Engine/Source/Runtime/CoreUObject/Public/UObject/Linker.h
@@ -166,35 +166,35 @@ public:
 };
 
 
-struct FLinkerNamePairKeyFuncs : DefaultKeyFuncs<FName, false>
+struct UE_DEPRECATED(4.23, "Outdated since display index replaced FName as key") FLinkerNamePairKeyFuncs : DefaultKeyFuncs<FName, false>
 {
 	static FORCEINLINE bool Matches(FName A, FName B)
 	{
 		// The linker requires that FNames preserve case, but the numeric suffix can be ignored since
 		// that is stored separately for each FName instance saved
-		return A.IsEqual(B, ENameCase::CaseSensitive, false/*bCompareNumber*/);
+		return A.GetDisplayIndex() == B.GetDisplayIndex();
 	}
 
 	static FORCEINLINE uint32 GetKeyHash(FName Key)
 	{
-		return Key.GetComparisonIndex();
+		return GetTypeHash(Key.GetDisplayIndex());
 	}
 };
 
 
 template<typename ValueType>
-struct TLinkerNameMapKeyFuncs : TDefaultMapKeyFuncs<FName, ValueType, false>
+struct UE_DEPRECATED(4.23, "Outdated since display indexes replaced FNames as keys") TLinkerNameMapKeyFuncs : TDefaultMapKeyFuncs<FName, ValueType, false>
 {
 	static FORCEINLINE bool Matches(FName A, FName B)
 	{
 		// The linker requires that FNames preserve case, but the numeric suffix can be ignored since
 		// that is stored separately for each FName instance saved
-		return A.IsEqual(B, ENameCase::CaseSensitive, false/*bCompareNumber*/);
+		return A.GetDisplayIndex() == B.GetDisplayIndex();
 	}
 
 	static FORCEINLINE uint32 GetKeyHash(FName Key)
 	{
-		return Key.GetComparisonIndex();
+		return GetTypeHash(Key.GetDisplayIndex());
 	}
 };
 
@@ -231,7 +231,7 @@ public:
 	FPackageFileSummary		Summary;
 
 	/** Names used by objects contained within this package */
-	TArray<FName>			NameMap;
+	TArray<FNameEntryId>	NameMap;
 
 	/** Gatherable text data contained within this package */
 	TArray<FGatherableTextData> GatherableTextDataMap;
diff --git a/Engine/Source/Runtime/CoreUObject/Public/UObject/LinkerLoad.h b/Engine/Source/Runtime/CoreUObject/Public/UObject/LinkerLoad.h
--- a/Engine/Source/Runtime/CoreUObject/Public/UObject/LinkerLoad.h
+++ b/Engine/Source/Runtime/CoreUObject/Public/UObject/LinkerLoad.h
@@ -806,12 +806,12 @@ private:
 		Value = ID;
 		return Ar;
 	}
-	void BadNameIndexError(NAME_INDEX NameIndex);
+	void BadNameIndexError(int32 NameIndex);
 	FORCEINLINE virtual FArchive& operator<<(FName& Name) override
 	{
 		Name = NAME_None;
 		FArchive& Ar = *this;
-		NAME_INDEX NameIndex;
+		int32 NameIndex;
 		Ar << NameIndex;
 		int32 Number;
 		Ar << Number;
@@ -825,11 +825,11 @@ private:
 		else
 		{
 			// if the name wasn't loaded (because it wasn't valid in this context)
-			const FName& MappedName = NameMap[NameIndex];
-			if (!MappedName.IsNone())
+			FNameEntryId MappedName = NameMap[NameIndex];
+			if (MappedName)
 			{
 				// simply create the name from the NameMap's name and the serialized instance number
-				Name = FName(MappedName, Number);
+				Name = FName::CreateFromDisplayId(MappedName, Number);
 			}
 		}
 		return *this;
diff --git a/Engine/Source/Runtime/CoreUObject/Public/UObject/LinkerSave.h b/Engine/Source/Runtime/CoreUObject/Public/UObject/LinkerSave.h
--- a/Engine/Source/Runtime/CoreUObject/Public/UObject/LinkerSave.h
+++ b/Engine/Source/Runtime/CoreUObject/Public/UObject/LinkerSave.h
@@ -43,7 +43,7 @@ public:
 	TMap<const UObject *, TArray<FName> > SearchableNamesObjectMap;
 
 	/** Index array - location of the name in the NameMap array for each FName is stored in the NameIndices array using the FName's Index */
-	TMap<FName, int32, FDefaultSetAllocator, TLinkerNameMapKeyFuncs<int32>> NameIndices;
+	TMap<FNameEntryId, int32> NameIndices;
 
 	/** Save context associated with this linker */
 	TRefCountPtr<FUObjectSerializeContext> SaveContext;
@@ -75,7 +75,7 @@ public:
 	FLinkerSave(UPackage* InParent, FArchive *InSaver, bool bForceByteSwapping, bool bInSaveUnversioned = false);
 
 	/** Returns the appropriate name index for the source name, or 0 if not found in NameIndices */
-	int32 MapName(const FName& Name) const;
+	int32 MapName( FNameEntryId Name) const;
 
 	/** Returns the appropriate package index for the source object, or default value if not found in ObjectIndicesMap */
 	FPackageIndex MapObject(const UObject* Object) const;
diff --git a/Engine/Source/Runtime/Engine/Private/NetworkDriver.cpp b/Engine/Source/Runtime/Engine/Private/NetworkDriver.cpp
--- a/Engine/Source/Runtime/Engine/Private/NetworkDriver.cpp
+++ b/Engine/Source/Runtime/Engine/Private/NetworkDriver.cpp
@@ -360,7 +360,8 @@ void UNetDriver::LoadChannelDefinitions()
 
 		for (FChannelDefinition& ChannelDef : ChannelDefinitions)
 		{
-			UE_CLOG((ChannelDef.ChannelName.GetComparisonIndex() > MAX_NETWORKED_HARDCODED_NAME), LogNet, Warning, TEXT("Channel name will be serialized as a string: %s"), *ChannelDef.ChannelName.ToString());
+			UE_CLOG(!ChannelDef.ChannelName.ToEName() || !ShouldReplicateAsInteger(*ChannelDef.ChannelName.ToEName()), 
+				LogNet, Warning, TEXT("Channel name will be serialized as a string: %s"), *ChannelDef.ChannelName.ToString());
 			UE_CLOG(ChannelDefinitionMap.Contains(ChannelDef.ChannelName), LogNet, Error, TEXT("Channel name is defined multiple times: %s"), *ChannelDef.ChannelName.ToString());
 			UE_CLOG(StaticChannelIndices.Contains(ChannelDef.StaticChannelIndex), LogNet, Error, TEXT("Channel static index is already in use: %s %i"), *ChannelDef.ChannelName.ToString(), ChannelDef.StaticChannelIndex);
 
diff --git a/Engine/Source/Runtime/Engine/Private/UnrealEngine.cpp b/Engine/Source/Runtime/Engine/Private/UnrealEngine.cpp
--- a/Engine/Source/Runtime/Engine/Private/UnrealEngine.cpp
+++ b/Engine/Source/Runtime/Engine/Private/UnrealEngine.cpp
@@ -4499,17 +4499,19 @@ bool UEngine::HandleCountDisabledParticleItemsCommand( const TCHAR* Cmd, FOutput
 	return true;
 }
 
-// View the last N number of names added to the name table. Useful for tracking down name table bloat
+// View all names or the last N names added to the name table. Useful for tracking down name table bloat
 bool UEngine::HandleViewnamesCommand( const TCHAR* Cmd, FOutputDevice& Ar )
 {
-	int32 NumNames = 0;
-	if (FParse::Value(Cmd,TEXT("NUM="),NumNames))
+	const TArray<const FNameEntry*> Entries = FName::DebugDump();
+
+	int32 NumLast = 0;
+	int32 BeginIdx = FParse::Value(Cmd, TEXT("NUM="), NumLast) ? FMath::Max(Entries.Num() - NumLast, 0) : 0;
+
+	for (int32 I = BeginIdx; I < Entries.Num(); ++I)
 	{
-		for (int32 NameIndex = FMath::Max<int32>(FName::GetMaxNames() - NumNames, 0); NameIndex < FName::GetMaxNames(); NameIndex++)
-		{
-			Ar.Logf(TEXT("%d->%s"), NameIndex, *FName::SafeString(NameIndex));
-		}
+		Ar.Log(*Entries[I]->GetPlainNameString());	
 	}
+
 	return true;
 }
 
diff --git a/Engine/Source/Runtime/Sockets/Private/BSDIPv6Sockets/SocketSubsystemBSDIPv6.cpp b/Engine/Source/Runtime/Sockets/Private/BSDIPv6Sockets/SocketSubsystemBSDIPv6.cpp
--- a/Engine/Source/Runtime/Sockets/Private/BSDIPv6Sockets/SocketSubsystemBSDIPv6.cpp
+++ b/Engine/Source/Runtime/Sockets/Private/BSDIPv6Sockets/SocketSubsystemBSDIPv6.cpp
@@ -57,18 +57,21 @@ FSocket* FSocketSubsystemBSDIPv6::CreateSocket(const FName& SocketType, const FS
 		ProtocolType = ESocketProtocolFamily::IPv6;
 	}
 
-	switch (SocketType.GetComparisonIndex())
+	if (const EName* Ename = SocketType.ToEName())
 	{
-	case NAME_DGram:
-		// Creates a data gram (UDP) socket
-		Socket = socket(AF_INET6, SOCK_DGRAM, IPPROTO_UDP);
-		NewSocket = (Socket != INVALID_SOCKET) ? InternalBSDSocketFactory(Socket, SOCKTYPE_Datagram, SocketDescription, ProtocolType) : NULL;
-		break;
-	case NAME_Stream:
-		// Creates a stream (TCP) socket
-		Socket = socket(AF_INET6, SOCK_STREAM, IPPROTO_TCP);
-		NewSocket = (Socket != INVALID_SOCKET) ? InternalBSDSocketFactory(Socket, SOCKTYPE_Streaming, SocketDescription, ProtocolType) : NULL;
-		break;
+		switch (*Ename)
+		{
+		case NAME_DGram:
+			// Creates a data gram (UDP) socket
+			Socket = socket(AF_INET6, SOCK_DGRAM, IPPROTO_UDP);
+			NewSocket = (Socket != INVALID_SOCKET) ? InternalBSDSocketFactory(Socket, SOCKTYPE_Datagram, SocketDescription, ProtocolType) : NULL;
+			break;
+		case NAME_Stream:
+			// Creates a stream (TCP) socket
+			Socket = socket(AF_INET6, SOCK_STREAM, IPPROTO_TCP);
+			NewSocket = (Socket != INVALID_SOCKET) ? InternalBSDSocketFactory(Socket, SOCKTYPE_Streaming, SocketDescription, ProtocolType) : NULL;
+			break;
+		}
 	}
 
 	if (!NewSocket)
