Minor correction to a comment

#rnx

diff --git a/Engine/Source/Runtime/Core/Public/Serialization/StructuredArchive.h b/Engine/Source/Runtime/Core/Public/Serialization/StructuredArchive.h
--- a/Engine/Source/Runtime/Core/Public/Serialization/StructuredArchive.h
+++ b/Engine/Source/Runtime/Core/Public/Serialization/StructuredArchive.h
@@ -311,7 +311,7 @@ public:
 	};
 
 	/**
-	 * Represents a map in the structured archive. A map is similar to an object, but keys can be read back out from an archive.
+	 * Represents a map in the structured archive. A map is similar to a record, but keys can be read back out from an archive.
 	 * (This is an important distinction for binary archives).
 	 */
 	class CORE_API FMap
