Obsolete comment removed.

diff --git a/Engine/Source/Runtime/Core/Public/Delegates/DelegateSignatureImpl.inl b/Engine/Source/Runtime/Core/Public/Delegates/DelegateSignatureImpl.inl
--- a/Engine/Source/Runtime/Core/Public/Delegates/DelegateSignatureImpl.inl
+++ b/Engine/Source/Runtime/Core/Public/Delegates/DelegateSignatureImpl.inl
@@ -995,9 +995,6 @@ protected:
 	/**
 	 * Removes a function from this multi-cast delegate's invocation list (performance is O(N)).
 	 *
-	 * The function is not actually removed, but deleted and marked as removed.
-	 * It will be removed next time the invocation list is compacted within Broadcast().
-	 *
 	 * @param Handle The handle of the delegate instance to remove.
 	 * @return  true if the delegate was successfully removed.
 	 */
