Fix lexical comparison of FName with substring of itself

Caused UHT non-deterministic ouput when sorting names, which caused needless recompilation. This in turn broke CI jobs that depended on engine code in prior steps not getting recompiled in later game specific steps.

diff --git a/Engine/Source/Runtime/Core/Private/UObject/UnrealNames.cpp b/Engine/Source/Runtime/Core/Private/UObject/UnrealNames.cpp
--- a/Engine/Source/Runtime/Core/Private/UObject/UnrealNames.cpp
+++ b/Engine/Source/Runtime/Core/Private/UObject/UnrealNames.cpp
@@ -1637,28 +1637,33 @@ int32 FName::Compare( const FName& Other ) const
 	}
 	else // Names don't match. This means we don't even need to check numbers.
 	{
-		const FNameEntry* const ThisEntry = GetComparisonNameEntry();
-		const FNameEntry* const OtherEntry = Other.GetComparisonNameEntry();
-		
-		// If one or both entries return an invalid name entry, the comparison fails - fallback to comparing the index
-		if (ThisEntry == nullptr || OtherEntry == nullptr)
+		FNameBuffer ThisBuffer, OtherBuffer;
+		FNameStringView ThisView =	GetComparisonNameEntry()->MakeView(ThisBuffer);
+		FNameStringView OtherView =	Other.GetComparisonNameEntry()->MakeView(OtherBuffer);
+
+		// If only one view is wide, convert the ansi view to wide as well
+		if (ThisView.bIsWide != OtherView.bIsWide)
 		{
-			return ComparisonIndex.CompareCreationOrder(Other.ComparisonIndex);
+			FNameStringView& AnsiView = ThisView.bIsWide ? OtherView : ThisView;
+			FNameBuffer& AnsiBuffer =	ThisView.bIsWide ? OtherBuffer : ThisBuffer;
+
+#ifndef WITH_CUSTOM_NAME_ENCODING
+			memcpy(AnsiBuffer.AnsiName, AnsiView.Ansi, AnsiView.Len * sizeof(ANSICHAR));
+			AnsiView.Ansi = AnsiBuffer.AnsiName;
+#endif
+
+			AnsiBuffer.ConvertInPlace<ANSICHAR, WIDECHAR>(AnsiView.Len);
+			AnsiView.bIsWide = true;
 		}
-		else if (ThisEntry->IsWide() == OtherEntry->IsWide())
+
+		int32 MinLen = FMath::Min(ThisView.Len, OtherView.Len);
+		if (int32 StrDiff = ThisView.bIsWide ?	FCStringWide::Strnicmp(ThisView.Wide, OtherView.Wide, MinLen) :
+												FCStringAnsi::Strnicmp(ThisView.Ansi, OtherView.Ansi, MinLen))
 		{
-			return ThisEntry->IsWide() ?	FCStringWide::Stricmp(ThisEntry->WideName, OtherEntry->WideName) :
-											FCStringAnsi::Stricmp(ThisEntry->AnsiName, OtherEntry->AnsiName);
+			return StrDiff;
 		}
-		else // Ansi/Wide mismatch, convert one string to wide
-		{
-			WIDECHAR TempBuffer[NAME_SIZE];
-			const FNameEntry* AnsiEntry = ThisEntry->IsWide() ? OtherEntry : ThisEntry;
-			FPlatformString::Convert(TempBuffer, NAME_SIZE, AnsiEntry->AnsiName, FCStringAnsi::Strlen(AnsiEntry->AnsiName));
 
-			return ThisEntry->IsWide() ? FCStringWide::Stricmp(ThisEntry->WideName, TempBuffer)
-										: FCStringWide::Stricmp(TempBuffer, OtherEntry->WideName);
-		}
+		return ThisView.Len - OtherView.Len;
 	}
 }
 
@@ -1918,6 +1923,25 @@ void FName::AutoTest()
 	check(!StringAndNumberEqualsString("abc", 3, NAME_NO_NUMBER_INTERNAL, "abc_0"));
 	check(!StringAndNumberEqualsString("abc", 3, NAME_NO_NUMBER_INTERNAL, "abc_"));
 
+	TArray<FName> Names;
+	Names.Add("FooB");
+	Names.Add("FooABCD");
+	Names.Add("FooABC");
+	Names.Add("FooAB");
+	Names.Add("FooA");
+	Names.Add("FooC");
+	const WIDECHAR FooWide[] = {'F', 'o', 'o', (WIDECHAR)2000, '\0'};
+	Names.Add(FooWide);
+	Algo::Sort(Names); // Todo: fix when removing default operator<
+
+	check(Names[0] == "FooA");
+	check(Names[1] == "FooAB");
+	check(Names[2] == "FooABC");
+	check(Names[3] == "FooABCD");
+	check(Names[4] == "FooB");
+	check(Names[5] == "FooC");
+	check(Names[6] == FooWide);
+
 #if 0
 	// Check hash table growth still yields the same unique FName ids
 	static int32 OverflowAtLeastTwiceCount = 4 * FNamePoolInitialSlots * FNamePoolShards;
