Make FArchive::IsTextFormat() return a compile time constant value when text archive support is disabled so we can have redundant text-only code paths stripped

diff --git a/Engine/Source/Runtime/Core/Public/Serialization/Archive.h b/Engine/Source/Runtime/Core/Public/Serialization/Archive.h
--- a/Engine/Source/Runtime/Core/Public/Serialization/Archive.h
+++ b/Engine/Source/Runtime/Core/Public/Serialization/Archive.h
@@ -920,7 +920,7 @@ PRAGMA_DISABLE_DEPRECATION_WARNINGS
 
 	FORCEINLINE bool IsTextFormat() const
 	{
-		return ArIsTextFormat;
+		return (ArIsTextFormat && WITH_TEXT_ARCHIVE_SUPPORT);
 	}
 
 	FORCEINLINE bool WantBinaryPropertySerialization() const
