Made FName creation during editor startup an additional  ~2x faster by tweaking a constant

Now using 1024 instead of 16 shards in dev / debugc config. This eliminates most contention between asset data gatherer thread and main thread and also reduce amount of rehashing needed.

diff --git a/Engine/Source/Runtime/Core/Private/UObject/UnrealNames.cpp b/Engine/Source/Runtime/Core/Private/UObject/UnrealNames.cpp
--- a/Engine/Source/Runtime/Core/Private/UObject/UnrealNames.cpp
+++ b/Engine/Source/Runtime/Core/Private/UObject/UnrealNames.cpp
@@ -390,10 +390,19 @@ private:
 	uint8* Blocks[FNameMaxBlocks] = {};
 };
 
+// Increasing shards reduces contention but uses more memory and adds cache pressure.
+// Reducing contention matters when multiple threads create FNames in parallel.
+// Contention exists in some tool scenarios, for instance between main thread
+// and asset data gatherer thread during editor startup.
+#if WITH_CASE_PRESERVING_NAME
+enum { FNamePoolShardBits = 10 };
+#else
 enum { FNamePoolShardBits = 4 };
+#endif
+
 enum { FNamePoolShards = 1 << FNamePoolShardBits };
 enum { FNamePoolInitialSlotBits = 8 };
-enum { FNamePoolInitialSlots = 1 << FNamePoolInitialSlotBits };
+enum { FNamePoolInitialSlotsPerShard = 1 << FNamePoolInitialSlotBits };
 
 /** Hashes name into 64 bits that determines shard and slot index.
  *	
@@ -530,9 +539,9 @@ public:
 	{
 		Entries = &InEntries;
 
-		Slots = (FNameSlot*)FMemory::Malloc(FNamePoolInitialSlots * sizeof(FNameSlot), alignof(FNameSlot));
-		memset(Slots, 0, FNamePoolInitialSlots * sizeof(FNameSlot));
-		CapacityMask = FNamePoolInitialSlots - 1;
+		Slots = (FNameSlot*)FMemory::Malloc(FNamePoolInitialSlotsPerShard * sizeof(FNameSlot), alignof(FNameSlot));
+		memset(Slots, 0, FNamePoolInitialSlotsPerShard * sizeof(FNameSlot));
+		CapacityMask = FNamePoolInitialSlotsPerShard - 1;
 	}
 
 	uint32 Capacity() const { return CapacityMask + 1; }
