Correct a typo in an FNV documentation comment

diff --git a/Engine/Source/Runtime/Core/Private/Misc/Fnv.cpp b/Engine/Source/Runtime/Core/Private/Misc/Fnv.cpp
--- a/Engine/Source/Runtime/Core/Private/Misc/Fnv.cpp
+++ b/Engine/Source/Runtime/Core/Private/Misc/Fnv.cpp
@@ -22,7 +22,7 @@ uint32 FFnv::MemFnv32( const void* InData, int32 Length, uint32 InOffset )
     return Fnv;
 }
 
-// generate a 32-bit hash from data in memory of length bytes
+// generate a 64-bit hash from data in memory of length bytes
 uint64 FFnv::MemFnv64( const void* InData, int32 Length, uint64 InOffset )
 {
     // constants from above reference
