FSetElement updated to use defaulted copies and moves so that the traits are reported correctly and TSparseArray can use faster assignment.

diff --git a/Engine/Source/Runtime/Core/Public/Containers/Set.h b/Engine/Source/Runtime/Core/Public/Containers/Set.h
--- a/Engine/Source/Runtime/Core/Public/Containers/Set.h
+++ b/Engine/Source/Runtime/Core/Public/Containers/Set.h
@@ -164,13 +164,10 @@ public:
 	/** Initialization constructor. */
 	template <typename InitType, typename = typename TEnableIf<!TAreTypesEqual<TSetElement, typename TDecay<InitType>::Type>::Value>::Type> explicit FORCEINLINE TSetElement(InitType&& InValue) : Value(Forward<InitType>(InValue)) {}
 
-	/** Copy/move constructors */
-	FORCEINLINE TSetElement(const TSetElement&  Rhs) : Value(                   Rhs.Value ), HashNextId(         Rhs.HashNextId ), HashIndex(Rhs.HashIndex) {}
-	FORCEINLINE TSetElement(      TSetElement&& Rhs) : Value(MoveTempIfPossible(Rhs.Value)), HashNextId(MoveTemp(Rhs.HashNextId)), HashIndex(Rhs.HashIndex) {}
-
-	/** Copy/move assignment */
-	FORCEINLINE TSetElement& operator=(const TSetElement&  Rhs) { Value =                    Rhs.Value ; HashNextId =          Rhs.HashNextId ; HashIndex = Rhs.HashIndex; return *this; }
-	FORCEINLINE TSetElement& operator=(      TSetElement&& Rhs) { Value = MoveTempIfPossible(Rhs.Value); HashNextId = MoveTemp(Rhs.HashNextId); HashIndex = Rhs.HashIndex; return *this; }
+	TSetElement(TSetElement&&) = default;
+	TSetElement(const TSetElement&) = default;
+	TSetElement& operator=(TSetElement&&) = default;
+	TSetElement& operator=(const TSetElement&) = default;
 
 	/** Serializer. */
 	FORCEINLINE friend FArchive& operator<<(FArchive& Ar,TSetElement& Element)
