Improved deprecation message and removed use of legacy typedef

Changes based on feedback from Marc Audy

diff --git a/Engine/Source/Runtime/Core/Private/UObject/UnrealNames.cpp b/Engine/Source/Runtime/Core/Private/UObject/UnrealNames.cpp
--- a/Engine/Source/Runtime/Core/Private/UObject/UnrealNames.cpp
+++ b/Engine/Source/Runtime/Core/Private/UObject/UnrealNames.cpp
@@ -919,7 +919,7 @@ static FNameEntryId DebugCastNameEntryId(int32 Id) { return (FNameEntryId&)(Id);
 * @param	Index	Name index to look up string for
 * @return			Associated name
 */
-const TCHAR* DebugFName(NAME_INDEX Index)
+const TCHAR* DebugFName(FNameEntryId Index)
 {
 	// Hardcoded static array. This function is only used inside the debugger so it should be fine to return it.
 	static TCHAR TempName[NAME_SIZE];
@@ -1769,7 +1769,7 @@ void FName::DisplayHash(FOutputDevice& Ar)
 	GetNamePool().LogStats(Ar);
 }
 
-FString FName::SafeString(NAME_INDEX InDisplayIndex, int32 InstanceNumber)
+FString FName::SafeString(FNameEntryId InDisplayIndex, int32 InstanceNumber)
 {
 	return FName(InDisplayIndex, InDisplayIndex, InstanceNumber).ToString();
 }
diff --git a/Engine/Source/Runtime/Core/Public/UObject/NameTypes.h b/Engine/Source/Runtime/Core/Public/UObject/NameTypes.h
--- a/Engine/Source/Runtime/Core/Public/UObject/NameTypes.h
+++ b/Engine/Source/Runtime/Core/Public/UObject/NameTypes.h
@@ -51,7 +51,7 @@ struct FNameEntryId
 
 	UE_DEPRECATED(4.23, "NAME_INDEX is replaced by FNameEntryId, which is no longer a contiguous integer. "
 						"Please use 'GetTypeHash(MyId)' instead of 'MyId' for hash functions. "
-						"ToUnstableInt() can be used in other cases if you *really* know what you are doing.")
+						"ToUnstableInt() can be used in other advanced cases.")
 	operator int32() const;
 
 	/** Get process specific integer */
@@ -375,15 +375,15 @@ struct FScriptName
 class CORE_API FName
 {
 public:
-	FORCEINLINE NAME_INDEX GetComparisonIndex() const
+	FORCEINLINE FNameEntryId GetComparisonIndex() const
 	{
 		checkName(IsWithinBounds(ComparisonIndex));
 		return ComparisonIndex;
 	}
 
-	FORCEINLINE NAME_INDEX GetDisplayIndex() const
+	FORCEINLINE FNameEntryId GetDisplayIndex() const
 	{
-		const NAME_INDEX Index = GetDisplayIndexFast();
+		const FNameEntryId Index = GetDisplayIndexFast();
 		checkName(IsWithinBounds(Index));
 		return Index;
 	}
@@ -608,7 +608,7 @@ public:
 	 * Create an FName from its component parts
 	 * Only call this if you *really* know what you're doing
 	 */
-	FORCEINLINE FName( const NAME_INDEX InComparisonIndex, const NAME_INDEX InDisplayIndex, const int32 InNumber )
+	FORCEINLINE FName( const FNameEntryId InComparisonIndex, const FNameEntryId InDisplayIndex, const int32 InNumber )
 		: ComparisonIndex( InComparisonIndex )
 #if WITH_CASE_PRESERVING_NAME
 		, DisplayIndex( InDisplayIndex )
@@ -713,7 +713,7 @@ public:
 	}
 
 	static void DisplayHash( class FOutputDevice& Ar );
-	static FString SafeString(NAME_INDEX InDisplayIndex, int32 InstanceNumber = NAME_NO_NUMBER_INTERNAL);
+	static FString SafeString(FNameEntryId InDisplayIndex, int32 InstanceNumber = NAME_NO_NUMBER_INTERNAL);
 
 	/**
 	 * @return Size of all name entries.
@@ -776,7 +776,7 @@ private:
 	friend const TCHAR* DebugFName(int32, int32);
 	friend const TCHAR* DebugFName(FName&);
 
-	FORCEINLINE NAME_INDEX GetDisplayIndexFast() const
+	FORCEINLINE FNameEntryId GetDisplayIndexFast() const
 	{
 #if WITH_CASE_PRESERVING_NAME
 		return DisplayIndex;
