TIdentity moved to its own header.

diff --git a/Engine/Source/Runtime/Core/Public/Templates/Identity.h b/Engine/Source/Runtime/Core/Public/Templates/Identity.h
new file mode 100644
--- /dev/null
+++ b/Engine/Source/Runtime/Core/Public/Templates/Identity.h
@@ -0,0 +1,18 @@
+// Copyright 1998-2019 Epic Games, Inc. All Rights Reserved.
+
+#pragma once
+
+/**
+ * Returns the same type passed to it.  This is useful in a few cases, but mainly for inhibiting template argument deduction in function arguments, e.g.:
+ *
+ * template <typename T>
+ * void Func1(T Val); // Can be called like Func(123) or Func<int>(123);
+ *
+ * template <typename T>
+ * void Func2(typename TIdentity<T>::Type Val); // Must be called like Func<int>(123)
+ */
+template <typename T>
+struct TIdentity
+{
+	typedef T Type;
+};
diff --git a/Engine/Source/Runtime/Core/Public/Templates/UnrealTemplate.h b/Engine/Source/Runtime/Core/Public/Templates/UnrealTemplate.h
--- a/Engine/Source/Runtime/Core/Public/Templates/UnrealTemplate.h
+++ b/Engine/Source/Runtime/Core/Public/Templates/UnrealTemplate.h
@@ -12,6 +12,7 @@
 #include "Templates/UnrealTypeTraits.h"
 #include "Templates/RemoveReference.h"
 #include "Templates/TypeCompatibleBytes.h"
+#include "Templates/Identity.h"
 #include "Traits/IsContiguousContainer.h"
 
 /*-----------------------------------------------------------------------------
@@ -583,21 +584,6 @@ struct TLosesQualifiersFromTo
 	enum { Value = !TAreTypesEqual<typename TCopyQualifiersFromTo<From, To>::Type, To>::Value };
 };
 
-/**
- * Returns the same type passed to it.  This is useful in a few cases, but mainly for inhibiting template argument deduction in function arguments, e.g.:
- *
- * template <typename T>
- * void Func1(T Val); // Can be called like Func(123) or Func<int>(123);
- *
- * template <typename T>
- * void Func2(typename TIdentity<T>::Type Val); // Must be called like Func<int>(123)
- */
-template <typename T>
-struct TIdentity
-{
-	typedef T Type;
-};
-
 /**
  * Equivalent to std::declval.
  *
