TModels template used for checking concepts.  Lots of comments telling you how to write concepts.

diff --git a/Engine/Source/Runtime/Core/Public/Templates/Models.h b/Engine/Source/Runtime/Core/Public/Templates/Models.h
new file mode 100644
--- /dev/null
+++ b/Engine/Source/Runtime/Core/Public/Templates/Models.h
@@ -0,0 +1,193 @@
+// Copyright 1998-2019 Epic Games, Inc. All Rights Reserved.
+
+#pragma once
+
+#include "Templates/Identity.h"
+
+/**
+ * Utilities for concepts checks.
+ *
+ * In this case, a successful concept check means that a given C++ expression is well-formed.
+ * No guarantees are given about the correctness, behavior or complexity of the runtime behaviour of that expression.
+ *
+ * Concepts are structs with a rather unusual definition:
+ *
+ * struct CConcept
+ * {
+ *     template <[...concept parameters...]>
+ *     auto Requires([...placeholder variables...]) -> decltype(
+ *         [...expression(s) to test the validity of...]
+ *     );
+ * };
+ *
+ * The prefix C is reserved for concepts, and concepts should be directly named as an adjective and not like a predicate, i.e.:
+ * CEqualityComparable  - good
+ * CIsComparable        - bad
+ * CHasEqualsOperator   - bad
+ *
+ * Concepts can be checked using the TModels trait:
+ *
+ * TModels<Concept, [...arguments...]>::Value
+ *
+ * The arguments are forwarded to the template parameters of the concept's Requires() function, which will attempt to
+ * compile the expression(s) in the return value, and SFINAE is utilized to test whether that succeeded.
+ *
+ * The placeholders are simply any variable declarations you need to write your expression(s).
+ *
+ * Note that 'legal C++' doesn't necessarily mean that the expression will compile when used.  The concept
+ * check only tests that the syntax is valid.  Instantiation of function template bodies may still fail.
+ * See the CContainerLvalueAddable example below.
+ */
+
+
+/**
+ * Traits class which does concept checking.
+ */
+template <typename Concept, typename... Args>
+struct TModels
+{
+	template <typename... Ts>
+	static char (&Resolve(decltype(&Concept::template Requires<Ts...>)*))[2];
+
+	template <typename... Ts>
+	static char (&Resolve(...))[1];
+
+	static constexpr bool Value = sizeof(Resolve<Args...>(0)) == 2;
+};
+
+/**
+ * Helper function which can be used as an expression in a concept to refine ('inherit') another concept.
+ * It should be used as expression-based variant of the TModels traits class.  If the arguments model
+ * the given concept, Refines<Concept, Args...> is a valid expression, otherwise it is not.
+ *
+ * See the CCopyablePointer example below.
+ */
+template <typename Concept, typename... Args>
+auto Refines() -> int(&)[!!TModels<Concept, Args...>::Value * 2 - 1];
+
+
+/************
+ * Examples *
+ ************/
+
+/**
+ * // Definition of a negatable type.
+ * struct CNegatable
+ * {
+ *     template <typename T>
+ *     auto Requires(const T& Val) -> decltype(
+ *         -Val
+ *     );
+ * };
+ *
+ * static_assert( TModels<CNegatable, int >::Value); // ints are negatable
+ * static_assert(!TModels<CNegatable, int*>::Value); // pointers are not negatable
+ */
+
+/**
+ * // Definition of an incrementable type.
+ * // Pre-increment and post-increment must both be legal.
+ * struct CIncrementable
+ * {
+ *     template <typename T>
+ *     auto Requires(T& Val) -> decltype(
+ *         ++Val,
+ *         Val++
+ *     );
+ * };
+ *
+ * static_assert( TModels<CIncrementable, int  >::Value); // ints are incrementable
+ * static_assert( TModels<CIncrementable, int* >::Value); // pointers are incrementable
+ * static_assert(!TModels<CIncrementable, int[]>::Value); // arrays are not incrementable
+ */
+
+/**
+ * // Definition of comparability between two types.
+ * // Requires both == and != and commutability.
+ * struct CEqualityComparable
+ * {
+ *     template <typename T, typename U>
+ *     auto Requires(const T& A, const U& B) -> decltype(
+ *         A == B,
+ *         B == A,
+ *         A != B,
+ *         B != A
+ *     );
+ * };
+ *
+ * static_assert( TModels<CEqualityComparable, FArchive*, FArchiveUObject*>::Value); // base pointers are comparable with derived pointers
+ * static_assert(!TModels<CEqualityComparable, int*,      float*          >::Value); // unrelated pointers are not comparable
+ */
+
+/**
+ * // Definition of a type with a nested ElementType and can have lvalues of that type added to it.
+ * // Note: this is just an example, not a good description of a container.
+ * struct CLvalueAddableContainer
+ * {
+ *     template <typename ContainerType>
+ *     auto Requires(ContainerType& Container, typename ContainerType::ElementType& Element) -> decltype(
+ *         Container.Add(Element)
+ *     );
+ * };
+ *
+ * static_assert(TModels<CLvalueAddableContainer, TArray<TUniquePtr<int>>>::Value); // success, but...
+ *
+ * TUniquePtr<int> Temp;
+ * Array.Add(Temp); // compile error when TArray attempts to copy the TUniquePtr
+ *
+ * This is because TArray doesn't remove its Add(const ElementType&) overload when the element type
+ * is non-copyable - instead we get a compile error inside the template instantiation when we attempt the call.
+ */
+
+/**
+ * // Definition of a type subtractable from itself and whose result is also (convertible to) itself.
+ * struct CGroupSubtractable
+ * {
+ *     template <typename T>
+ *     auto Requires(T& Result, const T& A, const T& B) -> decltype(
+ *         Result = A - B
+ *     );
+ * };
+ *
+ * static_assert( TModels<CGroupSubtractable, int >::Value); // ints form a group under subtraction
+ * static_assert(!TModels<CGroupSubtractable, int*>::Value); // pointers do not result in another pointer under subtraction and so do not form a group
+ */
+
+/**
+ * // Definition of a copyable type.
+ * struct CCopyable
+ * {
+ *     template <typename T>
+ *     auto Requires(T& Dest, const T& Val) -> decltype(
+ *         T(Val),
+ *         Dest = Val
+ *     );
+ * };
+ *
+ * // Definition of a dereferencable type.
+ * struct CDereferencable
+ * {
+ *     template <typename T>
+ *     auto Requires(const T& Val) -> decltype(
+ *         *Val
+ *     );
+ * };
+ *
+ * // Definition of a copyable pointer-like type which requires copyability and dereferencability (refining existing concepts) and bool-testability (stated directly).
+ * // Note: this is just an example, not a great concept.
+ * struct CCopyablePointer
+ * {
+ *     template <typename T>
+ *     auto Requires(bool& Result, bool AnyBool, const T& Val) -> decltype(
+ *         Refines<CCopyable, T>(),
+ *         Refines<CDereferencable, T>(),
+ *         Result = AnyBool && Val
+ *     );
+ * };
+ *
+ * static_assert( TModels<CCopyablePointer, int*           >::Value); // raw pointers model this concept
+ * static_assert( TModels<CCopyablePointer, TSharedPtr<int>>::Value); // TSharedPtrs model this concept
+ * static_assert(!TModels<CCopyablePointer, TUniquePtr<int>>::Value); // TUniquePtrs are dereferenceable and bool-testable, but not copyable
+ * static_assert(!TModels<CCopyablePointer, int            >::Value); // ints are copyable and bool-testable, but not dereferencable
+ * static_assert(!TModels<CCopyablePointer, FString        >::Value); // FStrings are copyable and dereferencable (to get the const TCHAR*), but not bool-testable
+ */
