Deprecation of old traits in favor of new concepts.

diff --git a/Engine/Source/Runtime/Core/Public/Concepts/EqualityComparable.h b/Engine/Source/Runtime/Core/Public/Concepts/EqualityComparable.h
new file mode 100644
--- /dev/null
+++ b/Engine/Source/Runtime/Core/Public/Concepts/EqualityComparable.h
@@ -0,0 +1,28 @@
+// Copyright 1998-2019 Epic Games, Inc. All Rights Reserved.
+
+#pragma once
+
+#include "CoreTypes.h"
+
+/**
+ * Describes a type comparable with another type.
+ */
+struct CEqualityComparableWith {
+	template <typename T>
+	auto Requires(bool& Result, const T& A, const T& B) -> decltype(
+		Result = A == B,
+		Result = B == A,
+		Result = A != B,
+		Result = B != A
+	);
+};
+
+/**
+ * Describes a type comparable with itself.
+ */
+struct CEqualityComparable {
+	template <typename T>
+	auto Requires() -> decltype(
+		Refines<CEqualityComparableWith, T, T>()
+	);
+};
diff --git a/Engine/Source/Runtime/Core/Public/Concepts/Insertable.h b/Engine/Source/Runtime/Core/Public/Concepts/Insertable.h
new file mode 100644
--- /dev/null
+++ b/Engine/Source/Runtime/Core/Public/Concepts/Insertable.h
@@ -0,0 +1,15 @@
+// Copyright 1998-2019 Epic Games, Inc. All Rights Reserved.
+
+#pragma once
+
+#include "CoreTypes.h"
+
+/**
+ * Describes a type with a GetTypeHash overload.
+ */
+struct CInsertable {
+	template <typename DestType, typename T>
+	auto Requires(DestType& Dest, T& Val) -> decltype(
+		Dest << Val
+	);
+};
diff --git a/Engine/Source/Runtime/Core/Public/Templates/HasInserterOperator.h b/Engine/Source/Runtime/Core/Public/Templates/HasInserterOperator.h
--- a/Engine/Source/Runtime/Core/Public/Templates/HasInserterOperator.h
+++ b/Engine/Source/Runtime/Core/Public/Templates/HasInserterOperator.h
@@ -41,7 +41,7 @@ namespace UE4HasInserterOperator_Private
  * Traits class which tests if a type has an operator<< overload between two types.
  */
 template <typename Dest, typename T>
-struct THasInserterOperator
+struct UE_DEPRECATED(4.23, "THasInserterOperator has been deprecated, please use TModels<CInsertable, Dest, T> instead.") THasInserterOperator
 {
 	enum { Value = UE4HasInserterOperator_Private::THasInserterOperator<Dest, T>::Value };
 };
diff --git a/Engine/Source/Runtime/Core/Public/Templates/HasOperatorEquals.h b/Engine/Source/Runtime/Core/Public/Templates/HasOperatorEquals.h
--- a/Engine/Source/Runtime/Core/Public/Templates/HasOperatorEquals.h
+++ b/Engine/Source/Runtime/Core/Public/Templates/HasOperatorEquals.h
@@ -44,7 +44,7 @@ namespace UE4HasOperatorEquals_Private
  * Traits class which tests if a type has an operator== overload.
  */
 template <typename T>
-struct THasOperatorEquals
+struct UE_DEPRECATED(4.23, "THasOperatorEquals has been deprecated, please use TModels<CEqualityComparable, T> instead.") THasOperatorEquals
 {
 	enum { Value = UE4HasOperatorEquals_Private::Equals<T>::Value };
 };
@@ -53,7 +53,7 @@ struct THasOperatorEquals
  * Traits class which tests if a type has an operator!= overload.
  */
 template <typename T>
-struct THasOperatorNotEquals
+struct UE_DEPRECATED(4.23, "THasOperatorNotEquals has been deprecated, please use TModels<CEqualityComparable, T> instead.") THasOperatorNotEquals
 {
 	enum { Value = UE4HasOperatorEquals_Private::NotEquals<T>::Value };
 };
