Removed some deprecated stuff from FArchive

diff --git a/Engine/Source/Runtime/Core/Private/Serialization/Archive.cpp b/Engine/Source/Runtime/Core/Private/Serialization/Archive.cpp
--- a/Engine/Source/Runtime/Core/Private/Serialization/Archive.cpp
+++ b/Engine/Source/Runtime/Core/Private/Serialization/Archive.cpp
@@ -605,13 +605,6 @@ public:
 };
 #endif		// WITH_MULTI_THREADED_COMPRESSION
 
-
-
-void FArchive::SerializeCompressed( void* V, int64 Length, ECompressionFlags Flags, bool bTreatBufferAsFileReader, bool bUsePlatformBitWindow )
-{
-	SerializeCompressed(V, Length, FCompression::GetCompressionFormatFromDeprecatedFlags(Flags), ECompressionFlags(Flags & COMPRESS_OptionsFlagsMask), bTreatBufferAsFileReader);
-}
-
 void FArchive::SerializeCompressed(void* V, int64 Length, FName CompressionFormat, ECompressionFlags Flags, bool bTreatBufferAsFileReader)
 {
 	if( IsLoading() )
diff --git a/Engine/Source/Runtime/Core/Public/Serialization/Archive.h b/Engine/Source/Runtime/Core/Public/Serialization/Archive.h
--- a/Engine/Source/Runtime/Core/Public/Serialization/Archive.h
+++ b/Engine/Source/Runtime/Core/Public/Serialization/Archive.h
@@ -746,8 +746,6 @@ public:
 	 * @param	bTreatBufferAsFileReader true if V is actually an FArchive, which is used when saving to read data - helps to avoid single huge allocations of source data
 	 * @param	bUsePlatformBitWindow use a platform specific bitwindow setting
 	 */
-	UE_DEPRECATED(4.20, "Use the FName based version of SerializeCompressed (which also removes the basically-unused bUsePlatformBitWindow)")
-	void SerializeCompressed(void* V, int64 Length, ECompressionFlags Flags = COMPRESS_NoFlags, bool bTreatBufferAsFileReader = false, bool bUsePlatformBitWindow = false);
 	void SerializeCompressed(void* V, int64 Length, FName CompressionFormat, ECompressionFlags Flags=COMPRESS_NoFlags, bool bTreatBufferAsFileReader=false);
 
 
@@ -1369,35 +1367,29 @@ private:
 	/** Copies all of the members except CustomVersionContainer */
 	void CopyTrivialFArchiveStatusMembers(const FArchive& ArchiveStatusToCopy);
 
-public:
 	/** Whether this archive is for loading data. */
-	UE_DEPRECATED(4.20, "Direct access to ArIsLoading has been deprecated - please use IsLoading() and SetIsLoading() instead.")
 	uint8 ArIsLoading : 1;
 
 	/** Whether this archive is for saving data. */
-	UE_DEPRECATED(4.20, "Direct access to ArIsSaving has been deprecated - please use IsSaving() and SetIsSaving() instead.")
 	uint8 ArIsSaving : 1;
 
 	/** Whether archive is transacting. */
-	UE_DEPRECATED(4.20, "Direct access to ArIsTransacting has been deprecated - please use IsTransacting() and SetIsTransacting() instead.")
 	uint8 ArIsTransacting : 1;
 
 	/** Whether this archive serializes to a text format. Text format archives should use high level constructs from FStructuredArchive for delimiting data rather than manually seeking through the file. */
-	UE_DEPRECATED(4.20, "Direct access to ArIsTextFormat has been deprecated - please use IsTextFormat() and SetIsTextFormat() instead.")
 	uint8 ArIsTextFormat : 1;
 
 	/** Whether this archive wants properties to be serialized in binary form instead of tagged. */
-	UE_DEPRECATED(4.20, "Direct access to ArWantBinaryPropertySerialization has been deprecated - please use WantBinaryPropertySerialization() and SetWantBinaryPropertySerialization() instead.")
 	uint8 ArWantBinaryPropertySerialization : 1;
 
 	/** Whether this archive wants to always save strings in unicode format */
-	UE_DEPRECATED(4.20, "Direct access to ArForceUnicode has been deprecated - please use ForceUnicode() and SetForceUnicode() instead.")
 	uint8 ArForceUnicode : 1;
 
 	/** Whether this archive saves to persistent storage. */
-	UE_DEPRECATED(4.20, "Direct access to ArIsPersistent has been deprecated - please use IsPersistent() and SetIsPersistent() instead.")
 	uint8 ArIsPersistent : 1;
 
+public:
+
 	/** Whether this archive contains errors. */
 	uint8 ArIsError : 1;
 
