Fix FText rounding to truncate rounding errors
- FText rounding relied on FMath::RoundHalfFrom/ToZero to truncate the fraction part before rounding from/to zero. This behavior in FMath has changed and this
change brings truncation back to FText rounding routines. The rounding errors is caused when rounding with x decimal digitis. For example FText RoundHalfFromZero
expects 1000.1225 to be rounded to 1000.123 but calling FMath::Modf results in 0.1249999. Trying to round this with 3 decimal digits without truncation
results in 0.122 => 1000.122.

diff --git a/Engine/Source/Runtime/Core/Private/Internationalization/FastDecimalFormat.cpp b/Engine/Source/Runtime/Core/Private/Internationalization/FastDecimalFormat.cpp
--- a/Engine/Source/Runtime/Core/Private/Internationalization/FastDecimalFormat.cpp
+++ b/Engine/Source/Runtime/Core/Private/Internationalization/FastDecimalFormat.cpp
@@ -279,7 +279,7 @@ void FractionalToString_SplitAndRoundNumber(const bool bIsNegative, const double
 	// Multiply the value to round by 10^DecimalPlacesToRoundTo - this will allow us to perform rounding calculations 
 	// that correctly trim any remaining fractional parts that are outside of our rounding range
 	double& ValueToRound = ((bIsRoundingEntireNumber) ? IntegralPart : FractionalPart);
-	ValueToRound *= Pow10Table[DecimalPlacesToRoundTo];
+	ValueToRound = FMath::TruncateToHalfIfClose(ValueToRound * Pow10Table[DecimalPlacesToRoundTo]);
 
 	// The rounding modes here mimic those of ICU. See http://userguide.icu-project.org/formatparse/numbers/rounding-modes
 	switch (InRoundingMode)
diff --git a/Engine/Source/Runtime/Core/Private/Math/UnrealMath.cpp b/Engine/Source/Runtime/Core/Private/Math/UnrealMath.cpp
--- a/Engine/Source/Runtime/Core/Private/Math/UnrealMath.cpp
+++ b/Engine/Source/Runtime/Core/Private/Math/UnrealMath.cpp
@@ -2786,42 +2786,37 @@ void FVector::GenerateClusterCenters(TArray<FVector>& Clusters, const TArray<FVe
 	}
 }
 
-namespace MathRoundingUtil
-{
-
-float TruncateToHalfIfClose(float F)
+float FMath::TruncateToHalfIfClose(float F, float Tolerance)
 {
 	float ValueToFudgeIntegralPart = 0.0f;
 	float ValueToFudgeFractionalPart = FMath::Modf(F, &ValueToFudgeIntegralPart);
 	if (F < 0.0f)
 	{
-		return ValueToFudgeIntegralPart + ((FMath::IsNearlyEqual(ValueToFudgeFractionalPart, -0.5f)) ? -0.5f : ValueToFudgeFractionalPart);
+		return ValueToFudgeIntegralPart + (FMath::IsNearlyEqual(ValueToFudgeFractionalPart, -0.5f, Tolerance) ? -0.5f : ValueToFudgeFractionalPart);
 	}
 	else
 	{
-		return ValueToFudgeIntegralPart + ((FMath::IsNearlyEqual(ValueToFudgeFractionalPart, 0.5f)) ? 0.5f : ValueToFudgeFractionalPart);
+		return ValueToFudgeIntegralPart + (FMath::IsNearlyEqual(ValueToFudgeFractionalPart, 0.5f, Tolerance) ? 0.5f : ValueToFudgeFractionalPart);
 	}
 }
 
-double TruncateToHalfIfClose(double F)
+double FMath::TruncateToHalfIfClose(double F, double Tolerance)
 {
 	double ValueToFudgeIntegralPart = 0.0;
 	double ValueToFudgeFractionalPart = FMath::Modf(F, &ValueToFudgeIntegralPart);
 	if (F < 0.0)
 	{
-		return ValueToFudgeIntegralPart + ((FMath::IsNearlyEqual(ValueToFudgeFractionalPart, -0.5)) ? -0.5 : ValueToFudgeFractionalPart);
+		return ValueToFudgeIntegralPart + (FMath::IsNearlyEqual(ValueToFudgeFractionalPart, -0.5, Tolerance) ? -0.5 : ValueToFudgeFractionalPart);
 	}
 	else
 	{
-		return ValueToFudgeIntegralPart + ((FMath::IsNearlyEqual(ValueToFudgeFractionalPart, 0.5)) ? 0.5 : ValueToFudgeFractionalPart);
+		return ValueToFudgeIntegralPart + (FMath::IsNearlyEqual(ValueToFudgeFractionalPart, 0.5, Tolerance) ? 0.5 : ValueToFudgeFractionalPart);
 	}
 }
 
-} // namespace GenericPlatformMathInternal
-
 float FMath::RoundHalfToEven(float F)
 {
-	F = MathRoundingUtil::TruncateToHalfIfClose(F);
+	F = FMath::TruncateToHalfIfClose(F);
 
 	const bool bIsNegative = F < 0.0f;
 	const bool bValueIsEven = static_cast<uint32>(FloorToFloat(((bIsNegative) ? -F : F))) % 2 == 0;
@@ -2839,7 +2834,7 @@ float FMath::RoundHalfToEven(float F)
 
 double FMath::RoundHalfToEven(double F)
 {
-	F = MathRoundingUtil::TruncateToHalfIfClose(F);
+	F = FMath::TruncateToHalfIfClose(F);
 
 	const bool bIsNegative = F < 0.0;
 	const bool bValueIsEven = static_cast<uint64>(FMath::FloorToDouble(((bIsNegative) ? -F : F))) % 2 == 0;
@@ -3441,4 +3436,4 @@ float FMath::PerlinNoise1D(const float Value)
 	const float v = rx1 * g[p[bx1]];
 
 	return 2.0f * (u + sx * (v - u));
-}
\ No newline at end of file
+}
diff --git a/Engine/Source/Runtime/Core/Public/Math/UnrealMathUtility.h b/Engine/Source/Runtime/Core/Public/Math/UnrealMathUtility.h
--- a/Engine/Source/Runtime/Core/Public/Math/UnrealMathUtility.h
+++ b/Engine/Source/Runtime/Core/Public/Math/UnrealMathUtility.h
@@ -1310,6 +1310,15 @@ struct FMath : public FPlatformMath
 	 */
 	static CORE_API bool PointsAreCoplanar(const TArray<FVector>& Points, const float Tolerance = 0.1f);
 
+	/**
+	 * Truncates a floating point number to half if closer than the given tolerance.
+	 * @param Value				Floating point number to truncate
+	 * @param Tolerance			Maximum allowed difference to 0.5 in order to truncate
+	 * @return					The truncated value
+	 */
+	static CORE_API float TruncateToHalfIfClose(float F, float Tolerance = SMALL_NUMBER);
+	static CORE_API double TruncateToHalfIfClose(double F, double Tolerance = SMALL_NUMBER);
+
 	/**
 	* Converts a floating point number to the nearest integer, equidistant ties go to the value which is closest to an even value: 1.5 becomes 2, 0.5 becomes 0
 	* @param F		Floating point value to convert
