Deprecated default FName operator < and >

We've had several cases where the previously default alphabetical sorting caused performance issues when it wasn't needed. We now require you to explicitly choose either slow alphabetical sorting or fast non-alphabetical order.

The fast order does not decode, access or compare actual string data, but the order is not stable between process runs.

diff --git a/Engine/Plugins/Animation/LiveLink/Source/LiveLink/Private/LiveLinkModule.cpp b/Engine/Plugins/Animation/LiveLink/Source/LiveLink/Private/LiveLinkModule.cpp
--- a/Engine/Plugins/Animation/LiveLink/Source/LiveLink/Private/LiveLinkModule.cpp
+++ b/Engine/Plugins/Animation/LiveLink/Source/LiveLink/Private/LiveLinkModule.cpp
@@ -1,4 +1,4 @@
-﻿// Copyright 1998-2019 Epic Games, Inc. All Rights Reserved.
+// Copyright 1998-2019 Epic Games, Inc. All Rights Reserved.
 
 #include "ILiveLinkModule.h"
 
@@ -149,7 +149,7 @@ public:
 		for (const FHeaderEntry& Header : Headers)
 		{
 			TArray<FName> Subjects = LiveLinkSubjects.FindChecked(Header.Key);
-			Subjects.Sort();
+			Subjects.Sort(FNameLexicalLess());
 			for (FName Subject : Subjects)
 			{
 				FName FullName = *FString::Format(TEXT("{0} ({1})"), { Subject.ToString(), Header.Value.ToString() });
diff --git a/Engine/Plugins/Compositing/Composure/Source/Composure/Private/MovieScene/MovieSceneComposureExportSectionTemplate.cpp b/Engine/Plugins/Compositing/Composure/Source/Composure/Private/MovieScene/MovieSceneComposureExportSectionTemplate.cpp
--- a/Engine/Plugins/Compositing/Composure/Source/Composure/Private/MovieScene/MovieSceneComposureExportSectionTemplate.cpp
+++ b/Engine/Plugins/Compositing/Composure/Source/Composure/Private/MovieScene/MovieSceneComposureExportSectionTemplate.cpp
@@ -81,7 +81,7 @@ private:
 struct FMovieSceneComposureExportPasses
 {
 	/** Map of internal transform pass name (or NAME_None for the Output), to the export config options */
-	TSortedMap<FName, FMovieSceneComposureExportPass, FDefaultAllocator, FNameSortIndexes> PassesToExport;
+	TSortedMap<FName, FMovieSceneComposureExportPass, FDefaultAllocator, FNameFastLess> PassesToExport;
 
 	void AddPass(const FMovieSceneComposureExportPass& InPass, ACompositingElement* CompShotElement)
 	{
diff --git a/Engine/Plugins/Editor/AssetManagerEditor/Source/AssetManagerEditor/Private/AssetManagerEditorModule.cpp b/Engine/Plugins/Editor/AssetManagerEditor/Source/AssetManagerEditor/Private/AssetManagerEditorModule.cpp
--- a/Engine/Plugins/Editor/AssetManagerEditor/Source/AssetManagerEditor/Private/AssetManagerEditorModule.cpp
+++ b/Engine/Plugins/Editor/AssetManagerEditor/Source/AssetManagerEditor/Private/AssetManagerEditorModule.cpp
@@ -182,7 +182,7 @@ void IAssetManagerEditorModule::GeneratePrimaryAssetTypeComboBoxStrings(TArray<
 	TArray<FPrimaryAssetTypeInfo> TypeInfos;
 
 	AssetManager.GetPrimaryAssetTypeInfoList(TypeInfos);
-	TypeInfos.Sort([](const FPrimaryAssetTypeInfo& LHS, const FPrimaryAssetTypeInfo& RHS) { return LHS.PrimaryAssetType < RHS.PrimaryAssetType; });
+	TypeInfos.Sort([](const FPrimaryAssetTypeInfo& LHS, const FPrimaryAssetTypeInfo& RHS) { return LHS.PrimaryAssetType.LexicalLess(RHS.PrimaryAssetType); });
 
 	// Can the field be cleared
 	if (bAllowClear)
@@ -1983,7 +1983,7 @@ void FAssetManagerEditorModule::LogAssetsWithMultipleLabels()
 		}
 	}
 
-	PackageToLabelMap.KeySort(TLess<FName>());
+	PackageToLabelMap.KeySort(FNameLexicalLess());
 
 	UE_LOG(LogAssetManagerEditor, Log, TEXT("\nAssets with multiple labels follow"));
 
@@ -2036,7 +2036,7 @@ void FAssetManagerEditorModule::DumpAssetDependencies(const TArray<FString>& Arg
 
 	Manager.GetPrimaryAssetTypeInfoList(TypeInfos);
 
-	TypeInfos.Sort([](const FPrimaryAssetTypeInfo& LHS, const FPrimaryAssetTypeInfo& RHS) { return LHS.PrimaryAssetType < RHS.PrimaryAssetType; });
+	TypeInfos.Sort([](const FPrimaryAssetTypeInfo& LHS, const FPrimaryAssetTypeInfo& RHS) { return LHS.PrimaryAssetType.LexicalLess(RHS.PrimaryAssetType); });
 
 	UE_LOG(LogAssetManagerEditor, Log, TEXT("=========== Asset Manager Dependencies ==========="));
 
@@ -2083,7 +2083,7 @@ void FAssetManagerEditorModule::DumpAssetDependencies(const TArray<FString>& Arg
 		{
 			UE_LOG(LogAssetManagerEditor, Log, TEXT("  Type %s:"), *TypeInfo.PrimaryAssetType.ToString());
 
-			DependencyInfos.Sort([](const FDependencyInfo& LHS, const FDependencyInfo& RHS) { return LHS.AssetName < RHS.AssetName; });
+			DependencyInfos.Sort([](const FDependencyInfo& LHS, const FDependencyInfo& RHS) { return LHS.AssetName.LexicalLess(RHS.AssetName); });
 
 			for (FDependencyInfo& DependencyInfo : DependencyInfos)
 			{
diff --git a/Engine/Plugins/Experimental/PythonScriptPlugin/Source/PythonScriptPlugin/Private/PyWrapperTypeRegistry.cpp b/Engine/Plugins/Experimental/PythonScriptPlugin/Source/PythonScriptPlugin/Private/PyWrapperTypeRegistry.cpp
--- a/Engine/Plugins/Experimental/PythonScriptPlugin/Source/PythonScriptPlugin/Private/PyWrapperTypeRegistry.cpp
+++ b/Engine/Plugins/Experimental/PythonScriptPlugin/Source/PythonScriptPlugin/Private/PyWrapperTypeRegistry.cpp
@@ -1978,7 +1978,7 @@ void FPyWrapperTypeRegistry::GenerateStubCodeForWrappedTypes(const EPyOnlineDocs
 	{
 		TArray<FName> ModuleNames;
 		GeneratedWrappedTypesForModule.GetKeys(ModuleNames);
-		ModuleNames.Sort();
+		ModuleNames.Sort(FNameLexicalLess());
 
 		bool bExportedImports = false;
 		for (const FName ModuleName : ModuleNames)
diff --git a/Engine/Plugins/FX/Niagara/Source/NiagaraEditor/Private/NiagaraGraph.cpp b/Engine/Plugins/FX/Niagara/Source/NiagaraEditor/Private/NiagaraGraph.cpp
--- a/Engine/Plugins/FX/Niagara/Source/NiagaraEditor/Private/NiagaraGraph.cpp
+++ b/Engine/Plugins/FX/Niagara/Source/NiagaraEditor/Private/NiagaraGraph.cpp
@@ -134,7 +134,7 @@ void UNiagaraGraph::PostLoad()
 	if (bAllZeroes && UniqueNames.Num() > 1)
 	{
 		// Just do the lexicographic sort and assign the call order to their ordered index value.
-		UniqueNames.Sort();
+		UniqueNames.Sort(FNameLexicalLess());
 		for (UNiagaraNodeInput* InputNode : InputNodes)
 		{
 			if (InputNode->Usage == ENiagaraInputNodeUsage::Parameter)
diff --git a/Engine/Plugins/FX/Niagara/Source/NiagaraEditor/Private/NiagaraHlslTranslator.cpp b/Engine/Plugins/FX/Niagara/Source/NiagaraEditor/Private/NiagaraHlslTranslator.cpp
--- a/Engine/Plugins/FX/Niagara/Source/NiagaraEditor/Private/NiagaraHlslTranslator.cpp
+++ b/Engine/Plugins/FX/Niagara/Source/NiagaraEditor/Private/NiagaraHlslTranslator.cpp
@@ -1107,7 +1107,7 @@ const FNiagaraTranslateResults &FHlslNiagaraTranslator::Translate(const FNiagara
 			// We sort the variables so that they end up in the same ordering between Spawn & Update...
 			SystemEngineReadVars.Sort([&](const FNiagaraVariable& A, const FNiagaraVariable& B)
 			{
-				return A.GetName() < B.GetName();
+				return A.GetName().LexicalLess(B.GetName());
 			});
 
 			{
@@ -1156,12 +1156,12 @@ const FNiagaraTranslateResults &FHlslNiagaraTranslator::Translate(const FNiagara
 		// We sort the variables so that they end up in the same ordering between Spawn & Update...
 		InstanceReadVars.Sort([&](const FNiagaraVariable& A, const FNiagaraVariable& B)
 		{
-			return A.GetName() < B.GetName();
+			return A.GetName().LexicalLess(B.GetName());
 		});
 		// We sort the variables so that they end up in the same ordering between Spawn & Update...
 		InstanceWriteVars.Sort([&](const FNiagaraVariable& A, const FNiagaraVariable& B)
 		{
-			return A.GetName() < B.GetName();
+			return A.GetName().LexicalLess(B.GetName());
 		});
 		//Define the simulation context. Which is a helper struct containing all the input, result and intermediate data needed for a single simulation.
 		//Allows us to reuse the same simulate function but provide different wrappers for final IO between GPU and CPU sims.
diff --git a/Engine/Plugins/FX/Niagara/Source/NiagaraEditor/Private/ViewModels/Stack/NiagaraStackFunctionInputCollection.cpp b/Engine/Plugins/FX/Niagara/Source/NiagaraEditor/Private/ViewModels/Stack/NiagaraStackFunctionInputCollection.cpp
--- a/Engine/Plugins/FX/Niagara/Source/NiagaraEditor/Private/ViewModels/Stack/NiagaraStackFunctionInputCollection.cpp
+++ b/Engine/Plugins/FX/Niagara/Source/NiagaraEditor/Private/ViewModels/Stack/NiagaraStackFunctionInputCollection.cpp
@@ -150,7 +150,7 @@ void UNiagaraStackFunctionInputCollection::RefreshChildrenInternal(const TArray<
 		}
 		else
 		{
-			return A.Pin->PinName < B.Pin->PinName;
+			return A.Pin->PinName.LexicalLess(B.Pin->PinName);
 		}
 	});
 
diff --git a/Engine/Plugins/FX/Niagara/Source/NiagaraEditor/Private/Widgets/SNiagaraParameterMapView.cpp b/Engine/Plugins/FX/Niagara/Source/NiagaraEditor/Private/Widgets/SNiagaraParameterMapView.cpp
--- a/Engine/Plugins/FX/Niagara/Source/NiagaraEditor/Private/Widgets/SNiagaraParameterMapView.cpp
+++ b/Engine/Plugins/FX/Niagara/Source/NiagaraEditor/Private/Widgets/SNiagaraParameterMapView.cpp
@@ -347,7 +347,7 @@ void SNiagaraParameterMapView::CollectAllActions(FGraphActionListBuilderBase& Ou
 		}
 	}
 
-	ParameterEntries.KeySort([](const FNiagaraVariable& A, const FNiagaraVariable& B) { return (A.GetName() < B.GetName()); });
+	ParameterEntries.KeySort([](const FNiagaraVariable& A, const FNiagaraVariable& B) { return A.GetName().LexicalLess(B.GetName()); });
 
 	const FText TooltipFormat = LOCTEXT("Parameters", "Name: {0} \nType: {1}");
 	for (const auto& ParameterEntry : ParameterEntries)
@@ -905,7 +905,7 @@ void SNiagaraAddParameterMenu::CollectAllActions(FGraphActionListBuilderBase& Ou
 		for (TWeakObjectPtr<UNiagaraGraph>& Graph : Graphs)
 		{
 			TMap<FNiagaraVariable, FNiagaraGraphParameterReferenceCollection> ParameterEntries = Graph.Get()->GetParameterMap();
-			ParameterEntries.KeySort([](const FNiagaraVariable& A, const FNiagaraVariable& B) { return (A.GetName() < B.GetName()); });
+			ParameterEntries.KeySort([](const FNiagaraVariable& A, const FNiagaraVariable& B) { return A.GetName().LexicalLess(B.GetName()); });
 
 			for (const auto& ParameterEntry : ParameterEntries)
 			{
@@ -961,7 +961,7 @@ void SNiagaraAddParameterMenu::AddParameterGroup(FGraphActionListBuilderBase& Ou
 {
 	if (bSort)
 	{
-		Variables.Sort([](const FNiagaraVariable& A, const FNiagaraVariable& B) { return (A.GetName() < B.GetName()); });
+		Variables.Sort([](const FNiagaraVariable& A, const FNiagaraVariable& B) { return A.GetName().LexicalLess(B.GetName()); });
 	}
 
 	for (FNiagaraVariable& Variable : Variables)
diff --git a/Engine/Plugins/Runtime/Steam/SteamAudio/Source/SteamAudio/Private/PhononProbeVolume.h b/Engine/Plugins/Runtime/Steam/SteamAudio/Source/SteamAudio/Private/PhononProbeVolume.h
--- a/Engine/Plugins/Runtime/Steam/SteamAudio/Source/SteamAudio/Private/PhononProbeVolume.h
+++ b/Engine/Plugins/Runtime/Steam/SteamAudio/Source/SteamAudio/Private/PhononProbeVolume.h
@@ -38,14 +38,14 @@ struct FBakedDataInfo
 	int32 Size;
 };
 
-inline bool operator==(const FBakedDataInfo& lhs, const FBakedDataInfo& rhs)
+inline bool operator==(const FBakedDataInfo& Lhs, const FBakedDataInfo& Rhs)
 {
-	return lhs.Name == rhs.Name && lhs.Size == rhs.Size;
+	return Lhs.Name == Rhs.Name && Lhs.Size == Rhs.Size;
 }
 
-inline bool operator<(const FBakedDataInfo& lhs, const FBakedDataInfo& rhs)
+inline bool operator<(const FBakedDataInfo& Lhs, const FBakedDataInfo& Rhs)
 {
-	return lhs.Name < rhs.Name;
+	return Lhs.Name.LexicalLess(Rhs.Name);
 }
 
 /**
diff --git a/Engine/Source/Developer/CollectionManager/Private/Collection.cpp b/Engine/Source/Developer/CollectionManager/Private/Collection.cpp
--- a/Engine/Source/Developer/CollectionManager/Private/Collection.cpp
+++ b/Engine/Source/Developer/CollectionManager/Private/Collection.cpp
@@ -194,7 +194,7 @@ bool FCollection::Save(const TArray<FText>& AdditionalChangelistText, FText& Out
 	{
 		// Write out the set as a sorted array to keep things in a known order for diffing
 		TArray<FName> ObjectList = ObjectSet.Array();
-		ObjectList.Sort();
+		ObjectList.Sort(FNameLexicalLess());
 
 		// Static collection. Save a flat list of all objects in the collection.
 		for (const FName& ObjectName : ObjectList)
@@ -581,7 +581,7 @@ void FCollection::PrintCollection() const
 
 		// Print the set as a sorted array to keep things in a sane order
 		TArray<FName> ObjectList = ObjectSet.Array();
-		ObjectList.Sort();
+		ObjectList.Sort(FNameLexicalLess());
 
 		for (const FName& ObjectName : ObjectList)
 		{
@@ -903,8 +903,8 @@ bool FCollection::CheckinCollection(const TArray<FText>& AdditionalChangelistTex
 			TArray<FName> ObjectsRemoved;
 			GetObjectDifferencesFromDisk(ObjectsAdded, ObjectsRemoved);
 
-			ObjectsAdded.Sort();
-			ObjectsRemoved.Sort();
+			ObjectsAdded.Sort(FNameLexicalLess());
+			ObjectsRemoved.Sort(FNameLexicalLess());
 
 			// Report added files
 			FFormatNamedArguments Args;
diff --git a/Engine/Source/Developer/GameplayDebugger/Private/GameplayDebuggerAddonManager.cpp b/Engine/Source/Developer/GameplayDebugger/Private/GameplayDebuggerAddonManager.cpp
--- a/Engine/Source/Developer/GameplayDebugger/Private/GameplayDebuggerAddonManager.cpp
+++ b/Engine/Source/Developer/GameplayDebugger/Private/GameplayDebuggerAddonManager.cpp
@@ -45,7 +45,7 @@ void FGameplayDebuggerAddonManager::NotifyCategoriesChanged()
 		int32 CategoryId;
 		int32 SlotIdx;
 
-		bool operator<(const FSlotInfo& Other) const { return (SlotIdx == Other.SlotIdx) ? (CategoryName < Other.CategoryName) : (SlotIdx < Other.SlotIdx); }
+		bool operator<(const FSlotInfo& Other) const { return (SlotIdx == Other.SlotIdx) ? CategoryName.LexicalLess(Other.CategoryName) : (SlotIdx < Other.SlotIdx); }
 	};
 
 	TArray<FSlotInfo> AssignList;
diff --git a/Engine/Source/Developer/LogVisualizer/Private/VisualLoggerCanvasRenderer.cpp b/Engine/Source/Developer/LogVisualizer/Private/VisualLoggerCanvasRenderer.cpp
--- a/Engine/Source/Developer/LogVisualizer/Private/VisualLoggerCanvasRenderer.cpp
+++ b/Engine/Source/Developer/LogVisualizer/Private/VisualLoggerCanvasRenderer.cpp
@@ -254,7 +254,7 @@ void FVisualLoggerCanvasRenderer::DrawHistogramGraphs(class UCanvas* Canvas, cla
 
 			auto& CategoriesForGraph = UsedGraphCategories.FindOrAdd(It->Key.ToString());
 
-			It->Value.GraphLines.KeySort(TLess<FName>());
+			It->Value.GraphLines.KeySort(FNameLexicalLess());
 			for (auto LinesIt(It->Value.GraphLines.CreateConstIterator()); LinesIt; ++LinesIt)
 			{
 				const FString DataName = LinesIt->Value.DataName.ToString();
diff --git a/Engine/Source/Developer/ModuleUI/Private/SModuleUI.cpp b/Engine/Source/Developer/ModuleUI/Private/SModuleUI.cpp
--- a/Engine/Source/Developer/ModuleUI/Private/SModuleUI.cpp
+++ b/Engine/Source/Developer/ModuleUI/Private/SModuleUI.cpp
@@ -202,7 +202,7 @@ void SModuleUI::UpdateModuleListItems()
 	{
 		FORCEINLINE bool operator()( const TSharedPtr<FModuleListItem>& A, const TSharedPtr<FModuleListItem>& B ) const
 		{
-			return A->ModuleName < B->ModuleName;
+			return A->ModuleName.LexicalLess(B->ModuleName);
 		}
 	};
 	ModuleListItems.Sort( FModuleSorter() );
diff --git a/Engine/Source/Developer/Profiler/Private/Widgets/SEventGraph.cpp b/Engine/Source/Developer/Profiler/Private/Widgets/SEventGraph.cpp
--- a/Engine/Source/Developer/Profiler/Private/Widgets/SEventGraph.cpp
+++ b/Engine/Source/Developer/Profiler/Private/Widgets/SEventGraph.cpp
@@ -1051,7 +1051,7 @@ void SEventGraph::FillThreadFilterOptions()
 	// Sort the thread names alphabetically
 	ThreadNamesForCombo.Sort([]( const TSharedPtr<FName> Lhs, const TSharedPtr<FName> Rhs ) 
 	{
-		return Lhs->IsNone() || ( !Rhs->IsNone() && *Lhs < *Rhs );
+		return Lhs->IsNone() || ( !Rhs->IsNone() && Lhs->LexicalLess(*Rhs) );
 	});
 
 	// Refresh the combo box
diff --git a/Engine/Source/Developer/Profiler/Private/Widgets/SFiltersAndPresets.h b/Engine/Source/Developer/Profiler/Private/Widgets/SFiltersAndPresets.h
--- a/Engine/Source/Developer/Profiler/Private/Widgets/SFiltersAndPresets.h
+++ b/Engine/Source/Developer/Profiler/Private/Widgets/SFiltersAndPresets.h
@@ -241,7 +241,7 @@ struct FGroupAndStatSorting
 	{
 		FORCEINLINE_DEBUGGABLE bool operator()( const FGroupOrStatNodePtr& A, const FGroupOrStatNodePtr& B ) const 
 		{
-			return A->GetName() < B->GetName();
+			return A->GetName().LexicalLess(B->GetName());
 		}
 	};
 
@@ -250,7 +250,7 @@ struct FGroupAndStatSorting
 	{
 		FORCEINLINE_DEBUGGABLE bool operator()( const FGroupOrStatNodePtr& A, const FGroupOrStatNodePtr& B ) const 
 		{
-			return A->GetGroupName() < B->GetGroupName();
+			return A->GetGroupName().LexicalLess(B->GetGroupName());
 		}
 	};
 
@@ -265,7 +265,7 @@ struct FGroupAndStatSorting
 			if( TypeA == TypeB )
 			{
 				// Sort by stat name.
-				return A->GetName() < B->GetName();
+				return A->GetName().LexicalLess(B->GetName());
 			}
 			else
 			{
diff --git a/Engine/Source/Editor/AnimGraph/Private/AnimGraphNode_ModifyCurve.cpp b/Engine/Source/Editor/AnimGraph/Private/AnimGraphNode_ModifyCurve.cpp
--- a/Engine/Source/Editor/AnimGraph/Private/AnimGraphNode_ModifyCurve.cpp
+++ b/Engine/Source/Editor/AnimGraph/Private/AnimGraphNode_ModifyCurve.cpp
@@ -48,7 +48,7 @@ TArray<FName> UAnimGraphNode_ModifyCurve::GetCurvesToAdd() const
 			CurvesToAdd.RemoveSingleSwap(ExistingCurveName, false);
 		}
 
-		CurvesToAdd.Sort();
+		CurvesToAdd.Sort(FNameLexicalLess());
 	}
 
 	return CurvesToAdd;
diff --git a/Engine/Source/Editor/AnimGraph/Private/PoseDriverDetails.cpp b/Engine/Source/Editor/AnimGraph/Private/PoseDriverDetails.cpp
--- a/Engine/Source/Editor/AnimGraph/Private/PoseDriverDetails.cpp
+++ b/Engine/Source/Editor/AnimGraph/Private/PoseDriverDetails.cpp
@@ -1064,7 +1064,7 @@ void FPoseDriverDetails::UpdateDrivenNameOptions()
 				{
 					TArray<FName> NameArray;
 					Mapping->FillNameArray(NameArray);
-					NameArray.Sort();
+					NameArray.Sort(FNameLexicalLess());
 
 					for (FName CurveName : NameArray)
 					{
diff --git a/Engine/Source/Editor/ContentBrowser/Private/AssetContextMenu.cpp b/Engine/Source/Editor/ContentBrowser/Private/AssetContextMenu.cpp
--- a/Engine/Source/Editor/ContentBrowser/Private/AssetContextMenu.cpp
+++ b/Engine/Source/Editor/ContentBrowser/Private/AssetContextMenu.cpp
@@ -1408,7 +1408,7 @@ bool FAssetContextMenu::AddCollectionMenuOptions(FMenuBuilder& MenuBuilder)
 
 			AvailableCollections.Sort([](const FCollectionNameType& One, const FCollectionNameType& Two) -> bool
 			{
-				return One.Name < Two.Name;
+				return One.Name.LexicalLess(Two.Name);
 			});
 
 			for (const FCollectionNameType& AvailableCollection : AvailableCollections)
diff --git a/Engine/Source/Editor/ContentBrowser/Private/CollectionContextMenu.cpp b/Engine/Source/Editor/ContentBrowser/Private/CollectionContextMenu.cpp
--- a/Engine/Source/Editor/ContentBrowser/Private/CollectionContextMenu.cpp
+++ b/Engine/Source/Editor/ContentBrowser/Private/CollectionContextMenu.cpp
@@ -255,7 +255,7 @@ void FCollectionContextMenu::MakeSaveDynamicCollectionSubMenu(FMenuBuilder& Menu
 
 	AvailableCollections.Sort([](const FCollectionNameType& One, const FCollectionNameType& Two) -> bool
 	{
-		return One.Name < Two.Name;
+		return One.Name.LexicalLess(Two.Name);
 	});
 
 	if (AvailableCollections.Num() > 0)
diff --git a/Engine/Source/Editor/ContentBrowser/Private/CollectionViewTypes.h b/Engine/Source/Editor/ContentBrowser/Private/CollectionViewTypes.h
--- a/Engine/Source/Editor/ContentBrowser/Private/CollectionViewTypes.h
+++ b/Engine/Source/Editor/ContentBrowser/Private/CollectionViewTypes.h
@@ -37,7 +37,7 @@ struct FCollectionItem
 	{
 		FORCEINLINE bool operator()( const TSharedPtr<FCollectionItem>& A, const TSharedPtr<FCollectionItem>& B ) const
 		{
-			return A->CollectionName < B->CollectionName;
+			return A->CollectionName.LexicalLess(B->CollectionName);
 		}
 	};
 
diff --git a/Engine/Source/Editor/GameProjectGeneration/Private/SProjectBrowser.cpp b/Engine/Source/Editor/GameProjectGeneration/Private/SProjectBrowser.cpp
--- a/Engine/Source/Editor/GameProjectGeneration/Private/SProjectBrowser.cpp
+++ b/Engine/Source/Editor/GameProjectGeneration/Private/SProjectBrowser.cpp
@@ -872,10 +872,7 @@ FReply SProjectBrowser::FindProjects()
 						TargetPlatforms.Add(PlatformInfo.PlatformInfoName);
 					}
 				}
-				TargetPlatforms.Sort([](const FName& One, const FName& Two) -> bool
-				{
-					return One < Two;
-				});
+				TargetPlatforms.Sort(FNameLexicalLess());
 
 				const bool bIsNewProjectItem = false;
 				TSharedRef<FProjectItem> NewProjectItem = MakeShareable( new FProjectItem(ProjectName, ProjectDescription, ProjectEngineIdentifier, bIsUpToDate, DynamicBrush, ProjectFilename, bIsNewProjectItem, TargetPlatforms, ProjectStatus.SupportsAllPlatforms() ) );
diff --git a/Engine/Source/Editor/Kismet/Private/BlueprintCompilationManager.cpp b/Engine/Source/Editor/Kismet/Private/BlueprintCompilationManager.cpp
--- a/Engine/Source/Editor/Kismet/Private/BlueprintCompilationManager.cpp
+++ b/Engine/Source/Editor/Kismet/Private/BlueprintCompilationManager.cpp
@@ -597,7 +597,7 @@ void FBlueprintCompilationManagerImpl::FlushCompilationQueueImpl(bool bSuppressB
 
 			if (DepthA == DepthB)
 			{
-				return A.GetFName() < B.GetFName(); 
+				return A.GetFName().LexicalLess(B.GetFName());
 			}
 			return DepthA < DepthB;
 		};
@@ -1633,7 +1633,7 @@ void FBlueprintCompilationManagerImpl::ReinstanceBatch(TArray<FReinstancingJob>&
 
 			if (DepthA == DepthB && A && B)
 			{
-				return A->GetFName() < B->GetFName(); 
+				return A->GetFName().LexicalLess(B->GetFName());
 			}
 			return DepthA < DepthB;
 		}
diff --git a/Engine/Source/Editor/LandscapeEditor/Private/LandscapeEdMode.cpp b/Engine/Source/Editor/LandscapeEditor/Private/LandscapeEdMode.cpp
--- a/Engine/Source/Editor/LandscapeEditor/Private/LandscapeEdMode.cpp
+++ b/Engine/Source/Editor/LandscapeEditor/Private/LandscapeEdMode.cpp
@@ -2598,7 +2598,7 @@ void FEdModeLandscape::UpdateTargetLayerDisplayOrder(ELandscapeLayerDisplayMode
 				SavedTargetNameList.Add(LandscapeTargetList[i]->LayerName);
 			}
 
-			SavedTargetNameList.Sort();
+			SavedTargetNameList.Sort(FNameLexicalLess());
 
 			// Then insert the non layer target that shouldn't be sorted
 			for (int32 i = 0; i < GetTargetLayerStartingIndex(); ++i)
diff --git a/Engine/Source/Editor/LevelEditor/Private/SLevelViewportToolBar.cpp b/Engine/Source/Editor/LevelEditor/Private/SLevelViewportToolBar.cpp
--- a/Engine/Source/Editor/LevelEditor/Private/SLevelViewportToolBar.cpp
+++ b/Engine/Source/Editor/LevelEditor/Private/SLevelViewportToolBar.cpp
@@ -1218,7 +1218,7 @@ static TMap<FName, TArray<UFoliageType*>> GroupFoliageByOuter(const TArray<UFoli
 		}
 	}
 
-	Result.KeySort([](const FName& A, const FName& B) { return (A < B && B != NAME_None); });
+	Result.KeySort([](const FName& A, const FName& B) { return (A.LexicalLess(B) && B != NAME_None); });
 	return Result;
 }
 
diff --git a/Engine/Source/Editor/MovieSceneTools/Private/MovieSceneEventCustomization.cpp b/Engine/Source/Editor/MovieSceneTools/Private/MovieSceneEventCustomization.cpp
--- a/Engine/Source/Editor/MovieSceneTools/Private/MovieSceneEventCustomization.cpp
+++ b/Engine/Source/Editor/MovieSceneTools/Private/MovieSceneEventCustomization.cpp
@@ -350,7 +350,7 @@ void FMovieSceneEventCustomization::PopulateQuickBindSubMenu(FMenuBuilder& MenuB
 			}
 		}
 	
-		Algo::SortBy(Functions, &UFunction::GetFName);
+		Algo::SortBy(Functions, &UFunction::GetFName, FNameLexicalLess());
 
 		for (UFunction* Function : Functions)
 		{
diff --git a/Engine/Source/Editor/MovieSceneTools/Private/TrackEditors/MaterialParameterCollectionTrackEditor.cpp b/Engine/Source/Editor/MovieSceneTools/Private/TrackEditors/MaterialParameterCollectionTrackEditor.cpp
--- a/Engine/Source/Editor/MovieSceneTools/Private/TrackEditors/MaterialParameterCollectionTrackEditor.cpp
+++ b/Engine/Source/Editor/MovieSceneTools/Private/TrackEditors/MaterialParameterCollectionTrackEditor.cpp
@@ -203,7 +203,7 @@ TSharedRef<SWidget> FMaterialParameterCollectionTrackEditor::OnGetAddParameterMe
 	MenuBuilder.BeginSection(NAME_None, LOCTEXT("ScalarParametersHeading", "Scalar"));
 	{
 		TArray<FCollectionScalarParameter> ScalarParameters = MPCTrack->MPC->ScalarParameters;
-		Algo::SortBy(ScalarParameters, &FCollectionParameterBase::ParameterName);
+		Algo::SortBy(ScalarParameters, &FCollectionParameterBase::ParameterName, FNameLexicalLess());
 
 		for (const FCollectionScalarParameter& Scalar : ScalarParameters)
 		{
@@ -220,7 +220,7 @@ TSharedRef<SWidget> FMaterialParameterCollectionTrackEditor::OnGetAddParameterMe
 	MenuBuilder.BeginSection(NAME_None, LOCTEXT("VectorParametersHeading", "Vector"));
 	{
 		TArray<FCollectionVectorParameter> VectorParameters = MPCTrack->MPC->VectorParameters;
-		Algo::SortBy(VectorParameters, &FCollectionParameterBase::ParameterName);
+		Algo::SortBy(VectorParameters, &FCollectionParameterBase::ParameterName, FNameLexicalLess());
 
 		for (const FCollectionVectorParameter& Vector : VectorParameters)
 		{
diff --git a/Engine/Source/Editor/MovieSceneTools/Private/TrackEditors/MaterialTrackEditor.cpp b/Engine/Source/Editor/MovieSceneTools/Private/TrackEditors/MaterialTrackEditor.cpp
--- a/Engine/Source/Editor/MovieSceneTools/Private/TrackEditors/MaterialTrackEditor.cpp
+++ b/Engine/Source/Editor/MovieSceneTools/Private/TrackEditors/MaterialTrackEditor.cpp
@@ -53,7 +53,7 @@ struct FParameterNameAndAction
 
 	bool operator<(FParameterNameAndAction const& Other) const
 	{
-		return ParameterName < Other.ParameterName;
+		return ParameterName.LexicalLess(Other.ParameterName);
 	}
 };
 
diff --git a/Engine/Source/Editor/MovieSceneTools/Private/TrackEditors/ParticleParameterTrackEditor.h b/Engine/Source/Editor/MovieSceneTools/Private/TrackEditors/ParticleParameterTrackEditor.h
--- a/Engine/Source/Editor/MovieSceneTools/Private/TrackEditors/ParticleParameterTrackEditor.h
+++ b/Engine/Source/Editor/MovieSceneTools/Private/TrackEditors/ParticleParameterTrackEditor.h
@@ -64,7 +64,7 @@ private:
 
 		bool operator<( FParameterNameAndAction const& Other ) const
 		{
-			return ParameterName < Other.ParameterName;
+			return ParameterName.LexicalLess(Other.ParameterName);
 		}
 	};
 
diff --git a/Engine/Source/Editor/Persona/Private/SAnimCurvePanel.cpp b/Engine/Source/Editor/Persona/Private/SAnimCurvePanel.cpp
--- a/Engine/Source/Editor/Persona/Private/SAnimCurvePanel.cpp
+++ b/Engine/Source/Editor/Persona/Private/SAnimCurvePanel.cpp
@@ -855,7 +855,7 @@ void SAnimCurvePanel::UpdatePanel()
 			MetadataNameMap->GetName(A.Name.UID, AName);
 			MetadataNameMap->GetName(B.Name.UID, BName);
 
-			return AName < BName;
+			return AName.LexicalLess(BName);
 		});
 
 		// Store expanded state before clearing the tracks
diff --git a/Engine/Source/Editor/Persona/Private/SAnimTrackCurvePanel.cpp b/Engine/Source/Editor/Persona/Private/SAnimTrackCurvePanel.cpp
--- a/Engine/Source/Editor/Persona/Private/SAnimTrackCurvePanel.cpp
+++ b/Engine/Source/Editor/Persona/Private/SAnimTrackCurvePanel.cpp
@@ -644,7 +644,7 @@ void SAnimTrackCurvePanel::UpdatePanel()
 			MetadataNameMap->GetName(A.Name.UID, AName);
 			MetadataNameMap->GetName(B.Name.UID, BName);
 
-			return AName < BName;
+			return AName.LexicalLess(BName);
 		});
 
 		// see if we need to clear or not
diff --git a/Engine/Source/Editor/Persona/Private/SSlotNameReferenceWindow.cpp b/Engine/Source/Editor/Persona/Private/SSlotNameReferenceWindow.cpp
--- a/Engine/Source/Editor/Persona/Private/SSlotNameReferenceWindow.cpp
+++ b/Engine/Source/Editor/Persona/Private/SSlotNameReferenceWindow.cpp
@@ -222,7 +222,7 @@ void SSlotNameReferenceWindow::UpdateInfo(FReferenceWindowInfo& UpdatedInfo)
 		// Sort for nicer display
 		UpdatedInfo.ReferencingMontages->Sort([](const FAssetData& A, const FAssetData& B)
 		{
-			return A.AssetName < B.AssetName;
+			return A.AssetName.LexicalLess(B.AssetName);
 		});
 
 		ReferencingMontages.Empty(UpdatedInfo.ReferencingMontages->Num());
diff --git a/Engine/Source/Editor/PinnedCommandList/Private/SPinnedCommandList.cpp b/Engine/Source/Editor/PinnedCommandList/Private/SPinnedCommandList.cpp
--- a/Engine/Source/Editor/PinnedCommandList/Private/SPinnedCommandList.cpp
+++ b/Engine/Source/Editor/PinnedCommandList/Private/SPinnedCommandList.cpp
@@ -722,7 +722,7 @@ void SPinnedCommandList::SortCommands()
 		}
 
 		// fallback to lexical sort
-		return CommandIdentifier0 < CommandIdentifier1;
+		return CommandIdentifier0.LexicalLess(CommandIdentifier1);
 	});
 
 	RefreshCommandWidgets();
diff --git a/Engine/Source/Editor/SceneOutliner/Private/SSceneOutliner.cpp b/Engine/Source/Editor/SceneOutliner/Private/SSceneOutliner.cpp
--- a/Engine/Source/Editor/SceneOutliner/Private/SSceneOutliner.cpp
+++ b/Engine/Source/Editor/SceneOutliner/Private/SSceneOutliner.cpp
@@ -2235,7 +2235,7 @@ namespace SceneOutliner
 		CacheFoldersEdit = InFolders;
 
 		// Sort folder names so parents appear before children
-		CacheFoldersEdit.Sort();
+		CacheFoldersEdit.Sort(FNameLexicalLess());
 
 		// Cache existing children
 		for (FName Folder : CacheFoldersEdit)
@@ -2423,7 +2423,7 @@ namespace SceneOutliner
 			// Sort in descending order so children will be deleted before parents
 			CacheFoldersDelete.Sort([](const FFolderTreeItem& FolderA, const FFolderTreeItem& FolderB)
 			{
-				return (FolderA.Path > FolderB.Path);
+				return FolderB.Path.LexicalLess(FolderA.Path);
 			});
 
 			for (FFolderTreeItem* Folder : CacheFoldersDelete)
diff --git a/Engine/Source/Editor/Sequencer/Private/ISequencerSection.cpp b/Engine/Source/Editor/Sequencer/Private/ISequencerSection.cpp
--- a/Engine/Source/Editor/Sequencer/Private/ISequencerSection.cpp
+++ b/Engine/Source/Editor/Sequencer/Private/ISequencerSection.cpp
@@ -119,7 +119,7 @@ void ISequencerSection::GenerateSectionLayout( ISectionLayoutBuilder& LayoutBuil
 		Pair.Value.Channels.Sort([](const FChannelData& A, const FChannelData& B){
 			if (A.MetaData.SortOrder == B.MetaData.SortOrder)
 			{
-				return A.MetaData.Name < B.MetaData.Name;
+				return A.MetaData.Name.LexicalLess(B.MetaData.Name);
 			}
 			return A.MetaData.SortOrder < B.MetaData.SortOrder;
 		});
diff --git a/Engine/Source/Editor/SkeletonEditor/Private/EditableSkeleton.cpp b/Engine/Source/Editor/SkeletonEditor/Private/EditableSkeleton.cpp
--- a/Engine/Source/Editor/SkeletonEditor/Private/EditableSkeleton.cpp
+++ b/Engine/Source/Editor/SkeletonEditor/Private/EditableSkeleton.cpp
@@ -425,9 +425,9 @@ void FEditableSkeleton::RemoveSmartnamesAndFixupAnimations(const FName& InContai
 		{ 
 			if (A.AssetClass == B.AssetClass)
 			{
-				return A.AssetName < B.AssetName;
+				return A.AssetName.LexicalLess(B.AssetName);
 			}
-			return A.AssetClass < B.AssetClass;
+			return A.AssetClass.LexicalLess(B.AssetClass);
 		});
 
 		for (FAssetData& Data : AnimationAssets)
diff --git a/Engine/Source/Editor/StaticMeshEditor/Private/StaticMeshEditorTools.cpp b/Engine/Source/Editor/StaticMeshEditor/Private/StaticMeshEditorTools.cpp
--- a/Engine/Source/Editor/StaticMeshEditor/Private/StaticMeshEditorTools.cpp
+++ b/Engine/Source/Editor/StaticMeshEditor/Private/StaticMeshEditorTools.cpp
@@ -3585,7 +3585,7 @@ TArray<FName> FLevelOfDetailSettingsLayout::GetLODScreenSizePlatformOverrideName
 {
 	TArray<FName> KeyArray;
 	LODScreenSizes[LODIndex].PerPlatform.GenerateKeyArray(KeyArray);
-	KeyArray.Sort();
+	KeyArray.Sort(FNameLexicalLess());
 	return KeyArray;
 }
 
@@ -4058,7 +4058,7 @@ TArray<FName> FLevelOfDetailSettingsLayout::GetMinLODPlatformOverrideNames() con
 	check(StaticMesh);
 	TArray<FName> KeyArray;
 	StaticMesh->MinLOD.PerPlatform.GenerateKeyArray(KeyArray);
-	KeyArray.Sort();
+	KeyArray.Sort(FNameLexicalLess());
 	return KeyArray;
 }
 
@@ -4168,7 +4168,7 @@ TArray<FName> FLevelOfDetailSettingsLayout::GetNumStreamedLODsPlatformOverrideNa
 	check(StaticMesh);
 	TArray<FName> KeyArray;
 	StaticMesh->NumStreamedLODs.PerPlatform.GenerateKeyArray(KeyArray);
-	KeyArray.Sort();
+	KeyArray.Sort(FNameLexicalLess());
 	return KeyArray;
 }
 
diff --git a/Engine/Source/Editor/UMGEditor/Private/WidgetBlueprintCompiler.cpp b/Engine/Source/Editor/UMGEditor/Private/WidgetBlueprintCompiler.cpp
--- a/Engine/Source/Editor/UMGEditor/Private/WidgetBlueprintCompiler.cpp
+++ b/Engine/Source/Editor/UMGEditor/Private/WidgetBlueprintCompiler.cpp
@@ -286,7 +286,7 @@ void FWidgetBlueprintCompilerContext::CreateClassVariablesFromBlueprint()
 	TArray<UWidget*> Widgets = WidgetBP->GetAllSourceWidgets();
 
 	// Sort the widgets alphabetically
-	Widgets.Sort( []( const UWidget& Lhs, const UWidget& Rhs ) { return Rhs.GetFName() < Lhs.GetFName(); } );
+	Widgets.Sort( []( const UWidget& Lhs, const UWidget& Rhs ) { return Rhs.GetFName().LexicalLess(Lhs.GetFName()); } );
 
 	// Add widget variables
 	for ( UWidget* Widget : Widgets )
diff --git a/Engine/Source/Editor/UnrealEd/Private/Commandlets/AssetRegUtilCommandlet.cpp b/Engine/Source/Editor/UnrealEd/Private/Commandlets/AssetRegUtilCommandlet.cpp
--- a/Engine/Source/Editor/UnrealEd/Private/Commandlets/AssetRegUtilCommandlet.cpp
+++ b/Engine/Source/Editor/UnrealEd/Private/Commandlets/AssetRegUtilCommandlet.cpp
@@ -146,7 +146,7 @@ struct FSortableDependencySort
 
 		if ((AClass != BClass) && bIsAClassGrouped && bIsBClassGrouped)
 		{
-			return AClass > BClass;
+			return BClass.LexicalLess(AClass);
 		}
 
 		if (A.DepSet != B.DepSet)
diff --git a/Engine/Source/Editor/UnrealEd/Private/Commandlets/AssetRegistryGenerator.cpp b/Engine/Source/Editor/UnrealEd/Private/Commandlets/AssetRegistryGenerator.cpp
--- a/Engine/Source/Editor/UnrealEd/Private/Commandlets/AssetRegistryGenerator.cpp
+++ b/Engine/Source/Editor/UnrealEd/Private/Commandlets/AssetRegistryGenerator.cpp
@@ -1235,7 +1235,7 @@ bool FAssetRegistryGenerator::GenerateAssetChunkInformationCSV(const FString& Ou
 	// Sort list so it's consistent over time
 	AssetDataList.Sort([](const FAssetData& A, const FAssetData& B)
 	{
-		return A.ObjectPath < B.ObjectPath;
+		return A.ObjectPath.LexicalLess(B.ObjectPath);
 	});
 
 	// Create file for all chunks
diff --git a/Engine/Source/Editor/UnrealEd/Private/Commandlets/GatherTextCommandletBase.cpp b/Engine/Source/Editor/UnrealEd/Private/Commandlets/GatherTextCommandletBase.cpp
--- a/Engine/Source/Editor/UnrealEd/Private/Commandlets/GatherTextCommandletBase.cpp
+++ b/Engine/Source/Editor/UnrealEd/Private/Commandlets/GatherTextCommandletBase.cpp
@@ -38,7 +38,7 @@ void UGatherTextCommandletBase::Initialize( const TSharedRef< FLocTextHelper >&
 		{
 			SplitPlatforms.Add(*SplitPlatformName, FString::Printf(TEXT("/%s/"), *SplitPlatformName));
 		}
-		SplitPlatforms.KeySort(TLess<FName>());
+		SplitPlatforms.KeySort(FNameLexicalLess());
 	}
 }
 
diff --git a/Engine/Source/Editor/UnrealEd/Private/Commandlets/GatherTextFromAssetsCommandlet.cpp b/Engine/Source/Editor/UnrealEd/Private/Commandlets/GatherTextFromAssetsCommandlet.cpp
--- a/Engine/Source/Editor/UnrealEd/Private/Commandlets/GatherTextFromAssetsCommandlet.cpp
+++ b/Engine/Source/Editor/UnrealEd/Private/Commandlets/GatherTextFromAssetsCommandlet.cpp
@@ -854,7 +854,7 @@ int32 UGatherTextFromAssetsCommandlet::Main(const FString& Params)
 	}
 	check(PackagesPendingGather.Num() == 0);
 
-	PackagesWithStaleGatherCache.Sort();
+	PackagesWithStaleGatherCache.Sort(FNameLexicalLess());
 
 	if (bReportStaleGatherCache)
 	{
diff --git a/Engine/Source/Editor/WorldBrowser/Private/Tiles/WorldTileModel.h b/Engine/Source/Editor/WorldBrowser/Private/Tiles/WorldTileModel.h
--- a/Engine/Source/Editor/WorldBrowser/Private/Tiles/WorldTileModel.h
+++ b/Engine/Source/Editor/WorldBrowser/Private/Tiles/WorldTileModel.h
@@ -32,7 +32,7 @@ public:
 		FORCEINLINE bool operator()(const TSharedPtr<FLevelModel>& A, 
 									const TSharedPtr<FLevelModel>& B) const
 		{
-			return A->GetLongPackageName() < B->GetLongPackageName();
+			return A->GetLongPackageName().LexicalLess(B->GetLongPackageName());
 		}
 	};
 
diff --git a/Engine/Source/Programs/UnrealHeaderTool/Private/CodeGenerator.cpp b/Engine/Source/Programs/UnrealHeaderTool/Private/CodeGenerator.cpp
--- a/Engine/Source/Programs/UnrealHeaderTool/Private/CodeGenerator.cpp
+++ b/Engine/Source/Programs/UnrealHeaderTool/Private/CodeGenerator.cpp
@@ -515,7 +515,7 @@ static FString OutputMetaDataCodeForObject(FOutputDevice& OutDeclaration, FOutpu
 
 		// We sort the metadata here so that we can get consistent output across multiple runs
 		// even when metadata is added in a different order
-		Algo::SortBy(KVPs, &KVPType::Key);
+		Algo::SortBy(KVPs, &KVPType::Key, FNameLexicalLess());
 
 		for (const KVPType& KVP : KVPs)
 		{
@@ -2397,7 +2397,7 @@ void FNativeClassHeaderGenerator::ExportNatives(FOutputDevice& Out, FClass* Clas
 			}
 		}
 
-		Algo::SortBy(NamedFunctionsToExport, [](const TTuple<UFunction*, FString>& Pair){ return Pair.Get<0>()->GetFName(); });
+		Algo::SortBy(NamedFunctionsToExport, [](const TTuple<UFunction*, FString>& Pair){ return Pair.Get<0>()->GetFName(); }, FNameLexicalLess());
 
 		if (NamedFunctionsToExport.Num() > 0)
 		{
diff --git a/Engine/Source/Runtime/AssetRegistry/Private/AssetRegistryState.cpp b/Engine/Source/Runtime/AssetRegistry/Private/AssetRegistryState.cpp
--- a/Engine/Source/Runtime/AssetRegistry/Private/AssetRegistryState.cpp
+++ b/Engine/Source/Runtime/AssetRegistry/Private/AssetRegistryState.cpp
@@ -1690,7 +1690,7 @@ void FAssetRegistryState::Dump(const TArray<FString>& Arguments, TArray<FString>
 
 		TArray<FName> Keys;
 		AssetMap.GenerateKeyArray(Keys);
-		Keys.Sort();
+		Keys.Sort(FNameLexicalLess());
 
 		TArray<FAssetData*> Items;
 		Items.Reserve(1024);
@@ -1767,7 +1767,7 @@ void FAssetRegistryState::Dump(const TArray<FString>& Arguments, TArray<FString>
 
 		TArray<FName> Keys;
 		CachedPackageData.GenerateKeyArray(Keys);
-		Keys.Sort();
+		Keys.Sort(FNameLexicalLess());
 
 		for (const FName& Key : Keys)
 		{
diff --git a/Engine/Source/Runtime/AssetRegistry/Public/AssetData.h b/Engine/Source/Runtime/AssetRegistry/Public/AssetData.h
--- a/Engine/Source/Runtime/AssetRegistry/Public/AssetData.h
+++ b/Engine/Source/Runtime/AssetRegistry/Public/AssetData.h
@@ -153,12 +153,12 @@ public:
 
 	bool operator>(const FAssetData& Other) const
 	{
-		return ObjectPath > Other.ObjectPath;
+		return  Other.ObjectPath.LexicalLess(ObjectPath);
 	}
 
 	bool operator<(const FAssetData& Other) const
 	{
-		return ObjectPath < Other.ObjectPath;
+		return ObjectPath.LexicalLess(Other.ObjectPath);
 	}
 
 	/** Checks to see if this AssetData refers to an asset or is NULL */
diff --git a/Engine/Source/Runtime/AssetRegistry/Public/AssetDataTagMap.h b/Engine/Source/Runtime/AssetRegistry/Public/AssetDataTagMap.h
--- a/Engine/Source/Runtime/AssetRegistry/Public/AssetDataTagMap.h
+++ b/Engine/Source/Runtime/AssetRegistry/Public/AssetDataTagMap.h
@@ -12,7 +12,7 @@
 #if !USE_COMPACT_ASSET_REGISTRY
 
 /** Type of tag map */
-typedef TSortedMap<FName, FString, FDefaultAllocator, FNameSortIndexes> FAssetDataTagMap;
+typedef TSortedMap<FName, FString, FDefaultAllocator, FNameFastLess> FAssetDataTagMap;
 
 #else
 
@@ -402,7 +402,7 @@ inline FAssetDataTagMapValueStorage::FStorageID::operator FString() const
 	return FAssetDataTagMapValueStorage::Get().IdToString(*this);
 }
 
-typedef TSortedMap<FName, FAssetDataTagMapValueStorage::FStorageID, FDefaultAllocator, FNameSortIndexes> FAssetDataTagMapBase;
+typedef TSortedMap<FName, FAssetDataTagMapValueStorage::FStorageID, FDefaultAllocator, FNameFastLess> FAssetDataTagMapBase;
 
 /** Wrapper of the underlying map that handles making sure that when the map dies, the underlying storage for the strings is freed. */
 class FAssetDataTagMap : public FAssetDataTagMapBase
@@ -476,7 +476,7 @@ public:
 	/** This is serialization compatible with the non-compact version...so we just load the non-compact version and iterate it to compress the data structure. */
 	friend FArchive& operator<<(FArchive& Ar, FAssetDataTagMap& This)
 	{
-		TSortedMap<FName, FString, FDefaultAllocator, FNameSortIndexes> EmulatedContainer;
+		TSortedMap<FName, FString, FDefaultAllocator, FNameFastLess> EmulatedContainer;
 		if (Ar.IsLoading())
 		{
 			Ar << EmulatedContainer;
diff --git a/Engine/Source/Runtime/Core/Private/Containers/ContainersTest.cpp b/Engine/Source/Runtime/Core/Private/Containers/ContainersTest.cpp
--- a/Engine/Source/Runtime/Core/Private/Containers/ContainersTest.cpp
+++ b/Engine/Source/Runtime/Core/Private/Containers/ContainersTest.cpp
@@ -503,13 +503,13 @@ bool FContainersFullTest::RunTest(const FString& Parameters)
 	RunContainerTests<TMap<FString, FContainerTestValueType, FDefaultSetAllocator, FCaseSensitiveLookupKeyFuncs<FContainerTestValueType>>, FString>();
 
 	RunContainerTests<TSortedMap<int32, FContainerTestValueType>, int32>();
-	RunContainerTests<TSortedMap<FName, FContainerTestValueType>, FName>();
+	RunContainerTests<TSortedMap<FName, FContainerTestValueType, FDefaultAllocator, FNameLexicalLess>, FName>();
 	RunContainerTests<TSortedMap<FString, FContainerTestValueType>, FString>();
 	RunContainerTests<TSortedMap<FString, FContainerTestValueType, TInlineAllocator<64>>, FString>();
 
 	// Verify use of FName index sorter with SortedMap
 
-	TSortedMap<FName, int32, FDefaultAllocator, FNameSortIndexes> NameMap;
+	TSortedMap<FName, int32, FDefaultAllocator, FNameFastLess> NameMap;
 	NameMap.Add(NAME_NameProperty);
 	NameMap.Add(NAME_FloatProperty);
 	NameMap.Add(NAME_None);
@@ -555,11 +555,11 @@ bool FContainerPerformanceTest::RunTest(const FString& Parameters)
 	RunPerformanceTest<TSortedMap<int32, FString>, int32>(TEXT("TSortedMap int32"), 1000, 1000000);
 	RunPerformanceTest<TSortedMap<int32, FString>, int32>(TEXT("TSortedMap int32"), 10000, 1000000);
 
-	RunPerformanceTest<TSortedMap<FName, FString>, FName>(TEXT("TSortedMap FName"), 1, 1000000);
-	RunPerformanceTest<TSortedMap<FName, FString>, FName>(TEXT("TSortedMap FName"), 10, 1000000);
-	RunPerformanceTest<TSortedMap<FName, FString>, FName>(TEXT("TSortedMap FName"), 100, 1000000);
-	RunPerformanceTest<TSortedMap<FName, FString>, FName>(TEXT("TSortedMap FName"), 1000, 1000000);
-	RunPerformanceTest<TSortedMap<FName, FString>, FName>(TEXT("TSortedMap FName"), 10000, 1000000);
+	RunPerformanceTest<TSortedMap<FName, FString, FDefaultAllocator, FNameLexicalLess>, FName>(TEXT("TSortedMap FName"), 1, 1000000);
+	RunPerformanceTest<TSortedMap<FName, FString, FDefaultAllocator, FNameLexicalLess>, FName>(TEXT("TSortedMap FName"), 10, 1000000);
+	RunPerformanceTest<TSortedMap<FName, FString, FDefaultAllocator, FNameLexicalLess>, FName>(TEXT("TSortedMap FName"), 100, 1000000);
+	RunPerformanceTest<TSortedMap<FName, FString, FDefaultAllocator, FNameLexicalLess>, FName>(TEXT("TSortedMap FName"), 1000, 1000000);
+	RunPerformanceTest<TSortedMap<FName, FString, FDefaultAllocator, FNameLexicalLess>, FName>(TEXT("TSortedMap FName"), 10000, 1000000);
 
 	RunPerformanceTest<TSortedMap<FString, FString>, FString>(TEXT("TSortedMap FString"), 1, 1000000);
 	RunPerformanceTest<TSortedMap<FString, FString>, FString>(TEXT("TSortedMap FString"), 10, 1000000);
diff --git a/Engine/Source/Runtime/Core/Private/Stats/StatsCommand.cpp b/Engine/Source/Runtime/Core/Private/Stats/StatsCommand.cpp
--- a/Engine/Source/Runtime/Core/Private/Stats/StatsCommand.cpp
+++ b/Engine/Source/Runtime/Core/Private/Stats/StatsCommand.cpp
@@ -1406,7 +1406,7 @@ struct FHUDGroupManager
 			AggregatedHierarchyHistory.CullByDepth( Params.MaxHierarchyDepth.Get() );
 
 			// Make sure the game thread is first.
-			AggregatedHierarchyHistory.Children.KeySort( TLess<FName>() );
+			AggregatedHierarchyHistory.Children.KeySort(FNameLexicalLess());
 
 			FComplexStatUtils::DiviveStatArray( AggregatedFlatHistory, NumFrames, EComplexStatField::IncSum, EComplexStatField::IncAve );
 			FComplexStatUtils::DiviveStatArray( AggregatedFlatHistory, NumFrames, EComplexStatField::ExcSum, EComplexStatField::ExcAve );
diff --git a/Engine/Source/Runtime/Core/Private/UObject/UnrealNames.cpp b/Engine/Source/Runtime/Core/Private/UObject/UnrealNames.cpp
--- a/Engine/Source/Runtime/Core/Private/UObject/UnrealNames.cpp
+++ b/Engine/Source/Runtime/Core/Private/UObject/UnrealNames.cpp
@@ -903,6 +903,45 @@ bool operator==(FNameEntryId Id, EName Ename)
 	return Id == GetNamePool().Find(Ename);
 }
 
+static int32 CompareDifferentIdsAlphabetically(FNameEntryId AId, FNameEntryId BId)
+{
+	checkSlow(AId != BId);
+
+	FNamePool& Pool = GetNamePool();
+	FNameBuffer ABuffer, BBuffer;
+	FNameStringView AView =	Pool.Resolve(AId).MakeView(ABuffer);
+	FNameStringView BView =	Pool.Resolve(BId).MakeView(BBuffer);
+
+	// If only one view is wide, convert the ansi view to wide as well
+	if (AView.bIsWide != BView.bIsWide)
+	{
+		FNameStringView& AnsiView = AView.bIsWide ? BView : AView;
+		FNameBuffer& AnsiBuffer =	AView.bIsWide ? BBuffer : ABuffer;
+
+#ifndef WITH_CUSTOM_NAME_ENCODING
+		memcpy(AnsiBuffer.AnsiName, AnsiView.Ansi, AnsiView.Len * sizeof(ANSICHAR));
+		AnsiView.Ansi = AnsiBuffer.AnsiName;
+#endif
+
+		AnsiBuffer.ConvertInPlace<ANSICHAR, WIDECHAR>(AnsiView.Len);
+		AnsiView.bIsWide = true;
+	}
+
+	int32 MinLen = FMath::Min(AView.Len, BView.Len);
+	if (int32 StrDiff = AView.bIsWide ?	FCStringWide::Strnicmp(AView.Wide, BView.Wide, MinLen) :
+										FCStringAnsi::Strnicmp(AView.Ansi, BView.Ansi, MinLen))
+	{
+		return StrDiff;
+	}
+
+	return AView.Len - BView.Len;
+}
+
+int32 FNameEntryId::CompareLexical(FNameEntryId Rhs) const
+{
+	return Value != Rhs.Value && CompareDifferentIdsAlphabetically(*this, Rhs);
+}
+
 #if !UE_BUILD_SHIPPING && !UE_BUILD_TEST
 	void CallNameCreationHook();
 #else
@@ -1643,36 +1682,9 @@ int32 FName::Compare( const FName& Other ) const
 	{
 		return GetNumber() - Other.GetNumber();
 	}
-	else // Names don't match. This means we don't even need to check numbers.
-	{
-		FNameBuffer ThisBuffer, OtherBuffer;
-		FNameStringView ThisView =	GetComparisonNameEntry()->MakeView(ThisBuffer);
-		FNameStringView OtherView =	Other.GetComparisonNameEntry()->MakeView(OtherBuffer);
 
-		// If only one view is wide, convert the ansi view to wide as well
-		if (ThisView.bIsWide != OtherView.bIsWide)
-		{
-			FNameStringView& AnsiView = ThisView.bIsWide ? OtherView : ThisView;
-			FNameBuffer& AnsiBuffer =	ThisView.bIsWide ? OtherBuffer : ThisBuffer;
-
-#ifndef WITH_CUSTOM_NAME_ENCODING
-			memcpy(AnsiBuffer.AnsiName, AnsiView.Ansi, AnsiView.Len * sizeof(ANSICHAR));
-			AnsiView.Ansi = AnsiBuffer.AnsiName;
-#endif
-
-			AnsiBuffer.ConvertInPlace<ANSICHAR, WIDECHAR>(AnsiView.Len);
-			AnsiView.bIsWide = true;
-		}
-
-		int32 MinLen = FMath::Min(ThisView.Len, OtherView.Len);
-		if (int32 StrDiff = ThisView.bIsWide ?	FCStringWide::Strnicmp(ThisView.Wide, OtherView.Wide, MinLen) :
-												FCStringAnsi::Strnicmp(ThisView.Ansi, OtherView.Ansi, MinLen))
-		{
-			return StrDiff;
-		}
-
-		return ThisView.Len - OtherView.Len;
-	}
+	// Names don't match. This means we don't even need to check numbers.
+	return CompareDifferentIdsAlphabetically(ComparisonIndex, Other.ComparisonIndex);
 }
 
 FString FName::GetPlainNameString() const
@@ -1941,7 +1953,7 @@ void FName::AutoTest()
 	Names.Add("FooC");
 	const WIDECHAR FooWide[] = {'F', 'o', 'o', (WIDECHAR)2000, '\0'};
 	Names.Add(FooWide);
-	Algo::Sort(Names); // Todo: fix when removing default operator<
+	Algo::Sort(Names, FNameLexicalLess()); 
 
 	check(Names[0] == "FooA");
 	check(Names[1] == "FooAB");
diff --git a/Engine/Source/Runtime/Core/Public/UObject/NameTypes.h b/Engine/Source/Runtime/Core/Public/UObject/NameTypes.h
--- a/Engine/Source/Runtime/Core/Public/UObject/NameTypes.h
+++ b/Engine/Source/Runtime/Core/Public/UObject/NameTypes.h
@@ -38,11 +38,19 @@ struct FNameEntryId
 {
 	FNameEntryId() : Value(0) {}
 	FNameEntryId(ENoInit) {}
-	
-	int32 CompareCreationOrder(FNameEntryId Rhs) const { return Value - Rhs.Value; };
 
-	/** Creation order*/
+	/** Slow alphabetical order that is stable / deterministic over process runs */
+	CORE_API int32 CompareLexical(FNameEntryId Rhs) const;
+	bool LexicalLess(FNameEntryId Rhs) const { return CompareLexical(Rhs) < 0; }
+
+	/** Fast non-alphabetical order that is only stable during this process' lifetime */
+	int32 CompareFast(FNameEntryId Rhs) const { return Value - Rhs.Value; };
+	bool FastLess(FNameEntryId Rhs) const { return CompareFast(Rhs) < 0; }
+
+	/** Fast non-alphabetical order that is only stable during this process' lifetime */
 	bool operator<(FNameEntryId Rhs) const { return Value < Rhs.Value; }
+
+	/** Fast non-alphabetical order that is only stable during this process' lifetime */
 	bool operator>(FNameEntryId Rhs) const { return Rhs.Value < Value; }
 	bool operator==(FNameEntryId Rhs) const { return Value == Rhs.Value; }
 	bool operator!=(FNameEntryId Rhs) const { return Value != Rhs.Value; }
@@ -454,20 +462,29 @@ public:
 		return !(*this == Other);
 	}
 
-	/**
-	 * Comparison operator used for sorting alphabetically.
-	 */
+	UE_DEPRECATED(4.23, "Please use FastLess() / FNameFastLess or LexicalLess() / FNameLexicalLess instead. "
+		"Default lexical sort order is deprecated to avoid unintended expensive sorting. ")
 	FORCEINLINE bool operator<( const FName& Other ) const
 	{
-		return Compare(Other) < 0;
+		return LexicalLess(Other);
 	}
 
-	/**
-	 * Comparison operator used for sorting alphabetically.
-	 */
+	UE_DEPRECATED(4.23, "Please use B.FastLess(A) or B.LexicalLess(A) instead of A > B.")
 	FORCEINLINE bool operator>(const FName& Other) const
 	{
-		return Compare(Other) > 0;
+		return Other.LexicalLess(*this);
+	}
+
+	/** Fast non-alphabetical order that is only stable during this process' lifetime. */
+	FORCEINLINE bool FastLess(const FName& Other) const
+	{
+		return CompareIndexes(Other) < 0;
+	}
+
+	/** Slow alphabetical order that is stable / deterministic over process runs. */
+	FORCEINLINE bool LexicalLess(const FName& Other) const
+	{
+		return Compare(Other) < 0;
 	}
 
 	FORCEINLINE bool IsNone() const
@@ -567,7 +584,7 @@ public:
 	 */
 	FORCEINLINE int32 CompareIndexes(const FName& Other) const
 	{
-		if (int32 ComparisonDiff = ComparisonIndex.CompareCreationOrder(Other.ComparisonIndex))
+		if (int32 ComparisonDiff = ComparisonIndex.CompareFast(Other.ComparisonIndex))
 		{
 			return ComparisonDiff;
 		}
@@ -858,14 +875,35 @@ inline bool operator!=(const CharType *LHS, const FName &RHS)
 /** FNames act like PODs. */
 template <> struct TIsPODType<FName> { enum { Value = true }; };
 
-
-/** Sort predicate to sort FName by index instead of alphabetically, pass to anything that wants TLess */
-struct FNameSortIndexes
+/** Fast non-alphabetical order that is only stable during this process' lifetime */
+struct FNameFastLess
 {
 	FORCEINLINE bool operator()(const FName& A, const FName& B) const
 	{
 		return A.CompareIndexes(B) < 0;
 	}
+
+	FORCEINLINE bool operator()(FNameEntryId A, FNameEntryId B) const
+	{
+		return A.FastLess(B);
+	}
+};
+
+UE_DEPRECATED(4.23, "Please use FNameFastLess instead.")
+typedef FNameFastLess FNameSortIndexes;
+
+/** Slow alphabetical order that is stable / deterministic over process runs */
+struct FNameLexicalLess
+{
+	FORCEINLINE bool operator()(const FName& A, const FName& B) const
+	{
+		return A.Compare(B) < 0;
+	}
+
+	FORCEINLINE bool operator()(FNameEntryId A, FNameEntryId B) const
+	{
+		return A.LexicalLess(B);
+	}
 };
 
 #ifndef WITH_CUSTOM_NAME_ENCODING
diff --git a/Engine/Source/Runtime/CoreUObject/Private/Misc/ExclusiveLoadPackageTimeTracker.cpp b/Engine/Source/Runtime/CoreUObject/Private/Misc/ExclusiveLoadPackageTimeTracker.cpp
--- a/Engine/Source/Runtime/CoreUObject/Private/Misc/ExclusiveLoadPackageTimeTracker.cpp
+++ b/Engine/Source/Runtime/CoreUObject/Private/Misc/ExclusiveLoadPackageTimeTracker.cpp
@@ -219,7 +219,7 @@ void FExclusiveLoadPackageTimeTracker::InternalDumpReport(const TArray<FString>&
 		ReportAr->Logf(TEXT("Dumping all loaded assets by exclusive load time:"));
 		if (bAlphaSort)
 		{
-			SortedLoadTimes.Sort([](const FLoadTime& A, const FLoadTime& B) { return A.TimeName < B.TimeName; });
+			SortedLoadTimes.Sort([](const FLoadTime& A, const FLoadTime& B) { return A.TimeName.LexicalLess(B.TimeName); });
 		}
 		else
 		{
@@ -255,7 +255,7 @@ void FExclusiveLoadPackageTimeTracker::InternalDumpReport(const TArray<FString>&
 		ReportAr->Logf(TEXT("Dumping all loaded assets by inclusive load time:"));
 		if (bAlphaSort)
 		{
-			SortedLoadTimes.Sort([](const FLoadTime& A, const FLoadTime& B) { return A.TimeName < B.TimeName; });
+			SortedLoadTimes.Sort([](const FLoadTime& A, const FLoadTime& B) { return A.TimeName.LexicalLess(B.TimeName); });
 		}
 		else
 		{
diff --git a/Engine/Source/Runtime/CoreUObject/Private/UObject/Linker.cpp b/Engine/Source/Runtime/CoreUObject/Private/UObject/Linker.cpp
--- a/Engine/Source/Runtime/CoreUObject/Private/UObject/Linker.cpp
+++ b/Engine/Source/Runtime/CoreUObject/Private/UObject/Linker.cpp
@@ -123,7 +123,7 @@ void FLinkerTables::SerializeSearchableNamesMap(FStructuredArchive::FSlot Slot)
 
 		for (TPair<FPackageIndex, TArray<FName> >& Pair : SearchableNamesMap)
 		{
-			Pair.Value.Sort();
+			Pair.Value.Sort(FNameLexicalLess());
 		}
 	}
 
diff --git a/Engine/Source/Runtime/CoreUObject/Private/UObject/LinkerLoad.cpp b/Engine/Source/Runtime/CoreUObject/Private/UObject/LinkerLoad.cpp
--- a/Engine/Source/Runtime/CoreUObject/Private/UObject/LinkerLoad.cpp
+++ b/Engine/Source/Runtime/CoreUObject/Private/UObject/LinkerLoad.cpp
@@ -5142,14 +5142,9 @@ bool AreObjectExportsEqualForDuplicateChecks(const FObjectExport& Lhs, const FOb
 bool ExportMapSorter(const FObjectExport& Lhs, const FObjectExport& Rhs)
 {
 	// Check names first.
-	if (Lhs.ObjectName < Rhs.ObjectName)
+	if (Lhs.ObjectName != Rhs.ObjectName)
 	{
-		return true;
-	}
-
-	if (Lhs.ObjectName > Rhs.ObjectName)
-	{
-		return false;
+		return Lhs.ObjectName.LexicalLess(Rhs.ObjectName);
 	}
 
 	// Names are equal, check classes.
diff --git a/Engine/Source/Runtime/CoreUObject/Private/UObject/SavePackage.cpp b/Engine/Source/Runtime/CoreUObject/Private/UObject/SavePackage.cpp
--- a/Engine/Source/Runtime/CoreUObject/Private/UObject/SavePackage.cpp
+++ b/Engine/Source/Runtime/CoreUObject/Private/UObject/SavePackage.cpp
@@ -3148,31 +3148,21 @@ EObjectMark UPackage::GetExcludedObjectMarksForTargetPlatform( const class ITarg
 bool ExportObjectSorter(const UObject& Lhs, const UObject& Rhs)
 {
 	// Check names first.
-	if (Lhs.GetFName() < Rhs.GetFName())
+	if (Lhs.GetFName() != Rhs.GetFName())
 	{
-		return true;
-	}
-
-	if (Lhs.GetFName() > Rhs.GetFName())
-	{
-		return false;
+		return Lhs.GetFName().LexicalLess(Rhs.GetFName());
 	}
 
 	// Names equal, compare class names.
-	if (Lhs.GetClass()->GetFName() < Rhs.GetClass()->GetFName())
+	if (Lhs.GetClass()->GetFName() != Rhs.GetClass()->GetFName())
 	{
-		return true;
-	}
-
-	if (Lhs.GetClass()->GetFName() > Rhs.GetClass()->GetFName())
-	{
-		return false;
+		return Lhs.GetClass()->GetFName().LexicalLess(Rhs.GetClass()->GetFName());
 	}
 
 	// Compare by outers if they exist.
 	if (Lhs.GetOuter() && Rhs.GetOuter())
 	{
-		return Lhs.GetOuter()->GetFName() < Rhs.GetOuter()->GetFName();
+		return Lhs.GetOuter()->GetFName().LexicalLess(Rhs.GetOuter()->GetFName());
 	}
 
 	if (Lhs.GetOuter())
diff --git a/Engine/Source/Runtime/CoreUObject/Private/UObject/UObjectClusters.cpp b/Engine/Source/Runtime/CoreUObject/Private/UObject/UObjectClusters.cpp
--- a/Engine/Source/Runtime/CoreUObject/Private/UObject/UObjectClusters.cpp
+++ b/Engine/Source/Runtime/CoreUObject/Private/UObject/UObjectClusters.cpp
@@ -383,7 +383,7 @@ void ListClusters(const TArray<FString>& Args)
 			Algo::SortBy(AllClusters, [](FUObjectCluster* A)
 			{
 				return GUObjectArray.IndexToObject(A->RootIndex)->Object->GetFName();
-			});
+			}, FNameLexicalLess());
 		}
 		else if (Arg == TEXT("SortByObjectCount"))
 		{
diff --git a/Engine/Source/Runtime/CoreUObject/Public/Serialization/BulkData.h b/Engine/Source/Runtime/CoreUObject/Public/Serialization/BulkData.h
--- a/Engine/Source/Runtime/CoreUObject/Public/Serialization/BulkData.h
+++ b/Engine/Source/Runtime/CoreUObject/Public/Serialization/BulkData.h
@@ -744,7 +744,7 @@ class FFormatContainer
 {
 	friend class UBodySetup;
 
-	TSortedMap<FName, FByteBulkData*, FDefaultAllocator, FNameSortIndexes> Formats;
+	TSortedMap<FName, FByteBulkData*, FDefaultAllocator, FNameFastLess> Formats;
 	uint32 Alignment;
 public:
 	~FFormatContainer()
diff --git a/Engine/Source/Runtime/Engine/Classes/Animation/Skeleton.h b/Engine/Source/Runtime/Engine/Classes/Animation/Skeleton.h
--- a/Engine/Source/Runtime/Engine/Classes/Animation/Skeleton.h
+++ b/Engine/Source/Runtime/Engine/Classes/Animation/Skeleton.h
@@ -373,7 +373,7 @@ public:
 	const TArray<FName>& GetExistingMarkerNames() const { return ExistingMarkerNames; }
 
 	// Register a new sync marker name
-	void RegisterMarkerName(FName MarkerName) { ExistingMarkerNames.AddUnique(MarkerName); ExistingMarkerNames.Sort(); }
+	void RegisterMarkerName(FName MarkerName) { ExistingMarkerNames.AddUnique(MarkerName); ExistingMarkerNames.Sort(FNameLexicalLess()); }
 
 	// Remove a sync marker name
 	void RemoveMarkerName(FName MarkerName) { ExistingMarkerNames.Remove(MarkerName); }
diff --git a/Engine/Source/Runtime/Engine/Classes/Components/MeshComponent.h b/Engine/Source/Runtime/Engine/Classes/Components/MeshComponent.h
--- a/Engine/Source/Runtime/Engine/Classes/Components/MeshComponent.h
+++ b/Engine/Source/Runtime/Engine/Classes/Components/MeshComponent.h
@@ -141,7 +141,7 @@ protected:
 		float ScalarParameterDefaultValue = 0.f;
 	};
 
-	TSortedMap<FName, FMaterialParameterCache, FDefaultAllocator, FNameSortIndexes> MaterialParameterCache;
+	TSortedMap<FName, FMaterialParameterCache, FDefaultAllocator, FNameFastLess> MaterialParameterCache;
 
 	UPROPERTY(EditAnywhere, AdvancedDisplay, BlueprintReadOnly, Category = MaterialParameters)
 	bool bEnableMaterialParameterCaching;
diff --git a/Engine/Source/Runtime/Engine/Classes/Components/SkinnedMeshComponent.h b/Engine/Source/Runtime/Engine/Classes/Components/SkinnedMeshComponent.h
--- a/Engine/Source/Runtime/Engine/Classes/Components/SkinnedMeshComponent.h
+++ b/Engine/Source/Runtime/Engine/Classes/Components/SkinnedMeshComponent.h
@@ -305,7 +305,7 @@ protected:
 	/**
 	*	Mapping for socket overrides, key is the Source socket name and the value is the override socket name
 	*/
-	TSortedMap<FName, FName, FDefaultAllocator, FNameSortIndexes> SocketOverrideLookup;
+	TSortedMap<FName, FName, FDefaultAllocator, FNameFastLess> SocketOverrideLookup;
 
 public:
 #if WITH_EDITORONLY_DATA
diff --git a/Engine/Source/Runtime/Engine/Classes/GameFramework/PlayerInput.h b/Engine/Source/Runtime/Engine/Classes/GameFramework/PlayerInput.h
--- a/Engine/Source/Runtime/Engine/Classes/GameFramework/PlayerInput.h
+++ b/Engine/Source/Runtime/Engine/Classes/GameFramework/PlayerInput.h
@@ -181,7 +181,7 @@ struct FInputActionKeyMapping
 	bool operator<(const FInputActionKeyMapping& Other) const
 	{
 		bool bResult = false;
-		if (ActionName < Other.ActionName)
+		if (ActionName.LexicalLess(Other.ActionName))
 		{
 			bResult = true;
 		}
@@ -234,7 +234,7 @@ struct FInputAxisKeyMapping
 	bool operator<(const FInputAxisKeyMapping& Other) const
 	{
 		bool bResult = false;
-		if (AxisName < Other.AxisName)
+		if (AxisName.LexicalLess(Other.AxisName))
 		{
 			bResult = true;
 		}
diff --git a/Engine/Source/Runtime/Engine/Private/Animation/AnimInstance.cpp b/Engine/Source/Runtime/Engine/Private/Animation/AnimInstance.cpp
--- a/Engine/Source/Runtime/Engine/Private/Animation/AnimInstance.cpp
+++ b/Engine/Source/Runtime/Engine/Private/Animation/AnimInstance.cpp
@@ -708,7 +708,7 @@ void OutputCurveMap(TMap<FName, float>& CurveMap, UCanvas* Canvas, FDisplayDebug
 {
 	TArray<FName> Names;
 	CurveMap.GetKeys(Names);
-	Names.Sort();
+	Names.Sort(FNameLexicalLess());
 	for (FName CurveName : Names)
 	{
 		FString CurveEntry = FString::Printf(TEXT("%s: %.3f"), *CurveName.ToString(), CurveMap[CurveName]);
diff --git a/Engine/Source/Runtime/Engine/Private/AnimationAsset.cpp b/Engine/Source/Runtime/Engine/Private/AnimationAsset.cpp
--- a/Engine/Source/Runtime/Engine/Private/AnimationAsset.cpp
+++ b/Engine/Source/Runtime/Engine/Private/AnimationAsset.cpp
@@ -153,7 +153,7 @@ void FAnimGroupInstance::Prepare(const FAnimGroupInstance* PreviousGroup)
 
 		bCanUseMarkerSync = ValidMarkers.Num() > 0;
 
-		ValidMarkers.Sort();
+		ValidMarkers.Sort(FNameLexicalLess());
 
 		if (!PreviousGroup || (ValidMarkers != PreviousGroup->ValidMarkers))
 		{
diff --git a/Engine/Source/Runtime/Engine/Private/AssetManager.cpp b/Engine/Source/Runtime/Engine/Private/AssetManager.cpp
--- a/Engine/Source/Runtime/Engine/Private/AssetManager.cpp
+++ b/Engine/Source/Runtime/Engine/Private/AssetManager.cpp
@@ -1138,7 +1138,7 @@ TSharedPtr<FStreamableHandle> UAssetManager::ChangeBundleStateForPrimaryAssets(c
 				NewBundleState.AddUnique(AddBundle);
 			}
 
-			NewBundleState.Sort();
+			NewBundleState.Sort(FNameLexicalLess());
 
 			// If the pending state is valid, check if it is different
 			if (NameData->PendingState.IsValid())
@@ -2330,7 +2330,7 @@ void UAssetManager::DumpAssetTypeSummary()
 
 	Manager.GetPrimaryAssetTypeInfoList(TypeInfos);
 
-	TypeInfos.Sort([](const FPrimaryAssetTypeInfo& LHS, const FPrimaryAssetTypeInfo& RHS) { return LHS.PrimaryAssetType < RHS.PrimaryAssetType; });
+	TypeInfos.Sort([](const FPrimaryAssetTypeInfo& LHS, const FPrimaryAssetTypeInfo& RHS) { return LHS.PrimaryAssetType.LexicalLess(RHS.PrimaryAssetType); });
 
 	UE_LOG(LogAssetManager, Log, TEXT("=========== Asset Manager Type Summary ==========="));
 
@@ -2358,7 +2358,7 @@ void UAssetManager::DumpLoadedAssetState()
 
 	Manager.GetPrimaryAssetTypeInfoList(TypeInfos);
 
-	TypeInfos.Sort([](const FPrimaryAssetTypeInfo& LHS, const FPrimaryAssetTypeInfo& RHS) { return LHS.PrimaryAssetType < RHS.PrimaryAssetType; });
+	TypeInfos.Sort([](const FPrimaryAssetTypeInfo& LHS, const FPrimaryAssetTypeInfo& RHS) { return LHS.PrimaryAssetType.LexicalLess(RHS.PrimaryAssetType); });
 
 	UE_LOG(LogAssetManager, Log, TEXT("=========== Asset Manager Loaded Asset State ==========="));
 
@@ -2404,7 +2404,7 @@ void UAssetManager::DumpLoadedAssetState()
 		{
 			UE_LOG(LogAssetManager, Log, TEXT("  Type %s:"), *TypeInfo.PrimaryAssetType.ToString());
 
-			LoadedInfos.Sort([](const FLoadedInfo& LHS, const FLoadedInfo& RHS) { return LHS.AssetName < RHS.AssetName; });
+			LoadedInfos.Sort([](const FLoadedInfo& LHS, const FLoadedInfo& RHS) { return LHS.AssetName.LexicalLess(RHS.AssetName); });
 
 			for (FLoadedInfo& LoadedInfo : LoadedInfos)
 			{
diff --git a/Engine/Source/Runtime/Engine/Private/Collision/Collision.cpp b/Engine/Source/Runtime/Engine/Private/Collision/Collision.cpp
--- a/Engine/Source/Runtime/Engine/Private/Collision/Collision.cpp
+++ b/Engine/Source/Runtime/Engine/Private/Collision/Collision.cpp
@@ -1028,7 +1028,7 @@ namespace CollisionResponseConsoleCommands
 		// Display Data
 		if (Results.Num() > 0)
 		{
-			Results.Sort([](const FName& A, const FName& B) { return A < B; });
+			Results.Sort(FNameLexicalLess());
 			for (FName& ResultName : Results)
 			{
 				UE_LOG(LogCollisionCommands, Log, TEXT("%s"), *ResultName.ToString());
diff --git a/Engine/Source/Runtime/Engine/Private/LevelTick.cpp b/Engine/Source/Runtime/Engine/Private/LevelTick.cpp
--- a/Engine/Source/Runtime/Engine/Private/LevelTick.cpp
+++ b/Engine/Source/Runtime/Engine/Private/LevelTick.cpp
@@ -1247,7 +1247,7 @@ static void RecordWorldCountsToCSV(UWorld* World)
 
 			bool bDetailed = (CVarDetailedTickContextForCSV.GetValueOnGameThread() != 0);
 
-			TSortedMap<FName, int32, FDefaultAllocator, FNameSortIndexes> TickContextToCountMap;
+			TSortedMap<FName, int32, FDefaultAllocator, FNameFastLess> TickContextToCountMap;
 			int32 EnabledCount;
 			FTickTaskManagerInterface::Get().GetEnabledTickFunctionCounts(World, TickContextToCountMap, EnabledCount, bDetailed);
 
diff --git a/Engine/Source/Runtime/Engine/Private/Materials/MaterialInterface.cpp b/Engine/Source/Runtime/Engine/Private/Materials/MaterialInterface.cpp
--- a/Engine/Source/Runtime/Engine/Private/Materials/MaterialInterface.cpp
+++ b/Engine/Source/Runtime/Engine/Private/Materials/MaterialInterface.cpp
@@ -543,7 +543,7 @@ void UMaterialInterface::SortTextureStreamingData(bool bForceSort, bool bFinalSo
 		}
 
 		// Sort by name to be compatible with FindTextureStreamingDataIndexRange
-		TextureStreamingData.Sort([](const FMaterialTextureInfo& Lhs, const FMaterialTextureInfo& Rhs) { return Lhs.TextureName < Rhs.TextureName; });
+		TextureStreamingData.Sort([](const FMaterialTextureInfo& Lhs, const FMaterialTextureInfo& Rhs) { return Lhs.TextureName.LexicalLess(Rhs.TextureName); });
 		bTextureStreamingDataSorted = true;
 	}
 #endif
@@ -565,7 +565,7 @@ bool UMaterialInterface::FindTextureStreamingDataIndexRange(FName TextureName, i
 		return false;
 	}
 
-	const int32 MatchingIndex = Algo::BinarySearchBy(TextureStreamingData, TextureName, &FMaterialTextureInfo::TextureName);
+	const int32 MatchingIndex = Algo::BinarySearchBy(TextureStreamingData, TextureName, &FMaterialTextureInfo::TextureName, FNameLexicalLess());
 	if (MatchingIndex != INDEX_NONE)
 	{
 		// Find the range of entries for this texture. 
diff --git a/Engine/Source/Runtime/Engine/Private/StaticMesh.cpp b/Engine/Source/Runtime/Engine/Private/StaticMesh.cpp
--- a/Engine/Source/Runtime/Engine/Private/StaticMesh.cpp
+++ b/Engine/Source/Runtime/Engine/Private/StaticMesh.cpp
@@ -1616,7 +1616,7 @@ void FStaticMeshLODSettings::Initialize(const FConfigFile& IniFile)
 		};
 	}
 
-	Groups.KeySort(TLess<FName>());
+	Groups.KeySort(FNameLexicalLess());
 	GroupName2Index.Empty(Groups.Num());
 	{
 		int32 GroupIdx = 0;
diff --git a/Engine/Source/Runtime/Engine/Private/TickTaskManager.cpp b/Engine/Source/Runtime/Engine/Private/TickTaskManager.cpp
--- a/Engine/Source/Runtime/Engine/Private/TickTaskManager.cpp
+++ b/Engine/Source/Runtime/Engine/Private/TickTaskManager.cpp
@@ -1188,7 +1188,7 @@ public:
 		DisabledCount += AllDisabledTickFunctions.Num();
 	}
 
-	FORCEINLINE void AddTickFunctionToMap(TSortedMap<FName, int32, FDefaultAllocator, FNameSortIndexes>& ClassNameToCountMap, FTickFunction* Function, bool bDetailed)
+	FORCEINLINE void AddTickFunctionToMap(TSortedMap<FName, int32, FDefaultAllocator, FNameFastLess>& ClassNameToCountMap, FTickFunction* Function, bool bDetailed)
 	{
 		FName ContextName = Function->DiagnosticContext(bDetailed);
 		// Find entry for this context (or add it if not present)
@@ -1197,7 +1197,7 @@ public:
 		CurrentCount++; 
 	}
 
-	void AddTickFunctionsToMap(TSortedMap<FName, int32, FDefaultAllocator, FNameSortIndexes>& ClassNameToCountMap, int32& EnabledCount, bool bDetailed)
+	void AddTickFunctionsToMap(TSortedMap<FName, int32, FDefaultAllocator, FNameFastLess>& ClassNameToCountMap, int32& EnabledCount, bool bDetailed)
 	{
 		// Add ticks from AllEnabledTickFunctions
 		for (TSet<FTickFunction*>::TIterator It(AllEnabledTickFunctions); It; ++It)
@@ -1639,7 +1639,7 @@ private:
 
 		if (bGrouped)
 		{
-			TSortedMap<FName, int32, FDefaultAllocator, FNameSortIndexes> TickContextToCountMap;
+			TSortedMap<FName, int32, FDefaultAllocator, FNameFastLess> TickContextToCountMap;
 			GetEnabledTickFunctionCounts(InWorld, TickContextToCountMap, EnabledCount, true);
 
 			// Build sorted array of tick context by count
@@ -1691,7 +1691,7 @@ private:
 	}
 
 
-	virtual void GetEnabledTickFunctionCounts(UWorld* InWorld, TSortedMap<FName, int32, FDefaultAllocator, FNameSortIndexes>& TickContextToCountMap, int32& EnabledCount, bool bDetailed)
+	virtual void GetEnabledTickFunctionCounts(UWorld* InWorld, TSortedMap<FName, int32, FDefaultAllocator, FNameFastLess>& TickContextToCountMap, int32& EnabledCount, bool bDetailed)
 	{
 		check(InWorld);
 		check(InWorld->TickTaskLevel);
diff --git a/Engine/Source/Runtime/Engine/Private/UnrealEngine.cpp b/Engine/Source/Runtime/Engine/Private/UnrealEngine.cpp
--- a/Engine/Source/Runtime/Engine/Private/UnrealEngine.cpp
+++ b/Engine/Source/Runtime/Engine/Private/UnrealEngine.cpp
@@ -9297,7 +9297,7 @@ struct FSoundInfo
 
 	bool CompareClass( const FSoundInfo& Other ) const
 	{
-		return ClassName < Other.ClassName;
+		return ClassName.LexicalLess(Other.ClassName);
 	}
 
 	bool CompareWaveInstancesNum( const FSoundInfo& Other ) const
@@ -15982,7 +15982,7 @@ int32 UEngine::RenderStatSounds(UWorld* World, FViewport* Viewport, FCanvas* Can
 			}
 			else if (ShowSounds & FViewportClient::ESoundShowFlags::Sort_Class)
 			{
-				AudioStats.StatSoundInfos.Sort([](const FAudioStats::FStatSoundInfo& A, const FAudioStats::FStatSoundInfo& B) { return A.SoundClassName < B.SoundClassName; });
+				AudioStats.StatSoundInfos.Sort([](const FAudioStats::FStatSoundInfo& A, const FAudioStats::FStatSoundInfo& B) { return A.SoundClassName.LexicalLess(B.SoundClassName); });
 				SortingName = TEXT("class");
 			}
 			else if (ShowSounds & FViewportClient::ESoundShowFlags::Sort_WavesNum)
diff --git a/Engine/Source/Runtime/Engine/Public/TickTaskManagerInterface.h b/Engine/Source/Runtime/Engine/Public/TickTaskManagerInterface.h
--- a/Engine/Source/Runtime/Engine/Public/TickTaskManagerInterface.h
+++ b/Engine/Source/Runtime/Engine/Public/TickTaskManagerInterface.h
@@ -60,7 +60,7 @@ public:
 	virtual void DumpAllTickFunctions(FOutputDevice& Ar, UWorld* InWorld, bool bEnabled, bool bDisabled, bool bGrouped) = 0;
 
 	/** Returns a map of enabled ticks, grouped by 'diagnostic context' string, along with count of enabled ticks */
-	virtual void GetEnabledTickFunctionCounts(UWorld* InWorld, TSortedMap<FName, int32, FDefaultAllocator, FNameSortIndexes>& TickContextToCountMap, int32& EnabledCount, bool bDetailed) = 0;
+	virtual void GetEnabledTickFunctionCounts(UWorld* InWorld, TSortedMap<FName, int32, FDefaultAllocator, FNameFastLess>& TickContextToCountMap, int32& EnabledCount, bool bDetailed) = 0;
 
 	/**
 	 * Singleton to retrieve the GLOBAL tick task manager
diff --git a/Engine/Source/Runtime/GameplayTags/Classes/GameplayTagContainer.h b/Engine/Source/Runtime/GameplayTags/Classes/GameplayTagContainer.h
--- a/Engine/Source/Runtime/GameplayTags/Classes/GameplayTagContainer.h
+++ b/Engine/Source/Runtime/GameplayTags/Classes/GameplayTagContainer.h
@@ -86,7 +86,7 @@ struct GAMEPLAYTAGS_API FGameplayTag
 
 	FORCEINLINE bool operator<(FGameplayTag const& Other) const
 	{
-		return TagName < Other.TagName;
+		return TagName.LexicalLess(Other.TagName);
 	}
 
 	/**
diff --git a/Engine/Source/Runtime/GameplayTags/Private/GameplayTagsManager.cpp b/Engine/Source/Runtime/GameplayTags/Private/GameplayTagsManager.cpp
--- a/Engine/Source/Runtime/GameplayTags/Private/GameplayTagsManager.cpp
+++ b/Engine/Source/Runtime/GameplayTags/Private/GameplayTagsManager.cpp
@@ -970,7 +970,7 @@ int32 UGameplayTagsManager::InsertTagIntoNodeArray(FName Tag, FName FullTag, TSh
 #endif				
 				break;
 			}
-			else if (SimpleTagName > Tag && WhereToInsert == INDEX_NONE)
+			else if (SimpleTagName.LexicalLess(Tag) && WhereToInsert == INDEX_NONE)
 			{
 				// Insert new node before this
 				WhereToInsert = CurIdx;
@@ -1912,7 +1912,7 @@ bool FGameplayTagTableRow::operator!=(FGameplayTagTableRow const& Other) const
 
 bool FGameplayTagTableRow::operator<(FGameplayTagTableRow const& Other) const
 {
-	return (Tag < Other.Tag);
+	return Tag.LexicalLess(Other.Tag);
 }
 
 FRestrictedGameplayTagTableRow::FRestrictedGameplayTagTableRow(FRestrictedGameplayTagTableRow const& Other)
diff --git a/Engine/Source/Runtime/InputCore/Classes/InputCoreTypes.h b/Engine/Source/Runtime/InputCore/Classes/InputCoreTypes.h
--- a/Engine/Source/Runtime/InputCore/Classes/InputCoreTypes.h
+++ b/Engine/Source/Runtime/InputCore/Classes/InputCoreTypes.h
@@ -83,7 +83,7 @@ struct INPUTCORE_API FKey
 
 	friend bool operator==(const FKey& KeyA, const FKey& KeyB) { return KeyA.KeyName == KeyB.KeyName; }
 	friend bool operator!=(const FKey& KeyA, const FKey& KeyB) { return KeyA.KeyName != KeyB.KeyName; }
-	friend bool operator<(const FKey& KeyA, const FKey& KeyB) { return KeyA.KeyName < KeyB.KeyName; }
+	friend bool operator<(const FKey& KeyA, const FKey& KeyB) { return KeyA.KeyName.LexicalLess(KeyB.KeyName); }
 	friend uint32 GetTypeHash(const FKey& Key) { return GetTypeHash(Key.KeyName); }
 
 	friend struct EKeys;
diff --git a/Engine/Source/Runtime/MovieScene/Private/Channels/MovieSceneChannelProxy.cpp b/Engine/Source/Runtime/MovieScene/Private/Channels/MovieSceneChannelProxy.cpp
--- a/Engine/Source/Runtime/MovieScene/Private/Channels/MovieSceneChannelProxy.cpp
+++ b/Engine/Source/Runtime/MovieScene/Private/Channels/MovieSceneChannelProxy.cpp
@@ -12,7 +12,7 @@ FMovieSceneChannelHandle FMovieSceneChannelProxy::MakeHandle(FName ChannelTypeNa
 
 const FMovieSceneChannelEntry* FMovieSceneChannelProxy::FindEntry(FName ChannelTypeName) const
 {
-	const int32 ChannelTypeIndex = Algo::BinarySearchBy(Entries, ChannelTypeName, &FMovieSceneChannelEntry::ChannelTypeName);
+	const int32 ChannelTypeIndex = Algo::BinarySearchBy(Entries, ChannelTypeName, &FMovieSceneChannelEntry::ChannelTypeName, FNameLexicalLess());
 
 	if (ChannelTypeIndex != INDEX_NONE)
 	{
diff --git a/Engine/Source/Runtime/MovieScene/Public/Channels/MovieSceneChannelProxy.h b/Engine/Source/Runtime/MovieScene/Public/Channels/MovieSceneChannelProxy.h
--- a/Engine/Source/Runtime/MovieScene/Public/Channels/MovieSceneChannelProxy.h
+++ b/Engine/Source/Runtime/MovieScene/Public/Channels/MovieSceneChannelProxy.h
@@ -306,7 +306,7 @@ int32 FMovieSceneChannelProxyData::AddInternal(ChannelType& InChannel)
 	FName ChannelTypeName = ChannelType::StaticStruct()->GetFName();
 
 	// Find the first entry that has a >= channel ID
-	int32 ChannelTypeIndex = Algo::LowerBoundBy(Entries, ChannelTypeName, &FMovieSceneChannelEntry::GetChannelTypeName);
+	int32 ChannelTypeIndex = Algo::LowerBoundBy(Entries, ChannelTypeName, &FMovieSceneChannelEntry::GetChannelTypeName, FNameLexicalLess());
 
 	// If the index we found isn't valid, or it's not the channel we want, we need to add a new entry there
 	if (ChannelTypeIndex >= Entries.Num() || Entries[ChannelTypeIndex].GetChannelTypeName() != ChannelTypeName)
