Temporary fix for the failing concept checks under (probably) old Clang versions.

diff --git a/Engine/Source/Runtime/Core/Public/Containers/Map.h b/Engine/Source/Runtime/Core/Public/Containers/Map.h
--- a/Engine/Source/Runtime/Core/Public/Containers/Map.h
+++ b/Engine/Source/Runtime/Core/Public/Containers/Map.h
@@ -11,6 +11,7 @@
 #include "Misc/AssertionMacros.h"
 #include "Misc/StructBuilder.h"
 #include "Templates/Function.h"
+#include "Templates/HasGetTypeHash.h"
 #include "Templates/Models.h"
 #include "Templates/Sorting.h"
 #include "Templates/Tuple.h"
@@ -93,12 +94,16 @@ struct TDefaultMapKeyFuncs : BaseKeyFuncs<TPair<KeyType,ValueType>,KeyType,bInAl
 	}
 };
 
+PRAGMA_DISABLE_DEPRECATION_WARNINGS
+
 template<typename KeyType, typename ValueType, bool bInAllowDuplicateKeys>
 struct TDefaultMapHashableKeyFuncs : TDefaultMapKeyFuncs<KeyType, ValueType, bInAllowDuplicateKeys>
 {
-	static_assert(TModels<CGetTypeHashable, KeyType>::Value, "TMap must have a hashable KeyType unless a custom key func is provided.");
+	static_assert(THasGetTypeHash<KeyType>::Value, "TMap must have a hashable KeyType unless a custom key func is provided.");
 };
 
+PRAGMA_ENABLE_DEPRECATION_WARNINGS
+
 /** 
  * The base class of maps from keys to values.  Implemented using a TSet of key-value pairs with a custom KeyFuncs, 
  * with the same O(1) addition, removal, and finding. 
diff --git a/Engine/Source/Runtime/Core/Public/Templates/HasGetTypeHash.h b/Engine/Source/Runtime/Core/Public/Templates/HasGetTypeHash.h
--- a/Engine/Source/Runtime/Core/Public/Templates/HasGetTypeHash.h
+++ b/Engine/Source/Runtime/Core/Public/Templates/HasGetTypeHash.h
@@ -48,7 +48,7 @@ namespace UE4GetTypeHashExists_Private
 * Traits class which tests if a type has a GetTypeHash overload.
 */
 template <typename T>
-struct UE_DEPRECATED(4.23, "THasGetTypeHash has been deprecated, please use TModels<CGetTypeHashable, T> instead.") THasGetTypeHash
+struct THasGetTypeHash
 {
 	enum { Value = UE4GetTypeHashExists_Private::GetTypeHashQuery<T>::Value };
 };
diff --git a/Engine/Source/Runtime/CoreUObject/Public/UObject/LazyObjectPtr.h b/Engine/Source/Runtime/CoreUObject/Public/UObject/LazyObjectPtr.h
--- a/Engine/Source/Runtime/CoreUObject/Public/UObject/LazyObjectPtr.h
+++ b/Engine/Source/Runtime/CoreUObject/Public/UObject/LazyObjectPtr.h
@@ -171,6 +171,7 @@ public:
 
 template <> struct TIsPODType<FLazyObjectPtr> { enum { Value = TIsPODType<TPersistentObjectPtr<FUniqueObjectGuid> >::Value }; };
 template <> struct TIsWeakPointerType<FLazyObjectPtr> { enum { Value = TIsWeakPointerType<TPersistentObjectPtr<FUniqueObjectGuid> >::Value }; };
+template <> struct THasGetTypeHash<FLazyObjectPtr> { enum { Value = THasGetTypeHash<TPersistentObjectPtr<FUniqueObjectGuid> >::Value }; };
 
 /**
  * TLazyObjectPtr is templatized version of the generic FLazyObjectPtr
diff --git a/Engine/Source/Runtime/CoreUObject/Public/UObject/SoftObjectPtr.h b/Engine/Source/Runtime/CoreUObject/Public/UObject/SoftObjectPtr.h
--- a/Engine/Source/Runtime/CoreUObject/Public/UObject/SoftObjectPtr.h
+++ b/Engine/Source/Runtime/CoreUObject/Public/UObject/SoftObjectPtr.h
@@ -107,6 +107,7 @@ public:
 
 template <> struct TIsPODType<FSoftObjectPtr> { enum { Value = TIsPODType<TPersistentObjectPtr<FSoftObjectPath> >::Value }; };
 template <> struct TIsWeakPointerType<FSoftObjectPtr> { enum { Value = TIsWeakPointerType<TPersistentObjectPtr<FSoftObjectPath> >::Value }; };
+template <> struct THasGetTypeHash<FSoftObjectPtr> { enum { Value = THasGetTypeHash<TPersistentObjectPtr<FSoftObjectPath> >::Value }; };
 
 /**
  * TSoftObjectPtr is templatized wrapper of the generic FSoftObjectPtr, it can be used in UProperties
