Added FName IsNone() documentation based on UDN request

diff --git a/Engine/Source/Runtime/Core/Public/UObject/NameTypes.h b/Engine/Source/Runtime/Core/Public/UObject/NameTypes.h
--- a/Engine/Source/Runtime/Core/Public/UObject/NameTypes.h
+++ b/Engine/Source/Runtime/Core/Public/UObject/NameTypes.h
@@ -487,6 +487,7 @@ public:
 		return Compare(Other) < 0;
 	}
 
+	/** True for FName(), FName(NAME_None) and FName("None") */
 	FORCEINLINE bool IsNone() const
 	{
 		return !ComparisonIndex && GetNumber() == NAME_NO_NUMBER_INTERNAL;
@@ -496,7 +497,7 @@ public:
 	 * Paranoid sanity check
 	 *
 	 * All FNames are valid except for stomped memory, dangling pointers, etc.
-	 * Should only be used to  investigate such bugs and not in production code.
+	 * Should only be used to investigate such bugs and not in production code.
 	 */
 	bool IsValid() const { return IsWithinBounds(ComparisonIndex); }
 
