TSharedPtr now comparable to nullptr.
Updated TWeakPtr's comparisons to use the correct nullptr type.

diff --git a/Engine/Source/Runtime/Core/Public/Templates/SharedPointer.h b/Engine/Source/Runtime/Core/Public/Templates/SharedPointer.h
--- a/Engine/Source/Runtime/Core/Public/Templates/SharedPointer.h
+++ b/Engine/Source/Runtime/Core/Public/Templates/SharedPointer.h
@@ -1396,6 +1396,30 @@ FORCEINLINE bool operator==( TSharedPtr< ObjectTypeA, Mode > const& InSharedPtrA
 }
 
 
+/**
+ * Global equality operator for TSharedPtr
+ *
+ * @return  True if the shared pointer is null
+ */
+template< class ObjectTypeA, ESPMode Mode >
+FORCEINLINE bool operator==( TSharedPtr< ObjectTypeA, Mode > const& InSharedPtrA, TYPE_OF_NULLPTR )
+{
+	return !InSharedPtrA.IsValid();
+}
+
+
+/**
+ * Global equality operator for TSharedPtr
+ *
+ * @return  True if the shared pointer is null
+ */
+template< class ObjectTypeB, ESPMode Mode >
+FORCEINLINE bool operator==( TYPE_OF_NULLPTR, TSharedPtr< ObjectTypeB, Mode > const& InSharedPtrB )
+{
+	return !InSharedPtrB.IsValid();
+}
+
+
 /**
  * Global inequality operator for TSharedPtr
  *
@@ -1408,6 +1432,30 @@ FORCEINLINE bool operator!=( TSharedPtr< ObjectTypeA, Mode > const& InSharedPtrA
 }
 
 
+/**
+ * Global inequality operator for TSharedPtr
+ *
+ * @return  True if the shared pointer is not null
+ */
+template< class ObjectTypeA, ESPMode Mode >
+FORCEINLINE bool operator!=( TSharedPtr< ObjectTypeA, Mode > const& InSharedPtrA, TYPE_OF_NULLPTR )
+{
+	return InSharedPtrA.IsValid();
+}
+
+
+/**
+ * Global inequality operator for TSharedPtr
+ *
+ * @return  True if the shared pointer is not null
+ */
+template< class ObjectTypeB, ESPMode Mode >
+FORCEINLINE bool operator!=( TYPE_OF_NULLPTR, TSharedPtr< ObjectTypeB, Mode > const& InSharedPtrB )
+{
+	return InSharedPtrB.IsValid();
+}
+
+
 /**
  * Tests to see if a TSharedRef is "equal" to a TSharedPtr (both are valid and refer to the same object)
  *
@@ -1522,7 +1570,7 @@ FORCEINLINE bool operator==( TSharedPtr< ObjectTypeA, Mode > const& InSharedPtrA
  * @return  True if the weak pointer is null
  */
 template< class ObjectTypeA, ESPMode Mode >
-FORCEINLINE bool operator==( TWeakPtr< ObjectTypeA, Mode > const& InWeakPtrA, decltype(nullptr) )
+FORCEINLINE bool operator==( TWeakPtr< ObjectTypeA, Mode > const& InWeakPtrA, TYPE_OF_NULLPTR )
 {
 	return !InWeakPtrA.IsValid();
 }
@@ -1534,7 +1582,7 @@ FORCEINLINE bool operator==( TWeakPtr< ObjectTypeA, Mode > const& InWeakPtrA, de
  * @return  True if the weak pointer is null
  */
 template< class ObjectTypeB, ESPMode Mode >
-FORCEINLINE bool operator==( decltype(nullptr), TWeakPtr< ObjectTypeB, Mode > const& InWeakPtrB )
+FORCEINLINE bool operator==( TYPE_OF_NULLPTR, TWeakPtr< ObjectTypeB, Mode > const& InWeakPtrB )
 {
 	return !InWeakPtrB.IsValid();
 }
@@ -1606,7 +1654,7 @@ FORCEINLINE bool operator!=( TSharedPtr< ObjectTypeA, Mode > const& InSharedPtrA
  * @return  True if the weak pointer is not null
  */
 template< class ObjectTypeA, ESPMode Mode >
-FORCEINLINE bool operator!=( TWeakPtr< ObjectTypeA, Mode > const& InWeakPtrA, decltype(nullptr) )
+FORCEINLINE bool operator!=( TWeakPtr< ObjectTypeA, Mode > const& InWeakPtrA, TYPE_OF_NULLPTR )
 {
 	return InWeakPtrA.IsValid();
 }
@@ -1618,7 +1666,7 @@ FORCEINLINE bool operator!=( TWeakPtr< ObjectTypeA, Mode > const& InWeakPtrA, de
  * @return  True if the weak pointer is not null
  */
 template< class ObjectTypeB, ESPMode Mode >
-FORCEINLINE bool operator!=( decltype(nullptr), TWeakPtr< ObjectTypeB, Mode > const& InWeakPtrB )
+FORCEINLINE bool operator!=( TYPE_OF_NULLPTR, TWeakPtr< ObjectTypeB, Mode > const& InWeakPtrB )
 {
 	return InWeakPtrB.IsValid();
 }
