A thread that shared the same name as a previous thread will not get tracked correctly because there will already be an entry in ShortNameToLongName.

(peafour-cherrypick of //Tasks/UE4/Dev-Core-Trace/[REDACTED_HANDLE] by [REDACTED_PERSON])

diff --git a/Engine/Source/Runtime/Core/Private/Stats/StatsData.cpp b/Engine/Source/Runtime/Core/Private/Stats/StatsData.cpp
--- a/Engine/Source/Runtime/Core/Private/Stats/StatsData.cpp
+++ b/Engine/Source/Runtime/Core/Private/Stats/StatsData.cpp
@@ -1686,27 +1686,27 @@ void FStatsThreadState::Condense(int64 TargetFrame, TArray<FStatMessage>& OutSta
 
 void FStatsThreadState::FindOrAddMetaData(FStatMessage const& Item)
 {
-	const FName LongName = Item.NameAndInfo.GetRawName();
 	const FName ShortName = Item.NameAndInfo.GetShortName();
+	const FName GroupName = Item.NameAndInfo.GetGroupName();
+
+	// Whether to add to the threads group.
+	const bool bIsThread = FStatConstants::NAME_ThreadGroup == GroupName;
+	if (bIsThread)
+	{
+		// The description of a thread group contains the thread id
+		const FString Desc = Item.NameAndInfo.GetDescription();
+		Threads.Add(FStatsUtils::ParseThreadID(Desc), ShortName);
+	}
 
 	FStatMessage* Result = ShortNameToLongName.Find(ShortName);
 	if (!Result)
 	{
+		const FName LongName = Item.NameAndInfo.GetRawName();
+
 		check(ShortName != LongName);
 		FStatMessage AsSet(Item);
 		AsSet.Clear();
 
-		const FName GroupName = Item.NameAndInfo.GetGroupName();
-
-		// Whether to add to the threads group.
-		const bool bIsThread = FStatConstants::NAME_ThreadGroup == GroupName;
-		if( bIsThread )
-		{
-			// The description of a thread group contains the thread id
-			const FString Desc = Item.NameAndInfo.GetDescription();
-			Threads.Add( FStatsUtils::ParseThreadID( Desc ), ShortName );
-		}
-
 		ShortNameToLongName.Add(ShortName, AsSet); // we want this to be a clear, but it should be a SetLongName
 		AsSet.NameAndInfo.SetField<EStatOperation>(EStatOperation::Set);
 		check(Item.NameAndInfo.GetField<EStatMetaFlags>());
