Potential fix for incremental build.

diff --git a/Engine/Source/Runtime/Core/Public/Algo/Reverse.h b/Engine/Source/Runtime/Core/Public/Algo/Reverse.h
--- a/Engine/Source/Runtime/Core/Public/Algo/Reverse.h
+++ b/Engine/Source/Runtime/Core/Public/Algo/Reverse.h
@@ -3,6 +3,7 @@
 #pragma once
 
 #include "CoreTypes.h"
+#include "Templates/UnrealTemplate.h" // for Swap
 
 namespace AlgoImpl
 {
