Fix incorrect check in FMinimalName::IsNone().
Fix various stat related crashes.

The mismatch between FName::IsNone() and FMinimalName::IsNone() could cause a mismatch between FCycleCounter::Start() and FCycleCounter::Stop().
This would skew ThreadStats->ScopeCount, which would trigger thread stat flushes of incorrectly or not at all.
When the game thread never flushes its stats, the stats thread never updates CurrentGameFrame from AdvanceFrameEventGameThread, and it will crash in AddToHistoryAndEmpty.

diff --git a/Engine/Source/Runtime/Core/Public/UObject/NameTypes.h b/Engine/Source/Runtime/Core/Public/UObject/NameTypes.h
--- a/Engine/Source/Runtime/Core/Public/UObject/NameTypes.h
+++ b/Engine/Source/Runtime/Core/Public/UObject/NameTypes.h
@@ -340,7 +340,7 @@ struct FMinimalName
 
 	FORCEINLINE bool IsNone() const
 	{
-		return Index && Number == NAME_NO_NUMBER_INTERNAL;
+		return !Index && Number == NAME_NO_NUMBER_INTERNAL;
 	}
 
 	/** Index into the Names array (used to find String portion of the string/number pair) */
