Serialize() fallbacks for when for you don't provide 'the other' FArchive/FSlot overload of a function.
Deprecated FBox constructor removed, since it was giving some false positives on Clang when doing the concept check.

diff --git a/Engine/Source/Runtime/Core/Public/Concepts/Insertable.h b/Engine/Source/Runtime/Core/Public/Concepts/Insertable.h
--- a/Engine/Source/Runtime/Core/Public/Concepts/Insertable.h
+++ b/Engine/Source/Runtime/Core/Public/Concepts/Insertable.h
@@ -7,9 +7,10 @@
 /**
  * Describes a type with a GetTypeHash overload.
  */
+template <typename DestType>
 struct CInsertable {
-	template <typename DestType, typename T>
-	auto Requires(DestType& Dest, T& Val) -> decltype(
+	template <typename T>
+	auto Requires(DestType Dest, T& Val) -> decltype(
 		Dest << Val
 	);
 };
diff --git a/Engine/Source/Runtime/Core/Public/Math/Box.h b/Engine/Source/Runtime/Core/Public/Math/Box.h
--- a/Engine/Source/Runtime/Core/Public/Math/Box.h
+++ b/Engine/Source/Runtime/Core/Public/Math/Box.h
@@ -32,13 +32,6 @@ public:
 	/** Default constructor (no initialization). */
 	FBox() { }
 
-	/** Creates and initializes a new box with zero extent and marks it as invalid. */
-	UE_DEPRECATED(4.16, "Use ForceInit constructor instead.")
-	FBox( int32 )
-	{
-		Init();
-	}
-
 	/**
 	 * Creates and initializes a new box with zero extent and marks it as invalid.
 	 *
diff --git a/Engine/Source/Runtime/Core/Public/Serialization/ArchiveFromStructuredArchive.h b/Engine/Source/Runtime/Core/Public/Serialization/ArchiveFromStructuredArchive.h
--- a/Engine/Source/Runtime/Core/Public/Serialization/ArchiveFromStructuredArchive.h
+++ b/Engine/Source/Runtime/Core/Public/Serialization/ArchiveFromStructuredArchive.h
@@ -9,6 +9,8 @@
 #include "UObject/NameTypes.h"
 #include "Containers/Map.h"
 #include "Containers/BitArray.h"
+#include "Concepts/Insertable.h"
+#include "Templates/Models.h"
 
 class CORE_API FArchiveFromStructuredArchive : public FArchiveProxy
 {
@@ -65,4 +67,24 @@ private:
 	TMap<UObject*, int32> ObjectToIndex;
 
 	FStructuredArchive::FSlot RootSlot;
-};
\ No newline at end of file
+};
+
+/**
+ * Adapter operator which allows a type to stream to an FStructuredArchive::FSlot when it already supports streaming to an FArchive.
+ *
+ * @param  Slot  The slot to read from or write to.
+ * @param  Obj   The object to read or write.
+ */
+template <typename T>
+typename TEnableIf<
+	TModels<CInsertable<FArchive&>, T>::Value &&
+	!TModels<CInsertable<FStructuredArchive::FSlot>, T>::Value
+>::Type operator<<(FStructuredArchive::FSlot Slot, T& Obj)
+{
+#if WITH_TEXT_ARCHIVE_SUPPORT
+	FArchiveFromStructuredArchive Ar(Slot);
+#else
+	FArchive& Ar = Slot.GetUnderlyingArchive();
+#endif
+	Ar << Obj;
+}
diff --git a/Engine/Source/Runtime/Core/Public/Serialization/StructuredArchiveFromArchive.h b/Engine/Source/Runtime/Core/Public/Serialization/StructuredArchiveFromArchive.h
--- a/Engine/Source/Runtime/Core/Public/Serialization/StructuredArchiveFromArchive.h
+++ b/Engine/Source/Runtime/Core/Public/Serialization/StructuredArchiveFromArchive.h
@@ -3,6 +3,8 @@
 #pragma once
 
 #include "Serialization/StructuredArchive.h"
+#include "Concepts/Insertable.h"
+#include "Templates/Models.h"
 
 class CORE_API FStructuredArchiveFromArchive
 {
@@ -23,3 +25,22 @@ private:
 	FStructuredArchive StructuredArchive;
 	FStructuredArchive::FSlot Slot;
 };
+
+/**
+ * Adapter operator which allows a type to stream to an FArchive when it already supports streaming to an FStructuredArchive::FSlot.
+ *
+ * @param  Ar   The archive to read from or write to.
+ * @param  Obj  The object to read or write.
+ *
+ * @return  A reference to the same archive as Ar.
+ */
+template <typename T>
+typename TEnableIf<
+	!TModels<CInsertable<FArchive&>, T>::Value && TModels<CInsertable<FStructuredArchive::FSlot>, T>::Value,
+	FArchive&
+>::Type operator<<(FArchive& Ar, T& Obj)
+{
+	FStructuredArchiveFromArchive ArAdapt(Ar);
+	ArAdapt.GetSlot() << Obj;
+	return Ar;
+}
diff --git a/Engine/Source/Runtime/Core/Public/Templates/HasInserterOperator.h b/Engine/Source/Runtime/Core/Public/Templates/HasInserterOperator.h
--- a/Engine/Source/Runtime/Core/Public/Templates/HasInserterOperator.h
+++ b/Engine/Source/Runtime/Core/Public/Templates/HasInserterOperator.h
@@ -41,7 +41,7 @@ namespace UE4HasInserterOperator_Private
  * Traits class which tests if a type has an operator<< overload between two types.
  */
 template <typename Dest, typename T>
-struct UE_DEPRECATED(4.23, "THasInserterOperator has been deprecated, please use TModels<CInsertable, Dest, T> instead.") THasInserterOperator
+struct UE_DEPRECATED(4.23, "THasInserterOperator has been deprecated, please use TModels<CInsertable<Dest>, T> instead.") THasInserterOperator
 {
 	enum { Value = UE4HasInserterOperator_Private::THasInserterOperator<Dest, T>::Value };
 };
