Warning for incorrect member init order enabled on VC.

diff --git a/Engine/Source/Runtime/Core/Public/Windows/WindowsPlatformCompilerSetup.h b/Engine/Source/Runtime/Core/Public/Windows/WindowsPlatformCompilerSetup.h
--- a/Engine/Source/Runtime/Core/Public/Windows/WindowsPlatformCompilerSetup.h
+++ b/Engine/Source/Runtime/Core/Public/Windows/WindowsPlatformCompilerSetup.h
@@ -236,6 +236,7 @@ static_assert(_MSC_VER != 1914 && _MSC_VER != 1915, "Visual Studio 2017 versions
 #pragma warning(default: 4928) // illegal copy-initialization; more than one user-defined conversion has been implicitly applied	https://docs.microsoft.com/en-us/cpp/error-messages/compiler-warnings/compiler-warning-level-1-c4928
 #pragma warning(default: 4931) // we are assuming the type library was built for number-bit pointers								https://docs.microsoft.com/en-us/cpp/error-messages/compiler-warnings/compiler-warning-level-4-c4931
 #pragma warning(default: 4946) // reinterpret_cast used between related classes: 'class1' and 'class2'								https://docs.microsoft.com/en-us/cpp/error-messages/compiler-warnings/compiler-warning-level-1-c4946
+#pragma warning(default: 5038) // data member 'A::y' will be initialized after data member 'A::x'									https://docs.microsoft.com/en-us/cpp/error-messages/compiler-warnings/c5038
 
 #pragma warning(default: 4996)
 
