Optimize usage of FArchive<->FStructuredArchive adaption objects when text asset support is compiled out. Adapters become simple passthroughs rather than proxies

diff --git a/Engine/Source/Runtime/Core/Private/Serialization/ArchiveFromStructuredArchive.cpp b/Engine/Source/Runtime/Core/Private/Serialization/ArchiveFromStructuredArchive.cpp
--- a/Engine/Source/Runtime/Core/Private/Serialization/ArchiveFromStructuredArchive.cpp
+++ b/Engine/Source/Runtime/Core/Private/Serialization/ArchiveFromStructuredArchive.cpp
@@ -3,6 +3,8 @@
 #include "Serialization/ArchiveFromStructuredArchive.h"
 #include "Internationalization/Text.h"
 
+#if WITH_TEXT_ARCHIVE_SUPPORT
+
 FArchiveFromStructuredArchive::FArchiveFromStructuredArchive(FStructuredArchive::FSlot Slot)
 	: FArchiveProxy(Slot.GetUnderlyingArchive())
 	, bPendingSerialize(true)
@@ -289,4 +291,6 @@ void FArchiveFromStructuredArchive::OpenArchive()
 			RootSlot.EnterStream();
 		}
 	}
-}
\ No newline at end of file
+}
+
+#endif
\ No newline at end of file
diff --git a/Engine/Source/Runtime/Core/Public/Serialization/ArchiveFromStructuredArchive.h b/Engine/Source/Runtime/Core/Public/Serialization/ArchiveFromStructuredArchive.h
--- a/Engine/Source/Runtime/Core/Public/Serialization/ArchiveFromStructuredArchive.h
+++ b/Engine/Source/Runtime/Core/Public/Serialization/ArchiveFromStructuredArchive.h
@@ -12,6 +12,8 @@
 #include "Concepts/Insertable.h"
 #include "Templates/Models.h"
 
+#if WITH_TEXT_ARCHIVE_SUPPORT
+
 class CORE_API FArchiveFromStructuredArchive : public FArchiveProxy
 {
 public:
@@ -69,6 +71,27 @@ private:
 	FStructuredArchive::FSlot RootSlot;
 };
 
+#else
+
+class CORE_API FArchiveFromStructuredArchive
+{
+public:
+
+	FArchiveFromStructuredArchive(FStructuredArchive::FSlot InSlot)
+		: Ar(InSlot.GetUnderlyingArchive())
+	{
+
+	}
+
+	operator FArchive& () { return Ar; }
+
+private:
+
+	FArchive& Ar;
+};
+
+#endif
+
 /**
  * Adapter operator which allows a type to stream to an FStructuredArchive::FSlot when it already supports streaming to an FArchive.
  *
@@ -87,4 +110,4 @@ typename TEnableIf<
 	FArchive& Ar = Slot.GetUnderlyingArchive();
 #endif
 	Ar << Obj;
-}
+}
\ No newline at end of file
diff --git a/Engine/Source/Runtime/CoreUObject/Private/Serialization/ArchiveUObjectFromStructuredArchive.cpp b/Engine/Source/Runtime/CoreUObject/Private/Serialization/ArchiveUObjectFromStructuredArchive.cpp
--- a/Engine/Source/Runtime/CoreUObject/Private/Serialization/ArchiveUObjectFromStructuredArchive.cpp
+++ b/Engine/Source/Runtime/CoreUObject/Private/Serialization/ArchiveUObjectFromStructuredArchive.cpp
@@ -2,6 +2,8 @@
 
 #include "Serialization/ArchiveUObjectFromStructuredArchive.h"
 
+#if WITH_TEXT_ARCHIVE_SUPPORT
+
 FArchiveUObjectFromStructuredArchive::FArchiveUObjectFromStructuredArchive(FStructuredArchive::FSlot Slot)
 	: FArchiveFromStructuredArchive(Slot)
 	, bPendingSerialize(true)
@@ -182,4 +184,6 @@ void FArchiveUObjectFromStructuredArchive::SerializeInternal(FStructuredArchive:
 
 		bPendingSerialize = true;
 	}
-}
\ No newline at end of file
+}
+
+#endif
\ No newline at end of file
diff --git a/Engine/Source/Runtime/CoreUObject/Private/UObject/PropertyMulticastDelegate.cpp b/Engine/Source/Runtime/CoreUObject/Private/UObject/PropertyMulticastDelegate.cpp
--- a/Engine/Source/Runtime/CoreUObject/Private/UObject/PropertyMulticastDelegate.cpp
+++ b/Engine/Source/Runtime/CoreUObject/Private/UObject/PropertyMulticastDelegate.cpp
@@ -470,7 +470,11 @@ FMulticastScriptDelegate::FInvocationList& UMulticastSparseDelegateProperty::Get
 void UMulticastSparseDelegateProperty::SerializeItem(FStructuredArchive::FSlot Slot, void* Value, void const* Defaults) const
 {
 	FArchiveUObjectFromStructuredArchive Ar(Slot);
+	SerializeItemInternal(Ar, Value, Defaults);
+}
 
+void UMulticastSparseDelegateProperty::SerializeItemInternal(FArchive& Ar, void* Value, void const* Defaults) const
+{
 	FSparseDelegate& SparseDelegate = *(FSparseDelegate*)Value;
 
 	if (Ar.IsLoading())
diff --git a/Engine/Source/Runtime/CoreUObject/Public/Serialization/ArchiveUObjectFromStructuredArchive.h b/Engine/Source/Runtime/CoreUObject/Public/Serialization/ArchiveUObjectFromStructuredArchive.h
--- a/Engine/Source/Runtime/CoreUObject/Public/Serialization/ArchiveUObjectFromStructuredArchive.h
+++ b/Engine/Source/Runtime/CoreUObject/Public/Serialization/ArchiveUObjectFromStructuredArchive.h
@@ -11,6 +11,8 @@
 #include "UObject/LazyObjectPtr.h"
 #include "UObject/WeakObjectPtr.h"
 
+#if WITH_TEXT_ARCHIVE_SUPPORT
+
 class COREUOBJECT_API FArchiveUObjectFromStructuredArchive : public FArchiveFromStructuredArchive
 {
 public:
@@ -42,4 +44,19 @@ private:
 	TMap<FSoftObjectPath, int32> SoftObjectPathToIndex;
 
 	virtual void SerializeInternal(FStructuredArchive::FRecord Record) override;
-};
\ No newline at end of file
+};
+
+#else
+
+class COREUOBJECT_API FArchiveUObjectFromStructuredArchive : public FArchiveFromStructuredArchive
+{
+public:
+
+	FArchiveUObjectFromStructuredArchive(FStructuredArchive::FSlot InSlot)
+		: FArchiveFromStructuredArchive(InSlot)
+	{
+
+	}
+};
+
+#endif
\ No newline at end of file
diff --git a/Engine/Source/Runtime/CoreUObject/Public/UObject/UnrealType.h b/Engine/Source/Runtime/CoreUObject/Public/UObject/UnrealType.h
--- a/Engine/Source/Runtime/CoreUObject/Public/UObject/UnrealType.h
+++ b/Engine/Source/Runtime/CoreUObject/Public/UObject/UnrealType.h
@@ -4346,6 +4346,9 @@ public:
 protected:
 	virtual FMulticastScriptDelegate::FInvocationList& GetInvocationList(const void* PropertyValue) const;
 	// End of UMulticastDelegateProperty interface
+
+private:
+	virtual void SerializeItemInternal(FArchive& Ar, void* Value, void const* Defaults) const;
 };
 
 /** Describes a single node in a custom property list. */
