Fixed broken implementation of PLATFORM_BREAK on iOS

- Turns out __asm__("svc 0") doesn't do anything if the debugger is attached so it's replaced by __builtin_trap(), which should also work for other configurations.

(peafour-cherrypick of 5723069 by wang.hao)

diff --git a/Engine/Source/Runtime/Core/Public/IOS/IOSPlatform.h b/Engine/Source/Runtime/Core/Public/IOS/IOSPlatform.h
--- a/Engine/Source/Runtime/Core/Public/IOS/IOSPlatform.h
+++ b/Engine/Source/Runtime/Core/Public/IOS/IOSPlatform.h
@@ -66,13 +66,8 @@ typedef FIOSPlatformTypes FPlatformTypes;
 #define PLATFORM_NEEDS_RHIRESOURCELIST					0
 #define PLATFORM_SUPPORTS_GEOMETRY_SHADERS				0
 #define PLATFORM_SUPPORTS_TESSELLATION_SHADERS			0
-#if WITH_SIMULATOR
-	#define PLATFORM_BREAK()							__asm__("int $3")
-#elif PLATFORM_64BITS
-	#define PLATFORM_BREAK()							__asm__("svc 0")
-#else
-	#define PLATFORM_BREAK()							__asm__("trap")
-#endif
+
+#define PLATFORM_BREAK()                                __builtin_trap()
 
 #define PLATFORM_CODE_SECTION(Name)						__attribute__((section("__TEXT,__" Name ",regular,pure_instructions"))) \
 														__attribute__((aligned(4)))
