Move insertion operator adapter definitions into their own header file and include from StructuredArchive.h so they have global visibility

diff --git a/Engine/Source/Runtime/Core/Public/Serialization/ArchiveAdapters.h b/Engine/Source/Runtime/Core/Public/Serialization/ArchiveAdapters.h
new file mode 100644
--- /dev/null
+++ b/Engine/Source/Runtime/Core/Public/Serialization/ArchiveAdapters.h
@@ -0,0 +1,45 @@
+// Copyright 1998-2019 Epic Games, Inc. All Rights Reserved.
+
+#pragma once
+
+#include "Concepts/Insertable.h"
+#include "Templates/Models.h"
+
+/**
+ * Adapter operator which allows a type to stream to an FStructuredArchive::FSlot when it already supports streaming to an FArchive.
+ *
+ * @param  Slot  The slot to read from or write to.
+ * @param  Obj   The object to read or write.
+ */
+template <typename T>
+typename TEnableIf<
+	TModels<CInsertable<FArchive&>, T>::Value &&
+	!TModels<CInsertable<FStructuredArchive::FSlot>, T>::Value
+>::Type operator<<(FStructuredArchive::FSlot Slot, T& Obj)
+{
+#if WITH_TEXT_ARCHIVE_SUPPORT
+	FArchiveFromStructuredArchive Ar(Slot);
+#else
+	FArchive& Ar = Slot.GetUnderlyingArchive();
+#endif
+	Ar << Obj;
+}
+
+/**
+ * Adapter operator which allows a type to stream to an FArchive when it already supports streaming to an FStructuredArchive::FSlot.
+ *
+ * @param  Ar   The archive to read from or write to.
+ * @param  Obj  The object to read or write.
+ *
+ * @return  A reference to the same archive as Ar.
+ */
+template <typename T>
+typename TEnableIf<
+	!TModels<CInsertable<FArchive&>, T>::Value && TModels<CInsertable<FStructuredArchive::FSlot>, T>::Value,
+	FArchive&
+>::Type operator<<(FArchive& Ar, T& Obj)
+{
+	FStructuredArchiveFromArchive ArAdapt(Ar);
+	ArAdapt.GetSlot() << Obj;
+	return Ar;
+}
\ No newline at end of file
diff --git a/Engine/Source/Runtime/Core/Public/Serialization/ArchiveFromStructuredArchive.h b/Engine/Source/Runtime/Core/Public/Serialization/ArchiveFromStructuredArchive.h
--- a/Engine/Source/Runtime/Core/Public/Serialization/ArchiveFromStructuredArchive.h
+++ b/Engine/Source/Runtime/Core/Public/Serialization/ArchiveFromStructuredArchive.h
@@ -90,24 +90,4 @@ private:
 	FArchive& Ar;
 };
 
-#endif
-
-/**
- * Adapter operator which allows a type to stream to an FStructuredArchive::FSlot when it already supports streaming to an FArchive.
- *
- * @param  Slot  The slot to read from or write to.
- * @param  Obj   The object to read or write.
- */
-template <typename T>
-typename TEnableIf<
-	TModels<CInsertable<FArchive&>, T>::Value &&
-	!TModels<CInsertable<FStructuredArchive::FSlot>, T>::Value
->::Type operator<<(FStructuredArchive::FSlot Slot, T& Obj)
-{
-#if WITH_TEXT_ARCHIVE_SUPPORT
-	FArchiveFromStructuredArchive Ar(Slot);
-#else
-	FArchive& Ar = Slot.GetUnderlyingArchive();
-#endif
-	Ar << Obj;
-}
\ No newline at end of file
+#endif
\ No newline at end of file
diff --git a/Engine/Source/Runtime/Core/Public/Serialization/StructuredArchive.h b/Engine/Source/Runtime/Core/Public/Serialization/StructuredArchive.h
--- a/Engine/Source/Runtime/Core/Public/Serialization/StructuredArchive.h
+++ b/Engine/Source/Runtime/Core/Public/Serialization/StructuredArchive.h
@@ -759,3 +759,4 @@ private:
 
 #endif
 
+#include "Serialization/ArchiveAdapters.h"
\ No newline at end of file
diff --git a/Engine/Source/Runtime/Core/Public/Serialization/StructuredArchiveFromArchive.h b/Engine/Source/Runtime/Core/Public/Serialization/StructuredArchiveFromArchive.h
--- a/Engine/Source/Runtime/Core/Public/Serialization/StructuredArchiveFromArchive.h
+++ b/Engine/Source/Runtime/Core/Public/Serialization/StructuredArchiveFromArchive.h
@@ -3,8 +3,6 @@
 #pragma once
 
 #include "Serialization/StructuredArchive.h"
-#include "Concepts/Insertable.h"
-#include "Templates/Models.h"
 
 class CORE_API FStructuredArchiveFromArchive
 {
@@ -25,22 +23,3 @@ private:
 	FStructuredArchive StructuredArchive;
 	FStructuredArchive::FSlot Slot;
 };
-
-/**
- * Adapter operator which allows a type to stream to an FArchive when it already supports streaming to an FStructuredArchive::FSlot.
- *
- * @param  Ar   The archive to read from or write to.
- * @param  Obj  The object to read or write.
- *
- * @return  A reference to the same archive as Ar.
- */
-template <typename T>
-typename TEnableIf<
-	!TModels<CInsertable<FArchive&>, T>::Value && TModels<CInsertable<FStructuredArchive::FSlot>, T>::Value,
-	FArchive&
->::Type operator<<(FArchive& Ar, T& Obj)
-{
-	FStructuredArchiveFromArchive ArAdapt(Ar);
-	ArAdapt.GetSlot() << Obj;
-	return Ar;
-}
