Add IsExplicitlyNull to check if a weak object ptr is really null instead of just missing, this is much faster to call than IsStale()
Comment and coding standard fixes for weak object ptr requested on review

diff --git a/Engine/Source/Runtime/Core/Public/UObject/WeakObjectPtrTemplates.h b/Engine/Source/Runtime/Core/Public/UObject/WeakObjectPtrTemplates.h
--- a/Engine/Source/Runtime/Core/Public/UObject/WeakObjectPtrTemplates.h
+++ b/Engine/Source/Runtime/Core/Public/UObject/WeakObjectPtrTemplates.h
@@ -9,28 +9,22 @@
 #include "Templates/AndOrNot.h"
 #include "Containers/Map.h"
 
-struct FWeakObjectPtr;
-
-/***
- * 
+/**
  * FWeakObjectPtr is a weak pointer to a UObject. 
  * It can return nullptr later if the object is garbage collected.
  * It has no impact on if the object is garbage collected or not.
  * It can't be directly used across a network.
  *
  * Most often it is used when you explicitly do NOT want to prevent something from being garbage collected.
- **/
+ */
 struct FWeakObjectPtr;
 
 template<class T=UObject, class TWeakObjectPtrBase=FWeakObjectPtr>
 struct TWeakObjectPtr;
 
-
-/***
-* 
-* TWeakObjectPtr is templatized version of the generic FWeakObjectPtr
-* 
-**/
+/**
+ * TWeakObjectPtr is the templated version of the generic FWeakObjectPtr
+ */
 template<class T, class TWeakObjectPtrBase>
 struct TWeakObjectPtr : private TWeakObjectPtrBase
 {
@@ -46,7 +40,7 @@ public:
 
 	/**
 	 * Construct from a null pointer
-	**/
+	 */
 	FORCEINLINE TWeakObjectPtr(TYPE_OF_NULLPTR) :
 		TWeakObjectPtrBase((UObject*)nullptr)
 	{
@@ -55,7 +49,7 @@ public:
 	/**
 	 * Construct from an object pointer
 	 * @param Object object to create a weak pointer to
-	**/
+	 */
 	template <
 		typename U,
 		typename = typename TEnableIf<TPointerIsConvertibleFromTo<const volatile U, const volatile T>::Value>::Type
@@ -88,7 +82,7 @@ public:
 	/**
 	 * Construct from another weak pointer of another type, intended for derived-to-base conversions
 	 * @param Other weak pointer to copy from
-	**/
+	 */
 	template <typename OtherT>
 	FORCEINLINE TWeakObjectPtr(const TWeakObjectPtr<OtherT, TWeakObjectPtrBase>& Other) :
 		TWeakObjectPtrBase(*(TWeakObjectPtrBase*)&Other) // we do a C-style cast to private base here to avoid clang 3.6.0 compilation problems with friend declarations
@@ -99,7 +93,7 @@ public:
 	}
 
 	/**
-	 * Reset the weak pointer back to the NULL state
+	 * Reset the weak pointer back to the null state
 	 */
 	FORCEINLINE void Reset()
 	{
@@ -109,7 +103,7 @@ public:
 	/**  
 	 * Copy from an object pointer
 	 * @param Object object to create a weak pointer to
-	**/
+	 */
 	template<class U>
 	FORCEINLINE typename TEnableIf<!TLosesQualifiersFromTo<U, T>::Value, TWeakObjectPtr&>::Type operator=(U* Object)
 	{
@@ -129,7 +123,7 @@ public:
 	/**  
 	 * Assign from another weak pointer, intended for derived-to-base conversions
 	 * @param Other weak pointer to copy from
-	**/
+	 */
 	template <typename OtherT>
 	FORCEINLINE TWeakObjectPtr& operator=(const TWeakObjectPtr<OtherT, TWeakObjectPtrBase>& Other)
 	{
@@ -145,8 +139,8 @@ public:
 	/**  
 	 * Dereference the weak pointer
 	 * @param bEvenIfPendingKill, if this is true, pendingkill objects are considered valid	
-	 * @return NULL if this object is gone or the weak pointer was NULL, otherwise a valid uobject pointer
-	**/
+	 * @return nullptr if this object is gone or the weak pointer is explicitly null, otherwise a valid uobject pointer
+	 */
 	FORCEINLINE T* Get(bool bEvenIfPendingKill) const
 	{
 		return (T*)TWeakObjectPtrBase::Get(bEvenIfPendingKill);
@@ -168,16 +162,16 @@ public:
 
 	/**  
 	 * Dereference the weak pointer
-	**/
-	FORCEINLINE T & operator*() const
+	 */
+	FORCEINLINE T& operator*() const
 	{
 		return *Get();
 	}
 
 	/**  
 	 * Dereference the weak pointer
-	**/
-	FORCEINLINE T * operator->() const
+	 */
+	FORCEINLINE T* operator->() const
 	{
 		return Get();
 	}
@@ -186,9 +180,9 @@ public:
 	 * Test if this points to a live UObject
 	 * @param bEvenIfPendingKill, if this is true, pendingkill objects are considered valid
 	 * @param bThreadsafeTest, if true then function will just give you information whether referenced 
-	 *							UObject is gone forever (@return false) or if it is still there (@return true, no object flags checked).
+	 *							UObject is gone forever (return false) or if it is still there (return true, no object flags checked).
 	 * @return true if Get() would return a valid non-null pointer
-	**/
+	 */
 	FORCEINLINE bool IsValid(bool bEvenIfPendingKill, bool bThreadsafeTest = false) const
 	{
 		return TWeakObjectPtrBase::IsValid(bEvenIfPendingKill, bThreadsafeTest);
@@ -208,12 +202,25 @@ public:
 	 * @param bIncludingIfPendingKill, if this is true, pendingkill objects are considered stale
 	 * @param bThreadsafeTest, set it to true when testing outside of Game Thread. Results in false if WeakObjPtr point to an existing object (no flags checked) 
 	 * @return true if this used to point at a real object but no longer does.
-	**/
+	 */
 	FORCEINLINE bool IsStale(bool bIncludingIfPendingKill = true, bool bThreadsafeTest = false) const
 	{
 		return TWeakObjectPtrBase::IsStale(bIncludingIfPendingKill, bThreadsafeTest);
 	}
+	
+	/**
+	 * Returns true if this pointer was explicitly assigned to null, was reset, or was never initialized.
+	 * If this returns true, IsValid() and IsStale() will both return false.
+	 */
+	FORCEINLINE bool IsExplicitlyNull() const
+	{
+		return TWeakObjectPtrBase::IsExplicitlyNull();
+	}
 
+	/**
+	 * Returns true if two weak pointers were originally set to the same object, even if they are now stale
+	 * @param Other weak pointer to compare to
+	 */
 	FORCEINLINE bool HasSameIndexAndSerialNumber(const TWeakObjectPtr& Other) const
 	{
 		return static_cast<const TWeakObjectPtrBase&>(*this).HasSameIndexAndSerialNumber(static_cast<const TWeakObjectPtrBase&>(Other));
@@ -225,6 +232,9 @@ public:
 		return GetTypeHash(static_cast<const TWeakObjectPtrBase&>(WeakObjectPtr));
 	}
 
+	/**
+	 * Weak object pointer serialization, this forwards to FArchive::operator<<(struct FWeakObjectPtr&) or an override
+	 */
 	friend FArchive& operator<<( FArchive& Ar, TWeakObjectPtr& WeakObjectPtr )
 	{
 		Ar << static_cast<TWeakObjectPtrBase&>(WeakObjectPtr);
@@ -239,6 +249,11 @@ FORCEINLINE TWeakObjectPtr<T> MakeWeakObjectPtr(T* Ptr)
 	return TWeakObjectPtr<T>(Ptr);
 }
 
+/**
+ * Compare weak pointers for equality.
+ * If both pointers would return nullptr from Get() they count as equal even if they were not initialized to the same object.
+ * @param Other weak pointer to compare to
+ */
 template <typename LhsT, typename RhsT, typename OtherTWeakObjectPtrBase>
 FORCENOINLINE bool operator==(const TWeakObjectPtr<LhsT, OtherTWeakObjectPtrBase>& Lhs, const TWeakObjectPtr<RhsT, OtherTWeakObjectPtrBase>& Rhs)
 {
@@ -285,6 +300,10 @@ FORCENOINLINE bool operator==(TYPE_OF_NULLPTR, const TWeakObjectPtr<RhsT, OtherT
 	return !Rhs.IsValid();
 }
 
+/**
+ * Compare weak pointers for inequality
+ * @param Other weak pointer to compare to
+ */
 template <typename LhsT, typename RhsT, typename OtherTWeakObjectPtrBase>
 FORCENOINLINE bool operator!=(const TWeakObjectPtr<LhsT, OtherTWeakObjectPtrBase>& Lhs, const TWeakObjectPtr<RhsT, OtherTWeakObjectPtrBase>& Rhs)
 {
@@ -402,6 +421,7 @@ template<class T> struct TIsPODType<TAutoWeakObjectPtr<T> > { enum { Value = tru
 template<class T> struct TIsZeroConstructType<TAutoWeakObjectPtr<T> > { enum { Value = true }; };
 template<class T> struct TIsWeakPointerType<TAutoWeakObjectPtr<T> > { enum { Value = true }; };
 
+/** Utility function to fill in a TArray<ClassName*> from a TArray<TWeakObjectPtr<ClassName>> */
 template<typename DestArrayType, typename SourceArrayType>
 void CopyFromWeakArray(DestArrayType& Dest, const SourceArrayType& Src)
 {
diff --git a/Engine/Source/Runtime/CoreUObject/Public/UObject/PersistentObjectPtr.h b/Engine/Source/Runtime/CoreUObject/Public/UObject/PersistentObjectPtr.h
--- a/Engine/Source/Runtime/CoreUObject/Public/UObject/PersistentObjectPtr.h
+++ b/Engine/Source/Runtime/CoreUObject/Public/UObject/PersistentObjectPtr.h
@@ -116,7 +116,9 @@ public:
 	FORCEINLINE UObject* Get() const
 	{
 		UObject* Object = WeakPtr.Get();
-		if (!Object && TObjectID::GetCurrentTag() != TagAtLastTest && ObjectID.IsValid())
+		
+		// Do a full resolve if the returned object is null and either we think we've loaded new objects, or the weak ptr may be stale
+		if (!Object && ObjectID.IsValid() && (TObjectID::GetCurrentTag() != TagAtLastTest || !WeakPtr.IsExplicitlyNull()))
 		{
 			Object = ObjectID.ResolveObject();
 			WeakPtr = Object;
diff --git a/Engine/Source/Runtime/CoreUObject/Public/UObject/WeakObjectPtr.h b/Engine/Source/Runtime/CoreUObject/Public/UObject/WeakObjectPtr.h
--- a/Engine/Source/Runtime/CoreUObject/Public/UObject/WeakObjectPtr.h
+++ b/Engine/Source/Runtime/CoreUObject/Public/UObject/WeakObjectPtr.h
@@ -9,21 +9,19 @@
 #include "CoreMinimal.h"
 #include "UObject/UObjectArray.h"
 
-/***
- * 
+/**
  * FWeakObjectPtr is a weak pointer to a UObject. 
- * It can return NULL later if the object is garbage collected.
+ * It can return nullptr later if the object is garbage collected.
  * It has no impact on if the object is garbage collected or not.
  * It can't be directly used across a network.
  *
  * Most often it is used when you explicitly do NOT want to prevent something from being garbage collected.
- **/
-
+ */
 struct FWeakObjectPtr
 {
 public:
 
-	/** NULL constructor **/
+	/** Null constructor **/
 	FORCEINLINE FWeakObjectPtr()
 	{
 		Reset();
@@ -31,8 +29,8 @@ public:
 	/**  
 	 * Construct from an object pointer
 	 * @param Object object to create a weak pointer to
-	**/
-	FORCEINLINE FWeakObjectPtr(const class UObject *Object)
+	 */
+	FORCEINLINE FWeakObjectPtr(const class UObject* Object)
 	{
 		(*this)=Object;
 	}
@@ -40,14 +38,14 @@ public:
 	/**  
 	 * Construct from another weak pointer
 	 * @param Other weak pointer to copy from
-	**/
-	FORCEINLINE FWeakObjectPtr(const FWeakObjectPtr &Other)
+	 */
+	FORCEINLINE FWeakObjectPtr(const FWeakObjectPtr& Other)
 	{
 		(*this)=Other;
 	}
 
 	/**
-	 * Reset the weak pointer back to the NULL state
+	 * Reset the weak pointer back to the null state
 	 */
 	FORCEINLINE void Reset()
 	{
@@ -58,24 +56,25 @@ public:
 	/**  
 	 * Copy from an object pointer
 	 * @param Object object to create a weak pointer to
-	**/
-	COREUOBJECT_API void operator=(const class UObject *Object);
+	 */
+	COREUOBJECT_API void operator=(const class UObject* Object);
 
 	/**  
 	 * Construct from another weak pointer
 	 * @param Other weak pointer to copy from
-	**/
-	FORCEINLINE void operator=(const FWeakObjectPtr &Other)
+	 */
+	FORCEINLINE void operator=(const FWeakObjectPtr& Other)
 	{
 		ObjectIndex = Other.ObjectIndex;
 		ObjectSerialNumber = Other.ObjectSerialNumber;
 	}
 
 	/**  
-	 * Compare weak pointers for equality
+	 * Compare weak pointers for equality.
+	 * If both pointers would return nullptr from Get() they count as equal even if they were not initialized to the same object.
 	 * @param Other weak pointer to compare to
-	**/
-	FORCEINLINE bool operator==(const FWeakObjectPtr &Other) const
+	 */
+	FORCEINLINE bool operator==(const FWeakObjectPtr& Other) const
 	{
 		return 
 			(ObjectIndex == Other.ObjectIndex && ObjectSerialNumber == Other.ObjectSerialNumber) ||
@@ -85,14 +84,18 @@ public:
 	/**  
 	 * Compare weak pointers for inequality
 	 * @param Other weak pointer to compare to
-	**/
-	FORCEINLINE bool operator!=(const FWeakObjectPtr &Other) const
+	 */
+	FORCEINLINE bool operator!=(const FWeakObjectPtr& Other) const
 	{
 		return 
 			(ObjectIndex != Other.ObjectIndex || ObjectSerialNumber != Other.ObjectSerialNumber) &&
 			(IsValid() || Other.IsValid());
 	}
 
+	/**
+	 * Returns true if two weak pointers were originally set to the same object, even if they are now stale
+	 * @param Other weak pointer to compare to
+	 */
 	FORCEINLINE bool HasSameIndexAndSerialNumber(const FWeakObjectPtr& Other) const
 	{
 		return ObjectIndex == Other.ObjectIndex && ObjectSerialNumber == Other.ObjectSerialNumber;
@@ -101,26 +104,26 @@ public:
 	/**  
 	 * Dereference the weak pointer.
 	 * @param bEvenIfPendingKill, if this is true, pendingkill objects are considered valid
-	 * @return NULL if this object is gone or the weak pointer was NULL, otherwise a UObject pointer
+	 * @return nullptr if this object is gone or the weak pointer is explicitly null, otherwise a valid uobject pointer
 	 */
-	COREUOBJECT_API class UObject *Get(bool bEvenIfPendingKill) const;
+	COREUOBJECT_API class UObject* Get(bool bEvenIfPendingKill) const;
 
 	/**  
 	 * Dereference the weak pointer. This is an optimized version implying bEvenIfPendingKill=false.
-	 * @return NULL if this object is gone or the weak pointer was NULL, otherwise a UObject pointer
+	 * @return nullptr if this object is gone or the weak pointer is explicitly null, otherwise a valid uobject pointer
 	 */
-	COREUOBJECT_API class UObject *Get(/*bool bEvenIfPendingKill = false*/) const;
+	COREUOBJECT_API class UObject* Get(/*bool bEvenIfPendingKill = false*/) const;
 
 	/** Dereference the weak pointer even if it is RF_PendingKill or RF_Unreachable */
-	COREUOBJECT_API class UObject *GetEvenIfUnreachable() const;
+	COREUOBJECT_API class UObject* GetEvenIfUnreachable() const;
 
 	/**  
 	 * Test if this points to a live UObject
 	 * @param bEvenIfPendingKill, if this is true, pendingkill are not considered invalid
 	 * @param bThreadsafeTest, if true then function will just give you information whether referenced 
-	 *							UObject is gone forever (@return false) or if it is still there (@return true, no object flags checked).
-	 * @return true if Get() would return a valid non-null pointer, if bThreadsafeTest == true then @see @param bThreadsafeTest
-	**/
+	 *							UObject is gone forever (return false) or if it is still there (return true, no object flags checked).
+	 * @return true if Get() would return a valid non-null pointer
+	 */
 	COREUOBJECT_API bool IsValid(bool bEvenIfPendingKill, bool bThreadsafeTest = false) const;
 
 	/**
@@ -134,9 +137,18 @@ public:
 	 * @param bIncludingIfPendingKill, if this is false, pendingkill objects are not considered stale
 	 * @param bThreadsafeTest, set it to true when testing outside of Game Thread. Results in false if WeakObjPtr point to an existing object (no flags checked)
 	 * @return true if this used to point at a real object but no longer does.
-	**/
+	 */
 	COREUOBJECT_API bool IsStale(bool bIncludingIfPendingKill = true, bool bThreadsafeTest = false) const;
 
+	/**
+	 * Returns true if this pointer was explicitly assigned to null, was reset, or was never initialized.
+	 * If this returns true, IsValid() and IsStale() will both return false.
+	 */
+	FORCEINLINE bool IsExplicitlyNull() const
+	{
+		return ObjectIndex == INDEX_NONE;
+	}
+
 	/** Hash function. */
 	friend uint32 GetTypeHash(const FWeakObjectPtr& WeakObjectPtr)
 	{
@@ -164,7 +176,7 @@ private:
 	/**  
 	 * internal function to test for serial number matches
 	 * @return true if the serial number in this matches the central table
-	**/
+	 */
 	FORCEINLINE_DEBUGGABLE bool SerialNumbersMatch() const
 	{
 		checkSlow(ObjectSerialNumber > FUObjectArray::START_SERIAL_NUMBER && ObjectIndex >= 0); // otherwise this is a corrupted weak pointer
