Allow a different heading to be specified for the dumping of a stack trace to log

diff --git a/Engine/Source/Runtime/Core/Private/Misc/AssertionMacros.cpp b/Engine/Source/Runtime/Core/Private/Misc/AssertionMacros.cpp
--- a/Engine/Source/Runtime/Core/Private/Misc/AssertionMacros.cpp
+++ b/Engine/Source/Runtime/Core/Private/Misc/AssertionMacros.cpp
@@ -472,7 +472,12 @@ FORCENOINLINE void VARARGS LowLevelFatalErrorHandler(const ANSICHAR* File, int32
 	StaticFailDebug(TEXT("LowLevelFatalError"), File, Line, DescriptionString, false, NumStackFramesToIgnore);
 }
 
-FORCENOINLINE void FDebug::DumpStackTraceToLog()
+void FDebug::DumpStackTraceToLog()
+{
+	DumpStackTraceToLog(TEXT("=== FDebug::DumpStackTrace(): ==="));
+}
+
+FORCENOINLINE void FDebug::DumpStackTraceToLog(const TCHAR* Heading)
 {
 #if !NO_LOGGING
 	// Walk the stack and dump it to the allocated memory.
@@ -492,7 +497,7 @@ FORCENOINLINE void FDebug::DumpStackTraceToLog()
 
 	// Dump the error and flush the log.
 	// ELogVerbosity::Error to make sure it gets printed in log for conveniency.
-	FDebug::LogFormattedMessageWithCallstack(LogOutputDevice.GetCategoryName(), __FILE__, __LINE__, TEXT("=== FDebug::DumpStackTrace(): ==="), ANSI_TO_TCHAR(StackTrace), ELogVerbosity::Error);
+	FDebug::LogFormattedMessageWithCallstack(LogOutputDevice.GetCategoryName(), __FILE__, __LINE__, Heading, ANSI_TO_TCHAR(StackTrace), ELogVerbosity::Error);
 	GLog->Flush();
 	FMemory::SystemFree(StackTrace);
 #endif
diff --git a/Engine/Source/Runtime/Core/Public/Misc/AssertionMacros.h b/Engine/Source/Runtime/Core/Public/Misc/AssertionMacros.h
--- a/Engine/Source/Runtime/Core/Public/Misc/AssertionMacros.h
+++ b/Engine/Source/Runtime/Core/Public/Misc/AssertionMacros.h
@@ -47,6 +47,9 @@ struct CORE_API FDebug
 	/** Dumps the stack trace into the log, meant to be used for debugging purposes. */
 	static void DumpStackTraceToLog();
 
+	/** Dumps the stack trace into the log with a custom heading, meant to be used for debugging purposes. */
+	static void DumpStackTraceToLog(const TCHAR* Heading);
+
 #if DO_CHECK || DO_GUARD_SLOW
 private:
 	static void VARARGS CheckVerifyFailedImpl(const ANSICHAR* Expr, const char* File, int32 Line, const TCHAR* Format, ...);
