Fix build errors on Windows

Need to avoid marking a static variable dllexport in an interface of a dllexported class

diff --git a/Engine/Source/Runtime/Core/Public/HAL/ThreadSingleton.h b/Engine/Source/Runtime/Core/Public/HAL/ThreadSingleton.h
--- a/Engine/Source/Runtime/Core/Public/HAL/ThreadSingleton.h
+++ b/Engine/Source/Runtime/Core/Public/HAL/ThreadSingleton.h
@@ -35,14 +35,25 @@ public:
 template < class T >
 class TThreadSingleton : public FTlsAutoCleanup
 {
+#if PLATFORM_UNIX || PLATFORM_APPLE
 	/**
 	 * @return TLS slot that holds a TThreadSingleton.
 	 */
-	static CORE_API uint32& GetTlsSlot()
+	CORE_API static uint32& GetTlsSlot()
 	{
 		static uint32 TlsSlot = 0xFFFFFFFF;
 		return TlsSlot;
 	}
+#else
+	/**
+	 * @return TLS slot that holds a TThreadSingleton.
+	 */
+	static uint32& GetTlsSlot()
+	{
+		static uint32 TlsSlot = 0xFFFFFFFF;
+		return TlsSlot;
+	}
+#endif
 
 protected:
 
@@ -79,4 +90,4 @@ public:
 	{
 		return (T*)FThreadSingletonInitializer::TryGet( T::GetTlsSlot() );
 	}
-};
+};
\ No newline at end of file
diff --git a/Engine/Source/Runtime/Slate/Public/Framework/Commands/Commands.h b/Engine/Source/Runtime/Slate/Public/Framework/Commands/Commands.h
--- a/Engine/Source/Runtime/Slate/Public/Framework/Commands/Commands.h
+++ b/Engine/Source/Runtime/Slate/Public/Framework/Commands/Commands.h
@@ -106,7 +106,11 @@ protected:
 	}
 
 	/** A static instance of the command set. */
+#if PLATFORM_UNIX || PLATFORM_APPLE
 	static SLATE_API TWeakPtr< CommandContextType > Instance;
+#else
+	static TWeakPtr< CommandContextType > Instance;
+#endif
 
 	/** Pure virtual to override; describe and instantiate the commands in here by using the UI COMMAND macro. */
 	virtual void RegisterCommands() = 0;
