Add support for VTABLE export on Android

diff --git a/Engine/Source/Runtime/Core/Public/Android/AndroidPlatform.h b/Engine/Source/Runtime/Core/Public/Android/AndroidPlatform.h
--- a/Engine/Source/Runtime/Core/Public/Android/AndroidPlatform.h
+++ b/Engine/Source/Runtime/Core/Public/Android/AndroidPlatform.h
@@ -115,6 +115,12 @@ typedef FAndroidTypes FPlatformTypes;
 
 #define ABSTRACT abstract
 
+// DLL export and import for types, only supported on clang
+#if (__clang_major__ > 3 || (__clang_major__ == 3 && __clang_minor__ >= 8))
+#define DLLEXPORT_VTABLE	__attribute__ ((__type_visibility__("default")))
+#define DLLIMPORT_VTABLE	__attribute__ ((__type_visibility__("default")))
+#endif
+
 // DLL export and import definitions
 #define DLLEXPORT			__attribute__((visibility("default")))
 #define DLLIMPORT			__attribute__((visibility("default")))
