Added support for ptrdiff_t (%td).
Added support for PRIxPTR, PRIXPTR.
Fixed bug that would clamp %lf to just 30 characters.
Removed FmtBuf clearing.

diff --git a/Engine/Source/Runtime/Core/Private/GenericPlatform/GenericWidePlatformString.cpp b/Engine/Source/Runtime/Core/Private/GenericPlatform/GenericWidePlatformString.cpp
--- a/Engine/Source/Runtime/Core/Private/GenericPlatform/GenericWidePlatformString.cpp
+++ b/Engine/Source/Runtime/Core/Private/GenericPlatform/GenericWidePlatformString.cpp
@@ -5,6 +5,7 @@
 #include "Templates/UnrealTemplate.h"
 #include "Logging/LogCategory.h"
 #include "Logging/LogMacros.h"
+#include <inttypes.h>
 
 #if PLATFORM_TCHAR_IS_CHAR16
 
@@ -240,6 +241,18 @@ void RunGetVarArgsTests()
 
 	TestGetVarArgs(OutputString, TEXT("Test F|%f|%e|%g|"), 123.456, 123.456, 123.456);
 	check(FString(OutputString) == FString(TEXT("Test F|123.456000|1.234560e+02|123.456|")));
+
+	TestGetVarArgs(OutputString, TEXT("Test G|%" PRIXPTR "|"), UPTRINT(49374));
+	check(FString(OutputString) == FString(TEXT("Test G|C0DE|")));
+
+	TestGetVarArgs(OutputString, TEXT("Test H|%" PRIxPTR "|"), UPTRINT(49374));
+	check(FString(OutputString) == FString(TEXT("Test H|c0de|")));
+
+	TestGetVarArgs(OutputString, TEXT("Test I|%" PRIu64 "|"), MAX_uint64);
+	check(FString(OutputString) == FString(TEXT("Test I|18446744073709551615|")));
+
+	TestGetVarArgs(OutputString, TEXT("Test J|%td|"), ptrdiff_t(-1));
+	check(FString(OutputString) == FString(TEXT("Test J|-1|")));
 }
 #endif
 
@@ -329,6 +342,11 @@ namespace
 		WIDECHAR* Ptr;
 		WIDECHAR* EndMinusOne; // when null, this means the iterator has already moved past the writable area of the buffer
 	};
+
+	static inline bool CharIsIntegerFormatSpecifier(TCHAR Char)
+	{
+		return (Char == 'i') | (Char == 'd') | (Char == 'u') | (Char == 'X') | (Char == 'x');
+	}
 }
 
 int32 FGenericWidePlatformString::GetVarArgs( WIDECHAR* Dest, SIZE_T DestSize, const WIDECHAR*& Fmt, va_list ArgPtr )
@@ -457,7 +475,7 @@ int32 FGenericWidePlatformString::GetVarArgs( WIDECHAR* Dest, SIZE_T DestSize, c
 				Src++;
 				int Val = va_arg(ArgPtr, int);
 				ANSICHAR AnsiNum[64];
-				ANSICHAR FmtBuf[30] = {0};
+				ANSICHAR FmtBuf[30];
 
 				// limit width to the buffer size
 				FieldLen = FMath::Min(static_cast<int>(sizeof(AnsiNum)) - 1, FieldLen);
@@ -506,7 +524,55 @@ int32 FGenericWidePlatformString::GetVarArgs( WIDECHAR* Dest, SIZE_T DestSize, c
 				size_t Val = va_arg(ArgPtr, size_t);
 
 				ANSICHAR AnsiNum[64];
-				ANSICHAR FmtBuf[30] = {0};
+				ANSICHAR FmtBuf[30];
+
+				// limit width to the buffer size
+				FieldLen = FMath::Min(static_cast<int>(sizeof(AnsiNum)) - 1, FieldLen);
+
+				// Yes, this is lame.
+				int CpyIdx = 0;
+				while (Percent < Src && CpyIdx < ARRAY_COUNT(FmtBuf))
+				{
+					if (LIKELY(*Percent != '*'))
+					{
+						FmtBuf[CpyIdx] = (ANSICHAR) *Percent;
+						++CpyIdx;
+					}
+					else
+					{
+						snprintf(&FmtBuf[CpyIdx], sizeof(FmtBuf) - CpyIdx, "%d", FieldLen);
+						while (CpyIdx < ARRAY_COUNT(FmtBuf) && FmtBuf[CpyIdx] != 0)
+						{
+							++CpyIdx;
+						}
+					}
+					++Percent;
+				}
+				FmtBuf[CpyIdx] = 0;
+
+				int RetCnt = snprintf(AnsiNum, sizeof(AnsiNum), FmtBuf, Val);
+				if (!DestIter.Write(AnsiNum, RetCnt))
+				{
+					return -1;
+				}
+				break;
+			}
+
+			case 't':
+			{
+				if (Src[1] == 0)
+				{
+					printf("Unknown percent [%lc] in FGenericWidePlatformString::GetVarArgs() [%s]\n.", Src[0], TCHAR_TO_ANSI(Fmt));
+					Src++;  // skip it, I guess.
+					break;
+				}
+
+				Src += 2;
+
+				ptrdiff_t Val = va_arg(ArgPtr, ptrdiff_t);
+
+				ANSICHAR AnsiNum[64];
+				ANSICHAR FmtBuf[30];
 
 				// limit width to the buffer size
 				FieldLen = FMath::Min(static_cast<int>(sizeof(AnsiNum)) - 1, FieldLen);
@@ -578,11 +644,11 @@ int32 FGenericWidePlatformString::GetVarArgs( WIDECHAR* Dest, SIZE_T DestSize, c
 					break;
 				}
 
-				// treat %ld as %d. Also shorts for %h will be promoted to ints
-				if ((Src[0] == 'l' && Src[1] == 'd') || Src[0] == 'h')
+				// treat %ld as %d. Also shorts for %h will be promoted to ints. This path also handles i, u, x and X
+				if ((Src[0] == 'l' && CharIsIntegerFormatSpecifier(Src[1])) || Src[0] == 'h')
 				{
 					Src+=2;
-					int Val = va_arg(ArgPtr, int);
+					long int Val = va_arg(ArgPtr, long int);
 					ANSICHAR AnsiNum[30];
 					ANSICHAR FmtBuf[30];
 
@@ -608,7 +674,7 @@ int32 FGenericWidePlatformString::GetVarArgs( WIDECHAR* Dest, SIZE_T DestSize, c
 				{
 					Src += 2;
 					double Val = va_arg(ArgPtr, double);
-					ANSICHAR AnsiNum[30];
+					ANSICHAR AnsiNum[384];
 					ANSICHAR FmtBuf[30];
 
 					// Yes, this is lame.
@@ -629,64 +695,46 @@ int32 FGenericWidePlatformString::GetVarArgs( WIDECHAR* Dest, SIZE_T DestSize, c
 					break;
 				}
 
+				// At this point we expect RemainingSize >= 3, but we know it's at least 2.
+				if (RemainingSize < 3)
+				{
+					printf("Unknown percent [%lc%lc] in FGenericWidePlatformString::GetVarArgs() [%s]\n.", Src[0], Src[1], TCHAR_TO_ANSI(Fmt));
+					Src++;  // skip it, I guess.
+					break;
+				}
+
 				if (Src[0] == 'l')
 				{
-					if (Src[1] != 'l' && Src[1] != 'u' && Src[1] != 'x')
+					if (Src[1] != 'l' || !CharIsIntegerFormatSpecifier(Src[2]))
 					{
-						printf("Unknown percent [%lc%lc] in FGenericWidePlatformString::GetVarArgs() [%s]\n.", Src[0], Src[1], TCHAR_TO_ANSI(Fmt));
-						Src++;  // skip it, I guess.
-						break;
-					}
-					else if (Src[1] == 'l' && RemainingSize == 2)
-					{
-						printf("Unknown percent [%lc%lc] in FGenericWidePlatformString::GetVarArgs() [%s]\n.", Src[0], Src[1], TCHAR_TO_ANSI(Fmt));
+						printf("Unknown percent [%lc%lc%lc] in FGenericWidePlatformString::GetVarArgs() [%s]\n.", Src[0], Src[1], Src[2], TCHAR_TO_ANSI(Fmt));
 						Src++;  // skip it, I guess.
 						break;
 					}
 				}
 				else if (Src[0] == 'I')
 				{
-					if (RemainingSize == 2)
-					{
-						printf("Unknown percent [%lc%lc] in FGenericWidePlatformString::GetVarArgs() [%s]\n.", Src[0], Src[1], TCHAR_TO_ANSI(Fmt));
-						Src++;  // skip it, I guess.
-						break;
-					}
-					else if (RemainingSize == 3)
-					{
-						printf("Unknown percent [%lc%lc%lc] in FGenericWidePlatformString::GetVarArgs() [%s]\n.", Src[0], Src[1], Src[2], TCHAR_TO_ANSI(Fmt));
-						Src++;  // skip it, I guess.
-						break;
-					}
-					else if (Src[1] != '6' || Src[2] != '4')
+					if (Src[1] != '6' || Src[2] != '4')
 					{
 						printf("Unknown percent [%lc%lc%lc] in FGenericWidePlatformString::GetVarArgs() [%s]\n.", Src[0], Src[1], Src[2], TCHAR_TO_ANSI(Fmt));
 						Src++;  // skip it, I guess.
 						break;
 					}
 				}
+				else
+				{
+					printf("Unknown percent [%lc%lc%lc] in FGenericWidePlatformString::GetVarArgs() [%s]\n.", Src[0], Src[1], Src[2], TCHAR_TO_ANSI(Fmt));
+					Src++;  // skip it, I guess.
+					break;
+				}
 
 				// Yes, this is lame.
-				int CpyIdx = 0;
 				unsigned long long Val = va_arg(ArgPtr, unsigned long long);
 				ANSICHAR AnsiNum[60];
 				ANSICHAR FmtBuf[30];
-				if (Src[0] == 'l' && Src[1] == 'l')
-				{
-					Src += 3;
-				}
-				else if (Src[0] == 'l')
-				{
-					Src += 2;
-				}
-				else
-				{
-					Src += 4;
-					strcpy(FmtBuf, "%L");
-					Percent += 4;
-					CpyIdx = 2;
-				}
 
+				Src += 4;
+				int CpyIdx = 0;
 				while (Percent < Src && CpyIdx < ARRAY_COUNT(FmtBuf))
 				{
 					FmtBuf[CpyIdx] = (ANSICHAR) *Percent;
