Added missing volatile keyword to InterlockedExhangePtr and InterlockedCompareExchangePointer.
Fixed incorrect atomic failures log output.
Fixed alignment checks in Interlocked implementations that have them.
Use InterlockedExchangePointer on Windows instead of figuring out pointer size and call other functions.

diff --git a/Engine/Source/Runtime/Core/Public/Clang/ClangPlatformAtomics.h b/Engine/Source/Runtime/Core/Public/Clang/ClangPlatformAtomics.h
--- a/Engine/Source/Runtime/Core/Public/Clang/ClangPlatformAtomics.h
+++ b/Engine/Source/Runtime/Core/Public/Clang/ClangPlatformAtomics.h
@@ -94,7 +94,7 @@ struct CORE_API FClangPlatformAtomics : public FGenericPlatformAtomics
 		return __sync_lock_test_and_set(Value, Exchange);
 	}
 
-	static FORCEINLINE void* InterlockedExchangePtr(void** Dest, void* Exchange)
+	static FORCEINLINE void* InterlockedExchangePtr(void*volatile* Dest, void* Exchange)
 	{
 		return __sync_lock_test_and_set(Dest, Exchange);
 	}
@@ -221,7 +221,7 @@ struct CORE_API FClangPlatformAtomics : public FGenericPlatformAtomics
 		return InterlockedCompareExchange((volatile int64*)Src, 0, 0);
 	}
 
-	static FORCEINLINE void* InterlockedCompareExchangePointer(void** Dest, void* Exchange, void* Comperand)
+	static FORCEINLINE void* InterlockedCompareExchangePointer(void*volatile* Dest, void* Exchange, void* Comperand)
 	{
 		return __sync_val_compare_and_swap(Dest, Comperand, Exchange);
 	}
diff --git a/Engine/Source/Runtime/Core/Public/GenericPlatform/GenericPlatformAtomics.h b/Engine/Source/Runtime/Core/Public/GenericPlatform/GenericPlatformAtomics.h
--- a/Engine/Source/Runtime/Core/Public/GenericPlatform/GenericPlatformAtomics.h
+++ b/Engine/Source/Runtime/Core/Public/GenericPlatform/GenericPlatformAtomics.h
@@ -141,7 +141,7 @@ struct FGenericPlatformAtomics
 	/**
 	 * Atomically swaps two pointers returning the original pointer to the caller
 	 */
-	static FORCEINLINE void* InterlockedExchangePtr( void** Dest, void* Exchange )
+	static FORCEINLINE void* InterlockedExchangePtr( void*volatile* Dest, void* Exchange )
 	{
 #if PLATFORM_64BITS
 		#error must implement
@@ -232,7 +232,7 @@ struct FGenericPlatformAtomics
 	 * Atomically compares the pointer to comparand and replaces with the exchange
 	 * pointer if they are equal and returns the original value
 	 */
-	static FORCEINLINE void* InterlockedCompareExchangePointer(void** Dest,void* Exchange,void* Comperand)
+	static FORCEINLINE void* InterlockedCompareExchangePointer(void*volatile* Dest,void* Exchange,void* Comperand)
 	{
 #if PLATFORM_64BITS
 		#error must implement
diff --git a/Engine/Source/Runtime/Core/Public/HTML5/HTML5PlatformAtomics.h b/Engine/Source/Runtime/Core/Public/HTML5/HTML5PlatformAtomics.h
--- a/Engine/Source/Runtime/Core/Public/HTML5/HTML5PlatformAtomics.h
+++ b/Engine/Source/Runtime/Core/Public/HTML5/HTML5PlatformAtomics.h
@@ -105,7 +105,7 @@ struct CORE_API FHTML5PlatformAtomics	: public FGenericPlatformAtomics
 //		return __sync_lock_test_and_set(Value, Exchange);
 	}
 
-	static FORCEINLINE void* InterlockedExchangePtr(void** Dest, void* Exchange)
+	static FORCEINLINE void* InterlockedExchangePtr(void*volatile* Dest, void* Exchange)
 	{
 		return __sync_lock_test_and_set(Dest, Exchange);
 	}
@@ -227,7 +227,7 @@ struct CORE_API FHTML5PlatformAtomics	: public FGenericPlatformAtomics
 		__atomic_store((volatile int64*)Src, &Val, __ATOMIC_RELAXED);
 	}
 
-	static FORCEINLINE void* InterlockedCompareExchangePointer(void** Dest, void* Exchange, void* Comperand)
+	static FORCEINLINE void* InterlockedCompareExchangePointer(void*volatile* Dest, void* Exchange, void* Comperand)
 	{
 		return __sync_val_compare_and_swap(Dest, Comperand, Exchange);
 	}
@@ -347,7 +347,7 @@ struct CORE_API FHTML5PlatformAtomics	: public FGenericPlatformAtomics
 		return Result;
 	}
 
-	static FORCEINLINE void* InterlockedExchangePtr( void** Dest, void* Exchange )
+	static FORCEINLINE void* InterlockedExchangePtr( void*volatile* Dest, void* Exchange )
 	{
 		void* Result = *Dest;
 		*Dest = Exchange;
@@ -474,7 +474,7 @@ struct CORE_API FHTML5PlatformAtomics	: public FGenericPlatformAtomics
 		*Src = Val;
 	}
 
-	static FORCEINLINE void* InterlockedCompareExchangePointer( void** Dest, void* Exchange, void* Comperand )
+	static FORCEINLINE void* InterlockedCompareExchangePointer( void*volatile* Dest, void* Exchange, void* Comperand )
 	{
 		void* Result = *Dest;
 		if (Result == Comperand)
diff --git a/Engine/Source/Runtime/Core/Public/Windows/WindowsPlatformAtomics.h b/Engine/Source/Runtime/Core/Public/Windows/WindowsPlatformAtomics.h
--- a/Engine/Source/Runtime/Core/Public/Windows/WindowsPlatformAtomics.h
+++ b/Engine/Source/Runtime/Core/Public/Windows/WindowsPlatformAtomics.h
@@ -146,20 +146,16 @@ struct CORE_API FWindowsPlatformAtomics
 		#endif
 	}
 
-	static FORCEINLINE void* InterlockedExchangePtr( void** Dest, void* Exchange )
+	static FORCEINLINE void* InterlockedExchangePtr( void*volatile* Dest, void* Exchange )
 	{
 		#if !(UE_BUILD_SHIPPING || UE_BUILD_TEST) 
-			if (IsAligned(Dest) == false)
+			if (IsAligned(Dest, alignof(void*)) == false)
 			{
-				HandleAtomicsFailure(TEXT("InterlockedExchangePointer requires Dest pointer to be aligned to %d bytes"), sizeof(void*));
+				HandleAtomicsFailure(TEXT("InterlockedExchangePointer requires Dest pointer to be aligned to %d bytes"), (int)alignof(void*));
 			}
 		#endif
 
-		#if PLATFORM_64BITS
-			return (void*)::_InterlockedExchange64((int64*)(Dest), (int64)(Exchange));
-		#else
-			return (void*)::_InterlockedExchange((long*)(Dest), (long)(Exchange));
-		#endif
+		return ::_InterlockedExchangePointer(Dest, Exchange);
 	}
 
 	static FORCEINLINE int8 InterlockedCompareExchange( volatile int8* Dest, int8 Exchange, int8 Comparand )
@@ -180,9 +176,9 @@ struct CORE_API FWindowsPlatformAtomics
 	static FORCEINLINE int64 InterlockedCompareExchange( volatile int64* Dest, int64 Exchange, int64 Comparand )
 	{
 		#if !(UE_BUILD_SHIPPING || UE_BUILD_TEST)
-			if (IsAligned(Dest) == false)
+			if (IsAligned(Dest, alignof(int64)) == false)
 			{
-				HandleAtomicsFailure(TEXT("InterlockedCompareExchangePointer requires Dest pointer to be aligned to %d bytes"), sizeof(void*));
+				HandleAtomicsFailure(TEXT("InterlockedCompareExchange int64 requires Dest pointer to be aligned to %d bytes"), (int)alignof(int64));
 			}
 		#endif
 
@@ -299,11 +295,11 @@ struct CORE_API FWindowsPlatformAtomics
 #if !(UE_BUILD_SHIPPING || UE_BUILD_TEST)
 		if (IsAligned(Dest,16) == false)
 		{
-			HandleAtomicsFailure(TEXT("InterlockedCompareExchangePointer requires Dest pointer to be aligned to 16 bytes") );
+			HandleAtomicsFailure(TEXT("InterlockedCompareExchange128 requires Dest pointer to be aligned to 16 bytes") );
 		}
 		if (IsAligned(Comparand,16) == false)
 		{
-			HandleAtomicsFailure(TEXT("InterlockedCompareExchangePointer requires Comparand pointer to be aligned to 16 bytes") );
+			HandleAtomicsFailure(TEXT("InterlockedCompareExchange128 requires Comparand pointer to be aligned to 16 bytes") );
 		}
 #endif
 
@@ -323,20 +319,16 @@ struct CORE_API FWindowsPlatformAtomics
 
 #endif // PLATFORM_HAS_128BIT_ATOMICS
 
-	static FORCEINLINE void* InterlockedCompareExchangePointer( void** Dest, void* Exchange, void* Comparand )
+	static FORCEINLINE void* InterlockedCompareExchangePointer( void*volatile* Dest, void* Exchange, void* Comparand )
 	{
 		#if !(UE_BUILD_SHIPPING || UE_BUILD_TEST)
-			if (IsAligned(Dest) == false)
+			if (IsAligned(Dest, alignof(void*)) == false)
 			{
-				HandleAtomicsFailure(TEXT("InterlockedCompareExchangePointer requires Dest pointer to be aligned to %d bytes"), sizeof(void*));
+				HandleAtomicsFailure(TEXT("InterlockedCompareExchangePointer requires Dest pointer to be aligned to %d bytes"), (int)alignof(void*));
 			}
 		#endif
 
-		#if PLATFORM_64BITS
-			return (void*)::_InterlockedCompareExchange64((int64*)Dest, (int64)Exchange, (int64)Comparand);
-		#else
-			return (void*)::_InterlockedCompareExchange((long*)Dest, (long)Exchange, (long)Comparand);
-		#endif
+		return ::_InterlockedCompareExchangePointer(Dest, Exchange, Comparand);
 	}
 
 	/**
