Fixed bug in Win32 implementation of InterlockedAdd on int64.

It returned the new value instead of the previous value.
System.Core.Misc.Atomic test now succeeds on Win32.

diff --git a/Engine/Source/Runtime/Core/Public/Windows/WindowsPlatformAtomics.h b/Engine/Source/Runtime/Core/Public/Windows/WindowsPlatformAtomics.h
--- a/Engine/Source/Runtime/Core/Public/Windows/WindowsPlatformAtomics.h
+++ b/Engine/Source/Runtime/Core/Public/Windows/WindowsPlatformAtomics.h
@@ -108,7 +108,7 @@ struct CORE_API FWindowsPlatformAtomics
 				int64 OldValue = *Value;
 				if (_InterlockedCompareExchange64(Value, OldValue + Amount, OldValue) == OldValue)
 				{
-					return OldValue + Amount;
+					return OldValue;
 				}
 			}
 		#endif
