Added InterlockedAnd, InterlockedOr and InterlockedXor to FPlatformAtomics.

HTML5 implementations of Increment, Decrement and Add is now using unsigned arithmetics.

Implemented tests for InterlockedAnd, InterlockedOr, InterlockedXor, InterlockedAdd, InterlockedIncrement, InterlockedDecrement, InterlockedExchange, InterlockedExchangePtr, InterlockedCompareExchange.

diff --git a/Engine/Source/Runtime/Core/Private/Tests/HAL/PlatformAtomicsTest.cpp b/Engine/Source/Runtime/Core/Private/Tests/HAL/PlatformAtomicsTest.cpp
new file mode 100644
--- /dev/null
+++ b/Engine/Source/Runtime/Core/Private/Tests/HAL/PlatformAtomicsTest.cpp
@@ -0,0 +1,505 @@
+// Copyright 1998-2019 Epic Games, Inc. All Rights Reserved.
+
+#include "CoreTypes.h"
+#include "HAL/PlatformAtomics.h"
+#include "Misc/AssertionMacros.h"
+#include "Misc/AutomationTest.h"
+
+#if WITH_DEV_AUTOMATION_TESTS
+
+IMPLEMENT_SIMPLE_AUTOMATION_TEST(FPlatformAtomicsTest, "System.Core.HAL.PlatformAtomics", EAutomationTestFlags::ApplicationContextMask | EAutomationTestFlags::SmokeFilter)
+
+template<typename T>
+static bool TestInterlocked(volatile T* Dest, T ExpectedReturnValue, T ExpectedFinalValue, TFunctionRef<T(volatile T*)> InterlockedFunc, FAutomationTestBase& Test, const TCHAR* FunctionName, const TCHAR* TypeName)
+{
+	const T ReturnValue = InterlockedFunc(Dest);
+	if (ReturnValue != ExpectedReturnValue)
+	{
+		check(false);
+		Test.AddError(FString::Printf(TEXT("FPlatformAtomics::Interlocked%s on %s failed"), FunctionName, TypeName));
+		return false;
+	}
+
+	if (FPlatformAtomics::AtomicRead(Dest) != ExpectedFinalValue)
+	{
+		check(false);
+		Test.AddError(FString::Printf(TEXT("FPlatformAtomics::Interlocked%s on %s failed"), FunctionName, TypeName));
+		return false;
+	}
+
+	return true;
+}
+
+static bool TestInterlockedAnd(FAutomationTestBase& Test)
+{
+	bool bSuccess = true;
+
+	constexpr TCHAR FunctionName[] = TEXT("And");
+	{
+		// Test and with value where some bits are set and some aren't set in the current value.
+		volatile int8 Value0 = 0x30;
+		bSuccess = TestInterlocked<int8>(&Value0, 0x30, 0x20, [](volatile int8* Dest) { return FPlatformAtomics::InterlockedAnd(Dest, 0x66); }, Test, FunctionName, TEXT("int8")) && bSuccess;
+
+		// Test and with all bits set.
+		volatile int8 Value1 = 0x30;
+		bSuccess = TestInterlocked<int8>(&Value1, 0x30, 0x30, [](volatile int8* Dest) { return FPlatformAtomics::InterlockedAnd(Dest, -1); }, Test, FunctionName, TEXT("int8")) && bSuccess;
+
+		// Test and with zero.
+		volatile int8 Value2 = 0x30;
+		bSuccess = TestInterlocked<int8>(&Value2, 0x30, 0, [](volatile int8* Dest) { return FPlatformAtomics::InterlockedAnd(Dest, 0); }, Test, FunctionName, TEXT("int8")) && bSuccess;
+	}
+
+	{
+		// Test and with value where some bits are set and some aren't set in the current value.
+		volatile int16 Value0 = 0x3030;
+		bSuccess = TestInterlocked<int16>(&Value0, 0x3030, 0x2020, [](volatile int16* Dest) { return FPlatformAtomics::InterlockedAnd(Dest, 0x6666); }, Test, FunctionName, TEXT("int16")) && bSuccess;
+
+		// Test and with all bits set.
+		volatile int16 Value1 = 0x3030;
+		bSuccess = TestInterlocked<int16>(&Value1, 0x3030, 0x3030, [](volatile int16* Dest) { return FPlatformAtomics::InterlockedAnd(Dest, -1); }, Test, FunctionName, TEXT("int16")) && bSuccess;
+
+		// Test and with zero.
+		volatile int16 Value2 = 0x3030;
+		bSuccess = TestInterlocked<int16>(&Value2, 0x3030, 0, [](volatile int16* Dest) { return FPlatformAtomics::InterlockedAnd(Dest, 0); }, Test, FunctionName, TEXT("int16")) && bSuccess;
+	}
+
+	{
+		// Test and with value where some bits are set and some aren't set in the current value.
+		volatile int32 Value0 = 0x30303030;
+		bSuccess = TestInterlocked<int32>(&Value0, 0x30303030, 0x20202020, [](volatile int32* Dest) { return FPlatformAtomics::InterlockedAnd(Dest, 0x66666666); }, Test, FunctionName, TEXT("int32")) && bSuccess;
+
+		// Test and with all bits set.
+		volatile int32 Value1 = 0x30303030;
+		bSuccess = TestInterlocked<int32>(&Value1, 0x30303030, 0x30303030, [](volatile int32* Dest) { return FPlatformAtomics::InterlockedAnd(Dest, -1); }, Test, FunctionName, TEXT("int32")) && bSuccess;
+
+		// Test and with zero.
+		volatile int32 Value2 = 0x30303030;
+		bSuccess = TestInterlocked<int32>(&Value2, 0x30303030, 0, [](volatile int32* Dest) { return FPlatformAtomics::InterlockedAnd(Dest, 0); }, Test, FunctionName, TEXT("int32")) && bSuccess;
+	}
+
+	{
+		// Test and with value where some bits are set and some aren't set in the current value.
+		alignas(8) volatile int64 Value0 = 0x3030303030303030LL;
+		bSuccess = TestInterlocked<int64>(&Value0, 0x3030303030303030LL, 0x2020202020202020LL, [](volatile int64* Dest) { return FPlatformAtomics::InterlockedAnd(Dest, 0x6666666666666666); }, Test, FunctionName, TEXT("int64")) && bSuccess;
+
+		// Test and with all bits set.
+		alignas(8) volatile int64 Value1 = 0x3030303030303030LL;
+		bSuccess = TestInterlocked<int64>(&Value1, 0x3030303030303030LL, 0x3030303030303030LL, [](volatile int64* Dest) { return FPlatformAtomics::InterlockedAnd(Dest, -1); }, Test, FunctionName, TEXT("int64")) && bSuccess;
+
+		// Test and with zero.
+		alignas(8) volatile int64 Value2 = 0x3030303030303030LL;
+		bSuccess = TestInterlocked<int64>(&Value2, 0x3030303030303030LL, 0, [](volatile int64* Dest) { return FPlatformAtomics::InterlockedAnd(Dest, 0); }, Test, FunctionName, TEXT("int64")) && bSuccess;
+	}
+
+	return bSuccess;
+}
+
+static bool TestInterlockedOr(FAutomationTestBase& Test)
+{
+	bool bSuccess = true;
+
+	constexpr TCHAR FunctionName[] = TEXT("Or");
+	{
+		// Test or with value where some bits are set and some aren't set in the current value.
+		volatile int8 Value0 = 0x30;
+		bSuccess = TestInterlocked<int8>(&Value0, 0x30, 0x76, [](volatile int8* Dest) { return FPlatformAtomics::InterlockedOr(Dest, 0x66); }, Test, FunctionName, TEXT("int8")) && bSuccess;
+
+		// Test or with all bits set.
+		volatile int8 Value1 = 0x30;
+		bSuccess = TestInterlocked<int8>(&Value1, 0x30, -1, [](volatile int8* Dest) { return FPlatformAtomics::InterlockedOr(Dest, -1); }, Test, FunctionName, TEXT("int8")) && bSuccess;
+
+		// Test or with zero.
+		volatile int8 Value2 = 0x30;
+		bSuccess = TestInterlocked<int8>(&Value2, 0x30, 0x30, [](volatile int8* Dest) { return FPlatformAtomics::InterlockedOr(Dest, 0); }, Test, FunctionName, TEXT("int8")) && bSuccess;
+	}
+
+	{
+		// Test or with value where some bits are set and some aren't set in the current value.
+		volatile int16 Value0 = 0x3030;
+		bSuccess = TestInterlocked<int16>(&Value0, 0x3030, 0x7676, [](volatile int16* Dest) { return FPlatformAtomics::InterlockedOr(Dest, 0x6666); }, Test, FunctionName, TEXT("int16")) && bSuccess;
+
+		// Test or with all bits set.
+		volatile int16 Value1 = 0x3030;
+		bSuccess = TestInterlocked<int16>(&Value1, 0x3030, -1, [](volatile int16* Dest) { return FPlatformAtomics::InterlockedOr(Dest, -1); }, Test, FunctionName, TEXT("int16")) && bSuccess;
+
+		// Test or with zero.
+		volatile int16 Value2 = 0x3030;
+		bSuccess = TestInterlocked<int16>(&Value2, 0x3030, 0x3030, [](volatile int16* Dest) { return FPlatformAtomics::InterlockedOr(Dest, 0); }, Test, FunctionName, TEXT("int16")) && bSuccess;
+	}
+
+	{
+		// Test or with value where some bits are set and some aren't set in the current value.
+		volatile int32 Value0 = 0x30303030;
+		bSuccess = TestInterlocked<int32>(&Value0, 0x30303030, 0x76767676, [](volatile int32* Dest) { return FPlatformAtomics::InterlockedOr(Dest, 0x66666666); }, Test, FunctionName, TEXT("int32")) && bSuccess;
+
+		// Test or with all bits set.
+		volatile int32 Value1 = 0x30303030;
+		bSuccess = TestInterlocked<int32>(&Value1, 0x30303030, -1, [](volatile int32* Dest) { return FPlatformAtomics::InterlockedOr(Dest, -1); }, Test, FunctionName, TEXT("int32")) && bSuccess;
+
+		// Test or with zero.
+		volatile int32 Value2 = 0x30303030;
+		bSuccess = TestInterlocked<int32>(&Value2, 0x30303030, 0x30303030, [](volatile int32* Dest) { return FPlatformAtomics::InterlockedOr(Dest, 0); }, Test, FunctionName, TEXT("int32")) && bSuccess;
+	}
+
+	{
+		// Test or with value where some bits are set and some aren't set in the current value.
+		alignas(8) volatile int64 Value0 = 0x3030303030303030LL;
+		bSuccess = TestInterlocked<int64>(&Value0, 0x3030303030303030LL, 0x7676767676767676LL, [](volatile int64* Dest) { return FPlatformAtomics::InterlockedOr(Dest, 0x6666666666666666LL); }, Test, FunctionName, TEXT("int64")) && bSuccess;
+
+		// Test or with all bits set.
+		alignas(8) volatile int64 Value1 = 0x3030303030303030LL;
+		bSuccess = TestInterlocked<int64>(&Value1, 0x3030303030303030LL, -1LL, [](volatile int64* Dest) { return FPlatformAtomics::InterlockedOr(Dest, -1); }, Test, FunctionName, TEXT("int64")) && bSuccess;
+
+		// Test or with zero.
+		alignas(8) volatile int64 Value2 = 0x3030303030303030LL;
+		bSuccess = TestInterlocked<int64>(&Value2, 0x3030303030303030LL, 0x3030303030303030LL, [](volatile int64* Dest) { return FPlatformAtomics::InterlockedOr(Dest, 0); }, Test, FunctionName, TEXT("int64")) && bSuccess;
+	}
+
+	return bSuccess;
+}
+
+static bool TestInterlockedXor(FAutomationTestBase& Test)
+{
+	bool bSuccess = true;
+
+	constexpr TCHAR FunctionName[] = TEXT("Xor");
+	{
+		// Test xor with value where some bits are set and some aren't set in the current value.
+		volatile int8 Value0 = 0x30;
+		bSuccess = TestInterlocked<int8>(&Value0, 0x30, 0x56, [](volatile int8* Dest) { return FPlatformAtomics::InterlockedXor(Dest, 0x66); }, Test, FunctionName, TEXT("int8")) && bSuccess;
+
+		// Test xor with all bits set.
+		volatile int8 Value1 = 0x30;
+		bSuccess = TestInterlocked<int8>(&Value1, 0x30, int8(~0x30), [](volatile int8* Dest) { return FPlatformAtomics::InterlockedXor(Dest, -1); }, Test, FunctionName, TEXT("int8")) && bSuccess;
+
+		// Test xor with zero.
+		volatile int8 Value2 = 0x30;
+		bSuccess = TestInterlocked<int8>(&Value2, 0x30, 0x30, [](volatile int8* Dest) { return FPlatformAtomics::InterlockedXor(Dest, 0); }, Test, FunctionName, TEXT("int8")) && bSuccess;
+	}
+
+	{
+		// Test xor with value where some bits are set and some aren't set in the current value.
+		volatile int16 Value0 = 0x3030;
+		bSuccess = TestInterlocked<int16>(&Value0, 0x3030, 0x5656, [](volatile int16* Dest) { return FPlatformAtomics::InterlockedXor(Dest, 0x6666); }, Test, FunctionName, TEXT("int16")) && bSuccess;
+
+		// Test xor with all bits set.
+		volatile int16 Value1 = 0x3030;
+		bSuccess = TestInterlocked<int16>(&Value1, 0x3030, int16(~0x3030), [](volatile int16* Dest) { return FPlatformAtomics::InterlockedXor(Dest, -1); }, Test, FunctionName, TEXT("int16")) && bSuccess;
+
+		// Test xor with zero.
+		volatile int16 Value2 = 0x3030;
+		bSuccess = TestInterlocked<int16>(&Value2, 0x3030, 0x3030, [](volatile int16* Dest) { return FPlatformAtomics::InterlockedXor(Dest, 0); }, Test, FunctionName, TEXT("int16")) && bSuccess;
+	}
+
+	{
+		// Test xor with value where some bits are set and some aren't set in the current value.
+		volatile int32 Value0 = 0x30303030;
+		bSuccess = TestInterlocked<int32>(&Value0, 0x30303030, 0x56565656, [](volatile int32* Dest) { return FPlatformAtomics::InterlockedXor(Dest, 0x66666666); }, Test, FunctionName, TEXT("int32")) && bSuccess;
+
+		// Test xor with all bits set.
+		volatile int32 Value1 = 0x30303030;
+		bSuccess = TestInterlocked<int32>(&Value1, 0x30303030, ~0x30303030, [](volatile int32* Dest) { return FPlatformAtomics::InterlockedXor(Dest, -1); }, Test, FunctionName, TEXT("int32")) && bSuccess;
+
+		// Test xor with zero.
+		volatile int32 Value2 = 0x30303030;
+		bSuccess = TestInterlocked<int32>(&Value2, 0x30303030, 0x30303030, [](volatile int32* Dest) { return FPlatformAtomics::InterlockedXor(Dest, 0); }, Test, FunctionName, TEXT("int32")) && bSuccess;
+	}
+
+	{
+		// Test xor with value where some bits are set and some aren't set in the current value.
+		alignas(8) volatile int64 Value0 = 0x3030303030303030LL;
+		bSuccess = TestInterlocked<int64>(&Value0, 0x3030303030303030LL, 0x5656565656565656LL, [](volatile int64* Dest) { return FPlatformAtomics::InterlockedXor(Dest, 0x6666666666666666LL); }, Test, FunctionName, TEXT("int64")) && bSuccess;
+
+		// Test xor with all bits set.
+		alignas(8) volatile int64 Value1 = 0x3030303030303030LL;
+		bSuccess = TestInterlocked<int64>(&Value1, 0x3030303030303030LL, ~0x3030303030303030LL, [](volatile int64* Dest) { return FPlatformAtomics::InterlockedXor(Dest, -1); }, Test, FunctionName, TEXT("int64")) && bSuccess;
+
+		// Test xor with zero.
+		alignas(8) volatile int64 Value2 = 0x3030303030303030LL;
+		bSuccess = TestInterlocked<int64>(&Value2, 0x3030303030303030LL, 0x3030303030303030LL, [](volatile int64* Dest) { return FPlatformAtomics::InterlockedXor(Dest, 0); }, Test, FunctionName, TEXT("int64")) && bSuccess;
+	}
+
+	return bSuccess;
+}
+
+static bool TestInterlockedAdd(FAutomationTestBase& Test)
+{
+	bool bSuccess = true;
+
+	constexpr TCHAR FunctionName[] = TEXT("Add");
+	{
+		// Test adding positive value
+		volatile int8 Value0 = 0x0F;
+		bSuccess = TestInterlocked<int8>(&Value0, 0x0F, 0x11, [](volatile int8* Dest) { return FPlatformAtomics::InterlockedAdd(Dest, +0x02); }, Test, FunctionName, TEXT("int8")) && bSuccess;
+
+		// Test adding negative value
+		volatile int8 Value1 = 0x11;
+		bSuccess = TestInterlocked<int8>(&Value1, 0x11, 0x0F, [](volatile int8* Dest) { return FPlatformAtomics::InterlockedAdd(Dest, -0x02); }, Test, FunctionName, TEXT("int8")) && bSuccess;
+
+		// Test value overflow
+		volatile int8 Value2 = MAX_int8 - 1;
+		bSuccess = TestInterlocked<int8>(&Value2, MAX_int8 - 1, MIN_int8 + 2, [](volatile int8* Dest) { return FPlatformAtomics::InterlockedAdd(Dest, +4); }, Test, FunctionName, TEXT("int8")) && bSuccess;
+
+		// Test value underflow
+		volatile int8 Value3 = MIN_int8 + 2;
+		bSuccess = TestInterlocked<int8>(&Value3, MIN_int8 + 2, MAX_int8 - 1, [](volatile int8* Dest) { return FPlatformAtomics::InterlockedAdd(Dest, -4); }, Test, FunctionName, TEXT("int8")) && bSuccess;
+	}
+
+	{
+		// Test adding positive value
+		volatile int16 Value0 = 0x0F00;
+		bSuccess = TestInterlocked<int16>(&Value0, 0x0F00, 0x1001, [](volatile int16* Dest) { return FPlatformAtomics::InterlockedAdd(Dest, 0x0101); }, Test, FunctionName, TEXT("int16")) && bSuccess;
+
+		// Test adding negative value
+		volatile int16 Value1 = 0x1001;
+		bSuccess = TestInterlocked<int16>(&Value1, 0x1001, 0x0F00, [](volatile int16* Dest) { return FPlatformAtomics::InterlockedAdd(Dest, -0x0101); }, Test, FunctionName, TEXT("int16")) && bSuccess;
+
+		// Test value overflow
+		volatile int16 Value2 = MAX_int16 - 1;
+		bSuccess = TestInterlocked<int16>(&Value2, MAX_int16 - 1, MIN_int16 + 2, [](volatile int16* Dest) { return FPlatformAtomics::InterlockedAdd(Dest, +4); }, Test, FunctionName, TEXT("int16")) && bSuccess;
+
+		// Test value underflow
+		volatile int16 Value3 = MIN_int16 + 2;
+		bSuccess = TestInterlocked<int16>(&Value3, MIN_int16 + 2, MAX_int16 - 1, [](volatile int16* Dest) { return FPlatformAtomics::InterlockedAdd(Dest, -4); }, Test, FunctionName, TEXT("int16")) && bSuccess;
+	}
+
+	{
+		// Test adding positive value
+		volatile int32 Value0 = 0x0F000000;
+		bSuccess = TestInterlocked<int32>(&Value0, 0x0F000000, 0x10010101, [](volatile int32* Dest) { return FPlatformAtomics::InterlockedAdd(Dest, 0x01010101); }, Test, FunctionName, TEXT("int32")) && bSuccess;
+
+		// Test adding negative value
+		volatile int32 Value1 = 0x10010101;
+		bSuccess = TestInterlocked<int32>(&Value1, 0x10010101, 0x0F000000, [](volatile int32* Dest) { return FPlatformAtomics::InterlockedAdd(Dest, -0x01010101); }, Test, FunctionName, TEXT("int32")) && bSuccess;
+
+		// Test value overflow
+		volatile int32 Value2 = MAX_int32 - 1;
+		bSuccess = TestInterlocked<int32>(&Value2, MAX_int32 - 1, MIN_int32 + 2, [](volatile int32* Dest) { return FPlatformAtomics::InterlockedAdd(Dest, +4); }, Test, FunctionName, TEXT("int32")) && bSuccess;
+
+		// Test value underflow
+		volatile int32 Value3 = MIN_int32 + 2;
+		bSuccess = TestInterlocked<int32>(&Value3, MIN_int32 + 2, MAX_int32 - 1, [](volatile int32* Dest) { return FPlatformAtomics::InterlockedAdd(Dest, -4); }, Test, FunctionName, TEXT("int32")) && bSuccess;
+	}
+
+	{
+		// Test adding positive value
+		alignas(8) volatile int64 Value0 = 0x0F00000000000000LL;
+		bSuccess = TestInterlocked<int64>(&Value0, 0x0F00000000000000LL, 0x1001010101010101LL, [](volatile int64* Dest) { return FPlatformAtomics::InterlockedAdd(Dest, 0x0101010101010101LL); }, Test, FunctionName, TEXT("int64")) && bSuccess;
+
+		// Test adding negative value
+		alignas(8) volatile int64 Value1 = 0x1001010101010101LL;
+		bSuccess = TestInterlocked<int64>(&Value1, 0x1001010101010101LL, 0x0F00000000000000LL, [](volatile int64* Dest) { return FPlatformAtomics::InterlockedAdd(Dest, -0x0101010101010101LL); }, Test, FunctionName, TEXT("int64")) && bSuccess;
+
+		// Test value overflow
+		alignas(8) volatile int64 Value2 = MAX_int64 - 1LL;
+		bSuccess = TestInterlocked<int64>(&Value2, MAX_int64 - 1, MIN_int64 + 2, [](volatile int64* Dest) { return FPlatformAtomics::InterlockedAdd(Dest, +4); }, Test, FunctionName, TEXT("int64")) && bSuccess;
+
+		// Test value underflow
+		alignas(8) volatile int64 Value3 = MIN_int64 + 2LL;
+		bSuccess = TestInterlocked<int64>(&Value3, MIN_int64 + 2, MAX_int64 - 1, [](volatile int64* Dest) { return FPlatformAtomics::InterlockedAdd(Dest, -4); }, Test, FunctionName, TEXT("int64")) && bSuccess;
+	}
+
+	return bSuccess;
+}
+
+static bool TestInterlockedIncrement(FAutomationTestBase& Test)
+{
+	bool bSuccess = true;
+
+	constexpr TCHAR FunctionName[] = TEXT("Increment");
+	{
+		volatile int8 Value0 = 0x0F;
+		bSuccess = TestInterlocked<int8>(&Value0, 0x10, 0x10, [](volatile int8* Dest) { return FPlatformAtomics::InterlockedIncrement(Dest); }, Test, FunctionName, TEXT("int8")) && bSuccess;
+
+		// Test value overflow
+		volatile int8 Value1 = MAX_int8;
+		bSuccess = TestInterlocked<int8>(&Value1, MIN_int8, MIN_int8, [](volatile int8* Dest) { return FPlatformAtomics::InterlockedIncrement(Dest); }, Test, FunctionName, TEXT("int8")) && bSuccess;
+	}
+
+	{
+		volatile int16 Value0 = 0x0F0F;
+		bSuccess = TestInterlocked<int16>(&Value0, 0x0F10, 0x0F10, [](volatile int16* Dest) { return FPlatformAtomics::InterlockedIncrement(Dest); }, Test, FunctionName, TEXT("int16")) && bSuccess;
+
+		// Test value overflow
+		volatile int16 Value1 = MAX_int16;
+		bSuccess = TestInterlocked<int16>(&Value1, MIN_int16, MIN_int16, [](volatile int16* Dest) { return FPlatformAtomics::InterlockedIncrement(Dest); }, Test, FunctionName, TEXT("int16")) && bSuccess;
+	}
+
+	{
+		volatile int32 Value0 = 0x0F00000F;
+		bSuccess = TestInterlocked<int32>(&Value0, 0x0F000010, 0x0F000010, [](volatile int32* Dest) { return FPlatformAtomics::InterlockedIncrement(Dest); }, Test, FunctionName, TEXT("int32")) && bSuccess;
+
+		// Test value overflow
+		volatile int32 Value1 = MAX_int32;
+		bSuccess = TestInterlocked<int32>(&Value1, MIN_int32, MIN_int32, [](volatile int32* Dest) { return FPlatformAtomics::InterlockedIncrement(Dest); }, Test, FunctionName, TEXT("int32")) && bSuccess;
+	}
+
+	{
+		alignas(8) volatile int64 Value0 = 0x0F0000000000000FLL;
+		bSuccess = TestInterlocked<int64>(&Value0, 0x0F00000000000010LL, 0x0F00000000000010LL, [](volatile int64* Dest) { return FPlatformAtomics::InterlockedIncrement(Dest); }, Test, FunctionName, TEXT("int64")) && bSuccess;
+
+		// Test value overflow
+		alignas(8) volatile int64 Value1 = MAX_int64;
+		bSuccess = TestInterlocked<int64>(&Value1, MIN_int64, MIN_int64, [](volatile int64* Dest) { return FPlatformAtomics::InterlockedIncrement(Dest); }, Test, FunctionName, TEXT("int64")) && bSuccess;
+	}
+
+	return bSuccess;
+}
+
+static bool TestInterlockedDecrement(FAutomationTestBase& Test)
+{
+	bool bSuccess = true;
+
+	constexpr TCHAR FunctionName[] = TEXT("Decrement");
+	{
+		volatile int8 Value0 = 0x10;
+		bSuccess = TestInterlocked<int8>(&Value0, 0x0F, 0x0F, [](volatile int8* Dest) { return FPlatformAtomics::InterlockedDecrement(Dest); }, Test, FunctionName, TEXT("int8")) && bSuccess;
+
+		// Test value underflow
+		volatile int8 Value1 = MIN_int8;
+		bSuccess = TestInterlocked<int8>(&Value1, MAX_int8, MAX_int8, [](volatile int8* Dest) { return FPlatformAtomics::InterlockedDecrement(Dest); }, Test, FunctionName, TEXT("int8")) && bSuccess;
+	}
+
+	{
+		volatile int16 Value0 = 0x0F10;
+		bSuccess = TestInterlocked<int16>(&Value0, 0x0F0F, 0x0F0F, [](volatile int16* Dest) { return FPlatformAtomics::InterlockedDecrement(Dest); }, Test, FunctionName, TEXT("int16")) && bSuccess;
+
+		// Test value underflow
+		volatile int16 Value1 = MIN_int16;
+		bSuccess = TestInterlocked<int16>(&Value1, MAX_int16, MAX_int16, [](volatile int16* Dest) { return FPlatformAtomics::InterlockedDecrement(Dest); }, Test, FunctionName, TEXT("int16")) && bSuccess;
+	}
+
+	{
+		volatile int32 Value0 = 0x0F000010;
+		bSuccess = TestInterlocked<int32>(&Value0, 0x0F00000F, 0x0F00000F, [](volatile int32* Dest) { return FPlatformAtomics::InterlockedDecrement(Dest); }, Test, FunctionName, TEXT("int32")) && bSuccess;
+
+		// Test value underflow
+		volatile int32 Value1 = MIN_int32;
+		bSuccess = TestInterlocked<int32>(&Value1, MAX_int32, MAX_int32, [](volatile int32* Dest) { return FPlatformAtomics::InterlockedDecrement(Dest); }, Test, FunctionName, TEXT("int32")) && bSuccess;
+	}
+
+	{
+		alignas(8) volatile int64 Value0 = 0x0F00000000000010LL;
+		bSuccess = TestInterlocked<int64>(&Value0, 0x0F0000000000000FLL, 0x0F0000000000000FLL, [](volatile int64* Dest) { return FPlatformAtomics::InterlockedDecrement(Dest); }, Test, FunctionName, TEXT("int64")) && bSuccess;
+
+		// Test value underflow
+		alignas(8) volatile int64 Value1 = MIN_int64;
+		bSuccess = TestInterlocked<int64>(&Value1, MAX_int64, MAX_int64, [](volatile int64* Dest) { return FPlatformAtomics::InterlockedDecrement(Dest); }, Test, FunctionName, TEXT("int64")) && bSuccess;
+	}
+
+	return bSuccess;
+}
+
+static bool TestInterlockedExchange(FAutomationTestBase& Test)
+{
+	bool bSuccess = true;
+
+	constexpr TCHAR FunctionName[] = TEXT("Exchange");
+	{
+		volatile int8 Value = 0x10;
+		bSuccess = TestInterlocked<int8>(&Value, 0x10, 0x01, [](volatile int8* Dest) { return FPlatformAtomics::InterlockedExchange(Dest, 0x01); }, Test, FunctionName, TEXT("int8")) && bSuccess;
+	}
+
+	{
+		volatile int16 Value = 0x1000;
+		bSuccess = TestInterlocked<int16>(&Value, 0x1000, 0x0001, [](volatile int16* Dest) { return FPlatformAtomics::InterlockedExchange(Dest, 0x0001); }, Test, FunctionName, TEXT("int16")) && bSuccess;
+	}
+
+	{
+		volatile int32 Value = 0x10000000;
+		bSuccess = TestInterlocked<int32>(&Value, 0x10000000, 0x00000101, [](volatile int32* Dest) { return FPlatformAtomics::InterlockedExchange(Dest, 0x00000101); }, Test, FunctionName, TEXT("int32")) && bSuccess;
+	}
+
+	{
+		alignas(8) volatile int64 Value = 0x1000000000000000LL;
+		bSuccess = TestInterlocked<int64>(&Value, 0x1000000000000000LL, 0x0000000001010101LL, [](volatile int64* Dest) { return FPlatformAtomics::InterlockedExchange(Dest, 0x0000000001010101LL); }, Test, FunctionName, TEXT("int64")) && bSuccess;
+	}
+
+	return bSuccess;
+}
+
+static bool TestInterlockedExchangePtr(FAutomationTestBase& Test)
+{
+	void* volatile Value = &Test;
+	const void* ReturnValue = FPlatformAtomics::InterlockedExchangePtr(&Value, nullptr);
+	if (ReturnValue != &Test)
+	{
+		Test.AddError(TEXT("FPlatformAtomics::InterlockedExchangePtr failed"));
+		return false;
+	}
+
+	if (Value != nullptr)
+	{
+		Test.AddError(TEXT("FPlatformAtomics::InterlockedExchangePtr failed"));
+		return false;
+	}
+
+	return true;
+}
+
+static bool TestInterlockedCompareExchange(FAutomationTestBase& Test)
+{
+	bool bSuccess = true;
+
+	constexpr TCHAR FunctionName[] = TEXT("CompareExchange");
+	{
+		volatile int8 Value = 0x10;
+
+		// Test value isn't changed when comparand differs
+		bSuccess = TestInterlocked<int8>(&Value, 0x10, 0x10, [](volatile int8* Dest) { return FPlatformAtomics::InterlockedCompareExchange(Dest, 0x10, 0x01); }, Test, FunctionName, TEXT("int8")) && bSuccess;
+
+		// Test value is changed when comparand matches
+		bSuccess = TestInterlocked<int8>(&Value, 0x10, 0x01, [](volatile int8* Dest) { return FPlatformAtomics::InterlockedCompareExchange(Dest, 0x01, 0x10); }, Test, FunctionName, TEXT("int8")) && bSuccess;
+	}
+
+	{
+		volatile int16 Value = 0x1000;
+
+		// Test value isn't changed when comparand differs
+		bSuccess = TestInterlocked<int16>(&Value, 0x1000, 0x1000, [](volatile int16* Dest) { return FPlatformAtomics::InterlockedCompareExchange(Dest, 0x1000, 0x0001); }, Test, FunctionName, TEXT("int16")) && bSuccess;
+
+		// Test value is changed when comparand matches
+		bSuccess = TestInterlocked<int16>(&Value, 0x1000, 0x0001, [](volatile int16* Dest) { return FPlatformAtomics::InterlockedCompareExchange(Dest, 0x0001, 0x1000); }, Test, FunctionName, TEXT("int16")) && bSuccess;
+	}
+
+	{
+		volatile int32 Value = 0x10000000;
+
+		// Test value isn't changed when comparand differs
+		bSuccess = TestInterlocked<int32>(&Value, 0x10000000, 0x10000000, [](volatile int32* Dest) { return FPlatformAtomics::InterlockedCompareExchange(Dest, 0x10000000, 0x00000101); }, Test, FunctionName, TEXT("int32")) && bSuccess;
+
+		// Test value is changed when comparand matches
+		bSuccess = TestInterlocked<int32>(&Value, 0x10000000, 0x00000101, [](volatile int32* Dest) { return FPlatformAtomics::InterlockedCompareExchange(Dest, 0x00000101, 0x10000000); }, Test, FunctionName, TEXT("int32")) && bSuccess;
+	}
+
+	{
+		alignas(8) volatile int64 Value = 0x1000000000000000LL;
+
+		// Test value isn't changed when comparand differs
+		bSuccess = TestInterlocked<int64>(&Value, 0x1000000000000000LL, 0x1000000000000000LL, [](volatile int64* Dest) { return FPlatformAtomics::InterlockedCompareExchange(Dest, 0x1000000000000000LL, 0x0000000001010101LL); }, Test, FunctionName, TEXT("int64")) && bSuccess;
+
+		// Test value is changed when comparand matches
+		bSuccess = TestInterlocked<int64>(&Value, 0x1000000000000000LL, 0x0000000001010101LL, [](volatile int64* Dest) { return FPlatformAtomics::InterlockedCompareExchange(Dest, 0x0000000001010101LL, 0x1000000000000000LL); }, Test, FunctionName, TEXT("int64")) && bSuccess;
+	}
+
+	return bSuccess;
+}
+
+bool FPlatformAtomicsTest::RunTest(const FString& Parameters)
+{
+	bool bSuccess = true;
+
+	bSuccess = TestInterlockedAnd(*this) && bSuccess;
+	bSuccess = TestInterlockedOr(*this) && bSuccess;
+	bSuccess = TestInterlockedXor(*this) && bSuccess;
+	bSuccess = TestInterlockedAdd(*this) && bSuccess;
+	bSuccess = TestInterlockedIncrement(*this) && bSuccess;
+	bSuccess = TestInterlockedDecrement(*this) && bSuccess;
+	bSuccess = TestInterlockedExchange(*this) && bSuccess;
+	bSuccess = TestInterlockedExchangePtr(*this) && bSuccess;
+	bSuccess = TestInterlockedCompareExchange(*this) && bSuccess;
+
+	return bSuccess;
+}
+
+#endif //WITH_DEV_AUTOMATION_TESTS
diff --git a/Engine/Source/Runtime/Core/Public/Clang/ClangPlatformAtomics.h b/Engine/Source/Runtime/Core/Public/Clang/ClangPlatformAtomics.h
--- a/Engine/Source/Runtime/Core/Public/Clang/ClangPlatformAtomics.h
+++ b/Engine/Source/Runtime/Core/Public/Clang/ClangPlatformAtomics.h
@@ -119,6 +119,66 @@ struct CORE_API FClangPlatformAtomics : public FGenericPlatformAtomics
 		return __sync_val_compare_and_swap(Dest, Comperand, Exchange);
 	}
 
+	static FORCEINLINE int8 InterlockedAnd(volatile int8* Value, const int8 AndValue)
+	{
+		return __sync_fetch_and_and(Value, AndValue);
+	}
+
+	static FORCEINLINE int16 InterlockedAnd(volatile int16* Value, const int16 AndValue)
+	{
+		return __sync_fetch_and_and(Value, AndValue);
+	}
+
+	static FORCEINLINE int32 InterlockedAnd(volatile int32* Value, const int32 AndValue)
+	{
+		return __sync_fetch_and_and(Value, AndValue);
+	}
+
+	static FORCEINLINE int64 InterlockedAnd(volatile int64* Value, const int64 AndValue)
+	{
+		return __sync_fetch_and_and(Value, AndValue);
+	}
+
+	static FORCEINLINE int8 InterlockedOr(volatile int8* Value, const int8 OrValue)
+	{
+		return __sync_fetch_and_or(Value, OrValue);
+	}
+
+	static FORCEINLINE int16 InterlockedOr(volatile int16* Value, const int16 OrValue)
+	{
+		return __sync_fetch_and_or(Value, OrValue);
+	}
+
+	static FORCEINLINE int32 InterlockedOr(volatile int32* Value, const int32 OrValue)
+	{
+		return __sync_fetch_and_or(Value, OrValue);
+	}
+
+	static FORCEINLINE int64 InterlockedOr(volatile int64* Value, const int64 OrValue)
+	{
+		return __sync_fetch_and_or(Value, OrValue);
+	}
+
+	static FORCEINLINE int8 InterlockedXor(volatile int8* Value, const int8 XorValue)
+	{
+		return __sync_fetch_and_xor(Value, XorValue);
+	}
+
+	static FORCEINLINE int16 InterlockedXor(volatile int16* Value, const int16 XorValue)
+	{
+		return __sync_fetch_and_xor(Value, XorValue);
+	}
+
+	static FORCEINLINE int32 InterlockedXor(volatile int32* Value, const int32 XorValue)
+	{
+		return __sync_fetch_and_xor(Value, XorValue);
+	}
+
+	static FORCEINLINE int64 InterlockedXor(volatile int64* Value, const int64 XorValue)
+	{
+		return __sync_fetch_and_xor(Value, XorValue);
+	}
+
 	static FORCEINLINE int8 AtomicRead(volatile const int8* Src)
 	{
 		int8 Result;
diff --git a/Engine/Source/Runtime/Core/Public/GenericPlatform/GenericPlatformAtomics.h b/Engine/Source/Runtime/Core/Public/GenericPlatform/GenericPlatformAtomics.h
--- a/Engine/Source/Runtime/Core/Public/GenericPlatform/GenericPlatformAtomics.h
+++ b/Engine/Source/Runtime/Core/Public/GenericPlatform/GenericPlatformAtomics.h
@@ -187,6 +187,48 @@ struct FGenericPlatformAtomics
 		#error must implement
 	}
 
+	/**
+	 * Atomically ands AndValue to the value pointed to and returns the old
+	 * value to the caller. Implemented for int8, int16 and int64 as well.
+	 */
+	static FORCEINLINE int32 InterlockedAnd(volatile int32* Value, const int32 AndValue)
+	{
+		int32 RetVal;
+		do
+		{
+			RetVal = *Value;
+		} while (InterlockedCompareExchange((int32*)Value, RetVal & AndValue, RetVal) != RetVal);
+		return RetVal;
+	}
+
+	/**
+	 * Atomically ors OrValue to the value pointed to and returns the old
+	 * value to the caller. Implemented for int8, int16 and int64 as well.
+	 */
+	static FORCEINLINE int32 InterlockedOr(volatile int32* Value, const int32 OrValue)
+	{
+		int32 RetVal;
+		do
+		{
+			RetVal = *Value;
+		} while (InterlockedCompareExchange((int32*)Value, RetVal | OrValue, RetVal) != RetVal);
+		return RetVal;
+	}
+
+	/**
+	 * Atomically xors XorValue to the value pointed to and returns the old
+	 * value to the caller. Implemented for int8, int16 and int64 as well.
+	 */
+	static FORCEINLINE int32 InterlockedXor(volatile int32* Value, const int32 XorValue)
+	{
+		int32 RetVal;
+		do
+		{
+			RetVal = *Value;
+		} while (InterlockedCompareExchange((int32*)Value, RetVal ^ XorValue, RetVal) != RetVal);
+		return RetVal;
+	}
+
 	/**
 	 * Atomic read of 32 bit value with an implicit memory barrier.
 	 */
diff --git a/Engine/Source/Runtime/Core/Public/HTML5/HTML5PlatformAtomics.h b/Engine/Source/Runtime/Core/Public/HTML5/HTML5PlatformAtomics.h
--- a/Engine/Source/Runtime/Core/Public/HTML5/HTML5PlatformAtomics.h
+++ b/Engine/Source/Runtime/Core/Public/HTML5/HTML5PlatformAtomics.h
@@ -131,6 +131,66 @@ struct CORE_API FHTML5PlatformAtomics	: public FGenericPlatformAtomics
 //		return __sync_val_compare_and_swap(Dest, Comperand, Exchange);
 	}
 
+	static FORCEINLINE int8 InterlockedAnd(volatile int8* Value, const int8 AndValue)
+	{
+		return __sync_fetch_and_and(Value, AndValue);
+	}
+
+	static FORCEINLINE int16 InterlockedAnd(volatile int16* Value, const int16 AndValue)
+	{
+		return __sync_fetch_and_and(Value, AndValue);
+	}
+
+	static FORCEINLINE int32 InterlockedAnd(volatile int32* Value, const int32 AndValue)
+	{
+		return __sync_fetch_and_and(Value, AndValue);
+	}
+
+	static FORCEINLINE int64 InterlockedAnd(volatile int64* Value, const int64 AndValue)
+	{
+		return __sync_fetch_and_and(Value, AndValue);
+	}
+
+	static FORCEINLINE int8 InterlockedOr(volatile int8* Value, const int8 OrValue)
+	{
+		return __sync_fetch_and_or(Value, OrValue);
+	}
+
+	static FORCEINLINE int16 InterlockedOr(volatile int16* Value, const int16 OrValue)
+	{
+		return __sync_fetch_and_or(Value, OrValue);
+	}
+
+	static FORCEINLINE int32 InterlockedOr(volatile int32* Value, const int32 OrValue)
+	{
+		return __sync_fetch_and_or(Value, OrValue);
+	}
+
+	static FORCEINLINE int64 InterlockedOr(volatile int64* Value, const int64 OrValue)
+	{
+		return __sync_fetch_and_or(Value, OrValue);
+	}
+
+	static FORCEINLINE int8 InterlockedXor(volatile int8* Value, const int8 XorValue)
+	{
+		return __sync_fetch_and_xor(Value, XorValue);
+	}
+
+	static FORCEINLINE int16 InterlockedXor(volatile int16* Value, const int16 XorValue)
+	{
+		return __sync_fetch_and_xor(Value, XorValue);
+	}
+
+	static FORCEINLINE int32 InterlockedXor(volatile int32* Value, const int32 XorValue)
+	{
+		return __sync_fetch_and_xor(Value, XorValue);
+	}
+
+	static FORCEINLINE int64 InterlockedXor(volatile int64* Value, const int64 XorValue)
+	{
+		return __sync_fetch_and_xor(Value, XorValue);
+	}
+
 	static FORCEINLINE int8 AtomicRead(volatile const int8* Src)
 	{
 		int8 Result;
@@ -245,77 +305,85 @@ struct CORE_API FHTML5PlatformAtomics	: public FGenericPlatformAtomics
 {
 	static FORCEINLINE int8 InterlockedIncrement( volatile int8* Value )
 	{
-		*Value += 1;
-		return *Value;
+		const int8 TempValue = static_cast<int8>(static_cast<uint8>(*Value) + 1U);
+		*Value = TempValue;
+		return TempValue;
 	}
 
 	static FORCEINLINE int16 InterlockedIncrement( volatile int16* Value )
 	{
-		*Value += 1;
-		return *Value;
+		const int16 TempValue = static_cast<int16>(static_cast<uint16>(*Value) + 1U);
+		*Value = TempValue;
+		return TempValue;
 	}
 
 	static FORCEINLINE int32 InterlockedIncrement( volatile int32* Value )
 	{
-		*Value += 1;
-		return *Value;
+		const int32 TempValue = static_cast<int32>(static_cast<uint32>(*Value) + 1U);
+		*Value = TempValue;
+		return TempValue;
 	}
 
 	static FORCEINLINE int64 InterlockedIncrement( volatile int64* Value )
 	{
-		*Value += 1;
-		return *Value;
+		const int64 TempValue = static_cast<int64>(static_cast<uint64>(*Value) + 1ULL);
+		*Value = TempValue;
+		return TempValue;
 	}
 
 	static FORCEINLINE int8 InterlockedDecrement( volatile int8* Value )
 	{
-		*Value -= 1;
-		return *Value;
+		const int8 TempValue = static_cast<int8>(static_cast<uint8>(*Value) - 1U);
+		*Value = TempValue;
+		return TempValue;
 	}
 
 	static FORCEINLINE int16 InterlockedDecrement( volatile int16* Value )
 	{
-		*Value -= 1;
-		return *Value;
+		const int16 TempValue = static_cast<int16>(static_cast<uint16>(*Value) - 1U);
+		*Value = TempValue;
+		return TempValue;
 	}
 
 	static FORCEINLINE int32 InterlockedDecrement( volatile int32* Value )
 	{
-		*Value -= 1;
-		return *Value;
+		const int32 TempValue = static_cast<int32>(static_cast<uint32>(*Value) - 1U);
+		*Value = TempValue;
+		return TempValue;
 	}
 
 	static FORCEINLINE int64 InterlockedDecrement( volatile int64* Value )
 	{
-		*Value -=1;
-		return *Value;
+		const int64 TempValue = static_cast<int64>(static_cast<uint64>(*Value) - 1ULL);
+		*Value = TempValue;
+		return TempValue;
 	}
 
 	static FORCEINLINE int8 InterlockedAdd( volatile int8* Value, int8 Amount )
 	{
 		int8 Result = *Value;
-		*Value += Amount;
+		*Value = static_cast<int8>(static_cast<uint8>(Result) + static_cast<uint8>(Amount));
 		return Result;
 	}
 
 	static FORCEINLINE int16 InterlockedAdd( volatile int16* Value, int16 Amount )
 	{
 		int16 Result = *Value;
-		*Value += Amount;
+		*Value = static_cast<int16>(static_cast<uint16>(Result) + static_cast<uint16>(Amount));
 		return Result;
 	}
 
 	static FORCEINLINE int32 InterlockedAdd( volatile int32* Value, int32 Amount )
 	{
 		int32 Result = *Value;
-		*Value += Amount;
+		*Value = static_cast<int32>(static_cast<uint32>(Result) + static_cast<uint32>(Amount));
 		return Result;
 	}
 
 	static FORCEINLINE int64 InterlockedAdd( volatile int64* Value, int64 Amount )
 	{
 		int64 Result = *Value;
-		*Value += Amount;
+		*Value = static_cast<int64>(static_cast<uint64>(Result) + static_cast<uint64>(Amount));
 		return Result;
 	}
 
@@ -394,6 +462,90 @@ struct CORE_API FHTML5PlatformAtomics	: public FGenericPlatformAtomics
 		return Result;
 	}
 
+	static FORCEINLINE int8 InterlockedAnd(volatile int8* Value, const int8 AndValue)
+	{
+		const int8 Result = *Value;
+		*Value = Result & AndValue;
+		return Result;
+	}
+
+	static FORCEINLINE int16 InterlockedAnd(volatile int16* Value, const int16 AndValue)
+	{
+		const int16 Result = *Value;
+		*Value = Result & AndValue;
+		return Result;
+	}
+
+	static FORCEINLINE int32 InterlockedAnd(volatile int32* Value, const int32 AndValue)
+	{
+		const int32 Result = *Value;
+		*Value = Result & AndValue;
+		return Result;
+	}
+
+	static FORCEINLINE int64 InterlockedAnd(volatile int64* Value, const int64 AndValue)
+	{
+		const int64 Result = *Value;
+		*Value = Result & AndValue;
+		return Result;
+	}
+
+	static FORCEINLINE int8 InterlockedOr(volatile int8* Value, const int8 OrValue)
+	{
+		const int8 Result = *Value;
+		*Value = Result | OrValue;
+		return Result;
+	}
+
+	static FORCEINLINE int16 InterlockedOr(volatile int16* Value, const int16 OrValue)
+	{
+		const int16 Result = *Value;
+		*Value = Result | OrValue;
+		return Result;
+	}
+
+	static FORCEINLINE int32 InterlockedOr(volatile int32* Value, const int32 OrValue)
+	{
+		const int32 Result = *Value;
+		*Value = Result | OrValue;
+		return Result;
+	}
+
+	static FORCEINLINE int64 InterlockedOr(volatile int64* Value, const int64 OrValue)
+	{
+		const int64 Result = *Value;
+		*Value = Result | OrValue;
+		return Result;
+	}
+
+	static FORCEINLINE int8 InterlockedXor(volatile int8* Value, const int8 XorValue)
+	{
+		const int8 Result = *Value;
+		*Value = Result ^ XorValue;
+		return Result;
+	}
+
+	static FORCEINLINE int16 InterlockedXor(volatile int16* Value, const int16 XorValue)
+	{
+		const int16 Result = *Value;
+		*Value = Result ^ XorValue;
+		return Result;
+	}
+
+	static FORCEINLINE int32 InterlockedXor(volatile int32* Value, const int32 XorValue)
+	{
+		const int32 Result = *Value;
+		*Value = Result ^ XorValue;
+		return Result;
+	}
+
+	static FORCEINLINE int64 InterlockedXor(volatile int64* Value, const int64 XorValue)
+	{
+		const int64 Result = *Value;
+		*Value = Result ^ XorValue;
+		return Result;
+	}
+
 	static FORCEINLINE int8 AtomicRead(volatile const int8* Src)
 	{
 		return *Src;
diff --git a/Engine/Source/Runtime/Core/Public/Windows/WindowsPlatformAtomics.h b/Engine/Source/Runtime/Core/Public/Windows/WindowsPlatformAtomics.h
--- a/Engine/Source/Runtime/Core/Public/Windows/WindowsPlatformAtomics.h
+++ b/Engine/Source/Runtime/Core/Public/Windows/WindowsPlatformAtomics.h
@@ -185,6 +185,102 @@ struct CORE_API FWindowsPlatformAtomics
 		return (int64)::_InterlockedCompareExchange64(Dest, Exchange, Comparand);
 	}
 
+	static FORCEINLINE int8 InterlockedAnd(volatile int8* Value, const int8 AndValue)
+	{
+		return (int8)::_InterlockedAnd8((volatile char*)Value, (char)AndValue);
+	}
+
+	static FORCEINLINE int16 InterlockedAnd(volatile int16* Value, const int16 AndValue)
+	{
+		return (int16)::_InterlockedAnd16((volatile short*)Value, (short)AndValue);
+	}
+
+	static FORCEINLINE int32 InterlockedAnd(volatile int32* Value, const int32 AndValue)
+	{
+		return (int32)::_InterlockedAnd((volatile long*)Value, (long)AndValue);
+	}
+
+	static FORCEINLINE int64 InterlockedAnd(volatile int64* Value, const int64 AndValue)
+	{
+		#if PLATFORM_64BITS
+			return (int64)::_InterlockedAnd64((volatile long long*)Value, (long long)AndValue);
+		#else
+			// No explicit instruction for 64-bit atomic and on 32-bit processors; has to be implemented in terms of CMPXCHG8B
+			for (;;)
+			{
+				const int64 OldValue = *Value;
+				if (_InterlockedCompareExchange64(Value, OldValue & AndValue, OldValue) == OldValue)
+				{
+					return OldValue;
+				}
+			}
+		#endif
+	}
+
+	static FORCEINLINE int8 InterlockedOr(volatile int8* Value, const int8 OrValue)
+	{
+		return (int8)::_InterlockedOr8((volatile char*)Value, (char)OrValue);
+	}
+
+	static FORCEINLINE int16 InterlockedOr(volatile int16* Value, const int16 OrValue)
+	{
+		return (int16)::_InterlockedOr16((volatile short*)Value, (short)OrValue);
+	}
+
+	static FORCEINLINE int32 InterlockedOr(volatile int32* Value, const int32 OrValue)
+	{
+		return (int32)::_InterlockedOr((volatile long*)Value, (long)OrValue);
+	}
+
+	static FORCEINLINE int64 InterlockedOr(volatile int64* Value, const int64 OrValue)
+	{
+		#if PLATFORM_64BITS
+			return (int64)::_InterlockedOr64((volatile long long*)Value, (long long)OrValue);
+		#else
+			// No explicit instruction for 64-bit atomic or on 32-bit processors; has to be implemented in terms of CMPXCHG8B
+			for (;;)
+			{
+				const int64 OldValue = *Value;
+				if (_InterlockedCompareExchange64(Value, OldValue | OrValue, OldValue) == OldValue)
+				{
+					return OldValue;
+				}
+			}
+		#endif
+	}
+
+	static FORCEINLINE int8 InterlockedXor(volatile int8* Value, const int8 XorValue)
+	{
+		return (int8)::_InterlockedXor8((volatile char*)Value, (char)XorValue);
+	}
+
+	static FORCEINLINE int16 InterlockedXor(volatile int16* Value, const int16 XorValue)
+	{
+		return (int16)::_InterlockedXor16((volatile short*)Value, (short)XorValue);
+	}
+
+	static FORCEINLINE int32 InterlockedXor(volatile int32* Value, const int32 XorValue)
+	{
+		return (int32)::_InterlockedXor((volatile long*)Value, (int32)XorValue);
+	}
+
+	static FORCEINLINE int64 InterlockedXor(volatile int64* Value, const int64 XorValue)
+	{
+		#if PLATFORM_64BITS
+			return (int64)::_InterlockedXor64((volatile long long*)Value, (long long)XorValue);
+		#else
+			// No explicit instruction for 64-bit atomic xor on 32-bit processors; has to be implemented in terms of CMPXCHG8B
+			for (;;)
+			{
+				const int64 OldValue = *Value;
+				if (_InterlockedCompareExchange64(Value, OldValue ^ XorValue, OldValue) == OldValue)
+				{
+					return OldValue;
+				}
+			}
+		#endif
+	}
+
 	static FORCEINLINE int8 AtomicRead(volatile const int8* Src)
 	{
 		return InterlockedCompareExchange((int8*)Src, 0, 0);
