Added AndExchange, OrExchange, XorExchange as well as operators &=, |= and ^= to TAtomic for integral types.

diff --git a/Engine/Source/Runtime/Core/Private/Tests/Misc/AtomicTest.cpp b/Engine/Source/Runtime/Core/Private/Tests/Misc/AtomicTest.cpp
--- a/Engine/Source/Runtime/Core/Private/Tests/Misc/AtomicTest.cpp
+++ b/Engine/Source/Runtime/Core/Private/Tests/Misc/AtomicTest.cpp
@@ -13,7 +13,7 @@ namespace
 {
 	// Returns an element type filled with a repeated byte
 	template <typename ElementType>
-	FORCEINLINE typename TEnableIf<sizeof(ElementType) != 1, ElementType>::Type GetByteFilledElement(uint8 Byte)
+	FORCEINLINE constexpr typename TEnableIf<sizeof(ElementType) != 1, ElementType>::Type GetByteFilledElement(uint8 Byte)
 	{
 		TUnsignedIntType_T<sizeof(ElementType)> Result = 0;
 
@@ -27,7 +27,7 @@ namespace
 	}
 
 	template <typename ElementType>
-	FORCEINLINE typename TEnableIf<sizeof(ElementType) == 1, ElementType>::Type GetByteFilledElement(uint8 Byte)
+	FORCEINLINE constexpr typename TEnableIf<sizeof(ElementType) == 1, ElementType>::Type GetByteFilledElement(uint8 Byte)
 	{
 		return (const ElementType&)Byte;
 	}
@@ -173,12 +173,129 @@ namespace
 		Data.Check();
 	}
 
+	template <typename ElementType>
+	void RunBitwiseOperationsAtomicTests()
+	{
+		ElementType NativeVal;
+		ElementType AtomicVal;
+
+		// And
+		{
+			constexpr ElementType Init = GetByteFilledElement<ElementType>(0x30);
+			constexpr ElementType AndValues[] = {GetByteFilledElement<ElementType>(0x66), GetByteFilledElement<ElementType>(0xFF), ElementType(0), };
+			constexpr ElementType ExpectedValues[] = {GetByteFilledElement<ElementType>(0x20), Init, ElementType(0), };
+
+			// operator &=
+			for (SIZE_T TestIt = 0; TestIt < ARRAY_COUNT(AndValues); ++TestIt)
+			{
+				TAtomicTestWrapper<ElementType> Data(Init);
+
+				AtomicVal = Data.GetAtomic() &= AndValues[TestIt];
+				NativeVal = Data.GetNative() &= AndValues[TestIt];
+				check(NativeVal == AtomicVal);
+
+				check(AtomicVal == ExpectedValues[TestIt]);
+				Data.Check();
+			}
+
+			// AndExchange
+			for (SIZE_T TestIt = 0; TestIt < ARRAY_COUNT(AndValues); ++TestIt)
+			{
+				TAtomicTestWrapper<ElementType> Data(Init);
+
+				AtomicVal = Data.GetAtomic().AndExchange(AndValues[TestIt]);
+				NativeVal = Data.GetNative();
+				check(NativeVal == AtomicVal);
+				NativeVal = Data.GetNative() &= AndValues[TestIt];
+				AtomicVal = Data.GetAtomic().Load();
+				check(NativeVal == AtomicVal);
+
+				check(AtomicVal == ExpectedValues[TestIt]);
+				Data.Check();
+			}
+		}
+
+		// Or
+		{
+			constexpr ElementType Init = GetByteFilledElement<ElementType>(0x30);
+			constexpr ElementType OrValues[] = {GetByteFilledElement<ElementType>(0x66), GetByteFilledElement<ElementType>(0xFF), ElementType(0), };
+			constexpr ElementType ExpectedValues[] = {GetByteFilledElement<ElementType>(0x76), GetByteFilledElement<ElementType>(0xFF), Init, };
+
+			// operator |=
+			for (SIZE_T TestIt = 0; TestIt < ARRAY_COUNT(OrValues); ++TestIt)
+			{
+				TAtomicTestWrapper<ElementType> Data(Init);
+
+				AtomicVal = Data.GetAtomic() |= OrValues[TestIt];
+				NativeVal = Data.GetNative() |= OrValues[TestIt];
+				check(NativeVal == AtomicVal);
+
+				check(AtomicVal == ExpectedValues[TestIt]);
+				Data.Check();
+			}
+
+			// OrExchange
+			for (SIZE_T TestIt = 0; TestIt < ARRAY_COUNT(OrValues); ++TestIt)
+			{
+				TAtomicTestWrapper<ElementType> Data(Init);
+
+				AtomicVal = Data.GetAtomic().OrExchange(OrValues[TestIt]);
+				NativeVal = Data.GetNative();
+				check(NativeVal == AtomicVal);
+				NativeVal = Data.GetNative() |= OrValues[TestIt];
+				AtomicVal = Data.GetAtomic().Load();
+				check(NativeVal == AtomicVal);
+
+				check(AtomicVal == ExpectedValues[TestIt]);
+				Data.Check();
+			}
+		}
+
+		// Xor
+		{
+			constexpr ElementType Init = GetByteFilledElement<ElementType>(0x30);
+			constexpr ElementType XorValues[] = {GetByteFilledElement<ElementType>(0x66), GetByteFilledElement<ElementType>(0xFF), ElementType(0), };
+			constexpr ElementType ExpectedValues[] = {GetByteFilledElement<ElementType>(0x56), GetByteFilledElement<ElementType>(~0x30), Init, };
+
+			// operator ^=
+			for (SIZE_T TestIt = 0; TestIt < ARRAY_COUNT(XorValues); ++TestIt)
+			{
+				TAtomicTestWrapper<ElementType> Data(Init);
+
+				AtomicVal = Data.GetAtomic() ^= XorValues[TestIt];
+				NativeVal = Data.GetNative() ^= XorValues[TestIt];
+				check(NativeVal == AtomicVal);
+
+				check(AtomicVal == ExpectedValues[TestIt]);
+				Data.Check();
+			}
+
+			// XorExchange
+			for (SIZE_T TestIt = 0; TestIt < ARRAY_COUNT(XorValues); ++TestIt)
+			{
+				TAtomicTestWrapper<ElementType> Data(Init);
+
+				AtomicVal = Data.GetAtomic().XorExchange(XorValues[TestIt]);
+				NativeVal = Data.GetNative();
+				check(NativeVal == AtomicVal);
+				NativeVal = Data.GetNative() ^= XorValues[TestIt];
+				AtomicVal = Data.GetAtomic().Load();
+				check(NativeVal == AtomicVal);
+
+				check(AtomicVal == ExpectedValues[TestIt]);
+				Data.Check();
+			}
+		}
+	}
+
 	template <typename ElementType>
 	void RunNumericAtomicTests()
 	{
 		RunBasicAtomicTests<ElementType>();
 
 		RunArithmeticAtomicTests<ElementType>(50);
+
+		RunBitwiseOperationsAtomicTests<ElementType>();
 	}
 
 	template <typename ElementType>
diff --git a/Engine/Source/Runtime/Core/Public/Templates/Atomic.h b/Engine/Source/Runtime/Core/Public/Templates/Atomic.h
--- a/Engine/Source/Runtime/Core/Public/Templates/Atomic.h
+++ b/Engine/Source/Runtime/Core/Public/Templates/Atomic.h
@@ -154,6 +154,27 @@ namespace UE4Atomic_Private
 		return *(const T*)&Result;
 	}
 
+	template <typename T>
+	FORCEINLINE T AndExchange(volatile T* Element, T AndValue)
+	{
+		auto Result = FPlatformAtomics::InterlockedAnd((volatile TUnderlyingIntegerType_T<T>*)Element, (TUnderlyingIntegerType_T<T>)AndValue);
+		return *(const T*)&Result;
+	}
+
+	template <typename T>
+	FORCEINLINE T OrExchange(volatile T* Element, T OrValue)
+	{
+		auto Result = FPlatformAtomics::InterlockedOr((volatile TUnderlyingIntegerType_T<T>*)Element, (TUnderlyingIntegerType_T<T>)OrValue);
+		return *(const T*)&Result;
+	}
+
+	template <typename T>
+	FORCEINLINE T XorExchange(volatile T* Element, T XorValue)
+	{
+		auto Result = FPlatformAtomics::InterlockedXor((volatile TUnderlyingIntegerType_T<T>*)Element, (TUnderlyingIntegerType_T<T>)XorValue);
+		return *(const T*)&Result;
+	}
+
 	template <typename T, bool bIsVoidPointer, bool bIsIntegral, bool bCanUsePlatformAtomics>
 	struct TAtomicBaseType
 	{
@@ -377,7 +398,78 @@ template <typename T>
 struct TAtomicBase_Integral : public TAtomicBase_Arithmetic<T, T>
 {
 public:
-	// TODO: Bitwise ops
+	/**
+	 * Performs Element &= Value, returning a copy of the new value of the element.
+	 *
+	 * @param  Value  The value to and the element with.
+	 *
+	 * @return A copy of the new value.
+	 */
+	FORCEINLINE T operator&=(const T Value)
+	{
+		return UE4Atomic_Private::AndExchange(&this->Element, Value) & Value;
+	}
+
+	/**
+	 * Performs Element |= Value, returning a copy of the new value of the element.
+	 *
+	 * @param  Value  The value to or the element with.
+	 *
+	 * @return A copy of the new value.
+	 */
+	FORCEINLINE T operator|=(const T Value)
+	{
+		return UE4Atomic_Private::OrExchange(&this->Element, Value) | Value;
+	}
+
+	/**
+	 * Performs Element ^= Value, returning a copy of the new value of the element.
+	 *
+	 * @param  Value  The value to xor the element with.
+	 *
+	 * @return A copy of the new value.
+	 */
+	FORCEINLINE T operator^=(const T Value)
+	{
+		return UE4Atomic_Private::XorExchange(&this->Element, Value) ^ Value;
+	}
+
+	/**
+	 * Performs Element &= Value, returning a copy of the previous value of the element.
+	 *
+	 * @param  Value  The value to and the element with.
+	 *
+	 * @return A copy of the previous value.
+	 */
+	FORCEINLINE T AndExchange(const T Value)
+	{
+		return UE4Atomic_Private::AndExchange(&this->Element, Value);
+	}
+
+	/**
+	 * Performs Element |= Value, returning a copy of the previous value of the element.
+	 *
+	 * @param  Value  The value to or the element with.
+	 *
+	 * @return A copy of the previous value.
+	 */
+	FORCEINLINE T OrExchange(const T Value)
+	{
+		return UE4Atomic_Private::OrExchange(&this->Element, Value);
+	}
+
+	/**
+	 * Performs Element ^= Value, returning a copy of the previous value of the element.
+	 *
+	 * @param  Value  The value to xor the element with.
+	 *
+	 * @return A copy of the previous value.
+	 */
+	FORCEINLINE T XorExchange(const T Value)
+	{
+		return UE4Atomic_Private::XorExchange(&this->Element, Value);
+	}
+
 
 protected:
 	TAtomicBase_Integral() = default;
