Make sure the code compiles even if some windows header files have been included before HAL/PlatformAtomics.h

diff --git a/Engine/Source/Runtime/Core/Private/Tests/HAL/PlatformAtomicsTest.cpp b/Engine/Source/Runtime/Core/Private/Tests/HAL/PlatformAtomicsTest.cpp
--- a/Engine/Source/Runtime/Core/Private/Tests/HAL/PlatformAtomicsTest.cpp
+++ b/Engine/Source/Runtime/Core/Private/Tests/HAL/PlatformAtomicsTest.cpp
@@ -1,12 +1,17 @@
 // Copyright 1998-2019 Epic Games, Inc. All Rights Reserved.
 
 #include "CoreTypes.h"
-#include "HAL/PlatformAtomics.h"
 #include "Misc/AssertionMacros.h"
 #include "Misc/AutomationTest.h"
 
 #if WITH_DEV_AUTOMATION_TESTS
 
+#if PLATFORM_WINDOWS
+#include "Windows/AllowWindowsPlatformAtomics.h"
+#include "Windows/HideWindowsPlatformAtomics.h"
+#endif
+#include "HAL/PlatformAtomics.h"
+
 IMPLEMENT_SIMPLE_AUTOMATION_TEST(FPlatformAtomicsTest, "System.Core.HAL.PlatformAtomics", EAutomationTestFlags::ApplicationContextMask | EAutomationTestFlags::SmokeFilter)
 
 template<typename T>
