Updated Windows macro defines and undefs to allow/hide InterlockedAnd, InterlockedOr and InterlockedXor.

diff --git a/Engine/Source/Runtime/Core/Public/Windows/AllowWindowsPlatformAtomics.h b/Engine/Source/Runtime/Core/Public/Windows/AllowWindowsPlatformAtomics.h
--- a/Engine/Source/Runtime/Core/Public/Windows/AllowWindowsPlatformAtomics.h
+++ b/Engine/Source/Runtime/Core/Public/Windows/AllowWindowsPlatformAtomics.h
@@ -12,6 +12,9 @@
 #define InterlockedExchange _InterlockedExchange
 #define InterlockedExchangeAdd _InterlockedExchangeAdd
 #define InterlockedCompareExchange _InterlockedCompareExchange
+#define InterlockedAnd _InterlockedAnd
+#define InterlockedOr _InterlockedOr
+#define InterlockedXor _InterlockedXor
 
 #if PLATFORM_64BITS
 	#define InterlockedCompareExchangePointer _InterlockedCompareExchangePointer
diff --git a/Engine/Source/Runtime/Core/Public/Windows/HideWindowsPlatformAtomics.h b/Engine/Source/Runtime/Core/Public/Windows/HideWindowsPlatformAtomics.h
--- a/Engine/Source/Runtime/Core/Public/Windows/HideWindowsPlatformAtomics.h
+++ b/Engine/Source/Runtime/Core/Public/Windows/HideWindowsPlatformAtomics.h
@@ -18,4 +18,7 @@
 #undef InterlockedCompareExchange64
 #undef InterlockedIncrement64
 #undef InterlockedDecrement64
+#undef InterlockedAnd
+#undef InterlockedOr
+#undef InterlockedXor
 
diff --git a/Engine/Source/Runtime/Core/Public/Windows/PostWindowsApi.h b/Engine/Source/Runtime/Core/Public/Windows/PostWindowsApi.h
--- a/Engine/Source/Runtime/Core/Public/Windows/PostWindowsApi.h
+++ b/Engine/Source/Runtime/Core/Public/Windows/PostWindowsApi.h
@@ -59,6 +59,9 @@ THIRD_PARTY_INCLUDES_END
 	#undef InterlockedCompareExchange64
 	#undef InterlockedIncrement64
 	#undef InterlockedDecrement64
+	#undef InterlockedAnd
+	#undef InterlockedOr
+	#undef InterlockedXor
 #endif
 
 // Restore any previously defined macros
