Added IncrementExchange, DecrementExchange, AddExchange and SubExchange to TAtomic.

diff --git a/Engine/Source/Runtime/Core/Private/Tests/Misc/AtomicTest.cpp b/Engine/Source/Runtime/Core/Private/Tests/Misc/AtomicTest.cpp
--- a/Engine/Source/Runtime/Core/Private/Tests/Misc/AtomicTest.cpp
+++ b/Engine/Source/Runtime/Core/Private/Tests/Misc/AtomicTest.cpp
@@ -165,11 +165,65 @@ namespace
 		NativeVal = Data.GetNative()--;
 		AtomicVal = Data.GetAtomic()--;
 		check(NativeVal == AtomicVal);
+		NativeVal = Data.GetNative();
+		AtomicVal = Data.GetAtomic();
+		check(NativeVal == AtomicVal);
 		Data.Check();
 
 		NativeVal = Data.GetNative()++;
 		AtomicVal = Data.GetAtomic()++;
 		check(NativeVal == AtomicVal);
+		NativeVal = Data.GetNative();
+		AtomicVal = Data.GetAtomic();
+		check(NativeVal == AtomicVal);
+		Data.Check();
+
+		NativeVal = Data.GetNative()--;
+		AtomicVal = Data.GetAtomic().DecrementExchange();
+		check(NativeVal == AtomicVal);
+		NativeVal = Data.GetNative();
+		AtomicVal = Data.GetAtomic();
+		check(NativeVal == AtomicVal);
+		Data.Check();
+
+		NativeVal = Data.GetNative()++;
+		AtomicVal = Data.GetAtomic().IncrementExchange();
+		check(NativeVal == AtomicVal);
+		NativeVal = Data.GetNative();
+		AtomicVal = Data.GetAtomic();
+		check(NativeVal == AtomicVal);
+		Data.Check();
+
+		NativeVal = Data.GetNative();
+		AtomicVal = Data.GetAtomic().AddExchange(47);
+		check(NativeVal == AtomicVal);
+		NativeVal = Data.GetNative() += 47;
+		AtomicVal = Data.GetAtomic();
+		check(NativeVal == AtomicVal);
+		Data.Check();
+
+		NativeVal = Data.GetNative();
+		AtomicVal = Data.GetAtomic().AddExchange(-11);
+		check(NativeVal == AtomicVal);
+		NativeVal = Data.GetNative() += -11;
+		AtomicVal = Data.GetAtomic();
+		check(NativeVal == AtomicVal);
+		Data.Check();
+
+		NativeVal = Data.GetNative();
+		AtomicVal = Data.GetAtomic().SubExchange(2);
+		check(NativeVal == AtomicVal);
+		NativeVal = Data.GetNative() -= 2;
+		AtomicVal = Data.GetAtomic();
+		check(NativeVal == AtomicVal);
+		Data.Check();
+
+		NativeVal = Data.GetNative();
+		AtomicVal = Data.GetAtomic().SubExchange(-9);
+		check(NativeVal == AtomicVal);
+		NativeVal = Data.GetNative() -= -9;
+		AtomicVal = Data.GetAtomic();
+		check(NativeVal == AtomicVal);
 		Data.Check();
 	}
 
diff --git a/Engine/Source/Runtime/Core/Public/Templates/Atomic.h b/Engine/Source/Runtime/Core/Public/Templates/Atomic.h
--- a/Engine/Source/Runtime/Core/Public/Templates/Atomic.h
+++ b/Engine/Source/Runtime/Core/Public/Templates/Atomic.h
@@ -362,7 +362,7 @@ public:
 	/**
 	 * Performs Element -= Value, returning a copy of the new value of the element.
 	 *
-	 * @param  Value  The value to add to the element.
+	 * @param  Value  The value to subtract from the element.
 	 *
 	 * @return A copy of the new, decremented value.
 	 */
@@ -371,6 +371,50 @@ public:
 		return UE4Atomic_Private::SubExchange(&this->Element, Value) - Value;
 	}
 
+	/**
+	 * Increments Element and returns a copy of the previous value of the element.
+	 *
+	 * @return A copy of the previous, unincremented value.
+	 */
+	FORCEINLINE T IncrementExchange()
+	{
+		return UE4Atomic_Private::IncrementExchange(&this->Element);
+	}
+
+	/**
+	 * Decrements Element and returns a copy of the previous value of the element.
+	 *
+	 * @return A copy of the previous, undecremented value.
+	 */
+	FORCEINLINE T DecrementExchange()
+	{
+		return UE4Atomic_Private::DecrementExchange(&this->Element);
+	}
+
+	/**
+	 * Adds Value to Element and returns a copy of the previous value of the element.
+	 *
+	 * @param  Value  The value to add to the element.
+	 *
+	 * @return A copy of the previous, unincremented value.
+	 */
+	FORCEINLINE T AddExchange(DiffType Value)
+	{
+		return UE4Atomic_Private::AddExchange(&this->Element, Value);
+	}
+
+	/**
+	 * Subtracts Value from Element and returns a copy of the previous value of the element.
+	 *
+	 * @param  Value  The value to subtract from the element.
+	 *
+	 * @return A copy of the previous, undecremented value.
+	 */
+	FORCEINLINE T SubExchange(DiffType Value)
+	{
+		return UE4Atomic_Private::SubExchange(&this->Element, Value);
+	}
+
 protected:
 	TAtomicBase_Arithmetic() = default;
 
