Improve efficiency of FString::AppendInt.

Added tests for the aforementioned method, System.Core.String.AppendInt

diff --git a/Engine/Source/Runtime/Core/Private/Containers/String.cpp b/Engine/Source/Runtime/Core/Private/Containers/String.cpp
--- a/Engine/Source/Runtime/Core/Private/Containers/String.cpp
+++ b/Engine/Source/Runtime/Core/Private/Containers/String.cpp
@@ -658,39 +658,29 @@ void FString::SerializeAsANSICharArray( FArchive& Ar, int32 MinCharacters ) cons
 	}
 }
 
-void FString::AppendInt( int32 InNum )
+void FString::AppendInt( int32 Num )
 {
-	int64 Num						= InNum; // This avoids having to deal with negating -MAX_int32-1
-	const TCHAR* NumberChar[11]		= { TEXT("0"), TEXT("1"), TEXT("2"), TEXT("3"), TEXT("4"), TEXT("5"), TEXT("6"), TEXT("7"), TEXT("8"), TEXT("9"), TEXT("-") };
-	bool bIsNumberNegative			= false;
+	const TCHAR* DigitToChar		= TEXT("9876543210123456789");
+	constexpr int32 ZeroDigitIndex	= 9;
+	bool bIsNumberNegative			= Num < 0;
 	const int32 TempBufferSize		= 16; // 16 is big enough
 	TCHAR TempNum[TempBufferSize];				
 	int32 TempAt					= TempBufferSize; // fill the temp string from the top down.
 
-	// Correctly handle negative numbers and convert to positive integer.
-	if( Num < 0 )
-	{
-		bIsNumberNegative = true;
-		Num = -Num;
-	}
-
-	TempNum[--TempAt] = 0; // NULL terminator
-
-	// Convert to string assuming base ten and a positive integer.
+	// Convert to string assuming base ten.
 	do 
 	{
-		TempNum[--TempAt] = *NumberChar[Num % 10];
+		TempNum[--TempAt] = DigitToChar[ZeroDigitIndex + (Num % 10)];
 		Num /= 10;
 	} while( Num );
 
-	// Append sign as we're going to reverse string afterwards.
 	if( bIsNumberNegative )
 	{
-		TempNum[--TempAt] = *NumberChar[10];
+		TempNum[--TempAt] = TEXT('-');
 	}
 
 	const TCHAR* CharPtr = TempNum + TempAt;
-	const int32 NumChars = TempBufferSize - TempAt - 1;
+	const int32 NumChars = TempBufferSize - TempAt;
 	Append(CharPtr, NumChars);
 }
 
diff --git a/Engine/Source/Runtime/Core/Private/Containers/StringTest.cpp b/Engine/Source/Runtime/Core/Private/Containers/StringTest.cpp
--- a/Engine/Source/Runtime/Core/Private/Containers/StringTest.cpp
+++ b/Engine/Source/Runtime/Core/Private/Containers/StringTest.cpp
@@ -46,4 +46,44 @@ bool FStringSanitizeFloatTest::RunTest(const FString& Parameters)
 	return true;
 }
 
+IMPLEMENT_SIMPLE_AUTOMATION_TEST(FStringAppendIntTest, "System.Core.String.AppendInt", EAutomationTestFlags::ApplicationContextMask | EAutomationTestFlags::SmokeFilter)
+bool FStringAppendIntTest::RunTest(const FString& Parameters)
+{
+	auto DoTest = [this](const TCHAR* Call, const FString& Result, const TCHAR* InExpected)
+	{
+		if (!Result.Equals(InExpected, ESearchCase::CaseSensitive))
+		{
+			AddError(FString::Printf(TEXT("'%s' failure: result '%s' (expected '%s')"), Call, *Result, InExpected));
+		}
+	};
+
+	{
+		FString Zero;
+		Zero.AppendInt(0);
+		DoTest(TEXT("AppendInt(0)"), Zero, TEXT("0"));
+	}
+
+	{
+		FString IntMin;
+		IntMin.AppendInt(MIN_int32);
+		DoTest(TEXT("AppendInt(MIN_int32)"), IntMin, TEXT("-2147483648"));
+	}
+
+	{
+		FString IntMin;
+		IntMin.AppendInt(MAX_int32);
+		DoTest(TEXT("AppendInt(MAX_int32)"), IntMin, TEXT("2147483647"));
+	}
+
+	{
+		FString AppendMultipleInts;
+		AppendMultipleInts.AppendInt(1);
+		AppendMultipleInts.AppendInt(-2);
+		AppendMultipleInts.AppendInt(3);
+		DoTest(TEXT("AppendInt(1);AppendInt(-2);AppendInt(3)"), AppendMultipleInts, TEXT("1-23"));
+	}
+
+	return true;
+}
+
 #endif // WITH_DEV_AUTOMATION_TESTS
