Handle all format specifiers starting with an l or ll followed by d, i, u, x, or X.
Fixed bug that would clamp doubles to just 30 characters.
Removed FmtBuf clearing.

Added defines for format specifiers for Unreal specific types int64, uint64, SIZE_T, SSIZE_T, PTRINT, UPTRINT. Defines exist for decimal and hexadecimal outputs. They are named with all capital letters starting with the type name and ending with _FMT. For SIZE_T the defines are SIZE_T_FMT, SIZE_T_x_FMT and SIZE_T_X_FMT for decimal output, lowercase hexadecimal output and uppercase hexadecimal output respectively. The other types follow the same pattern.

diff --git a/Engine/Source/Runtime/Core/Private/GenericPlatform/GenericWidePlatformString.cpp b/Engine/Source/Runtime/Core/Private/GenericPlatform/GenericWidePlatformString.cpp
--- a/Engine/Source/Runtime/Core/Private/GenericPlatform/GenericWidePlatformString.cpp
+++ b/Engine/Source/Runtime/Core/Private/GenericPlatform/GenericWidePlatformString.cpp
@@ -235,11 +235,20 @@ void RunGetVarArgsTests()
 	TestGetVarArgs(OutputString, TEXT("Test D|%p|"), 0x12345);
 	check(FString(OutputString) == FString(TEXT("Test D|0x12345|")));
 
-	TestGetVarArgs(OutputString, TEXT("Test E|%lld|"), 12345678912345);
+	TestGetVarArgs(OutputString, TEXT("Test E|%" INT64_FMT "|"), int64(12345678912345LL));
 	check(FString(OutputString) == FString(TEXT("Test E|12345678912345|")));
 
 	TestGetVarArgs(OutputString, TEXT("Test F|%f|%e|%g|"), 123.456, 123.456, 123.456);
 	check(FString(OutputString) == FString(TEXT("Test F|123.456000|1.234560e+02|123.456|")));
+
+	TestGetVarArgs(OutputString, TEXT("Test G|%" UPTRINT_X_FMT"|"), UPTRINT(49374));
+	check(FString(OutputString) == FString(TEXT("Test G|C0DE|")));
+
+	TestGetVarArgs(OutputString, TEXT("Test H|%" UPTRINT_x_FMT "|"), UPTRINT(49374));
+	check(FString(OutputString) == FString(TEXT("Test H|c0de|")));
+
+	TestGetVarArgs(OutputString, TEXT("Test I|%" UINT64_FMT "|"), MAX_uint64);
+	check(FString(OutputString) == FString(TEXT("Test I|18446744073709551615|")));
 }
 #endif
 
@@ -329,6 +338,11 @@ namespace
 		WIDECHAR* Ptr;
 		WIDECHAR* EndMinusOne; // when null, this means the iterator has already moved past the writable area of the buffer
 	};
+
+	static inline bool CharIsIntegerFormatSpecifier(TCHAR Char)
+	{
+		return (Char == 'i') | (Char == 'd') | (Char == 'u') | (Char == 'X') | (Char == 'x');
+	}
 }
 
 int32 FGenericWidePlatformString::GetVarArgs( WIDECHAR* Dest, SIZE_T DestSize, const WIDECHAR*& Fmt, va_list ArgPtr )
@@ -457,7 +471,7 @@ int32 FGenericWidePlatformString::GetVarArgs( WIDECHAR* Dest, SIZE_T DestSize, c
 				Src++;
 				int Val = va_arg(ArgPtr, int);
 				ANSICHAR AnsiNum[64];
-				ANSICHAR FmtBuf[30] = {0};
+				ANSICHAR FmtBuf[30];
 
 				// limit width to the buffer size
 				FieldLen = FMath::Min(static_cast<int>(sizeof(AnsiNum)) - 1, FieldLen);
@@ -506,7 +520,7 @@ int32 FGenericWidePlatformString::GetVarArgs( WIDECHAR* Dest, SIZE_T DestSize, c
 				size_t Val = va_arg(ArgPtr, size_t);
 
 				ANSICHAR AnsiNum[64];
-				ANSICHAR FmtBuf[30] = {0};
+				ANSICHAR FmtBuf[30];
 
 				// limit width to the buffer size
 				FieldLen = FMath::Min(static_cast<int>(sizeof(AnsiNum)) - 1, FieldLen);
@@ -578,11 +592,11 @@ int32 FGenericWidePlatformString::GetVarArgs( WIDECHAR* Dest, SIZE_T DestSize, c
 					break;
 				}
 
-				// treat %ld as %d. Also shorts for %h will be promoted to ints
-				if ((Src[0] == 'l' && Src[1] == 'd') || Src[0] == 'h')
+				// treat %ld as %d. Also shorts for %h will be promoted to ints. This path also handles %li, %lu, %lx and %lX.
+				if ((Src[0] == 'l' && CharIsIntegerFormatSpecifier(Src[1])) || Src[0] == 'h')
 				{
 					Src+=2;
-					int Val = va_arg(ArgPtr, int);
+					long int Val = va_arg(ArgPtr, long int);
 					ANSICHAR AnsiNum[30];
 					ANSICHAR FmtBuf[30];
 
@@ -608,7 +622,7 @@ int32 FGenericWidePlatformString::GetVarArgs( WIDECHAR* Dest, SIZE_T DestSize, c
 				{
 					Src += 2;
 					double Val = va_arg(ArgPtr, double);
-					ANSICHAR AnsiNum[30];
+					ANSICHAR AnsiNum[340]; // make sure buffer is large enough to handle the longest double values, such as the largest number which is 316 characters including decimal point and 6 decimals
 					ANSICHAR FmtBuf[30];
 
 					// Yes, this is lame.
@@ -629,62 +643,65 @@ int32 FGenericWidePlatformString::GetVarArgs( WIDECHAR* Dest, SIZE_T DestSize, c
 					break;
 				}
 
+				// At this point we expect RemainingSize >= 3, but we know it's at least 2.
+				if (RemainingSize < 3)
+				{
+					printf("Unknown percent [%lc%lc] in FGenericWidePlatformString::GetVarArgs() [%s]\n.", Src[0], Src[1], TCHAR_TO_ANSI(Fmt));
+					Src++;  // skip it, I guess.
+					break;
+				}
+
 				if (Src[0] == 'l')
 				{
-					if (Src[1] != 'l' && Src[1] != 'u' && Src[1] != 'x')
+					if (Src[1] != 'l' || !CharIsIntegerFormatSpecifier(Src[2]))
 					{
-						printf("Unknown percent [%lc%lc] in FGenericWidePlatformString::GetVarArgs() [%s]\n.", Src[0], Src[1], TCHAR_TO_ANSI(Fmt));
-						Src++;  // skip it, I guess.
-						break;
-					}
-					else if (Src[1] == 'l' && RemainingSize == 2)
-					{
-						printf("Unknown percent [%lc%lc] in FGenericWidePlatformString::GetVarArgs() [%s]\n.", Src[0], Src[1], TCHAR_TO_ANSI(Fmt));
+						printf("Unknown percent [%lc%lc%lc] in FGenericWidePlatformString::GetVarArgs() [%s]\n.", Src[0], Src[1], Src[2], TCHAR_TO_ANSI(Fmt));
 						Src++;  // skip it, I guess.
 						break;
 					}
 				}
 				else if (Src[0] == 'I')
 				{
-					if (RemainingSize == 2)
-					{
-						printf("Unknown percent [%lc%lc] in FGenericWidePlatformString::GetVarArgs() [%s]\n.", Src[0], Src[1], TCHAR_TO_ANSI(Fmt));
-						Src++;  // skip it, I guess.
-						break;
-					}
-					else if (RemainingSize == 3)
+					if (RemainingSize < 4)
 					{
 						printf("Unknown percent [%lc%lc%lc] in FGenericWidePlatformString::GetVarArgs() [%s]\n.", Src[0], Src[1], Src[2], TCHAR_TO_ANSI(Fmt));
 						Src++;  // skip it, I guess.
 						break;
 					}
-					else if (Src[1] != '6' || Src[2] != '4')
+					else if (Src[1] != '6' || Src[2] != '4' || !CharIsIntegerFormatSpecifier(Src[3]))
 					{
-						printf("Unknown percent [%lc%lc%lc] in FGenericWidePlatformString::GetVarArgs() [%s]\n.", Src[0], Src[1], Src[2], TCHAR_TO_ANSI(Fmt));
+						printf("Unknown percent [%lc%lc%lc%lc] in FGenericWidePlatformString::GetVarArgs() [%s]\n.", Src[0], Src[1], Src[2], Src[3], TCHAR_TO_ANSI(Fmt));
 						Src++;  // skip it, I guess.
 						break;
 					}
 				}
+				else
+				{
+					printf("Unknown percent [%lc%lc%lc] in FGenericWidePlatformString::GetVarArgs() [%s]\n.", Src[0], Src[1], Src[2], TCHAR_TO_ANSI(Fmt));
+					Src++;  // skip it, I guess.
+					break;
+				}
 
 				// Yes, this is lame.
-				int CpyIdx = 0;
 				unsigned long long Val = va_arg(ArgPtr, unsigned long long);
 				ANSICHAR AnsiNum[60];
 				ANSICHAR FmtBuf[30];
-				if (Src[0] == 'l' && Src[1] == 'l')
-				{
-					Src += 3;
-				}
-				else if (Src[0] == 'l')
+
+				int CpyIdx;
+				if (Src[0] == 'I')
 				{
-					Src += 2;
+					// Convert I64 syntax to ll 
+					Src += 4;
+					Percent += 3;
+					CpyIdx = 3;
+					FmtBuf[0] = '%';
+					FmtBuf[1] = 'l';
+					FmtBuf[2] = 'l';
 				}
 				else
 				{
-					Src += 4;
-					strcpy(FmtBuf, "%L");
-					Percent += 4;
-					CpyIdx = 2;
+					Src += 3;
+					CpyIdx = 0;
 				}
 
 				while (Percent < Src && CpyIdx < ARRAY_COUNT(FmtBuf))
@@ -709,8 +726,7 @@ int32 FGenericWidePlatformString::GetVarArgs( WIDECHAR* Dest, SIZE_T DestSize, c
 			{
 				Src++;
 				double Val = va_arg(ArgPtr, double);
-				// doubles in the form of 1e+9999 can get quite large, make sure we have enough room for them
-				ANSICHAR AnsiNum[48];
+				ANSICHAR AnsiNum[340]; // make sure buffer is large enough to handle the longest double values, such as the largest number which is 316 characters including decimal point and 6 decimals
 				ANSICHAR FmtBuf[30];
 
 				// Yes, this is lame.
diff --git a/Engine/Source/Runtime/Core/Private/Misc/CStringTest.cpp b/Engine/Source/Runtime/Core/Private/Misc/CStringTest.cpp
new file mode 100644
--- /dev/null
+++ b/Engine/Source/Runtime/Core/Private/Misc/CStringTest.cpp
@@ -0,0 +1,96 @@
+// Copyright 1998-2019 Epic Games, Inc. All Rights Reserved.
+
+#include "CoreTypes.h"
+#include "Misc/AutomationTest.h"
+#include "Misc/AssertionMacros.h"
+#include "Misc/CString.h"
+
+#if WITH_DEV_AUTOMATION_TESTS
+
+// This class is a workaround for clang compilers causing error "'va_start' used in function with fixed args" when using it in a lambda in RunTest()
+class FCStringGetVarArgsTestBase : public FAutomationTestBase
+{
+public:
+	FCStringGetVarArgsTestBase(const FString& InName, const bool bInComplexTask)
+	: FAutomationTestBase(InName, bInComplexTask)
+	{
+	}
+
+protected:
+	void DoTest(const TCHAR* ExpectedOutput, const TCHAR* Format, ...)
+	{
+		constexpr SIZE_T OutputBufferCharacterCount = 512;
+		TCHAR OutputBuffer[OutputBufferCharacterCount];
+		va_list ArgPtr;
+		va_start(ArgPtr, Format);
+		const int32 Result = FCString::GetVarArgs(OutputBuffer, OutputBufferCharacterCount, Format, ArgPtr);
+		va_end(ArgPtr);
+
+		if (Result < 0)
+		{
+			this->AddError(FString::Printf(TEXT("'%s' could not be parsed."), Format));
+			return;
+		}
+
+		if (FCString::Strcmp(OutputBuffer, ExpectedOutput) != 0)
+		{
+			this->AddError(FString::Printf(TEXT("'%s' resulted in '%s', expected '%s'."), Format, OutputBuffer, ExpectedOutput));
+			return;
+		}
+	}
+};
+
+IMPLEMENT_CUSTOM_SIMPLE_AUTOMATION_TEST(FCStringGetVarArgsTest, FCStringGetVarArgsTestBase, "System.Core.Misc.CString.GetVarArgs", EAutomationTestFlags::ApplicationContextMask | EAutomationTestFlags::SmokeFilter)
+bool FCStringGetVarArgsTest::RunTest(const FString& Parameters)
+{
+#if PLATFORM_64BITS
+	DoTest(TEXT("SIZE_T_FMT |18446744073709551615|"), TEXT("SIZE_T_FMT |%" SIZE_T_FMT "|"), SIZE_T(MAX_uint64));
+	DoTest(TEXT("SIZE_T_x_FMT |ffffffffffffffff|"), TEXT("SIZE_T_x_FMT |%" SIZE_T_x_FMT "|"), UPTRINT(MAX_uint64));
+	DoTest(TEXT("SIZE_T_X_FMT |FFFFFFFFFFFFFFFF|"), TEXT("SIZE_T_X_FMT |%" SIZE_T_X_FMT "|"), UPTRINT(MAX_uint64));
+
+	DoTest(TEXT("SSIZE_T_FMT |-9223372036854775808|"), TEXT("SSIZE_T_FMT |%" SSIZE_T_FMT "|"), SSIZE_T(MIN_int64));
+	DoTest(TEXT("SSIZE_T_x_FMT |ffffffffffffffff|"), TEXT("SSIZE_T_x_FMT |%" SSIZE_T_x_FMT "|"), SSIZE_T(-1));
+	DoTest(TEXT("SSIZE_T_X_FMT |FFFFFFFFFFFFFFFF|"), TEXT("SSIZE_T_X_FMT |%" SSIZE_T_X_FMT "|"), SSIZE_T(-1));
+
+	DoTest(TEXT("PTRINT_FMT |-9223372036854775808|"), TEXT("PTRINT_FMT |%" PTRINT_FMT "|"), PTRINT(MIN_int64));
+	DoTest(TEXT("PTRINT_x_FMT |ffffffffffffffff|"), TEXT("PTRINT_x_FMT |%" PTRINT_x_FMT "|"), PTRINT(-1));
+	DoTest(TEXT("PTRINT_X_FMT |FFFFFFFFFFFFFFFF|"), TEXT("PTRINT_X_FMT |%" PTRINT_X_FMT "|"), PTRINT(-1));
+
+	DoTest(TEXT("UPTRINT_FMT |18446744073709551615|"), TEXT("UPTRINT_FMT |%" UPTRINT_FMT "|"), UPTRINT(MAX_uint64));
+	DoTest(TEXT("UPTRINT_x_FMT |ffffffffffffffff|"), TEXT("UPTRINT_x_FMT |%" UPTRINT_x_FMT "|"), UPTRINT(MAX_uint64));
+	DoTest(TEXT("UPTRINT_X_FMT |FFFFFFFFFFFFFFFF|"), TEXT("UPTRINT_X_FMT |%" UPTRINT_X_FMT "|"), UPTRINT(MAX_uint64));
+#else
+	DoTest(TEXT("SIZE_T_FMT |4294967295|"), TEXT("SIZE_T_FMT |%" SIZE_T_FMT "|"), SIZE_T(MAX_uint32));
+	DoTest(TEXT("SIZE_T_x_FMT |ffffffff|"), TEXT("SIZE_T_x_FMT |%" SIZE_T_x_FMT "|"), UPTRINT(MAX_uint32));
+	DoTest(TEXT("SIZE_T_X_FMT |FFFFFFFF|"), TEXT("SIZE_T_X_FMT |%" SIZE_T_X_FMT "|"), UPTRINT(MAX_uint32));
+
+	DoTest(TEXT("SSIZE_T_FMT |-2147483648|"), TEXT("SSIZE_T_FMT |%" SSIZE_T_FMT "|"), SSIZE_T(MIN_int32));
+	DoTest(TEXT("SSIZE_T_x_FMT |ffffffff|"), TEXT("SSIZE_T_x_FMT |%" SSIZE_T_x_FMT "|"), SSIZE_T(-1));
+	DoTest(TEXT("SSIZE_T_X_FMT |FFFFFFFF|"), TEXT("SSIZE_T_X_FMT |%" SSIZE_T_X_FMT "|"), SSIZE_T(-1));
+
+	DoTest(TEXT("PTRINT_FMT |-2147483648|"), TEXT("PTRINT_FMT |%" PTRINT_FMT "|"), PTRINT(MIN_int32));
+	DoTest(TEXT("PTRINT_x_FMT |ffffffff|"), TEXT("PTRINT_x_FMT |%" PTRINT_x_FMT "|"), PTRINT(-1));
+	DoTest(TEXT("PTRINT_X_FMT |FFFFFFFF|"), TEXT("PTRINT_X_FMT |%" PTRINT_X_FMT "|"), PTRINT(-1));
+
+	DoTest(TEXT("UPTRINT_FMT |4294967295|"), TEXT("UPTRINT_FMT |%" UPTRINT_FMT "|"), UPTRINT(MAX_uint32));
+	DoTest(TEXT("UPTRINT_x_FMT |ffffffff|"), TEXT("UPTRINT_x_FMT |%" UPTRINT_x_FMT "|"), UPTRINT(MAX_uint32));
+	DoTest(TEXT("UPTRINT_X_FMT |FFFFFFFF|"), TEXT("UPTRINT_X_FMT |%" UPTRINT_X_FMT "|"), UPTRINT(MAX_uint32));
+#endif
+
+	DoTest(TEXT("INT64_FMT |-9223372036854775808|"), TEXT("INT64_FMT |%" INT64_FMT "|"), MIN_int64);
+	DoTest(TEXT("INT64_x_FMT |ffffffffffffffff|"), TEXT("INT64_x_FMT |%" INT64_x_FMT "|"), int64(-1));
+	DoTest(TEXT("INT64_X_FMT |FFFFFFFFFFFFFFFF|"), TEXT("INT64_X_FMT |%" INT64_X_FMT "|"), int64(-1));
+
+	DoTest(TEXT("UINT64_FMT |18446744073709551615|"), TEXT("UINT64_FMT |%" UINT64_FMT "|"), MAX_uint64);
+	DoTest(TEXT("UINT64_x_FMT |ffffffffffffffff|"), TEXT("UINT64_x_FMT |%" UINT64_x_FMT "|"), MAX_uint64);
+	DoTest(TEXT("UINT64_X_FMT |FFFFFFFFFFFFFFFF|"), TEXT("UINT64_X_FMT |%" UINT64_X_FMT "|"), MAX_uint64);
+
+	DoTest(TEXT("|LEFT                |               RIGHT|     33.33|66.67     |"), TEXT("|%-20s|%20s|%10.2f|%-10.2f|"), TEXT("LEFT"), TEXT("RIGHT"), 33.333333, 66.666666);
+
+	DoTest(TEXT("Percents|%%%3|"), TEXT("Percents|%%%%%%%d|"), 3);
+
+	DoTest(TEXT("Integer arguments|12345|54321|123ABC|f|99|"), TEXT("Integer arguments|%d|%i|%X|%x|%u|"), 12345, 54321, 0x123AbC, 15, 99);
+	return true;
+}
+
+#endif // WITH_DEV_AUTOMATION_TESTS
diff --git a/Engine/Source/Runtime/Core/Public/Android/AndroidString.h b/Engine/Source/Runtime/Core/Public/Android/AndroidString.h
--- a/Engine/Source/Runtime/Core/Public/Android/AndroidString.h
+++ b/Engine/Source/Runtime/Core/Public/Android/AndroidString.h
@@ -582,3 +582,60 @@ struct FAndroidPlatformString : public FGenericPlatformString
 
 typedef FAndroidPlatformString FPlatformString;
 
+// Format specifiers to be able to print values of these types correctly, for example when using UE_LOG.
+// SIZE_T format specifier
+#if PLATFORM_64BITS
+#define SIZE_T_FMT "llu"
+// SIZE_T format specifier for lowercase hexadecimal output
+#define SIZE_T_x_FMT "llx"
+// SIZE_T format specifier for uppercase hexadecimal output
+#define SIZE_T_X_FMT "llX"
+
+// SSIZE_T format specifier
+#define SSIZE_T_FMT "lld"
+// SSIZE_T format specifier for lowercase hexadecimal output
+#define SSIZE_T_x_FMT "llx"
+// SSIZE_T format specifier for uppercase hexadecimal output
+#define SSIZE_T_X_FMT "llX"
+#else
+#define SIZE_T_FMT "u"
+// SIZE_T format specifier for lowercase hexadecimal output
+#define SIZE_T_x_FMT "x"
+// SIZE_T format specifier for uppercase hexadecimal output
+#define SIZE_T_X_FMT "X"
+
+// SSIZE_T format specifier
+#define SSIZE_T_FMT "d"
+// SSIZE_T format specifier for lowercase hexadecimal output
+#define SSIZE_T_x_FMT "x"
+// SSIZE_T format specifier for uppercase hexadecimal output
+#define SSIZE_T_X_FMT "X"
+#endif
+
+// PTRINT format specifier for decimal output
+#define PTRINT_FMT SSIZE_T_FMT
+// PTRINT format specifier for lowercase hexadecimal output
+#define PTRINT_x_FMT SSIZE_T_x_FMT
+// PTRINT format specifier for uppercase hexadecimal output
+#define PTRINT_X_FMT SSIZE_T_X_FMT
+
+// UPTRINT format specifier for decimal output
+#define UPTRINT_FMT SIZE_T_FMT
+// UPTRINT format specifier for lowercase hexadecimal output
+#define UPTRINT_x_FMT SIZE_T_x_FMT
+// UPTRINT format specifier for uppercase hexadecimal output
+#define UPTRINT_X_FMT SIZE_T_X_FMT
+
+// int64 format specifier for decimal output
+#define INT64_FMT "lld"
+// int64 format specifier for lowercase hexadecimal output
+#define INT64_x_FMT "llx"
+// int64 format specifier for uppercase hexadecimal output
+#define INT64_X_FMT "llX"
+
+// uint64 format specifier for decimal output
+#define UINT64_FMT "llu"
+// uint64 format specifier for lowercase hexadecimal output
+#define UINT64_x_FMT "llx"
+// uint64 format specifier for uppercase hexadecimal output
+#define UINT64_X_FMT "llX"
diff --git a/Engine/Source/Runtime/Core/Public/Apple/ApplePlatformString.h b/Engine/Source/Runtime/Core/Public/Apple/ApplePlatformString.h
--- a/Engine/Source/Runtime/Core/Public/Apple/ApplePlatformString.h
+++ b/Engine/Source/Runtime/Core/Public/Apple/ApplePlatformString.h
@@ -100,3 +100,46 @@ struct FApplePlatformString
 };
 
 typedef FApplePlatformString FPlatformString;
+
+// Format specifiers to be able to print values of these types correctly, for example when using UE_LOG.
+// SIZE_T format specifier
+#define SIZE_T_FMT "zu"
+// SIZE_T format specifier for lowercase hexadecimal output
+#define SIZE_T_x_FMT "zx"
+// SIZE_T format specifier for uppercase hexadecimal output
+#define SIZE_T_X_FMT "zX"
+
+// SSIZE_T format specifier
+#define SSIZE_T_FMT "lld"
+// SSIZE_T format specifier for lowercase hexadecimal output
+#define SSIZE_T_x_FMT "llx"
+// SSIZE_T format specifier for uppercase hexadecimal output
+#define SSIZE_T_X_FMT "llX"
+
+// PTRINT format specifier for decimal output
+#define PTRINT_FMT SSIZE_T_FMT
+// PTRINT format specifier for lowercase hexadecimal output
+#define PTRINT_x_FMT SSIZE_T_x_FMT
+// PTRINT format specifier for uppercase hexadecimal output
+#define PTRINT_X_FMT SSIZE_T_X_FMT
+
+// UPTRINT format specifier for decimal output
+#define UPTRINT_FMT "llu"
+// UPTRINT format specifier for lowercase hexadecimal output
+#define UPTRINT_x_FMT "llx"
+// UPTRINT format specifier for uppercase hexadecimal output
+#define UPTRINT_X_FMT "llX"
+
+// int64 format specifier for decimal output
+#define INT64_FMT SSIZE_T_FMT
+// int64 format specifier for lowercase hexadecimal output
+#define INT64_x_FMT SSIZE_T_x_FMT
+// int64 format specifier for uppercase hexadecimal output
+#define INT64_X_FMT SSIZE_T_X_FMT
+
+// uint64 format specifier for decimal output
+#define UINT64_FMT "llu"
+// uint64 format specifier for lowercase hexadecimal output
+#define UINT64_x_FMT "llx"
+// uint64 format specifier for uppercase hexadecimal output
+#define UINT64_X_FMT "llX"
diff --git a/Engine/Source/Runtime/Core/Public/HTML5/HTML5PlatformString.h b/Engine/Source/Runtime/Core/Public/HTML5/HTML5PlatformString.h
--- a/Engine/Source/Runtime/Core/Public/HTML5/HTML5PlatformString.h
+++ b/Engine/Source/Runtime/Core/Public/HTML5/HTML5PlatformString.h
@@ -20,3 +20,45 @@ struct FHTML5PlatformString : public FStandardPlatformString
 // default implementation
 typedef FHTML5PlatformString FPlatformString;
 
+// Format specifiers to be able to print values of these types correctly, for example when using UE_LOG.
+// SIZE_T format specifier
+#define SIZE_T_FMT "lu"
+// SIZE_T format specifier for lowercase hexadecimal output
+#define SIZE_T_x_FMT "lx"
+// SIZE_T format specifier for uppercase hexadecimal output
+#define SIZE_T_X_FMT "lX"
+
+// SSIZE_T format specifier
+#define SSIZE_T_FMT "ld"
+// SSIZE_T format specifier for lowercase hexadecimal output
+#define SSIZE_T_x_FMT "lx"
+// SSIZE_T format specifier for uppercase hexadecimal output
+#define SSIZE_T_X_FMT "lX"
+
+// PTRINT format specifier for decimal output
+#define PTRINT_FMT "d"
+// PTRINT format specifier for lowercase hexadecimal output
+#define PTRINT_x_FMT "x"
+// PTRINT format specifier for uppercase hexadecimal output
+#define PTRINT_X_FMT "X"
+
+// UPTRINT format specifier for decimal output
+#define UPTRINT_FMT "u"
+// UPTRINT format specifier for lowercase hexadecimal output
+#define UPTRINT_x_FMT "x"
+// UPTRINT format specifier for uppercase hexadecimal output
+#define UPTRINT_X_FMT "X"
+
+// int64 format specifier for decimal output
+#define INT64_FMT "lld"
+// int64 format specifier for lowercase hexadecimal output
+#define INT64_x_FMT "llx"
+// int64 format specifier for uppercase hexadecimal output
+#define INT64_X_FMT "llX"
+
+// uint64 format specifier for decimal output
+#define UINT64_FMT "llu"
+// uint64 format specifier for lowercase hexadecimal output
+#define UINT64_x_FMT "llx"
+// uint64 format specifier for uppercase hexadecimal output
+#define UINT64_X_FMT "llX"
diff --git a/Engine/Source/Runtime/Core/Public/Unix/UnixPlatformString.h b/Engine/Source/Runtime/Core/Public/Unix/UnixPlatformString.h
--- a/Engine/Source/Runtime/Core/Public/Unix/UnixPlatformString.h
+++ b/Engine/Source/Runtime/Core/Public/Unix/UnixPlatformString.h
@@ -36,3 +36,62 @@ struct FUnixPlatformString : public
 };
 
 typedef FUnixPlatformString FPlatformString;
+
+// Format specifiers to be able to print values of these types correctly, for example when using UE_LOG.
+// SIZE_T format specifier
+#define SIZE_T_FMT "zu"
+// SIZE_T format specifier for lowercase hexadecimal output
+#define SIZE_T_x_FMT "zx"
+// SIZE_T format specifier for uppercase hexadecimal output
+#define SIZE_T_X_FMT "zX"
+
+#if PLATFORM_64BITS
+// SSIZE_T format specifier
+#define SSIZE_T_FMT "lld"
+// SSIZE_T format specifier for lowercase hexadecimal output
+#define SSIZE_T_x_FMT "llx"
+// SSIZE_T format specifier for uppercase hexadecimal output
+#define SSIZE_T_X_FMT "llX"
+
+// UPTRINT format specifier for decimal output
+#define UPTRINT_FMT "llu"
+// UPTRINT format specifier for lowercase hexadecimal output
+#define UPTRINT_x_FMT "llx"
+// UPTRINT format specifier for uppercase hexadecimal output
+#define UPTRINT_X_FMT "llX"
+#else
+// SSIZE_T format specifier
+#define SSIZE_T_FMT "d"
+// SSIZE_T format specifier for lowercase hexadecimal output
+#define SSIZE_T_x_FMT "x"
+// SSIZE_T format specifier for uppercase hexadecimal output
+#define SSIZE_T_X_FMT "X"
+
+// UPTRINT format specifier for decimal output
+#define UPTRINT_FMT "u"
+// UPTRINT format specifier for lowercase hexadecimal output
+#define UPTRINT_x_FMT "x"
+// UPTRINT format specifier for uppercase hexadecimal output
+#define UPTRINT_X_FMT "X"
+#endif
+
+// PTRINT format specifier for decimal output
+#define PTRINT_FMT SSIZE_T_FMT
+// PTRINT format specifier for lowercase hexadecimal output
+#define PTRINT_x_FMT SSIZE_T_x_FMT
+// PTRINT format specifier for uppercase hexadecimal output
+#define PTRINT_X_FMT SSIZE_T_X_FMT
+
+// int64 format specifier for decimal output
+#define INT64_FMT "lld"
+// int64 format specifier for lowercase hexadecimal output
+#define INT64_x_FMT "llx"
+// int64 format specifier for uppercase hexadecimal output
+#define INT64_X_FMT "llX"
+
+// uint64 format specifier for decimal output
+#define UINT64_FMT "llu"
+// uint64 format specifier for lowercase hexadecimal output
+#define UINT64_x_FMT "llx"
+// uint64 format specifier for uppercase hexadecimal output
+#define UINT64_X_FMT "llX"
diff --git a/Engine/Source/Runtime/Core/Public/Windows/WindowsPlatformString.h b/Engine/Source/Runtime/Core/Public/Windows/WindowsPlatformString.h
--- a/Engine/Source/Runtime/Core/Public/Windows/WindowsPlatformString.h
+++ b/Engine/Source/Runtime/Core/Public/Windows/WindowsPlatformString.h
@@ -25,3 +25,76 @@ struct FWindowsPlatformString
 
 
 typedef FWindowsPlatformString FPlatformString;
+
+// Format specifiers to be able to print values of these types correctly, for example when using UE_LOG.
+#if PLATFORM_64BITS
+// SIZE_T format specifier
+#define SIZE_T_FMT "I64u"
+// SIZE_T format specifier for lowercase hexadecimal output
+#define SIZE_T_x_FMT "I64x"
+// SIZE_T format specifier for uppercase hexadecimal output
+#define SIZE_T_X_FMT "I64X"
+
+// SSIZE_T format specifier
+#define SSIZE_T_FMT "I64d"
+// SSIZE_T format specifier for lowercase hexadecimal output
+#define SSIZE_T_x_FMT "I64x"
+// SSIZE_T format specifier for uppercase hexadecimal output
+#define SSIZE_T_X_FMT "I64X"
+
+// PTRINT format specifier for decimal output
+#define PTRINT_FMT "lld"
+// PTRINT format specifier for lowercase hexadecimal output
+#define PTRINT_x_FMT "llx"
+// PTRINT format specifier for uppercase hexadecimal output
+#define PTRINT_X_FMT "llX"
+
+// UPTRINT format specifier for decimal output
+#define UPTRINT_FMT "llu"
+// UPTRINT format specifier for lowercase hexadecimal output
+#define UPTRINT_x_FMT "llx"
+// UPTRINT format specifier for uppercase hexadecimal output
+#define UPTRINT_X_FMT "llX"
+#else
+// SIZE_T format specifier
+#define SIZE_T_FMT "lu"
+// SIZE_T format specifier for lowercase hexadecimal output
+#define SIZE_T_x_FMT "lx"
+// SIZE_T format specifier for uppercase hexadecimal output
+#define SIZE_T_X_FMT "lX"
+
+// SSIZE_T format specifier
+#define SSIZE_T_FMT "ld"
+// SSIZE_T format specifier for lowercase hexadecimal output
+#define SSIZE_T_x_FMT "lx"
+// SSIZE_T format specifier for uppercase hexadecimal output
+#define SSIZE_T_X_FMT "lX"
+
+// PTRINT format specifier for decimal output
+#define PTRINT_FMT "d"
+// PTRINT format specifier for lowercase hexadecimal output
+#define PTRINT_x_FMT "x"
+// PTRINT format specifier for uppercase hexadecimal output
+#define PTRINT_X_FMT "X"
+
+// UPTRINT format specifier for decimal output
+#define UPTRINT_FMT "u"
+// UPTRINT format specifier for lowercase hexadecimal output
+#define UPTRINT_x_FMT "x"
+// UPTRINT format specifier for uppercase hexadecimal output
+#define UPTRINT_X_FMT "X"
+#endif
+
+// int64 format specifier for decimal output
+#define INT64_FMT "lld"
+// int64 format specifier for lowercase hexadecimal output
+#define INT64_x_FMT "llx"
+// int64 format specifier for uppercase hexadecimal output
+#define INT64_X_FMT "llX"
+
+// uint64 format specifier for decimal output
+#define UINT64_FMT "llu"
+// uint64 format specifier for lowercase hexadecimal output
+#define UINT64_x_FMT "llx"
+// uint64 format specifier for uppercase hexadecimal output
+#define UINT64_X_FMT "llX"
