Changing filter type of FCStringGetVarArgsTest so that it doesn't run automatically at engine startup, at least until the test passes on Mac.

diff --git a/Engine/Source/Runtime/Core/Private/Misc/CStringTest.cpp b/Engine/Source/Runtime/Core/Private/Misc/CStringTest.cpp
--- a/Engine/Source/Runtime/Core/Private/Misc/CStringTest.cpp
+++ b/Engine/Source/Runtime/Core/Private/Misc/CStringTest.cpp
@@ -40,7 +40,7 @@ protected:
 	}
 };
 
-IMPLEMENT_CUSTOM_SIMPLE_AUTOMATION_TEST(FCStringGetVarArgsTest, FCStringGetVarArgsTestBase, "System.Core.Misc.CString.GetVarArgs", EAutomationTestFlags::ApplicationContextMask | EAutomationTestFlags::SmokeFilter)
+IMPLEMENT_CUSTOM_SIMPLE_AUTOMATION_TEST(FCStringGetVarArgsTest, FCStringGetVarArgsTestBase, "System.Core.Misc.CString.GetVarArgs", EAutomationTestFlags::ApplicationContextMask | EAutomationTestFlags::EngineFilter)
 bool FCStringGetVarArgsTest::RunTest(const FString& Parameters)
 {
 #if PLATFORM_64BITS
