Added missing IWYU header

diff --git a/Engine/Source/Runtime/Core/Public/Misc/InlineValue.h b/Engine/Source/Runtime/Core/Public/Misc/InlineValue.h
--- a/Engine/Source/Runtime/Core/Public/Misc/InlineValue.h
+++ b/Engine/Source/Runtime/Core/Public/Misc/InlineValue.h
@@ -8,6 +8,7 @@
 #include "Templates/MemoryOps.h"
 #include "Templates/PointerIsConvertibleFromTo.h"
 #include "Templates/Decay.h"
+#include "Templates/TypeCompatibleBytes.h"
 
 /**
  * A container type that houses an instance of BaseType in inline memory where it is <= MaxInlineSize,
