Trick Static Analysis into not complaining about a define constant of intended behaviour

diff --git a/Engine/Source/Runtime/Core/Private/Windows/WindowsPlatformCrashContext.cpp b/Engine/Source/Runtime/Core/Private/Windows/WindowsPlatformCrashContext.cpp
--- a/Engine/Source/Runtime/Core/Private/Windows/WindowsPlatformCrashContext.cpp
+++ b/Engine/Source/Runtime/Core/Private/Windows/WindowsPlatformCrashContext.cpp
@@ -315,7 +315,7 @@ int32 ReportCrashUsingCrashReportClient(FWindowsPlatformCrashContext& InContext,
 	bool bNoDialog = FApp::IsUnattended() || ReportUI == EErrorReportUI::ReportInUnattendedMode || IsRunningDedicatedServer();
 
 	bool bImplicitSend = false;
-	if (!UE_EDITOR && GConfig && ReportUI != EErrorReportUI::ReportInUnattendedMode)
+	if (GConfig && ReportUI != EErrorReportUI::ReportInUnattendedMode && !UE_EDITOR)
 	{
 		// Only check if we are in a non-editor build
 		GConfig->GetBool(TEXT("CrashReportClient"), TEXT("bImplicitSend"), bImplicitSend, GEngineIni);
