Fix compile errors with Visual Studio 2019.

diff --git a/Engine/Source/Runtime/Core/Public/Windows/WindowsPlatformCompilerPreSetup.h b/Engine/Source/Runtime/Core/Public/Windows/WindowsPlatformCompilerPreSetup.h
--- a/Engine/Source/Runtime/Core/Public/Windows/WindowsPlatformCompilerPreSetup.h
+++ b/Engine/Source/Runtime/Core/Public/Windows/WindowsPlatformCompilerPreSetup.h
@@ -97,6 +97,7 @@
 			__pragma(warning(push)) \
 			__pragma(warning(disable: 4510))  /* '<class>': default constructor could not be generated. */ \
 			__pragma(warning(disable: 4610))  /* object '<class>' can never be instantiated - user-defined constructor required. */ \
+			__pragma(warning(disable: 4800))  /* Implicit conversion from '<type>' to bool. Possible information loss. */ \
 			__pragma(warning(disable: 4946))  /* reinterpret_cast used between related classes: '<class1>' and '<class1>' */ \
 			__pragma(warning(disable: 4996))  /* '<obj>' was declared deprecated. */ \
 			__pragma(warning(disable: 6011))  /* Dereferencing NULL pointer '<ptr>'. */ \
diff --git a/Engine/Source/Runtime/Experimental/GeometryCollectionEngine/Private/GeometryCollection/GeometryCollectionSceneProxy.cpp b/Engine/Source/Runtime/Experimental/GeometryCollectionEngine/Private/GeometryCollection/GeometryCollectionSceneProxy.cpp
--- a/Engine/Source/Runtime/Experimental/GeometryCollectionEngine/Private/GeometryCollection/GeometryCollectionSceneProxy.cpp
+++ b/Engine/Source/Runtime/Experimental/GeometryCollectionEngine/Private/GeometryCollection/GeometryCollectionSceneProxy.cpp
@@ -417,7 +417,7 @@ void FGeometryCollectionSceneProxy::SetDynamicData_RenderThread(FGeometryCollect
 
 		if (bSupportsManualVertexFetch)
 		{
-			const bool bLocalGeometryCollectionTripleBufferUploads = GGeometryCollectionTripleBufferUploads;
+			const bool bLocalGeometryCollectionTripleBufferUploads = (GGeometryCollectionTripleBufferUploads != 0);
 
 			if (bLocalGeometryCollectionTripleBufferUploads && TransformBuffers.Num() == 1)
 			{
