Just ifdef versus use the compiler to check for errors

diff --git a/Engine/Source/Runtime/Core/Private/Windows/WindowsPlatformCrashContext.cpp b/Engine/Source/Runtime/Core/Private/Windows/WindowsPlatformCrashContext.cpp
--- a/Engine/Source/Runtime/Core/Private/Windows/WindowsPlatformCrashContext.cpp
+++ b/Engine/Source/Runtime/Core/Private/Windows/WindowsPlatformCrashContext.cpp
@@ -315,11 +315,13 @@ int32 ReportCrashUsingCrashReportClient(FWindowsPlatformCrashContext& InContext,
 	bool bNoDialog = FApp::IsUnattended() || ReportUI == EErrorReportUI::ReportInUnattendedMode || IsRunningDedicatedServer();
 
 	bool bImplicitSend = false;
-	if (GConfig && ReportUI != EErrorReportUI::ReportInUnattendedMode && !UE_EDITOR)
+#if !UE_EDITOR
+	if (GConfig && ReportUI != EErrorReportUI::ReportInUnattendedMode)
 	{
 		// Only check if we are in a non-editor build
 		GConfig->GetBool(TEXT("CrashReportClient"), TEXT("bImplicitSend"), bImplicitSend, GEngineIni);
 	}
+#endif
 
 	bool bSendUnattendedBugReports = true;
 	if (GConfig)
