Fixed race condition in FCrashReportingThread initialization

diff --git a/Engine/Source/Runtime/Core/Private/Windows/WindowsPlatformCrashContext.cpp b/Engine/Source/Runtime/Core/Private/Windows/WindowsPlatformCrashContext.cpp
--- a/Engine/Source/Runtime/Core/Private/Windows/WindowsPlatformCrashContext.cpp
+++ b/Engine/Source/Runtime/Core/Private/Windows/WindowsPlatformCrashContext.cpp
@@ -711,14 +711,15 @@ public:
 		, ExceptionInfo(nullptr)
 		, CrashHandledEvent(nullptr)
 	{
+		// Synchronization objects
+		CrashEvent = CreateEvent(nullptr, true, 0, nullptr);
+		CrashHandledEvent = CreateEvent(nullptr, true, 0, nullptr);
+
 		// Create a background thread that will process the crash and generate crash reports
 		Thread = CreateThread(NULL, 0, CrashReportingThreadProc, this, 0, &ThreadId);
 		if (Thread)
 		{
 			SetThreadPriority(Thread, THREAD_PRIORITY_BELOW_NORMAL);
-			// Synchronization objects
-			CrashEvent = CreateEvent(nullptr, true, 0, nullptr);
-			CrashHandledEvent = CreateEvent(nullptr, true, 0, nullptr);
 		}
 	}
 
@@ -732,13 +733,15 @@ public:
 			if (WaitForSingleObject(Thread, 1000) == WAIT_OBJECT_0)
 			{
 				CloseHandle(Thread);
-				CloseHandle(CrashEvent);
-				CloseHandle(CrashHandledEvent);
 			}
 			Thread = nullptr;
-			CrashEvent = nullptr;
-			CrashHandledEvent = nullptr;
 		}
+
+		CloseHandle(CrashEvent);
+		CrashEvent = nullptr;
+
+		CloseHandle(CrashHandledEvent);
+		CrashHandledEvent = nullptr;
 	}
 
 	/** The thread that crashed calls this function which will trigger the CR thread to report the crash */
