Fixed TInlineAllocator::GetAllocatedSize reporting SecondaryData AllocatedSize even when there wasn't any SecondaryData

diff --git a/Engine/Source/Runtime/Core/Public/Containers/ContainerAllocationPolicies.h b/Engine/Source/Runtime/Core/Public/Containers/ContainerAllocationPolicies.h
--- a/Engine/Source/Runtime/Core/Public/Containers/ContainerAllocationPolicies.h
+++ b/Engine/Source/Runtime/Core/Public/Containers/ContainerAllocationPolicies.h
@@ -545,7 +545,11 @@ public:
 
 		SIZE_T GetAllocatedSize(int32 NumAllocatedElements, SIZE_T NumBytesPerElement) const
 		{
-			return SecondaryData.GetAllocatedSize(NumAllocatedElements, NumBytesPerElement);
+			if (NumAllocatedElements > NumInlineElements)
+			{
+				return SecondaryData.GetAllocatedSize(NumAllocatedElements, NumBytesPerElement);
+			}
+			return 0;
 		}
 
 		bool HasAllocation()
