Allocate MeshDrawCommands at lowest free index instead of last freed index
Shrink MeshDrawCommands periodically when r.DoLazyStaticMeshUpdate is enabled
Measured in replay ES31: Cached MeshDrawCommands 5.236Mb -> 2.216Mb

diff --git a/Engine/Source/Runtime/Core/Public/Containers/BitArray.h b/Engine/Source/Runtime/Core/Public/Containers/BitArray.h
--- a/Engine/Source/Runtime/Core/Public/Containers/BitArray.h
+++ b/Engine/Source/Runtime/Core/Public/Containers/BitArray.h
@@ -642,13 +642,13 @@ public:
 	 * Finds the first zero bit in the array, sets it to true, and returns the bit index.
 	 * If there is none, INDEX_NONE is returned.
 	 */
-	int32 FindAndSetFirstZeroBit()
+	int32 FindAndSetFirstZeroBit(int32 ConservativeStartIndex = 0)
 	{
 		// Iterate over the array until we see a word with a zero bit.
 		uint32* RESTRICT DwordArray = GetData();
 		const int32 LocalNumBits = NumBits;
 		const int32 DwordCount = FMath::DivideAndRoundUp(LocalNumBits, NumBitsPerDWORD);
-		int32 DwordIndex = 0;
+		int32 DwordIndex = FMath::DivideAndRoundDown(ConservativeStartIndex, NumBitsPerDWORD);
 		while (DwordIndex < DwordCount && DwordArray[DwordIndex] == (uint32)-1)
 		{
 			++DwordIndex;
diff --git a/Engine/Source/Runtime/Core/Public/Containers/SparseArray.h b/Engine/Source/Runtime/Core/Public/Containers/SparseArray.h
--- a/Engine/Source/Runtime/Core/Public/Containers/SparseArray.h
+++ b/Engine/Source/Runtime/Core/Public/Containers/SparseArray.h
@@ -141,6 +141,59 @@ public:
 		return Allocation.Index;
 	}
 
+	FSparseArrayAllocationInfo AddUninitializedAtLowestFreeIndex(int32& LowestFreeIndexSearchStart)
+	{
+		int32 Index;
+		if(NumFreeIndices)
+		{
+			Index = AllocationFlags.FindAndSetFirstZeroBit(LowestFreeIndexSearchStart);
+			LowestFreeIndexSearchStart = Index + 1;
+
+			auto& IndexData = GetData(Index);
+
+			// Update FirstFreeIndex
+			if (FirstFreeIndex == Index)
+			{
+				FirstFreeIndex = IndexData.NextFreeIndex;
+			}
+
+			// Link our next and prev free nodes together
+			if (IndexData.NextFreeIndex >= 0)
+			{
+				GetData(IndexData.NextFreeIndex).PrevFreeIndex = IndexData.PrevFreeIndex;
+			}
+
+			if (IndexData.PrevFreeIndex >= 0)
+			{
+				GetData(IndexData.PrevFreeIndex).NextFreeIndex = IndexData.NextFreeIndex;
+			}
+
+			--NumFreeIndices;
+		}
+		else
+		{
+			// Add a new element.
+			Index = Data.AddUninitialized(1);
+			AllocationFlags.Add(true);
+		}
+
+		FSparseArrayAllocationInfo Result;
+		Result.Index = Index;
+		Result.Pointer = &GetData(Result.Index).ElementData;
+		return Result;
+	}
+
+	/** 
+	 * Add an element at the lowest free index, instead of the last freed index. 
+	 * This requires a search which can be accelerated with LowestFreeIndexSearchStart.
+	 */
+	int32 AddAtLowestFreeIndex(const ElementType& Element, int32& LowestFreeIndexSearchStart)
+	{
+		FSparseArrayAllocationInfo Allocation = AddUninitializedAtLowestFreeIndex(LowestFreeIndexSearchStart);
+		new(Allocation) ElementType(Element);
+		return Allocation.Index;
+	}
+
 	/**
 	 * Allocates space for an element in the array at a given index.  The element is not initialized, and you must use the corresponding placement new operator
 	 * to construct the element in the allocated memory.
diff --git a/Engine/Source/Runtime/Renderer/Private/MeshPassProcessor.cpp b/Engine/Source/Runtime/Renderer/Private/MeshPassProcessor.cpp
--- a/Engine/Source/Runtime/Renderer/Private/MeshPassProcessor.cpp
+++ b/Engine/Source/Runtime/Renderer/Private/MeshPassProcessor.cpp
@@ -1037,7 +1037,8 @@ FMeshDrawCommand& FCachedPassMeshDrawListContext::AddCommand(const FMeshDrawComm
 	{
 		// Only one FMeshDrawCommand supported per FStaticMesh in a pass
 		check(CommandInfo.CommandIndex == -1);
-		CommandInfo.CommandIndex = DrawList.MeshDrawCommands.Add(Initializer);
+		// Allocate at lowest free index so that 'r.DoLazyStaticMeshUpdate' can shrink the TSparseArray more effectively
+		CommandInfo.CommandIndex = DrawList.MeshDrawCommands.AddAtLowestFreeIndex(Initializer, DrawList.LowestFreeIndexSearchStart);
 		return DrawList.MeshDrawCommands[CommandInfo.CommandIndex];
 	}
 }
diff --git a/Engine/Source/Runtime/Renderer/Private/PrimitiveSceneInfo.cpp b/Engine/Source/Runtime/Renderer/Private/PrimitiveSceneInfo.cpp
--- a/Engine/Source/Runtime/Renderer/Private/PrimitiveSceneInfo.cpp
+++ b/Engine/Source/Runtime/Renderer/Private/PrimitiveSceneInfo.cpp
@@ -312,6 +312,9 @@ void FPrimitiveSceneInfo::RemoveCachedMeshDrawCommands()
 		{
 			FCachedPassMeshDrawList& PassDrawList = Scene->CachedDrawLists[CachedCommand.MeshPass];
 			PassDrawList.MeshDrawCommands.RemoveAt(CachedCommand.CommandIndex);
+
+			// Track the lowest index that might be free for faster AddAtLowestFreeIndex
+			PassDrawList.LowestFreeIndexSearchStart = FMath::Min(PassDrawList.LowestFreeIndexSearchStart, CachedCommand.CommandIndex);
 		}
 	}
 
diff --git a/Engine/Source/Runtime/Renderer/Private/SceneVisibility.cpp b/Engine/Source/Runtime/Renderer/Private/SceneVisibility.cpp
--- a/Engine/Source/Runtime/Renderer/Private/SceneVisibility.cpp
+++ b/Engine/Source/Runtime/Renderer/Private/SceneVisibility.cpp
@@ -2773,9 +2773,17 @@ void FSceneRenderer::PreVisibilityFrameSetup(FRHICommandListImmediate& RHICmdLis
 		{
 			QUICK_SCOPE_CYCLE_COUNTER(STAT_PreVisibilityFrameSetup_EvictionForLazyStaticMeshUpdate);
 			static int32 RollingRemoveIndex = 0;
+			static int32 RollingPassShrinkIndex = 0;
 			if (RollingRemoveIndex >= Scene->Primitives.Num())
 			{
 				RollingRemoveIndex = 0;
+				RollingPassShrinkIndex++;
+				if (RollingPassShrinkIndex >= ARRAY_COUNT(Scene->CachedDrawLists))
+				{
+					RollingPassShrinkIndex = 0;
+				}
+				// Periodically shrink the SparseArray containing cached mesh draw commands which we are causing to be regenerated with UpdateStaticMeshes
+				Scene->CachedDrawLists[RollingPassShrinkIndex].MeshDrawCommands.Shrink();
 			}
 			const int32 NumRemovedPerFrame = 10;
 			for (int32 NumRemoved = 0; NumRemoved < NumRemovedPerFrame && RollingRemoveIndex < Scene->Primitives.Num(); NumRemoved++, RollingRemoveIndex++)
diff --git a/Engine/Source/Runtime/Renderer/Public/MeshPassProcessor.h b/Engine/Source/Runtime/Renderer/Public/MeshPassProcessor.h
--- a/Engine/Source/Runtime/Renderer/Public/MeshPassProcessor.h
+++ b/Engine/Source/Runtime/Renderer/Public/MeshPassProcessor.h
@@ -714,8 +714,14 @@ public:
 class FCachedPassMeshDrawList
 {
 public:
+
+	FCachedPassMeshDrawList() :
+		LowestFreeIndexSearchStart(0)
+	{}
+
 	/** Indices held by FStaticMeshBatch::CachedMeshDrawCommands must be stable */
 	TSparseArray<FMeshDrawCommand> MeshDrawCommands;
+	int32 LowestFreeIndexSearchStart;
 };
 
 typedef TArray<int32, TInlineAllocator<5>> FDrawCommandIndices;
