Water changes
- Added better defaults for depth and width
- Hid the static mesh from the water actor. This is defined by the water type and should in general not be changed.  An override is still around as an advanced property as needed
- Tied depth and width properties to spline scale x and y.  They should always be in lockstep now.
- Added fixup code to fix up exising tests

diff --git a/Engine/Source/Runtime/Core/Public/UObject/FortniteMainBranchObjectVersion.h b/Engine/Source/Runtime/Core/Public/UObject/FortniteMainBranchObjectVersion.h
--- a/Engine/Source/Runtime/Core/Public/UObject/FortniteMainBranchObjectVersion.h
+++ b/Engine/Source/Runtime/Core/Public/UObject/FortniteMainBranchObjectVersion.h
@@ -90,6 +90,9 @@ struct CORE_API FFortniteMainBranchObjectVersion
 		// Support Virtual Bone in Retarget Manager
 		SupportVirtualBoneInRetargeting,
 
+		// Fixup bad defaults in water metadata
+		FixUpWaterMetadata,
+
 		// -----<new versions can be added above this line>-------------------------------------------------
 		VersionPlusOne,
 		LatestVersion = VersionPlusOne - 1
