Fixed TArrayView assignment operator, to work properly with compatible types.

diff --git a/Engine/Source/Runtime/Core/Public/Containers/ArrayView.h b/Engine/Source/Runtime/Core/Public/Containers/ArrayView.h
--- a/Engine/Source/Runtime/Core/Public/Containers/ArrayView.h
+++ b/Engine/Source/Runtime/Core/Public/Containers/ArrayView.h
@@ -170,11 +170,9 @@ public:
 		typename = typename TEnableIf<TIsCompatibleElementType<OtherElementType>::Value>::Type>
 	FORCEINLINE TArrayView& operator=(const TArrayView<OtherElementType>& Other)
 	{
-		if (this != &Other)
-		{
-			DataPtr = Other.DataPtr;
-			ArrayNum = Other.ArrayNum;
-		}
+		DataPtr = Other.GetData();
+		ArrayNum = Other.Num();
+
 		return *this;
 	}
 
