FPaths::GetCleanFilename now returns empty strings for directories so need a way to get leaf directories over abusing that function

Move previous implementation into its own GetPathLeaf function which should allow getting the leaf of the path regardless if its a file or dir

diff --git a/Engine/Source/Runtime/Core/Private/Misc/Paths.cpp b/Engine/Source/Runtime/Core/Private/Misc/Paths.cpp
--- a/Engine/Source/Runtime/Core/Private/Misc/Paths.cpp
+++ b/Engine/Source/Runtime/Core/Private/Misc/Paths.cpp
@@ -759,6 +759,30 @@ FString FPaths::GetPath(FString&& InPath)
 	return Result;
 }
 
+FString FPaths::GetPathLeaf(const FString& InPath)
+{
+	static_assert(INDEX_NONE == -1, "INDEX_NONE assumed to be -1");
+
+	int32 EndPos   = InPath.FindLastCharByPredicate(UE4Paths_Private::IsNotSlashOrBackslash) + 1;
+	int32 StartPos = InPath.FindLastCharByPredicate(UE4Paths_Private::IsSlashOrBackslash, EndPos) + 1;
+
+	FString Result = InPath.Mid(StartPos, EndPos - StartPos);
+	return Result;
+}
+
+FString FPaths::GetPathLeaf(FString&& InPath)
+{
+	static_assert(INDEX_NONE == -1, "INDEX_NONE assumed to be -1");
+
+	int32 EndPos   = InPath.FindLastCharByPredicate(UE4Paths_Private::IsNotSlashOrBackslash) + 1;
+	int32 StartPos = InPath.FindLastCharByPredicate(UE4Paths_Private::IsSlashOrBackslash, EndPos) + 1;
+
+	InPath.RemoveAt(EndPos, InPath.Len() - EndPos, false);
+	InPath.RemoveAt(0, StartPos, false);
+
+	return MoveTemp(InPath);
+}
+
 FString FPaths::ChangeExtension(const FString& InPath, const FString& InNewExtension)
 {
 	int32 Pos = INDEX_NONE;
diff --git a/Engine/Source/Runtime/Core/Public/Misc/Paths.h b/Engine/Source/Runtime/Core/Public/Misc/Paths.h
--- a/Engine/Source/Runtime/Core/Public/Misc/Paths.h
+++ b/Engine/Source/Runtime/Core/Public/Misc/Paths.h
@@ -397,6 +397,12 @@ public:
 	// Returns the path in front of the filename
 	static FString GetPath(FString&& InPath);
 
+	// Returns the leaf in the path
+	static FString GetPathLeaf(const FString& InPath);
+
+	// Returns the leaf in the path
+	static FString GetPathLeaf(FString&& InPath);
+
 	/** Changes the extension of the given filename (does nothing if the file has no extension) */
 	static FString ChangeExtension(const FString& InPath, const FString& InNewExtension);
 
diff --git a/Engine/Source/Runtime/CrashReportCore/Public/GenericErrorReport.h b/Engine/Source/Runtime/CrashReportCore/Public/GenericErrorReport.h
--- a/Engine/Source/Runtime/CrashReportCore/Public/GenericErrorReport.h
+++ b/Engine/Source/Runtime/CrashReportCore/Public/GenericErrorReport.h
@@ -90,8 +90,8 @@ public:
 	 */
 	FString GetReportDirectoryLeafName() const
 	{
-		// Using GetCleanFilename to actually get directory name
-		return FPaths::GetCleanFilename(ReportDirectory);
+		// Using GetPathLeaf to actually get directory name
+		return FPaths::GetPathLeaf(ReportDirectory);
 	}
 
 	/**
