Set our current thread when getting into the default crash handler

diff --git a/Engine/Source/Runtime/Core/Private/Unix/UnixPlatformCrashContext.cpp b/Engine/Source/Runtime/Core/Private/Unix/UnixPlatformCrashContext.cpp
--- a/Engine/Source/Runtime/Core/Private/Unix/UnixPlatformCrashContext.cpp
+++ b/Engine/Source/Runtime/Core/Private/Unix/UnixPlatformCrashContext.cpp
@@ -680,6 +680,7 @@ void DefaultCrashHandler(const FUnixCrashContext & Context)
 	const_cast<FUnixCrashContext&>(Context).CaptureStackTrace();
 	if (GLog)
 	{
+		GLog->SetCurrentThreadAsMasterThread();
 		GLog->Flush();
 	}
 	if (GWarn)
