Bug fix to reenable hardware decompression for XboxOne
Also fixes https://udn.unrealengine.com/questions/514260/xboxone-cook-crash-xboxonezlib.html#comment-514260-form

diff --git a/Engine/Source/Runtime/Core/Private/Misc/Compression.cpp b/Engine/Source/Runtime/Core/Private/Misc/Compression.cpp
--- a/Engine/Source/Runtime/Core/Private/Misc/Compression.cpp
+++ b/Engine/Source/Runtime/Core/Private/Misc/Compression.cpp
@@ -23,6 +23,7 @@ DEFINE_LOG_CATEGORY(LogCompression);
 DECLARE_STATS_GROUP( TEXT( "Compression" ), STATGROUP_Compression, STATCAT_Advanced );
 
 TMap<FName, struct ICompressionFormat*> FCompression::CompressionFormats;
+FCriticalSection FCompression::CompressionFormatsCriticalSection;
 
 
 static void *zalloc(void *opaque, unsigned int size, unsigned int num)
@@ -306,6 +307,7 @@ static ECompressionFlags CheckGlobalCompressionFlags(ECompressionFlags Flags)
 
 ICompressionFormat* FCompression::GetCompressionFormat(FName FormatName, bool bErrorOnFailure)
 {
+	FScopeLock Lock(&CompressionFormatsCriticalSection);
 	ICompressionFormat** ExistingFormat = CompressionFormats.Find(FormatName);
 	if (ExistingFormat == nullptr)
 	{
diff --git a/Engine/Source/Runtime/Core/Public/Misc/Compression.h b/Engine/Source/Runtime/Core/Public/Misc/Compression.h
--- a/Engine/Source/Runtime/Core/Public/Misc/Compression.h
+++ b/Engine/Source/Runtime/Core/Public/Misc/Compression.h
@@ -5,6 +5,7 @@
 #include "CoreMinimal.h"
 #include "Templates/Atomic.h"
 #include "Misc/CompressionFlags.h"
+#include "HAL/CriticalSection.h"
 
 class IMemoryReadStream;
 
@@ -103,6 +104,8 @@ private:
 
 	/** Mapping of Compression FNames to their compressor objects */
 	static TMap<FName, struct ICompressionFormat*> CompressionFormats;
+	static FCriticalSection CompressionFormatsCriticalSection;
+
 };
 
 
