Fix Framepro's DynamicWString wrapper class - each time the assignment operator is called, it chops off the last character of the string.

diff --git a/Engine/Source/Runtime/Core/Private/FramePro/FramePro.cpp b/Engine/Source/Runtime/Core/Private/FramePro/FramePro.cpp
--- a/Engine/Source/Runtime/Core/Private/FramePro/FramePro.cpp
+++ b/Engine/Source/Runtime/Core/Private/FramePro/FramePro.cpp
@@ -1079,7 +1079,7 @@ namespace FramePro
 			mp_Value = (wchar_t*)mp_Allocator->Alloc((len + 1) * sizeof(wchar_t));
 			for(int i = 0; i < len; ++i)
 				mp_Value[i] = p_value[i];
-			mp_Value[len - 1] = L'\0';
+			mp_Value[len] = L'\0';
 		}
 
 		//------------------------------------------------------------------------
@@ -1092,7 +1092,7 @@ namespace FramePro
 			mp_Value = (wchar_t*)mp_Allocator->Alloc((len + 1) * sizeof(wchar_t));
 			for(int i = 0; i < len; ++i)
 				mp_Value[i] = p_value[i];
-			mp_Value[len - 1] = L'\0';
+			mp_Value[len] = L'\0';
 		}
 
 		//------------------------------------------------------------------------
