Fix a non unity build break

diff --git a/Engine/Source/Runtime/Core/Private/Misc/Compression.cpp b/Engine/Source/Runtime/Core/Private/Misc/Compression.cpp
--- a/Engine/Source/Runtime/Core/Private/Misc/Compression.cpp
+++ b/Engine/Source/Runtime/Core/Private/Misc/Compression.cpp
@@ -5,6 +5,7 @@
 #include "HAL/UnrealMemory.h"
 #include "Logging/LogMacros.h"
 #include "Misc/Parse.h"
+#include "Misc/ScopeLock.h"
 #include "Misc/CommandLine.h"
 #include "Stats/Stats.h"
 #include "Misc/ConfigCacheIni.h"
