properly set the memory values in the OOM crash report

diff --git a/Engine/Source/Runtime/Core/Private/Apple/ApplePlatformMemory.cpp b/Engine/Source/Runtime/Core/Private/Apple/ApplePlatformMemory.cpp
--- a/Engine/Source/Runtime/Core/Private/Apple/ApplePlatformMemory.cpp
+++ b/Engine/Source/Runtime/Core/Private/Apple/ApplePlatformMemory.cpp
@@ -257,7 +257,7 @@ FPlatformMemoryStats FApplePlatformMemory::GetStats()
 {
 	const FPlatformMemoryConstants& MemoryConstants = FPlatformMemory::GetConstants();
 	
-	FPlatformMemoryStats MemoryStats;
+	static FPlatformMemoryStats MemoryStats;
 	
 	// Gather platform memory stats.
 	vm_statistics Stats;
