Core: Added FAutoConsoleCommandWithWorldArgsAndOutputDevice wrapper to pair with FConsoleCommandWithWorldArgsAndOutputDeviceDelegate

diff --git a/Engine/Source/Runtime/Core/Public/HAL/IConsoleManager.h b/Engine/Source/Runtime/Core/Public/HAL/IConsoleManager.h
--- a/Engine/Source/Runtime/Core/Public/HAL/IConsoleManager.h
+++ b/Engine/Source/Runtime/Core/Public/HAL/IConsoleManager.h
@@ -1234,8 +1234,27 @@ public:
 		: FAutoConsoleObject(IConsoleManager::Get().RegisterConsoleCommand(Name, Help, Command, Flags))
 	{
 	}
+};
 
-	
+
+/**
+ * Autoregistering console command with world, args, an output device
+ */
+class CORE_API FAutoConsoleCommandWithWorldArgsAndOutputDevice : private FAutoConsoleObject
+{
+public:
+	/**
+	 * Register a console command that takes an output device
+	 *
+	 * @param	Name		The name of this command (must not be nullptr)
+	 * @param	Help		Help text for this command
+	 * @param	Command		The user function to call when this command is executed
+	 * @param	Flags		Optional flags bitmask
+	 */
+	FAutoConsoleCommandWithWorldArgsAndOutputDevice(const TCHAR* Name, const TCHAR* Help, const FConsoleCommandWithWorldArgsAndOutputDeviceDelegate& Command, uint32 Flags = ECVF_Default)
+		: FAutoConsoleObject(IConsoleManager::Get().RegisterConsoleCommand(Name, Help, Command, Flags))
+	{
+	}
 };
 
 /**
