Remove duplicate function definition
#rnx

diff --git a/Engine/Source/Runtime/Core/Public/HAL/IConsoleManager.h b/Engine/Source/Runtime/Core/Public/HAL/IConsoleManager.h
--- a/Engine/Source/Runtime/Core/Public/HAL/IConsoleManager.h
+++ b/Engine/Source/Runtime/Core/Public/HAL/IConsoleManager.h
@@ -1257,25 +1257,5 @@ public:
 	}
 };
 
-/**
- * Autoregistering console command with a world, arguments, and output device
- */
-class CORE_API FAutoConsoleCommandWithWorldArgsAndOutputDevice : private FAutoConsoleObject
-{
-public:
-	/**
-	* Register a console command that takes arguments, a world argument and an output device
-	*
-	* @param	Name		The name of this command (must not be nullptr)
-	* @param	Help		Help text for this command
-	* @param	Command		The user function to call when this command is executed
-	* @param	Flags		Optional flags bitmask
-	*/
-	FAutoConsoleCommandWithWorldArgsAndOutputDevice(const TCHAR* Name, const TCHAR* Help, const FConsoleCommandWithWorldArgsAndOutputDeviceDelegate& Command, uint32 Flags = ECVF_Default)
-		: FAutoConsoleObject(IConsoleManager::Get().RegisterConsoleCommand(Name, Help, Command, Flags))
-	{
-	}
-};
-
 CORE_API DECLARE_LOG_CATEGORY_EXTERN(LogConsoleResponse, Log, All);
 
