Add FPlatformMisc::GetPlatformChunkInstall() to Linux

diff --git a/Engine/Source/Runtime/Core/Private/Unix/UnixPlatformMisc.cpp b/Engine/Source/Runtime/Core/Private/Unix/UnixPlatformMisc.cpp
--- a/Engine/Source/Runtime/Core/Private/Unix/UnixPlatformMisc.cpp
+++ b/Engine/Source/Runtime/Core/Private/Unix/UnixPlatformMisc.cpp
@@ -20,7 +20,7 @@
 #include "Misc/CommandLine.h"
 #include "Misc/App.h"
 #include "Unix/UnixPlatformCrashContext.h"
-
+#include "GenericPlatform/GenericPlatformChunkInstall.h"
 
 #if PLATFORM_HAS_CPUID
 	#include <cpuid.h>
@@ -1041,3 +1041,37 @@ FString FUnixPlatformMisc::GetLoginId()
 {
 	return FString::Printf(TEXT("%s-%08x"), *GetOperatingSystemId(), static_cast<uint32>(geteuid()));
 }
+
+IPlatformChunkInstall* FUnixPlatformMisc::GetPlatformChunkInstall()
+{
+	static IPlatformChunkInstall* ChunkInstall = nullptr;
+	static bool bIniChecked = false;
+	if (!ChunkInstall || !bIniChecked)
+	{
+		IPlatformChunkInstallModule* PlatformChunkInstallModule = nullptr;
+		if (!GEngineIni.IsEmpty())
+		{
+			FString InstallModule;
+			GConfig->GetString(TEXT("StreamingInstall"), TEXT("DefaultProviderName"), InstallModule, GEngineIni);
+			FModuleStatus Status;
+			if (FModuleManager::Get().QueryModule(*InstallModule, Status))
+			{
+				PlatformChunkInstallModule = FModuleManager::LoadModulePtr<IPlatformChunkInstallModule>(*InstallModule);
+				if (PlatformChunkInstallModule != nullptr)
+				{
+					// Attempt to grab the platform installer
+					ChunkInstall = PlatformChunkInstallModule->GetPlatformChunkInstall();
+				}
+			}
+			bIniChecked = true;
+		}
+
+		if (PlatformChunkInstallModule == nullptr)
+		{
+			// Placeholder instance
+			ChunkInstall = FGenericPlatformMisc::GetPlatformChunkInstall();
+		}
+	}
+
+	return ChunkInstall;
+}
diff --git a/Engine/Source/Runtime/Core/Public/Unix/UnixPlatformMisc.h b/Engine/Source/Runtime/Core/Public/Unix/UnixPlatformMisc.h
--- a/Engine/Source/Runtime/Core/Public/Unix/UnixPlatformMisc.h
+++ b/Engine/Source/Runtime/Core/Public/Unix/UnixPlatformMisc.h
@@ -136,6 +136,8 @@ struct CORE_API FUnixPlatformMisc : public FGenericPlatformMisc
 
 	static void CreateGuid(FGuid& Result);
 
+	static IPlatformChunkInstall* GetPlatformChunkInstall();
+
 #if STATS || ENABLE_STATNAMEDEVENTS
 	static void BeginNamedEventFrame();
 	static void BeginNamedEvent(const struct FColor& Color, const TCHAR* Text);
