Install Bundle Manager
- Add ability to force a full install.
- Add ability to query full install size specifically.

diff --git a/Engine/Source/Runtime/Core/Public/GenericPlatform/GenericPlatformInstallBundleManager.h b/Engine/Source/Runtime/Core/Public/GenericPlatform/GenericPlatformInstallBundleManager.h
--- a/Engine/Source/Runtime/Core/Public/GenericPlatform/GenericPlatformInstallBundleManager.h
+++ b/Engine/Source/Runtime/Core/Public/GenericPlatform/GenericPlatformInstallBundleManager.h
@@ -111,6 +111,7 @@ enum class EInstallBundleRequestFlags : uint32
 	CheckForCellularDataUsage						= (1 << 0),
 	UseBackgroundDownloads							= (1 << 1),
 	SendNotificationIfDownloadCompletesInBackground = (1 << 2),
+	ForceNoPatching									= (1 << 3),
 	Defaults = UseBackgroundDownloads,
 };
 ENUM_CLASS_FLAGS(EInstallBundleRequestFlags)
@@ -236,6 +237,13 @@ struct FInstallBundleContentState
 	uint64 FreeSpace = 0;
 };
 
+enum class EInstallBundleGetContentStateFlags : uint32
+{
+	None = 0,
+	ForceNoPatching = (1 << 0),
+};
+ENUM_CLASS_FLAGS(EInstallBundleGetContentStateFlags);
+
 enum class EInstallBundleRequestInfoFlags : int32
 {
 	None							= 0,
@@ -299,8 +307,8 @@ public:
 	virtual FInstallBundleRequestInfo RequestUpdateContent(FName BundleName, EInstallBundleRequestFlags Flags) = 0;
 	virtual FInstallBundleRequestInfo RequestUpdateContent(TArrayView<FName> BundleNames, EInstallBundleRequestFlags Flags) = 0;
 
-	virtual void GetContentState(FName BundleName, bool bAddDependencies, FInstallBundleGetContentStateDelegate Callback, FName RequestTag = TEXT("None")) = 0;
-	virtual void GetContentState(TArrayView<FName> BundleNames, bool bAddDependencies, FInstallBundleGetContentStateDelegate Callback, FName RequestTag = TEXT("None")) = 0;
+	virtual void GetContentState(FName BundleName, EInstallBundleGetContentStateFlags Flags, bool bAddDependencies, FInstallBundleGetContentStateDelegate Callback, FName RequestTag = TEXT("None")) = 0;
+	virtual void GetContentState(TArrayView<FName> BundleNames, EInstallBundleGetContentStateFlags Flags, bool bAddDependencies, FInstallBundleGetContentStateDelegate Callback, FName RequestTag = TEXT("None")) = 0;
 
     virtual void CancelAllGetContentStateRequestsForTag(FName RequestTag) = 0;
     
diff --git a/Engine/Source/Runtime/NullInstallBundleManager/Source/NullInstallBundleManager.cpp b/Engine/Source/Runtime/NullInstallBundleManager/Source/NullInstallBundleManager.cpp
--- a/Engine/Source/Runtime/NullInstallBundleManager/Source/NullInstallBundleManager.cpp
+++ b/Engine/Source/Runtime/NullInstallBundleManager/Source/NullInstallBundleManager.cpp
@@ -50,14 +50,14 @@ class FNullInstallBundleManager : public IPlatformInstallBundleManager
 		return RetInfo;
 	}
 
-	virtual void GetContentState(FName BundleName, bool bAddDependencies, FInstallBundleGetContentStateDelegate Callback, FName RequestTag) override
+	virtual void GetContentState(FName BundleName, EInstallBundleGetContentStateFlags Flags, bool bAddDependencies, FInstallBundleGetContentStateDelegate Callback, FName RequestTag) override
 	{
 		FInstallBundleContentState State;
 		State.State = EInstallBundleContentState::UpToDate;
 		Callback.ExecuteIfBound(State);
 	}
 
-	virtual void GetContentState(TArrayView<FName> BundleNames, bool bAddDependencies, FInstallBundleGetContentStateDelegate Callback, FName RequestTag) override
+	virtual void GetContentState(TArrayView<FName> BundleNames, EInstallBundleGetContentStateFlags Flags, bool bAddDependencies, FInstallBundleGetContentStateDelegate Callback, FName RequestTag) override
 	{
 		FInstallBundleContentState State;
 		State.State = EInstallBundleContentState::UpToDate;
