Copy of CL 7605175

Changes to make new Clang in Xcode 11 beta compile the engine witout warnings

diff --git a/Engine/Source/Developer/Localization/Public/LocTextHelper.h b/Engine/Source/Developer/Localization/Public/LocTextHelper.h
--- a/Engine/Source/Developer/Localization/Public/LocTextHelper.h
+++ b/Engine/Source/Developer/Localization/Public/LocTextHelper.h
@@ -113,6 +113,10 @@ private:
 public:
 	FLocTextConflicts() {}
 
+	/** Publicly movable */
+	FLocTextConflicts(FLocTextConflicts&&) = default;
+	FLocTextConflicts& operator=(FLocTextConflicts&&) = default;
+	
 	/**
 	 * Add a new conflict entry.
 	 *
diff --git a/Engine/Source/Runtime/Core/Private/Apple/AppleLLM.cpp b/Engine/Source/Runtime/Core/Private/Apple/AppleLLM.cpp
--- a/Engine/Source/Runtime/Core/Private/Apple/AppleLLM.cpp
+++ b/Engine/Source/Runtime/Core/Private/Apple/AppleLLM.cpp
@@ -26,15 +26,10 @@ const FLLMTagInfoApple ELLMTagNamesApple[] =
 	{ TEXT("Objective-C"),				GET_STATFNAME(STAT_ObjectiveCLLM),		NAME_None },									// ELLMTagApple::ObjectiveC
 };
 
-typedef id (*AllocWithZoneFunc)(id Obj, SEL Sel, struct _NSZone* Zone);
-typedef void (*DeallocFunc)(id Obj, SEL Sel);
-
-static AllocWithZoneFunc AllocWithZoneOriginal = nullptr;
-static DeallocFunc DeallocOriginal = nullptr;
-
-
-static IMP AllocWithZoneOriginalIMP = nullptr;
-static IMP DeallocOriginalIMP = nullptr;
+typedef id (*AllocWithZoneIMP)(id Obj, SEL Sel, struct _NSZone* Zone);
+typedef id (*DeallocIMP)(id Obj, SEL Sel);
+static AllocWithZoneIMP AllocWithZoneOriginal = nullptr;
+static DeallocIMP DeallocOriginal = nullptr;
 
 static id AllocWithZoneInterposer(id Obj, SEL Sel, struct _NSZone * Zone)
 {
@@ -124,14 +119,11 @@ void AppleLLM::Initialise()
 	Method AllocZone = class_getClassMethod([NSObject class], @selector(allocWithZone:));
 	Method Dealloc = class_getInstanceMethod([NSObject class], @selector(dealloc));
 	
-	AllocWithZoneOriginalIMP = method_getImplementation(AllocZone);
-	DeallocOriginalIMP = method_getImplementation(Dealloc);
-	
-	check(AllocWithZoneOriginalIMP);
-	check(DeallocOriginalIMP);
-    
-    AllocWithZoneOriginal = (AllocWithZoneFunc) AllocWithZoneOriginalIMP;
-    DeallocOriginal = (DeallocFunc) DeallocOriginalIMP;
+	AllocWithZoneOriginal = (AllocWithZoneIMP)method_getImplementation(AllocZone);
+	DeallocOriginal = (DeallocIMP)method_getImplementation(Dealloc);
+
+	check(AllocWithZoneOriginal);
+	check(DeallocOriginal);
 	
 	method_setImplementation(AllocZone, (IMP)&AllocWithZoneInterposer);
 	method_setImplementation(Dealloc, (IMP)&DeallocInterposer);
diff --git a/Engine/Source/Runtime/Landscape/Private/LandscapeGizmoActor.cpp b/Engine/Source/Runtime/Landscape/Private/LandscapeGizmoActor.cpp
--- a/Engine/Source/Runtime/Landscape/Private/LandscapeGizmoActor.cpp
+++ b/Engine/Source/Runtime/Landscape/Private/LandscapeGizmoActor.cpp
@@ -26,18 +26,6 @@
 #include "Components/BillboardComponent.h"
 #include "HAL/PlatformApplicationMisc.h"
 
-namespace
-{
-	enum PreviewType
-	{
-		Invalid = -1,
-		Both = 0,
-		Add = 1,
-		Sub = 2,
-	};
-}
-
-
 class FLandscapeGizmoMeshRenderProxy : public FMaterialRenderProxy
 {
 public:
