Fix build failure

diff --git a/Engine/Source/Runtime/Core/Private/Unix/UnixPlatformMisc.cpp b/Engine/Source/Runtime/Core/Private/Unix/UnixPlatformMisc.cpp
--- a/Engine/Source/Runtime/Core/Private/Unix/UnixPlatformMisc.cpp
+++ b/Engine/Source/Runtime/Core/Private/Unix/UnixPlatformMisc.cpp
@@ -20,6 +20,7 @@
 #include "Misc/CommandLine.h"
 #include "Misc/App.h"
 #include "Unix/UnixPlatformCrashContext.h"
+#include "Misc/ConfigCacheIni.h"
 #include "GenericPlatform/GenericPlatformChunkInstall.h"
 
 #if PLATFORM_HAS_CPUID
