Include a non-zero patch component in licensee compatibile-with engine versions. Since licensee-built engines do not obey the same guaranteed compatibility between engine hotfixes, we need to ensure that we always treat newer licensee changelists as monotonically incrementing (rather than early-out in FEngineVersion::GetNewest because the patch version is greater).

diff --git a/Engine/Source/Runtime/Core/Private/Misc/EngineVersion.cpp b/Engine/Source/Runtime/Core/Private/Misc/EngineVersion.cpp
--- a/Engine/Source/Runtime/Core/Private/Misc/EngineVersion.cpp
+++ b/Engine/Source/Runtime/Core/Private/Misc/EngineVersion.cpp
@@ -11,7 +11,7 @@
 FEngineVersion FEngineVersion::CurrentVersion(ENGINE_MAJOR_VERSION, ENGINE_MINOR_VERSION, ENGINE_PATCH_VERSION, BuildSettings::GetCurrentChangelist() | (BuildSettings::IsLicenseeVersion()? (1U << 31) : 0), BuildSettings::GetBranchName());
 
 // Version which this engine maintains strict API and package compatibility with. By default, we always maintain compatibility with the current major/minor version, unless we're built at a different changelist.
-FEngineVersion FEngineVersion::CompatibleWithVersion(ENGINE_MAJOR_VERSION, ENGINE_MINOR_VERSION, 0, BuildSettings::GetCompatibleChangelist() | (BuildSettings::IsLicenseeVersion()? (1U << 31) : 0), BuildSettings::GetBranchName());
+FEngineVersion FEngineVersion::CompatibleWithVersion(ENGINE_MAJOR_VERSION, ENGINE_MINOR_VERSION, BuildSettings::IsLicenseeVersion()? ENGINE_PATCH_VERSION : 0, BuildSettings::GetCompatibleChangelist() | (BuildSettings::IsLicenseeVersion()? (1U << 31) : 0), BuildSettings::GetBranchName());
 
 
 FEngineVersionBase::FEngineVersionBase()
@@ -55,7 +55,7 @@ bool FEngineVersionBase::HasChangelist() const
 EVersionComparison FEngineVersionBase::GetNewest(const FEngineVersionBase &First, const FEngineVersionBase &Second, EVersionComponent *OutComponent)
 {
 	EVersionComponent LocalComponent = EVersionComponent::Minor;
-	auto& Component = OutComponent ? *OutComponent : LocalComponent;
+	EVersionComponent& Component = OutComponent ? *OutComponent : LocalComponent;
 
 	// Compare major versions
 	if (First.GetMajor() != Second.GetMajor())
