Fix FTextSnapshot not detecting that Empty == Empty in some cases such as PIE startup.
This caused output log to hang for several seconds rebuilding all lines of text.
Also fixes invariant string comparision on FTextSnapshot.

diff --git a/Engine/Source/Runtime/Core/Private/Internationalization/Text.cpp b/Engine/Source/Runtime/Core/Private/Internationalization/Text.cpp
--- a/Engine/Source/Runtime/Core/Private/Internationalization/Text.cpp
+++ b/Engine/Source/Runtime/Core/Private/Internationalization/Text.cpp
@@ -1458,10 +1458,15 @@ FTextSnapshot::FTextSnapshot()
 
 FTextSnapshot::FTextSnapshot(const FText& InText)
 	: TextDataPtr(InText.TextData)
-	, GlobalHistoryRevision(InText.TextData->GetGlobalHistoryRevision())
-	, LocalHistoryRevision(InText.TextData->GetLocalHistoryRevision())
+	, GlobalHistoryRevision(GetGlobalHistoryRevisionForText(InText))
+	, LocalHistoryRevision(GetLocalHistoryRevisionForText(InText))
 	, Flags(InText.Flags)
 {
+	if (InText.IsEmpty())
+	{
+		GlobalHistoryRevision = 0;
+		LocalHistoryRevision = 0;
+	}
 }
 
 bool FTextSnapshot::IdenticalTo(const FText& InText) const
@@ -1471,8 +1476,8 @@ bool FTextSnapshot::IdenticalTo(const FText& InText) const
 	InText.Rebuild();
 
 	return TextDataPtr == InText.TextData 
-		&& GlobalHistoryRevision == InText.TextData->GetGlobalHistoryRevision()
-		&& LocalHistoryRevision == InText.TextData->GetLocalHistoryRevision()
+		&& GlobalHistoryRevision == GetGlobalHistoryRevisionForText(InText)
+		&& LocalHistoryRevision == GetLocalHistoryRevisionForText(InText)
 		&& Flags == InText.Flags;
 }
 
@@ -1484,11 +1489,21 @@ bool FTextSnapshot::IsDisplayStringEqualTo(const FText& InText) const
 
 	// We have to assume that the display string has changed if the history of the text has changed
 	// (due to a culture change), as we no longer have the old display string to compare against
-	return GlobalHistoryRevision == InText.TextData->GetGlobalHistoryRevision()
-		&& LocalHistoryRevision == InText.TextData->GetLocalHistoryRevision()
+	return GlobalHistoryRevision == GetGlobalHistoryRevisionForText(InText)
+		&& LocalHistoryRevision == GetLocalHistoryRevisionForText(InText)
 		&& TextDataPtr.IsValid() && TextDataPtr->GetDisplayString().Equals(InText.ToString(), ESearchCase::CaseSensitive);
 }
 
+uint16 FTextSnapshot::GetGlobalHistoryRevisionForText(const FText& InText)
+{
+	return (InText.IsEmpty() || InText.IsCultureInvariant()) ? 0 : InText.TextData->GetGlobalHistoryRevision();
+}
+
+uint16 FTextSnapshot::GetLocalHistoryRevisionForText(const FText& InText)
+{
+	return (InText.IsEmpty() || InText.IsCultureInvariant()) ? 0 : InText.TextData->GetLocalHistoryRevision();
+}
+
 FScopedTextIdentityPreserver::FScopedTextIdentityPreserver(FText& InTextToPersist)
 	: TextToPersist(InTextToPersist)
 	, HadFoundNamespaceAndKey(false)
diff --git a/Engine/Source/Runtime/Core/Public/Internationalization/Text.h b/Engine/Source/Runtime/Core/Public/Internationalization/Text.h
--- a/Engine/Source/Runtime/Core/Public/Internationalization/Text.h
+++ b/Engine/Source/Runtime/Core/Public/Internationalization/Text.h
@@ -1030,6 +1030,13 @@ public:
 	bool IsDisplayStringEqualTo(const FText& InText) const;
 
 private:
+
+	/** Get adjusted global history revision used for comparison */
+	static uint16 GetGlobalHistoryRevisionForText(const FText& InText);
+
+	/** Get adjusted local history revision used for comparison */
+	static uint16 GetLocalHistoryRevisionForText(const FText& InText);
+
 	/** A pointer to the text data for the FText that we took a snapshot of (used for an efficient pointer compare) */
 	TSharedPtr<ITextData, ESPMode::ThreadSafe> TextDataPtr;
 
