Remove uncessary code
#rnx

diff --git a/Engine/Source/Runtime/Core/Private/Internationalization/Text.cpp b/Engine/Source/Runtime/Core/Private/Internationalization/Text.cpp
--- a/Engine/Source/Runtime/Core/Private/Internationalization/Text.cpp
+++ b/Engine/Source/Runtime/Core/Private/Internationalization/Text.cpp
@@ -1462,11 +1462,6 @@ FTextSnapshot::FTextSnapshot(const FText& InText)
 	, LocalHistoryRevision(GetLocalHistoryRevisionForText(InText))
 	, Flags(InText.Flags)
 {
-	if (InText.IsEmpty())
-	{
-		GlobalHistoryRevision = 0;
-		LocalHistoryRevision = 0;
-	}
 }
 
 bool FTextSnapshot::IdenticalTo(const FText& InText) const
