macOS 10.15 has a new max limit of open files (10240, equal to OPEN_MAX, instead of RLIM_INFINITY). Changed the code in MacPlatformMisc that checks and sets the limit to warn only if the resulting limit is lower than OPEN_MAX.

diff --git a/Engine/Source/Runtime/Core/Private/Mac/MacPlatformMisc.cpp b/Engine/Source/Runtime/Core/Private/Mac/MacPlatformMisc.cpp
--- a/Engine/Source/Runtime/Core/Private/Mac/MacPlatformMisc.cpp
+++ b/Engine/Source/Runtime/Core/Private/Mac/MacPlatformMisc.cpp
@@ -442,11 +442,7 @@ void FMacPlatformMisc::PlatformPreInit()
 	int32 Result = getrlimit(RLIMIT_NOFILE, &Limit);
 	if (Result == 0)
 	{
-		if(Limit.rlim_max != RLIM_INFINITY)
-		{
-			UE_LOG(LogInit, Warning, TEXT("Hard Max File Limit Too Small: %llu, should be RLIM_INFINITY, UE4 may be unstable."), Limit.rlim_max);
-		}
-		if(Limit.rlim_max == RLIM_INFINITY)
+		if (Limit.rlim_max == RLIM_INFINITY)
 		{
 			Limit.rlim_cur = MaxFilesPerProc;
 		}
@@ -455,12 +451,16 @@ void FMacPlatformMisc::PlatformPreInit()
 			Limit.rlim_cur = FMath::Min(Limit.rlim_max, (rlim_t)MaxFilesPerProc);
 		}
 	}
+	if (Limit.rlim_cur < OPEN_MAX)
+	{
+		UE_LOG(LogInit, Warning, TEXT("Open files limit too small: %llu, should be at least OPEN_MAX (%llu). rlim_max is %llu, kern.maxfilesperproc is %u. UE4 may be unstable."), Limit.rlim_cur, OPEN_MAX, Limit.rlim_max, MaxFilesPerProc);
+	}
 	Result = setrlimit(RLIMIT_NOFILE, &Limit);
 	if (Result != 0)
 	{
 		UE_LOG(LogInit, Warning, TEXT("Failed to change open file limit, UE4 may be unstable."));
 	}
-	
+
 	FApplePlatformSymbolication::EnableCoreSymbolication(!FPlatformProcess::IsSandboxedApplication() && IS_PROGRAM);
 }
 
