Added TryDecommitPending to page cache system

#rnx

diff --git a/Engine/Source/Runtime/Core/Public/HAL/PageCache.h b/Engine/Source/Runtime/Core/Public/HAL/PageCache.h
--- a/Engine/Source/Runtime/Core/Public/HAL/PageCache.h
+++ b/Engine/Source/Runtime/Core/Public/HAL/PageCache.h
@@ -286,6 +286,75 @@ public:
 		}
 		return size_t(NumFound) * PageSize;
 	}
+	template<typename T>
+	size_t TryDecommitPending(size_t Size, T&& DecommitFunction)
+	{
+		if (!NumPages)
+		{
+			return 0; // this page cache was never really set up, nothing to sweep
+		}
+		check(Size > 0 && IsAligned(Size, PageSize));
+		uint32 NumNeed = Size / PageSize;
+
+		uint32 LastDecommitPage = MAX_uint32;
+		uint32 StartDecommitPage = MAX_uint32;
+		uint32 NumFound = 0;
+		FScopeLock Lock(&CriticalSection);
+
+		auto DecommitRange = [&]()
+		{
+			uint32 DecommitPageCount = 1 + LastDecommitPage - StartDecommitPage;
+
+			if (DecommitFunction(PageIndexToAddr(StartDecommitPage), DecommitPageCount * PageSize))
+			{
+				CommittedPages -= DecommitPageCount;
+				DecommittedPages += DecommitPageCount;
+				check(PendingDecommittedPages >= DecommitPageCount);
+				PendingDecommittedPages -= DecommitPageCount;
+				NumFound += DecommitPageCount;
+
+				uint32 EndBitIndex = StartDecommitPage + DecommitPageCount;
+				for (uint32 BitIndex = StartDecommitPage; BitIndex < EndBitIndex; BitIndex++)
+				{
+					CurrentlyCommitted.FreeBit(BitIndex);
+					NotPendingDecommit.AllocBit(BitIndex);
+				}
+			}
+		};
+		while (NumFound < NumNeed)
+		{
+			check(SweepPage < NumPages);
+			uint32 Index = NotPendingDecommit.NextAllocBit(SweepPage);
+			if (Index == MAX_uint32)
+			{
+				SweepPage = 0;
+				break;
+			}
+			check(CurrentlyCommitted.IsAllocated(Index) && !NotPendingDecommit.IsAllocated(Index));
+			check(CommittedPages);
+
+			if (StartDecommitPage == MAX_uint32 || LastDecommitPage + 1 != Index)
+			{
+				if (StartDecommitPage != MAX_uint32)
+				{
+					DecommitRange();
+				}
+				StartDecommitPage = Index;
+			}
+			LastDecommitPage = Index;
+			SweepPage = Index + 1;
+			if (SweepPage >= NumPages)
+			{
+				SweepPage = 0;
+				break;
+			}
+		}
+		if (StartDecommitPage != MAX_uint32)
+		{
+			DecommitRange();
+		}
+		return size_t(NumFound) * PageSize;
+	}
 	size_t GetFreeableMemory()
 	{
 		FScopeLock Lock(&CriticalSection);
