Fixed summary message in CrashReportClient having an "Assertion Failed:" prefix when triggering a fatal error.

diff --git a/Engine/Source/Runtime/Core/Private/Misc/AssertionMacros.cpp b/Engine/Source/Runtime/Core/Private/Misc/AssertionMacros.cpp
--- a/Engine/Source/Runtime/Core/Private/Misc/AssertionMacros.cpp
+++ b/Engine/Source/Runtime/Core/Private/Misc/AssertionMacros.cpp
@@ -454,6 +454,22 @@ void VARARGS FDebug::AssertFailed(const ANSICHAR* Expr, const ANSICHAR* File, in
 	va_end(Args);
 }
 
+void FDebug::ProcessFatalError()
+{
+	if (GIsCriticalError)
+	{
+		return;
+	}
+
+	// This is not perfect because another thread might crash and be handled before this assert
+	// but this static varible will report the crash as an assert. Given complexity of a thread
+	// aware solution, this should be good enough. If crash reports are obviously wrong we can
+	// look into fixing this.
+	bHasAsserted = true;
+
+	GError->Logf(TEXT("%s"), GErrorHist);
+}
+
 #if DO_CHECK || DO_GUARD_SLOW
 FORCENOINLINE bool VARARGS FDebug::OptionallyLogFormattedEnsureMessageReturningFalseImpl( bool bLog, const ANSICHAR* Expr, const ANSICHAR* File, int32 Line, const TCHAR* FormattedMsg, ... )
 {
diff --git a/Engine/Source/Runtime/Core/Public/Logging/LogMacros.h b/Engine/Source/Runtime/Core/Public/Logging/LogMacros.h
--- a/Engine/Source/Runtime/Core/Public/Logging/LogMacros.h
+++ b/Engine/Source/Runtime/Core/Public/Logging/LogMacros.h
@@ -110,7 +110,7 @@ private:
 		{ \
 			LowLevelFatalErrorHandler(UE_LOG_SOURCE_FILE(__FILE__), __LINE__, Format, ##__VA_ARGS__); \
 			_DebugBreakAndPromptForRemote(); \
-			FDebug::AssertFailed("", UE_LOG_SOURCE_FILE(__FILE__), __LINE__, Format, ##__VA_ARGS__); \
+			FDebug::ProcessFatalError(); \
 			UE_LOG_EXPAND_IS_FATAL(Verbosity, CA_ASSUME(false);, PREPROCESSOR_NOTHING) \
 		} \
 	}
@@ -124,7 +124,7 @@ private:
 			{ \
 				LowLevelFatalErrorHandler(UE_LOG_SOURCE_FILE(__FILE__), __LINE__, Format, ##__VA_ARGS__); \
 				_DebugBreakAndPromptForRemote(); \
-				FDebug::AssertFailed("", UE_LOG_SOURCE_FILE(__FILE__), __LINE__, Format, ##__VA_ARGS__); \
+				FDebug::ProcessFatalError(); \
 				UE_LOG_EXPAND_IS_FATAL(Verbosity, CA_ASSUME(false);, PREPROCESSOR_NOTHING) \
 			} \
 		} \
@@ -201,7 +201,7 @@ private:
 					{ \
 						FMsg::Logf_Internal(UE_LOG_SOURCE_FILE(__FILE__), __LINE__, CategoryName.GetCategoryName(), ELogVerbosity::Verbosity, Format, ##__VA_ARGS__); \
 						_DebugBreakAndPromptForRemote(); \
-						FDebug::AssertFailed("", UE_LOG_SOURCE_FILE(__FILE__), __LINE__, Format, ##__VA_ARGS__); \
+						FDebug::ProcessFatalError(); \
 						CA_ASSUME(false); \
 					}, \
 					{ \
@@ -247,7 +247,7 @@ private:
 						{ \
 							FMsg::Logf_Internal(UE_LOG_SOURCE_FILE(__FILE__), __LINE__, CategoryName.GetCategoryName(), ELogVerbosity::Verbosity, Format, ##__VA_ARGS__); \
 							_DebugBreakAndPromptForRemote(); \
-							FDebug::AssertFailed("", UE_LOG_SOURCE_FILE(__FILE__), __LINE__, Format, ##__VA_ARGS__); \
+							FDebug::ProcessFatalError(); \
 							CA_ASSUME(false); \
 						}, \
 						{ \
diff --git a/Engine/Source/Runtime/Core/Public/Misc/AssertionMacros.h b/Engine/Source/Runtime/Core/Public/Misc/AssertionMacros.h
--- a/Engine/Source/Runtime/Core/Public/Misc/AssertionMacros.h
+++ b/Engine/Source/Runtime/Core/Public/Misc/AssertionMacros.h
@@ -38,6 +38,9 @@ struct CORE_API FDebug
 	/** Logs final assert message and exits the program. */
 	static void VARARGS AssertFailed(const ANSICHAR* Expr, const ANSICHAR* File, int32 Line, const TCHAR* Format = TEXT(""), ...);
 
+	/** Triggers a fatal error, using the error formatted to GErrorHist via a previous call to FMsg*/
+	static void ProcessFatalError();
+
 	// returns true if an assert has occurred
 	static bool HasAsserted();
 
@@ -373,6 +376,6 @@ CORE_API void VARARGS LowLevelFatalErrorHandler(const ANSICHAR* File, int32 Line
 		static_assert(TIsArrayOrRefOfType<decltype(Format), TCHAR>::Value, "Formatting string must be a TCHAR array."); \
 		LowLevelFatalErrorHandler(__FILE__, __LINE__, Format, ##__VA_ARGS__); \
 		_DebugBreakAndPromptForRemote(); \
-		FDebug::AssertFailed("", __FILE__, __LINE__, Format, ##__VA_ARGS__); \
+		FDebug::ProcessFatalError(); \
 	}
 
