Content browser no longer builds it's package roots from full drive paths.

diff --git a/Engine/Source/Editor/ContentBrowser/Private/FrontendFilters.cpp b/Engine/Source/Editor/ContentBrowser/Private/FrontendFilters.cpp
--- a/Engine/Source/Editor/ContentBrowser/Private/FrontendFilters.cpp
+++ b/Engine/Source/Editor/ContentBrowser/Private/FrontendFilters.cpp
@@ -1106,9 +1106,9 @@ FString FFrontendFilter_ArbitraryComparisonOperation::ConvertOperationToString(E
 
 FFrontendFilter_ShowOtherDevelopers::FFrontendFilter_ShowOtherDevelopers(TSharedPtr<FFrontendFilterCategory> InCategory)
 	: FFrontendFilter(InCategory)
-	, BaseDeveloperPath(FPackageName::FilenameToLongPackageName(FPaths::GameDevelopersDir()))
+	, BaseDeveloperPath(TEXT("/Game/Developers/"))
 	, BaseDeveloperPathAnsi()
-	, UserDeveloperPath(FPackageName::FilenameToLongPackageName(FPaths::GameUserDeveloperDir()))
+	, UserDeveloperPath(BaseDeveloperPath + FPaths::GameUserDeveloperFolderName() + TEXT("/"))
 	, bIsOnlyOneDeveloperPathSelected(false)
 	, bShowOtherDeveloperAssets(false)
 {
diff --git a/Engine/Source/Runtime/Core/Private/Misc/Paths.cpp b/Engine/Source/Runtime/Core/Private/Misc/Paths.cpp
--- a/Engine/Source/Runtime/Core/Private/Misc/Paths.cpp
+++ b/Engine/Source/Runtime/Core/Private/Misc/Paths.cpp
@@ -413,18 +413,18 @@ FString FPaths::GameDevelopersDir()
 	return FPaths::ProjectContentDir() + TEXT("Developers/");
 }
 
-FString FPaths::GameUserDeveloperDir()
+FString FPaths::GameUserDeveloperFolderName()
 {
 	static FString UserFolder;
 
-	if ( UserFolder.Len() == 0 )
+	if (UserFolder.Len() == 0)
 	{
 		// The user folder is the user name without any invalid characters
 		const FString InvalidChars = INVALID_LONGPACKAGE_CHARACTERS;
 		const FString& UserName = FPlatformProcess::UserName();
-		
+
 		UserFolder = UserName;
-		
+
 		for (int32 CharIdx = 0; CharIdx < InvalidChars.Len(); ++CharIdx)
 		{
 			const FString Char = InvalidChars.Mid(CharIdx, 1);
@@ -432,7 +432,12 @@ FString FPaths::GameUserDeveloperDir()
 		}
 	}
 
-	return FPaths::GameDevelopersDir() + UserFolder + TEXT("/");
+	return UserFolder;
+}
+
+FString FPaths::GameUserDeveloperDir()
+{
+	return FPaths::GameDevelopersDir() + GameUserDeveloperFolderName() + TEXT("/");
 }
 
 FString FPaths::DiffDir()
diff --git a/Engine/Source/Runtime/Core/Public/Misc/Paths.h b/Engine/Source/Runtime/Core/Public/Misc/Paths.h
--- a/Engine/Source/Runtime/Core/Public/Misc/Paths.h
+++ b/Engine/Source/Runtime/Core/Public/Misc/Paths.h
@@ -289,6 +289,9 @@ public:
 	/** Returns the directory that contains subfolders for developer-specific content */
 	static FString GameDevelopersDir();
 
+	/** Returns The folder name for the developer-specific directory for the current user */
+	static FString GameUserDeveloperFolderName();
+
 	/** Returns The directory that contains developer-specific content for the current user */
 	static FString GameUserDeveloperDir();
 
