Use the Oculus runtime check (IsStandaloneStereoOnlyDevice) to ensure Vulkan is always used if available and specified on Android builds.

diff --git a/Engine/Source/Runtime/Core/Private/Android/AndroidPlatformMisc.cpp b/Engine/Source/Runtime/Core/Private/Android/AndroidPlatformMisc.cpp
--- a/Engine/Source/Runtime/Core/Private/Android/AndroidPlatformMisc.cpp
+++ b/Engine/Source/Runtime/Core/Private/Android/AndroidPlatformMisc.cpp
@@ -1765,7 +1765,8 @@ bool FAndroidMisc::ShouldUseVulkan()
 		static const auto CVarDisableVulkan = IConsoleManager::Get().FindTConsoleVariableDataInt(TEXT("r.Android.DisableVulkanSupport"));
 
 		const bool bVulkanAvailable = IsVulkanAvailable();
-		const bool bVulkanDisabledCVar = CVarDisableVulkan->GetValueOnAnyThread() == 1;
+
+		const bool bVulkanDisabledCVar = !IsStandaloneStereoOnlyDevice() && CVarDisableVulkan->GetValueOnAnyThread() == 1;
 
 		if (bVulkanAvailable && !bVulkanDisabledCVar)
 		{
