Fix initialization order issues with lock free lists and global static variables

diff --git a/Engine/Source/Runtime/Core/Private/Containers/LockFreeList.cpp b/Engine/Source/Runtime/Core/Private/Containers/LockFreeList.cpp
--- a/Engine/Source/Runtime/Core/Private/Containers/LockFreeList.cpp
+++ b/Engine/Source/Runtime/Core/Private/Containers/LockFreeList.cpp
@@ -118,7 +118,7 @@ public:
 				TLS.PartialBundle = GlobalFreeListBundles.Pop();
 				if (!TLS.PartialBundle)
 				{
-					int32 FirstIndex = FLockFreeLinkPolicy::LinkAllocator.Alloc(NUM_PER_BUNDLE);
+					int32 FirstIndex = FLockFreeLinkPolicy::GetLinkAllocator().Alloc(NUM_PER_BUNDLE);
 					for (int32 Index = 0; Index < NUM_PER_BUNDLE; Index++)
 					{
 						TLink* Event = FLockFreeLinkPolicy::IndexToLink(FirstIndex + Index);
@@ -205,19 +205,40 @@ private:
 	FLockFreePointerListLIFORoot<PLATFORM_CACHE_LINE_SIZE> GlobalFreeListBundles;
 };
 
-static LockFreeLinkAllocator_TLSCache GLockFreeLinkAllocator;
+static LockFreeLinkAllocator_TLSCache* GLockFreeLinkAllocator = nullptr;
+
+LockFreeLinkAllocator_TLSCache* InternalGetFreeLinkAllocator()
+{
+	static LockFreeLinkAllocator_TLSCache LockFreeLinkAllocator;
+	return &LockFreeLinkAllocator;
+}
+
+LockFreeLinkAllocator_TLSCache& GetLockFreeLinkAllocator()
+{
+	if (GLockFreeLinkAllocator == nullptr)
+	{
+		GLockFreeLinkAllocator = InternalGetFreeLinkAllocator();
+	}
+	return *GLockFreeLinkAllocator;
+}
 
 void FLockFreeLinkPolicy::FreeLockFreeLink(FLockFreeLinkPolicy::TLinkPtr Item)
 {
-	GLockFreeLinkAllocator.Push(Item);
+	GetLockFreeLinkAllocator().Push(Item);
 }
 
 FLockFreeLinkPolicy::TLinkPtr FLockFreeLinkPolicy::AllocLockFreeLink() TSAN_SAFE
 {
-	FLockFreeLinkPolicy::TLinkPtr Result = GLockFreeLinkAllocator.Pop();
+	FLockFreeLinkPolicy::TLinkPtr Result = GetLockFreeLinkAllocator().Pop();
 	// this can only really be a mem stomp
 	checkLockFreePointerList(Result && !FLockFreeLinkPolicy::DerefLink(Result)->DoubleNext.GetPtr() && !FLockFreeLinkPolicy::DerefLink(Result)->Payload && !FLockFreeLinkPolicy::DerefLink(Result)->SingleNext);
 	return Result;
 }
 
-FLockFreeLinkPolicy::TAllocator FLockFreeLinkPolicy::LinkAllocator;
+FLockFreeLinkPolicy::TAllocator* FLockFreeLinkPolicy::InternalGetLinkAllocator()
+{
+	static TAllocator LinkAllocator;
+	return &LinkAllocator;
+}
+
+FLockFreeLinkPolicy::TAllocator* FLockFreeLinkPolicy::LinkAllocatorPtr = nullptr;
diff --git a/Engine/Source/Runtime/Core/Public/Containers/LockFreeList.h b/Engine/Source/Runtime/Core/Public/Containers/LockFreeList.h
--- a/Engine/Source/Runtime/Core/Public/Containers/LockFreeList.h
+++ b/Engine/Source/Runtime/Core/Public/Containers/LockFreeList.h
@@ -249,11 +249,11 @@ struct FLockFreeLinkPolicy
 
 	static FORCEINLINE FIndexedLockFreeLink* DerefLink(uint32 Ptr)
 	{
-		return LinkAllocator.GetItem(Ptr);
+		return GetLinkAllocator().GetItem(Ptr);
 	}
 	static FORCEINLINE FIndexedLockFreeLink* IndexToLink(uint32 Index)
 	{
-		return LinkAllocator.GetItem(Index);
+		return GetLinkAllocator().GetItem(Index);
 	}
 	static FORCEINLINE uint32 IndexToPtr(uint32 Index)
 	{
@@ -262,7 +262,18 @@ struct FLockFreeLinkPolicy
 
 	CORE_API static uint32 AllocLockFreeLink();
 	CORE_API static void FreeLockFreeLink(uint32 Item);
-	CORE_API static TAllocator LinkAllocator;
+	CORE_API static TAllocator& GetLinkAllocator()
+	{
+		if (LinkAllocatorPtr == nullptr)
+		{
+			LinkAllocatorPtr = InternalGetLinkAllocator();
+		}
+		return *LinkAllocatorPtr;
+	}
+
+private:
+	static TAllocator* LinkAllocatorPtr;
+	static TAllocator* InternalGetLinkAllocator();
 };
 
 template<int TPaddingForCacheContention, uint64 TABAInc = 1>
