Mark  internal members that will be inlined as CORE_API
#rnx

diff --git a/Engine/Source/Runtime/Core/Public/Containers/LockFreeList.h b/Engine/Source/Runtime/Core/Public/Containers/LockFreeList.h
--- a/Engine/Source/Runtime/Core/Public/Containers/LockFreeList.h
+++ b/Engine/Source/Runtime/Core/Public/Containers/LockFreeList.h
@@ -272,8 +272,8 @@ struct FLockFreeLinkPolicy
 	}
 
 private:
-	static TAllocator* LinkAllocatorPtr;
-	static TAllocator* InternalGetLinkAllocator();
+	CORE_API static TAllocator* LinkAllocatorPtr;
+	CORE_API static TAllocator* InternalGetLinkAllocator();
 };
 
 template<int TPaddingForCacheContention, uint64 TABAInc = 1>
