Revert LockFreeList changes from CL# 7778249 and instead port CL# 7310781 from Dev-Core

diff --git a/Engine/Source/Runtime/Core/Private/Containers/LockFreeList.cpp b/Engine/Source/Runtime/Core/Private/Containers/LockFreeList.cpp
--- a/Engine/Source/Runtime/Core/Private/Containers/LockFreeList.cpp
+++ b/Engine/Source/Runtime/Core/Private/Containers/LockFreeList.cpp
@@ -118,7 +118,7 @@ public:
 				TLS.PartialBundle = GlobalFreeListBundles.Pop();
 				if (!TLS.PartialBundle)
 				{
-					int32 FirstIndex = FLockFreeLinkPolicy::GetLinkAllocator().Alloc(NUM_PER_BUNDLE);
+					int32 FirstIndex = FLockFreeLinkPolicy::LinkAllocator.Alloc(NUM_PER_BUNDLE);
 					for (int32 Index = 0; Index < NUM_PER_BUNDLE; Index++)
 					{
 						TLink* Event = FLockFreeLinkPolicy::IndexToLink(FirstIndex + Index);
@@ -205,40 +205,19 @@ private:
 	FLockFreePointerListLIFORoot<PLATFORM_CACHE_LINE_SIZE> GlobalFreeListBundles;
 };
 
-static LockFreeLinkAllocator_TLSCache* GLockFreeLinkAllocator = nullptr;
-
-LockFreeLinkAllocator_TLSCache* InternalGetFreeLinkAllocator()
-{
-	static LockFreeLinkAllocator_TLSCache LockFreeLinkAllocator;
-	return &LockFreeLinkAllocator;
-}
-
-LockFreeLinkAllocator_TLSCache& GetLockFreeLinkAllocator()
-{
-	if (GLockFreeLinkAllocator == nullptr)
-	{
-		GLockFreeLinkAllocator = InternalGetFreeLinkAllocator();
-	}
-	return *GLockFreeLinkAllocator;
-}
+static LockFreeLinkAllocator_TLSCache GLockFreeLinkAllocator;
 
 void FLockFreeLinkPolicy::FreeLockFreeLink(FLockFreeLinkPolicy::TLinkPtr Item)
 {
-	GetLockFreeLinkAllocator().Push(Item);
+	GLockFreeLinkAllocator.Push(Item);
 }
 
 FLockFreeLinkPolicy::TLinkPtr FLockFreeLinkPolicy::AllocLockFreeLink() TSAN_SAFE
 {
-	FLockFreeLinkPolicy::TLinkPtr Result = GetLockFreeLinkAllocator().Pop();
+	FLockFreeLinkPolicy::TLinkPtr Result = GLockFreeLinkAllocator.Pop();
 	// this can only really be a mem stomp
 	checkLockFreePointerList(Result && !FLockFreeLinkPolicy::DerefLink(Result)->DoubleNext.GetPtr() && !FLockFreeLinkPolicy::DerefLink(Result)->Payload && !FLockFreeLinkPolicy::DerefLink(Result)->SingleNext);
 	return Result;
 }
 
-FLockFreeLinkPolicy::TAllocator* FLockFreeLinkPolicy::InternalGetLinkAllocator()
-{
-	static TAllocator LinkAllocator;
-	return &LinkAllocator;
-}
-
-FLockFreeLinkPolicy::TAllocator* FLockFreeLinkPolicy::LinkAllocatorPtr = nullptr;
+FLockFreeLinkPolicy::TAllocator FLockFreeLinkPolicy::LinkAllocator;
diff --git a/Engine/Source/Runtime/Core/Public/Containers/LockFreeList.h b/Engine/Source/Runtime/Core/Public/Containers/LockFreeList.h
--- a/Engine/Source/Runtime/Core/Public/Containers/LockFreeList.h
+++ b/Engine/Source/Runtime/Core/Public/Containers/LockFreeList.h
@@ -249,11 +249,11 @@ struct FLockFreeLinkPolicy
 
 	static FORCEINLINE FIndexedLockFreeLink* DerefLink(uint32 Ptr)
 	{
-		return GetLinkAllocator().GetItem(Ptr);
+		return LinkAllocator.GetItem(Ptr);
 	}
 	static FORCEINLINE FIndexedLockFreeLink* IndexToLink(uint32 Index)
 	{
-		return GetLinkAllocator().GetItem(Index);
+		return LinkAllocator.GetItem(Index);
 	}
 	static FORCEINLINE uint32 IndexToPtr(uint32 Index)
 	{
@@ -262,18 +262,7 @@ struct FLockFreeLinkPolicy
 
 	CORE_API static uint32 AllocLockFreeLink();
 	CORE_API static void FreeLockFreeLink(uint32 Item);
-	CORE_API static TAllocator& GetLinkAllocator()
-	{
-		if (LinkAllocatorPtr == nullptr)
-		{
-			LinkAllocatorPtr = InternalGetLinkAllocator();
-		}
-		return *LinkAllocatorPtr;
-	}
-
-private:
-	CORE_API static TAllocator* LinkAllocatorPtr;
-	CORE_API static TAllocator* InternalGetLinkAllocator();
+	CORE_API static TAllocator LinkAllocator;
 };
 
 template<int TPaddingForCacheContention, uint64 TABAInc = 1>
diff --git a/Engine/Source/Runtime/CoreUObject/Private/UObject/UObjectBase.cpp b/Engine/Source/Runtime/CoreUObject/Private/UObject/UObjectBase.cpp
--- a/Engine/Source/Runtime/CoreUObject/Private/UObject/UObjectBase.cpp
+++ b/Engine/Source/Runtime/CoreUObject/Private/UObject/UObjectBase.cpp
@@ -33,9 +33,16 @@ DEFINE_LOG_CATEGORY_STATIC(LogUObjectBootstrap, Display, Display);
 /** Whether uobject system is initialized.												*/
 namespace Internal
 {
-	bool GObjInitialized = false;
+	static bool& GetUObjectSubsystemInitialised()
+	{
+		static bool ObjInitialized = false;
+		return ObjInitialized;
+	}
 };
-
+bool UObjectInitialized()
+{
+	return Internal::GetUObjectSubsystemInitialised();
+}
 
 /** Objects to automatically register once the object system is ready.					*/
 struct FPendingRegistrantInfo
@@ -207,7 +214,7 @@ void UObjectBase::CreateStatID() const
  */
 void UObjectBase::DeferredRegister(UClass *UClassStaticClass,const TCHAR* PackageName,const TCHAR* InName)
 {
-	check(Internal::GObjInitialized);
+	check(UObjectInitialized());
 	// Set object properties.
 	UPackage* Package = CreatePackage(nullptr, PackageName);
 	check(Package);
@@ -1119,7 +1126,7 @@ void UObjectBaseInit()
 	InitAsyncThread();
 
 	// Note initialized.
-	Internal::GObjInitialized = true;
+	Internal::GetUObjectSubsystemInitialised() = true;
 
 	UObjectProcessRegistrants();
 }
@@ -1130,7 +1137,7 @@ void UObjectBaseInit()
 void UObjectBaseShutdown()
 {
 	GUObjectArray.ShutdownUObjectArray();
-	Internal::GObjInitialized = false;
+	Internal::GetUObjectSubsystemInitialised() = false;
 }
 
 /**
diff --git a/Engine/Source/Runtime/CoreUObject/Public/UObject/UObjectBase.h b/Engine/Source/Runtime/CoreUObject/Public/UObject/UObjectBase.h
--- a/Engine/Source/Runtime/CoreUObject/Public/UObject/UObjectBase.h
+++ b/Engine/Source/Runtime/CoreUObject/Public/UObject/UObjectBase.h
@@ -272,20 +272,13 @@ private:
 #endif // HACK_HEADER_GENERATOR
 };
 
-namespace Internal
-{
-	/** Internal state indicating whether uobject system is initialized. Do not use that directly. Use
-		UObjectInitialized() instead. */
-	COREUOBJECT_API extern bool GObjInitialized;
-}
-
 /**
  * Checks to see if the UObject subsystem is fully bootstrapped and ready to go.
  * If true, then all objects are registered and auto registration of natives is over, forever.
  *
  * @return true if the UObject subsystem is initialized.
  */
-FORCEINLINE bool UObjectInitialized() { return Internal::GObjInitialized; }
+COREUOBJECT_API bool UObjectInitialized();
 
 /**
  * Force a pending registrant to register now instead of in the natural order
diff --git a/Engine/Source/Runtime/Engine/Private/Tickable.cpp b/Engine/Source/Runtime/Engine/Private/Tickable.cpp
--- a/Engine/Source/Runtime/Engine/Private/Tickable.cpp
+++ b/Engine/Source/Runtime/Engine/Private/Tickable.cpp
@@ -70,7 +70,14 @@ FTickableGameObject::FTickableGameObject()
 {
 	FTickableStatics& Statics = FTickableStatics::Get();
 
-	Statics.NewTickableObjects.Push(this);
+	if (UObjectInitialized())
+	{
+		Statics.NewTickableObjects.Push(this);
+	}
+	else
+	{
+		AddTickableObject(Statics.TickableObjects, this);
+	}
 }
 
 FTickableGameObject::~FTickableGameObject()
