Improved Lake body collision.  Collision is now generated by forming convex hulls that follow the shape of the lake.

diff --git a/Engine/Source/Runtime/Core/Public/UObject/FortniteMainBranchObjectVersion.h b/Engine/Source/Runtime/Core/Public/UObject/FortniteMainBranchObjectVersion.h
--- a/Engine/Source/Runtime/Core/Public/UObject/FortniteMainBranchObjectVersion.h
+++ b/Engine/Source/Runtime/Core/Public/UObject/FortniteMainBranchObjectVersion.h
@@ -96,6 +96,9 @@ struct CORE_API FFortniteMainBranchObjectVersion
 		// Move the location of water metadata
 		MoveWaterMetadataToActor,
 
+		// Replaced lake collision component
+		ReplaceLakeCollision,
+
 		// -----<new versions can be added above this line>-------------------------------------------------
 		VersionPlusOne,
 		LatestVersion = VersionPlusOne - 1
