Added an overloaded TQueue::Peek, so we can peek at elements that do not allow copying (such as TUniquePtr).

diff --git a/Engine/Source/Runtime/Core/Public/Containers/Queue.h b/Engine/Source/Runtime/Core/Public/Containers/Queue.h
--- a/Engine/Source/Runtime/Core/Public/Containers/Queue.h
+++ b/Engine/Source/Runtime/Core/Public/Containers/Queue.h
@@ -203,6 +203,34 @@ public:
 		return true;
 	}
 
+	/**
+	 * Peek at the queue's tail item without removing it.
+	 *
+	 * This version of Peek allows peeking at a queue of items that do not allow
+	 * copying, such as TUniquePtr.
+	 *
+	 * @return Pointer to the item, or nullptr if queue is empty
+	 */
+	ItemType* Peek()
+	{
+		if (Tail->NextNode == nullptr)
+		{
+			return nullptr;
+		}
+
+		return &Tail->NextNode->Item;
+	}
+
+	const ItemType* Peek() const
+	{
+		if (Tail->NextNode == nullptr)
+		{
+			return nullptr;
+		}
+
+		return &Tail->NextNode->Item;
+	}
+
 	/**
 	 * Removes the item from the tail of the queue.
 	 *
