fix for magic number for nano malloc region. we now figure it out at startup as iOS 11 uses a different region

diff --git a/Engine/Source/Runtime/Core/Private/HAL/MallocBinned.cpp b/Engine/Source/Runtime/Core/Private/HAL/MallocBinned.cpp
--- a/Engine/Source/Runtime/Core/Private/HAL/MallocBinned.cpp
+++ b/Engine/Source/Runtime/Core/Private/HAL/MallocBinned.cpp
@@ -142,14 +142,10 @@ struct FMallocBinned::PoolHashBucket
 #define PLAT_PAGE_SIZE_LIMIT 16384
 #define PLAT_BINNED_ALLOC_POOLSIZE 16384
 #define PLAT_SMALL_BLOCK_POOL_SIZE 224
-#define PLAT_SMALL_BLOCK_START 0x280000000
-#define PLAT_SMALL_BLOCK_END 0x2a0000000
 #else
 #define PLAT_PAGE_SIZE_LIMIT 65536
 #define PLAT_BINNED_ALLOC_POOLSIZE 65536
 #define PLAT_SMALL_BLOCK_POOL_SIZE 0
-#define PLAT_SMALL_BLOCK_START 0
-#define PLAT_SMALL_BLOCK_END 0
 #endif
 
 struct FMallocBinned::Private
@@ -738,8 +734,8 @@ struct FMallocBinned::Private
 };
 
 #if USE_OS_SMALL_BLOCK_ALLOC
-uint64 FMallocBinned::Private::SmallBlockStartPtr = PLAT_SMALL_BLOCK_START;
-uint64 FMallocBinned::Private::SmallBlockEndPtr = PLAT_SMALL_BLOCK_END;
+uint64 FMallocBinned::Private::SmallBlockStartPtr = 0;
+uint64 FMallocBinned::Private::SmallBlockEndPtr = 0;
 #endif
 
 void FMallocBinned::GetAllocatorStats( FGenericMemoryStats& out_Stats )
@@ -838,6 +834,13 @@ FMallocBinned::FMallocBinned(uint32 InPageSize, uint64 AddressLimit)
 	check(PageSize <= 65536); // There is internal limit on page size of 64k
 	check(AddressLimit > PageSize); // Check to catch 32 bit overflow in AddressLimit
 
+#if USE_OS_SMALL_BLOCK_ALLOC
+	void* ptr = ::malloc(16);
+	Private::SmallBlockStartPtr = (uint64_t)ptr & ~(0x7fffffff);
+	Private::SmallBlockEndPtr = Private::SmallBlockStartPtr + 0x80000000;
+	::free(ptr);
+#endif
+	
 	/** Shift to get the reference from the indirect tables */
 	PoolBitShift = FPlatformMath::CeilLogTwo(PageSize);
 	IndirectPoolBitShift = FPlatformMath::CeilLogTwo(PageSize/sizeof(FPoolInfo));
