CoreDelegates::OnBeforeHandleSystemError is unnecessary if we move the call earlier in WindowsPlatformCrashContext.

Largely a rollback of most of CL 7473700 in Release.

Source: 7839628 (Dev-Editor)

#rnx

diff --git a/Engine/Source/Runtime/Core/Private/Misc/CoreDelegates.cpp b/Engine/Source/Runtime/Core/Private/Misc/CoreDelegates.cpp
--- a/Engine/Source/Runtime/Core/Private/Misc/CoreDelegates.cpp
+++ b/Engine/Source/Runtime/Core/Private/Misc/CoreDelegates.cpp
@@ -42,7 +42,6 @@ FCoreDelegates::FOnUserControllerPairingChange FCoreDelegates::OnControllerPairi
 FCoreDelegates::FOnSafeFrameChangedEvent FCoreDelegates::OnSafeFrameChangedEvent;
 FCoreDelegates::FOnHandleSystemEnsure FCoreDelegates::OnHandleSystemEnsure;
 FCoreDelegates::FOnHandleSystemError FCoreDelegates::OnHandleSystemError;
-FCoreDelegates::FOnHandleSystemError FCoreDelegates::OnBeforeHandleSystemError;
 
 FCoreDelegates::FOnActorLabelChanged FCoreDelegates::OnActorLabelChanged;
 
diff --git a/Engine/Source/Runtime/Core/Private/Windows/WindowsPlatformCrashContext.cpp b/Engine/Source/Runtime/Core/Private/Windows/WindowsPlatformCrashContext.cpp
--- a/Engine/Source/Runtime/Core/Private/Windows/WindowsPlatformCrashContext.cpp
+++ b/Engine/Source/Runtime/Core/Private/Windows/WindowsPlatformCrashContext.cpp
@@ -822,7 +822,8 @@ private:
 
 		GLog->PanicFlushThreadedLogs();
 
-		FCoreDelegates::OnBeforeHandleSystemError.Broadcast();
+		// Then try run time crash processing and broadcast information about a crash.
+		FCoreDelegates::OnHandleSystemError.Broadcast();
 
 		// Get the default settings for the crash context
 		ECrashContextType Type = ECrashContextType::Crash;
@@ -871,9 +872,6 @@ private:
 			WriteMinidump(CrashContext, MiniDumpFilenameW, ExceptionInfo);
 		}
 
-		// Then try run time crash processing and broadcast information about a crash.
-		FCoreDelegates::OnHandleSystemError.Broadcast();
-
 		const bool bGenerateRuntimeCallstack =
 #if UE_LOG_CRASH_CALLSTACK
 			true;
diff --git a/Engine/Source/Runtime/Core/Public/Misc/CoreDelegates.h b/Engine/Source/Runtime/Core/Public/Misc/CoreDelegates.h
--- a/Engine/Source/Runtime/Core/Public/Misc/CoreDelegates.h
+++ b/Engine/Source/Runtime/Core/Public/Misc/CoreDelegates.h
@@ -185,10 +185,6 @@ public:
 	// Callback when an error (crash) has occurred
 	static FOnHandleSystemError OnHandleSystemError;
 
-	// Callback when an error (crash) has occurred. Runs before the crash reporter.
-	// This should be used very carefully.
-	static FOnHandleSystemError OnBeforeHandleSystemError;
-
 	// Called when an actor label is changed
 	static FOnActorLabelChanged OnActorLabelChanged;
 
diff --git a/Engine/Source/Runtime/Engine/Private/EngineSessionManager.cpp b/Engine/Source/Runtime/Engine/Private/EngineSessionManager.cpp
--- a/Engine/Source/Runtime/Engine/Private/EngineSessionManager.cpp
+++ b/Engine/Source/Runtime/Engine/Private/EngineSessionManager.cpp
@@ -128,7 +128,7 @@ namespace
 void FEngineSessionManager::Initialize()
 {
 	// Register for crash and app state callbacks
-	FCoreDelegates::OnBeforeHandleSystemError.AddRaw(this, &FEngineSessionManager::OnCrashing);
+	FCoreDelegates::OnHandleSystemError.AddRaw(this, &FEngineSessionManager::OnCrashing);
 	FCoreDelegates::ApplicationHasReactivatedDelegate.AddRaw(this, &FEngineSessionManager::OnAppReactivate);
 	FCoreDelegates::ApplicationWillDeactivateDelegate.AddRaw(this, &FEngineSessionManager::OnAppDeactivate);
 	FCoreDelegates::ApplicationWillEnterBackgroundDelegate.AddRaw(this, &FEngineSessionManager::OnAppBackground);
