[FEmbeddedCommunication] Add support for ticking multiple functors per frame

diff --git a/Engine/Source/Runtime/Core/Private/Misc/EmbeddedCommunication.cpp b/Engine/Source/Runtime/Core/Private/Misc/EmbeddedCommunication.cpp
--- a/Engine/Source/Runtime/Core/Private/Misc/EmbeddedCommunication.cpp
+++ b/Engine/Source/Runtime/Core/Private/Misc/EmbeddedCommunication.cpp
@@ -305,28 +305,47 @@ void FEmbeddedCommunication::WakeGameThread()
 bool FEmbeddedCommunication::TickGameThread(float DeltaTime)
 {
 #if BUILD_EMBEDDED_APP
-
-	TFunction<void()> LambdaToCall;
-	bool bFoundFunction = false;
+	bool bEnableTickMultipleFunctors = false;
+	GConfig->GetBool(TEXT("EmbeddedCommunication"), TEXT("bEnableTickMultipleFunctors"), bEnableTickMultipleFunctors, GEngineIni);
+	double TickMaxTimeSeconds = 0.1;
+	if (bEnableTickMultipleFunctors)
 	{
-		FScopeLock Lock(&GEmbeddedLock);
+		GConfig->GetDouble(TEXT("EmbeddedCommunication"), TEXT("TickMaxTimeSeconds"), TickMaxTimeSeconds, GEngineIni);
+	}
 
-		for (int Queue = 0; Queue < ARRAY_COUNT(GEmbeddedQueues); Queue++)
+	double TimeSliceEnd = FPlatformTime::Seconds() + TickMaxTimeSeconds;
+	bool bLambdaWasCalled = false;
+	do
+	{
+		TFunction<void()> LambdaToCall;
 		{
-			if (GEmbeddedQueues[Queue].Dequeue(LambdaToCall))
+			FScopeLock Lock(&GEmbeddedLock);
+
+			for (int Queue = 0; Queue < ARRAY_COUNT(GEmbeddedQueues); Queue++)
 			{
-				break;
+				if (GEmbeddedQueues[Queue].Dequeue(LambdaToCall))
+				{
+					break;
+				}
 			}
 		}
-	}
 
-	// call the function if we found one
-	if (LambdaToCall)
-	{
-		LambdaToCall();
+		if (LambdaToCall)
+		{
+			LambdaToCall();
+			bLambdaWasCalled = true;
+		}
+		else
+		{
+			break;
+		}
 	}
+	while (bEnableTickMultipleFunctors
+		&& FPlatformTime::Seconds() < TimeSliceEnd);
+
 	// sleep if nothing is going on
-	else if (GRenderingWakeMap.Num() == 0
+	if (!bLambdaWasCalled
+		&& GRenderingWakeMap.Num() == 0
 		&& GTickWakeMap.Num() == 0
 		&& GTickWithoutSleepCount <= 0)
  	{
