Fix deprecation warnings.

#rnx

diff --git a/Engine/Source/Runtime/Core/Private/GenericPlatform/GenericPlatformMisc.cpp b/Engine/Source/Runtime/Core/Private/GenericPlatform/GenericPlatformMisc.cpp
--- a/Engine/Source/Runtime/Core/Private/GenericPlatform/GenericPlatformMisc.cpp
+++ b/Engine/Source/Runtime/Core/Private/GenericPlatform/GenericPlatformMisc.cpp
@@ -118,19 +118,19 @@ namespace EBuildConfigurations
 	{
 		switch (Configuration)
 		{
-		case Debug:
+		case EBuildConfiguration::Debug:
 			return NSLOCTEXT("UnrealBuildConfigurations", "DebugName", "Debug");
 
-		case DebugGame:
+		case EBuildConfiguration::DebugGame:
 			return NSLOCTEXT("UnrealBuildConfigurations", "DebugGameName", "DebugGame");
 
-		case Development:
+		case EBuildConfiguration::Development:
 			return NSLOCTEXT("UnrealBuildConfigurations", "DevelopmentName", "Development");
 
-		case Shipping:
+		case EBuildConfiguration::Shipping:
 			return NSLOCTEXT("UnrealBuildConfigurations", "ShippingName", "Shipping");
 
-		case Test:
+		case EBuildConfiguration::Test:
 			return NSLOCTEXT("UnrealBuildConfigurations", "TestName", "Test");
 
 		default:
