#UE4 Fix reading DedicatedServerGame.ini and DedicatedServerEngine.ini

diff --git a/Engine/Source/Runtime/Core/Private/Misc/ConfigCacheIni.cpp b/Engine/Source/Runtime/Core/Private/Misc/ConfigCacheIni.cpp
--- a/Engine/Source/Runtime/Core/Private/Misc/ConfigCacheIni.cpp
+++ b/Engine/Source/Runtime/Core/Private/Misc/ConfigCacheIni.cpp
@@ -3133,7 +3133,7 @@ static void ConditionalOverrideIniFilename(FString& IniFilename, const TCHAR* Ba
 
 const int Flag_Required = 1;
 const int Flag_AllowCommandLineOverride = 2;
-const int Flag_DedicatedServerOnly = 4; // replaces Default, Base, and (NOT {PLATFORM} yet) with DedicatedServer
+const int Flag_DedicatedServerOnly = 4; // replaces Default, Base, and (NOT {PLATFORM} yet) with an empty string
 const int Flag_GenerateCacheKey = 8;
 
 
@@ -3322,9 +3322,9 @@ static void GetSourceIniHierarchyFilenames(const TCHAR* InBaseIniName, const TCH
 					}
 					if (IsRunningDedicatedServer())
 					{
-						ExpansionPath = ExpansionPath.Replace(TEXT("Base"), TEXT("DedicatedServer"));
-						ExpansionPath = ExpansionPath.Replace(TEXT("Default"), TEXT("DedicatedServer"));
-						// ExpansionPath = ExpansionPath.Replace(TEXT("{PLATFORM}"), TEXT("DedicatedServer"));
+						ExpansionPath = ExpansionPath.Replace(TEXT("Base"), TEXT(""));
+						ExpansionPath = ExpansionPath.Replace(TEXT("Default"), TEXT(""));
+						// ExpansionPath = ExpansionPath.Replace(TEXT("{PLATFORM}"), TEXT(""));
 					}
 					else
 					{
