[Switch] Enable read speed calculation and analytic.

diff --git a/Engine/Source/Runtime/Core/Private/HAL/DiskUtilizationTracker.cpp b/Engine/Source/Runtime/Core/Private/HAL/DiskUtilizationTracker.cpp
new file mode 100644
--- /dev/null
+++ b/Engine/Source/Runtime/Core/Private/HAL/DiskUtilizationTracker.cpp
@@ -0,0 +1,83 @@
+#include "HAL/DiskUtilizationTracker.h"
+
+#if TRACK_DISK_UTILIZATION
+
+void FDiskUtilizationTracker::MaybePrint()
+{
+#if !UE_BUILD_SHIPPING && SPEW_DISK_UTILIZATION
+	static double LastPrintSeconds = 0.0;
+
+	double CurrentSeconds = FPlatformTime::Seconds();
+
+	// if we haven't printed or haven't in a while and there's been some I/O emit stats
+	if (((LastPrintSeconds == 0.0) || ((CurrentSeconds - LastPrintSeconds) > PrintFrequencySeconds)) && (TotalIOTime > 0.0))
+	{
+		{
+			// emit recent I/O info
+			static uint64 LastReads = 0;
+			static uint64 LastBytesRead = 0;
+
+			static double LastIOTime = 0.0;
+			static double LastIdleTime = 0.0;
+
+			static uint32 LastSeeks = 0;
+			static uint64 LastSeekDistance = 0;
+
+			if ((LastPrintSeconds > 0.0) && (TotalBytesRead > LastBytesRead))
+			{
+				float TimeInterval = CurrentSeconds - LastPrintSeconds;
+
+				double RecentIOTime = TotalIOTime - LastIOTime;
+				double RecentIdleTime = TotalIdleTime - LastIdleTime;
+
+				float Utilization = float(100.0 * RecentIOTime / (RecentIOTime + RecentIdleTime));
+
+				uint64 RecentBytesRead = TotalBytesRead - LastBytesRead;
+
+				double OverallThroughput = double(RecentBytesRead) / (RecentIOTime + RecentIdleTime) / (1024 * 1024);
+				double ReadThroughput = double(RecentBytesRead) / RecentIOTime / (1024 * 1024);
+
+				uint32 RecentSeeks = NumSeeks - LastSeeks;
+				uint64 RecentSeekDistance = TotalSeekDistance - LastSeekDistance;
+
+				double KBPerSeek = RecentSeeks ? double(RecentBytesRead) / (1024 * RecentSeeks) : 0;
+				double AvgSeek = RecentSeeks ? double(RecentSeekDistance) / double(RecentSeeks) : 0;
+
+				uint64 RecentReads = NumReads - LastReads;
+
+				FPlatformMisc::LowLevelOutputDebugStringf(TEXT("Recent Disk Utilization: %5.2f%% over %6.2fs\t%.2f MB/s\t%.2f Actual MB/s\t(%d Reads, %d Seeks, %.2f kbytes / seek, %.2f ave seek)\r\n"), 
+					Utilization, TimeInterval, OverallThroughput, ReadThroughput, RecentReads, RecentSeeks, KBPerSeek, AvgSeek);
+			}
+
+			LastReads = NumReads;
+			LastBytesRead = TotalBytesRead;
+
+			LastIOTime = TotalIOTime;
+			LastIdleTime = TotalIOTime;
+
+			LastSeeks = NumSeeks;
+			LastSeekDistance = TotalSeekDistance;
+		}
+
+		{
+			// emit recent I/O info
+			float Utilization = float(100.0 * TotalIOTime / (TotalIOTime + TotalIdleTime));
+
+			double OverallThroughput = GetOverallThroughputMBS();
+			double ReadThroughput = GetReadThrougputMBS();
+
+			double KBPerSeek = NumSeeks ? double(TotalBytesRead) / (1024 * NumSeeks) : 0;
+			double AvgSeek = NumSeeks ? double(TotalSeekDistance) / double(NumSeeks) : 0;
+
+			FPlatformMisc::LowLevelOutputDebugStringf(TEXT("Overall Disk Utilization: %5.2f%%\t%.2f MB/s\t%.2f Actual MB/s\t(%d Reads, %d Seeks, %.2f kbytes / seek, %.2f ave seek)\r\n"),
+				Utilization, OverallThroughput, ReadThroughput, NumReads, NumSeeks, KBPerSeek, AvgSeek);
+		}
+	}
+
+	LastPrintSeconds = CurrentSeconds;
+
+#endif //!UE_BUILD_SHIPPING && SPEW_DISK_UTILIZATION
+}
+
+struct FDiskUtilizationTracker GDiskUtilizationTracker;
+#endif // TRACK_DISK_UTILIZATION
\ No newline at end of file
diff --git a/Engine/Source/Runtime/Core/Public/HAL/DiskUtilizationTracker.h b/Engine/Source/Runtime/Core/Public/HAL/DiskUtilizationTracker.h
new file mode 100644
--- /dev/null
+++ b/Engine/Source/Runtime/Core/Public/HAL/DiskUtilizationTracker.h
@@ -0,0 +1,167 @@
+// Copyright 1998-2019 Epic Games, Inc. All Rights Reserved.
+
+#pragma once
+
+#include "HAL/PlatformMisc.h"
+
+#ifndef TRACK_DISK_UTILIZATION
+#define TRACK_DISK_UTILIZATION 0
+#endif
+
+#if TRACK_DISK_UTILIZATION
+
+#ifndef SPEW_DISK_UTILIZATION
+#define SPEW_DISK_UTILIZATION 0
+#endif // SPEW_DISK_UTILIZATION
+
+struct FDiskUtilizationTracker
+{
+	FCriticalSection CriticalSection;
+
+	uint64 NumReads;
+	uint64 NumSeeks;
+
+	uint64 IdleStartCycle;
+	uint64 ReadStartCycle;
+
+	uint64 TotalBytesRead;
+	uint64 TotalSeekDistance;
+
+	double TotalIOTime;
+	double TotalIdleTime;
+
+	uint64 InFlightBytes;
+	int32  InFlightReads;
+
+	FDiskUtilizationTracker()
+		: NumReads(0)
+		, NumSeeks(0)
+		, IdleStartCycle(0)
+		, ReadStartCycle(0)
+		, TotalBytesRead(0)
+		, TotalSeekDistance(0)
+		, TotalIOTime(0.0)
+		, TotalIdleTime(0.0)
+		, InFlightBytes(0)
+		, InFlightReads(0)
+	{
+	}
+
+	void StartRead(uint64 InReadBytes, uint64 InSeekDistance = 0)
+	{
+		// update total reads
+		NumReads++;
+
+		// update seek data
+		if (InSeekDistance > 0)
+		{
+			NumSeeks++;
+			TotalSeekDistance += InSeekDistance;
+		}
+
+		{
+			FScopeLock Lock(&CriticalSection);
+
+			if (InFlightReads == 0)
+			{
+				// if this is the first started read from idle start 
+				ReadStartCycle = FPlatformTime::Cycles64();
+
+				// update idle time (if we've been idle)
+				if (IdleStartCycle > 0)
+				{
+
+					TotalIdleTime += double(ReadStartCycle - IdleStartCycle) * FPlatformTime::GetSecondsPerCycle64();
+				}
+			}
+
+			InFlightBytes += InReadBytes;
+			InFlightReads++;
+		}
+	}
+
+	void FinishRead()
+	{
+
+		// if we're the last in flight read update the start idle counter
+		{
+			FScopeLock Lock(&CriticalSection);
+			check(InFlightReads > 0);
+
+			if (--InFlightReads == 0)
+			{
+				IdleStartCycle = FPlatformTime::Cycles64();
+
+				// update our read counters
+				TotalIOTime += double(IdleStartCycle - ReadStartCycle) * FPlatformTime::GetSecondsPerCycle64();
+				TotalBytesRead += InFlightBytes;
+
+				InFlightBytes = 0;
+			}
+
+		}
+
+		MaybePrint();
+	}
+
+	double GetOverallThroughputMBS()
+	{
+		return double(TotalBytesRead) / (TotalIOTime + TotalIdleTime) / (1024 * 1024);
+	}
+
+	double GetReadThrougputMBS()
+	{
+		return double(TotalBytesRead) / (TotalIOTime) / (1024 * 1024);
+	}
+
+	void ResetStats()
+	{
+		// make sure we're not reseting stats while a read is in flight. that'd mess things up.
+		check(InFlightReads == 0);
+
+		NumSeeks = 0;
+		NumReads = 0;
+		
+		IdleStartCycle = 0;
+		ReadStartCycle = 0;
+
+		TotalBytesRead = 0;
+		TotalSeekDistance = 0;
+
+		TotalIOTime = 0.0;
+		TotalIdleTime = 0.0;
+
+		InFlightBytes = 0;
+		InFlightReads = 0;
+	}
+
+	static constexpr float PrintFrequencySeconds = 0.5f;
+
+	void MaybePrint();
+};
+
+extern FDiskUtilizationTracker GDiskUtilizationTracker;
+
+struct FScopedDiskUtilizationTracker
+{
+	FScopedDiskUtilizationTracker(uint64 InReadBytes, uint64 InSeekDistance) 
+	{
+		GDiskUtilizationTracker.StartRead(InReadBytes, InSeekDistance);
+	}
+
+	~FScopedDiskUtilizationTracker()
+	{
+		GDiskUtilizationTracker.FinishRead();
+	}
+};
+
+#else
+
+struct FScopedDiskUtilizationTracker
+{
+	FScopedDiskUtilizationTracker(uint64 Size, uint64 SeekDistance)
+	{
+	}
+};
+
+#endif // TRACK_DISK_UTILIZATION
\ No newline at end of file
diff --git a/Engine/Source/Runtime/Core/Public/HAL/PlatformFileCommon.h b/Engine/Source/Runtime/Core/Public/HAL/PlatformFileCommon.h
--- a/Engine/Source/Runtime/Core/Public/HAL/PlatformFileCommon.h
+++ b/Engine/Source/Runtime/Core/Public/HAL/PlatformFileCommon.h
@@ -5,139 +5,14 @@
 #include "CoreTypes.h"
 #include "HAL/ThreadSafeCounter.h"
 #include "Misc/ScopeLock.h"
-
+#include "HAL/DiskUtilizationTracker.h"
 #include "HAL/ThreadSafeCounter.h"
 #include "Misc/ScopeLock.h"
 #include "GenericPlatform/GenericPlatformFile.h"
 #include "HAL/PlatformProcess.h"
 
-
 #define MANAGE_FILE_HANDLES (#) // this is not longer used, this will error on any attempt to use it
 
-#define SPEW_DISK_UTLIZATION (0)
-
-#if SPEW_DISK_UTLIZATION
-
-static const float DiskUtilizationTrackerPrintFrequency = 0.1f;
-
-struct FDiskUtilizationTracker
-{
-	FCriticalSection Crit;
-	int32 NumRequests;
-	double LastTime;
-	double LastPrint;
-	double WorkTime;
-	double IdleTime;
-	uint64 AmountRead;
-	uint64 NumSeeks;
-	uint64 TotalSeekDistance;
-	uint64 NumReads;
-
-	FDiskUtilizationTracker()
-		: NumRequests(0)
-		, LastTime(-1.0)
-		, LastPrint(-1.0)
-		, WorkTime(0.0)
-		, IdleTime(0.0)
-		, AmountRead(0)
-		, NumSeeks(0)
-		, TotalSeekDistance(0)
-		, NumReads(0)
-	{
-	}
-
-	void Start(uint64 Size, uint64 SeekDistance)
-	{
-		FScopeLock Lock(&Crit);
-		double Now = FPlatformTime::Seconds();
-		if (NumRequests++ == 0)
-		{
-			// was idle, no longer is
-			if (LastTime != -1.0)
-			{
-				IdleTime += (Now - LastTime);
-			}
-			LastTime = Now;
-			MaybePrint();
-		}
-		AmountRead += Size;
-		NumReads++;
-		if (SeekDistance > 0)
-		{
-			NumSeeks++;
-			TotalSeekDistance += SeekDistance;
-		}
-	}
-	void Stop()
-	{
-		FScopeLock Lock(&Crit);
-		double Now = FPlatformTime::Seconds();
-		if (--NumRequests == 0)
-		{
-			// was working, no longer is
-			check(LastTime > 0);
-			WorkTime += (Now - LastTime);
-			LastTime = Now;
-			MaybePrint();
-		}
-	}
-	void MaybePrint()
-	{
-		if (LastPrint < 0.0)
-		{
-			LastPrint = LastTime;
-			return;
-		}
-		float TimeInterval = float(LastTime - LastPrint);
-		if (TimeInterval > DiskUtilizationTrackerPrintFrequency && IdleTime + WorkTime > 0.0)
-		{
-			LastPrint = LastTime;
-			float Result = float(100.0 * WorkTime / (IdleTime + WorkTime));
-			double MBS = ((double)AmountRead / TimeInterval) / (1024 * 1024);
-			double KBPerSeek = 0.0;
-			if (NumSeeks)
-			{
-				KBPerSeek = ((double)AmountRead) / (1024 * NumSeeks);
-			}
-			double ActualMBS = MBS * Result / 100.0;
-			double AvgSeek = NumSeeks ? ((double)TotalSeekDistance / (double)NumSeeks) : 0;
-			FPlatformMisc::LowLevelOutputDebugStringf(TEXT("Disk: %5.2f%% utilization over %6.2fs\t%.2f MB/s\t%.2f Actual MB/s\t(%d Reads, %d Seeks, %.2f kbytes / seek, %.2f ave seek)\r\n"), Result, TimeInterval, MBS, ActualMBS, NumReads, NumSeeks, KBPerSeek, AvgSeek);
-			WorkTime = 0.0;
-			IdleTime = 0.0;
-			AmountRead = 0;
-			NumSeeks = 0;
-			TotalSeekDistance = 0;
-			NumReads = 0;
-		}
-	}
-} GDiskUtilizationTracker;
-
-
-struct FScopedDiskUtilizationTracker
-{
-	FScopedDiskUtilizationTracker(uint64 Size, uint64 SeekDistance)
-	{
-		GDiskUtilizationTracker.Start(Size, SeekDistance);
-	}
-	~FScopedDiskUtilizationTracker()
-	{
-		GDiskUtilizationTracker.Stop();
-	}
-
-};
-
-
-
-#else
-struct FScopedDiskUtilizationTracker
-{
-	FScopedDiskUtilizationTracker(uint64 Size, uint64 SeekDistance)
-	{
-	}
-};
-#endif
-
-
 class FRegisteredFileHandle : public IFileHandle
 {
 public:
