#UE4 Added copyright

diff --git a/Engine/Source/Runtime/Core/Private/HAL/DiskUtilizationTracker.cpp b/Engine/Source/Runtime/Core/Private/HAL/DiskUtilizationTracker.cpp
--- a/Engine/Source/Runtime/Core/Private/HAL/DiskUtilizationTracker.cpp
+++ b/Engine/Source/Runtime/Core/Private/HAL/DiskUtilizationTracker.cpp
@@ -1,3 +1,5 @@
+// Copyright 1998-2019 Epic Games, Inc. All Rights Reserved.
+
 #include "HAL/DiskUtilizationTracker.h"
 
 #if TRACK_DISK_UTILIZATION
