UE-56031 - Logs on IOS are writing on top of the previous log

diff --git a/Engine/Source/Runtime/Core/Private/IOS/IOSPlatformFile.cpp b/Engine/Source/Runtime/Core/Private/IOS/IOSPlatformFile.cpp
--- a/Engine/Source/Runtime/Core/Private/IOS/IOSPlatformFile.cpp
+++ b/Engine/Source/Runtime/Core/Private/IOS/IOSPlatformFile.cpp
@@ -838,6 +838,12 @@ IFileHandle* FIOSPlatformFile::OpenWrite(const TCHAR* Filename, bool bAppend, bo
 	}
 	FString IOSFilename = ConvertToIOSPath(NormalizeFilename(Filename), true, bCreatePublicFiles);
 	int32 Handle = open(TCHAR_TO_UTF8(*IOSFilename), Flags, S_IRUSR | S_IWUSR | S_IRGRP | S_IWGRP | S_IROTH | S_IWOTH);
+
+	if (!bAppend)
+	{
+		ftruncate(Handle, 0);
+	}
+
 	if (Handle != -1)
 	{
 		FIOSFileHandle* FileHandleIOS = new FIOSFileHandle(Handle, IOSFilename, false);
