Improvements to TQueue. (removed duplicated code)

diff --git a/Engine/Source/Runtime/Core/Public/Containers/Queue.h b/Engine/Source/Runtime/Core/Public/Containers/Queue.h
--- a/Engine/Source/Runtime/Core/Public/Containers/Queue.h
+++ b/Engine/Source/Runtime/Core/Public/Containers/Queue.h
@@ -221,14 +221,9 @@ public:
 		return &Tail->NextNode->Item;
 	}
 
-	const ItemType* Peek() const
+	FORCEINLINE const ItemType* Peek() const
 	{
-		if (Tail->NextNode == nullptr)
-		{
-			return nullptr;
-		}
-
-		return &Tail->NextNode->Item;
+		return const_cast<TQueue*>(this)->Peek();
 	}
 
 	/**
