FString slack optimization in FTextKeyState.

diff --git a/Engine/Source/Runtime/Core/Private/Internationalization/TextKey.cpp b/Engine/Source/Runtime/Core/Private/Internationalization/TextKey.cpp
--- a/Engine/Source/Runtime/Core/Private/Internationalization/TextKey.cpp
+++ b/Engine/Source/Runtime/Core/Private/Internationalization/TextKey.cpp
@@ -24,7 +24,7 @@ public:
 		const FString* StrPtr = KeysTable.Find(SrcKey);
 		if (!StrPtr)
 		{
-			FString StrCopy(InStrLen, InStr); // Need to copy the string here so we can reference its internal allocation as the key
+			FString StrCopy = CopyString(InStrLen, InStr); // Need to copy the string here so we can reference its internal allocation as the key
 			const FKeyData DestKey(*StrCopy, StrCopy.Len(), SrcKey.StrHash);
 			StrPtr = &KeysTable.Add(DestKey, MoveTemp(StrCopy));
 			check(DestKey.Str == **StrPtr); // The move must have moved the allocation we referenced in the key
@@ -44,7 +44,7 @@ public:
 		const FString* StrPtr = KeysTable.Find(SrcKey);
 		if (!StrPtr)
 		{
-			FString StrCopy(InStrLen, InStr); // Need to copy the string here so we can reference its internal allocation as the key
+			FString StrCopy = CopyString(InStrLen, InStr); // Need to copy the string here so we can reference its internal allocation as the key
 			const FKeyData DestKey(*StrCopy, StrCopy.Len(), SrcKey.StrHash);
 			StrPtr = &KeysTable.Add(DestKey, MoveTemp(StrCopy));
 			check(DestKey.Str == **StrPtr); // The move must have moved the allocation we referenced in the key
@@ -143,6 +143,16 @@ private:
 		uint32 StrHash;
 	};
 
+	FORCEINLINE FString CopyString(const int32 InStrLen, const TCHAR* InStr)
+	{
+		// We do this rather than use the FString constructor directly, 
+		// as this method avoids slack being added to the allocation
+		FString Str;
+		Str.Reserve(InStrLen);
+		Str.AppendChars(InStr, InStrLen);
+		return Str;
+	}
+
 	FCriticalSection SynchronizationObject;
 	TMap<FKeyData, FString> KeysTable;
 };
