Fixing the issues with re-opening previous blueprint project from ProjectDir being called to early in DataDrivenPlatformInfo

diff --git a/Engine/Source/Runtime/Core/Private/Misc/DataDrivenPlatformInfoRegistry.cpp b/Engine/Source/Runtime/Core/Private/Misc/DataDrivenPlatformInfoRegistry.cpp
--- a/Engine/Source/Runtime/Core/Private/Misc/DataDrivenPlatformInfoRegistry.cpp
+++ b/Engine/Source/Runtime/Core/Private/Misc/DataDrivenPlatformInfoRegistry.cpp
@@ -3,6 +3,7 @@
 #include "Misc/DataDrivenPlatformInfoRegistry.h"
 #include "Misc/Paths.h"
 #include "HAL/FileManager.h"
+#include "Misc/FileHelper.h"
 
 
 
@@ -36,20 +37,28 @@ bool FDataDrivenPlatformInfoRegistry::LoadDataDrivenIniFile(int32 Index, FConfig
 		return false;
 	}
 
-	FConfigCacheIni::LoadExternalIniFile(IniFile, *FPaths::GetBaseFilename(IniFilenames[Index]), nullptr, *FPaths::GetPath(IniFilenames[Index]), false);
-
-	// platform extension paths are different (platform/engine/config, not engine/config/platform)
-	if (IniFilenames[Index].StartsWith(FPaths::PlatformExtensionsDir()))
-	{
-		PlatformName = FPaths::GetCleanFilename(FPaths::GetPath(FPaths::GetPath(FPaths::GetPath(IniFilenames[Index]))));
-	}
-	else
+	// manually load a FConfigFile object from a source ini file so that we don't do any SavedConfigDir processing or anything
+	// (there's a possibility this is called before the ProjectDir is set)
+	FString IniContents;
+	if (FFileHelper::LoadFileToString(IniContents, *IniFilenames[Index]))
 	{
-		// this could be 'Engine' for a shared DataDrivenPlatformInfo file
-		PlatformName = FPaths::GetCleanFilename(FPaths::GetPath(IniFilenames[Index]));
+		IniFile.ProcessInputFileContents(IniContents);
+
+		// platform extension paths are different (platform/engine/config, not engine/config/platform)
+		if (IniFilenames[Index].StartsWith(FPaths::PlatformExtensionsDir()))
+		{
+			PlatformName = FPaths::GetCleanFilename(FPaths::GetPath(FPaths::GetPath(FPaths::GetPath(IniFilenames[Index]))));
+		}
+		else
+		{
+			// this could be 'Engine' for a shared DataDrivenPlatformInfo file
+			PlatformName = FPaths::GetCleanFilename(FPaths::GetPath(IniFilenames[Index]));
+		}
+
+		return true;
 	}
 
-	return true;
+	return false;
 }
 
 /**
