addition of MMIO memory tracking with LLM

diff --git a/Engine/Source/Runtime/ApplicationCore/Private/IOS/IOSAppDelegate.cpp b/Engine/Source/Runtime/ApplicationCore/Private/IOS/IOSAppDelegate.cpp
--- a/Engine/Source/Runtime/ApplicationCore/Private/IOS/IOSAppDelegate.cpp
+++ b/Engine/Source/Runtime/ApplicationCore/Private/IOS/IOSAppDelegate.cpp
@@ -399,7 +399,11 @@ static IOSAppDelegate* CachedDelegate = nil;
 
 -(void)RecordPeakMemory
 {
-    FIOSPlatformMemory::GetStats();
+    FPlatformMemoryStats MemoryStats = FIOSPlatformMemory::GetStats();
+
+#if ENABLE_LOW_LEVEL_MEM_TRACKER
+	FLowLevelMemTracker::Get().SetTagAmountForTracker(ELLMTracker::Platform, ELLMTag::PlatformVM, MemoryStats.UsedVirtual, false);
+#endif
 }
 
 -(void)InitIdleTimerSettings
diff --git a/Engine/Source/Runtime/Core/Private/HAL/LowLevelMemTracker.cpp b/Engine/Source/Runtime/Core/Private/HAL/LowLevelMemTracker.cpp
--- a/Engine/Source/Runtime/Core/Private/HAL/LowLevelMemTracker.cpp
+++ b/Engine/Source/Runtime/Core/Private/HAL/LowLevelMemTracker.cpp
@@ -104,6 +104,8 @@ DECLARE_LLM_MEMORY_STAT(TEXT("SkeletalMesh"), STAT_SkeletalMeshLLM, STATGROUP_LL
 DECLARE_LLM_MEMORY_STAT(TEXT("InstancedMesh"), STAT_InstancedMeshLLM, STATGROUP_LLMFULL);
 DECLARE_LLM_MEMORY_STAT(TEXT("Landscape"), STAT_LandscapeLLM, STATGROUP_LLMFULL);
 DECLARE_LLM_MEMORY_STAT(TEXT("VideoStreaming"), STAT_VideoStreamingLLM, STATGROUP_LLMFULL);
+DECLARE_LLM_MEMORY_STAT(TEXT("MMIO"), STAT_PlatformMMIOLLM, STATGROUP_LLMPlatform);
+DECLARE_LLM_MEMORY_STAT(TEXT("VirtualMemory"), STAT_PlatformVMLLM, STATGROUP_LLMPlatform);
 
 /*
 * LLM Summary stats referenced by ELLMTagNames
diff --git a/Engine/Source/Runtime/Core/Private/IOS/IOSPlatformFile.cpp b/Engine/Source/Runtime/Core/Private/IOS/IOSPlatformFile.cpp
--- a/Engine/Source/Runtime/Core/Private/IOS/IOSPlatformFile.cpp
+++ b/Engine/Source/Runtime/Core/Private/IOS/IOSPlatformFile.cpp
@@ -461,6 +461,7 @@ public:
 	
 	virtual IMappedFileRegion* MapRegion(int64 Offset = 0, int64 BytesToMap = MAX_int64, bool bPreloadHint = false) override
 	{
+		LLM_PLATFORM_SCOPE(ELLMTag::PlatformMMIO);
 		check(Offset < GetFileSize()); // don't map zero bytes and don't map off the end of the file
 		BytesToMap = FMath::Min<int64>(BytesToMap, GetFileSize() - Offset);
 		check(BytesToMap > 0); // don't map zero bytes
@@ -486,6 +487,7 @@ public:
 			UE_LOG(LogIOS, Warning, TEXT("Failed to map memory %s, error is %d"), *Filename, errno);
 			return nullptr;
 		}
+		LLM(FLowLevelMemTracker::Get().OnLowLevelAlloc(ELLMTracker::Platform, AlignedMapPtr, AlignedSize));
 
 		// create a mapping for this range
 		const uint8* MapPtr = AlignedMapPtr + Offset - AlignedOffset;
@@ -496,9 +498,11 @@ public:
 	
 	void UnMap(FIOSMappedFileRegion* Region)
 	{
+		LLM_PLATFORM_SCOPE(ELLMTag::PlatformMMIO);
 		check(NumOutstandingRegions > 0);
 		NumOutstandingRegions--;
 		
+		LLM(FLowLevelMemTracker::Get().OnLowLevelFree(ELLMTracker::Platform, (void*)Region->AlignedPtr));
 		int Res = munmap((void*)Region->AlignedPtr, Region->AlignedSize);
 		checkf(Res == 0, TEXT("Failed to unmap, error is %d, errno is %d [params: %x, %d]"), Res, errno, MappedPtr, GetFileSize());
 	}
diff --git a/Engine/Source/Runtime/Core/Public/HAL/LowLevelMemTracker.h b/Engine/Source/Runtime/Core/Public/HAL/LowLevelMemTracker.h
--- a/Engine/Source/Runtime/Core/Public/HAL/LowLevelMemTracker.h
+++ b/Engine/Source/Runtime/Core/Public/HAL/LowLevelMemTracker.h
@@ -218,6 +218,8 @@ enum class ELLMTagSet : uint8
 	macro(Landscape,							"Landscape",					GET_STATFNAME(STAT_LandscapeLLM),							GET_STATFNAME(STAT_EngineSummaryLLM),			ELLMTag::Meshes)\
 	macro(CsvProfiler,							"CsvProfiler",					GET_STATFNAME(STAT_CsvProfilerLLM),							GET_STATFNAME(STAT_EngineSummaryLLM),			-1)\
 	macro(VideoStreaming,						"VideoStreaming",				GET_STATFNAME(STAT_VideoStreamingLLM),						GET_STATFNAME(STAT_EngineSummaryLLM),			-1)\
+	macro(PlatformMMIO,							"MMIO",							GET_STATFNAME(STAT_PlatformMMIOLLM),						NAME_None,										-1)\
+	macro(PlatformVM,							"Virtual Memory",				GET_STATFNAME(STAT_PlatformVMLLM),							NAME_None,										-1)\
 
 /*
  * Enum values to be passed in to LLM_SCOPE() macro
