PR #6077 : Added ToString() method to FRandomStream

diff --git a/Engine/Source/Runtime/Core/Public/Math/RandomStream.h b/Engine/Source/Runtime/Core/Public/Math/RandomStream.h
--- a/Engine/Source/Runtime/Core/Public/Math/RandomStream.h
+++ b/Engine/Source/Runtime/Core/Public/Math/RandomStream.h
@@ -328,6 +328,16 @@ public:
 	 */
 	CORE_API bool ExportTextItem(FString& ValueStr, FRandomStream const& DefaultValue, class UObject* Parent, int32 PortFlags, class UObject* ExportRootScope) const;
 
+	/**
+	 * Get a textual representation of the RandomStream.
+	 *
+	 * @return Text describing the RandomStream.
+	 */
+	FString ToString() const
+	{
+		return FString::Printf(TEXT("FRandomStream(InitialSeed=%i, Seed=%i)"), InitialSeed, Seed);
+	}
+
 protected:
 
 	/**
