Added constructors to FAutoConsoleVariable that allow passing in a FConsoleVariableDelegate as a OnChangedCallback.

#rnx

diff --git a/Engine/Source/Runtime/Core/Public/HAL/IConsoleManager.h b/Engine/Source/Runtime/Core/Public/HAL/IConsoleManager.h
--- a/Engine/Source/Runtime/Core/Public/HAL/IConsoleManager.h
+++ b/Engine/Source/Runtime/Core/Public/HAL/IConsoleManager.h
@@ -844,6 +844,46 @@ public:
 		: FAutoConsoleObject(IConsoleManager::Get().RegisterConsoleVariable(Name, DefaultValue, Help, Flags))
 	{
 	}
+
+	/**
+	 * Create a int console variable
+	 * @param Name must not be 0
+	 * @param Help must not be 0
+	 * @param Callback Delegate called when the variable changes. @see IConsoleVariable::SetOnChangedCallback
+	 * @param Flags bitmask combined from EConsoleVariableFlags
+	 */
+	FAutoConsoleVariable(const TCHAR* Name, int32 DefaultValue, const TCHAR* Help, const FConsoleVariableDelegate& Callback, uint32 Flags = ECVF_Default)
+		: FAutoConsoleObject(IConsoleManager::Get().RegisterConsoleVariable(Name, DefaultValue, Help, Flags))
+	{
+		AsVariable()->SetOnChangedCallback(Callback);
+	}
+
+	/**
+	 * Create a float console variable
+	 * @param Name must not be 0
+	 * @param Help must not be 0
+	 * @param Callback Delegate called when the variable changes. @see IConsoleVariable::SetOnChangedCallback
+	 * @param Flags bitmask combined from EConsoleVariableFlags
+	 */
+	FAutoConsoleVariable(const TCHAR* Name, float DefaultValue, const TCHAR* Help, const FConsoleVariableDelegate& Callback, uint32 Flags = ECVF_Default)
+		: FAutoConsoleObject(IConsoleManager::Get().RegisterConsoleVariable(Name, DefaultValue, Help, Flags))
+	{
+		AsVariable()->SetOnChangedCallback(Callback);
+	}
+
+	/**
+	 * Create a string console variable
+	 * @param Name must not be 0
+	 * @param Help must not be 0
+	 * @param Callback Delegate called when the variable changes. @see IConsoleVariable::SetOnChangedCallback
+	 * @param Flags bitmask combined from EConsoleVariableFlags
+	 */
+	FAutoConsoleVariable(const TCHAR* Name, const TCHAR* DefaultValue, const TCHAR* Help, const FConsoleVariableDelegate& Callback, uint32 Flags = ECVF_Default)
+		: FAutoConsoleObject(IConsoleManager::Get().RegisterConsoleVariable(Name, DefaultValue, Help, Flags))
+	{
+		AsVariable()->SetOnChangedCallback(Callback);
+	}
+
 	/** Dereference back to a console variable**/
 	FORCEINLINE IConsoleVariable& operator*()
 	{
