Restore removed version so assets saved while the version was in place can still load
#rnx

diff --git a/Engine/Source/Runtime/Core/Public/UObject/ReleaseObjectVersion.h b/Engine/Source/Runtime/Core/Public/UObject/ReleaseObjectVersion.h
--- a/Engine/Source/Runtime/Core/Public/UObject/ReleaseObjectVersion.h
+++ b/Engine/Source/Runtime/Core/Public/UObject/ReleaseObjectVersion.h
@@ -80,6 +80,9 @@ struct CORE_API FReleaseObjectVersion
 		// DepthOfFieldFocalDistance at 0 now disables DOF instead of DepthOfFieldFstop at 0.
 		FocalDistanceDisablesDOF,
 
+		// Removed versioning, but version entry must still exist to keep assets saved with this version loadable
+		Unused_SoundClass2DReverbSend,
+
 		// -----<new versions can be added above this line>-------------------------------------------------
 		VersionPlusOne,
 		LatestVersion = VersionPlusOne - 1
