UE-78915 HoloLens - Packaged apps sometimes crash on launch
-Task jobs were not projected by critical section, and could get into unexpected states by being processed more than once simultaneously.
-Corrected hololens implementation of TryLock so that it keeps the lock when it finds it available rather than immediately releasing it, matching the behavior on other platforms.

diff --git a/Engine/Source/Runtime/Core/Public/HoloLens/HoloLensCriticalSection.h b/Engine/Source/Runtime/Core/Public/HoloLens/HoloLensCriticalSection.h
--- a/Engine/Source/Runtime/Core/Public/HoloLens/HoloLensCriticalSection.h
+++ b/Engine/Source/Runtime/Core/Public/HoloLens/HoloLensCriticalSection.h
@@ -51,13 +51,14 @@ public:
 	}
 
 	/**
-	* Quick test for seeing if the lock is already being used.
+	 * Attempt to take a lock and returns whether or not a lock was taken.
+	 *
+	 * @return true if a lock was taken, false otherwise.
 	*/
 	FORCEINLINE bool TryLock()
 	{
 		if (Windows::TryEnterCriticalSection(&CriticalSection))
 		{
-			Windows::LeaveCriticalSection(&CriticalSection);
 			return true;
 		};
 		return false;
