Fixed off by 1 error in FNameEntry::GetSize. FNameEntry strings are not null terminated, so a +1 when calculating size is incorrect.

This function is hit when running console command 'ViewNames'.

#rnx

diff --git a/Engine/Source/Runtime/Core/Private/UObject/UnrealNames.cpp b/Engine/Source/Runtime/Core/Private/UObject/UnrealNames.cpp
--- a/Engine/Source/Runtime/Core/Private/UObject/UnrealNames.cpp
+++ b/Engine/Source/Runtime/Core/Private/UObject/UnrealNames.cpp
@@ -1162,7 +1162,7 @@ int32 FNameEntry::GetSize(const TCHAR* Name)
 
 int32 FNameEntry::GetSize(int32 Length, bool bIsPureAnsi)
 {
-	int32 Bytes = GetDataOffset() + (Length + 1) * (bIsPureAnsi ? sizeof(ANSICHAR) : sizeof(TCHAR));
+	int32 Bytes = GetDataOffset() + Length * (bIsPureAnsi ? sizeof(ANSICHAR) : sizeof(WIDECHAR));
 	return Align(Bytes, alignof(FNameEntry));
 }
 
