Tick FEmbeddedCommunication directly in FEngineLoop::Tick instead of CoreTicker

diff --git a/Engine/Source/Runtime/Core/Private/Misc/EmbeddedCommunication.cpp b/Engine/Source/Runtime/Core/Private/Misc/EmbeddedCommunication.cpp
--- a/Engine/Source/Runtime/Core/Private/Misc/EmbeddedCommunication.cpp
+++ b/Engine/Source/Runtime/Core/Private/Misc/EmbeddedCommunication.cpp
@@ -72,8 +72,6 @@ void FEmbeddedCommunication::Init()
  	FScopeLock Lock(&GEmbeddedLock);
  	GSleepEvent = FPlatformProcess::GetSynchEventFromPool(false);
 	GTickWithoutSleepCount = 0;
-
-	FTicker::GetCoreTicker().AddTicker(FTickerDelegate::CreateStatic(&FEmbeddedCommunication::TickGameThread));
 #endif
 }
 
@@ -128,6 +126,7 @@ void FEmbeddedCommunication::ForceTick(int ID, float MinTimeSlice, float MaxTime
 		//We have to manually tick everything as we are looping the main thread here
 		FTicker::GetCoreTicker().Tick(Now - LastTime);
 		FThreadManager::Get().Tick();
+		TickGameThread(Now - LastTime);
 		
 		FPlatformProcess::Sleep(DeltaTime);
 
@@ -290,10 +289,8 @@ void FEmbeddedCommunication::RunOnGameThread(int Priority, TFunction<void()> Lam
 void FEmbeddedCommunication::WakeGameThread()
 {
 #if BUILD_EMBEDDED_APP
-	// Allow 2 ticks without a sleep
-	// Our sleep happens in the core ticker's tick, and that order gets reversed every tick,
-	// so the caller isn't guaranteed to get a tick before our next sleep
-	GTickWithoutSleepCount = 2;
+	// Make sure we tick at least one full tick
+	GTickWithoutSleepCount = 1;
 	// wake up the game thread!
 	if (GSleepEvent)
 	{
diff --git a/Engine/Source/Runtime/Launch/Private/LaunchEngineLoop.cpp b/Engine/Source/Runtime/Launch/Private/LaunchEngineLoop.cpp
--- a/Engine/Source/Runtime/Launch/Private/LaunchEngineLoop.cpp
+++ b/Engine/Source/Runtime/Launch/Private/LaunchEngineLoop.cpp
@@ -4486,6 +4486,9 @@ void FEngineLoop::Tick()
 
 			FTicker::GetCoreTicker().Tick(FApp::GetDeltaTime());
 			FThreadManager::Get().Tick();
+#if BUILD_EMBEDDED_APP
+			FEmbeddedCommunication::TickGameThread(FApp::GetDeltaTime());
+#endif
 			GEngine->TickDeferredCommands();		
 		}
 
