Core - Adding full support for bool to the console variable system.  Before the support was very spotty, they worker on refs, but nothing else, and nobody properly implemented the IsVariable___() functions.  So I made a pass at getting it working throughout the system, and shuring up specializations to provide overrides for the IsVariable___() functions so that you can correctly determine what the underlying type of the variable is if you need it (i need it).  Also, some things used TTypeFromString<T>, other places atoi(), so trying to normalize around TTypeFromString.

Bob.Tellez, Michael.Noland

diff --git a/Engine/Source/Runtime/Core/Private/HAL/ConsoleManager.cpp b/Engine/Source/Runtime/Core/Private/HAL/ConsoleManager.cpp
--- a/Engine/Source/Runtime/Core/Private/HAL/ConsoleManager.cpp
+++ b/Engine/Source/Runtime/Core/Private/HAL/ConsoleManager.cpp
@@ -318,15 +318,19 @@ public:
 		}
 	}
 
-	virtual int32 GetInt() const;
-	virtual float GetFloat() const;
-	virtual FString GetString() const;
-	virtual bool IsVariableInt() const override { return false; }
-	virtual bool IsVariableFloat() const override { return false; }
-	virtual bool IsVariableString() const override { return false; }
-	virtual class TConsoleVariableData<int32>* AsVariableInt() { return 0; }
-	virtual class TConsoleVariableData<float>* AsVariableFloat() { return 0; }
-	virtual class TConsoleVariableData<FString>* AsVariableString() { return 0; }
+	virtual bool GetBool() const override;
+	virtual int32 GetInt() const override;
+	virtual float GetFloat() const override;
+	virtual FString GetString() const override;
+
+	virtual bool IsVariableBool() const { return false; }
+	virtual bool IsVariableFloat() const { return false; }
+	virtual bool IsVariableString() const { return false; }
+
+	virtual class TConsoleVariableData<bool>* AsVariableBool() { return nullptr; }
+	virtual class TConsoleVariableData<int32>* AsVariableInt() { return nullptr; }
+	virtual class TConsoleVariableData<float>* AsVariableFloat() { return nullptr; }
+	virtual class TConsoleVariableData<FString>* AsVariableString() { return nullptr; }
 
 private: // ----------------------------------------------------
 
@@ -347,9 +351,55 @@ private: // ----------------------------------------------------
 	}
 };
 
+// specialization for all
+
+template<> bool FConsoleVariable<bool>::IsVariableBool() const
+{
+	return true;
+}
+template<> bool FConsoleVariable<int32>::IsVariableInt() const
+{
+	return true;
+}
+template<> bool FConsoleVariable<float>::IsVariableFloat() const
+{
+	return true;
+}
+template<> bool FConsoleVariable<FString>::IsVariableString() const
+{
+	return true;
+}
+
+
+// specialization for bool
+
+template<> bool FConsoleVariable<bool>::GetBool() const
+{
+	return Value();
+}
+template<> int32 FConsoleVariable<bool>::GetInt() const
+{
+	return Value() ? 1 : 0;
+}
+template<> float FConsoleVariable<bool>::GetFloat() const
+{
+	return Value() ? 1.0f : 0.0f;
+}
+template<> FString FConsoleVariable<bool>::GetString() const
+{
+	return Value() ? TEXT("true") : TEXT("false");
+}
+template<> TConsoleVariableData<bool>* FConsoleVariable<bool>::AsVariableBool()
+{
+	return &Data;
+}
 
 // specialization for int32
 
+template<> bool FConsoleVariable<int32>::GetBool() const
+{
+	return Value() != 0;
+}
 template<> int32 FConsoleVariable<int32>::GetInt() const
 {
 	return Value();
@@ -363,11 +413,6 @@ template<> FString FConsoleVariable<int32>::GetString() const
 	return FString::Printf(TEXT("%d"), Value());
 }
 
-template<> bool FConsoleVariable<int32>::IsVariableInt() const
-{
-	return true; 
-}
-
 template<> TConsoleVariableData<int32>* FConsoleVariable<int32>::AsVariableInt()
 {
 	return &Data; 
@@ -375,6 +420,10 @@ template<> TConsoleVariableData<int32>* FConsoleVariable<int32>::AsVariableInt()
 
 // specialization for float
 
+template<> bool FConsoleVariable<float>::GetBool() const
+{
+	return Value() != 0;
+}
 template<> int32 FConsoleVariable<float>::GetInt() const
 {
 	return (int32)Value();
@@ -387,10 +436,6 @@ template<> FString FConsoleVariable<float>::GetString() const
 {
 	return FString::Printf(TEXT("%g"), Value());
 }
-template<> bool FConsoleVariable<float>::IsVariableFloat() const
-{
-	return true;
-}
 template<> TConsoleVariableData<float>* FConsoleVariable<float>::AsVariableFloat()
 {
 	return &Data;
@@ -406,13 +451,23 @@ template<> void FConsoleVariable<FString>::Set(const TCHAR* InValue, EConsoleVar
 		OnChanged(SetBy);
 	}
 }
+template<> bool FConsoleVariable<FString>::GetBool() const
+{
+	bool OutValue = false;
+	TTypeFromString<bool>::FromString(OutValue, *Value());
+	return OutValue;
+}
 template<> int32 FConsoleVariable<FString>::GetInt() const
 {
-	return FCString::Atoi(*Value());
+	int32 OutValue = 0;
+	TTypeFromString<int32>::FromString(OutValue, *Value());
+	return OutValue;
 }
 template<> float FConsoleVariable<FString>::GetFloat() const
 {
-	return FCString::Atof(*Value());
+	float OutValue = 0.0f;
+	TTypeFromString<float>::FromString(OutValue, *Value());
+	return OutValue;
 }
 template<> FString FConsoleVariable<FString>::GetString() const
 {
@@ -453,6 +508,16 @@ public:
 			OnChanged(SetBy);
 		}
 	}
+
+	virtual bool IsVariableBool() const { return false; }
+	virtual bool IsVariableInt() const { return false; }
+	virtual bool IsVariableFloat() const { return false; }
+	virtual bool IsVariableString() const { return false; }
+
+	virtual bool GetBool() const
+	{
+		return (bool)MainValue;
+	}
 	virtual int32 GetInt() const
 	{
 		return (int32)MainValue;
@@ -494,25 +559,31 @@ private: // ----------------------------------------------------
 	}
 };
 
-// specialization for float
-
 template <>
-FString FConsoleVariableRef<float>::GetString() const
+bool FConsoleVariableRef<bool>::IsVariableBool() const
 {
-	// otherwise we get 2.1f would become "2.100000"
-	return FString::SanitizeFloat(RefValue);
+	return true;
 }
-template<> bool FConsoleVariableRef<bool>::IsVariableInt() const
+
+template <>
+bool FConsoleVariableRef<int32>::IsVariableInt() const
 {
 	return true;
 }
-template<> bool FConsoleVariableRef<int32>::IsVariableInt() const
+
+template <>
+bool FConsoleVariableRef<float>::IsVariableFloat() const
 {
 	return true;
 }
-template<> bool FConsoleVariableRef<float>::IsVariableFloat() const
+
+// specialization for float
+
+template <>
+FString FConsoleVariableRef<float>::GetString() const
 {
-	return true;
+	// otherwise we get 2.1f would become "2.100000"
+	return FString::SanitizeFloat(RefValue);
 }
 
 // string version
@@ -541,9 +612,15 @@ public:
 			OnChanged(SetBy);
 		}
 	}
+	virtual bool GetBool() const
+	{
+		bool Result = false;
+		TTypeFromString<bool>::FromString(Result, *MainValue);
+		return Result;
+	}
 	virtual int32 GetInt() const
 	{
-		int32 Result;
+		int32 Result = 0;
 		TTypeFromString<int32>::FromString(Result, *MainValue);
 		return Result;
 	}
@@ -615,6 +692,10 @@ public:
 			OnChanged(SetBy);
 		}
 	}
+	virtual bool GetBool() const
+	{
+		return GetInt() != 0;
+	}
 	virtual int32 GetInt() const
 	{
 		// we apply the bitmask on game thread (showflags) so we don't have to do any special thread handling
@@ -926,6 +1007,11 @@ public:
 	}
 };
 
+IConsoleVariable* FConsoleManager::RegisterConsoleVariable(const TCHAR* Name, bool DefaultValue, const TCHAR* Help, uint32 Flags)
+{
+	return AddConsoleObject(Name, new FConsoleVariable<bool>(DefaultValue, Help, (EConsoleVariableFlags)Flags))->AsVariable();
+}
+
 IConsoleVariable* FConsoleManager::RegisterConsoleVariable(const TCHAR* Name, int32 DefaultValue, const TCHAR* Help, uint32 Flags)
 {
 	return AddConsoleObject(Name, new FConsoleVariable<int32>(DefaultValue, Help, (EConsoleVariableFlags)Flags))->AsVariable();
@@ -936,6 +1022,11 @@ IConsoleVariable* FConsoleManager::RegisterConsoleVariable(const TCHAR* Name, fl
 	return AddConsoleObject(Name, new FConsoleVariable<float>(DefaultValue, Help, (EConsoleVariableFlags)Flags))->AsVariable();
 }
 
+IConsoleVariable* FConsoleManager::RegisterConsoleVariable(const TCHAR* Name, const TCHAR* DefaultValue, const TCHAR* Help, uint32 Flags)
+{
+	return RegisterConsoleVariable(Name, FString(DefaultValue), Help, Flags);
+}
+
 IConsoleVariable* FConsoleManager::RegisterConsoleVariable(const TCHAR* Name, const FString& DefaultValue, const TCHAR* Help, uint32 Flags)
 {
 	// not supported
@@ -943,6 +1034,11 @@ IConsoleVariable* FConsoleManager::RegisterConsoleVariable(const TCHAR* Name, co
 	return AddConsoleObject(Name, new FConsoleVariable<FString>(DefaultValue, Help, (EConsoleVariableFlags)Flags))->AsVariable();
 }
 
+IConsoleVariable* FConsoleManager::RegisterConsoleVariableRef(const TCHAR* Name, bool& RefValue, const TCHAR* Help, uint32 Flags)
+{
+	return AddConsoleObject(Name, new FConsoleVariableRef<bool>(RefValue, Help, (EConsoleVariableFlags)Flags))->AsVariable();
+}
+
 IConsoleVariable* FConsoleManager::RegisterConsoleVariableRef(const TCHAR* Name, int32& RefValue, const TCHAR* Help, uint32 Flags)
 {
 	return AddConsoleObject(Name, new FConsoleVariableRef<int32>(RefValue, Help, (EConsoleVariableFlags)Flags))->AsVariable();
@@ -953,11 +1049,6 @@ IConsoleVariable* FConsoleManager::RegisterConsoleVariableRef(const TCHAR* Name,
 	return AddConsoleObject(Name, new FConsoleVariableRef<float>(RefValue, Help, (EConsoleVariableFlags)Flags))->AsVariable();
 }
 
-IConsoleVariable* FConsoleManager::RegisterConsoleVariableRef(const TCHAR* Name, bool& RefValue, const TCHAR* Help, uint32 Flags)
-{
-	return AddConsoleObject(Name, new FConsoleVariableRef<bool>(RefValue, Help, (EConsoleVariableFlags)Flags))->AsVariable();
-}
-
 IConsoleVariable* FConsoleManager::RegisterConsoleVariableRef(const TCHAR* Name, FString& RefValue, const TCHAR* Help, uint32 Flags)
 {
 	return AddConsoleObject(Name, new FConsoleVariableStringRef(RefValue, Help, (EConsoleVariableFlags)Flags))->AsVariable();
diff --git a/Engine/Source/Runtime/Core/Public/HAL/ConsoleManager.h b/Engine/Source/Runtime/Core/Public/HAL/ConsoleManager.h
--- a/Engine/Source/Runtime/Core/Public/HAL/ConsoleManager.h
+++ b/Engine/Source/Runtime/Core/Public/HAL/ConsoleManager.h
@@ -46,12 +46,15 @@ public:
 
 	// interface IConsoleManager -----------------------------------
 
+	virtual IConsoleVariable* RegisterConsoleVariable(const TCHAR* Name, bool DefaultValue, const TCHAR* Help, uint32 Flags) override;
 	virtual IConsoleVariable* RegisterConsoleVariable(const TCHAR* Name, int32 DefaultValue, const TCHAR* Help, uint32 Flags) override;
 	virtual IConsoleVariable* RegisterConsoleVariable(const TCHAR* Name, float DefaultValue, const TCHAR* Help, uint32 Flags) override;
+	virtual IConsoleVariable* RegisterConsoleVariable(const TCHAR* Name, const TCHAR* DefaultValue, const TCHAR* Help, uint32 Flags) override;
 	virtual IConsoleVariable* RegisterConsoleVariable(const TCHAR* Name, const FString& DefaultValue, const TCHAR* Help, uint32 Flags) override;
+
+	virtual IConsoleVariable* RegisterConsoleVariableRef(const TCHAR* Name, bool& RefValue, const TCHAR* Help, uint32 Flags) override;
 	virtual IConsoleVariable* RegisterConsoleVariableRef(const TCHAR* Name, int32& RefValue, const TCHAR* Help, uint32 Flags) override;
 	virtual IConsoleVariable* RegisterConsoleVariableRef(const TCHAR* Name, float& RefValue, const TCHAR* Help, uint32 Flags) override;
-	virtual IConsoleVariable* RegisterConsoleVariableRef(const TCHAR* Name, bool& RefValue, const TCHAR* Help, uint32 Flags) override;
 	virtual IConsoleVariable* RegisterConsoleVariableRef(const TCHAR* Name, FString& RefValue, const TCHAR* Help, uint32 Flags) override;
 	virtual IConsoleVariable* RegisterConsoleVariableBitRef(const TCHAR* CVarName, const TCHAR* FlagName, uint32 BitNumber, uint8* Force0MaskPtr, uint8* Force1MaskPtr, const TCHAR* Help, uint32 Flags) override;
 	
diff --git a/Engine/Source/Runtime/Core/Public/HAL/IConsoleManager.h b/Engine/Source/Runtime/Core/Public/HAL/IConsoleManager.h
--- a/Engine/Source/Runtime/Core/Public/HAL/IConsoleManager.h
+++ b/Engine/Source/Runtime/Core/Public/HAL/IConsoleManager.h
@@ -213,10 +213,17 @@ public:
 		return 0; 
 	}
 
+	virtual bool IsVariableBool() const { return false; }
 	virtual bool IsVariableInt() const { return false; }
 	virtual bool IsVariableFloat() const { return false; }
 	virtual bool IsVariableString() const { return false; }
 
+	virtual class TConsoleVariableData<bool>* AsVariableBool()
+	{
+		ensureMsgf(false, TEXT("Attempted to access variable data of a console variable type that doesn't support it.  For example FindTConsoleVariableData* on a FAutoConsoleVariableRef."));
+		return 0;
+	}
+
 	virtual class TConsoleVariableData<int32>* AsVariableInt()
 	{
 		ensureMsgf(false, TEXT("Attempted to access variable data of a console variable type that doesn't support it.  For example FindTConsoleVariableData* on a FAutoConsoleVariableRef."));
@@ -271,6 +278,11 @@ public:
 	 * @param SetBy anything in ECVF_LastSetMask e.g. ECVF_SetByScalability
 	 **/
 	virtual void Set(const TCHAR* InValue, EConsoleVariableFlags SetBy = ECVF_SetByCode) = 0;
+
+	/**
+	 * Get the internal value as a bool, works on bools, ints and floats.
+	 */
+	virtual bool GetBool() const = 0;
 	/**
 	 * Get the internal value as int (should not be used on strings).
 	 * @return value is not rounded (simple cast)
@@ -295,6 +307,12 @@ public:
 
 	// convenience methods
 
+	/** Set the internal value from the specified bool. */
+	void Set(bool InValue, EConsoleVariableFlags SetBy = ECVF_SetByCode)
+	{
+		// inefficient but no common code path
+		Set(InValue ? TEXT("true") : TEXT("false"), SetBy);
+	}
 	/** Set the internal value from the specified int. */
 	void Set(int32 InValue, EConsoleVariableFlags SetBy = ECVF_SetByCode)
 	{
@@ -307,6 +325,12 @@ public:
 		// inefficient but no common code path
 		Set(*FString::Printf(TEXT("%g"), InValue), SetBy);
 	}
+
+	void SetWithCurrentPriority(bool InValue)
+	{
+		EConsoleVariableFlags CurFlags = (EConsoleVariableFlags)(GetFlags() & ECVF_SetByMask);
+		Set(InValue, CurFlags);
+	}
 	void SetWithCurrentPriority(int32 InValue)
 	{
 		EConsoleVariableFlags CurFlags = (EConsoleVariableFlags)(GetFlags() & ECVF_SetByMask);
@@ -468,6 +492,13 @@ public:
  */
 struct CORE_API IConsoleManager
 {
+	/**
+	 * Create a bool console variable
+	 * @param Name must not be 0
+	 * @param Help must not be 0
+	 * @param Flags bitmask combined from EConsoleVariableFlags
+	 */
+	virtual IConsoleVariable* RegisterConsoleVariable(const TCHAR* Name, bool DefaultValue, const TCHAR* Help, uint32 Flags = ECVF_Default) = 0;
 	/**
 	 * Create a int console variable
 	 * @param Name must not be 0
@@ -482,6 +513,13 @@ struct CORE_API IConsoleManager
 	 * @param Flags bitmask combined from EConsoleVariableFlags
 	 */
 	virtual IConsoleVariable* RegisterConsoleVariable(const TCHAR* Name, float DefaultValue, const TCHAR* Help, uint32 Flags = ECVF_Default) = 0;
+	/**
+	 * Create a string console variable
+	 * @param Name must not be 0
+	 * @param Help must not be 0
+	 * @param Flags bitmask combined from EConsoleVariableFlags
+	 */
+	virtual IConsoleVariable* RegisterConsoleVariable(const TCHAR* Name, const TCHAR* DefaultValue, const TCHAR* Help, uint32 Flags = ECVF_Default) = 0;
 	/**
 	 * Create a string console variable
 	 * @param Name must not be 0
@@ -489,6 +527,13 @@ struct CORE_API IConsoleManager
 	 * @param Flags bitmask combined from EConsoleVariableFlags
 	 */
 	virtual IConsoleVariable* RegisterConsoleVariable(const TCHAR* Name, const FString& DefaultValue, const TCHAR* Help, uint32 Flags = ECVF_Default) = 0;
+	/**
+	 * Create a reference to a bool console variable
+	 * @param Name must not be 0
+	 * @param Help must not be 0
+	 * @param Flags bitmask combined from EConsoleVariableFlags
+	 */
+	virtual IConsoleVariable* RegisterConsoleVariableRef(const TCHAR* Name, bool& RefValue, const TCHAR* Help, uint32 Flags = ECVF_Default) = 0;
 	/**
 	 * Create a reference to a int console variable
 	 * @param Name must not be 0
@@ -504,13 +549,6 @@ struct CORE_API IConsoleManager
 	 */
 	virtual IConsoleVariable* RegisterConsoleVariableRef(const TCHAR* Name, float& RefValue, const TCHAR* Help, uint32 Flags = ECVF_Default) = 0;
 	/**
-	* Create a reference to a bool console variable
-	* @param Name must not be 0
-	* @param Help must not be 0
-	* @param Flags bitmask combined from EConsoleVariableFlags
-	*/
-	virtual IConsoleVariable* RegisterConsoleVariableRef(const TCHAR* Name, bool& RefValue, const TCHAR* Help, uint32 Flags = ECVF_Default) = 0;
-	/**
 	* Create a reference to a string console variable
 	* @param Name must not be 0
 	* @param Help must not be 0
@@ -814,6 +852,16 @@ private:
 class CORE_API FAutoConsoleVariable : private FAutoConsoleObject
 {
 public:
+	/**
+	 * Create a bool console variable
+	 * @param Name must not be 0
+	 * @param Help must not be 0
+	 * @param Flags bitmask combined from EConsoleVariableFlags
+	 */
+	FAutoConsoleVariable(const TCHAR* Name, bool DefaultValue, const TCHAR* Help, uint32 Flags = ECVF_Default)
+		: FAutoConsoleObject(IConsoleManager::Get().RegisterConsoleVariable(Name, DefaultValue, Help, Flags))
+	{
+	}
 	/**
 	 * Create a int console variable
 	 * @param Name must not be 0
@@ -1146,6 +1194,13 @@ private:
 	TConsoleVariableData<T>* Ref;
 };
 
+template <>
+inline TAutoConsoleVariable<bool>::TAutoConsoleVariable(const TCHAR* Name, const bool& DefaultValue, const TCHAR* Help, uint32 Flags)
+	: FAutoConsoleObject(IConsoleManager::Get().RegisterConsoleVariable(Name, DefaultValue, Help, Flags))
+{
+	Ref = AsVariable()->AsVariableBool();
+}
+
 template <>
 inline TAutoConsoleVariable<int32>::TAutoConsoleVariable(const TCHAR* Name, const int32& DefaultValue, const TCHAR* Help, uint32 Flags)
 	: FAutoConsoleObject(IConsoleManager::Get().RegisterConsoleVariable(Name, DefaultValue, Help, Flags))
