Fixing build issues after merge.

diff --git a/Engine/Source/Runtime/Core/Private/HAL/ConsoleManager.cpp b/Engine/Source/Runtime/Core/Private/HAL/ConsoleManager.cpp
--- a/Engine/Source/Runtime/Core/Private/HAL/ConsoleManager.cpp
+++ b/Engine/Source/Runtime/Core/Private/HAL/ConsoleManager.cpp
@@ -323,14 +323,15 @@ public:
 	virtual float GetFloat() const override;
 	virtual FString GetString() const override;
 
-	virtual bool IsVariableBool() const { return false; }
-	virtual bool IsVariableFloat() const { return false; }
-	virtual bool IsVariableString() const { return false; }
+virtual bool IsVariableBool() const override { return false; }
+	virtual bool IsVariableInt() const override { return false; }
+	virtual bool IsVariableFloat() const override { return false; }
+	virtual bool IsVariableString() const override { return false; }
 
-	virtual class TConsoleVariableData<bool>* AsVariableBool() { return nullptr; }
-	virtual class TConsoleVariableData<int32>* AsVariableInt() { return nullptr; }
-	virtual class TConsoleVariableData<float>* AsVariableFloat() { return nullptr; }
-	virtual class TConsoleVariableData<FString>* AsVariableString() { return nullptr; }
+	virtual class TConsoleVariableData<bool>* AsVariableBool() override { return nullptr; }
+	virtual class TConsoleVariableData<int32>* AsVariableInt() override { return nullptr; }
+	virtual class TConsoleVariableData<float>* AsVariableFloat() override { return nullptr; }
+	virtual class TConsoleVariableData<FString>* AsVariableString() override { return nullptr; }
 
 private: // ----------------------------------------------------
 
@@ -473,10 +474,6 @@ template<> FString FConsoleVariable<FString>::GetString() const
 {
 	return Value();
 }
-template<> bool FConsoleVariable<FString>::IsVariableString() const
-{
-	return true;
-}
 template<> TConsoleVariableData<FString>* FConsoleVariable<FString>::AsVariableString()
 {
 	return &Data;
@@ -509,10 +506,10 @@ public:
 		}
 	}
 
-	virtual bool IsVariableBool() const { return false; }
-	virtual bool IsVariableInt() const { return false; }
-	virtual bool IsVariableFloat() const { return false; }
-	virtual bool IsVariableString() const { return false; }
+	virtual bool IsVariableBool() const override { return false; }
+	virtual bool IsVariableInt() const override { return false; }
+	virtual bool IsVariableFloat() const override { return false; }
+	virtual bool IsVariableString() const override { return false; }
 
 	virtual bool GetBool() const
 	{
@@ -530,9 +527,6 @@ public:
 	{
 		return TTypeToString<T>::ToString(MainValue);
 	}
-	virtual bool IsVariableInt() const override { return false; }
-	virtual bool IsVariableFloat() const override { return false; }
-	virtual bool IsVariableString() const override { return false; }
 
 private: // ----------------------------------------------------
 
diff --git a/Engine/Source/Runtime/Core/Public/HAL/IConsoleManager.h b/Engine/Source/Runtime/Core/Public/HAL/IConsoleManager.h
--- a/Engine/Source/Runtime/Core/Public/HAL/IConsoleManager.h
+++ b/Engine/Source/Runtime/Core/Public/HAL/IConsoleManager.h
@@ -26,8 +26,7 @@ template <class T> class TConsoleVariableData;
  *
  * The variable should be creates early in the initialization but not before (not in global variable construction).
  * Choose the right variable type, consider using a console command if more functionality is needed (see Exec()).
- * Available types: int, float, int&, float&, string
- * There is no bool type as the int type provides enough functionality (0=false, 1=true) and can be extended easily (e.g. 2=auto, -1=debug)
+ * Available types: bool, int, float, bool&, int&, float&, string
  * Always provide a good help text, other should be able to understand the function of the console variable by reading this help.
  * The help length should be limited to a reasonable width in order to work well for low res screen resolutions.
  *
