Break Iterator improvements/fixes
 - Added IBreakIterator::SetStringRef as an optimization for platforms using UTF-16 strings where it can avoid a string copy.
 - Fixed a bug where various break iterators were using the wrong FString consructor for their IBreakIterator::SetString(const TCHAR*, const int32) overload.

diff --git a/Engine/Source/Runtime/Core/Private/Internationalization/CamelCaseBreakIterator.cpp b/Engine/Source/Runtime/Core/Private/Internationalization/CamelCaseBreakIterator.cpp
--- a/Engine/Source/Runtime/Core/Private/Internationalization/CamelCaseBreakIterator.cpp
+++ b/Engine/Source/Runtime/Core/Private/Internationalization/CamelCaseBreakIterator.cpp
@@ -26,11 +26,17 @@ void FCamelCaseBreakIterator::SetString(const FString& InString)
 
 void FCamelCaseBreakIterator::SetString(const TCHAR* const InString, const int32 InStringLength) 
 {
-	String = FString(InString, InStringLength);
+	String = FString(InStringLength, InString);
 	UpdateBreakPointsArray();
 	ResetToBeginning();
 }
 
+void FCamelCaseBreakIterator::SetStringRef(const FString* InString)
+{
+	// TODO: Support external string references as an optimization
+	SetString(*InString);
+}
+
 void FCamelCaseBreakIterator::ClearString()
 {
 	String.Reset();
diff --git a/Engine/Source/Runtime/Core/Private/Internationalization/CamelCaseBreakIterator.h b/Engine/Source/Runtime/Core/Private/Internationalization/CamelCaseBreakIterator.h
--- a/Engine/Source/Runtime/Core/Private/Internationalization/CamelCaseBreakIterator.h
+++ b/Engine/Source/Runtime/Core/Private/Internationalization/CamelCaseBreakIterator.h
@@ -19,6 +19,7 @@ public:
 	virtual void SetString(const FText& InText) override;
 	virtual void SetString(const FString& InString) override;
 	virtual void SetString(const TCHAR* const InString, const int32 InStringLength) override;
+	virtual void SetStringRef(const FString* InString) override;
 	virtual void ClearString() override;
 
 	virtual int32 GetCurrentPosition() const override;
diff --git a/Engine/Source/Runtime/Core/Private/Internationalization/ICUBreakIterator.cpp b/Engine/Source/Runtime/Core/Private/Internationalization/ICUBreakIterator.cpp
--- a/Engine/Source/Runtime/Core/Private/Internationalization/ICUBreakIterator.cpp
+++ b/Engine/Source/Runtime/Core/Private/Internationalization/ICUBreakIterator.cpp
@@ -107,6 +107,12 @@ void FICUBreakIterator::SetString(const TCHAR* const InString, const int32 InStr
 	ResetToBeginning();
 }
 
+void FICUBreakIterator::SetStringRef(const FString* InString)
+{
+	GetInternalBreakIterator()->adoptText(new FICUTextCharacterIterator(InString)); // ICUBreakIterator takes ownership of this instance
+	ResetToBeginning();
+}
+
 void FICUBreakIterator::ClearString()
 {
 	GetInternalBreakIterator()->adoptText(new FICUTextCharacterIterator(FString())); // ICUBreakIterator takes ownership of this instance
diff --git a/Engine/Source/Runtime/Core/Private/Internationalization/ICUBreakIterator.h b/Engine/Source/Runtime/Core/Private/Internationalization/ICUBreakIterator.h
--- a/Engine/Source/Runtime/Core/Private/Internationalization/ICUBreakIterator.h
+++ b/Engine/Source/Runtime/Core/Private/Internationalization/ICUBreakIterator.h
@@ -47,6 +47,7 @@ public:
 	virtual void SetString(const FText& InText) override;
 	virtual void SetString(const FString& InString) override;
 	virtual void SetString(const TCHAR* const InString, const int32 InStringLength) override;
+	virtual void SetStringRef(const FString* InString) override;
 	virtual void ClearString() override;
 
 	virtual int32 GetCurrentPosition() const override;
diff --git a/Engine/Source/Runtime/Core/Private/Internationalization/ICULineBreakIterator.cpp b/Engine/Source/Runtime/Core/Private/Internationalization/ICULineBreakIterator.cpp
--- a/Engine/Source/Runtime/Core/Private/Internationalization/ICULineBreakIterator.cpp
+++ b/Engine/Source/Runtime/Core/Private/Internationalization/ICULineBreakIterator.cpp
@@ -52,6 +52,7 @@ public:
 	virtual void SetString(const FText& InText) override;
 	virtual void SetString(const FString& InString) override;
 	virtual void SetString(const TCHAR* const InString, const int32 InStringLength) override;
+	virtual void SetStringRef(const FString* InString) override;
 	virtual void ClearString() override;
 
 	virtual int32 GetCurrentPosition() const override;
@@ -105,6 +106,12 @@ void FICULineBreakIterator::SetString(const TCHAR* const InString, const int32 I
 	ResetToBeginning();
 }
 
+void FICULineBreakIterator::SetStringRef(const FString* InString)
+{
+	GetInternalLineBreakIterator()->adoptText(new FICUTextCharacterIterator(InString)); // ICUBreakIterator takes ownership of this instance
+	ResetToBeginning();
+}
+
 void FICULineBreakIterator::ClearString()
 {
 	GetInternalLineBreakIterator()->adoptText(new FICUTextCharacterIterator(FString())); // ICUBreakIterator takes ownership of this instance
diff --git a/Engine/Source/Runtime/Core/Private/Internationalization/ICUTextCharacterIterator.cpp b/Engine/Source/Runtime/Core/Private/Internationalization/ICUTextCharacterIterator.cpp
--- a/Engine/Source/Runtime/Core/Private/Internationalization/ICUTextCharacterIterator.cpp
+++ b/Engine/Source/Runtime/Core/Private/Internationalization/ICUTextCharacterIterator.cpp
@@ -24,7 +24,7 @@ FICUTextCharacterIterator_NativeUTF16::FICUTextCharacterIterator_NativeUTF16(con
 }
 
 FICUTextCharacterIterator_NativeUTF16::FICUTextCharacterIterator_NativeUTF16(const TCHAR* const InString, const int32 InStringLength)
-	: InternalString(InString, InStringLength)
+	: InternalString(InStringLength, InString)
 	, StringPtr(&InternalString)
 {
 	setText(reinterpret_cast<const UChar*>(**StringPtr), StringPtr->Len()); // scary cast from TCHAR* to UChar* so that this builds on platforms where TCHAR isn't UTF-16 (but we won't use it!)
@@ -129,7 +129,7 @@ FICUTextCharacterIterator_ConvertToUnicodeString::FICUTextCharacterIterator_Conv
 }
 
 FICUTextCharacterIterator_ConvertToUnicodeString::FICUTextCharacterIterator_ConvertToUnicodeString(const TCHAR* const InString, const int32 InStringLength)
-	: FICUTextCharacterIterator_ConvertToUnicodeStringPrivate(FString(InString, InStringLength))
+	: FICUTextCharacterIterator_ConvertToUnicodeStringPrivate(FString(InStringLength, InString))
 	, icu::StringCharacterIterator(InternalString)
 {
 }
diff --git a/Engine/Source/Runtime/Core/Private/Internationalization/LegacyCharacterBoundaryIterator.cpp b/Engine/Source/Runtime/Core/Private/Internationalization/LegacyCharacterBoundaryIterator.cpp
--- a/Engine/Source/Runtime/Core/Private/Internationalization/LegacyCharacterBoundaryIterator.cpp
+++ b/Engine/Source/Runtime/Core/Private/Internationalization/LegacyCharacterBoundaryIterator.cpp
@@ -18,6 +18,7 @@ public:
 	virtual void SetString(const FText& InText) override;
 	virtual void SetString(const FString& InString) override;
 	virtual void SetString(const TCHAR* const InString, const int32 InStringLength) override;
+	virtual void SetStringRef(const FString* InString) override;
 	virtual void ClearString() override;
 
 	virtual int32 GetCurrentPosition() const override;
@@ -55,10 +56,16 @@ void FLegacyCharacterBoundaryIterator::SetString(const FString& InString)
 
 void FLegacyCharacterBoundaryIterator::SetString(const TCHAR* const InString, const int32 InStringLength) 
 {
-	String = FString(InString, InStringLength);
+	String = FString(InStringLength, InString);
 	ResetToBeginning();
 }
 
+void FLegacyCharacterBoundaryIterator::SetStringRef(const FString* InString)
+{
+	// TODO: Support external string references as an optimization
+	SetString(*InString);
+}
+
 void FLegacyCharacterBoundaryIterator::ClearString()
 {
 	String = FString();
diff --git a/Engine/Source/Runtime/Core/Private/Internationalization/LegacyLineBreakIterator.cpp b/Engine/Source/Runtime/Core/Private/Internationalization/LegacyLineBreakIterator.cpp
--- a/Engine/Source/Runtime/Core/Private/Internationalization/LegacyLineBreakIterator.cpp
+++ b/Engine/Source/Runtime/Core/Private/Internationalization/LegacyLineBreakIterator.cpp
@@ -18,6 +18,7 @@ public:
 	virtual void SetString(const FText& InText) override;
 	virtual void SetString(const FString& InString) override;
 	virtual void SetString(const TCHAR* const InString, const int32 InStringLength) override;
+	virtual void SetStringRef(const FString* InString) override;
 	virtual void ClearString() override;
 
 	virtual int32 GetCurrentPosition() const override;
@@ -55,10 +56,16 @@ void FLegacyLineBreakIterator::SetString(const FString& InString)
 
 void FLegacyLineBreakIterator::SetString(const TCHAR* const InString, const int32 InStringLength) 
 {
-	String = FString(InString, InStringLength);
+	String = FString(InStringLength, InString);
 	ResetToBeginning();
 }
 
+void FLegacyLineBreakIterator::SetStringRef(const FString* InString)
+{
+	// TODO: Support external string references as an optimization
+	SetString(*InString);
+}
+
 void FLegacyLineBreakIterator::ClearString()
 {
 	String = FString();
diff --git a/Engine/Source/Runtime/Core/Private/Internationalization/LegacyWordBreakIterator.cpp b/Engine/Source/Runtime/Core/Private/Internationalization/LegacyWordBreakIterator.cpp
--- a/Engine/Source/Runtime/Core/Private/Internationalization/LegacyWordBreakIterator.cpp
+++ b/Engine/Source/Runtime/Core/Private/Internationalization/LegacyWordBreakIterator.cpp
@@ -18,6 +18,7 @@ public:
 	virtual void SetString(const FText& InText) override;
 	virtual void SetString(const FString& InString) override;
 	virtual void SetString(const TCHAR* const InString, const int32 InStringLength) override;
+	virtual void SetStringRef(const FString* InString) override;
 	virtual void ClearString() override;
 
 	virtual int32 GetCurrentPosition() const override;
@@ -55,10 +56,16 @@ void FLegacyWordBreakIterator::SetString(const FString& InString)
 
 void FLegacyWordBreakIterator::SetString(const TCHAR* const InString, const int32 InStringLength) 
 {
-	String = FString(InString, InStringLength);
+	String = FString(InStringLength, InString);
 	ResetToBeginning();
 }
 
+void FLegacyWordBreakIterator::SetStringRef(const FString* InString)
+{
+	// TODO: Support external string references as an optimization
+	SetString(*InString);
+}
+
 void FLegacyWordBreakIterator::ClearString()
 {
 	String = FString();
diff --git a/Engine/Source/Runtime/Core/Public/Internationalization/IBreakIterator.h b/Engine/Source/Runtime/Core/Public/Internationalization/IBreakIterator.h
--- a/Engine/Source/Runtime/Core/Public/Internationalization/IBreakIterator.h
+++ b/Engine/Source/Runtime/Core/Public/Internationalization/IBreakIterator.h
@@ -12,6 +12,7 @@ public:
 	virtual void SetString(const FText& InText) = 0;
 	virtual void SetString(const FString& InString) = 0;
 	virtual void SetString(const TCHAR* const InString, const int32 InStringLength) = 0;
+	virtual void SetStringRef(const FString* InString) = 0;
 	virtual void ClearString() = 0;
 
 	virtual int32 GetCurrentPosition() const = 0;
