Fix issue where a quoted application path would be quoted again before being executed on Windows. Causing CIS errors with automated tests.

#rnx

diff --git a/Engine/Source/Runtime/Core/Private/Windows/WindowsPlatformProcess.cpp b/Engine/Source/Runtime/Core/Private/Windows/WindowsPlatformProcess.cpp
--- a/Engine/Source/Runtime/Core/Private/Windows/WindowsPlatformProcess.cpp
+++ b/Engine/Source/Runtime/Core/Private/Windows/WindowsPlatformProcess.cpp
@@ -725,7 +725,15 @@ bool FWindowsPlatformProcess::ExecProcess(const TCHAR* URL, const TCHAR* Params,
 
 	bool bSuccess = false;
 
-	FString CommandLine = FString::Printf(TEXT("\"%s\" %s"), URL, Params);
+	FString CommandLine;
+	if (URL[0] != '\"') // Don't quote executable name if it's already quoted
+	{
+		CommandLine = FString::Printf(TEXT("\"%s\" %s"), URL, Params); 
+	}
+	else
+	{
+		CommandLine = FString::Printf(TEXT("%s %s"), URL, Params);
+	}
 
 	PROCESS_INFORMATION ProcInfo;
 	if (CreateProcess(NULL, CommandLine.GetCharArray().GetData(), NULL, NULL, TRUE, NORMAL_PRIORITY_CLASS | DETACHED_PROCESS | EXTENDED_STARTUPINFO_PRESENT, NULL, OptionalWorkingDirectory, &StartupInfoEx.StartupInfo, &ProcInfo))
