Core - When setting by bool, and then converting to a string to then set the value, use 0 or 1, as those are backwards compatible with existing code that might now be calling the new bool version, which in turn would would need to send "0" or "1" to work correctly with ConsoleVariable<int32>'s that already exist.

Bob.Tellez, Michael.Noland

diff --git a/Engine/Source/Runtime/Core/Private/HAL/ConsoleManager.cpp b/Engine/Source/Runtime/Core/Private/HAL/ConsoleManager.cpp
--- a/Engine/Source/Runtime/Core/Private/HAL/ConsoleManager.cpp
+++ b/Engine/Source/Runtime/Core/Private/HAL/ConsoleManager.cpp
@@ -293,7 +293,7 @@ void OnCVarChange(T& Dst, const T& Src, EConsoleVariableFlags Flags, EConsoleVar
 	}
 }
 
-// T: int32, float, FString
+// T: bool, int32, float, FString
 template <class T>
 class FConsoleVariable : public FConsoleVariableBase
 {
@@ -308,10 +308,11 @@ public:
 	virtual void Release()
 	{
 		delete this; 
-	} 
+	}
+
 	virtual void Set(const TCHAR* InValue, EConsoleVariableFlags SetBy)
 	{
-		if(CanChange(SetBy))
+		if (CanChange(SetBy))
 		{
 			TTypeFromString<T>::FromString(Data.ShadowedValue[0], InValue);
 			OnChanged(SetBy);
@@ -496,7 +497,8 @@ public:
 	virtual void Release()
 	{
 		delete this; 
-	} 
+	}
+
 	virtual void Set(const TCHAR* InValue, EConsoleVariableFlags SetBy)
 	{
 		if(CanChange(SetBy))
diff --git a/Engine/Source/Runtime/Core/Public/HAL/IConsoleManager.h b/Engine/Source/Runtime/Core/Public/HAL/IConsoleManager.h
--- a/Engine/Source/Runtime/Core/Public/HAL/IConsoleManager.h
+++ b/Engine/Source/Runtime/Core/Public/HAL/IConsoleManager.h
@@ -309,8 +309,11 @@ public:
 	/** Set the internal value from the specified bool. */
 	void Set(bool InValue, EConsoleVariableFlags SetBy = ECVF_SetByCode)
 	{
+		// NOTE: Bool needs to use 1 and 0 here rather than true/false, as this may be a int32 or something
+		// and eventually this code calls, TTypeFromString<T>::FromString which won't handle the true/false,
+		// but 1 and 0 will work for whatever.
 		// inefficient but no common code path
-		Set(InValue ? TEXT("true") : TEXT("false"), SetBy);
+		Set(InValue ? TEXT("1") : TEXT("0"), SetBy);
 	}
 	/** Set the internal value from the specified int. */
 	void Set(int32 InValue, EConsoleVariableFlags SetBy = ECVF_SetByCode)
