Adding nullcheck to prevent crash.

diff --git a/Engine/Source/Runtime/Core/Private/Misc/VarargsHelper.h b/Engine/Source/Runtime/Core/Private/Misc/VarargsHelper.h
--- a/Engine/Source/Runtime/Core/Private/Misc/VarargsHelper.h
+++ b/Engine/Source/Runtime/Core/Private/Misc/VarargsHelper.h
@@ -28,6 +28,10 @@
 		FMemory::SystemFree(AllocatedBuffer); \
 		/* We need to use malloc here directly as GMalloc might not be safe. */ \
 		Buffer = AllocatedBuffer = (TCHAR*) FMemory::SystemMalloc( BufferSize * sizeof(TCHAR) ); \
+		if (Buffer == NULL) \
+		{ \
+			return; \
+		} \
 		GET_VARARGS_RESULT( Buffer, BufferSize, BufferSize-1, Fmt, Fmt, Result ); \
 		BufferSize *= 2; \
 	}; \
