Disable spawning CrashReportClient based on a ini setting

diff --git a/Engine/Source/Runtime/Core/Private/Unix/UnixPlatformCrashContext.cpp b/Engine/Source/Runtime/Core/Private/Unix/UnixPlatformCrashContext.cpp
--- a/Engine/Source/Runtime/Core/Private/Unix/UnixPlatformCrashContext.cpp
+++ b/Engine/Source/Runtime/Core/Private/Unix/UnixPlatformCrashContext.cpp
@@ -628,21 +628,32 @@ void FUnixCrashContext::GenerateCrashInfoAndLaunchReporter(bool bReportingNonCra
 			}
 			else
 			{
-				// spin here until CrashReporter exits
-				FProcHandle RunningProc = FPlatformProcess::CreateProc(RelativePathToCrashReporter, *CrashReportClientArguments, true, false, false, NULL, 0, NULL, NULL);
-
-				// do not wait indefinitely - can be more generous about the hitch than in ensure() case
-				// NOTE: Chris.Wood - increased from 3 to 8 mins because server crashes were timing out and getting lost
-				// NOTE: Do not increase above 8.5 mins without altering watchdog scripts to match
-				const double kCrashTimeOut = 8 * 60.0;
-
-				const double kCrashSleepInterval = 1.0;
-				if (!UnixCrashReporterTracker::WaitForProcWithTimeout(RunningProc, kCrashTimeOut, kCrashSleepInterval))
+				// Things can be setup to allow for a global crash handler to capture the core from a crash and allow another process
+				// to handle spawning of this process
+				bool bStartCRCFromEngineHandler = true;
+				if (GConfig)
 				{
-					FPlatformProcess::TerminateProc(RunningProc);
+					GConfig->GetBool(TEXT("CrashReportClient"), TEXT("bStartCRCFromEngineHandler"), bStartCRCFromEngineHandler, GEngineIni);
 				}
 
-				FPlatformProcess::CloseProc(RunningProc);
+				if (bStartCRCFromEngineHandler)
+				{
+					// spin here until CrashReporter exits
+					FProcHandle RunningProc = FPlatformProcess::CreateProc(RelativePathToCrashReporter, *CrashReportClientArguments, true, false, false, NULL, 0, NULL, NULL);
+
+					// do not wait indefinitely - can be more generous about the hitch than in ensure() case
+					// NOTE: Chris.Wood - increased from 3 to 8 mins because server crashes were timing out and getting lost
+					// NOTE: Do not increase above 8.5 mins without altering watchdog scripts to match
+					const double kCrashTimeOut = 8 * 60.0;
+
+					const double kCrashSleepInterval = 1.0;
+					if (!UnixCrashReporterTracker::WaitForProcWithTimeout(RunningProc, kCrashTimeOut, kCrashSleepInterval))
+					{
+						FPlatformProcess::TerminateProc(RunningProc);
+					}
+
+					FPlatformProcess::CloseProc(RunningProc);
+				}
 			}
 		}
 	}
