Converted LEET to a code-based culture, and added support for general custom code-based cultures

This avoids the LEET culture being tied to the ICU specific implementation and version, and allows games to inject custom cultures of their own without having to deal with ICU data.

diff --git a/Engine/Source/Runtime/Core/Private/Internationalization/Culture.cpp b/Engine/Source/Runtime/Core/Private/Internationalization/Culture.cpp
--- a/Engine/Source/Runtime/Core/Private/Internationalization/Culture.cpp
+++ b/Engine/Source/Runtime/Core/Private/Internationalization/Culture.cpp
@@ -5,54 +5,17 @@
 #if UE_ENABLE_ICU
 #include "Internationalization/ICUCulture.h"
 #else
-#include "LegacyCulture.h"
+#include "Internationalization/LegacyCulture.h"
 #endif
 
-#if UE_ENABLE_ICU
-FCulturePtr FCulture::Create(const FString& LocaleName)
-{
-	return MakeShareable(new FCulture(LocaleName));
-}
-#else
-FCulturePtr FCulture::Create(
-	const FText& InDisplayName, 
-	const FString& InEnglishName, 
-	const int InKeyboardLayoutId, 
-	const int InLCID, 
-	const FString& InName, 
-	const FString& InNativeName, 
-	const FString& InUnrealLegacyThreeLetterISOLanguageName, 
-	const FString& InThreeLetterISOLanguageName, 
-	const FString& InTwoLetterISOLanguageName,
-	const FDecimalNumberFormattingRules& InDecimalNumberFormattingRules,
-	const FDecimalNumberFormattingRules& InPercentFormattingRules,
-	const FDecimalNumberFormattingRules& InBaseCurrencyFormattingRules
-	)
+FCultureRef FCulture::Create(TUniquePtr<FCultureImplementation>&& InImplementation)
 {
-	return MakeShareable(new FCulture(InDisplayName, InEnglishName, InKeyboardLayoutId, InLCID, InName, InNativeName, InUnrealLegacyThreeLetterISOLanguageName, InThreeLetterISOLanguageName, InTwoLetterISOLanguageName, InDecimalNumberFormattingRules, InPercentFormattingRules, InBaseCurrencyFormattingRules));
+	check(InImplementation.IsValid());
+	return MakeShareable(new FCulture(MoveTemp(InImplementation)));
 }
-#endif
 
-#if UE_ENABLE_ICU
-FCulture::FCulture(const FString& LocaleName)
-	: Implementation( new FICUCultureImplementation( LocaleName ) )
-#else
-FCulture::FCulture(
-	const FText& InDisplayName, 
-	const FString& InEnglishName, 
-	const int InKeyboardLayoutId, 
-	const int InLCID, 
-	const FString& InName, 
-	const FString& InNativeName, 
-	const FString& InUnrealLegacyThreeLetterISOLanguageName, 
-	const FString& InThreeLetterISOLanguageName, 
-	const FString& InTwoLetterISOLanguageName,
-	const FDecimalNumberFormattingRules& InDecimalNumberFormattingRules,
-	const FDecimalNumberFormattingRules& InPercentFormattingRules,
-	const FDecimalNumberFormattingRules& InBaseCurrencyFormattingRules
-	) 
-	: Implementation( new FLegacyCultureImplementation(InDisplayName, InEnglishName, InKeyboardLayoutId, InLCID, InName, InNativeName, InUnrealLegacyThreeLetterISOLanguageName, InThreeLetterISOLanguageName, InTwoLetterISOLanguageName, InDecimalNumberFormattingRules, InPercentFormattingRules, InBaseCurrencyFormattingRules) )
-#endif
+FCulture::FCulture(TUniquePtr<FCultureImplementation>&& InImplementation)
+	: Implementation(MoveTemp(InImplementation))
 	, CachedDisplayName(Implementation->GetDisplayName())
 	, CachedEnglishName(Implementation->GetEnglishName())
 	, CachedName(Implementation->GetName())
@@ -61,17 +24,18 @@ FCulture::FCulture(
 	, CachedThreeLetterISOLanguageName(Implementation->GetThreeLetterISOLanguageName())
 	, CachedTwoLetterISOLanguageName(Implementation->GetTwoLetterISOLanguageName())
 	, CachedNativeLanguage(Implementation->GetNativeLanguage())
-#if UE_ENABLE_ICU
 	, CachedRegion(Implementation->GetRegion())
-#endif
 	, CachedNativeRegion(Implementation->GetNativeRegion())
-#if UE_ENABLE_ICU
 	, CachedScript(Implementation->GetScript())
 	, CachedVariant(Implementation->GetVariant())
-#endif
 { 
 }
 
+FCulture::~FCulture()
+{
+	// Explicit destructor due to forward declaration of FCultureImplementation
+}
+
 const FString& FCulture::GetDisplayName() const
 {
 	return CachedDisplayName;
@@ -160,7 +124,7 @@ FString FCulture::CreateCultureName(const FString& LanguageCode, const FString&
 
 FString FCulture::GetCanonicalName(const FString& Name)
 {
-	return FImplementation::GetCanonicalName(Name);
+	return FCultureImplementation::GetCanonicalName(Name);
 }
 
 const FString& FCulture::GetName() const
diff --git a/Engine/Source/Runtime/Core/Private/Internationalization/CultureImplementation.h b/Engine/Source/Runtime/Core/Private/Internationalization/CultureImplementation.h
new file mode 100644
--- /dev/null
+++ b/Engine/Source/Runtime/Core/Private/Internationalization/CultureImplementation.h
@@ -0,0 +1,42 @@
+// Copyright 1998-2019 Epic Games, Inc. All Rights Reserved.
+
+#pragma once
+
+#include "CoreTypes.h"
+#include "Containers/Array.h"
+#include "Containers/UnrealString.h"
+
+struct FDecimalNumberFormattingRules;
+enum class ETextPluralForm : uint8;
+enum class ETextPluralType : uint8;
+
+/**
+ * Internal implementation of a culture.
+ * @note See FCulture for the API description.
+ */
+class ICultureImplementation
+{
+public:
+	virtual ~ICultureImplementation() = default;
+
+	virtual FString GetDisplayName() const = 0;
+	virtual FString GetEnglishName() const = 0;
+	virtual int GetKeyboardLayoutId() const = 0;
+	virtual int GetLCID() const = 0;
+	virtual FString GetName() const = 0;
+	virtual FString GetNativeName() const = 0;
+	virtual FString GetUnrealLegacyThreeLetterISOLanguageName() const = 0;
+	virtual FString GetThreeLetterISOLanguageName() const = 0;
+	virtual FString GetTwoLetterISOLanguageName() const = 0;
+	virtual FString GetNativeLanguage() const = 0;
+	virtual FString GetNativeRegion() const = 0;
+	virtual FString GetRegion() const = 0;
+	virtual FString GetScript() const = 0;
+	virtual FString GetVariant() const = 0;
+	virtual const FDecimalNumberFormattingRules& GetDecimalNumberFormattingRules() = 0;
+	virtual const FDecimalNumberFormattingRules& GetPercentFormattingRules() = 0;
+	virtual const FDecimalNumberFormattingRules& GetCurrencyFormattingRules(const FString& InCurrencyCode) = 0;
+	virtual ETextPluralForm GetPluralForm(int32 Val, const ETextPluralType PluralType) const = 0;
+	virtual ETextPluralForm GetPluralForm(double Val, const ETextPluralType PluralType) const = 0;
+	virtual const TArray<ETextPluralForm>& GetValidPluralForms(const ETextPluralType PluralType) const = 0;
+};
diff --git a/Engine/Source/Runtime/Core/Private/Internationalization/Cultures/InvariantCulture.h b/Engine/Source/Runtime/Core/Private/Internationalization/Cultures/InvariantCulture.h
--- a/Engine/Source/Runtime/Core/Private/Internationalization/Cultures/InvariantCulture.h
+++ b/Engine/Source/Runtime/Core/Private/Internationalization/Cultures/InvariantCulture.h
@@ -4,6 +4,7 @@
 #if !UE_ENABLE_ICU
 
 #include "Internationalization/Culture.h"
+#include "Internationalization/LegacyCulture.h"
 #include "Internationalization/Internationalization.h"
 #include "Internationalization/FastDecimalFormat.h"
 
@@ -47,7 +48,7 @@ public:
 		BaseCurrencyFormattingRules.PrimaryGroupingSize = 3;
 		BaseCurrencyFormattingRules.SecondaryGroupingSize = 3;
 
-		FCultureRef Culture = FCulture::Create(
+		FCultureRef Culture = FCulture::Create(MakeUnique<FLegacyCultureImplementation>(
 			LOCTEXT("InvariantCultureDisplayName", "Invariant Language (Invariant Country)"),	//const FText DisplayName
 			FString(TEXT("Invariant Language (Invariant Country)")),							//const FString EnglishName
 			1033,																				//const int KeyboardLayoutId
@@ -60,7 +61,7 @@ public:
 			DecimalNumberFormattingRules,														//const FDecimalNumberFormattingRules InDecimalNumberFormattingRules
 			PercentFormattingRules,																//const FDecimalNumberFormattingRules InPercentFormattingRules
 			BaseCurrencyFormattingRules															//const FDecimalNumberFormattingRules InBaseCurrencyFormattingRules
-			).ToSharedRef();
+			));
 
 		return Culture;
 	}
diff --git a/Engine/Source/Runtime/Core/Private/Internationalization/Cultures/LeetCulture.cpp b/Engine/Source/Runtime/Core/Private/Internationalization/Cultures/LeetCulture.cpp
new file mode 100644
--- /dev/null
+++ b/Engine/Source/Runtime/Core/Private/Internationalization/Cultures/LeetCulture.cpp
@@ -0,0 +1,18 @@
+// Copyright 1998-2019 Epic Games, Inc. All Rights Reserved.
+
+#include "Internationalization/Cultures/LeetCulture.h"
+
+#if ENABLE_LOC_TESTING
+
+FLeetCulture::FLeetCulture(const FCultureRef& InInvariantCulture)
+	: InvariantCulture(InInvariantCulture)
+{
+}
+
+const FString& FLeetCulture::StaticGetName()
+{
+	static const FString LeetCultureName = TEXT("LEET");
+	return LeetCultureName;
+}
+
+#endif
diff --git a/Engine/Source/Runtime/Core/Private/Internationalization/Cultures/LeetCulture.h b/Engine/Source/Runtime/Core/Private/Internationalization/Cultures/LeetCulture.h
new file mode 100644
--- /dev/null
+++ b/Engine/Source/Runtime/Core/Private/Internationalization/Cultures/LeetCulture.h
@@ -0,0 +1,36 @@
+// Copyright 1998-2019 Epic Games, Inc. All Rights Reserved.
+
+#pragma once
+
+#include "Internationalization/LocTesting.h"
+#include "Internationalization/ICustomCulture.h"
+
+#if ENABLE_LOC_TESTING
+
+class FLeetCulture : public ICustomCulture
+{
+public:
+	explicit FLeetCulture(const FCultureRef& InInvariantCulture);
+
+	static const FString& StaticGetName();
+
+	//~ ICustomCulture interface
+	virtual FCultureRef GetBaseCulture() const override { return InvariantCulture; }
+	virtual FString GetDisplayName() const override { return FLeetCulture::StaticGetName(); }
+	virtual FString GetEnglishName() const override { return FLeetCulture::StaticGetName(); }
+	virtual FString GetName() const override { return FLeetCulture::StaticGetName(); }
+	virtual FString GetNativeName() const override { return FLeetCulture::StaticGetName(); }
+	virtual FString GetUnrealLegacyThreeLetterISOLanguageName() const override { return TEXT("INT"); }
+	virtual FString GetThreeLetterISOLanguageName() const override { return FLeetCulture::StaticGetName(); }
+	virtual FString GetTwoLetterISOLanguageName() const override { return FLeetCulture::StaticGetName(); }
+	virtual FString GetNativeLanguage() const override { return FLeetCulture::StaticGetName(); }
+	virtual FString GetNativeRegion() const override { return FString(); }
+	virtual FString GetRegion() const override { return FString(); }
+	virtual FString GetScript() const override { return FString(); }
+	virtual FString GetVariant() const override { return FString(); }
+
+private:
+	FCultureRef InvariantCulture;
+};
+
+#endif
diff --git a/Engine/Source/Runtime/Core/Private/Internationalization/CustomCultureImplementation.h b/Engine/Source/Runtime/Core/Private/Internationalization/CustomCultureImplementation.h
new file mode 100644
--- /dev/null
+++ b/Engine/Source/Runtime/Core/Private/Internationalization/CustomCultureImplementation.h
@@ -0,0 +1,59 @@
+// Copyright 1998-2019 Epic Games, Inc. All Rights Reserved.
+
+#pragma once
+
+#include "CoreTypes.h"
+#include "Internationalization/Culture.h"
+#include "Internationalization/ICustomCulture.h"
+
+#if UE_ENABLE_ICU
+#include "Internationalization/ICUCulture.h"
+#else
+#include "Internationalization/LegacyCulture.h"
+#endif
+
+class FCustomCultureImplementation : public FCultureImplementation
+{
+public:
+	FCustomCultureImplementation(const TSharedRef<ICustomCulture>& InCustomCulture)
+#if UE_ENABLE_ICU
+		: FCultureImplementation(InCustomCulture->GetBaseCulture()->GetName())
+#else
+		: FCultureImplementation(
+			FText::AsCultureInvariant(InCustomCulture->GetBaseCulture()->GetDisplayName()),
+			InCustomCulture->GetBaseCulture()->GetEnglishName(),
+			InCustomCulture->GetBaseCulture()->GetKeyboardLayoutId(),
+			InCustomCulture->GetBaseCulture()->GetLCID(),
+			InCustomCulture->GetBaseCulture()->GetName(),
+			InCustomCulture->GetBaseCulture()->GetNativeName(),
+			InCustomCulture->GetBaseCulture()->GetUnrealLegacyThreeLetterISOLanguageName(),
+			InCustomCulture->GetBaseCulture()->GetThreeLetterISOLanguageName(),
+			InCustomCulture->GetBaseCulture()->GetTwoLetterISOLanguageName(),
+			InCustomCulture->GetBaseCulture()->GetDecimalNumberFormattingRules(),
+			InCustomCulture->GetBaseCulture()->GetPercentFormattingRules(),
+			InCustomCulture->GetBaseCulture()->GetCurrencyFormattingRules(FString())
+		)
+#endif
+		, CustomCulture(InCustomCulture)
+	{
+	}
+
+	virtual ~FCustomCultureImplementation() = default;
+
+	//~ ICultureImplementation interface
+	virtual FString GetDisplayName() const override { return CustomCulture->GetDisplayName(); }
+	virtual FString GetEnglishName() const override { return CustomCulture->GetEnglishName(); }
+	virtual FString GetName() const override { return CustomCulture->GetName(); }
+	virtual FString GetNativeName() const override { return CustomCulture->GetNativeName(); }
+	virtual FString GetUnrealLegacyThreeLetterISOLanguageName() const override { return CustomCulture->GetUnrealLegacyThreeLetterISOLanguageName(); }
+	virtual FString GetThreeLetterISOLanguageName() const override { return CustomCulture->GetThreeLetterISOLanguageName(); }
+	virtual FString GetTwoLetterISOLanguageName() const override { return CustomCulture->GetTwoLetterISOLanguageName(); }
+	virtual FString GetNativeLanguage() const override { return CustomCulture->GetNativeLanguage(); }
+	virtual FString GetNativeRegion() const override { return CustomCulture->GetNativeRegion(); }
+	virtual FString GetRegion() const override { return CustomCulture->GetRegion(); }
+	virtual FString GetScript() const override { return CustomCulture->GetScript(); }
+	virtual FString GetVariant() const override { return CustomCulture->GetVariant(); }
+
+private:
+	TSharedRef<ICustomCulture> CustomCulture;
+};
diff --git a/Engine/Source/Runtime/Core/Private/Internationalization/FastDecimalFormat.cpp b/Engine/Source/Runtime/Core/Private/Internationalization/FastDecimalFormat.cpp
--- a/Engine/Source/Runtime/Core/Private/Internationalization/FastDecimalFormat.cpp
+++ b/Engine/Source/Runtime/Core/Private/Internationalization/FastDecimalFormat.cpp
@@ -253,7 +253,7 @@ FORCEINLINE int32 IntegralToString_Common(const uint64 InVal, const FDecimalNumb
 	// Perform the initial format to a decimal string
 	return IntegralToString_UInt64ToString(
 		InVal, 
-		InFormattingOptions.UseGrouping, 
+		InFormattingOptions.UseGrouping && InFormattingRules.PrimaryGroupingSize > 0,
 		InFormattingRules.PrimaryGroupingSize, 
 		InFormattingRules.SecondaryGroupingSize, 
 		InFormattingRules.GroupingSeparatorCharacter, 
diff --git a/Engine/Source/Runtime/Core/Private/Internationalization/ICUCulture.cpp b/Engine/Source/Runtime/Core/Private/Internationalization/ICUCulture.cpp
--- a/Engine/Source/Runtime/Core/Private/Internationalization/ICUCulture.cpp
+++ b/Engine/Source/Runtime/Core/Private/Internationalization/ICUCulture.cpp
@@ -1,6 +1,8 @@
 // Copyright 1998-2019 Epic Games, Inc. All Rights Reserved.
 
 #include "Internationalization/ICUCulture.h"
+#include "Internationalization/Cultures/LeetCulture.h"
+#include "Internationalization/Internationalization.h"
 #include "Misc/ScopeLock.h"
 #include "Containers/SortedMap.h"
 #include "Internationalization/FastDecimalFormat.h"
@@ -135,7 +137,7 @@ TArray<ETextPluralForm> ICUPluralRulesToUEValidPluralForms(const icu::PluralRule
 	return UEPluralForms;
 }
 
-FCulture::FICUCultureImplementation::FICUCultureImplementation(const FString& LocaleName)
+FICUCultureImplementation::FICUCultureImplementation(const FString& LocaleName)
 	: ICULocale( TCHAR_TO_ANSI( *LocaleName ) )
 {
 	if (ICULocale.isBogus())
@@ -156,37 +158,41 @@ FCulture::FICUCultureImplementation::FICUCultureImplementation(const FString& Lo
 	}
 }
 
-FString FCulture::FICUCultureImplementation::GetDisplayName() const
+FString FICUCultureImplementation::GetDisplayName() const
 {
 	icu::UnicodeString ICUResult;
 	ICULocale.getDisplayName(ICUResult);
 	return ICUUtilities::ConvertString(ICUResult);
 }
 
-FString FCulture::FICUCultureImplementation::GetEnglishName() const
+FString FICUCultureImplementation::GetEnglishName() const
 {
 	icu::UnicodeString ICUResult;
 	ICULocale.getDisplayName(icu::Locale("en"), ICUResult);
 	return ICUUtilities::ConvertString(ICUResult);
 }
 
-int FCulture::FICUCultureImplementation::GetKeyboardLayoutId() const
+int FICUCultureImplementation::GetKeyboardLayoutId() const
 {
 	return 0;
 }
 
-int FCulture::FICUCultureImplementation::GetLCID() const
+int FICUCultureImplementation::GetLCID() const
 {
 	return ICULocale.getLCID();
 }
 
-FString FCulture::FICUCultureImplementation::GetCanonicalName(const FString& Name)
+FString FICUCultureImplementation::GetCanonicalName(const FString& Name)
 {
 	auto IsLanguageCode = [](const FString& InCode)
 	{
-		// Language codes must be 2 or 3 letters, or out special "LEET" language
-		static const FString LeetLanguageCode = TEXT("LEET");
-		return InCode.Len() == 2 || InCode.Len() == 3 || InCode == LeetLanguageCode;
+		// Language codes must be 2 or 3 letters, or our special "LEET" language
+		return InCode.Len() == 2 
+			|| InCode.Len() == 3 
+#if ENABLE_LOC_TESTING
+			|| InCode == FLeetCulture::StaticGetName()
+#endif
+			;
 	};
 
 	auto IsScriptCode = [](const FString& InCode)
@@ -313,6 +319,12 @@ FString FCulture::FICUCultureImplementation::GetCanonicalName(const FString& Nam
 	// Sanitize any nastiness from the culture code
 	const FString SanitizedName = ICUUtilities::SanitizeCultureCode(Name);
 
+	// If the name matches a custom culture, then just accept it as-is
+	if (FInternationalization::Get().GetCustomCulture(SanitizedName))
+	{
+		return SanitizedName;
+	}
+
 	// These will be populated as the string is processed and are used to re-build the canonized string
 	TArray<FNameTag, TInlineAllocator<4>> ParsedNameTags;
 	TSortedMap<FString, FString, TInlineAllocator<4>> ParsedKeywords;
@@ -511,21 +523,21 @@ FString FCulture::FICUCultureImplementation::GetCanonicalName(const FString& Nam
 	return CanonicalName;
 }
 
-FString FCulture::FICUCultureImplementation::GetName() const
+FString FICUCultureImplementation::GetName() const
 {
 	FString Result = ICULocale.getName();
 	Result.ReplaceInline(TEXT("_"), TEXT("-"), ESearchCase::IgnoreCase);
 	return Result;
 }
 
-FString FCulture::FICUCultureImplementation::GetNativeName() const
+FString FICUCultureImplementation::GetNativeName() const
 {
 	icu::UnicodeString ICUResult;
 	ICULocale.getDisplayName(ICULocale, ICUResult);
 	return ICUUtilities::ConvertString(ICUResult);
 }
 
-FString FCulture::FICUCultureImplementation::GetUnrealLegacyThreeLetterISOLanguageName() const
+FString FICUCultureImplementation::GetUnrealLegacyThreeLetterISOLanguageName() const
 {
 	FString Result( ICULocale.getISO3Language() );
 
@@ -543,17 +555,17 @@ FString FCulture::FICUCultureImplementation::GetUnrealLegacyThreeLetterISOLangua
 	return Result;
 }
 
-FString FCulture::FICUCultureImplementation::GetThreeLetterISOLanguageName() const
+FString FICUCultureImplementation::GetThreeLetterISOLanguageName() const
 {
 	return ICULocale.getISO3Language();
 }
 
-FString FCulture::FICUCultureImplementation::GetTwoLetterISOLanguageName() const
+FString FICUCultureImplementation::GetTwoLetterISOLanguageName() const
 {
 	return ICULocale.getLanguage();
 }
 
-FString FCulture::FICUCultureImplementation::GetNativeLanguage() const
+FString FICUCultureImplementation::GetNativeLanguage() const
 {
 	icu::UnicodeString ICUNativeLanguage;
 	ICULocale.getDisplayLanguage(ICULocale, ICUNativeLanguage);
@@ -572,12 +584,12 @@ FString FCulture::FICUCultureImplementation::GetNativeLanguage() const
 	return NativeLanguage;
 }
 
-FString FCulture::FICUCultureImplementation::GetRegion() const
+FString FICUCultureImplementation::GetRegion() const
 {
 	return ICULocale.getCountry();
 }
 
-FString FCulture::FICUCultureImplementation::GetNativeRegion() const
+FString FICUCultureImplementation::GetNativeRegion() const
 {
 	icu::UnicodeString ICUNativeCountry;
 	ICULocale.getDisplayCountry(ICULocale, ICUNativeCountry);
@@ -596,17 +608,17 @@ FString FCulture::FICUCultureImplementation::GetNativeRegion() const
 	return NativeCountry;
 }
 
-FString FCulture::FICUCultureImplementation::GetScript() const
+FString FICUCultureImplementation::GetScript() const
 {
 	return ICULocale.getScript();
 }
 
-FString FCulture::FICUCultureImplementation::GetVariant() const
+FString FICUCultureImplementation::GetVariant() const
 {
 	return ICULocale.getVariant();
 }
 
-TSharedRef<const icu::BreakIterator> FCulture::FICUCultureImplementation::GetBreakIterator(const EBreakIteratorType Type)
+TSharedRef<const icu::BreakIterator> FICUCultureImplementation::GetBreakIterator(const EBreakIteratorType Type)
 {
 	TSharedPtr<const icu::BreakIterator> Result;
 
@@ -642,7 +654,7 @@ TSharedRef<const icu::BreakIterator> FCulture::FICUCultureImplementation::GetBre
 	return Result.ToSharedRef();
 }
 
-TSharedRef<const icu::Collator, ESPMode::ThreadSafe> FCulture::FICUCultureImplementation::GetCollator(const ETextComparisonLevel::Type ComparisonLevel)
+TSharedRef<const icu::Collator, ESPMode::ThreadSafe> FICUCultureImplementation::GetCollator(const ETextComparisonLevel::Type ComparisonLevel)
 {
 	if (!ICUCollator.IsValid())
 	{
@@ -664,7 +676,7 @@ TSharedRef<const icu::Collator, ESPMode::ThreadSafe> FCulture::FICUCultureImplem
 	}
 }
 
-TSharedRef<const icu::DateFormat, ESPMode::ThreadSafe> FCulture::FICUCultureImplementation::GetDateFormatter(const EDateTimeStyle::Type DateStyle, const FString& TimeZone)
+TSharedRef<const icu::DateFormat, ESPMode::ThreadSafe> FICUCultureImplementation::GetDateFormatter(const EDateTimeStyle::Type DateStyle, const FString& TimeZone)
 {
 	if (!ICUDateFormat.IsValid())
 	{
@@ -711,7 +723,7 @@ TSharedRef<const icu::DateFormat, ESPMode::ThreadSafe> FCulture::FICUCultureImpl
 	}
 }
 
-TSharedRef<const icu::DateFormat, ESPMode::ThreadSafe> FCulture::FICUCultureImplementation::GetTimeFormatter(const EDateTimeStyle::Type TimeStyle, const FString& TimeZone)
+TSharedRef<const icu::DateFormat, ESPMode::ThreadSafe> FICUCultureImplementation::GetTimeFormatter(const EDateTimeStyle::Type TimeStyle, const FString& TimeZone)
 {
 	if (!ICUTimeFormat.IsValid())
 	{
@@ -758,7 +770,7 @@ TSharedRef<const icu::DateFormat, ESPMode::ThreadSafe> FCulture::FICUCultureImpl
 	}
 }
 
-TSharedRef<const icu::DateFormat, ESPMode::ThreadSafe> FCulture::FICUCultureImplementation::GetDateTimeFormatter(const EDateTimeStyle::Type DateStyle, const EDateTimeStyle::Type TimeStyle, const FString& TimeZone)
+TSharedRef<const icu::DateFormat, ESPMode::ThreadSafe> FICUCultureImplementation::GetDateTimeFormatter(const EDateTimeStyle::Type DateStyle, const EDateTimeStyle::Type TimeStyle, const FString& TimeZone)
 {
 	if (!ICUDateTimeFormat.IsValid())
 	{
@@ -871,7 +883,7 @@ FDecimalNumberFormattingRules ExtractNumberFormattingRulesFromICUDecimalFormatte
 
 } // anonymous namespace
 
-const FDecimalNumberFormattingRules& FCulture::FICUCultureImplementation::GetDecimalNumberFormattingRules()
+const FDecimalNumberFormattingRules& FICUCultureImplementation::GetDecimalNumberFormattingRules()
 {
 	if (UEDecimalNumberFormattingRules.IsValid())
 	{
@@ -901,7 +913,7 @@ const FDecimalNumberFormattingRules& FCulture::FICUCultureImplementation::GetDec
 	return *UEDecimalNumberFormattingRules;
 }
 
-const FDecimalNumberFormattingRules& FCulture::FICUCultureImplementation::GetPercentFormattingRules()
+const FDecimalNumberFormattingRules& FICUCultureImplementation::GetPercentFormattingRules()
 {
 	if (UEPercentFormattingRules.IsValid())
 	{
@@ -931,7 +943,7 @@ const FDecimalNumberFormattingRules& FCulture::FICUCultureImplementation::GetPer
 	return *UEPercentFormattingRules;
 }
 
-const FDecimalNumberFormattingRules& FCulture::FICUCultureImplementation::GetCurrencyFormattingRules(const FString& InCurrencyCode)
+const FDecimalNumberFormattingRules& FICUCultureImplementation::GetCurrencyFormattingRules(const FString& InCurrencyCode)
 {
 	const FString SanitizedCurrencyCode = ICUUtilities::SanitizeCurrencyCode(InCurrencyCode);
 	const bool bUseDefaultFormattingRules = SanitizedCurrencyCode.IsEmpty();
@@ -1002,7 +1014,7 @@ const FDecimalNumberFormattingRules& FCulture::FICUCultureImplementation::GetCur
 	}
 }
 
-ETextPluralForm FCulture::FICUCultureImplementation::GetPluralForm(int32 Val, const ETextPluralType PluralType) const
+ETextPluralForm FICUCultureImplementation::GetPluralForm(int32 Val, const ETextPluralType PluralType) const
 {
 	checkf(Val >= 0, TEXT("GetPluralFormImpl requires a positive value"));
 
@@ -1012,7 +1024,7 @@ ETextPluralForm FCulture::FICUCultureImplementation::GetPluralForm(int32 Val, co
 	return ICUPluralFormToUE(ICUPluralFormTag);
 }
 
-ETextPluralForm FCulture::FICUCultureImplementation::GetPluralForm(double Val, const ETextPluralType PluralType) const
+ETextPluralForm FICUCultureImplementation::GetPluralForm(double Val, const ETextPluralType PluralType) const
 {
 	checkf(!FMath::IsNegativeDouble(Val), TEXT("GetPluralFormImpl requires a positive value"));
 
@@ -1022,7 +1034,7 @@ ETextPluralForm FCulture::FICUCultureImplementation::GetPluralForm(double Val, c
 	return ICUPluralFormToUE(ICUPluralFormTag);
 }
 
-const TArray<ETextPluralForm>& FCulture::FICUCultureImplementation::GetValidPluralForms(const ETextPluralType PluralType) const
+const TArray<ETextPluralForm>& FICUCultureImplementation::GetValidPluralForms(const ETextPluralType PluralType) const
 {
 	return (PluralType == ETextPluralType::Cardinal) ? UEAvailableCardinalPluralForms : UEAvailableOrdinalPluralForms;
 }
diff --git a/Engine/Source/Runtime/Core/Private/Internationalization/ICUCulture.h b/Engine/Source/Runtime/Core/Private/Internationalization/ICUCulture.h
--- a/Engine/Source/Runtime/Core/Private/Internationalization/ICUCulture.h
+++ b/Engine/Source/Runtime/Core/Private/Internationalization/ICUCulture.h
@@ -9,6 +9,7 @@
 #include "Templates/SharedPointer.h"
 #include "Internationalization/Text.h"
 #include "Internationalization/Culture.h"
+#include "Internationalization/CultureImplementation.h"
 
 struct FDecimalNumberFormattingRules;
 
@@ -157,44 +158,40 @@ enum class EBreakIteratorType
 	Title
 };
 
-class FCulture::FICUCultureImplementation
+class FICUCultureImplementation : public ICultureImplementation
 {
 	friend class FCulture;
 	friend class FText;
 	friend class FTextChronoFormatter;
 	friend class FICUBreakIteratorManager;
 
-	FICUCultureImplementation(const FString& LocaleName);
-
-	FString GetDisplayName() const;
-
-	FString GetEnglishName() const;
-
-	int GetKeyboardLayoutId() const;
-
-	int GetLCID() const;
+public:
+	explicit FICUCultureImplementation(const FString& LocaleName);
+	virtual ~FICUCultureImplementation() = default;
+
+	//~ ICultureImplementation interface
+	virtual FString GetDisplayName() const override;
+	virtual FString GetEnglishName() const override;
+	virtual int GetKeyboardLayoutId() const override;
+	virtual int GetLCID() const override;
+	virtual FString GetName() const override;
+	virtual FString GetNativeName() const override;
+	virtual FString GetUnrealLegacyThreeLetterISOLanguageName() const override;
+	virtual FString GetThreeLetterISOLanguageName() const override;
+	virtual FString GetTwoLetterISOLanguageName() const override;
+	virtual FString GetNativeLanguage() const override;
+	virtual FString GetNativeRegion() const override;
+	virtual FString GetRegion() const override;
+	virtual FString GetScript() const override;
+	virtual FString GetVariant() const override;
+	virtual const FDecimalNumberFormattingRules& GetDecimalNumberFormattingRules() override;
+	virtual const FDecimalNumberFormattingRules& GetPercentFormattingRules() override;
+	virtual const FDecimalNumberFormattingRules& GetCurrencyFormattingRules(const FString& InCurrencyCode) override;
+	virtual ETextPluralForm GetPluralForm(int32 Val, const ETextPluralType PluralType) const override;
+	virtual ETextPluralForm GetPluralForm(double Val, const ETextPluralType PluralType) const override;
+	virtual const TArray<ETextPluralForm>& GetValidPluralForms(const ETextPluralType PluralType) const override;
 
 	static FString GetCanonicalName(const FString& Name);
-	
-	FString GetName() const;
-
-	FString GetNativeName() const;
-
-	FString GetUnrealLegacyThreeLetterISOLanguageName() const;
-
-	FString GetThreeLetterISOLanguageName() const;
-
-	FString GetTwoLetterISOLanguageName() const;
-
-	FString GetNativeLanguage() const;
-
-	FString GetRegion() const;
-
-	FString GetNativeRegion() const;
-
-	FString GetScript() const;
-
-	FString GetVariant() const;
 
 	TSharedRef<const icu::BreakIterator> GetBreakIterator(const EBreakIteratorType Type);
 	TSharedRef<const icu::Collator, ESPMode::ThreadSafe> GetCollator(const ETextComparisonLevel::Type ComparisonLevel);
@@ -202,15 +199,7 @@ class FCulture::FICUCultureImplementation
 	TSharedRef<const icu::DateFormat, ESPMode::ThreadSafe> GetTimeFormatter(const EDateTimeStyle::Type TimeStyle, const FString& TimeZone);
 	TSharedRef<const icu::DateFormat, ESPMode::ThreadSafe> GetDateTimeFormatter(const EDateTimeStyle::Type DateStyle, const EDateTimeStyle::Type TimeStyle, const FString& TimeZone);
 
-	const FDecimalNumberFormattingRules& GetDecimalNumberFormattingRules();
-	const FDecimalNumberFormattingRules& GetPercentFormattingRules();
-	const FDecimalNumberFormattingRules& GetCurrencyFormattingRules(const FString& InCurrencyCode);
-
-	ETextPluralForm GetPluralForm(int32 Val, const ETextPluralType PluralType) const;
-	ETextPluralForm GetPluralForm(double Val, const ETextPluralType PluralType) const;
-
-	const TArray<ETextPluralForm>& GetValidPluralForms(const ETextPluralType PluralType) const;
-
+private:
 	icu::Locale ICULocale;
 	TSharedPtr<const icu::BreakIterator> ICUGraphemeBreakIterator;
 	TSharedPtr<const icu::BreakIterator> ICUWordBreakIterator;
diff --git a/Engine/Source/Runtime/Core/Private/Internationalization/ICUInternationalization.cpp b/Engine/Source/Runtime/Core/Private/Internationalization/ICUInternationalization.cpp
--- a/Engine/Source/Runtime/Core/Private/Internationalization/ICUInternationalization.cpp
+++ b/Engine/Source/Runtime/Core/Private/Internationalization/ICUInternationalization.cpp
@@ -6,6 +6,7 @@
 #include "Misc/Paths.h"
 #include "Internationalization/Culture.h"
 #include "Internationalization/Internationalization.h"
+#include "Internationalization/Cultures/LeetCulture.h"
 #include "Stats/Stats.h"
 #include "Misc/CoreStats.h"
 #include "Misc/ConfigCacheIni.h"
@@ -14,6 +15,7 @@
 #if UE_ENABLE_ICU
 
 #include "Internationalization/ICURegex.h"
+#include "Internationalization/ICUCulture.h"
 #include "Internationalization/ICUUtilities.h"
 #include "Internationalization/ICUBreakIterator.h"
 THIRD_PARTY_INCLUDES_START
@@ -26,6 +28,7 @@ THIRD_PARTY_INCLUDES_END
 DEFINE_LOG_CATEGORY_STATIC(LogICUInternationalization, Log, All);
 
 static_assert(sizeof(UChar) == 2, "UChar (from ICU) is assumed to always be 2-bytes!");
+static_assert(PLATFORM_LITTLE_ENDIAN, "ICU data is only built for little endian platforms. You'll need to rebuild the data for your platform and update this code!");
 
 namespace
 {
@@ -112,7 +115,11 @@ bool FICUInternationalization::Initialize()
 		if (FPaths::DirectoryExists(PotentialDataDirectory))
 		{
 			u_setDataDirectory(StringCast<char>(*PotentialDataDirectory).Get());
+#if WITH_ICU_V64
+			ICUDataDirectory = PotentialDataDirectory / TEXT("icudt64l") / TEXT(""); // We include the sub-folder here as it prevents ICU I/O requests outside of this folder
+#else	// WITH_ICU_V64
 			ICUDataDirectory = PotentialDataDirectory / TEXT("icudt53l") / TEXT(""); // We include the sub-folder here as it prevents ICU I/O requests outside of this folder
+#endif	// WITH_ICU_V64
 			break;
 		}
 	}
@@ -161,6 +168,10 @@ bool FICUInternationalization::Initialize()
 	I18N->CurrentLanguage = I18N->DefaultLanguage;
 	I18N->CurrentLocale = I18N->DefaultLocale;
 
+#if ENABLE_LOC_TESTING
+	I18N->AddCustomCulture(MakeShared<FLeetCulture>(I18N->InvariantCulture.ToSharedRef()));
+#endif
+
 	InitializeTimeZone();
 	InitializeInvariantGregorianCalendar();
 
@@ -542,11 +553,15 @@ void FICUInternationalization::HandleLanguageChanged(const FString& Name)
 
 void FICUInternationalization::GetCultureNames(TArray<FString>& CultureNames) const
 {
-	CultureNames.Reset(AllAvailableCultures.Num());
+	CultureNames.Reset(AllAvailableCultures.Num() + I18N->CustomCultures.Num());
 	for (const FICUCultureData& CultureData : AllAvailableCultures)
 	{
 		CultureNames.Add(CultureData.Name);
 	}
+	for (const FCultureRef& CustomCulture : I18N->CustomCultures)
+	{
+		CultureNames.Add(CustomCulture->GetName());
+	}
 }
 
 TArray<FString> FICUInternationalization::GetPrioritizedCultureNames(const FString& Name)
@@ -679,10 +694,15 @@ FCulturePtr FICUInternationalization::FindOrMakeCanonizedCulture(const FString&
 	// If no cached culture is found, try to make one.
 	FCulturePtr NewCulture;
 
+	// Is this a custom culture?
+	if (FCulturePtr CustomCulture = I18N->GetCustomCulture(Name))
+	{
+		NewCulture = CustomCulture;
+	}
 	// Is this in our list of available cultures?
-	if (AllAvailableCulturesMap.Contains(Name))
+	else if (AllAvailableCulturesMap.Contains(Name))
 	{
-		NewCulture = FCulture::Create(Name);
+		NewCulture = FCulture::Create(MakeUnique<FICUCultureImplementation>(Name));
 	}
 	else
 	{
@@ -692,7 +712,7 @@ FCulturePtr FICUInternationalization::FindOrMakeCanonizedCulture(const FString&
 		{
 			if (ICUStatus != U_USING_DEFAULT_WARNING || AllowDefaultFallback == EAllowDefaultCultureFallback::Yes)
 			{
-				NewCulture = FCulture::Create(Name);
+				NewCulture = FCulture::Create(MakeUnique<FICUCultureImplementation>(Name));
 			}
 			ures_close(ICUResourceBundle);
 		}
diff --git a/Engine/Source/Runtime/Core/Private/Internationalization/Internationalization.cpp b/Engine/Source/Runtime/Core/Private/Internationalization/Internationalization.cpp
--- a/Engine/Source/Runtime/Core/Private/Internationalization/Internationalization.cpp
+++ b/Engine/Source/Runtime/Core/Private/Internationalization/Internationalization.cpp
@@ -7,6 +7,7 @@
 #include "Misc/LazySingleton.h"
 #include "Misc/Paths.h"
 #include "Internationalization/Culture.h"
+#include "Internationalization/CustomCultureImplementation.h"
 #include "Internationalization/TextCache.h"
 
 #if UE_ENABLE_ICU
@@ -43,7 +44,7 @@ FText FInternationalization::ForUseOnlyByLocMacroAndGraphNodeTextLiterals_Create
 
 bool FInternationalization::SetCurrentCulture(const FString& InCultureName)
 {
-	FCulturePtr NewCulture = Implementation->GetCulture(InCultureName);
+	FCulturePtr NewCulture = GetCulture(InCultureName);
 
 	if (NewCulture.IsValid())
 	{
@@ -64,7 +65,7 @@ bool FInternationalization::SetCurrentCulture(const FString& InCultureName)
 
 bool FInternationalization::SetCurrentLanguage(const FString& InCultureName)
 {
-	FCulturePtr NewCulture = Implementation->GetCulture(InCultureName);
+	FCulturePtr NewCulture = GetCulture(InCultureName);
 
 	if (NewCulture.IsValid())
 	{
@@ -83,7 +84,7 @@ bool FInternationalization::SetCurrentLanguage(const FString& InCultureName)
 
 bool FInternationalization::SetCurrentLocale(const FString& InCultureName)
 {
-	FCulturePtr NewCulture = Implementation->GetCulture(InCultureName);
+	FCulturePtr NewCulture = GetCulture(InCultureName);
 
 	if (NewCulture.IsValid())
 	{
@@ -100,7 +101,7 @@ bool FInternationalization::SetCurrentLocale(const FString& InCultureName)
 
 bool FInternationalization::SetCurrentLanguageAndLocale(const FString& InCultureName)
 {
-	FCulturePtr NewCulture = Implementation->GetCulture(InCultureName);
+	FCulturePtr NewCulture = GetCulture(InCultureName);
 
 	if (NewCulture.IsValid())
 	{
@@ -120,7 +121,7 @@ bool FInternationalization::SetCurrentLanguageAndLocale(const FString& InCulture
 
 bool FInternationalization::SetCurrentAssetGroupCulture(const FName& InAssetGroupName, const FString& InCultureName)
 {
-	FCulturePtr NewCulture = Implementation->GetCulture(InCultureName);
+	FCulturePtr NewCulture = GetCulture(InCultureName);
 
 	if (NewCulture.IsValid())
 	{
@@ -194,7 +195,7 @@ void FInternationalization::RestoreCultureState(const FCultureStateSnapshot& InS
 	// Apply the language
 	if (!InSnapshot.Language.IsEmpty())
 	{
-		FCulturePtr NewCulture = Implementation->GetCulture(InSnapshot.Language);
+		FCulturePtr NewCulture = GetCulture(InSnapshot.Language);
 
 		if (NewCulture.IsValid())
 		{
@@ -212,7 +213,7 @@ void FInternationalization::RestoreCultureState(const FCultureStateSnapshot& InS
 	// Apply the locale
 	if (!InSnapshot.Locale.IsEmpty())
 	{
-		FCulturePtr NewCulture = Implementation->GetCulture(InSnapshot.Locale);
+		FCulturePtr NewCulture = GetCulture(InSnapshot.Locale);
 
 		if (NewCulture.IsValid())
 		{
@@ -230,7 +231,7 @@ void FInternationalization::RestoreCultureState(const FCultureStateSnapshot& InS
 	CurrentAssetGroupCultures.Reset(InSnapshot.AssetGroups.Num());
 	for (const auto& AssetGroupCultureNamePair : InSnapshot.AssetGroups)
 	{
-		FCulturePtr NewCulture = Implementation->GetCulture(AssetGroupCultureNamePair.Value);
+		FCulturePtr NewCulture = GetCulture(AssetGroupCultureNamePair.Value);
 		if (NewCulture.IsValid())
 		{
 			bChangedCulture = true;
@@ -275,6 +276,7 @@ void FInternationalization::Terminate()
 	DefaultLanguage.Reset();
 	DefaultLocale.Reset();
 
+	CustomCultures.Reset();
 	InvariantCulture.Reset();
 
 	Implementation->Terminate();
@@ -393,6 +395,20 @@ void FInternationalization::LoadAllCultureData()
 	Implementation->LoadAllCultureData();
 }
 
+void FInternationalization::AddCustomCulture(const TSharedRef<ICustomCulture>& InCustomCulture)
+{
+	CustomCultures.Add(FCulture::Create(MakeUnique<FCustomCultureImplementation>(InCustomCulture)));
+}
+
+FCulturePtr FInternationalization::GetCustomCulture(const FString& InCultureName) const
+{
+	const FCultureRef* CulturePtr = CustomCultures.FindByPredicate([&InCultureName](const FCultureRef& InPotentialCulture)
+	{
+		return InPotentialCulture->GetName() == InCultureName;
+	});
+	return CulturePtr ? FCulturePtr(*CulturePtr) : FCulturePtr();
+}
+
 bool FInternationalization::IsCultureRemapped(const FString& Name, FString* OutMappedCulture)
 {
 	return Implementation->IsCultureRemapped(Name, OutMappedCulture);
@@ -436,7 +452,7 @@ TArray<FCultureRef> FInternationalization::GetAvailableCultures(const TArray<FSt
 				const TArray<FString> PrioritizedParentCultureNames = Culture->GetPrioritizedParentCultureNames();
 				for (const FString& PrioritizedParentCultureName : PrioritizedParentCultureNames)
 				{
-					if (InCultureNames.Contains(PrioritizedParentCultureName) && Implementation->IsCultureAllowed(Culture->GetName()))
+					if (InCultureNames.Contains(PrioritizedParentCultureName) && IsCultureAllowed(Culture->GetName()))
 					{
 						AvailableCultures.AddUnique(Culture.ToSharedRef());
 						break;
diff --git a/Engine/Source/Runtime/Core/Private/Internationalization/LegacyCulture.cpp b/Engine/Source/Runtime/Core/Private/Internationalization/LegacyCulture.cpp
--- a/Engine/Source/Runtime/Core/Private/Internationalization/LegacyCulture.cpp
+++ b/Engine/Source/Runtime/Core/Private/Internationalization/LegacyCulture.cpp
@@ -6,7 +6,7 @@
 
 #if !UE_ENABLE_ICU
 
-FCulture::FLegacyCultureImplementation::FLegacyCultureImplementation(
+FLegacyCultureImplementation::FLegacyCultureImplementation(
 	const FText& InDisplayName, 
 	const FString& InEnglishName, 
 	const int InKeyboardLayoutId, 
@@ -35,42 +35,42 @@ FCulture::FLegacyCultureImplementation::FLegacyCultureImplementation(
 { 
 }
 
-FString FCulture::FLegacyCultureImplementation::GetDisplayName() const
+FString FLegacyCultureImplementation::GetDisplayName() const
 {
 	return DisplayName.ToString();
 }
 
-FString FCulture::FLegacyCultureImplementation::GetEnglishName() const
+FString FLegacyCultureImplementation::GetEnglishName() const
 {
 	return EnglishName;
 }
 
-int FCulture::FLegacyCultureImplementation::GetKeyboardLayoutId() const
+int FLegacyCultureImplementation::GetKeyboardLayoutId() const
 {
 	return KeyboardLayoutId;
 }
 
-int FCulture::FLegacyCultureImplementation::GetLCID() const
+int FLegacyCultureImplementation::GetLCID() const
 {
 	return LCID;
 }
 
-FString FCulture::FLegacyCultureImplementation::GetName() const
+FString FLegacyCultureImplementation::GetName() const
 {
 	return Name;
 }
 
-FString FCulture::FLegacyCultureImplementation::GetCanonicalName(const FString& Name)
+FString FLegacyCultureImplementation::GetCanonicalName(const FString& Name)
 {
 	return Name;
 }
 
-FString FCulture::FLegacyCultureImplementation::GetNativeName() const
+FString FLegacyCultureImplementation::GetNativeName() const
 {
 	return NativeName;
 }
 
-FString FCulture::FLegacyCultureImplementation::GetNativeLanguage() const
+FString FLegacyCultureImplementation::GetNativeLanguage() const
 {
 	int32 LastBracket = INDEX_NONE;
 	int32 FirstBracket = INDEX_NONE;
@@ -81,7 +81,7 @@ FString FCulture::FLegacyCultureImplementation::GetNativeLanguage() const
 	return NativeName;
 }
 
-FString FCulture::FLegacyCultureImplementation::GetNativeRegion() const
+FString FLegacyCultureImplementation::GetNativeRegion() const
 {
 	int32 LastBracket = INDEX_NONE;
 	int32 FirstBracket = INDEX_NONE;
@@ -92,32 +92,47 @@ FString FCulture::FLegacyCultureImplementation::GetNativeRegion() const
 	return NativeName;
 }
 
-FString FCulture::FLegacyCultureImplementation::GetUnrealLegacyThreeLetterISOLanguageName() const
+FString FLegacyCultureImplementation::GetRegion() const
+{
+	return FString();
+}
+
+FString FLegacyCultureImplementation::GetScript() const
+{
+	return FString();
+}
+
+FString FLegacyCultureImplementation::GetVariant() const
+{
+	return FString();
+}
+
+FString FLegacyCultureImplementation::GetUnrealLegacyThreeLetterISOLanguageName() const
 {
 	return UnrealLegacyThreeLetterISOLanguageName;
 }
 
-FString FCulture::FLegacyCultureImplementation::GetThreeLetterISOLanguageName() const
+FString FLegacyCultureImplementation::GetThreeLetterISOLanguageName() const
 {
 	return ThreeLetterISOLanguageName;
 }
 
-FString FCulture::FLegacyCultureImplementation::GetTwoLetterISOLanguageName() const
+FString FLegacyCultureImplementation::GetTwoLetterISOLanguageName() const
 {
 	return TwoLetterISOLanguageName;
 }
 
-const FDecimalNumberFormattingRules& FCulture::FLegacyCultureImplementation::GetDecimalNumberFormattingRules()
+const FDecimalNumberFormattingRules& FLegacyCultureImplementation::GetDecimalNumberFormattingRules()
 {
 	return DecimalNumberFormattingRules;
 }
 
-const FDecimalNumberFormattingRules& FCulture::FLegacyCultureImplementation::GetPercentFormattingRules()
+const FDecimalNumberFormattingRules& FLegacyCultureImplementation::GetPercentFormattingRules()
 {
 	return PercentFormattingRules;
 }
 
-const FDecimalNumberFormattingRules& FCulture::FLegacyCultureImplementation::GetCurrencyFormattingRules(const FString& InCurrencyCode)
+const FDecimalNumberFormattingRules& FLegacyCultureImplementation::GetCurrencyFormattingRules(const FString& InCurrencyCode)
 {
 	const bool bUseDefaultFormattingRules = InCurrencyCode.IsEmpty();
 
@@ -191,19 +206,19 @@ ETextPluralForm GetDefaultPluralForm(T Val, const ETextPluralType PluralType)
 
 }
 
-ETextPluralForm FCulture::FLegacyCultureImplementation::GetPluralForm(int32 Val, const ETextPluralType PluralType) const
+ETextPluralForm FLegacyCultureImplementation::GetPluralForm(int32 Val, const ETextPluralType PluralType) const
 {
 	checkf(Val >= 0, TEXT("GetPluralFormImpl requires a positive value"));
 	return GetDefaultPluralForm(Val, PluralType);
 }
 
-ETextPluralForm FCulture::FLegacyCultureImplementation::GetPluralForm(double Val, const ETextPluralType PluralType) const
+ETextPluralForm FLegacyCultureImplementation::GetPluralForm(double Val, const ETextPluralType PluralType) const
 {
 	checkf(!FMath::IsNegativeDouble(Val), TEXT("GetPluralFormImpl requires a positive value"));
 	return GetDefaultPluralForm((int64)Val, PluralType);
 }
 
-const TArray<ETextPluralForm>& FCulture::FLegacyCultureImplementation::GetValidPluralForms(const ETextPluralType PluralType) const
+const TArray<ETextPluralForm>& FLegacyCultureImplementation::GetValidPluralForms(const ETextPluralType PluralType) const
 {
 	if (PluralType == ETextPluralType::Cardinal)
 	{
diff --git a/Engine/Source/Runtime/Core/Private/Internationalization/LegacyCulture.h b/Engine/Source/Runtime/Core/Private/Internationalization/LegacyCulture.h
--- a/Engine/Source/Runtime/Core/Private/Internationalization/LegacyCulture.h
+++ b/Engine/Source/Runtime/Core/Private/Internationalization/LegacyCulture.h
@@ -6,13 +6,15 @@
 #include "Containers/Map.h"
 #include "Internationalization/Text.h"
 #include "Internationalization/Culture.h"
+#include "Internationalization/CultureImplementation.h"
 #include "Internationalization/FastDecimalFormat.h"
 
 #if !UE_ENABLE_ICU
-class FCulture::FLegacyCultureImplementation
+class FLegacyCultureImplementation : public ICultureImplementation
 {
 	friend class FCulture;
 
+public:
 	FLegacyCultureImplementation(
 		const FText& InDisplayName, 
 		const FString& InEnglishName, 
@@ -27,43 +29,33 @@ class FCulture::FLegacyCultureImplementation
 		const FDecimalNumberFormattingRules& InPercentFormattingRules,
 		const FDecimalNumberFormattingRules& InBaseCurrencyFormattingRules
 		);
-
-	FString GetDisplayName() const;
-
-	FString GetEnglishName() const;
-
-	int GetKeyboardLayoutId() const;
-
-	int GetLCID() const;
+	virtual ~FLegacyCultureImplementation() = default;
+
+	//~ ICultureImplementation interface
+	virtual FString GetDisplayName() const override;
+	virtual FString GetEnglishName() const override;
+	virtual int GetKeyboardLayoutId() const override;
+	virtual int GetLCID() const override;
+	virtual FString GetName() const override;
+	virtual FString GetNativeName() const override;
+	virtual FString GetUnrealLegacyThreeLetterISOLanguageName() const override;
+	virtual FString GetThreeLetterISOLanguageName() const override;
+	virtual FString GetTwoLetterISOLanguageName() const override;
+	virtual FString GetNativeLanguage() const override;
+	virtual FString GetNativeRegion() const override;
+	virtual FString GetRegion() const override;
+	virtual FString GetScript() const override;
+	virtual FString GetVariant() const override;
+	virtual const FDecimalNumberFormattingRules& GetDecimalNumberFormattingRules() override;
+	virtual const FDecimalNumberFormattingRules& GetPercentFormattingRules() override;
+	virtual const FDecimalNumberFormattingRules& GetCurrencyFormattingRules(const FString& InCurrencyCode) override;
+	virtual ETextPluralForm GetPluralForm(int32 Val, const ETextPluralType PluralType) const override;
+	virtual ETextPluralForm GetPluralForm(double Val, const ETextPluralType PluralType) const override;
+	virtual const TArray<ETextPluralForm>& GetValidPluralForms(const ETextPluralType PluralType) const override;
 
 	static FString GetCanonicalName(const FString& Name);
 
-	FString GetName() const;
-
-	FString GetNativeName() const;
-
-	FString GetUnrealLegacyThreeLetterISOLanguageName() const;
-
-	FString GetThreeLetterISOLanguageName() const;
-
-	FString GetTwoLetterISOLanguageName() const;
-
-	FString GetNativeLanguage() const;
-
-	FString GetNativeRegion() const;
-
-	const FDecimalNumberFormattingRules& GetDecimalNumberFormattingRules();
-
-	const FDecimalNumberFormattingRules& GetPercentFormattingRules();
-
-	const FDecimalNumberFormattingRules& GetCurrencyFormattingRules(const FString& InCurrencyCode);
-
-	ETextPluralForm GetPluralForm(int32 Val, const ETextPluralType PluralType) const;
-
-	ETextPluralForm GetPluralForm(double Val, const ETextPluralType PluralType) const;
-
-	const TArray<ETextPluralForm>& GetValidPluralForms(const ETextPluralType PluralType) const;
-
+private:
 	// Full localized culture name
 	const FText DisplayName;
 
diff --git a/Engine/Source/Runtime/Core/Private/Internationalization/LegacyInternationalization.cpp b/Engine/Source/Runtime/Core/Private/Internationalization/LegacyInternationalization.cpp
--- a/Engine/Source/Runtime/Core/Private/Internationalization/LegacyInternationalization.cpp
+++ b/Engine/Source/Runtime/Core/Private/Internationalization/LegacyInternationalization.cpp
@@ -1,5 +1,7 @@
 // Copyright 1998-2019 Epic Games, Inc. All Rights Reserved.
+
 #include "Internationalization/LegacyInternationalization.h"
+#include "Internationalization/Cultures/LeetCulture.h"
 
 #if !UE_ENABLE_ICU
 
@@ -19,6 +21,10 @@ bool FLegacyInternationalization::Initialize()
 	I18N->CurrentLanguage = I18N->InvariantCulture;
 	I18N->CurrentLocale = I18N->InvariantCulture;
 
+#if ENABLE_LOC_TESTING
+	I18N->AddCustomCulture(MakeShared<FLeetCulture>(I18N->InvariantCulture.ToSharedRef()));
+#endif
+
 	return true;
 }
 
@@ -46,7 +52,12 @@ void FLegacyInternationalization::HandleLanguageChanged(const FString& Name)
 
 void FLegacyInternationalization::GetCultureNames(TArray<FString>& CultureNames) const
 {
-	CultureNames.Add(TEXT(""));
+	CultureNames.Reset(1 + I18N->CustomCultures.Num());
+	CultureNames.Add(FString());
+	for (const FCultureRef& CustomCulture : I18N->CustomCultures)
+	{
+		CultureNames.Add(CustomCulture->GetName());
+	}
 }
 
 TArray<FString> FLegacyInternationalization::GetPrioritizedCultureNames(const FString& Name)
@@ -58,7 +69,12 @@ TArray<FString> FLegacyInternationalization::GetPrioritizedCultureNames(const FS
 
 FCulturePtr FLegacyInternationalization::GetCulture(const FString& Name)
 {
-	return Name.IsEmpty() ? I18N->InvariantCulture : nullptr;
+	FCulturePtr Culture = I18N->GetCustomCulture(Name);
+	if (!Culture && Name.IsEmpty())
+	{
+		Culture = I18N->InvariantCulture;
+	}
+	return Culture;
 }
 
 #endif
diff --git a/Engine/Source/Runtime/Core/Private/Internationalization/TextLocalizationManager.cpp b/Engine/Source/Runtime/Core/Private/Internationalization/TextLocalizationManager.cpp
--- a/Engine/Source/Runtime/Core/Private/Internationalization/TextLocalizationManager.cpp
+++ b/Engine/Source/Runtime/Core/Private/Internationalization/TextLocalizationManager.cpp
@@ -6,6 +6,7 @@
 #include "Internationalization/LocalizationResourceTextSource.h"
 #include "Internationalization/PolyglotTextSource.h"
 #include "Internationalization/StringTableRegistry.h"
+#include "Internationalization/Cultures/LeetCulture.h"
 #include "GenericPlatform/GenericPlatformFile.h"
 #include "HAL/FileManager.h"
 #include "HAL/LowLevelMemTracker.h"
@@ -47,9 +48,9 @@ void ApplyDefaultCultureSettings(const ELocalizationLoadFlags LocLoadFlags)
 
 	// Set culture according to configuration now that configs are available.
 #if ENABLE_LOC_TESTING
-	if (FCommandLine::IsInitialized() && FParse::Param(FCommandLine::Get(), TEXT("LEET")))
+	if (FCommandLine::IsInitialized() && FParse::Param(FCommandLine::Get(), *FLeetCulture::StaticGetName()))
 	{
-		I18N.SetCurrentCulture(TEXT("LEET"));
+		I18N.SetCurrentCulture(FLeetCulture::StaticGetName());
 	}
 	else
 #endif
@@ -188,7 +189,7 @@ void ApplyDefaultCultureSettings(const ELocalizationLoadFlags LocLoadFlags)
 			FString TargetCultureName = InRequestedCulture;
 
 #if ENABLE_LOC_TESTING
-			if (TargetCultureName != TEXT("LEET"))
+			if (TargetCultureName != FLeetCulture::StaticGetName())
 #endif
 			{
 				ELocalizationLoadFlags ValidationFlags = ELocalizationLoadFlags::None;
@@ -572,7 +573,7 @@ FTextDisplayStringRef FTextLocalizationManager::GetDisplayString(const FTextKey&
 	}
 
 #if ENABLE_LOC_TESTING
-	const bool bShouldLEETIFYAll = bIsInitialized && FInternationalization::Get().GetCurrentLanguage()->GetName() == TEXT("LEET");
+	const bool bShouldLEETIFYAll = bIsInitialized && FInternationalization::Get().GetCurrentLanguage()->GetName() == FLeetCulture::StaticGetName();
 
 	// Attempt to set bShouldLEETIFYUnlocalizedString appropriately, only once, after the commandline is initialized and parsed.
 	static bool bShouldLEETIFYUnlocalizedString = false;
@@ -903,7 +904,7 @@ void FTextLocalizationManager::LoadLocalizationResourcesForPrioritizedCultures(T
 
 #if ENABLE_LOC_TESTING
 	// The leet culture is fake. Just leet-ify existing strings.
-	if (PrioritizedCultureNames[0] == TEXT("LEET"))
+	if (PrioritizedCultureNames[0] == FLeetCulture::StaticGetName())
 	{
 		// Lock while updating the tables
 		{
diff --git a/Engine/Source/Runtime/Core/Private/Tests/Internationalization/TextTest.cpp b/Engine/Source/Runtime/Core/Private/Tests/Internationalization/TextTest.cpp
--- a/Engine/Source/Runtime/Core/Private/Tests/Internationalization/TextTest.cpp
+++ b/Engine/Source/Runtime/Core/Private/Tests/Internationalization/TextTest.cpp
@@ -11,6 +11,7 @@
 #include "Internationalization/FastDecimalFormat.h"
 #include "Internationalization/StringTableRegistry.h"
 #include "Internationalization/Internationalization.h"
+#include "Internationalization/Cultures/LeetCulture.h"
 #include "Serialization/MemoryWriter.h"
 #include "Serialization/MemoryReader.h"
 #include "Misc/AutomationTest.h"
@@ -496,9 +497,10 @@ bool FTextTest::RunTest (const FString& Parameters)
 			// The original string in the native language.
 			FormattedTestLayer2_OriginalLanguageSourceString = FormattedTestLayer2.BuildSourceString();
 
+#if ENABLE_LOC_TESTING
 			{
 				// Swap to "LEET" culture to check if rebuilding works (verify the whole)
-				I18N.SetCurrentCulture("LEET");
+				I18N.SetCurrentCulture(FLeetCulture::StaticGetName());
 
 				// When changes are made to FormattedTestLayer2, please pull out the newly translated LEET string and update the below if-statement to keep the test passing!
 				FString LEETTranslatedString = FormattedTestLayer2.ToString();
@@ -512,6 +514,7 @@ bool FTextTest::RunTest (const FString& Parameters)
 					AddError( TEXT("Desired Output=") + DesiredOutput );
 				}
 			}
+#endif
 
 			// Swap to French-Canadian to check if rebuilding works (verify each numerical component)
 			{
@@ -563,8 +566,9 @@ bool FTextTest::RunTest (const FString& Parameters)
 			}
 		}
 
+#if ENABLE_LOC_TESTING
 		{
-			I18N.SetCurrentCulture("LEET");
+			I18N.SetCurrentCulture(FLeetCulture::StaticGetName());
 
 			FText FormattedEnglishTextHistoryAsLeet;
 			FText FormattedFrenchCanadianTextHistoryAsLeet;
@@ -627,6 +631,7 @@ bool FTextTest::RunTest (const FString& Parameters)
 			}
 		}
 	}
+#endif
 #else
 	AddWarning("ICU is disabled thus locale-aware formatting needed in rebuilding source text from history is disabled.");
 #endif
diff --git a/Engine/Source/Runtime/Core/Public/Internationalization/Culture.h b/Engine/Source/Runtime/Core/Public/Internationalization/Culture.h
--- a/Engine/Source/Runtime/Core/Public/Internationalization/Culture.h
+++ b/Engine/Source/Runtime/Core/Public/Internationalization/Culture.h
@@ -4,13 +4,21 @@
 #include "CoreTypes.h"
 #include "Containers/Array.h"
 #include "Containers/UnrealString.h"
-#include "Templates/SharedPointer.h"
+#include "Templates/UniquePtr.h"
 #include "Internationalization/CulturePointer.h"
 
 struct FDecimalNumberFormattingRules;
 enum class ETextPluralForm : uint8;
 enum class ETextPluralType : uint8;
 
+#if UE_ENABLE_ICU
+class FICUCultureImplementation;
+typedef FICUCultureImplementation FCultureImplementation;
+#else
+class FLegacyCultureImplementation;
+typedef FLegacyCultureImplementation FCultureImplementation;
+#endif
+
 class CORE_API FCulture
 {
 #if UE_ENABLE_ICU
@@ -20,24 +28,9 @@ class CORE_API FCulture
 #endif
 
 public:
-#if UE_ENABLE_ICU
-	static FCulturePtr Create(const FString& LocaleName);
-#else
-	static FCulturePtr Create(
-		const FText& InDisplayName, 
-		const FString& InEnglishName, 
-		const int InKeyboardLayoutId, 
-		const int InLCID, 
-		const FString& InName, 
-		const FString& InNativeName, 
-		const FString& InUnrealLegacyThreeLetterISOLanguageName, 
-		const FString& InThreeLetterISOLanguageName, 
-		const FString& InTwoLetterISOLanguageName,
-		const FDecimalNumberFormattingRules& InDecimalNumberFormattingRules,
-		const FDecimalNumberFormattingRules& InPercentFormattingRules,
-		const FDecimalNumberFormattingRules& InBaseCurrencyFormattingRules
-		);
-#endif
+	~FCulture();
+	
+	static FCultureRef Create(TUniquePtr<FCultureImplementation>&& InImplementation);
 
 	const FString& GetDisplayName() const;
 
@@ -106,37 +99,10 @@ public:
 	void HandleCultureChanged();
 
 private:
-#if UE_ENABLE_ICU
-	class FICUCultureImplementation;
-	typedef FICUCultureImplementation FImplementation;
-	TSharedRef<FICUCultureImplementation> Implementation;
-#else
-	class FLegacyCultureImplementation;
-	typedef FLegacyCultureImplementation FImplementation;
-	TSharedRef<FLegacyCultureImplementation> Implementation;
-#endif
+	explicit FCulture(TUniquePtr<FCultureImplementation>&& InImplementation);
 
-protected:
-#if UE_ENABLE_ICU
-	FCulture(const FString& LocaleName);
-#else
-	FCulture(
-		const FText& InDisplayName, 
-		const FString& InEnglishName, 
-		const int InKeyboardLayoutId, 
-		const int InLCID, 
-		const FString& InName, 
-		const FString& InNativeName, 
-		const FString& InUnrealLegacyThreeLetterISOLanguageName, 
-		const FString& InThreeLetterISOLanguageName, 
-		const FString& InTwoLetterISOLanguageName,
-		const FDecimalNumberFormattingRules& InDecimalNumberFormattingRules,
-		const FDecimalNumberFormattingRules& InPercentFormattingRules,
-		const FDecimalNumberFormattingRules& InBaseCurrencyFormattingRules
-		);
-#endif
+	TUniquePtr<FCultureImplementation> Implementation;
 
-private:
 	FString CachedDisplayName;
 	FString CachedEnglishName;
 	FString CachedName;
diff --git a/Engine/Source/Runtime/Core/Public/Internationalization/ICustomCulture.h b/Engine/Source/Runtime/Core/Public/Internationalization/ICustomCulture.h
new file mode 100644
--- /dev/null
+++ b/Engine/Source/Runtime/Core/Public/Internationalization/ICustomCulture.h
@@ -0,0 +1,30 @@
+// Copyright 1998-2019 Epic Games, Inc. All Rights Reserved.
+
+#pragma once
+
+#include "CoreTypes.h"
+#include "Internationalization/CulturePointer.h"
+
+/**
+ * Interface for a custom culture.
+ * @note See FCulture for the API description.
+ */
+class ICustomCulture
+{
+public:
+	virtual ~ICustomCulture() = default;
+
+	virtual FCultureRef GetBaseCulture() const = 0;
+	virtual FString GetDisplayName() const = 0;
+	virtual FString GetEnglishName() const = 0;
+	virtual FString GetName() const = 0;
+	virtual FString GetNativeName() const = 0;
+	virtual FString GetUnrealLegacyThreeLetterISOLanguageName() const = 0;
+	virtual FString GetThreeLetterISOLanguageName() const = 0;
+	virtual FString GetTwoLetterISOLanguageName() const = 0;
+	virtual FString GetNativeLanguage() const = 0;
+	virtual FString GetNativeRegion() const = 0;
+	virtual FString GetRegion() const = 0;
+	virtual FString GetScript() const = 0;
+	virtual FString GetVariant() const = 0;
+};
diff --git a/Engine/Source/Runtime/Core/Public/Internationalization/Internationalization.h b/Engine/Source/Runtime/Core/Public/Internationalization/Internationalization.h
--- a/Engine/Source/Runtime/Core/Public/Internationalization/Internationalization.h
+++ b/Engine/Source/Runtime/Core/Public/Internationalization/Internationalization.h
@@ -14,6 +14,7 @@
 #define LOC_DEFINE_REGION
 
 class FCulture;
+class ICustomCulture;
 class FICUInternationalization;
 class FLegacyInternationalization;
 
@@ -161,6 +162,12 @@ public:
 	/** Load and cache the data needed for every culture we know about (this is usually done per-culture as required) */
 	CORE_API void LoadAllCultureData();
 
+	/** Add a new custom culture */
+	CORE_API void AddCustomCulture(const TSharedRef<ICustomCulture>& InCustomCulture);
+
+	/** Get a new custom culture from its name */
+	CORE_API FCulturePtr GetCustomCulture(const FString& InCultureName) const;
+
 	/** Has the given culture been remapped in this build? */
 	CORE_API bool IsCultureRemapped(const FString& Name, FString* OutMappedCulture);
 
@@ -244,6 +251,11 @@ private:
 	 * An invariant culture that can be used when you don't care about localization/internationalization.
 	 */
 	FCulturePtr InvariantCulture;
+
+	/**
+	 * Custom cultures registered via AddCustomCulture.
+	 */
+	TArray<FCultureRef> CustomCultures;
 };
 
 namespace UE4LocGen_Private
