Fixed compile error

#rnx

diff --git a/Engine/Source/Runtime/Core/Public/Internationalization/LocTesting.h b/Engine/Source/Runtime/Core/Public/Internationalization/LocTesting.h
--- a/Engine/Source/Runtime/Core/Public/Internationalization/LocTesting.h
+++ b/Engine/Source/Runtime/Core/Public/Internationalization/LocTesting.h
@@ -2,6 +2,8 @@
 
 #pragma once
 
+#include "Misc/Build.h"
+
 #ifndef ENABLE_LOC_TESTING
 	#define ENABLE_LOC_TESTING (UE_BUILD_DEBUG || UE_BUILD_DEVELOPMENT || UE_BUILD_TEST)
 #endif
