Fixed ComputeSquaredDistanceToPoint to no longer use manhattan distance when point is outside more than one side

diff --git a/Engine/Source/Runtime/Core/Public/Math/Box.h b/Engine/Source/Runtime/Core/Public/Math/Box.h
--- a/Engine/Source/Runtime/Core/Public/Math/Box.h
+++ b/Engine/Source/Runtime/Core/Public/Math/Box.h
@@ -536,40 +536,7 @@ FORCEINLINE FBox& FBox::operator+=( const FBox& Other )
 
 FORCEINLINE FVector FBox::GetClosestPointTo( const FVector& Point ) const
 {
-	// start by considering the point inside the box
-	FVector ClosestPoint = Point;
-
-	// now clamp to inside box if it's outside
-	if (Point.X < Min.X)
-	{
-		ClosestPoint.X = Min.X;
-	}
-	else if (Point.X > Max.X)
-	{
-		ClosestPoint.X = Max.X;
-	}
-
-	// now clamp to inside box if it's outside
-	if (Point.Y < Min.Y)
-	{
-		ClosestPoint.Y = Min.Y;
-	}
-	else if (Point.Y > Max.Y)
-	{
-		ClosestPoint.Y = Max.Y;
-	}
-
-	// Now clamp to inside box if it's outside.
-	if (Point.Z < Min.Z)
-	{
-		ClosestPoint.Z = Min.Z;
-	}
-	else if (Point.Z > Max.Z)
-	{
-		ClosestPoint.Z = Max.Z;
-	}
-
-	return ClosestPoint;
+	return ::GetClosestPointFromBoxToPoint(Min, Max, Point);
 }
 
 
diff --git a/Engine/Source/Runtime/Core/Public/Math/Vector.h b/Engine/Source/Runtime/Core/Public/Math/Vector.h
--- a/Engine/Source/Runtime/Core/Public/Math/Vector.h
+++ b/Engine/Source/Runtime/Core/Public/Math/Vector.h
@@ -1072,52 +1072,64 @@ FORCEINLINE uint32 GetTypeHash(const FVector& Vector)
 	}
 #endif
 
-
-/** 
- * Util to calculate distance from a point to a bounding box 
+/**
+ * Util to calculate closest point from a point to a bounding box
  *
- * @param Mins 3D Point defining the lower values of the axis of the bound box
+ * @param Min 3D Point defining the lower values of the axis of the bound box
  * @param Max 3D Point defining the lower values of the axis of the bound box
  * @param Point 3D position of interest
- * @return the distance from the Point to the bounding box.
+ * @return If point is outside box, returns closest point on box.  If point inside, returns point
  */
-FORCEINLINE float ComputeSquaredDistanceFromBoxToPoint(const FVector& Mins, const FVector& Maxs, const FVector& Point)
+FORCEINLINE FVector GetClosestPointFromBoxToPoint(const FVector& Min, const FVector& Max, const FVector& Point)
 {
-	// Accumulates the distance as we iterate axis
-	float DistSquared = 0.f;
+	// start by considering the point inside the box
+	FVector ClosestPoint = Point;
 
-	// Check each axis for min/max and add the distance accordingly
-	// NOTE: Loop manually unrolled for > 2x speed up
-	if (Point.X < Mins.X)
+	// now clamp to inside box if it's outside
+	if (Point.X < Min.X)
 	{
-		DistSquared += FMath::Square(Point.X - Mins.X);
+		ClosestPoint.X = Min.X;
 	}
-	else if (Point.X > Maxs.X)
+	else if (Point.X > Max.X)
 	{
-		DistSquared += FMath::Square(Point.X - Maxs.X);
+		ClosestPoint.X = Max.X;
 	}
-	
-	if (Point.Y < Mins.Y)
+
+	// now clamp to inside box if it's outside
+	if (Point.Y < Min.Y)
 	{
-		DistSquared += FMath::Square(Point.Y - Mins.Y);
+		ClosestPoint.Y = Min.Y;
 	}
-	else if (Point.Y > Maxs.Y)
+	else if (Point.Y > Max.Y)
 	{
-		DistSquared += FMath::Square(Point.Y - Maxs.Y);
+		ClosestPoint.Y = Max.Y;
 	}
-	
-	if (Point.Z < Mins.Z)
+
+	// Now clamp to inside box if it's outside.
+	if (Point.Z < Min.Z)
 	{
-		DistSquared += FMath::Square(Point.Z - Mins.Z);
+		ClosestPoint.Z = Min.Z;
 	}
-	else if (Point.Z > Maxs.Z)
+	else if (Point.Z > Max.Z)
 	{
-		DistSquared += FMath::Square(Point.Z - Maxs.Z);
+		ClosestPoint.Z = Max.Z;
 	}
-	
-	return DistSquared;
+
+	return ClosestPoint;
 }
 
+/** 
+ * Util to calculate distance from a point to a bounding box 
+ *
+ * @param Mins 3D Point defining the lower values of the axis of the bound box
+ * @param Max 3D Point defining the lower values of the axis of the bound box
+ * @param Point 3D position of interest
+ * @return the distance from the Point to the bounding box.
+ */
+FORCEINLINE float ComputeSquaredDistanceFromBoxToPoint(const FVector& Mins, const FVector& Maxs, const FVector& Point)
+{
+	return (GetClosestPointFromBoxToPoint(Mins, Maxs, Point) - Point).SizeSquared();
+}
 
 FORCEINLINE FVector::FVector(const FVector2D V, float InZ)
 	: X(V.X), Y(V.Y), Z(InZ)
