PR #6132: Fix 2 profiler frames created per engine tick (Contributed by [REDACTED_PERSON])

diff --git a/Engine/Source/Runtime/Core/Private/Windows/WindowsPlatformMisc.cpp b/Engine/Source/Runtime/Core/Private/Windows/WindowsPlatformMisc.cpp
--- a/Engine/Source/Runtime/Core/Private/Windows/WindowsPlatformMisc.cpp
+++ b/Engine/Source/Runtime/Core/Private/Windows/WindowsPlatformMisc.cpp
@@ -857,12 +857,6 @@ void FWindowsPlatformMisc::BeginNamedEventFrame()
 {
 #if FRAMEPRO_ENABLED
 	FFrameProProfiler::FrameStart();
-#elif UE_EXTERNAL_PROFILING_ENABLED
-	FExternalProfiler* Profiler = FActiveExternalProfilerBase::GetActiveProfiler();
-	if (Profiler)
-	{
-		Profiler->FrameSync();
-	}
 #endif
 }
 
