reserve more memory for the crash handler

diff --git a/Engine/Source/Runtime/Core/Private/IOS/IOSPlatformMisc.cpp b/Engine/Source/Runtime/Core/Private/IOS/IOSPlatformMisc.cpp
--- a/Engine/Source/Runtime/Core/Private/IOS/IOSPlatformMisc.cpp
+++ b/Engine/Source/Runtime/Core/Private/IOS/IOSPlatformMisc.cpp
@@ -1666,7 +1666,7 @@ void FIOSPlatformMisc::SetCrashHandler(void (* CrashHandler)(const FGenericCrash
     if (!FIOSApplicationInfo::CrashReporter && !FIOSApplicationInfo::CrashMalloc)
     {
         // configure the crash handler malloc zone to reserve a little memory for itself
-        FIOSApplicationInfo::CrashMalloc = new FIOSMallocCrashHandler(128*1024);
+        FIOSApplicationInfo::CrashMalloc = new FIOSMallocCrashHandler(4*1024*1024);
         
         PLCrashReporterConfig* Config = [[[PLCrashReporterConfig alloc] initWithSignalHandlerType: PLCrashReporterSignalHandlerTypeBSD symbolicationStrategy: PLCrashReporterSymbolicationStrategyNone crashReportFolder: FIOSApplicationInfo::TemporaryCrashReportFolder().GetNSString() crashReportName: FIOSApplicationInfo::TemporaryCrashReportName().GetNSString()] autorelease];
         FIOSApplicationInfo::CrashReporter = [[PLCrashReporter alloc] initWithConfiguration: Config];
