Integrating cl from Dev-Rendering to add support to RHI to automation test blacklist

diff --git a/Engine/Source/Runtime/Core/Private/Misc/AutomationTest.cpp b/Engine/Source/Runtime/Core/Private/Misc/AutomationTest.cpp
--- a/Engine/Source/Runtime/Core/Private/Misc/AutomationTest.cpp
+++ b/Engine/Source/Runtime/Core/Private/Misc/AutomationTest.cpp
@@ -13,6 +13,7 @@
 #include "Misc/OutputDeviceRedirector.h"
 #include "Internationalization/Regex.h"
 #include <inttypes.h>
+#include "Misc/App.h"
 
 DEFINE_LOG_CATEGORY_STATIC(LogAutomationTest, Warning, All);
 
@@ -415,7 +416,7 @@ void FAutomationTestFramework::BuildTestBlacklistFromConfig()
 					if (Section.Key() == TEXT("BlacklistTest"))
 					{
 						FString BlacklistValue = Section.Value().GetValue();
-						FString Map, Test, Reason, Warn, ListName;
+						FString Map, Test, Reason, RHIs, Warn, ListName;
 						bool bSuccess = false;
 
 						if (FParse::Value(*BlacklistValue, TEXT("Test="), Test, true))
@@ -423,6 +424,7 @@ void FAutomationTestFramework::BuildTestBlacklistFromConfig()
 							ListName = FString(Test);
 							FParse::Value(*BlacklistValue, TEXT("Map="), Map, true);
 							FParse::Value(*BlacklistValue, TEXT("Reason="), Reason);
+							FParse::Value(*BlacklistValue, TEXT("RHIs="), RHIs);
 							FParse::Value(*BlacklistValue, TEXT("Warn="), Warn);
 
 							if (Map.IsEmpty())
@@ -447,10 +449,19 @@ void FAutomationTestFramework::BuildTestBlacklistFromConfig()
 							}
 							else
 							{
+								ListName.RemoveSpacesInline();
 								FBlacklistEntry& Entry = TestBlacklist.Add(ListName);
 								Entry.Map = Map;
 								Entry.Test = Test;
 								Entry.Reason = Reason;
+								if (!RHIs.IsEmpty())
+								{
+									RHIs.ToLower().ParseIntoArray(Entry.RHIs, TEXT(","), true);
+									for (FString& RHI : Entry.RHIs)
+									{
+										RHI.TrimStartAndEndInline();
+									}
+								}
 								Entry.bWarn = Warn.ToBool();
 							}
 						}
@@ -480,6 +491,11 @@ bool FAutomationTestFramework::IsBlacklisted(const FString& TestName, FString* O
 
 	if (Entry)
 	{
+		if (Entry->RHIs.Num() != 0 && !Entry->RHIs.Contains(FApp::GetGraphicsRHI().ToLower()))
+		{
+			return false;
+		}
+
 		if (OutReason != nullptr)
 		{
 			*OutReason = Entry->Reason;
@@ -567,7 +583,7 @@ void FAutomationTestFramework::GetValidTestNames( TArray<FAutomationTestInfo>& T
 				FString BlacklistReason;
 				bool bWarn(false);
 				FString TestName = Test.GetDisplayName();
-				if (!IsBlacklisted(TestName, &BlacklistReason, &bWarn))
+				if (!IsBlacklisted(TestName.Replace(TEXT(" "), TEXT("")), &BlacklistReason, &bWarn))
 				{
 					TestInfo.Add(Test);
 				}
diff --git a/Engine/Source/Runtime/Core/Public/Misc/AutomationTest.h b/Engine/Source/Runtime/Core/Public/Misc/AutomationTest.h
--- a/Engine/Source/Runtime/Core/Public/Misc/AutomationTest.h
+++ b/Engine/Source/Runtime/Core/Public/Misc/AutomationTest.h
@@ -918,6 +918,7 @@ private:
 		FString Map;
 		FString Test;
 		FString Reason;
+		TArray<FString> RHIs;
 		bool bWarn;
 	};
 
