Undefine GetTempFileName() after including Windows.h to prevent name conflicts.

#rnx

diff --git a/Engine/Source/Runtime/Core/Public/Windows/PostWindowsApi.h b/Engine/Source/Runtime/Core/Public/Windows/PostWindowsApi.h
--- a/Engine/Source/Runtime/Core/Public/Windows/PostWindowsApi.h
+++ b/Engine/Source/Runtime/Core/Public/Windows/PostWindowsApi.h
@@ -44,6 +44,7 @@ THIRD_PARTY_INCLUDES_END
 #undef GetClassInfo
 #undef Yield
 #undef IMediaEventSink
+#undef GetTempFileName
 
 // Undefine all the atomics. AllowWindowsPlatformAtomics/HideWindowsPlatformAtomics temporarily defining these macros.
 #if !PLATFORM_XBOXONE
