Fix FMallocBinned2::Trim() crashing with -onethread -nothreading (UE-79811)

diff --git a/Engine/Source/Runtime/Core/Private/HAL/MallocBinned2.cpp b/Engine/Source/Runtime/Core/Private/HAL/MallocBinned2.cpp
--- a/Engine/Source/Runtime/Core/Private/HAL/MallocBinned2.cpp
+++ b/Engine/Source/Runtime/Core/Private/HAL/MallocBinned2.cpp
@@ -1096,7 +1096,7 @@ void FMallocBinned2::Trim(bool bTrimThreadCaches)
 			FlushCurrentThreadCache();
 		};
 		// Skip task threads on desktop platforms as it is too slow and they don't have much memory
-		FTaskGraphInterface::BroadcastSlow_OnlyUseForSpecialPurposes(!PLATFORM_DESKTOP, false, Broadcast);
+		FTaskGraphInterface::BroadcastSlow_OnlyUseForSpecialPurposes((!PLATFORM_DESKTOP && FPlatformProcess::SupportsMultithreading() && FApp::ShouldUseThreadingForPerformance()), false, Broadcast);
 		//UE_LOG(LogTemp, Display, TEXT("Trim Broadcast = %6.2fms"), 1000.0f * float(FPlatformTime::Seconds() - StartTime));
 	}
 	{
