Enable vectorization of FQuat and FRotator for hardware with NEON support

diff --git a/Engine/Source/Runtime/Core/Private/Math/UnrealMath.cpp b/Engine/Source/Runtime/Core/Private/Math/UnrealMath.cpp
--- a/Engine/Source/Runtime/Core/Private/Math/UnrealMath.cpp
+++ b/Engine/Source/Runtime/Core/Private/Math/UnrealMath.cpp
@@ -400,7 +400,7 @@ FQuat FRotator::Quaternion() const
 
 	DiagnosticCheckNaN();
 
-#if PLATFORM_ENABLE_VECTORINTRINSICS
+#if PLATFORM_ENABLE_VECTORINTRINSICS || PLATFORM_ENABLE_VECTORINTRINSICS_NEON
 	const VectorRegister Angles = MakeVectorRegister(Pitch, Yaw, Roll, 0.0f);
 	const VectorRegister AnglesNoWinding = VectorMod(Angles, GlobalVectorConstants::Float360);
 	const VectorRegister HalfAngles = VectorMultiply(AnglesNoWinding, GlobalVectorConstants::DEG_TO_RAD_HALF);
diff --git a/Engine/Source/Runtime/Core/Public/Math/Quat.h b/Engine/Source/Runtime/Core/Public/Math/Quat.h
--- a/Engine/Source/Runtime/Core/Public/Math/Quat.h
+++ b/Engine/Source/Runtime/Core/Public/Math/Quat.h
@@ -789,7 +789,7 @@ FORCEINLINE FQuat FQuat::operator-(const FQuat& Q) const
 
 FORCEINLINE bool FQuat::Equals(const FQuat& Q, float Tolerance) const
 {
-#if PLATFORM_ENABLE_VECTORINTRINSICS
+#if PLATFORM_ENABLE_VECTORINTRINSICS || PLATFORM_ENABLE_VECTORINTRINSICS_NEON
 	const VectorRegister ToleranceV = VectorLoadFloat1(&Tolerance);
 	const VectorRegister A = VectorLoadAligned(this);
 	const VectorRegister B = VectorLoadAligned(&Q);
@@ -918,7 +918,7 @@ FORCEINLINE float FQuat::operator|(const FQuat& Q) const
 
 FORCEINLINE void FQuat::Normalize(float Tolerance)
 {
-#if PLATFORM_ENABLE_VECTORINTRINSICS
+#if PLATFORM_ENABLE_VECTORINTRINSICS || PLATFORM_ENABLE_VECTORINTRINSICS_NEON
 	const VectorRegister Vector = VectorLoadAligned(this);
 
 	const VectorRegister SquareSum = VectorDot4(Vector, Vector);
diff --git a/Engine/Source/Runtime/Core/Public/Math/RotationTranslationMatrix.h b/Engine/Source/Runtime/Core/Public/Math/RotationTranslationMatrix.h
--- a/Engine/Source/Runtime/Core/Public/Math/RotationTranslationMatrix.h
+++ b/Engine/Source/Runtime/Core/Public/Math/RotationTranslationMatrix.h
@@ -31,7 +31,7 @@ public:
 
 FORCEINLINE FRotationTranslationMatrix::FRotationTranslationMatrix(const FRotator& Rot, const FVector& Origin)
 {
-#if PLATFORM_ENABLE_VECTORINTRINSICS
+#if PLATFORM_ENABLE_VECTORINTRINSICS || PLATFORM_ENABLE_VECTORINTRINSICS_NEON
 
 	const VectorRegister Angles = MakeVectorRegister(Rot.Pitch, Rot.Yaw, Rot.Roll, 0.0f);
 	const VectorRegister HalfAngles = VectorMultiply(Angles, GlobalVectorConstants::DEG_TO_RAD);
diff --git a/Engine/Source/Runtime/Core/Public/Math/Rotator.h b/Engine/Source/Runtime/Core/Public/Math/Rotator.h
--- a/Engine/Source/Runtime/Core/Public/Math/Rotator.h
+++ b/Engine/Source/Runtime/Core/Public/Math/Rotator.h
@@ -537,7 +537,7 @@ FORCEINLINE FRotator FRotator::operator-=( const FRotator& R )
 
 FORCEINLINE bool FRotator::IsNearlyZero(float Tolerance) const
 {
-#if PLATFORM_ENABLE_VECTORINTRINSICS
+#if PLATFORM_ENABLE_VECTORINTRINSICS || PLATFORM_ENABLE_VECTORINTRINSICS_NEON
 	const VectorRegister RegA = VectorLoadFloat3_W0(this);
 	const VectorRegister Norm = VectorNormalizeRotator(RegA);
 	const VectorRegister AbsNorm = VectorAbs(Norm);
@@ -559,7 +559,7 @@ FORCEINLINE bool FRotator::IsZero() const
 
 FORCEINLINE bool FRotator::Equals(const FRotator& R, float Tolerance) const
 {
-#if PLATFORM_ENABLE_VECTORINTRINSICS
+#if PLATFORM_ENABLE_VECTORINTRINSICS || PLATFORM_ENABLE_VECTORINTRINSICS_NEON
 	const VectorRegister RegA = VectorLoadFloat3_W0(this);
 	const VectorRegister RegB = VectorLoadFloat3_W0(&R);
 	const VectorRegister NormDelta = VectorNormalizeRotator(VectorSubtract(RegA, RegB));
@@ -678,7 +678,7 @@ FORCEINLINE FRotator FRotator::GetDenormalized() const
 
 FORCEINLINE void FRotator::Normalize()
 {
-#if PLATFORM_ENABLE_VECTORINTRINSICS
+#if PLATFORM_ENABLE_VECTORINTRINSICS || PLATFORM_ENABLE_VECTORINTRINSICS_NEON
 	VectorRegister VRotator = VectorLoadFloat3_W0(this);
 	VRotator = VectorNormalizeRotator(VRotator);
 	VectorStoreFloat3(VRotator, this);
diff --git a/Engine/Source/Runtime/Core/Public/Math/ScalarRegister.h b/Engine/Source/Runtime/Core/Public/Math/ScalarRegister.h
--- a/Engine/Source/Runtime/Core/Public/Math/ScalarRegister.h
+++ b/Engine/Source/Runtime/Core/Public/Math/ScalarRegister.h
@@ -13,7 +13,7 @@
 */
 //Currently disabled because FBoneAtom became FTransform and we want to iterate quickly on it.
 #define ENABLE_VECTORIZED_FBONEATOM		0 //PLATFORM_ENABLE_VECTORINTRINSICS && !__ARM_NEON__ && 1 // currently no support for VectorPermute needed for the boneatoms
-#define ENABLE_VECTORIZED_TRANSFORM		PLATFORM_ENABLE_VECTORINTRINSICS
+#define ENABLE_VECTORIZED_TRANSFORM		(PLATFORM_ENABLE_VECTORINTRINSICS || PLATFORM_ENABLE_VECTORINTRINSICS_NEON)
 
 #if ENABLE_VECTORIZED_FBONEATOM || ENABLE_VECTORIZED_TRANSFORM		
 
