Report the value of GIsRequestingExit during a crash

diff --git a/Engine/Source/Runtime/Core/Private/GenericPlatform/GenericPlatformCrashContext.cpp b/Engine/Source/Runtime/Core/Private/GenericPlatform/GenericPlatformCrashContext.cpp
--- a/Engine/Source/Runtime/Core/Private/GenericPlatform/GenericPlatformCrashContext.cpp
+++ b/Engine/Source/Runtime/Core/Private/GenericPlatform/GenericPlatformCrashContext.cpp
@@ -337,6 +337,7 @@ void FGenericCrashContext::SerializeContentToBuffer() const
 	AddCrashProperty( TEXT( "AppDefaultLocale" ), *NCachedCrashContextProperties::DefaultLocale );
 	AddCrashProperty( TEXT( "BuildVersion" ), FApp::GetBuildVersion() );
 	AddCrashProperty( TEXT( "IsUE4Release" ), NCachedCrashContextProperties::bIsUE4Release );
+	AddCrashProperty( TEXT( "IsRequestingExit" ), GIsRequestingExit );
 
 	// Remove periods from user names to match AutoReporter user names
 	// The name prefix is read by CrashRepository.AddNewCrash in the website code
