KAIROS:
Clean up keep awake tag usage based on [REDACTED_HANDLE] feedback.  It no longer depends on the native wrapper to keep the render thread awak while the game is running.

diff --git a/Engine/Source/Runtime/Core/Public/Misc/EmbeddedCommunication.h b/Engine/Source/Runtime/Core/Public/Misc/EmbeddedCommunication.h
--- a/Engine/Source/Runtime/Core/Public/Misc/EmbeddedCommunication.h
+++ b/Engine/Source/Runtime/Core/Public/Misc/EmbeddedCommunication.h
@@ -175,6 +175,9 @@ public:
 		return *this;
 	}
 
+	bool GetNeedsRendering() const { return bNeedsRendering; }
+	FName GetRequester() const { return Requester;  }
+
 private:
 	FName Requester;
 	bool bNeedsRendering = false;
