Fix CIS (App.h usage) (UE-79818).

diff --git a/Engine/Source/Runtime/Core/Private/HAL/MallocBinned2.cpp b/Engine/Source/Runtime/Core/Private/HAL/MallocBinned2.cpp
--- a/Engine/Source/Runtime/Core/Private/HAL/MallocBinned2.cpp
+++ b/Engine/Source/Runtime/Core/Private/HAL/MallocBinned2.cpp
@@ -13,6 +13,7 @@
 #include "HAL/IConsoleManager.h"
 #include "HAL/MemoryMisc.h"
 #include "HAL/PlatformMisc.h"
+#include "Misc/App.h"
 
 #if BINNED2_ALLOW_RUNTIME_TWEAKING
 
