make crash context creation less prone to low memory situation OOMs

diff --git a/Engine/Source/Runtime/Core/Private/IOS/IOSMallocZone.cpp b/Engine/Source/Runtime/Core/Private/IOS/IOSMallocZone.cpp
--- a/Engine/Source/Runtime/Core/Private/IOS/IOSMallocZone.cpp
+++ b/Engine/Source/Runtime/Core/Private/IOS/IOSMallocZone.cpp
@@ -93,6 +93,12 @@ void FIOSMallocCrashHandler::Enable( FIOSCrashContext* Context, uint32 CrashedTh
 	GMalloc = this;
 }
 	
+void FIOSMallocCrashHandler::SetContext( FIOSCrashContext* Context )
+{
+	check(Context);
+	CrashContext = Context;
+}
+
 void* FIOSMallocCrashHandler::Malloc( SIZE_T Size, uint32 Alignment )
 {
 	void* Result = nullptr;
diff --git a/Engine/Source/Runtime/Core/Private/IOS/IOSPlatformMisc.cpp b/Engine/Source/Runtime/Core/Private/IOS/IOSPlatformMisc.cpp
--- a/Engine/Source/Runtime/Core/Private/IOS/IOSPlatformMisc.cpp
+++ b/Engine/Source/Runtime/Core/Private/IOS/IOSPlatformMisc.cpp
@@ -1535,16 +1535,20 @@ static void DefaultCrashHandler(FIOSCrashContext const& Context)
 static uint32 GIOSStackIgnoreDepth = 6;
 
 // true system specific crash handler that gets called first
+static FIOSCrashContext TempCrashContext(ECrashContextType::Crash, TEXT("Temp Context"));
 static void PlatformCrashHandler(int32 Signal, siginfo_t* Info, void* Context)
 {
+	// switch to crash handler malloc to avoid malloc reentrancy
+	check(FIOSApplicationInfo::CrashMalloc);
+	FIOSApplicationInfo::CrashMalloc->Enable(&TempCrashContext, FPlatformTLS::GetCurrentThreadId());
+	
     FIOSCrashContext CrashContext(ECrashContextType::Crash, TEXT("Caught signal"));
     CrashContext.IgnoreDepth = GIOSStackIgnoreDepth;
     CrashContext.InitFromSignal(Signal, Info, Context);
-    
-    // switch to crash handler malloc to avoid malloc reentrancy
-    check(FIOSApplicationInfo::CrashMalloc);
-    FIOSApplicationInfo::CrashMalloc->Enable(&CrashContext, FPlatformTLS::GetCurrentThreadId());
-    
+	
+	// switch to the crash malloc to the new context now that we have everything
+	FIOSApplicationInfo::CrashMalloc->SetContext(&CrashContext);
+	
     if (GCrashHandlerPointer)
     {
         GCrashHandlerPointer(CrashContext);
diff --git a/Engine/Source/Runtime/Core/Public/IOS/IOSMallocZone.h b/Engine/Source/Runtime/Core/Public/IOS/IOSMallocZone.h
--- a/Engine/Source/Runtime/Core/Public/IOS/IOSMallocZone.h
+++ b/Engine/Source/Runtime/Core/Public/IOS/IOSMallocZone.h
@@ -47,6 +47,8 @@ public:
 	
 	void Enable( FIOSCrashContext* Context, uint32 CrashedThreadId );
 	
+	void SetContext( FIOSCrashContext* Context );
+	
 	// FMalloc interface.
 	virtual void* Malloc( SIZE_T Size, uint32 Alignment ) override;
 	
