Fix tracking of peak memory from Mac platform stats.

diff --git a/Engine/Source/Runtime/Core/Private/Mac/MacPlatformMemory.cpp b/Engine/Source/Runtime/Core/Private/Mac/MacPlatformMemory.cpp
--- a/Engine/Source/Runtime/Core/Private/Mac/MacPlatformMemory.cpp
+++ b/Engine/Source/Runtime/Core/Private/Mac/MacPlatformMemory.cpp
@@ -134,7 +134,7 @@ FPlatformMemoryStats FMacPlatformMemory::GetStats()
 {
 	const FPlatformMemoryConstants& MemoryConstants = FPlatformMemory::GetConstants();
 
-	FPlatformMemoryStats MemoryStats;
+	static FPlatformMemoryStats MemoryStats;
 
 	// Gather platform memory stats.
 	vm_statistics Stats;
