Add a FApp::GetBuildTargetType() utility method.

#rnx

diff --git a/Engine/Source/Programs/UnrealBuildTool/Configuration/TargetRules.cs b/Engine/Source/Programs/UnrealBuildTool/Configuration/TargetRules.cs
--- a/Engine/Source/Programs/UnrealBuildTool/Configuration/TargetRules.cs
+++ b/Engine/Source/Programs/UnrealBuildTool/Configuration/TargetRules.cs
@@ -1476,6 +1476,7 @@ namespace UnrealBuildTool
 			else if(Type == global::UnrealBuildTool.TargetType.Client)
 			{
 				GlobalDefinitions.Add("UE_GAME=1");
+				GlobalDefinitions.Add("UE_CLIENT=1");
 			}
 			else if(Type == global::UnrealBuildTool.TargetType.Editor)
 			{
diff --git a/Engine/Source/Runtime/Core/Core.Build.cs b/Engine/Source/Runtime/Core/Core.Build.cs
--- a/Engine/Source/Runtime/Core/Core.Build.cs
+++ b/Engine/Source/Runtime/Core/Core.Build.cs
@@ -199,10 +199,20 @@ public class Core : ModuleRules
             PublicDefinitions.Add("WITH_DIRECTXMATH=0");
         }
 
-        // Decide if validating memory allocator (aka MallocStomp) can be used on the current platform.
-        // Run-time validation must be enabled through '-stompmalloc' command line argument.
+		// Set a macro to allow FApp::GetBuildTargetType() to detect client targts
+		if(Target.Type == TargetRules.TargetType.Client)
+		{
+			PrivateDefinitions.Add("IS_CLIENT_TARGET=1");
+		}
+		else
+		{
+			PrivateDefinitions.Add("IS_CLIENT_TARGET=1");
+		}
+
+		// Decide if validating memory allocator (aka MallocStomp) can be used on the current platform.
+		// Run-time validation must be enabled through '-stompmalloc' command line argument.
 
-        bool bWithMallocStomp = false;
+		bool bWithMallocStomp = false;
         if (Target.Configuration != UnrealTargetConfiguration.Shipping)
         {
 			if (Target.Platform == UnrealTargetPlatform.Mac || Target.Platform == UnrealTargetPlatform.Linux || Target.Platform == UnrealTargetPlatform.Win64)
diff --git a/Engine/Source/Runtime/Core/Private/Misc/App.cpp b/Engine/Source/Runtime/Core/Private/Misc/App.cpp
--- a/Engine/Source/Runtime/Core/Private/Misc/App.cpp
+++ b/Engine/Source/Runtime/Core/Private/Misc/App.cpp
@@ -87,6 +87,23 @@ EBuildConfiguration FApp::GetBuildConfiguration()
 #endif
 }
 
+EBuildTargetType FApp::GetBuildTargetType()
+{
+#if IS_CLIENT_TARGET
+	return EBuildTargetType::Client;
+#elif UE_GAME
+	return EBuildTargetType::Game;
+#elif UE_EDITOR
+	return EBuildTargetType::Editor;
+#elif UE_SERVER
+	return EBuildTargetType::Server;
+#elif IS_PROGRAM
+	return EBuildTargetType::Program;
+#else
+	static_assert(false, "No host type is set.");
+#endif
+}
+
 #if UE_BUILD_DEVELOPMENT
 void FApp::SetDebugGame(bool bInIsDebugGame)
 {
diff --git a/Engine/Source/Runtime/Core/Public/Misc/App.h b/Engine/Source/Runtime/Core/Public/Misc/App.h
--- a/Engine/Source/Runtime/Core/Public/Misc/App.h
+++ b/Engine/Source/Runtime/Core/Public/Misc/App.h
@@ -41,6 +41,13 @@ public:
 	 */
 	static EBuildConfiguration GetBuildConfiguration();
 
+	/**
+	 * Gets the target type of the current application (eg. client, server, etc...)
+	 *
+	 * @return The build target type
+	 */
+	static EBuildTargetType GetBuildTargetType();
+
 #if UE_BUILD_DEVELOPMENT
 	/**
 	 * For development configurations, sets whether the application should load DebugGame game modules.
