fix for app flavor symbol naming

diff --git a/Engine/Source/Runtime/Core/Private/GenericPlatform/GenericPlatformCrashContext.cpp b/Engine/Source/Runtime/Core/Private/GenericPlatform/GenericPlatformCrashContext.cpp
--- a/Engine/Source/Runtime/Core/Private/GenericPlatform/GenericPlatformCrashContext.cpp
+++ b/Engine/Source/Runtime/Core/Private/GenericPlatform/GenericPlatformCrashContext.cpp
@@ -325,6 +325,9 @@ void FGenericCrashContext::SerializeContentToBuffer() const
 #ifdef UE_BUILD_FLAVOR
 	Symbols = FString::Printf(TEXT( "%s-%s" ), *Symbols, *FString(UE_BUILD_FLAVOR));
 #endif
+#ifdef UE_APP_FLAVOR
+	Symbols = FString::Printf(TEXT( "%s-%s" ), *Symbols, *FString(UE_APP_FLAVOR));
+#endif
 
 	AddCrashProperty( TEXT( "Symbols" ), Symbols);
 
