`FThread` as a simpler API for system threads compared with FRunnable/FRunnableThread.
The API modeled after `std::thread`.
The implementation is based on `FRunnableThread`.

diff --git a/Engine/Source/Runtime/Core/Private/HAL/Thread.cpp b/Engine/Source/Runtime/Core/Private/HAL/Thread.cpp
new file mode 100644
--- /dev/null
+++ b/Engine/Source/Runtime/Core/Private/HAL/Thread.cpp
@@ -0,0 +1,127 @@
+// Copyright 1998-2018 Epic Games, Inc. All Rights Reserved.
+
+#include "HAL/Thread.h"
+#include "HAL/Runnable.h"
+#include "HAL/RunnableThread.h"
+#include "HAL/PlatformTLS.h"
+#include "HAL/ThreadSafeBool.h"
+#include "Misc/AssertionMacros.h"
+#include "Templates/UnrealTemplate.h"
+#include "Templates/UniquePtr.h"
+
+class FThreadImpl final : public FRunnable
+{
+public:
+	FThreadImpl(
+		TCHAR const* ThreadName,
+		TUniqueFunction<void()>&& InThreadFunction,
+		uint32 StackSize,
+		EThreadPriority ThreadPriority,
+		uint64 ThreadAffinityMask
+	) 
+		: ThreadFunction(MoveTemp(InThreadFunction))
+		, RunnableThread(FRunnableThread::Create(this, ThreadName, StackSize, ThreadPriority, ThreadAffinityMask))
+	{
+		check(IsJoinable());
+	}
+
+	// Provides a reference to self. This can't be done by `SharedFromThis` in the constructor because the instance is not constructed yet
+	void Initialize(const TSharedPtr<FThreadImpl, ESPMode::ThreadSafe>& InSelf)
+	{
+		Self = InSelf;
+	}
+
+	bool IsJoinable() const
+	{
+		return RunnableThread.IsValid() && FPlatformTLS::GetCurrentThreadId() != RunnableThread->GetThreadID();
+	}
+
+	void Join()
+	{
+		check(IsJoinable());
+
+		RunnableThread->WaitForCompletion();
+		RunnableThread.Reset();
+	}
+
+	uint32 GetThreadId() const
+	{
+		return RunnableThread.IsValid() ? RunnableThread->GetThreadID() : FThread::InvalidThreadId;
+	}
+
+private:
+	uint32 Run() override
+	{
+		// the thread can be started before `RunnableThread` member is initialized
+		//check(RunnableThread.IsValid());
+
+		ThreadFunction();
+
+		while (!Self.IsValid());
+
+		// we're about to exit the thread. Release the reference to self. If the thread is detached, it's the only reference
+		// and so this instance will be deleted. No member access should be performed after this point.
+		Self.Reset();
+
+		return 0;
+	}
+
+private:
+	// Two strong references are hold for `FThreadImpl`, one in parent `FThread` and another here. The reference in `FThreadImpl` is released
+	// on `Detach`, as to physically detach `FThreadImpl` from `FThread`. The reference below (`Self`) is released before exiting the thread,
+	// as the work was done and `FThreadImpl` doesn't mind to be deleted. Releasing the last reference deletes the instance, so 
+	// no member access can be performed after that.
+	// This must be the declared before `RunnableThread` so it's already initialized when the thread is created, otherwise the thread can complete
+	// before `Self` is initialized.
+	TSharedPtr<FThreadImpl, ESPMode::ThreadSafe> Self;
+	
+	TUniqueFunction<void()> ThreadFunction;
+	TUniquePtr<FRunnableThread> RunnableThread;
+};
+
+FThread::FThread(
+	TCHAR const* ThreadName,
+	TUniqueFunction<void()>&& ThreadFunction,
+	uint32 StackSize,
+	EThreadPriority ThreadPriority,
+	uint64 ThreadAffinityMask
+)
+	: Impl(MakeShared<FThreadImpl, ESPMode::ThreadSafe>(ThreadName, MoveTemp(ThreadFunction), StackSize, ThreadPriority, ThreadAffinityMask))
+{
+	Impl->Initialize(Impl);
+}
+
+FThread& FThread::operator=(FThread&& Other)
+{
+	checkf(!IsJoinable(), TEXT("Joinable thread cannot be assigned"));
+	Impl = MoveTemp(Other.Impl);
+	return *this;
+}
+
+FThread::~FThread()
+{
+	checkf(!Impl.IsValid(), TEXT("FThread must be either joined or detached before destruction"));
+}
+
+bool FThread::IsJoinable() const
+{
+	return Impl.IsValid() && Impl->IsJoinable();
+}
+
+void FThread::Join()
+{
+	check(Impl.IsValid());
+	Impl->Join();
+	Impl.Reset(); // no need to keep this memory anymore (it's the last reference), and this also simplifies the check in the destructor
+}
+
+uint32 FThread::GetThreadId() const
+{
+	return Impl.IsValid() ? Impl->GetThreadId() : FThread::InvalidThreadId;
+}
+
+void FThread::Detach()
+{
+	check(Impl.IsValid());
+	Impl.Reset(); // releases this reference, `FThreadImpl` instance can still be alive until the thread finishes its execution
+}
diff --git a/Engine/Source/Runtime/Core/Private/Tests/HAL/ThreadTest.cpp b/Engine/Source/Runtime/Core/Private/Tests/HAL/ThreadTest.cpp
new file mode 100644
--- /dev/null
+++ b/Engine/Source/Runtime/Core/Private/Tests/HAL/ThreadTest.cpp
@@ -0,0 +1,181 @@
+#include "CoreTypes.h"
+#include "Misc/AssertionMacros.h"
+#include "Misc/AutomationTest.h"
+
+#if WITH_DEV_AUTOMATION_TESTS
+
+#include "HAL/Thread.h"
+#include "HAL/ThreadSafeBool.h"
+#include "HAL/ThreadSafeCounter.h"
+#include "HAL/Event.h"
+#include "Containers/Queue.h"
+
+IMPLEMENT_SIMPLE_AUTOMATION_TEST(FThreadTest, "System.Core.HAL.Thread", EAutomationTestFlags::ApplicationContextMask | EAutomationTestFlags::EngineFilter)
+
+namespace
+{
+	void TestIsJoinableAfterCreation(FThreadTest& This)
+	{
+		FThread Thread(TEXT("Test.Thread"), []() { /*NOOP*/ });
+		This.TestTrue(TEXT("FThread must be joinable after construction"), Thread.IsJoinable());
+		Thread.Join();
+	}
+
+	void TestIsJoinableAfterCompletion(FThreadTest& This)
+	{
+		FThreadSafeBool bDone = false;
+		FThread Thread(TEXT("Test.Thread"), [&bDone]() { bDone = true; });
+		while (!bDone); // wait for completion
+		This.TestTrue(TEXT("FThread must still be joinable after completion"), Thread.IsJoinable());
+		Thread.Join();
+	}
+
+	void TestIsNotJoinableAfterJoining(FThreadTest& This)
+	{
+		FThread Thread(TEXT("Test.Thread"), []() { /*NOOP*/ });
+		Thread.Join();
+		This.TestFalse(TEXT("FThread must not be joinable after joining"), Thread.IsJoinable());
+	}
+
+	void TestIsNotJoinableAfterDetaching(FThreadTest& This)
+	{
+		FThread Thread(TEXT("Test.Thread"), []() { /*NOOP*/ });
+		Thread.Detach();
+		This.TestFalse(TEXT("FThread must not be joinable after detaching"), Thread.IsJoinable());
+	}
+
+	void TestAssertIfNotJoinedOrDetached(FThreadTest& This)
+	{
+		// this does fails the `check`, but it seems there's no way to test this by UE4 Automation Testing, so commented out
+		//FThread Thread(TEXT("Test.Thread"), []() { /*NOOP*/ });
+		// should assert in the destructor
+	}
+
+	void TestDefaultConstruction(FThreadTest& This)
+	{
+		{
+			FThread Thread;
+			This.TestFalse(TEXT("Default-constructed FThread must be not joinable"), Thread.IsJoinable());
+		}
+		{	// check that default constructed thread can be "upgraded" to joinable thread
+			FThread Thread;
+			Thread = FThread(TEXT("Test.Thread"), []() { /* NOOP */ });
+			This.TestTrue(TEXT("Move-constructed FThread from joinable thread must be joinable"), Thread.IsJoinable());
+			Thread.Join();
+		}
+	}
+
+	void TestMovability(FThreadTest& This)
+	{
+		{	// move constructor with default-constructed thread
+			FThread Src;
+			FThread Dst(MoveTemp(Src));
+			This.TestFalse(TEXT("Default-constructed thread must stay not joinable after moving out"), Src.IsJoinable());
+			This.TestFalse(TEXT("Move-constructed thread from not joinable thread must be not joinable"), Dst.IsJoinable());
+		}
+		{	// move constructor with joinable thread
+			FThread Src(TEXT("Test.Thread"), []() { /* NOOP */ });
+			FThread Dst(MoveTemp(Src));
+			This.TestFalse(TEXT("Moved out thread must be not joinable"), Src.IsJoinable());
+			This.TestTrue(TEXT("Move-constructed thread from joinable thread must be joinable"), Dst.IsJoinable());
+			Dst.Join();
+		}
+		{	// move assignment operator
+			FThread Src(TEXT("Test.Thread"), []() { /* NOOP */ });
+			FThread Dst;
+			Dst = MoveTemp(Src);
+			This.TestFalse(TEXT("Moved out thread must be not joinable"), Src.IsJoinable());
+			This.TestTrue(TEXT("Move-assigned thread from joinable thread must be joinable"), Dst.IsJoinable());
+			Dst.Join();
+		}
+		{	// Failure test for move assignment operator of joinable thread
+			//FThread Src(TEXT("Test.Thread"), []() { /* NOOP */ });
+			//FThread Dst(TEXT("Test.Thread"), []() { /* NOOP */ });
+			//Dst = MoveTemp(Src); // must assert that joinable thread wasn't joined before move-assignment, no way to test this
+			//Dst.Join();
+		}
+		{	// Move assignment operator of thread that has been joined
+			FThread Src(TEXT("Test.Thread"), []() { /* NOOP */ });
+			FThread Dst(TEXT("Test.Thread"), []() { /* NOOP */ });
+			Dst.Join();
+			Dst = MoveTemp(Src);
+			Dst.Join();
+		}
+	}
+
+	// An example of possible implementation of Consumer/Producer idiom
+	void TestTypicalUseCase(FThreadTest& This)
+	{
+		FThreadSafeBool bQuitRequested = false;
+		using FWork = uint32;
+		TQueue<FWork> WorkQueue;
+		FEvent* WorkQueuedEvent = FPlatformProcess::GetSynchEventFromPool();
+
+		FThread WorkerThread(TEXT("Test.Thread"), [&bQuitRequested, &WorkQueue, WorkQueuedEvent]() 
+		{
+			while (!bQuitRequested)
+			{
+				// get work
+				FWork Work;
+				if (!WorkQueue.Dequeue(Work))
+				{
+					WorkQueuedEvent->Wait();
+					// first check if quitting requested
+					continue;
+				}
+
+				// do work
+				UE_LOG(LogTemp, Log, TEXT("Work #%d consumed"), Work);
+			}
+
+			UE_LOG(LogTemp, Log, TEXT("Quit"));
+		});
+
+		// produce work
+		const int WorkNum = 3;
+		for (FWork Work = 0; Work != WorkNum; ++Work)
+		{
+			WorkQueue.Enqueue(Work);
+			WorkQueuedEvent->Trigger();
+			UE_LOG(LogTemp, Log, TEXT("Work #%d produced"), Work);
+		}
+
+		UE_LOG(LogTemp, Log, TEXT("Request to quit"));
+		bQuitRequested = true;
+		// the thread can be blocked waiting for work, unblock it
+		WorkQueuedEvent->Trigger();
+		WorkerThread.Join();
+
+		FPlatformProcess::ReturnSynchEventToPool(WorkQueuedEvent);
+
+		// example of output:
+		// Work #0 produced
+		//	Work #0 consumed
+		//	Work #1 produced
+		//	Work #1 consumed
+		//	Work #2 produced
+		//	Work #2 consumed
+		//	Request to quit
+		//	The thread 0x96e0 has exited with code 0 (0x0).
+		//	Quit
+	}
+}
+
+bool FThreadTest::RunTest(const FString& Parameters)
+{
+	TestIsJoinableAfterCreation(*this);
+	TestIsJoinableAfterCompletion(*this);
+	TestIsNotJoinableAfterJoining(*this);
+	TestIsNotJoinableAfterDetaching(*this);
+
+	TestAssertIfNotJoinedOrDetached(*this);
+
+	TestDefaultConstruction(*this);
+	TestMovability(*this);
+
+	TestTypicalUseCase(*this);
+
+	return true;
+}
+
+#endif //WITH_DEV_AUTOMATION_TESTS
diff --git a/Engine/Source/Runtime/Core/Public/HAL/Thread.h b/Engine/Source/Runtime/Core/Public/HAL/Thread.h
new file mode 100644
--- /dev/null
+++ b/Engine/Source/Runtime/Core/Public/HAL/Thread.h
@@ -0,0 +1,97 @@
+// Copyright 1998-2018 Epic Games, Inc. All Rights Reserved.
+
+#pragma once
+
+#include "HAL/Platform.h"
+#include "HAL/PlatformAffinity.h"
+#include "Templates/Function.h"
+#include "Templates/SharedPointer.h"
+
+/**
+ * Simple API for system threads. 
+ * Before using, please make sure you really need a new system thread. By default and in the majority of cases
+ * parallel processing should be done by TaskGraph.
+ * For richer functionality check `FRunnable`/`FRunnableThread`.
+ * It's up to user to provide a way to notify the thread function to exit on demand.
+ * Before destroying the instance it must be either `Join`ed or `Detach`ed.
+ * Example:
+ *		FThread Thread{TEXT("New thread"), []() { do_something_important(); }};
+ *		// ... continue in the caller thread
+ *		Thread.Join();
+ * For more verbose example check `TestTypicalUseCase` in `ThreadTest.cpp`
+ */
+class CORE_API FThread final
+{
+public:
+	/**
+	 * Creates new "empty" thread object that doesn't represent a system thread
+	 */
+	FThread()
+	{}
+
+	/**
+	 * Creates and immediately starts a new system thread that will execute `ThreadFunction` argument.
+	 * Can return before the thread is actually started or when it already finished execution.
+	 * @param ThreadName Name of the thread
+	 * @param ThreadFunction The function that will be executed by the newly created thread
+	 * @param StackSize The size of the stack to create. 0 means use the current thread's stack size
+	 * @param ThreadPriority Tells the thread whether it needs to adjust its priority or not. Defaults to normal priority
+	 */
+	FThread(
+		TCHAR const* ThreadName,
+		TUniqueFunction<void()>&& ThreadFunction,
+		uint32 StackSize = 0,
+		EThreadPriority ThreadPriority = TPri_Normal,
+		uint64 ThreadAffinityMask = FPlatformAffinity::GetNoAffinityMask()
+	);
+
+	// non-copyable
+	FThread(const FThread&) = delete;
+	FThread& operator=(const FThread&) = delete;
+
+	FThread(FThread&&) = default;
+	/**
+	 * Move assignment operator.
+	 * Asserts if the instance is joinable.
+	 */
+	FThread& operator=(FThread&& Other);
+
+	/**
+	 * Destructor asserts if the instance is not joined or detached.
+	 */
+	~FThread();
+
+	/**
+	 * Checks if the thread object identifies an active thread of execution. 
+	 * A thread that has finished executing code, but has not yet been joined is still considered an active 
+	 * thread of execution and is therefore joinable.
+	 * @see Join
+	 */
+	bool IsJoinable() const;
+
+	/**
+	 * Blocks the current thread until the thread identified by `this` finishes its execution.
+	 * The completion of the thread identified by `this` synchronizes with the corresponding successful return from Join().
+	 * No synchronization is performed on `this` itself. Concurrently calling Join() on the same FThread object 
+	 * from multiple threads constitutes a data race that results in undefined behavior.
+	 * @see IsJoinable
+	 */
+	void Join();
+
+	static constexpr uint32 InvalidThreadId = static_cast<uint32>(-1);
+
+	/**
+	 * @return Thread ID for this thread
+	 */
+	uint32 GetThreadId() const;
+
+	/**
+	 * Separates the thread of execution from the thread object, allowing execution to continue independently. 
+	 * Any allocated resources will be freed once the thread exits.
+	 * After calling detach `this` no longer owns any thread.
+	 */
+	void Detach();
+
+private:
+	TSharedPtr<class FThreadImpl, ESPMode::ThreadSafe> Impl; // "shared" with `FThreadImpl::Self`
+};
