added/fixed copyright notice

diff --git a/Engine/Source/Runtime/Core/Private/HAL/Thread.cpp b/Engine/Source/Runtime/Core/Private/HAL/Thread.cpp
--- a/Engine/Source/Runtime/Core/Private/HAL/Thread.cpp
+++ b/Engine/Source/Runtime/Core/Private/HAL/Thread.cpp
@@ -1,4 +1,4 @@
-// Copyright 1998-2018 Epic Games, Inc. All Rights Reserved.
+// Copyright 1998-2019 Epic Games, Inc. All Rights Reserved.
 
 #include "HAL/Thread.h"
 #include "HAL/Runnable.h"
diff --git a/Engine/Source/Runtime/Core/Private/Tests/HAL/ThreadTest.cpp b/Engine/Source/Runtime/Core/Private/Tests/HAL/ThreadTest.cpp
--- a/Engine/Source/Runtime/Core/Private/Tests/HAL/ThreadTest.cpp
+++ b/Engine/Source/Runtime/Core/Private/Tests/HAL/ThreadTest.cpp
@@ -1,3 +1,5 @@
+// Copyright 1998-2019 Epic Games, Inc. All Rights Reserved.
+
 #include "CoreTypes.h"
 #include "Misc/AssertionMacros.h"
 #include "Misc/AutomationTest.h"
diff --git a/Engine/Source/Runtime/Core/Public/HAL/Thread.h b/Engine/Source/Runtime/Core/Public/HAL/Thread.h
--- a/Engine/Source/Runtime/Core/Public/HAL/Thread.h
+++ b/Engine/Source/Runtime/Core/Public/HAL/Thread.h
@@ -1,4 +1,4 @@
-// Copyright 1998-2018 Epic Games, Inc. All Rights Reserved.
+// Copyright 1998-2019 Epic Games, Inc. All Rights Reserved.
 
 #pragma once
 
