disabled false static analysis warning

diff --git a/Engine/Source/Runtime/Core/Private/Tests/HAL/ThreadTest.cpp b/Engine/Source/Runtime/Core/Private/Tests/HAL/ThreadTest.cpp
--- a/Engine/Source/Runtime/Core/Private/Tests/HAL/ThreadTest.cpp
+++ b/Engine/Source/Runtime/Core/Private/Tests/HAL/ThreadTest.cpp
@@ -27,7 +27,7 @@ namespace
 	{
 		FThreadSafeBool bDone = false;
 		FThread Thread(TEXT("Test.Thread"), [&bDone]() { bDone = true; });
-		while (!bDone); // wait for completion
+		while (!bDone); // wait for completion //-V529
 		This.TestTrue(TEXT("FThread must still be joinable after completion"), Thread.IsJoinable());
 		Thread.Join();
 	}
