Fixed error message in CharTest

diff --git a/Engine/Source/Runtime/Core/Private/Tests/Misc/CharTest.cpp b/Engine/Source/Runtime/Core/Private/Tests/Misc/CharTest.cpp
--- a/Engine/Source/Runtime/Core/Private/Tests/Misc/CharTest.cpp
+++ b/Engine/Source/Runtime/Core/Private/Tests/Misc/CharTest.cpp
@@ -35,7 +35,7 @@ bool TCharTest::RunTest(const FString& Parameters)
 	const char* CurrentLocale = setlocale(LC_CTYPE, nullptr);
 	if (strcmp("C", CurrentLocale))
 	{
-		AddError(FString::Printf(TEXT("Locale is \"%s\" but should be \"C\". Did something call setlocale()?"), CurrentLocale));
+		AddError(FString::Printf(TEXT("Locale is \"%s\" but should be \"C\". Did something call setlocale()?"), ANSI_TO_TCHAR(CurrentLocale)));
 	}
 	else
 	{
