Move the Build.h include into Platform.h to prevent non-unity compile errors in Thread.h when HAL/Platform.h is included first.

diff --git a/Engine/Source/Runtime/Core/Public/CoreTypes.h b/Engine/Source/Runtime/Core/Public/CoreTypes.h
--- a/Engine/Source/Runtime/Core/Public/CoreTypes.h
+++ b/Engine/Source/Runtime/Core/Public/CoreTypes.h
@@ -6,7 +6,6 @@
 	Low level includes.
 ----------------------------------------------------------------------------*/
 
-#include "Misc/Build.h"
 #include "HAL/Platform.h"
 #include "ProfilingDebugging/UMemoryDefines.h"
 #include "Misc/CoreMiscDefines.h"
diff --git a/Engine/Source/Runtime/Core/Public/HAL/Platform.h b/Engine/Source/Runtime/Core/Public/HAL/Platform.h
--- a/Engine/Source/Runtime/Core/Public/HAL/Platform.h
+++ b/Engine/Source/Runtime/Core/Public/HAL/Platform.h
@@ -2,6 +2,8 @@
 
 #pragma once
 
+#include "Misc/Build.h"
+
 // define all other platforms to be zero
 //@port Define the platform here to be zero when compiling for other platforms
 #if !defined(PLATFORM_WINDOWS)
