fixed up case in which command line was queried before being initialized

diff --git a/Engine/Source/Runtime/Core/Private/GenericPlatform/GenericPlatformMisc.cpp b/Engine/Source/Runtime/Core/Private/GenericPlatform/GenericPlatformMisc.cpp
--- a/Engine/Source/Runtime/Core/Private/GenericPlatform/GenericPlatformMisc.cpp
+++ b/Engine/Source/Runtime/Core/Private/GenericPlatform/GenericPlatformMisc.cpp
@@ -1183,7 +1183,7 @@ void FGenericPlatformMisc::UpdateHotfixableEnsureSettings()
 	else
 	{
 		float HandleEnsurePercentOnCmdLine = 100.0f;
-		if (FParse::Value(FCommandLine::Get(), TEXT("handleensurepercent="), HandleEnsurePercentOnCmdLine))
+		if (!FCommandLine::IsInitialized() && FParse::Value(FCommandLine::Get(), TEXT("handleensurepercent="), HandleEnsurePercentOnCmdLine))
 		{
 			GenericPlatformMisc::GEnsureChance = HandleEnsurePercentOnCmdLine / 100.0;
 		}
