missed a cosmetic improvement suggested during code review: `static_cast<uint32>(-1)` -> `~uint32(0)`

diff --git a/Engine/Source/Runtime/Core/Public/HAL/Thread.h b/Engine/Source/Runtime/Core/Public/HAL/Thread.h
--- a/Engine/Source/Runtime/Core/Public/HAL/Thread.h
+++ b/Engine/Source/Runtime/Core/Public/HAL/Thread.h
@@ -78,7 +78,7 @@ public:
 	 */
 	void Join();
 
-	static constexpr uint32 InvalidThreadId = static_cast<uint32>(-1);
+	static constexpr uint32 InvalidThreadId = ~uint32(0);
 
 	/**
 	 * @return Thread ID for this thread
