AddAdditionalRootDirectory now enforces a slash at the end.

diff --git a/Engine/Source/Runtime/Core/Private/GenericPlatform/GenericPlatformMisc.cpp b/Engine/Source/Runtime/Core/Private/GenericPlatform/GenericPlatformMisc.cpp
--- a/Engine/Source/Runtime/Core/Private/GenericPlatform/GenericPlatformMisc.cpp
+++ b/Engine/Source/Runtime/Core/Private/GenericPlatform/GenericPlatformMisc.cpp
@@ -807,7 +807,12 @@ const TArray<FString>& FGenericPlatformMisc::GetAdditionalRootDirectories()
 void FGenericPlatformMisc::AddAdditionalRootDirectory(const FString& RootDir)
 {
 	TArray<FString>& RootDirectories = TLazySingleton<FStaticData>::Get().AdditionalRootDirectories;
-	RootDirectories.Add(RootDir);
+	FString NewRootDirectory = RootDir;
+	if (!NewRootDirectory.EndsWith(GetDefaultPathSeparator()))
+	{
+		NewRootDirectory = NewRootDirectory + GetDefaultPathSeparator();
+	}
+	RootDirectories.Add(NewRootDirectory);
 }
 
 const TCHAR* FGenericPlatformMisc::EngineDir()
diff --git a/Engine/Source/Runtime/Core/Private/IOS/IOSPlatformFile.cpp b/Engine/Source/Runtime/Core/Private/IOS/IOSPlatformFile.cpp
--- a/Engine/Source/Runtime/Core/Private/IOS/IOSPlatformFile.cpp
+++ b/Engine/Source/Runtime/Core/Private/IOS/IOSPlatformFile.cpp
@@ -10,6 +10,7 @@
 #include "Misc/CoreMisc.h"
 #include "Misc/CommandLine.h"
 #include "Misc/App.h"
+#include "Misc/Paths.h"
 
 #include "Async/MappedFileHandle.h"
 #include <sys/mman.h>
@@ -1062,10 +1063,10 @@ FString FIOSPlatformFile::ConvertToIOSPath(const FString& Filename, bool bForWri
         AdditionalRootDirectory.ReplaceInline(TEXT(".."), TEXT(""));
 		if (Result.StartsWith(AdditionalRootDirectory))
 		{
-			static FString ReadPathBase = FString([NSSearchPathForDirectoriesInDomains(NSDocumentDirectory, NSUserDomainMask, YES) objectAtIndex:0]) + TEXT("/");
-			
-			Result = ReadPathBase + Result.Mid(0,AdditionalRootDirectory.Len()) + TEXT("/") + Result.Mid(AdditionalRootDirectory.Len()+1).ToLower();
-			
+			static FString ReadPathBase = FString([NSSearchPathForDirectoriesInDomains(NSDocumentDirectory, NSUserDomainMask, YES) objectAtIndex:0]);
+
+			// lowercase the second half of the path because ios
+			Result = FPaths::Combine(ReadPathBase, Result.Mid(0,AdditionalRootDirectory.Len()), Result.Mid(AdditionalRootDirectory.Len()+1).ToLower());
 			return Result;
         }
 	}
