Fix FIOSPlatformFile::ConvertToIOSPath not handling paths without trailing slashes.

diff --git a/Engine/Source/Runtime/Core/Private/GenericPlatform/GenericPlatformMisc.cpp b/Engine/Source/Runtime/Core/Private/GenericPlatform/GenericPlatformMisc.cpp
--- a/Engine/Source/Runtime/Core/Private/GenericPlatform/GenericPlatformMisc.cpp
+++ b/Engine/Source/Runtime/Core/Private/GenericPlatform/GenericPlatformMisc.cpp
@@ -808,10 +808,7 @@ void FGenericPlatformMisc::AddAdditionalRootDirectory(const FString& RootDir)
 {
 	TArray<FString>& RootDirectories = TLazySingleton<FStaticData>::Get().AdditionalRootDirectories;
 	FString NewRootDirectory = RootDir;
-	if (!NewRootDirectory.EndsWith(GetDefaultPathSeparator()))
-	{
-		NewRootDirectory = NewRootDirectory + GetDefaultPathSeparator();
-	}
+	FPaths::MakePlatformFilename(NewRootDirectory);
 	RootDirectories.Add(NewRootDirectory);
 }
 
@@ -1409,4 +1406,4 @@ void FGenericPlatformMisc::ParseChunkIdPakchunkIndexMapping(TArray<FString> Chun
 			OutMapping.Add(ChunkId, PakchunkIndex);
 		}
 	}
-}
\ No newline at end of file
+}
diff --git a/Engine/Source/Runtime/Core/Private/IOS/IOSPlatformFile.cpp b/Engine/Source/Runtime/Core/Private/IOS/IOSPlatformFile.cpp
--- a/Engine/Source/Runtime/Core/Private/IOS/IOSPlatformFile.cpp
+++ b/Engine/Source/Runtime/Core/Private/IOS/IOSPlatformFile.cpp
@@ -1052,7 +1052,9 @@ FString FIOSPlatformFile::ConvertToIOSPath(const FString& Filename, bool bForWri
 	{
 		return Result;
 	}
-    
+	
+	FPaths::MakePlatformFilename(Result);
+	
 	Result.ReplaceInline(TEXT("../"), TEXT(""));
 	Result.ReplaceInline(TEXT(".."), TEXT(""));
 	Result.ReplaceInline(FPlatformProcess::BaseDir(), TEXT(""));
@@ -1061,7 +1063,8 @@ FString FIOSPlatformFile::ConvertToIOSPath(const FString& Filename, bool bForWri
 	{
         AdditionalRootDirectory.ReplaceInline(TEXT("../"), TEXT(""));
         AdditionalRootDirectory.ReplaceInline(TEXT(".."), TEXT(""));
-		if (Result.StartsWith(AdditionalRootDirectory))
+		if (Result.StartsWith(AdditionalRootDirectory) &&
+			(Result.Len() == AdditionalRootDirectory.Len() || Result.Mid(AdditionalRootDirectory.Len()).StartsWith(FPlatformMisc::GetDefaultPathSeparator())))
 		{
 			static FString ReadPathBase = FString([NSSearchPathForDirectoriesInDomains(NSDocumentDirectory, NSUserDomainMask, YES) objectAtIndex:0]);
 
