ensure there are no null llm tags. hopefully temporary until a good repro is found for the null tag that's occuring in CIS.

diff --git a/Engine/Source/Runtime/Core/Private/HAL/LowLevelMemTracker.cpp b/Engine/Source/Runtime/Core/Private/HAL/LowLevelMemTracker.cpp
--- a/Engine/Source/Runtime/Core/Private/HAL/LowLevelMemTracker.cpp
+++ b/Engine/Source/Runtime/Core/Private/HAL/LowLevelMemTracker.cpp
@@ -481,6 +481,8 @@ bool FLowLevelMemTracker::IsEnabled()
 FLowLevelMemTracker* FLowLevelMemTracker::TrackerInstance;
 bool FLowLevelMemTracker::bIsDisabled; // must start off enabled because allocations happen before the command line enables/disables us
 
+static const TCHAR* InvalidLLMTagName = TEXT("?");
+
 FLowLevelMemTracker::FLowLevelMemTracker()
 	: bFirstTimeUpdating(true)
 	, bCanEnable(true)
@@ -507,7 +509,7 @@ FLowLevelMemTracker::FLowLevelMemTracker()
 
 	for (int32 Index = 0; Index < LLM_CUSTOM_TAG_COUNT; Index++ )
 	{
-		CustomTags[Index].Name = nullptr;
+		CustomTags[Index].Name = InvalidLLMTagName;
 	}
 
 	for (int32 Index = 0; Index < LLM_TAG_COUNT; Index++ )
@@ -904,7 +906,7 @@ bool FLowLevelMemTracker::FindTagByName( const TCHAR* Name, uint64& OutTag ) con
 		}
 		for ( int32 PlatformTagIndex = 0; PlatformTagIndex < LLM_CUSTOM_TAG_COUNT; PlatformTagIndex++ )
 		{
-			if( CustomTags[PlatformTagIndex].Name != nullptr && FCString::Stricmp( Name, CustomTags[PlatformTagIndex].Name ) == 0 )
+			if( CustomTags[PlatformTagIndex].Name != nullptr && CustomTags[PlatformTagIndex].Name != InvalidLLMTagName && FCString::Stricmp( Name, CustomTags[PlatformTagIndex].Name ) == 0 )
 			{
 				OutTag = LLM_CUSTOM_TAG_START + PlatformTagIndex;
 				return true;
