Fix value of IS_CLIENT_TARGET flag.

#rnx

diff --git a/Engine/Source/Runtime/Core/Core.Build.cs b/Engine/Source/Runtime/Core/Core.Build.cs
--- a/Engine/Source/Runtime/Core/Core.Build.cs
+++ b/Engine/Source/Runtime/Core/Core.Build.cs
@@ -206,7 +206,7 @@ public class Core : ModuleRules
 		}
 		else
 		{
-			PrivateDefinitions.Add("IS_CLIENT_TARGET=1");
+			PrivateDefinitions.Add("IS_CLIENT_TARGET=0");
 		}
 
 		// Decide if validating memory allocator (aka MallocStomp) can be used on the current platform.
