Fix FTransform::InverseFast vectorized version: since W wasn't used through the whole computation, the compiler somehow deduced that it could ignore setting the W component to zero (most probably because fpu exceptions are turned off).
Force W to 1 before going through GetSafeScaleReciprocal to avoid W=inf after, in the case where FPU exceptions are accidentally turned on by some other code.

#rnx

diff --git a/Engine/Source/Runtime/Core/Public/Math/TransformVectorized.h b/Engine/Source/Runtime/Core/Public/Math/TransformVectorized.h
--- a/Engine/Source/Runtime/Core/Public/Math/TransformVectorized.h
+++ b/Engine/Source/Runtime/Core/Public/Math/TransformVectorized.h
@@ -1373,7 +1373,7 @@ private:
 		checkSlow(VectorAnyGreaterThan(VectorAbs(Scale3D), GlobalVectorConstants::SmallNumber));
 
 		// Invert the scale
-		const VectorRegister InvScale = VectorSet_W0(GetSafeScaleReciprocal(Scale3D , ScalarRegister(GlobalVectorConstants::SmallNumber)));
+		const VectorRegister InvScale = VectorSet_W0(GetSafeScaleReciprocal(VectorSet_W1(Scale3D), ScalarRegister(GlobalVectorConstants::SmallNumber)));
 
 		// Invert the rotation
 		const VectorRegister InvRotation = VectorQuaternionInverse(Rotation);
