Implementing Nan/Inf check for Neon SIMD
#rnx

diff --git a/Engine/Source/Runtime/Core/Public/Math/UnrealMathNeon.h b/Engine/Source/Runtime/Core/Public/Math/UnrealMathNeon.h
--- a/Engine/Source/Runtime/Core/Public/Math/UnrealMathNeon.h
+++ b/Engine/Source/Runtime/Core/Public/Math/UnrealMathNeon.h
@@ -1333,8 +1333,23 @@ FORCEINLINE void VectorSinCos(  VectorRegister* VSinAngles, VectorRegister* VCos
 // Returns true if the vector contains a component that is either NAN or +/-infinite.
 inline bool VectorContainsNaNOrInfinite(const VectorRegister& Vec)
 {
-	checkf(false, TEXT("Not implemented for NEON")); //@TODO: Implement this method for NEON
-	return false;
+	// https://en.wikipedia.org/wiki/IEEE_754-1985
+	// Infinity is represented with all exponent bits set, with the correct sign bit.
+	// NaN is represented with all exponent bits set, plus at least one fraction/significant bit set.
+	// This means finite values will not have all exponent bits set, so check against those bits.
+
+	union { float F; uint32 U; } InfUnion;
+	InfUnion.U = 0x7F800000;
+	const float Inf = InfUnion.F;
+	const VectorRegister FloatInfinity = MakeVectorRegister(Inf, Inf, Inf, Inf);
+
+	// Mask off Exponent
+	VectorRegister ExpTest = VectorBitwiseAnd(Vec, FloatInfinity);
+	// Compare to full exponent & combine resulting flags into lane 0
+	static const int8x16_t Table = {0,4,8,12, 0,0,0,0, 0,0,0,0, 0,0,0,0};
+	uint8x16_t res = (uint8x16_t)VectorCompareEQ(ExpTest, FloatInfinity);
+	// If we have all zeros, all elements are finite
+	return vgetq_lane_u32((uint32x4_t)vqtbx1q_u8(res, res, Table), 0) != 0;
 }
 
 //TODO: Vectorize
