Merge of CLs 7879375 and 7974119 - Engine changes.

diff --git a/Engine/Source/Runtime/Core/Public/Math/UnrealMathNeon.h b/Engine/Source/Runtime/Core/Public/Math/UnrealMathNeon.h
--- a/Engine/Source/Runtime/Core/Public/Math/UnrealMathNeon.h
+++ b/Engine/Source/Runtime/Core/Public/Math/UnrealMathNeon.h
@@ -1168,7 +1168,17 @@ FORCEINLINE void VectorStoreSignedByte4(VectorRegister Vec, void* Ptr)
 template <bool bAligned>
 FORCEINLINE void VectorStoreHalf4(VectorRegister Vec, void* RESTRICT Ptr)
 {
+#if PLATFORM_ANDROID_ARM
+	float16x4_t f16x4;
+
+	for (int x = 0; x < 4; x++)
+	{
+		f16x4[x] = Vec[x];
+	}
+#else
 	float16x4_t f16x4 = (float16x4_t)vcvt_f16_f32(Vec);
+#endif
+
 	if (bAligned)
 	{
 		vst1_u8( (uint8_t *)Ptr, f16x4 );
