Decrease startup times on Linux non-editor by not traversing the filesystem (UE-80192).

- Regular behavior is to search the whole filesystem for a possibly misspelled cAsE in the filename. Make game and server binaries not do that,
  since the bulk of (possibly mispelled) asset names should be read from a case-insensitive pak file anyway. Any files outside are expected to be
  given with the proper case.

diff --git a/Engine/Source/Runtime/Core/Private/Unix/UnixPlatformFile.cpp b/Engine/Source/Runtime/Core/Private/Unix/UnixPlatformFile.cpp
--- a/Engine/Source/Runtime/Core/Private/Unix/UnixPlatformFile.cpp
+++ b/Engine/Source/Runtime/Core/Private/Unix/UnixPlatformFile.cpp
@@ -10,6 +10,7 @@
 #include <sys/file.h>
 
 #include "HAL/PlatformFileCommon.h"
+#include "HAL/PlatformFilemanager.h"
 
 DEFINE_LOG_CATEGORY_STATIC(LogUnixPlatformFile, Log, All);
 
@@ -628,6 +629,14 @@ public:
 			}
 			else
 			{
+#if (UE_GAME || UE_SERVER)
+				// Games (including clients) and servers have no business traversing the filesystem when reading from the pak files - make sure the paths are correct!
+				static bool bReadingFromPakFiles = FPlatformFileManager::Get().FindPlatformFile(TEXT("PakFile")) != nullptr;
+				if (LIKELY(bReadingFromPakFiles))
+				{
+					return -1;
+				}
+#endif
 				// perform a case-insensitive search
 				// make sure we get the absolute filename
 				checkf(Filename[0] == TEXT('/'), TEXT("Filename '%s' given to OpenCaseInsensitiveRead is not absolute!"), *Filename);
