Enable LOGTRACE for Linux

#rnx

diff --git a/Engine/Source/Runtime/Core/Public/Logging/LogTrace.h b/Engine/Source/Runtime/Core/Public/Logging/LogTrace.h
--- a/Engine/Source/Runtime/Core/Public/Logging/LogTrace.h
+++ b/Engine/Source/Runtime/Core/Public/Logging/LogTrace.h
@@ -4,7 +4,7 @@
 
 #include "CoreTypes.h"
 
-#if !IS_PROGRAM && !UE_BUILD_SHIPPING && (PLATFORM_WINDOWS || PLATFORM_PS4 || PLATFORM_XBOXONE)
+#if !IS_PROGRAM && !UE_BUILD_SHIPPING && (PLATFORM_WINDOWS || PLATFORM_PS4 || PLATFORM_UNIX || PLATFORM_XBOXONE)
 #define LOGTRACE_ENABLED 1
 #else
 #define LOGTRACE_ENABLED 0
@@ -50,4 +50,4 @@ private:
 #else
 #define TRACE_LOG_CATEGORY(Category, Name, DefaultVerbosity)
 #define TRACE_LOG_MESSAGE(Category, Verbosity, Format, ...)
-#endif
\ No newline at end of file
+#endif
