Fix Exclusive csv tag. When wait events where happening on the RT, the conditional wait end  marker was inserted, and exclusive markers above in the hierarchy where not computing the correct timing. When this happen in between the InitViews_Scene and the InitViews_Shadows, because of the occlusion queries, the InitViews_Scene showed the total time of the function instead of exclusive time.

diff --git a/Engine/Source/Runtime/Core/Public/ProfilingDebugging/CsvProfiler.h b/Engine/Source/Runtime/Core/Public/ProfilingDebugging/CsvProfiler.h
--- a/Engine/Source/Runtime/Core/Public/ProfilingDebugging/CsvProfiler.h
+++ b/Engine/Source/Runtime/Core/Public/ProfilingDebugging/CsvProfiler.h
@@ -406,10 +406,13 @@ public:
 
 	~FScopedCsvWaitConditional()
 	{
+		if (bCondition)
+		{
 #if CSV_EXCLUSIVE_TIMING_STATS_EMIT_NAMED_EVENTS
-		FPlatformMisc::EndNamedEvent();
+			FPlatformMisc::EndNamedEvent();
 #endif
-		FCsvProfiler::EndWait();
+			FCsvProfiler::EndWait();
+		}
 	}
 	bool bCondition;
 };
