Fixed `FThread` deadlock that happens when thread function is executed before `FThread` instance initialisation is completed. This was caused by the wrong assumption that `TSharedPtr<T, ESPMode::ThreadSafe>::IsValid()` method is "volatile"

diff --git a/Engine/Source/Runtime/Core/Private/HAL/Thread.cpp b/Engine/Source/Runtime/Core/Private/HAL/Thread.cpp
--- a/Engine/Source/Runtime/Core/Private/HAL/Thread.cpp
+++ b/Engine/Source/Runtime/Core/Private/HAL/Thread.cpp
@@ -29,6 +29,7 @@ public:
 	void Initialize(const TSharedPtr<FThreadImpl, ESPMode::ThreadSafe>& InSelf)
 	{
 		Self = InSelf;
+		bIsInitialized = true;
 	}
 
 	bool IsJoinable() const
@@ -57,7 +58,10 @@ private:
 
 		ThreadFunction();
 
-		while (!Self.IsValid());
+		// busy-wait till `Self` is initialized before releasing it
+		while (!bIsInitialized);
+
+		check(Self.IsValid());
 
 		// we're about to exit the thread. Release the reference to self. If the thread is detached, it's the only reference
 		// and so this instance will be deleted. No member access should be performed after this point.
@@ -74,6 +78,8 @@ private:
 	// This must be the declared before `RunnableThread` so it's already initialized when the thread is created, otherwise the thread can complete
 	// before `Self` is initialized.
 	TSharedPtr<FThreadImpl, ESPMode::ThreadSafe> Self;
+
+	TAtomic<bool> bIsInitialized = false;
 	
 	TUniqueFunction<void()> ThreadFunction;
 	TUniquePtr<FRunnableThread> RunnableThread;
diff --git a/Engine/Source/Runtime/Core/Private/Tests/HAL/ThreadTest.cpp b/Engine/Source/Runtime/Core/Private/Tests/HAL/ThreadTest.cpp
--- a/Engine/Source/Runtime/Core/Private/Tests/HAL/ThreadTest.cpp
+++ b/Engine/Source/Runtime/Core/Private/Tests/HAL/ThreadTest.cpp
@@ -21,6 +21,7 @@ namespace
 		FThread Thread(TEXT("Test.Thread"), []() { /*NOOP*/ });
 		This.TestTrue(TEXT("FThread must be joinable after construction"), Thread.IsJoinable());
 		Thread.Join();
+		UE_LOG(LogTemp, Log, TEXT("%s completed"), TEXT(__FUNCTION__));
 	}
 
 	void TestIsJoinableAfterCompletion(FThreadTest& This)
@@ -30,6 +31,7 @@ namespace
 		while (!bDone); // wait for completion //-V529
 		This.TestTrue(TEXT("FThread must still be joinable after completion"), Thread.IsJoinable());
 		Thread.Join();
+		UE_LOG(LogTemp, Log, TEXT("%s completed"), TEXT(__FUNCTION__));
 	}
 
 	void TestIsNotJoinableAfterJoining(FThreadTest& This)
@@ -37,6 +39,7 @@ namespace
 		FThread Thread(TEXT("Test.Thread"), []() { /*NOOP*/ });
 		Thread.Join();
 		This.TestFalse(TEXT("FThread must not be joinable after joining"), Thread.IsJoinable());
+		UE_LOG(LogTemp, Log, TEXT("%s completed"), TEXT(__FUNCTION__));
 	}
 
 	void TestIsNotJoinableAfterDetaching(FThreadTest& This)
@@ -44,6 +47,7 @@ namespace
 		FThread Thread(TEXT("Test.Thread"), []() { /*NOOP*/ });
 		Thread.Detach();
 		This.TestFalse(TEXT("FThread must not be joinable after detaching"), Thread.IsJoinable());
+		UE_LOG(LogTemp, Log, TEXT("%s completed"), TEXT(__FUNCTION__));
 	}
 
 	void TestAssertIfNotJoinedOrDetached(FThreadTest& This)
@@ -65,6 +69,7 @@ namespace
 			This.TestTrue(TEXT("Move-constructed FThread from joinable thread must be joinable"), Thread.IsJoinable());
 			Thread.Join();
 		}
+		UE_LOG(LogTemp, Log, TEXT("%s completed"), TEXT(__FUNCTION__));
 	}
 
 	void TestMovability(FThreadTest& This)
@@ -103,6 +108,7 @@ namespace
 			Dst = MoveTemp(Src);
 			Dst.Join();
 		}
+		UE_LOG(LogTemp, Log, TEXT("%s completed"), TEXT(__FUNCTION__));
 	}
 
 	// An example of possible implementation of Consumer/Producer idiom
@@ -160,11 +166,15 @@ namespace
 		//	Request to quit
 		//	The thread 0x96e0 has exited with code 0 (0x0).
 		//	Quit
+	
+		UE_LOG(LogTemp, Log, TEXT("%s completed"), TEXT(__FUNCTION__));
 	}
 }
 
 bool FThreadTest::RunTest(const FString& Parameters)
 {
+	UE_LOG(LogTemp, Log, TEXT("%s"), TEXT(__FUNCTION__));
+
 	TestIsJoinableAfterCreation(*this);
 	TestIsJoinableAfterCompletion(*this);
 	TestIsNotJoinableAfterJoining(*this);
