fixed linux compilation

diff --git a/Engine/Source/Runtime/Core/Private/HAL/Thread.cpp b/Engine/Source/Runtime/Core/Private/HAL/Thread.cpp
--- a/Engine/Source/Runtime/Core/Private/HAL/Thread.cpp
+++ b/Engine/Source/Runtime/Core/Private/HAL/Thread.cpp
@@ -79,8 +79,8 @@ private:
 	// before `Self` is initialized.
 	TSharedPtr<FThreadImpl, ESPMode::ThreadSafe> Self;
 
-	TAtomic<bool> bIsInitialized = false;
-	
+	TAtomic<bool> bIsInitialized{ false };
+
 	TUniqueFunction<void()> ThreadFunction;
 	TUniquePtr<FRunnableThread> RunnableThread;
 };
diff --git a/Engine/Source/Runtime/Core/Private/Tests/HAL/ThreadTest.cpp b/Engine/Source/Runtime/Core/Private/Tests/HAL/ThreadTest.cpp
--- a/Engine/Source/Runtime/Core/Private/Tests/HAL/ThreadTest.cpp
+++ b/Engine/Source/Runtime/Core/Private/Tests/HAL/ThreadTest.cpp
@@ -11,6 +11,7 @@
 #include "HAL/ThreadSafeCounter.h"
 #include "HAL/Event.h"
 #include "Containers/Queue.h"
+#include "Containers/StringConv.h"
 
 IMPLEMENT_SIMPLE_AUTOMATION_TEST(FThreadTest, "System.Core.HAL.Thread", EAutomationTestFlags::ApplicationContextMask | EAutomationTestFlags::EngineFilter)
 
@@ -21,7 +22,7 @@ namespace
 		FThread Thread(TEXT("Test.Thread"), []() { /*NOOP*/ });
 		This.TestTrue(TEXT("FThread must be joinable after construction"), Thread.IsJoinable());
 		Thread.Join();
-		UE_LOG(LogTemp, Log, TEXT("%s completed"), TEXT(__FUNCTION__));
+		UE_LOG(LogTemp, Log, TEXT("%s completed"), StringCast<TCHAR>(__FUNCTION__).Get());
 	}
 
 	void TestIsJoinableAfterCompletion(FThreadTest& This)
@@ -31,7 +32,7 @@ namespace
 		while (!bDone); // wait for completion //-V529
 		This.TestTrue(TEXT("FThread must still be joinable after completion"), Thread.IsJoinable());
 		Thread.Join();
-		UE_LOG(LogTemp, Log, TEXT("%s completed"), TEXT(__FUNCTION__));
+		UE_LOG(LogTemp, Log, TEXT("%s completed"), StringCast<TCHAR>(__FUNCTION__).Get());
 	}
 
 	void TestIsNotJoinableAfterJoining(FThreadTest& This)
@@ -39,7 +40,7 @@ namespace
 		FThread Thread(TEXT("Test.Thread"), []() { /*NOOP*/ });
 		Thread.Join();
 		This.TestFalse(TEXT("FThread must not be joinable after joining"), Thread.IsJoinable());
-		UE_LOG(LogTemp, Log, TEXT("%s completed"), TEXT(__FUNCTION__));
+		UE_LOG(LogTemp, Log, TEXT("%s completed"), StringCast<TCHAR>(__FUNCTION__).Get());
 	}
 
 	void TestIsNotJoinableAfterDetaching(FThreadTest& This)
@@ -47,7 +48,7 @@ namespace
 		FThread Thread(TEXT("Test.Thread"), []() { /*NOOP*/ });
 		Thread.Detach();
 		This.TestFalse(TEXT("FThread must not be joinable after detaching"), Thread.IsJoinable());
-		UE_LOG(LogTemp, Log, TEXT("%s completed"), TEXT(__FUNCTION__));
+		UE_LOG(LogTemp, Log, TEXT("%s completed"), StringCast<TCHAR>(__FUNCTION__).Get());
 	}
 
 	void TestAssertIfNotJoinedOrDetached(FThreadTest& This)
@@ -69,7 +70,7 @@ namespace
 			This.TestTrue(TEXT("Move-constructed FThread from joinable thread must be joinable"), Thread.IsJoinable());
 			Thread.Join();
 		}
-		UE_LOG(LogTemp, Log, TEXT("%s completed"), TEXT(__FUNCTION__));
+		UE_LOG(LogTemp, Log, TEXT("%s completed"), StringCast<TCHAR>(__FUNCTION__).Get());
 	}
 
 	void TestMovability(FThreadTest& This)
@@ -108,7 +109,7 @@ namespace
 			Dst = MoveTemp(Src);
 			Dst.Join();
 		}
-		UE_LOG(LogTemp, Log, TEXT("%s completed"), TEXT(__FUNCTION__));
+		UE_LOG(LogTemp, Log, TEXT("%s completed"), StringCast<TCHAR>(__FUNCTION__).Get());
 	}
 
 	// An example of possible implementation of Consumer/Producer idiom
@@ -167,13 +168,13 @@ namespace
 		//	The thread 0x96e0 has exited with code 0 (0x0).
 		//	Quit
 	
-		UE_LOG(LogTemp, Log, TEXT("%s completed"), TEXT(__FUNCTION__));
+		UE_LOG(LogTemp, Log, TEXT("%s completed"), StringCast<TCHAR>(__FUNCTION__).Get());
 	}
 }
 
 bool FThreadTest::RunTest(const FString& Parameters)
 {
-	UE_LOG(LogTemp, Log, TEXT("%s"), TEXT(__FUNCTION__));
+	UE_LOG(LogTemp, Log, TEXT("%s"), StringCast<TCHAR>(__FUNCTION__).Get());
 
 	TestIsJoinableAfterCreation(*this);
 	TestIsJoinableAfterCompletion(*this);
