Fix compile error while loop has empty body -Werror,-Wempty-body

diff --git a/Engine/Source/Runtime/Core/Private/HAL/Thread.cpp b/Engine/Source/Runtime/Core/Private/HAL/Thread.cpp
--- a/Engine/Source/Runtime/Core/Private/HAL/Thread.cpp
+++ b/Engine/Source/Runtime/Core/Private/HAL/Thread.cpp
@@ -59,7 +59,7 @@ private:
 		ThreadFunction();
 
 		// busy-wait till `Self` is initialized before releasing it
-		while (!bIsInitialized);
+		while (!bIsInitialized) {}
 
 		check(Self.IsValid());
 
