Silence VS2019 compiler warning

diff --git a/Engine/Source/Runtime/Core/Public/UObject/NameTypes.h b/Engine/Source/Runtime/Core/Public/UObject/NameTypes.h
--- a/Engine/Source/Runtime/Core/Public/UObject/NameTypes.h
+++ b/Engine/Source/Runtime/Core/Public/UObject/NameTypes.h
@@ -1025,7 +1025,7 @@ private:
 
 		bool IsLiteral() const
 		{
-			return LiteralFlag & Int;
+			return (LiteralFlag & Int) != 0;
 		}
 
 		FNameEntryId AsName() const
@@ -1058,7 +1058,7 @@ private:
 
 		bool IsLiteral() const
 		{
-			return LiteralFlag & Int;
+			return (LiteralFlag & Int) != 0;
 		}
 
 		FNameEntryId AsName() const
