Allow FString instances containing code units outside of the basic multilingual plane to be losslessly processed regardless of whether TCHAR is 2 or 4 bytes.

Most UE4 platforms use a 2-byte TCHAR, however some still use a 4-byte TCHAR. The platforms that use a 4-byte TCHAR expect their string data to be UTF-32, however there are parts of UE4 that serialize FString data as a series of UCS2CHAR, simply narrowing or widening each TCHAR in turn. This can result in invalid or corrupted UTF-32 strings (either UTF-32 strings containing UTF-16 surrogates, or UTF-32 code points that have been truncated to 2-bytes), which leads to either odd behavior or crashes.

This change updates the parts of UE4 that process FString data as a series of 2-byte values to do so on the correct UTF-16 interpretation of the data, converting to/from UTF-32 as required on platforms that use a 4-byte TCHAR. This conversion is a no-op on platforms that use a 2-byte TCHAR as the string is already assumed to be valid UTF-16 data. It should also be noted that while FString may contain UTF-16 code units on platforms using a 2-byte TCHAR, this change doesn't do anything to make FString represent a Unicode string on those platforms (ie, a string that understands and works on code points), but is rather just a bag of code units.

Two new variable-width string converters have be added to facilitate the conversion (modelled after the TCHAR<->UTF-8 converters), TUTF16ToUTF32_Convert and TUTF32ToUTF16_Convert. These are used for both TCHAR<->UTF16CHAR conversion when needed, but also for TCHAR<->wchar_t conversion on platforms that use char16_t for TCHAR along with having a 4-byte wchar_t (as defined by the new PLATFORM_WCHAR_IS_4_BYTES option).

These conversion routines are accessed either via the conversion macros (TCHAR_TO_UTF16, UTF16_TO_TCHAR, TCHAR_TO_WCHAR, and WCHAR_TO_TCHAR), or by using a conversion struct (FTCHARToUTF16, FUTF16ToTCHAR, FTCHARToWChar, and FWCharToTCHAR), which is the same pattern as the existing TCHAR<->UTF-8 conversion. Both the macros and the structs are defined as no-ops when the conversion isn't needed, but always exist so that code can be written in a portable way.

Very little code actually needed updating to use UTF-16, as the vast majority makes no assumptions about the size of TCHAR, nor how FString should be serialized. The main places were the FString archive serialization and the JSON reader/writer, along with some minor fixes to the UTF-8 conversion logic for platforms using a 4-byte TCHAR.

Tests have been added to verify that an FString representing a UTF-32 code point can be losslessly converted to/from UTF-8 and UTF-16, and serialized to/from an archive.

diff --git a/Engine/Source/Editor/KismetCompiler/Private/KismetCompilerVMBackend.cpp b/Engine/Source/Editor/KismetCompiler/Private/KismetCompilerVMBackend.cpp
--- a/Engine/Source/Editor/KismetCompiler/Private/KismetCompilerVMBackend.cpp
+++ b/Engine/Source/Editor/KismetCompiler/Private/KismetCompilerVMBackend.cpp
@@ -422,11 +422,14 @@ public:
 		}
 		else
 		{
+			// Note: This is a no-op on platforms that are using a 16-bit TCHAR
+			FTCHARToUTF16 UTF16String(*String, String.Len() + 1); // include the null terminator
+
 			Writer << EX_UnicodeStringConst;
 			uint16 OutCh;
-			for (const TCHAR* Ch = *String; *Ch; ++Ch)
+			for (const UTF16CHAR* Ch = UTF16String.Get(); *Ch; ++Ch)
 			{
-				OutCh = CharCast<UCS2CHAR>(*Ch);
+				OutCh = *Ch;
 				Writer << OutCh;
 			}
 
diff --git a/Engine/Source/Editor/UnrealEd/Private/ScriptDisassembler.cpp b/Engine/Source/Editor/UnrealEd/Private/ScriptDisassembler.cpp
--- a/Engine/Source/Editor/UnrealEd/Private/ScriptDisassembler.cpp
+++ b/Engine/Source/Editor/UnrealEd/Private/ScriptDisassembler.cpp
@@ -193,6 +193,9 @@ FString FKismetBytecodeDisassembler::ReadString16(int32& ScriptIndex)
 	}
 	while ((Script[ScriptIndex-1] != 0) || (Script[ScriptIndex-2] != 0));
 
+	// Inline combine any surrogate pairs in the data when loading into a UTF-32 string
+	StringConv::InlineCombineSurrogates(Result);
+
 	return Result;
 }
 
diff --git a/Engine/Source/Runtime/Core/Private/Containers/String.cpp b/Engine/Source/Runtime/Core/Private/Containers/String.cpp
--- a/Engine/Source/Runtime/Core/Private/Containers/String.cpp
+++ b/Engine/Source/Runtime/Core/Private/Containers/String.cpp
@@ -1409,17 +1409,20 @@ void FString::AppendfImpl(FString& AppendToMe, const TCHAR* Fmt, ...)
 	}
 }
 
+static_assert(PLATFORM_LITTLE_ENDIAN, "FString serialization needs updating to support big-endian platforms!");
+
 FArchive& operator<<( FArchive& Ar, FString& A )
 {
-	// > 0 for ANSICHAR, < 0 for UCS2CHAR serialization
+	// > 0 for ANSICHAR, < 0 for UTF16CHAR serialization
+	static_assert(sizeof(UTF16CHAR) == sizeof(UCS2CHAR), "UTF16CHAR and UCS2CHAR are assumed to be the same size!");
 
 	if (Ar.IsLoading())
 	{
 		int32 SaveNum;
 		Ar << SaveNum;
 
-		bool LoadUCS2Char = SaveNum < 0;
-		if (LoadUCS2Char)
+		bool bLoadUnicodeChar = SaveNum < 0;
+		if (bLoadUnicodeChar)
 		{
 			// If SaveNum cannot be negated due to integer overflow, Ar is corrupted.
 			if (SaveNum == MIN_int32)
@@ -1449,16 +1452,17 @@ FArchive& operator<<( FArchive& Ar, FString& A )
 
 		if (SaveNum)
 		{
-			if (LoadUCS2Char)
+			if (bLoadUnicodeChar)
 			{
-				// read in the unicode string and byteswap it, etc
+				// read in the unicode string
 				auto Passthru = StringMemoryPassthru<UCS2CHAR>(A.Data.GetData(), SaveNum, SaveNum);
 				Ar.Serialize(Passthru.Get(), SaveNum * sizeof(UCS2CHAR));
 				// Ensure the string has a null terminator
 				Passthru.Get()[SaveNum-1] = '\0';
 				Passthru.Apply();
 
-				INTEL_ORDER_TCHARARRAY(A.Data.GetData())
+				// Inline combine any surrogate pairs in the data when loading into a UTF-32 string
+				StringConv::InlineCombineSurrogates(A);
 
 				// Since Microsoft's vsnwprintf implementation raises an invalid parameter warning
 				// with a character of 0xffff, scan for it and terminate the string there.
@@ -1488,30 +1492,31 @@ FArchive& operator<<( FArchive& Ar, FString& A )
 	}
 	else
 	{
-		bool  SaveUCS2Char = Ar.IsForcingUnicode() || !FCString::IsPureAnsi(*A);
-		int32 Num        = A.Data.Num();
-		int32 SaveNum    = SaveUCS2Char ? -Num : Num;
+		A.Data.CountBytes(Ar);
 
-		Ar << SaveNum;
+		const bool bSaveUnicodeChar = Ar.IsForcingUnicode() || !FCString::IsPureAnsi(*A);
+		if (bSaveUnicodeChar)
+		{
+			// Note: This is a no-op on platforms that are using a 16-bit TCHAR
+ 			FTCHARToUTF16 UTF16String(*A, A.Len() + 1); // include the null terminator
+			int32 Num = UTF16String.Length() + 1; // include the null terminator
 
-		A.Data.CountBytes( Ar );
+			int32 SaveNum = -Num;
+			Ar << SaveNum;
 
-		if (SaveNum)
-		{
-			if (SaveUCS2Char)
+			if (Num)
 			{
-				// TODO - This is creating a temporary in order to byte-swap.  Need to think about how to make this not necessary.
-				#if !PLATFORM_LITTLE_ENDIAN
-					FString  ATemp  = A;
-					FString& A      = ATemp;
-					INTEL_ORDER_TCHARARRAY(A.Data.GetData());
-				#endif
-
-				Ar.Serialize((void*)StringCast<UCS2CHAR>(A.Data.GetData(), Num).Get(), sizeof(UCS2CHAR)* Num);
+				Ar.Serialize((void*)UTF16String.Get(), sizeof(UTF16CHAR) * Num);
 			}
-			else
+		}
+		else
+		{
+			int32 Num = A.Data.Num();
+			Ar << Num;
+
+			if (Num)
 			{
-				Ar.Serialize((void*)StringCast<ANSICHAR>(A.Data.GetData(), Num).Get(), sizeof(ANSICHAR)* Num);
+				Ar.Serialize((void*)StringCast<ANSICHAR>(A.Data.GetData(), Num).Get(), sizeof(ANSICHAR) * Num);
 			}
 		}
 	}
diff --git a/Engine/Source/Runtime/Core/Private/Containers/StringTest.cpp b/Engine/Source/Runtime/Core/Private/Containers/StringTest.cpp
--- a/Engine/Source/Runtime/Core/Private/Containers/StringTest.cpp
+++ b/Engine/Source/Runtime/Core/Private/Containers/StringTest.cpp
@@ -4,6 +4,8 @@
 #include "Misc/AutomationTest.h"
 #include "Misc/AssertionMacros.h"
 #include "Containers/UnrealString.h"
+#include "Serialization/MemoryReader.h"
+#include "Serialization/MemoryWriter.h"
 
 #if WITH_DEV_AUTOMATION_TESTS
 
@@ -86,4 +88,51 @@ bool FStringAppendIntTest::RunTest(const FString& Parameters)
 	return true;
 }
 
+IMPLEMENT_SIMPLE_AUTOMATION_TEST(FStringUnicodeTest, "System.Core.String.Unicode", EAutomationTestFlags::ApplicationContextMask | EAutomationTestFlags::SmokeFilter)
+bool FStringUnicodeTest::RunTest(const FString& Parameters)
+{
+	auto DoTest = [this](const TCHAR* Call, const FString& Result, const TCHAR* InExpected)
+	{
+		if (!Result.Equals(InExpected, ESearchCase::CaseSensitive))
+		{
+			AddError(FString::Printf(TEXT("'%s' failure: result '%s' (expected '%s')"), Call, *Result, InExpected));
+		}
+	};
+
+	// Test data used to verify basic processing of a Unicode character outside the BMP
+	FString TestStr;
+	if (FUnicodeChar::CodepointToString(128512, TestStr))
+	{
+		// Verify that the string can be serialized and deserialized without losing any data
+		{
+			TArray<uint8> StringData;
+			FString FromArchive = TestStr;
+
+			FMemoryWriter Writer(StringData);
+			Writer << FromArchive;
+
+			FromArchive.Reset();
+			FMemoryReader Reader(StringData);
+			Reader << FromArchive;
+
+			DoTest(TEXT("FromArchive"), FromArchive, *TestStr);
+		}
+
+		// Verify that the string can be converted from/to UTF-8 without losing any data
+		{
+			const FString FromUtf8 = UTF8_TO_TCHAR(TCHAR_TO_UTF8(*TestStr));
+			DoTest(TEXT("FromUtf8"), FromUtf8, *TestStr);
+		}
+
+		// Verify that the string can be converted from/to UTF-16 without losing any data
+		{
+			const FString FromUtf16 = UTF16_TO_TCHAR(TCHAR_TO_UTF16(*TestStr));
+			DoTest(TEXT("FromUtf16"), FromUtf16, *TestStr);
+		}
+	}
+
+
+	return true;
+}
+
 #endif // WITH_DEV_AUTOMATION_TESTS
diff --git a/Engine/Source/Runtime/Core/Private/Internationalization/TextKey.cpp b/Engine/Source/Runtime/Core/Private/Internationalization/TextKey.cpp
--- a/Engine/Source/Runtime/Core/Private/Internationalization/TextKey.cpp
+++ b/Engine/Source/Runtime/Core/Private/Internationalization/TextKey.cpp
@@ -168,36 +168,37 @@ namespace TextKeyUtil
 static const int32 InlineStringSize = 128;
 typedef TArray<TCHAR, TInlineAllocator<InlineStringSize>> FInlineStringBuffer;
 
+static_assert(PLATFORM_LITTLE_ENDIAN, "FTextKey serialization needs updating to support big-endian platforms!");
+
 bool SaveKeyString(FArchive& Ar, const TCHAR* InStrPtr)
 {
 	// Note: This serialization should be compatible with the FString serialization, but avoids creating an FString if the FTextKey is already cached
-	// > 0 for ANSICHAR, < 0 for UCS2CHAR serialization
+	// > 0 for ANSICHAR, < 0 for UTF16CHAR serialization
 	check(!Ar.IsLoading());
 
-	const bool SaveUCS2Char = Ar.IsForcingUnicode() || !FCString::IsPureAnsi(InStrPtr);
-	const int32 Num = FCString::Strlen(InStrPtr) + 1; // include the null terminator
-
-	int32 SaveNum = SaveUCS2Char ? -Num : Num;
-	Ar << SaveNum;
-
-	if (SaveNum)
+	const bool bSaveUnicodeChar = Ar.IsForcingUnicode() || !FCString::IsPureAnsi(InStrPtr);
+	if (bSaveUnicodeChar)
 	{
-		if (SaveUCS2Char)
-		{
-			const TCHAR* LittleEndianStrPtr = InStrPtr;
+		// Note: This is a no-op on platforms that are using a 16-bit TCHAR
+		FTCHARToUTF16 UTF16String(InStrPtr);
+		const int32 Num = UTF16String.Length() + 1; // include the null terminator
 
-			// TODO - This is creating a temporary in order to byte-swap.  Need to think about how to make this not necessary.
-#if !PLATFORM_LITTLE_ENDIAN
-			FString LittleEndianStr = FString(Num - 1, LittleEndianStrPtr);
-			INTEL_ORDER_TCHARARRAY(LittleEndianStr.GetData());
-			LittleEndianStrPtr = *LittleEndianStr;
-#endif
+		int32 SaveNum = -Num;
+		Ar << SaveNum;
 
-			Ar.Serialize((void*)StringCast<UCS2CHAR>(LittleEndianStrPtr, Num).Get(), sizeof(UCS2CHAR)* Num);
+		if (Num)
+		{
+			Ar.Serialize((void*)UTF16String.Get(), sizeof(UTF16CHAR) * Num);
 		}
-		else
+	}
+	else
+	{
+		int32 Num = FCString::Strlen(InStrPtr) + 1; // include the null terminator
+		Ar << Num;
+
+		if (Num)
 		{
-			Ar.Serialize((void*)StringCast<ANSICHAR>(InStrPtr, Num).Get(), sizeof(ANSICHAR)* Num);
+			Ar.Serialize((void*)StringCast<ANSICHAR>(InStrPtr, Num).Get(), sizeof(ANSICHAR) * Num);
 		}
 	}
 
@@ -207,14 +208,14 @@ bool SaveKeyString(FArchive& Ar, const TCHAR* InStrPtr)
 bool LoadKeyString(FArchive& Ar, FInlineStringBuffer& OutStrBuffer)
 {
 	// Note: This serialization should be compatible with the FString serialization, but avoids creating an FString if the FTextKey is already cached
-	// > 0 for ANSICHAR, < 0 for UCS2CHAR serialization
+	// > 0 for ANSICHAR, < 0 for UTF16CHAR serialization
 	check(Ar.IsLoading());
 
 	int32 SaveNum = 0;
 	Ar << SaveNum;
 
-	const bool LoadUCS2Char = SaveNum < 0;
-	if (LoadUCS2Char)
+	const bool bLoadUnicodeChar = SaveNum < 0;
+	if (bLoadUnicodeChar)
 	{
 		SaveNum = -SaveNum;
 	}
@@ -241,15 +242,16 @@ bool LoadKeyString(FArchive& Ar, FInlineStringBuffer& OutStrBuffer)
 
 	if (SaveNum)
 	{
-		if (LoadUCS2Char)
+		if (bLoadUnicodeChar)
 		{
-			// Read in the Unicode string and byte-swap it
+			// Read in the Unicode string
 			auto Passthru = StringMemoryPassthru<UCS2CHAR, TCHAR, InlineStringSize>(OutStrBuffer.GetData(), SaveNum, SaveNum);
 			Ar.Serialize(Passthru.Get(), SaveNum * sizeof(UCS2CHAR));
 			Passthru.Get()[SaveNum - 1] = 0; // Ensure the string has a null terminator
 			Passthru.Apply();
 
-			INTEL_ORDER_TCHARARRAY(StrBuffer.GetData())
+			// Inline combine any surrogate pairs in the data when loading into a UTF-32 string
+			StringConv::InlineCombineSurrogates_Array(OutStrBuffer);
 		}
 		else
 		{
diff --git a/Engine/Source/Runtime/Core/Private/Misc/FileHelper.cpp b/Engine/Source/Runtime/Core/Private/Misc/FileHelper.cpp
--- a/Engine/Source/Runtime/Core/Private/Misc/FileHelper.cpp
+++ b/Engine/Source/Runtime/Core/Private/Misc/FileHelper.cpp
@@ -60,6 +60,7 @@ void FFileHelper::BufferToString( FString& Result, const uint8* Buffer, int32 Si
 	TArray<TCHAR>& ResultArray = Result.GetCharArray();
 	ResultArray.Empty();
 
+	bool bIsUnicode = false;
 	if( Size >= 2 && !( Size & 1 ) && Buffer[0] == 0xff && Buffer[1] == 0xfe )
 	{
 		// Unicode Intel byte order. Less 1 for the FFFE header, additional 1 for null terminator.
@@ -68,6 +69,7 @@ void FFileHelper::BufferToString( FString& Result, const uint8* Buffer, int32 Si
 		{
 			ResultArray[ i ] = CharCast<TCHAR>( (UCS2CHAR)(( uint16 )Buffer[i * 2 + 2] + ( uint16 )Buffer[i * 2 + 3] * 256) );
 		}
+		bIsUnicode = true;
 	}
 	else if( Size >= 2 && !( Size & 1 ) && Buffer[0] == 0xfe && Buffer[1] == 0xff )
 	{
@@ -77,6 +79,7 @@ void FFileHelper::BufferToString( FString& Result, const uint8* Buffer, int32 Si
 		{
 			ResultArray[ i ] = CharCast<TCHAR>( (UCS2CHAR)(( uint16 )Buffer[i * 2 + 3] + ( uint16 )Buffer[i * 2 + 2] * 256) );
 		}
+		bIsUnicode = true;
 	}
 	else
 	{
@@ -102,6 +105,12 @@ void FFileHelper::BufferToString( FString& Result, const uint8* Buffer, int32 Si
 	{
 		// Else ensure null terminator is present
 		ResultArray.Last() = 0;
+
+		if (bIsUnicode)
+		{
+			// Inline combine any surrogate pairs in the data when loading into a UTF-32 string
+			StringConv::InlineCombineSurrogates(Result);
+		}
 	}
 }
 
@@ -211,33 +220,32 @@ bool FFileHelper::SaveStringToFile( const FString& String, const TCHAR* Filename
 	if( String.IsEmpty() )
 		return true;
 
-	const TCHAR* StrPtr = *String;
-
-	bool SaveAsUnicode = EncodingOptions == EEncodingOptions::ForceUnicode || ( EncodingOptions == EEncodingOptions::AutoDetect && !FCString::IsPureAnsi(StrPtr) );
+	bool SaveAsUnicode = EncodingOptions == EEncodingOptions::ForceUnicode || ( EncodingOptions == EEncodingOptions::AutoDetect && !FCString::IsPureAnsi(*String) );
 	if( EncodingOptions == EEncodingOptions::ForceUTF8 )
 	{
 		UTF8CHAR UTF8BOM[] = { 0xEF, 0xBB, 0xBF };
 		Ar->Serialize( &UTF8BOM, ARRAY_COUNT(UTF8BOM) * sizeof(UTF8CHAR) );
 
-		FTCHARToUTF8 UTF8String(StrPtr);
+		FTCHARToUTF8 UTF8String(*String, String.Len());
 		Ar->Serialize( (UTF8CHAR*)UTF8String.Get(), UTF8String.Length() * sizeof(UTF8CHAR) );
 	}
 	else if ( EncodingOptions == EEncodingOptions::ForceUTF8WithoutBOM )
 	{
-		FTCHARToUTF8 UTF8String(StrPtr);
+		FTCHARToUTF8 UTF8String(*String, String.Len());
 		Ar->Serialize((UTF8CHAR*)UTF8String.Get(), UTF8String.Length() * sizeof(UTF8CHAR));
 	}
 	else if (SaveAsUnicode)
 	{
-		UCS2CHAR BOM = UNICODE_BOM;
-		Ar->Serialize( &BOM, sizeof(UCS2CHAR) );
+		UTF16CHAR BOM = UNICODE_BOM;
+		Ar->Serialize( &BOM, sizeof(UTF16CHAR) );
 
-		auto Src = StringCast<UCS2CHAR>(StrPtr, String.Len());
-		Ar->Serialize( (UCS2CHAR*)Src.Get(), Src.Length() * sizeof(UCS2CHAR) );
+		// Note: This is a no-op on platforms that are using a 16-bit TCHAR
+		FTCHARToUTF16 UTF16String(*String, String.Len());
+		Ar->Serialize((UTF16CHAR*)UTF16String.Get(), UTF16String.Length() * sizeof(UTF16CHAR));
 	}
 	else
 	{
-		auto Src = StringCast<ANSICHAR>(StrPtr, String.Len());
+		auto Src = StringCast<ANSICHAR>(*String, String.Len());
 		Ar->Serialize( (ANSICHAR*)Src.Get(), Src.Length() * sizeof(ANSICHAR) );
 	}
 
diff --git a/Engine/Source/Runtime/Core/Private/Tests/Misc/TCharToUTF8.cpp b/Engine/Source/Runtime/Core/Private/Tests/Misc/TCharToUTF8.cpp
--- a/Engine/Source/Runtime/Core/Private/Tests/Misc/TCharToUTF8.cpp
+++ b/Engine/Source/Runtime/Core/Private/Tests/Misc/TCharToUTF8.cpp
@@ -11,80 +11,95 @@ bool FTCharToUTF8Test::RunTest(const FString& Parameters)
 {
 	{
 		// Hebrew letter final nun
-		const TCHAR TestString[] = TEXT("\x05DF");
+		const FString TestString = UTF16_TO_TCHAR(u"\x05DF");
 		const ANSICHAR ExpectedResult[] = "\xD7\x9F";
 		{
-			const FTCHARToUTF8 ConvertedValue(TestString);
+			const FTCHARToUTF8 ConvertedValue(*TestString);
 			const FUTF8ToTCHAR ReverseConvertedValue(ConvertedValue.Get());
 			TestTrue(FString(TEXT("Expected test string 1 to have different encoded value")), FCStringAnsi::Strcmp(ExpectedResult, ConvertedValue.Get()) == 0);
-			TestTrue(FString(TEXT("Expected test string 1 to have different decoded value")), FCString::Strcmp(TestString, ReverseConvertedValue.Get()) == 0);
+			TestTrue(FString(TEXT("Expected test string 1 to have different decoded value")), FCString::Strcmp(*TestString, ReverseConvertedValue.Get()) == 0);
 		}
 
 		{
 			const FUTF8ToTCHAR ConvertedValue(ExpectedResult);
 			const FTCHARToUTF8 ReverseConvertedValue(ConvertedValue.Get());
-			TestTrue(FString(TEXT("Expected test string 2 to have different encoded value")), FCString::Strcmp(TestString, ConvertedValue.Get()) == 0);
+			TestTrue(FString(TEXT("Expected test string 2 to have different encoded value")), FCString::Strcmp(*TestString, ConvertedValue.Get()) == 0);
 			TestTrue(FString(TEXT("Expected test string 2 to have different decoded value")), FCStringAnsi::Strcmp(ExpectedResult, ReverseConvertedValue.Get()) == 0);
 		}
 	}
 
 	{
 		// Smiling face with open mouth and tightly-closed eyes and a Four of Circles
-		const TCHAR TestString[] = TEXT("\xD83D\xDE06\xD83C\xDC1C");
+		const FString TestString = UTF16_TO_TCHAR(u"\xD83D\xDE06\xD83C\xDC1C");
 		const ANSICHAR ExpectedResult[] = "\xF0\x9F\x98\x86\xF0\x9F\x80\x9C";
 		{
-			const FTCHARToUTF8 ConvertedValue(TestString);
+			const FTCHARToUTF8 ConvertedValue(*TestString);
 			const FUTF8ToTCHAR ReverseConvertedValue(ConvertedValue.Get());
 			TestTrue(FString(TEXT("Expected test string 3 to have different encoded value")), FCStringAnsi::Strcmp(ExpectedResult, ConvertedValue.Get()) == 0);
-			TestTrue(FString(TEXT("Expected test string 3 to have different decoded value")), FCString::Strcmp(TestString, ReverseConvertedValue.Get()) == 0);
+			TestTrue(FString(TEXT("Expected test string 3 to have different decoded value")), FCString::Strcmp(*TestString, ReverseConvertedValue.Get()) == 0);
 		}
 
 		{
 			const FUTF8ToTCHAR ConvertedValue(ExpectedResult);
 			const FTCHARToUTF8 ReverseConvertedValue(ConvertedValue.Get());
-			TestTrue(FString(TEXT("Expected test string 4 to have different encoded value")), FCString::Strcmp(TestString, ConvertedValue.Get()) == 0);
+			TestTrue(FString(TEXT("Expected test string 4 to have different encoded value")), FCString::Strcmp(*TestString, ConvertedValue.Get()) == 0);
 			TestTrue(FString(TEXT("Expected test string 4 to have different decoded value")), FCStringAnsi::Strcmp(ExpectedResult, ReverseConvertedValue.Get()) == 0);
 		}
 	}
 
 	{
 		// "Internationalization" Test string
-		const TCHAR TestString[] = TEXT("\x0049\x00F1\x0074\x00EB\x0072\x006E\x00E2\x0074\x0069\x00F4\x006E\x00E0\x006C\x0069\x007A\x00E6\x0074\x0069\x00F8\x006E");
+		const FString TestString = UTF16_TO_TCHAR(u"\x0049\x00F1\x0074\x00EB\x0072\x006E\x00E2\x0074\x0069\x00F4\x006E\x00E0\x006C\x0069\x007A\x00E6\x0074\x0069\x00F8\x006E");
 		const ANSICHAR ExpectedResult[] = "\x49\xC3\xB1\x74\xC3\xAB\x72\x6E\xC3\xA2\x74\x69\xC3\xB4\x6E\xC3\xA0\x6C\x69\x7A\xC3\xA6\x74\x69\xC3\xB8\x6E";
 		{
-			const FTCHARToUTF8 ConvertedValue(TestString);
+			const FTCHARToUTF8 ConvertedValue(*TestString);
 			const FUTF8ToTCHAR ReverseConvertedValue(ConvertedValue.Get());
 			TestTrue(FString(TEXT("Expected test string 5 to have different encoded value")), FCStringAnsi::Strcmp(ExpectedResult, ConvertedValue.Get()) == 0);
-			TestTrue(FString(TEXT("Expected test string 5 to have different decoded value")), FCString::Strcmp(TestString, ReverseConvertedValue.Get()) == 0);
+			TestTrue(FString(TEXT("Expected test string 5 to have different decoded value")), FCString::Strcmp(*TestString, ReverseConvertedValue.Get()) == 0);
 		}
 
 		{
 			const FUTF8ToTCHAR ConvertedValue(ExpectedResult);
 			const FTCHARToUTF8 ReverseConvertedValue(ConvertedValue.Get());
-			TestTrue(FString(TEXT("Expected test string 6 to have different encoded value")), FCString::Strcmp(TestString, ConvertedValue.Get()) == 0);
+			TestTrue(FString(TEXT("Expected test string 6 to have different encoded value")), FCString::Strcmp(*TestString, ConvertedValue.Get()) == 0);
 			TestTrue(FString(TEXT("Expected test string 6 to have different decoded value")), FCStringAnsi::Strcmp(ExpectedResult, ReverseConvertedValue.Get()) == 0);
 		}
 	}
 
 	{
 		// UE Test string
-		const TCHAR TestString[] = TEXT("\xD835\xDE50\x043F\xD835\xDDCB\xD835\xDE26\xD835\xDD52\x006C\x0020\xD835\xDC6C\xD835\xDD2B\xFF47\xD835\xDCBE\xD835\xDD93\xD835\xDD56\x0020\xD835\xDE92\xFF53\x0020\xD835\xDC6E\xD835\xDE9B\x212E\xD835\xDF36\xFF54\x0021");
+		const FString TestString = UTF16_TO_TCHAR(u"\xD835\xDE50\x043F\xD835\xDDCB\xD835\xDE26\xD835\xDD52\x006C\x0020\xD835\xDC6C\xD835\xDD2B\xFF47\xD835\xDCBE\xD835\xDD93\xD835\xDD56\x0020\xD835\xDE92\xFF53\x0020\xD835\xDC6E\xD835\xDE9B\x212E\xD835\xDF36\xFF54\x0021");
 		const ANSICHAR ExpectedResult[] = "\xF0\x9D\x99\x90\xD0\xBF\xF0\x9D\x97\x8B\xF0\x9D\x98\xA6\xF0\x9D\x95\x92\x6C\x20\xF0\x9D\x91\xAC\xF0\x9D\x94\xAB\xEF\xBD\x87\xF0\x9D\x92\xBE\xF0\x9D\x96\x93\xF0\x9D\x95\x96\x20\xF0\x9D\x9A\x92\xEF\xBD\x93\x20\xF0\x9D\x91\xAE\xF0\x9D\x9A\x9B\xE2\x84\xAE\xF0\x9D\x9C\xB6\xEF\xBD\x94\x21";
 		{
-			const FTCHARToUTF8 ConvertedValue(TestString);
+			const FTCHARToUTF8 ConvertedValue(*TestString);
 			const FUTF8ToTCHAR ReverseConvertedValue(ConvertedValue.Get());
 			TestTrue(FString(TEXT("Expected test string 7 to have different encoded value")), FCStringAnsi::Strcmp(ExpectedResult, ConvertedValue.Get()) == 0);
-			TestTrue(FString(TEXT("Expected test string 7 to have different decoded value")), FCString::Strcmp(TestString, ReverseConvertedValue.Get()) == 0);
+			TestTrue(FString(TEXT("Expected test string 7 to have different decoded value")), FCString::Strcmp(*TestString, ReverseConvertedValue.Get()) == 0);
 		}
 
 		{
 			const FUTF8ToTCHAR ConvertedValue(ExpectedResult);
 			const FTCHARToUTF8 ReverseConvertedValue(ConvertedValue.Get());
-			TestTrue(FString(TEXT("Expected test string 8 to have different encoded value")), FCString::Strcmp(TestString, ConvertedValue.Get()) == 0);
+			TestTrue(FString(TEXT("Expected test string 8 to have different encoded value")), FCString::Strcmp(*TestString, ConvertedValue.Get()) == 0);
 			TestTrue(FString(TEXT("Expected test string 8 to have different decoded value")), FCStringAnsi::Strcmp(ExpectedResult, ReverseConvertedValue.Get()) == 0);
 		}
 	}
 
+	{
+		// Null terminator
+		const FString TestString = TEXT("Test");
+		{
+			const FTCHARToUTF8 ConvertedValue(*TestString);
+			const bool bIsNullTerminated = *(ConvertedValue.Get() + ConvertedValue.Length()) == 0;
+			TestTrue(TEXT("Expected converted value to be null-terminated"), bIsNullTerminated);
+		}
+		{
+			const FTCHARToUTF8 ConvertedValue(*TestString, TestString.Len() + 1); // include the null terminator
+			const bool bIsNullTerminated = *(ConvertedValue.Get() + ConvertedValue.Length()) == 0;
+			TestTrue(TEXT("Expected converted value to be null-terminated"), bIsNullTerminated);
+		}
+	}
+
 	return !HasAnyErrors();
 }
 
diff --git a/Engine/Source/Runtime/Core/Private/UObject/UnrealNames.cpp b/Engine/Source/Runtime/Core/Private/UObject/UnrealNames.cpp
--- a/Engine/Source/Runtime/Core/Private/UObject/UnrealNames.cpp
+++ b/Engine/Source/Runtime/Core/Private/UObject/UnrealNames.cpp
@@ -2143,6 +2143,8 @@ void FNameEntry::Write(FStructuredArchive::FSlot Slot) const
 	Slot << EntrySerialized;
 }
 
+static_assert(PLATFORM_LITTLE_ENDIAN, "FNameEntrySerialized serialization needs updating to support big-endian platforms!");
+
 FArchive& operator<<(FArchive& Ar, FNameEntrySerialized& E)
 {
 	if (Ar.IsLoading())
@@ -2182,12 +2184,15 @@ FArchive& operator<<(FArchive& Ar, FNameEntrySerialized& E)
 			// get the pointer to the wide array 
 			WIDECHAR* WideName = const_cast<WIDECHAR*>(E.GetWideName());
 
-			// read in the UCS2CHAR string and byteswap it, etc
+			// read in the UCS2CHAR string
 			auto Sink = StringMemoryPassthru<UCS2CHAR>(WideName, StringLen, StringLen);
 			Ar.Serialize(Sink.Get(), StringLen * sizeof(UCS2CHAR));
 			Sink.Apply();
 
-			INTEL_ORDER_TCHARARRAY(WideName)
+#if PLATFORM_TCHAR_IS_4_BYTES
+			// Inline combine any surrogate pairs in the data when loading into a UTF-32 string
+			StringLen = StringConv::InlineCombineSurrogates_Buffer(WideName, StringLen);
+#endif	// PLATFORM_TCHAR_IS_4_BYTES
 		}
 		else
 		{
diff --git a/Engine/Source/Runtime/Core/Public/Android/AndroidPlatform.h b/Engine/Source/Runtime/Core/Public/Android/AndroidPlatform.h
--- a/Engine/Source/Runtime/Core/Public/Android/AndroidPlatform.h
+++ b/Engine/Source/Runtime/Core/Public/Android/AndroidPlatform.h
@@ -42,6 +42,7 @@ typedef FAndroidTypes FPlatformTypes;
 #define PLATFORM_MAX_FILEPATH_LENGTH_DEPRECATED		ANDROID_MAX_PATH
 #define PLATFORM_SUPPORTS_TEXTURE_STREAMING			1
 #define PLATFORM_REQUIRES_FILESERVER				1
+#define PLATFORM_WCHAR_IS_4_BYTES					1
 #define PLATFORM_TCHAR_IS_CHAR16					1
 #define PLATFORM_HAS_NO_EPROCLIM					1
 #define PLATFORM_USES_ES2							1
diff --git a/Engine/Source/Runtime/Core/Public/Containers/StringConv.h b/Engine/Source/Runtime/Core/Public/Containers/StringConv.h
--- a/Engine/Source/Runtime/Core/Public/Containers/StringConv.h
+++ b/Engine/Source/Runtime/Core/Public/Containers/StringConv.h
@@ -34,7 +34,7 @@ public:
 	}
 };
 
-namespace UE4StringConv_Private
+namespace StringConv
 {
 	constexpr const uint16 HIGH_SURROGATE_START_CODEPOINT = 0xD800;
 	constexpr const uint16 HIGH_SURROGATE_END_CODEPOINT   = 0xDBFF;
@@ -43,6 +43,17 @@ namespace UE4StringConv_Private
 	constexpr const uint32 ENCODED_SURROGATE_START_CODEPOINT = 0x10000;
 	constexpr const uint32 ENCODED_SURROGATE_END_CODEPOINT   = 0x10FFFF;
 
+	/** Is the provided Codepoint within the range of valid codepoints? */
+	static FORCEINLINE bool IsValidCodepoint(const uint32 Codepoint)
+	{
+		if ((Codepoint > 0x10FFFF) ||						// No Unicode codepoints above 10FFFFh, (for now!)
+			(Codepoint == 0xFFFE) || (Codepoint == 0xFFFF)) // illegal values.
+		{
+			return false;
+		}
+		return true;
+	}
+
 	/** Is the provided Codepoint within the range of the high-surrogates? */
 	static FORCEINLINE bool IsHighSurrogate(const uint32 Codepoint)
 	{
@@ -55,12 +66,96 @@ namespace UE4StringConv_Private
 		return Codepoint >= LOW_SURROGATE_START_CODEPOINT && Codepoint <= LOW_SURROGATE_END_CODEPOINT;
 	}
 
+	static FORCEINLINE uint32 EncodeSurrogate(const uint16 HighSurrogate, const uint16 LowSurrogate)
+	{
+		return ((HighSurrogate - HIGH_SURROGATE_START_CODEPOINT) << 10) + (LowSurrogate - LOW_SURROGATE_START_CODEPOINT) + 0x10000;
+	}
+
+	static FORCEINLINE void DecodeSurrogate(const uint32 Codepoint, uint16& OutHighSurrogate, uint16& OutLowSurrogate)
+	{
+		const uint32 TmpCodepoint = Codepoint - 0x10000;
+		OutHighSurrogate = (TmpCodepoint >> 10) + HIGH_SURROGATE_START_CODEPOINT;
+		OutLowSurrogate = (TmpCodepoint & 0x3FF) + LOW_SURROGATE_START_CODEPOINT;
+	}
+
 	/** Is the provided Codepoint outside of the range of the basic multilingual plane, but within the valid range of UTF8/16? */
 	static FORCEINLINE bool IsEncodedSurrogate(const uint32 Codepoint)
 	{
 		return Codepoint >= ENCODED_SURROGATE_START_CODEPOINT && Codepoint <= ENCODED_SURROGATE_END_CODEPOINT;
 	}
 
+	/** Inline combine any UTF-16 surrogate pairs in the given null-terminated character buffer, returning the updated length */
+	template<typename CharType>
+	static FORCEINLINE int32 InlineCombineSurrogates_Buffer(CharType* StrBuffer, int32 StrLen)
+	{
+		static_assert(sizeof(CharType) == 4, "CharType must be 4-bytes!");
+
+		for (int32 Index = 0; Index < StrLen; ++Index)
+		{
+			CharType Codepoint = StrBuffer[Index];
+
+			// Check if this character is a high-surrogate
+			if (StringConv::IsHighSurrogate(Codepoint))
+			{
+				// Ensure our string has enough characters to read from
+				if ((Index + 1) >= StrLen)
+				{
+					// Unpaired surrogate - replace with bogus char
+					StrBuffer[Index] = UNICODE_BOGUS_CHAR_CODEPOINT;
+					break;
+				}
+
+				const uint32 HighCodepoint = Codepoint;
+				Codepoint = StrBuffer[Index + 1];
+
+				// If our High Surrogate is set, check if this character is the matching low-surrogate
+				if (StringConv::IsLowSurrogate(Codepoint))
+				{
+					const uint32 LowCodepoint = Codepoint;
+
+					// Combine our high and low surrogates together to a single Unicode codepoint
+					Codepoint = StringConv::EncodeSurrogate(HighCodepoint, LowCodepoint);
+
+					StrBuffer[Index] = Codepoint;
+					{
+						const int32 ArrayNum = StrLen + 1;
+						const int32 IndexToRemove = Index + 1;
+						const int32 NumToMove = ArrayNum - IndexToRemove - 1;
+						if (NumToMove > 0)
+						{
+							FMemory::Memmove(StrBuffer + IndexToRemove, StrBuffer + IndexToRemove + 1, NumToMove * sizeof(CharType));
+						}
+					}
+					--StrLen;
+					continue;
+				}
+
+				// Unpaired surrogate - replace with bogus char
+				StrBuffer[Index] = UNICODE_BOGUS_CHAR_CODEPOINT;
+			}
+			else if (StringConv::IsLowSurrogate(Codepoint))
+			{
+				// Unpaired surrogate - replace with bogus char
+				StrBuffer[Index] = UNICODE_BOGUS_CHAR_CODEPOINT;
+			}
+		}
+
+		return StrLen;
+	}
+
+	/** Inline combine any UTF-16 surrogate pairs in the given null-terminated TCHAR array */
+	template<typename AllocatorType>
+	static FORCEINLINE void InlineCombineSurrogates_Array(TArray<TCHAR, AllocatorType>& StrBuffer)
+	{
+#if PLATFORM_TCHAR_IS_4_BYTES
+		const int32 NewStrLen = InlineCombineSurrogates_Buffer(StrBuffer.GetData(), StrBuffer.Num() - 1);
+		StrBuffer.SetNum(NewStrLen + 1, /*bAllowShrinking*/false);
+#endif	// PLATFORM_TCHAR_IS_4_BYTES
+	}
+}
+
+namespace UE4StringConv_Private
+{
 	/**
 	 * This is a basic object which counts how many times it has been incremented
 	 */
@@ -117,15 +212,11 @@ public:
 
 		const BufferType OutputIteratorStartPosition = OutputIterator;
 
-		if (Codepoint > 0x10FFFF)   // No Unicode codepoints above 10FFFFh, (for now!)
-		{
-			Codepoint = UNICODE_BOGUS_CHAR_CODEPOINT;
-		}
-		else if ((Codepoint == 0xFFFE) || (Codepoint == 0xFFFF))  // illegal values.
+		if (!StringConv::IsValidCodepoint(Codepoint))
 		{
 			Codepoint = UNICODE_BOGUS_CHAR_CODEPOINT;
 		}
-		else if (UE4StringConv_Private::IsHighSurrogate(Codepoint) || UE4StringConv_Private::IsLowSurrogate(Codepoint)) // UTF-8 Characters are not allowed to encode codepoints in the surrogate pair range
+		else if (StringConv::IsHighSurrogate(Codepoint) || StringConv::IsLowSurrogate(Codepoint)) // UTF-8 Characters are not allowed to encode codepoints in the surrogate pair range
 		{
 			Codepoint = UNICODE_BOGUS_CHAR_CODEPOINT;
 		}
@@ -209,22 +300,34 @@ private:
 	template <typename DestBufferType>
 	static void Convert_Impl(DestBufferType& Dest, int32 DestLen, const TCHAR* Source, const int32 SourceLen)
 	{
-		uint32 HighCodepoint = MAX_uint32;
+#if PLATFORM_TCHAR_IS_4_BYTES
+		for (int32 i = 0; i < SourceLen; ++i)
+		{
+			uint32 Codepoint = static_cast<uint32>(Source[i]);
+
+			if (!WriteCodepointToBuffer(Codepoint, Dest, DestLen))
+			{
+				// Could not write data, bail out
+				return;
+			}
+		}
+#else	// PLATFORM_TCHAR_IS_4_BYTES
+		uint32 HighSurrogate = MAX_uint32;
 
 		for (int32 i = 0; i < SourceLen; ++i)
 		{
-			const bool bHighSurrogateIsSet = HighCodepoint != MAX_uint32;
+			const bool bHighSurrogateIsSet = HighSurrogate != MAX_uint32;
 			uint32 Codepoint = static_cast<uint32>(Source[i]);
 
 			// Check if this character is a high-surrogate
-			if (UE4StringConv_Private::IsHighSurrogate(Codepoint))
+			if (StringConv::IsHighSurrogate(Codepoint))
 			{
 				// Ensure we don't already have a high-surrogate set
 				if (bHighSurrogateIsSet)
 				{
 					// Already have a high-surrogate in this pair
 					// Write our stored value (will be converted into bogus character)
-					if (!WriteCodepointToBuffer(HighCodepoint, Dest, DestLen))
+					if (!WriteCodepointToBuffer(HighSurrogate, Dest, DestLen))
 					{
 						// Could not write data, bail out
 						return;
@@ -232,23 +335,23 @@ private:
 				}
 
 				// Store our code point for our next character
-				HighCodepoint = Codepoint;
+				HighSurrogate = Codepoint;
 				continue;
 			}
 
 			// If our High Surrogate is set, check if this character is the matching low-surrogate
 			if (bHighSurrogateIsSet)
 			{
-				if (UE4StringConv_Private::IsLowSurrogate(Codepoint))
+				if (StringConv::IsLowSurrogate(Codepoint))
 				{
-					const uint32 LowCodepoint = Codepoint;
+					const uint32 LowSurrogate = Codepoint;
 					// Combine our high and low surrogates together to a single Unicode codepoint
-					Codepoint = ((HighCodepoint - UE4StringConv_Private::HIGH_SURROGATE_START_CODEPOINT) << 10) + (LowCodepoint - UE4StringConv_Private::LOW_SURROGATE_START_CODEPOINT) + 0x10000;
+					Codepoint = StringConv::EncodeSurrogate(HighSurrogate, LowSurrogate);
 				}
 				else
 				{
-					// Did not find matching low-surrogate, write out a bogus character for our stored HighCodepoint
-					if (!WriteCodepointToBuffer(HighCodepoint, Dest, DestLen))
+					// Did not find matching low-surrogate, write out a bogus character for our stored HighSurrogate
+					if (!WriteCodepointToBuffer(HighSurrogate, Dest, DestLen))
 					{
 						// Could not write data, bail out
 						return;
@@ -256,7 +359,7 @@ private:
 				}
 
 				// Reset our high-surrogate now that we've used (or discarded) its value
-				HighCodepoint = MAX_uint32;
+				HighSurrogate = MAX_uint32;
 			}
 
 			if (!WriteCodepointToBuffer(Codepoint, Dest, DestLen))
@@ -265,6 +368,7 @@ private:
 				return;
 			}
 		}
+#endif	// PLATFORM_TCHAR_IS_4_BYTES
 	}
 
 	template <typename DestBufferType>
@@ -289,7 +393,7 @@ public:
 	typedef TCHAR    ToType;
 
 	/**
-	 * Converts the UTF-8 string to UTF-16.
+	 * Converts the UTF-8 string to UTF-16 or UTF-32.
 	 *
 	 * @param Dest      The destination buffer of the converted string.
 	 * @param DestLen   The length of the destination buffer.
@@ -306,7 +410,7 @@ public:
 	 *
 	 * @param Source string to read and determine amount of TCHARs to represent
 	 * @param SourceLen Length of source string; we will not read past this amount, even if the characters tell us to
-	 * @return The length of the string in UTF-16 characters.
+	 * @return The length of the string in UTF-16 or UTF-32 characters.
 	 */
 	static int32 ConvertedLength(const ANSICHAR* Source, const int32 SourceLen)
 	{
@@ -319,7 +423,7 @@ public:
 private:
 	static uint32 CodepointFromUtf8(const ANSICHAR*& SourceString, const uint32 SourceLengthRemaining)
 	{
-		checkSlow(SourceLengthRemaining > 0)
+		checkSlow(SourceLengthRemaining > 0);
 
 		const ANSICHAR* OctetPtr = SourceString;
 
@@ -392,7 +496,7 @@ private:
 			Codepoint = ( ((Octet << 12)) | ((Octet2-128) << 6) | ((Octet3-128)) );
 
 			// UTF-8 characters cannot be in the UTF-16 surrogates range
-			if (UE4StringConv_Private::IsHighSurrogate(Codepoint) || UE4StringConv_Private::IsLowSurrogate(Codepoint))
+			if (StringConv::IsHighSurrogate(Codepoint) || StringConv::IsLowSurrogate(Codepoint))
 			{
 				++SourceString;  // Sequence was not valid UTF-8. Skip the first byte and continue.
 				return UNICODE_BOGUS_CHAR_CODEPOINT;
@@ -555,17 +659,19 @@ private:
 			// Read our codepoint, advancing the source pointer
 			uint32 Codepoint = CodepointFromUtf8(Source, SourceEnd - Source);
 
+#if !PLATFORM_TCHAR_IS_4_BYTES
 			// We want to write out two chars
-			if (UE4StringConv_Private::IsEncodedSurrogate(Codepoint))
+			if (StringConv::IsEncodedSurrogate(Codepoint))
 			{
 				// We need two characters to write the surrogate pair
 				if (DestLen >= 2)
 				{
-					Codepoint -= 0x10000;
-					const TCHAR HighSurrogate = (Codepoint >> 10) + UE4StringConv_Private::HIGH_SURROGATE_START_CODEPOINT;
-					const TCHAR LowSurrogate = (Codepoint & 0x3FF) + UE4StringConv_Private::LOW_SURROGATE_START_CODEPOINT;
-					*(ConvertedBuffer++) = HighSurrogate;
-					*(ConvertedBuffer++) = LowSurrogate;
+					uint16 HighSurrogate = 0;
+					uint16 LowSurrogate = 0;
+					StringConv::DecodeSurrogate(Codepoint, HighSurrogate, LowSurrogate);
+
+					*(ConvertedBuffer++) = (TCHAR)HighSurrogate;
+					*(ConvertedBuffer++) = (TCHAR)LowSurrogate;
 					DestLen -= 2;
 					continue;
 				}
@@ -573,11 +679,248 @@ private:
 				// If we don't have space, write a bogus character instead (we should have space for it)
 				Codepoint = UNICODE_BOGUS_CHAR_CODEPOINT;
 			}
-			else if (Codepoint > UE4StringConv_Private::ENCODED_SURROGATE_END_CODEPOINT)
+			else if (Codepoint > StringConv::ENCODED_SURROGATE_END_CODEPOINT)
 			{
 				// Ignore values higher than the supplementary plane range
 				Codepoint = UNICODE_BOGUS_CHAR_CODEPOINT;
 			}
+#endif	// !PLATFORM_TCHAR_IS_4_BYTES
+
+			*(ConvertedBuffer++) = Codepoint;
+			--DestLen;
+		}
+	}
+};
+
+template<typename InFromType, typename InToType>
+class TUTF32ToUTF16_Convert
+{
+	static_assert(sizeof(InFromType) == 4, "FromType must be 4 bytes!");
+	static_assert(sizeof(InToType) == 2, "ToType must be 2 bytes!");
+
+public:
+	typedef InFromType FromType;
+	typedef InToType   ToType;
+
+	/**
+	 * Convert Codepoint into UTF-16 characters.
+	 *
+	 * @param Codepoint Codepoint to expand into UTF-16 code units
+	 * @param OutputIterator Output iterator to write UTF-16 code units into
+	 * @param OutputIteratorNumRemaining Maximum number of UTF-16 code units that can be written to OutputIterator
+	 * @return Number of characters written for Codepoint
+	 */
+	template <typename BufferType>
+	static int32 Utf16FromCodepoint(uint32 Codepoint, BufferType OutputIterator, uint32 OutputIteratorNumRemaining)
+	{
+		// Ensure we have at least one character in size to write
+		checkSlow(OutputIteratorNumRemaining >= 1);
+
+		const BufferType OutputIteratorStartPosition = OutputIterator;
+
+		if (!StringConv::IsValidCodepoint(Codepoint))
+		{
+			Codepoint = UNICODE_BOGUS_CHAR_CODEPOINT;
+		}
+		else if (StringConv::IsHighSurrogate(Codepoint) || StringConv::IsLowSurrogate(Codepoint)) // Surrogate pairs range should not be present in UTF-32
+		{
+			Codepoint = UNICODE_BOGUS_CHAR_CODEPOINT;
+		}
+
+		// We want to write out two chars
+		if (StringConv::IsEncodedSurrogate(Codepoint))
+		{
+			// We need two characters to write the surrogate pair
+			if (OutputIteratorNumRemaining >= 2)
+			{
+				uint16 HighSurrogate = 0;
+				uint16 LowSurrogate = 0;
+				StringConv::DecodeSurrogate(Codepoint, HighSurrogate, LowSurrogate);
+
+				*(OutputIterator++) = (ToType)HighSurrogate;
+				*(OutputIterator++) = (ToType)LowSurrogate;
+			}
+			else
+			{
+				// If we don't have space, write a bogus character instead (we should have space for it)
+				*(OutputIterator++) = UNICODE_BOGUS_CHAR_CODEPOINT;
+			}
+		}
+		else if (Codepoint > StringConv::ENCODED_SURROGATE_END_CODEPOINT)
+		{
+			// Ignore values higher than the supplementary plane range
+			*(OutputIterator++) = UNICODE_BOGUS_CHAR_CODEPOINT;
+		}
+		else
+		{
+			// Normal codepoint
+			*(OutputIterator++) = Codepoint;
+		}
+
+		return static_cast<int32>(OutputIterator - OutputIteratorStartPosition);
+	}
+
+	/**
+	 * Converts the string to the desired format.
+	 *
+	 * @param Dest      The destination buffer of the converted string.
+	 * @param DestLen   The length of the destination buffer.
+	 * @param Source    The source string to convert.
+	 * @param SourceLen The length of the source string.
+	 */
+	static FORCEINLINE void Convert(ToType* Dest, int32 DestLen, const FromType* Source, int32 SourceLen)
+	{
+		Convert_Impl(Dest, DestLen, Source, SourceLen);
+	}
+
+	/**
+	 * Determines the length of the converted string.
+	 *
+	 * @return The length of the string in UTF-16 code units.
+	 */
+	static FORCEINLINE int32 ConvertedLength(const FromType* Source, int32 SourceLen)
+	{
+		UE4StringConv_Private::FCountingOutputIterator Dest;
+		const int32 DestLen = SourceLen * 2;
+		Convert_Impl(Dest, DestLen, Source, SourceLen);
+
+		return Dest.GetCount();
+	}
+
+private:
+	template <typename DestBufferType>
+	static void Convert_Impl(DestBufferType& Dest, int32 DestLen, const FromType* Source, const int32 SourceLen)
+	{
+		for (int32 i = 0; i < SourceLen; ++i)
+		{
+			uint32 Codepoint = static_cast<uint32>(Source[i]);
+
+			if (!WriteCodepointToBuffer(Codepoint, Dest, DestLen))
+			{
+				// Could not write data, bail out
+				return;
+			}
+		}
+	}
+
+	template <typename DestBufferType>
+	static bool WriteCodepointToBuffer(const uint32 Codepoint, DestBufferType& Dest, int32& DestLen)
+	{
+		int32 WrittenChars = Utf16FromCodepoint(Codepoint, Dest, DestLen);
+		if (WrittenChars < 1)
+		{
+			return false;
+		}
+
+		Dest += WrittenChars;
+		DestLen -= WrittenChars;
+		return true;
+	}
+};
+
+template<typename InFromType, typename InToType>
+class TUTF16ToUTF32_Convert
+{
+	static_assert(sizeof(InFromType) == 2, "FromType must be 2 bytes!");
+	static_assert(sizeof(InToType) == 4, "ToType must be 4 bytes!");
+
+public:
+	typedef InFromType FromType;
+	typedef InToType   ToType;
+
+	/**
+	 * Converts the UTF-16 string to UTF-32.
+	 *
+	 * @param Dest      The destination buffer of the converted string.
+	 * @param DestLen   The length of the destination buffer.
+	 * @param Source    The source string to convert.
+	 * @param SourceLen The length of the source string.
+	 */
+	static FORCEINLINE void Convert(ToType* Dest, const int32 DestLen, const FromType* Source, const int32 SourceLen)
+	{
+		Convert_Impl(Dest, DestLen, Source, SourceLen);
+	}
+
+	/**
+	 * Determines the length of the converted string.
+	 *
+	 * @param Source string to read and determine amount of ToType chars to represent
+	 * @param SourceLen Length of source string; we will not read past this amount, even if the characters tell us to
+	 * @return The length of the string in UTF-32 characters.
+	 */
+	static int32 ConvertedLength(const FromType* Source, const int32 SourceLen)
+	{
+		UE4StringConv_Private::FCountingOutputIterator Dest;
+		Convert_Impl(Dest, MAX_int32, Source, SourceLen);
+
+		return Dest.GetCount();
+	}
+
+private:
+	static uint32 CodepointFromUtf16(const FromType*& SourceString, const uint32 SourceLengthRemaining)
+	{
+		checkSlow(SourceLengthRemaining > 0);
+
+		const FromType* CodeUnitPtr = SourceString;
+
+		uint32 Codepoint = *CodeUnitPtr;
+
+		// Check if this character is a high-surrogate
+		if (StringConv::IsHighSurrogate(Codepoint))
+		{
+			// Ensure our string has enough characters to read from
+			if (SourceLengthRemaining < 2)
+			{
+				// Skip to end and write out a single char (we always have room for at least 1 char)
+				SourceString += SourceLengthRemaining;
+				return UNICODE_BOGUS_CHAR_CODEPOINT;
+			}
+
+			const uint32 HighSurrogate = Codepoint;
+			Codepoint = *(++CodeUnitPtr);
+
+			// If our High Surrogate is set, check if this character is the matching low-surrogate
+			if (StringConv::IsLowSurrogate(Codepoint))
+			{
+				const uint32 LowSurrogate = Codepoint;
+
+				// Combine our high and low surrogates together to a single Unicode codepoint
+				Codepoint = StringConv::EncodeSurrogate(HighSurrogate, LowSurrogate);
+
+				SourceString += 2;  // skip to next possible start of codepoint.
+				return Codepoint;
+			}
+			else
+			{
+				++SourceString; // Did not find matching low-surrogate, write out a bogus character and continue
+				return UNICODE_BOGUS_CHAR_CODEPOINT;
+			}
+		}
+		else if (StringConv::IsLowSurrogate(Codepoint))
+		{
+			// Unpaired low surrogate
+			++SourceString;
+			return UNICODE_BOGUS_CHAR_CODEPOINT;
+		}
+		else
+		{
+			// Single codepoint
+			++SourceString;
+			return Codepoint;
+		}
+	}
+
+	/**
+	 * Read Source string, converting the data from UTF-16 into UTF-32, and placing these in the Destination
+	 */
+	template <typename DestBufferType>
+	static void Convert_Impl(DestBufferType& ConvertedBuffer, int32 DestLen, const FromType* Source, const int32 SourceLen)
+	{
+		const FromType* SourceEnd = Source + SourceLen;
+		while (Source < SourceEnd && DestLen > 0)
+		{
+			// Read our codepoint, advancing the source pointer
+			uint32 Codepoint = CodepointFromUtf16(Source, SourceEnd - Source);
 
 			*(ConvertedBuffer++) = Codepoint;
 			--DestLen;
@@ -632,7 +975,7 @@ public:
 		}
 		else
 		{
-			Ptr          = NULL;
+			Ptr = nullptr;
 			StringLength = 0;
 		}
 	}
@@ -641,11 +984,19 @@ public:
 	{
 		if (Source)
 		{
-			Init(Source, SourceLen, ENullTerminatedString::No);
+			ENullTerminatedString::Type NullTerminated = ENullTerminatedString::No;
+			if (SourceLen > 0 && Source[SourceLen-1] == 0)
+			{
+				// Given buffer is null-terminated
+				NullTerminated = ENullTerminatedString::Yes;
+				SourceLen -= 1;
+			}
+
+			Init(Source, SourceLen, NullTerminated);
 		}
 		else
 		{
-			Ptr          = NULL;
+			Ptr = nullptr;
 			StringLength = 0;
 		}
 	}
@@ -661,8 +1012,9 @@ public:
 
 	/**
 	 * Accessor for the converted string.
+	 * @note The string may not be null-terminated if constructed from an explicitly sized buffer that didn't include the null-terminator.
 	 *
-	 * @return A const pointer to the null-terminated converted string.
+	 * @return A const pointer to the converted string.
 	 */
 	FORCEINLINE const ToType* Get() const
 	{
@@ -681,14 +1033,98 @@ public:
 
 private:
 	// Non-copyable
-	TStringConversion(const TStringConversion&);
-	TStringConversion& operator=(const TStringConversion&);
+	TStringConversion(const TStringConversion&) = delete;
+	TStringConversion& operator=(const TStringConversion&) = delete;
 
 	ToType* Ptr;
 	int32   StringLength;
 };
 
 
+/**
+ * Class takes one type of string and and stores it as-is.
+ * For API compatibility with TStringConversion when the To and From types are already in the same format.
+ */
+template<typename InFromType, typename InToType = InFromType>
+class TStringPointer
+{
+	static_assert(sizeof(InFromType) == sizeof(InToType), "FromType must be the same size as ToType!");
+
+public:
+	typedef InFromType FromType;
+	typedef InToType   ToType;
+
+public:
+	explicit TStringPointer(const FromType* Source)
+	{
+		if (Source)
+		{
+			Ptr = (const ToType*)Source;
+			StringLength = -1; // Length calculated on-demand
+		}
+		else
+		{
+			Ptr = nullptr;
+			StringLength = 0;
+		}
+	}
+
+	TStringPointer(const FromType* Source, int32 SourceLen)
+	{
+		if (Source)
+		{
+			if (SourceLen > 0 && Source[SourceLen-1] == 0)
+			{
+				// Given buffer is null-terminated
+				SourceLen -= 1;
+			}
+
+			Ptr = (const ToType*)Source;
+			StringLength = SourceLen;
+		}
+		else
+		{
+			Ptr = nullptr;
+			StringLength = 0;
+		}
+	}
+
+	/**
+	 * Move constructor
+	 */
+	TStringPointer(TStringPointer&& Other) = default;
+
+	/**
+	 * Accessor for the string.
+	 * @note The string may not be null-terminated if constructed from an explicitly sized buffer that didn't include the null-terminator.
+	 *
+	 * @return A const pointer to the string.
+	 */
+	FORCEINLINE const ToType* Get() const
+	{
+		return Ptr;
+	}
+
+	/**
+	 * Length of the string.
+	 *
+	 * @return The number of characters in the string, excluding any null terminator.
+	 */
+	FORCEINLINE int32 Length() const
+	{
+		return StringLength == -1 ? TCString<ToType>::Strlen(Ptr) : StringLength;
+	}
+
+private:
+	// Non-copyable
+	TStringPointer(const TStringPointer&) = delete;
+	TStringPointer& operator=(const TStringPointer&) = delete;
+
+	const ToType* Ptr;
+	int32 StringLength;
+};
+
+
 /**
  * NOTE: The objects these macros declare have very short lifetimes. They are
  * meant to be used as parameters to functions. You cannot assign a variable
@@ -716,39 +1152,33 @@ typedef TStringConversion<FUTF8ToTCHAR_Convert> FUTF8ToTCHAR;
 #define TCHAR_TO_UTF8(str) (ANSICHAR*)FTCHARToUTF8((const TCHAR*)str).Get()
 #define UTF8_TO_TCHAR(str) (TCHAR*)FUTF8ToTCHAR((const ANSICHAR*)str).Get()
 
-// special needs handling for going from char16_t to wchar_t for third party libraries that need wchar_t
-#if PLATFORM_TCHAR_IS_CHAR16
-#define TCHAR_TO_WCHAR(str) (wchar_t*)StringCast<wchar_t>(static_cast<const TCHAR*>(str)).Get()
-#define WCHAR_TO_TCHAR(str) (TCHAR*)StringCast<TCHAR>(static_cast<const wchar_t*>(str)).Get()
+// special handling for platforms still using a 32-bit TCHAR
+#if PLATFORM_TCHAR_IS_4_BYTES
+typedef TStringConversion<TUTF32ToUTF16_Convert<TCHAR, UTF16CHAR>> FTCHARToUTF16;
+typedef TStringConversion<TUTF16ToUTF32_Convert<UTF16CHAR, TCHAR>> FUTF16ToTCHAR;
+#define TCHAR_TO_UTF16(str) (UTF16CHAR*)FTCHARToUTF16((const TCHAR*)str).Get()
+#define UTF16_TO_TCHAR(str) (TCHAR*)FUTF16ToTCHAR((const UTF16CHAR*)str).Get()
 #else
-#define TCHAR_TO_WCHAR(str) str
-#define WCHAR_TO_TCHAR(str) str
+static_assert(sizeof(TCHAR) == sizeof(UTF16CHAR), "TCHAR and UTF16CHAR are expected to be the same size for inline conversion! PLATFORM_TCHAR_IS_4_BYTES is not configured correctly for this platform.");
+typedef TStringPointer<TCHAR, UTF16CHAR> FTCHARToUTF16;
+typedef TStringPointer<UTF16CHAR, TCHAR> FUTF16ToTCHAR;
+#define TCHAR_TO_UTF16(str) (UTF16CHAR*)(str)
+#define UTF16_TO_TCHAR(str) (TCHAR*)(str)
 #endif
 
-// This seemingly-pointless class is intended to be API-compatible with TStringConversion
-// and is returned by StringCast when no string conversion is necessary.
-template <typename T>
-class TStringPointer
-{
-public:
-	FORCEINLINE explicit TStringPointer(const T* InPtr)
-		: Ptr(InPtr)
-	{
-	}
-
-	FORCEINLINE const T* Get() const
-	{
-		return Ptr;
-	}
-
-	FORCEINLINE int32 Length() const
-	{
-		return Ptr ? TCString<T>::Strlen(Ptr) : 0;
-	}
-
-private:
-	const T* Ptr;
-};
+// special handling for going from char16_t to wchar_t for third party libraries that need wchar_t
+#if PLATFORM_TCHAR_IS_CHAR16 && PLATFORM_WCHAR_IS_4_BYTES
+typedef TStringConversion<TUTF16ToUTF32_Convert<TCHAR, wchar_t>> FTCHARToWChar;
+typedef TStringConversion<TUTF32ToUTF16_Convert<wchar_t, TCHAR>> FWCharToTCHAR;
+#define TCHAR_TO_WCHAR(str) (wchar_t*)FTCHARToWChar((const TCHAR*)str).Get()
+#define WCHAR_TO_TCHAR(str) (TCHAR*)FWCharToTCHAR((const wchar_t*)str).Get()
+#else
+static_assert(sizeof(TCHAR) == sizeof(wchar_t), "TCHAR and wchar_t are expected to be the same size for inline conversion! PLATFORM_WCHAR_IS_4_BYTES is not configured correctly for this platform.");
+typedef TStringPointer<TCHAR, wchar_t> FTCHARToWChar;
+typedef TStringPointer<wchar_t, TCHAR> FWCharToTCHAR;
+#define TCHAR_TO_WCHAR(str) (wchar_t*)(str)
+#define WCHAR_TO_TCHAR(str) (TCHAR*)(str)
+#endif
 
 /**
  * StringCast example usage:
@@ -792,7 +1222,7 @@ FORCEINLINE typename TEnableIf<!FPlatformString::TAreEncodingsCompatible<To, Fro
 template <typename To, typename From>
 FORCEINLINE typename TEnableIf<FPlatformString::TAreEncodingsCompatible<To, From>::Value, TStringPointer<To>>::Type StringCast(const From* Str, int32 Len)
 {
-	return TStringPointer<To>((const To*)Str);
+	return TStringPointer<To>((const To*)Str, Len);
 }
 
 /**
diff --git a/Engine/Source/Runtime/Core/Public/Containers/UnrealString.h b/Engine/Source/Runtime/Core/Public/Containers/UnrealString.h
--- a/Engine/Source/Runtime/Core/Public/Containers/UnrealString.h
+++ b/Engine/Source/Runtime/Core/Public/Containers/UnrealString.h
@@ -2238,6 +2238,18 @@ FORCEINLINE const TCHAR* ToCStr(const FString& Str)
 	return *Str;
 }
 
+namespace StringConv
+{
+	template<typename AllocatorType>
+	static void InlineCombineSurrogates_Array(TArray<TCHAR, AllocatorType>& StrBuffer);
+
+	/** Inline combine any UTF-16 surrogate pairs in the given string */
+	static FORCEINLINE void InlineCombineSurrogates(FString& Str)
+	{
+		InlineCombineSurrogates_Array(Str.GetCharArray());
+	}
+}
+
 /*----------------------------------------------------------------------------
 	Special archivers.
 ----------------------------------------------------------------------------*/
diff --git a/Engine/Source/Runtime/Core/Public/HAL/Platform.h b/Engine/Source/Runtime/Core/Public/HAL/Platform.h
--- a/Engine/Source/Runtime/Core/Public/HAL/Platform.h
+++ b/Engine/Source/Runtime/Core/Public/HAL/Platform.h
@@ -226,6 +226,9 @@
 #ifndef PLATFORM_TCHAR_IS_4_BYTES
 	#define PLATFORM_TCHAR_IS_4_BYTES			0
 #endif
+#ifndef PLATFORM_WCHAR_IS_4_BYTES
+	#define PLATFORM_WCHAR_IS_4_BYTES			0
+#endif
 #ifndef PLATFORM_TCHAR_IS_CHAR16
 	#define PLATFORM_TCHAR_IS_CHAR16			0
 #endif
@@ -863,6 +866,9 @@ namespace TypeTests
 	static_assert(!PLATFORM_TCHAR_IS_4_BYTES || sizeof(TCHAR) == 4, "TCHAR size must be 4 bytes.");
 	static_assert(PLATFORM_TCHAR_IS_4_BYTES || sizeof(TCHAR) == 2, "TCHAR size must be 2 bytes.");
 
+	static_assert(!PLATFORM_WCHAR_IS_4_BYTES || sizeof(wchar_t) == 4, "wchar_t size must be 4 bytes.");
+	static_assert(PLATFORM_WCHAR_IS_4_BYTES || sizeof(wchar_t) == 2, "wchar_t size must be 2 bytes.");
+
 	static_assert(PLATFORM_32BITS || PLATFORM_64BITS, "Type tests pointer size failed.");
 	static_assert(PLATFORM_32BITS != PLATFORM_64BITS, "Type tests pointer exclusive failed.");
 	static_assert(!PLATFORM_64BITS || sizeof(void*) == 8, "Pointer size is 64bit, but pointers are short.");
diff --git a/Engine/Source/Runtime/Core/Public/HTML5/HTML5Platform.h b/Engine/Source/Runtime/Core/Public/HTML5/HTML5Platform.h
--- a/Engine/Source/Runtime/Core/Public/HTML5/HTML5Platform.h
+++ b/Engine/Source/Runtime/Core/Public/HTML5/HTML5Platform.h
@@ -29,6 +29,7 @@ typedef FHTML5Types FPlatformTypes;
 #define PLATFORM_USE_LS_SPEC_FOR_WIDECHAR			1
 #define PLATFORM_HAS_BSD_TIME						1
 #define PLATFORM_MAX_FILEPATH_LENGTH_DEPRECATED		PATH_MAX
+#define PLATFORM_WCHAR_IS_4_BYTES					1
 #define PLATFORM_TCHAR_IS_4_BYTES					1
 #define PLATFORM_HAS_BSD_SOCKETS					1
 #define PLATFORM_HAS_BSD_SOCKET_FEATURE_MSG_DONTWAIT	1
diff --git a/Engine/Source/Runtime/Core/Public/IOS/IOSPlatform.h b/Engine/Source/Runtime/Core/Public/IOS/IOSPlatform.h
--- a/Engine/Source/Runtime/Core/Public/IOS/IOSPlatform.h
+++ b/Engine/Source/Runtime/Core/Public/IOS/IOSPlatform.h
@@ -35,6 +35,7 @@ typedef FIOSPlatformTypes FPlatformTypes;
 #define PLATFORM_LITTLE_ENDIAN							1
 #define PLATFORM_SUPPORTS_PRAGMA_PACK					1
 #define PLATFORM_COMPILER_DISTINGUISHES_INT_AND_LONG	1
+#define PLATFORM_WCHAR_IS_4_BYTES						1
 #define PLATFORM_TCHAR_IS_CHAR16						1
 #define PLATFORM_USE_SYSTEM_VSWPRINTF					0
 #define PLATFORM_HAS_BSD_TIME							1
diff --git a/Engine/Source/Runtime/Core/Public/Mac/MacPlatform.h b/Engine/Source/Runtime/Core/Public/Mac/MacPlatform.h
--- a/Engine/Source/Runtime/Core/Public/Mac/MacPlatform.h
+++ b/Engine/Source/Runtime/Core/Public/Mac/MacPlatform.h
@@ -35,6 +35,7 @@ typedef FMacPlatformTypes FPlatformTypes;
 //#define PLATFORM_USE_LS_SPEC_FOR_WIDECHAR			1
 #define PLATFORM_USE_SYSTEM_VSWPRINTF				0
 #define PLATFORM_COMPILER_DISTINGUISHES_INT_AND_LONG			1
+#define PLATFORM_WCHAR_IS_4_BYTES					1
 #define PLATFORM_TCHAR_IS_4_BYTES					1
 #define PLATFORM_HAS_BSD_TIME						1
 //#define PLATFORM_USE_PTHREADS						1
diff --git a/Engine/Source/Runtime/Core/Public/Unix/UnixPlatform.h b/Engine/Source/Runtime/Core/Public/Unix/UnixPlatform.h
--- a/Engine/Source/Runtime/Core/Public/Unix/UnixPlatform.h
+++ b/Engine/Source/Runtime/Core/Public/Unix/UnixPlatform.h
@@ -45,6 +45,7 @@ typedef FUnixPlatformTypes FPlatformTypes;
 #define PLATFORM_COMPILER_DISTINGUISHES_INT_AND_LONG	1
 #define PLATFORM_SUPPORTS_PRAGMA_PACK					1
 #define PLATFORM_USE_LS_SPEC_FOR_WIDECHAR				1
+#define PLATFORM_WCHAR_IS_4_BYTES						1
 #if PLATFORM_UNIX_USE_CHAR16
 #define PLATFORM_TCHAR_IS_CHAR16						1
 #else
diff --git a/Engine/Source/Runtime/CoreUObject/Private/Serialization/PropertyLocalizationDataGathering.cpp b/Engine/Source/Runtime/CoreUObject/Private/Serialization/PropertyLocalizationDataGathering.cpp
--- a/Engine/Source/Runtime/CoreUObject/Private/Serialization/PropertyLocalizationDataGathering.cpp
+++ b/Engine/Source/Runtime/CoreUObject/Private/Serialization/PropertyLocalizationDataGathering.cpp
@@ -617,6 +617,9 @@ private:
 				iCode += sizeof(uint16);
 			}
 			while (Script[iCode-1] || Script[iCode-2]);
+
+			// Inline combine any surrogate pairs in the data when loading into a UTF-32 string
+			StringConv::InlineCombineSurrogates(LastParsedString);
 		}
 		else
 		{
diff --git a/Engine/Source/Runtime/CoreUObject/Private/UObject/ScriptCore.cpp b/Engine/Source/Runtime/CoreUObject/Private/UObject/ScriptCore.cpp
--- a/Engine/Source/Runtime/CoreUObject/Private/UObject/ScriptCore.cpp
+++ b/Engine/Source/Runtime/CoreUObject/Private/UObject/ScriptCore.cpp
@@ -3048,7 +3048,11 @@ IMPLEMENT_VM_FUNCTION( EX_StringConst, execStringConst );
 
 DEFINE_FUNCTION(UObject::execUnicodeStringConst)
 {
- 	*(FString*)RESULT_PARAM = FString((UCS2CHAR*)Stack.Code);
+	FString& ResultStr = *(FString*)RESULT_PARAM;
+	ResultStr = FString((UCS2CHAR*)Stack.Code);
+
+	// Inline combine any surrogate pairs in the data when loading into a UTF-32 string
+	StringConv::InlineCombineSurrogates(ResultStr);
 
 	while( *(uint16*)Stack.Code )
 	{
diff --git a/Engine/Source/Runtime/Json/Public/Policies/JsonPrintPolicy.h b/Engine/Source/Runtime/Json/Public/Policies/JsonPrintPolicy.h
--- a/Engine/Source/Runtime/Json/Public/Policies/JsonPrintPolicy.h
+++ b/Engine/Source/Runtime/Json/Public/Policies/JsonPrintPolicy.h
@@ -49,3 +49,21 @@ inline void TJsonPrintPolicy<TCHAR>::WriteString( FArchive* Stream, const FStrin
 {
 	Stream->Serialize((void*)*String, String.Len() * sizeof(TCHAR));
 }
+
+
+#if !PLATFORM_TCHAR_IS_CHAR16
+
+/**
+ * Specialization for UTF16CHAR that writes FString data UTF-16.
+ */
+template <>
+inline void TJsonPrintPolicy<UTF16CHAR>::WriteString(FArchive* Stream, const FString& String)
+{
+	// Note: This is a no-op on platforms that are using a 16-bit TCHAR
+	FTCHARToUTF16 UTF16String(*String, String.Len());
+
+	Stream->Serialize((void*)UTF16String.Get(), UTF16String.Length() * sizeof(UTF16CHAR));
+}
+
+#endif
+
diff --git a/Engine/Source/Runtime/Json/Public/Serialization/JsonReader.h b/Engine/Source/Runtime/Json/Public/Serialization/JsonReader.h
--- a/Engine/Source/Runtime/Json/Public/Serialization/JsonReader.h
+++ b/Engine/Source/Runtime/Json/Public/Serialization/JsonReader.h
@@ -591,6 +591,10 @@ private:
 		}
 
 		StringValue = MoveTemp(String);
+
+		// Inline combine any surrogate pairs in the data when loading into a UTF-32 string
+		StringConv::InlineCombineSurrogates(StringValue);
+
 		return true;
 	}
 
diff --git a/Engine/Source/Runtime/SlateCore/Private/Fonts/SlateTextShaper.cpp b/Engine/Source/Runtime/SlateCore/Private/Fonts/SlateTextShaper.cpp
--- a/Engine/Source/Runtime/SlateCore/Private/Fonts/SlateTextShaper.cpp
+++ b/Engine/Source/Runtime/SlateCore/Private/Fonts/SlateTextShaper.cpp
@@ -9,26 +9,6 @@
 DECLARE_CYCLE_STAT(TEXT("Shape Bidirectional Text"), STAT_SlateShapeBidirectionalText, STATGROUP_Slate);
 DECLARE_CYCLE_STAT(TEXT("Shape Unidirectional Text"), STAT_SlateShapeUnidirectionalText, STATGROUP_Slate);
 
-namespace SurrogatePairUtil
-{
-	template <bool IsUnicode, size_t TCHARSize>
-	bool IsSurrogatePairImpl(const TCHAR InHighChar, const TCHAR InLowChar)
-	{
-		return false;
-	}
-
-	template <>
-	bool IsSurrogatePairImpl<true, 2>(const TCHAR InHighChar, const TCHAR InLowChar)
-	{
-		return (InHighChar >= 0xD800 && InHighChar <= 0xDBFF) && (InLowChar >= 0xDC00 && InLowChar <= 0xDFFF);
-	}
-
-	bool IsSurrogatePair(const TCHAR InHighChar, const TCHAR InLowChar)
-	{
-		return IsSurrogatePairImpl<FPlatformString::IsUnicodeEncoded, sizeof(TCHAR)>(InHighChar, InLowChar);
-	}
-}
-
 namespace
 {
 
@@ -228,13 +208,15 @@ void FSlateTextShaper::PerformTextShaping(const TCHAR* InText, const int32 InTex
 						return true;
 					}
 
+#if !PLATFORM_TCHAR_IS_4_BYTES
 					{
 						const int32 NextTextIndex = RunningTextIndex + 1;
-						if (NextTextIndex < TextEndIndex && SurrogatePairUtil::IsSurrogatePair(Char, InText[NextTextIndex]))
+						if (NextTextIndex < TextEndIndex && StringConv::IsHighSurrogate(Char) && StringConv::IsLowSurrogate(InText[NextTextIndex]))
 						{
 							return true;
 						}
 					}
+#endif	// !PLATFORM_TCHAR_IS_4_BYTES
 				}
 			}
 
@@ -483,14 +465,16 @@ void FSlateTextShaper::PerformHarfBuzzTextShaping(const TCHAR* InText, const int
 			
 			// Combine any surrogate pairs into a complete codepoint
 			UTF32CHAR CurrentCodepoint = InText[RunningTextIndex];
+#if !PLATFORM_TCHAR_IS_4_BYTES
 			if (ClusterSize > 1)
 			{
 				const int32 NextTextIndex = RunningTextIndex + 1;
-				if (SurrogatePairUtil::IsSurrogatePair(InText[RunningTextIndex], InText[NextTextIndex]))
+				if (StringConv::IsHighSurrogate(InText[RunningTextIndex]) && StringConv::IsLowSurrogate(InText[NextTextIndex]))
 				{
-					CurrentCodepoint = ((InText[RunningTextIndex] - 0xD800) << 10) + (InText[NextTextIndex] - 0xDC00) + 0x10000;
+					CurrentCodepoint = StringConv::EncodeSurrogate(InText[RunningTextIndex], InText[NextTextIndex]);
 				}
 			}
+#endif	// !PLATFORM_TCHAR_IS_4_BYTES
 
 			// Substitute whitespace characters with spaces since not all fonts support all kinds of whitespace characters (but we don't care since we don't render them anyway)
 			if (FText::IsWhitespace(CurrentCodepoint))
diff --git a/Engine/Source/Runtime/WebBrowser/Private/MobileJS/MobileJSStructDeserializerBackend.cpp b/Engine/Source/Runtime/WebBrowser/Private/MobileJS/MobileJSStructDeserializerBackend.cpp
--- a/Engine/Source/Runtime/WebBrowser/Private/MobileJS/MobileJSStructDeserializerBackend.cpp
+++ b/Engine/Source/Runtime/WebBrowser/Private/MobileJS/MobileJSStructDeserializerBackend.cpp
@@ -99,8 +99,10 @@ FMobileJSStructDeserializerBackend::FMobileJSStructDeserializerBackend(FMobileJS
 	, JsonData()
 	, Reader(JsonData)
 {
-	auto Convert = StringCast<UCS2CHAR>(*JsonString);
-	JsonData.Append((uint8*)Convert.Get(), JsonString.Len() * sizeof(UCS2CHAR));
+	// Note: This is a no-op on platforms that are using a 16-bit TCHAR
+	FTCHARToUTF16 UTF16String(*JsonString, JsonString.Len());
+
+	JsonData.Append((uint8*)UTF16String.Get(), UTF16String.Length() * sizeof(UTF16CHAR));
 }
 
 #endif // PLATFORM_ANDROID || PLATFORM_IOS
\ No newline at end of file
diff --git a/Engine/Source/Runtime/WebBrowser/Private/MobileJS/MobileJSStructSerializerBackend.cpp b/Engine/Source/Runtime/WebBrowser/Private/MobileJS/MobileJSStructSerializerBackend.cpp
--- a/Engine/Source/Runtime/WebBrowser/Private/MobileJS/MobileJSStructSerializerBackend.cpp
+++ b/Engine/Source/Runtime/WebBrowser/Private/MobileJS/MobileJSStructSerializerBackend.cpp
@@ -46,9 +46,8 @@ void FMobileJSStructSerializerBackend::WriteUObject(const FStructSerializerState
 FString FMobileJSStructSerializerBackend::ToString()
 {
 	ReturnBuffer.Add(0);
-	ReturnBuffer.Add(0);
-	auto CastObject = StringCast<TCHAR>((UCS2CHAR*)ReturnBuffer.GetData());
-	return FString(CastObject.Get(), CastObject.Length());
+	ReturnBuffer.Add(0); // Add two as we're dealing with UTF-16, so 2 bytes
+	return UTF16_TO_TCHAR((UTF16CHAR*)ReturnBuffer.GetData());
 }
 
 FMobileJSStructSerializerBackend::FMobileJSStructSerializerBackend(TSharedRef<class FMobileJSScripting> InScripting)
diff --git a/Engine/Source/Runtime/WebBrowser/Private/Native/NativeJSScripting.cpp b/Engine/Source/Runtime/WebBrowser/Private/Native/NativeJSScripting.cpp
--- a/Engine/Source/Runtime/WebBrowser/Private/Native/NativeJSScripting.cpp
+++ b/Engine/Source/Runtime/WebBrowser/Private/Native/NativeJSScripting.cpp
@@ -243,10 +243,8 @@ FString FNativeJSScripting::ConvertStruct(UStruct* TypeInfo, const void* StructP
 
 	// Extract the result value from the serialized JSON object:
 	ReturnBuffer.Add(0);
-	ReturnBuffer.Add(0);
-	auto CastObject = StringCast<TCHAR>((UCS2CHAR*)ReturnBuffer.GetData());
-	FString ResultJS = FString(CastObject.Get(), CastObject.Length());
-	return ResultJS;
+	ReturnBuffer.Add(0); // Add two as we're dealing with UTF-16, so 2 bytes
+	return UTF16_TO_TCHAR((UTF16CHAR*)ReturnBuffer.GetData());
 }
 
 FString FNativeJSScripting::ConvertObject(UObject* Object)
@@ -432,11 +430,13 @@ bool FNativeJSScripting::HandleExecuteUObjectMethodMessage(const TArray<FString>
 		Params.AddUninitialized(ParamsSize);
 		Function->InitializeStruct(Params.GetData());
 
+		// Note: This is a no-op on platforms that are using a 16-bit TCHAR
+		FTCHARToUTF16 UTF16String(*MessageArgs[3], MessageArgs[3].Len());
+
 		TArray<uint8> JsonData;
-		auto Convert = StringCast<UCS2CHAR>(*MessageArgs[3]);
-		JsonData.Append((uint8*)Convert.Get(), MessageArgs[3].Len() * sizeof(UCS2CHAR));
-		FMemoryReader Reader(JsonData);
+		JsonData.Append((uint8*)UTF16String.Get(), UTF16String.Length() * sizeof(UTF16CHAR));
 
+		FMemoryReader Reader(JsonData);
 		FNativeJSStructDeserializerBackend Backend = FNativeJSStructDeserializerBackend(SharedThis(this), Reader);
 		FStructDeserializer::Deserialize(Params.GetData(), *Function, Backend);
 	}
@@ -468,9 +468,8 @@ bool FNativeJSScripting::HandleExecuteUObjectMethodMessage(const TArray<FString>
 
 			// Extract the result value from the serialized JSON object:
 			ReturnBuffer.Add(0);
-			ReturnBuffer.Add(0);
-			auto CastObject = StringCast<TCHAR>((UCS2CHAR*)ReturnBuffer.GetData());
-			FString ResultJS = FString(CastObject.Get(), CastObject.Length());
+			ReturnBuffer.Add(0); // Add two as we're dealing with UTF-16, so 2 bytes
+			const FString ResultJS = UTF16_TO_TCHAR((UTF16CHAR*)ReturnBuffer.GetData());
 
 			InvokeJSFunctionRaw(ResultCallbackId, ResultJS, false);
 		}
