Back out change to ComputeSquaredDistanceToPoint.  These are equivalent.

#rnx

diff --git a/Engine/Source/Runtime/Core/Public/Math/Vector.h b/Engine/Source/Runtime/Core/Public/Math/Vector.h
--- a/Engine/Source/Runtime/Core/Public/Math/Vector.h
+++ b/Engine/Source/Runtime/Core/Public/Math/Vector.h
@@ -1072,64 +1072,52 @@ FORCEINLINE uint32 GetTypeHash(const FVector& Vector)
 	}
 #endif
 
-/**
- * Util to calculate closest point from a point to a bounding box
+
+/** 
+ * Util to calculate distance from a point to a bounding box 
  *
- * @param Min 3D Point defining the lower values of the axis of the bound box
+ * @param Mins 3D Point defining the lower values of the axis of the bound box
  * @param Max 3D Point defining the lower values of the axis of the bound box
  * @param Point 3D position of interest
- * @return If point is outside box, returns closest point on box.  If point inside, returns point
+ * @return the distance from the Point to the bounding box.
  */
-FORCEINLINE FVector GetClosestPointFromBoxToPoint(const FVector& Min, const FVector& Max, const FVector& Point)
+FORCEINLINE float ComputeSquaredDistanceFromBoxToPoint(const FVector& Mins, const FVector& Maxs, const FVector& Point)
 {
-	// start by considering the point inside the box
-	FVector ClosestPoint = Point;
+	// Accumulates the distance as we iterate axis
+	float DistSquared = 0.f;
 
-	// now clamp to inside box if it's outside
-	if (Point.X < Min.X)
+	// Check each axis for min/max and add the distance accordingly
+	// NOTE: Loop manually unrolled for > 2x speed up
+	if (Point.X < Mins.X)
 	{
-		ClosestPoint.X = Min.X;
+		DistSquared += FMath::Square(Point.X - Mins.X);
 	}
-	else if (Point.X > Max.X)
+	else if (Point.X > Maxs.X)
 	{
-		ClosestPoint.X = Max.X;
+		DistSquared += FMath::Square(Point.X - Maxs.X);
 	}
-
-	// now clamp to inside box if it's outside
-	if (Point.Y < Min.Y)
+	
+	if (Point.Y < Mins.Y)
 	{
-		ClosestPoint.Y = Min.Y;
+		DistSquared += FMath::Square(Point.Y - Mins.Y);
 	}
-	else if (Point.Y > Max.Y)
+	else if (Point.Y > Maxs.Y)
 	{
-		ClosestPoint.Y = Max.Y;
+		DistSquared += FMath::Square(Point.Y - Maxs.Y);
 	}
-
-	// Now clamp to inside box if it's outside.
-	if (Point.Z < Min.Z)
+	
+	if (Point.Z < Mins.Z)
 	{
-		ClosestPoint.Z = Min.Z;
+		DistSquared += FMath::Square(Point.Z - Mins.Z);
 	}
-	else if (Point.Z > Max.Z)
+	else if (Point.Z > Maxs.Z)
 	{
-		ClosestPoint.Z = Max.Z;
+		DistSquared += FMath::Square(Point.Z - Maxs.Z);
 	}
-
-	return ClosestPoint;
+	
+	return DistSquared;
 }
 
-/** 
- * Util to calculate distance from a point to a bounding box 
- *
- * @param Mins 3D Point defining the lower values of the axis of the bound box
- * @param Max 3D Point defining the lower values of the axis of the bound box
- * @param Point 3D position of interest
- * @return the distance from the Point to the bounding box.
- */
-FORCEINLINE float ComputeSquaredDistanceFromBoxToPoint(const FVector& Mins, const FVector& Maxs, const FVector& Point)
-{
-	return (GetClosestPointFromBoxToPoint(Mins, Maxs, Point) - Point).SizeSquared();
-}
 
 FORCEINLINE FVector::FVector(const FVector2D V, float InZ)
 	: X(V.X), Y(V.Y), Z(InZ)
