Build fix.  Part of CL 8677313

#rnx

diff --git a/Engine/Source/Runtime/Core/Public/Math/Box.h b/Engine/Source/Runtime/Core/Public/Math/Box.h
--- a/Engine/Source/Runtime/Core/Public/Math/Box.h
+++ b/Engine/Source/Runtime/Core/Public/Math/Box.h
@@ -536,7 +536,40 @@ FORCEINLINE FBox& FBox::operator+=( const FBox& Other )
 
 FORCEINLINE FVector FBox::GetClosestPointTo( const FVector& Point ) const
 {
-	return ::GetClosestPointFromBoxToPoint(Min, Max, Point);
+	// start by considering the point inside the box
+	FVector ClosestPoint = Point;
+
+	// now clamp to inside box if it's outside
+	if (Point.X < Min.X)
+	{
+		ClosestPoint.X = Min.X;
+	}
+	else if (Point.X > Max.X)
+	{
+		ClosestPoint.X = Max.X;
+	}
+
+	// now clamp to inside box if it's outside
+	if (Point.Y < Min.Y)
+	{
+		ClosestPoint.Y = Min.Y;
+	}
+	else if (Point.Y > Max.Y)
+	{
+		ClosestPoint.Y = Max.Y;
+	}
+
+	// Now clamp to inside box if it's outside.
+	if (Point.Z < Min.Z)
+	{
+		ClosestPoint.Z = Min.Z;
+	}
+	else if (Point.Z > Max.Z)
+	{
+		ClosestPoint.Z = Max.Z;
+	}
+
+	return ClosestPoint;
 }
 
 
