Fix non-unity build

#rnx

diff --git a/Engine/Source/Runtime/Core/Private/Containers/String.cpp b/Engine/Source/Runtime/Core/Private/Containers/String.cpp
--- a/Engine/Source/Runtime/Core/Private/Containers/String.cpp
+++ b/Engine/Source/Runtime/Core/Private/Containers/String.cpp
@@ -1614,3 +1614,8 @@ void FTextRange::CalculateLineRangesFromString(const FString& Input, TArray<FTex
 		LineRanges.Emplace(FTextRange(LineBeginIndex, Input.Len()));
 	}
 }
+
+void StringConv::InlineCombineSurrogates(FString& Str)
+{
+	InlineCombineSurrogates_Array(Str.GetCharArray());
+}
diff --git a/Engine/Source/Runtime/Core/Public/Containers/UnrealString.h b/Engine/Source/Runtime/Core/Public/Containers/UnrealString.h
--- a/Engine/Source/Runtime/Core/Public/Containers/UnrealString.h
+++ b/Engine/Source/Runtime/Core/Public/Containers/UnrealString.h
@@ -2240,14 +2240,8 @@ FORCEINLINE const TCHAR* ToCStr(const FString& Str)
 
 namespace StringConv
 {
-	template<typename AllocatorType>
-	static void InlineCombineSurrogates_Array(TArray<TCHAR, AllocatorType>& StrBuffer);
-
 	/** Inline combine any UTF-16 surrogate pairs in the given string */
-	static FORCEINLINE void InlineCombineSurrogates(FString& Str)
-	{
-		InlineCombineSurrogates_Array(Str.GetCharArray());
-	}
+	CORE_API void InlineCombineSurrogates(FString& Str);
 }
 
 /*----------------------------------------------------------------------------
