Ensure we always load ICU data from loose files, regardless of the compiled-in default

diff --git a/Engine/Source/Runtime/Core/Private/Internationalization/ICUInternationalization.cpp b/Engine/Source/Runtime/Core/Private/Internationalization/ICUInternationalization.cpp
--- a/Engine/Source/Runtime/Core/Private/Internationalization/ICUInternationalization.cpp
+++ b/Engine/Source/Runtime/Core/Private/Internationalization/ICUInternationalization.cpp
@@ -143,6 +143,7 @@ bool FICUInternationalization::Initialize()
 		UE_LOG(LogICUInternationalization, Fatal, TEXT("ICU data directory was not discovered:\n%s"), *GetPrioritizedDataDirectoriesString());
 	}
 
+	udata_setFileAccess(UDATA_FILES_FIRST, &ICUStatus); // We always need to load loose ICU data files
 	u_setDataFileFunctions(this, &FICUInternationalization::OpenDataFile, &FICUInternationalization::CloseDataFile, &ICUStatus);
 	u_init(&ICUStatus);
 	checkf(U_SUCCESS(ICUStatus), TEXT("Failed to open ICUInternationalization data file, missing or corrupt?"));
