Add delegate in the csv for start / End capture.

#rnx

diff --git a/Engine/Source/Runtime/Core/Private/ProfilingDebugging/CsvProfiler.cpp b/Engine/Source/Runtime/Core/Private/ProfilingDebugging/CsvProfiler.cpp
--- a/Engine/Source/Runtime/Core/Private/ProfilingDebugging/CsvProfiler.cpp
+++ b/Engine/Source/Runtime/Core/Private/ProfilingDebugging/CsvProfiler.cpp
@@ -2471,6 +2471,9 @@ void FCsvProfiler::BeginFrame()
 			else
 			{
 				UE_LOG(LogCsvProfiler, Display, TEXT("Capture Starting"));
+				
+				// signal external profiler that we are capturing
+				OnCSVProfileStartDelegate.Broadcast();
 
 				// Latch the cvars when we start a capture
 				int32 BufferSize = FMath::Max(CVarCsvWriteBufferSize.GetValueOnAnyThread(), 0);
@@ -2623,6 +2626,9 @@ void FCsvProfiler::EndFrame()
 			}
 			else
 			{
+				// signal external profiler that we are done
+				OnCSVProfileEndDelegate.Broadcast();
+
 				// Signal to the processing thread to write the file out (if we have one).
 				GCsvProfilerIsWritingFile = true;
 				GCsvProfilerIsCapturing = false;
diff --git a/Engine/Source/Runtime/Core/Public/ProfilingDebugging/CsvProfiler.h b/Engine/Source/Runtime/Core/Public/ProfilingDebugging/CsvProfiler.h
--- a/Engine/Source/Runtime/Core/Public/ProfilingDebugging/CsvProfiler.h
+++ b/Engine/Source/Runtime/Core/Public/ProfilingDebugging/CsvProfiler.h
@@ -272,6 +272,12 @@ public:
 
 	CORE_API FString GetOutputFilename() const { return OutputFilename; }
 
+	DECLARE_MULTICAST_DELEGATE(FOnCSVProfileStart);
+	FOnCSVProfileStart& OnCSVProfileStart() { return OnCSVProfileStartDelegate; }
+
+	DECLARE_MULTICAST_DELEGATE(FOnCSVProfileEnd);
+	FOnCSVProfileEnd& OnCSVProfileEnd() { return OnCSVProfileEndDelegate; }
+	
 	DECLARE_MULTICAST_DELEGATE_OneParam(FOnCSVProfileFinished, const FString& /*Filename */);
 	FOnCSVProfileFinished& OnCSVProfileFinished() { return OnCSVProfileFinishedDelegate; }
 
@@ -310,6 +316,9 @@ private:
 
 	ECsvProfilerFlags CurrentFlags;
 
+	FOnCSVProfileStart OnCSVProfileStartDelegate;
+	FOnCSVProfileEnd OnCSVProfileEndDelegate;
+	
 	FOnCSVProfileFinished OnCSVProfileFinishedDelegate;
 };
 
diff --git a/Engine/Source/Runtime/Core/Public/ProfilingDebugging/ExternalProfiler.h b/Engine/Source/Runtime/Core/Public/ProfilingDebugging/ExternalProfiler.h
--- a/Engine/Source/Runtime/Core/Public/ProfilingDebugging/ExternalProfiler.h
+++ b/Engine/Source/Runtime/Core/Public/ProfilingDebugging/ExternalProfiler.h
@@ -44,6 +44,9 @@ public:
 	/** Mark where the profiler should consider the frame boundary to be. */
 	virtual void FrameSync() = 0;
 
+	/** Initialize  profiler, register some delegates..*/
+	virtual void Register() {}
+
 	/** Pauses profiling. */
 	virtual void ProfilerPauseFunction() = 0;
 
diff --git a/Engine/Source/Runtime/Launch/Private/LaunchEngineLoop.cpp b/Engine/Source/Runtime/Launch/Private/LaunchEngineLoop.cpp
--- a/Engine/Source/Runtime/Launch/Private/LaunchEngineLoop.cpp
+++ b/Engine/Source/Runtime/Launch/Private/LaunchEngineLoop.cpp
@@ -3657,6 +3657,14 @@ int32 FEngineLoop::Init()
 	FEmbeddedCommunication::AllowSleep(TEXT("Startup"));
 	FEmbeddedCommunication::KeepAwake(TEXT("FirstTicks"), true);
 #endif
+	
+#if UE_EXTERNAL_PROFILING_ENABLED
+	FExternalProfiler* ActiveProfiler = FActiveExternalProfilerBase::GetActiveProfiler();
+	if (ActiveProfiler)
+	{
+		ActiveProfiler->Register();
+	}
+#endif		// UE_EXTERNAL_PROFILING_ENABLED
 
 	return 0;
 }
