Add missing brackets
#rnx

diff --git a/Engine/Source/Runtime/Core/Private/Misc/ConfigCacheIni.cpp b/Engine/Source/Runtime/Core/Private/Misc/ConfigCacheIni.cpp
--- a/Engine/Source/Runtime/Core/Private/Misc/ConfigCacheIni.cpp
+++ b/Engine/Source/Runtime/Core/Private/Misc/ConfigCacheIni.cpp
@@ -941,7 +941,7 @@ FString FConfigFile::GenerateExportedPropertyLine(const FString& PropertyName, c
 namespace CommandlineOverrideSpecifiers
 {
 	// -ini:IniName:[Section1]:Key1=Value1,[Section2]:Key2=Value2
-	const TCHAR	IniFileOverrideIdentifier = TEXT("-iniFile=");
+	const TCHAR	IniFileOverrideIdentifier[] = TEXT("-iniFile=");
 	const TCHAR IniSwitchIdentifier[]     = TEXT("-ini:");
 	const TCHAR IniNameEndIdentifier[]    = TEXT(":[");
 	const TCHAR SectionStartIdentifier[]  = TEXT("[");
