Thread detaching implementation is incorrect. The problem is that detached thread needs to destroy itself when completed. The latest point it can do this is in overriden `FRunnable::Exit()`, but it's not late enough. It's called by `FRunnableThread::Run()` and after the call, `FRunnableThread` still expects to be alive, e.g. to free TLS (`FRunnableThread::FreeTls`). Right now I don't see any way to implement `Detach` correctly w/o modifying underlying FRunnable family. As this is marked as a critical bug, for now I disabled `Detach` functionality and corresponding test (detaching is not used and wasn't released yet). Any FRunnable modifications would need to be discussed with Core team.

Other changes:

Modified the related test (that is now disabled) to always crash on PS4 w/o the fix, as before it was crashing only sometimes. On Windows it still doesn't crash, presumably due to memory deallocation specifics.

Made test thread names to be more descriptive to help debugging.

diff --git a/Engine/Source/Runtime/Core/Private/HAL/Thread.cpp b/Engine/Source/Runtime/Core/Private/HAL/Thread.cpp
--- a/Engine/Source/Runtime/Core/Private/HAL/Thread.cpp
+++ b/Engine/Source/Runtime/Core/Private/HAL/Thread.cpp
@@ -58,6 +58,11 @@ private:
 
 		ThreadFunction();
 
+		return 0;
+	}
+
+	void Exit() override
+	{
 		// busy-wait till `Self` is initialized before releasing it
 		while (!bIsInitialized) {}
 
@@ -66,8 +71,6 @@ private:
 		// we're about to exit the thread. Release the reference to self. If the thread is detached, it's the only reference
 		// and so this instance will be deleted. No member access should be performed after this point.
 		Self.Reset();
-
-		return 0;
 	}
 
 private:
@@ -126,8 +129,12 @@ uint32 FThread::GetThreadId() const
 	return Impl.IsValid() ? Impl->GetThreadId() : FThread::InvalidThreadId;
 }
 
+#if 0 // disabled as it doesn't work as intended
+
 void FThread::Detach()
 {
 	check(Impl.IsValid());
 	Impl.Reset(); // releases this reference, `FThreadImpl` instance can still be alive until the thread finishes its execution
 }
+
+#endif
\ No newline at end of file
diff --git a/Engine/Source/Runtime/Core/Private/Tests/HAL/ThreadTest.cpp b/Engine/Source/Runtime/Core/Private/Tests/HAL/ThreadTest.cpp
--- a/Engine/Source/Runtime/Core/Private/Tests/HAL/ThreadTest.cpp
+++ b/Engine/Source/Runtime/Core/Private/Tests/HAL/ThreadTest.cpp
@@ -19,7 +19,7 @@ namespace
 {
 	void TestIsJoinableAfterCreation(FThreadTest& This)
 	{
-		FThread Thread(TEXT("Test.Thread"), []() { /*NOOP*/ });
+		FThread Thread(TEXT("Test.Thread.TestIsJoinableAfterCreation"), []() { /*NOOP*/ });
 		This.TestTrue(TEXT("FThread must be joinable after construction"), Thread.IsJoinable());
 		Thread.Join();
 		UE_LOG(LogTemp, Log, TEXT("%s completed"), StringCast<TCHAR>(__FUNCTION__).Get());
@@ -28,7 +28,7 @@ namespace
 	void TestIsJoinableAfterCompletion(FThreadTest& This)
 	{
 		FThreadSafeBool bDone = false;
-		FThread Thread(TEXT("Test.Thread"), [&bDone]() { bDone = true; });
+		FThread Thread(TEXT("Test.Thread.TestIsJoinableAfterCompletion"), [&bDone]() { bDone = true; });
 		while (!bDone); // wait for completion //-V529
 		This.TestTrue(TEXT("FThread must still be joinable after completion"), Thread.IsJoinable());
 		Thread.Join();
@@ -37,24 +37,44 @@ namespace
 
 	void TestIsNotJoinableAfterJoining(FThreadTest& This)
 	{
-		FThread Thread(TEXT("Test.Thread"), []() { /*NOOP*/ });
+		FThread Thread(TEXT("Test.Thread.TestIsNotJoinableAfterJoining"), []() { /*NOOP*/ });
 		Thread.Join();
 		This.TestFalse(TEXT("FThread must not be joinable after joining"), Thread.IsJoinable());
 		UE_LOG(LogTemp, Log, TEXT("%s completed"), StringCast<TCHAR>(__FUNCTION__).Get());
 	}
 
+#if 0 // detaching is not implemented
 	void TestIsNotJoinableAfterDetaching(FThreadTest& This)
 	{
-		FThread Thread(TEXT("Test.Thread"), []() { /*NOOP*/ });
-		Thread.Detach();
-		This.TestFalse(TEXT("FThread must not be joinable after detaching"), Thread.IsJoinable());
+		// two cases: it's either the calling thread detaches from the thread before the thread is completed
+		{
+			TAtomic<bool> bReady{ false };
+			FThread Thread(TEXT("Test.Thread"), [&bReady]()
+				{
+					while (!bReady) {}
+				});
+			Thread.Detach();
+			bReady = true; // make sure `Detach` is called before thread function exit
+			This.TestFalse(TEXT("FThread must not be joinable after detaching"), Thread.IsJoinable());
+		}
+		// or thread function is completed fast and `FThreadImpl` releases the reference to itself before
+		// `Detach` call
+		{
+			TAtomic<bool> bReady{ false };
+			FThread Thread(TEXT("Test.Thread"), [&bReady]() { /*NOOP*/});
+			FPlatformProcess::Sleep(0.1); // let the thread exit before detaching
+			Thread.Detach();
+			bReady = true; // make sure `Detach` is called before thread function exit
+			This.TestFalse(TEXT("FThread must not be joinable after detaching"), Thread.IsJoinable());
+		}
 		UE_LOG(LogTemp, Log, TEXT("%s completed"), StringCast<TCHAR>(__FUNCTION__).Get());
 	}
+#endif
 
 	void TestAssertIfNotJoinedOrDetached(FThreadTest& This)
 	{
 		// this does fails the `check`, but it seems there's no way to test this by UE4 Automation Testing, so commented out
-		//FThread Thread(TEXT("Test.Thread"), []() { /*NOOP*/ });
+		FThread Thread(TEXT("Test.Thread.TestAssertIfNotJoinedOrDetached"), []() { /*NOOP*/ });
 		// should assert in the destructor
 	}
 
@@ -66,7 +86,7 @@ namespace
 		}
 		{	// check that default constructed thread can be "upgraded" to joinable thread
 			FThread Thread;
-			Thread = FThread(TEXT("Test.Thread"), []() { /* NOOP */ });
+			Thread = FThread(TEXT("Test.Thread.TestDefaultConstruction"), []() { /* NOOP */ });
 			This.TestTrue(TEXT("Move-constructed FThread from joinable thread must be joinable"), Thread.IsJoinable());
 			Thread.Join();
 		}
@@ -82,14 +102,14 @@ namespace
 			This.TestFalse(TEXT("Move-constructed thread from not joinable thread must be not joinable"), Dst.IsJoinable());
 		}
 		{	// move constructor with joinable thread
-			FThread Src(TEXT("Test.Thread"), []() { /* NOOP */ });
+			FThread Src(TEXT("Test.Thread.TestMovability.1"), []() { /* NOOP */ });
 			FThread Dst(MoveTemp(Src));
 			This.TestFalse(TEXT("Moved out thread must be not joinable"), Src.IsJoinable());
 			This.TestTrue(TEXT("Move-constructed thread from joinable thread must be joinable"), Dst.IsJoinable());
 			Dst.Join();
 		}
 		{	// move assignment operator
-			FThread Src(TEXT("Test.Thread"), []() { /* NOOP */ });
+			FThread Src(TEXT("Test.Thread.TestMovability.2"), []() { /* NOOP */ });
 			FThread Dst;
 			Dst = MoveTemp(Src);
 			This.TestFalse(TEXT("Moved out thread must be not joinable"), Src.IsJoinable());
@@ -103,8 +123,8 @@ namespace
 			//Dst.Join();
 		}
 		{	// Move assignment operator of thread that has been joined
-			FThread Src(TEXT("Test.Thread"), []() { /* NOOP */ });
-			FThread Dst(TEXT("Test.Thread"), []() { /* NOOP */ });
+			FThread Src(TEXT("Test.Thread.TestMovability.3"), []() { /* NOOP */ });
+			FThread Dst(TEXT("Test.Thread.TestMovability.4"), []() { /* NOOP */ });
 			Dst.Join();
 			Dst = MoveTemp(Src);
 			Dst.Join();
@@ -120,7 +140,7 @@ namespace
 		TQueue<FWork> WorkQueue;
 		FEvent* WorkQueuedEvent = FPlatformProcess::GetSynchEventFromPool();
 
-		FThread WorkerThread(TEXT("Test.Thread"), [&bQuitRequested, &WorkQueue, WorkQueuedEvent]() 
+		FThread WorkerThread(TEXT("Test.Thread.TestTypicalUseCase"), [&bQuitRequested, &WorkQueue, WorkQueuedEvent]() 
 		{
 			while (!bQuitRequested)
 			{
@@ -179,9 +199,12 @@ bool FThreadTest::RunTest(const FString& Parameters)
 	TestIsJoinableAfterCreation(*this);
 	TestIsJoinableAfterCompletion(*this);
 	TestIsNotJoinableAfterJoining(*this);
+	
+#if 0 // detaching is not implemented
 	TestIsNotJoinableAfterDetaching(*this);
+#endif
 
-	TestAssertIfNotJoinedOrDetached(*this);
+	//TestAssertIfNotJoinedOrDetached(*this);
 
 	TestDefaultConstruction(*this);
 	TestMovability(*this);
diff --git a/Engine/Source/Runtime/Core/Public/HAL/Thread.h b/Engine/Source/Runtime/Core/Public/HAL/Thread.h
--- a/Engine/Source/Runtime/Core/Public/HAL/Thread.h
+++ b/Engine/Source/Runtime/Core/Public/HAL/Thread.h
@@ -85,6 +85,8 @@ public:
 	 */
 	uint32 GetThreadId() const;
 
+#if 0 // disabled as it doesn't work as intended
+
 	/**
 	 * Separates the thread of execution from the thread object, allowing execution to continue independently. 
 	 * Any allocated resources will be freed once the thread exits.
@@ -92,6 +94,8 @@ public:
 	 */
 	void Detach();
 
+#endif
+
 private:
 	TSharedPtr<class FThreadImpl, ESPMode::ThreadSafe> Impl; // "shared" with `FThreadImpl::Self`
 };
