Change LLMcsv tracker files to allow reading while being written to.

diff --git a/Engine/Source/Runtime/Core/Private/HAL/LowLevelMemTracker.cpp b/Engine/Source/Runtime/Core/Private/HAL/LowLevelMemTracker.cpp
--- a/Engine/Source/Runtime/Core/Private/HAL/LowLevelMemTracker.cpp
+++ b/Engine/Source/Runtime/Core/Private/HAL/LowLevelMemTracker.cpp
@@ -1999,7 +1999,7 @@ void FLLMCsvWriter::WriteGraph(FLLMCustomTag* CustomTags, const int32* ParentTag
 		const TCHAR* TrackerName = GetTrackerCsvName(Tracker);
 		const FDateTime FileDate = FDateTime::Now();
 		FString Filename = FString::Printf(TEXT("%s/%s_%s.csv"), *Directory, TrackerName, *FileDate.ToString());
-		Archive = IFileManager::Get().CreateFileWriter(*Filename);
+		Archive = IFileManager::Get().CreateFileWriter(*Filename, FILEWRITE_AllowRead);
 		LLMCheck(Archive);
 
 		// create space for column titles that are filled in as we get them
