Change EmbeddedCommunication to use regular logging mechanism (so we can change verbosity)

diff --git a/Engine/Source/Runtime/Core/Private/Misc/EmbeddedCommunication.cpp b/Engine/Source/Runtime/Core/Private/Misc/EmbeddedCommunication.cpp
--- a/Engine/Source/Runtime/Core/Private/Misc/EmbeddedCommunication.cpp
+++ b/Engine/Source/Runtime/Core/Private/Misc/EmbeddedCommunication.cpp
@@ -204,46 +204,39 @@ void FEmbeddedCommunication::AllowSleep(FName Requester)
 }
 
 #if BUILD_EMBEDDED_APP
-static FName NativeLogCategory("Native");
+DEFINE_LOG_CATEGORY_STATIC(LogBridge, Log, All);
+
+#define UE_LOG_BRIDGE(Verbosity, String) \
+{ \
+	UE_LOG(LogBridge, Verbosity, TEXT("%s"), String); \
+}
 #endif
 
 void FEmbeddedCommunication::UELogError(const TCHAR* String)
 {
 #if BUILD_EMBEDDED_APP
-	if (GWarn)
-	{
-		GWarn->Log(NativeLogCategory, ELogVerbosity::Error, String);
-	}
+	UE_LOG_BRIDGE(Error, String);
 #endif
 }
 
 void FEmbeddedCommunication::UELogWarning(const TCHAR* String)
 {
 #if BUILD_EMBEDDED_APP
-	if (GWarn)
-	{
-		GWarn->Log(NativeLogCategory, ELogVerbosity::Warning, String);
-	}
+	UE_LOG_BRIDGE(Warning, String);
 #endif
 }
 
 void FEmbeddedCommunication::UELogLog(const TCHAR* String)
 {
 #if BUILD_EMBEDDED_APP
-	if (GLog)
-	{
-		GLog->Log(NativeLogCategory, ELogVerbosity::Log, String);
-	}
+	UE_LOG_BRIDGE(Log, String);
 #endif
 }
 
 void FEmbeddedCommunication::UELogVerbose(const TCHAR* String)
 {
 #if BUILD_EMBEDDED_APP
-	if (GLog)
-	{
-		GLog->Log(NativeLogCategory, ELogVerbosity::Verbose, String);
-	}
+	UE_LOG_BRIDGE(Verbose, String);
 #endif
 }
 
