Make the logging interval configurable, with a separate override value for shipping if necessary. Reducing the flush interval helps I/O performance on some platforms

diff --git a/Engine/Source/Runtime/Core/Private/Misc/OutputDeviceFile.cpp b/Engine/Source/Runtime/Core/Private/Misc/OutputDeviceFile.cpp
--- a/Engine/Source/Runtime/Core/Private/Misc/OutputDeviceFile.cpp
+++ b/Engine/Source/Runtime/Core/Private/Misc/OutputDeviceFile.cpp
@@ -34,218 +34,241 @@
 typedef uint8 UTF8BOMType[3];
 static UTF8BOMType UTF8BOM = { 0xEF, 0xBB, 0xBF };
 
+static float GLogFlushIntervalSec = 0.2f;
+static FAutoConsoleVariableRef CVarLogFlushInterval(
+	TEXT("log.flushInterval"),
+	GLogFlushIntervalSec,
+	TEXT("Logging interval in seconds"),
+	ECVF_Default );
+
+#if UE_BUILD_SHIPPING
+static float GLogFlushIntervalSec_Shipping = 0.0f;
+static FAutoConsoleVariableRef CVarLogFlushIntervalShipping(
+	TEXT("log.flushInterval.Shipping"),
+	GLogFlushIntervalSec_Shipping,
+	TEXT("Logging interval in shipping. If set, this overrides archive.FlushInterval"),
+	ECVF_Default);
+#endif
 
+inline double GetLogFlushIntervalSec()
+{
+#if UE_BUILD_SHIPPING
+	return double((GLogFlushIntervalSec_Shipping>0.0f) ? GLogFlushIntervalSec_Shipping : GLogFlushIntervalSec);
+#else
+	return double(GLogFlushIntervalSec);
+#endif
+}
 
-	/** [WRITER THREAD] Flushes the archive and reset the flush timer. */
-	void FAsyncWriter::FlushArchiveAndResetTimer()
-	{
-		// This should be the one and only place where we flush because we want the flush to happen only on the 
-		// async writer thread (if threading is enabled)
-		Ar.Flush();
-		LastArchiveFlushTime = FPlatformTime::Seconds();
-	}
 
-	/** [WRITER THREAD] Serialize the contents of the ring buffer to disk */
-	void FAsyncWriter::SerializeBufferToArchive()
-	{
-		// Unix and PS4 use FPlatformMallocCrash during a crash, which means this function is not allowed to perform any allocations
-		// or else it will deadlock when flushing logs during crash handling. Ideally scoped named events would be disabled while crashing.
-		// GIsCriticalError is not always true when crashing (i.e. the case of a GPF) so there is no way to know to skip this behavior only when crashing
+/** [WRITER THREAD] Flushes the archive and reset the flush timer. */
+void FAsyncWriter::FlushArchiveAndResetTimer()
+{
+	// This should be the one and only place where we flush because we want the flush to happen only on the 
+	// async writer thread (if threading is enabled)
+	Ar.Flush();
+	LastArchiveFlushTime = FPlatformTime::Seconds();
+}
+
+/** [WRITER THREAD] Serialize the contents of the ring buffer to disk */
+void FAsyncWriter::SerializeBufferToArchive()
+{
+	// Unix and PS4 use FPlatformMallocCrash during a crash, which means this function is not allowed to perform any allocations
+	// or else it will deadlock when flushing logs during crash handling. Ideally scoped named events would be disabled while crashing.
+	// GIsCriticalError is not always true when crashing (i.e. the case of a GPF) so there is no way to know to skip this behavior only when crashing
 #if PLATFORM_ALLOW_ALLOCATIONS_IN_FASYNCWRITER_SERIALIZEBUFFERTOARCHIVE
-		SCOPED_NAMED_EVENT(FAsyncWriter_SerializeBufferToArchive, FColor::Cyan);
+	SCOPED_NAMED_EVENT(FAsyncWriter_SerializeBufferToArchive, FColor::Cyan);
 #endif
-		while (SerializeRequestCounter.GetValue() > 0)
-		{
-			// Grab a local copy of the end pos. It's ok if it changes on the client thread later on.
-			// We won't be modifying it anyway and will later serialize new data in the next iteration.
-			// Here we only serialize what we know exists at the beginning of this function.
-			int32 ThisThreadStartPos = BufferStartPos.Load(EMemoryOrder::Relaxed);
-			int32 ThisThreadEndPos   = BufferEndPos  .Load(EMemoryOrder::Relaxed);
+	while (SerializeRequestCounter.GetValue() > 0)
+	{
+		// Grab a local copy of the end pos. It's ok if it changes on the client thread later on.
+		// We won't be modifying it anyway and will later serialize new data in the next iteration.
+		// Here we only serialize what we know exists at the beginning of this function.
+		int32 ThisThreadStartPos = BufferStartPos.Load(EMemoryOrder::Relaxed);
+		int32 ThisThreadEndPos   = BufferEndPos  .Load(EMemoryOrder::Relaxed);
 
-			if (ThisThreadEndPos >= ThisThreadStartPos)
-			{
-				Ar.Serialize(Buffer.GetData() + ThisThreadStartPos, ThisThreadEndPos - ThisThreadStartPos);
-			}
-			else
-			{
-				// Data is wrapped around the ring buffer
-				Ar.Serialize(Buffer.GetData() + ThisThreadStartPos, Buffer.Num() - ThisThreadStartPos);
-				Ar.Serialize(Buffer.GetData(), ThisThreadEndPos);
-			}
+		if (ThisThreadEndPos >= ThisThreadStartPos)
+		{
+			Ar.Serialize(Buffer.GetData() + ThisThreadStartPos, ThisThreadEndPos - ThisThreadStartPos);
+		}
+		else
+		{
+			// Data is wrapped around the ring buffer
+			Ar.Serialize(Buffer.GetData() + ThisThreadStartPos, Buffer.Num() - ThisThreadStartPos);
+			Ar.Serialize(Buffer.GetData(), ThisThreadEndPos);
+		}
 
-			// Modify the start pos. Only the worker thread modifies this value so it's ok to not guard it with a critical section.
-			BufferStartPos = ThisThreadEndPos;
+		// Modify the start pos. Only the worker thread modifies this value so it's ok to not guard it with a critical section.
+		BufferStartPos = ThisThreadEndPos;
 
-			// Decrement the request counter, we now know we serialized at least one request.
-			// We might have serialized more requests but it's irrelevant, the counter will go down to 0 eventually
-			SerializeRequestCounter.Decrement();
+		// Decrement the request counter, we now know we serialized at least one request.
+		// We might have serialized more requests but it's irrelevant, the counter will go down to 0 eventually
+		SerializeRequestCounter.Decrement();
 
-			// Flush the archive periodically if running on a separate thread
-			if (Thread)
-			{
-				if ((FPlatformTime::Seconds() - LastArchiveFlushTime) > ArchiveFlushIntervalSec)
-				{
-					FlushArchiveAndResetTimer();
-				}
-			}
-			// If no threading is available or when we explicitly requested flush (see FlushBuffer), flush immediately after writing.
-			// In some rare cases we may flush twice (see above) but that's ok. We need a clear division between flushing because of the timer
-			// and force flush on demand.
-			if (WantsArchiveFlush.GetValue() > 0)
+		// Flush the archive periodically if running on a separate thread
+		if (Thread)
+		{
+			if ((FPlatformTime::Seconds() - LastArchiveFlushTime) > GetLogFlushIntervalSec() )
 			{
 				FlushArchiveAndResetTimer();
-				int32 FlushCount = WantsArchiveFlush.Decrement();
-				check(FlushCount >= 0);
 			}
 		}
-	}
-
-	/** [CLIENT THREAD] Flush the memory buffer (doesn't force the archive to flush). Can only be used from inside of BufferPosCritical lock. */
-	void FAsyncWriter::FlushBuffer()
-	{
-		SerializeRequestCounter.Increment();
-		if (!Thread)
+		// If no threading is available or when we explicitly requested flush (see FlushBuffer), flush immediately after writing.
+		// In some rare cases we may flush twice (see above) but that's ok. We need a clear division between flushing because of the timer
+		// and force flush on demand.
+		if (WantsArchiveFlush.GetValue() > 0)
 		{
-			SerializeBufferToArchive();
+			FlushArchiveAndResetTimer();
+			int32 FlushCount = WantsArchiveFlush.Decrement();
+			check(FlushCount >= 0);
 		}
-		while (SerializeRequestCounter.GetValue() != 0)
-		{
-			FPlatformProcess::SleepNoStats(0);
-		}
-		// Make sure there's been no unexpected concurrency
-		check(SerializeRequestCounter.GetValue() == 0);
 	}
+}
 
-	FAsyncWriter::FAsyncWriter(FArchive& InAr)
-		: Thread(nullptr)
-		, Ar(InAr)
-		, BufferStartPos(0)
-		, BufferEndPos(0)
-		, LastArchiveFlushTime(0.0)
-		, ArchiveFlushIntervalSec(0.2)
+/** [CLIENT THREAD] Flush the memory buffer (doesn't force the archive to flush). Can only be used from inside of BufferPosCritical lock. */
+void FAsyncWriter::FlushBuffer()
+{
+	SerializeRequestCounter.Increment();
+	if (!Thread)
 	{
-		Buffer.AddUninitialized(InitialBufferSize);
+		SerializeBufferToArchive();
+	}
+	while (SerializeRequestCounter.GetValue() != 0)
+	{
+		FPlatformProcess::SleepNoStats(0);
+	}
+	// Make sure there's been no unexpected concurrency
+	check(SerializeRequestCounter.GetValue() == 0);
+}
 
-		float CommandLineInterval = 0.0;
-		if (FParse::Value(FCommandLine::Get(), TEXT("LOGFLUSHINTERVAL="), CommandLineInterval))
-		{
-			ArchiveFlushIntervalSec = CommandLineInterval;
-		}
+FAsyncWriter::FAsyncWriter(FArchive& InAr)
+	: Thread(nullptr)
+	, Ar(InAr)
+	, BufferStartPos(0)
+	, BufferEndPos(0)
+	, LastArchiveFlushTime(0.0)
+{
+	Buffer.AddUninitialized(InitialBufferSize);
 
-		if (FPlatformProcess::SupportsMultithreading())
-		{
-			FString WriterName = FString::Printf(TEXT("FAsyncWriter_%s"), *FPaths::GetBaseFilename(Ar.GetArchiveName()));
-			FPlatformAtomics::InterlockedExchangePtr((void**)&Thread, FRunnableThread::Create(this, *WriterName, 0, TPri_BelowNormal));
-		}
+	float CommandLineInterval = 0.0;
+	if (FParse::Value(FCommandLine::Get(), TEXT("LOGFLUSHINTERVAL="), CommandLineInterval))
+	{
+		GLogFlushIntervalSec = CommandLineInterval;
 	}
 
-	FAsyncWriter::~FAsyncWriter()
+	if (FPlatformProcess::SupportsMultithreading())
 	{
-		Flush();
-		delete Thread;
-		Thread = nullptr;
+		FString WriterName = FString::Printf(TEXT("FAsyncWriter_%s"), *FPaths::GetBaseFilename(Ar.GetArchiveName()));
+		FPlatformAtomics::InterlockedExchangePtr((void**)&Thread, FRunnableThread::Create(this, *WriterName, 0, TPri_BelowNormal));
 	}
+}
+
+FAsyncWriter::~FAsyncWriter()
+{
+	Flush();
+	delete Thread;
+	Thread = nullptr;
+}
 
-	/** [CLIENT THREAD] Serialize data to buffer that will later be saved to disk by the async thread */
-	void FAsyncWriter::Serialize(void* InData, int64 Length)
+/** [CLIENT THREAD] Serialize data to buffer that will later be saved to disk by the async thread */
+void FAsyncWriter::Serialize(void* InData, int64 Length)
+{
+	if (!InData || Length <= 0)
 	{
-		if (!InData || Length <= 0)
-		{
-			return;
-		}
+		return;
+	}
 
-		const uint8* Data = (uint8*)InData;
+	const uint8* Data = (uint8*)InData;
 
-		FScopeLock WriteLock(&BufferPosCritical);
+	FScopeLock WriteLock(&BufferPosCritical);
 
-		const int32 ThisThreadEndPos = BufferEndPos.Load(EMemoryOrder::Relaxed);
+	const int32 ThisThreadEndPos = BufferEndPos.Load(EMemoryOrder::Relaxed);
 
-		// Store the local copy of the current buffer start pos. It may get moved by the worker thread but we don't
-		// care about it too much because we only modify BufferEndPos. Copy should be atomic enough. We only use it
-		// for checking the remaining space in the buffer so underestimating is ok.
+	// Store the local copy of the current buffer start pos. It may get moved by the worker thread but we don't
+	// care about it too much because we only modify BufferEndPos. Copy should be atomic enough. We only use it
+	// for checking the remaining space in the buffer so underestimating is ok.
+	{
+		const int32 ThisThreadStartPos = BufferStartPos.Load(EMemoryOrder::Relaxed);
+		// Calculate the remaining size in the ring buffer
+		const int32 BufferFreeSize = ThisThreadStartPos <= ThisThreadEndPos ? (Buffer.Num() - ThisThreadEndPos + ThisThreadStartPos) : (ThisThreadStartPos - ThisThreadEndPos);
+		// Make sure the buffer is BIGGER than we require otherwise we may calculate the wrong (0) buffer EndPos for StartPos = 0 and Length = Buffer.Num()
+		if (BufferFreeSize <= Length)
 		{
-			const int32 ThisThreadStartPos = BufferStartPos.Load(EMemoryOrder::Relaxed);
-			// Calculate the remaining size in the ring buffer
-			const int32 BufferFreeSize = ThisThreadStartPos <= ThisThreadEndPos ? (Buffer.Num() - ThisThreadEndPos + ThisThreadStartPos) : (ThisThreadStartPos - ThisThreadEndPos);
-			// Make sure the buffer is BIGGER than we require otherwise we may calculate the wrong (0) buffer EndPos for StartPos = 0 and Length = Buffer.Num()
-			if (BufferFreeSize <= Length)
-			{
-				// Force the async thread to call SerializeBufferToArchive even if it's currently empty
-				FlushBuffer();
+			// Force the async thread to call SerializeBufferToArchive even if it's currently empty
+			FlushBuffer();
 
-				// Resize the buffer if needed
-				if (Length >= Buffer.Num())
-				{
-					// Keep the buffer bigger than we require so that % Buffer.Num() does not return 0 for Lengths = Buffer.Num()
-					Buffer.SetNumUninitialized(Length + 1);
-				}
+			// Resize the buffer if needed
+			if (Length >= Buffer.Num())
+			{
+				// Keep the buffer bigger than we require so that % Buffer.Num() does not return 0 for Lengths = Buffer.Num()
+				Buffer.SetNumUninitialized(Length + 1);
 			}
 		}
-
-		// We now know there's enough space in the buffer to copy data
-		const int32 WritePos = ThisThreadEndPos;
-		if ((WritePos + Length) <= Buffer.Num())
-		{
-			// Copy straight into the ring buffer
-			FMemory::Memcpy(Buffer.GetData() + WritePos, Data, Length);
-		}
-		else
-		{
-			// Wrap around the ring buffer
-			int32 BufferSizeToEnd = Buffer.Num() - WritePos;
-			FMemory::Memcpy(Buffer.GetData() + WritePos, Data, BufferSizeToEnd);
-			FMemory::Memcpy(Buffer.GetData(), Data + BufferSizeToEnd, Length - BufferSizeToEnd);
-		}
-
-		// Update the end position and let the async thread know we need to write to disk
-		BufferEndPos = (ThisThreadEndPos + Length) % Buffer.Num();
-		SerializeRequestCounter.Increment();
-
-		// No async thread? Serialize now.
-		if (!Thread)
-		{
-			SerializeBufferToArchive();
-		}
 	}
 
-	/** Flush all buffers to disk */
-	void FAsyncWriter::Flush()
+	// We now know there's enough space in the buffer to copy data
+	const int32 WritePos = ThisThreadEndPos;
+	if ((WritePos + Length) <= Buffer.Num())
+	{
+		// Copy straight into the ring buffer
+		FMemory::Memcpy(Buffer.GetData() + WritePos, Data, Length);
+	}
+	else
 	{
-		FScopeLock WriteLock(&BufferPosCritical);
-		WantsArchiveFlush.Increment();
-		FlushBuffer();
+		// Wrap around the ring buffer
+		int32 BufferSizeToEnd = Buffer.Num() - WritePos;
+		FMemory::Memcpy(Buffer.GetData() + WritePos, Data, BufferSizeToEnd);
+		FMemory::Memcpy(Buffer.GetData(), Data + BufferSizeToEnd, Length - BufferSizeToEnd);
 	}
 
-	//~ Begin FRunnable Interface.
-	bool FAsyncWriter::Init()
+	// Update the end position and let the async thread know we need to write to disk
+	BufferEndPos = (ThisThreadEndPos + Length) % Buffer.Num();
+	SerializeRequestCounter.Increment();
+
+	// No async thread? Serialize now.
+	if (!Thread)
 	{
-		return true;
+		SerializeBufferToArchive();
 	}
+}
+
+/** Flush all buffers to disk */
+void FAsyncWriter::Flush()
+{
+	FScopeLock WriteLock(&BufferPosCritical);
+	WantsArchiveFlush.Increment();
+	FlushBuffer();
+}
+
+//~ Begin FRunnable Interface.
+bool FAsyncWriter::Init()
+{
+	return true;
+}
 	
-	uint32 FAsyncWriter::Run()
+uint32 FAsyncWriter::Run()
+{
+	while (StopTaskCounter.GetValue() == 0)
 	{
-		while (StopTaskCounter.GetValue() == 0)
+		if (SerializeRequestCounter.GetValue() > 0)
 		{
-			if (SerializeRequestCounter.GetValue() > 0)
-			{
-				SerializeBufferToArchive();
-			}
-			else if ((FPlatformTime::Seconds() - LastArchiveFlushTime) > ArchiveFlushIntervalSec)
-			{
-				FlushArchiveAndResetTimer();
-			}
-			else
-			{
-				FPlatformProcess::SleepNoStats(0.01f);
-			}
+			SerializeBufferToArchive();
+		}
+		else if ((FPlatformTime::Seconds() - LastArchiveFlushTime) > GetLogFlushIntervalSec() )
+		{
+			FlushArchiveAndResetTimer();
+		}
+		else
+		{
+			FPlatformProcess::SleepNoStats(0.01f);
 		}
-		return 0;
 	}
+	return 0;
+}
 
-	void FAsyncWriter::Stop()
-	{
-		StopTaskCounter.Increment();
-	}
+void FAsyncWriter::Stop()
+{
+	StopTaskCounter.Increment();
+}
 
 
 /**
diff --git a/Engine/Source/Runtime/Core/Public/Misc/OutputDeviceFile.h b/Engine/Source/Runtime/Core/Public/Misc/OutputDeviceFile.h
--- a/Engine/Source/Runtime/Core/Public/Misc/OutputDeviceFile.h
+++ b/Engine/Source/Runtime/Core/Public/Misc/OutputDeviceFile.h
@@ -49,9 +49,6 @@ class CORE_API FAsyncWriter : public FRunnable, public FArchive
 	/** [WRITER THREAD] Last time the archive was flushed. used in threaded situations to flush the underlying archive at a certain maximum rate. */
 	double LastArchiveFlushTime;
 
-	/** [WRITER THREAD] Archive flush interval. */
-	double ArchiveFlushIntervalSec;
-
 	/** [WRITER THREAD] Flushes the archive and reset the flush timer. */
 	void FlushArchiveAndResetTimer();
 
