Fix a nonunity build error

diff --git a/Engine/Source/Runtime/Core/Private/Misc/OutputDeviceFile.cpp b/Engine/Source/Runtime/Core/Private/Misc/OutputDeviceFile.cpp
--- a/Engine/Source/Runtime/Core/Private/Misc/OutputDeviceFile.cpp
+++ b/Engine/Source/Runtime/Core/Private/Misc/OutputDeviceFile.cpp
@@ -25,6 +25,7 @@
 #include "Misc/OutputDeviceHelper.h"
 #include "Math/Color.h"
 #include "Templates/Atomic.h"
+#include "HAL/ConsoleManager.h"
 
 /** Used by tools which include only core to disable log file creation. */
 #ifndef ALLOW_LOG_FILE
