Add plugin config overrides to config hierarchy so that changes from them don't get lost when reloading.
- Remove duplicate INI reloading hack from RuntimeOptions.  Instead make the INI global so that it is there at startup and gets properly patched by plugin loading.
- Remove conflicting INI settings from ESP.

diff --git a/Engine/Config/BaseRuntimeOptions.ini b/Engine/Config/BaseRuntimeOptions.ini
new file mode 100644
--- /dev/null
+++ b/Engine/Config/BaseRuntimeOptions.ini
@@ -0,0 +1,2 @@
+; Required to exist for to properly merge overrides at runtime
+[/Script/Engine.RuntimeOptionsBase]
\ No newline at end of file
diff --git a/Engine/Source/Editor/ConfigEditor/Private/SConfigEditor.cpp b/Engine/Source/Editor/ConfigEditor/Private/SConfigEditor.cpp
--- a/Engine/Source/Editor/ConfigEditor/Private/SConfigEditor.cpp
+++ b/Engine/Source/Editor/ConfigEditor/Private/SConfigEditor.cpp
@@ -84,7 +84,7 @@ void SConfigEditor::CreateDisplayObjectForSelectedTargetPlatform()
 	FConfigFile PlatformIniFile;
 	LocalConfigCache->LoadLocalIniFile(PlatformIniFile, *ClassConfigName, true, *SelectedTargetPlatform);
 
-	for (auto& IniFile : PlatformIniFile.SourceIniHierarchy)
+	for (const auto& IniFile : PlatformIniFile.SourceIniHierarchy)
 	{
 		UPropertyConfigFileDisplayRow* ConfigFilePropertyObj = NewObject<UPropertyConfigFileDisplayRow>(GetTransientPackage(), *IniFile.Value.Filename);
 		ConfigFilePropertyObj->InitWithConfigAndProperty(IniFile.Value.Filename, PropHelper->EditProperty.Get());
diff --git a/Engine/Source/Runtime/Core/Private/Misc/ConfigCacheIni.cpp b/Engine/Source/Runtime/Core/Private/Misc/ConfigCacheIni.cpp
--- a/Engine/Source/Runtime/Core/Private/Misc/ConfigCacheIni.cpp
+++ b/Engine/Source/Runtime/Core/Private/Misc/ConfigCacheIni.cpp
@@ -473,6 +473,134 @@
 }
 
 
+enum class EConfigLayerFlags : int32
+{
+	Required					= (1 << 0),
+	AllowCommandLineOverride	= (1 << 1),
+	DedicatedServerOnly			= (1 << 2), // replaces Default, Base, and (NOT {PLATFORM} yet) with an empty string
+	GenerateCacheKey			= (1 << 3),
+};
+ENUM_CLASS_FLAGS(EConfigLayerFlags);
+
+/**
+ * Structure to define all the layers of the config system. Layers can be expanded by expansion files (NoRedist, etc), or by ini platform parents
+ * (coming soon from another branch)
+ */
+struct FConfigLayer
+{
+	// Used by the editor to display in the ini-editor
+	const TCHAR* EditorName;
+	// Path to the ini file (with variables)
+	const TCHAR* Path;
+	// Path to the platform extension version
+	const TCHAR* PlatformExtensionPath;
+	// Special flag
+	EConfigLayerFlags Flag;
+
+} GConfigLayers[] =
+{
+	/**************************************************
+	**** CRITICAL NOTES
+	**** If you change this array, you need to also change EnumerateConfigFileLocations() in ConfigHierarchy.cs!!!
+	**** And maybe UObject::GetDefaultConfigFilename(), UObject::GetGlobalUserConfigFilename()
+	**************************************************/
+
+	// Engine/Base.ini
+	{ TEXT("AbsoluteBase"),				TEXT("{ENGINE}Base.ini"), TEXT(""), EConfigLayerFlags::Required },
+	
+	// Engine/Base*.ini
+	{ TEXT("Base"),						TEXT("{ENGINE}{ED}{EF}Base{TYPE}.ini") },
+	// Engine/Platform/BasePlatform*.ini
+	{ TEXT("BasePlatform"),				TEXT("{ENGINE}{ED}{PLATFORM}/{EF}Base{PLATFORM}{TYPE}.ini"), TEXT("{EXTENGINE}/{ED}{EF}Base{PLATFORM}{TYPE}.ini"),  },
+	// Project/Default*.ini
+	{ TEXT("ProjectDefault"),			TEXT("{PROJECT}{ED}{EF}Default{TYPE}.ini"), TEXT(""), EConfigLayerFlags::AllowCommandLineOverride | EConfigLayerFlags::GenerateCacheKey },
+	// Engine/Platform/Platform*.ini
+	{ TEXT("EnginePlatform"),			TEXT("{ENGINE}{ED}{PLATFORM}/{EF}{PLATFORM}{TYPE}.ini"), TEXT("{EXTENGINE}/{ED}{EF}{PLATFORM}{TYPE}.ini") },
+	// Project/Platform/Platform*.ini
+	{ TEXT("ProjectPlatform"),			TEXT("{PROJECT}{ED}{PLATFORM}/{EF}{PLATFORM}{TYPE}.ini"), TEXT("{EXTPROJECT}/{ED}{EF}{PLATFORM}{TYPE}.ini") },
+	
+	// UserSettings/.../User*.ini
+	{ TEXT("UserSettingsDir"),			TEXT("{USERSETTINGS}Unreal Engine/Engine/Config/User{TYPE}.ini") },
+	// UserDir/.../User*.ini
+	{ TEXT("UserDir"),					TEXT("{USER}Unreal Engine/Engine/Config/User{TYPE}.ini") },
+	// Project/User*.ini
+	{ TEXT("GameDirUser"),				TEXT("{PROJECT}User{TYPE}.ini"), TEXT(""), EConfigLayerFlags::GenerateCacheKey },
+};
+
+
+/**
+ * This describes extra files per layer, most so that Epic needs to be able to ship a project, but still share the project 
+ * with licensees. These settings are the settings that should not be shared outside of Epic because they could cause 
+ * problems if a licensee blindly copy and pasted the settings (they can't copy what they don't have!)
+ */
+struct FConfigLayerExpansion
+{
+	// The subdirectory for this expansion (ie "NoRedist")
+	const TCHAR* DirectoryPrefix;
+	// The filename prefix for this expansion (ie "Shippable")
+	const TCHAR* FilePrefix;
+	// Optional flags 
+	EConfigLayerFlags Flag;
+} GConfigLayerExpansions[] = 
+{
+	/**************************************************
+	**** CRITICAL NOTES
+	**** If you change this array, you need to also change EnumerateConfigFileLocations() in ConfigHierarchy.cs!!!
+	**************************************************/
+	
+	// The base expansion (ie, no expansion)
+	{ TEXT(""), TEXT("") },
+
+	// When running a dedicated server - does not support {PLATFORM} layers, so those are skipped
+	{ TEXT(""), TEXT("DedicatedServer"),  EConfigLayerFlags::DedicatedServerOnly },
+	// This file is remapped in UAT from inside NFL or NoRedist, because those directories are stripped while packaging
+	{ TEXT(""), TEXT("Shippable") },
+	// Hidden directory from licensees
+	{ TEXT("NotForLicensees/"), TEXT("") },
+	// Settings that need to be hidden from licensees, but are needed for shipping
+	{ TEXT("NotForLicensees/"), TEXT("Shippable") },
+	// Hidden directory from non-Epic
+	{ TEXT("NoRedist/"), TEXT("") },
+	// Settings that need to be hidden from non-Epic, but are needed for shipping
+	{ TEXT("NoRedist/"), TEXT("Shippable") },
+};
+
+constexpr int32 MaxPlatformIndex = 99;
+
+constexpr int32 GetStaticKey(int32 LayerIndex, int32 LayerExpansionIndex, int32 PlatformIndex)
+{
+	return LayerIndex * 10000 + LayerExpansionIndex * 100 + PlatformIndex;
+}
+
+constexpr int32 MaxStaticHierarchyKey = GetStaticKey(ARRAY_COUNT(GConfigLayers) - 1, ARRAY_COUNT(GConfigLayerExpansions) - 1, MaxPlatformIndex);
+
 /*-----------------------------------------------------------------------------
+	FConfigFileHierarchy
+-----------------------------------------------------------------------------*/
+FConfigFileHierarchy::FConfigFileHierarchy() : KeyGen(MaxStaticHierarchyKey)
+{
+}
+
+
+int32 FConfigFileHierarchy::GenerateDynamicKey()
+{
+	return ++KeyGen;
+}
+
+int32 FConfigFileHierarchy::AddStaticLayer(FIniFilename Filename, int32 LayerIndex, int32 LayerExpansionIndex /*= 0*/, int32 PlatformIndex /*= 0*/)
+{
+	int32 Key = GetStaticKey(LayerIndex, LayerExpansionIndex, PlatformIndex);
+	Emplace(Key, MoveTemp(Filename));
+	return Key;
+}
+
+int32 FConfigFileHierarchy::AddDynamicLayer(FIniFilename Filename)
+{
+	int32 Key = GenerateDynamicKey();
+	Emplace(Key, MoveTemp(Filename));
+	return Key;
+}
+
+/*-----------------------------------------------------------------------------
 	FConfigFile
 -----------------------------------------------------------------------------*/
@@ -1041,6 +1169,29 @@
 }
 
 
+void FConfigFile::AddDynamicLayerToHeirarchy(const FString& Filename)
+{
+	// Don't allow dynamic layers in editor
+	if (GIsEditor)
+		return;
+
+	FString ConfigContent;
+	if (!FFileHelper::LoadFileToString(ConfigContent, *Filename))
+		return;
+
+	if (SourceConfigFile)
+	{
+		SourceConfigFile->SourceIniHierarchy.AddDynamicLayer(FIniFilename(Filename));
+		SourceConfigFile->CombineFromBuffer(ConfigContent);
+	}
+
+	SourceIniHierarchy.AddDynamicLayer(FIniFilename(Filename));
+	CombineFromBuffer(ConfigContent);
+
+	// Disable saving since dynamic layers are only for runtime
+	NoSave = true;
+}
+
 /**
  * This will completely load .ini file hierarchy into the passed in FConfigFile. The passed in FConfigFile will then
  * have the data after combining all of those .ini 
@@ -3168,102 +3319,6 @@
 }
 
 
-
-
-
-
-const int Flag_Required = 1;
-const int Flag_AllowCommandLineOverride = 2;
-const int Flag_DedicatedServerOnly = 4; // replaces Default, Base, and (NOT {PLATFORM} yet) with an empty string
-const int Flag_GenerateCacheKey = 8;
-
-
 /**
- * Structure to define all the layers of the config system. Layers can be expanded by expansion files (NoRedist, etc), or by ini platform parents
- * (coming soon from another branch)
- */
-struct FConfigLayer
-{
-	// Used by the editor to display in the ini-editor
-	const TCHAR* EditorName;
-	// Path to the ini file (with variables)
-	const TCHAR* Path;
-	// Path to the platform extension version
-	const TCHAR* PlatformExtensionPath;
-	// Special flag
-	int32  Flag;
-
-} GConfigLayers[] =
-{
-	/**************************************************
-	**** CRITICAL NOTES
-	**** If you change this array, you need to also change EnumerateConfigFileLocations() in ConfigHierarchy.cs!!!
-	**** And maybe UObject::GetDefaultConfigFilename(), UObject::GetGlobalUserConfigFilename()
-	**************************************************/
-
-	// Engine/Base.ini
-	{ TEXT("AbsoluteBase"),				TEXT("{ENGINE}Base.ini"), TEXT(""), Flag_Required },
-	
-	// Engine/Base*.ini
-	{ TEXT("Base"),						TEXT("{ENGINE}{ED}{EF}Base{TYPE}.ini") },
-	// Engine/Platform/BasePlatform*.ini
-	{ TEXT("BasePlatform"),				TEXT("{ENGINE}{ED}{PLATFORM}/{EF}Base{PLATFORM}{TYPE}.ini"), TEXT("{EXTENGINE}/{ED}{EF}Base{PLATFORM}{TYPE}.ini"),  },
-	// Project/Default*.ini
-	{ TEXT("ProjectDefault"),			TEXT("{PROJECT}{ED}{EF}Default{TYPE}.ini"), TEXT(""), Flag_AllowCommandLineOverride | Flag_GenerateCacheKey },
-	// Engine/Platform/Platform*.ini
-	{ TEXT("EnginePlatform"),			TEXT("{ENGINE}{ED}{PLATFORM}/{EF}{PLATFORM}{TYPE}.ini"), TEXT("{EXTENGINE}/{ED}{EF}{PLATFORM}{TYPE}.ini") },
-	// Project/Platform/Platform*.ini
-	{ TEXT("ProjectPlatform"),			TEXT("{PROJECT}{ED}{PLATFORM}/{EF}{PLATFORM}{TYPE}.ini"), TEXT("{EXTPROJECT}/{ED}{EF}{PLATFORM}{TYPE}.ini") },
-	
-	// UserSettings/.../User*.ini
-	{ TEXT("UserSettingsDir"),			TEXT("{USERSETTINGS}Unreal Engine/Engine/Config/User{TYPE}.ini") },
-	// UserDir/.../User*.ini
-	{ TEXT("UserDir"),					TEXT("{USER}Unreal Engine/Engine/Config/User{TYPE}.ini") },
-	// Project/User*.ini
-	{ TEXT("GameDirUser"),				TEXT("{PROJECT}User{TYPE}.ini"), TEXT(""), Flag_GenerateCacheKey },
-};
-
-
-/**
- * This describes extra files per layer, most so that Epic needs to be able to ship a project, but still share the project 
- * with licensees. These settings are the settings that should not be shared outside of Epic because they could cause 
- * problems if a licensee blindly copy and pasted the settings (they can't copy what they don't have!)
- */
-struct FConfigLayerExpansion
-{
-	// The subdirectory for this expansion (ie "NoRedist")
-	const TCHAR* DirectoryPrefix;
-	// The filename prefix for this expansion (ie "Shippable")
-	const TCHAR* FilePrefix;
-	// Optional flags 
-	int32 Flag;
-} GConfigLayerExpansions[] = 
-{
-	/**************************************************
-	**** CRITICAL NOTES
-	**** If you change this array, you need to also change EnumerateConfigFileLocations() in ConfigHierarchy.cs!!!
-	**************************************************/
-	
-	// The base expansion (ie, no expansion)
-	{ TEXT(""), TEXT("") },
-
-	// When running a dedicated server - does not support {PLATFORM} layers, so those are skipped
-	{ TEXT(""), TEXT("DedicatedServer"),  Flag_DedicatedServerOnly },
-	// This file is remapped in UAT from inside NFL or NoRedist, because those directories are stripped while packaging
-	{ TEXT(""), TEXT("Shippable") },
-	// Hidden directory from licensees
-	{ TEXT("NotForLicensees/"), TEXT("") },
-	// Settings that need to be hidden from licensees, but are needed for shipping
-	{ TEXT("NotForLicensees/"), TEXT("Shippable") },
-	// Hidden directory from non-Epic
-	{ TEXT("NoRedist/"), TEXT("") },
-	// Settings that need to be hidden from non-Epic, but are needed for shipping
-	{ TEXT("NoRedist/"), TEXT("Shippable") },
-};
-
-
-
-
-/**
  * Creates a chain of ini filenames to load and combine.
  *
@@ -3271,11 +3326,10 @@
  * @param InPlatformName Platform name, nullptr means to use the current platform
  * @param OutHierarchy An array which is to receive the generated hierachy of ini filenames.
  */
-static void GetSourceIniHierarchyFilenames(const TCHAR* InBaseIniName, const TCHAR* InPlatformName, const TCHAR* EngineConfigDir, const TCHAR* SourceConfigDir, FConfigFile& OutFile, bool bRequireDefaultIni)
+void FConfigFile::AddStaticLayersToHierarchy(const TCHAR* InBaseIniName, const TCHAR* InPlatformName, const TCHAR* EngineConfigDir, const TCHAR* SourceConfigDir)
 {
-	FConfigFileHierarchy& OutHierarchy = OutFile.SourceIniHierarchy;
-	OutFile.SourceEngineConfigDir = EngineConfigDir;
-	OutFile.SourceProjectConfigDir = SourceConfigDir;
+	SourceEngineConfigDir = EngineConfigDir;
+	SourceProjectConfigDir = SourceConfigDir;
 	
 	// get the platform name
 	const FString PlatformName(InPlatformName ? InPlatformName : ANSI_TO_TCHAR(FPlatformProperties::IniPlatformName()));
@@ -3293,9 +3347,6 @@
 		const bool bHasPlatformTag = FCString::Strstr(Layer.Path, TEXT("{PLATFORM}")) != nullptr;
 		const bool bHasProjectTag = FCString::Strstr(Layer.Path, TEXT("{PROJECT}")) != nullptr;
 
-		// calculate a increasing-by-layer index
-		int32 ConfigFileIndex = LayerIndex * 10000;
-
 		// start replacing basic variables
 		FString LayerPath;
 		// you can only have PROJECT or ENGINE, not both
@@ -3330,7 +3381,7 @@
 #if IS_PROGRAM
 		const bool bIsRequired = false;
 #else
-		const bool bIsRequired = ((Layer.Flag & Flag_Required) != 0) && (EngineConfigDir == FPaths::EngineConfigDir());
+		const bool bIsRequired = EnumHasAnyFlags(Layer.Flag, EConfigLayerFlags::Required) && (EngineConfigDir == FPaths::EngineConfigDir());
 #endif
 
 		// expand if it it has {ED} or {EF} expansion tags
@@ -3346,15 +3397,12 @@
 			{
 				const FConfigLayerExpansion& Expansion = GConfigLayerExpansions[ExpansionIndex];
 
-				// apply expansion level
-				int32 ExpansionFileIndex = ConfigFileIndex + ExpansionIndex * 100;
-
 				// replace the expansion tags
 				FString ExpansionPath = LayerPath.Replace(TEXT("{ED}"), Expansion.DirectoryPrefix, ESearchCase::CaseSensitive);
 				ExpansionPath = ExpansionPath.Replace(TEXT("{EF}"), Expansion.FilePrefix, ESearchCase::CaseSensitive);
 
 				// check for dedicated server expansion
-				if (Expansion.Flag & Flag_DedicatedServerOnly)
+				if (EnumHasAnyFlags(Expansion.Flag, EConfigLayerFlags::DedicatedServerOnly))
 				{
 					// it's unclear how a platform DS ini would be named, so for now, not supported
 					if (bHasPlatformTag)
@@ -3375,21 +3423,22 @@
 				}
 
 				// allow for override, only on BASE EXPANSION!
-				if ((Layer.Flag & Flag_AllowCommandLineOverride) && ExpansionIndex == 0)
+				if (EnumHasAnyFlags(Layer.Flag, EConfigLayerFlags::AllowCommandLineOverride) && ExpansionIndex == 0)
 				{
-					checkfSlow(!bHasPlatformTag, TEXT("Flag_AllowCommandLineOverride config %s shouldn't have a PLATFORM in it"), *Layer.Path);
+					checkfSlow(!bHasPlatformTag, TEXT("EConfigLayerFlags::AllowCommandLineOverride config %s shouldn't have a PLATFORM in it"), *Layer.Path);
 
 					ConditionalOverrideIniFilename(ExpansionPath, InBaseIniName);
 				}
 
 				// check if we should be generating the cache key - only at the end of all expansions
-				bool bGenerateCacheKey = (Layer.Flag & Flag_GenerateCacheKey) && ExpansionIndex == ARRAY_COUNT(GConfigLayerExpansions);
-				checkfSlow(!(bGenerateCacheKey && bHasPlatformTag), TEXT("Flag_GenerateCacheKey shouldn't have a platform tag"));
+				bool bGenerateCacheKey = EnumHasAnyFlags(Layer.Flag, EConfigLayerFlags::GenerateCacheKey) && ExpansionIndex == ARRAY_COUNT(GConfigLayerExpansions);
+				checkfSlow(!(bGenerateCacheKey && bHasPlatformTag), TEXT("EConfigLayerFlags::GenerateCacheKey shouldn't have a platform tag"));
 
 				const FDataDrivenPlatformInfoRegistry::FPlatformInfo& Info = FDataDrivenPlatformInfoRegistry::GetPlatformInfo(PlatformName);
 
 				// go over parents, and then this platform, unless there's no platform tag, then we simply want to run through the loop one time to add it to the 
 				int NumPlatforms = bHasPlatformTag ? Info.IniParentChain.Num() + 1 : 1;
+				check(NumPlatforms < MaxPlatformIndex);
 				for (int PlatformIndex = 0; PlatformIndex < NumPlatforms; PlatformIndex++)
 				{
 					// the platform to work on (active platform is always last)
@@ -3399,11 +3448,11 @@
 					FString PlatformPath = ExpansionPath.Replace(TEXT("{PLATFORM}"), *CurrentPlatform, ESearchCase::CaseSensitive);
 
 					// add this to the list!
-					OutHierarchy.Add(ExpansionFileIndex, FIniFilename(PlatformPath, bIsRequired,
-						bGenerateCacheKey ? GenerateHierarchyCacheKey(OutHierarchy, PlatformPath, InBaseIniName) : FString(TEXT(""))));
-
-					// we can incrememnt it here because we reset it next expansion file
-					ExpansionFileIndex++;
+					SourceIniHierarchy.AddStaticLayer(
+						FIniFilename(PlatformPath, bIsRequired, bGenerateCacheKey ? 
+							GenerateHierarchyCacheKey(SourceIniHierarchy, PlatformPath, InBaseIniName) :
+							FString(TEXT(""))),
+						LayerIndex, ExpansionIndex, PlatformIndex);
 				}
 			}
 		}
@@ -3411,11 +3460,14 @@
 		else
 		{
 			checkfSlow(!bHasPlatformTag, TEXT("Non-expanded config %s shouldn't have a PLATFORM in it"), *Layer.Path);
-			checkfSlow(!(Layer.Flag & Flag_AllowCommandLineOverride), TEXT("Non-expanded config can't have a Flag_AllowCommandLineOverride"));
+			checkfSlow(!EnumHasAnyFlags(Layer.Flag, EConfigLayerFlags::AllowCommandLineOverride), TEXT("Non-expanded config can't have a EConfigLayerFlags::AllowCommandLineOverride"));
 
 			// final layer needs to generate a hierarchy cache
-			OutHierarchy.Add(ConfigFileIndex, FIniFilename(LayerPath, bIsRequired, 
-				(Layer.Flag & Flag_GenerateCacheKey) ? GenerateHierarchyCacheKey(OutHierarchy, LayerPath, InBaseIniName) : FString(TEXT(""))));
+			SourceIniHierarchy.AddStaticLayer(
+				FIniFilename(LayerPath, bIsRequired, EnumHasAnyFlags(Layer.Flag, EConfigLayerFlags::GenerateCacheKey) ? 
+					GenerateHierarchyCacheKey(SourceIniHierarchy, LayerPath, InBaseIniName) :
+					FString(TEXT(""))),
+				LayerIndex);
 		}
 	}
 }
@@ -3516,7 +3568,9 @@
 	// Load scalability settings.
 	FConfigCacheIni::LoadGlobalIniFile(GScalabilityIni, TEXT("Scalability"), ScalabilityPlatformOverride);
 	// Load driver blacklist
-	FConfigCacheIni::LoadGlobalIniFile(GHardwareIni, TEXT("Hardware"));
+	FConfigCacheIni::LoadGlobalIniFile(GHardwareIni, TEXT("Hardware"));	
+	// Load runtime options
+	FConfigCacheIni::LoadGlobalIniFile(GRuntimeOptionsIni, TEXT("RuntimeOptions"));
 	
 	// Load user game settings .ini, allowing merging. This also updates the user .ini if necessary.
 #if PLATFORM_PS4
@@ -3580,7 +3634,17 @@
 	// make a new entry in GConfig (overwriting what's already there)
 	FConfigFile& NewConfigFile = GConfig->Add(FinalIniFilename, FConfigFile());
 
-	return LoadExternalIniFile(NewConfigFile, BaseIniName, *EngineConfigDir, *SourceConfigDir, true, Platform, bForceReload, true, bAllowGeneratedIniWhenCooked, GeneratedConfigDir);
+	return LoadExternalIniFile(
+		NewConfigFile, 
+		BaseIniName, 
+		*EngineConfigDir, 
+		*SourceConfigDir, 
+		/*bIsBaseIniName*/ true,
+		Platform, 
+		bForceReload, 
+		/*bWriteDestIni*/ true,
+		bAllowGeneratedIniWhenCooked, 
+		GeneratedConfigDir);
 }
 
 bool FConfigCacheIni::LoadLocalIniFile(FConfigFile& ConfigFile, const TCHAR* IniName, bool bIsBaseIniName, const TCHAR* Platform, bool bForceReload )
@@ -3608,7 +3672,15 @@
 		}
 	}
 
-	return LoadExternalIniFile(ConfigFile, IniName, *EngineConfigDir, *SourceConfigDir, bIsBaseIniName, Platform, bForceReload, false);
+	return LoadExternalIniFile(
+		ConfigFile, 
+		IniName, 
+		*EngineConfigDir, 
+		*SourceConfigDir, 
+		bIsBaseIniName, 
+		Platform, 
+		bForceReload, 
+		/*bWriteDestIni*/ false);
 }
 
 bool FConfigCacheIni::LoadExternalIniFile(FConfigFile& ConfigFile, const TCHAR* IniName, const TCHAR* EngineConfigDir, const TCHAR* SourceConfigDir, bool bIsBaseIniName, const TCHAR* Platform, bool bForceReload, bool bWriteDestIni, bool bAllowGeneratedIniWhenCooked, const TCHAR* GeneratedConfigDir)
@@ -3641,7 +3713,7 @@
 #endif
 		FString DestIniFilename = GetDestIniFilename(IniName, Platform, GeneratedConfigDir);
 
-		GetSourceIniHierarchyFilenames( IniName, Platform, EngineConfigDir, SourceConfigDir, ConfigFile, false );
+		ConfigFile.AddStaticLayersToHierarchy( IniName, Platform, EngineConfigDir, SourceConfigDir );
 
 		if (bForceReload)
 		{
@@ -3657,7 +3729,6 @@
 		ConfigFile.SourceConfigFile = new FConfigFile();
 
 		// now generate and make sure it's up to date (using IniName as a Base for an ini filename)
-		const bool bAllowGeneratedINIs = true;
 		// @todo This bNeedsWrite afaict is always true even if it loaded a completely valid generated/final .ini, and the write below will
 		// just write out the exact same thing it read in!
 		bool bNeedsWrite = GenerateDestIniFile(ConfigFile, DestIniFilename, ConfigFile.SourceIniHierarchy, bAllowGeneratedIniWhenCooked, true);
@@ -3770,7 +3841,7 @@
 	{
 		// get the standard ini files
 		SourceIniHierarchy.Empty();
-		GetSourceIniHierarchyFilenames(IniRootName, OverridePlatform, *FPaths::EngineConfigDir(), *FPaths::SourceConfigDir(), *this, false);
+		AddStaticLayersToHierarchy(IniRootName, OverridePlatform, *FPaths::EngineConfigDir(), *FPaths::SourceConfigDir());
 
 		// now chop off this file and any after it
 		TArray<int32> Keys;
@@ -3800,7 +3871,6 @@
 
 		// now when Write it called below, it will diff against SourceIniHierarchy
 		LoadIniFileHierarchy(SourceIniHierarchy, *SourceConfigFile, true);
-
 	}
 
 	// take what we got above (which has the sections skipped), and then append the new sections
diff --git a/Engine/Source/Runtime/Core/Private/Misc/CoreGlobals.cpp b/Engine/Source/Runtime/Core/Private/Misc/CoreGlobals.cpp
--- a/Engine/Source/Runtime/Core/Private/Misc/CoreGlobals.cpp
+++ b/Engine/Source/Runtime/Core/Private/Misc/CoreGlobals.cpp
@@ -220,6 +220,7 @@ FString				GHardwareIni;												/* Hardware ini filename */
 FString				GInputIni;													/* Input ini filename */
 FString				GGameIni;													/* Game ini filename */
 FString				GGameUserSettingsIni;										/* User Game Settings ini filename */
+FString				GRuntimeOptionsIni;											/* Runtime Options ini filename */
 
 float					GNearClippingPlane				= 10.0f;				/* Near clipping plane */
 
diff --git a/Engine/Source/Runtime/Core/Public/CoreGlobals.h b/Engine/Source/Runtime/Core/Public/CoreGlobals.h
--- a/Engine/Source/Runtime/Core/Public/CoreGlobals.h
+++ b/Engine/Source/Runtime/Core/Public/CoreGlobals.h
@@ -302,6 +302,7 @@ extern CORE_API FString GHardwareIni;
 extern CORE_API FString GInputIni;
 extern CORE_API FString GGameIni;
 extern CORE_API FString GGameUserSettingsIni;
+extern CORE_API FString GRuntimeOptionsIni;
 
 extern CORE_API float GNearClippingPlane;
 
diff --git a/Engine/Source/Runtime/Core/Public/Misc/ConfigCacheIni.h b/Engine/Source/Runtime/Core/Public/Misc/ConfigCacheIni.h
--- a/Engine/Source/Runtime/Core/Public/Misc/ConfigCacheIni.h
+++ b/Engine/Source/Runtime/Core/Public/Misc/ConfigCacheIni.h
@@ -260,11 +260,11 @@ struct FIniFilename
 	/** Ini filename */
 	FString Filename;
 	/** If true this ini file is required to generate the output ini. */
-	bool bRequired;
+	bool bRequired = false;
 	/** Used as ID for looking up an INI Hierarchy */
 	FString CacheKey;
 
-	FIniFilename(const FString& InFilename, bool InIsRequired, FString InCacheKey=FString(TEXT("")))
+	explicit FIniFilename(const FString& InFilename, bool InIsRequired=false, FString InCacheKey=FString(TEXT("")))
 		: Filename(InFilename)
 		, bRequired(InIsRequired) 
 		, CacheKey(InCacheKey)
@@ -281,14 +281,30 @@ struct FConfigCommandlineOverride
 #endif // ALLOW_INI_OVERRIDE_FROM_COMMANDLINE
 
 
-typedef TMap<int32, FIniFilename> FConfigFileHierarchy;
+class FConfigFileHierarchy : public TMap<int32, FIniFilename>
+{
+private:
+	int32 KeyGen = 0;
+
+public:
+	FConfigFileHierarchy();
+
+private:
+	int32 GenerateDynamicKey();
+
+	int32 AddStaticLayer(FIniFilename Filename, int32 LayerIndex, int32 LayerExpansionIndex = 0, int32 PlatformIndex = 0);
+	int32 AddDynamicLayer(FIniFilename Filename);
+
+	friend class FConfigFile;
+};
 
 // One config file.
 
 class FConfigFile : public TMap<FString,FConfigSection>
 {
 public:
-	bool Dirty, NoSave;
+	bool Dirty;
+	bool NoSave;
 
 	/** The name of this config file */	
 	FName Name;
@@ -390,6 +406,9 @@ public:
 	/** Checks the command line for any overridden config file settings */
 	CORE_API static void OverrideFileFromCommandline(FString& Filename);
 
+	/** Appends a new INI file to the SourceIniHierarchy and combines it */
+	CORE_API void AddDynamicLayerToHeirarchy(const FString& Filename);
+
 private:
 
 	// This holds per-object config class names, with their ArrayOfStructKeys. Since the POC sections are all unique,
@@ -413,6 +432,17 @@ private:
 	 */
 	void ProcessPropertyAndWriteForDefaults(int32 IniCombineThreshold, const TArray<FConfigValue>& InCompletePropertyToProcess, FString& OutText, const FString& SectionName, const FString& PropertyName);
 
+	/**
+	 * Creates a chain of ini filenames to load and combine.
+	 *
+	 * @param InBaseIniName Ini name.
+	 * @param InPlatformName Platform name, nullptr means to use the current platform
+	 * @param OutHierarchy An array which is to receive the generated hierachy of ini filenames.
+	 */
+	void AddStaticLayersToHierarchy(const TCHAR* InBaseIniName, const TCHAR* InPlatformName, const TCHAR* EngineConfigDir, const TCHAR* SourceConfigDir);
+
+	// for AddStaticLayersToHierarchy
+	friend class FConfigCacheIni;
 };
 
 /**
diff --git a/Engine/Source/Runtime/Engine/Classes/Engine/RuntimeOptionsBase.h b/Engine/Source/Runtime/Engine/Classes/Engine/RuntimeOptionsBase.h
--- a/Engine/Source/Runtime/Engine/Classes/Engine/RuntimeOptionsBase.h
+++ b/Engine/Source/Runtime/Engine/Classes/Engine/RuntimeOptionsBase.h
@@ -62,6 +62,4 @@ protected:
 	void ApplyCommandlineOverrides();
 	void RegisterSupportedConsoleVariables(bool bDuringReload);
 	void InitializeRuntimeOptions();
-
-	void ApplyPluginOverridesIfPresent(const FString& PluginName);
 };
diff --git a/Engine/Source/Runtime/Engine/Private/RuntimeOptionsBase.cpp b/Engine/Source/Runtime/Engine/Private/RuntimeOptionsBase.cpp
--- a/Engine/Source/Runtime/Engine/Private/RuntimeOptionsBase.cpp
+++ b/Engine/Source/Runtime/Engine/Private/RuntimeOptionsBase.cpp
@@ -47,37 +47,6 @@ void URuntimeOptionsBase::InitializeRuntimeOptions()
 	}
 }
 
-void URuntimeOptionsBase::ApplyPluginOverridesIfPresent(const FString& PluginName)
-{
-	check(HasAnyFlags(RF_ClassDefaultObject));
-
-	TSharedPtr<IPlugin> Plugin = IPluginManager::Get().FindPlugin(*PluginName);
-	if (Plugin.IsValid() && Plugin->IsEnabled())
-	{
-		const FString PluginConfigDir = Plugin->GetBaseDir() / TEXT("Config/");
-		TArray<FString> PluginConfigs;
-		IFileManager::Get().FindFiles(PluginConfigs, *PluginConfigDir, TEXT("ini"));
-		const FString PlaformName = FPlatformProperties::PlatformName();
-
-		for (const FString& ConfigFile : PluginConfigs)
-		{
-			const FString PluginConfigFilename = FString::Printf(TEXT("%s%s/%s.ini"), *FPaths::GeneratedConfigDir(), *PlaformName, *FPaths::GetBaseFilename(ConfigFile));
-			if (FConfigFile* FoundConfig = GConfig->Find(PluginConfigFilename, false))
-			{
-				FString PluginConfigContent;
-				if (FFileHelper::LoadFileToString(PluginConfigContent, *FPaths::Combine(PluginConfigDir, ConfigFile)))
-				{
-					FoundConfig->CombineFromBuffer(PluginConfigContent);
-					// if plugin config overrides are applied then don't save
-					FoundConfig->NoSave = true;
-
-					FConfigFile::OverrideFromCommandline(FoundConfig, PluginConfigFilename);
-				}
-			}
-		}
-	}
-}
-
 void URuntimeOptionsBase::ApplyCommandlineOverrides()
 {
 #if UE_RUNTIMEOPTIONSBASE_SUPPORT_COMMANDLINE
diff --git a/Engine/Source/Runtime/Projects/Private/PluginManager.cpp b/Engine/Source/Runtime/Projects/Private/PluginManager.cpp
--- a/Engine/Source/Runtime/Projects/Private/PluginManager.cpp
+++ b/Engine/Source/Runtime/Projects/Private/PluginManager.cpp
@@ -679,13 +679,7 @@ bool FPluginManager::ConfigureEnabledPlugins()
 						{
 							UE_LOG(LogPluginManager, Log, TEXT("Found config from plugin[%s] %s"), *Plugin.GetName(), *PluginConfigFilename);
 
-							FString PluginConfigContent;
-							if (FFileHelper::LoadFileToString(PluginConfigContent, *FPaths::Combine(PluginConfigDir, ConfigFile)))
-							{
-								FoundConfig->CombineFromBuffer(PluginConfigContent);
-								// if plugin config overrides are applied then don't save
-								FoundConfig->NoSave = true;
-							}
+							FoundConfig->AddDynamicLayerToHeirarchy(FPaths::Combine(PluginConfigDir, ConfigFile));
 
 #if ALLOW_INI_OVERRIDE_FROM_COMMANDLINE
 							// Don't allow plugins to stomp command line overrides
