Fixed FText::IdenticalTo comparing two instances with null localized strings as identical

Invariant and generated FText instances will have null localized strings unless they're persisted, and should never be considered identical to anything else.

diff --git a/Engine/Source/Runtime/Core/Private/Internationalization/Text.cpp b/Engine/Source/Runtime/Core/Private/Internationalization/Text.cpp
--- a/Engine/Source/Runtime/Core/Private/Internationalization/Text.cpp
+++ b/Engine/Source/Runtime/Core/Private/Internationalization/Text.cpp
@@ -1175,10 +1175,18 @@ bool FText::GetHistoricNumericData(FHistoricTextNumericData& OutHistoricNumericD
 
 bool FText::IdenticalTo( const FText& Other ) const
 {
-	// If both instances point to the same data or localized string, then both instances are considered identical.
+	// If both instances point to the same data, then both instances are considered identical.
+	if (TextData == Other.TextData)
+	{
+		return true;
+	}
+
+	// If both instances point to the same localized string, then both instances are considered identical.
 	// This is fast as it skips a lexical compare, however it can also return false for two instances that have identical strings, but in different pointers.
-	// For instance, this method will return false for two FText objects created from FText::FromString("Wooble") as they each have unique, non-shared instances.
-	return TextData == Other.TextData || TextData->GetLocalizedString() == Other.TextData->GetLocalizedString();
+	// For instance, this method will return false for two FText objects created from FText::FromString("Wooble") as they each have unique (or null), non-shared instances.
+	FTextDisplayStringPtr DisplayStringPtr = TextData->GetLocalizedString();
+	FTextDisplayStringPtr OtherDisplayStringPtr = Other.TextData->GetLocalizedString();
+	return DisplayStringPtr && OtherDisplayStringPtr && DisplayStringPtr == OtherDisplayStringPtr;
 }
 
 void operator<<(FStructuredArchive::FSlot Slot, FFormatArgumentValue& Value)
