Change back to using GWarn/GLog for native bridge logging, but checking the category is active.

diff --git a/Engine/Source/Runtime/Core/Private/Misc/EmbeddedCommunication.cpp b/Engine/Source/Runtime/Core/Private/Misc/EmbeddedCommunication.cpp
--- a/Engine/Source/Runtime/Core/Private/Misc/EmbeddedCommunication.cpp
+++ b/Engine/Source/Runtime/Core/Private/Misc/EmbeddedCommunication.cpp
@@ -204,39 +204,48 @@ void FEmbeddedCommunication::AllowSleep(FName Requester)
 }
 
 #if BUILD_EMBEDDED_APP
+
 DEFINE_LOG_CATEGORY_STATIC(LogBridge, Log, All);
 
-#define UE_LOG_BRIDGE(Verbosity, String) \
-{ \
-	UE_LOG(LogBridge, Verbosity, TEXT("%s"), String); \
-}
-#endif
+#endif // BUILD_EMBEDDED_APP
 
 void FEmbeddedCommunication::UELogError(const TCHAR* String)
 {
 #if BUILD_EMBEDDED_APP
-	UE_LOG_BRIDGE(Error, String);
+	if (GWarn && UE_LOG_ACTIVE(LogBridge, Error))
+	{
+		GWarn->Log("LogBridge", ELogVerbosity::Error, String);
+	}
 #endif
 }
 
 void FEmbeddedCommunication::UELogWarning(const TCHAR* String)
 {
 #if BUILD_EMBEDDED_APP
-	UE_LOG_BRIDGE(Warning, String);
+	if (GWarn && UE_LOG_ACTIVE(LogBridge, Warning))
+	{
+		GWarn->Log("LogBridge", ELogVerbosity::Warning, String);
+	}
 #endif
 }
 
 void FEmbeddedCommunication::UELogLog(const TCHAR* String)
 {
 #if BUILD_EMBEDDED_APP
-	UE_LOG_BRIDGE(Log, String);
+	if (GLog && UE_LOG_ACTIVE(LogBridge, Log))
+	{
+		GLog->Log("LogBridge", ELogVerbosity::Log, String);
+	}
 #endif
 }
 
 void FEmbeddedCommunication::UELogVerbose(const TCHAR* String)
 {
 #if BUILD_EMBEDDED_APP
-	UE_LOG_BRIDGE(Verbose, String);
+	if (GLog && UE_LOG_ACTIVE(LogBridge, Verbose))
+	{
+		GLog->Log("LogBridge", ELogVerbosity::Verbose, String);
+	}
 #endif
 }
 
