Turning on Nano_Malloc page consolidation. Allows us to consolidate different nano_malloc values into larger blocks for use in our internal pages.
#tests iOS Client

diff --git a/Engine/Source/Runtime/Core/Private/HAL/MallocBinned.cpp b/Engine/Source/Runtime/Core/Private/HAL/MallocBinned.cpp
--- a/Engine/Source/Runtime/Core/Private/HAL/MallocBinned.cpp
+++ b/Engine/Source/Runtime/Core/Private/HAL/MallocBinned.cpp
@@ -9,6 +9,8 @@
 #include "Misc/BufferedOutputDevice.h"
 
 #include "HAL/MemoryMisc.h"
+#include "HAL/LowLevelMemStats.h"
+#include "HAL/LowLevelMemTracker.h"
 
 /** Malloc binned allocator specific stats. */
 DEFINE_STAT(STAT_Binned_OsCurrent);
@@ -21,6 +23,23 @@ DEFINE_STAT(STAT_Binned_CurrentAllocs);
 DEFINE_STAT(STAT_Binned_TotalAllocs);
 DEFINE_STAT(STAT_Binned_SlackCurrent);
 
+#if USE_OS_SMALL_BLOCK_GRAB_MEMORY_FROM_OS
+DEFINE_STAT(STAT_Binned_NanoMallocPages_Current);
+DEFINE_STAT(STAT_Binned_NanoMallocPages_Peak);
+DEFINE_STAT(STAT_Binned_NanoMallocPages_WasteCurrent);
+DEFINE_STAT(STAT_Binned_NanoMallocPages_WastePeak);
+#endif
+
+#if PLATFORM_IOS
+#define PLAT_PAGE_SIZE_LIMIT 16384
+#define PLAT_BINNED_ALLOC_POOLSIZE 16384
+#define PLAT_SMALL_BLOCK_POOL_SIZE 256 //224
+#else
+#define PLAT_PAGE_SIZE_LIMIT 65536
+#define PLAT_BINNED_ALLOC_POOLSIZE 65536
+#define PLAT_SMALL_BLOCK_POOL_SIZE 0
+#endif //PLATFORM_IOS
+
 /** Information about a piece of free memory. 8 bytes */
 struct FMallocBinned::FFreeMem
 {
@@ -136,18 +155,6 @@ struct FMallocBinned::PoolHashBucket
 	}
 };
 
-#define USE_OS_SMALL_BLOCK_ALLOC PLATFORM_IOS
-
-#if PLATFORM_IOS
-#define PLAT_PAGE_SIZE_LIMIT 16384
-#define PLAT_BINNED_ALLOC_POOLSIZE 16384
-#define PLAT_SMALL_BLOCK_POOL_SIZE 224
-#else
-#define PLAT_PAGE_SIZE_LIMIT 65536
-#define PLAT_BINNED_ALLOC_POOLSIZE 65536
-#define PLAT_SMALL_BLOCK_POOL_SIZE 0
-#endif
-
 struct FMallocBinned::Private
 {
 	/** Default alignment for binned allocator */
@@ -160,10 +167,244 @@ struct FMallocBinned::Private
 	// On IOS can push small allocs in to a pre-allocated small block pool
 	enum { SMALL_BLOCK_POOL_SIZE = PLAT_SMALL_BLOCK_POOL_SIZE };
 
+	static bool bHasInitializedStatsMetadata;
+	
 #if USE_OS_SMALL_BLOCK_ALLOC
 	static uint64 SmallBlockStartPtr;
 	static uint64 SmallBlockEndPtr;
-#endif
+	enum { SMALL_BLOCK_MAX_TOTAL_POOL_SIZE = 0x20000000 };
+
+#if USE_OS_SMALL_BLOCK_GRAB_MEMORY_FROM_OS
+	enum { SMALL_BLOCK_GRAB_ALLOC_ALIGN = 16 };
+	enum { SMALL_BLOCK_GRAB_MAX_ALLOC_SIZE = 256 };
+	enum { SMALL_BLOCK_GRAB_MIN_ALLOC_SIZE = SMALL_BLOCK_GRAB_ALLOC_ALIGN };
+	
+	enum { SMALL_BLOCK_GRAB_TEMP_MEM_ARRAY_SIZE = (SMALL_BLOCK_MAX_TOTAL_POOL_SIZE / SMALL_BLOCK_GRAB_MIN_ALLOC_SIZE) };
+	
+	//Sizes that we attempt to grab sequential allocated memory in
+	enum { BLOCK_GRAB_TARGET_BIN_SIZE = PAGE_SIZE_LIMIT };
+	
+	//if set to true, we will only grab sequential blocks that are page aligned. If set, make sure our
+	//BLOCK_GRAB_TARGET_BIN_SIZE above is set to something that is also a multiple of page size
+	enum { SMALL_BLOCK_GRAB_ENSURE_PAGE_ALIGNMENT = true };
+	
+	//How much of the small block pool we should use max for our GRAB_MEMORY_FROM_OS override
+	enum { MAX_SMALL_BLOCK_USEAGE_SIZE = 250 * 1024 * 1024 }; //250MB
+	enum { MAX_NUM_BLOCK_START_ADDRESSES = MAX_SMALL_BLOCK_USEAGE_SIZE / BLOCK_GRAB_TARGET_BIN_SIZE };
+	
+	//Storage for all our free SmallBlockGrab free start addresses
+	static uint64 SmallBlockGrab_FreeStartPointers[MAX_NUM_BLOCK_START_ADDRESSES];
+	
+	//How many actual allocations we have available in our SmallBlockGrab_FreeStartPointers array
+	static int NumFreeSmallBlockGrabAllocations;
+	
+	static bool IsSmallBlockGrabAllocation(void* ptr)
+	{
+		return (   (((uint64)ptr) >= Private::SmallBlockStartPtr)
+				&& (((uint64)ptr) < Private::SmallBlockEndPtr)  );
+	}
+	
+	static bool IsNanoMallocPageAligned(void* ptr)
+	{
+		uint64 AddressOffset = ((uint64)ptr) - SmallBlockStartPtr;
+		return ((AddressOffset % PAGE_SIZE_LIMIT) == 0);
+	}
+	
+	static FCriticalSection SmallLock;
+	
+	static FFreeMem* GetAllocFromSmallBlockGrab(FMallocBinned& Allocator, UPTRINT OsBytes)
+	{
+		//Make sure we can hold this memory in a free SmallBlockGrab
+		if (OsBytes > Private::BLOCK_GRAB_TARGET_BIN_SIZE)
+		{
+			//FPlatformMisc::LowLevelOutputDebugStringf(TEXT("OSBytes: %d"), OsBytes);
+			return nullptr;
+		}
+		
+		FScopeLock ScopedLock( &SmallLock );
+		if (Private::NumFreeSmallBlockGrabAllocations == 0)
+		{
+			//FPlatformMisc::LowLevelOutputDebugStringf(TEXT("OSBytes: %d  FreeAllocations: %d"), OsBytes, Private::NumFreeSmallBlockGrabAllocations);
+			return nullptr;
+		}
+		
+		int GrabIndex = --NumFreeSmallBlockGrabAllocations;
+		check(GrabIndex > 0);
+		void* FreeMemory = (void*)SmallBlockGrab_FreeStartPointers[GrabIndex];
+		
+		//Flag that we have used this value for verification
+		SmallBlockGrab_FreeStartPointers[GrabIndex] = 0xBADF00D;
+		
+		BINNED_PEAK_STATCOUNTER(Allocator.NanoMallocPagesPeak,    BINNED_ADD_STATCOUNTER(Allocator.NanoMallocPagesCurrent, BLOCK_GRAB_TARGET_BIN_SIZE));
+		BINNED_PEAK_STATCOUNTER(Allocator.NanoMallocWastePagesPeak, BINNED_ADD_STATCOUNTER(Allocator.NanoMallocPagesWaste, BLOCK_GRAB_TARGET_BIN_SIZE - OsBytes));
+		
+		return (FFreeMem*)FreeMemory;
+	}
+	
+	static void FreeSmallBlockGrab(void* ptr, FMallocBinned& Allocator, SIZE_T Size)
+	{
+		FScopeLock ScopedLock( &SmallLock );
+
+		check(IsSmallBlockGrabAllocation(ptr));
+		check(!SMALL_BLOCK_GRAB_ENSURE_PAGE_ALIGNMENT || IsNanoMallocPageAligned(ptr));
+		
+		int NewIndex = NumFreeSmallBlockGrabAllocations;
+		++NumFreeSmallBlockGrabAllocations;
+		check(NewIndex < MAX_NUM_BLOCK_START_ADDRESSES);
+		
+		SmallBlockGrab_FreeStartPointers[NewIndex] = (uint64) ptr;
+			
+		BINNED_ADD_STATCOUNTER(Allocator.NanoMallocPagesCurrent,    -(int64)BLOCK_GRAB_TARGET_BIN_SIZE);
+		BINNED_ADD_STATCOUNTER(Allocator.NanoMallocPagesWaste, -(int64)(BLOCK_GRAB_TARGET_BIN_SIZE - Size));
+	}
+	
+	//Helper struct for storing all our allocation data in an intermediate
+	struct FMemoryAllocationGrabberHelper
+	{
+	public:
+		uint32 NumFoundCombinedBlocks;
+		
+		FMemoryAllocationGrabberHelper()
+		: NumFoundCombinedBlocks(0)
+		{
+			ResetArrays();
+		}
+		
+		void AddGrabbedMemory(uint64 AddressAddedAt, uint32 AllocatedSize)
+		{
+			assert((AddressAddedAt > 0) && (AllocatedSize > 0) && (SmallBlockStartPtr > 0));
+			{
+				uint32 IndexToAddAt = 0;
+				IndexToAddAt = (AddressAddedAt - SmallBlockStartPtr) / SMALL_BLOCK_GRAB_MIN_ALLOC_SIZE;
+				
+				assert((IndexToAddAt > 0) && (IndexToAddAt < SMALL_BLOCK_GRAB_TEMP_MEM_ARRAY_SIZE));
+				{
+					assert(GrabbedMemoryBlocks[IndexToAddAt] == 0);
+					{
+						GrabbedMemoryBlocks[IndexToAddAt] = AllocatedSize;
+					}
+				}
+			}
+		}
+		
+		void ValidateMemoryBlocks()
+		{
+			for (int Index = 0; Index < SMALL_BLOCK_GRAB_TEMP_MEM_ARRAY_SIZE;)
+			{
+				if (GrabbedMemoryBlocks[Index] > 0)
+				{
+					uint32 IndicesToParse = (GrabbedMemoryBlocks[Index] / SMALL_BLOCK_GRAB_MIN_ALLOC_SIZE);
+					
+					//Verify no data shows up before we would expect it in the array
+					//Since the value is our data size, we shouldn't have allocated anything between our Start Index and our expected End Index
+					for (int SizeVerifyIndex = (Index + 1); SizeVerifyIndex < IndicesToParse; ++SizeVerifyIndex)
+					{
+						assert(GrabbedMemoryBlocks[Index + SizeVerifyIndex] == 0);
+					}
+					
+					//Go ahead and write a temp value in there just to test accessing all these memory blocks
+					const uint64 MemoryAddress = GetAddressForIndex(Index);
+					new ((void*)MemoryAddress) int(5);
+					
+					Index += IndicesToParse;
+				}
+				else
+				{
+					++Index;
+				}
+			}
+		}
+		
+		uint64 GetAddressForIndex(uint32 Index)
+		{
+			return (Index * SMALL_BLOCK_GRAB_MIN_ALLOC_SIZE) + SmallBlockStartPtr;
+		}
+		
+		void FindGrabbedBlockSequentials(uint64 OutArrayCombinedAddressStartPtrs[MAX_NUM_BLOCK_START_ADDRESSES])
+		{
+			//If we care about page alignment, we want to make sure the target sizes are a multiple of page alignment
+			static_assert((!Private::SMALL_BLOCK_GRAB_ENSURE_PAGE_ALIGNMENT || ((Private::BLOCK_GRAB_TARGET_BIN_SIZE % PAGE_SIZE_LIMIT) == 0)),
+						  "If SMALL_BLOCK_GRAB_ENSURE_PAGE_ALIGNMENT is true, then BLOCK_GRAB_TARGET_BIN_SIZE must be a multiple of PAGE_SIZE_LIMIT.");
+			
+			static uint64 CurrentSequentialBlockFoundCount = 0;
+			static uint64 CurrentBlockStartIndex = 0;
+			static bool CurrentBlockMeetsPageAlignmentRequirements = IsGrabbedBlockIndexPageAligned(0);
+			
+			for (int Index = 0; Index < SMALL_BLOCK_GRAB_TEMP_MEM_ARRAY_SIZE;)
+			{
+				bool bDidEndSequentialChain = false;
+				if (CurrentBlockMeetsPageAlignmentRequirements && (GrabbedMemoryBlocks[Index] > 0))
+				{
+					CurrentSequentialBlockFoundCount += GrabbedMemoryBlocks[Index];
+					
+					//End this chain if we already have a large enough bucket
+					bDidEndSequentialChain = (CurrentSequentialBlockFoundCount >= BLOCK_GRAB_TARGET_BIN_SIZE);
+					
+					uint32 IndicesToAdvance = (GrabbedMemoryBlocks[Index] / SMALL_BLOCK_GRAB_MIN_ALLOC_SIZE);
+					Index += IndicesToAdvance;
+				}
+				else
+				{
+					bDidEndSequentialChain = true;
+					++Index;
+				}
+				
+				if (bDidEndSequentialChain)
+				{
+					//If we are large, save as a large sequential block as long as we haven't saved too many
+					if ((CurrentSequentialBlockFoundCount >= BLOCK_GRAB_TARGET_BIN_SIZE) && (NumFoundCombinedBlocks < MAX_NUM_BLOCK_START_ADDRESSES))
+					{
+						OutArrayCombinedAddressStartPtrs[NumFoundCombinedBlocks] = GetAddressForIndex(CurrentBlockStartIndex);
+						++NumFoundCombinedBlocks;
+					}
+					
+					//if we didn't make it to either, go ahead and de-allocate memory in this chain
+					else
+					{
+						for (int FreeIndex = CurrentBlockStartIndex; FreeIndex < Index;++FreeIndex)
+						{
+							//Only free memory that we actually allocated
+							if (GrabbedMemoryBlocks[FreeIndex] > 0)
+							{
+								uint64 PointerToFree = GetAddressForIndex(FreeIndex);
+								::free((void*)PointerToFree);
+							}
+						}
+					}
+					
+					CurrentSequentialBlockFoundCount = 0;
+					CurrentBlockStartIndex = Index;
+					CurrentBlockMeetsPageAlignmentRequirements = IsGrabbedBlockIndexPageAligned(CurrentBlockStartIndex);
+				}
+			}
+		}
+		
+		//Look if we are at the start of a page with this memory block index
+		bool IsGrabbedBlockIndexPageAligned(int Index)
+		{
+			if (!Private::SMALL_BLOCK_GRAB_ENSURE_PAGE_ALIGNMENT)
+			{
+				return true;
+			}
+			
+			const uint64 MemLoc = GetAddressForIndex(Index);
+			const uint64 MemoryOffsetFromStart = MemLoc - SmallBlockStartPtr;
+			
+			return ((MemoryOffsetFromStart % PAGE_SIZE_LIMIT) == 0);
+		}
+		
+	private:
+		//Tracker holding the state of all possible memory allocations.
+		//0 if un-allocated. Otherwise holds the data size allocated
+		uint32 GrabbedMemoryBlocks[SMALL_BLOCK_GRAB_TEMP_MEM_ARRAY_SIZE];
+		
+		void ResetArrays()
+		{
+			memset(&GrabbedMemoryBlocks, 0, sizeof(GrabbedMemoryBlocks));
+		}
+	};
+#endif //USE_OS_SMALL_BLOCK_GRAB_MEMORY_FROM_OS
+#endif //USE_OS_SMALL_BLOCK_ALLOC
 	
 	// Implementation. 
 	static CA_NO_RETURN void OutOfMemory(uint64 Size, uint32 Alignment=0)
@@ -368,12 +609,20 @@ struct FMallocBinned::Private
 
 		checkSlow(Blocks >= 1);
 		checkSlow(Blocks * Table->BlockSize <= Bytes && PoolSize >= Bytes);
-
-		// Allocate memory.
+		
 		FFreeMem* Free = nullptr;
 		SIZE_T ActualPoolSize; //TODO: use this to reduce waste?
-		Free = (FFreeMem*)OSAlloc(Allocator, OsBytes, ActualPoolSize);
-
+		
+#if USE_OS_SMALL_BLOCK_GRAB_MEMORY_FROM_OS
+		Free = GetAllocFromSmallBlockGrab(Allocator, OsBytes);
+#endif //USE_OS_SMALL_BLOCK_GRAB_MEMORY_FROM_OS
+		
+		// Allocate memory if we haven't yet
+		if (Free == nullptr)
+		{
+			Free = (FFreeMem*)OSAlloc(Allocator, OsBytes, ActualPoolSize);
+		}
+		
 		checkSlow(!((UPTRINT)Free & (PageSize - 1)));
 		if( !Free )
 		{
@@ -455,9 +704,9 @@ struct FMallocBinned::Private
 #if PLATFORM_IOS || PLATFORM_MAC
         if (Pool == nullptr)
         {
-#if USE_OS_SMALL_BLOCK_ALLOC
+#if USE_OS_SMALL_BLOCK_ALLOC && !USE_OS_SMALL_BLOCK_GRAB_MEMORY_FROM_OS
 			// check the small block allocation range
-			if ((uint64)Ptr >= SmallBlockStartPtr && (uint64)Ptr <= SmallBlockEndPtr)
+			if ((uint64)Ptr >= SmallBlockStartPtr && (uint64)Ptr < SmallBlockEndPtr)
 			{
 				SmallOSFree(Allocator, Ptr, Private::SMALL_BLOCK_POOL_SIZE);
 				return;
@@ -581,7 +830,16 @@ struct FMallocBinned::Private
 #endif
 		if (Size > MAX_CACHED_OS_FREES_BYTE_LIMIT / 4)
 		{
-			FPlatformMemory::BinnedFreeToOS(Ptr, Size);
+			#if USE_OS_SMALL_BLOCK_GRAB_MEMORY_FROM_OS
+				if (IsSmallBlockGrabAllocation(ptr))
+				{
+					FreeSmallBlockGrab(ptr, Allocator, Size);
+				}
+				else
+			#endif //USE_OS_SMALL_BLOCK_GRAB_MEMORY_FROM_OS
+			{
+				FPlatformMemory::BinnedFreeToOS(Ptr, Size);
+			}
 			return;
 		}
 
@@ -596,14 +854,33 @@ struct FMallocBinned::Private
 			{
 				FMemory::Memmove(&Allocator.FreedPageBlocks[0], &Allocator.FreedPageBlocks[1], sizeof(FFreePageBlock) * Allocator.FreedPageBlocksNum);
 			}
-			FPlatformMemory::BinnedFreeToOS(FreePtr, FreeSize);
+			
+			#if USE_OS_SMALL_BLOCK_GRAB_MEMORY_FROM_OS
+				if (IsSmallBlockGrabAllocation(FreePtr))
+				{
+					FreeSmallBlockGrab(FreePtr, Allocator, Size);
+				}
+				else
+			#endif //USE_OS_SMALL_BLOCK_GRAB_MEMORY_FROM_OS
+			{
+				FPlatformMemory::BinnedFreeToOS(FreePtr, FreeSize);
+			}
 		}
 		Allocator.FreedPageBlocks[Allocator.FreedPageBlocksNum].Ptr      = Ptr;
 		Allocator.FreedPageBlocks[Allocator.FreedPageBlocksNum].ByteSize = Size;
 		Allocator.CachedTotal += Size;
 		++Allocator.FreedPageBlocksNum;
 #else
-		FPlatformMemory::BinnedFreeToOS(Ptr, Size);
+		#if USE_OS_SMALL_BLOCK_GRAB_MEMORY_FROM_OS
+			if (Private::IsSmallBlockGrabAllocation(Ptr))
+			{
+				FreeSmallBlockGrab(Ptr, Allocator, Size);
+			}
+			else
+		#endif //USE_OS_SMALL_BLOCK_GRAB_MEMORY_FROM_OS
+		{
+			FPlatformMemory::BinnedFreeToOS(Ptr, Size);
+		}
 #endif
 	}
 
@@ -738,10 +1015,23 @@ uint64 FMallocBinned::Private::SmallBlockStartPtr = 0;
 uint64 FMallocBinned::Private::SmallBlockEndPtr = 0;
 #endif
 
+#if USE_OS_SMALL_BLOCK_GRAB_MEMORY_FROM_OS
+uint64 FMallocBinned::Private::SmallBlockGrab_FreeStartPointers[MAX_NUM_BLOCK_START_ADDRESSES];
+int FMallocBinned::Private::NumFreeSmallBlockGrabAllocations = 0;
+FCriticalSection FMallocBinned::Private::SmallLock;
+#endif
+
+bool FMallocBinned::Private::bHasInitializedStatsMetadata = false;
+
 void FMallocBinned::GetAllocatorStats( FGenericMemoryStats& out_Stats )
 {
 	FMalloc::GetAllocatorStats( out_Stats );
 
+	if (!Private::bHasInitializedStatsMetadata)
+	{
+		InitializeStatsMetadata();
+	}
+	
 #if	STATS
 	SIZE_T	LocalOsCurrent = 0;
 	SIZE_T	LocalOsPeak = 0;
@@ -752,7 +1042,14 @@ void FMallocBinned::GetAllocatorStats( FGenericMemoryStats& out_Stats )
 	SIZE_T	LocalCurrentAllocs = 0;
 	SIZE_T	LocalTotalAllocs = 0;
 	SIZE_T	LocalSlackCurrent = 0;
-
+	
+#if USE_OS_SMALL_BLOCK_GRAB_MEMORY_FROM_OS
+	SIZE_T LocalNanoMallocPagesCurrent;
+	SIZE_T LocalNanoMallocPagesPeak;
+	SIZE_T LocalNanoMallocPagesWaste;
+	SIZE_T LocalNanoMallocWastePagesPeak;
+#endif //USE_OS_SMALL_BLOCK_GRAB_MEMORY_FROM_OS
+	
 	{
 #ifdef USE_INTERNAL_LOCKS
 		FScopeLock ScopedLock( &AccessGuard );
@@ -770,6 +1067,13 @@ void FMallocBinned::GetAllocatorStats( FGenericMemoryStats& out_Stats )
 		LocalCurrentAllocs = CurrentAllocs;
 		LocalTotalAllocs = TotalAllocs;
 		LocalSlackCurrent = SlackCurrent;
+
+#if USE_OS_SMALL_BLOCK_GRAB_MEMORY_FROM_OS
+		LocalNanoMallocPagesCurrent = NanoMallocPagesCurrent;
+		LocalNanoMallocPagesPeak = NanoMallocPagesPeak;
+		LocalNanoMallocPagesWaste = NanoMallocPagesWaste;
+		LocalNanoMallocWastePagesPeak = NanoMallocWastePagesPeak;
+#endif //USE_OS_SMALL_BLOCK_GRAB_MEMORY_FROM_OS
 	}
 
 	// Malloc binned stats.
@@ -782,11 +1086,21 @@ void FMallocBinned::GetAllocatorStats( FGenericMemoryStats& out_Stats )
 	out_Stats.Add( GET_STATDESCRIPTION( STAT_Binned_CurrentAllocs ), LocalCurrentAllocs );
 	out_Stats.Add( GET_STATDESCRIPTION( STAT_Binned_TotalAllocs ), LocalTotalAllocs );
 	out_Stats.Add( GET_STATDESCRIPTION( STAT_Binned_SlackCurrent ), LocalSlackCurrent );
+	
+#if USE_OS_SMALL_BLOCK_GRAB_MEMORY_FROM_OS
+	out_Stats.Add( GET_STATDESCRIPTION( STAT_Binned_NanoMallocPages_Current ), LocalNanoMallocPagesCurrent );
+	out_Stats.Add( GET_STATDESCRIPTION( STAT_Binned_NanoMallocPages_Peak ), LocalNanoMallocPagesPeak );
+	out_Stats.Add( GET_STATDESCRIPTION( STAT_Binned_NanoMallocPages_WasteCurrent ), LocalNanoMallocPagesWaste );
+	out_Stats.Add( GET_STATDESCRIPTION( STAT_Binned_NanoMallocPages_WastePeak ), LocalNanoMallocWastePagesPeak );
+#endif //USE_OS_SMALL_BLOCK_GRAB_MEMORY_FROM_OS
+	
 #endif // STATS
 }
 
 void FMallocBinned::InitializeStatsMetadata()
 {
+	Private::bHasInitializedStatsMetadata = true;
+	
 	FMalloc::InitializeStatsMetadata();
 
 	// Initialize stats metadata here instead of UpdateStats.
@@ -800,6 +1114,13 @@ void FMallocBinned::InitializeStatsMetadata()
 	GET_STATFNAME(STAT_Binned_CurrentAllocs);
 	GET_STATFNAME(STAT_Binned_TotalAllocs);
 	GET_STATFNAME(STAT_Binned_SlackCurrent);
+	
+#if USE_OS_SMALL_BLOCK_GRAB_MEMORY_FROM_OS && ENABLE_LOW_LEVEL_MEM_TRACKER
+	FLowLevelMemTracker::Get().RegisterPlatformTag((int32)ELLMTagNanoMallocGrabber::NanoMallocPagesCurrent, TEXT("Nano Malloc Pages Current"), GET_STATFNAME(STAT_Binned_NanoMallocPages_Current), GET_STATFNAME(STAT_EngineSummaryLLM));
+	FLowLevelMemTracker::Get().RegisterPlatformTag((int32)ELLMTagNanoMallocGrabber::NanoMallocPagesPeak, TEXT("Nano Malloc Pages Peak"), GET_STATFNAME(STAT_Binned_NanoMallocPages_Peak), GET_STATFNAME(STAT_EngineSummaryLLM));
+	FLowLevelMemTracker::Get().RegisterPlatformTag((int32)ELLMTagNanoMallocGrabber::NanoMallocPagesWasteCurrent, TEXT("Nano Malloc Pages Waste Current"), GET_STATFNAME(STAT_Binned_NanoMallocPages_WasteCurrent), GET_STATFNAME(STAT_EngineSummaryLLM));
+	FLowLevelMemTracker::Get().RegisterPlatformTag((int32)ELLMTagNanoMallocGrabber::NanoMallocPagesWastePeak, TEXT("Nano Malloc Pages Waste Peak"), GET_STATFNAME(STAT_Binned_NanoMallocPages_WastePeak), GET_STATFNAME(STAT_EngineSummaryLLM));
+#endif //USE_OS_SMALL_BLOCK_GRAB_MEMORY_FROM_OS
 }
 
 FMallocBinned::FMallocBinned(uint32 InPageSize, uint64 AddressLimit)
@@ -827,6 +1148,12 @@ FMallocBinned::FMallocBinned(uint32 InPageSize, uint64 AddressLimit)
 	,	TotalAllocs		( 0 )
 	,	SlackCurrent	( 0 )
 	,	MemTime			( 0.0 )
+#if USE_OS_SMALL_BLOCK_GRAB_MEMORY_FROM_OS
+	,	NanoMallocPagesCurrent 		( 0 )
+	,	NanoMallocPagesPeak			( 0 )
+	,	NanoMallocPagesWaste		( 0 )
+	,	NanoMallocWastePagesPeak	( 0 )
+#endif //USE_OS_SMALL_BLOCK_GRAB_MEMORY_FROM_OS
 #endif
 {
 	check(!(PageSize & (PageSize - 1)));
@@ -842,10 +1169,45 @@ FMallocBinned::FMallocBinned(uint32 InPageSize, uint64 AddressLimit)
 	Private::SmallBlockStartPtr = (uint64_t)ptr & ~(0x1fffffff);
 	
 	//Add pool size to start pointer. Magic number comes from found instruments pool size of 512MB
-	Private::SmallBlockEndPtr = Private::SmallBlockStartPtr + 0x20000000;
+	Private::SmallBlockEndPtr = Private::SmallBlockStartPtr + Private::SMALL_BLOCK_MAX_TOTAL_POOL_SIZE;
 	
 	::free(ptr);
-#endif
+
+#if USE_OS_SMALL_BLOCK_GRAB_MEMORY_FROM_OS
+
+	void* GrabberHelperLoc = ::malloc(sizeof(Private::FMemoryAllocationGrabberHelper));
+	Private::FMemoryAllocationGrabberHelper* GrabberHelper = new (GrabberHelperLoc) Private::FMemoryAllocationGrabberHelper();
+	
+	//Grab as much as we can in smaller and smaller amounts until we can get no more
+	int CurrentAllocationSize = Private::SMALL_BLOCK_GRAB_MAX_ALLOC_SIZE;
+	int CurrentGrabIndex = 0;
+	while ((CurrentAllocationSize >= Private::SMALL_BLOCK_GRAB_MIN_ALLOC_SIZE))
+	{
+		void* NewData = ::malloc(CurrentAllocationSize);
+		if (((uint64)NewData) >= Private::SmallBlockStartPtr && ((uint64)NewData) < Private::SmallBlockEndPtr)
+		{
+			GrabberHelper->AddGrabbedMemory((uint64)NewData, CurrentAllocationSize);
+		}
+		else
+		{
+			//Free memory we just allocated as it was outside of the nano_malloc range
+			::free(NewData);
+			
+			//Lower how much data we are trying to allocate now that 1 failed on this size
+			CurrentAllocationSize -= Private::SMALL_BLOCK_GRAB_ALLOC_ALIGN;
+		}
+	}
+	
+	GrabberHelper->ValidateMemoryBlocks();
+	
+	GrabberHelper->FindGrabbedBlockSequentials(Private::SmallBlockGrab_FreeStartPointers);
+	Private::NumFreeSmallBlockGrabAllocations = GrabberHelper->NumFoundCombinedBlocks;
+	
+	//Free our GrabberHelper now that we are fully allocated and good to go
+	::free(GrabberHelperLoc);
+	
+#endif //USE_OS_SMALL_BLOCK_GRAB_MEMORY_FROM_OS
+#endif //USE_OS_SMALL_BLOCK_ALLOC
 	
 	/** Shift to get the reference from the indirect tables */
 	PoolBitShift = FPlatformMath::CeilLogTwo(PageSize);
@@ -958,7 +1320,7 @@ void* FMallocBinned::Malloc(SIZE_T Size, uint32 Alignment)
 	
 	FFreeMem* Free = nullptr;
 	bool bUsePools = true;
-#if USE_OS_SMALL_BLOCK_ALLOC
+#if USE_OS_SMALL_BLOCK_ALLOC && !USE_OS_SMALL_BLOCK_GRAB_MEMORY_FROM_OS
 	if (Size <= Private::SMALL_BLOCK_POOL_SIZE)
 	{
 		bUsePools = false;
@@ -1084,9 +1446,13 @@ void* FMallocBinned::Realloc( void* Ptr, SIZE_T NewSize, uint32 Alignment )
 	void* NewPtr = Ptr;
 	if( Ptr && NewSize )
 	{
-#if USE_OS_SMALL_BLOCK_ALLOC
+		FPoolInfo* Pool = Private::FindPoolInfo(*this, (UPTRINT)Ptr, BasePtr);
+		
+#if USE_OS_SMALL_BLOCK_ALLOC && !USE_OS_SMALL_BLOCK_GRAB_MEMORY_FROM_OS
 		// special case for really small blocks
-		if ((uint64)Ptr >= Private::SmallBlockStartPtr && (uint64)Ptr < Private::SmallBlockEndPtr)
+		if ((Pool == nullptr)
+			&& ((uint64)Ptr >= Private::SmallBlockStartPtr)
+			&& ((uint64)Ptr < Private::SmallBlockEndPtr))
 		{
 			NewPtr = ::realloc(Ptr, NewSize);
 
@@ -1102,7 +1468,7 @@ void* FMallocBinned::Realloc( void* Ptr, SIZE_T NewSize, uint32 Alignment )
 		else
 #endif
 		{
-			FPoolInfo* Pool = Private::FindPoolInfo(*this, (UPTRINT)Ptr, BasePtr);
+			
 
 			if( Pool->TableIndex < BinnedOSTableIndex )
 			{
@@ -1184,7 +1550,7 @@ bool FMallocBinned::GetAllocationSize(void *Original, SIZE_T &SizeOut)
 #if PLATFORM_IOS || PLATFORM_MAC
 	if (Pool == nullptr)
 	{
-#if USE_OS_SMALL_BLOCK_ALLOC
+#if USE_OS_SMALL_BLOCK_ALLOC && !USE_OS_SMALL_BLOCK_GRAB_MEMORY_FROM_OS
 		if ((uint64)Original >= Private::SmallBlockStartPtr && (uint64)Original < Private::SmallBlockEndPtr)
 		{
 			SizeOut = Private::SMALL_BLOCK_POOL_SIZE;
@@ -1230,7 +1596,7 @@ SIZE_T FMallocBinned::QuantizeSize(SIZE_T Size, uint32 Alignment)
 	Size = FMath::Max<SIZE_T>(PoolTable[0].BlockSize, Size + (Alignment - SpareBytesCount));
 
 	SIZE_T Result;
-#if USE_OS_SMALL_BLOCK_ALLOC
+#if USE_OS_SMALL_BLOCK_ALLOC && !USE_OS_SMALL_BLOCK_GRAB_MEMORY_FROM_OS
 	if (Size <= Private::SMALL_BLOCK_POOL_SIZE)
 	{
 		UPTRINT AlignedSize = Align(Size, Alignment);
@@ -1309,6 +1675,13 @@ void FMallocBinned::UpdateStats()
 	SIZE_T	LocalTotalAllocs = 0;
 	SIZE_T	LocalSlackCurrent = 0;
 
+#if USE_OS_SMALL_BLOCK_GRAB_MEMORY_FROM_OS
+	SIZE_T LocalNanoMallocPagesCurrent = 0;
+	SIZE_T LocalNanoMallocPagesPeak = 0;
+	SIZE_T LocalNanoMallocPagesWaste = 0;
+	SIZE_T LocalNanoMallocWastePagesPeak = 0;
+#endif //USE_OS_SMALL_BLOCK_GRAB_MEMORY_FROM_OS
+	
 	{
 #ifdef USE_INTERNAL_LOCKS
 		FScopeLock ScopedLock( &AccessGuard );
@@ -1325,6 +1698,13 @@ void FMallocBinned::UpdateStats()
 		LocalCurrentAllocs = CurrentAllocs;
 		LocalTotalAllocs = TotalAllocs;
 		LocalSlackCurrent = SlackCurrent;
+		
+#if USE_OS_SMALL_BLOCK_GRAB_MEMORY_FROM_OS
+		LocalNanoMallocPagesCurrent = NanoMallocPagesCurrent;
+		LocalNanoMallocPagesPeak = NanoMallocPagesPeak;
+		LocalNanoMallocPagesWaste = NanoMallocPagesWaste;
+		LocalNanoMallocWastePagesPeak = NanoMallocWastePagesPeak;
+#endif //USE_OS_SMALL_BLOCK_GRAB_MEMORY_FROM_OS
 	}
 
 	SET_MEMORY_STAT( STAT_Binned_OsCurrent, LocalOsCurrent );
@@ -1336,6 +1716,13 @@ void FMallocBinned::UpdateStats()
 	SET_DWORD_STAT( STAT_Binned_CurrentAllocs, LocalCurrentAllocs );
 	SET_DWORD_STAT( STAT_Binned_TotalAllocs, LocalTotalAllocs );
 	SET_MEMORY_STAT( STAT_Binned_SlackCurrent, LocalSlackCurrent );
+	
+#if USE_OS_SMALL_BLOCK_GRAB_MEMORY_FROM_OS
+	SET_MEMORY_STAT( STAT_Binned_NanoMallocPages_Current, LocalNanoMallocPagesCurrent );
+	SET_MEMORY_STAT( STAT_Binned_NanoMallocPages_Peak, LocalNanoMallocPagesPeak );
+	SET_MEMORY_STAT( STAT_Binned_NanoMallocPages_WasteCurrent, LocalNanoMallocPagesWaste );
+	SET_MEMORY_STAT( STAT_Binned_NanoMallocPages_WastePeak, LocalNanoMallocWastePagesPeak );
+#endif //USE_OS_SMALL_BLOCK_GRAB_MEMORY_FROM_OS
 #endif
 }
 
@@ -1365,6 +1752,12 @@ void FMallocBinned::DumpAllocatorStats( class FOutputDevice& Ar )
 		MEM_TIME( BufferedOutput.CategorizedLogf( LogMemory.GetCategoryName(), ELogVerbosity::Log, TEXT( "Seconds     % 5.3f" ), MemTime ) );
 		MEM_TIME( BufferedOutput.CategorizedLogf( LogMemory.GetCategoryName(), ELogVerbosity::Log, TEXT( "MSec/Allc   % 5.5f" ), 1000.0 * MemTime / MemAllocs ) );
 
+#if USE_OS_SMALL_BLOCK_GRAB_MEMORY_FROM_OS
+		BufferedOutput.CategorizedLogf( LogMemory.GetCategoryName(), ELogVerbosity::Log, TEXT( "NanoMallocPage Stats:" ) );
+		BufferedOutput.CategorizedLogf( LogMemory.GetCategoryName(), ELogVerbosity::Log, TEXT( "Current %.2f MB, peak %.2f MB" ), NanoMallocPagesCurrent / (1024.0f * 1024.0f),  NanoMallocPagesPeak / (1024.0f * 1024.0f) );
+		BufferedOutput.CategorizedLogf( LogMemory.GetCategoryName(), ELogVerbosity::Log, TEXT( "Current Waste %.2f MB, peak %.2f MB" ), NanoMallocPagesWaste / (1024.0f * 1024.0f),  NanoMallocWastePagesPeak / (1024.0f * 1024.0f) );
+#endif //USE_OS_SMALL_BLOCK_GRAB_MEMORY_FROM_OS
+		
 		// This is the memory tracked inside individual allocation pools
 		BufferedOutput.CategorizedLogf( LogMemory.GetCategoryName(), ELogVerbosity::Log, TEXT( "" ) );
 		BufferedOutput.CategorizedLogf( LogMemory.GetCategoryName(), ELogVerbosity::Log, TEXT( "Block Size Num Pools Max Pools Cur Allocs Total Allocs Min Req Max Req Mem Used Mem Slack Mem Waste Efficiency" ) );
diff --git a/Engine/Source/Runtime/Core/Public/HAL/MallocBinned.h b/Engine/Source/Runtime/Core/Public/HAL/MallocBinned.h
--- a/Engine/Source/Runtime/Core/Public/HAL/MallocBinned.h
+++ b/Engine/Source/Runtime/Core/Public/HAL/MallocBinned.h
@@ -9,6 +9,7 @@
 #include "Stats/Stats.h"
 #include "HAL/MallocJemalloc.h"
 #include "HAL/LowLevelMemTracker.h"
+#include "HAL/LowLevelMemStats.h"
 
 #define MEM_TIME(st)
 
@@ -18,6 +19,14 @@
 //#define CACHE_FREED_OS_ALLOCS
 #endif
 
+#ifndef USE_OS_SMALL_BLOCK_ALLOC
+#define USE_OS_SMALL_BLOCK_ALLOC PLATFORM_IOS
+#endif //USE_OS_SMALL_BLOCK_ALLOC
+
+#ifndef USE_OS_SMALL_BLOCK_GRAB_MEMORY_FROM_OS
+#define USE_OS_SMALL_BLOCK_GRAB_MEMORY_FROM_OS (USE_OS_SMALL_BLOCK_ALLOC && PLATFORM_IOS)
+#endif //USE_OS_SMALL_BLOCK_GRAB_MEMORY_FROM_OS
+
 #ifdef USE_INTERNAL_LOCKS
 //#   define USE_COARSE_GRAIN_LOCKS
 #endif
@@ -90,6 +99,24 @@ DECLARE_DWORD_COUNTER_STAT_EXTERN(TEXT("Binned Current Allocs"),	STAT_Binned_Cur
 DECLARE_DWORD_COUNTER_STAT_EXTERN(TEXT("Binned Total Allocs"),		STAT_Binned_TotalAllocs,STATGROUP_MemoryAllocator, CORE_API);
 DECLARE_MEMORY_STAT_EXTERN(TEXT("Binned Slack Current"),	STAT_Binned_SlackCurrent,STATGROUP_MemoryAllocator, CORE_API);
 
+#if USE_OS_SMALL_BLOCK_GRAB_MEMORY_FROM_OS && ENABLE_LOW_LEVEL_MEM_TRACKER
+DECLARE_LLM_MEMORY_STAT_EXTERN(TEXT("Nano Malloc Pages Current"), STAT_Binned_NanoMallocPages_Current,STATGROUP_LLMPlatform, CORE_API);
+DECLARE_LLM_MEMORY_STAT_EXTERN(TEXT("Nano Malloc Pages Peak"), STAT_Binned_NanoMallocPages_Peak,STATGROUP_LLMPlatform, CORE_API);
+DECLARE_LLM_MEMORY_STAT_EXTERN(TEXT("Nano Malloc Pages Waste Current"), STAT_Binned_NanoMallocPages_WasteCurrent,STATGROUP_LLMPlatform, CORE_API);
+DECLARE_LLM_MEMORY_STAT_EXTERN(TEXT("Nano Malloc Pages Waste Peak"),STAT_Binned_NanoMallocPages_WastePeak,STATGROUP_LLMPlatform, CORE_API);
+#endif //USE_OS_SMALL_BLOCK_GRAB_MEMORY_FROM_OS
+
+#if USE_OS_SMALL_BLOCK_GRAB_MEMORY_FROM_OS && ENABLE_LOW_LEVEL_MEM_TRACKER
+enum class ELLMTagNanoMallocGrabber : LLM_TAG_TYPE
+{
+	NanoMallocPagesCurrent = (LLM_TAG_TYPE)ELLMTag::PlatformTagStart, // Use Instruments for detailed breakdown!
+	NanoMallocPagesPeak,
+	NanoMallocPagesWasteCurrent,
+	NanoMallocPagesWastePeak,
+	Count
+};
+#endif //#if USE_OS_SMALL_BLOCK_GRAB_MEMORY_FROM_OS
+
 
 //
 // Optimized virtual memory allocator.
@@ -106,7 +133,7 @@ private:
 	/** Maximum allocation for the pooled allocator */
 	enum { EXTENDED_PAGE_POOL_ALLOCATION_COUNT = 2 };
 	enum { MAX_POOLED_ALLOCATION_SIZE   = 32768+1 };
-
+	
 	// Forward declares.
 	struct FFreeMem;
 	struct FPoolTable;
@@ -198,7 +225,7 @@ private:
 	FCriticalSection	AccessGuard;
 
 	// PageSize dependent constants
-	uint64 MaxHashBuckets; 
+	uint64 MaxHashBuckets;
 	uint64 MaxHashBucketBits;
 	uint64 MaxHashBucketWaste;
 	uint64 MaxBookKeepingOverhead;
@@ -218,12 +245,12 @@ private:
 	FPoolTable	OsTable;
 	FPoolTable	PagePoolTable[EXTENDED_PAGE_POOL_ALLOCATION_COUNT];
 	FPoolTable* MemSizeToPoolTable[MAX_POOLED_ALLOCATION_SIZE+EXTENDED_PAGE_POOL_ALLOCATION_COUNT];
-
+	
 	PoolHashBucket* HashBuckets;
 	PoolHashBucket* HashBucketFreeList;
-
+	
 	uint32		PageSize;
-
+	
 #ifdef CACHE_FREED_OS_ALLOCS
 	FFreePageBlock	FreedPageBlocks[MAX_CACHED_OS_FREES];
 	uint32			FreedPageBlocksNum;
@@ -242,7 +269,15 @@ private:
 	/** OsCurrent - WasteCurrent - UsedCurrent. */
 	BINNED_STAT		SlackCurrent;
 	double		MemTime;
-#endif
+	
+#if USE_OS_SMALL_BLOCK_GRAB_MEMORY_FROM_OS
+	BINNED_STAT		NanoMallocPagesCurrent;
+	BINNED_STAT		NanoMallocPagesPeak;
+	BINNED_STAT		NanoMallocPagesWaste;
+	BINNED_STAT		NanoMallocWastePagesPeak;
+#endif //USE_OS_SMALL_BLOCK_GRAB_MEMORY_FROM_OS
+	
+#endif //STATS
 
 public:
 	// FMalloc interface.
