Do not do case-insensitive search in cooked games

diff --git a/Engine/Source/Runtime/Core/Private/Unix/UnixPlatformFile.cpp b/Engine/Source/Runtime/Core/Private/Unix/UnixPlatformFile.cpp
--- a/Engine/Source/Runtime/Core/Private/Unix/UnixPlatformFile.cpp
+++ b/Engine/Source/Runtime/Core/Private/Unix/UnixPlatformFile.cpp
@@ -14,6 +14,8 @@
 
 DEFINE_LOG_CATEGORY_STATIC(LogUnixPlatformFile, Log, All);
 
+#define UNIX_PLATFORM_FILE_SPEEDUP_FILE_OPERATIONS	((!WITH_EDITOR && !IS_PROGRAM) || !PLATFORM_LINUX) 
+
 // make an FTimeSpan object that represents the "epoch" for time_t (from a stat struct)
 const FDateTime UnixEpoch(1970, 1, 1);
 
@@ -686,11 +688,16 @@ FString FUnixPlatformFile::NormalizeDirectory(const TCHAR* Directory, bool bIsFo
 bool FUnixPlatformFile::FileExists(const TCHAR* Filename)
 {
 	FString CaseSensitiveFilename;
-	if (!GCaseInsensMapper.MapCaseInsensitiveFile(NormalizeFilename(Filename, false), CaseSensitiveFilename))
+	FString NormalizedFilename = NormalizeFilename(Filename, false);
+#if !UNIX_PLATFORM_FILE_SPEEDUP_FILE_OPERATIONS
+	if (!GCaseInsensMapper.MapCaseInsensitiveFile(NormalizedFilename, CaseSensitiveFilename))
 	{
 		// could not find the file
 		return false;
 	}
+#else
+	CaseSensitiveFilename = NormalizedFilename;
+#endif
 
 	struct stat FileInfo;
 	if(stat(TCHAR_TO_UTF8(*CaseSensitiveFilename), &FileInfo) != -1)
@@ -703,11 +710,16 @@ bool FUnixPlatformFile::FileExists(const TCHAR* Filename)
 int64 FUnixPlatformFile::FileSize(const TCHAR* Filename)
 {
 	FString CaseSensitiveFilename;
-	if (!GCaseInsensMapper.MapCaseInsensitiveFile(NormalizeFilename(Filename, false), CaseSensitiveFilename))
+	FString NormalizedFilename = NormalizeFilename(Filename, false);
+#if !UNIX_PLATFORM_FILE_SPEEDUP_FILE_OPERATIONS
+	if (!GCaseInsensMapper.MapCaseInsensitiveFile(NormalizedFilename, CaseSensitiveFilename))
 	{
 		// could not find the file
 		return -1;
 	}
+#else
+	CaseSensitiveFilename = NormalizedFilename;
+#endif
 
 	struct stat FileInfo;
 	FileInfo.st_size = -1;
@@ -972,11 +984,16 @@ IFileHandle* FUnixPlatformFile::OpenWrite(const TCHAR* Filename, bool bAppend, b
 bool FUnixPlatformFile::DirectoryExists(const TCHAR* Directory)
 {
 	FString CaseSensitiveFilename;
-	if (!GCaseInsensMapper.MapCaseInsensitiveFile(NormalizeFilename(Directory, false), CaseSensitiveFilename))
+	FString NormalizedFilename = NormalizeFilename(Directory, false);
+#if !UNIX_PLATFORM_FILE_SPEEDUP_FILE_OPERATIONS
+	if (!GCaseInsensMapper.MapCaseInsensitiveFile(NormalizedFilename, CaseSensitiveFilename))
 	{
 		// could not find the file
 		return false;
 	}
+#else
+	CaseSensitiveFilename = NormalizedFilename;
+#endif
 
 	struct stat FileInfo;
 	if (stat(TCHAR_TO_UTF8(*CaseSensitiveFilename), &FileInfo) != -1)
@@ -995,11 +1012,15 @@ bool FUnixPlatformFile::DeleteDirectory(const TCHAR* Directory)
 {
 	FString CaseSensitiveFilename;
 	FString IntendedFilename(NormalizeFilename(Directory,true));
+#if !UNIX_PLATFORM_FILE_SPEEDUP_FILE_OPERATIONS
 	if (!GCaseInsensMapper.MapCaseInsensitiveFile(IntendedFilename, CaseSensitiveFilename))
 	{
 		// could not find the directory
 		return false;
 	}
+#else
+	CaseSensitiveFilename = IntendedFilename;
+#endif
 
 	GetFileMapCache().Invalidate(IntendedFilename);
 
