Fixing Error in bad Check in NanoMallocGrabber code.
#test iOS Client

diff --git a/Engine/Source/Runtime/Core/Private/HAL/MallocBinned.cpp b/Engine/Source/Runtime/Core/Private/HAL/MallocBinned.cpp
--- a/Engine/Source/Runtime/Core/Private/HAL/MallocBinned.cpp
+++ b/Engine/Source/Runtime/Core/Private/HAL/MallocBinned.cpp
@@ -229,7 +229,7 @@ struct FMallocBinned::Private
 		}
 		
 		int GrabIndex = --NumFreeSmallBlockGrabAllocations;
-		check(GrabIndex > 0);
+		check(GrabIndex >= 0);
 		void* FreeMemory = (void*)SmallBlockGrab_FreeStartPointers[GrabIndex];
 		
 		//Flag that we have used this value for verification
