[Switch] Updated Disk I/O tracking and analytics.

diff --git a/Engine/Source/Runtime/Core/Private/HAL/DiskUtilizationTracker.cpp b/Engine/Source/Runtime/Core/Private/HAL/DiskUtilizationTracker.cpp
--- a/Engine/Source/Runtime/Core/Private/HAL/DiskUtilizationTracker.cpp
+++ b/Engine/Source/Runtime/Core/Private/HAL/DiskUtilizationTracker.cpp
@@ -3,6 +3,28 @@
 #include "HAL/DiskUtilizationTracker.h"
 
 #if TRACK_DISK_UTILIZATION
+#include "ProfilingDebugging/CsvProfiler.h"
+
+CSV_DEFINE_CATEGORY(DiskIO, true); 
+
+DEFINE_LOG_CATEGORY_STATIC(LogDiskIO, Log, All);
+
+static FAutoConsoleCommand GDumpShortTermIOStats(
+	TEXT("disk.DumpShortTermStats"),
+	TEXT("Dumps short term disk I/O stats."),
+	FConsoleCommandDelegate::CreateLambda([]
+	{
+		UE_LOG(LogDiskIO, Display, TEXT("Disk I/O short term stats:"));
+		GDiskUtilizationTracker.GetShortTermStats().Dump();
+	})
+);
+
+void FDiskUtilizationTracker::UtilizationStats::Dump() const
+{
+	UE_LOG(LogDiskIO, Display, TEXT("Total Reads: %llu Total Bytes Read: %llu"), TotalReads, TotalBytesRead);
+	UE_LOG(LogDiskIO, Display, TEXT("Total IO Time: %fs Total Idle Time: %fs"), TotalIOTime, TotalIdleTime);
+	UE_LOG(LogDiskIO, Display, TEXT("Read Throughput: %fMB/s Pct Time Idle: %f%%"), GetReadThrougputMBS(), GetPercentTimeIdle());
+}
 
 void FDiskUtilizationTracker::MaybePrint()
 {
diff --git a/Engine/Source/Runtime/Core/Public/HAL/DiskUtilizationTracker.h b/Engine/Source/Runtime/Core/Public/HAL/DiskUtilizationTracker.h
--- a/Engine/Source/Runtime/Core/Public/HAL/DiskUtilizationTracker.h
+++ b/Engine/Source/Runtime/Core/Public/HAL/DiskUtilizationTracker.h
@@ -2,13 +2,18 @@
 
 #pragma once
 
-#include "HAL/PlatformMisc.h"
-
 #ifndef TRACK_DISK_UTILIZATION
 #define TRACK_DISK_UTILIZATION 0
 #endif
 
 #if TRACK_DISK_UTILIZATION
+#include "HAL/PlatformMisc.h"
+#include "HAL/ThreadSafeBool.h"
+#include "ProfilingDebugging/CsvProfiler.h"
+
+#include <nn/nn_Log.h>
+
+CSV_DECLARE_CATEGORY_EXTERN(DiskIO);
 
 #ifndef SPEW_DISK_UTILIZATION
 #define SPEW_DISK_UTILIZATION 0
@@ -16,47 +21,111 @@
 
 struct FDiskUtilizationTracker
 {
-	FCriticalSection CriticalSection;
+	struct UtilizationStats
+	{
+		UtilizationStats() :
+			TotalReads(0),
+			TotalSeeks(0),
+			TotalBytesRead(0),
+			TotalSeekDistance(0),
+			TotalIOTime(0.0),
+			TotalIdleTime(0.0)
+		{}
+
+		double GetOverallThroughputMBS() const
+		{
+			return double(TotalBytesRead) / (TotalIOTime + TotalIdleTime) / (1024 * 1024);
+		}
 
-	uint64 NumReads;
-	uint64 NumSeeks;
+		double GetReadThrougputMBS() const
+		{
+			return double(TotalBytesRead) / (TotalIOTime) / (1024 * 1024);
+		}
 
-	uint64 IdleStartCycle;
-	uint64 ReadStartCycle;
+		double GetTotalIdleTimeInSeconds() const
+		{
+			return TotalIdleTime;
+		}
+
+		double GetTotalIOTimeInSeconds() const
+		{
+			return TotalIOTime;
+		}
+
+		double GetPercentTimeIdle() const
+		{
+			double TotalTime = TotalIOTime + TotalIdleTime;
+
+			return TotalTime > 0.0 ? (100.0f * TotalIdleTime) / TotalTime : 0.0;
+		}
+
+		void Reset()
+		{
+			TotalReads = 0;
+			TotalSeeks = 0;
+			TotalBytesRead = 0;
+			TotalSeekDistance = 0;
+			TotalIOTime = 0.0;
+			TotalIdleTime = 0.0;
+		}
 
-	uint64 TotalBytesRead;
-	uint64 TotalSeekDistance;
+		void Dump() const;
 
-	double TotalIOTime;
-	double TotalIdleTime;
+		uint64 TotalReads;
+		uint64 TotalSeeks;
+
+		uint64 TotalBytesRead;
+		uint64 TotalSeekDistance;
+
+		double TotalIOTime;
+		double TotalIdleTime;
+	};
+
+	UtilizationStats LongTermStats;
+	UtilizationStats ShortTermStats;
+
+	FCriticalSection CriticalSection;
+
+	uint64 IdleStartCycle;
+	uint64 ReadStartCycle;
 
 	uint64 InFlightBytes;
 	int32  InFlightReads;
 
-	FDiskUtilizationTracker()
-		: NumReads(0)
-		, NumSeeks(0)
-		, IdleStartCycle(0)
-		, ReadStartCycle(0)
-		, TotalBytesRead(0)
-		, TotalSeekDistance(0)
-		, TotalIOTime(0.0)
-		, TotalIdleTime(0.0)
-		, InFlightBytes(0)
-		, InFlightReads(0)
+	FThreadSafeBool bResetShortTermStats;
+
+	FDiskUtilizationTracker() :
+		IdleStartCycle(0),
+		ReadStartCycle(0),
+		InFlightBytes(0),
+		InFlightReads(0)
 	{
 	}
 
 	void StartRead(uint64 InReadBytes, uint64 InSeekDistance = 0)
 	{
+		static bool bBreak = false;
+
+		bool bReset = bResetShortTermStats.AtomicSet(false);
+
+		if (bReset)
+		{
+			ShortTermStats.Reset();
+			bBreak = true;
+		}
+
 		// update total reads
-		NumReads++;
+		LongTermStats.TotalReads++;
+		ShortTermStats.TotalReads++;
 
 		// update seek data
 		if (InSeekDistance > 0)
 		{
-			NumSeeks++;
-			TotalSeekDistance += InSeekDistance;
+			LongTermStats.TotalSeeks++;
+			ShortTermStats.TotalSeeks++;
+
+			LongTermStats.TotalSeekDistance += InSeekDistance;
+			ShortTermStats.TotalSeekDistance += InSeekDistance;
 		}
 
 		{
@@ -70,8 +139,12 @@ struct FDiskUtilizationTracker
 				// update idle time (if we've been idle)
 				if (IdleStartCycle > 0)
 				{
+					const double IdleTime = double(ReadStartCycle - IdleStartCycle) * FPlatformTime::GetSecondsPerCycle64();
+
+					LongTermStats.TotalIdleTime += IdleTime;
+					ShortTermStats.TotalIdleTime += bReset ? 0 : IdleTime;
 
-					TotalIdleTime += double(ReadStartCycle - IdleStartCycle) * FPlatformTime::GetSecondsPerCycle64();
+					CSV_CUSTOM_STAT(DiskIO, AccumulatedIdleTime, float(IdleTime), ECsvCustomStatOp::Accumulate);
 				}
 			}
 
@@ -93,46 +166,41 @@ struct FDiskUtilizationTracker
 				IdleStartCycle = FPlatformTime::Cycles64();
 
 				// update our read counters
-				TotalIOTime += double(IdleStartCycle - ReadStartCycle) * FPlatformTime::GetSecondsPerCycle64();
-				TotalBytesRead += InFlightBytes;
+				const double IOTime = double(IdleStartCycle - ReadStartCycle) * FPlatformTime::GetSecondsPerCycle64();
+
+				LongTermStats.TotalIOTime += IOTime;
+				ShortTermStats.TotalIOTime += IOTime;
+
+				LongTermStats.TotalBytesRead += InFlightBytes;
+				ShortTermStats.TotalBytesRead += InFlightBytes;
+
+				CSV_CUSTOM_STAT(DiskIO, AccumulatedIOTime, float(IOTime), ECsvCustomStatOp::Accumulate);
 
 				InFlightBytes = 0;
 			}
 
 		}
-
 		MaybePrint();
 	}
 
-	double GetOverallThroughputMBS()
+	uint32 GetOutstandingRequests() const
 	{
-		return double(TotalBytesRead) / (TotalIOTime + TotalIdleTime) / (1024 * 1024);
+		return InFlightReads;
 	}
 
-	double GetReadThrougputMBS()
+	struct UtilizationStats& GetLongTermStats()
 	{
-		return double(TotalBytesRead) / (TotalIOTime) / (1024 * 1024);
+		return LongTermStats;
 	}
 
-	void ResetStats()
+	struct UtilizationStats& GetShortTermStats()
 	{
-		// make sure we're not reseting stats while a read is in flight. that'd mess things up.
-		check(InFlightReads == 0);
-
-		NumSeeks = 0;
-		NumReads = 0;
-		
-		IdleStartCycle = 0;
-		ReadStartCycle = 0;
-
-		TotalBytesRead = 0;
-		TotalSeekDistance = 0;
-
-		TotalIOTime = 0.0;
-		TotalIdleTime = 0.0;
+		return ShortTermStats;
+	}
 
-		InFlightBytes = 0;
-		InFlightReads = 0;
+	void ResetShortTermStats()
+	{
+		bResetShortTermStats = true;
 	}
 
 	static constexpr float PrintFrequencySeconds = 0.5f;
diff --git a/Engine/Source/Runtime/PakFile/Private/IPlatformFilePak.cpp b/Engine/Source/Runtime/PakFile/Private/IPlatformFilePak.cpp
--- a/Engine/Source/Runtime/PakFile/Private/IPlatformFilePak.cpp
+++ b/Engine/Source/Runtime/PakFile/Private/IPlatformFilePak.cpp
@@ -21,6 +21,7 @@
 #include "Templates/Greater.h"
 #include "Serialization/ArchiveProxy.h"
 #include "Misc/Base64.h"
+#include "HAL/DiskUtilizationTracker.h"
 #include "Stats/StatsMisc.h"
 #include "HAL/PlatformFilemanager.h"
 #include "HAL/ThreadHeartBeat.h"
@@ -4116,6 +4117,11 @@ void FPakPlatformFile::Tick()
 		GPreCacheTotalLoadedLastTick = GPreCacheTotalLoaded;
 }
 #endif
+#if TRACK_DISK_UTILIZATION && CSV_PROFILER
+	CSV_CUSTOM_STAT(DiskIO, OutstandingIORequests, int32(GDiskUtilizationTracker.GetOutstandingRequests()), ECsvCustomStatOp::Set);
+	CSV_CUSTOM_STAT(DiskIO, BusyTime, float(GDiskUtilizationTracker.GetShortTermStats().GetTotalIOTimeInSeconds()), ECsvCustomStatOp::Set);
+	CSV_CUSTOM_STAT(DiskIO, IdleTime, float(GDiskUtilizationTracker.GetShortTermStats().GetTotalIdleTimeInSeconds()), ECsvCustomStatOp::Set);
+#endif
 }
 
 class FMappedFilePakProxy final : public IMappedFileHandle
