#UE4 CIS fix

diff --git a/Engine/Source/Runtime/Core/Public/HAL/DiskUtilizationTracker.h b/Engine/Source/Runtime/Core/Public/HAL/DiskUtilizationTracker.h
--- a/Engine/Source/Runtime/Core/Public/HAL/DiskUtilizationTracker.h
+++ b/Engine/Source/Runtime/Core/Public/HAL/DiskUtilizationTracker.h
@@ -6,6 +6,8 @@
 #define TRACK_DISK_UTILIZATION 0
 #endif
 
+#include "HAL/Platform.h"
+
 #if TRACK_DISK_UTILIZATION
 #include "HAL/PlatformMisc.h"
 #include "HAL/ThreadSafeBool.h"
