#UE4 Include fix

diff --git a/Engine/Source/Runtime/Core/Public/HAL/DiskUtilizationTracker.h b/Engine/Source/Runtime/Core/Public/HAL/DiskUtilizationTracker.h
--- a/Engine/Source/Runtime/Core/Public/HAL/DiskUtilizationTracker.h
+++ b/Engine/Source/Runtime/Core/Public/HAL/DiskUtilizationTracker.h
@@ -6,7 +6,7 @@
 #define TRACK_DISK_UTILIZATION 0
 #endif
 
-#include "HAL/Platform.h"
+#include "CoreTypes.h"
 
 #if TRACK_DISK_UTILIZATION
 #include "HAL/PlatformMisc.h"
