Fixed crash in LogTrace with nullptr string arguments

diff --git a/Engine/Source/Runtime/Core/Public/ProfilingDebugging/FormatArgsTrace.h b/Engine/Source/Runtime/Core/Public/ProfilingDebugging/FormatArgsTrace.h
--- a/Engine/Source/Runtime/Core/Public/ProfilingDebugging/FormatArgsTrace.h
+++ b/Engine/Source/Runtime/Core/Public/ProfilingDebugging/FormatArgsTrace.h
@@ -60,7 +60,14 @@ private:
 	template <typename T, typename CharType = typename TRemoveCV<typename TRemovePointer<T>::Type>::Type>
 	static typename TEnableIf<TIsStringArgument<T>::Value, uint64>::Type GetArgumentEncodedSize(T Argument)
 	{
-		return (TCString<CharType>::Strlen(Argument) + 1) * sizeof(CharType);
+		if (Argument != nullptr)
+		{
+			return (TCString<CharType>::Strlen(Argument) + 1) * sizeof(CharType);
+		}
+		else
+		{
+			return sizeof(CharType);
+		}
 	}
 
 	constexpr static uint64 GetArgumentsEncodedSize()
@@ -112,9 +119,18 @@ private:
 	static typename TEnableIf<TIsStringArgument<T>::Value>::Type EncodeArgumentInternal(uint8*& TypeCodesPtr, uint8*& PayloadPtr, T Argument)
 	{
 		*TypeCodesPtr++ = FormatArgTypeCode_CategoryString | sizeof(CharType);
-		uint16 Length = (TCString<CharType>::Strlen(Argument) + 1) * sizeof(CharType);
-		memcpy(PayloadPtr, Argument, Length);
-		PayloadPtr += Length;
+		if (Argument != nullptr)
+		{
+			uint16 Length = (TCString<CharType>::Strlen(Argument) + 1) * sizeof(CharType);
+			memcpy(PayloadPtr, Argument, Length);
+			PayloadPtr += Length;
+		}
+		else
+		{
+			CharType Terminator { 0 };
+			memcpy(PayloadPtr, &Terminator, sizeof(CharType));
+			PayloadPtr += sizeof(CharType);
+		}
 	}
 
 	constexpr static void EncodeArgumentsInternal(uint8*& TypeCodesPtr, uint8*& PayloadPtr)
