Remove check on trace cpu scope overflow

diff --git a/Engine/Source/Runtime/Core/Private/ProfilingDebugging/CpuProfilerTrace.cpp b/Engine/Source/Runtime/Core/Private/ProfilingDebugging/CpuProfilerTrace.cpp
--- a/Engine/Source/Runtime/Core/Private/ProfilingDebugging/CpuProfilerTrace.cpp
+++ b/Engine/Source/Runtime/Core/Private/ProfilingDebugging/CpuProfilerTrace.cpp
@@ -135,7 +135,10 @@ uint16 FCpuProfilerTrace::OutputEventType(const TCHAR* Name, ECpuProfilerGroup G
 {
 	static TAtomic<uint32> NextSpecId(1);
 	uint32 WideId = NextSpecId++;
-	check(WideId <= 0xFFFF);
+	if (WideId > 0xFFFF)
+	{
+		WideId = 0xFFFF;
+	}
 	uint16 Id = uint16(WideId);
 	uint16 NameSize = (FCString::Strlen(Name) + 1) * sizeof(TCHAR);
 	UE_TRACE_LOG(CpuProfiler, EventSpec, NameSize)
