Deprecating ARRAY_COUNT and changing it to UE_ARRAY_COUNT.

diff --git a/Engine/Source/Runtime/Core/Private/Internationalization/FastDecimalFormat.cpp b/Engine/Source/Runtime/Core/Private/Internationalization/FastDecimalFormat.cpp
--- a/Engine/Source/Runtime/Core/Private/Internationalization/FastDecimalFormat.cpp
+++ b/Engine/Source/Runtime/Core/Private/Internationalization/FastDecimalFormat.cpp
@@ -35,7 +35,7 @@ static const uint64 Pow10Table[] = {
 	1000000000000000000,	// 10^18
 };
 
-static_assert(ARRAY_COUNT(Pow10Table) - 1 >= MaxFractionalPrintPrecision, "Pow10Table must at big enough to index any value up-to MaxFractionalPrintPrecision");
+static_assert(UE_ARRAY_COUNT(Pow10Table) - 1 >= MaxFractionalPrintPrecision, "Pow10Table must at big enough to index any value up-to MaxFractionalPrintPrecision");
 
 enum class EDecimalNumberSigningStringsFlags : uint8
 {
@@ -377,7 +377,7 @@ void IntegralToString(const bool bIsNegative, const uint64 InVal, const FDecimal
 
 	// Deal with the integral part (produces a string of the integral part, inserting group separators if requested and required, and padding as needed)
 	TCHAR IntegralPartBuffer[MinRequiredIntegralBufferSize];
-	const int32 IntegralPartLen = IntegralToString_Common(InVal, InFormattingRules, InFormattingOptions, IntegralPartBuffer, ARRAY_COUNT(IntegralPartBuffer));
+	const int32 IntegralPartLen = IntegralToString_Common(InVal, InFormattingRules, InFormattingOptions, IntegralPartBuffer, UE_ARRAY_COUNT(IntegralPartBuffer));
 
 	// Deal with any forced fractional part (produces a string zeros up to the required minimum length)
 	TCHAR FractionalPartBuffer[MinRequiredIntegralBufferSize];
@@ -419,14 +419,14 @@ void FractionalToString(const double InVal, const FDecimalNumberFormattingRules&
 
 	// Deal with the integral part (produces a string of the integral part, inserting group separators if requested and required, and padding as needed)
 	TCHAR IntegralPartBuffer[MinRequiredIntegralBufferSize];
-	const int32 IntegralPartLen = IntegralToString_Common(static_cast<uint64>(IntegralPart), InFormattingRules, InFormattingOptions, IntegralPartBuffer, ARRAY_COUNT(IntegralPartBuffer));
+	const int32 IntegralPartLen = IntegralToString_Common(static_cast<uint64>(IntegralPart), InFormattingRules, InFormattingOptions, IntegralPartBuffer, UE_ARRAY_COUNT(IntegralPartBuffer));
 
 	// Deal with the fractional part (produces a string of the fractional part, potentially padding with zeros up to InFormattingOptions.MaximumFractionalDigits)
 	TCHAR FractionalPartBuffer[MinRequiredIntegralBufferSize];
 	int32 FractionalPartLen = 0;
 	if (FractionalPart != 0.0)
 	{
-		FractionalPartLen = IntegralToString_UInt64ToString(static_cast<uint64>(FractionalPart), false, 0, 0, ' ', InFormattingRules.DigitCharacters, 0, InFormattingOptions.MaximumFractionalDigits, FractionalPartBuffer, ARRAY_COUNT(FractionalPartBuffer));
+		FractionalPartLen = IntegralToString_UInt64ToString(static_cast<uint64>(FractionalPart), false, 0, 0, ' ', InFormattingRules.DigitCharacters, 0, InFormattingOptions.MaximumFractionalDigits, FractionalPartBuffer, UE_ARRAY_COUNT(FractionalPartBuffer));
 	
 		{
 			// Pad the fractional part with any leading zeros that may have been lost when the number was split
@@ -564,7 +564,7 @@ bool StringToIntegral_StringToUInt64(const TCHAR*& InBuffer, const TCHAR* InBuff
 
 		// Process numeric characters (also test European numerals in case they were used by a language that doesn't normally use them)
 		bool bValidChar = false;
-		for (int32 CharIndex = 0; CharIndex < ARRAY_COUNT(InFormattingRules.DigitCharacters); ++CharIndex)
+		for (int32 CharIndex = 0; CharIndex < UE_ARRAY_COUNT(InFormattingRules.DigitCharacters); ++CharIndex)
 		{
 			if ((*InBuffer == InFormattingRules.DigitCharacters[CharIndex]) || (*InBuffer == EuropeanNumerals[CharIndex]))
 			{
diff --git a/Engine/Source/Runtime/Core/Private/Internationalization/TextHistory.h b/Engine/Source/Runtime/Core/Private/Internationalization/TextHistory.h
--- a/Engine/Source/Runtime/Core/Private/Internationalization/TextHistory.h
+++ b/Engine/Source/Runtime/Core/Private/Internationalization/TextHistory.h
@@ -70,14 +70,14 @@ static const TCHAR LocalSuffix[] = TEXT("_LOCAL");
 	Buffer = Func(Buffer, ##__VA_ARGS__);									\
 	if (!Buffer) { return nullptr; }
 
-#define TEXT_STRINGIFICATION_PEEK_MARKER(T)					TextStringificationUtil::PeekMarker(Buffer, T, ARRAY_COUNT(T) - 1)
-#define TEXT_STRINGIFICATION_PEEK_INSENSITIVE_MARKER(T)		TextStringificationUtil::PeekInsensitiveMarker(Buffer, T, ARRAY_COUNT(T) - 1)
+#define TEXT_STRINGIFICATION_PEEK_MARKER(T)					TextStringificationUtil::PeekMarker(Buffer, T, UE_ARRAY_COUNT(T) - 1)
+#define TEXT_STRINGIFICATION_PEEK_INSENSITIVE_MARKER(T)		TextStringificationUtil::PeekInsensitiveMarker(Buffer, T, UE_ARRAY_COUNT(T) - 1)
 bool PeekMarker(const TCHAR* Buffer, const TCHAR* InMarker, const int32 InMarkerLen);
 bool PeekInsensitiveMarker(const TCHAR* Buffer, const TCHAR* InMarker, const int32 InMarkerLen);
 
-#define TEXT_STRINGIFICATION_SKIP_MARKER(T)					TEXT_STRINGIFICATION_FUNC_MODIFY_BUFFER_AND_VALIDATE(TextStringificationUtil::SkipMarker, T, ARRAY_COUNT(T) - 1)
-#define TEXT_STRINGIFICATION_SKIP_INSENSITIVE_MARKER(T)		TEXT_STRINGIFICATION_FUNC_MODIFY_BUFFER_AND_VALIDATE(TextStringificationUtil::SkipInsensitiveMarker, T, ARRAY_COUNT(T) - 1)
-#define TEXT_STRINGIFICATION_SKIP_MARKER_LEN(T)				Buffer += (ARRAY_COUNT(T) - 1)
+#define TEXT_STRINGIFICATION_SKIP_MARKER(T)					TEXT_STRINGIFICATION_FUNC_MODIFY_BUFFER_AND_VALIDATE(TextStringificationUtil::SkipMarker, T, UE_ARRAY_COUNT(T) - 1)
+#define TEXT_STRINGIFICATION_SKIP_INSENSITIVE_MARKER(T)		TEXT_STRINGIFICATION_FUNC_MODIFY_BUFFER_AND_VALIDATE(TextStringificationUtil::SkipInsensitiveMarker, T, UE_ARRAY_COUNT(T) - 1)
+#define TEXT_STRINGIFICATION_SKIP_MARKER_LEN(T)				Buffer += (UE_ARRAY_COUNT(T) - 1)
 const TCHAR* SkipMarker(const TCHAR* Buffer, const TCHAR* InMarker, const int32 InMarkerLen);
 const TCHAR* SkipInsensitiveMarker(const TCHAR* Buffer, const TCHAR* InMarker, const int32 InMarkerLen);
 
diff --git a/Engine/Source/Runtime/Core/Private/Linux/LinuxPlatformProcess.cpp b/Engine/Source/Runtime/Core/Private/Linux/LinuxPlatformProcess.cpp
--- a/Engine/Source/Runtime/Core/Private/Linux/LinuxPlatformProcess.cpp
+++ b/Engine/Source/Runtime/Core/Private/Linux/LinuxPlatformProcess.cpp
@@ -12,7 +12,7 @@ const TCHAR* FLinuxPlatformProcess::BaseDir()
 	if (!bHaveResult)
 	{
 		char SelfPath[UNIX_MAX_PATH] = {0};
-		if (readlink( "/proc/self/exe", SelfPath, ARRAY_COUNT(SelfPath) - 1) == -1)
+		if (readlink( "/proc/self/exe", SelfPath, UE_ARRAY_COUNT(SelfPath) - 1) == -1)
 		{
 			int ErrNo = errno;
 			UE_LOG(LogHAL, Fatal, TEXT("readlink() failed with errno = %d (%s)"), ErrNo,
@@ -20,11 +20,11 @@ const TCHAR* FLinuxPlatformProcess::BaseDir()
 			// unreachable
 			return CachedResult;
 		}
-		SelfPath[ARRAY_COUNT(SelfPath) - 1] = 0;
+		SelfPath[UE_ARRAY_COUNT(SelfPath) - 1] = 0;
 
-		FCString::Strncpy(CachedResult, UTF8_TO_TCHAR(dirname(SelfPath)), ARRAY_COUNT(CachedResult) - 1);
-		CachedResult[ARRAY_COUNT(CachedResult) - 1] = 0;
-		FCString::Strncat(CachedResult, TEXT("/"), ARRAY_COUNT(CachedResult) - 1);
+		FCString::Strncpy(CachedResult, UTF8_TO_TCHAR(dirname(SelfPath)), UE_ARRAY_COUNT(CachedResult) - 1);
+		CachedResult[UE_ARRAY_COUNT(CachedResult) - 1] = 0;
+		FCString::Strncat(CachedResult, TEXT("/"), UE_ARRAY_COUNT(CachedResult) - 1);
 		bHaveResult = true;
 	}
 	return CachedResult;
diff --git a/Engine/Source/Runtime/Core/Private/Logging/LogMacros.cpp b/Engine/Source/Runtime/Core/Private/Logging/LogMacros.cpp
--- a/Engine/Source/Runtime/Core/Private/Logging/LogMacros.cpp
+++ b/Engine/Source/Runtime/Core/Private/Logging/LogMacros.cpp
@@ -55,10 +55,10 @@ void FMsg::LogfImpl(const ANSICHAR* File, int32 Line, const FLogCategoryName& Ca
 			// We're using one big shared static buffer here, so guard against re-entry
 			FScopeLock MsgLock(&MsgLogfStaticBufferGuard);
 			// Print to a large static buffer so we can keep the stack allocation below 16K
-			GET_VARARGS(MsgLogfStaticBuffer, ARRAY_COUNT(MsgLogfStaticBuffer), ARRAY_COUNT(MsgLogfStaticBuffer) - 1, Fmt, Fmt);
+			GET_VARARGS(MsgLogfStaticBuffer, UE_ARRAY_COUNT(MsgLogfStaticBuffer), UE_ARRAY_COUNT(MsgLogfStaticBuffer) - 1, Fmt, Fmt);
 			// Copy the message to the stack-allocated buffer)
-			FCString::Strncpy(Message, MsgLogfStaticBuffer, ARRAY_COUNT(Message) - 1);
-			Message[ARRAY_COUNT(Message) - 1] = '\0';
+			FCString::Strncpy(Message, MsgLogfStaticBuffer, UE_ARRAY_COUNT(Message) - 1);
+			Message[UE_ARRAY_COUNT(Message) - 1] = '\0';
 		}
 
 		const int32 NumStackFramesToIgnore = 1;
@@ -102,10 +102,10 @@ void FMsg::Logf_InternalImpl(const ANSICHAR* File, int32 Line, const FLogCategor
 			// We're using one big shared static buffer here, so guard against re-entry
 			FScopeLock MsgLock(&MsgLogfStaticBufferGuard);
 			// Print to a large static buffer so we can keep the stack allocation below 16K
-			GET_VARARGS(MsgLogfStaticBuffer, ARRAY_COUNT(MsgLogfStaticBuffer), ARRAY_COUNT(MsgLogfStaticBuffer) - 1, Fmt, Fmt);
+			GET_VARARGS(MsgLogfStaticBuffer, UE_ARRAY_COUNT(MsgLogfStaticBuffer), UE_ARRAY_COUNT(MsgLogfStaticBuffer) - 1, Fmt, Fmt);
 			// Copy the message to the stack-allocated buffer)
-			FCString::Strncpy(Message, MsgLogfStaticBuffer, ARRAY_COUNT(Message) - 1);
-			Message[ARRAY_COUNT(Message) - 1] = '\0';
+			FCString::Strncpy(Message, MsgLogfStaticBuffer, UE_ARRAY_COUNT(Message) - 1);
+			Message[UE_ARRAY_COUNT(Message) - 1] = '\0';
 		}
 
 		const int32 NumStackFramesToIgnore = 1;
diff --git a/Engine/Source/Runtime/Core/Private/Mac/MacErrorOutputDevice.cpp b/Engine/Source/Runtime/Core/Private/Mac/MacErrorOutputDevice.cpp
--- a/Engine/Source/Runtime/Core/Private/Mac/MacErrorOutputDevice.cpp
+++ b/Engine/Source/Runtime/Core/Private/Mac/MacErrorOutputDevice.cpp
@@ -34,8 +34,8 @@ void FMacErrorOutputDevice::Serialize( const TCHAR* Msg, ELogVerbosity::Type Ver
 		{
 			UE_LOG(LogMac, Error, TEXT("appError called: %s"), Msg );
 		}
-		FCString::Strncpy( GErrorHist, Msg, ARRAY_COUNT(GErrorHist) - 5 );
-		FCString::Strncat( GErrorHist, TEXT("\r\n\r\n"), ARRAY_COUNT(GErrorHist) - 1  );
+		FCString::Strncpy( GErrorHist, Msg, UE_ARRAY_COUNT(GErrorHist) - 5 );
+		FCString::Strncat( GErrorHist, TEXT("\r\n\r\n"), UE_ARRAY_COUNT(GErrorHist) - 1  );
 		ErrorPos = FCString::Strlen(GErrorHist);
 	}
 	else
@@ -86,7 +86,7 @@ void FMacErrorOutputDevice::HandleError()
 	GIsRunning				= 0;
 	GIsCriticalError		= 1;
 	GLogConsole				= NULL;
-	GErrorHist[ARRAY_COUNT(GErrorHist)-1]=0;
+	GErrorHist[UE_ARRAY_COUNT(GErrorHist)-1]=0;
 
 	// Dump the error and flush the log.
 	UE_LOG(LogMac, Log, TEXT("=== Critical error: ===") LINE_TERMINATOR TEXT("%s") LINE_TERMINATOR, GErrorExceptionDescription);
diff --git a/Engine/Source/Runtime/Core/Private/Mac/MacPlatformMisc.cpp b/Engine/Source/Runtime/Core/Private/Mac/MacPlatformMisc.cpp
--- a/Engine/Source/Runtime/Core/Private/Mac/MacPlatformMisc.cpp
+++ b/Engine/Source/Runtime/Core/Private/Mac/MacPlatformMisc.cpp
@@ -202,7 +202,7 @@ struct FMacApplicationInfo
 		TempSysCtlBufferSize = PATH_MAX+1;
 		sysctlbyname("machdep.cpu.brand_string", MachineCPUString, &TempSysCtlBufferSize, NULL, 0);
 		
-		gethostname(MachineName, ARRAY_COUNT(MachineName));
+		gethostname(MachineName, UE_ARRAY_COUNT(MachineName));
 		
 		FString CrashVideoPath = FPaths::ProjectLogDir() + TEXT("CrashVideo.avi");
 
diff --git a/Engine/Source/Runtime/Core/Private/Mac/MacPlatformProcess.cpp b/Engine/Source/Runtime/Core/Private/Mac/MacPlatformProcess.cpp
--- a/Engine/Source/Runtime/Core/Private/Mac/MacPlatformProcess.cpp
+++ b/Engine/Source/Runtime/Core/Private/Mac/MacPlatformProcess.cpp
@@ -897,8 +897,8 @@ const TCHAR* FMacPlatformProcess::ComputerName()
 
 	if( !Result[0] )
 	{
-		ANSICHAR AnsiResult[ARRAY_COUNT(Result)];
-		gethostname(AnsiResult, ARRAY_COUNT(Result));
+		ANSICHAR AnsiResult[UE_ARRAY_COUNT(Result)];
+		gethostname(AnsiResult, UE_ARRAY_COUNT(Result));
 		FCString::Strcpy(Result, ANSI_TO_TCHAR(AnsiResult));
 	}
 	return Result;
diff --git a/Engine/Source/Runtime/Core/Private/Math/Box.cpp b/Engine/Source/Runtime/Core/Private/Math/Box.cpp
--- a/Engine/Source/Runtime/Core/Private/Math/Box.cpp
+++ b/Engine/Source/Runtime/Core/Private/Math/Box.cpp
@@ -103,7 +103,7 @@ FBox FBox::InverseTransformBy(const FTransform& M) const
 
 	FBox NewBox(ForceInit);
 
-	for (int32 VertexIndex = 0; VertexIndex < ARRAY_COUNT(Vertices); VertexIndex++)
+	for (int32 VertexIndex = 0; VertexIndex < UE_ARRAY_COUNT(Vertices); VertexIndex++)
 	{
 		FVector4 ProjectedVertex = M.InverseTransformPosition(Vertices[VertexIndex]);
 		NewBox += ProjectedVertex;
@@ -129,7 +129,7 @@ FBox FBox::TransformProjectBy(const FMatrix& ProjM) const
 
 	FBox NewBox(ForceInit);
 
-	for (int32 VertexIndex = 0; VertexIndex < ARRAY_COUNT(Vertices); VertexIndex++)
+	for (int32 VertexIndex = 0; VertexIndex < UE_ARRAY_COUNT(Vertices); VertexIndex++)
 	{
 		FVector4 ProjectedVertex = ProjM.TransformPosition(Vertices[VertexIndex]);
 		NewBox += ((FVector)ProjectedVertex) / ProjectedVertex.W;
diff --git a/Engine/Source/Runtime/Core/Private/Math/UnitConversion.cpp b/Engine/Source/Runtime/Core/Private/Math/UnitConversion.cpp
--- a/Engine/Source/Runtime/Core/Private/Math/UnitConversion.cpp
+++ b/Engine/Source/Runtime/Core/Private/Math/UnitConversion.cpp
@@ -296,7 +296,7 @@ struct FUnitExpressionParser
 			FStringToken UnitToken;
 			const FStringToken StartingToken = NumberToken.GetValue();
 
-			for (int32 Index = 0; Index < ARRAY_COUNT(ParseCandidates); ++Index)
+			for (int32 Index = 0; Index < UE_ARRAY_COUNT(ParseCandidates); ++Index)
 			{
 				TOptional<FStringToken> CandidateUnitToken = Stream.ParseTokenIgnoreCase(ParseCandidates[Index].String, &NumberToken.GetValue());
 				const size_t MatchedCharCount = NumberToken.GetValue().GetTokenEndPos() - StartingToken.GetTokenEndPos();
@@ -499,8 +499,8 @@ EUnitType FUnitConversion::GetUnitType(EUnit InUnit)
 /** Get the unit abbreviation the specified unit type */
 const TCHAR* FUnitConversion::GetUnitDisplayString(EUnit Unit)
 {
-	static_assert(ARRAY_COUNT(UnitTypes) == (uint8)EUnit::Unspecified, "Type array does not match size of unit enum");
-	static_assert(ARRAY_COUNT(DisplayStrings) == (uint8)EUnit::Unspecified, "Display String array does not match size of unit enum");
+	static_assert(UE_ARRAY_COUNT(UnitTypes) == (uint8)EUnit::Unspecified, "Type array does not match size of unit enum");
+	static_assert(UE_ARRAY_COUNT(DisplayStrings) == (uint8)EUnit::Unspecified, "Display String array does not match size of unit enum");
 	
 	if (Unit != EUnit::Unspecified)
 	{
@@ -517,7 +517,7 @@ TOptional<EUnit> FUnitConversion::UnitFromString(const TCHAR* UnitString)
 		return TOptional<EUnit>();
 	}
 
-	for (int32 Index = 0; Index < ARRAY_COUNT(ParseCandidates); ++Index)
+	for (int32 Index = 0; Index < UE_ARRAY_COUNT(ParseCandidates); ++Index)
 	{
 		if (FCString::Stricmp(UnitString, ParseCandidates[Index].String) == 0)
 		{
diff --git a/Engine/Source/Runtime/Core/Private/Misc/AssertionMacros.cpp b/Engine/Source/Runtime/Core/Private/Misc/AssertionMacros.cpp
--- a/Engine/Source/Runtime/Core/Private/Misc/AssertionMacros.cpp
+++ b/Engine/Source/Runtime/Core/Private/Misc/AssertionMacros.cpp
@@ -98,7 +98,7 @@ static void AssertFailedImplV(const ANSICHAR* Expr, const ANSICHAR* File, int32
 	bHasAsserted = true;
 
 	TCHAR DescriptionString[4096];
-	FCString::GetVarArgs(DescriptionString, ARRAY_COUNT(DescriptionString), Format, Args);
+	FCString::GetVarArgs(DescriptionString, UE_ARRAY_COUNT(DescriptionString), Format, Args);
 
 	TCHAR ErrorString[MAX_SPRINTF];
 	FCString::Sprintf(ErrorString, TEXT("%s"), ANSI_TO_TCHAR(Expr));
@@ -117,7 +117,7 @@ FORCENOINLINE void StaticFailDebug( const TCHAR* Error, const ANSICHAR* File, in
 
 	TCHAR DescriptionAndTrace[4096];
 
-	FCString::Strncpy(DescriptionAndTrace, Description, ARRAY_COUNT(DescriptionAndTrace) - 1);
+	FCString::Strncpy(DescriptionAndTrace, Description, UE_ARRAY_COUNT(DescriptionAndTrace) - 1);
 
 	// some platforms (Windows, Mac, Linux) generate this themselves by throwing an exception and capturing
 	// the backtrace later on
@@ -127,10 +127,10 @@ FORCENOINLINE void StaticFailDebug( const TCHAR* Error, const ANSICHAR* File, in
 		if (StackTrace != NULL)
 		{
 			StackTrace[0] = 0;
-			FPlatformStackWalk::StackWalkAndDump(StackTrace, ARRAY_COUNT(StackTrace), NumStackFramesToIgnore + 1);
+			FPlatformStackWalk::StackWalkAndDump(StackTrace, UE_ARRAY_COUNT(StackTrace), NumStackFramesToIgnore + 1);
 
-			FCString::Strncat(DescriptionAndTrace, TEXT("\n"), ARRAY_COUNT(DescriptionAndTrace) - 1);
-			FCString::Strncat(DescriptionAndTrace, ANSI_TO_TCHAR(StackTrace), ARRAY_COUNT(DescriptionAndTrace) - 1);
+			FCString::Strncat(DescriptionAndTrace, TEXT("\n"), UE_ARRAY_COUNT(DescriptionAndTrace) - 1);
+			FCString::Strncat(DescriptionAndTrace, ANSI_TO_TCHAR(StackTrace), UE_ARRAY_COUNT(DescriptionAndTrace) - 1);
 		}
 	}
 
@@ -139,15 +139,15 @@ FORCENOINLINE void StaticFailDebug( const TCHAR* Error, const ANSICHAR* File, in
 
 	// Copy the detailed error into the error message.
 	TCHAR ErrorMessage[4096];
-	if (FCString::Snprintf( ErrorMessage, ARRAY_COUNT( ErrorMessage ), TEXT( "%s" ) FILE_LINE_DESC TEXT( "\n%s\n" ), Error, ANSI_TO_TCHAR( File ), Line, DescriptionAndTrace) < 0)
+	if (FCString::Snprintf( ErrorMessage, UE_ARRAY_COUNT( ErrorMessage ), TEXT( "%s" ) FILE_LINE_DESC TEXT( "\n%s\n" ), Error, ANSI_TO_TCHAR( File ), Line, DescriptionAndTrace) < 0)
 	{
 		// Description and callstack was too long to fit in GErrorMessage. Use only description
-		FCString::Snprintf( ErrorMessage, ARRAY_COUNT( ErrorMessage ), TEXT( "%s" ) FILE_LINE_DESC TEXT( "\n%s\n<< callstack too long >>" ), Error, ANSI_TO_TCHAR( File ), Line, Description);
+		FCString::Snprintf( ErrorMessage, UE_ARRAY_COUNT( ErrorMessage ), TEXT( "%s" ) FILE_LINE_DESC TEXT( "\n%s\n<< callstack too long >>" ), Error, ANSI_TO_TCHAR( File ), Line, Description);
 	}
 
 	// Copy the error message to the error history.
-	FCString::Strncpy( GErrorHist, ErrorMessage, ARRAY_COUNT( GErrorHist ) );
-	FCString::Strncat( GErrorHist, TEXT( "\r\n\r\n" ), ARRAY_COUNT( GErrorHist ) );
+	FCString::Strncpy( GErrorHist, ErrorMessage, UE_ARRAY_COUNT( GErrorHist ) );
+	FCString::Strncat( GErrorHist, TEXT( "\r\n\r\n" ), UE_ARRAY_COUNT( GErrorHist ) );
 }
 
 
@@ -192,7 +192,7 @@ void FDebug::LogFormattedMessageWithCallstack(const FName& InLogName, const ANSI
 		// Find the end of the current line
 		const TCHAR* LineEnd = LineStart;
 		TCHAR* SingleLineWritePos = SingleLine;
-		int32 SpaceRemaining = ARRAY_COUNT(SingleLine) - 1;
+		int32 SpaceRemaining = UE_ARRAY_COUNT(SingleLine) - 1;
 
 		while (SpaceRemaining > 0 && *LineEnd != 0 && *LineEnd != '\r' && *LineEnd != '\n')
 		{
@@ -252,7 +252,7 @@ void FDebug::LogAssertFailedMessageImplV(const ANSICHAR* Expr, const ANSICHAR* F
 	if( !GIsCriticalError )
 	{
 		TCHAR DescriptionString[4096];
-		FCString::GetVarArgs( DescriptionString, ARRAY_COUNT(DescriptionString), Fmt, Args );
+		FCString::GetVarArgs( DescriptionString, UE_ARRAY_COUNT(DescriptionString), Fmt, Args );
 
 		TCHAR ErrorString[MAX_SPRINTF];
 		FCString::Sprintf( ErrorString, TEXT( "Assertion failed: %s" ), ANSI_TO_TCHAR( Expr ) );
@@ -316,7 +316,7 @@ FORCENOINLINE void FDebug::EnsureFailed(const ANSICHAR* Expr, const ANSICHAR* Fi
 
 		// Create a final string that we'll output to the log (and error history buffer)
 		TCHAR ErrorMsg[16384];
-		FCString::Snprintf(ErrorMsg, ARRAY_COUNT(ErrorMsg), TEXT("Ensure condition failed: %s [File:%s] [Line: %i]") LINE_TERMINATOR TEXT("%s") LINE_TERMINATOR TEXT("Stack: ") LINE_TERMINATOR, ANSI_TO_TCHAR(Expr), ANSI_TO_TCHAR(File), Line, Msg);
+		FCString::Snprintf(ErrorMsg, UE_ARRAY_COUNT(ErrorMsg), TEXT("Ensure condition failed: %s [File:%s] [Line: %i]") LINE_TERMINATOR TEXT("%s") LINE_TERMINATOR TEXT("Stack: ") LINE_TERMINATOR, ANSI_TO_TCHAR(Expr), ANSI_TO_TCHAR(File), Line, Msg);
 
 		// No debugger attached, so generate a call stack and submit a crash report
 		// Walk the stack and dump it to the allocated memory.
@@ -340,7 +340,7 @@ FORCENOINLINE void FDebug::EnsureFailed(const ANSICHAR* Expr, const ANSICHAR* Fi
 			}
 
 			// Also append the stack trace
-			FCString::Strncat(ErrorMsg, ANSI_TO_TCHAR(StackTrace), ARRAY_COUNT(ErrorMsg) - 1);
+			FCString::Strncat(ErrorMsg, ANSI_TO_TCHAR(StackTrace), UE_ARRAY_COUNT(ErrorMsg) - 1);
 			FMemory::SystemFree(StackTrace);
 
 			// Dump the error and flush the log.
@@ -496,7 +496,7 @@ FORCENOINLINE bool VARARGS FDebug::OptionallyLogFormattedEnsureMessageReturningF
 FORCENOINLINE void VARARGS LowLevelFatalErrorHandler(const ANSICHAR* File, int32 Line, const TCHAR* Format, ...)
 {
 	TCHAR DescriptionString[4096];
-	GET_VARARGS( DescriptionString, ARRAY_COUNT(DescriptionString), ARRAY_COUNT(DescriptionString)-1, Format, Format );
+	GET_VARARGS( DescriptionString, UE_ARRAY_COUNT(DescriptionString), UE_ARRAY_COUNT(DescriptionString)-1, Format, Format );
 
 	const int32 NumStackFramesToIgnore = 1; // Just ignore this frame
 	StaticFailDebug(TEXT("LowLevelFatalError"), File, Line, DescriptionString, false, NumStackFramesToIgnore);
diff --git a/Engine/Source/Runtime/Core/Private/Misc/CommandLine.cpp b/Engine/Source/Runtime/Core/Private/Misc/CommandLine.cpp
--- a/Engine/Source/Runtime/Core/Private/Misc/CommandLine.cpp
+++ b/Engine/Source/Runtime/Core/Private/Misc/CommandLine.cpp
@@ -57,12 +57,12 @@ bool FCommandLine::Set(const TCHAR* NewCommandLine)
 {
 	if (!bIsInitialized)
 	{
-		FCString::Strncpy(OriginalCmdLine, NewCommandLine, ARRAY_COUNT(OriginalCmdLine));
-		FCString::Strncpy(LoggingOriginalCmdLine, NewCommandLine, ARRAY_COUNT(LoggingOriginalCmdLine));
+		FCString::Strncpy(OriginalCmdLine, NewCommandLine, UE_ARRAY_COUNT(OriginalCmdLine));
+		FCString::Strncpy(LoggingOriginalCmdLine, NewCommandLine, UE_ARRAY_COUNT(LoggingOriginalCmdLine));
 	}
 
-	FCString::Strncpy( CmdLine, NewCommandLine, ARRAY_COUNT(CmdLine) );
-	FCString::Strncpy(LoggingCmdLine, NewCommandLine, ARRAY_COUNT(LoggingCmdLine));
+	FCString::Strncpy( CmdLine, NewCommandLine, UE_ARRAY_COUNT(CmdLine) );
+	FCString::Strncpy(LoggingCmdLine, NewCommandLine, UE_ARRAY_COUNT(LoggingCmdLine));
 	// If configured as part of the build, strip out any unapproved args
 	WhitelistCommandLines();
 
@@ -85,7 +85,7 @@ bool FCommandLine::Set(const TCHAR* NewCommandLine)
 
 void FCommandLine::Append(const TCHAR* AppendString)
 {
-	FCString::Strncat( CmdLine, AppendString, ARRAY_COUNT(CmdLine) );
+	FCString::Strncat( CmdLine, AppendString, UE_ARRAY_COUNT(CmdLine) );
 	// If configured as part of the build, strip out any unapproved args
 	WhitelistCommandLines();
 }
@@ -144,16 +144,16 @@ void FCommandLine::WhitelistCommandLines()
 	}
 	// Process the original command line
 	TArray<FString> OriginalList = FilterCommandLine(OriginalCmdLine);
-	BuildWhitelistCommandLine(OriginalCmdLine, ARRAY_COUNT(OriginalCmdLine), OriginalList);
+	BuildWhitelistCommandLine(OriginalCmdLine, UE_ARRAY_COUNT(OriginalCmdLine), OriginalList);
 	// Process the current command line
 	TArray<FString> CmdList = FilterCommandLine(CmdLine);
-	BuildWhitelistCommandLine(CmdLine, ARRAY_COUNT(CmdLine), CmdList);
+	BuildWhitelistCommandLine(CmdLine, UE_ARRAY_COUNT(CmdLine), CmdList);
 	// Process the command line for logging purposes
 	TArray<FString> LoggingCmdList = FilterCommandLineForLogging(LoggingCmdLine);
-	BuildWhitelistCommandLine(LoggingCmdLine, ARRAY_COUNT(LoggingCmdLine), LoggingCmdList);
+	BuildWhitelistCommandLine(LoggingCmdLine, UE_ARRAY_COUNT(LoggingCmdLine), LoggingCmdList);
 	// Process the original command line for logging purposes
 	TArray<FString> LoggingOriginalCmdList = FilterCommandLineForLogging(LoggingOriginalCmdLine);
-	BuildWhitelistCommandLine(LoggingOriginalCmdLine, ARRAY_COUNT(LoggingOriginalCmdLine), LoggingOriginalCmdList);
+	BuildWhitelistCommandLine(LoggingOriginalCmdLine, UE_ARRAY_COUNT(LoggingOriginalCmdLine), LoggingOriginalCmdList);
 }
 
 TArray<FString> FCommandLine::FilterCommandLine(TCHAR* CommandLine)
diff --git a/Engine/Source/Runtime/Core/Private/Misc/ConfigCacheIni.cpp b/Engine/Source/Runtime/Core/Private/Misc/ConfigCacheIni.cpp
--- a/Engine/Source/Runtime/Core/Private/Misc/ConfigCacheIni.cpp
+++ b/Engine/Source/Runtime/Core/Private/Misc/ConfigCacheIni.cpp
@@ -572,7 +572,7 @@ constexpr int32 GetStaticKey(int32 LayerIndex, int32 LayerExpansionIndex, int32
 	return LayerIndex * 10000 + LayerExpansionIndex * 100 + PlatformIndex;
 }
 
-constexpr int32 MaxStaticHierarchyKey = GetStaticKey(ARRAY_COUNT(GConfigLayers) - 1, ARRAY_COUNT(GConfigLayerExpansions) - 1, MaxPlatformIndex);
+constexpr int32 MaxStaticHierarchyKey = GetStaticKey(UE_ARRAY_COUNT(GConfigLayers) - 1, UE_ARRAY_COUNT(GConfigLayerExpansions) - 1, MaxPlatformIndex);
 
 /*-----------------------------------------------------------------------------
 	FConfigFileHierarchy
@@ -3341,7 +3341,7 @@ void FConfigFile::AddStaticLayersToHierarchy(const TCHAR* InBaseIniName, const T
 	bool bHasPlatformExtensionProjectConfigDir = FPaths::DirectoryExists(*PlatformExtensionProjectConfigDir);
 
 	// go over all the config layers
-	for (int32 LayerIndex = 0; LayerIndex < ARRAY_COUNT(GConfigLayers); LayerIndex++)
+	for (int32 LayerIndex = 0; LayerIndex < UE_ARRAY_COUNT(GConfigLayers); LayerIndex++)
 	{
 		const FConfigLayer& Layer = GConfigLayers[LayerIndex];
 		const bool bHasPlatformTag = FCString::Strstr(Layer.Path, TEXT("{PLATFORM}")) != nullptr;
@@ -3390,10 +3390,10 @@ void FConfigFile::AddStaticLayersToHierarchy(const TCHAR* InBaseIniName, const T
 			// we assume none of the more special tags in expanded ones
 			checkfSlow(FCString::Strstr(Layer.Path, TEXT("{USERSETTINGS}")) == nullptr && FCString::Strstr(Layer.Path, TEXT("{USER}")) == nullptr, TEXT("Expanded config %s shouldn't have a {USER*} tags in it"), *Layer.Path);
 			// last layer is expected to not be expanded
-			checkfSlow(LayerIndex < ARRAY_COUNT(GConfigLayers) - 1, TEXT("Final layer %s shouldn't be an expansion layer, as it needs to generate the  hierarchy cache key"), *Layer.Path);
+			checkfSlow(LayerIndex < UE_ARRAY_COUNT(GConfigLayers) - 1, TEXT("Final layer %s shouldn't be an expansion layer, as it needs to generate the  hierarchy cache key"), *Layer.Path);
 
 			// loop over all the possible expansions
-			for (int32 ExpansionIndex = 0; ExpansionIndex < ARRAY_COUNT(GConfigLayerExpansions); ExpansionIndex++)
+			for (int32 ExpansionIndex = 0; ExpansionIndex < UE_ARRAY_COUNT(GConfigLayerExpansions); ExpansionIndex++)
 			{
 				const FConfigLayerExpansion& Expansion = GConfigLayerExpansions[ExpansionIndex];
 
@@ -3431,7 +3431,7 @@ void FConfigFile::AddStaticLayersToHierarchy(const TCHAR* InBaseIniName, const T
 				}
 
 				// check if we should be generating the cache key - only at the end of all expansions
-				bool bGenerateCacheKey = EnumHasAnyFlags(Layer.Flag, EConfigLayerFlags::GenerateCacheKey) && ExpansionIndex == ARRAY_COUNT(GConfigLayerExpansions);
+				bool bGenerateCacheKey = EnumHasAnyFlags(Layer.Flag, EConfigLayerFlags::GenerateCacheKey) && ExpansionIndex == UE_ARRAY_COUNT(GConfigLayerExpansions);
 				checkfSlow(!(bGenerateCacheKey && bHasPlatformTag), TEXT("EConfigLayerFlags::GenerateCacheKey shouldn't have a platform tag"));
 
 				const FDataDrivenPlatformInfoRegistry::FPlatformInfo& Info = FDataDrivenPlatformInfoRegistry::GetPlatformInfo(PlatformName);
@@ -4348,7 +4348,7 @@ private:
 			TEXT("SectionName")
 		};
 		using UnderType = __underlying_type(HistoryType);
-		static_assert(static_cast<UnderType>(HistoryType::Count) == ARRAY_COUNT(Strings), "");
+		static_assert(static_cast<UnderType>(HistoryType::Count) == UE_ARRAY_COUNT(Strings), "");
 		return Strings[static_cast<UnderType>(Type)];
 	}
 
diff --git a/Engine/Source/Runtime/Core/Private/Misc/EmbeddedCommunication.cpp b/Engine/Source/Runtime/Core/Private/Misc/EmbeddedCommunication.cpp
--- a/Engine/Source/Runtime/Core/Private/Misc/EmbeddedCommunication.cpp
+++ b/Engine/Source/Runtime/Core/Private/Misc/EmbeddedCommunication.cpp
@@ -274,7 +274,7 @@ bool FEmbeddedCommunication::IsAwakeForRendering()
 void FEmbeddedCommunication::RunOnGameThread(int Priority, TFunction<void()> Lambda)
 {
 #if BUILD_EMBEDDED_APP
-	check(Priority < ARRAY_COUNT(GEmbeddedQueues));
+	check(Priority < UE_ARRAY_COUNT(GEmbeddedQueues));
 
 	{
 		FScopeLock Lock(&GEmbeddedLock);
@@ -323,7 +323,7 @@ bool FEmbeddedCommunication::TickGameThread(float DeltaTime)
 		{
 			FScopeLock Lock(&GEmbeddedLock);
 
-			for (int Queue = 0; Queue < ARRAY_COUNT(GEmbeddedQueues); Queue++)
+			for (int Queue = 0; Queue < UE_ARRAY_COUNT(GEmbeddedQueues); Queue++)
 			{
 				if (GEmbeddedQueues[Queue].Dequeue(LambdaToCall))
 				{
diff --git a/Engine/Source/Runtime/Core/Private/Misc/FileHelper.cpp b/Engine/Source/Runtime/Core/Private/Misc/FileHelper.cpp
--- a/Engine/Source/Runtime/Core/Private/Misc/FileHelper.cpp
+++ b/Engine/Source/Runtime/Core/Private/Misc/FileHelper.cpp
@@ -224,7 +224,7 @@ bool FFileHelper::SaveStringToFile( const FString& String, const TCHAR* Filename
 	if( EncodingOptions == EEncodingOptions::ForceUTF8 )
 	{
 		UTF8CHAR UTF8BOM[] = { 0xEF, 0xBB, 0xBF };
-		Ar->Serialize( &UTF8BOM, ARRAY_COUNT(UTF8BOM) * sizeof(UTF8CHAR) );
+		Ar->Serialize( &UTF8BOM, UE_ARRAY_COUNT(UTF8BOM) * sizeof(UTF8CHAR) );
 
 		FTCHARToUTF8 UTF8String(*String, String.Len());
 		Ar->Serialize( (UTF8CHAR*)UTF8String.Get(), UTF8String.Length() * sizeof(UTF8CHAR) );
@@ -257,7 +257,7 @@ bool FFileHelper::SaveStringArrayToFile( const TArray<FString>& Lines, const TCH
 	int32 Length = 10;
 	for(const FString& Line : Lines)
 	{
-		Length += Line.Len() + ARRAY_COUNT(LINE_TERMINATOR);
+		Length += Line.Len() + UE_ARRAY_COUNT(LINE_TERMINATOR);
 	}
 	
 	FString CombinedString;
diff --git a/Engine/Source/Runtime/Core/Private/Misc/FrameRate.cpp b/Engine/Source/Runtime/Core/Private/Misc/FrameRate.cpp
--- a/Engine/Source/Runtime/Core/Private/Misc/FrameRate.cpp
+++ b/Engine/Source/Runtime/Core/Private/Misc/FrameRate.cpp
@@ -200,7 +200,7 @@ bool FFrameRate::ComputeGridSpacing(float PixelsPerSecond, double& OutMajorInter
 
 		// Showing hours, minutes or seconds
 		static const int32 DesirableBases[]  = { 1, 2, 5, 10, 30, 60 };
-		static const int32 NumDesirableBases = ARRAY_COUNT(DesirableBases);
+		static const int32 NumDesirableBases = UE_ARRAY_COUNT(DesirableBases);
 
 		const int32 Scale     = FMath::CeilToInt(DesiredMajorTickPx / PixelsPerSecond / TimeOrder);
 		const int32 BaseIndex = FMath::Min(Algo::LowerBound(DesirableBases, Scale), NumDesirableBases-1);
diff --git a/Engine/Source/Runtime/Core/Private/Misc/OutputDeviceAnsiError.cpp b/Engine/Source/Runtime/Core/Private/Misc/OutputDeviceAnsiError.cpp
--- a/Engine/Source/Runtime/Core/Private/Misc/OutputDeviceAnsiError.cpp
+++ b/Engine/Source/Runtime/Core/Private/Misc/OutputDeviceAnsiError.cpp
@@ -31,8 +31,8 @@ void FOutputDeviceAnsiError::Serialize( const TCHAR* Msg, ELogVerbosity::Type Ve
 		// First appError.
 		GIsCriticalError = 1;
 		UE_LOG(LogHAL, Error, TEXT("appError called: %s"), Msg );
-		FCString::Strncpy( GErrorHist, Msg, ARRAY_COUNT(GErrorHist) );
-		FCString::Strncat( GErrorHist, TEXT("\r\n\r\n"), ARRAY_COUNT(GErrorHist) );
+		FCString::Strncpy( GErrorHist, Msg, UE_ARRAY_COUNT(GErrorHist) );
+		FCString::Strncat( GErrorHist, TEXT("\r\n\r\n"), UE_ARRAY_COUNT(GErrorHist) );
 		ErrorPos = FCString::Strlen(GErrorHist);
 	}
 	else
@@ -76,7 +76,7 @@ void FOutputDeviceAnsiError::HandleError()
 	GIsRunning			= 0;
 	GIsCriticalError	= 1;
 	GLogConsole			= NULL;
-	GErrorHist[ARRAY_COUNT(GErrorHist)-1]=0;
+	GErrorHist[UE_ARRAY_COUNT(GErrorHist)-1]=0;
 
 	if (GLog)
 	{
diff --git a/Engine/Source/Runtime/Core/Private/Misc/OutputDeviceFile.cpp b/Engine/Source/Runtime/Core/Private/Misc/OutputDeviceFile.cpp
--- a/Engine/Source/Runtime/Core/Private/Misc/OutputDeviceFile.cpp
+++ b/Engine/Source/Runtime/Core/Private/Misc/OutputDeviceFile.cpp
@@ -297,7 +297,7 @@ FOutputDeviceFile::FOutputDeviceFile( const TCHAR* InFilename, bool bInDisableBa
 {
 	if( InFilename )
 	{
-		FCString::Strncpy( Filename, InFilename, ARRAY_COUNT(Filename) );
+		FCString::Strncpy( Filename, InFilename, UE_ARRAY_COUNT(Filename) );
 	}
 	else
 	{
@@ -319,7 +319,7 @@ void FOutputDeviceFile::SetFilename(const TCHAR* InFilename)
 	// Close any existing file.
 	TearDown();
 
-	FCString::Strncpy( Filename, InFilename, ARRAY_COUNT(Filename) );
+	FCString::Strncpy( Filename, InFilename, UE_ARRAY_COUNT(Filename) );
 }
 
 /**
@@ -421,7 +421,7 @@ bool FOutputDeviceFile::CreateWriter(uint32 MaxAttempts)
 			{
 				CreateBackupCopy(*FinalFilename);
 			}
-			FCString::Strcpy(Filename, ARRAY_COUNT(Filename), *FinalFilename);
+			FCString::Strcpy(Filename, UE_ARRAY_COUNT(Filename), *FinalFilename);
 			Ar = IFileManager::Get().CreateFileWriter(*FinalFilename, WriteFlags);
 		} while (!Ar && FileIndex < MaxAttempts);
 	}
diff --git a/Engine/Source/Runtime/Core/Private/Misc/Parse.cpp b/Engine/Source/Runtime/Core/Private/Misc/Parse.cpp
--- a/Engine/Source/Runtime/Core/Private/Misc/Parse.cpp
+++ b/Engine/Source/Runtime/Core/Private/Misc/Parse.cpp
@@ -320,7 +320,7 @@ bool FParse::Value( const TCHAR* Stream, const TCHAR* Match, FString& Value, boo
 	if (StreamLen < 4096)
 	{
 		TCHAR Temp[4096]=TEXT("");
-		if (FParse::Value(Stream, Match, Temp, ARRAY_COUNT(Temp), bShouldStopOnSeparator))
+		if (FParse::Value(Stream, Match, Temp, UE_ARRAY_COUNT(Temp), bShouldStopOnSeparator))
 		{
 			Value = Temp;
 			return true;
@@ -520,7 +520,7 @@ bool FParse::Value( const TCHAR* Stream, const TCHAR* Match, uint64& Value )
 bool FParse::Value( const TCHAR* Stream, const TCHAR* Match, int64& Value )
 {
 	TCHAR Temp[4096]=TEXT(""), *Ptr=Temp;
-	if( FParse::Value( Stream, Match, Temp, ARRAY_COUNT(Temp) ) )
+	if( FParse::Value( Stream, Match, Temp, UE_ARRAY_COUNT(Temp) ) )
 	{
 		Value = 0;
 		bool Negative = (*Ptr=='-');
@@ -667,7 +667,7 @@ bool FParse::Bool( const TCHAR* Stream, const TCHAR* Match, bool& OnOff )
 bool FParse::Value( const TCHAR* Stream, const TCHAR* Match, struct FGuid& Guid )
 {
 	TCHAR Temp[256];
-	if( !FParse::Value( Stream, Match, Temp, ARRAY_COUNT(Temp) ) )
+	if( !FParse::Value( Stream, Match, Temp, UE_ARRAY_COUNT(Temp) ) )
 		return false;
 
 	Guid.A = Guid.B = Guid.C = Guid.D = 0;
@@ -916,7 +916,7 @@ bool FParse::Token( const TCHAR*& Str, FString& Arg, bool UseEscape )
 FString FParse::Token( const TCHAR*& Str, bool UseEscape )
 {
 	TCHAR Buffer[1024];
-	if( FParse::Token( Str, Buffer, ARRAY_COUNT(Buffer), UseEscape ) )
+	if( FParse::Token( Str, Buffer, UE_ARRAY_COUNT(Buffer), UseEscape ) )
 		return Buffer;
 	else
 		return TEXT("");
@@ -1267,7 +1267,7 @@ bool FParseLineExtendedTest::RunTest(const FString& Parameters)
 	int32 LinesConsumed = 0;
 	FString Result;
 
-	for (int32 Index = 0; Index < ARRAY_COUNT(Tests); ++Index)
+	for (int32 Index = 0; Index < UE_ARRAY_COUNT(Tests); ++Index)
 	{
 		LinesConsumed = 0;
 		Result.Reset();
diff --git a/Engine/Source/Runtime/Core/Private/Misc/Paths.cpp b/Engine/Source/Runtime/Core/Private/Misc/Paths.cpp
--- a/Engine/Source/Runtime/Core/Private/Misc/Paths.cpp
+++ b/Engine/Source/Runtime/Core/Private/Misc/Paths.cpp
@@ -993,7 +993,7 @@ void FPaths::NormalizeDirectoryName(FString& InPath)
 bool FPaths::CollapseRelativeDirectories(FString& InPath)
 {
 	const TCHAR ParentDir[] = TEXT("/..");
-	const int32 ParentDirLength = ARRAY_COUNT( ParentDir ) - 1; // To avoid hardcoded values
+	const int32 ParentDirLength = UE_ARRAY_COUNT( ParentDir ) - 1; // To avoid hardcoded values
 
 	for (;;)
 	{
diff --git a/Engine/Source/Runtime/Core/Private/Misc/SecureHash.cpp b/Engine/Source/Runtime/Core/Private/Misc/SecureHash.cpp
--- a/Engine/Source/Runtime/Core/Private/Misc/SecureHash.cpp
+++ b/Engine/Source/Runtime/Core/Private/Misc/SecureHash.cpp
@@ -799,27 +799,27 @@ void FSHA1::HMACBuffer(const void* Key, uint32 KeySize, const void* Data, uint64
 
 
 	// Start concatenating/hashing the pads/data etc: Hash(OKeyPad + Hash(IKeyPad + Data))
-	uint8* IKeyPad_Data = new uint8[ARRAY_COUNT(IKeyPad) + DataSize];
+	uint8* IKeyPad_Data = new uint8[UE_ARRAY_COUNT(IKeyPad) + DataSize];
 
-	FMemory::Memcpy(IKeyPad_Data, IKeyPad, ARRAY_COUNT(IKeyPad));
-	FMemory::Memcpy(IKeyPad_Data + ARRAY_COUNT(IKeyPad), Data, DataSize);
+	FMemory::Memcpy(IKeyPad_Data, IKeyPad, UE_ARRAY_COUNT(IKeyPad));
+	FMemory::Memcpy(IKeyPad_Data + UE_ARRAY_COUNT(IKeyPad), Data, DataSize);
 
 
 	uint8 IKeyPad_Data_Hash[HashSize];
 
-	HashBuffer(IKeyPad_Data, ARRAY_COUNT(IKeyPad) + DataSize, IKeyPad_Data_Hash);
+	HashBuffer(IKeyPad_Data, UE_ARRAY_COUNT(IKeyPad) + DataSize, IKeyPad_Data_Hash);
 
 	delete[] IKeyPad_Data;
 
 
-	uint8 OKeyPad_IHash[ARRAY_COUNT(OKeyPad) + HashSize];
+	uint8 OKeyPad_IHash[UE_ARRAY_COUNT(OKeyPad) + HashSize];
 
-	FMemory::Memcpy(OKeyPad_IHash, OKeyPad, ARRAY_COUNT(OKeyPad));
-	FMemory::Memcpy(OKeyPad_IHash + ARRAY_COUNT(OKeyPad), IKeyPad_Data_Hash, HashSize);
+	FMemory::Memcpy(OKeyPad_IHash, OKeyPad, UE_ARRAY_COUNT(OKeyPad));
+	FMemory::Memcpy(OKeyPad_IHash + UE_ARRAY_COUNT(OKeyPad), IKeyPad_Data_Hash, HashSize);
 
 
 	// Output the final hash
-	HashBuffer(OKeyPad_IHash, ARRAY_COUNT(OKeyPad_IHash), OutHash);
+	HashBuffer(OKeyPad_IHash, UE_ARRAY_COUNT(OKeyPad_IHash), OutHash);
 }
 
 
diff --git a/Engine/Source/Runtime/Core/Private/Misc/TextFilterUtils.cpp b/Engine/Source/Runtime/Core/Private/Misc/TextFilterUtils.cpp
--- a/Engine/Source/Runtime/Core/Private/Misc/TextFilterUtils.cpp
+++ b/Engine/Source/Runtime/Core/Private/Misc/TextFilterUtils.cpp
@@ -208,7 +208,7 @@ namespace FastToUpper
 		0,		// ~
 		0,		// 	
 	};
-	static const int32 ToUpperAdjustmentTableCount = ARRAY_COUNT(ToUpperAdjustmentTable);
+	static const int32 ToUpperAdjustmentTableCount = UE_ARRAY_COUNT(ToUpperAdjustmentTable);
 
 	FORCEINLINE TCHAR ToUpper(const TCHAR InChar)
 	{
diff --git a/Engine/Source/Runtime/Core/Private/Misc/VarargsHelper.h b/Engine/Source/Runtime/Core/Private/Misc/VarargsHelper.h
--- a/Engine/Source/Runtime/Core/Private/Misc/VarargsHelper.h
+++ b/Engine/Source/Runtime/Core/Private/Misc/VarargsHelper.h
@@ -20,7 +20,7 @@
 \
 	/* first, try using the stack buffer */ \
 	Buffer = StackBuffer; \
-	GET_VARARGS_RESULT( Buffer, ARRAY_COUNT(StackBuffer), ARRAY_COUNT(StackBuffer) - 1, Fmt, Fmt, Result ); \
+	GET_VARARGS_RESULT( Buffer, UE_ARRAY_COUNT(StackBuffer), UE_ARRAY_COUNT(StackBuffer) - 1, Fmt, Fmt, Result ); \
 \
 	/* if that fails, then use heap allocation to make enough space */ \
 	while(Result == -1) \
diff --git a/Engine/Source/Runtime/Core/Private/ProfilingDebugging/CsvProfiler.cpp b/Engine/Source/Runtime/Core/Private/ProfilingDebugging/CsvProfiler.cpp
--- a/Engine/Source/Runtime/Core/Private/ProfilingDebugging/CsvProfiler.cpp
+++ b/Engine/Source/Runtime/Core/Private/ProfilingDebugging/CsvProfiler.cpp
@@ -2938,7 +2938,7 @@ void FCsvProfiler::RecordEventfInternal(int32 CategoryIndex, const TCHAR* Fmt, .
 	{
 		LLM_SCOPE(ELLMTag::CsvProfiler);
 		TCHAR Buffer[256];
-		GET_VARARGS(Buffer, ARRAY_COUNT(Buffer), ARRAY_COUNT(Buffer) - 1, Fmt, Fmt);
+		GET_VARARGS(Buffer, UE_ARRAY_COUNT(Buffer), UE_ARRAY_COUNT(Buffer) - 1, Fmt, Fmt);
 		Buffer[255] = '\0';
 		FString Str = Buffer;
 		RecordEvent(CategoryIndex, Str);
diff --git a/Engine/Source/Runtime/Core/Private/Tests/Internationalization/TextTest.cpp b/Engine/Source/Runtime/Core/Private/Tests/Internationalization/TextTest.cpp
--- a/Engine/Source/Runtime/Core/Private/Tests/Internationalization/TextTest.cpp
+++ b/Engine/Source/Runtime/Core/Private/Tests/Internationalization/TextTest.cpp
@@ -675,7 +675,7 @@ bool FTextRoundingTest::RunTest (const FString& Parameters)
 		TEXT("ToPositiveInfinity"),
 	};
 
-	static_assert(ERoundingMode::ToPositiveInfinity == ARRAY_COUNT(RoundingModeNames) - 1, "RoundingModeNames array needs updating");
+	static_assert(ERoundingMode::ToPositiveInfinity == UE_ARRAY_COUNT(RoundingModeNames) - 1, "RoundingModeNames array needs updating");
 
 	static const double InputValues[] = {
 		1000.1224,
@@ -707,7 +707,7 @@ bool FTextRoundingTest::RunTest (const FString& Parameters)
 		-1000.1256,
 	};
 
-	static const TCHAR* OutputValues[][ARRAY_COUNT(RoundingModeNames)] = 
+	static const TCHAR* OutputValues[][UE_ARRAY_COUNT(RoundingModeNames)] = 
 	{
 		// HalfToEven        | HalfFromZero      | HalfToZero        | FromZero          | ToZero            | ToNegativeInfinity | ToPositiveInfinity
 		{  TEXT("1000.122"),   TEXT("1000.122"),   TEXT("1000.122"),   TEXT("1000.123"),   TEXT("1000.122"),   TEXT("1000.122"),    TEXT("1000.123") },
@@ -739,7 +739,7 @@ bool FTextRoundingTest::RunTest (const FString& Parameters)
 		{ TEXT("-1000.126"),  TEXT("-1000.126"),  TEXT("-1000.126"),  TEXT("-1000.126"),  TEXT("-1000.125"),  TEXT("-1000.126"),   TEXT("-1000.125") },
 	};
 
-	static_assert(ARRAY_COUNT(InputValues) == ARRAY_COUNT(OutputValues), "The size of InputValues does not match OutputValues");
+	static_assert(UE_ARRAY_COUNT(InputValues) == UE_ARRAY_COUNT(OutputValues), "The size of InputValues does not match OutputValues");
 
 	FInternationalization& I18N = FInternationalization::Get();
 
@@ -767,7 +767,7 @@ bool FTextRoundingTest::RunTest (const FString& Parameters)
 	{
 		FormattingOptions.SetRoundingMode(InRoundingMode);
 
-		for (int32 TestValueIndex = 0; TestValueIndex < ARRAY_COUNT(InputValues); ++TestValueIndex)
+		for (int32 TestValueIndex = 0; TestValueIndex < UE_ARRAY_COUNT(InputValues); ++TestValueIndex)
 		{
 			DoSingleTest(InputValues[TestValueIndex], OutputValues[TestValueIndex][InRoundingMode], RoundingModeNames[InRoundingMode]);
 		}
diff --git a/Engine/Source/Runtime/Core/Private/Tests/Math/UnrealMathTest.cpp b/Engine/Source/Runtime/Core/Private/Tests/Math/UnrealMathTest.cpp
--- a/Engine/Source/Runtime/Core/Private/Tests/Math/UnrealMathTest.cpp
+++ b/Engine/Source/Runtime/Core/Private/Tests/Math/UnrealMathTest.cpp
@@ -1352,10 +1352,10 @@ bool FVectorRegisterAbstractionTest::RunTest(const FString& Parameters)
 		}
 
 		// Test loop
-		for (int32 QIndex = 0; QIndex < ARRAY_COUNT(TestQuats); ++QIndex)
+		for (int32 QIndex = 0; QIndex < UE_ARRAY_COUNT(TestQuats); ++QIndex)
 		{
 			const FQuat& Q = TestQuats[QIndex];
-			for (int32 VIndex = 0; VIndex < ARRAY_COUNT(TestVectors); ++VIndex)
+			for (int32 VIndex = 0; VIndex < UE_ARRAY_COUNT(TestVectors); ++VIndex)
 			{
 				const FVector& V = TestVectors[VIndex];
 				TEST_QUAT_ROTATE(QIndex, VIndex, Q, V, TestQuaternionRotateVectorScalar, Tolerance);
diff --git a/Engine/Source/Runtime/Core/Private/Tests/Misc/AtomicTest.cpp b/Engine/Source/Runtime/Core/Private/Tests/Misc/AtomicTest.cpp
--- a/Engine/Source/Runtime/Core/Private/Tests/Misc/AtomicTest.cpp
+++ b/Engine/Source/Runtime/Core/Private/Tests/Misc/AtomicTest.cpp
@@ -240,7 +240,7 @@ namespace
 			constexpr ElementType ExpectedValues[] = {GetByteFilledElement<ElementType>(0x20), Init, ElementType(0), };
 
 			// operator &=
-			for (SIZE_T TestIt = 0; TestIt < ARRAY_COUNT(AndValues); ++TestIt)
+			for (SIZE_T TestIt = 0; TestIt < UE_ARRAY_COUNT(AndValues); ++TestIt)
 			{
 				TAtomicTestWrapper<ElementType> Data(Init);
 
@@ -253,7 +253,7 @@ namespace
 			}
 
 			// AndExchange
-			for (SIZE_T TestIt = 0; TestIt < ARRAY_COUNT(AndValues); ++TestIt)
+			for (SIZE_T TestIt = 0; TestIt < UE_ARRAY_COUNT(AndValues); ++TestIt)
 			{
 				TAtomicTestWrapper<ElementType> Data(Init);
 
@@ -276,7 +276,7 @@ namespace
 			constexpr ElementType ExpectedValues[] = {GetByteFilledElement<ElementType>(0x76), GetByteFilledElement<ElementType>(0xFF), Init, };
 
 			// operator |=
-			for (SIZE_T TestIt = 0; TestIt < ARRAY_COUNT(OrValues); ++TestIt)
+			for (SIZE_T TestIt = 0; TestIt < UE_ARRAY_COUNT(OrValues); ++TestIt)
 			{
 				TAtomicTestWrapper<ElementType> Data(Init);
 
@@ -289,7 +289,7 @@ namespace
 			}
 
 			// OrExchange
-			for (SIZE_T TestIt = 0; TestIt < ARRAY_COUNT(OrValues); ++TestIt)
+			for (SIZE_T TestIt = 0; TestIt < UE_ARRAY_COUNT(OrValues); ++TestIt)
 			{
 				TAtomicTestWrapper<ElementType> Data(Init);
 
@@ -312,7 +312,7 @@ namespace
 			constexpr ElementType ExpectedValues[] = {GetByteFilledElement<ElementType>(0x56), GetByteFilledElement<ElementType>(~0x30), Init, };
 
 			// operator ^=
-			for (SIZE_T TestIt = 0; TestIt < ARRAY_COUNT(XorValues); ++TestIt)
+			for (SIZE_T TestIt = 0; TestIt < UE_ARRAY_COUNT(XorValues); ++TestIt)
 			{
 				TAtomicTestWrapper<ElementType> Data(Init);
 
@@ -325,7 +325,7 @@ namespace
 			}
 
 			// XorExchange
-			for (SIZE_T TestIt = 0; TestIt < ARRAY_COUNT(XorValues); ++TestIt)
+			for (SIZE_T TestIt = 0; TestIt < UE_ARRAY_COUNT(XorValues); ++TestIt)
 			{
 				TAtomicTestWrapper<ElementType> Data(Init);
 
diff --git a/Engine/Source/Runtime/Core/Private/UObject/UnrealNames.cpp b/Engine/Source/Runtime/Core/Private/UObject/UnrealNames.cpp
--- a/Engine/Source/Runtime/Core/Private/UObject/UnrealNames.cpp
+++ b/Engine/Source/Runtime/Core/Private/UObject/UnrealNames.cpp
@@ -1360,7 +1360,7 @@ FString FName::NameToDisplayString( const FString& InDisplayName, const bool bIs
 
 			// Search for a word that needs case repaired
 			bool bIsArticle = false;
-			for( int32 CurArticleIndex = 0; CurArticleIndex < ARRAY_COUNT( Articles ); ++CurArticleIndex )
+			for( int32 CurArticleIndex = 0; CurArticleIndex < UE_ARRAY_COUNT( Articles ); ++CurArticleIndex )
 			{
 				// Make sure the character following the string we're testing is not lowercase (we don't want to match "in" with "instance")
 				const int32 ArticleLength = FCString::Strlen( Articles[ CurArticleIndex ] );
diff --git a/Engine/Source/Runtime/Core/Private/Unix/UnixErrorOutputDevice.cpp b/Engine/Source/Runtime/Core/Private/Unix/UnixErrorOutputDevice.cpp
--- a/Engine/Source/Runtime/Core/Private/Unix/UnixErrorOutputDevice.cpp
+++ b/Engine/Source/Runtime/Core/Private/Unix/UnixErrorOutputDevice.cpp
@@ -41,8 +41,8 @@ void FUnixErrorOutputDevice::Serialize(const TCHAR* Msg, ELogVerbosity::Type Ver
 		{
 			UE_LOG(LogCore, Error, TEXT("appError called: %s"), Msg );
 		}
-		FCString::Strncpy( GErrorHist, Msg, ARRAY_COUNT(GErrorHist) - 5 );
-		FCString::Strncat( GErrorHist, TEXT("\r\n\r\n"), ARRAY_COUNT(GErrorHist) - 1 );
+		FCString::Strncpy( GErrorHist, Msg, UE_ARRAY_COUNT(GErrorHist) - 5 );
+		FCString::Strncat( GErrorHist, TEXT("\r\n\r\n"), UE_ARRAY_COUNT(GErrorHist) - 1 );
 		ErrorPos = FCString::Strlen(GErrorHist);
 	}
 	else
@@ -100,7 +100,7 @@ void FUnixErrorOutputDevice::HandleError()
 		GIsRunning = 0;
 		GIsCriticalError = 1;
 		GLogConsole = NULL;
-		GErrorHist[ARRAY_COUNT(GErrorHist)-1] = 0;
+		GErrorHist[UE_ARRAY_COUNT(GErrorHist)-1] = 0;
 
 		// Dump the error and flush the log.
 		UE_LOG(LogCore, Log, TEXT("=== Critical error: ===") LINE_TERMINATOR TEXT("%s") LINE_TERMINATOR, GErrorExceptionDescription);
diff --git a/Engine/Source/Runtime/Core/Private/Unix/UnixPlatformCrashContext.cpp b/Engine/Source/Runtime/Core/Private/Unix/UnixPlatformCrashContext.cpp
--- a/Engine/Source/Runtime/Core/Private/Unix/UnixPlatformCrashContext.cpp
+++ b/Engine/Source/Runtime/Core/Private/Unix/UnixPlatformCrashContext.cpp
@@ -101,7 +101,7 @@ void FUnixCrashContext::InitFromSignal(int32 InSignal, siginfo_t* InInfo, void*
 	Info = InInfo;
 	Context = reinterpret_cast< ucontext_t* >( InContext );
 
-	FCString::Strcat(SignalDescription, ARRAY_COUNT( SignalDescription ) - 1, *DescribeSignal(Signal, Info, Context));
+	FCString::Strcat(SignalDescription, UE_ARRAY_COUNT( SignalDescription ) - 1, *DescribeSignal(Signal, Info, Context));
 }
 
 void FUnixCrashContext::InitFromEnsureHandler(const TCHAR* EnsureMessage, const void* CrashAddress)
@@ -116,10 +116,10 @@ void FUnixCrashContext::InitFromEnsureHandler(const TCHAR* EnsureMessage, const
 	Context = nullptr;
 
 	// set signal description to a more human-readable one for ensures
-	FCString::Strcpy(SignalDescription, ARRAY_COUNT(SignalDescription) - 1, EnsureMessage);
+	FCString::Strcpy(SignalDescription, UE_ARRAY_COUNT(SignalDescription) - 1, EnsureMessage);
 
 	// only need the first string
-	for (int Idx = 0; Idx < ARRAY_COUNT(SignalDescription); ++Idx)
+	for (int Idx = 0; Idx < UE_ARRAY_COUNT(SignalDescription); ++Idx)
 	{
 		if (SignalDescription[Idx] == TEXT('\n'))
 		{
@@ -155,7 +155,7 @@ void CreateExceptionInfoString(int32 Signal, siginfo_t* Info, ucontext_t *Contex
 {
 	FString ErrorString = TEXT("Unhandled Exception: ");
 	ErrorString += DescribeSignal(Signal, Info, Context);
-	FCString::Strncpy(GErrorExceptionDescription, *ErrorString, FMath::Min(ErrorString.Len() + 1, (int32)ARRAY_COUNT(GErrorExceptionDescription)));
+	FCString::Strncpy(GErrorExceptionDescription, *ErrorString, FMath::Min(ErrorString.Len() + 1, (int32)UE_ARRAY_COUNT(GErrorExceptionDescription)));
 }
 
 namespace
@@ -303,7 +303,7 @@ void FUnixCrashContext::CaptureStackTrace()
 		printf("StackTrace:\n%s\n", StackTrace);
 #endif
 
-		FCString::Strncat( GErrorHist, UTF8_TO_TCHAR(StackTrace), ARRAY_COUNT(GErrorHist) - 1 );
+		FCString::Strncat( GErrorHist, UTF8_TO_TCHAR(StackTrace), UE_ARRAY_COUNT(GErrorHist) - 1 );
 		CreateExceptionInfoString(Signal, Info, Context);
 
 		FMemory::Free( StackTrace );
diff --git a/Engine/Source/Runtime/Core/Private/Unix/UnixPlatformMemory.cpp b/Engine/Source/Runtime/Core/Private/Unix/UnixPlatformMemory.cpp
--- a/Engine/Source/Runtime/Core/Private/Unix/UnixPlatformMemory.cpp
+++ b/Engine/Source/Runtime/Core/Private/Unix/UnixPlatformMemory.cpp
@@ -666,7 +666,7 @@ FPlatformMemoryStats FUnixPlatformMemory::GetStats()
 		do
 		{
 			char LineBuffer[256] = {0};
-			char *Line = fgets(LineBuffer, ARRAY_COUNT(LineBuffer), FileGlobalMemStats);
+			char *Line = fgets(LineBuffer, UE_ARRAY_COUNT(LineBuffer), FileGlobalMemStats);
 			if (Line == nullptr)
 			{
 				break;	// eof or an error
@@ -712,7 +712,7 @@ FPlatformMemoryStats FUnixPlatformMemory::GetStats()
 		do
 		{
 			char LineBuffer[256] = {0};
-			char *Line = fgets(LineBuffer, ARRAY_COUNT(LineBuffer), ProcMemStats);
+			char *Line = fgets(LineBuffer, UE_ARRAY_COUNT(LineBuffer), ProcMemStats);
 			if (Line == nullptr)
 			{
 				break;	// eof or an error
@@ -767,7 +767,7 @@ FExtendedPlatformMemoryStats FUnixPlatformMemory::GetExtendedStats()
 		do
 		{
 			char LineBuffer[256] = { 0 };
-			char *Line = fgets(LineBuffer, ARRAY_COUNT(LineBuffer), ProcSMaps);
+			char *Line = fgets(LineBuffer, UE_ARRAY_COUNT(LineBuffer), ProcSMaps);
 			if (Line == nullptr)
 			{
 				break;	// eof or an error
diff --git a/Engine/Source/Runtime/Core/Private/Unix/UnixPlatformMisc.cpp b/Engine/Source/Runtime/Core/Private/Unix/UnixPlatformMisc.cpp
--- a/Engine/Source/Runtime/Core/Private/Unix/UnixPlatformMisc.cpp
+++ b/Engine/Source/Runtime/Core/Private/Unix/UnixPlatformMisc.cpp
@@ -601,7 +601,7 @@ FString FUnixPlatformMisc::GetCPUVendor()
 
 		VendorResult.Buffer[12] = 0;
 
-		FCString::Strncpy(Result, UTF8_TO_TCHAR(VendorResult.Buffer), ARRAY_COUNT(Result));
+		FCString::Strncpy(Result, UTF8_TO_TCHAR(VendorResult.Buffer), UE_ARRAY_COUNT(Result));
 #else
 		// use /proc?
 #endif // PLATFORM_HAS_CPUID
@@ -657,7 +657,7 @@ FString FUnixPlatformMisc::GetCPUBrand()
 			}
 		}
 
-		FCString::Strncpy(Result, UTF8_TO_TCHAR(BrandString), ARRAY_COUNT(Result));
+		FCString::Strncpy(Result, UTF8_TO_TCHAR(BrandString), UE_ARRAY_COUNT(Result));
 #else
 		// use /proc?
 #endif // PLATFORM_HAS_CPUID
diff --git a/Engine/Source/Runtime/Core/Private/Unix/UnixPlatformProcess.cpp b/Engine/Source/Runtime/Core/Private/Unix/UnixPlatformProcess.cpp
--- a/Engine/Source/Runtime/Core/Private/Unix/UnixPlatformProcess.cpp
+++ b/Engine/Source/Runtime/Core/Private/Unix/UnixPlatformProcess.cpp
@@ -141,8 +141,8 @@ const TCHAR* FUnixPlatformProcess::ComputerName()
 			SysName = "Unix Computer";
 		}
 
-		FCString::Strcpy(CachedResult, ARRAY_COUNT(CachedResult) - 1, UTF8_TO_TCHAR(SysName));
-		CachedResult[ARRAY_COUNT(CachedResult) - 1] = 0;
+		FCString::Strcpy(CachedResult, UE_ARRAY_COUNT(CachedResult) - 1, UTF8_TO_TCHAR(SysName));
+		CachedResult[UE_ARRAY_COUNT(CachedResult) - 1] = 0;
 		bHaveResult = true;
 	}
 	return CachedResult;
@@ -163,7 +163,7 @@ const TCHAR* FUnixPlatformProcess::UserName(bool bOnlyAlphaNumeric)
 			{
 				const TCHAR *Src = *TempName;
 				TCHAR * Dst = Name;
-				for (; *Src != 0 && (Dst - Name) < ARRAY_COUNT(Name) - 1; ++Src)
+				for (; *Src != 0 && (Dst - Name) < UE_ARRAY_COUNT(Name) - 1; ++Src)
 				{
 					if (FChar::IsAlnum(*Src))
 					{
@@ -174,7 +174,7 @@ const TCHAR* FUnixPlatformProcess::UserName(bool bOnlyAlphaNumeric)
 			}
 			else
 			{
-				FCString::Strncpy(Name, *TempName, ARRAY_COUNT(Name) - 1);
+				FCString::Strncpy(Name, *TempName, UE_ARRAY_COUNT(Name) - 1);
 			}
 		}
 		else
@@ -207,8 +207,8 @@ const TCHAR* FUnixPlatformProcess::UserDir()
 				if (DocLen > 0)
 				{
 					DocPath[DocLen] = '\0';
-					FCString::Strncpy(Result, UTF8_TO_TCHAR(DocPath), ARRAY_COUNT(Result));
-					FCString::Strncat(Result, TEXT("/"), ARRAY_COUNT(Result));
+					FCString::Strncpy(Result, UTF8_TO_TCHAR(DocPath), UE_ARRAY_COUNT(Result));
+					FCString::Strncat(Result, TEXT("/"), UE_ARRAY_COUNT(Result));
 				}
 			}
 			pclose(FilePtr);
@@ -217,8 +217,8 @@ const TCHAR* FUnixPlatformProcess::UserDir()
 		// if xdg-user-dir did not work, use $HOME
 		if (!Result[0])
 		{
-			FCString::Strncpy(Result, FPlatformProcess::UserHomeDir(), ARRAY_COUNT(Result));
-			FCString::Strncat(Result, TEXT("/Documents/"), ARRAY_COUNT(Result));
+			FCString::Strncpy(Result, FPlatformProcess::UserHomeDir(), UE_ARRAY_COUNT(Result));
+			FCString::Strncat(Result, TEXT("/Documents/"), UE_ARRAY_COUNT(Result));
 		}
 	}
 	return Result;
@@ -235,7 +235,7 @@ const TCHAR* FUnixPlatformProcess::UserHomeDir()
 		const char * VarValue = secure_getenv("HOME");
 		if (VarValue)
 		{
-			FCString::Strcpy(CachedResult, ARRAY_COUNT(CachedResult) - 1, UTF8_TO_TCHAR(VarValue));
+			FCString::Strcpy(CachedResult, UE_ARRAY_COUNT(CachedResult) - 1, UTF8_TO_TCHAR(VarValue));
 			bHaveHome = true;
 		}
 		else
@@ -243,7 +243,7 @@ const TCHAR* FUnixPlatformProcess::UserHomeDir()
 			struct passwd * UserInfo = getpwuid(geteuid());
 			if (NULL != UserInfo && NULL != UserInfo->pw_dir)
 			{
-				FCString::Strcpy(CachedResult, ARRAY_COUNT(CachedResult) - 1, UTF8_TO_TCHAR(UserInfo->pw_dir));
+				FCString::Strcpy(CachedResult, UE_ARRAY_COUNT(CachedResult) - 1, UTF8_TO_TCHAR(UserInfo->pw_dir));
 				bHaveHome = true;
 			}
 			else
@@ -271,8 +271,8 @@ const TCHAR* FUnixPlatformProcess::ApplicationSettingsDir()
 	static TCHAR Result[UNIX_MAX_PATH] = TEXT("");
 	if (!Result[0])
 	{
-		FCString::Strncpy(Result, FPlatformProcess::UserHomeDir(), ARRAY_COUNT(Result));
-		FCString::Strncat(Result, TEXT("/.config/Epic/"), ARRAY_COUNT(Result));
+		FCString::Strncpy(Result, FPlatformProcess::UserHomeDir(), UE_ARRAY_COUNT(Result));
+		FCString::Strncat(Result, TEXT("/.config/Epic/"), UE_ARRAY_COUNT(Result));
 	}
 	return Result;
 }
@@ -327,17 +327,17 @@ const TCHAR* FUnixPlatformProcess::ExecutablePath()
 	if (!bHaveResult)
 	{
 		char SelfPath[ PlatformProcessLimits::MaxBaseDirLength ] = {0};
-		if (readlink( "/proc/self/exe", SelfPath, ARRAY_COUNT(SelfPath) - 1) == -1)
+		if (readlink( "/proc/self/exe", SelfPath, UE_ARRAY_COUNT(SelfPath) - 1) == -1)
 		{
 			int ErrNo = errno;
 			UE_LOG(LogHAL, Fatal, TEXT("readlink() failed with errno = %d (%s)"), ErrNo,
 				StringCast< TCHAR >(strerror(ErrNo)).Get());
 			return CachedResult;
 		}
-		SelfPath[ARRAY_COUNT(SelfPath) - 1] = 0;
+		SelfPath[UE_ARRAY_COUNT(SelfPath) - 1] = 0;
 
-		FCString::Strcpy(CachedResult, ARRAY_COUNT(CachedResult) - 1, UTF8_TO_TCHAR(SelfPath));
-		CachedResult[ARRAY_COUNT(CachedResult) - 1] = 0;
+		FCString::Strcpy(CachedResult, UE_ARRAY_COUNT(CachedResult) - 1, UTF8_TO_TCHAR(SelfPath));
+		CachedResult[UE_ARRAY_COUNT(CachedResult) - 1] = 0;
 		bHaveResult = true;
 	}
 	return CachedResult;
@@ -350,17 +350,17 @@ const TCHAR* FUnixPlatformProcess::ExecutableName(bool bRemoveExtension)
 	if (!bHaveResult)
 	{
 		char SelfPath[ PlatformProcessLimits::MaxBaseDirLength ] = {0};
-		if (readlink( "/proc/self/exe", SelfPath, ARRAY_COUNT(SelfPath) - 1) == -1)
+		if (readlink( "/proc/self/exe", SelfPath, UE_ARRAY_COUNT(SelfPath) - 1) == -1)
 		{
 			int ErrNo = errno;
 			UE_LOG(LogHAL, Fatal, TEXT("readlink() failed with errno = %d (%s)"), ErrNo,
 				StringCast< TCHAR >(strerror(ErrNo)).Get());
 			return CachedResult;
 		}
-		SelfPath[ARRAY_COUNT(SelfPath) - 1] = 0;
+		SelfPath[UE_ARRAY_COUNT(SelfPath) - 1] = 0;
 
-		FCString::Strcpy(CachedResult, ARRAY_COUNT(CachedResult) - 1, UTF8_TO_TCHAR(basename(SelfPath)));
-		CachedResult[ARRAY_COUNT(CachedResult) - 1] = 0;
+		FCString::Strcpy(CachedResult, UE_ARRAY_COUNT(CachedResult) - 1, UTF8_TO_TCHAR(basename(SelfPath)));
+		CachedResult[UE_ARRAY_COUNT(CachedResult) - 1] = 0;
 		bHaveResult = true;
 	}
 	return CachedResult;
@@ -389,7 +389,7 @@ FString FUnixPlatformProcess::GetApplicationName( uint32 ProcessId )
 	FCStringAnsi::Sprintf(ReadLinkCmd, "/proc/%d/exe", ProcessId);
 	
 	char ProcessPath[ PlatformProcessLimits::MaxBaseDirLength ] = {0};
-	int32 Ret = readlink(ReadLinkCmd, ProcessPath, ARRAY_COUNT(ProcessPath) - 1);
+	int32 Ret = readlink(ReadLinkCmd, ProcessPath, UE_ARRAY_COUNT(ProcessPath) - 1);
 	if (Ret != -1)
 	{
 		Output = UTF8_TO_TCHAR(ProcessPath);
diff --git a/Engine/Source/Runtime/Core/Private/Unix/UnixPlatformStackWalk.cpp b/Engine/Source/Runtime/Core/Private/Unix/UnixPlatformStackWalk.cpp
--- a/Engine/Source/Runtime/Core/Private/Unix/UnixPlatformStackWalk.cpp
+++ b/Engine/Source/Runtime/Core/Private/Unix/UnixPlatformStackWalk.cpp
@@ -530,10 +530,10 @@ bool FUnixPlatformStackWalk::ProgramCounterToHumanReadableString( int32 CurrentC
 				if (UnixContext)
 				{
 					// append Module!FunctioName [Source.cpp:X] to MinidumpCallstackInfo
-					FCStringAnsi::Strncat(UnixContext->MinidumpCallstackInfo, TempSymbolInfo.ModuleName, ARRAY_COUNT( UnixContext->MinidumpCallstackInfo ) - 1);
-					FCStringAnsi::Strncat(UnixContext->MinidumpCallstackInfo, "!", ARRAY_COUNT( UnixContext->MinidumpCallstackInfo ) - 1);
-					FCStringAnsi::Strncat(UnixContext->MinidumpCallstackInfo, TempSymbolInfo.FunctionName, ARRAY_COUNT( UnixContext->MinidumpCallstackInfo ) - 1);
-					FCStringAnsi::Strncat(UnixContext->MinidumpCallstackInfo, TempArray, ARRAY_COUNT( UnixContext->MinidumpCallstackInfo ) - 1);
+					FCStringAnsi::Strncat(UnixContext->MinidumpCallstackInfo, TempSymbolInfo.ModuleName, UE_ARRAY_COUNT( UnixContext->MinidumpCallstackInfo ) - 1);
+					FCStringAnsi::Strncat(UnixContext->MinidumpCallstackInfo, "!", UE_ARRAY_COUNT( UnixContext->MinidumpCallstackInfo ) - 1);
+					FCStringAnsi::Strncat(UnixContext->MinidumpCallstackInfo, TempSymbolInfo.FunctionName, UE_ARRAY_COUNT( UnixContext->MinidumpCallstackInfo ) - 1);
+					FCStringAnsi::Strncat(UnixContext->MinidumpCallstackInfo, TempArray, UE_ARRAY_COUNT( UnixContext->MinidumpCallstackInfo ) - 1);
 				}
 			}
 			else
@@ -559,10 +559,10 @@ bool FUnixPlatformStackWalk::ProgramCounterToHumanReadableString( int32 CurrentC
 
 				if (UnixContext)
 				{
-					FCStringAnsi::Strncat(UnixContext->MinidumpCallstackInfo, ModuleName != nullptr ? ModuleName : "Unknown", ARRAY_COUNT( UnixContext->MinidumpCallstackInfo ) - 1);
-					FCStringAnsi::Strncat(UnixContext->MinidumpCallstackInfo, "!", ARRAY_COUNT( UnixContext->MinidumpCallstackInfo ) - 1);
-					FCStringAnsi::Strncat(UnixContext->MinidumpCallstackInfo, FunctionName != nullptr ? FunctionName : "UnknownFunction", ARRAY_COUNT( UnixContext->MinidumpCallstackInfo ) - 1);
-					FCStringAnsi::Strncat(UnixContext->MinidumpCallstackInfo, FunctionName && TempSymbolInfo.SymbolDisplacement ? "(+": "(", ARRAY_COUNT( UnixContext->MinidumpCallstackInfo ) - 1);
+					FCStringAnsi::Strncat(UnixContext->MinidumpCallstackInfo, ModuleName != nullptr ? ModuleName : "Unknown", UE_ARRAY_COUNT( UnixContext->MinidumpCallstackInfo ) - 1);
+					FCStringAnsi::Strncat(UnixContext->MinidumpCallstackInfo, "!", UE_ARRAY_COUNT( UnixContext->MinidumpCallstackInfo ) - 1);
+					FCStringAnsi::Strncat(UnixContext->MinidumpCallstackInfo, FunctionName != nullptr ? FunctionName : "UnknownFunction", UE_ARRAY_COUNT( UnixContext->MinidumpCallstackInfo ) - 1);
+					FCStringAnsi::Strncat(UnixContext->MinidumpCallstackInfo, FunctionName && TempSymbolInfo.SymbolDisplacement ? "(+": "(", UE_ARRAY_COUNT( UnixContext->MinidumpCallstackInfo ) - 1);
 				}
 
 				if (TempSymbolInfo.SymbolDisplacement > 0x0)
@@ -572,7 +572,7 @@ bool FUnixPlatformStackWalk::ProgramCounterToHumanReadableString( int32 CurrentC
 
 					if (UnixContext)
 					{
-						FCStringAnsi::Strncat(UnixContext->MinidumpCallstackInfo, TempArray, ARRAY_COUNT( UnixContext->MinidumpCallstackInfo ) - 1);
+						FCStringAnsi::Strncat(UnixContext->MinidumpCallstackInfo, TempArray, UE_ARRAY_COUNT( UnixContext->MinidumpCallstackInfo ) - 1);
 					}
 				}
 
@@ -580,13 +580,13 @@ bool FUnixPlatformStackWalk::ProgramCounterToHumanReadableString( int32 CurrentC
 
 				if (UnixContext)
 				{
-					FCStringAnsi::Strncat(UnixContext->MinidumpCallstackInfo, ")", ARRAY_COUNT( UnixContext->MinidumpCallstackInfo ) - 1);
+					FCStringAnsi::Strncat(UnixContext->MinidumpCallstackInfo, ")", UE_ARRAY_COUNT( UnixContext->MinidumpCallstackInfo ) - 1);
 				}
 			}
 
 			if (UnixContext)
 			{
-				FCStringAnsi::Strncat(UnixContext->MinidumpCallstackInfo, "\r\n", ARRAY_COUNT( UnixContext->MinidumpCallstackInfo ) - 1);	// this one always uses Windows line terminators
+				FCStringAnsi::Strncat(UnixContext->MinidumpCallstackInfo, "\r\n", UE_ARRAY_COUNT( UnixContext->MinidumpCallstackInfo ) - 1);	// this one always uses Windows line terminators
 			}
 		}
 		return true;
diff --git a/Engine/Source/Runtime/Core/Private/Unix/UnixPlatformTime.cpp b/Engine/Source/Runtime/Core/Private/Unix/UnixPlatformTime.cpp
--- a/Engine/Source/Runtime/Core/Private/Unix/UnixPlatformTime.cpp
+++ b/Engine/Source/Runtime/Core/Private/Unix/UnixPlatformTime.cpp
@@ -201,7 +201,7 @@ int FUnixTime::CalibrateAndSelectClock()
 		int ChosenClock = 0;	
 		// Skip measuring CLOCK_REALTIME, so it's never picked up if any other is available.
 		// CLOCK_REALTIME should be always supported - as a last resort.
-		for (int Idx = 1; Idx < ARRAY_COUNT(Clocks); ++Idx)
+		for (int Idx = 1; Idx < UE_ARRAY_COUNT(Clocks); ++Idx)
 		{
 			Clocks[Idx].Rate = CallsPerSecondBenchmark(Clocks[Idx].Id, Clocks[Idx].Desc);
 			if (Clocks[Idx].Rate > Clocks[ChosenClock].Rate)
diff --git a/Engine/Source/Runtime/Core/Private/Windows/WindowsErrorOutputDevice.cpp b/Engine/Source/Runtime/Core/Private/Windows/WindowsErrorOutputDevice.cpp
--- a/Engine/Source/Runtime/Core/Private/Windows/WindowsErrorOutputDevice.cpp
+++ b/Engine/Source/Runtime/Core/Private/Windows/WindowsErrorOutputDevice.cpp
@@ -99,7 +99,7 @@ void FWindowsErrorOutputDevice::HandleError()
 	GIsRunning				= 0;
 	GIsCriticalError		= 1;
 	GLogConsole				= NULL;
-	GErrorHist[ARRAY_COUNT(GErrorHist)-1]=0;
+	GErrorHist[UE_ARRAY_COUNT(GErrorHist)-1]=0;
 
 	// Trigger the OnSystemFailure hook if it exists
 	// make sure it happens after GIsGuarded is set to 0 in case this hook crashes
diff --git a/Engine/Source/Runtime/Core/Private/Windows/WindowsPlatformAtomics.cpp b/Engine/Source/Runtime/Core/Private/Windows/WindowsPlatformAtomics.cpp
--- a/Engine/Source/Runtime/Core/Private/Windows/WindowsPlatformAtomics.cpp
+++ b/Engine/Source/Runtime/Core/Private/Windows/WindowsPlatformAtomics.cpp
@@ -13,7 +13,7 @@ void FWindowsPlatformAtomics::HandleAtomicsFailure( const TCHAR* InFormat, ... )
 	va_list Ptr;
 
 	va_start( Ptr, InFormat );	
-	FCString::GetVarArgs( TempStr, ARRAY_COUNT(TempStr), InFormat, Ptr );
+	FCString::GetVarArgs( TempStr, UE_ARRAY_COUNT(TempStr), InFormat, Ptr );
 	va_end( Ptr );
 
 	UE_LOG(LogWindows, Log,  TempStr );
diff --git a/Engine/Source/Runtime/Core/Private/Windows/WindowsPlatformCrashContext.cpp b/Engine/Source/Runtime/Core/Private/Windows/WindowsPlatformCrashContext.cpp
--- a/Engine/Source/Runtime/Core/Private/Windows/WindowsPlatformCrashContext.cpp
+++ b/Engine/Source/Runtime/Core/Private/Windows/WindowsPlatformCrashContext.cpp
@@ -578,7 +578,7 @@ FORCENOINLINE void ReportAssert(const TCHAR* ErrorMessage, int NumStackFramesToI
 	FAssertInfo Info(ErrorMessage, NumStackFramesToIgnore + 2); // +2 for this function and RaiseException()
 
 	ULONG_PTR Arguments[] = { (ULONG_PTR)&Info };
-	::RaiseException( AssertExceptionCode, 0, ARRAY_COUNT(Arguments), Arguments );
+	::RaiseException( AssertExceptionCode, 0, UE_ARRAY_COUNT(Arguments), Arguments );
 }
 
 FORCENOINLINE void ReportGPUCrash(const TCHAR* ErrorMessage, int NumStackFramesToIgnore)
@@ -589,7 +589,7 @@ FORCENOINLINE void ReportGPUCrash(const TCHAR* ErrorMessage, int NumStackFramesT
 	FAssertInfo Info(ErrorMessage, NumStackFramesToIgnore + 2); // +2 for this function and RaiseException()
 
 	ULONG_PTR Arguments[] = { (ULONG_PTR)&Info };
-	::RaiseException( GPUCrashExceptionCode, 0, ARRAY_COUNT(Arguments), Arguments );
+	::RaiseException( GPUCrashExceptionCode, 0, UE_ARRAY_COUNT(Arguments), Arguments );
 }
 
 void ReportHang(const TCHAR* ErrorMessage, const uint64* StackFrames, int32 NumStackFrames, uint32 HungThreadId)
@@ -696,7 +696,7 @@ void CreateExceptionInfoString(EXCEPTION_RECORD* ExceptionRecord)
 		ErrorString += FString::Printf(TEXT("0x%08x"), (uint32)ExceptionRecord->ExceptionCode);
 	}
 
-	FCString::Strncpy(GErrorExceptionDescription, *ErrorString, ARRAY_COUNT(GErrorExceptionDescription));
+	FCString::Strncpy(GErrorExceptionDescription, *ErrorString, UE_ARRAY_COUNT(GErrorExceptionDescription));
 
 #undef HANDLE_CASE
 }
@@ -894,11 +894,11 @@ private:
 			if (ExceptionInfo->ExceptionRecord->ExceptionCode != 1 && ExceptionInfo->ExceptionRecord->ExceptionCode != AssertExceptionCode)
 			{
 				CreateExceptionInfoString(ExceptionInfo->ExceptionRecord);
-				FCString::Strncat(GErrorHist, GErrorExceptionDescription, ARRAY_COUNT(GErrorHist));
-				FCString::Strncat(GErrorHist, TEXT("\r\n\r\n"), ARRAY_COUNT(GErrorHist));
+				FCString::Strncat(GErrorHist, GErrorExceptionDescription, UE_ARRAY_COUNT(GErrorHist));
+				FCString::Strncat(GErrorHist, TEXT("\r\n\r\n"), UE_ARRAY_COUNT(GErrorHist));
 			}
 
-			FCString::Strncat(GErrorHist, ANSI_TO_TCHAR(StackTrace), ARRAY_COUNT(GErrorHist));
+			FCString::Strncat(GErrorHist, ANSI_TO_TCHAR(StackTrace), UE_ARRAY_COUNT(GErrorHist));
 
 			GMalloc->Free(StackTrace);
 		}
diff --git a/Engine/Source/Runtime/Core/Private/Windows/WindowsPlatformProcess.cpp b/Engine/Source/Runtime/Core/Private/Windows/WindowsPlatformProcess.cpp
--- a/Engine/Source/Runtime/Core/Private/Windows/WindowsPlatformProcess.cpp
+++ b/Engine/Source/Runtime/Core/Private/Windows/WindowsPlatformProcess.cpp
@@ -908,7 +908,7 @@ const TCHAR* FWindowsPlatformProcess::BaseDir()
 			{
 				hCurrentModule = hInstance;
 			}
-			GetModuleFileName(hCurrentModule, Result, ARRAY_COUNT(Result));
+			GetModuleFileName(hCurrentModule, Result, UE_ARRAY_COUNT(Result));
 			FString TempResult(Result);
 			TempResult = TempResult.Replace(TEXT("\\"), TEXT("/"));
 			FCString::Strcpy(Result, *TempResult);
@@ -1023,7 +1023,7 @@ const TCHAR* FWindowsPlatformProcess::ComputerName()
 	static TCHAR Result[256]=TEXT("");
 	if( !Result[0] )
 	{
-		uint32 Size=ARRAY_COUNT(Result);
+		uint32 Size=UE_ARRAY_COUNT(Result);
 		GetComputerName( Result, (::DWORD*)&Size );
 	}
 	return Result;
@@ -1037,7 +1037,7 @@ const TCHAR* FWindowsPlatformProcess::UserName(bool bOnlyAlphaNumeric/* = true*/
 	{
 		if( !ResultAlpha[0] )
 		{
-			uint32 Size=ARRAY_COUNT(ResultAlpha);
+			uint32 Size=UE_ARRAY_COUNT(ResultAlpha);
 			GetUserName( ResultAlpha, (::DWORD*)&Size );
 			TCHAR *c, *d;
 			for( c=ResultAlpha, d=ResultAlpha; *c!=0; c++ )
@@ -1051,7 +1051,7 @@ const TCHAR* FWindowsPlatformProcess::UserName(bool bOnlyAlphaNumeric/* = true*/
 	{
 		if( !Result[0] )
 		{
-			uint32 Size=ARRAY_COUNT(Result);
+			uint32 Size=UE_ARRAY_COUNT(Result);
 			GetUserName( Result, (::DWORD*)&Size );
 		}
 		return Result;
@@ -1104,7 +1104,7 @@ const TCHAR* FWindowsPlatformProcess::ExecutablePath()
 	static TCHAR Result[512]=TEXT("");
 	if( !Result[0] )
 	{
-		if ( !GetModuleFileName( hInstance, Result, ARRAY_COUNT(Result) ) )
+		if ( !GetModuleFileName( hInstance, Result, UE_ARRAY_COUNT(Result) ) )
 		{
 			Result[0] = 0;
 		}
@@ -1119,13 +1119,13 @@ const TCHAR* FWindowsPlatformProcess::ExecutableName(bool bRemoveExtension)
 	if( !Result[0] )
 	{
 		// Get complete path for the executable
-		if ( GetModuleFileName( hInstance, Result, ARRAY_COUNT(Result) ) != 0 )
+		if ( GetModuleFileName( hInstance, Result, UE_ARRAY_COUNT(Result) ) != 0 )
 		{
 			// Remove all of the path information by finding the base filename
 			FString FileName = Result;
 			FString FileNameWithExt = Result;
-			FCString::Strncpy( Result, *( FPaths::GetBaseFilename(FileName) ), ARRAY_COUNT(Result) );
-			FCString::Strncpy( ResultWithExt, *( FPaths::GetCleanFilename(FileNameWithExt) ), ARRAY_COUNT(ResultWithExt) );
+			FCString::Strncpy( Result, *( FPaths::GetBaseFilename(FileName) ), UE_ARRAY_COUNT(Result) );
+			FCString::Strncpy( ResultWithExt, *( FPaths::GetCleanFilename(FileNameWithExt) ), UE_ARRAY_COUNT(ResultWithExt) );
 		}
 		// If the call failed, zero out the memory to be safe
 		else
@@ -1165,7 +1165,7 @@ const FString FWindowsPlatformProcess::GetModulesDirectory()
 		}
 
 		// Get the directory for it
-		GetModuleFileName(hCurrentModule, Result, ARRAY_COUNT(Result));
+		GetModuleFileName(hCurrentModule, Result, UE_ARRAY_COUNT(Result));
 		*FCString::Strrchr(Result, '\\') = 0;
 
 		// Normalize the resulting path
diff --git a/Engine/Source/Runtime/Core/Private/Windows/WindowsRunnableThread.cpp b/Engine/Source/Runtime/Core/Private/Windows/WindowsRunnableThread.cpp
--- a/Engine/Source/Runtime/Core/Private/Windows/WindowsRunnableThread.cpp
+++ b/Engine/Source/Runtime/Core/Private/Windows/WindowsRunnableThread.cpp
@@ -60,8 +60,8 @@ uint32 FRunnableThreadWin::GuardedRun()
 			GWarn->Flush();
 
 			// Append the thread name at the end of the error report.
-			FCString::Strncat( GErrorHist, LINE_TERMINATOR TEXT( "Crash in runnable thread " ), ARRAY_COUNT( GErrorHist ) );
-			FCString::Strncat( GErrorHist, *ThreadName, ARRAY_COUNT( GErrorHist ) );
+			FCString::Strncat( GErrorHist, LINE_TERMINATOR TEXT( "Crash in runnable thread " ), UE_ARRAY_COUNT( GErrorHist ) );
+			FCString::Strncat( GErrorHist, *ThreadName, UE_ARRAY_COUNT( GErrorHist ) );
 
 			// Crashed.
 			ExitCode = 1;		
diff --git a/Engine/Source/Runtime/Core/Public/Android/AndroidPlatformRunnableThread.h b/Engine/Source/Runtime/Core/Public/Android/AndroidPlatformRunnableThread.h
--- a/Engine/Source/Runtime/Core/Public/Android/AndroidPlatformRunnableThread.h
+++ b/Engine/Source/Runtime/Core/Public/Android/AndroidPlatformRunnableThread.h
@@ -52,7 +52,7 @@ private:
 
 				// cut out the middle and replace with a substitute
 				const TCHAR Dash[] = TEXT("-");
-				const int32 DashLen = ARRAY_COUNT(Dash) - 1;
+				const int32 DashLen = UE_ARRAY_COUNT(Dash) - 1;
 				int NumToLeave = (AndroidThreadNameLimit - DashLen) / 2;
 
 				SizeLimitedThreadName = Temp.Left(AndroidThreadNameLimit - (NumToLeave + DashLen));
diff --git a/Engine/Source/Runtime/Core/Public/HAL/IConsoleManager.h b/Engine/Source/Runtime/Core/Public/HAL/IConsoleManager.h
--- a/Engine/Source/Runtime/Core/Public/HAL/IConsoleManager.h
+++ b/Engine/Source/Runtime/Core/Public/HAL/IConsoleManager.h
@@ -1249,7 +1249,7 @@ public:
 	// constructor
 	TConsoleVariableData(const T DefaultValue)
 	{
-		for(uint32 i = 0; i < ARRAY_COUNT(ShadowedValue); ++i)
+		for(uint32 i = 0; i < UE_ARRAY_COUNT(ShadowedValue); ++i)
 		{
 			ShadowedValue[i] = DefaultValue;
 		}
diff --git a/Engine/Source/Runtime/Core/Public/HAL/IPlatformFileCachedWrapper.h b/Engine/Source/Runtime/Core/Public/HAL/IPlatformFileCachedWrapper.h
--- a/Engine/Source/Runtime/Core/Public/HAL/IPlatformFileCachedWrapper.h
+++ b/Engine/Source/Runtime/Core/Public/HAL/IPlatformFileCachedWrapper.h
@@ -212,7 +212,7 @@ private:
 	}
 	int32 GetCacheIndex(int64 Pos) const
 	{
-		for (uint32 i=0; i<ARRAY_COUNT(CacheStart); ++i)
+		for (uint32 i=0; i<UE_ARRAY_COUNT(CacheStart); ++i)
 		{
 			if (Pos >= CacheStart[i] && Pos < CacheEnd[i]) 
 			{
diff --git a/Engine/Source/Runtime/Core/Public/IOS/IOSCommandLineHelper.h b/Engine/Source/Runtime/Core/Public/IOS/IOSCommandLineHelper.h
--- a/Engine/Source/Runtime/Core/Public/IOS/IOSCommandLineHelper.h
+++ b/Engine/Source/Runtime/Core/Public/IOS/IOSCommandLineHelper.h
@@ -62,7 +62,7 @@ class FIOSCommandLineHelper
 			{
 				FPlatformMisc::LowLevelOutputDebugStringf(TEXT("Checking for command line in %s... FOUND!") LINE_TERMINATOR, *CommandLineFilePath);
 				char CommandLine[CMD_LINE_MAX] = {0};
-				char* DataExists = fgets(CommandLine, ARRAY_COUNT(CommandLine) - 1, CommandLineFile);
+				char* DataExists = fgets(CommandLine, UE_ARRAY_COUNT(CommandLine) - 1, CommandLineFile);
 				if (DataExists)
 				{
 					// chop off trailing spaces
diff --git a/Engine/Source/Runtime/Core/Public/Math/BigInt.h b/Engine/Source/Runtime/Core/Public/Math/BigInt.h
--- a/Engine/Source/Runtime/Core/Public/Math/BigInt.h
+++ b/Engine/Source/Runtime/Core/Public/Math/BigInt.h
@@ -138,7 +138,7 @@ public:
 	{
 		// If we weren't given enough data to fill the int, zero first
 		// TODO: Only zero the bits we aren't going to copy over
-		if (InNumBytes < (ARRAY_COUNT(Bits) * sizeof(uint32)))
+		if (InNumBytes < (UE_ARRAY_COUNT(Bits) * sizeof(uint32)))
 		{
 			Zero();
 		}
diff --git a/Engine/Source/Runtime/Core/Public/Misc/App.h b/Engine/Source/Runtime/Core/Public/Misc/App.h
--- a/Engine/Source/Runtime/Core/Public/Misc/App.h
+++ b/Engine/Source/Runtime/Core/Public/Misc/App.h
@@ -184,9 +184,9 @@ public:
 	FORCEINLINE static void SetProjectName(const TCHAR* InProjectName)
 	{
 		// At the moment Strcpy is not safe as we don't check the buffer size on all platforms, so we use strncpy here.
-		FCString::Strncpy(GInternalProjectName, InProjectName, ARRAY_COUNT(GInternalProjectName));
+		FCString::Strncpy(GInternalProjectName, InProjectName, UE_ARRAY_COUNT(GInternalProjectName));
 		// And make sure the ProjectName string is null terminated.
-		GInternalProjectName[ARRAY_COUNT(GInternalProjectName) - 1] = 0;
+		GInternalProjectName[UE_ARRAY_COUNT(GInternalProjectName) - 1] = 0;
 	}
 
 	UE_DEPRECATED(4.18, "SetGameName() has been superseded by SetProjectName().")
diff --git a/Engine/Source/Runtime/Core/Public/Misc/Paths.h b/Engine/Source/Runtime/Core/Public/Misc/Paths.h
--- a/Engine/Source/Runtime/Core/Public/Misc/Paths.h
+++ b/Engine/Source/Runtime/Core/Public/Misc/Paths.h
@@ -600,7 +600,7 @@ public:
 		const TCHAR* Paths[] = { GetTCharPtr(Forward<PathTypes>(InPaths))... };
 		FString Out;
 		
-		CombineInternal(Out, Paths, ARRAY_COUNT(Paths));
+		CombineInternal(Out, Paths, UE_ARRAY_COUNT(Paths));
 		return Out;
 	}
 
diff --git a/Engine/Source/Runtime/Core/Public/Misc/TVariant.h b/Engine/Source/Runtime/Core/Public/Misc/TVariant.h
--- a/Engine/Source/Runtime/Core/Public/Misc/TVariant.h
+++ b/Engine/Source/Runtime/Core/Public/Misc/TVariant.h
@@ -184,7 +184,7 @@ namespace UE4Variant_Details
 		static void Destruct(SIZE_T TypeIndex, void* Value)
 		{
 			static constexpr void(*Destructors[])(void*) = { &TDestructorCaller<Ts>::Destruct... };
-			check(TypeIndex < ARRAY_COUNT(Destructors));
+			check(TypeIndex < UE_ARRAY_COUNT(Destructors));
 			Destructors[TypeIndex](Value);
 		}
 	};
@@ -208,7 +208,7 @@ namespace UE4Variant_Details
 		static void Construct(SIZE_T TypeIndex, void* Storage, const void* Value)
 		{
 			static constexpr void(*CopyConstructors[])(void*, const void*) = { &TCopyConstructorCaller<Ts>::Construct... };
-			check(TypeIndex < ARRAY_COUNT(CopyConstructors));
+			check(TypeIndex < UE_ARRAY_COUNT(CopyConstructors));
 			CopyConstructors[TypeIndex](Storage, Value);
 		}
 	};
@@ -233,7 +233,7 @@ namespace UE4Variant_Details
 		static void Construct(SIZE_T TypeIndex, void* Target, void* Source)
 		{
 			static constexpr void(*MoveConstructors[])(void*, void*) = { &TMoveConstructorCaller<Ts>::Construct... };
-			check(TypeIndex < ARRAY_COUNT(MoveConstructors));
+			check(TypeIndex < UE_ARRAY_COUNT(MoveConstructors));
 			MoveConstructors[TypeIndex](Target, Source);
 		}
 	};
@@ -246,7 +246,7 @@ namespace UE4Variant_Details
 		static bool IsSame(SIZE_T TypeIndex)
 		{
 			static constexpr bool bIsSameType[] = { TIsSame<Ts, LookupType>::Value... };
-			check(TypeIndex < ARRAY_COUNT(bIsSameType));
+			check(TypeIndex < UE_ARRAY_COUNT(bIsSameType));
 			return bIsSameType[TypeIndex];
 		}
 	};
diff --git a/Engine/Source/Runtime/Core/Public/Modules/ModuleManager.h b/Engine/Source/Runtime/Core/Public/Modules/ModuleManager.h
--- a/Engine/Source/Runtime/Core/Public/Modules/ModuleManager.h
+++ b/Engine/Source/Runtime/Core/Public/Modules/ModuleManager.h
@@ -789,13 +789,13 @@ class FDefaultGameModuleImpl
 		{ \
 			const uint8 Exponent[] = { ExponentValue }; \
 			const uint8 Modulus[] = { ModulusValue }; \
-			OutExponent.SetNum(ARRAY_COUNT(Exponent)); \
-			OutModulus.SetNum(ARRAY_COUNT(Modulus)); \
-			for(int ByteIdx = 0; ByteIdx < ARRAY_COUNT(Exponent); ByteIdx++) \
+			OutExponent.SetNum(UE_ARRAY_COUNT(Exponent)); \
+			OutModulus.SetNum(UE_ARRAY_COUNT(Modulus)); \
+			for(int ByteIdx = 0; ByteIdx < UE_ARRAY_COUNT(Exponent); ByteIdx++) \
 			{ \
 				OutExponent[ByteIdx] = Exponent[ByteIdx]; \
 			} \
-			for(int ByteIdx = 0; ByteIdx < ARRAY_COUNT(Modulus); ByteIdx++) \
+			for(int ByteIdx = 0; ByteIdx < UE_ARRAY_COUNT(Modulus); ByteIdx++) \
 			{ \
 				OutModulus[ByteIdx] = Modulus[ByteIdx]; \
 			} \
@@ -846,7 +846,7 @@ class FDefaultGameModuleImpl
 			{ \
 				FAutoSet##ModuleName() \
 				{ \
-					FCString::Strncpy(GInternalProjectName, TEXT( GameName ), ARRAY_COUNT(GInternalProjectName)); \
+					FCString::Strncpy(GInternalProjectName, TEXT( GameName ), UE_ARRAY_COUNT(GInternalProjectName)); \
 				} \
 			} AutoSet##ModuleName; \
 			IMPLEMENT_LIVE_CODING_ENGINE_DIR() \
diff --git a/Engine/Source/Runtime/Core/Public/ProfilingDebugging/DiagnosticTable.h b/Engine/Source/Runtime/Core/Public/ProfilingDebugging/DiagnosticTable.h
--- a/Engine/Source/Runtime/Core/Public/ProfilingDebugging/DiagnosticTable.h
+++ b/Engine/Source/Runtime/Core/Public/ProfilingDebugging/DiagnosticTable.h
@@ -47,7 +47,7 @@ public:
 
 		// Format the input string.
 		TCHAR ColumnText[4096];
-		GET_VARARGS(ColumnText,ARRAY_COUNT(ColumnText),ARRAY_COUNT(ColumnText)-1,Format,Format);
+		GET_VARARGS(ColumnText,UE_ARRAY_COUNT(ColumnText),UE_ARRAY_COUNT(ColumnText)-1,Format,Format);
 
 		// If this isn't the first column in the row, add a comma to separate the columns.
 		if(CurrentRow.Len())
diff --git a/Engine/Source/Runtime/Core/Public/ProfilingDebugging/MallocProfiler.h b/Engine/Source/Runtime/Core/Public/ProfilingDebugging/MallocProfiler.h
--- a/Engine/Source/Runtime/Core/Public/ProfilingDebugging/MallocProfiler.h
+++ b/Engine/Source/Runtime/Core/Public/ProfilingDebugging/MallocProfiler.h
@@ -209,14 +209,14 @@ struct FCallStackInfo
 		Ar << CallStackInfo.CRC;
 		// Serialize valid callstack indices.
 		int32 i=0;
-		for( ; i<ARRAY_COUNT(CallStackInfo.AddressIndices) && CallStackInfo.AddressIndices[i]!=-1; i++ )
+		for( ; i<UE_ARRAY_COUNT(CallStackInfo.AddressIndices) && CallStackInfo.AddressIndices[i]!=-1; i++ )
 		{
 			Ar << CallStackInfo.AddressIndices[i];
 		}
 		// Terminate list of address indices with -1 if we have a normal callstack.
 		int32 Stopper = -1;
 		// And terminate with -2 if the callstack was truncated.
-		if( i== ARRAY_COUNT(CallStackInfo.AddressIndices) )
+		if( i== UE_ARRAY_COUNT(CallStackInfo.AddressIndices) )
 		{
 			Stopper = -2;
 		}
diff --git a/Engine/Source/Runtime/Core/Public/Stats/StatsData.h b/Engine/Source/Runtime/Core/Public/Stats/StatsData.h
--- a/Engine/Source/Runtime/Core/Public/Stats/StatsData.h
+++ b/Engine/Source/Runtime/Core/Public/Stats/StatsData.h
@@ -76,7 +76,7 @@ template <typename T>
 void ParseTypedValue( const TCHAR* Stream, const TCHAR* Match, T& Out )
 {
 	TCHAR Temp[64] = TEXT( "" );
-	if( FParse::Value( Stream, Match, Temp, ARRAY_COUNT( Temp ) ) )
+	if( FParse::Value( Stream, Match, Temp, UE_ARRAY_COUNT( Temp ) ) )
 	{
 		LexFromString( Out, Temp );
 	}
diff --git a/Engine/Source/Runtime/Core/Public/Templates/UnrealTemplate.h b/Engine/Source/Runtime/Core/Public/Templates/UnrealTemplate.h
--- a/Engine/Source/Runtime/Core/Public/Templates/UnrealTemplate.h
+++ b/Engine/Source/Runtime/Core/Public/Templates/UnrealTemplate.h
@@ -131,15 +131,16 @@ FORCEINLINE const T& AsConst(T& Ref)
 ----------------------------------------------------------------------------*/
 
 #ifdef __clang__
-	template <typename T>
-	auto ArrayCountHelper(T& t) -> typename TEnableIf<__is_array(T), char(&)[sizeof(t) / sizeof(t[0]) + 1]>::Type;
+template <typename T>
+auto UE4ArrayCountHelper(T& t) -> typename TEnableIf<__is_array(T), char(&)[sizeof(t) / sizeof(t[0]) + 1]>::Type;
 #else
-	template <typename T, uint32 N>
-	char (&ArrayCountHelper(const T (&)[N]))[N + 1];
+template <typename T, uint32 N>
+char(&UE4ArrayCountHelper(const T(&)[N]))[N + 1];
 #endif
 
 // Number of elements in an array.
-#define ARRAY_COUNT( array ) (sizeof(ArrayCountHelper(array)) - 1)
+#define UE_ARRAY_COUNT( array ) (sizeof(UE4ArrayCountHelper(array)) - 1)
+#define ARRAY_COUNT( array ) DEPRECATED_MACRO(4.24, "The ARRAY_COUNT macro has been deprecated in favor of UE_ARRAY_COUNT.") UE_ARRAY_COUNT( array )
 
 // Offset of a struct member.
 #ifndef UNREAL_CODE_ANALYZER
diff --git a/Engine/Source/Runtime/Core/Public/Unix/UnixPlatformRunnableThread.h b/Engine/Source/Runtime/Core/Public/Unix/UnixPlatformRunnableThread.h
--- a/Engine/Source/Runtime/Core/Public/Unix/UnixPlatformRunnableThread.h
+++ b/Engine/Source/Runtime/Core/Public/Unix/UnixPlatformRunnableThread.h
@@ -236,7 +236,7 @@ private:
 
 				// cut out the middle and replace with a substitute
 				const TCHAR Dash[] = TEXT("-");
-				const int32 DashLen = ARRAY_COUNT(Dash) - 1;
+				const int32 DashLen = UE_ARRAY_COUNT(Dash) - 1;
 				int NumToLeave = (EConstants::UnixThreadNameLimit - DashLen) / 2;
 
 				SizeLimitedThreadName = Temp.Left(EConstants::UnixThreadNameLimit - (NumToLeave + DashLen));
diff --git a/Engine/Source/Runtime/CoreUObject/Private/Misc/PackageName.cpp b/Engine/Source/Runtime/CoreUObject/Private/Misc/PackageName.cpp
--- a/Engine/Source/Runtime/CoreUObject/Private/Misc/PackageName.cpp
+++ b/Engine/Source/Runtime/CoreUObject/Private/Misc/PackageName.cpp
@@ -756,7 +756,7 @@ bool FPackageName::FindPackageFileWithoutExtension(const FString& InPackageFilen
 
 		// Loop through all known extensions and check if the file exists
 
-		for (int32 ExtensionIndex = 0; ExtensionIndex < ARRAY_COUNT(PackageExtensions); ++ExtensionIndex)
+		for (int32 ExtensionIndex = 0; ExtensionIndex < UE_ARRAY_COUNT(PackageExtensions); ++ExtensionIndex)
 		{
 			FString   PackageFilename = InPackageFilename + *PackageExtensions[ExtensionIndex];
 			FDateTime Timestamp       = FileManager.GetTimeStamp(*PackageFilename);
@@ -779,7 +779,7 @@ bool FPackageName::FindPackageFileWithoutExtension(const FString& InPackageFilen
 			&TextMapPackageExtension
 		};
 
-		for (int32 ExtensionIndex = 0; ExtensionIndex < ARRAY_COUNT(TextPackageExtensions); ++ExtensionIndex)
+		for (int32 ExtensionIndex = 0; ExtensionIndex < UE_ARRAY_COUNT(TextPackageExtensions); ++ExtensionIndex)
 		{
 			FString   PackageFilename = InPackageFilename + *TextPackageExtensions[ExtensionIndex];
 			FDateTime Timestamp		  = FileManager.GetTimeStamp(*PackageFilename);
diff --git a/Engine/Source/Runtime/CoreUObject/Private/Serialization/Formatters/JsonArchiveInputFormatter.cpp b/Engine/Source/Runtime/CoreUObject/Private/Serialization/Formatters/JsonArchiveInputFormatter.cpp
--- a/Engine/Source/Runtime/CoreUObject/Private/Serialization/Formatters/JsonArchiveInputFormatter.cpp
+++ b/Engine/Source/Runtime/CoreUObject/Private/Serialization/Formatters/JsonArchiveInputFormatter.cpp
@@ -286,7 +286,7 @@ void FJsonArchiveInputFormatter::Serialize(FName& Value)
 	const TCHAR Prefix[] = TEXT("Name:");
 	if (ensure(StringValue.StartsWith(Prefix)))
 	{
-		Value = FName(*StringValue + ARRAY_COUNT(Prefix) - 1);
+		Value = FName(*StringValue + UE_ARRAY_COUNT(Prefix) - 1);
 	}
 	else
 	{
@@ -300,7 +300,7 @@ void FJsonArchiveInputFormatter::Serialize(UObject*& Value)
 	const TCHAR Prefix[] = TEXT("Object:");
 	if (ValueStack.Top()->TryGetString(StringValue) && StringValue.StartsWith(Prefix))
 	{
-		Value = ResolveObjectName(*StringValue + ARRAY_COUNT(Prefix) - 1);
+		Value = ResolveObjectName(*StringValue + UE_ARRAY_COUNT(Prefix) - 1);
 	}
 	else
 	{
@@ -335,7 +335,7 @@ void FJsonArchiveInputFormatter::Serialize(FSoftObjectPath& Value)
 	const TCHAR Prefix[] = TEXT("Object:");
 	if (ValueStack.Top()->TryGetString(StringValue) && StringValue.StartsWith(Prefix))
 	{
-		Value.SetPath(*StringValue + ARRAY_COUNT(Prefix) - 1);
+		Value.SetPath(*StringValue + UE_ARRAY_COUNT(Prefix) - 1);
 	}
 	else
 	{
@@ -350,7 +350,7 @@ void FJsonArchiveInputFormatter::Serialize(FLazyObjectPtr& Value)
 	if (ValueStack.Top()->TryGetString(StringValue) && StringValue.StartsWith(Prefix))
 	{
 		FUniqueObjectGuid Guid;
-		Guid.FromString(*StringValue + ARRAY_COUNT(Prefix) - 1);
+		Guid.FromString(*StringValue + UE_ARRAY_COUNT(Prefix) - 1);
 		Value = Guid;
 	}
 	else
diff --git a/Engine/Source/Runtime/CoreUObject/Private/Serialization/Formatters/JsonArchiveOutputFormatter.cpp b/Engine/Source/Runtime/CoreUObject/Private/Serialization/Formatters/JsonArchiveOutputFormatter.cpp
--- a/Engine/Source/Runtime/CoreUObject/Private/Serialization/Formatters/JsonArchiveOutputFormatter.cpp
+++ b/Engine/Source/Runtime/CoreUObject/Private/Serialization/Formatters/JsonArchiveOutputFormatter.cpp
@@ -14,7 +14,7 @@
 
 FJsonArchiveOutputFormatter::FJsonArchiveOutputFormatter(FArchive& InInner) 
 	: Inner(InInner)
-	, Newline(LINE_TERMINATOR_ANSI, ARRAY_COUNT(LINE_TERMINATOR_ANSI) - 1)
+	, Newline(LINE_TERMINATOR_ANSI, UE_ARRAY_COUNT(LINE_TERMINATOR_ANSI) - 1)
 {
 	Inner.SetIsTextFormat(true);
 
@@ -369,13 +369,13 @@ void FJsonArchiveOutputFormatter::Serialize(void* Data, uint64 DataSize)
 
 		// Convert the hash to a string
 		ANSICHAR DigestString[(FSHA1::DigestSize * 2) + 1];
-		for(int32 Idx = 0; Idx < ARRAY_COUNT(Digest); Idx++)
+		for(int32 Idx = 0; Idx < UE_ARRAY_COUNT(Digest); Idx++)
 		{
 			static const ANSICHAR HexDigits[] = "0123456789abcdef";
 			DigestString[(Idx * 2) + 0] = HexDigits[Digest[Idx] >> 4];
 			DigestString[(Idx * 2) + 1] = HexDigits[Digest[Idx] & 15];
 		}
-		DigestString[ARRAY_COUNT(DigestString) - 1] = 0;
+		DigestString[UE_ARRAY_COUNT(DigestString) - 1] = 0;
 
 		// Write the digest
 		Write("\"Digest\": \"");
diff --git a/Engine/Source/Runtime/CoreUObject/Private/UObject/ErrorException.cpp b/Engine/Source/Runtime/CoreUObject/Private/UObject/ErrorException.cpp
--- a/Engine/Source/Runtime/CoreUObject/Private/UObject/ErrorException.cpp
+++ b/Engine/Source/Runtime/CoreUObject/Private/UObject/ErrorException.cpp
@@ -9,7 +9,7 @@
 void VARARGS FError::ThrowfImpl(const TCHAR* Fmt, ...)
 {
 	static TCHAR TempStr[4096];
-	GET_VARARGS( TempStr, ARRAY_COUNT(TempStr), ARRAY_COUNT(TempStr)-1, Fmt, Fmt );
+	GET_VARARGS( TempStr, UE_ARRAY_COUNT(TempStr), UE_ARRAY_COUNT(TempStr)-1, Fmt, Fmt );
 #if !PLATFORM_EXCEPTIONS_DISABLED
 	throw( TempStr );
 #else
diff --git a/Engine/Source/Runtime/CoreUObject/Private/UObject/LinkerLoad.cpp b/Engine/Source/Runtime/CoreUObject/Private/UObject/LinkerLoad.cpp
--- a/Engine/Source/Runtime/CoreUObject/Private/UObject/LinkerLoad.cpp
+++ b/Engine/Source/Runtime/CoreUObject/Private/UObject/LinkerLoad.cpp
@@ -1936,7 +1936,7 @@ FLinkerLoad::ELinkerStatus FLinkerLoad::CreateExportHash()
 	// Zero initialize hash on first iteration.
 	if( ExportHashIndex == 0 )
 	{
-		for( int32 i=0; i<ARRAY_COUNT(ExportHash); i++ )
+		for( int32 i=0; i<UE_ARRAY_COUNT(ExportHash); i++ )
 		{
 			ExportHash[i] = INDEX_NONE;
 		}
@@ -1947,7 +1947,7 @@ FLinkerLoad::ELinkerStatus FLinkerLoad::CreateExportHash()
 	{
 		FObjectExport& Export = ExportMap[ExportHashIndex];
 
-		const int32 iHash = HashNames( Export.ObjectName, GetExportClassName(ExportHashIndex), GetExportClassPackage(ExportHashIndex) ) & (ARRAY_COUNT(ExportHash)-1);
+		const int32 iHash = HashNames( Export.ObjectName, GetExportClassName(ExportHashIndex), GetExportClassPackage(ExportHashIndex) ) & (UE_ARRAY_COUNT(ExportHash)-1);
 		Export.HashNext = ExportHash[iHash];
 		ExportHash[iHash] = ExportHashIndex;
 
@@ -2810,7 +2810,7 @@ bool FLinkerLoad::VerifyImportInner(const int32 ImportIndex, FString& WarningSuf
 			Pkg = CreatePackage(NULL, *Top->ObjectName.ToString() );
 
 			// Find this import within its existing linker.
-			int32 iHash = HashNames( Import.ObjectName, Import.ClassName, Import.ClassPackage) & (ARRAY_COUNT(ExportHash)-1);
+			int32 iHash = HashNames( Import.ObjectName, Import.ClassName, Import.ClassPackage) & (UE_ARRAY_COUNT(ExportHash)-1);
 
 			//@Package name transition, if we can match without shortening the names, then we must not take a shortened match
 			bool bMatchesWithoutShortening = false;
@@ -3275,7 +3275,7 @@ FName FLinkerLoad::ResolveResourceName( FPackageIndex ResourceIndex )
 // Find the index of a specified object without regard to specific package.
 int32 FLinkerLoad::FindExportIndex( FName ClassName, FName ClassPackage, FName ObjectName, FPackageIndex ExportOuterIndex )
 {
-	int32 iHash = HashNames( ObjectName, ClassName, ClassPackage ) & (ARRAY_COUNT(ExportHash)-1);
+	int32 iHash = HashNames( ObjectName, ClassName, ClassPackage ) & (UE_ARRAY_COUNT(ExportHash)-1);
 
 	for( int32 i=ExportHash[iHash]; i!=INDEX_NONE; i=ExportMap[i].HashNext )
 	{
diff --git a/Engine/Source/Runtime/CoreUObject/Private/UObject/Obj.cpp b/Engine/Source/Runtime/CoreUObject/Private/UObject/Obj.cpp
--- a/Engine/Source/Runtime/CoreUObject/Private/UObject/Obj.cpp
+++ b/Engine/Source/Runtime/CoreUObject/Private/UObject/Obj.cpp
@@ -3016,7 +3016,7 @@ static void PerformSetCommand( const TCHAR* Str, FOutputDevice& Ar, bool bNotify
 {
 	// Set a class default variable.
 	TCHAR ObjectName[256], PropertyName[256];
-	if (FParse::Token(Str, ObjectName, ARRAY_COUNT(ObjectName), true) && FParse::Token(Str, PropertyName, ARRAY_COUNT(PropertyName), true))
+	if (FParse::Token(Str, ObjectName, UE_ARRAY_COUNT(ObjectName), true) && FParse::Token(Str, PropertyName, UE_ARRAY_COUNT(PropertyName), true))
 	{
 		UClass* Class = FindObject<UClass>(ANY_PACKAGE, ObjectName);
 		if (Class != NULL)
@@ -3225,11 +3225,11 @@ bool StaticExec( UWorld* InWorld, const TCHAR* Cmd, FOutputDevice& Ar )
 		UClass* Class;
 		UProperty* Property;
 		if
-		(	FParse::Token( Str, ClassName, ARRAY_COUNT(ClassName), 1 )
+		(	FParse::Token( Str, ClassName, UE_ARRAY_COUNT(ClassName), 1 )
 		&&	(Class=FindObject<UClass>( ANY_PACKAGE, ClassName))!=NULL )
 		{
 			if
-			(	FParse::Token( Str, PropertyName, ARRAY_COUNT(PropertyName), 1 )
+			(	FParse::Token( Str, PropertyName, UE_ARRAY_COUNT(PropertyName), 1 )
 			&&	(Property=FindField<UProperty>( Class, PropertyName))!=NULL )
 			{
 				FString	Temp;
@@ -3257,7 +3257,7 @@ bool StaticExec( UWorld* InWorld, const TCHAR* Cmd, FOutputDevice& Ar )
 		UClass* Class;
 		FString PropWildcard;
 
-		if ( FParse::Token(Str, ClassName, ARRAY_COUNT(ClassName), 1) &&
+		if ( FParse::Token(Str, ClassName, UE_ARRAY_COUNT(ClassName), 1) &&
 			(Class = FindObject<UClass>(ANY_PACKAGE, ClassName)) != NULL &&
 			FParse::Token(Str, PropWildcard, true) )
 		{
@@ -3386,10 +3386,10 @@ bool StaticExec( UWorld* InWorld, const TCHAR* Cmd, FOutputDevice& Ar )
 		UClass* Class;
 		UProperty* Property;
 
-		if ( FParse::Token(Str,ClassName,ARRAY_COUNT(ClassName), 1) &&
+		if ( FParse::Token(Str,ClassName,UE_ARRAY_COUNT(ClassName), 1) &&
 			(Class=FindObject<UClass>( ANY_PACKAGE, ClassName)) != NULL )
 		{
-			FParse::Token(Str,PropertyName,ARRAY_COUNT(PropertyName),1);
+			FParse::Token(Str,PropertyName,UE_ARRAY_COUNT(PropertyName),1);
 			{
 				Property=FindField<UProperty>(Class,PropertyName);
 				{
@@ -3402,7 +3402,7 @@ bool StaticExec( UWorld* InWorld, const TCHAR* Cmd, FOutputDevice& Ar )
 					// Check for a specific object name
 					TCHAR ObjNameStr[256];
 					FName ObjName(NAME_None);
-					if (FParse::Value(Str,TEXT("NAME="),ObjNameStr,ARRAY_COUNT(ObjNameStr)))
+					if (FParse::Value(Str,TEXT("NAME="),ObjNameStr,UE_ARRAY_COUNT(ObjNameStr)))
 					{
 						ObjName = FName(ObjNameStr);
 					}
@@ -3530,7 +3530,7 @@ bool StaticExec( UWorld* InWorld, const TCHAR* Cmd, FOutputDevice& Ar )
 		// LISTFUNCS <classname>
 		TCHAR ClassName[256];
 
-		if (FParse::Token(Str, ClassName, ARRAY_COUNT(ClassName), true))
+		if (FParse::Token(Str, ClassName, UE_ARRAY_COUNT(ClassName), true))
 		{
 			//if ( (Property=FindField<UProperty>(Class,PropertyName)) != NULL )
 			UClass* Class = FindObject<UClass>(ANY_PACKAGE, ClassName);
@@ -3557,7 +3557,7 @@ bool StaticExec( UWorld* InWorld, const TCHAR* Cmd, FOutputDevice& Ar )
 		// LISTFUNC <classname> <functionname>
 		TCHAR ClassName[256];
 		TCHAR FunctionName[256];
-		if (FParse::Token(Str, ClassName, ARRAY_COUNT(ClassName), true) && FParse::Token(Str, FunctionName, ARRAY_COUNT(FunctionName), true))
+		if (FParse::Token(Str, ClassName, UE_ARRAY_COUNT(ClassName), true) && FParse::Token(Str, FunctionName, UE_ARRAY_COUNT(FunctionName), true))
 		{
 			UClass* Class = FindObject<UClass>(ANY_PACKAGE, ClassName);
 
@@ -3910,7 +3910,7 @@ bool StaticExec( UWorld* InWorld, const TCHAR* Cmd, FOutputDevice& Ar )
 				ReferencerClass = NULL;
 			}
 			TCHAR TempStr[1024];
-			if (FParse::Value(Str, TEXT("REFNAME="), TempStr, ARRAY_COUNT(TempStr)))
+			if (FParse::Value(Str, TEXT("REFNAME="), TempStr, UE_ARRAY_COUNT(TempStr)))
 			{
 				ReferencerName = TempStr;
 			}
@@ -4172,7 +4172,7 @@ bool StaticExec( UWorld* InWorld, const TCHAR* Cmd, FOutputDevice& Ar )
 			// Dump all object flags for objects rooted at the named object.
 			TCHAR ObjectName[NAME_SIZE];
 			UObject* Obj = NULL;
-			if ( FParse::Token(Str,ObjectName,ARRAY_COUNT(ObjectName), 1) )
+			if ( FParse::Token(Str,ObjectName,UE_ARRAY_COUNT(ObjectName), 1) )
 			{
 				Obj = FindObject<UObject>(ANY_PACKAGE,ObjectName);
 			}
@@ -4224,7 +4224,7 @@ bool StaticExec( UWorld* InWorld, const TCHAR* Cmd, FOutputDevice& Ar )
 	{
 		TCHAR ClassName[256];
 		// Determine the object/class name
-		if (FParse::Token(Str,ClassName,ARRAY_COUNT(ClassName),1))
+		if (FParse::Token(Str,ClassName,UE_ARRAY_COUNT(ClassName),1))
 		{
 			// Try to find a corresponding class
 			UClass* ClassToReload = FindObject<UClass>(ANY_PACKAGE,ClassName);
diff --git a/Engine/Source/Runtime/CoreUObject/Private/UObject/ScriptCore.cpp b/Engine/Source/Runtime/CoreUObject/Private/UObject/ScriptCore.cpp
--- a/Engine/Source/Runtime/CoreUObject/Private/UObject/ScriptCore.cpp
+++ b/Engine/Source/Runtime/CoreUObject/Private/UObject/ScriptCore.cpp
@@ -605,7 +605,7 @@ COREUOBJECT_API uint8 GRegisterNative( int32 NativeBytecodeIndex, const FNativeF
 	if (!bInitialized)
 	{
 		bInitialized = true;
-		for (uint32 i = 0; i < ARRAY_COUNT(GNatives); i++)
+		for (uint32 i = 0; i < UE_ARRAY_COUNT(GNatives); i++)
 		{
 			GNatives[i] = &UObject::execUndefined;
 		}
@@ -613,7 +613,7 @@ COREUOBJECT_API uint8 GRegisterNative( int32 NativeBytecodeIndex, const FNativeF
 
 	if( NativeBytecodeIndex != INDEX_NONE )
 	{
-		if( NativeBytecodeIndex<0 || (uint32)NativeBytecodeIndex>ARRAY_COUNT(GNatives) || GNatives[NativeBytecodeIndex]!=&UObject::execUndefined) 
+		if( NativeBytecodeIndex<0 || (uint32)NativeBytecodeIndex>UE_ARRAY_COUNT(GNatives) || GNatives[NativeBytecodeIndex]!=&UObject::execUndefined) 
 		{
 #if WITH_HOT_RELOAD
 			if (GIsHotReload)
@@ -641,7 +641,7 @@ COREUOBJECT_API uint8 GRegisterCast( int32 CastCode, const FNativeFuncPtr& Func
 	if (!bInitialized)
 	{
 		bInitialized = true;
-		for (uint32 i = 0; i < ARRAY_COUNT(GCasts); i++)
+		for (uint32 i = 0; i < UE_ARRAY_COUNT(GCasts); i++)
 		{
 			GCasts[i] = &UObject::execUndefined;
 		}
@@ -656,7 +656,7 @@ COREUOBJECT_API uint8 GRegisterCast( int32 CastCode, const FNativeFuncPtr& Func
 #if WITH_HOT_RELOAD
 			!GIsHotReload && 
 #endif
-			(CastCode<0 || (uint32)CastCode>ARRAY_COUNT(GCasts) || GCasts[CastCode]!=&UObject::execUndefined) ) 
+			(CastCode<0 || (uint32)CastCode>UE_ARRAY_COUNT(GCasts) || GCasts[CastCode]!=&UObject::execUndefined) ) 
 		{
 			GCastDuplicate = CastCode;
 		}
diff --git a/Engine/Source/Runtime/CoreUObject/Private/UObject/UObjectBaseUtility.cpp b/Engine/Source/Runtime/CoreUObject/Private/UObject/UObjectBaseUtility.cpp
--- a/Engine/Source/Runtime/CoreUObject/Private/UObject/UObjectBaseUtility.cpp
+++ b/Engine/Source/Runtime/CoreUObject/Private/UObject/UObjectBaseUtility.cpp
@@ -462,7 +462,7 @@ void FScopeCycleCounterUObject::TrackObjectForMallocProfiling(const FName InPack
 	{
 		// "Package:/Path/To/Package"
 		ScratchSpaceBuffer.Reset();
-		ScratchSpaceBuffer.Append(PackageTagCategory, ARRAY_COUNT(PackageTagCategory) - 1);
+		ScratchSpaceBuffer.Append(PackageTagCategory, UE_ARRAY_COUNT(PackageTagCategory) - 1);
 		AppendNameToBuffer(InPackageName);
 		ScratchSpaceBuffer.Add(0);
 		PackageTag = FName(ScratchSpaceBuffer.GetData());
@@ -470,7 +470,7 @@ void FScopeCycleCounterUObject::TrackObjectForMallocProfiling(const FName InPack
 
 		// "Object:/Path/To/Package/ObjectName"
 		ScratchSpaceBuffer.Reset();
-		ScratchSpaceBuffer.Append(ObjectTagCategory, ARRAY_COUNT(ObjectTagCategory) - 1);
+		ScratchSpaceBuffer.Append(ObjectTagCategory, UE_ARRAY_COUNT(ObjectTagCategory) - 1);
 		AppendNameToBuffer(InPackageName);
 		ScratchSpaceBuffer.Add(TEXT('/'));
 		AppendNameToBuffer(InObjectName);
@@ -483,7 +483,7 @@ void FScopeCycleCounterUObject::TrackObjectForMallocProfiling(const FName InPack
 	{
 		// "Class:ClassName"
 		ScratchSpaceBuffer.Reset();
-		ScratchSpaceBuffer.Append(ClassTagCategory, ARRAY_COUNT(ClassTagCategory) - 1);
+		ScratchSpaceBuffer.Append(ClassTagCategory, UE_ARRAY_COUNT(ClassTagCategory) - 1);
 		AppendNameToBuffer(InClassName);
 		ScratchSpaceBuffer.Add(0);
 		ClassTag = FName(ScratchSpaceBuffer.GetData());
diff --git a/Engine/Source/Runtime/CoreUObject/Private/UObject/UObjectGlobals.cpp b/Engine/Source/Runtime/CoreUObject/Private/UObject/UObjectGlobals.cpp
--- a/Engine/Source/Runtime/CoreUObject/Private/UObject/UObjectGlobals.cpp
+++ b/Engine/Source/Runtime/CoreUObject/Private/UObject/UObjectGlobals.cpp
@@ -780,7 +780,7 @@ bool ResolveName(UObject*& InPackage, FString& InOutName, bool Create, bool Thro
 bool ParseObject( const TCHAR* Stream, const TCHAR* Match, UClass* Class, UObject*& DestRes, UObject* InParent, bool* bInvalidObject )
 {
 	TCHAR TempStr[1024];
-	if( !FParse::Value( Stream, Match, TempStr, ARRAY_COUNT(TempStr) ) )
+	if( !FParse::Value( Stream, Match, TempStr, UE_ARRAY_COUNT(TempStr) ) )
 	{
 		// Match not found
 		return 0;
diff --git a/Engine/Source/Runtime/CrashReportCore/Private/GenericErrorReport.cpp b/Engine/Source/Runtime/CrashReportCore/Private/GenericErrorReport.cpp
--- a/Engine/Source/Runtime/CrashReportCore/Private/GenericErrorReport.cpp
+++ b/Engine/Source/Runtime/CrashReportCore/Private/GenericErrorReport.cpp
@@ -299,7 +299,7 @@ bool FGenericErrorReport::TryReadDiagnosticsFile()
 			else if (Line.StartsWith(ExceptionLineStart))
 			{
 				// Not subtracting 1 from the array count so it gobbles the initial quote
-				Exception = Line.RightChop(ARRAY_COUNT(ExceptionLineStart)).LeftChop(1);
+				Exception = Line.RightChop(UE_ARRAY_COUNT(ExceptionLineStart)).LeftChop(1);
 			}
 			break;
 		}
