More ARRAY_COUNT fixes.

diff --git a/Engine/Source/Developer/DesktopPlatform/Private/DesktopPlatformBase.cpp b/Engine/Source/Developer/DesktopPlatform/Private/DesktopPlatformBase.cpp
--- a/Engine/Source/Developer/DesktopPlatform/Private/DesktopPlatformBase.cpp
+++ b/Engine/Source/Developer/DesktopPlatform/Private/DesktopPlatformBase.cpp
@@ -774,7 +774,7 @@ struct FTargetFileVisitor : IPlatformFile::FDirectoryStatVisitor
 		int32 Length = FCString::Strlen(FileNameOrDirectory);
 
 		static const TCHAR TargetExt[] = TEXT(".Target.cs");
-		static const int32 TargetExtLen = ARRAY_COUNT(TargetExt) - 1;
+		static const int32 TargetExtLen = UE_ARRAY_COUNT(TargetExt) - 1;
 		if (Length > TargetExtLen && FCString::Stricmp(FileNameOrDirectory + Length - TargetExtLen, TargetExt) == 0)
 		{
 			FString TargetName = FPaths::GetCleanFilename(FString(Length - TargetExtLen, FileNameOrDirectory));
@@ -782,7 +782,7 @@ struct FTargetFileVisitor : IPlatformFile::FDirectoryStatVisitor
 		}
 		
 		static const TCHAR ModuleExt[] = TEXT(".Build.cs");
-		static const int32 ModuleExtLen = ARRAY_COUNT(ModuleExt) - 1;
+		static const int32 ModuleExtLen = UE_ARRAY_COUNT(ModuleExt) - 1;
 		if (Length > ModuleExtLen && FCString::Stricmp(FileNameOrDirectory + Length - ModuleExtLen, ModuleExt) == 0)
 		{
 			bSearchSubDirectories = false;
diff --git a/Engine/Source/Programs/UnrealHeaderTool/Private/HeaderParser.cpp b/Engine/Source/Programs/UnrealHeaderTool/Private/HeaderParser.cpp
--- a/Engine/Source/Programs/UnrealHeaderTool/Private/HeaderParser.cpp
+++ b/Engine/Source/Programs/UnrealHeaderTool/Private/HeaderParser.cpp
@@ -2078,7 +2078,7 @@ TMap<FName, FString> FHeaderParser::GetParameterToolTipsFromFunctionComment(cons
 		int32 ParamStart = Input.Find(ParamTag, ESearchCase::CaseSensitive, ESearchDir::FromStart, Offset);
 		if(ParamStart != INDEX_NONE)
 		{
-			ParamStart = ParamStart + ARRAY_COUNT(ParamTag);
+			ParamStart = ParamStart + UE_ARRAY_COUNT(ParamTag);
 			Offset = ParamStart;
 		}
 		else
@@ -2086,7 +2086,7 @@ TMap<FName, FString> FHeaderParser::GetParameterToolTipsFromFunctionComment(cons
 			ParamStart = Input.Find(ReturnTag, ESearchCase::CaseSensitive, ESearchDir::FromStart, Offset);
 			if (ParamStart != INDEX_NONE)
 			{
-				ParamStart = ParamStart + ARRAY_COUNT(ReturnTag);
+				ParamStart = ParamStart + UE_ARRAY_COUNT(ReturnTag);
 				Offset = ParamStart;
 				ParamPrefix = ReturnParamPrefix;
 			}
diff --git a/Engine/Source/Runtime/Core/Private/Logging/LogSuppressionInterface.cpp b/Engine/Source/Runtime/Core/Private/Logging/LogSuppressionInterface.cpp
--- a/Engine/Source/Runtime/Core/Private/Logging/LogSuppressionInterface.cpp
+++ b/Engine/Source/Runtime/Core/Private/Logging/LogSuppressionInterface.cpp
@@ -136,7 +136,7 @@ class FLogSuppressionImplementation: public FLogSuppressionInterface, private FS
 			FString Command = SubCmds[Index].TrimStart();
 			if (Command.StartsWith(LogString))
 			{
-				Command = Command.Right(Command.Len() - (ARRAY_COUNT(LogString) - 1));
+				Command = Command.Right(Command.Len() - (UE_ARRAY_COUNT(LogString) - 1));
 			}
 			TArray<FString> CommandParts;
 			Command.ParseIntoArrayWS(CommandParts);
diff --git a/Engine/Source/Runtime/Core/Private/Misc/ConfigCacheIni.cpp b/Engine/Source/Runtime/Core/Private/Misc/ConfigCacheIni.cpp
--- a/Engine/Source/Runtime/Core/Private/Misc/ConfigCacheIni.cpp
+++ b/Engine/Source/Runtime/Core/Private/Misc/ConfigCacheIni.cpp
@@ -1157,7 +1157,7 @@ void FConfigFile::OverrideFromCommandline(FConfigFile* File, const FString& File
 				CommandlineOption.Section = CommandlineOption.Section.Replace(CommandlineOverrideSpecifiers::PropertyStartIdentifier, TEXT(""));
 				CommandlineOption.Section = CommandlineOption.Section.Replace(CommandlineOverrideSpecifiers::SectionStartIdentifier, TEXT(""));
 
-				CommandlineOption.PropertyKey = SectionAndKey.Mid(SectionNameEndIndex + ARRAY_COUNT(CommandlineOverrideSpecifiers::PropertyStartIdentifier) - 1);
+				CommandlineOption.PropertyKey = SectionAndKey.Mid(SectionNameEndIndex + UE_ARRAY_COUNT(CommandlineOverrideSpecifiers::PropertyStartIdentifier) - 1);
 				CommandlineOption.PropertyValue = Value;
 
 				// now put it into this into the cache
diff --git a/Engine/Source/Runtime/Core/Private/Windows/WindowsPlatformMisc.cpp b/Engine/Source/Runtime/Core/Private/Windows/WindowsPlatformMisc.cpp
--- a/Engine/Source/Runtime/Core/Private/Windows/WindowsPlatformMisc.cpp
+++ b/Engine/Source/Runtime/Core/Private/Windows/WindowsPlatformMisc.cpp
@@ -381,7 +381,7 @@ int32 FWindowsOSVersionHelper::GetOSVersions( FString& OutOSVersionLabel, FStrin
 	OSVersionLabel[0] = 0;
 	OSSubVersionLabel[0] = 0;
 
-	int32 Result = GetOSVersionsHelper( OSVersionLabel, ARRAY_COUNT(OSVersionLabel), OSSubVersionLabel, ARRAY_COUNT(OSSubVersionLabel) );
+	int32 Result = GetOSVersionsHelper( OSVersionLabel, UE_ARRAY_COUNT(OSVersionLabel), OSSubVersionLabel, UE_ARRAY_COUNT(OSSubVersionLabel) );
 
 	OutOSVersionLabel = OSVersionLabel;
 	OutOSSubVersionLabel = OSSubVersionLabel;
@@ -2521,7 +2521,7 @@ void FWindowsPlatformMisc::GetOSVersions( FString& OutOSVersionLabel, FString& O
 		{
 			OSVersionLabel[0] = 0;
 			OSSubVersionLabel[0] = 0;
-			GetOSVersionsHelper( OSVersionLabel, ARRAY_COUNT(OSVersionLabel), OSSubVersionLabel, ARRAY_COUNT(OSSubVersionLabel) );
+			GetOSVersionsHelper( OSVersionLabel, UE_ARRAY_COUNT(OSVersionLabel), OSSubVersionLabel, UE_ARRAY_COUNT(OSSubVersionLabel) );
 		}
 
 		TCHAR OSVersionLabel[128];
@@ -2540,7 +2540,7 @@ FString FWindowsPlatformMisc::GetOSVersion()
 		FOSVersionInitializer()
 		{
 			CachedOSVersion[0] = 0;
-			if (!GetOSVersionHelper(CachedOSVersion, ARRAY_COUNT(CachedOSVersion)))
+			if (!GetOSVersionHelper(CachedOSVersion, UE_ARRAY_COUNT(CachedOSVersion)))
 			{
 				CachedOSVersion[0] = 0;
 			}
diff --git a/Engine/Source/Runtime/Engine/Public/PrecomputedVolumetricLightmap.h b/Engine/Source/Runtime/Engine/Public/PrecomputedVolumetricLightmap.h
--- a/Engine/Source/Runtime/Engine/Public/PrecomputedVolumetricLightmap.h
+++ b/Engine/Source/Runtime/Engine/Public/PrecomputedVolumetricLightmap.h
@@ -87,7 +87,7 @@ public:
 	{
 		AmbientVector.Texture.SafeRelease();
 
-		for (int32 i = 0; i < ARRAY_COUNT(SHCoefficients); i++)
+		for (int32 i = 0; i < UE_ARRAY_COUNT(SHCoefficients); i++)
 		{
 			SHCoefficients[i].Texture.SafeRelease();
 		}
@@ -97,7 +97,7 @@ public:
 
 		AmbientVector.UAV.SafeRelease();
 
-		for (int32 i = 0; i < ARRAY_COUNT(SHCoefficients); i++)
+		for (int32 i = 0; i < UE_ARRAY_COUNT(SHCoefficients); i++)
 		{
 			SHCoefficients[i].UAV.SafeRelease();
 		}
@@ -123,7 +123,7 @@ public:
 	{
 		AmbientVector.bNeedsCPUAccess = InAccess;
 
-		for (int32 i = 0; i < ARRAY_COUNT(SHCoefficients); i++)
+		for (int32 i = 0; i < UE_ARRAY_COUNT(SHCoefficients); i++)
 		{
 			SHCoefficients[i].bNeedsCPUAccess = InAccess;
 		}
