- Potential threading fix for crash with GTickWakeMap

diff --git a/Engine/Source/Runtime/Core/Private/Misc/EmbeddedCommunication.cpp b/Engine/Source/Runtime/Core/Private/Misc/EmbeddedCommunication.cpp
--- a/Engine/Source/Runtime/Core/Private/Misc/EmbeddedCommunication.cpp
+++ b/Engine/Source/Runtime/Core/Private/Misc/EmbeddedCommunication.cpp
@@ -369,6 +369,8 @@ bool FEmbeddedCommunication::TickGameThread(float DeltaTime)
 FString FEmbeddedCommunication::GetDebugInfo()
 {
 #if BUILD_EMBEDDED_APP
+	FScopeLock Lock(&GEmbeddedLock);
+
 	FString Str = "";
 	for (auto It : GRenderingWakeMap)
 	{
