Fix FGuardValue_Bitfield macro expansion to no longer hide previous usage within same or nested scope.

diff --git a/Engine/Source/Runtime/Core/Public/Templates/UnrealTemplate.h b/Engine/Source/Runtime/Core/Public/Templates/UnrealTemplate.h
--- a/Engine/Source/Runtime/Core/Public/Templates/UnrealTemplate.h
+++ b/Engine/Source/Runtime/Core/Public/Templates/UnrealTemplate.h
@@ -301,9 +301,9 @@ private:
  * Macro variant on TGuardValue<bool> that can deal with bitfields which cannot be passed by reference in to TGuardValue
  */
 #define FGuardValue_Bitfield(ReferenceValue, NewValue) \
-	const bool TempBitfield##__LINE__ = ReferenceValue; \
+	const bool PREPROCESSOR_JOIN(TempBitfield, __LINE__) = ReferenceValue; \
 	ReferenceValue = NewValue; \
-	const TGuardValue_Bitfield_Cleanup<TFunction<void()>> TempBitfieldCleanup##__LINE__([&](){ ReferenceValue = TempBitfield##__LINE__; });
+	const TGuardValue_Bitfield_Cleanup<TFunction<void()>> PREPROCESSOR_JOIN(TempBitfieldCleanup, __LINE__)([&](){ ReferenceValue = PREPROCESSOR_JOIN(TempBitfield, __LINE__); });
 
 /** 
  * Commonly used to make sure a value is incremented, and then decremented anyway the function can terminate.
