Changed a check to a checkSlow in FNamePool::Find to reduce the impact of FName creation through EName (happens a lot when comparing to NAME_None).

diff --git a/Engine/Source/Runtime/Core/Private/UObject/UnrealNames.cpp b/Engine/Source/Runtime/Core/Private/UObject/UnrealNames.cpp
--- a/Engine/Source/Runtime/Core/Private/UObject/UnrealNames.cpp
+++ b/Engine/Source/Runtime/Core/Private/UObject/UnrealNames.cpp
@@ -807,7 +807,7 @@ static bool IsPureAnsi(const WIDECHAR* Str, const int32 Len)
 
 FNameEntryId FNamePool::Find(EName Ename) const
 {
-	check(Ename < NAME_MaxHardcodedNameIndex);
+	checkSlow(Ename < NAME_MaxHardcodedNameIndex);
 	return ENameToEntry[Ename];
 }
 
