Added missing ELLMTag::Fname scopes around FName allocs.
This steals back some of the LLM tagged memory from other categories that should have been marked as FName.

In MB:
AssetRegistry:  61.98   -> 43.60
AsyncLoading:   25.27   -> 19.54
EngineInit:     20.12   -> 19.12
EnginePreInit:  38.33   -> 13.33
GC:             10.21   -> 7.57
Stats:          51.10   -> 50.78
UObject:        270.85  -> 268.20
Untagged:       32.30   -> 30.67
FName:          0       -> 54.25

#rnx

diff --git a/Engine/Source/Runtime/Core/Private/UObject/UnrealNames.cpp b/Engine/Source/Runtime/Core/Private/UObject/UnrealNames.cpp
--- a/Engine/Source/Runtime/Core/Private/UObject/UnrealNames.cpp
+++ b/Engine/Source/Runtime/Core/Private/UObject/UnrealNames.cpp
@@ -265,6 +265,7 @@ public:
 	/** Initializes all member variables. */
 	FNameEntryAllocator()
 	{
+		LLM_SCOPE(ELLMTag::FName);
 		Blocks[0] = (uint8*)FMemory::Malloc(BlockSizeBytes, FPlatformMemory::GetConstants().PageSize);
 	}
 
@@ -376,6 +377,7 @@ private:
 	/** Allocates a new pool. */
 	void AllocateNewBlock()
 	{
+		LLM_SCOPE(ELLMTag::FName);
 		// Null-terminate final entry to allow DebugDump() entry iteration
 		if (CurrentByteCursor + FNameEntry::GetDataOffset() <= BlockSizeBytes)
 		{
@@ -548,6 +550,7 @@ class alignas(PLATFORM_CACHE_LINE_SIZE) FNamePoolShardBase : FNoncopyable
 public:
 	void Initialize(FNameEntryAllocator& InEntries)
 	{
+		LLM_SCOPE(ELLMTag::FName);
 		Entries = &InEntries;
 
 		Slots = (FNameSlot*)FMemory::Malloc(FNamePoolInitialSlotsPerShard * sizeof(FNameSlot), alignof(FNameSlot));
@@ -642,6 +645,7 @@ private:
 
 	void Grow()
 	{
+		LLM_SCOPE(ELLMTag::FName);
 		FNameSlot* const OldSlots = Slots;
 		const uint32 OldUsedSlots = UsedSlots;
 		const uint32 OldCapacity = Capacity();
