Removed the FArchive-based serialization overloads on FLinearColor and FVector, and swapped the related structops traits over to use the structured serializer

diff --git a/Engine/Source/Runtime/Core/Public/Math/Color.h b/Engine/Source/Runtime/Core/Public/Math/Color.h
--- a/Engine/Source/Runtime/Core/Public/Math/Color.h
+++ b/Engine/Source/Runtime/Core/Public/Math/Color.h
@@ -60,18 +60,6 @@ struct FLinearColor
 	CORE_API explicit FLinearColor(const FFloat16Color& C);
 
 	// Serializer.
-
-	friend FArchive& operator<<(FArchive& Ar,FLinearColor& Color)
-	{
-		return Ar << Color.R << Color.G << Color.B << Color.A;
-	}
-
-	bool Serialize( FArchive& Ar )
-	{
-		Ar << *this;
-		return true;
-	}
-
 	friend void operator<<(FStructuredArchive::FSlot Slot, FLinearColor& Color)
 	{
 		FStructuredArchive::FRecord Record = Slot.EnterRecord();
diff --git a/Engine/Source/Runtime/Core/Public/Math/Vector.h b/Engine/Source/Runtime/Core/Public/Math/Vector.h
--- a/Engine/Source/Runtime/Core/Public/Math/Vector.h
+++ b/Engine/Source/Runtime/Core/Public/Math/Vector.h
@@ -982,20 +982,6 @@ public:
 	 */
 	static CORE_API void GenerateClusterCenters(TArray<FVector>& Clusters, const TArray<FVector>& Points, int32 NumIterations, int32 NumConnectionsToBeValid);
 
-	/**
-	 * Serializer.
-	 *
-	 * @param Ar Serialization Archive.
-	 * @param V Vector to serialize.
-	 * @return Reference to Archive after serialization.
-	 */
-	friend FArchive& operator<<(FArchive& Ar, FVector& V)
-	{
-		// @warning BulkSerialize: FVector is serialized as memory dump
-		// See TArray::BulkSerialize for detailed description of implied limitations.
-		return Ar << V.X << V.Y << V.Z;
-	}
-
 	/**
 	 * Structured archive slot serializer.
 	 *
@@ -1006,18 +992,12 @@ public:
 	{
 		// @warning BulkSerialize: FVector is serialized as memory dump
 		// See TArray::BulkSerialize for detailed description of implied limitations.
-		FStructuredArchive::FStream Stream = Slot.EnterStream();
-		Stream.EnterElement() << V.X;
-		Stream.EnterElement() << V.Y;
-		Stream.EnterElement() << V.Z;
+		FStructuredArchive::FRecord Record = Slot.EnterRecord();
+		Record << NAMED_ITEM("X", V.X);
+		Record << NAMED_ITEM("Y", V.Y);
+		Record << NAMED_ITEM("Z", V.Z);
 	}
 	
-	bool Serialize(FArchive& Ar)
-	{
-		Ar << *this;
-		return true;
-	}
-
 	bool Serialize(FStructuredArchive::FSlot Slot)
 	{
 		Slot << *this;
diff --git a/Engine/Source/Runtime/CoreUObject/Private/UObject/Property.cpp b/Engine/Source/Runtime/CoreUObject/Private/UObject/Property.cpp
--- a/Engine/Source/Runtime/CoreUObject/Private/UObject/Property.cpp
+++ b/Engine/Source/Runtime/CoreUObject/Private/UObject/Property.cpp
@@ -30,7 +30,7 @@ struct TStructOpsTypeTraits<FVector> : public TStructOpsTypeTraitsBase2<FVector>
 		WithZeroConstructor = true,
 		WithNetSerializer = true,
 		WithNetSharedSerialization = true,
-		WithSerializer = true,
+		WithStructuredSerializer = true,
 	};
 };
 IMPLEMENT_STRUCT(Vector);
@@ -172,7 +172,7 @@ struct TStructOpsTypeTraits<FLinearColor> : public TStructOpsTypeTraitsBase2<FLi
 	{
 		WithNoInitConstructor = true,
 		WithZeroConstructor = true,
-		WithSerializer = true,
+		WithStructuredSerializer = true,
 	};
 };
 IMPLEMENT_STRUCT(LinearColor);
@@ -297,7 +297,7 @@ struct TStructOpsTypeTraits<FSoftObjectPath> : public TStructOpsTypeTraitsBase2<
 	enum
 	{
 		WithZeroConstructor = true,
-		WithSerializer = true,
+		WithStructuredSerializer = true,
 		WithCopy = true,
 		WithIdenticalViaEquality = true,
 		WithExportTextItem = true,
