Incremental build fix.

diff --git a/Engine/Source/Runtime/Core/Public/Serialization/ArchiveAdapters.h b/Engine/Source/Runtime/Core/Public/Serialization/ArchiveAdapters.h
--- a/Engine/Source/Runtime/Core/Public/Serialization/ArchiveAdapters.h
+++ b/Engine/Source/Runtime/Core/Public/Serialization/ArchiveAdapters.h
@@ -3,6 +3,8 @@
 #pragma once
 
 #include "Concepts/Insertable.h"
+#include "Serialization/ArchiveFromStructuredArchive.h"
+#include "Serialization/StructuredArchiveFromArchive.h"
 #include "Templates/Models.h"
 
 /**
