Back out changelist 6276384.

diff --git a/Engine/Source/Runtime/Core/Public/Serialization/ArchiveAdapters.h b/Engine/Source/Runtime/Core/Public/Serialization/ArchiveAdapters.h
--- a/Engine/Source/Runtime/Core/Public/Serialization/ArchiveAdapters.h
+++ b/Engine/Source/Runtime/Core/Public/Serialization/ArchiveAdapters.h
@@ -3,8 +3,6 @@
 #pragma once
 
 #include "Concepts/Insertable.h"
-#include "Serialization/ArchiveFromStructuredArchive.h"
-#include "Serialization/StructuredArchiveFromArchive.h"
 #include "Templates/Models.h"
 
 /**
