Switched to faster and more sane implementation GetTypeHash(FSHAHash) to help speed up asset loading

diff --git a/Engine/Source/Runtime/Core/Private/Misc/SecureHash.cpp b/Engine/Source/Runtime/Core/Private/Misc/SecureHash.cpp
--- a/Engine/Source/Runtime/Core/Private/Misc/SecureHash.cpp
+++ b/Engine/Source/Runtime/Core/Private/Misc/SecureHash.cpp
@@ -605,11 +605,6 @@ FArchive& operator<<( FArchive& Ar, FSHAHash& G )
 	return Ar;
 }
 
-uint32 GetTypeHash(FSHAHash const& InKey)
-{
-	return FCrc::MemCrc32(InKey.Hash, sizeof(InKey.Hash));
-}
-
 FString LexToString(const FSHAHash& InHash)
 {
 	return InHash.ToString();
diff --git a/Engine/Source/Runtime/Core/Public/Misc/SecureHash.h b/Engine/Source/Runtime/Core/Public/Misc/SecureHash.h
--- a/Engine/Source/Runtime/Core/Public/Misc/SecureHash.h
+++ b/Engine/Source/Runtime/Core/Public/Misc/SecureHash.h
@@ -193,7 +193,7 @@ typedef union
 class CORE_API FSHAHash
 {
 public:
-	uint8 Hash[20];
+	alignas(uint32) uint8 Hash[20];
 
 	FSHAHash()
 	{
@@ -222,7 +222,10 @@ public:
 
 	friend CORE_API FArchive& operator<<( FArchive& Ar, FSHAHash& G );
 	
-	friend CORE_API uint32 GetTypeHash(FSHAHash const& InKey);
+	friend uint32 GetTypeHash(const FSHAHash& InKey)
+	{
+		return *reinterpret_cast<const uint32*>(InKey.Hash);
+	}
 
 	friend CORE_API FString LexToString(const FSHAHash&);
 	friend CORE_API void LexFromString(FSHAHash& Hash, const TCHAR*);
