Allow cvars to migrated across hot reload.

diff --git a/Engine/Source/Runtime/Core/Private/HAL/ConsoleManager.cpp b/Engine/Source/Runtime/Core/Private/HAL/ConsoleManager.cpp
--- a/Engine/Source/Runtime/Core/Private/HAL/ConsoleManager.cpp
+++ b/Engine/Source/Runtime/Core/Private/HAL/ConsoleManager.cpp
@@ -303,7 +303,7 @@ public:
 	{
 	}
 
-	// interface IConsoleVariable -----------------------------------
+	// interface IConsoleVariable ----------------------------------- 
 
 	virtual void Release()
 	{
@@ -321,7 +321,9 @@ public:
 	virtual int32 GetInt() const;
 	virtual float GetFloat() const;
 	virtual FString GetString() const;
-	virtual bool IsVariableInt() const { return false; }
+	virtual bool IsVariableInt() const override { return false; }
+	virtual bool IsVariableFloat() const override { return false; }
+	virtual bool IsVariableString() const override { return false; }
 	virtual class TConsoleVariableData<int32>* AsVariableInt() { return 0; }
 	virtual class TConsoleVariableData<float>* AsVariableFloat() { return 0; }
 	virtual class TConsoleVariableData<FString>* AsVariableString() { return 0; }
@@ -385,9 +387,13 @@ template<> FString FConsoleVariable<float>::GetString() const
 {
 	return FString::Printf(TEXT("%g"), Value());
 }
+template<> bool FConsoleVariable<float>::IsVariableFloat() const
+{
+	return true;
+}
 template<> TConsoleVariableData<float>* FConsoleVariable<float>::AsVariableFloat()
 {
-	return &Data; 
+	return &Data;
 }
 
 // specialization for FString
@@ -412,7 +418,10 @@ template<> FString FConsoleVariable<FString>::GetString() const
 {
 	return Value();
 }
-
+template<> bool FConsoleVariable<FString>::IsVariableString() const
+{
+	return true;
+}
 template<> TConsoleVariableData<FString>* FConsoleVariable<FString>::AsVariableString()
 {
 	return &Data;
@@ -456,6 +465,9 @@ public:
 	{
 		return TTypeToString<T>::ToString(MainValue);
 	}
+	virtual bool IsVariableInt() const override { return false; }
+	virtual bool IsVariableFloat() const override { return false; }
+	virtual bool IsVariableString() const override { return false; }
 
 private: // ----------------------------------------------------
 
@@ -490,6 +502,18 @@ FString FConsoleVariableRef<float>::GetString() const
 	// otherwise we get 2.1f would become "2.100000"
 	return FString::SanitizeFloat(RefValue);
 }
+template<> bool FConsoleVariableRef<bool>::IsVariableInt() const
+{
+	return true;
+}
+template<> bool FConsoleVariableRef<int32>::IsVariableInt() const
+{
+	return true;
+}
+template<> bool FConsoleVariableRef<float>::IsVariableFloat() const
+{
+	return true;
+}
 
 // string version
 
@@ -533,6 +557,10 @@ public:
 	{
 		return MainValue;
 	}
+	virtual bool IsVariableString() const override
+	{
+		return true;
+	}
 
 private: // ----------------------------------------------------
 
@@ -1412,6 +1440,31 @@ IConsoleObject* FConsoleManager::AddConsoleObject(const TCHAR* Name, IConsoleObj
 				ConsoleObjects.Add(Name, Var);
 				return Var;
 			}
+			else if (GIsHotReload)
+			{
+				// Variable is being replaced due to a hot reload - copy state across to new variable, but only if the type hasn't changed
+				{
+					if (ExistingVar->IsVariableFloat())
+					{
+						Var->Set(ExistingVar->GetFloat());
+					}
+				}
+				{
+					if (ExistingVar->IsVariableInt())
+					{
+						Var->Set(ExistingVar->GetInt());
+					}
+				}
+				{
+					if (ExistingVar->IsVariableString())
+					{
+						Var->Set(*ExistingVar->GetString());
+					}
+				}
+				ExistingVar->Release();
+				ConsoleObjects.Add(Name, Var);
+				return Var;
+			}
 			else
 			{
 				// Copy data over from the new variable,
diff --git a/Engine/Source/Runtime/Core/Public/HAL/IConsoleManager.h b/Engine/Source/Runtime/Core/Public/HAL/IConsoleManager.h
--- a/Engine/Source/Runtime/Core/Public/HAL/IConsoleManager.h
+++ b/Engine/Source/Runtime/Core/Public/HAL/IConsoleManager.h
@@ -214,6 +214,8 @@ public:
 	}
 
 	virtual bool IsVariableInt() const { return false; }
+	virtual bool IsVariableFloat() const { return false; }
+	virtual bool IsVariableString() const { return false; }
 
 	virtual class TConsoleVariableData<int32>* AsVariableInt()
 	{
