CIS fix.

diff --git a/Engine/Source/Runtime/Core/Private/HAL/ConsoleManager.cpp b/Engine/Source/Runtime/Core/Private/HAL/ConsoleManager.cpp
--- a/Engine/Source/Runtime/Core/Private/HAL/ConsoleManager.cpp
+++ b/Engine/Source/Runtime/Core/Private/HAL/ConsoleManager.cpp
@@ -1440,6 +1440,7 @@ IConsoleObject* FConsoleManager::AddConsoleObject(const TCHAR* Name, IConsoleObj
 				ConsoleObjects.Add(Name, Var);
 				return Var;
 			}
+#if WITH_HOT_RELOAD
 			else if (GIsHotReload)
 			{
 				// Variable is being replaced due to a hot reload - copy state across to new variable, but only if the type hasn't changed
@@ -1465,6 +1466,7 @@ IConsoleObject* FConsoleManager::AddConsoleObject(const TCHAR* Name, IConsoleObj
 				ConsoleObjects.Add(Name, Var);
 				return Var;
 			}
+#endif
 			else
 			{
 				// Copy data over from the new variable,
