Fixes to allocator API misuse.

diff --git a/Engine/Source/Developer/ShaderCompilerCommon/Private/HlslUtils.h b/Engine/Source/Developer/ShaderCompilerCommon/Private/HlslUtils.h
--- a/Engine/Source/Developer/ShaderCompilerCommon/Private/HlslUtils.h
+++ b/Engine/Source/Developer/ShaderCompilerCommon/Private/HlslUtils.h
@@ -129,7 +129,7 @@ namespace CrossCompiler
 			{
 				return Data;
 			}
-			void ResizeAllocation(int32 PreviousNumElements, int32 NumElements, int32 NumBytesPerElement)
+			void ResizeAllocation(int32 PreviousNumElements, int32 NumElements, SIZE_T NumBytesPerElement)
 			{
 				void* OldData = Data;
 				if (NumElements)
@@ -147,25 +147,25 @@ namespace CrossCompiler
 					}
 				}
 			}
-			int32 CalculateSlackReserve(int32 NumElements, int32 NumBytesPerElement) const
+			int32 CalculateSlackReserve(int32 NumElements, SIZE_T NumBytesPerElement) const
 			{
 				return DefaultCalculateSlackReserve(NumElements, NumBytesPerElement, false);
 			}
-			int32 CalculateSlackShrink(int32 NumElements, int32 NumAllocatedElements, int32 NumBytesPerElement) const
+			int32 CalculateSlackShrink(int32 NumElements, int32 NumAllocatedElements, SIZE_T NumBytesPerElement) const
 			{
 				return DefaultCalculateSlackShrink(NumElements, NumAllocatedElements, NumBytesPerElement, false);
 			}
-			int32 CalculateSlackGrow(int32 NumElements, int32 NumAllocatedElements, int32 NumBytesPerElement) const
+			int32 CalculateSlackGrow(int32 NumElements, int32 NumAllocatedElements, SIZE_T NumBytesPerElement) const
 			{
 				return DefaultCalculateSlackGrow(NumElements, NumAllocatedElements, NumBytesPerElement, false);
 			}
 
-			int32 GetAllocatedSize(int32 NumAllocatedElements, int32 NumBytesPerElement) const
+			SIZE_T GetAllocatedSize(int32 NumAllocatedElements, SIZE_T NumBytesPerElement) const
 			{
 				return NumAllocatedElements * NumBytesPerElement;
 			}
 
-			bool HasAllocation()
+			bool HasAllocation() const
 			{
 				return !!Data;
 			}
diff --git a/Engine/Source/Runtime/Core/Private/HAL/Allocators/AnsiAllocator.h b/Engine/Source/Runtime/Core/Private/HAL/Allocators/AnsiAllocator.h
--- a/Engine/Source/Runtime/Core/Private/HAL/Allocators/AnsiAllocator.h
+++ b/Engine/Source/Runtime/Core/Private/HAL/Allocators/AnsiAllocator.h
@@ -76,15 +76,15 @@ public:
 				Data = nullptr;
 			}
 		}
-		int32 CalculateSlackReserve(int32 NumElements, int32 NumBytesPerElement) const
+		int32 CalculateSlackReserve(int32 NumElements, SIZE_T NumBytesPerElement) const
 		{
 			return DefaultCalculateSlackReserve(NumElements, NumBytesPerElement, false);
 		}
-		int32 CalculateSlackShrink(int32 NumElements, int32 NumAllocatedElements, int32 NumBytesPerElement) const
+		int32 CalculateSlackShrink(int32 NumElements, int32 NumAllocatedElements, SIZE_T NumBytesPerElement) const
 		{
 			return DefaultCalculateSlackShrink(NumElements, NumAllocatedElements, NumBytesPerElement, false);
 		}
-		int32 CalculateSlackGrow(int32 NumElements, int32 NumAllocatedElements, int32 NumBytesPerElement) const
+		int32 CalculateSlackGrow(int32 NumElements, int32 NumAllocatedElements, SIZE_T NumBytesPerElement) const
 		{
 			return DefaultCalculateSlackGrow(NumElements, NumAllocatedElements, NumBytesPerElement, false);
 		}
@@ -94,7 +94,7 @@ public:
 			return NumAllocatedElements * NumBytesPerElement;
 		}
 
-		bool HasAllocation()
+		bool HasAllocation() const
 		{
 			return !!Data;
 		}
diff --git a/Engine/Source/Runtime/Core/Public/Containers/ContainerAllocationPolicies.h b/Engine/Source/Runtime/Core/Public/Containers/ContainerAllocationPolicies.h
--- a/Engine/Source/Runtime/Core/Public/Containers/ContainerAllocationPolicies.h
+++ b/Engine/Source/Runtime/Core/Public/Containers/ContainerAllocationPolicies.h
@@ -170,7 +170,7 @@ public:
 		 * @param CurrentNumSlackElements - The current number of elements allocated.
 		 * @param NumBytesPerElement - The number of bytes/element.
 		 */
-		int32 CalculateSlack(
+		int32 CalculateSlackReserve(
 			int32 NumElements,
 			int32 CurrentNumSlackElements,
 			SIZE_T NumBytesPerElement
@@ -200,7 +200,15 @@ public:
 			SIZE_T NumBytesPerElement
 			) const;
 
+		/**
+		 * Returns the size of any requested heap allocation currently owned by the allocator.
+		 * @param NumAllocatedElements - The number of elements allocated by the container.
+		 * @param NumBytesPerElement - The number of bytes/element.
+		 */
 		SIZE_T GetAllocatedSize(int32 NumAllocatedElements, SIZE_T NumBytesPerElement) const;
+
+		/** Returns true if the allocator has made any heap allocations */
+		bool HasAllocation() const;
 	};
 
 	/**
@@ -273,15 +281,15 @@ public:
 				Data = (FScriptContainerElement*)FMemory::Realloc( Data, NumElements*NumBytesPerElement, Alignment );
 			}
 		}
-		FORCEINLINE int32 CalculateSlackReserve(int32 NumElements, int32 NumBytesPerElement) const
+		FORCEINLINE int32 CalculateSlackReserve(int32 NumElements, SIZE_T NumBytesPerElement) const
 		{
 			return DefaultCalculateSlackReserve(NumElements, NumBytesPerElement, true, Alignment);
 		}
-		FORCEINLINE int32 CalculateSlackShrink(int32 NumElements, int32 NumAllocatedElements, int32 NumBytesPerElement) const
+		FORCEINLINE int32 CalculateSlackShrink(int32 NumElements, int32 NumAllocatedElements, SIZE_T NumBytesPerElement) const
 		{
 			return DefaultCalculateSlackShrink(NumElements, NumAllocatedElements, NumBytesPerElement, true, Alignment);
 		}
-		FORCEINLINE int32 CalculateSlackGrow(int32 NumElements, int32 NumAllocatedElements, int32 NumBytesPerElement) const
+		FORCEINLINE int32 CalculateSlackGrow(int32 NumElements, int32 NumAllocatedElements, SIZE_T NumBytesPerElement) const
 		{
 			return DefaultCalculateSlackGrow(NumElements, NumAllocatedElements, NumBytesPerElement, true, Alignment);
 		}
@@ -291,7 +299,7 @@ public:
 			return NumAllocatedElements * NumBytesPerElement;
 		}
 
-		bool HasAllocation()
+		bool HasAllocation() const
 		{
 			return !!Data;
 		}
@@ -384,15 +392,15 @@ public:
 				Data = (FScriptContainerElement*)FMemory::Realloc( Data, NumElements*NumBytesPerElement );
 			}
 		}
-		FORCEINLINE int32 CalculateSlackReserve(int32 NumElements, int32 NumBytesPerElement) const
+		FORCEINLINE int32 CalculateSlackReserve(int32 NumElements, SIZE_T NumBytesPerElement) const
 		{
 			return DefaultCalculateSlackReserve(NumElements, NumBytesPerElement, true);
 		}
-		FORCEINLINE int32 CalculateSlackShrink(int32 NumElements, int32 NumAllocatedElements, int32 NumBytesPerElement) const
+		FORCEINLINE int32 CalculateSlackShrink(int32 NumElements, int32 NumAllocatedElements, SIZE_T NumBytesPerElement) const
 		{
 			return DefaultCalculateSlackShrink(NumElements, NumAllocatedElements, NumBytesPerElement, true);
 		}
-		FORCEINLINE int32 CalculateSlackGrow(int32 NumElements, int32 NumAllocatedElements, int32 NumBytesPerElement) const
+		FORCEINLINE int32 CalculateSlackGrow(int32 NumElements, int32 NumAllocatedElements, SIZE_T NumBytesPerElement) const
 		{
 			return DefaultCalculateSlackGrow(NumElements, NumAllocatedElements, NumBytesPerElement, true);
 		}
@@ -402,7 +410,7 @@ public:
 			return NumAllocatedElements * NumBytesPerElement;
 		}
 
-		bool HasAllocation()
+		bool HasAllocation() const
 		{
 			return !!Data;
 		}
@@ -528,14 +536,14 @@ public:
 				NumInlineElements :
 				SecondaryData.CalculateSlackReserve(NumElements, NumBytesPerElement);
 		}
-		FORCEINLINE int32 CalculateSlackShrink(int32 NumElements, int32 NumAllocatedElements, int32 NumBytesPerElement) const
+		FORCEINLINE int32 CalculateSlackShrink(int32 NumElements, int32 NumAllocatedElements, SIZE_T NumBytesPerElement) const
 		{
 			// If the elements use less space than the inline allocation, only use the inline allocation as slack.
 			return NumElements <= NumInlineElements ?
 				NumInlineElements :
 				SecondaryData.CalculateSlackShrink(NumElements, NumAllocatedElements, NumBytesPerElement);
 		}
-		FORCEINLINE int32 CalculateSlackGrow(int32 NumElements, int32 NumAllocatedElements, int32 NumBytesPerElement) const
+		FORCEINLINE int32 CalculateSlackGrow(int32 NumElements, int32 NumAllocatedElements, SIZE_T NumBytesPerElement) const
 		{
 			// If the elements use less space than the inline allocation, only use the inline allocation as slack.
 			return NumElements <= NumInlineElements ?
@@ -552,7 +560,7 @@ public:
 			return 0;
 		}
 
-		bool HasAllocation()
+		bool HasAllocation() const
 		{
 			return SecondaryData.HasAllocation();
 		}
@@ -676,13 +684,13 @@ public:
 			return (NumElements <= NumInlineElements) ? NumInlineElements : DefaultCalculateSlackReserve(NumElements, NumBytesPerElement, true);
 		}
 
-		FORCEINLINE int32 CalculateSlackShrink(int32 NumElements, int32 NumAllocatedElements, int32 NumBytesPerElement) const
+		FORCEINLINE int32 CalculateSlackShrink(int32 NumElements, int32 NumAllocatedElements, SIZE_T NumBytesPerElement) const
 		{
 			// If the elements use less space than the inline allocation, only use the inline allocation as slack.
 			return (NumElements <= NumInlineElements) ? NumInlineElements : DefaultCalculateSlackShrink(NumElements, NumAllocatedElements, NumBytesPerElement, true);
 		}
 
-		FORCEINLINE int32 CalculateSlackGrow(int32 NumElements, int32 NumAllocatedElements, int32 NumBytesPerElement) const
+		FORCEINLINE int32 CalculateSlackGrow(int32 NumElements, int32 NumAllocatedElements, SIZE_T NumBytesPerElement) const
 		{
 			// If the elements use less space than the inline allocation, only use the inline allocation as slack.
 			return (NumElements <= NumInlineElements) ? NumInlineElements : DefaultCalculateSlackGrow(NumElements, NumAllocatedElements, NumBytesPerElement, true);
@@ -777,13 +785,13 @@ public:
 			checkSlow(NumElements <= NumInlineElements);
 			return NumInlineElements;
 		}
-		FORCEINLINE int32 CalculateSlackShrink(int32 NumElements, int32 NumAllocatedElements, int32 NumBytesPerElement) const
+		FORCEINLINE int32 CalculateSlackShrink(int32 NumElements, int32 NumAllocatedElements, SIZE_T NumBytesPerElement) const
 		{
 			// Ensure the requested allocation will fit in the inline data area.
 			checkSlow(NumAllocatedElements <= NumInlineElements);
 			return NumInlineElements;
 		}
-		FORCEINLINE int32 CalculateSlackGrow(int32 NumElements, int32 NumAllocatedElements, int32 NumBytesPerElement) const
+		FORCEINLINE int32 CalculateSlackGrow(int32 NumElements, int32 NumAllocatedElements, SIZE_T NumBytesPerElement) const
 		{
 			// Ensure the requested allocation will fit in the inline data area.
 			checkSlow(NumElements <= NumInlineElements);
@@ -795,7 +803,7 @@ public:
 			return 0;
 		}
 
-		bool HasAllocation()
+		bool HasAllocation() const
 		{
 			return false;
 		}
diff --git a/Engine/Source/Runtime/Core/Public/Misc/MemStack.h b/Engine/Source/Runtime/Core/Public/Misc/MemStack.h
--- a/Engine/Source/Runtime/Core/Public/Misc/MemStack.h
+++ b/Engine/Source/Runtime/Core/Public/Misc/MemStack.h
@@ -303,7 +303,7 @@ public:
 		{
 			return Data;
 		}
-		void ResizeAllocation(int32 PreviousNumElements,int32 NumElements,int32 NumBytesPerElement)
+		void ResizeAllocation(int32 PreviousNumElements,int32 NumElements, SIZE_T NumBytesPerElement)
 		{
 			void* OldData = Data;
 			if( NumElements )
@@ -322,25 +322,25 @@ public:
 				}
 			}
 		}
-		FORCEINLINE int32 CalculateSlackReserve(int32 NumElements, int32 NumBytesPerElement) const
+		FORCEINLINE int32 CalculateSlackReserve(int32 NumElements, SIZE_T NumBytesPerElement) const
 		{
 			return DefaultCalculateSlackReserve(NumElements, NumBytesPerElement, false, Alignment);
 		}
-		FORCEINLINE int32 CalculateSlackShrink(int32 NumElements, int32 NumAllocatedElements, int32 NumBytesPerElement) const
+		FORCEINLINE int32 CalculateSlackShrink(int32 NumElements, int32 NumAllocatedElements, SIZE_T NumBytesPerElement) const
 		{
 			return DefaultCalculateSlackShrink(NumElements, NumAllocatedElements, NumBytesPerElement, false, Alignment);
 		}
-		FORCEINLINE int32 CalculateSlackGrow(int32 NumElements, int32 NumAllocatedElements, int32 NumBytesPerElement) const
+		FORCEINLINE int32 CalculateSlackGrow(int32 NumElements, int32 NumAllocatedElements, SIZE_T NumBytesPerElement) const
 		{
 			return DefaultCalculateSlackGrow(NumElements, NumAllocatedElements, NumBytesPerElement, false, Alignment);
 		}
 
-		FORCEINLINE int32 GetAllocatedSize(int32 NumAllocatedElements, int32 NumBytesPerElement) const
+		FORCEINLINE SIZE_T GetAllocatedSize(int32 NumAllocatedElements, SIZE_T NumBytesPerElement) const
 		{
 			return NumAllocatedElements * NumBytesPerElement;
 		}
 
-		bool HasAllocation()
+		bool HasAllocation() const
 		{
 			return !!Data;
 		}
diff --git a/Engine/Source/Runtime/Engine/Classes/Animation/AnimSequence.h b/Engine/Source/Runtime/Engine/Classes/Animation/AnimSequence.h
--- a/Engine/Source/Runtime/Engine/Classes/Animation/AnimSequence.h
+++ b/Engine/Source/Runtime/Engine/Classes/Animation/AnimSequence.h
@@ -486,17 +486,17 @@ public:
 				Data = (FScriptContainerElement*)FMemory::Realloc(Data, NumElements*NumBytesPerElement, Alignment);
 			}
 		}
-		int32 CalculateSlackReserve(int32 NumElements, int32 NumBytesPerElement) const
+		int32 CalculateSlackReserve(int32 NumElements, SIZE_T NumBytesPerElement) const
 		{
 			check(!MappedHandle && !MappedRegion); // this could be supported, but it probably is never what you want, so we will just assert.
 			return DefaultCalculateSlackReserve(NumElements, NumBytesPerElement, true, Alignment);
 		}
-		int32 CalculateSlackShrink(int32 NumElements, int32 NumAllocatedElements, int32 NumBytesPerElement) const
+		int32 CalculateSlackShrink(int32 NumElements, int32 NumAllocatedElements, SIZE_T NumBytesPerElement) const
 		{
 			check(!MappedHandle && !MappedRegion); // this could be supported, but it probably is never what you want, so we will just assert.
 			return DefaultCalculateSlackShrink(NumElements, NumAllocatedElements, NumBytesPerElement, true, Alignment);
 		}
-		int32 CalculateSlackGrow(int32 NumElements, int32 NumAllocatedElements, int32 NumBytesPerElement) const
+		int32 CalculateSlackGrow(int32 NumElements, int32 NumAllocatedElements, SIZE_T NumBytesPerElement) const
 		{
 			check(!MappedHandle && !MappedRegion); // this could be supported, but it probably is never what you want, so we will just assert.
 			return DefaultCalculateSlackGrow(NumElements, NumAllocatedElements, NumBytesPerElement, true, Alignment);
@@ -507,7 +507,7 @@ public:
 			return NumAllocatedElements * NumBytesPerElement;
 		}
 
-		bool HasAllocation()
+		bool HasAllocation() const
 		{
 			return !!Data;
 		}
