Removed seemingly pointless sleep before setting Windows thread names since this showed up in profiles

Code has been around since 2009 (CL309293 in UE3) .

diff --git a/Engine/Source/Runtime/Core/Private/Windows/WindowsRunnableThread.h b/Engine/Source/Runtime/Core/Private/Windows/WindowsRunnableThread.h
--- a/Engine/Source/Runtime/Core/Private/Windows/WindowsRunnableThread.h
+++ b/Engine/Source/Runtime/Core/Private/Windows/WindowsRunnableThread.h
@@ -48,7 +48,6 @@ class FRunnableThreadWin
 			uint32 dwFlags;		// Reserved for future use, must be zero.
 		};
 
-		Sleep(10);
 		THREADNAME_INFO ThreadNameInfo;
 		ThreadNameInfo.dwType		= 0x1000;
 		ThreadNameInfo.szName		= ThreadName;
