Allocators must now define a SizeType which is the signed type used for representing indices and counts.
FDefaultAllocator64 and TArray64 aliases added to enable a 64-bit-indexed version of TArray.
Some forward declarations removed which no longer compile because they are now aliases instead of classes.

diff --git a/Engine/Source/Developer/ShaderCompilerCommon/Private/HlslUtils.h b/Engine/Source/Developer/ShaderCompilerCommon/Private/HlslUtils.h
--- a/Engine/Source/Developer/ShaderCompilerCommon/Private/HlslUtils.h
+++ b/Engine/Source/Developer/ShaderCompilerCommon/Private/HlslUtils.h
@@ -109,6 +109,8 @@ namespace CrossCompiler
 	class FLinearAllocatorPolicy
 	{
 	public:
+		using SizeType = int32;
+
 		// Unreal allocator magic
 		enum { NeedsElementType = false };
 		enum { RequireRangeCheck = true };
@@ -129,7 +131,7 @@ namespace CrossCompiler
 			{
 				return Data;
 			}
-			void ResizeAllocation(int32 PreviousNumElements, int32 NumElements, SIZE_T NumBytesPerElement)
+			void ResizeAllocation(SizeType PreviousNumElements, SizeType NumElements, SIZE_T NumBytesPerElement)
 			{
 				void* OldData = Data;
 				if (NumElements)
@@ -142,25 +144,25 @@ namespace CrossCompiler
 					// If the container previously held elements, copy them into the new allocation.
 					if (OldData && PreviousNumElements)
 					{
-						const int32 NumCopiedElements = FMath::Min(NumElements, PreviousNumElements);
+						const SizeType NumCopiedElements = FMath::Min(NumElements, PreviousNumElements);
 						FMemory::Memcpy(Data, OldData, NumCopiedElements * NumBytesPerElement);
 					}
 				}
 			}
-			int32 CalculateSlackReserve(int32 NumElements, SIZE_T NumBytesPerElement) const
+			SizeType CalculateSlackReserve(SizeType NumElements, SIZE_T NumBytesPerElement) const
 			{
 				return DefaultCalculateSlackReserve(NumElements, NumBytesPerElement, false);
 			}
-			int32 CalculateSlackShrink(int32 NumElements, int32 NumAllocatedElements, SIZE_T NumBytesPerElement) const
+			SizeType CalculateSlackShrink(SizeType NumElements, SizeType NumAllocatedElements, SIZE_T NumBytesPerElement) const
 			{
 				return DefaultCalculateSlackShrink(NumElements, NumAllocatedElements, NumBytesPerElement, false);
 			}
-			int32 CalculateSlackGrow(int32 NumElements, int32 NumAllocatedElements, SIZE_T NumBytesPerElement) const
+			SizeType CalculateSlackGrow(SizeType NumElements, SizeType NumAllocatedElements, SIZE_T NumBytesPerElement) const
 			{
 				return DefaultCalculateSlackGrow(NumElements, NumAllocatedElements, NumBytesPerElement, false);
 			}
 
-			SIZE_T GetAllocatedSize(int32 NumAllocatedElements, SIZE_T NumBytesPerElement) const
+			SIZE_T GetAllocatedSize(SizeType NumAllocatedElements, SIZE_T NumBytesPerElement) const
 			{
 				return NumAllocatedElements * NumBytesPerElement;
 			}
diff --git a/Engine/Source/Runtime/Core/Private/HAL/Allocators/AnsiAllocator.h b/Engine/Source/Runtime/Core/Private/HAL/Allocators/AnsiAllocator.h
--- a/Engine/Source/Runtime/Core/Private/HAL/Allocators/AnsiAllocator.h
+++ b/Engine/Source/Runtime/Core/Private/HAL/Allocators/AnsiAllocator.h
@@ -14,6 +14,7 @@
 class CORE_API FAnsiAllocator
 {
 public:
+	using SizeType = int32;
 
 	enum { NeedsElementType = false };
 	enum { RequireRangeCheck = true };
@@ -61,7 +62,7 @@ public:
 		{
 			return Data;
 		}
-		void ResizeAllocation(int32 PreviousNumElements, int32 NumElements, SIZE_T NumBytesPerElement)
+		void ResizeAllocation(SizeType PreviousNumElements, SizeType NumElements, SIZE_T NumBytesPerElement)
 		{
 			// Avoid calling FMemory::Realloc( nullptr, 0 ) as ANSI C mandates returning a valid pointer which is not what we want.
 			if (NumElements)
@@ -76,20 +77,20 @@ public:
 				Data = nullptr;
 			}
 		}
-		int32 CalculateSlackReserve(int32 NumElements, SIZE_T NumBytesPerElement) const
+		SizeType CalculateSlackReserve(SizeType NumElements, SIZE_T NumBytesPerElement) const
 		{
 			return DefaultCalculateSlackReserve(NumElements, NumBytesPerElement, false);
 		}
-		int32 CalculateSlackShrink(int32 NumElements, int32 NumAllocatedElements, SIZE_T NumBytesPerElement) const
+		SizeType CalculateSlackShrink(SizeType NumElements, SizeType NumAllocatedElements, SIZE_T NumBytesPerElement) const
 		{
 			return DefaultCalculateSlackShrink(NumElements, NumAllocatedElements, NumBytesPerElement, false);
 		}
-		int32 CalculateSlackGrow(int32 NumElements, int32 NumAllocatedElements, SIZE_T NumBytesPerElement) const
+		SizeType CalculateSlackGrow(SizeType NumElements, SizeType NumAllocatedElements, SIZE_T NumBytesPerElement) const
 		{
 			return DefaultCalculateSlackGrow(NumElements, NumAllocatedElements, NumBytesPerElement, false);
 		}
 
-		SIZE_T GetAllocatedSize(int32 NumAllocatedElements, SIZE_T NumBytesPerElement) const
+		SIZE_T GetAllocatedSize(SizeType NumAllocatedElements, SIZE_T NumBytesPerElement) const
 		{
 			return NumAllocatedElements * NumBytesPerElement;
 		}
diff --git a/Engine/Source/Runtime/Core/Public/Containers/Array.h b/Engine/Source/Runtime/Core/Public/Containers/Array.h
--- a/Engine/Source/Runtime/Core/Public/Containers/Array.h
+++ b/Engine/Source/Runtime/Core/Public/Containers/Array.h
@@ -6,6 +6,7 @@
 #include "Misc/AssertionMacros.h"
 #include "HAL/UnrealMemory.h"
 #include "Templates/AreTypesEqual.h"
+#include "Templates/IsSigned.h"
 #include "Templates/UnrealTypeTraits.h"
 #include "Templates/UnrealTemplate.h"
 #include "Containers/ContainerAllocationPolicies.h"
@@ -30,16 +31,16 @@
 /**
  * Generic iterator which can operate on types that expose the following:
  * - A type called ElementType representing the contained type.
- * - A method IndexType Num() const that returns the number of items in the container.
- * - A method bool IsValidIndex(IndexType index) which returns whether a given index is valid in the container.
- * - A method T& operator\[\](IndexType index) which returns a reference to a contained object by index.
- * - A method void RemoveAt(IndexType index) which removes the element at index
+ * - A method SizeType Num() const that returns the number of items in the container.
+ * - A method bool IsValidIndex(SizeType index) which returns whether a given index is valid in the container.
+ * - A method T& operator\[\](SizeType index) which returns a reference to a contained object by index.
+ * - A method void RemoveAt(SizeType index) which removes the element at index
  */
-template< typename ContainerType, typename ElementType, typename IndexType>
+template< typename ContainerType, typename ElementType, typename SizeType>
 class TIndexedContainerIterator
 {
 public:
-	TIndexedContainerIterator(ContainerType& InContainer, IndexType StartIndex = 0)
+	TIndexedContainerIterator(ContainerType& InContainer, SizeType StartIndex = 0)
 		: Container(InContainer)
 		, Index    (StartIndex)
 	{
@@ -72,24 +73,24 @@ public:
 	}
 
 	/** iterator arithmetic support */
-	TIndexedContainerIterator& operator+=(int32 Offset)
+	TIndexedContainerIterator& operator+=(SizeType Offset)
 	{
 		Index += Offset;
 		return *this;
 	}
 
-	TIndexedContainerIterator operator+(int32 Offset) const
+	TIndexedContainerIterator operator+(SizeType Offset) const
 	{
 		TIndexedContainerIterator Tmp(*this);
 		return Tmp += Offset;
 	}
 
-	TIndexedContainerIterator& operator-=(int32 Offset)
+	TIndexedContainerIterator& operator-=(SizeType Offset)
 	{
 		return *this += -Offset;
 	}
 
-	TIndexedContainerIterator operator-(int32 Offset) const
+	TIndexedContainerIterator operator-(SizeType Offset) const
 	{
 		TIndexedContainerIterator Tmp(*this);
 		return Tmp -= Offset;
@@ -112,7 +113,7 @@ public:
 	}
 
 	/** Returns an index to the current element. */
-	IndexType GetIndex() const
+	SizeType GetIndex() const
 	{
 		return Index;
 	}
@@ -142,13 +143,13 @@ public:
 private:
 
 	ContainerType& Container;
-	IndexType      Index;
+	SizeType      Index;
 };
 
 
 /** operator + */
-template <typename ContainerType, typename ElementType, typename IndexType>
-FORCEINLINE TIndexedContainerIterator<ContainerType, ElementType, IndexType> operator+(int32 Offset, TIndexedContainerIterator<ContainerType, ElementType, IndexType> RHS)
+template <typename ContainerType, typename ElementType, typename SizeType>
+FORCEINLINE TIndexedContainerIterator<ContainerType, ElementType, SizeType> operator+(SizeType Offset, TIndexedContainerIterator<ContainerType, ElementType, SizeType> RHS)
 {
 	return RHS + Offset;
 }
@@ -159,7 +160,7 @@ FORCEINLINE TIndexedContainerIterator<ContainerType, ElementType, IndexType> ope
 	 * Pointer-like iterator type for ranged-for loops which checks that the
 	 * container hasn't been resized during iteration.
 	 */
-	template <typename ElementType>
+	template <typename ElementType, typename SizeType>
 	struct TCheckedPointerIterator
 	{
 		// This iterator type only supports the minimal functionality needed to support
@@ -167,7 +168,7 @@ FORCEINLINE TIndexedContainerIterator<ContainerType, ElementType, IndexType> ope
 		//
 		// We do add an operator-- to help FString implementation
 
-		explicit TCheckedPointerIterator(const int32& InNum, ElementType* InPtr)
+		explicit TCheckedPointerIterator(const SizeType& InNum, ElementType* InPtr)
 			: Ptr       (InPtr)
 			, CurrentNum(InNum)
 			, InitialNum(InNum)
@@ -192,9 +193,9 @@ FORCEINLINE TIndexedContainerIterator<ContainerType, ElementType, IndexType> ope
 		}
 
 	private:
-		ElementType* Ptr;
-		const int32& CurrentNum;
-		int32        InitialNum;
+		ElementType*    Ptr;
+		const SizeType& CurrentNum;
+		SizeType        InitialNum;
 
 		FORCEINLINE friend bool operator!=(const TCheckedPointerIterator& Lhs, const TCheckedPointerIterator& Rhs)
 		{
@@ -259,6 +260,18 @@ namespace UE4Array_Private
 				)
 		};
 	};
+
+	template <typename AllocatorType>
+	struct TAllocatorSizeType
+	{
+		template <typename T>
+		static typename T::SizeType ResolveIndexType(typename T::SizeType*);
+
+		template <typename T>
+		static int32 ResolveIndexType(...);
+
+		using Type = decltype(ResolveIndexType<AllocatorType>(nullptr));
+	};
 }
 
 
@@ -280,10 +293,12 @@ class TArray
 	friend class TArray;
 
 public:
-
+	typedef typename InAllocator::SizeType SizeType;
 	typedef InElementType ElementType;
 	typedef InAllocator   Allocator;
 
+	static_assert(TIsSigned<SizeType>::Value, "TArray only supports signed index types");
+
 	/**
 	 * Constructor, initializes element number counters.
 	 */
@@ -299,7 +314,7 @@ public:
 	 * @param Count The number of elements to copy from Ptr.
 	 * @see Append
 	 */
-	FORCEINLINE TArray(const ElementType* Ptr, int32 Count)
+	FORCEINLINE TArray(const ElementType* Ptr, SizeType Count)
 	{
 		check(Ptr != nullptr || Count == 0);
 
@@ -314,7 +329,7 @@ public:
 		// This is not strictly legal, as std::initializer_list's iterators are not guaranteed to be pointers, but
 		// this appears to be the case on all of our implementations.  Also, if it's not true on a new implementation,
 		// it will fail to compile rather than behave badly.
-		CopyToEmpty(InitList.begin(), (int32)InitList.size(), 0, 0);
+		CopyToEmpty(InitList.begin(), (SizeType)InitList.size(), 0, 0);
 	}
 
 	/**
@@ -345,7 +360,7 @@ public:
 	 * @param ExtraSlack Tells how much extra memory should be preallocated
 	 *                   at the end of the array in the number of elements.
 	 */
-	FORCEINLINE TArray(const TArray& Other, int32 ExtraSlack)
+	FORCEINLINE TArray(const TArray& Other, SizeType ExtraSlack)
 	{
 		CopyToEmpty(Other.GetData(), Other.Num(), 0, ExtraSlack);
 	}
@@ -362,7 +377,7 @@ public:
 		// This is not strictly legal, as std::initializer_list's iterators are not guaranteed to be pointers, but
 		// this appears to be the case on all of our implementations.  Also, if it's not true on a new implementation,
 		// it will fail to compile rather than behave badly.
-		CopyToEmpty(InitList.begin(), (int32)InitList.size(), ArrayMax, 0);
+		CopyToEmpty(InitList.begin(), (SizeType)InitList.size(), ArrayMax, 0);
 		return *this;
 	}
 
@@ -409,7 +424,7 @@ private:
 	 * @param FromArray Array to move from.
 	 */
 	template <typename FromArrayType, typename ToArrayType>
-	static FORCEINLINE typename TEnableIf<UE4Array_Private::TCanMoveTArrayPointersBetweenArrayTypes<FromArrayType, ToArrayType>::Value>::Type MoveOrCopy(ToArrayType& ToArray, FromArrayType& FromArray, int32 PrevMax)
+	static FORCEINLINE typename TEnableIf<UE4Array_Private::TCanMoveTArrayPointersBetweenArrayTypes<FromArrayType, ToArrayType>::Value>::Type MoveOrCopy(ToArrayType& ToArray, FromArrayType& FromArray, SizeType PrevMax)
 	{
 		ToArray.AllocatorInstance.MoveToEmpty(FromArray.AllocatorInstance);
 
@@ -430,7 +445,7 @@ private:
 	 *                   at the end of the array in the number of elements.
 	 */
 	template <typename FromArrayType, typename ToArrayType>
-	static FORCEINLINE typename TEnableIf<!UE4Array_Private::TCanMoveTArrayPointersBetweenArrayTypes<FromArrayType, ToArrayType>::Value>::Type MoveOrCopy(ToArrayType& ToArray, FromArrayType& FromArray, int32 PrevMax)
+	static FORCEINLINE typename TEnableIf<!UE4Array_Private::TCanMoveTArrayPointersBetweenArrayTypes<FromArrayType, ToArrayType>::Value>::Type MoveOrCopy(ToArrayType& ToArray, FromArrayType& FromArray, SizeType PrevMax)
 	{
 		ToArray.CopyToEmpty(FromArray.GetData(), FromArray.Num(), PrevMax, 0);
 	}
@@ -446,7 +461,7 @@ private:
 	 *                   at the end of the array in the number of elements.
 	 */
 	template <typename FromArrayType, typename ToArrayType>
-	static FORCEINLINE typename TEnableIf<UE4Array_Private::TCanMoveTArrayPointersBetweenArrayTypes<FromArrayType, ToArrayType>::Value>::Type MoveOrCopyWithSlack(ToArrayType& ToArray, FromArrayType& FromArray, int32 PrevMax, int32 ExtraSlack)
+	static FORCEINLINE typename TEnableIf<UE4Array_Private::TCanMoveTArrayPointersBetweenArrayTypes<FromArrayType, ToArrayType>::Value>::Type MoveOrCopyWithSlack(ToArrayType& ToArray, FromArrayType& FromArray, SizeType PrevMax, SizeType ExtraSlack)
 	{
 		MoveOrCopy(ToArray, FromArray, PrevMax);
 
@@ -464,7 +479,7 @@ private:
 	 *                   at the end of the array in the number of elements.
 	 */
 	template <typename FromArrayType, typename ToArrayType>
-	static FORCEINLINE typename TEnableIf<!UE4Array_Private::TCanMoveTArrayPointersBetweenArrayTypes<FromArrayType, ToArrayType>::Value>::Type MoveOrCopyWithSlack(ToArrayType& ToArray, FromArrayType& FromArray, int32 PrevMax, int32 ExtraSlack)
+	static FORCEINLINE typename TEnableIf<!UE4Array_Private::TCanMoveTArrayPointersBetweenArrayTypes<FromArrayType, ToArrayType>::Value>::Type MoveOrCopyWithSlack(ToArrayType& ToArray, FromArrayType& FromArray, SizeType PrevMax, SizeType ExtraSlack)
 	{
 		ToArray.CopyToEmpty(FromArray.GetData(), FromArray.Num(), PrevMax, ExtraSlack);
 	}
@@ -499,7 +514,7 @@ public:
 	 *                   at the end of the array in the number of elements.
 	 */
 	template <typename OtherElementType>
-	TArray(TArray<OtherElementType, Allocator>&& Other, int32 ExtraSlack)
+	TArray(TArray<OtherElementType, Allocator>&& Other, SizeType ExtraSlack)
 	{
 		// We don't implement move semantics for general OtherAllocators, as there's no way
 		// to tell if they're compatible with the current one.  Probably going to be a pretty
@@ -572,7 +587,7 @@ public:
 	 *
 	 * @returns Number of bytes allocated by this container.
 	 */
-	FORCEINLINE uint32 GetAllocatedSize(void) const
+	FORCEINLINE SIZE_T GetAllocatedSize(void) const
 	{
 		return AllocatorInstance.GetAllocatedSize(ArrayMax, sizeof(ElementType));
 	}
@@ -582,7 +597,7 @@ public:
 	 *
 	 * @see Num, Shrink
 	 */
-	FORCEINLINE int32 GetSlack() const
+	FORCEINLINE SizeType GetSlack() const
 	{
 		return ArrayMax - ArrayNum;
 	}
@@ -601,7 +616,7 @@ public:
 	 *
 	 * @param Index Index to check.
 	 */
-	FORCEINLINE void RangeCheck(int32 Index) const
+	FORCEINLINE void RangeCheck(SizeType Index) const
 	{
 		CheckInvariants();
 
@@ -618,7 +633,7 @@ public:
 	 * @param Index Index to test.
 	 * @returns True if index is valid. False otherwise.
 	 */
-	FORCEINLINE bool IsValidIndex(int32 Index) const
+	FORCEINLINE bool IsValidIndex(SizeType Index) const
 	{
 		return Index >= 0 && Index < ArrayNum;
 	}
@@ -629,7 +644,7 @@ public:
 	 * @returns Number of elements in array.
 	 * @see GetSlack
 	 */
-	FORCEINLINE int32 Num() const
+	FORCEINLINE SizeType Num() const
 	{
 		return ArrayNum;
 	}
@@ -640,7 +655,7 @@ public:
 	 * @returns Maximum number of elements in array.
 	 * @see GetSlack
 	 */
-	FORCEINLINE int32 Max() const
+	FORCEINLINE SizeType Max() const
 	{
 		return ArrayMax;
 	}
@@ -650,7 +665,7 @@ public:
 	 *
 	 * @returns Reference to indexed element.
 	 */
-	FORCEINLINE ElementType& operator[](int32 Index)
+	FORCEINLINE ElementType& operator[](SizeType Index)
 	{
 		RangeCheck(Index);
 		return GetData()[Index];
@@ -663,7 +678,7 @@ public:
 	 *
 	 * @returns Reference to indexed element.
 	 */
-	FORCEINLINE const ElementType& operator[](int32 Index) const
+	FORCEINLINE const ElementType& operator[](SizeType Index) const
 	{
 		RangeCheck(Index);
 		return GetData()[Index];
@@ -736,7 +751,7 @@ public:
 	 * @param IndexFromTheEnd (Optional) Index from the end of array (default = 0).
 	 * @returns Reference to n-th last element from the array.
 	 */
-	FORCEINLINE ElementType& Last(int32 IndexFromTheEnd = 0)
+	FORCEINLINE ElementType& Last(SizeType IndexFromTheEnd = 0)
 	{
 		RangeCheck(ArrayNum - IndexFromTheEnd - 1);
 		return GetData()[ArrayNum - IndexFromTheEnd - 1];
@@ -750,7 +765,7 @@ public:
 	 * @param IndexFromTheEnd (Optional) Index from the end of array (default = 0).
 	 * @returns Reference to n-th last element from the array.
 	 */
-	FORCEINLINE const ElementType& Last(int32 IndexFromTheEnd = 0) const
+	FORCEINLINE const ElementType& Last(SizeType IndexFromTheEnd = 0) const
 	{
 		RangeCheck(ArrayNum - IndexFromTheEnd - 1);
 		return GetData()[ArrayNum - IndexFromTheEnd - 1];
@@ -778,7 +793,7 @@ public:
 	 * @returns True if found. False otherwise.
 	 * @see FindLast, FindLastByPredicate
 	 */
-	FORCEINLINE bool Find(const ElementType& Item, int32& Index) const
+	FORCEINLINE bool Find(const ElementType& Item, SizeType& Index) const
 	{
 		Index = this->Find(Item);
 		return Index != INDEX_NONE;
@@ -791,14 +806,14 @@ public:
 	 * @returns Index of the found element. INDEX_NONE otherwise.
 	 * @see FindLast, FindLastByPredicate
 	 */
-	int32 Find(const ElementType& Item) const
+	SizeType Find(const ElementType& Item) const
 	{
 		const ElementType* RESTRICT Start = GetData();
 		for (const ElementType* RESTRICT Data = Start, *RESTRICT DataEnd = Data + ArrayNum; Data != DataEnd; ++Data)
 		{
 			if (*Data == Item)
 			{
-				return static_cast<int32>(Data - Start);
+				return static_cast<SizeType>(Data - Start);
 			}
 		}
 		return INDEX_NONE;
@@ -812,7 +827,7 @@ public:
 	 * @returns True if found. False otherwise.
 	 * @see Find, FindLastByPredicate
 	 */
-	FORCEINLINE bool FindLast(const ElementType& Item, int32& Index) const
+	FORCEINLINE bool FindLast(const ElementType& Item, SizeType& Index) const
 	{
 		Index = this->FindLast(Item);
 		return Index != INDEX_NONE;
@@ -824,14 +839,14 @@ public:
 	 * @param Item Item to look for.
 	 * @returns Index of the found element. INDEX_NONE otherwise.
 	 */
-	int32 FindLast(const ElementType& Item) const
+	SizeType FindLast(const ElementType& Item) const
 	{
 		for (const ElementType* RESTRICT Start = GetData(), *RESTRICT Data = Start + ArrayNum; Data != Start; )
 		{
 			--Data;
 			if (*Data == Item)
 			{
-				return static_cast<int32>(Data - Start);
+				return static_cast<SizeType>(Data - Start);
 			}
 		}
 		return INDEX_NONE;
@@ -845,7 +860,7 @@ public:
 	 * @returns Index of the found element. INDEX_NONE otherwise.
 	 */
 	template <typename Predicate>
-	int32 FindLastByPredicate(Predicate Pred, int32 Count) const
+	SizeType FindLastByPredicate(Predicate Pred, SizeType Count) const
 	{
 		check(Count >= 0 && Count <= this->Num());
 		for (const ElementType* RESTRICT Start = GetData(), *RESTRICT Data = Start + Count; Data != Start; )
@@ -853,7 +868,7 @@ public:
 			--Data;
 			if (Pred(*Data))
 			{
-				return static_cast<int32>(Data - Start);
+				return static_cast<SizeType>(Data - Start);
 			}
 		}
 		return INDEX_NONE;
@@ -866,7 +881,7 @@ public:
 	 * @returns Index of the found element. INDEX_NONE otherwise.
 	 */
 	template <typename Predicate>
-	FORCEINLINE int32 FindLastByPredicate(Predicate Pred) const
+	FORCEINLINE SizeType FindLastByPredicate(Predicate Pred) const
 	{
 		return FindLastByPredicate(Pred, ArrayNum);
 	}
@@ -879,14 +894,14 @@ public:
 	 * @returns Index to the first matching element, or INDEX_NONE if none is found.
 	 */
 	template <typename KeyType>
-	int32 IndexOfByKey(const KeyType& Key) const
+	SizeType IndexOfByKey(const KeyType& Key) const
 	{
 		const ElementType* RESTRICT Start = GetData();
 		for (const ElementType* RESTRICT Data = Start, *RESTRICT DataEnd = Start + ArrayNum; Data != DataEnd; ++Data)
 		{
 			if (*Data == Key)
 			{
-				return static_cast<int32>(Data - Start);
+				return static_cast<SizeType>(Data - Start);
 			}
 		}
 		return INDEX_NONE;
@@ -899,14 +914,14 @@ public:
 	 * @returns Index to the first matching element, or INDEX_NONE if none is found.
 	 */
 	template <typename Predicate>
-	int32 IndexOfByPredicate(Predicate Pred) const
+	SizeType IndexOfByPredicate(Predicate Pred) const
 	{
 		const ElementType* RESTRICT Start = GetData();
 		for (const ElementType* RESTRICT Data = Start, *RESTRICT DataEnd = Start + ArrayNum; Data != DataEnd; ++Data)
 		{
 			if (Pred(*Data))
 			{
-				return static_cast<int32>(Data - Start);
+				return static_cast<SizeType>(Data - Start);
 			}
 		}
 		return INDEX_NONE;
@@ -1043,7 +1058,7 @@ public:
 	 */
 	bool operator==(const TArray& OtherArray) const
 	{
-		int32 Count = Num();
+		SizeType Count = Num();
 
 		return Count == OtherArray.Num() && CompareItems(GetData(), OtherArray.GetData(), Count);
 	}
@@ -1071,8 +1086,8 @@ public:
 		A.CountBytes(Ar);
 
 		// For net archives, limit serialization to 16MB, to protect against excessive allocation
-		const int32 MaxNetArraySerialize = (16 * 1024 * 1024) / sizeof(ElementType);
-		int32 SerializeNum = (Ar.IsLoading() ? 0 : A.ArrayNum);
+		constexpr SizeType MaxNetArraySerialize = (16 * 1024 * 1024) / sizeof(ElementType);
+		SizeType SerializeNum = Ar.IsLoading() ? 0 : A.ArrayNum;
 
 		Ar << SerializeNum;
 
@@ -1097,7 +1112,7 @@ public:
 				// Required for resetting ArrayNum
 				A.Empty(SerializeNum);
 
-				for (int32 i=0; i<SerializeNum; i++)
+				for (SizeType i=0; i<SerializeNum; i++)
 				{
 					Ar << *::new(A) ElementType;
 				}
@@ -1106,7 +1121,7 @@ public:
 			{
 				A.ArrayNum = SerializeNum;
 
-				for (int32 i=0; i<A.ArrayNum; i++)
+				for (SizeType i=0; i<A.ArrayNum; i++)
 				{
 					Ar << A[i];
 				}
@@ -1170,7 +1185,7 @@ public:
 				// Serialize the number of elements, block allocate the right amount of memory and deserialize
 				// the data as a giant memory blob in a single call to Serialize. Please see the function header
 				// for detailed documentation on limitations and implications.
-				int32 NewArrayNum = 0;
+				SizeType NewArrayNum = 0;
 				Ar << NewArrayNum;
 				Empty(NewArrayNum);
 				AddUninitialized(NewArrayNum);
@@ -1178,7 +1193,7 @@ public:
 			}
 			else if (Ar.IsSaving())
 			{
-				int32 ArrayCount = Num();
+				SizeType ArrayCount = Num();
 				Ar << ArrayCount;
 				Ar.Serialize(GetData(), ArrayCount * SerializedElementSize);
 			}
@@ -1205,12 +1220,12 @@ public:
 	 * @param Count Number of elements to add.
 	 * @returns Number of elements in array before addition.
 	 */
-	FORCEINLINE int32 AddUninitialized(int32 Count = 1)
+	FORCEINLINE SizeType AddUninitialized(SizeType Count = 1)
 	{
 		CheckInvariants();
 		checkSlow(Count >= 0);
 
-		const int32 OldNum = ArrayNum;
+		const SizeType OldNum = ArrayNum;
 		if ((ArrayNum += Count) > ArrayMax)
 		{
 			ResizeGrow(OldNum);
@@ -1218,6 +1233,26 @@ public:
 		return OldNum;
 	}
 
+private:
+	template <typename OtherSizeType>
+	void InsertUninitializedImpl(SizeType Index, OtherSizeType Count)
+	{
+		CheckInvariants();
+		checkSlow((Count >= 0) & (Index >= 0) & (Index <= ArrayNum));
+
+		SizeType NewNum = Count;
+		checkf((OtherSizeType)NewNum == Count, TEXT("Invalid number of elements to add to this array type: %llu"), (unsigned long long)NewNum);
+
+		const SizeType OldNum = ArrayNum;
+		if ((ArrayNum += Count) > ArrayMax)
+		{
+			ResizeGrow(OldNum);
+		}
+		ElementType* Data = GetData() + Index;
+		RelocateConstructItems<ElementType>(Data + Count, Data, OldNum - Index);
+	}
+
+public:
 	/**
 	 * Inserts a given number of uninitialized elements into the array at given
 	 * location.
@@ -1230,18 +1265,9 @@ public:
 	 * @param Count Number of elements to add.
 	 * @see Insert, InsertZeroed, InsertDefaulted
 	 */
-	void InsertUninitialized(int32 Index, int32 Count = 1)
+	FORCEINLINE void InsertUninitialized(SizeType Index, SizeType Count = 1)
 	{
-		CheckInvariants();
-		checkSlow((Count >= 0) & (Index >= 0) & (Index <= ArrayNum));
-
-		const int32 OldNum = ArrayNum;
-		if ((ArrayNum += Count) > ArrayMax)
-		{
-			ResizeGrow(OldNum);
-		}
-		ElementType* Data = GetData() + Index;
-		RelocateConstructItems<ElementType>(Data + Count, Data, OldNum - Index);
+		InsertUninitializedImpl(Index, Count);
 	}
 
 	/**
@@ -1256,9 +1282,9 @@ public:
 	 * @param Count Number of elements to add.
 	 * @see Insert, InsertUninitialized, InsertDefaulted
 	 */
-	void InsertZeroed(int32 Index, int32 Count = 1)
+	void InsertZeroed(SizeType Index, SizeType Count = 1)
 	{
-		InsertUninitialized(Index, Count);
+		InsertUninitializedImpl(Index, Count);
 		FMemory::Memzero(GetData() + Index, Count * sizeof(ElementType));
 	}
 
@@ -1273,9 +1299,9 @@ public:
 	 * @return A reference to the newly-inserted element.
 	 * @see Insert_GetRef, InsertDefaulted_GetRef
 	 */
-	ElementType& InsertZeroed_GetRef(int32 Index)
+	ElementType& InsertZeroed_GetRef(SizeType Index)
 	{
-		InsertUninitialized(Index, 1);
+		InsertUninitializedImpl(Index, 1);
 		ElementType* Ptr = GetData() + Index;
 		FMemory::Memzero(Ptr, sizeof(ElementType));
 		return *Ptr;
@@ -1289,9 +1315,9 @@ public:
 	 * @param Count Number of elements to add.
 	 * @see Insert, InsertUninitialized, InsertZeroed
 	 */
-	void InsertDefaulted(int32 Index, int32 Count = 1)
+	void InsertDefaulted(SizeType Index, SizeType Count = 1)
 	{
-		InsertUninitialized(Index, Count);
+		InsertUninitializedImpl(Index, Count);
 		DefaultConstructItems<ElementType>(GetData() + Index, Count);
 	}
 
@@ -1303,9 +1329,9 @@ public:
 	 * @return A reference to the newly-inserted element.
 	 * @see Insert_GetRef, InsertZeroed_GetRef
 	 */
-	ElementType& InsertDefaulted_GetRef(int32 Index)
+	ElementType& InsertDefaulted_GetRef(SizeType Index)
 	{
-		InsertUninitialized(Index, 1);
+		InsertUninitializedImpl(Index, 1);
 		ElementType* Ptr = GetData() + Index;
 		DefaultConstructItems<ElementType>(Ptr, 1);
 		return *Ptr;
@@ -1318,11 +1344,11 @@ public:
 	 * @param InIndex Tells where to insert the new elements.
 	 * @returns Location at which the item was inserted.
 	 */
-	int32 Insert(std::initializer_list<ElementType> InitList, const int32 InIndex)
+	SizeType Insert(std::initializer_list<ElementType> InitList, const SizeType InIndex)
 	{
-		int32 NumNewElements = (int32)InitList.size();
+		SizeType NumNewElements = (SizeType)InitList.size();
 
-		InsertUninitialized(InIndex, NumNewElements);
+		InsertUninitializedImpl(InIndex, NumNewElements);
 		ConstructItems<ElementType>(GetData() + InIndex, InitList.begin(), NumNewElements);
 
 		return InIndex;
@@ -1336,13 +1362,13 @@ public:
 	 * @returns Location at which the item was inserted.
 	 */
 	template <typename OtherAllocator>
-	int32 Insert(const TArray<ElementType, OtherAllocator>& Items, const int32 InIndex)
+	SizeType Insert(const TArray<ElementType, OtherAllocator>& Items, const SizeType InIndex)
 	{
 		check((const void*)this != (const void*)&Items);
 
-		int32 NumNewElements = Items.Num();
+		auto NumNewElements = Items.Num();
 
-		InsertUninitialized(InIndex, NumNewElements);
+		InsertUninitializedImpl(InIndex, NumNewElements);
 		ConstructItems<ElementType>(GetData() + InIndex, Items.GetData(), NumNewElements);
 
 		return InIndex;
@@ -1356,13 +1382,13 @@ public:
 	 * @returns Location at which the item was inserted.
 	 */
 	template <typename OtherAllocator>
-	int32 Insert(TArray<ElementType, OtherAllocator>&& Items, const int32 InIndex)
+	SizeType Insert(TArray<ElementType, OtherAllocator>&& Items, const SizeType InIndex)
 	{
 		check((const void*)this != (const void*)&Items);
 
-		int32 NumNewElements = Items.Num();
+		auto NumNewElements = Items.Num();
 
-		InsertUninitialized(InIndex, NumNewElements);
+		InsertUninitializedImpl(InIndex, NumNewElements);
 		RelocateConstructItems<ElementType>(GetData() + InIndex, Items.GetData(), NumNewElements);
 		Items.ArrayNum = 0;
 
@@ -1378,11 +1404,11 @@ public:
 	 * @return The index of the first element inserted.
 	 * @see Add, Remove
 	 */
-	int32 Insert(const ElementType* Ptr, int32 Count, int32 Index)
+	SizeType Insert(const ElementType* Ptr, SizeType Count, SizeType Index)
 	{
 		check(Ptr != nullptr);
 
-		InsertUninitialized(Index, Count);
+		InsertUninitializedImpl(Index, Count);
 		ConstructItems<ElementType>(GetData() + Index, Ptr, Count);
 
 		return Index;
@@ -1410,13 +1436,13 @@ public:
 	 * @returns Location at which the insert was done.
 	 * @see Add, Remove
 	 */
-	int32 Insert(ElementType&& Item, int32 Index)
+	SizeType Insert(ElementType&& Item, SizeType Index)
 	{
 		CheckAddress(&Item);
 
 		// construct a copy in place at Index (this new operator will insert at 
 		// Index, then construct that memory with Item)
-		InsertUninitialized(Index, 1);
+		InsertUninitializedImpl(Index, 1);
 		new(GetData() + Index) ElementType(MoveTempIfPossible(Item));
 		return Index;
 	}
@@ -1429,13 +1455,13 @@ public:
 	 * @returns Location at which the insert was done.
 	 * @see Add, Remove
 	 */
-	int32 Insert(const ElementType& Item, int32 Index)
+	SizeType Insert(const ElementType& Item, SizeType Index)
 	{
 		CheckAddress(&Item);
 
 		// construct a copy in place at Index (this new operator will insert at 
 		// Index, then construct that memory with Item)
-		InsertUninitialized(Index, 1);
+		InsertUninitializedImpl(Index, 1);
 		new(GetData() + Index) ElementType(Item);
 		return Index;
 	}
@@ -1449,13 +1475,13 @@ public:
 	 * @return A reference to the newly-inserted element.
 	 * @see Add, Remove
 	 */
-	ElementType& Insert_GetRef(ElementType&& Item, int32 Index)
+	ElementType& Insert_GetRef(ElementType&& Item, SizeType Index)
 	{
 		CheckAddress(&Item);
 
 		// construct a copy in place at Index (this new operator will insert at 
 		// Index, then construct that memory with Item)
-		InsertUninitialized(Index, 1);
+		InsertUninitializedImpl(Index, 1);
 		ElementType* Ptr = GetData() + Index;
 		new(Ptr) ElementType(MoveTempIfPossible(Item));
 		return *Ptr;
@@ -1469,20 +1495,20 @@ public:
 	 * @return A reference to the newly-inserted element.
 	 * @see Add, Remove
 	 */
-	ElementType& Insert_GetRef(const ElementType& Item, int32 Index)
+	ElementType& Insert_GetRef(const ElementType& Item, SizeType Index)
 	{
 		CheckAddress(&Item);
 
 		// construct a copy in place at Index (this new operator will insert at 
 		// Index, then construct that memory with Item)
-		InsertUninitialized(Index, 1);
+		InsertUninitializedImpl(Index, 1);
 		ElementType* Ptr = GetData() + Index;
 		new(Ptr) ElementType(Item);
 		return *Ptr;
 	}
 
 private:
-	void RemoveAtImpl(int32 Index, int32 Count, bool bAllowShrinking)
+	void RemoveAtImpl(SizeType Index, SizeType Count, bool bAllowShrinking)
 	{
 		if (Count)
 		{
@@ -1492,7 +1518,7 @@ private:
 			DestructItems(GetData() + Index, Count);
 
 			// Skip memmove in the common case that there is nothing to move.
-			int32 NumToMove = ArrayNum - Index - Count;
+			SizeType NumToMove = ArrayNum - Index - Count;
 			if (NumToMove)
 			{
 				FMemory::Memmove
@@ -1520,7 +1546,7 @@ public:
 	 * @param Count (Optional) Number of elements to remove. Default is 1.
 	 * @param bAllowShrinking (Optional) Tells if this call can shrink array if suitable after remove. Default is true.
 	 */
-	FORCEINLINE void RemoveAt(int32 Index)
+	FORCEINLINE void RemoveAt(SizeType Index)
 	{
 		RemoveAtImpl(Index, 1, true);
 	}
@@ -1534,14 +1560,14 @@ public:
 	 * @param bAllowShrinking (Optional) Tells if this call can shrink array if suitable after remove. Default is true.
 	 */
 	template <typename CountType>
-	FORCEINLINE void RemoveAt(int32 Index, CountType Count, bool bAllowShrinking = true)
+	FORCEINLINE void RemoveAt(SizeType Index, CountType Count, bool bAllowShrinking = true)
 	{
 		static_assert(!TAreTypesEqual<CountType, bool>::Value, "TArray::RemoveAt: unexpected bool passed as the Count argument");
 		RemoveAtImpl(Index, Count, bAllowShrinking);
 	}
 
 private:
-	void RemoveAtSwapImpl(int32 Index, int32 Count = 1, bool bAllowShrinking = true)
+	void RemoveAtSwapImpl(SizeType Index, SizeType Count = 1, bool bAllowShrinking = true)
 	{
 		if (Count)
 		{
@@ -1551,9 +1577,9 @@ private:
 			DestructItems(GetData() + Index, Count);
 
 			// Replace the elements in the hole created by the removal with elements from the end of the array, so the range of indices used by the array is contiguous.
-			const int32 NumElementsInHole = Count;
-			const int32 NumElementsAfterHole = ArrayNum - (Index + Count);
-			const int32 NumElementsToMoveIntoHole = FPlatformMath::Min(NumElementsInHole, NumElementsAfterHole);
+			const SizeType NumElementsInHole = Count;
+			const SizeType NumElementsAfterHole = ArrayNum - (Index + Count);
+			const SizeType NumElementsToMoveIntoHole = FPlatformMath::Min(NumElementsInHole, NumElementsAfterHole);
 			if (NumElementsToMoveIntoHole)
 			{
 				FMemory::Memcpy(
@@ -1584,7 +1610,7 @@ public:
 	 * @param bAllowShrinking (Optional) Tells if this call can shrink array if
 	 *                        suitable after remove. Default is true.
 	 */
-	FORCEINLINE void RemoveAtSwap(int32 Index)
+	FORCEINLINE void RemoveAtSwap(SizeType Index)
 	{
 		RemoveAtSwapImpl(Index, 1, true);
 	}
@@ -1602,7 +1628,7 @@ public:
 	 *                        suitable after remove. Default is true.
 	 */
 	template <typename CountType>
-	FORCEINLINE void RemoveAtSwap(int32 Index, CountType Count, bool bAllowShrinking = true)
+	FORCEINLINE void RemoveAtSwap(SizeType Index, CountType Count, bool bAllowShrinking = true)
 	{
 		static_assert(!TAreTypesEqual<CountType, bool>::Value, "TArray::RemoveAtSwap: unexpected bool passed as the Count argument");
 		RemoveAtSwapImpl(Index, Count, bAllowShrinking);
@@ -1614,7 +1640,7 @@ public:
 	 *
 	 * @param NewSize The expected usage size after calling this function.
 	 */
-	void Reset(int32 NewSize = 0)
+	void Reset(SizeType NewSize = 0)
 	{
 		// If we have space to hold the excepted size, then don't reallocate
 		if (NewSize <= ArrayMax)
@@ -1633,7 +1659,7 @@ public:
 	 *
 	 * @param Slack (Optional) The expected usage size after empty operation. Default is 0.
 	 */
-	void Empty(int32 Slack = 0)
+	void Empty(SizeType Slack = 0)
 	{
 		DestructItems(GetData(), ArrayNum);
 
@@ -1652,12 +1678,12 @@ public:
 	 * @param NewNum New size of the array.
 	 * @param bAllowShrinking Tell if this function can shrink the memory in-use if suitable.
 	 */
-	void SetNum(int32 NewNum, bool bAllowShrinking = true)
+	void SetNum(SizeType NewNum, bool bAllowShrinking = true)
 	{
 		if (NewNum > Num())
 		{
-			const int32 Diff = NewNum - ArrayNum;
-			const int32 Index = AddUninitialized(Diff);
+			const SizeType Diff = NewNum - ArrayNum;
+			const SizeType Index = AddUninitialized(Diff);
 			DefaultConstructItems<ElementType>((uint8*)AllocatorInstance.GetAllocation() + Index * sizeof(ElementType), Diff);
 		}
 		else if (NewNum < Num())
@@ -1671,7 +1697,7 @@ public:
 	 *
 	 * @param NewNum New size of the array.
 	 */
-	void SetNumZeroed(int32 NewNum, bool bAllowShrinking = true)
+	void SetNumZeroed(SizeType NewNum, bool bAllowShrinking = true)
 	{
 		if (NewNum > Num())
 		{
@@ -1688,7 +1714,7 @@ public:
 	 *
 	 * @param NewNum New size of the array.
 	 */
-	void SetNumUninitialized(int32 NewNum, bool bAllowShrinking = true)
+	void SetNumUninitialized(SizeType NewNum, bool bAllowShrinking = true)
 	{
 		if (NewNum > Num())
 		{
@@ -1704,7 +1730,7 @@ public:
 	 * Does nothing except setting the new number of elements in the array. Does not destruct items, does not de-allocate memory.
 	 * @param NewNum New number of elements in the array, must be <= the current number of elements in the array.
 	 */
-	void SetNumUnsafeInternal(int32 NewNum)
+	void SetNumUnsafeInternal(SizeType NewNum)
 	{
 		checkSlow(NewNum <= Num() && NewNum >= 0);
 		ArrayNum = NewNum;
@@ -1723,7 +1749,7 @@ public:
 	{
 		check((void*)this != (void*)&Source);
 
-		int32 SourceCount = Source.Num();
+		SizeType SourceCount = Source.Num();
 
 		// Do nothing if the source is empty.
 		if (!SourceCount)
@@ -1749,7 +1775,7 @@ public:
 	{
 		check((void*)this != (void*)&Source);
 
-		int32 SourceCount = Source.Num();
+		SizeType SourceCount = Source.Num();
 
 		// Do nothing if the source is empty.
 		if (!SourceCount)
@@ -1772,11 +1798,11 @@ public:
 	 * @param Count The number of elements to insert from Ptr.
 	 * @see Add, Insert
 	 */
-	void Append(const ElementType* Ptr, int32 Count)
+	void Append(const ElementType* Ptr, SizeType Count)
 	{
 		check(Ptr != nullptr || Count == 0);
 
-		int32 Pos = AddUninitialized(Count);
+		SizeType Pos = AddUninitialized(Count);
 		ConstructItems<ElementType>(GetData() + Pos, Ptr, Count);
 	}
 
@@ -1788,9 +1814,9 @@ public:
 	 */
 	FORCEINLINE void Append(std::initializer_list<ElementType> InitList)
 	{
-		int32 Count = (int32)InitList.size();
+		SizeType Count = (SizeType)InitList.size();
 
-		int32 Pos = AddUninitialized(Count);
+		SizeType Pos = AddUninitialized(Count);
 		ConstructItems<ElementType>(GetData() + Pos, InitList.begin(), Count);
 	}
 
@@ -1838,9 +1864,9 @@ public:
 	 * @return		Index to the new item
 	 */
 	template <typename... ArgsType>
-	FORCEINLINE int32 Emplace(ArgsType&&... Args)
+	FORCEINLINE SizeType Emplace(ArgsType&&... Args)
 	{
-		const int32 Index = AddUninitialized(1);
+		const SizeType Index = AddUninitialized(1);
 		new(GetData() + Index) ElementType(Forward<ArgsType>(Args)...);
 		return Index;
 	}
@@ -1854,7 +1880,7 @@ public:
 	template <typename... ArgsType>
 	FORCEINLINE ElementType& Emplace_GetRef(ArgsType&&... Args)
 	{
-		const int32 Index = AddUninitialized(1);
+		const SizeType Index = AddUninitialized(1);
 		ElementType* Ptr = GetData() + Index;
 		new(Ptr) ElementType(Forward<ArgsType>(Args)...);
 		return *Ptr;
@@ -1867,9 +1893,9 @@ public:
 	 * @param Args	The arguments to forward to the constructor of the new item.
 	 */
 	template <typename... ArgsType>
-	FORCEINLINE void EmplaceAt(int32 Index, ArgsType&&... Args)
+	FORCEINLINE void EmplaceAt(SizeType Index, ArgsType&&... Args)
 	{
-		InsertUninitialized(Index, 1);
+		InsertUninitializedImpl(Index, 1);
 		new(GetData() + Index) ElementType(Forward<ArgsType>(Args)...);
 	}
 
@@ -1881,9 +1907,9 @@ public:
 	 * @return A reference to the newly-inserted element.
 	 */
 	template <typename... ArgsType>
-	FORCEINLINE ElementType& EmplaceAt_GetRef(int32 Index, ArgsType&&... Args)
+	FORCEINLINE ElementType& EmplaceAt_GetRef(SizeType Index, ArgsType&&... Args)
 	{
-		InsertUninitialized(Index, 1);
+		InsertUninitializedImpl(Index, 1);
 		ElementType* Ptr = GetData() + Index;
 		new(Ptr) ElementType(Forward<ArgsType>(Args)...);
 		return *Ptr;
@@ -1898,7 +1924,7 @@ public:
 	 * @return Index to the new item
 	 * @see AddDefaulted, AddUnique, AddZeroed, Append, Insert
 	 */
-	FORCEINLINE int32 Add(ElementType&& Item)
+	FORCEINLINE SizeType Add(ElementType&& Item)
 	{
 		CheckAddress(&Item);
 		return Emplace(MoveTempIfPossible(Item));
@@ -1911,7 +1937,7 @@ public:
 	 * @return Index to the new item
 	 * @see AddDefaulted, AddUnique, AddZeroed, Append, Insert
 	 */
-	FORCEINLINE int32 Add(const ElementType& Item)
+	FORCEINLINE SizeType Add(const ElementType& Item)
 	{
 		CheckAddress(&Item);
 		return Emplace(Item);
@@ -1957,9 +1983,9 @@ public:
 	 * @return Index to the first of the new items.
 	 * @see Add, AddDefaulted, AddUnique, Append, Insert
 	 */
-	int32 AddZeroed(int32 Count = 1)
+	SizeType AddZeroed(SizeType Count = 1)
 	{
-		const int32 Index = AddUninitialized(Count);
+		const SizeType Index = AddUninitialized(Count);
 		FMemory::Memzero((uint8*)AllocatorInstance.GetAllocation() + Index*sizeof(ElementType), Count*sizeof(ElementType));
 		return Index;
 	}
@@ -1977,7 +2003,7 @@ public:
 	 */
 	ElementType& AddZeroed_GetRef()
 	{
-		const int32 Index = AddUninitialized(1);
+		const SizeType Index = AddUninitialized(1);
 		ElementType* Ptr = GetData() + Index;
 		FMemory::Memzero(Ptr, sizeof(ElementType));
 		return *Ptr;
@@ -1991,9 +2017,9 @@ public:
 	 * @return Index to the first of the new items.
 	 * @see Add, AddZeroed, AddUnique, Append, Insert
 	 */
-	int32 AddDefaulted(int32 Count = 1)
+	SizeType AddDefaulted(SizeType Count = 1)
 	{
-		const int32 Index = AddUninitialized(Count);
+		const SizeType Index = AddUninitialized(Count);
 		DefaultConstructItems<ElementType>((uint8*)AllocatorInstance.GetAllocation() + Index * sizeof(ElementType), Count);
 		return Index;
 	}
@@ -2007,7 +2033,7 @@ public:
 	 */
 	ElementType& AddDefaulted_GetRef()
 	{
-		const int32 Index = AddUninitialized(1);
+		const SizeType Index = AddUninitialized(1);
 		ElementType* Ptr = GetData() + Index;
 		DefaultConstructItems<ElementType>(Ptr, 1);
 		return *Ptr;
@@ -2022,9 +2048,9 @@ private:
 	 * @returns Index of the element in the array.
 	 */
 	template <typename ArgsType>
-	int32 AddUniqueImpl(ArgsType&& Args)
+	SizeType AddUniqueImpl(ArgsType&& Args)
 	{
-		int32 Index;
+		SizeType Index;
 		if (Find(Args, Index))
 		{
 			return Index;
@@ -2044,7 +2070,7 @@ public:
 	 * @returns Index of the element in the array.
 	 * @see Add, AddDefaulted, AddZeroed, Append, Insert
 	 */
-	FORCEINLINE int32 AddUnique(ElementType&& Item) { return AddUniqueImpl(MoveTempIfPossible(Item)); }
+	FORCEINLINE SizeType AddUnique(ElementType&& Item) { return AddUniqueImpl(MoveTempIfPossible(Item)); }
 
 	/**
 	 * Adds unique element to array if it doesn't exist.
@@ -2053,7 +2079,7 @@ public:
 	 * @returns Index of the element in the array.
 	 * @see Add, AddDefaulted, AddZeroed, Append, Insert
 	 */
-	FORCEINLINE int32 AddUnique(const ElementType& Item) { return AddUniqueImpl(Item); }
+	FORCEINLINE SizeType AddUnique(const ElementType& Item) { return AddUniqueImpl(Item); }
 
 	/**
 	 * Reserves memory such that the array can contain at least Number elements.
@@ -2061,7 +2087,7 @@ public:
 	 * @param Number The number of elements that the array should be able to contain after allocation.
 	 * @see Shrink
 	 */
-	FORCEINLINE void Reserve(int32 Number)
+	FORCEINLINE void Reserve(SizeType Number)
 	{
 		checkSlow(Number >= 0);
 		if (Number > ArrayMax)
@@ -2076,10 +2102,10 @@ public:
 	 * @param Element The element to fill array with.
 	 * @param Number The number of elements that the array should be able to contain after allocation.
 	 */
-	void Init(const ElementType& Element, int32 Number)
+	void Init(const ElementType& Element, SizeType Number)
 	{
 		Empty(Number);
-		for (int32 Index = 0; Index < Number; ++Index)
+		for (SizeType Index = 0; Index < Number; ++Index)
 		{
 			new(*this) ElementType(Element);
 		}
@@ -2093,9 +2119,9 @@ public:
 	 * @returns The number of items removed. For RemoveSingleItem, this is always either 0 or 1.
 	 * @see Add, Insert, Remove, RemoveAll, RemoveAllSwap
 	 */
-	int32 RemoveSingle(const ElementType& Item)
+	SizeType RemoveSingle(const ElementType& Item)
 	{
-		int32 Index = Find(Item);
+		SizeType Index = Find(Item);
 		if (Index == INDEX_NONE)
 		{
 			return 0;
@@ -2105,7 +2131,7 @@ public:
 
 		// Destruct items that match the specified Item.
 		DestructItems(RemovePtr, 1);
-		const int32 NextIndex = Index + 1;
+		const SizeType NextIndex = Index + 1;
 		RelocateConstructItems<ElementType>(RemovePtr, RemovePtr + 1, ArrayNum - (Index + 1));
 
 		// Update the array count
@@ -2123,7 +2149,7 @@ public:
 	 * @returns Number of removed elements.
 	 * @see Add, Insert, RemoveAll, RemoveAllSwap, RemoveSingle, RemoveSwap
 	 */
-	int32 Remove(const ElementType& Item)
+	SizeType Remove(const ElementType& Item)
 	{
 		CheckAddress(&Item);
 
@@ -2140,25 +2166,25 @@ public:
 	 * @see Add, Insert, RemoveAllSwap, RemoveSingle, RemoveSwap
 	 */
 	template <class PREDICATE_CLASS>
-	int32 RemoveAll(const PREDICATE_CLASS& Predicate)
+	SizeType RemoveAll(const PREDICATE_CLASS& Predicate)
 	{
-		const int32 OriginalNum = ArrayNum;
+		const SizeType OriginalNum = ArrayNum;
 		if (!OriginalNum)
 		{
 			return 0; // nothing to do, loop assumes one item so need to deal with this edge case here
 		}
 
-		int32 WriteIndex = 0;
-		int32 ReadIndex = 0;
+		SizeType WriteIndex = 0;
+		SizeType ReadIndex = 0;
 		bool NotMatch = !Predicate(GetData()[ReadIndex]); // use a ! to guarantee it can't be anything other than zero or one
 		do
 		{
-			int32 RunStartIndex = ReadIndex++;
+			SizeType RunStartIndex = ReadIndex++;
 			while (ReadIndex < OriginalNum && NotMatch == !Predicate(GetData()[ReadIndex]))
 			{
 				ReadIndex++;
 			}
-			int32 RunLength = ReadIndex - RunStartIndex;
+			SizeType RunLength = ReadIndex - RunStartIndex;
 			checkSlow(RunLength > 0);
 			if (NotMatch)
 			{
@@ -2190,7 +2216,7 @@ public:
 	template <class PREDICATE_CLASS>
 	void RemoveAllSwap(const PREDICATE_CLASS& Predicate, bool bAllowShrinking = true)
 	{
-		for (int32 ItemIndex = 0; ItemIndex < Num();)
+		for (SizeType ItemIndex = 0; ItemIndex < Num();)
 		{
 			if (Predicate((*this)[ItemIndex]))
 			{
@@ -2212,9 +2238,9 @@ public:
 	 * @returns The number of items removed. For RemoveSingleItem, this is always either 0 or 1.
 	 * @see Add, Insert, Remove, RemoveAll, RemoveAllSwap, RemoveSwap
 	 */
-	int32 RemoveSingleSwap(const ElementType& Item, bool bAllowShrinking = true)
+	SizeType RemoveSingleSwap(const ElementType& Item, bool bAllowShrinking = true)
 	{
-		int32 Index = Find(Item);
+		SizeType Index = Find(Item);
 		if (Index == INDEX_NONE)
 		{
 			return 0;
@@ -2236,12 +2262,12 @@ public:
 	 * @returns Number of elements removed.
 	 * @see Add, Insert, Remove, RemoveAll, RemoveAllSwap
 	 */
-	int32 RemoveSwap(const ElementType& Item)
+	SizeType RemoveSwap(const ElementType& Item)
 	{
 		CheckAddress(&Item);
 
-		const int32 OriginalNum = ArrayNum;
-		for (int32 Index = 0; Index < ArrayNum; Index++)
+		const SizeType OriginalNum = ArrayNum;
+		for (SizeType Index = 0; Index < ArrayNum; Index++)
 		{
 			if ((*this)[Index] == Item)
 			{
@@ -2257,7 +2283,7 @@ public:
 	 * @param FirstIndexToSwap Position of the first element to swap.
 	 * @param SecondIndexToSwap Position of the second element to swap.
 	 */
-	FORCEINLINE void SwapMemory(int32 FirstIndexToSwap, int32 SecondIndexToSwap)
+	FORCEINLINE void SwapMemory(SizeType FirstIndexToSwap, SizeType SecondIndexToSwap)
 	{
 		FMemory::Memswap(
 			(uint8*)AllocatorInstance.GetAllocation() + (sizeof(ElementType)*FirstIndexToSwap),
@@ -2274,7 +2300,7 @@ public:
 	 * @param FirstIndexToSwap Position of the first element to swap.
 	 * @param SecondIndexToSwap Position of the second element to swap.
 	 */
-	FORCEINLINE void Swap(int32 FirstIndexToSwap, int32 SecondIndexToSwap)
+	FORCEINLINE void Swap(SizeType FirstIndexToSwap, SizeType SecondIndexToSwap)
 	{
 		check((FirstIndexToSwap >= 0) && (SecondIndexToSwap >= 0));
 		check((ArrayNum > FirstIndexToSwap) && (ArrayNum > SecondIndexToSwap));
@@ -2299,10 +2325,10 @@ public:
 	 * @returns True if element was found. False otherwise.
 	 */
 	template<typename SearchType>
-	bool FindItemByClass(SearchType **Item = nullptr, int32 *ItemIndex = nullptr, int32 StartIndex = 0) const
+	bool FindItemByClass(SearchType **Item = nullptr, SizeType *ItemIndex = nullptr, SizeType StartIndex = 0) const
 	{
 		UClass* SearchClass = SearchType::StaticClass();
-		for (int32 Idx = StartIndex; Idx < ArrayNum; Idx++)
+		for (SizeType Idx = StartIndex; Idx < ArrayNum; Idx++)
 		{
 			if ((*this)[Idx] != nullptr && (*this)[Idx]->IsA(SearchClass))
 			{
@@ -2321,8 +2347,8 @@ public:
 	}
 
 	// Iterators
-	typedef TIndexedContainerIterator<      TArray,       ElementType, int32> TIterator;
-	typedef TIndexedContainerIterator<const TArray, const ElementType, int32> TConstIterator;
+	typedef TIndexedContainerIterator<      TArray,       ElementType, SizeType> TIterator;
+	typedef TIndexedContainerIterator<const TArray, const ElementType, SizeType> TConstIterator;
 
 	/**
 	 * Creates an iterator for the contents of this array
@@ -2345,8 +2371,8 @@ public:
 	}
 
 	#if TARRAY_RANGED_FOR_CHECKS
-		typedef TCheckedPointerIterator<      ElementType> RangedForIteratorType;
-		typedef TCheckedPointerIterator<const ElementType> RangedForConstIteratorType;
+		typedef TCheckedPointerIterator<      ElementType, SizeType> RangedForIteratorType;
+		typedef TCheckedPointerIterator<const ElementType, SizeType> RangedForConstIteratorType;
 	#else
 		typedef       ElementType* RangedForIteratorType;
 		typedef const ElementType* RangedForConstIteratorType;
@@ -2440,7 +2466,7 @@ private:
 	 * @param Index Position to get.
 	 * @returns Reference to the element at given position.
 	 */
-	FORCENOINLINE const ElementType& DebugGet(int32 Index) const
+	FORCENOINLINE const ElementType& DebugGet(SizeType Index) const
 	{
 		return GetData()[Index];
 	}
@@ -2448,14 +2474,14 @@ private:
 
 private:
 
-	FORCENOINLINE void ResizeGrow(int32 OldNum)
+	FORCENOINLINE void ResizeGrow(SizeType OldNum)
 	{
 		ArrayMax = AllocatorInstance.CalculateSlackGrow(ArrayNum, ArrayMax, sizeof(ElementType));
 		AllocatorInstance.ResizeAllocation(OldNum, ArrayMax, sizeof(ElementType));
 	}
 	FORCENOINLINE void ResizeShrink()
 	{
-		const int32 NewArrayMax = AllocatorInstance.CalculateSlackShrink(ArrayNum, ArrayMax, sizeof(ElementType));
+		const SizeType NewArrayMax = AllocatorInstance.CalculateSlackShrink(ArrayNum, ArrayMax, sizeof(ElementType));
 		if (NewArrayMax != ArrayMax)
 		{
 			ArrayMax = NewArrayMax;
@@ -2463,7 +2489,7 @@ private:
 			AllocatorInstance.ResizeAllocation(ArrayNum, ArrayMax, sizeof(ElementType));
 		}
 	}
-	FORCENOINLINE void ResizeTo(int32 NewMax)
+	FORCENOINLINE void ResizeTo(SizeType NewMax)
 	{
 		if (NewMax)
 		{
@@ -2475,7 +2501,7 @@ private:
 			AllocatorInstance.ResizeAllocation(ArrayNum, ArrayMax, sizeof(ElementType));
 		}
 	}
-	FORCENOINLINE void ResizeForCopy(int32 NewMax, int32 PrevMax)
+	FORCENOINLINE void ResizeForCopy(SizeType NewMax, SizeType PrevMax)
 	{
 		if (NewMax)
 		{
@@ -2499,14 +2525,17 @@ private:
 	 *                   the end of the buffer. Counted in elements. Zero by
 	 *                   default.
 	 */
-	template <typename OtherElementType>
-	void CopyToEmpty(const OtherElementType* OtherData, int32 OtherNum, int32 PrevMax, int32 ExtraSlack)
+	template <typename OtherElementType, typename OtherSizeType>
+	void CopyToEmpty(const OtherElementType* OtherData, OtherSizeType OtherNum, SizeType PrevMax, SizeType ExtraSlack)
 	{
+		SizeType NewNum = OtherNum;
+		checkf((OtherSizeType)NewNum == OtherNum, TEXT("Invalid number of elements to add to this array type: %llu"), (unsigned long long)NewNum);
+
 		checkSlow(ExtraSlack >= 0);
-		ArrayNum = OtherNum;
+		ArrayNum = NewNum;
 		if (OtherNum || ExtraSlack || PrevMax)
 		{
-			ResizeForCopy(OtherNum + ExtraSlack, PrevMax);
+			ResizeForCopy(NewNum + ExtraSlack, PrevMax);
 			ConstructItems<ElementType>(GetData(), OtherData, OtherNum);
 		}
 		else
@@ -2524,8 +2553,8 @@ protected:
 		>::Result ElementAllocatorType;
 
 	ElementAllocatorType AllocatorInstance;
-	int32	  ArrayNum;
-	int32	  ArrayMax;
+	SizeType             ArrayNum;
+	SizeType             ArrayMax;
 
 	/**
 	 * Implicit heaps
@@ -2574,12 +2603,12 @@ public:
 	 *        The auto-dereferencing behavior does not occur with smart pointers.
 	 */
 	template <class PREDICATE_CLASS>
-	int32 HeapPush(ElementType&& InItem, const PREDICATE_CLASS& Predicate)
+	SizeType HeapPush(ElementType&& InItem, const PREDICATE_CLASS& Predicate)
 	{
 		// Add at the end, then sift up
 		Add(MoveTempIfPossible(InItem));
 		TDereferenceWrapper<ElementType, PREDICATE_CLASS> PredicateWrapper(Predicate);
-		int32 Result = AlgoImpl::HeapSiftUp(GetData(), 0, Num() - 1, FIdentityFunctor(), PredicateWrapper);
+		SizeType Result = AlgoImpl::HeapSiftUp(GetData(), 0, Num() - 1, FIdentityFunctor(), PredicateWrapper);
 
 		return Result;
 	}
@@ -2596,12 +2625,12 @@ public:
 	 *        The auto-dereferencing behavior does not occur with smart pointers.
 	 */
 	template <class PREDICATE_CLASS>
-	int32 HeapPush(const ElementType& InItem, const PREDICATE_CLASS& Predicate)
+	SizeType HeapPush(const ElementType& InItem, const PREDICATE_CLASS& Predicate)
 	{
 		// Add at the end, then sift up
 		Add(InItem);
 		TDereferenceWrapper<ElementType, PREDICATE_CLASS> PredicateWrapper(Predicate);
-		int32 Result = AlgoImpl::HeapSiftUp(GetData(), 0, Num() - 1, FIdentityFunctor(), PredicateWrapper);
+		SizeType Result = AlgoImpl::HeapSiftUp(GetData(), 0, Num() - 1, FIdentityFunctor(), PredicateWrapper);
 
 		return Result;
 	}
@@ -2617,7 +2646,7 @@ public:
 	 *        Therefore, your array will be heapified by the values being pointed to, rather than the pointers' values.
 	 *        The auto-dereferencing behavior does not occur with smart pointers.
 	 */
-	int32 HeapPush(ElementType&& InItem)
+	SizeType HeapPush(ElementType&& InItem)
 	{
 		return HeapPush(MoveTempIfPossible(InItem), TLess<ElementType>());
 	}
@@ -2633,7 +2662,7 @@ public:
 	 *        Therefore, your array will be heapified by the values being pointed to, rather than the pointers' values.
 	 *        The auto-dereferencing behavior does not occur with smart pointers.
 	 */
-	int32 HeapPush(const ElementType& InItem)
+	SizeType HeapPush(const ElementType& InItem)
 	{
 		return HeapPush(InItem, TLess<ElementType>());
 	}
@@ -2753,7 +2782,7 @@ public:
 	 *        The auto-dereferencing behavior does not occur with smart pointers.
 	 */
 	template <class PREDICATE_CLASS>
-	void HeapRemoveAt(int32 Index, const PREDICATE_CLASS& Predicate, bool bAllowShrinking = true)
+	void HeapRemoveAt(SizeType Index, const PREDICATE_CLASS& Predicate, bool bAllowShrinking = true)
 	{
 		RemoveAtSwap(Index, 1, bAllowShrinking);
 
@@ -2773,7 +2802,7 @@ public:
 	 *        Therefore, your array will be heapified by the values being pointed to, rather than the pointers' values.
 	 *        The auto-dereferencing behavior does not occur with smart pointers.
 	 */
-	void HeapRemoveAt(int32 Index, bool bAllowShrinking = true)
+	void HeapRemoveAt(SizeType Index, bool bAllowShrinking = true)
 	{
 		HeapRemoveAt(Index, TLess< ElementType >(), bAllowShrinking);
 	}
@@ -2846,13 +2875,13 @@ template <typename InElementType, typename Allocator> struct TIsTArray<const vol
 template <typename T,typename Allocator> void* operator new( size_t Size, TArray<T,Allocator>& Array )
 {
 	check(Size == sizeof(T));
-	const int32 Index = Array.AddUninitialized(1);
+	const auto Index = Array.AddUninitialized(1);
 	return &Array[Index];
 }
-template <typename T,typename Allocator> void* operator new( size_t Size, TArray<T,Allocator>& Array, int32 Index )
+template <typename T,typename Allocator> void* operator new( size_t Size, TArray<T,Allocator>& Array, typename TArray<T, Allocator>::SizeType Index )
 {
 	check(Size == sizeof(T));
-	Array.InsertUninitialized(Index,1);
+	Array.InsertUninitialized(Index);
 	return &Array[Index];
 }
 
diff --git a/Engine/Source/Runtime/Core/Public/Containers/ContainerAllocationPolicies.h b/Engine/Source/Runtime/Core/Public/Containers/ContainerAllocationPolicies.h
--- a/Engine/Source/Runtime/Core/Public/Containers/ContainerAllocationPolicies.h
+++ b/Engine/Source/Runtime/Core/Public/Containers/ContainerAllocationPolicies.h
@@ -12,6 +12,9 @@
 
 class FDefaultBitArrayAllocator;
 
+template<int IndexSize> class TSizedDefaultAllocator;
+using FDefaultAllocator = TSizedDefaultAllocator<32>;
+
 /** branchless pointer selection
 * return A ? A : B;
 **/
@@ -24,13 +27,14 @@ ReferencedType* IfAThenAElseB(ReferencedType* A,ReferencedType* B);
 template<typename PredicateType,typename ReferencedType>
 ReferencedType* IfPThenAElseB(PredicateType Predicate,ReferencedType* A,ReferencedType* B);
 
-FORCEINLINE int32 DefaultCalculateSlackShrink(int32 NumElements, int32 NumAllocatedElements, SIZE_T BytesPerElement, bool bAllowQuantize, uint32 Alignment = DEFAULT_ALIGNMENT)
+template <typename SizeType>
+FORCEINLINE SizeType DefaultCalculateSlackShrink(SizeType NumElements, SizeType NumAllocatedElements, SIZE_T BytesPerElement, bool bAllowQuantize, uint32 Alignment = DEFAULT_ALIGNMENT)
 {
-	int32 Retval;
+	SizeType Retval;
 	checkSlow(NumElements < NumAllocatedElements);
 
 	// If the container has too much slack, shrink it to exactly fit the number of elements.
-	const uint32 CurrentSlackElements = NumAllocatedElements - NumElements;
+	const SizeType CurrentSlackElements = NumAllocatedElements - NumElements;
 	const SIZE_T CurrentSlackBytes = (NumAllocatedElements - NumElements)*BytesPerElement;
 	const bool bTooManySlackBytes = CurrentSlackBytes >= 16384;
 	const bool bTooManySlackElements = 3 * NumElements < 2 * NumAllocatedElements;
@@ -53,7 +57,8 @@ FORCEINLINE int32 DefaultCalculateSlackShrink(int32 NumElements, int32 NumAlloca
 	return Retval;
 }
 
-FORCEINLINE int32 DefaultCalculateSlackGrow(int32 NumElements, int32 NumAllocatedElements, SIZE_T BytesPerElement, bool bAllowQuantize, uint32 Alignment = DEFAULT_ALIGNMENT)
+template <typename SizeType>
+FORCEINLINE SizeType DefaultCalculateSlackGrow(SizeType NumElements, SizeType NumAllocatedElements, SIZE_T BytesPerElement, bool bAllowQuantize, uint32 Alignment = DEFAULT_ALIGNMENT)
 {
 #if !defined(AGGRESSIVE_MEMORY_SAVING)
 	#error "AGGRESSIVE_MEMORY_SAVING must be defined"
@@ -66,7 +71,7 @@ FORCEINLINE int32 DefaultCalculateSlackGrow(int32 NumElements, int32 NumAllocate
 	const SIZE_T ConstantGrow = 16;
 #endif
 
-	int32 Retval;
+	SizeType Retval;
 	checkSlow(NumElements > NumAllocatedElements && NumElements > 0);
 
 	SIZE_T Grow = FirstGrow; // this is the amount for the first alloc
@@ -86,15 +91,16 @@ FORCEINLINE int32 DefaultCalculateSlackGrow(int32 NumElements, int32 NumAllocate
 	// NumElements and MaxElements are stored in 32 bit signed integers so we must be careful not to overflow here.
 	if (NumElements > Retval)
 	{
-		Retval = MAX_int32;
+		Retval = TNumericLimits<SizeType>::Max();
 	}
 
 	return Retval;
 }
 
-FORCEINLINE int32 DefaultCalculateSlackReserve(int32 NumElements, SIZE_T BytesPerElement, bool bAllowQuantize, uint32 Alignment = DEFAULT_ALIGNMENT)
+template <typename SizeType>
+FORCEINLINE SizeType DefaultCalculateSlackReserve(SizeType NumElements, SIZE_T BytesPerElement, bool bAllowQuantize, uint32 Alignment = DEFAULT_ALIGNMENT)
 {
-	int32 Retval = NumElements;
+	SizeType Retval = NumElements;
 	checkSlow(NumElements > 0);
 	if (bAllowQuantize)
 	{
@@ -102,7 +108,7 @@ FORCEINLINE int32 DefaultCalculateSlackReserve(int32 NumElements, SIZE_T BytesPe
 		// NumElements and MaxElements are stored in 32 bit signed integers so we must be careful not to overflow here.
 		if (NumElements > Retval)
 		{
-			Retval = MAX_int32;
+			Retval = TNumericLimits<SizeType>::Max();
 		}
 	}
 
@@ -130,6 +136,8 @@ struct TAllocatorTraits : TAllocatorTraitsBase<AllocatorType>
 class FContainerAllocatorInterface
 {
 public:
+	/** The integral type to be used for element counts and indices used by the allocator and container - must be signed */
+	using SizeType = int32;
 
 	/** Determines whether the user of the allocator may use the ForAnyElementType inner class. */
 	enum { NeedsElementType = true };
@@ -159,8 +167,8 @@ public:
 		 * @param NumBytesPerElement - The number of bytes/element.
 		 */
 		void ResizeAllocation(
-			int32 PreviousNumElements,
-			int32 NumElements,
+			SizeType PreviousNumElements,
+			SizeType NumElements,
 			SIZE_T NumBytesPerElement
 			);
 
@@ -170,33 +178,33 @@ public:
 		 * @param CurrentNumSlackElements - The current number of elements allocated.
 		 * @param NumBytesPerElement - The number of bytes/element.
 		 */
-		int32 CalculateSlackReserve(
-			int32 NumElements,
-			int32 CurrentNumSlackElements,
+		SizeType CalculateSlackReserve(
+			SizeType NumElements,
+			SizeType CurrentNumSlackElements,
 			SIZE_T NumBytesPerElement
 			) const;
 
 		/**
-		* Calculates the amount of slack to allocate for an array that has just shrunk to a given number of elements.
-		* @param NumElements - The number of elements to allocate space for.
-		* @param CurrentNumSlackElements - The current number of elements allocated.
-		* @param NumBytesPerElement - The number of bytes/element.
-		*/
-		int32 CalculateSlackShrink(
-			int32 NumElements,
-			int32 CurrentNumSlackElements,
+		 * Calculates the amount of slack to allocate for an array that has just shrunk to a given number of elements.
+		 * @param NumElements - The number of elements to allocate space for.
+		 * @param CurrentNumSlackElements - The current number of elements allocated.
+		 * @param NumBytesPerElement - The number of bytes/element.
+		 */
+		SizeType CalculateSlackShrink(
+			SizeType NumElements,
+			SizeType CurrentNumSlackElements,
 			SIZE_T NumBytesPerElement
 			) const;
 
 		/**
-		* Calculates the amount of slack to allocate for an array that has just grown to a given number of elements.
-		* @param NumElements - The number of elements to allocate space for.
-		* @param CurrentNumSlackElements - The current number of elements allocated.
-		* @param NumBytesPerElement - The number of bytes/element.
-		*/
-		int32 CalculateSlackGrow(
-			int32 NumElements,
-			int32 CurrentNumSlackElements,
+		 * Calculates the amount of slack to allocate for an array that has just grown to a given number of elements.
+		 * @param NumElements - The number of elements to allocate space for.
+		 * @param CurrentNumSlackElements - The current number of elements allocated.
+		 * @param NumBytesPerElement - The number of bytes/element.
+		 */
+		SizeType CalculateSlackGrow(
+			SizeType NumElements,
+			SizeType CurrentNumSlackElements,
 			SIZE_T NumBytesPerElement
 			) const;
 
@@ -205,7 +213,7 @@ public:
 		 * @param NumAllocatedElements - The number of elements allocated by the container.
 		 * @param NumBytesPerElement - The number of bytes/element.
 		 */
-		SIZE_T GetAllocatedSize(int32 NumAllocatedElements, SIZE_T NumBytesPerElement) const;
+		SIZE_T GetAllocatedSize(SizeType NumAllocatedElements, SIZE_T NumBytesPerElement) const;
 
 		/** Returns true if the allocator has made any heap allocations */
 		bool HasAllocation() const;
@@ -223,6 +231,7 @@ template<uint32 Alignment = DEFAULT_ALIGNMENT>
 class TAlignedHeapAllocator
 {
 public:
+	using SizeType = int32;
 
 	enum { NeedsElementType = false };
 	enum { RequireRangeCheck = true };
@@ -269,8 +278,8 @@ public:
 			return Data;
 		}
 		void ResizeAllocation(
-			int32 PreviousNumElements,
-			int32 NumElements,
+			SizeType PreviousNumElements,
+			SizeType NumElements,
 			SIZE_T NumBytesPerElement
 			)
 		{
@@ -281,20 +290,20 @@ public:
 				Data = (FScriptContainerElement*)FMemory::Realloc( Data, NumElements*NumBytesPerElement, Alignment );
 			}
 		}
-		FORCEINLINE int32 CalculateSlackReserve(int32 NumElements, SIZE_T NumBytesPerElement) const
+		FORCEINLINE SizeType CalculateSlackReserve(SizeType NumElements, SIZE_T NumBytesPerElement) const
 		{
 			return DefaultCalculateSlackReserve(NumElements, NumBytesPerElement, true, Alignment);
 		}
-		FORCEINLINE int32 CalculateSlackShrink(int32 NumElements, int32 NumAllocatedElements, SIZE_T NumBytesPerElement) const
+		FORCEINLINE SizeType CalculateSlackShrink(SizeType NumElements, SizeType NumAllocatedElements, SIZE_T NumBytesPerElement) const
 		{
 			return DefaultCalculateSlackShrink(NumElements, NumAllocatedElements, NumBytesPerElement, true, Alignment);
 		}
-		FORCEINLINE int32 CalculateSlackGrow(int32 NumElements, int32 NumAllocatedElements, SIZE_T NumBytesPerElement) const
+		FORCEINLINE SizeType CalculateSlackGrow(SizeType NumElements, SizeType NumAllocatedElements, SIZE_T NumBytesPerElement) const
 		{
 			return DefaultCalculateSlackGrow(NumElements, NumAllocatedElements, NumBytesPerElement, true, Alignment);
 		}
 
-		SIZE_T GetAllocatedSize(int32 NumAllocatedElements, SIZE_T NumBytesPerElement) const
+		SIZE_T GetAllocatedSize(SizeType NumAllocatedElements, SIZE_T NumBytesPerElement) const
 		{
 			return NumAllocatedElements * NumBytesPerElement;
 		}
@@ -335,15 +344,28 @@ struct TAllocatorTraits<TAlignedHeapAllocator<Alignment>> : TAllocatorTraitsBase
 	enum { IsZeroConstruct = true };
 };
 
+template <int IndexSize>
+struct TBitsToSizeType
+{
+	static_assert(IndexSize, "Unsupported allocator index size.");
+};
+
+template <> struct TBitsToSizeType<8>  { using Type = int8; };
+template <> struct TBitsToSizeType<16> { using Type = int16; };
+template <> struct TBitsToSizeType<32> { using Type = int32; };
+template <> struct TBitsToSizeType<64> { using Type = int64; };
+
 /** The indirect allocation policy always allocates the elements indirectly. */
-class CORE_API FHeapAllocator
+template <int IndexSize>
+class TSizedHeapAllocator
 {
 public:
+	using SizeType = typename TBitsToSizeType<IndexSize>::Type;
 
 	enum { NeedsElementType = false };
 	enum { RequireRangeCheck = true };
 
-	class CORE_API ForAnyElementType
+	class ForAnyElementType
 	{
 	public:
 		/** Default constructor. */
@@ -383,7 +405,7 @@ public:
 		{
 			return Data;
 		}
-		FORCEINLINE void ResizeAllocation(int32 PreviousNumElements, int32 NumElements, SIZE_T NumBytesPerElement)
+		FORCEINLINE void ResizeAllocation(SizeType PreviousNumElements, SizeType NumElements, SIZE_T NumBytesPerElement)
 		{
 			// Avoid calling FMemory::Realloc( nullptr, 0 ) as ANSI C mandates returning a valid pointer which is not what we want.
 			if (Data || NumElements)
@@ -392,20 +414,20 @@ public:
 				Data = (FScriptContainerElement*)FMemory::Realloc( Data, NumElements*NumBytesPerElement );
 			}
 		}
-		FORCEINLINE int32 CalculateSlackReserve(int32 NumElements, SIZE_T NumBytesPerElement) const
+		FORCEINLINE SizeType CalculateSlackReserve(SizeType NumElements, SIZE_T NumBytesPerElement) const
 		{
 			return DefaultCalculateSlackReserve(NumElements, NumBytesPerElement, true);
 		}
-		FORCEINLINE int32 CalculateSlackShrink(int32 NumElements, int32 NumAllocatedElements, SIZE_T NumBytesPerElement) const
+		FORCEINLINE SizeType CalculateSlackShrink(SizeType NumElements, SizeType NumAllocatedElements, SizeType NumBytesPerElement) const
 		{
 			return DefaultCalculateSlackShrink(NumElements, NumAllocatedElements, NumBytesPerElement, true);
 		}
-		FORCEINLINE int32 CalculateSlackGrow(int32 NumElements, int32 NumAllocatedElements, SIZE_T NumBytesPerElement) const
+		FORCEINLINE SizeType CalculateSlackGrow(SizeType NumElements, SizeType NumAllocatedElements, SizeType NumBytesPerElement) const
 		{
 			return DefaultCalculateSlackGrow(NumElements, NumAllocatedElements, NumBytesPerElement, true);
 		}
 
-		SIZE_T GetAllocatedSize(int32 NumAllocatedElements, SIZE_T NumBytesPerElement) const
+		SIZE_T GetAllocatedSize(SizeType NumAllocatedElements, SIZE_T NumBytesPerElement) const
 		{
 			return NumAllocatedElements * NumBytesPerElement;
 		}
@@ -439,14 +461,14 @@ public:
 	};
 };
 
-template <>
-struct TAllocatorTraits<FHeapAllocator> : TAllocatorTraitsBase<FHeapAllocator>
+template <uint8 IndexSize>
+struct TAllocatorTraits<TSizedHeapAllocator<IndexSize>> : TAllocatorTraitsBase<TSizedHeapAllocator<IndexSize>>
 {
 	enum { SupportsMove    = true };
 	enum { IsZeroConstruct = true };
 };
 
-class FDefaultAllocator;
+using FHeapAllocator = TSizedHeapAllocator<32>;
 
 /**
  * The inline allocation policy allocates up to a specified number of elements in the same allocation as the container.
@@ -457,6 +479,7 @@ template <uint32 NumInlineElements, typename SecondaryAllocator = FDefaultAlloca
 class TInlineAllocator
 {
 public:
+	using SizeType = int32;
 
 	enum { NeedsElementType = true };
 	enum { RequireRangeCheck = true };
@@ -497,7 +520,7 @@ public:
 			return IfAThenAElseB<ElementType>(SecondaryData.GetAllocation(),GetInlineElements());
 		}
 
-		void ResizeAllocation(int32 PreviousNumElements,int32 NumElements,SIZE_T NumBytesPerElement)
+		void ResizeAllocation(SizeType PreviousNumElements, SizeType NumElements,SIZE_T NumBytesPerElement)
 		{
 			// Check if the new allocation will fit in the inline data area.
 			if(NumElements <= NumInlineElements)
@@ -529,21 +552,21 @@ public:
 			}
 		}
 
-		FORCEINLINE int32 CalculateSlackReserve(int32 NumElements, SIZE_T NumBytesPerElement) const
+		FORCEINLINE SizeType CalculateSlackReserve(SizeType NumElements, SIZE_T NumBytesPerElement) const
 		{
 			// If the elements use less space than the inline allocation, only use the inline allocation as slack.
 			return NumElements <= NumInlineElements ?
 				NumInlineElements :
 				SecondaryData.CalculateSlackReserve(NumElements, NumBytesPerElement);
 		}
-		FORCEINLINE int32 CalculateSlackShrink(int32 NumElements, int32 NumAllocatedElements, SIZE_T NumBytesPerElement) const
+		FORCEINLINE SizeType CalculateSlackShrink(SizeType NumElements, SizeType NumAllocatedElements, SIZE_T NumBytesPerElement) const
 		{
 			// If the elements use less space than the inline allocation, only use the inline allocation as slack.
 			return NumElements <= NumInlineElements ?
 				NumInlineElements :
 				SecondaryData.CalculateSlackShrink(NumElements, NumAllocatedElements, NumBytesPerElement);
 		}
-		FORCEINLINE int32 CalculateSlackGrow(int32 NumElements, int32 NumAllocatedElements, SIZE_T NumBytesPerElement) const
+		FORCEINLINE SizeType CalculateSlackGrow(SizeType NumElements, SizeType NumAllocatedElements, SIZE_T NumBytesPerElement) const
 		{
 			// If the elements use less space than the inline allocation, only use the inline allocation as slack.
 			return NumElements <= NumInlineElements ?
@@ -551,7 +574,7 @@ public:
 				SecondaryData.CalculateSlackGrow(NumElements, NumAllocatedElements, NumBytesPerElement);
 		}
 
-		SIZE_T GetAllocatedSize(int32 NumAllocatedElements, SIZE_T NumBytesPerElement) const
+		SIZE_T GetAllocatedSize(SizeType NumAllocatedElements, SIZE_T NumBytesPerElement) const
 		{
 			if (NumAllocatedElements > NumInlineElements)
 			{
@@ -600,6 +623,7 @@ template <uint32 NumInlineElements>
 class TNonRelocatableInlineAllocator
 {
 public:
+	using SizeType = int32;
 
 	enum { NeedsElementType = true };
 	enum { RequireRangeCheck = true };
@@ -647,7 +671,7 @@ public:
 			return Data;
 		}
 
-		void ResizeAllocation(int32 PreviousNumElements,int32 NumElements,SIZE_T NumBytesPerElement)
+		void ResizeAllocation(SizeType PreviousNumElements, SizeType NumElements,SIZE_T NumBytesPerElement)
 		{
 			// Check if the new allocation will fit in the inline data area.
 			if(NumElements <= NumInlineElements)
@@ -678,25 +702,25 @@ public:
 			}
 		}
 
-		FORCEINLINE int32 CalculateSlackReserve(int32 NumElements, SIZE_T NumBytesPerElement) const
+		FORCEINLINE SizeType CalculateSlackReserve(SizeType NumElements, SIZE_T NumBytesPerElement) const
 		{
 			// If the elements use less space than the inline allocation, only use the inline allocation as slack.
 			return (NumElements <= NumInlineElements) ? NumInlineElements : DefaultCalculateSlackReserve(NumElements, NumBytesPerElement, true);
 		}
 
-		FORCEINLINE int32 CalculateSlackShrink(int32 NumElements, int32 NumAllocatedElements, SIZE_T NumBytesPerElement) const
+		FORCEINLINE SizeType CalculateSlackShrink(SizeType NumElements, SizeType NumAllocatedElements, SIZE_T NumBytesPerElement) const
 		{
 			// If the elements use less space than the inline allocation, only use the inline allocation as slack.
 			return (NumElements <= NumInlineElements) ? NumInlineElements : DefaultCalculateSlackShrink(NumElements, NumAllocatedElements, NumBytesPerElement, true);
 		}
 
-		FORCEINLINE int32 CalculateSlackGrow(int32 NumElements, int32 NumAllocatedElements, SIZE_T NumBytesPerElement) const
+		FORCEINLINE SizeType CalculateSlackGrow(SizeType NumElements, SizeType NumAllocatedElements, SIZE_T NumBytesPerElement) const
 		{
 			// If the elements use less space than the inline allocation, only use the inline allocation as slack.
 			return (NumElements <= NumInlineElements) ? NumInlineElements : DefaultCalculateSlackGrow(NumElements, NumAllocatedElements, NumBytesPerElement, true);
 		}
 
-		SIZE_T GetAllocatedSize(int32 NumAllocatedElements, SIZE_T NumBytesPerElement) const
+		SIZE_T GetAllocatedSize(SizeType NumAllocatedElements, SIZE_T NumBytesPerElement) const
 		{
 			return HasAllocation()? (NumAllocatedElements * NumBytesPerElement) : 0;
 		}
@@ -740,6 +764,7 @@ template <uint32 NumInlineElements>
 class TFixedAllocator
 {
 public:
+	using SizeType = int32;
 
 	enum { NeedsElementType = true };
 	enum { RequireRangeCheck = true };
@@ -773,32 +798,32 @@ public:
 			return GetInlineElements();
 		}
 
-		void ResizeAllocation(int32 PreviousNumElements,int32 NumElements,SIZE_T NumBytesPerElement)
+		void ResizeAllocation(SizeType PreviousNumElements, SizeType NumElements,SIZE_T NumBytesPerElement)
 		{
 			// Ensure the requested allocation will fit in the inline data area.
 			checkSlow(NumElements <= NumInlineElements);
 		}
 
-		FORCEINLINE int32 CalculateSlackReserve(int32 NumElements, SIZE_T NumBytesPerElement) const
+		FORCEINLINE SizeType CalculateSlackReserve(SizeType NumElements, SIZE_T NumBytesPerElement) const
 		{
 			// Ensure the requested allocation will fit in the inline data area.
 			checkSlow(NumElements <= NumInlineElements);
 			return NumInlineElements;
 		}
-		FORCEINLINE int32 CalculateSlackShrink(int32 NumElements, int32 NumAllocatedElements, SIZE_T NumBytesPerElement) const
+		FORCEINLINE SizeType CalculateSlackShrink(SizeType NumElements, SizeType NumAllocatedElements, SIZE_T NumBytesPerElement) const
 		{
 			// Ensure the requested allocation will fit in the inline data area.
 			checkSlow(NumAllocatedElements <= NumInlineElements);
 			return NumInlineElements;
 		}
-		FORCEINLINE int32 CalculateSlackGrow(int32 NumElements, int32 NumAllocatedElements, SIZE_T NumBytesPerElement) const
+		FORCEINLINE SizeType CalculateSlackGrow(SizeType NumElements, SizeType NumAllocatedElements, SIZE_T NumBytesPerElement) const
 		{
 			// Ensure the requested allocation will fit in the inline data area.
 			checkSlow(NumElements <= NumInlineElements);
 			return NumInlineElements;
 		}
 
-		SIZE_T GetAllocatedSize(int32 NumAllocatedElements, SIZE_T NumBytesPerElement) const
+		SIZE_T GetAllocatedSize(SizeType NumAllocatedElements, SIZE_T NumBytesPerElement) const
 		{
 			return 0;
 		}
@@ -840,9 +865,6 @@ struct TAllocatorTraits<TFixedAllocator<NumInlineElements>> : TAllocatorTraitsBa
 // Sparse array allocation definitions
 //
 
-class FDefaultAllocator;
-class FDefaultBitArrayAllocator;
-
 /** Encapsulates the allocators used by a sparse array in a single type. */
 template<typename InElementAllocator = FDefaultAllocator,typename InBitArrayAllocator = FDefaultBitArrayAllocator>
 class TSparseArrayAllocator
@@ -921,8 +943,6 @@ public:
 	typedef InHashAllocator        HashAllocator;
 };
 
-class FDefaultAllocator;
-
 /** An inline set allocator that allows sizing of the inline allocations for a set number of elements. */
 template<
 	uint32   NumInlineElements,
@@ -1006,11 +1026,14 @@ public:
  * 'forward' these TAllocatorTraits specializations below.
  */
 
-class FDefaultAllocator            : public FHeapAllocator          { public: typedef FHeapAllocator          Typedef; };
+template <int IndexSize> class TSizedDefaultAllocator : public TSizedHeapAllocator<IndexSize> { public: typedef TSizedHeapAllocator<IndexSize> Typedef; };
+
 class FDefaultSetAllocator         : public TSetAllocator<>         { public: typedef TSetAllocator<>         Typedef; };
 class FDefaultBitArrayAllocator    : public TInlineAllocator<4>     { public: typedef TInlineAllocator<4>     Typedef; };
 class FDefaultSparseArrayAllocator : public TSparseArrayAllocator<> { public: typedef TSparseArrayAllocator<> Typedef; };
 
+template <int IndexSize> struct TAllocatorTraits<TSizedDefaultAllocator<IndexSize>> : TAllocatorTraits<typename TSizedDefaultAllocator<IndexSize>::Typedef> {};
+
 template <> struct TAllocatorTraits<FDefaultAllocator>            : TAllocatorTraits<typename FDefaultAllocator           ::Typedef> {};
 template <> struct TAllocatorTraits<FDefaultSetAllocator>         : TAllocatorTraits<typename FDefaultSetAllocator        ::Typedef> {};
 template <> struct TAllocatorTraits<FDefaultBitArrayAllocator>    : TAllocatorTraits<typename FDefaultBitArrayAllocator   ::Typedef> {};
diff --git a/Engine/Source/Runtime/Core/Public/Containers/ContainersFwd.h b/Engine/Source/Runtime/Core/Public/Containers/ContainersFwd.h
--- a/Engine/Source/Runtime/Core/Public/Containers/ContainersFwd.h
+++ b/Engine/Source/Runtime/Core/Public/Containers/ContainersFwd.h
@@ -3,12 +3,15 @@
 #pragma once
 
 /// @cond DOXYGEN_WARNINGS
-class FDefaultAllocator;
+template<int IndexSize> class TSizedDefaultAllocator;
+using FDefaultAllocator = TSizedDefaultAllocator<32>;
+using FDefaultAllocator64 = TSizedDefaultAllocator<64>;
 class FDefaultSetAllocator;
 
 class FString;
 
 template<typename T, typename Allocator = FDefaultAllocator> class TArray;
+template<typename T> using TArray64 = TArray<T, FDefaultAllocator64>;
 template<typename T> class TTransArray;
 template<typename KeyType, typename ValueType, bool bInAllowDuplicateKeys> struct TDefaultMapHashableKeyFuncs;
 template<typename KeyType, typename ValueType, typename SetAllocator = FDefaultSetAllocator, typename KeyFuncs = TDefaultMapHashableKeyFuncs<KeyType, ValueType, false> > class TMap;
diff --git a/Engine/Source/Runtime/Core/Public/Misc/MemStack.h b/Engine/Source/Runtime/Core/Public/Misc/MemStack.h
--- a/Engine/Source/Runtime/Core/Public/Misc/MemStack.h
+++ b/Engine/Source/Runtime/Core/Public/Misc/MemStack.h
@@ -271,6 +271,7 @@ template<uint32 Alignment = DEFAULT_ALIGNMENT>
 class TMemStackAllocator
 {
 public:
+	using SizeType = int32;
 
 	enum { NeedsElementType = true };
 	enum { RequireRangeCheck = true };
@@ -303,7 +304,7 @@ public:
 		{
 			return Data;
 		}
-		void ResizeAllocation(int32 PreviousNumElements,int32 NumElements, SIZE_T NumBytesPerElement)
+		void ResizeAllocation(SizeType PreviousNumElements, SizeType NumElements,SIZE_T NumBytesPerElement)
 		{
 			void* OldData = Data;
 			if( NumElements )
@@ -317,25 +318,25 @@ public:
 				// If the container previously held elements, copy them into the new allocation.
 				if(OldData && PreviousNumElements)
 				{
-					const int32 NumCopiedElements = FMath::Min(NumElements,PreviousNumElements);
+					const SizeType NumCopiedElements = FMath::Min(NumElements,PreviousNumElements);
 					FMemory::Memcpy(Data,OldData,NumCopiedElements * NumBytesPerElement);
 				}
 			}
 		}
-		FORCEINLINE int32 CalculateSlackReserve(int32 NumElements, SIZE_T NumBytesPerElement) const
+		FORCEINLINE SizeType CalculateSlackReserve(SizeType NumElements, SIZE_T NumBytesPerElement) const
 		{
 			return DefaultCalculateSlackReserve(NumElements, NumBytesPerElement, false, Alignment);
 		}
-		FORCEINLINE int32 CalculateSlackShrink(int32 NumElements, int32 NumAllocatedElements, SIZE_T NumBytesPerElement) const
+		FORCEINLINE SizeType CalculateSlackShrink(SizeType NumElements, SizeType NumAllocatedElements, SIZE_T NumBytesPerElement) const
 		{
 			return DefaultCalculateSlackShrink(NumElements, NumAllocatedElements, NumBytesPerElement, false, Alignment);
 		}
-		FORCEINLINE int32 CalculateSlackGrow(int32 NumElements, int32 NumAllocatedElements, SIZE_T NumBytesPerElement) const
+		FORCEINLINE SizeType CalculateSlackGrow(SizeType NumElements, SizeType NumAllocatedElements, SIZE_T NumBytesPerElement) const
 		{
 			return DefaultCalculateSlackGrow(NumElements, NumAllocatedElements, NumBytesPerElement, false, Alignment);
 		}
 
-		FORCEINLINE SIZE_T GetAllocatedSize(int32 NumAllocatedElements, SIZE_T NumBytesPerElement) const
+		FORCEINLINE SIZE_T GetAllocatedSize(SizeType NumAllocatedElements, SIZE_T NumBytesPerElement) const
 		{
 			return NumAllocatedElements * NumBytesPerElement;
 		}
diff --git a/Engine/Source/Runtime/Core/Public/Templates/MemoryOps.h b/Engine/Source/Runtime/Core/Public/Templates/MemoryOps.h
--- a/Engine/Source/Runtime/Core/Public/Templates/MemoryOps.h
+++ b/Engine/Source/Runtime/Core/Public/Templates/MemoryOps.h
@@ -39,8 +39,8 @@ namespace UE4MemoryOps_Private
  * @param	Elements	The address of the first memory location to construct at.
  * @param	Count		The number of elements to destruct.
  */
-template <typename ElementType>
-FORCEINLINE typename TEnableIf<!TIsZeroConstructType<ElementType>::Value>::Type DefaultConstructItems(void* Address, int32 Count)
+template <typename ElementType, typename SizeType>
+FORCEINLINE typename TEnableIf<!TIsZeroConstructType<ElementType>::Value>::Type DefaultConstructItems(void* Address, SizeType Count)
 {
 	ElementType* Element = (ElementType*)Address;
 	while (Count)
@@ -52,8 +52,8 @@ FORCEINLINE typename TEnableIf<!TIsZeroConstructType<ElementType>::Value>::Type
 }
 
 
-template <typename ElementType>
-FORCEINLINE typename TEnableIf<TIsZeroConstructType<ElementType>::Value>::Type DefaultConstructItems(void* Elements, int32 Count)
+template <typename ElementType, typename SizeType>
+FORCEINLINE typename TEnableIf<TIsZeroConstructType<ElementType>::Value>::Type DefaultConstructItems(void* Elements, SizeType Count)
 {
 	FMemory::Memset(Elements, 0, sizeof(ElementType) * Count);
 }
@@ -90,8 +90,8 @@ FORCEINLINE typename TEnableIf<TIsTriviallyDestructible<ElementType>::Value>::Ty
  *
  * @note: This function is optimized for values of T, and so will not dynamically dispatch destructor calls if T's destructor is virtual.
  */
-template <typename ElementType>
-FORCEINLINE typename TEnableIf<!TIsTriviallyDestructible<ElementType>::Value>::Type DestructItems(ElementType* Element, int32 Count)
+template <typename ElementType, typename SizeType>
+FORCEINLINE typename TEnableIf<!TIsTriviallyDestructible<ElementType>::Value>::Type DestructItems(ElementType* Element, SizeType Count)
 {
 	while (Count)
 	{
@@ -105,8 +105,8 @@ FORCEINLINE typename TEnableIf<!TIsTriviallyDestructible<ElementType>::Value>::T
 }
 
 
-template <typename ElementType>
-FORCEINLINE typename TEnableIf<TIsTriviallyDestructible<ElementType>::Value>::Type DestructItems(ElementType* Elements, int32 Count)
+template <typename ElementType, typename SizeType>
+FORCEINLINE typename TEnableIf<TIsTriviallyDestructible<ElementType>::Value>::Type DestructItems(ElementType* Elements, SizeType Count)
 {
 }
 
@@ -118,8 +118,8 @@ FORCEINLINE typename TEnableIf<TIsTriviallyDestructible<ElementType>::Value>::Ty
  * @param	Source		A pointer to the first argument to pass to the constructor.
  * @param	Count		The number of elements to copy.
  */
-template <typename DestinationElementType, typename SourceElementType>
-FORCEINLINE typename TEnableIf<!TIsBitwiseConstructible<DestinationElementType, SourceElementType>::Value>::Type ConstructItems(void* Dest, const SourceElementType* Source, int32 Count)
+template <typename DestinationElementType, typename SourceElementType, typename SizeType>
+FORCEINLINE typename TEnableIf<!TIsBitwiseConstructible<DestinationElementType, SourceElementType>::Value>::Type ConstructItems(void* Dest, const SourceElementType* Source, SizeType Count)
 {
 	while (Count)
 	{
@@ -131,8 +131,8 @@ FORCEINLINE typename TEnableIf<!TIsBitwiseConstructible<DestinationElementType,
 }
 
 
-template <typename DestinationElementType, typename SourceElementType>
-FORCEINLINE typename TEnableIf<TIsBitwiseConstructible<DestinationElementType, SourceElementType>::Value>::Type ConstructItems(void* Dest, const SourceElementType* Source, int32 Count)
+template <typename DestinationElementType, typename SourceElementType, typename SizeType>
+FORCEINLINE typename TEnableIf<TIsBitwiseConstructible<DestinationElementType, SourceElementType>::Value>::Type ConstructItems(void* Dest, const SourceElementType* Source, SizeType Count)
 {
 	FMemory::Memcpy(Dest, Source, sizeof(SourceElementType) * Count);
 }
@@ -145,8 +145,8 @@ FORCEINLINE typename TEnableIf<TIsBitwiseConstructible<DestinationElementType, S
  * @param	Source		A pointer to the first item to assign.
  * @param	Count		The number of elements to assign.
  */
-template <typename ElementType>
-FORCEINLINE typename TEnableIf<!TIsTriviallyCopyAssignable<ElementType>::Value>::Type CopyAssignItems(ElementType* Dest, const ElementType* Source, int32 Count)
+template <typename ElementType, typename SizeType>
+FORCEINLINE typename TEnableIf<!TIsTriviallyCopyAssignable<ElementType>::Value>::Type CopyAssignItems(ElementType* Dest, const ElementType* Source, SizeType Count)
 {
 	while (Count)
 	{
@@ -158,8 +158,8 @@ FORCEINLINE typename TEnableIf<!TIsTriviallyCopyAssignable<ElementType>::Value>:
 }
 
 
-template <typename ElementType>
-FORCEINLINE typename TEnableIf<TIsTriviallyCopyAssignable<ElementType>::Value>::Type CopyAssignItems(ElementType* Dest, const ElementType* Source, int32 Count)
+template <typename ElementType, typename SizeType>
+FORCEINLINE typename TEnableIf<TIsTriviallyCopyAssignable<ElementType>::Value>::Type CopyAssignItems(ElementType* Dest, const ElementType* Source, SizeType Count)
 {
 	FMemory::Memcpy(Dest, Source, sizeof(ElementType) * Count);
 }
@@ -173,8 +173,8 @@ FORCEINLINE typename TEnableIf<TIsTriviallyCopyAssignable<ElementType>::Value>::
  * @param	Source		A pointer to the first item to relocate.
  * @param	Count		The number of elements to relocate.
  */
-template <typename DestinationElementType, typename SourceElementType>
-FORCEINLINE typename TEnableIf<!UE4MemoryOps_Private::TCanBitwiseRelocate<DestinationElementType, SourceElementType>::Value>::Type RelocateConstructItems(void* Dest, const SourceElementType* Source, int32 Count)
+template <typename DestinationElementType, typename SourceElementType, typename SizeType>
+FORCEINLINE typename TEnableIf<!UE4MemoryOps_Private::TCanBitwiseRelocate<DestinationElementType, SourceElementType>::Value>::Type RelocateConstructItems(void* Dest, const SourceElementType* Source, SizeType Count)
 {
 	while (Count)
 	{
@@ -188,8 +188,8 @@ FORCEINLINE typename TEnableIf<!UE4MemoryOps_Private::TCanBitwiseRelocate<Destin
 	}
 }
 
-template <typename DestinationElementType, typename SourceElementType>
-FORCEINLINE typename TEnableIf<UE4MemoryOps_Private::TCanBitwiseRelocate<DestinationElementType, SourceElementType>::Value>::Type RelocateConstructItems(void* Dest, const SourceElementType* Source, int32 Count)
+template <typename DestinationElementType, typename SourceElementType, typename SizeType>
+FORCEINLINE typename TEnableIf<UE4MemoryOps_Private::TCanBitwiseRelocate<DestinationElementType, SourceElementType>::Value>::Type RelocateConstructItems(void* Dest, const SourceElementType* Source, SizeType Count)
 {
 	/* All existing UE containers seem to assume trivial relocatability (i.e. memcpy'able) of their members,
 	 * so we're going to assume that this is safe here.  However, it's not generally possible to assume this
@@ -209,8 +209,8 @@ FORCEINLINE typename TEnableIf<UE4MemoryOps_Private::TCanBitwiseRelocate<Destina
  * @param	Source		A pointer to the first item to move from.
  * @param	Count		The number of elements to move.
  */
-template <typename ElementType>
-FORCEINLINE typename TEnableIf<!TIsTriviallyCopyConstructible<ElementType>::Value>::Type MoveConstructItems(void* Dest, const ElementType* Source, int32 Count)
+template <typename ElementType, typename SizeType>
+FORCEINLINE typename TEnableIf<!TIsTriviallyCopyConstructible<ElementType>::Value>::Type MoveConstructItems(void* Dest, const ElementType* Source, SizeType Count)
 {
 	while (Count)
 	{
@@ -221,8 +221,8 @@ FORCEINLINE typename TEnableIf<!TIsTriviallyCopyConstructible<ElementType>::Valu
 	}
 }
 
-template <typename ElementType>
-FORCEINLINE typename TEnableIf<TIsTriviallyCopyConstructible<ElementType>::Value>::Type MoveConstructItems(void* Dest, const ElementType* Source, int32 Count)
+template <typename ElementType, typename SizeType>
+FORCEINLINE typename TEnableIf<TIsTriviallyCopyConstructible<ElementType>::Value>::Type MoveConstructItems(void* Dest, const ElementType* Source, SizeType Count)
 {
 	FMemory::Memmove(Dest, Source, sizeof(ElementType) * Count);
 }
@@ -234,8 +234,8 @@ FORCEINLINE typename TEnableIf<TIsTriviallyCopyConstructible<ElementType>::Value
  * @param	Source		A pointer to the first item to move assign.
  * @param	Count		The number of elements to move assign.
  */
-template <typename ElementType>
-FORCEINLINE typename TEnableIf<!TIsTriviallyCopyAssignable<ElementType>::Value>::Type MoveAssignItems(ElementType* Dest, const ElementType* Source, int32 Count)
+template <typename ElementType, typename SizeType>
+FORCEINLINE typename TEnableIf<!TIsTriviallyCopyAssignable<ElementType>::Value>::Type MoveAssignItems(ElementType* Dest, const ElementType* Source, SizeType Count)
 {
 	while (Count)
 	{
@@ -246,21 +246,21 @@ FORCEINLINE typename TEnableIf<!TIsTriviallyCopyAssignable<ElementType>::Value>:
 	}
 }
 
-template <typename ElementType>
-FORCEINLINE typename TEnableIf<TIsTriviallyCopyAssignable<ElementType>::Value>::Type MoveAssignItems(ElementType* Dest, const ElementType* Source, int32 Count)
+template <typename ElementType, typename SizeType>
+FORCEINLINE typename TEnableIf<TIsTriviallyCopyAssignable<ElementType>::Value>::Type MoveAssignItems(ElementType* Dest, const ElementType* Source, SizeType Count)
 {
 	FMemory::Memmove(Dest, Source, sizeof(ElementType) * Count);
 }
 
-template <typename ElementType>
-FORCEINLINE typename TEnableIf<TTypeTraits<ElementType>::IsBytewiseComparable, bool>::Type CompareItems(const ElementType* A, const ElementType* B, int32 Count)
+template <typename ElementType, typename SizeType>
+FORCEINLINE typename TEnableIf<TTypeTraits<ElementType>::IsBytewiseComparable, bool>::Type CompareItems(const ElementType* A, const ElementType* B, SizeType Count)
 {
 	return !FMemory::Memcmp(A, B, sizeof(ElementType) * Count);
 }
 
 
-template <typename ElementType>
-FORCEINLINE typename TEnableIf<!TTypeTraits<ElementType>::IsBytewiseComparable, bool>::Type CompareItems(const ElementType* A, const ElementType* B, int32 Count)
+template <typename ElementType, typename SizeType>
+FORCEINLINE typename TEnableIf<!TTypeTraits<ElementType>::IsBytewiseComparable, bool>::Type CompareItems(const ElementType* A, const ElementType* B, SizeType Count)
 {
 	while (Count)
 	{
diff --git a/Engine/Source/Runtime/Engine/Classes/Animation/AnimSequence.h b/Engine/Source/Runtime/Engine/Classes/Animation/AnimSequence.h
--- a/Engine/Source/Runtime/Engine/Classes/Animation/AnimSequence.h
+++ b/Engine/Source/Runtime/Engine/Classes/Animation/AnimSequence.h
@@ -425,6 +425,7 @@ template<uint32 Alignment = DEFAULT_ALIGNMENT>
 class TMaybeMappedAllocator
 {
 public:
+	using SizeType = int32;
 
 	enum { NeedsElementType = false };
 	enum { RequireRangeCheck = true };
@@ -473,8 +474,8 @@ public:
 			return Data;
 		}
 		void ResizeAllocation(
-			int32 PreviousNumElements,
-			int32 NumElements,
+			SizeType PreviousNumElements,
+			SizeType NumElements,
 			SIZE_T NumBytesPerElement
 		)
 		{
@@ -486,23 +487,23 @@ public:
 				Data = (FScriptContainerElement*)FMemory::Realloc(Data, NumElements*NumBytesPerElement, Alignment);
 			}
 		}
-		int32 CalculateSlackReserve(int32 NumElements, SIZE_T NumBytesPerElement) const
+		SizeType CalculateSlackReserve(SizeType NumElements, SIZE_T NumBytesPerElement) const
 		{
 			check(!MappedHandle && !MappedRegion); // this could be supported, but it probably is never what you want, so we will just assert.
 			return DefaultCalculateSlackReserve(NumElements, NumBytesPerElement, true, Alignment);
 		}
-		int32 CalculateSlackShrink(int32 NumElements, int32 NumAllocatedElements, SIZE_T NumBytesPerElement) const
+		SizeType CalculateSlackShrink(SizeType NumElements, SizeType NumAllocatedElements, SIZE_T NumBytesPerElement) const
 		{
 			check(!MappedHandle && !MappedRegion); // this could be supported, but it probably is never what you want, so we will just assert.
 			return DefaultCalculateSlackShrink(NumElements, NumAllocatedElements, NumBytesPerElement, true, Alignment);
 		}
-		int32 CalculateSlackGrow(int32 NumElements, int32 NumAllocatedElements, SIZE_T NumBytesPerElement) const
+		SizeType CalculateSlackGrow(SizeType NumElements, SizeType NumAllocatedElements, SIZE_T NumBytesPerElement) const
 		{
 			check(!MappedHandle && !MappedRegion); // this could be supported, but it probably is never what you want, so we will just assert.
 			return DefaultCalculateSlackGrow(NumElements, NumAllocatedElements, NumBytesPerElement, true, Alignment);
 		}
 
-		SIZE_T GetAllocatedSize(int32 NumAllocatedElements, SIZE_T NumBytesPerElement) const
+		SIZE_T GetAllocatedSize(SizeType NumAllocatedElements, SIZE_T NumBytesPerElement) const
 		{
 			return NumAllocatedElements * NumBytesPerElement;
 		}
diff --git a/Engine/Source/Runtime/Engine/Public/Net/Util/ResizableCircularQueue.h b/Engine/Source/Runtime/Engine/Public/Net/Util/ResizableCircularQueue.h
--- a/Engine/Source/Runtime/Engine/Public/Net/Util/ResizableCircularQueue.h
+++ b/Engine/Source/Runtime/Engine/Public/Net/Util/ResizableCircularQueue.h
@@ -12,8 +12,6 @@
  * an extra element or maintain explicit empty state.
  * Capacity most be power of two and it supports growing
  */
-class FDefaultAllocator;
-
 template<typename T, typename AllocatorT = FDefaultAllocator>
 class TResizableCircularQueue
 {
