Unused WIP trait removed.

diff --git a/Engine/Source/Runtime/Core/Public/Containers/Array.h b/Engine/Source/Runtime/Core/Public/Containers/Array.h
--- a/Engine/Source/Runtime/Core/Public/Containers/Array.h
+++ b/Engine/Source/Runtime/Core/Public/Containers/Array.h
@@ -260,18 +260,6 @@ namespace UE4Array_Private
 				)
 		};
 	};
-
-	template <typename AllocatorType>
-	struct TAllocatorSizeType
-	{
-		template <typename T>
-		static typename T::SizeType ResolveIndexType(typename T::SizeType*);
-
-		template <typename T>
-		static int32 ResolveIndexType(...);
-
-		using Type = decltype(ResolveIndexType<AllocatorType>(nullptr));
-	};
 }
 
 
