TIsConst metafunction.

diff --git a/Engine/Source/Runtime/Core/Public/Templates/IsConst.h b/Engine/Source/Runtime/Core/Public/Templates/IsConst.h
new file mode 100644
--- /dev/null
+++ b/Engine/Source/Runtime/Core/Public/Templates/IsConst.h
@@ -0,0 +1,18 @@
+// Copyright 1998-2019 Epic Games, Inc. All Rights Reserved.
+
+#pragma once
+
+/**
+ * Traits class which tests if a type is const.
+ */
+template <typename T>
+struct TIsConst
+{
+	static constexpr bool Value = false;
+};
+
+template <typename T>
+struct TIsConst<const T>
+{
+	static constexpr bool Value = true;
+};
