Copying //UE4/Dev-Core to Main (//UE4/Main)

diff --git a/Engine/Source/Editor/UnrealEd/Private/CookOnTheFlyServer.cpp b/Engine/Source/Editor/UnrealEd/Private/CookOnTheFlyServer.cpp
--- a/Engine/Source/Editor/UnrealEd/Private/CookOnTheFlyServer.cpp
+++ b/Engine/Source/Editor/UnrealEd/Private/CookOnTheFlyServer.cpp
@@ -29,6 +29,7 @@
 #include "UObject/Class.h"
 #include "UObject/UObjectIterator.h"
 #include "UObject/Package.h"
+#include "UObject/SavePackage.h"
 #include "UObject/MetaData.h"
 #include "UObject/ConstructorHelpers.h"
 #include "UObject/UObjectArray.h"
@@ -311,7 +312,7 @@ void ClearHierarchyTimers()
 
 #define CREATE_TIMER(name, incrementScope) FScopeTimer ScopeTimer##name(__COUNTER__, #name, incrementScope); 
 
-#define SCOPE_TIMER(name)				TRACE_CPUPROFILER_EVENT_SCOPE_TEXT(TEXT(#name)); CREATE_TIMER(name, true); ScopeTimer##name.Start();
+#define SCOPE_TIMER(name)				TRACE_CPUPROFILER_EVENT_SCOPE(name); CREATE_TIMER(name, true); ScopeTimer##name.Start();
 
 #else
 #define SCOPE_TIMER(name)
@@ -1422,6 +1423,7 @@ public:
 	bool							bErrorOnEngineContentUse = false;
 	bool							bDisableUnsolicitedPackages = false;
 	bool							bFullLoadAndSave = false;
+	bool							bPackageStore = false;
 	TArray<FName>					TargetPlatformNames;
 	TArray<FName>					StartupPackages;
 
@@ -1464,7 +1466,7 @@ UCookOnTheFlyServer::~UCookOnTheFlyServer()
 // this tick only happens in the editor cook commandlet directly calls tick on the side
 void UCookOnTheFlyServer::Tick(float DeltaTime)
 {
-	TRACE_CPUPROFILER_EVENT_SCOPE_TEXT(TEXT("UCookOnTheFlyServer::Tick"));
+	TRACE_CPUPROFILER_EVENT_SCOPE(UCookOnTheFlyServer::Tick);
 
 	check(IsCookingInEditor());
 
@@ -1982,6 +1984,11 @@ bool UCookOnTheFlyServer::IsCookByTheBookMode() const
 	return CurrentCookMode == ECookMode::CookByTheBookFromTheEditor || CurrentCookMode == ECookMode::CookByTheBook;
 }
 
+bool UCookOnTheFlyServer::IsUsingPackageStore() const
+{
+	return CookByTheBookOptions->bPackageStore;
+}
+
 bool UCookOnTheFlyServer::IsCookOnTheFlyMode() const
 {
 	return CurrentCookMode == ECookMode::CookOnTheFly || CurrentCookMode == ECookMode::CookOnTheFlyFromTheEditor; 
@@ -3621,13 +3628,6 @@ void UCookOnTheFlyServer::SaveCookedPackage(UPackage* Package, uint32 SaveFlags,
 	return SaveCookedPackage( Package, SaveFlags, TargetPlatformNames, SavePackageResults);
 }
 
-bool UCookOnTheFlyServer::ShouldConsiderCompressedPackageFileLengthRequirements() const
-{
-	bool bConsiderCompressedPackageFileLengthRequirements = true;
-	GConfig->GetBool(TEXT("CookSettings"), TEXT("bConsiderCompressedPackageFileLengthRequirements"), bConsiderCompressedPackageFileLengthRequirements, GEditorIni);
-	return bConsiderCompressedPackageFileLengthRequirements;
-}
-
 bool UCookOnTheFlyServer::MakePackageFullyLoaded(UPackage* Package) const
 {
 	if (Package->IsFullyLoaded())
@@ -3999,15 +3999,11 @@ void UCookOnTheFlyServer::SaveCookedPackage(UPackage* Package, uint32 SaveFlags,
 						World->PersistentLevel->HandleLegacyMapBuildData();
 					}
 
-					// need to subtract 32 because the SavePackage code creates temporary files with longer file names then the one we provide
-					// projects may ignore this restriction if desired
-					static bool bConsiderCompressedPackageFileLengthRequirements = ShouldConsiderCompressedPackageFileLengthRequirements();
-					const int32 CompressedPackageFileLengthRequirement = bConsiderCompressedPackageFileLengthRequirements ? 32 : 0;
 					const FString FullFilename = FPaths::ConvertRelativePathToFull(PlatFilename);
-					if (FullFilename.Len() >= (FPlatformMisc::GetMaxPathLength() - CompressedPackageFileLengthRequirement))
+					if (FullFilename.Len() >= FPlatformMisc::GetMaxPathLength())
 					{
-						LogCookerMessage(FString::Printf(TEXT("Couldn't save package, filename is too long: %s"), *PlatFilename), EMessageSeverity::Error);
-						UE_LOG(LogCook, Error, TEXT("Couldn't save package, filename is too long :%s"), *PlatFilename);
+						LogCookerMessage(FString::Printf(TEXT("Couldn't save package, filename is too long (%d >= %d): %s"), FullFilename.Len(), FPlatformMisc::GetMaxPathLength(), *PlatFilename), EMessageSeverity::Error);
+						UE_LOG(LogCook, Error, TEXT("Couldn't save package, filename is too long (%d >= %d): %s"), FullFilename.Len(), FPlatformMisc::GetMaxPathLength(), *PlatFilename);
 						Result = ESavePackageResult::Error;
 					}
 					else
@@ -4033,7 +4029,12 @@ void UCookOnTheFlyServer::SaveCookedPackage(UPackage* Package, uint32 SaveFlags,
 						}
 						else
 						{
-							Result = GEditor->Save(Package, World, FlagsToCook, *PlatFilename, GError, NULL, bSwap, false, SaveFlags, Target, FDateTime::MinValue(), false);
+							FSavePackageContext* const SavePackageContext = IsUsingPackageStore() ? SavePackageContexts[PlatformIndex] : nullptr;
+
+							Result = GEditor->Save(	Package, World, FlagsToCook, *PlatFilename, 
+													GError, nullptr, bSwap, false, SaveFlags, Target, 
+													FDateTime::MinValue(), false, /*DiffMap*/ nullptr, 
+													SavePackageContext);
 						}
 						GIsCookerLoadingPackage = false;
 						{
@@ -4081,7 +4082,7 @@ void UCookOnTheFlyServer::SaveCookedPackage(UPackage* Package, uint32 SaveFlags,
 
 void UCookOnTheFlyServer::Initialize( ECookMode::Type DesiredCookMode, ECookInitializationFlags InCookFlags, const FString &InOutputDirectoryOverride )
 {
-	TRACE_CPUPROFILER_EVENT_SCOPE_TEXT(TEXT("UCookOnTheFlyServer::Initialize"));
+	TRACE_CPUPROFILER_EVENT_SCOPE(UCookOnTheFlyServer::Initialize);
 
 	OutputDirectoryOverride = InOutputDirectoryOverride;
 	CurrentCookMode = DesiredCookMode;
@@ -5075,7 +5076,7 @@ FName UCookOnTheFlyServer::ConvertCookedPathToUncookedPath(
 
 void UCookOnTheFlyServer::GetAllCookedFiles(TMap<FName, FName>& UncookedPathToCookedPath, const FString& SandboxRootDir)
 {
-	TRACE_CPUPROFILER_EVENT_SCOPE_TEXT(TEXT("UCookOnTheFlyServer::GetAllCookedFiles"));
+	TRACE_CPUPROFILER_EVENT_SCOPE(UCookOnTheFlyServer::GetAllCookedFiles);
 
 	TArray<FString> CookedFiles;
 	{
@@ -5104,7 +5105,7 @@ void UCookOnTheFlyServer::GetAllCookedFiles(TMap<FName, FName>& UncookedPathToCo
 
 void UCookOnTheFlyServer::PopulateCookedPackagesFromDisk(const TArray<ITargetPlatform*>& Platforms)
 {
-	TRACE_CPUPROFILER_EVENT_SCOPE_TEXT(TEXT("UCookOnTheFlyServer::PopulateCookedPackagesFromDisk"));
+	TRACE_CPUPROFILER_EVENT_SCOPE(UCookOnTheFlyServer::PopulateCookedPackagesFromDisk);
 
 	// See what files are out of date in the sandbox folder
 	for (int32 Index = 0; Index < Platforms.Num(); Index++)
@@ -5368,7 +5369,7 @@ const FString ExtractPackageNameFromObjectPath( const FString ObjectPath )
 
 void UCookOnTheFlyServer::CleanSandbox(const bool bIterative)
 {
-	TRACE_CPUPROFILER_EVENT_SCOPE_TEXT(TEXT("UCookOnTheFlyServer::CleanSandbox"));
+	TRACE_CPUPROFILER_EVENT_SCOPE(UCookOnTheFlyServer::CleanSandbox);
 
 	const TArray<ITargetPlatform*>& Platforms = GetCookingTargetPlatforms();
 
@@ -5443,7 +5444,7 @@ void UCookOnTheFlyServer::CleanSandbox(const bool bIterative)
 
 void UCookOnTheFlyServer::GenerateAssetRegistry()
 {
-	TRACE_CPUPROFILER_EVENT_SCOPE_TEXT(TEXT("UCookOnTheFlyServer::GenerateAssetRegistry"));
+	TRACE_CPUPROFILER_EVENT_SCOPE(UCookOnTheFlyServer::GenerateAssetRegistry);
 
 	// Cache asset registry for later
 	FAssetRegistryModule& AssetRegistryModule = FModuleManager::LoadModuleChecked<FAssetRegistryModule>(TEXT("AssetRegistry"));
@@ -5582,7 +5583,7 @@ void UCookOnTheFlyServer::AddFileToCook( TArray<FName>& InOutFilesToCook, const
 
 void UCookOnTheFlyServer::CollectFilesToCook(TArray<FName>& FilesInPath, const TArray<FString>& CookMaps, const TArray<FString>& InCookDirectories, const TArray<FString> &IniMapSections, ECookByTheBookOptions FilesToCookFlags)
 {
-	TRACE_CPUPROFILER_EVENT_SCOPE_TEXT(TEXT("UCookOnTheFlyServer::GenerateAssetRegistry"));
+	TRACE_CPUPROFILER_EVENT_SCOPE(UCookOnTheFlyServer::GenerateAssetRegistry);
 
 #if OUTPUT_TIMING
 	SCOPE_TIMER(CollectFilesToCook);
@@ -6330,6 +6331,19 @@ void UCookOnTheFlyServer::CookByTheBookFinished()
 	UE_LOG(LogCook, Display, TEXT("Finishing up..."));
 
 	UPackage::WaitForAsyncFileWrites();
+	
+	if (IsCookByTheBookMode() && IsUsingPackageStore())
+	{
+		SCOPE_TIMER(SavingNameMap);
+		UE_LOG(LogCook, Display, TEXT("Saving name map(s)..."));
+		const TArray<ITargetPlatform*>& TargetPlatforms = GetCookingTargetPlatforms();
+		for (int32 PlatformIndex = 0; PlatformIndex < TargetPlatforms.Num(); ++PlatformIndex)
+		{
+			INameMapSaver& NameMapSaver = SavePackageContexts[PlatformIndex]->HeaderSaver.NameMapSaver;
+			NameMapSaver.End();
+		}
+		UE_LOG(LogCook, Display, TEXT("Done saving name map(s)"));
+	}
 
 	GetDerivedDataCacheRef().WaitForQuiescence(true);
 	
@@ -6702,6 +6716,37 @@ void UCookOnTheFlyServer::InitializeSandbox()
 	}
 }
 
+void UCookOnTheFlyServer::InitializePackageStore(const TArray<FName>& TargetPlatformNames)
+{
+	// TODO: should ideally support cook-on-the-fly
+
+	if (!(IsCookByTheBookMode() && IsUsingPackageStore()))
+	{
+		// set up reusable buffers for non-package store case???
+		SavePackageContexts.SetNum(TargetPlatformNames.Num());
+
+		return;
+	}
+
+	const FString NameMapFilename = FPaths::RootDir() / TEXT("megafile.unamemap");
+	const FString NameMapSandboxFilename = ConvertToFullSandboxPath(*NameMapFilename, true);
+
+	FString NameMapCookedSandboxFilename;
+	SavePackageContexts.Reserve(TargetPlatformNames.Num());
+
+	for (const FName PlatformName : TargetPlatformNames)
+	{
+		NameMapCookedSandboxFilename = NameMapSandboxFilename.Replace(TEXT("[Platform]"), *PlatformName.ToString());
+
+		// just leak all memory for now
+		FPackageStoreNameMapSaver* NameMapSaver = new FPackageStoreNameMapSaver(*NameMapCookedSandboxFilename);
+		FPackageHeaderSaver* PackageHeaderSaver = new FPackageHeaderSaver(*NameMapSaver);
+		FLooseFileWriter* LooseFileWriter		= new FLooseFileWriter();
+		FSavePackageContext* SavePackageContext = new FSavePackageContext(*PackageHeaderSaver, *LooseFileWriter);
+
+		SavePackageContexts.Add(SavePackageContext);
+	}
+}
 
 void UCookOnTheFlyServer::InitializeTargetPlatforms()
 {
@@ -6811,6 +6856,7 @@ void UCookOnTheFlyServer::StartCookByTheBook( const FCookByTheBookStartupOptions
 	CookByTheBookOptions->CreateReleaseVersion = CreateReleaseVersion;
 	CookByTheBookOptions->bDisableUnsolicitedPackages = !!(CookOptions & ECookByTheBookOptions::DisableUnsolicitedPackages);
 	CookByTheBookOptions->bFullLoadAndSave = !!(CookOptions & ECookByTheBookOptions::FullLoadAndSave);
+	CookByTheBookOptions->bPackageStore = !!(CookOptions & ECookByTheBookOptions::PackageStore);
 	CookByTheBookOptions->bErrorOnEngineContentUse = CookByTheBookStartupOptions.bErrorOnEngineContentUse;
 
 	GenerateAssetRegistry();
@@ -6955,6 +7001,8 @@ void UCookOnTheFlyServer::StartCookByTheBook( const FCookByTheBookStartupOptions
 	InitializeSandbox();
 	InitializeTargetPlatforms();
 
+	InitializePackageStore(TargetPlatformNames);
+
 	if (CurrentCookMode == ECookMode::CookByTheBook && !IsCookFlagSet(ECookInitializationFlags::Iterative))
 	{
 		StartSavingEDLCookInfoForVerification();
diff --git a/Engine/Source/Editor/UnrealEd/Private/EditorEngine.cpp b/Engine/Source/Editor/UnrealEd/Private/EditorEngine.cpp
--- a/Engine/Source/Editor/UnrealEd/Private/EditorEngine.cpp
+++ b/Engine/Source/Editor/UnrealEd/Private/EditorEngine.cpp
@@ -4398,8 +4398,9 @@ void UEditorEngine::CleanupPhysicsSceneThatWasInitializedForSave(UWorld* World,
 }
 
 FSavePackageResultStruct UEditorEngine::Save( UPackage* InOuter, UObject* InBase, EObjectFlags TopLevelFlags, const TCHAR* Filename,
-				 FOutputDevice* Error, FLinkerLoad* Conform, bool bForceByteSwapping, bool bWarnOfLongFilename, 
-				 uint32 SaveFlags, const class ITargetPlatform* TargetPlatform, const FDateTime& FinalTimeStamp, bool bSlowTask, FArchiveDiffMap* InOutDiffMap)
+				 FOutputDevice* Error, FLinkerNull* Conform, bool bForceByteSwapping, bool bWarnOfLongFilename, 
+				 uint32 SaveFlags, const class ITargetPlatform* TargetPlatform, const FDateTime& FinalTimeStamp, bool bSlowTask, FArchiveDiffMap* InOutDiffMap,
+				 FSavePackageContext* SavePackageContext)
 {
 	FScopedSlowTask SlowTask(100, FText(), bSlowTask);
 
@@ -4444,7 +4445,7 @@ FSavePackageResultStruct UEditorEngine::Save( UPackage* InOuter, UObject* InBase
 	SlowTask.EnterProgressFrame(70);
 
 	UPackage::PreSavePackageEvent.Broadcast(InOuter);
-	const FSavePackageResultStruct Result = UPackage::Save(InOuter, Base, TopLevelFlags, Filename, Error, Conform, bForceByteSwapping, bWarnOfLongFilename, SaveFlags, TargetPlatform, FinalTimeStamp, bSlowTask, InOutDiffMap);
+	const FSavePackageResultStruct Result = UPackage::Save(InOuter, Base, TopLevelFlags, Filename, Error, Conform, bForceByteSwapping, bWarnOfLongFilename, SaveFlags, TargetPlatform, FinalTimeStamp, bSlowTask, InOutDiffMap, SavePackageContext);
 
 	SlowTask.EnterProgressFrame(10);
 
@@ -4495,7 +4496,7 @@ FSavePackageResultStruct UEditorEngine::Save( UPackage* InOuter, UObject* InBase
 }
 
 bool UEditorEngine::SavePackage(UPackage* InOuter, UObject* InBase, EObjectFlags TopLevelFlags, const TCHAR* Filename,
-	FOutputDevice* Error, FLinkerLoad* Conform, bool bForceByteSwapping, bool bWarnOfLongFilename,
+	FOutputDevice* Error, FLinkerNull* Conform, bool bForceByteSwapping, bool bWarnOfLongFilename,
 	uint32 SaveFlags, const class ITargetPlatform* TargetPlatform, const FDateTime& FinalTimeStamp, bool bSlowTask)
 {
 	// Workaround to avoid function signature change while keeping both bool and ESavePackageResult versions of SavePackage
diff --git a/Engine/Source/Programs/BenchmarkTool/BenchmarkTool.Build.cs b/Engine/Source/Programs/BenchmarkTool/BenchmarkTool.Build.cs
new file mode 100644
--- /dev/null
+++ b/Engine/Source/Programs/BenchmarkTool/BenchmarkTool.Build.cs
@@ -0,0 +1,16 @@
+// Copyright 1998-2019 Epic Games, Inc. All Rights Reserved.
+
+using UnrealBuildTool;
+
+public class BenchmarkTool : ModuleRules
+{
+	public BenchmarkTool(ReadOnlyTargetRules Target) : base(Target)
+	{
+		PublicIncludePaths.Add("Runtime/Launch/Public");
+
+		PrivateIncludePaths.Add("Runtime/Launch/Private");		// For LaunchEngineLoop.cpp include
+
+		PrivateDependencyModuleNames.Add("Core");
+		PrivateDependencyModuleNames.Add("Projects");
+	}
+}
diff --git a/Engine/Source/Programs/BenchmarkTool/Private/BenchmarkTool.h b/Engine/Source/Programs/BenchmarkTool/Private/BenchmarkTool.h
new file mode 100644
--- /dev/null
+++ b/Engine/Source/Programs/BenchmarkTool/Private/BenchmarkTool.h
@@ -0,0 +1,6 @@
+// Copyright 1998-2019 Epic Games, Inc. All Rights Reserved.
+
+#pragma once
+
+#include "CoreMinimal.h"
+
diff --git a/Engine/Source/Runtime/Core/Private/IO/IoDispatcher.cpp b/Engine/Source/Runtime/Core/Private/IO/IoDispatcher.cpp
new file mode 100644
--- /dev/null
+++ b/Engine/Source/Runtime/Core/Private/IO/IoDispatcher.cpp
@@ -0,0 +1,1156 @@
+// Copyright 1998-2019 Epic Games, Inc. All Rights Reserved.
+
+#include "IO/IoDispatcher.h"
+#include "HAL/FileManager.h"
+#include "Async/MappedFileHandle.h"
+#include "Containers/ChunkedArray.h"
+#include "HAL/CriticalSection.h"
+#include "Logging/LogMacros.h"
+#include "Misc/CString.h"
+#include "Misc/EventPool.h"
+#include "Misc/ScopeRWLock.h"
+#include "Misc/SecureHash.h"
+#include "Misc/StringBuilder.h"
+#include "Memory/MemoryArena.h"
+#include "Misc/LazySingleton.h"
+
+DEFINE_LOG_CATEGORY_STATIC(LogIoDispatch, Log, All);
+
+const TCHAR* GetIoErrorText(EIoErrorCode ErrorCode)
+{
+	static constexpr const TCHAR* ErrorCodeText[]
+	{
+		TEXT("OK"),
+		TEXT("Unknown Status"),
+		TEXT("Invalid Code"),
+		TEXT("Cancelled"),
+		TEXT("FileOpen Failed"),
+		TEXT("File Not Open"),
+		TEXT("Write Error"),
+		TEXT("Not Found"),
+		TEXT("Corrupt Toc"),
+	};
+
+	return ErrorCodeText[static_cast<uint32>(ErrorCode)];
+}
+
+const FIoStatus FIoStatus::Ok		{ EIoErrorCode::Ok,				TEXT("OK")				};
+const FIoStatus FIoStatus::Unknown	{ EIoErrorCode::Unknown,		TEXT("Unknown Status")	};
+const FIoStatus FIoStatus::Invalid	{ EIoErrorCode::InvalidCode,	TEXT("Invalid Code")	};
+
+FIoStatus::FIoStatus(EIoErrorCode Code)
+:	ErrorCode(Code)
+,	ErrorMessage(GetIoErrorText(Code))
+{
+}
+
+FIoStatus::FIoStatus(EIoErrorCode Code, const FStringView& InErrorMessage)
+: ErrorCode(Code)
+, ErrorMessage(InErrorMessage.Data(), InErrorMessage.Len())
+{
+}
+
+FIoStatus& 
+FIoStatus::operator=(const FIoStatus& Other)
+{
+	ErrorCode = Other.ErrorCode;
+	ErrorMessage = Other.ErrorMessage;
+
+	return *this;
+}
+
+bool		
+FIoStatus::operator==(const FIoStatus& Other) const
+{
+	return ErrorCode == Other.ErrorCode
+		&& ErrorMessage == Other.ErrorMessage;
+}
+
+FString	
+FIoStatus::ToString() const
+{
+	return ErrorMessage;
+}
+
+//////////////////////////////////////////////////////////////////////////
+
+FIoStatusBuilder::FIoStatusBuilder(EIoErrorCode InStatusCode)
+:	StatusCode(InStatusCode)
+{
+}
+
+FIoStatusBuilder::FIoStatusBuilder(const FIoStatus& InStatus, FStringView String)
+:	StatusCode(InStatus.ErrorCode)
+{
+	Message.Append(String.Data(), String.Len());
+}
+
+FIoStatusBuilder::~FIoStatusBuilder()
+{
+}
+
+FIoStatusBuilder::operator FIoStatus()
+{
+	return FIoStatus(StatusCode, Message);
+}
+
+FIoStatusBuilder& 
+FIoStatusBuilder::operator<<(FStringView String)
+{
+	Message.Append(String.Data(), String.Len());
+
+	return *this;
+}
+
+FIoStatusBuilder operator<<(const FIoStatus& Status, FStringView String)
+{ 
+	return FIoStatusBuilder(Status, String);
+}
+
+//////////////////////////////////////////////////////////////////////////
+
+void StatusOrCrash(const FIoStatus& Status)
+{
+	// TODO
+}
+
+//////////////////////////////////////////////////////////////////////////
+
+void FIoChunkId::GenerateFromData(const void* InData, SIZE_T InDataSize)
+{
+	FSHA1 Sha1;
+	Sha1.Update(reinterpret_cast<const uint8*>(InData), InDataSize);
+	Sha1.Final();
+	
+	uint8 Sha1Hash[20];
+	Sha1.GetHash(Sha1Hash);
+
+	memcpy(Id, Sha1Hash, sizeof Id);
+}
+
+//////////////////////////////////////////////////////////////////////////
+
+class FIoBufferManager
+{
+public:
+	TChunkedArray<FIoBuffer::BufCore, 65536> BufCores;
+
+	void AllocBufCore()
+	{
+		new(BufCores) FIoBuffer::BufCore;
+	}
+};
+
+
+//////////////////////////////////////////////////////////////////////////
+
+class FIoBufferPool
+{
+public:
+	FIoBufferPool();
+	~FIoBufferPool();
+
+	TIoStatusOr<FIoBuffer>	Alloc(uint64 RequiredCapacity);
+	void					Free(FIoBuffer InBuffer);
+};
+
+FIoBufferPool::FIoBufferPool()
+{
+}
+
+FIoBufferPool::~FIoBufferPool()
+{
+}
+
+TIoStatusOr<FIoBuffer> FIoBufferPool::Alloc(uint64 RequiredCapacity)
+{
+	return FIoBuffer(RequiredCapacity);
+}
+
+void FIoBufferPool::Free(FIoBuffer InBuffer)
+{
+	auto Destroyer = MoveTemp(InBuffer);
+}
+
+//////////////////////////////////////////////////////////////////////////
+
+FIoBuffer::BufCore::BufCore()
+{
+}
+
+FIoBuffer::BufCore::~BufCore()
+{
+	if (IsMemoryOwned())
+	{
+		FMemory::Free(Data());
+	}
+}
+
+FIoBuffer::BufCore::BufCore(const uint8* InData, uint64 InSize, bool InOwnsMemory)
+{
+	SetDataAndSize(InData, InSize);
+
+	SetIsOwned(InOwnsMemory);
+}
+
+FIoBuffer::BufCore::BufCore(const uint8* InData, uint64 InSize, const BufCore* InOuter)
+:	OuterCore(InOuter)
+{
+	SetDataAndSize(InData, InSize);
+}
+
+FIoBuffer::BufCore::BufCore(uint64 InSize)
+{
+	uint8* NewBuffer = reinterpret_cast<uint8*>(FMemory::Malloc(InSize));
+
+	SetDataAndSize(NewBuffer, InSize);
+
+	SetIsOwned(true);
+}
+
+FIoBuffer::BufCore::BufCore(ECloneTag, uint8* InData, uint64 InSize)
+:	FIoBuffer::BufCore(InSize)
+{
+	FMemory::Memcpy(Data(), InData, InSize);
+}
+
+void
+FIoBuffer::BufCore::CheckRefCount() const
+{
+	// Verify that Release() is not being called on an object which is already at a zero refcount
+	check(NumRefs != 0);
+}
+
+void
+FIoBuffer::BufCore::SetDataAndSize(const uint8* InData, uint64 InSize)
+{
+	// This is intentionally not split into SetData and SetSize to enable different storage
+	// strategies for flags in the future (in unused pointer bits)
+
+	DataPtr			= const_cast<uint8*>(InData);
+	DataSizeLow		= uint32(InSize & 0xffffffffu);
+	DataSizeHigh	= (InSize >> 32) & 0xffu;
+}
+
+void
+FIoBuffer::BufCore::SetSize(uint64 InSize)
+{
+	SetDataAndSize(Data(), InSize);
+}
+
+void
+FIoBuffer::BufCore::MakeOwned()
+{
+	if (IsMemoryOwned())
+		return;
+
+	const uint64 BufferSize = DataSize();
+	uint8* NewBuffer		= reinterpret_cast<uint8*>(FMemory::Malloc(BufferSize));
+
+	FMemory::Memcpy(NewBuffer, Data(), BufferSize);
+
+	SetDataAndSize(NewBuffer, BufferSize);
+
+	SetIsOwned(true);
+}
+
+//////////////////////////////////////////////////////////////////////////
+
+FIoBuffer::FIoBuffer()
+:	CorePtr(new BufCore)
+{
+}
+
+FIoBuffer::FIoBuffer(uint64 InSize)
+:	CorePtr(new BufCore(InSize))
+{
+}
+
+FIoBuffer::FIoBuffer(const void* Data, uint64 InSize, const FIoBuffer& OuterBuffer)
+:	CorePtr(new BufCore((uint8*) Data, InSize, OuterBuffer.CorePtr))
+{
+}
+
+FIoBuffer::FIoBuffer(FIoBuffer::EWrapTag, const void* Data, uint64 InSize)
+:	CorePtr(new BufCore((uint8*)Data, InSize, /* ownership */ false))
+{
+}
+
+FIoBuffer::FIoBuffer(FIoBuffer::EAssumeOwnershipTag, const void* Data, uint64 InSize)
+:	CorePtr(new BufCore((uint8*)Data, InSize, /* ownership */ true))
+{
+}
+
+FIoBuffer::FIoBuffer(FIoBuffer::ECloneTag, const void* Data, uint64 InSize)
+:	CorePtr(new BufCore(Clone, (uint8*)Data, InSize))
+{
+}
+
+void		
+FIoBuffer::MakeOwned() const
+{
+	CorePtr->MakeOwned();
+}
+
+//////////////////////////////////////////////////////////////////////////
+
+static const char TocMagicImg[] = "-==--==--==--==-";
+
+struct FIoStoreTocHeader
+{
+	uint8	TocMagic[16];
+	uint32	TocHeaderSize;
+	uint32	TocEntryCount;
+	uint32	TocEntrySize;	// For sanity checking
+	uint32	TocPad[25];
+
+	void MakeMagic()
+	{
+		FMemory::Memcpy(TocMagic, TocMagicImg, sizeof TocMagic);
+	}
+
+	bool CheckMagic()
+	{
+		return FMemory::Memcmp(TocMagic, TocMagicImg, sizeof TocMagic) == 0;
+	}
+};
+
+struct FIoStoreTocEntry
+{
+	// We use 5 bytes for offset and size, this is enough to represent 
+	// an offset and size of 1PB
+	uint8		OffsetAndLength[5 + 5];
+
+	// TBD: should the chunk ID use content addressing, or names, or a
+	//      mix of both?
+	FIoChunkId	ChunkId;
+
+	inline uint64 GetOffset() const
+	{
+		return OffsetAndLength[4]
+			| (uint64(OffsetAndLength[3]) << 8)
+			| (uint64(OffsetAndLength[2]) << 16)
+			| (uint64(OffsetAndLength[1]) << 24)
+			| (uint64(OffsetAndLength[0]) << 32)
+			;
+	}
+
+	inline uint64 GetLength() const
+	{
+		return OffsetAndLength[9]
+			| (uint64(OffsetAndLength[8]) << 8)
+			| (uint64(OffsetAndLength[7]) << 16)
+			| (uint64(OffsetAndLength[6]) << 24)
+			| (uint64(OffsetAndLength[5]) << 32)
+			;
+	}
+
+	inline void SetOffset(uint64 Offset)
+	{
+		OffsetAndLength[0] = uint8(Offset >> 32);
+		OffsetAndLength[1] = uint8(Offset >> 24);
+		OffsetAndLength[2] = uint8(Offset >> 16);
+		OffsetAndLength[3] = uint8(Offset >>  8);
+		OffsetAndLength[4] = uint8(Offset >>  0);
+	}
+
+	inline void SetLength(uint64 Length)
+	{
+		OffsetAndLength[5] = uint8(Length >> 32);
+		OffsetAndLength[6] = uint8(Length >> 24);
+		OffsetAndLength[7] = uint8(Length >> 16);
+		OffsetAndLength[8] = uint8(Length >> 8);
+		OffsetAndLength[9] = uint8(Length >> 0);
+	}
+};
+
+//////////////////////////////////////////////////////////////////////////
+
+FIoStoreEnvironment::FIoStoreEnvironment()
+{
+}
+
+FIoStoreEnvironment::~FIoStoreEnvironment()
+{
+}
+
+void FIoStoreEnvironment::InitializeFileEnvironment(FStringView InRootPath)
+{
+	RootPath = InRootPath.ToString();
+}
+
+//////////////////////////////////////////////////////////////////////////
+
+/** Very simple string class
+
+	The backing memory is 
+  */
+class FCompactString
+{
+	typedef TCHAR CharT;
+
+public:
+						FCompactString();
+	explicit			FCompactString(FStringView InString);
+						FCompactString(const FCompactString& InString);
+						FCompactString(FCompactString&& InString);
+						~FCompactString();
+
+	inline const CharT* operator*() const						{ return Chars; }
+
+	inline FCompactString& operator=(FStringView Rhs)			{ Assign(Rhs.Data(), Rhs.Len()); return *this; }
+	inline FCompactString& operator=(const FCompactString& Rhs) { Assign(*Rhs); return *this; }
+	inline FCompactString& operator=(FCompactString&& Rhs)		{ /* swap guts */ auto Temp = Chars; Chars = Rhs.Chars; Rhs.Chars = Temp; return *this; }
+
+private:
+	void Initialize(const CharT* InChars, SIZE_T InCharCount);
+	void Initialize(const CharT* InChars);
+	void Assign(const CharT* InChars, SIZE_T InCharCount);
+	void Assign(const CharT* InChars);
+
+	TArenaPointer<CharT>	Chars;
+};
+
+FCompactString::FCompactString()
+{
+}
+
+FCompactString::~FCompactString()
+{
+}
+
+FCompactString::FCompactString(FCompactString&& InString)
+{
+	Chars = InString.Chars;
+	InString.Chars = nullptr;
+}
+
+FCompactString::FCompactString(const FCompactString& InString)
+{
+	Initialize(*InString);
+}
+
+FCompactString::FCompactString(FStringView InString)
+{
+	Initialize(InString.Data(), InString.Len());
+}
+
+void 
+FCompactString::Initialize(const CharT* InChars, SIZE_T InCharCount)
+{
+	// The count does not include the NUL terminator
+	++InCharCount;
+
+	Chars = reinterpret_cast<CharT*>(FMemory::Malloc(InCharCount * sizeof(CharT)));
+
+	FMemory::Memcpy(Chars, InChars, InCharCount * sizeof(CharT));
+}
+
+void
+FCompactString::Initialize(const CharT* InChars)
+{
+	Initialize(InChars, TCString<CharT>::Strlen(InChars));
+}
+
+void
+FCompactString::Assign(const CharT* InChars, SIZE_T InCharCount)
+{
+	// The count does not include the NUL terminator
+	++InCharCount;
+
+	if (Chars)
+	{
+		FMemory::Free(Chars);
+	}
+
+	Chars = reinterpret_cast<CharT*>(FMemory::Malloc(InCharCount * sizeof(CharT)));
+
+	FMemory::Memcpy(Chars, InChars, InCharCount * sizeof(CharT));
+}
+
+void
+FCompactString::Assign(const CharT* InChars)
+{
+	Assign(InChars, TCString<CharT>::Strlen(InChars));
+}
+
+//////////////////////////////////////////////////////////////////////////
+
+class FIoStoreReaderImpl
+{
+public:
+	FIoStoreReaderImpl(FIoStoreEnvironment& InEnvironment)
+		: Environment(InEnvironment)
+	{
+	}
+
+	FIoStatus Open(FStringView UniqueId);
+
+	TIoStatusOr<FIoBuffer> Lookup(const FIoChunkId& ChunkId)
+	{
+		const FIoStoreTocEntry* Entry = Toc.Find(ChunkId);
+
+		if (!Entry)
+		{
+			return FIoStatus(EIoErrorCode::NotFound);
+		}
+
+		return FIoBuffer(FIoBuffer::Wrap, MappedRegion->GetMappedPtr() + Entry->GetOffset(), Entry->GetLength());
+	}
+
+private:
+	FIoStoreEnvironment&				Environment;
+	FCompactString						UniqueId;
+	TMap<FIoChunkId, FIoStoreTocEntry>	Toc;
+	TUniquePtr<IFileHandle>				ContainerFileHandle;
+	TUniquePtr<IMappedFileHandle>		ContainerMappedFileHandle;
+	TUniquePtr<IMappedFileRegion>		MappedRegion;
+};
+
+FIoStatus FIoStoreReaderImpl::Open(FStringView InUniqueId)
+{
+	IPlatformFile& Ipf = IPlatformFile::GetPlatformPhysical();
+
+	UniqueId = InUniqueId;
+
+	const FString& RootPath = Environment.GetRootPath();
+
+	TStringBuilder<256> ContainerFilePath;
+	ContainerFilePath.Append(RootPath);
+	if (ContainerFilePath.LastChar() != '/')
+		ContainerFilePath.Append(TEXT('/'));
+
+	TStringBuilder<256> TocFilePath;
+	TocFilePath.Append(ContainerFilePath);
+	TocFilePath.Append(TEXT("Container.utoc"));
+
+	TUniquePtr<IFileHandle>	TocFileHandle;
+	TocFileHandle.Reset(Ipf.OpenRead(*TocFilePath, /* allowwrite */ false));
+
+	if (!TocFileHandle)
+	{
+		return FIoStatusBuilder(EIoErrorCode::FileOpenFailed) << TEXT("Failed to open IoStore TOC file '") << *TocFilePath << TEXT("'");
+	}
+
+	ContainerFilePath.Append(TEXT("Container.ucas"));
+	ContainerFileHandle.Reset(Ipf.OpenRead(*ContainerFilePath, /* allowwrite */ false));
+	
+	if (!ContainerFileHandle)
+	{
+		return FIoStatusBuilder(EIoErrorCode::FileOpenFailed) << TEXT("Failed to open IoStore container file '") << *ContainerFilePath << TEXT("'");
+	}
+
+	const uint64 ContainerSize = ContainerFileHandle->Size();
+
+	ContainerMappedFileHandle.Reset(Ipf.OpenMapped(*ContainerFilePath));
+	MappedRegion.Reset(ContainerMappedFileHandle->MapRegion());
+
+	if (!ContainerMappedFileHandle)
+	{
+		return FIoStatusBuilder(EIoErrorCode::FileOpenFailed) << TEXT("Failed to memory map IoStore container file '") << *ContainerFilePath << TEXT("'");
+	}
+
+	// Parse TOC
+	//
+	// This should ultimately be a read-in-place operation but it looks like this for now
+
+	FIoStoreTocHeader Header;
+	TocFileHandle->Read(reinterpret_cast<uint8*>(&Header), sizeof Header);
+
+	if (!Header.CheckMagic())
+	{
+		return FIoStatusBuilder(EIoErrorCode::CorruptToc) << TEXT("TOC header magic mismatch while reading '") << *TocFilePath << TEXT("'");
+	}
+
+	if (Header.TocHeaderSize != sizeof Header)
+	{
+		return FIoStatusBuilder(EIoErrorCode::CorruptToc) << TEXT("TOC header size mismatch while reading '") << *TocFilePath << TEXT("'");
+	}
+
+	if (Header.TocEntrySize != sizeof(FIoStoreTocEntry))
+	{
+		return FIoStatusBuilder(EIoErrorCode::CorruptToc) << TEXT("TOC entry size mismatch while reading '") << *TocFilePath << TEXT("'");
+	}
+
+	uint32 EntryCount = Header.TocEntryCount;
+
+	while(EntryCount--)
+	{
+		FIoStoreTocEntry Entry;
+		bool Success = TocFileHandle->Read(reinterpret_cast<uint8*>(&Entry), sizeof Entry);
+
+		if (!Success)
+		{
+			return FIoStatusBuilder(EIoErrorCode::CorruptToc) << TEXT("failed to read TOC entry while reading '") << *TocFilePath << TEXT("'");
+		}
+
+		if ((Entry.GetOffset() + Entry.GetLength()) > ContainerSize)
+		{
+			// TODO: add details
+			return FIoStatusBuilder(EIoErrorCode::CorruptToc) << TEXT("TOC entry out of container bounds while reading '") << *TocFilePath << TEXT("'");
+		}
+
+		Toc.Add(Entry.ChunkId, Entry);
+	}
+
+	if (TocFileHandle->Tell() != TocFileHandle->Size())
+	{
+		return FIoStatusBuilder(EIoErrorCode::CorruptToc) << TEXT("TOC file contains trailing garbage '") << *TocFilePath << TEXT("'");
+	}
+
+	return FIoStatus::Ok;
+}
+
+FIoStoreReader::FIoStoreReader(FIoStoreEnvironment& Environment)
+:	Impl(new FIoStoreReaderImpl(Environment))
+{
+}
+
+FIoStoreReader::~FIoStoreReader()
+{
+	delete Impl;
+}
+
+FIoStatus 
+FIoStoreReader::Initialize(FStringView UniqueId)
+{
+	return Impl->Open(UniqueId);
+}
+
+//////////////////////////////////////////////////////////////////////////
+
+class FIoStoreWriterImpl
+{
+public:
+	FIoStoreWriterImpl(FIoStoreEnvironment& InEnvironment)
+	:	Environment(InEnvironment)
+	{
+	}
+
+	FIoStatus Open()
+	{
+		IPlatformFile& Ipf = IPlatformFile::GetPlatformPhysical();
+
+		const FString& RootPath = Environment.GetRootPath();
+		Ipf.CreateDirectoryTree(*RootPath);
+
+		TStringBuilder<256> ContainerFilePath;
+		ContainerFilePath.Append(RootPath);
+		if (ContainerFilePath.LastChar() != '/')
+			ContainerFilePath.Append(TEXT('/'));
+
+		TStringBuilder<256> TocFilePath;
+		TocFilePath.Append(ContainerFilePath);
+
+		ContainerFilePath.Append(TEXT("Container.ucas"));
+		TocFilePath.Append(TEXT("Container.utoc"));
+
+		ContainerFileHandle.Reset(Ipf.OpenWrite(*ContainerFilePath, /* append */ false, /* allowread */ true));
+
+		if (!ContainerFileHandle)
+		{ 
+			return FIoStatusBuilder(EIoErrorCode::FileOpenFailed) << TEXT("Failed to open IoStore container file '") << *ContainerFilePath << TEXT("'");
+		}
+
+		TocFileHandle.Reset(Ipf.OpenWrite(*TocFilePath, /* append */ false, /* allowread */ true));
+
+		if (!TocFileHandle)
+		{
+			return FIoStatusBuilder(EIoErrorCode::FileOpenFailed) << TEXT("Failed to open IoStore TOC file '") << *TocFilePath << TEXT("'");
+		}
+
+		return FIoStatus::Ok;
+	}
+
+	UE_NODISCARD FIoStatus Append(FIoChunkId ChunkId, FIoBuffer Chunk)
+	{
+		if (!ContainerFileHandle)
+		{
+			return FIoStatus(EIoErrorCode::FileNotOpen, TEXT("No container file to append to"));
+		}
+
+		FIoStoreTocEntry TocEntry;
+
+		TocEntry.SetOffset(ContainerFileHandle->Tell());
+		TocEntry.SetLength(Chunk.DataSize());
+		TocEntry.ChunkId = ChunkId;
+
+		IsMetadataDirty = true;
+
+		const bool Success = ContainerFileHandle->Write(Chunk.Data(), Chunk.DataSize());
+
+		if (Success)
+		{
+			Toc.Add(ChunkId, TocEntry);
+
+			return FIoStatus::Ok;
+		}
+		else
+		{
+			return FIoStatus(EIoErrorCode::WriteError, TEXT("Append failed"));
+		}
+	}
+
+	UE_NODISCARD FIoStatus FlushMetadata()
+	{
+		TocFileHandle->Seek(0);
+
+		FIoStoreTocHeader TocHeader;
+		FMemory::Memset(TocHeader, 0);
+
+		TocHeader.MakeMagic();
+		TocHeader.TocHeaderSize = sizeof TocHeader;
+		TocHeader.TocEntryCount = Toc.Num();
+		TocHeader.TocEntrySize = sizeof(FIoStoreTocEntry);
+
+		const bool Success = TocFileHandle->Write(reinterpret_cast<const uint8*>(&TocHeader), sizeof TocHeader);
+
+		if (!Success)
+		{
+			return FIoStatus(EIoErrorCode::WriteError, TEXT("TOC write failed"));
+		}
+
+		for (auto& _: Toc)
+		{
+			FIoStoreTocEntry& TocEntry = _.Value;
+			
+			TocFileHandle->Write(reinterpret_cast<const uint8*>(&TocEntry), sizeof TocEntry);
+		}
+
+		return FIoStatus::Ok;
+	}
+
+private:
+	FIoStoreEnvironment&				Environment;
+	TMap<FIoChunkId, FIoStoreTocEntry>	Toc;
+	TUniquePtr<IFileHandle>				ContainerFileHandle;
+	TUniquePtr<IFileHandle>				TocFileHandle;
+	bool								IsMetadataDirty = true;
+};
+
+FIoStoreWriter::FIoStoreWriter(FIoStoreEnvironment& InEnvironment)
+:	Impl(new FIoStoreWriterImpl(InEnvironment))
+{
+}
+
+FIoStoreWriter::~FIoStoreWriter()
+{
+	Impl->FlushMetadata();
+
+	delete Impl;
+}
+
+FIoStatus FIoStoreWriter::Initialize()
+{
+	return Impl->Open();
+}
+
+void FIoStoreWriter::Append(FIoChunkId ChunkId, FIoBuffer Chunk)
+{
+	Impl->Append(ChunkId, Chunk);
+}
+
+void FIoStoreWriter::FlushMetadata()
+{
+	Impl->FlushMetadata();
+}
+
+//////////////////////////////////////////////////////////////////////////
+
+class FIoBatchImpl
+{
+public:
+	uint32			FirstRequest = ~uint32(0);
+	TAtomic<uint32>	OutstandingRequests { 0 };
+
+	void Reset() 
+	{
+		// TODO: need to release IO requests back to pool
+
+		FirstRequest		= ~uint32(0);
+		OutstandingRequests = 0;
+	}
+};
+
+class FIoRequestImpl
+{
+public:
+	FIoReadOptions			Options;
+	TIoStatusOr<FIoBuffer>	Result;
+	uint32					NextRequestId = ~uint32(0);
+	FIoChunkId				ChunkId;
+};
+
+//////////////////////////////////////////////////////////////////////////
+
+class FIoStoreImpl : public FRefCountBase
+{
+public:
+	void Mount(FIoStoreReader* IoStore)
+	{
+		FWriteScopeLock _(RwLockIoStores);
+		// TODO prevent duplicates?
+		IoStores.Add(IoStore);
+	}
+
+	void Unmount(FIoStoreReader* IoStore)
+	{
+		FWriteScopeLock _(RwLockIoStores);
+		IoStores.Remove(IoStore);
+	}
+
+	TIoStatusOr<FIoBuffer> Resolve(FIoChunkId& ChunkId)
+	{
+		for (const auto& IoStore : IoStores)
+		{
+			TIoStatusOr<FIoBuffer> Result = IoStore->Impl->Lookup(ChunkId);
+
+			if (Result.IsOk())
+			{
+				return Result;
+			}
+		}
+
+		return FIoStatus(EIoErrorCode::NotFound);
+	}
+
+private:
+	FRWLock									RwLockIoStores;
+	TArray<TRefCountPtr<FIoStoreReader>>	IoStores;
+};
+
+//////////////////////////////////////////////////////////////////////////
+
+class FIoWorker : FRefCountBase
+{
+public:
+					FIoWorker(FIoStoreImpl* InIoStore);
+	virtual			~FIoWorker();
+
+	virtual void	IssueBatch(FIoBatchImpl& Batch) = 0;
+
+private:
+	TRefCountPtr<FIoStoreImpl> IoStore;
+};
+
+FIoWorker::FIoWorker(FIoStoreImpl* InIoStore)
+:	IoStore(InIoStore)
+{
+}
+
+FIoWorker::~FIoWorker()
+{
+}
+
+class FIoWorkerFile : FIoWorker
+{
+public:
+	FIoWorkerFile(FIoStoreImpl* InIoStore)
+	:	FIoWorker(InIoStore)
+	{
+	}
+
+	~FIoWorkerFile()
+	{
+	}
+
+	virtual void IssueBatch(FIoBatchImpl& Batch) override
+	{
+	}
+};
+
+//////////////////////////////////////////////////////////////////////////
+
+class FIoDispatcherImpl
+{
+	FIoDispatcher* Owner;
+
+	// Batch allocation - could use something more lightweight but I'm not sure this
+	// is a sufficiently high frequency operation to warrant something lock-free
+
+	FRWLock	RwLockBatchId;
+	uint32	FirstFreeBatchId = 0;
+
+	// Request allocation - could use something more lightweight but I'm not sure this
+	// is a sufficiently high frequency operation to warrant something lock-free
+
+	FRWLock	RwLockRequestId;
+	uint32	FirstFreeRequestId = ~0u;
+
+public:
+	FIoDispatcherImpl(FIoDispatcher* InOwner)
+	:	Owner(InOwner)
+	{
+		for (int i = 0; i < (MaxBatches - 1); ++i)
+		{
+			Batches[i].FirstRequest = i + 1;
+		}
+
+		Batches[MaxBatches - 1].FirstRequest = ~uint32(0);
+
+		for (int i = 0; i < InitialIoRequests; ++i)
+		{
+			ReserveNewRequest();
+		}
+
+		Batches[MaxBatches - 1].FirstRequest = ~uint32(0);
+
+		IoStore = new FIoStoreImpl;
+	}
+
+	static const int	MaxBatches = 65536;
+	FIoBatchImpl		Batches[MaxBatches];
+	int32				NumBatches = 0;
+
+	static const int	MaxIoRequests = 65536;
+	static const int	InitialIoRequests = 1;
+
+	TChunkedArray<FIoRequestImpl, MaxIoRequests> Requests;
+
+	FIoBatchImpl& MapBatch(uint32 BatchId)
+	{
+		return Batches[BatchId];
+	}
+
+	void ReserveNewRequest()
+	{
+		// NOTE: RwLockRequestId must be held while this is called
+
+		FIoRequestImpl& NewRequest = *new(Requests) FIoRequestImpl;
+		NewRequest.NextRequestId = FirstFreeRequestId;
+		FirstFreeRequestId = Requests.Num() - 1;
+	}
+
+	FIoRequestImpl& AllocRequest(const FIoChunkId& Chunk, FIoReadOptions Options, uint32 BatchId)
+	{
+		const uint32 RequestId = AllocRequestId();
+
+		FIoRequestImpl& Request = Requests[RequestId];
+		FIoBatchImpl& Batch		= MapBatch(BatchId);
+
+		Request.ChunkId			= Chunk;
+		Request.Options			= Options;
+		Request.NextRequestId	= Batch.FirstRequest;
+		Request.Result			= FIoStatus::Unknown;
+
+		Batch.FirstRequest		= RequestId;
+
+		return Request;
+	}
+
+	uint32 AllocRequestId()
+	{
+		FWriteScopeLock _(RwLockRequestId);
+
+		if (FirstFreeRequestId == ~uint32(0))
+		{
+			ReserveNewRequest();
+		}
+
+		const uint32 RequestId	= FirstFreeRequestId;
+
+		FIoRequestImpl& Request	= Requests[RequestId];
+		FirstFreeRequestId		= Request.NextRequestId;
+		Request.NextRequestId	= ~uint32(0);
+
+		return RequestId;
+	}
+
+	uint32 FreeRequestId(const uint32 RequestId)
+	{
+		FWriteScopeLock _(RwLockRequestId);
+
+		uint32 NextRequestId				= Requests[RequestId].NextRequestId;
+		Requests[RequestId].NextRequestId	= FirstFreeRequestId;
+		FirstFreeRequestId					= RequestId;
+
+		return NextRequestId;
+	}
+
+	uint32 AllocBatchId()
+	{
+		check(NumBatches < MaxBatches);
+
+		FWriteScopeLock _(RwLockBatchId);
+
+		const uint32 BatchId	= FirstFreeBatchId;
+		FIoBatchImpl& Batch		= Batches[BatchId];
+		FirstFreeBatchId		= Batch.FirstRequest;
+		Batch.FirstRequest		= ~uint32(0);
+
+		NumBatches++;
+
+		return BatchId;
+	}
+
+	void FreeBatchId(const uint32 BatchId)
+	{
+		FIoBatchImpl& Batch = Batches[BatchId];
+
+		uint32 RequestId = Batch.FirstRequest; 
+		while (RequestId != ~uint32(0))
+		{
+			RequestId = FreeRequestId(RequestId);
+		}
+	
+		FWriteScopeLock _(RwLockBatchId);
+
+		Batches[BatchId].FirstRequest	= FirstFreeBatchId;
+		FirstFreeBatchId				= BatchId;
+
+		NumBatches--;
+		check(NumBatches >= 0);
+	}
+
+	template<typename Func>
+	void IterateBatch(const uint32 BatchId, Func&& InCallbackFunction)
+	{
+		FIoBatchImpl& Batch = Batches[BatchId];
+
+		uint32 RequestId = Batch.FirstRequest;
+
+		while (RequestId != ~uint32(0))
+		{
+			FIoRequestImpl& Request = Requests[RequestId];
+
+			const bool DoContinue = InCallbackFunction(Request);
+
+			if (!DoContinue)
+			{
+				return;
+			}
+
+			RequestId = Request.NextRequestId;
+		}
+	}
+
+	void IssueBatch(const uint32 BatchId)
+	{
+		// At this point the batch is immutable and we should start
+		// doing the work.
+
+		IterateBatch(BatchId, [this](FIoRequestImpl& Request) {
+			TIoStatusOr<FIoBuffer> Resolved = IoStore->Resolve(Request.ChunkId);
+
+			Request.Result = Resolved;
+
+			return true;
+		});
+	}
+
+	void Mount(FIoStoreReader* IoStoreReader)
+	{
+		IoStore->Mount(IoStoreReader);
+	}
+
+	void Unmount(FIoStoreReader* IoStoreReader)
+	{
+		IoStore->Unmount(IoStoreReader);
+	}
+
+private:
+	TRefCountPtr<FIoStoreImpl>	IoStore;
+};
+
+//////////////////////////////////////////////////////////////////////////
+
+FIoDispatcher::FIoDispatcher()
+:	Impl(new FIoDispatcherImpl(this))
+{
+}
+
+FIoDispatcher::~FIoDispatcher()
+{
+}
+
+void		
+FIoDispatcher::Mount(FIoStoreReader* IoStore)
+{
+	Impl->Mount(IoStore);
+}
+
+void		
+FIoDispatcher::Unmount(FIoStoreReader* IoStore)
+{
+	Impl->Unmount(IoStore);
+}
+
+FIoBatch
+FIoDispatcher::NewBatch()
+{
+	return FIoBatch(Impl, Impl->AllocBatchId());
+}
+
+void
+FIoDispatcher::FreeBatch(FIoBatch Batch)
+{
+	Impl->FreeBatchId(Batch.BatchId);
+}
+
+//////////////////////////////////////////////////////////////////////////
+
+FIoBatch::FIoBatch(FIoDispatcherImpl* OwningIoDispatcher, uint32 InBatchId)
+:	Dispatcher(OwningIoDispatcher)
+,	BatchId(InBatchId)
+,	CompletionEvent(TLazySingleton<FEventPool<EEventPoolTypes::ManualReset>>::Get().GetEventFromPool())
+{
+}
+
+FIoRequest
+FIoBatch::Read(const FIoChunkId& Chunk, FIoReadOptions Options)
+{
+	return FIoRequest(Dispatcher->AllocRequest(Chunk, Options, BatchId));
+}
+
+void 
+FIoBatch::ForEachRequest(TFunction<bool(FIoRequest&)> Callback)
+{
+	Dispatcher->IterateBatch(BatchId, [&](FIoRequestImpl& InRequest) {
+		FIoRequest Request(InRequest);
+		return Callback(Request);
+	});
+}
+
+void 
+FIoBatch::Issue()
+{
+	Dispatcher->IssueBatch(BatchId);
+}
+
+void 
+FIoBatch::Wait()
+{
+	unimplemented();
+}
+
+void 
+FIoBatch::Cancel()
+{
+	unimplemented();
+}
+
+//////////////////////////////////////////////////////////////////////////
+
+bool		
+FIoRequest::IsOk() const 
+{ 
+	return Impl->Result.IsOk(); 
+}
+
+FIoStatus	
+FIoRequest::Status() const
+{ 
+	return Impl->Result.Status();
+}
+
+FIoBuffer	
+FIoRequest::GetChunk()
+{
+	return Impl->Result.ValueOrDie();
+}
+
+const FIoChunkId&
+FIoRequest::GetChunkId() const
+{
+	return Impl->ChunkId;
+}
diff --git a/Engine/Source/Runtime/Core/Private/Mac/MacPlatformMisc.cpp b/Engine/Source/Runtime/Core/Private/Mac/MacPlatformMisc.cpp
--- a/Engine/Source/Runtime/Core/Private/Mac/MacPlatformMisc.cpp
+++ b/Engine/Source/Runtime/Core/Private/Mac/MacPlatformMisc.cpp
@@ -1627,7 +1627,7 @@ static void GracefulTerminationHandler(int32 Signal, siginfo_t* Info, void* Cont
 	}
 	else
 	{
-		_Exit(0);
+		_Exit(1);
 	}
 }
 
@@ -1986,7 +1986,7 @@ void FMacCrashContext::GenerateCrashInfoAndLaunchReporter() const
 		raise(Signal);
 	}
 	
-	_Exit(0);
+	_Exit(1);
 }
 
 void FMacCrashContext::GenerateEnsureInfoAndLaunchReporter() const
diff --git a/Engine/Source/Runtime/Core/Private/Misc/StringBuilder.cpp b/Engine/Source/Runtime/Core/Private/Misc/StringBuilder.cpp
new file mode 100644
--- /dev/null
+++ b/Engine/Source/Runtime/Core/Private/Misc/StringBuilder.cpp
@@ -0,0 +1,123 @@
+// Copyright 1998-2019 Epic Games, Inc. All Rights Reserved.
+
+#include "Misc/StringBuilder.h"
+#include "Misc/AssertionMacros.h"
+#include "HAL/PlatformMath.h"
+#include "HAL/UnrealMemory.h"
+
+//#include "doctest.h"
+
+static inline uint64_t NextPowerOfTwo(uint64_t x)
+{
+	x -= 1;
+	x |= (x >> 1);
+	x |= (x >> 2);
+	x |= (x >> 4);
+	x |= (x >> 8);
+	x |= (x >> 16);
+
+	return x + 1;
+}
+
+template<typename C>
+TStringBuilderImpl<C>::~TStringBuilderImpl()
+{
+	if (bIsDynamic)
+		FreeBuffer(Base, End - Base);
+}
+
+template<typename C>
+void
+TStringBuilderImpl<C>::Extend(int32 ExtraCapacity)
+{
+	check(bIsExtendable);
+
+	const SIZE_T OldCapacity = End - Base;
+	const SIZE_T NewCapacity = NextPowerOfTwo(OldCapacity + ExtraCapacity);
+
+	C* NewBase = (C*)AllocBuffer(NewCapacity);
+
+	SIZE_T Pos = CurPos - Base;
+	memcpy(NewBase, Base, Pos * sizeof(C));
+
+	if (bIsDynamic)
+	{
+		FreeBuffer(Base, OldCapacity);
+	}
+
+	Base		= NewBase;
+	CurPos		= NewBase + Pos;
+	End			= NewBase + NewCapacity;
+	bIsDynamic	= true;
+}
+
+template<typename C>
+void*
+TStringBuilderImpl<C>::AllocBuffer(int32 ByteCount)
+{
+	return FMemory::Malloc(ByteCount * sizeof(C));
+}
+
+template<typename C>
+void
+TStringBuilderImpl<C>::FreeBuffer(void* Buffer, int32 ByteCount)
+{
+	FMemory::Free(Buffer);
+}
+
+// Instantiate templates once
+
+template class TStringBuilderImpl<ANSICHAR>;
+template class TStringBuilderImpl<TCHAR>;
+
+#if 0
+#if !(UE_BUILD_SHIPPING || UE_BUILD_TEST)
+
+TEST_CASE("Core.StringView")
+{
+	{
+		TAnsiStringBuilder<32> boo;
+		CHECK(boo.Len() == 0);
+		CHECK(strcmp(boo.ToString(), "") == 0);
+
+		{
+			FAnsiStringView view = boo;
+			CHECK(boo.Len() == view.Len());
+		}
+
+		boo.Append("bababa");
+
+		CHECK(strcmp(boo.ToString(), "bababa") == 0);
+
+		{
+			FAnsiStringView view = boo;
+			CHECK(boo.Len() == view.Len());
+		}
+	}
+
+	{
+		TWideStringBuilder<32> booboo;
+		booboo.Append(TEXT("babababa"));
+
+		{
+			FWideStringView view = booboo;
+			CHECK(booboo.Len() == view.Len());
+		}
+
+		//booboo << TEXT("bobo");
+		//booboo << "baba";	// ANSI
+		//booboo << 'b';
+
+		char foo[] = "aaaaa";
+		char noo[] = "oOoOo";
+
+		booboo << noo;
+		booboo << foo;
+
+		FString ouch(TEXT("ouch"));
+		booboo << ouch;
+	}
+}
+
+#endif
+#endif
diff --git a/Engine/Source/Runtime/Core/Private/Serialization/Archive.cpp b/Engine/Source/Runtime/Core/Private/Serialization/Archive.cpp
--- a/Engine/Source/Runtime/Core/Private/Serialization/Archive.cpp
+++ b/Engine/Source/Runtime/Core/Private/Serialization/Archive.cpp
@@ -35,8 +35,6 @@ FArchive::FArchive()
 #if DEVIRTUALIZE_FLinkerLoad_Serialize
 	ActiveFPLB = &InlineFPLB;
 #endif
-	CustomVersionContainer = nullptr;
-
 	SerializedPropertyChain = nullptr;
 
 #if USE_STABLE_LOCALIZATION_KEYS
diff --git a/Engine/Source/Runtime/Core/Private/Templates/RefCounting.cpp b/Engine/Source/Runtime/Core/Private/Templates/RefCounting.cpp
new file mode 100644
--- /dev/null
+++ b/Engine/Source/Runtime/Core/Private/Templates/RefCounting.cpp
@@ -0,0 +1,10 @@
+// Copyright 1998-2019 Epic Games, Inc. All Rights Reserved.
+
+#include "Templates/RefCounting.h"
+
+void 
+FRefCountBase::CheckRefCount() const
+{
+	// Verify that Release() is not being called on an object which is already at a zero refcount
+	check(NumRefs != 0);
+}
diff --git a/Engine/Source/Runtime/Core/Private/Windows/WindowsPlatformMisc.cpp b/Engine/Source/Runtime/Core/Private/Windows/WindowsPlatformMisc.cpp
--- a/Engine/Source/Runtime/Core/Private/Windows/WindowsPlatformMisc.cpp
+++ b/Engine/Source/Runtime/Core/Private/Windows/WindowsPlatformMisc.cpp
@@ -1641,7 +1641,7 @@ int32 FWindowsPlatformMisc::NumberOfCores()
 	static int32 CoreCount = 0;
 	if (CoreCount == 0)
 	{
-		if (FParse::Param(FCommandLine::Get(), TEXT("usehyperthreading")))
+		if (FCommandLine::IsInitialized() && FParse::Param(FCommandLine::Get(), TEXT("usehyperthreading")))
 		{
 			CoreCount = NumberOfCoresIncludingHyperthreads();
 		}
@@ -1676,7 +1676,16 @@ int32 FWindowsPlatformMisc::NumberOfCores()
 			}
 			FMemory::Free(InfoBuffer);
 		}
+
+		// Optionally limit number of threads (we don't necessarily scale super well with very high core counts)
+
+		int32 LimitCount = 32768;
+		if (FParse::Value(FCommandLine::Get(), TEXT("-corelimit="), LimitCount))
+		{
+			CoreCount = FMath::Min(CoreCount, LimitCount);
+		}
 	}
+
 	return CoreCount;
 }
 
@@ -1689,7 +1698,16 @@ int32 FWindowsPlatformMisc::NumberOfCoresIncludingHyperthreads()
 		SYSTEM_INFO SI;
 		GetSystemInfo(&SI);
 		CoreCount = (int32)SI.dwNumberOfProcessors;
+
+		// Optionally limit number of threads (we don't necessarily scale super well with very high core counts)
+
+		int32 LimitCount = 32768;
+		if (FParse::Value(FCommandLine::Get(), TEXT("-corelimit="), LimitCount))
+		{
+			CoreCount = FMath::Min(CoreCount, LimitCount);
+		}
 	}
+
 	return CoreCount;
 }
 
@@ -1785,7 +1803,7 @@ HWND FWindowsPlatformMisc::GetTopLevelWindowHandle(uint32 ProcessId)
 FORCENOINLINE void FWindowsPlatformMisc::RaiseException( uint32 ExceptionCode )
 {
 	/** This is the last place to gather memory stats before exception. */
-	FGenericCrashContext::CrashMemoryStats = FPlatformMemory::GetStats();
+	FGenericCrashContext::SetMemoryStats(FPlatformMemory::GetStats());
 
 	::RaiseException( ExceptionCode, 0, 0, NULL );
 }
diff --git a/Engine/Source/Runtime/Core/Public/IO/IoDispatcher.h b/Engine/Source/Runtime/Core/Public/IO/IoDispatcher.h
new file mode 100644
--- /dev/null
+++ b/Engine/Source/Runtime/Core/Public/IO/IoDispatcher.h
@@ -0,0 +1,562 @@
+// Copyright 1998-2019 Epic Games, Inc. All Rights Reserved.
+
+#pragma once
+
+#include "CoreTypes.h"
+#include "Containers/UnrealString.h"
+#include "Misc/StringView.h"
+#include "Misc/StringBuilder.h"
+#include "Templates/RefCounting.h"
+#include "Templates/UnrealTemplate.h"
+#include "HAL/PlatformAtomics.h"
+
+#if __cplusplus >= 201703L
+#	define UE_NODISCARD		[[nodiscard]]
+#	define UE_NORETURN		[[noreturn]]
+#else
+#	define UE_NODISCARD		
+#	define UE_NORETURN		
+#endif
+
+class FIoRequest;
+class FIoDispatcher;
+class FIoStoreReader;
+class FIoStoreWriter;
+
+class FIoRequestImpl;
+class FIoDispatcherImpl;
+class FIoStoreReaderImpl;
+class FIoStoreWriterImpl;
+
+//////////////////////////////////////////////////////////////////////////
+//
+// IO Status classes modeled after Google Status / StatusOr
+//
+// TODO: prevent nullptr value in StatusOr
+//
+
+enum class EIoErrorCode
+{
+	Ok,
+	Unknown,
+	InvalidCode,
+	Cancelled,
+	FileOpenFailed,
+	FileNotOpen,
+	WriteError,
+	NotFound,
+	CorruptToc
+};
+
+class FIoStatus
+{
+public:
+						FIoStatus() = default;
+						~FIoStatus() = default;
+
+	CORE_API			FIoStatus(EIoErrorCode Code, const FStringView& InErrorMessage);
+	CORE_API			FIoStatus(EIoErrorCode Code);
+	CORE_API FIoStatus&	operator=(const FIoStatus& Other);
+
+	CORE_API bool		operator==(const FIoStatus& Other) const;
+			 bool		operator!=(const FIoStatus& Other) const { return !operator==(Other); }
+
+	inline bool			IsOk() const { return ErrorCode == EIoErrorCode::Ok; }
+	inline bool			IsCompleted() const { return ErrorCode != EIoErrorCode::Unknown; }
+	inline EIoErrorCode	GetErrorCode() const { return ErrorCode; }
+	CORE_API FString	ToString() const;
+
+	static const FIoStatus Ok;
+	static const FIoStatus Unknown;
+	static const FIoStatus Invalid;
+
+private:
+	EIoErrorCode	ErrorCode = EIoErrorCode::Ok;
+	FString			ErrorMessage;
+
+	friend class FIoStatusBuilder;
+};
+
+/** Helper to make it easier to generate meaningful error messages
+    for FIoStatusBuilder.
+  */
+class FIoStatusBuilder
+{
+	EIoErrorCode		StatusCode;
+	FString				Message;
+public:
+	CORE_API explicit	FIoStatusBuilder(EIoErrorCode StatusCode);
+	CORE_API			FIoStatusBuilder(const FIoStatus& InStatus, FStringView String);
+	CORE_API			~FIoStatusBuilder();
+
+	inline operator FIoStatus();
+
+	CORE_API FIoStatusBuilder& operator<<(FStringView String);
+};
+
+CORE_API FIoStatusBuilder operator<<(const FIoStatus& Status, FStringView String);
+
+template<typename T>
+class TIoStatusOr
+{
+	template<typename U> friend class TIoStatusOr;
+
+public:
+	TIoStatusOr() : StatusValue(FIoStatus::Unknown), Value{} {}
+	TIoStatusOr(const FIoStatus& Status);
+	TIoStatusOr(const T& Value);
+
+	template<typename U>
+	TIoStatusOr(const TIoStatusOr<U>& Other);
+
+	TIoStatusOr<T>& operator=(const TIoStatusOr<T>& other);
+
+	template<typename U>
+	TIoStatusOr<T>& operator=(const TIoStatusOr<U>& other);
+
+	const FIoStatus&	Status() const;
+	bool				IsOk() const;
+
+	const T&			ValueOrDie();
+	T					ConsumeValueOrDie();
+
+	void				Reset();
+
+private:
+	FIoStatus	StatusValue;
+	T			Value;
+};
+
+CORE_API void StatusOrCrash(const FIoStatus& Status);
+
+template<typename T>
+void TIoStatusOr<T>::Reset()
+{
+	StatusValue = FIoStatus::Unknown;
+	Value = T();
+}
+
+template<typename T>
+const T& TIoStatusOr<T>::ValueOrDie()
+{
+	if (!StatusValue.IsOk())
+	{
+		StatusOrCrash(StatusValue);
+	}
+
+	return Value;
+}
+
+template<typename T>
+T TIoStatusOr<T>::ConsumeValueOrDie()
+{
+	if (!StatusValue.IsOk())
+	{
+		StatusOrCrash(StatusValue);
+	}
+
+	return MoveTemp(Value);
+}
+
+template<typename T>
+TIoStatusOr<T>::TIoStatusOr(const FIoStatus& Status)
+{
+	if (Status.IsOk())
+	{
+		// This doesn't make a whole lot of sense. If everything's ok then
+		// we should be returning a T
+		StatusValue = FIoStatus::Invalid;
+	}
+	else
+	{
+		StatusValue = Status;
+	}
+}
+
+template<typename T>
+TIoStatusOr<T>::TIoStatusOr(const T& InValue)
+{
+	Value = InValue;
+}
+
+template<typename T>
+bool TIoStatusOr<T>::IsOk() const
+{
+	return StatusValue.IsOk();
+}
+
+template<typename T>
+const FIoStatus& TIoStatusOr<T>::Status() const
+{
+	return StatusValue;
+}
+
+template<typename T>
+TIoStatusOr<T>&
+TIoStatusOr<T>::operator=(const TIoStatusOr<T>& Other)
+{
+	StatusValue = Other.StatusValue;
+	Value = Other.Value;
+	return *this;
+}
+
+template<typename T>
+template<typename U>
+TIoStatusOr<T>::TIoStatusOr(const TIoStatusOr<U>& Other)
+:	StatusValue(Other.StatusValue)
+,	Value(Other.StatusValue.IsOk() ? Other.Value : T()) 
+{
+}
+
+template<typename T>
+template<typename U>
+TIoStatusOr<T>& TIoStatusOr<T>::operator=(const TIoStatusOr<U>& Other)
+{
+	StatusValue = Other.StatusValue;
+	if (StatusValue.IsOk())
+	{
+		Value = Other.Value;
+	}
+	return *this;
+}
+
+/** Reference to buffer data used by I/O dispatcher APIs
+  */
+class FIoBuffer
+{
+public:
+	enum EAssumeOwnershipTag	{ AssumeOwnership };
+	enum ECloneTag				{ Clone };
+	enum EWrapTag				{ Wrap };
+
+	CORE_API			FIoBuffer();
+	CORE_API explicit	FIoBuffer(uint64 InSize);
+	CORE_API			FIoBuffer(const void* Data, uint64 InSize, const FIoBuffer& OuterBuffer);
+
+	CORE_API			FIoBuffer(EAssumeOwnershipTag,	const void* Data, uint64 InSize);
+	CORE_API			FIoBuffer(ECloneTag,			const void* Data, uint64 InSize);
+	CORE_API			FIoBuffer(EWrapTag,				const void* Data, uint64 InSize);
+
+	// Note: we currently rely on implicit move constructor, thus we do not declare any
+	//		 destructor or copy/assignment operators or copy constructors
+
+	inline const uint8*	Data() const			{ return CorePtr->Data(); }
+	inline uint8*		Data()					{ return CorePtr->Data(); }
+	inline uint64		DataSize() const		{ return CorePtr->DataSize(); }
+
+	inline void			SetSize(uint64 InSize)	{ return CorePtr->SetSize(InSize); }
+
+	inline bool			IsAvailable() const;
+	inline bool			IsMemoryOwned() const	{ return CorePtr->IsMemoryOwned(); }
+
+	inline void			EnsureOwned() const		{ if (!CorePtr->IsMemoryOwned()) { MakeOwned(); } }
+
+	CORE_API void		MakeOwned() const;
+
+private:
+	/** Core buffer object. For internal use only, used by FIoBuffer
+
+		Contains all state pertaining to a buffer.
+	  */
+	struct BufCore
+	{
+					BufCore();
+		CORE_API	~BufCore();
+
+		explicit	BufCore(uint64 InSize);
+					BufCore(const uint8* InData, uint64 InSize, bool InOwnsMemory);
+					BufCore(const uint8* InData, uint64 InSize, const BufCore* InOuter);
+					BufCore(ECloneTag, uint8* InData, uint64 InSize);
+
+					BufCore(const BufCore& Rhs) = delete;
+		
+		BufCore& operator=(const BufCore& Rhs) = delete;
+
+		inline uint8* Data()			{ return DataPtr; }
+		inline uint64 DataSize() const	{ return DataSizeLow | (uint64(DataSizeHigh) << 32); }
+
+		//
+
+		void	SetDataAndSize(const uint8* InData, uint64 InSize);
+		void	SetSize(uint64 InSize);
+
+		void	MakeOwned();
+
+		inline void SetIsOwned(bool InOwnsMemory)
+		{
+			if (InOwnsMemory)
+			{
+				Flags |= OwnsMemory;
+			}
+			else
+			{
+				Flags &= ~OwnsMemory;
+			}
+		}
+
+		inline uint32 AddRef() const
+		{
+			return uint32(FPlatformAtomics::InterlockedIncrement(&NumRefs));
+		}
+
+		inline uint32 Release() const
+		{
+#if DO_CHECK
+			CheckRefCount();
+#endif
+
+			const int32 Refs = FPlatformAtomics::InterlockedDecrement(&NumRefs);
+			if (Refs == 0)
+			{
+				delete this;
+			}
+
+			return uint32(Refs);
+		}
+
+		uint32 GetRefCount() const
+		{
+			return uint32(NumRefs);
+		}
+
+		bool			IsMemoryOwned() const	{ return Flags & OwnsMemory; }
+
+	private:
+		CORE_API void				CheckRefCount() const;
+
+		uint8*						DataPtr = nullptr;
+
+		uint32						DataSizeLow = 0;
+		mutable int32				NumRefs = 0;
+
+		// Reference-counted outer "core", used for views into other buffer
+		//
+		// Ultimately this should probably just be an index into a pool
+		TRefCountPtr<const BufCore>	OuterCore;
+
+		// TODO: These two could be packed in the MSB of DataPtr on x64
+		uint8		DataSizeHigh = 0;	// High 8 bits of size (40 bits total)
+		uint8		Flags = 0;
+
+		enum
+		{
+			OwnsMemory		= 1 << 0,	// Buffer memory is owned by this instance
+			ReadOnlyBuffer	= 1 << 1,	// Buffer memory is immutable
+			
+			FlagsMask		= (1 << 2) - 1
+		};
+
+		void EnsureDataIsResident() {}
+	};
+
+	// Reference-counted "core"
+	//
+	// Ultimately this should probably just be an index into a pool
+	TRefCountPtr<BufCore>	CorePtr;
+	
+	friend class FIoBufferManager;
+};
+
+class FIoChunkId
+{
+public:
+	friend uint32 GetTypeHash(FIoChunkId InId)
+	{
+		uint32 Hash = 5381;
+		for (int i = 0; i < sizeof Id; ++i)
+		{
+			Hash = Hash * 33 + InId.Id[i];
+		}
+		return Hash;
+	}
+
+	inline bool operator==(const FIoChunkId& Rhs) const
+	{
+		return 0 == FMemory::Memcmp(Id, Rhs.Id, sizeof Id);
+	}
+
+	CORE_API void GenerateFromData(const void* InData, SIZE_T InDataSize);
+
+	void Set(const void* InIdPtr, SIZE_T InSize)
+	{
+		check(InSize == sizeof Id);
+		FMemory::Memcpy(Id, InIdPtr, sizeof Id);
+	}
+
+private:
+	uint8	Id[12];
+};
+
+//////////////////////////////////////////////////////////////////////////
+
+class FIoReadOptions
+{
+public:
+	FIoReadOptions() = default;
+	~FIoReadOptions() = default;
+
+	void SetRange(uint64 Offset, uint32 Size)
+	{
+		RequestedOffset = Offset;
+		RequestedSize	= Size;
+	}
+
+	void SetTargetVa(uint64 VaTargetAddress)
+	{
+		TargetVa = VaTargetAddress;
+	}
+
+	void ForGPU()
+	{
+		Flags |= EFlags::GPUMemory;
+	}
+
+private:
+	uint64	TargetVa		= 0;
+	uint64	RequestedOffset = 0;
+	uint32	RequestedSize	= ~uint32(0);
+	uint32	Flags			= 0;
+
+	enum EFlags
+	{
+		GPUMemory = 1 << 0,
+	};
+};
+
+//////////////////////////////////////////////////////////////////////////
+
+/**
+  */
+class FIoRequest
+{
+public:
+	FIoRequest() = default;
+
+	FIoRequest(const FIoRequest&) = default;
+	FIoRequest& operator=(const FIoRequest&) = default;
+
+	CORE_API bool				IsOk() const;
+	CORE_API FIoStatus			Status() const;
+	CORE_API FIoBuffer			GetChunk();
+	CORE_API const FIoChunkId&	GetChunkId() const;
+
+private:
+	FIoRequestImpl* Impl = nullptr;
+
+	explicit FIoRequest(FIoRequestImpl& InImpl)
+	: Impl(&InImpl)
+	{
+	}
+
+	friend class FIoDispatcherImpl;
+	friend class FIoBatch;
+};
+
+/** I/O batch
+
+	This is a primitive used to group I/O requests for synchronization
+	purposes
+  */
+class FIoBatch
+{
+	friend class FIoDispatcher;
+
+	FIoBatch(FIoDispatcherImpl* OwningIoDispatcher, uint32 InBatchId);
+
+public:
+	CORE_API FIoRequest Read(const FIoChunkId& Chunk, FIoReadOptions Options);
+
+	CORE_API void ForEachRequest(TFunction<bool(FIoRequest&)> Callback);
+
+	CORE_API void Issue();
+	CORE_API void Wait();
+	CORE_API void Cancel();
+
+private:
+	FIoDispatcherImpl*	Dispatcher		= nullptr;
+	uint32				BatchId			= ~uint32(0);
+	FEvent*				CompletionEvent = nullptr;
+};
+
+/** I/O dispatcher
+  */
+class FIoDispatcher
+{
+public:
+	CORE_API			FIoDispatcher();
+	CORE_API virtual	~FIoDispatcher();
+
+	CORE_API void		Mount(FIoStoreReader* IoStore);
+	CORE_API void		Unmount(FIoStoreReader* IoStore);
+
+	CORE_API FIoBatch	NewBatch();
+	CORE_API void		FreeBatch(FIoBatch Batch);
+
+	FIoDispatcher(const FIoDispatcher&) = default;
+	FIoDispatcher& operator=(const FIoDispatcher&) = delete;
+
+private:
+	FIoDispatcherImpl* Impl;
+
+	friend class FIoRequest;
+	friend class FIoBatch;
+};
+
+//////////////////////////////////////////////////////////////////////////
+
+/** Helper used to manage creation of I/O store file handles etc 
+  */
+class FIoStoreEnvironment
+{
+public:
+	CORE_API FIoStoreEnvironment();
+	CORE_API ~FIoStoreEnvironment();
+
+	CORE_API void InitializeFileEnvironment(FStringView InRootPath);
+
+	CORE_API const FString& GetRootPath() const { return RootPath; }
+
+private:
+	FString			RootPath;
+};
+
+class FIoStoreReader : public FRefCountBase
+{
+public:
+	CORE_API FIoStoreReader(FIoStoreEnvironment& Environment);
+	CORE_API ~FIoStoreReader();
+
+	/** This will parse the manifests in the environment and
+		populate the table of contents. To be useful the IO dispatcher
+		needs to have access to the information, which is what
+		the I/O dispatcher Mount()/Unmount() calls are for.
+	  */
+	CORE_API FIoStatus Initialize(FStringView UniqueId);
+
+private:
+	FIoStoreReaderImpl* Impl;
+
+	friend class FIoDispatcherImpl;
+	friend class FIoStoreImpl;
+};
+
+//////////////////////////////////////////////////////////////////////////
+
+class FIoStoreWriter
+{
+public:
+	CORE_API 			FIoStoreWriter(FIoStoreEnvironment& InEnvironment);
+	CORE_API virtual	~FIoStoreWriter();
+
+	FIoStoreWriter(const FIoStoreWriter&) = delete;
+	FIoStoreWriter& operator=(const FIoStoreWriter&) = delete;
+
+	CORE_API FIoStatus	Initialize();
+	CORE_API void		Append(FIoChunkId ChunkId, FIoBuffer Chunk);
+	CORE_API void		FlushMetadata();
+
+private:
+	FIoStoreWriterImpl*		Impl;
+};
+
+//////////////////////////////////////////////////////////////////////////
diff --git a/Engine/Source/Runtime/Core/Public/Memory/MemoryArena.h b/Engine/Source/Runtime/Core/Public/Memory/MemoryArena.h
new file mode 100644
--- /dev/null
+++ b/Engine/Source/Runtime/Core/Public/Memory/MemoryArena.h
@@ -0,0 +1,226 @@
+// Copyright 1998-2019 Epic Games, Inc. All Rights Reserved.
+
+#pragma once
+
+#include "CoreTypes.h"
+
+#if !defined(UE_RESTRICT)
+#	if defined _MSC_VER
+#		define UE_RESTRICT __declspec(restrict)
+#	else
+#		define UE_RESTRICT
+#	endif
+#endif
+
+#if !defined(UE_NOALIAS)
+#	if defined _MSC_VER
+#		define UE_NOALIAS __declspec(noalias)
+#	else
+#		define UE_NOALIAS
+#	endif
+#endif
+
+#define UE_WITH_HEAPARENA	0
+#define UE_WITH_ARENAMAP	0
+
+class FMemoryArena;
+class FHeapArena;
+
+/** Tagged Arena Pointer
+
+	Stores an arena ID alongside a pointer. On 64-bit architectures the arena ID
+	can be encoded in the otherwise unused MSB part of the pointer for zero
+	memory overhead vs a regular pointer. For 32-bit architectures the arena ID
+	is stored in a separate member.
+
+	A null pointer can still have an arena ID, which is useful for cases where
+	you have a container that starts out empty but still want to track which arena
+	should be used for subsequent allocations.
+
+  */
+class FArenaPointer
+{
+public:
+	FArenaPointer() = default;
+
+	enum { NoTag = 0 };
+
+#if PLATFORM_32BITS
+private:
+	void*	Ptr = 0;
+	uint16	ArenaTag = 0;
+
+public:
+	FORCEINLINE const uint16	ArenaIndex() const									{ return ArenaTag; }
+	FORCEINLINE void*			Pointer()											{ return Ptr; }
+	FORCEINLINE const void*		Pointer() const										{ return Ptr; }
+	FORCEINLINE void			SetPointerAndArena(void* InPtr, uint16 InArenaTag)	{ Ptr = InPtr; ArenaTag = InArenaTag; }
+
+	inline FArenaPointer(void* InPtr, uint16 InArenaTag) : Ptr(InPtr), ArenaTag(InArenaTag) {}
+
+#else
+
+private:
+	static const int	ArenaShift	= 48;
+	static const uint64 ArenaMask	= 0xffff000000000000;
+	static const uint64 PointerMask = 0x0000ffffFFFFffff;
+
+	void*	TaggedPointer = nullptr;
+
+public:
+	FORCEINLINE const uint16	ArenaIndex() const	{ return uint16(UPTRINT(TaggedPointer) >> ArenaShift); }
+	FORCEINLINE void*			Pointer()			{ return reinterpret_cast<void*>(UPTRINT(TaggedPointer) & PointerMask); }
+	FORCEINLINE const void*		Pointer() const		{ return reinterpret_cast<const void*>(UPTRINT(TaggedPointer) & PointerMask); }
+	FORCEINLINE void			SetPointerAndArena(void* InPtr, uint16 InArenaTag) { TaggedPointer = reinterpret_cast<void*>(UPTRINT(InPtr) | (UPTRINT(InArenaTag) << ArenaShift)); }
+
+	inline FArenaPointer(void* Ptr, uint16 ArenaIndex) : TaggedPointer(reinterpret_cast<void*>(UPTRINT(Ptr) | (UPTRINT(ArenaIndex) << ArenaShift))) {};
+#endif
+
+	inline operator bool() const	{ return !!Pointer(); }
+
+	void Free() const;
+
+	CORE_API FMemoryArena& Arena() const;
+};
+
+template<typename T>
+class TArenaPointer : public FArenaPointer
+{
+public:
+	TArenaPointer() = default;
+	TArenaPointer(T* Ptr, uint16 ArenaIndex) : FArenaPointer(reinterpret_cast<T*>(Ptr), ArenaIndex) {}
+
+	inline TArenaPointer& operator=(T* Rhs) { SetPointerAndArena(Rhs, NoTag); return *this; }
+	
+	inline operator T* ()					{ return reinterpret_cast<T*>(Pointer()); }
+	inline operator const T* () const		{ return reinterpret_cast<const T*>(Pointer()); }
+	inline T* operator -> ()				{ return reinterpret_cast<T*>(Pointer()); }
+	inline const T* operator -> () const	{ return reinterpret_cast<const T*>(Pointer()); }
+};
+
+/** Memory arena interface
+  */
+class FMemoryArena
+{
+public:
+	CORE_API			FMemoryArena();
+	CORE_API virtual	~FMemoryArena();
+
+	FMemoryArena(const FMemoryArena&) = delete;
+	FMemoryArena& operator=(const FMemoryArena&) = delete;
+
+	CORE_API UE_RESTRICT UE_NOALIAS void*	Alloc(SIZE_T Size, SIZE_T Alignment);
+	CORE_API UE_NOALIAS void				Free(const void* MemoryBlock);
+
+	CORE_API SIZE_T				BlockSize(const void* MemoryBlock) const;
+	CORE_API const TCHAR*		GetDebugName() const;
+
+private:
+	CORE_API virtual void*		InternalAlloc(SIZE_T Size, SIZE_T Alignment) = 0;
+	CORE_API virtual void		InternalFree(const void* MemoryBlock, SIZE_T MemoryBlockSize);
+	CORE_API virtual SIZE_T		InternalBlockSize(const void* MemoryBlock) const = 0;
+
+	CORE_API virtual const TCHAR* InternalGetDebugName() const;
+
+	enum { FlagNoFree = 1 << 0 };
+
+	uint16			ArenaFlags	= 0;
+
+public:
+	uint16			ArenaId		= ~uint16(0);
+};
+
+inline void FArenaPointer::Free() const 
+{ 
+	return Arena().Free(Pointer());
+}
+
+// These are somewhat temporary, to support (experimental) arena-based container allocators
+
+CORE_API FArenaPointer ArenaRealloc(FArenaPointer InPtr, SIZE_T OldSize, SIZE_T NewSize, SIZE_T Alignment);
+CORE_API FArenaPointer ArenaRealloc(FMemoryArena* Arena, void* InPtr, SIZE_T OldSize, SIZE_T NewSize, SIZE_T Alignment);
+
+/** Heap arena
+
+	Manages a dedicated area of memory, and allows user to allocate blocks from
+	it
+  */
+
+#if UE_WITH_HEAPARENA
+class FHeapArena : public FMemoryArena
+{
+public:
+	CORE_API			FHeapArena();
+	CORE_API			~FHeapArena();
+
+private:
+	CORE_API virtual void*			InternalAlloc(SIZE_T Size, SIZE_T Alignment) override;
+	CORE_API virtual void			InternalFree(const void* MemoryBlock, SIZE_T MemoryBlockSize) override;
+	CORE_API virtual SIZE_T			InternalBlockSize(const void* MemoryBlock) const override;
+	CORE_API virtual const TCHAR*	InternalGetDebugName() const override;
+
+	void* HeapHandle = nullptr;
+};
+#endif
+
+/** Default heap allocator
+
+	All allocations are passed through to UE4 main heap allocation functions
+
+  */
+class FMallocArena final : public FMemoryArena
+{
+public:
+	CORE_API			FMallocArena();
+	CORE_API			~FMallocArena();
+
+private:
+	CORE_API virtual void*			InternalAlloc(SIZE_T Size, SIZE_T Alignment) override;
+	CORE_API virtual void			InternalFree(const void* MemoryBlock, SIZE_T MemoryBlockSize) override;
+	CORE_API virtual SIZE_T			InternalBlockSize(const void* MemoryBlock) const override;
+	CORE_API virtual const TCHAR*	InternalGetDebugName() const override;
+};
+
+/** CRT heap allocator
+
+	All allocations are passed through to CRT memory allocation functions
+
+  */
+class FAnsiArena final : public FMemoryArena
+{
+public:
+	CORE_API			FAnsiArena();
+	CORE_API			~FAnsiArena();
+
+private:
+	CORE_API virtual void*			InternalAlloc(SIZE_T Size, SIZE_T Alignment) override;
+	CORE_API virtual void			InternalFree(const void* MemoryBlock, SIZE_T MemoryBlockSize) override;
+	CORE_API virtual SIZE_T			InternalBlockSize(const void* MemoryBlock) const override;
+	CORE_API virtual const TCHAR*	InternalGetDebugName() const override;
+};
+
+/** Memory arena map
+
+	Provides an efficient mechanism for mapping pointers to arenas
+  */
+
+#if UE_WITH_ARENAMAP
+class FArenaMap
+{
+public:
+	CORE_API static void Initialize();
+
+	// Reset the state of the arena map - this is here only for testing and benchmarking purposes
+	// and would never normally be called during normal execution
+	CORE_API static void Reset();
+
+	// Associate an arena with a virtual address range
+	CORE_API static void SetRangeToArena(const void* VaBase, SIZE_T VaSize, FMemoryArena* ArenaPtr);
+
+	// Mark range as unallocated
+	CORE_API static void ClearRange(const void* VaBase, SIZE_T VaSize);
+
+	// Map a virtual address to a memory arena
+	CORE_API static FMemoryArena* MapPtrToArena(const void* VaBase);
+};
+#endif
diff --git a/Engine/Source/Runtime/Core/Public/Misc/StringBuilder.h b/Engine/Source/Runtime/Core/Public/Misc/StringBuilder.h
new file mode 100644
--- /dev/null
+++ b/Engine/Source/Runtime/Core/Public/Misc/StringBuilder.h
@@ -0,0 +1,357 @@
+// Copyright 1998-2019 Epic Games, Inc. All Rights Reserved.
+
+#pragma once
+
+#include "CoreTypes.h"
+#include "Misc/CString.h"
+#include "Misc/StringView.h"
+
+#define USE_STRING_LITERAL_PATH 1
+
+/** String Builder implementation
+
+	This class helps with the common task of constructing new strings. 
+	
+	It does this by allocating buffer space which is used to hold the 
+	constructed string. The intent is that the builder is allocated on
+	the stack as a function local variable to avoid heap allocations.
+
+	The buffer is always contiguous and the class is not intended to be
+	used to construct extremely large strings. 
+	
+	This is not intended to be used as a mechanism for holding on to 
+	strings for a long time - the use case is explicitly to aid in 
+	*constructing* strings on the stack	and subsequently passing the 
+	string into a function call or a more permanent string storage 
+	mechanism like FString et al
+
+	The amount of buffer space to allocate is specified via a template
+	parameter and if the constructed string should overflow this initial
+	buffer, a new buffer will be allocated using regular dynamic memory
+	allocations. For instances where you absolutely must not allocate
+	any memory, you should use the fixed variants named
+	TFixed*StringBuilder
+
+	Overflow allocation should be the exceptional case however -- always
+	try to size the buffer so that it can hold the vast majority of 
+	strings you expect to construct.
+
+	Be mindful that stack is a limited resource, so if you are writing a
+	highly recursive function you may want to use some other mechanism
+	to build your strings.
+
+  */
+
+template<typename C>
+class TStringBuilderImpl
+{
+public:
+				TStringBuilderImpl() = default;
+	CORE_API	~TStringBuilderImpl();
+
+				TStringBuilderImpl(const TStringBuilderImpl&) = delete;
+				TStringBuilderImpl(TStringBuilderImpl&&) = delete;
+
+	const TStringBuilderImpl& operator=(const TStringBuilderImpl&) = delete;
+	const TStringBuilderImpl& operator=(TStringBuilderImpl&&) = delete;
+
+	inline int32	Len() const			{ return int32(CurPos - Base); }
+	inline const C* ToString() const	{ EnsureNulTerminated(); return Base; }
+	inline const C* operator*() const	{ EnsureNulTerminated(); return Base; }
+
+	inline const C	LastChar() const	{ return *(CurPos - 1); }
+
+	inline TStringBuilderImpl& Append(C Char)
+	{
+		EnsureCapacity(1);
+
+		*CurPos++ = Char;
+
+		return *this;
+	}
+
+#if USE_STRING_LITERAL_PATH
+	inline TStringBuilderImpl& AppendAnsi(const ANSICHAR*& NulTerminatedString)
+#else
+	inline TStringBuilderImpl& AppendAnsi(const ANSICHAR* NulTerminatedString)
+#endif
+	{
+		const int32 Length = TCString<ANSICHAR>::Strlen(NulTerminatedString);
+
+		EnsureCapacity(Length);
+
+		C* RESTRICT Dest = CurPos;
+		CurPos += Length;
+
+		for (int32 i = 0; i < Length; ++i)
+		{
+			Dest[i] = NulTerminatedString[i];
+		}
+
+		return *this;
+	}
+
+	inline TStringBuilderImpl& AppendAnsi(const FAnsiStringView& AnsiString)
+	{
+		return AppendAnsi(AnsiString.Data(), AnsiString.Len());
+	}
+
+	inline TStringBuilderImpl& AppendAnsi(const ANSICHAR* NulTerminatedString, const int32 Length)
+	{
+		EnsureCapacity(Length);
+
+		C* RESTRICT Dest = CurPos;
+		CurPos += Length;
+
+		for (int32 i = 0; i < Length; ++i)
+		{
+			Dest[i] = NulTerminatedString[i];
+		}
+
+		return *this;
+	}
+
+#if USE_STRING_LITERAL_PATH
+	inline TStringBuilderImpl& Append(const C*& NulTerminatedString)
+#else
+	inline TStringBuilderImpl& Append(const C* NulTerminatedString)
+#endif
+	{
+		const int32 Length = TCString<C>::Strlen(NulTerminatedString);
+
+		EnsureCapacity(Length);
+		C* RESTRICT Dest = CurPos;
+		CurPos += Length;
+
+		FMemory::Memcpy(Dest, NulTerminatedString, Length * sizeof(C));
+
+		return *this;
+	}
+
+	inline TStringBuilderImpl& Append(const C* NulTerminatedString, int32 MaxChars)
+	{
+		const int32 StringLength = TCString<C>::Strlen(NulTerminatedString);
+		const int32 Length = FPlatformMath::Min<int32>(MaxChars, StringLength);
+
+		EnsureCapacity(Length);
+		C* RESTRICT Dest = CurPos;
+		CurPos += Length;
+
+		FMemory::Memcpy(Dest, NulTerminatedString, Length * sizeof(C));
+
+		return *this;
+	}
+
+#if USE_STRING_LITERAL_PATH
+	template<int32 ArrayLength>
+	inline TStringBuilderImpl& Append(C(&CharArray)[ArrayLength])
+	{
+		return Append(&CharArray[0]);
+	}
+
+	template<int32 LiteralLength>
+	inline TStringBuilderImpl& Append(const C (&StringLiteral)[LiteralLength])
+	{
+		constexpr int32 Length = LiteralLength - 1;
+		EnsureCapacity(Length);
+		C* RESTRICT Dest = CurPos;
+		CurPos += Length;
+
+		FMemory::Memcpy(Dest, StringLiteral, Length * sizeof(C));
+
+		return *this;
+	}
+
+	template<int32 LiteralLength>
+	inline TStringBuilderImpl& AppendAnsi(const ANSICHAR(&StringLiteral)[LiteralLength])
+	{
+		constexpr int32 Length = LiteralLength - 1;
+		EnsureCapacity(Length);
+
+		C* RESTRICT Dest = CurPos;
+		CurPos += Length;
+
+		for (int32 i = 0; i < Length; ++i)
+		{
+			Dest[i] = StringLiteral[i];
+		}
+
+		return *this;
+	}
+#endif
+
+protected:
+	inline void Initialize(C* InBase, int32 InCapacity)
+	{
+		Base	= InBase;
+		CurPos	= InBase;
+		End		= Base + InCapacity;
+	}
+
+	inline void EnsureNulTerminated() const
+	{
+		if (*CurPos)
+		{
+			*CurPos = 0;
+		}
+	}
+
+	inline void EnsureCapacity(int32 RequiredAdditionalCapacity)
+	{
+		// precondition: we know the current buffer has enough capacity
+		// for the existing string including NUL terminator
+
+		if ((CurPos + RequiredAdditionalCapacity) < End)
+			return;
+
+		Extend(RequiredAdditionalCapacity);
+	}
+
+	CORE_API void	Extend(int32 extraCapacity);
+	CORE_API void*	AllocBuffer(int32 byteCount);
+	CORE_API void	FreeBuffer(void* buffer, int32 byteCount);
+
+	C*			Base;
+	C* 			CurPos;
+	C*			End;
+	bool		bIsDynamic = false;
+	bool		bIsExtendable = false;
+};
+
+//////////////////////////////////////////////////////////////////////////
+
+extern template class TStringBuilderImpl<ANSICHAR>;
+
+class FAnsiStringBuilderBase : public TStringBuilderImpl<ANSICHAR>
+{
+public:
+	operator FAnsiStringView() const { return FAnsiStringView(Base, CurPos - Base); }
+
+protected:
+	inline FAnsiStringBuilderBase(ANSICHAR* BufferPointer, int32 BufferCapacity)
+	{
+		Initialize(BufferPointer, BufferCapacity);
+	}
+
+	~FAnsiStringBuilderBase() = default;
+};
+
+template<int32 N>
+class TFixedAnsiStringBuilder : public FAnsiStringBuilderBase
+{
+public:
+	inline TFixedAnsiStringBuilder()
+	: FAnsiStringBuilderBase(StringBuffer, N)
+	{
+	}
+
+	~TFixedAnsiStringBuilder() = default;
+
+private:
+	ANSICHAR	StringBuffer[N];
+};
+
+template<int32 N>
+class TAnsiStringBuilder : public FAnsiStringBuilderBase
+{
+public:
+	inline TAnsiStringBuilder()
+	: FAnsiStringBuilderBase(StringBuffer, N)
+	{
+		bIsExtendable = true;
+	}
+
+	~TAnsiStringBuilder() = default;
+
+private:
+	ANSICHAR	StringBuffer[N];
+};
+
+//////////////////////////////////////////////////////////////////////////
+
+extern template class TStringBuilderImpl<TCHAR>;
+
+class FStringBuilderBase : public TStringBuilderImpl<TCHAR>
+{
+public:
+	operator FStringView() const { return FStringView(Base, CurPos - Base); }
+
+	using TStringBuilderImpl<TCHAR>::Append;
+
+	inline TStringBuilderImpl& Append(const FStringView& StringView)
+	{
+		const int32 StringLength = StringView.Len();
+
+		EnsureCapacity(StringLength);
+		TCHAR* RESTRICT Dest = CurPos;
+
+		FMemory::Memcpy(CurPos, StringView.Data(), StringLength * sizeof(TCHAR));
+
+		CurPos += StringLength;
+
+		return *this;
+	}
+
+protected:
+	FStringBuilderBase(TCHAR* BufferPointer, int32 BufferCapacity)
+	{
+		Initialize(BufferPointer, BufferCapacity);
+	}
+
+	~FStringBuilderBase() = default;
+};
+
+template<int32 N>
+class TFixedStringBuilder : public FStringBuilderBase
+{
+public:
+	inline TFixedStringBuilder()
+	: FStringBuilderBase(StringBuffer, N)
+	{
+	}
+
+	~TFixedStringBuilder() = default;
+
+private:
+	TCHAR	StringBuffer[N];
+};
+
+template<int32 N>
+class TStringBuilder : public FStringBuilderBase
+{
+public:
+	inline TStringBuilder()
+	: FStringBuilderBase(StringBuffer, N)
+	{
+		bIsExtendable = true;
+	}
+
+	~TStringBuilder() = default;
+
+private:
+	TCHAR	StringBuffer[N];
+};
+
+// Append operator implementations
+
+inline FStringBuilderBase&					operator<<(FStringBuilderBase& Builder, const TCHAR Char)						{ Builder.Append(Char); return Builder; }
+
+#if USE_STRING_LITERAL_PATH
+
+inline FStringBuilderBase&					operator<<(FStringBuilderBase& Builder, const TCHAR*& NulTerminatedString)		{ Builder.Append(NulTerminatedString); return Builder; }
+inline FStringBuilderBase&					operator<<(FStringBuilderBase& Builder, const ANSICHAR*& NulTerminatedString)	{ Builder.AppendAnsi(NulTerminatedString); return Builder; }
+
+template<int32 N> FStringBuilderBase&		operator<<(FStringBuilderBase& Builder, const TCHAR (&StringLiteral)[N])		{ Builder.Append(StringLiteral); return Builder; }
+template<int32 N> FStringBuilderBase&		operator<<(FStringBuilderBase& Builder, const ANSICHAR (&StringLiteral)[N])		{ Builder.AppendAnsi(StringLiteral); return Builder; }
+
+template<int32 N> FStringBuilderBase&		operator<<(FStringBuilderBase& Builder, TCHAR (&NulTerminatedString)[N])		{ Builder.Append(&NulTerminatedString[0], TCString<TCHAR>::Strlen(NulTerminatedString)); return Builder; }
+template<int32 N> FStringBuilderBase&		operator<<(FStringBuilderBase& Builder, ANSICHAR (&NulTerminatedString)[N])		{ Builder.AppendAnsi(&NulTerminatedString[0], TCString<ANSICHAR>::Strlen(NulTerminatedString)); return Builder; }
+
+#else
+
+inline FStringBuilderBase&					operator<<(FStringBuilderBase& Builder, const TCHAR* NulTerminatedString)		{ Builder.Append(NulTerminatedString); return Builder; }
+inline FStringBuilderBase&					operator<<(FStringBuilderBase& Builder, const ANSICHAR* NulTerminatedString)	{ Builder.AppendAnsi(NulTerminatedString); return Builder; }
+
+#endif
+
+inline FStringBuilderBase&					operator<<(FStringBuilderBase& Builder, const FStringView& Str)					{ Builder.Append(Str.Data(), Str.Len()); return Builder; }
diff --git a/Engine/Source/Runtime/Core/Public/Misc/StringView.h b/Engine/Source/Runtime/Core/Public/Misc/StringView.h
new file mode 100644
--- /dev/null
+++ b/Engine/Source/Runtime/Core/Public/Misc/StringView.h
@@ -0,0 +1,223 @@
+// Copyright 1998-2019 Epic Games, Inc. All Rights Reserved.
+
+#pragma once
+
+#include "CoreTypes.h"
+#include "Misc/CString.h"
+#include "Containers/UnrealString.h"
+
+/** String View
+
+	* A string view is implicitly constructible from FString and const char* style strings
+
+	* A string view does not own any data nor does it attempt to control any lifetimes, it 
+	  merely points at a subrange of characters in some other string. It's up to the user
+	  to ensure the underlying string stays valid for the lifetime of the string view.
+
+	* A string view does not represent a NUL terminated string and therefore you should
+	  never pass in the pointer returned by Data() into a C-string API accepting only a
+	  pointer. You must either use a string builder to make a properly terminated string,
+	  or the ToString() function below if you absolutely must, and can live with the
+	  knowledge that you just added yet one more memory allocation even though memory
+	  allocations are not cheap.
+
+	String views are a good fit for arguments to functions which don't wish to care
+	which style of string construction is used by the caller. If you accept strings via
+	string views then the caller is free to use FString, FStringBuilder or raw C strings
+	or any other type which can be converted into a string builder.
+
+	I.e a function such as
+
+	void DoFoo(const FStringView& InString);
+
+	May be called as:
+
+	void MultiFoo()
+	{
+	     FString MyFoo(TEXT("Zoo"));
+		 const TCHAR* MyFooStr = *MyFoo;
+
+		 TStringBuilder<64> BuiltFoo;
+		 BuiltFoo.Append(TEXT("ABC"));
+
+		 DoFoo(MyFoo);
+		 DoFoo(MyFooStr);
+		 DoFoo(TEXT("ABC"));
+		 DoFoo(BuiltFoo);
+	}
+
+  */
+
+template<typename C>
+class TStringViewImpl
+{
+public:
+	typedef int32 SizeType;
+
+	inline TStringViewImpl(const C* InData)
+	:	DataPtr(InData)
+	,	Size(TCString<C>::Strlen(InData))
+	{
+	}
+
+	inline TStringViewImpl(const C* InData, SizeType InSize)
+	:	DataPtr(InData)
+	,	Size(InSize)
+	{
+	}
+
+	inline const C& operator[](SizeType Pos) const	{ return DataPtr[Pos]; }
+
+	inline const C* Data() const					{ return DataPtr; }
+
+	// Capacity
+
+	inline SizeType	Len() const						{ return Size; }
+	inline bool		IsEmpty() const					{ return Size == 0; }
+
+	// Modifiers
+
+	inline void		RemovePrefix(SizeType CharCount)	{ DataPtr += CharCount; Size -= CharCount; }
+	inline void		RemoveSuffix(SizeType CharCount)	{ Size -= CharCount; }
+
+	// Operations
+
+	inline SIZE_T			CopyString(C* Dest, SizeType CharCount, SizeType Position = 0)	{ memcpy(Dest, DataPtr + Position, FMath::Min(Size - Position, CharCount)); }
+	inline TStringViewImpl& SubStr(SizeType Position, SizeType CharCount)					{ return TStringViewImpl(DataPtr + Position, FMath::Min(Size - Position, CharCount)); }
+
+	// Comparison
+
+	// Compare
+
+	inline bool				StartsWith(const TStringViewImpl& Prefix) const				{ return 0 == TCString<C>::Strnicmp(Prefix.Data(), this->Data(), Prefix.Len()); }
+	inline bool				StartsWith(C Prefix) const									{ return Size >= 1 && DataPtr[0] == Prefix; }
+
+	// EndsWith
+
+	inline bool				EndsWith(const TStringViewImpl& Suffix) const 
+	{ 
+		const SIZE_T SuffixLen = Suffix.Len();
+		if (SuffixLen > this->Len())
+		{
+			return false;
+		}
+
+		return 0 == TCString<C>::Strnicmp(Suffix.Data(), this->Data() + this->Len() - SuffixLen, SuffixLen);
+	}
+
+	inline bool FindChar(C InChar, int32& OutIndex) const
+	{
+		for (SizeType i = 0; i < Size; ++i)
+		{
+			if (DataPtr[i] == InChar)
+			{
+				OutIndex = (int32) i;
+				return true;
+			}
+		}
+
+		return false;
+	}
+
+	inline bool FindLastChar(C InChar, int32& OutIndex) const
+	{
+		if (Size == 0)
+		{
+			return false;
+		}
+
+		for (SizeType i = Size - 1; i >= 0; --i)
+		{
+			if (DataPtr[i] == InChar)
+			{
+				OutIndex = (int32)i;
+				return true;
+			}
+		}
+
+		return false;
+	}
+
+	// FindLastOf
+	// FindFirstNotOf
+	// FindLastNotOf
+
+	inline const C* begin() const					{ return DataPtr; }
+	inline const C* end() const						{ return DataPtr + Size; }
+
+protected:
+	const C*	DataPtr;
+	SizeType	Size;
+};
+
+// Case-insensitive comparison operators
+
+template<typename C>
+inline bool operator==(const TStringViewImpl<C>& lhs, const TStringViewImpl<C>& rhs)
+{
+	if (rhs.Len() != lhs.Len())
+	{
+		return false;
+	}
+
+	return lhs.StartsWith(rhs);
+}
+
+template<typename C>
+inline bool operator!=(const TStringViewImpl<C>& lhs, const TStringViewImpl<C>& rhs)
+{
+	return !operator==(lhs, rhs);
+}
+
+//////////////////////////////////////////////////////////////////////////
+
+class FAnsiStringView : public TStringViewImpl<ANSICHAR>
+{
+public:
+	FAnsiStringView(const ANSICHAR* InString, SizeType InStrlen)
+	:	TStringViewImpl<ANSICHAR>(InString, InStrlen)
+	{
+	}
+
+	FAnsiStringView(const ANSICHAR* InString)
+	:	TStringViewImpl<ANSICHAR>(InString)
+	{
+	}
+};
+
+class FStringView : public TStringViewImpl<TCHAR>
+{
+public:
+	FStringView(const FString& InString)
+	:	TStringViewImpl<TCHAR>(*InString, InString.Len())
+	{
+	}
+
+	FStringView(const TCHAR* InString)
+	:	TStringViewImpl<TCHAR>(InString)
+	{
+	}
+
+	FStringView(const TCHAR* InString, SizeType InStrlen)
+	:	TStringViewImpl<TCHAR>(InString, InStrlen)
+	{
+	}
+
+	/** Convert to a dynamic string
+
+		Remember that a string view is not necessarily
+		NUL terminated so whenever you end up passing a
+		string into an API which accepts a plain C string
+		then you will need to ensure it's NUL terminated
+		first. This is one way of doing that, involving
+		a memory allocation as well as a string copy,
+		along with a memory deallocation when the string
+		goes away.
+
+		I would encourage you to use a TStringBuilder<>
+		instead since this avoids hitting the heap, and
+		the heap is always more expensive even if it
+		seems seductively convenient.
+	  */
+	CORE_API FString ToString() const;
+};
diff --git a/Engine/Source/Runtime/Core/Public/Serialization/Archive.h b/Engine/Source/Runtime/Core/Public/Serialization/Archive.h
--- a/Engine/Source/Runtime/Core/Public/Serialization/Archive.h
+++ b/Engine/Source/Runtime/Core/Public/Serialization/Archive.h
@@ -1566,7 +1566,7 @@ private:
 	* Stored as a pointer to a heap-allocated object because of a 3-way dependency between TArray, FCustomVersionContainer and FArchive, which is too much work to change right now.
 	* Keeping it as a heap-allocated object also helps with performance in some cases as we don't need to construct it for archives that don't care about custom versions.
 	*/
-	mutable FCustomVersionContainer* CustomVersionContainer;
+	mutable FCustomVersionContainer* CustomVersionContainer = nullptr;
 
 public:
 	/** Custom property list attribute. If the flag below is set, only these properties will be iterated during serialization. If NULL, then no properties will be iterated. */
diff --git a/Engine/Source/Runtime/Core/Public/Templates/RefCounting.h b/Engine/Source/Runtime/Core/Public/Templates/RefCounting.h
--- a/Engine/Source/Runtime/Core/Public/Templates/RefCounting.h
+++ b/Engine/Source/Runtime/Core/Public/Templates/RefCounting.h
@@ -5,6 +5,7 @@
 #include "CoreTypes.h"
 #include "Misc/AssertionMacros.h"
 #include "Serialization/Archive.h"
+#include "HAL/PlatformAtomics.h"
 
 /** A virtual interface for ref counted objects to implement. */
 class IRefCountedObject
@@ -16,9 +17,55 @@ public:
 	virtual uint32 GetRefCount() const = 0;
 };
 
+/**
+ * Base class implementing thread-safe reference counting.
+ */
+class FRefCountBase
+{
+public:
+			FRefCountBase() = default;
+	virtual ~FRefCountBase() = default;
+
+	FRefCountBase(const FRefCountBase& Rhs) = delete;
+	FRefCountBase& operator=(const FRefCountBase& Rhs) = delete;
+
+	inline uint32 AddRef() const
+	{
+		return uint32(FPlatformAtomics::InterlockedIncrement(&NumRefs));
+	}
+
+	inline uint32 Release() const
+	{
+#if DO_CHECK
+		CheckRefCount();
+#endif
+
+		const int32 Refs = FPlatformAtomics::InterlockedDecrement(&NumRefs);
+		if (Refs == 0)
+		{
+			delete this;
+		}
+
+		return uint32(Refs);
+	}
+
+	uint32 GetRefCount() const
+	{
+		return uint32(NumRefs);
+	}
+
+private:
+	mutable int32 NumRefs = 0;
+
+	CORE_API void CheckRefCount() const;
+};
 
 /**
  * The base class of reference counted objects.
+ *
+ * This class should not be used for new code as it does not use atomic operations to update 
+ * the reference count.
+ *
  */
 class CORE_API FRefCountedObject
 {
diff --git a/Engine/Source/Runtime/Core/Public/Templates/UnrealTemplate.h b/Engine/Source/Runtime/Core/Public/Templates/UnrealTemplate.h
--- a/Engine/Source/Runtime/Core/Public/Templates/UnrealTemplate.h
+++ b/Engine/Source/Runtime/Core/Public/Templates/UnrealTemplate.h
@@ -75,14 +75,13 @@ auto GetData(T&& Container) -> decltype(Container.GetData())
 	return Container.GetData();
 }
 
-template <typename T, SIZE_T N>
-CONSTEXPR T* GetData(T (&Container)[N])
-{
-	return Container;
-}
+template <typename T, SIZE_T N> CONSTEXPR       T* GetData(      T (& Container)[N]) { return Container; }
+template <typename T, SIZE_T N> CONSTEXPR       T* GetData(      T (&&Container)[N]) { return Container; }
+template <typename T, SIZE_T N> CONSTEXPR const T* GetData(const T (& Container)[N]) { return Container; }
+template <typename T, SIZE_T N> CONSTEXPR const T* GetData(const T (&&Container)[N]) { return Container; }
 
 template <typename T>
-CONSTEXPR T* GetData(std::initializer_list<T> List)
+CONSTEXPR const T* GetData(std::initializer_list<T> List)
 {
 	return List.begin();
 }
@@ -96,11 +95,10 @@ SIZE_T GetNum(T&& Container)
 	return (SIZE_T)Container.Num();
 }
 
-template <typename T, SIZE_T N>
-CONSTEXPR SIZE_T GetNum(T (&Container)[N])
-{
-	return N;
-}
+template <typename T, SIZE_T N> CONSTEXPR SIZE_T GetNum(      T (& Container)[N]) { return N; }
+template <typename T, SIZE_T N> CONSTEXPR SIZE_T GetNum(      T (&&Container)[N]) { return N; }
+template <typename T, SIZE_T N> CONSTEXPR SIZE_T GetNum(const T (& Container)[N]) { return N; }
+template <typename T, SIZE_T N> CONSTEXPR SIZE_T GetNum(const T (&&Container)[N]) { return N; }
 
 template <typename T>
 CONSTEXPR SIZE_T GetNum(std::initializer_list<T> List)
@@ -131,11 +129,11 @@ FORCEINLINE const T& AsConst(T& Ref)
 ----------------------------------------------------------------------------*/
 
 #ifdef __clang__
-template <typename T>
-auto UE4ArrayCountHelper(T& t) -> typename TEnableIf<__is_array(T), char(&)[sizeof(t) / sizeof(t[0]) + 1]>::Type;
+	template <typename T>
+	auto UE4ArrayCountHelper(T& t) -> typename TEnableIf<__is_array(T), char(&)[sizeof(t) / sizeof(t[0]) + 1]>::Type;
 #else
-template <typename T, uint32 N>
-char(&UE4ArrayCountHelper(const T(&)[N]))[N + 1];
+	template <typename T, uint32 N>
+	char (&UE4ArrayCountHelper(const T (&)[N]))[N + 1];
 #endif
 
 // Number of elements in an array.
diff --git a/Engine/Source/Runtime/Core/Public/Windows/WindowsPlatformCompilerSetup.h b/Engine/Source/Runtime/Core/Public/Windows/WindowsPlatformCompilerSetup.h
--- a/Engine/Source/Runtime/Core/Public/Windows/WindowsPlatformCompilerSetup.h
+++ b/Engine/Source/Runtime/Core/Public/Windows/WindowsPlatformCompilerSetup.h
@@ -35,7 +35,7 @@ static_assert(_MSC_VER != 1914 && _MSC_VER != 1915, "Visual Studio 2017 versions
 #pragma warning (error: 4600      4602 4603 4604      4606           4609 4610 4611 4612 4613      4615 4616      4618      4620 4621 4622 4623 4624 4625 4626 4627 4628 4629 4630 4631 4632 4633 4634 4635 4636 4637 4638 4639 4640 4641 4642      4644 4645 4646      4648 4649 4650      4652 4653 4654 4655 4656 4657 4658 4659      4661 4662                     4667      4669 4670 4671 4672 4673 4674      4676 4677 4678 4679 4680 4681 4682 4683 4684 4685 4686 4687 4688 4689 4690 4691      4693 4694 4695 4696      4698     )
 #pragma warning (error: 4700      4702 4703           4706           4709 4710 4711           4714 4715 4716 4717 4718 4719 4720 4721 4722 4723 4724 4725 4726 4727 4728 4729      4731 4732 4733 4734                     4739 4740 4741 4742 4743 4744 4745 4746 4747      4749 4750 4751 4752 4753 4754 4755 4756 4757                4761           4764                               4771 4772           4775 4776 4777 4778                                    4786      4788 4789           4792 4793 4794                4798 4799)
 #pragma warning (error: 4800 4801      4803 4804 4805 4806 4807 4808 4809 4810 4811 4812 4813      4815 4816 4817                4821 4822 4823                4827      4829                     4834 4835                4839 4840 4841 4842 4843 4844 4845 4846 4847 4848 4849                                                                                      4867      4869           4872                                    4880 4881 4882 4883                                                                                )
-#pragma warning (error: 4900                     4905 4906                4910      4912 4913           4916 4917 4918      4920 4921                4925 4926 4927 4928 4929 4930 4931 4932      4934 4935 4936 4937 4938 4939                     4944 4945 4946 4947 4948 4949 4950 4951 4952 4953 4954 4955 4956 4957 4958 4959 4960 4961      4963 4964 4965 4966                4970 4971 4972 4973 4974                               4981           4984 4985                4989 4990 4991 4992           4995      4997 4998 4999)
+#pragma warning (error: 4900                     4905 4906                4910      4912 4913           4916 4917 4918      4920 4921                4925 4926 4927 4928 4929 4930 4931 4932      4934 4935 4936 4937 4938 4939                     4944 4945 4946 4947 4948 4949 4950 4951 4952 4953 4954 4955 4956 4957 4958 4959 4960 4961      4963 4964 4965 4966                4970 4971 4972 4973 4974                               4981                4985                4989 4990 4991 4992           4995      4997 4998 4999)
 #pragma warning (error:                                                                                                                                                                                               5038                                                                                                                                                                                                                                                                                                                 )
 
 //
@@ -237,6 +237,7 @@ static_assert(_MSC_VER != 1914 && _MSC_VER != 1915, "Visual Studio 2017 versions
 #pragma warning(default: 4931) // we are assuming the type library was built for number-bit pointers								https://docs.microsoft.com/en-us/cpp/error-messages/compiler-warnings/compiler-warning-level-4-c4931
 #pragma warning(default: 4946) // reinterpret_cast used between related classes: 'class1' and 'class2'								https://docs.microsoft.com/en-us/cpp/error-messages/compiler-warnings/compiler-warning-level-1-c4946
 #pragma warning(default: 5038) // data member 'A::y' will be initialized after data member 'A::x'									https://docs.microsoft.com/en-us/cpp/error-messages/compiler-warnings/c5038
+#pragma warning(disable: 4984) // 'if constexpr' is a C++17 language extension
 
 #pragma warning(default: 4996)
 
