Fix InlineAllocation not defined if TFUNCTION_USES_INLINE_STORAGE is 0
#rnx

diff --git a/Engine/Source/Runtime/Core/Public/Templates/Function.h b/Engine/Source/Runtime/Core/Public/Templates/Function.h
--- a/Engine/Source/Runtime/Core/Public/Templates/Function.h
+++ b/Engine/Source/Runtime/Core/Public/Templates/Function.h
@@ -412,11 +412,13 @@ namespace UE4Function_Private
 			using OwnedType = TStorageOwnerTypeT<FunctorType, bUnique, !bUseInline>;
 
 			void* NewAlloc;
+#if TFUNCTION_USES_INLINE_STORAGE
 			if constexpr (bUseInline)
 			{
 				NewAlloc = &InlineAllocation;
 			}
 			else
+#endif
 			{
 				NewAlloc = FMemory::Malloc(sizeof(OwnedType), alignof(OwnedType));
 				HeapAllocation = NewAlloc;
