Rearranged members to hopefully work around issue which I could not repro (could be that some #include site has a custom pragma pack leak)

diff --git a/Engine/Source/Runtime/Core/Public/ProfilingDebugging/ScopedTimers.h b/Engine/Source/Runtime/Core/Public/ProfilingDebugging/ScopedTimers.h
--- a/Engine/Source/Runtime/Core/Public/ProfilingDebugging/ScopedTimers.h
+++ b/Engine/Source/Runtime/Core/Public/ProfilingDebugging/ScopedTimers.h
@@ -151,9 +151,9 @@ class FScopedSwitchedDurationTimer
 {
 public:
 	explicit FScopedSwitchedDurationTimer(double& AccumulatorIn, bool bDoFunctionalityIn)
-		: bDoFunctionality(bDoFunctionalityIn)
-		, StartTime(bDoFunctionalityIn ? FPlatformTime::Seconds() : 0)
-		, Accumulator(AccumulatorIn)		
+		: StartTime(bDoFunctionalityIn ? FPlatformTime::Seconds() : 0)
+		, Accumulator(AccumulatorIn)
+		, bDoFunctionality(bDoFunctionalityIn)
 	{
 	}
 
@@ -172,11 +172,11 @@ public:
 	}
 
 protected:
-	const bool bDoFunctionality;
 	/** Start time, captured in ctor. */
 	double StartTime;
 	/** Time variable to update. */
 	double& Accumulator;
+	const bool bDoFunctionality;
 };
 
 /**
