Probable fix for HoloLens compile error.

diff --git a/Engine/Source/Runtime/Core/Public/MSVC/MSVCPlatform.h b/Engine/Source/Runtime/Core/Public/MSVC/MSVCPlatform.h
--- a/Engine/Source/Runtime/Core/Public/MSVC/MSVCPlatform.h
+++ b/Engine/Source/Runtime/Core/Public/MSVC/MSVCPlatform.h
@@ -6,7 +6,7 @@
 
 #pragma once
 
-#if _MSC_FULL_VER >= 191125507 && defined(_MSVC_LANG) && _MSVC_LANG >= 201402
+#if _MSC_FULL_VER >= 191125507 && defined(__cpp_if_constexpr)
 	#define PLATFORM_COMPILER_HAS_IF_CONSTEXPR 1
 #else
 	#define PLATFORM_COMPILER_HAS_IF_CONSTEXPR 0
