Added macros to push/pop the default platform packing

This is needed for some Third-Party Software that requires the default packing (as that's how it was compiled) and has structures affected by the change in packing size.

#rnx

diff --git a/Engine/Source/Runtime/Core/Private/Internationalization/ICUInternationalization.cpp b/Engine/Source/Runtime/Core/Private/Internationalization/ICUInternationalization.cpp
--- a/Engine/Source/Runtime/Core/Private/Internationalization/ICUInternationalization.cpp
+++ b/Engine/Source/Runtime/Core/Private/Internationalization/ICUInternationalization.cpp
@@ -20,7 +20,9 @@
 #include "Internationalization/ICUBreakIterator.h"
 THIRD_PARTY_INCLUDES_START
 	#include <unicode/locid.h>
-	#include <unicode/timezone.h>
+PRAGMA_PUSH_PLATFORM_DEFAULT_PACKING
+	#include <unicode/timezone.h> // icu::Calendar can be affected by the non-standard packing UE4 uses, so force the platform default
+PRAGMA_POP_PLATFORM_DEFAULT_PACKING
 	#include <unicode/uclean.h>
 	#include <unicode/udata.h>
 THIRD_PARTY_INCLUDES_END
diff --git a/Engine/Source/Runtime/Core/Private/Internationalization/ICUInternationalization.h b/Engine/Source/Runtime/Core/Private/Internationalization/ICUInternationalization.h
--- a/Engine/Source/Runtime/Core/Private/Internationalization/ICUInternationalization.h
+++ b/Engine/Source/Runtime/Core/Private/Internationalization/ICUInternationalization.h
@@ -15,7 +15,9 @@ class FInternationalization;
 
 THIRD_PARTY_INCLUDES_START
 	#include <unicode/umachine.h>
-	#include <unicode/gregocal.h>
+PRAGMA_PUSH_PLATFORM_DEFAULT_PACKING
+	#include <unicode/gregocal.h> // icu::Calendar can be affected by the non-standard packing UE4 uses, so force the platform default
+PRAGMA_POP_PLATFORM_DEFAULT_PACKING
 THIRD_PARTY_INCLUDES_END
 
 // This should be defined by ICU.build.cs
diff --git a/Engine/Source/Runtime/Core/Public/Clang/ClangPlatformCompilerPreSetup.h b/Engine/Source/Runtime/Core/Public/Clang/ClangPlatformCompilerPreSetup.h
--- a/Engine/Source/Runtime/Core/Public/Clang/ClangPlatformCompilerPreSetup.h
+++ b/Engine/Source/Runtime/Core/Public/Clang/ClangPlatformCompilerPreSetup.h
@@ -122,6 +122,14 @@
 		_Pragma("clang diagnostic pop")
 #endif // PRAGMA_POP
 
+#ifndef PRAGMA_PUSH_PLATFORM_DEFAULT_PACKING
+	#define PRAGMA_PUSH_PLATFORM_DEFAULT_PACKING
+#endif // PRAGMA_PUSH_PLATFORM_DEFAULT_PACKING
+	
+#ifndef PRAGMA_POP_PLATFORM_DEFAULT_PACKING
+	#define PRAGMA_POP_PLATFORM_DEFAULT_PACKING
+#endif // PRAGMA_POP_PLATFORM_DEFAULT_PACKING
+
 #ifndef EMIT_CUSTOM_WARNING_AT_LINE
 	#define EMIT_CUSTOM_WARNING_AT_LINE(Line, Warning) \
 		_Pragma(PREPROCESSOR_TO_STRING(message(Warning)))
diff --git a/Engine/Source/Runtime/Core/Public/HoloLens/HoloLensPlatformCompilerPreSetup.h b/Engine/Source/Runtime/Core/Public/HoloLens/HoloLensPlatformCompilerPreSetup.h
--- a/Engine/Source/Runtime/Core/Public/HoloLens/HoloLensPlatformCompilerPreSetup.h
+++ b/Engine/Source/Runtime/Core/Public/HoloLens/HoloLensPlatformCompilerPreSetup.h
@@ -67,6 +67,27 @@
 		__pragma(warning(pop))
 #endif // PRAGMA_POP
 
+// UE4 uses a struct packing of 4 for Win32 and 8 for Win64, 
+// and the default packing is 8 for Win32 and 16 for Win64
+#ifdef _WIN64
+	#ifndef PRAGMA_PUSH_PLATFORM_DEFAULT_PACKING
+		#define PRAGMA_PUSH_PLATFORM_DEFAULT_PACKING \
+			__pragma(pack(push)) \
+			__pragma(pack(16))
+	#endif // PRAGMA_PUSH_PLATFORM_DEFAULT_PACKING
+#else // _WIN64
+	#ifndef PRAGMA_PUSH_PLATFORM_DEFAULT_PACKING
+		#define PRAGMA_PUSH_PLATFORM_DEFAULT_PACKING \
+			__pragma(pack(push)) \
+			__pragma(pack(8))
+	#endif // PRAGMA_PUSH_PLATFORM_DEFAULT_PACKING
+#endif // _WIN64
+	
+#ifndef PRAGMA_POP_PLATFORM_DEFAULT_PACKING
+	#define PRAGMA_POP_PLATFORM_DEFAULT_PACKING \
+		__pragma(pack(pop))
+#endif // PRAGMA_POP_PLATFORM_DEFAULT_PACKING
+
 // Disable common CA warnings around SDK includes
 #ifndef THIRD_PARTY_INCLUDES_START
 	#define THIRD_PARTY_INCLUDES_START \
diff --git a/Engine/Source/Runtime/Core/Public/Windows/MinimalWindowsApi.h b/Engine/Source/Runtime/Core/Public/Windows/MinimalWindowsApi.h
--- a/Engine/Source/Runtime/Core/Public/Windows/MinimalWindowsApi.h
+++ b/Engine/Source/Runtime/Core/Public/Windows/MinimalWindowsApi.h
@@ -18,11 +18,8 @@
 	#define MINIMAL_WINDOWS_API extern "C" __declspec(dllimport)
 #endif
 
-// Set up the default struct packing on Win32
-#if PLATFORM_WINDOWS && PLATFORM_32BITS
-	#pragma pack(push)
-	#pragma pack(8)
-#endif
+// The 10.0.18362.0 SDK introduces an error if the packing isn't the default for the platform.
+PRAGMA_PUSH_PLATFORM_DEFAULT_PACKING
 
 // Undefine Windows types that we're going to redefine. This should be done by HideWindowsPlatformTypes.h after including any system headers, 
 // but it's less friction to just undefine them here.
@@ -211,6 +208,4 @@ namespace Windows
 #pragma pop_macro("TRUE")
 
 // Restore the packing setting
-#if PLATFORM_WINDOWS && PLATFORM_32BITS
-	#pragma pack(pop)
-#endif
+PRAGMA_POP_PLATFORM_DEFAULT_PACKING
diff --git a/Engine/Source/Runtime/Core/Public/Windows/PostWindowsApi.h b/Engine/Source/Runtime/Core/Public/Windows/PostWindowsApi.h
--- a/Engine/Source/Runtime/Core/Public/Windows/PostWindowsApi.h
+++ b/Engine/Source/Runtime/Core/Public/Windows/PostWindowsApi.h
@@ -76,9 +76,7 @@ THIRD_PARTY_INCLUDES_END
 #pragma pop_macro("FALSE")
 
 // Restore the struct packing setting
-#if PLATFORM_WINDOWS && PLATFORM_32BITS
-	#pragma pack(pop)
-#endif
+PRAGMA_POP_PLATFORM_DEFAULT_PACKING
 
 // Redefine CDECL to our version of the #define.  <AJS> Is this really necessary?
 #define CDECL	    __cdecl					/* Standard C function */
diff --git a/Engine/Source/Runtime/Core/Public/Windows/PreWindowsApi.h b/Engine/Source/Runtime/Core/Public/Windows/PreWindowsApi.h
--- a/Engine/Source/Runtime/Core/Public/Windows/PreWindowsApi.h
+++ b/Engine/Source/Runtime/Core/Public/Windows/PreWindowsApi.h
@@ -4,10 +4,7 @@
 #pragma warning(disable:4103)
 
 // The 10.0.18362.0 SDK introduces an error if the packing isn't the default for the platform.
-#if PLATFORM_WINDOWS && PLATFORM_32BITS
-	#pragma pack(push)
-	#pragma pack(8)
-#endif
+PRAGMA_PUSH_PLATFORM_DEFAULT_PACKING
 
 // Save these macros for later; Windows redefines them
 #pragma push_macro("MAX_uint8")
diff --git a/Engine/Source/Runtime/Core/Public/Windows/WindowsPlatformCompilerPreSetup.h b/Engine/Source/Runtime/Core/Public/Windows/WindowsPlatformCompilerPreSetup.h
--- a/Engine/Source/Runtime/Core/Public/Windows/WindowsPlatformCompilerPreSetup.h
+++ b/Engine/Source/Runtime/Core/Public/Windows/WindowsPlatformCompilerPreSetup.h
@@ -3,6 +3,27 @@
 #pragma once
 
 #if defined(__clang__)
+	// UE4 uses a struct packing of 4 for Win32 and 8 for Win64, 
+	// and the default packing is 8 for Win32 and 16 for Win64
+	#ifdef _WIN64
+		#ifndef PRAGMA_PUSH_PLATFORM_DEFAULT_PACKING
+			#define PRAGMA_PUSH_PLATFORM_DEFAULT_PACKING \
+				_Pragma("pack(push)") \
+				_Pragma("pack(16)")
+		#endif // PRAGMA_PUSH_PLATFORM_DEFAULT_PACKING
+	#else // _WIN64
+		#ifndef PRAGMA_PUSH_PLATFORM_DEFAULT_PACKING
+			#define PRAGMA_PUSH_PLATFORM_DEFAULT_PACKING \
+				_Pragma("pack(push)") \
+				_Pragma("pack(8)")
+		#endif // PRAGMA_PUSH_PLATFORM_DEFAULT_PACKING
+	#endif // _WIN64
+	
+	#ifndef PRAGMA_POP_PLATFORM_DEFAULT_PACKING
+		#define PRAGMA_POP_PLATFORM_DEFAULT_PACKING \
+			_Pragma("pack(pop)")
+	#endif // PRAGMA_POP_PLATFORM_DEFAULT_PACKING
+
 	#include "Clang/ClangPlatformCompilerPreSetup.h"
 
 	// Disable common CA warnings around SDK includes
@@ -91,6 +112,27 @@
 			__pragma(warning(pop))
 	#endif // PRAGMA_POP
 
+	// UE4 uses a struct packing of 4 for Win32 and 8 for Win64, 
+	// and the default packing is 8 for Win32 and 16 for Win64
+	#ifdef _WIN64
+		#ifndef PRAGMA_PUSH_PLATFORM_DEFAULT_PACKING
+			#define PRAGMA_PUSH_PLATFORM_DEFAULT_PACKING \
+				__pragma(pack(push)) \
+				__pragma(pack(16))
+		#endif // PRAGMA_PUSH_PLATFORM_DEFAULT_PACKING
+	#else // _WIN64
+		#ifndef PRAGMA_PUSH_PLATFORM_DEFAULT_PACKING
+			#define PRAGMA_PUSH_PLATFORM_DEFAULT_PACKING \
+				__pragma(pack(push)) \
+				__pragma(pack(8))
+		#endif // PRAGMA_PUSH_PLATFORM_DEFAULT_PACKING
+	#endif // _WIN64
+	
+	#ifndef PRAGMA_POP_PLATFORM_DEFAULT_PACKING
+		#define PRAGMA_POP_PLATFORM_DEFAULT_PACKING \
+			__pragma(pack(pop))
+	#endif // PRAGMA_POP_PLATFORM_DEFAULT_PACKING
+
 	// Disable common CA warnings around SDK includes
 	#ifndef THIRD_PARTY_INCLUDES_START
 		#define THIRD_PARTY_INCLUDES_START \
