Fix serialization with new spatial idx

diff --git a/Engine/Source/Runtime/Core/Public/UObject/ExternalPhysicsCustomObjectVersion.h b/Engine/Source/Runtime/Core/Public/UObject/ExternalPhysicsCustomObjectVersion.h
--- a/Engine/Source/Runtime/Core/Public/UObject/ExternalPhysicsCustomObjectVersion.h
+++ b/Engine/Source/Runtime/Core/Public/UObject/ExternalPhysicsCustomObjectVersion.h
@@ -27,6 +27,9 @@ struct CORE_API FExternalPhysicsCustomObjectVersion
 		//Global elements have bounds so we can skip them
 		GlobalElementsHaveBounds,
 
+		//SpatialIdx serialized
+		SpatialIdxSerialized,
+
 		// -----<new versions can be added above this line>-------------------------------------------------
 		VersionPlusOne,
 		LatestVersion = VersionPlusOne - 1
diff --git a/Engine/Source/Runtime/Experimental/Chaos/Private/Chaos/PBDRigidsEvolution.cpp b/Engine/Source/Runtime/Experimental/Chaos/Private/Chaos/PBDRigidsEvolution.cpp
--- a/Engine/Source/Runtime/Experimental/Chaos/Private/Chaos/PBDRigidsEvolution.cpp
+++ b/Engine/Source/Runtime/Experimental/Chaos/Private/Chaos/PBDRigidsEvolution.cpp
@@ -391,6 +391,7 @@ namespace Chaos
 			AccelerationStructureTaskComplete = nullptr;
 			for (auto& Particle : Particles.GetNonDisabledView())
 			{
+				Particle.SetSpatialIdx(FSpatialAccelerationIdx{ 0,0 });
 				DirtyParticle(Particle);
 			}
 
diff --git a/Engine/Source/Runtime/Experimental/Chaos/Public/Chaos/GeometryParticles.h b/Engine/Source/Runtime/Experimental/Chaos/Public/Chaos/GeometryParticles.h
--- a/Engine/Source/Runtime/Experimental/Chaos/Public/Chaos/GeometryParticles.h
+++ b/Engine/Source/Runtime/Experimental/Chaos/Public/Chaos/GeometryParticles.h
@@ -303,6 +303,15 @@ namespace Chaos
 					}
 				}
 			}
+
+			if (Ar.CustomVer(FExternalPhysicsCustomObjectVersion::GUID) < FExternalPhysicsCustomObjectVersion::SpatialIdxSerialized)
+			{
+				MSpatialIdx.AddZeroed(MGeometry.Num());
+			}
+			else
+			{
+				Ar << MSpatialIdx;
+			}
 		}
 
 		CHAOS_API EParticleType ParticleType() const { return MParticleType; }
diff --git a/Engine/Source/Runtime/Experimental/Chaos/Public/Chaos/GeometryParticlesfwd.h b/Engine/Source/Runtime/Experimental/Chaos/Public/Chaos/GeometryParticlesfwd.h
--- a/Engine/Source/Runtime/Experimental/Chaos/Public/Chaos/GeometryParticlesfwd.h
+++ b/Engine/Source/Runtime/Experimental/Chaos/Public/Chaos/GeometryParticlesfwd.h
@@ -33,4 +33,9 @@ namespace Chaos
 			return ((const uint16&)*this) == ((const uint16&)Rhs);
 		}
 	};
+
+	inline FArchive& operator<<(FArchive& Ar, FSpatialAccelerationIdx& Idx)
+	{
+		return Ar << (uint16&)Idx;
+	}
 }
\ No newline at end of file
diff --git a/Engine/Source/Runtime/Experimental/Chaos/Public/Chaos/PBDRigidsEvolution.h b/Engine/Source/Runtime/Experimental/Chaos/Public/Chaos/PBDRigidsEvolution.h
--- a/Engine/Source/Runtime/Experimental/Chaos/Public/Chaos/PBDRigidsEvolution.h
+++ b/Engine/Source/Runtime/Experimental/Chaos/Public/Chaos/PBDRigidsEvolution.h
@@ -389,6 +389,9 @@ protected:
 			Ar << AccelerationHandle;
 			Ar << bUpdate;
 			Ar << bDelete;
+
+			Ar << UpdatedSpatialIdx;
+			Ar << DeletedSpatialIdx;
 		}
 	};
 
