FAndroidMisc::FileExistsInPlatformPackage() got lost in the engine merge - restore it.  Also fix spelling.
This should fix Android Test/Ship builds not being able to find BuildInfo.ini

diff --git a/Engine/Source/Runtime/Core/Private/Android/AndroidPlatformMisc.cpp b/Engine/Source/Runtime/Core/Private/Android/AndroidPlatformMisc.cpp
--- a/Engine/Source/Runtime/Core/Private/Android/AndroidPlatformMisc.cpp
+++ b/Engine/Source/Runtime/Core/Private/Android/AndroidPlatformMisc.cpp
@@ -1135,6 +1135,20 @@ FString FAndroidMisc::LoadTextFileFromPlatformPackage(const FString& RelativePat
 	return FString();
 }
 
+bool FAndroidMisc::FileExistsInPlatformPackage(const FString& RelativePath)
+{
+#if USE_ANDROID_JNI
+	AAssetManager* AssetMgr = AndroidThunkCpp_GetAssetManager();
+	AAsset* asset = AAssetManager_open(AssetMgr, TCHAR_TO_UTF8(*RelativePath), AASSET_MODE_UNKNOWN);
+	if (asset)
+	{
+		AAsset_close(asset);
+		return true;
+	}
+#endif
+	return false;
+}
+
 void FAndroidMisc::SetVersionInfo( FString InAndroidVersion, FString InDeviceMake, FString InDeviceModel, FString InDeviceBuildNumber, FString InOSLanguage )
 {
 	AndroidVersion = InAndroidVersion;
diff --git a/Engine/Source/Runtime/Core/Private/GenericPlatform/GenericPlatformMisc.cpp b/Engine/Source/Runtime/Core/Private/GenericPlatform/GenericPlatformMisc.cpp
--- a/Engine/Source/Runtime/Core/Private/GenericPlatform/GenericPlatformMisc.cpp
+++ b/Engine/Source/Runtime/Core/Private/GenericPlatform/GenericPlatformMisc.cpp
@@ -1372,7 +1372,7 @@ FString FGenericPlatformMisc::LoadTextFileFromPlatformPackage(const FString& Rel
 	return Result;
 }
 
-bool FGenericPlatformMisc::FileExitsInPlatformPackage(const FString& RelativePath)
+bool FGenericPlatformMisc::FileExistsInPlatformPackage(const FString& RelativePath)
 {
 	FString Path = RootDir() / RelativePath;
 	return IFileManager::Get().FileExists(*Path);
diff --git a/Engine/Source/Runtime/Core/Private/IOS/IOSPlatformMisc.cpp b/Engine/Source/Runtime/Core/Private/IOS/IOSPlatformMisc.cpp
--- a/Engine/Source/Runtime/Core/Private/IOS/IOSPlatformMisc.cpp
+++ b/Engine/Source/Runtime/Core/Private/IOS/IOSPlatformMisc.cpp
@@ -1115,7 +1115,7 @@ FString FIOSPlatformMisc::LoadTextFileFromPlatformPackage(const FString& Relativ
 	return FString(UTF8_TO_TCHAR(FileContents.GetData()));
 }
 
-bool FIOSPlatformMisc::FileExitsInPlatformPackage(const FString& RelativePath)
+bool FIOSPlatformMisc::FileExistsInPlatformPackage(const FString& RelativePath)
 {
 	FString FilePath = FString([[NSBundle mainBundle] bundlePath]) / RelativePath;
 
diff --git a/Engine/Source/Runtime/Core/Public/Android/AndroidPlatformMisc.h b/Engine/Source/Runtime/Core/Public/Android/AndroidPlatformMisc.h
--- a/Engine/Source/Runtime/Core/Public/Android/AndroidPlatformMisc.h
+++ b/Engine/Source/Runtime/Core/Public/Android/AndroidPlatformMisc.h
@@ -160,6 +160,7 @@ public:
 	static void ShareURL(const FString& URL, const FText& Description, int32 LocationHintX, int32 LocationHintY);
 
 	static FString LoadTextFileFromPlatformPackage(const FString& RelativePath);
+	static bool FileExistsInPlatformPackage(const FString& RelativePath);
 
 	// ANDROID ONLY:
 	static void SetVersionInfo(FString AndroidVersion, FString DeviceMake, FString DeviceModel, FString DeviceBuildNumber, FString OSLanguage);
diff --git a/Engine/Source/Runtime/Core/Public/GenericPlatform/GenericPlatformMisc.h b/Engine/Source/Runtime/Core/Public/GenericPlatform/GenericPlatformMisc.h
--- a/Engine/Source/Runtime/Core/Public/GenericPlatform/GenericPlatformMisc.h
+++ b/Engine/Source/Runtime/Core/Public/GenericPlatform/GenericPlatformMisc.h
@@ -1381,7 +1381,7 @@ public:
 	*/
 	static FString LoadTextFileFromPlatformPackage(const FString& RelativePath);
 
-	static bool FileExitsInPlatformPackage(const FString& RelativePath);
+	static bool FileExistsInPlatformPackage(const FString& RelativePath);
 
 	/**
 	 * Frees any memory retained by FGenericPlatformMisc.
diff --git a/Engine/Source/Runtime/Core/Public/IOS/IOSPlatformMisc.h b/Engine/Source/Runtime/Core/Public/IOS/IOSPlatformMisc.h
--- a/Engine/Source/Runtime/Core/Public/IOS/IOSPlatformMisc.h
+++ b/Engine/Source/Runtime/Core/Public/IOS/IOSPlatformMisc.h
@@ -79,7 +79,7 @@ struct CORE_API FIOSPlatformMisc : public FApplePlatformMisc
 	static void ShareURL(const FString& URL, const FText& Description, int32 LocationHintX, int32 LocationHintY);
 
 	static FString LoadTextFileFromPlatformPackage(const FString& RelativePath);
-	static bool FileExitsInPlatformPackage(const FString& RelativePath);
+	static bool FileExistsInPlatformPackage(const FString& RelativePath);
 
 	static void EnableVoiceChat(bool bEnable);
 	static bool IsVoiceChatEnabled();
