Fixed CIS error for Clang on Windows

#rnx

diff --git a/Engine/Source/Runtime/Core/Public/Windows/PostWindowsApi.h b/Engine/Source/Runtime/Core/Public/Windows/PostWindowsApi.h
--- a/Engine/Source/Runtime/Core/Public/Windows/PostWindowsApi.h
+++ b/Engine/Source/Runtime/Core/Public/Windows/PostWindowsApi.h
@@ -78,6 +78,11 @@ THIRD_PARTY_INCLUDES_END
 // Restore the struct packing setting
 PRAGMA_POP_PLATFORM_DEFAULT_PACKING
 
+// Restore the warning that the pack size is changed in this header.
+#ifdef __clang__
+	#pragma clang diagnostic pop
+#endif	// __clang__
+
 // Redefine CDECL to our version of the #define.  <AJS> Is this really necessary?
 #define CDECL	    __cdecl					/* Standard C function */
 
diff --git a/Engine/Source/Runtime/Core/Public/Windows/PreWindowsApi.h b/Engine/Source/Runtime/Core/Public/Windows/PreWindowsApi.h
--- a/Engine/Source/Runtime/Core/Public/Windows/PreWindowsApi.h
+++ b/Engine/Source/Runtime/Core/Public/Windows/PreWindowsApi.h
@@ -1,7 +1,12 @@
 // Copyright 1998-2019 Epic Games, Inc. All Rights Reserved.
 
 // Disable the warning that the pack size is changed in this header.
-#pragma warning(disable:4103)
+#ifdef __clang__
+	#pragma clang diagnostic push
+	#pragma clang diagnostic ignored "-Wpragma-pack"
+#else	// __clang__
+	#pragma warning(disable:4103)
+#endif	// __clang__
 
 // The 10.0.18362.0 SDK introduces an error if the packing isn't the default for the platform.
 PRAGMA_PUSH_PLATFORM_DEFAULT_PACKING
