Add a CSV_PROFILER_USE_CUSTOM_FRAME_TIMINGS preprocesser macro to allow a game to call FCsvProfiler::BeginFrame and EndFrame instead of using the engine defaults.

diff --git a/Engine/Source/Runtime/Core/Private/ProfilingDebugging/CsvProfiler.cpp b/Engine/Source/Runtime/Core/Private/ProfilingDebugging/CsvProfiler.cpp
--- a/Engine/Source/Runtime/Core/Private/ProfilingDebugging/CsvProfiler.cpp
+++ b/Engine/Source/Runtime/Core/Private/ProfilingDebugging/CsvProfiler.cpp
@@ -2450,10 +2450,12 @@ FCsvProfiler::FCsvProfiler()
 {
 	check(IsInGameThread());
 
+#if !CSV_PROFILER_USE_CUSTOM_FRAME_TIMINGS
 	FCoreDelegates::OnBeginFrame.AddStatic(CsvProfilerBeginFrame);
 	FCoreDelegates::OnEndFrame.AddStatic(CsvProfilerEndFrame);
 	FCoreDelegates::OnBeginFrameRT.AddStatic(CsvProfilerBeginFrameRT);
 	FCoreDelegates::OnEndFrameRT.AddStatic(CsvProfilerEndFrameRT);
+#endif
 }
 
 FCsvProfiler::~FCsvProfiler()
diff --git a/Engine/Source/Runtime/Core/Public/ProfilingDebugging/CsvProfiler.h b/Engine/Source/Runtime/Core/Public/ProfilingDebugging/CsvProfiler.h
--- a/Engine/Source/Runtime/Core/Public/ProfilingDebugging/CsvProfiler.h
+++ b/Engine/Source/Runtime/Core/Public/ProfilingDebugging/CsvProfiler.h
@@ -29,6 +29,10 @@
 #define CSV_PROFILER_ALLOW_DEBUG_FEATURES (!UE_BUILD_SHIPPING)
 #endif
 
+#ifndef CSV_PROFILER_USE_CUSTOM_FRAME_TIMINGS
+#define CSV_PROFILER_USE_CUSTOM_FRAME_TIMINGS 0
+#endif
+
 #if WITH_SERVER_CODE
   #define CSV_PROFILER (WITH_ENGINE && 1)
 #else
