Merging //UE4/Dev-VR to Dev-Main (//UE4/Dev-Main) Change: 9164498
UnrealMathNeon.h too many parameters in nan check on hololens.  Temporarily went back to it just returning false #if PLATFORM_HOLOLENS

diff --git a/Engine/Source/Runtime/Core/Public/Math/UnrealMathNeon.h b/Engine/Source/Runtime/Core/Public/Math/UnrealMathNeon.h
--- a/Engine/Source/Runtime/Core/Public/Math/UnrealMathNeon.h
+++ b/Engine/Source/Runtime/Core/Public/Math/UnrealMathNeon.h
@@ -1343,6 +1343,10 @@ FORCEINLINE void VectorSinCos(  VectorRegister* VSinAngles, VectorRegister* VCos
 // Returns true if the vector contains a component that is either NAN or +/-infinite.
 inline bool VectorContainsNaNOrInfinite(const VectorRegister& Vec)
 {
+#if PLATFORM_HOLOLENS
+	//TODO : The implementation below does not compile on Hololens Arm64.Temporarily working around this problem.
+	return false;
+#else
 	// https://en.wikipedia.org/wiki/IEEE_754-1985
 	// Infinity is represented with all exponent bits set, with the correct sign bit.
 	// NaN is represented with all exponent bits set, plus at least one fraction/significant bit set.
@@ -1360,6 +1364,7 @@ inline bool VectorContainsNaNOrInfinite(const VectorRegister& Vec)
 	uint8x16_t res = (uint8x16_t)VectorCompareEQ(ExpTest, FloatInfinity);
 	// If we have all zeros, all elements are finite
 	return vgetq_lane_u32((uint32x4_t)vqtbx1q_u8(res, res, Table), 0) != 0;
+#endif
 }
 
 //TODO: Vectorize
