Fix GetAdditionalRootDirectories() / AddAdditionalRootDirectories() not being threadsafe.  We occasionally hit a race condition with this during IOS install.

diff --git a/Engine/Source/Runtime/Core/Private/GenericPlatform/GenericPlatformMisc.cpp b/Engine/Source/Runtime/Core/Private/GenericPlatform/GenericPlatformMisc.cpp
--- a/Engine/Source/Runtime/Core/Private/GenericPlatform/GenericPlatformMisc.cpp
+++ b/Engine/Source/Runtime/Core/Private/GenericPlatform/GenericPlatformMisc.cpp
@@ -3,6 +3,8 @@
 #include "GenericPlatform/GenericPlatformMisc.h"
 #include "Misc/AssertionMacros.h"
 #include "HAL/PlatformFilemanager.h"
+#include "HAL/CriticalSection.h"
+#include "Misc/ScopeRWLock.h"
 #include "Math/UnrealMathUtility.h"
 #include "HAL/UnrealMemory.h"
 #include "Containers/Array.h"
@@ -61,27 +63,27 @@ bool LexTryParseString(EBuildConfiguration& OutConfiguration, const TCHAR* Confi
 		return true;
 	}
 	else if (FCString::Stricmp(Configuration, TEXT("DebugGame")) == 0)
-	{
+		{
 		OutConfiguration = EBuildConfiguration::DebugGame;
 		return true;
-	}
+		}
 	else if (FCString::Stricmp(Configuration, TEXT("Development")) == 0)
-	{
+		{
 		OutConfiguration = EBuildConfiguration::Development;
 		return true;
-	}
+		}
 	else if (FCString::Stricmp(Configuration, TEXT("Shipping")) == 0)
-	{
+		{
 		OutConfiguration = EBuildConfiguration::Shipping;
 		return true;
-	}
+		}
 	else if(FCString::Stricmp(Configuration, TEXT("Test")) == 0)
-	{
+		{
 		OutConfiguration = EBuildConfiguration::Test;
 		return true;
-	}
+		}
 	else if(FCString::Stricmp(Configuration, TEXT("Unknown")) == 0)
-	{
+		{
 		OutConfiguration = EBuildConfiguration::Unknown;
 		return true;
 	}
@@ -89,27 +91,27 @@ bool LexTryParseString(EBuildConfiguration& OutConfiguration, const TCHAR* Confi
 	{
 		OutConfiguration = EBuildConfiguration::Unknown;
 		return false;
+		}
 	}
-}
 
 const TCHAR* LexToString( EBuildConfiguration Configuration )
-{
-	switch (Configuration)
 	{
+		switch (Configuration)
+		{
 	case EBuildConfiguration::Debug:
-		return TEXT("Debug");
+				return TEXT("Debug");
 	case EBuildConfiguration::DebugGame:
-		return TEXT("DebugGame");
+				return TEXT("DebugGame");
 	case EBuildConfiguration::Development:
-		return TEXT("Development");
+				return TEXT("Development");
 	case EBuildConfiguration::Shipping:
-		return TEXT("Shipping");
+				return TEXT("Shipping");
 	case EBuildConfiguration::Test:
-		return TEXT("Test");
-	default:
-		return TEXT("Unknown");
+				return TEXT("Test");
+			default:
+				return TEXT("Unknown");
+		}
 	}
-}
 
 namespace EBuildConfigurations
 {
@@ -167,7 +169,7 @@ bool LexTryParseString(EBuildTargetType& OutType, const TCHAR* Type)
 		return true;
 	}
 	else if (FCString::Strcmp(Type, TEXT("Server")) == 0)
-	{
+{
 		OutType = EBuildTargetType::Server;
 		return true;
 	}
@@ -177,40 +179,40 @@ bool LexTryParseString(EBuildTargetType& OutType, const TCHAR* Type)
 		return true;
 	}
 	else if (FCString::Strcmp(Type, TEXT("Program")) == 0)
-	{
+		{
 		OutType = EBuildTargetType::Program;
 		return true;
-	}
+		}
 	else if (FCString::Strcmp(Type, TEXT("Unknown")) == 0)
-	{
+		{
 		OutType = EBuildTargetType::Unknown;
 		return true;
-	}
+		}
 	else
-	{
+		{
 		OutType = EBuildTargetType::Unknown;
 		return false;
+		}
 	}
-}
 
 const TCHAR* LexToString(EBuildTargetType Type)
-{
-	switch (Type)
 	{
+	switch (Type)
+		{
 	case EBuildTargetType::Editor:
-		return TEXT("Editor");
+				return TEXT("Editor");
 	case EBuildTargetType::Game:
-		return TEXT("Game");
+				return TEXT("Game");
 	case EBuildTargetType::Server:
-		return TEXT("Server");
+				return TEXT("Server");
 	case EBuildTargetType::Client:
 		return TEXT("Client");
 	case EBuildTargetType::Program:
 		return TEXT("Program");
-	default:
-		return TEXT("Unknown");
+			default:
+				return TEXT("Unknown");
+		}
 	}
-}
 
 EBuildTargetType EBuildTargets::FromString(const FString& Target)
 {
@@ -279,6 +281,7 @@ struct FGenericPlatformMisc::FStaticData
 {
 	FString         RootDir;
 	TArray<FString> AdditionalRootDirectories;
+	FRWLock         AdditionalRootDirectoriesLock;
 	FString         EngineDirectory;
 	FString         LaunchDir;
 	FString         ProjectDir;
@@ -799,13 +802,17 @@ const TCHAR* FGenericPlatformMisc::RootDir()
 	return *Path;
 }
 
-const TArray<FString>& FGenericPlatformMisc::GetAdditionalRootDirectories()
+TArray<FString> FGenericPlatformMisc::GetAdditionalRootDirectories()
 {
+	FRWScopeLock Lock(TLazySingleton<FStaticData>::Get().AdditionalRootDirectoriesLock, SLT_ReadOnly);
+
 	return TLazySingleton<FStaticData>::Get().AdditionalRootDirectories;
 }
 
 void FGenericPlatformMisc::AddAdditionalRootDirectory(const FString& RootDir)
 {
+	FRWScopeLock Lock(TLazySingleton<FStaticData>::Get().AdditionalRootDirectoriesLock, SLT_Write);
+
 	TArray<FString>& RootDirectories = TLazySingleton<FStaticData>::Get().AdditionalRootDirectories;
 	FString NewRootDirectory = RootDir;
 	FPaths::MakePlatformFilename(NewRootDirectory);
diff --git a/Engine/Source/Runtime/Core/Public/GenericPlatform/GenericPlatformMisc.h b/Engine/Source/Runtime/Core/Public/GenericPlatform/GenericPlatformMisc.h
--- a/Engine/Source/Runtime/Core/Public/GenericPlatform/GenericPlatformMisc.h
+++ b/Engine/Source/Runtime/Core/Public/GenericPlatform/GenericPlatformMisc.h
@@ -26,44 +26,44 @@ class TFunction;
 #define UE_DEBUG_BREAK() ((void)(FPlatformMisc::IsDebuggerPresent() && ([] () { UE_DEBUG_BREAK_IMPL(); } (), 1)))
 #endif
 
-/**
+	/**
  * Available build configurations. Mirorred from UnrealTargetConfiguration.
- */
+	 */
 enum class EBuildConfiguration : uint8
-{
-	/** Unknown build configuration. */
-	Unknown,
+	{
+		/** Unknown build configuration. */
+		Unknown,
 
-	/** Debug build. */
-	Debug,
+		/** Debug build. */
+		Debug,
 
-	/** DebugGame build. */
-	DebugGame,
+		/** DebugGame build. */
+		DebugGame,
 
-	/** Development build. */
-	Development,
+		/** Development build. */
+		Development,
 
-	/** Shipping build. */
-	Shipping,
+		/** Shipping build. */
+		Shipping,
 
-	/** Test build. */
-	Test
-};
+		/** Test build. */
+		Test
+	};
 
-/**
- * Returns the string representation of the specified EBuildConfiguration value.
- *
+	/**
+	 * Returns the string representation of the specified EBuildConfiguration value.
+	 *
  * @param Configuration The string to get the EBuildConfiguration for.
  * @return An EBuildConfiguration value.
- */
+	 */
 CORE_API bool LexTryParseString(EBuildConfiguration& OutConfiguration, const TCHAR* Configuration);
 
-/**
- * Returns the string representation of the specified EBuildConfiguration value.
- *
- * @param Configuration The value to get the string for.
- * @return The string representation.
- */
+	/**
+	 * Returns the string representation of the specified EBuildConfiguration value.
+	 *
+	 * @param Configuration The value to get the string for.
+	 * @return The string representation.
+	 */
 CORE_API const TCHAR* LexToString(EBuildConfiguration Configuration);
 
 namespace EBuildConfigurations
@@ -116,18 +116,18 @@ namespace EBuildConfigurations
 	CORE_API FText ToText( EBuildConfiguration Configuration );
 }
 
-/**
+	/**
  * Enumerates build target types.
- */
+	 */
 enum class EBuildTargetType : uint8
-{
-	/** Unknown build target. */
-	Unknown,
+	{
+		/** Unknown build target. */
+		Unknown,
 
-	/** Game target. */
-	Game,
+		/** Game target. */
+		Game,
 
-	/** Server target. */
+		/** Server target. */
 	Server,
 
 	/** Client target. */
@@ -138,23 +138,23 @@ enum class EBuildTargetType : uint8
 
 	/** Program target. */
 	Program,
-};
+	};
 
-/**
- * Returns the string representation of the specified EBuildTarget value.
- *
+	/**
+	 * Returns the string representation of the specified EBuildTarget value.
+	 *
  * @param OutType The value to get the string for.
  * @param Text The text to parse.
  * @return The string representation.
- */
+	 */
 CORE_API bool LexTryParseString(EBuildTargetType& OutType, const TCHAR* Text);
 
-/**
+	/**
  * Returns the string representation of the specified EBuildTargetType value.
- *
+	 *
  * @param Target The string to get the EBuildTargetType for.
  * @return An EBuildTarget::Type value.
- */
+	 */
 CORE_API const TCHAR* LexToString(EBuildTargetType Type);
 
 namespace EBuildTargets
@@ -874,7 +874,7 @@ public:
 	static const TCHAR* RootDir();
 
 	/** get additional directories which can be considered as root directories */
-	static const TArray<FString>& GetAdditionalRootDirectories();
+	static TArray<FString> GetAdditionalRootDirectories();
 	/** add an additional root directory */
 	static void AddAdditionalRootDirectory(const FString& RootDir);
 
