Wrapped missed ICU include for correct padding

Added a static_assert as a compile time check, rather than letting it compile and then crash in a random place later

#rnx

diff --git a/Engine/Source/Runtime/Core/Private/Internationalization/ICUCulture.h b/Engine/Source/Runtime/Core/Private/Internationalization/ICUCulture.h
--- a/Engine/Source/Runtime/Core/Private/Internationalization/ICUCulture.h
+++ b/Engine/Source/Runtime/Core/Private/Internationalization/ICUCulture.h
@@ -20,7 +20,9 @@ THIRD_PARTY_INCLUDES_START
 	#include <unicode/coll.h>
 	#include <unicode/numfmt.h>
 	#include <unicode/decimfmt.h>
-	#include <unicode/datefmt.h>
+PRAGMA_PUSH_PLATFORM_DEFAULT_PACKING
+	#include <unicode/datefmt.h> // icu::Calendar can be affected by the non-standard packing UE4 uses, so force the platform default
+PRAGMA_POP_PLATFORM_DEFAULT_PACKING
 	#include <unicode/plurrule.h>
 THIRD_PARTY_INCLUDES_END
 
diff --git a/Engine/Source/Runtime/Core/Private/Internationalization/ICUInternationalization.cpp b/Engine/Source/Runtime/Core/Private/Internationalization/ICUInternationalization.cpp
--- a/Engine/Source/Runtime/Core/Private/Internationalization/ICUInternationalization.cpp
+++ b/Engine/Source/Runtime/Core/Private/Internationalization/ICUInternationalization.cpp
@@ -32,6 +32,14 @@ DEFINE_LOG_CATEGORY_STATIC(LogICUInternationalization, Log, All);
 static_assert(sizeof(UChar) == 2, "UChar (from ICU) is assumed to always be 2-bytes!");
 static_assert(PLATFORM_LITTLE_ENDIAN, "ICU data is only built for little endian platforms. You'll need to rebuild the data for your platform and update this code!");
 
+#if WITH_ICU_V64 && PLATFORM_WINDOWS
+	#if PLATFORM_32BITS
+		static_assert(sizeof(icu::Calendar) == 608, "icu::Calendar should be 608-bytes! Ensure relevant includes are wrapped in PRAGMA_PUSH_PLATFORM_DEFAULT_PACKING and PRAGMA_POP_PLATFORM_DEFAULT_PACKING.");
+	#else
+		static_assert(sizeof(icu::Calendar) == 616, "icu::Calendar should be 616-bytes! Ensure relevant includes are wrapped in PRAGMA_PUSH_PLATFORM_DEFAULT_PACKING and PRAGMA_POP_PLATFORM_DEFAULT_PACKING.");
+	#endif
+#endif
+
 namespace
 {
 	struct FICUOverrides
