Make FName detect number strings between int32 and uint32 max, even though FName's number is signed.

These semantics match the old UE 4.22 FName implementation.

diff --git a/Engine/Source/Runtime/Core/Private/UObject/UnrealNames.cpp b/Engine/Source/Runtime/Core/Private/UObject/UnrealNames.cpp
--- a/Engine/Source/Runtime/Core/Private/UObject/UnrealNames.cpp
+++ b/Engine/Source/Runtime/Core/Private/UObject/UnrealNames.cpp
@@ -1408,7 +1408,7 @@ bool FName::IsWithinBounds(FNameEntryId Id)
 -----------------------------------------------------------------------------*/
 
 template<class CharType>
-static bool NumberEqualsString(int32 Number, const CharType* Str)
+static bool NumberEqualsString(uint32 Number, const CharType* Str)
 {
 	CharType* End = nullptr;
 	return TCString<CharType>::Strtoi64(Str, &End, 10) == Number && End && *End == '\0';
@@ -1530,8 +1530,8 @@ struct FNameHelper
 			{
 				// Attempt to convert what's following it to a number
 				// This relies on Name being null-terminated
-				uint64 Number = TCString<CharType>::Atoi64(Name + Len - Digits);
-				if (Number < MAX_uint32)
+				int64 Number = TCString<CharType>::Atoi64(Name + Len - Digits);
+				if (Number < MAX_int32)
 				{
 					View.Len -= 1 + Digits;
 					return MakeWithNumber(View, FindType, static_cast<uint32>(NAME_EXTERNAL_TO_INTERNAL(Number)));
@@ -1954,6 +1954,7 @@ void FName::AutoTest()
 	check(FCStringAnsi::Strlen("ABC_9") == FName("ABC_9").GetStringLength());
 	check(FCStringAnsi::Strlen("ABC_10") == FName("ABC_10").GetStringLength());
 	check(FCStringAnsi::Strlen("ABC_2000000000") == FName("ABC_2000000000").GetStringLength());
+	check(FCStringAnsi::Strlen("ABC_4000000000") == FName("ABC_4000000000").GetStringLength());
 
 	const FName NullName(static_cast<ANSICHAR*>(nullptr));
 	check(NullName.IsNone());
@@ -2029,6 +2030,7 @@ void FName::AutoTest()
 	check(NumberEqualsString(0, "0"));
 	check(NumberEqualsString(11, "11"));
 	check(NumberEqualsString(2147483647, "2147483647"));
+	check(NumberEqualsString(4294967294, "4294967294"));
 
 	check(!NumberEqualsString(0, "1"));
 	check(!NumberEqualsString(1, "0"));
