Re-enabled CPU scope trace events on Linux. They were disabled temporarily due to a Linux merge error.

#rnx

(ushell-p4-cherrypick of 9439746 by [REDACTED_PERSON])

diff --git a/Engine/Source/Runtime/Core/Public/ProfilingDebugging/CpuProfilerTrace.h b/Engine/Source/Runtime/Core/Public/ProfilingDebugging/CpuProfilerTrace.h
--- a/Engine/Source/Runtime/Core/Public/ProfilingDebugging/CpuProfilerTrace.h
+++ b/Engine/Source/Runtime/Core/Public/ProfilingDebugging/CpuProfilerTrace.h
@@ -4,7 +4,7 @@
 
 #include "CoreTypes.h"
 
-#if !IS_PROGRAM && !UE_BUILD_SHIPPING && (PLATFORM_WINDOWS || PLATFORM_PS4 || PLATFORM_XBOXONE || PLATFORM_SWITCH || PLATFORM_APPLE || PLATFORM_ANDROID)
+#if !IS_PROGRAM && !UE_BUILD_SHIPPING && (PLATFORM_WINDOWS || PLATFORM_UNIX || PLATFORM_PS4 || PLATFORM_XBOXONE || PLATFORM_SWITCH || PLATFORM_APPLE || PLATFORM_ANDROID)
 #define CPUPROFILERTRACE_ENABLED 1
 #else
 #define CPUPROFILERTRACE_ENABLED 0
