Added FLazyName number support

diff --git a/Engine/Source/Runtime/Core/Private/UObject/UnrealNames.cpp b/Engine/Source/Runtime/Core/Private/UObject/UnrealNames.cpp
--- a/Engine/Source/Runtime/Core/Private/UObject/UnrealNames.cpp
+++ b/Engine/Source/Runtime/Core/Private/UObject/UnrealNames.cpp
@@ -1509,11 +1509,15 @@ struct FNameHelper
 		{
 			return FName();
 		}
-
-		using CharType = typename ViewType::CharType;
-		const CharType* Name = View.Str;
-		const int32 Len = View.Len;
 		
+		uint32 InternalNumber = ParseNumber(View.Str, /* may be shortened */ View.Len);
+		return MakeWithNumber(View, FindType, InternalNumber);
+	}
+
+	template<typename CharType>
+	static uint32 ParseNumber(const CharType* Name, int32& InOutLen)
+	{
+		const int32 Len = InOutLen;
 		int32 Digits = 0;
 		for (const CharType* It = Name + Len - 1; It >= Name && *It >= '0' && *It <= '9'; --It)
 		{
@@ -1533,13 +1537,13 @@ struct FNameHelper
 				int64 Number = TCString<CharType>::Atoi64(Name + Len - Digits);
 				if (Number < MAX_int32)
 				{
-					View.Len -= 1 + Digits;
-					return MakeWithNumber(View, FindType, static_cast<uint32>(NAME_EXTERNAL_TO_INTERNAL(Number)));
+					InOutLen -= 1 + Digits;
+					return static_cast<uint32>(NAME_EXTERNAL_TO_INTERNAL(Number));
 				}
 			}
 		}
 
-		return MakeWithNumber(View, FindType, NAME_NO_NUMBER_INTERNAL);
+		return NAME_NO_NUMBER_INTERNAL;
 	}
 
 	static FName MakeWithNumber(FNameAnsiStringView	 View, EFindName FindType, int32 InternalNumber)
@@ -1629,7 +1633,6 @@ struct FNameHelper
 			: FNameStringView(LoadedEntry.AnsiName, FCStringAnsi::Strlen(LoadedEntry.AnsiName));
 
 		return Make(View, FNAME_Add, NAME_NO_NUMBER_INTERNAL);
-
 	}
 
 	template<class CharType>
@@ -2075,6 +2078,16 @@ void FName::AutoTest()
 	check(Names[5] == "FooC");
 	check(Names[6] == FooWide);
 
+	check(FLazyName("Hej") == FName("Hej"));
+	check(FLazyName("Hej_0") == FName("Hej_0"));
+	check(FLazyName("Hej_00") == FName("Hej_00"));
+	check(FLazyName("Hej_1") == FName("Hej_1"));
+	check(FLazyName("Hej_01") == FName("Hej_01"));
+	check(FLazyName("Hej_-1") == FName("Hej_-1"));
+	check(FLazyName("Hej__0") == FName("Hej__0"));
+	check(FLazyName("Hej_2147483647") == FName("Hej_2147483647"));
+	check(FLazyName("Hej_123") == FLazyName(FName("Hej_123")));
+
 #if 0
 	// Check hash table growth still yields the same unique FName ids
 	static int32 OverflowAtLeastTwiceCount = 4 * FNamePoolInitialSlotsPerShard * FNamePoolShards;
@@ -2221,6 +2234,11 @@ void FName::TearDown()
 	}
 }
 
+uint32 FLazyName::ParseNumber(const TCHAR* Str, int32 Len)
+{
+	return FNameHelper::ParseNumber(Str, Len);
+}
+
 #if !UE_BUILD_SHIPPING && !UE_BUILD_TEST
 
 #include "Containers/StackTracker.h"
diff --git a/Engine/Source/Runtime/Core/Public/UObject/NameTypes.h b/Engine/Source/Runtime/Core/Public/UObject/NameTypes.h
--- a/Engine/Source/Runtime/Core/Public/UObject/NameTypes.h
+++ b/Engine/Source/Runtime/Core/Public/UObject/NameTypes.h
@@ -992,10 +992,12 @@ public:
 	template<int N>
 	FLazyName(const TCHAR(&Literal)[N])
 		: Either(Literal)
+		, Number(ParseNumber(Literal, N))
 	{}
 
 	explicit FLazyName(FName Name)
 		: Either(Name.GetComparisonIndex())
+		, Number(Name.GetNumber())
 	{}
 
 	operator FName() const
@@ -1006,7 +1008,7 @@ public:
 		if (Copy.IsName())
 		{
 			FNameEntryId Id = Copy.AsName();
-			return FName(Id, Id, 0);
+			return FName(Id, Id, Number);
 		}
 
 		// Resolve to FName
@@ -1015,7 +1017,7 @@ public:
 		// Deliberately unsynchronized write of word-sized int, ok if multiple threads resolve same lazy name
 		Either = FLiteralOrName(Id);
 
-		return FName(Id, Id, 0);
+		return FName(Id, Id, Number);
 	}
 
 private:
@@ -1091,6 +1093,9 @@ private:
 	};
 
 	mutable FLiteralOrName Either;
+	mutable uint32 Number = 0;
+
+	CORE_API static uint32 ParseNumber(const TCHAR* Literal, int32 Len);
 
 public:
 	friend bool operator==(FName Name, FLazyName Lazy)
