There was a discrepancy with check()'s formatted prints between logs and crash reports/core dumps as the vaarg was incorrectly reused.

#rnx

(ushell-p4-cherrypick of 9487152 by [REDACTED_PERSON])

diff --git a/Engine/Source/Runtime/Core/Private/Misc/AssertionMacros.cpp b/Engine/Source/Runtime/Core/Private/Misc/AssertionMacros.cpp
--- a/Engine/Source/Runtime/Core/Private/Misc/AssertionMacros.cpp
+++ b/Engine/Source/Runtime/Core/Private/Misc/AssertionMacros.cpp
@@ -439,15 +439,19 @@ void FORCENOINLINE FDebug::CheckVerifyFailedImpl(
 	...)
 {
 	va_list Args;
+
 	va_start(Args, Format);
 	FDebug::LogAssertFailedMessageImplV(Expr, File, Line, Format, Args);
+	va_end(Args);
 
 	if (!FPlatformMisc::IsDebuggerPresent())
 	{
 		FPlatformMisc::PromptForRemoteDebugging(false);
+
+		va_start(Args, Format);
 		AssertFailedImplV(Expr, File, Line, Format, Args);
+		va_end(Args);
 	}
-	va_end(Args);
 }
 
 #endif // DO_CHECK || DO_GUARD_SLOW
