Mark GetExtrema as deprecated due to unsafe indexing operations, to be replaced with either the [] operator or direct access to Min and Max.

diff --git a/Engine/Source/Runtime/Core/Public/Math/Box.h b/Engine/Source/Runtime/Core/Public/Math/Box.h
--- a/Engine/Source/Runtime/Core/Public/Math/Box.h
+++ b/Engine/Source/Runtime/Core/Public/Math/Box.h
@@ -250,23 +250,13 @@ public:
 		return 0.5f * (Max - Min);
 	}
 
-	/**
-	 * Gets a reference to the specified point of the bounding box.
-	 *
-	 * @param PointIndex The index of the extrema point to return.
-	 * @return A reference to the point.
-	 */
+	UE_DEPRECATED(4.24, "This method performed unsafe operations and should be replaced with using .Min and .Max directly or using the [] operator on this class instead.")
 	FORCEINLINE FVector& GetExtrema( int PointIndex )
 	{
 		return (&Min)[PointIndex];
 	}
 
-	/**
-	 * Gets a read-only reference to the specified point of the bounding box.
-	 *
-	 * @param PointIndex The index of extrema point to return.
-	 * @return A read-only reference to the point.
-	 */
+	UE_DEPRECATED(4.24, "This method performed unsafe operations and should be replaced with using .Min and .Max directly or using the [] operator on this class instead.")
 	FORCEINLINE const FVector& GetExtrema( int PointIndex ) const
 	{
 		return (&Min)[PointIndex];
