Copying //UE4/Dev-Documentation to Samples-Main (//UE4/Samples-Main)
#rnx

diff --git a/Engine/Documentation/Source/Shared/Editor/MapErrors/MapErrors.INT.udn b/Engine/Documentation/Source/Shared/Editor/MapErrors/MapErrors.INT.udn
--- a/Engine/Documentation/Source/Shared/Editor/MapErrors/MapErrors.INT.udn
+++ b/Engine/Documentation/Source/Shared/Editor/MapErrors/MapErrors.INT.udn
@@ -2,7 +2,11 @@ Availability:Public
 Title:Map Check Errors
 Crumbs:%ROOT%, Engine, Engine/UI/LevelEditor
 Description:Descriptions of errors generated during map builds or running a map check in Unreal Editor.
-Version: 4.9
+Version: 4.22
+Type:reference
+Parent:Shared/Editor
+Tags:map
+Tags:error
 
 [TOC(start:2 end:3)]
 
diff --git a/Engine/Source/Editor/FoliageEdit/Private/SFoliageEdit.cpp b/Engine/Source/Editor/FoliageEdit/Private/SFoliageEdit.cpp
--- a/Engine/Source/Editor/FoliageEdit/Private/SFoliageEdit.cpp
+++ b/Engine/Source/Editor/FoliageEdit/Private/SFoliageEdit.cpp
@@ -295,7 +295,7 @@ void SFoliageEdit::Construct(const FArguments& InArgs)
 									.Visibility(this, &SFoliageEdit::GetVisibility_SpawnInCurrentLevelMode)
 									.OnCheckStateChanged(this, &SFoliageEdit::OnCheckStateChanged_SpawnInCurrentLevelMode)
 									.IsChecked(this, &SFoliageEdit::GetCheckState_SpawnInCurrentLevelMode)
-									.ToolTipText(LOCTEXT("SpawnInCurrentLevelModeTooltips", "Wether to place foliage meshes in the current level or in the level containing the mesh being painted on."))
+									.ToolTipText(LOCTEXT("SpawnInCurrentLevelModeTooltips", "Whether to place foliage meshes in the current level or in the level containing the mesh being painted on."))
 									[
 										SNew(STextBlock)
 										.Text(LOCTEXT("SpawnInCurrentLevelMode", "Place in Current Level"))
diff --git a/Engine/Source/Programs/AutomationTool/BuildGraph/Tasks/AgeStoreTask.cs b/Engine/Source/Programs/AutomationTool/BuildGraph/Tasks/AgeStoreTask.cs
--- a/Engine/Source/Programs/AutomationTool/BuildGraph/Tasks/AgeStoreTask.cs
+++ b/Engine/Source/Programs/AutomationTool/BuildGraph/Tasks/AgeStoreTask.cs
@@ -35,7 +35,7 @@ namespace Win.Automation
 		public int Days;
 
 		/// <summary>
-		/// The root of the build dir to check for existing buildversion named directories
+		/// The root of the build directory to check for existing buildversion named directories.
 		/// </summary>
 		[TaskParameter(Optional = true)]
 		public string BuildDir;
@@ -43,14 +43,14 @@ namespace Win.Automation
 		/// <summary>
 		/// A substring to match in directory file names before deleting symbols. This allows the "age store" task
 		/// to avoid deleting symbols from other builds in the case where multiple builds share the same symbol server.
-		/// Specific use of the filter value is determined by the symbol server structure defined by the platform tool chain.
+		/// Specific use of the filter value is determined by the symbol server structure defined by the platform toolchain.
 		/// </summary>
 		[TaskParameter(Optional = true)]
 		public string Filter;
 	}
 
 	/// <summary>
-	/// Task which strips symbols from a set of files. This task is named after the AGESTORE utility that comes with the Microsoft debugger tools SDK, but is actually a separate implementation. The main
+	/// Task that strips symbols from a set of files. This task is named after the AGESTORE utility that comes with the Microsoft debugger tools SDK, but is actually a separate implementation. The main
 	/// difference is that it uses the last modified time rather than last access time to determine which files to delete.
 	/// </summary>
 	[TaskElement("AgeStore", typeof(AgeStoreTaskParameters))]
diff --git a/Engine/Source/Programs/AutomationTool/BuildGraph/Tasks/CommandTask.cs b/Engine/Source/Programs/AutomationTool/BuildGraph/Tasks/CommandTask.cs
--- a/Engine/Source/Programs/AutomationTool/BuildGraph/Tasks/CommandTask.cs
+++ b/Engine/Source/Programs/AutomationTool/BuildGraph/Tasks/CommandTask.cs
@@ -22,18 +22,18 @@ namespace BuildGraph.Tasks
 	}
 
 	/// <summary>
-	/// Parameters for a task which calls another UAT command
+	/// Parameters for a task that calls another UAT command
 	/// </summary>
 	public class CommandTaskParameters
 	{
 		/// <summary>
-		/// The command name to execute
+		/// The command name to execute.
 		/// </summary>
 		[TaskParameter]
 		public string Name;
 
 		/// <summary>
-		/// Arguments to be passed to the command
+		/// Arguments to be passed to the command.
 		/// </summary>
 		[TaskParameter(Optional = true)]
 		public string Arguments;
diff --git a/Engine/Source/Programs/AutomationTool/BuildGraph/Tasks/CommandletTask.cs b/Engine/Source/Programs/AutomationTool/BuildGraph/Tasks/CommandletTask.cs
--- a/Engine/Source/Programs/AutomationTool/BuildGraph/Tasks/CommandletTask.cs
+++ b/Engine/Source/Programs/AutomationTool/BuildGraph/Tasks/CommandletTask.cs
@@ -18,19 +18,19 @@ namespace AutomationTool.Tasks
 	public class CommandletTaskParameters
 	{
 		/// <summary>
-		/// The commandlet name to execute
+		/// The commandlet name to execute.
 		/// </summary>
 		[TaskParameter]
 		public string Name;
 
 		/// <summary>
-		/// The project to run the editor with
+		/// The project to run the editor with.
 		/// </summary>
 		[TaskParameter(Optional = true, ValidationType = TaskParameterValidationType.FileSpec)]
 		public string Project;
 
 		/// <summary>
-		/// Arguments to be passed to the commandlet
+		/// Arguments to be passed to the commandlet.
 		/// </summary>
 		[TaskParameter(Optional = true)]
 		public string Arguments;
diff --git a/Engine/Source/Programs/AutomationTool/BuildGraph/Tasks/CompileTask.cs b/Engine/Source/Programs/AutomationTool/BuildGraph/Tasks/CompileTask.cs
--- a/Engine/Source/Programs/AutomationTool/BuildGraph/Tasks/CompileTask.cs
+++ b/Engine/Source/Programs/AutomationTool/BuildGraph/Tasks/CompileTask.cs
@@ -18,55 +18,55 @@ namespace AutomationTool
 	public class CompileTaskParameters
 	{
 		/// <summary>
-		/// The target to compile
+		/// The target to compile.
 		/// </summary>
 		[TaskParameter]
 		public string Target;
 
 		/// <summary>
-		/// The configuration to compile
+		/// The configuration to compile.
 		/// </summary>
 		[TaskParameter]
 		public UnrealTargetConfiguration Configuration;
 
 		/// <summary>
-		/// The platform to compile for
+		/// The platform to compile for.
 		/// </summary>
 		[TaskParameter]
 		public UnrealTargetPlatform Platform;
 
 		/// <summary>
-		/// The project to compile with
+		/// The project to compile with.
 		/// </summary>
 		[TaskParameter(Optional = true)]
 		public string Project;
 
 		/// <summary>
-		/// Additional arguments for UnrealBuildTool
+		/// Additional arguments for UnrealBuildTool.
 		/// </summary>
 		[TaskParameter(Optional = true)]
 		public string Arguments;
 
 		/// <summary>
-		/// Whether to allow using XGE for compilation
+		/// Whether to allow using XGE for compilation.
 		/// </summary>
 		[TaskParameter(Optional = true)]
 		public bool AllowXGE = true;
 
 		/// <summary>
-		/// Whether to allow using the parallel executor for this compile
+		/// Whether to allow using the parallel executor for this compile.
 		/// </summary>
 		[TaskParameter(Optional = true)]
 		public bool AllowParallelExecutor = true;
 
 		/// <summary>
-		/// Whether to allow cleaning this target. If unspecified, targets are cleaned if the -Clean argument passed on the command line.
+		/// Whether to allow cleaning this target. If unspecified, targets are cleaned if the -Clean argument is passed on the command line.
 		/// </summary>
 		[TaskParameter(Optional = true)]
 		public bool? Clean = null;
 
 		/// <summary>
-		/// Tag to be applied to build products of this task
+		/// Tag to be applied to build products of this task.
 		/// </summary>
 		[TaskParameter(Optional = true, ValidationType = TaskParameterValidationType.TagList)]
 		public string Tag;
diff --git a/Engine/Source/Programs/AutomationTool/BuildGraph/Tasks/CookTask.cs b/Engine/Source/Programs/AutomationTool/BuildGraph/Tasks/CookTask.cs
--- a/Engine/Source/Programs/AutomationTool/BuildGraph/Tasks/CookTask.cs
+++ b/Engine/Source/Programs/AutomationTool/BuildGraph/Tasks/CookTask.cs
@@ -1,4 +1,4 @@
-﻿// Copyright 1998-2019 Epic Games, Inc. All Rights Reserved.
+// Copyright 1998-2019 Epic Games, Inc. All Rights Reserved.
 
 using AutomationTool;
 using System;
@@ -18,37 +18,37 @@ namespace BuildGraph.Tasks
 	public class CookTaskParameters
 	{
 		/// <summary>
-		/// Project file to be cooked
+		/// Project file to be cooked.
 		/// </summary>
 		[TaskParameter]
 		public string Project;
 
 		/// <summary>
-		/// The cook platform to target (eg. WindowsNoEditor)
+		/// The cook platform to target (for example, WindowsNoEditor).
 		/// </summary>
 		[TaskParameter]
 		public string Platform;
 
 		/// <summary>
-		/// List of maps to be cooked, separated by '+' characters
+		/// List of maps to be cooked, separated by '+' characters.
 		/// </summary>
 		[TaskParameter(Optional = true)]
 		public string Maps;
 
 		/// <summary>
-		/// Additional arguments to be passed to the cooker
+		/// Additional arguments to be passed to the cooker.
 		/// </summary>
 		[TaskParameter(Optional = true)]
 		public bool Versioned = false;
 	
 		/// <summary>
-		/// Additional arguments to be passed to the cooker
+		/// Additional arguments to be passed to the cooker.
 		/// </summary>
 		[TaskParameter(Optional = true)]
 		public string Arguments = "";
 
 		/// <summary>
-		/// Tag to be applied to build products of this task
+		/// Tag to be applied to build products of this task.
 		/// </summary>
 		[TaskParameter(Optional = true, ValidationType = TaskParameterValidationType.TagList)]
 		public string Tag;
diff --git a/Engine/Source/Programs/AutomationTool/BuildGraph/Tasks/CopyTask.cs b/Engine/Source/Programs/AutomationTool/BuildGraph/Tasks/CopyTask.cs
--- a/Engine/Source/Programs/AutomationTool/BuildGraph/Tasks/CopyTask.cs
+++ b/Engine/Source/Programs/AutomationTool/BuildGraph/Tasks/CopyTask.cs
@@ -19,31 +19,31 @@ namespace BuildGraph.Tasks
 	public class CopyTaskParameters
 	{
 		/// <summary>
-		/// Filter to be applied to the list of input files. Optional.
+		/// Optional filter to be applied to the list of input files.
 		/// </summary>
 		[TaskParameter(Optional = true, ValidationType = TaskParameterValidationType.FileSpec)]
 		public string Files;
 
 		/// <summary>
-		/// The pattern(s) to copy from (eg. Engine/*.txt)
+		/// The pattern(s) to copy from (for example, Engine/*.txt).
 		/// </summary>
 		[TaskParameter(ValidationType = TaskParameterValidationType.FileSpec)]
 		public string From;
 
 		/// <summary>
-		/// The directory or to copy to
+		/// The directory to copy to.
 		/// </summary>
 		[TaskParameter(ValidationType = TaskParameterValidationType.FileSpec)]
 		public string To;
 
 		/// <summary>
-		/// Whether or not to overwrite existing files
+		/// Whether or not to overwrite existing files.
 		/// </summary>
 		[TaskParameter(Optional = true)]
 		public bool Overwrite = true;
 
 		/// <summary>
-		/// Tag to be applied to build products of this task
+		/// Tag to be applied to build products of this task.
 		/// </summary>
 		[TaskParameter(Optional = true, ValidationType = TaskParameterValidationType.TagList)]
 		public string Tag;
diff --git a/Engine/Source/Programs/AutomationTool/BuildGraph/Tasks/CsCompileTask.cs b/Engine/Source/Programs/AutomationTool/BuildGraph/Tasks/CsCompileTask.cs
--- a/Engine/Source/Programs/AutomationTool/BuildGraph/Tasks/CsCompileTask.cs
+++ b/Engine/Source/Programs/AutomationTool/BuildGraph/Tasks/CsCompileTask.cs
@@ -17,50 +17,50 @@ namespace AutomationTool.Tasks
 	public class CsCompileTaskParameters
 	{
 		/// <summary>
-		/// The C# project file to be compile. More than one project file can be specified by separating with semicolons.
+		/// The C# project file to compile. Using semicolons, more than one project file can be specified.
 		/// </summary>
 		[TaskParameter]
 		public string Project;
 
 		/// <summary>
-		/// The configuration to compile
+		/// The configuration to compile.
 		/// </summary>
 		[TaskParameter(Optional = true)]
 		public string Configuration;
 
 		/// <summary>
-		/// The platform to compile
+		/// The platform to compile.
 		/// </summary>
 		[TaskParameter(Optional = true)]
 		public string Platform;
 
 		/// <summary>
-		/// The target to build
+		/// The target to build.
 		/// </summary>
 		[TaskParameter(Optional = true)]
 		public string Target;
 
 		/// <summary>
-		/// Additional options to pass to the compiler
+		/// Additional options to pass to the compiler.
 		/// </summary>
 		[TaskParameter(Optional = true)]
 		public string Arguments;
 
 		/// <summary>
-		/// Only enumerate build products; do not actually compile the projects.
+		/// Only enumerate build products -- do not actually compile the projects.
 		/// </summary>
 		[TaskParameter(Optional = true)]
 		public bool EnumerateOnly;
 
 		/// <summary>
-		/// Tag to be applied to build products of this task
+		/// Tag to be applied to build products of this task.
 		/// </summary>
 		[TaskParameter(Optional = true, ValidationType = TaskParameterValidationType.TagList)]
 		public string Tag;
 
 		/// <summary>
-		/// Tag to be applied to any non-private references the projects have
-		/// (i.e. those that are external and not copied into the output dir)
+		/// Tag to be applied to any non-private references the projects have.
+		/// (for example, those that are external and not copied into the output directory).
 		/// </summary>
 		[TaskParameter(Optional = true, ValidationType = TaskParameterValidationType.TagList)]
 		public string TagReferences;
diff --git a/Engine/Source/Programs/AutomationTool/BuildGraph/Tasks/DeleteTask.cs b/Engine/Source/Programs/AutomationTool/BuildGraph/Tasks/DeleteTask.cs
--- a/Engine/Source/Programs/AutomationTool/BuildGraph/Tasks/DeleteTask.cs
+++ b/Engine/Source/Programs/AutomationTool/BuildGraph/Tasks/DeleteTask.cs
@@ -1,4 +1,4 @@
-﻿// Copyright 1998-2019 Epic Games, Inc. All Rights Reserved.
+// Copyright 1998-2019 Epic Games, Inc. All Rights Reserved.
 
 using AutomationTool;
 using System;
@@ -19,7 +19,7 @@ namespace BuildGraph.Tasks
 	public class DeleteTaskParameters
 	{
 		/// <summary>
-		/// List of file specifications separated by semicolons (eg. *.cpp;Engine/.../*.bat), or the name of a tag set
+		/// List of file specifications separated by semicolons (for example, *.cpp;Engine/.../*.bat), or the name of a tag set
 		/// </summary>
 		[TaskParameter(ValidationType = TaskParameterValidationType.FileSpec)]
 		public string Files;
diff --git a/Engine/Source/Programs/AutomationTool/BuildGraph/Tasks/LogTask.cs b/Engine/Source/Programs/AutomationTool/BuildGraph/Tasks/LogTask.cs
--- a/Engine/Source/Programs/AutomationTool/BuildGraph/Tasks/LogTask.cs
+++ b/Engine/Source/Programs/AutomationTool/BuildGraph/Tasks/LogTask.cs
@@ -1,4 +1,4 @@
-﻿// Copyright 1998-2019 Epic Games, Inc. All Rights Reserved.
+// Copyright 1998-2019 Epic Games, Inc. All Rights Reserved.
 
 using AutomationTool;
 using System;
@@ -19,7 +19,7 @@ namespace BuildGraph.Tasks
 	public class LogTaskParameters
 	{
 		/// <summary>
-		/// Message to print out
+		/// Message to print out.
 		/// </summary>
 		[TaskParameter(Optional = true)]
 		public string Message;
@@ -38,7 +38,7 @@ namespace BuildGraph.Tasks
 	}
 
 	/// <summary>
-	/// Print a message (and other optional diagnostic information) to the output log
+	/// Print a message (and other optional diagnostic information) to the output log.
 	/// </summary>
 	[TaskElement("Log", typeof(LogTaskParameters))]
 	public class LogTask : CustomTask
diff --git a/Engine/Source/Programs/AutomationTool/BuildGraph/Tasks/MoveTask.cs b/Engine/Source/Programs/AutomationTool/BuildGraph/Tasks/MoveTask.cs
--- a/Engine/Source/Programs/AutomationTool/BuildGraph/Tasks/MoveTask.cs
+++ b/Engine/Source/Programs/AutomationTool/BuildGraph/Tasks/MoveTask.cs
@@ -18,25 +18,25 @@ namespace BuildGraph.Tasks
 	public class MoveTaskParameters
 	{
 		/// <summary>
-		/// Filter to be applied to the list of input files. Optional.
+		/// Optional filter to be applied to the list of input files.
 		/// </summary>
 		[TaskParameter(Optional = true, ValidationType = TaskParameterValidationType.FileSpec)]
 		public string Files;
 
 		/// <summary>
-		/// The pattern(s) to copy from (eg. Engine/*.txt)
+		/// The pattern(s) to copy from (for example, Engine/*.txt).
 		/// </summary>
 		[TaskParameter(ValidationType = TaskParameterValidationType.FileSpec)]
 		public string From;
 
 		/// <summary>
-		/// The directory or to copy to
+		/// The directory to copy to.
 		/// </summary>
 		[TaskParameter(ValidationType = TaskParameterValidationType.FileSpec)]
 		public string To;
 
 		/// <summary>
-		/// Tag to be applied to build products of this task
+		/// Tag to be applied to build products of this task.
 		/// </summary>
 		[TaskParameter(Optional = true, ValidationType = TaskParameterValidationType.TagList)]
 		public string Tag;
diff --git a/Engine/Source/Programs/AutomationTool/BuildGraph/Tasks/MsBuildTask.cs b/Engine/Source/Programs/AutomationTool/BuildGraph/Tasks/MsBuildTask.cs
--- a/Engine/Source/Programs/AutomationTool/BuildGraph/Tasks/MsBuildTask.cs
+++ b/Engine/Source/Programs/AutomationTool/BuildGraph/Tasks/MsBuildTask.cs
@@ -1,4 +1,4 @@
-﻿// Copyright 1998-2019 Epic Games, Inc. All Rights Reserved.
+// Copyright 1998-2019 Epic Games, Inc. All Rights Reserved.
 
 using System;
 using System.Collections.Generic;
@@ -15,32 +15,32 @@ namespace AutomationTool.Tasks
 	/// </summary>
 	public class MsBuildTaskParameters
 	{
-		/// <summary>
-		/// The C# project file to be compile. More than one project file can be specified by separating with semicolons.
-		/// </summary>
-		[TaskParameter]
+        /// <summary>
+        /// The C# project file to compile. Using semicolons, more than one project file can be specified.
+        /// </summary>
+        [TaskParameter]
 		public string Project;
 
 		/// <summary>
-		/// The configuration to compile
+		/// The configuration to compile.
 		/// </summary>
 		[TaskParameter(Optional = true)]
 		public string Configuration;
 
 		/// <summary>
-		/// The platform to compile
+		/// The platform to compile.
 		/// </summary>
 		[TaskParameter(Optional = true)]
 		public string Platform;
 
 		/// <summary>
-		/// Additional options to pass to the compiler
+		/// Additional options to pass to the compiler.
 		/// </summary>
 		[TaskParameter(Optional = true)]
 		public string Arguments;
 
 		/// <summary>
-		/// The MSBuild output verbosity
+		/// The MSBuild output verbosity.
 		/// </summary>
 		[TaskParameter(Optional = true)]
 		public string Verbosity = "minimal";
diff --git a/Engine/Source/Programs/AutomationTool/BuildGraph/Tasks/PakFileTask.cs b/Engine/Source/Programs/AutomationTool/BuildGraph/Tasks/PakFileTask.cs
--- a/Engine/Source/Programs/AutomationTool/BuildGraph/Tasks/PakFileTask.cs
+++ b/Engine/Source/Programs/AutomationTool/BuildGraph/Tasks/PakFileTask.cs
@@ -19,19 +19,19 @@ namespace BuildGraph.Tasks
 	public class PakFileTaskParameters
 	{
 		/// <summary>
-		/// List of files, wildcards and tag sets to add to the pak file, separated by ';' characters.
+		/// List of files, wildcards, and tag sets to add to the pak file, separated by ';' characters.
 		/// </summary>
 		[TaskParameter(ValidationType = TaskParameterValidationType.FileSpec)]
 		public string Files;
 
 		/// <summary>
-		/// PAK file to output
+		/// PAK file to output.
 		/// </summary>
 		[TaskParameter]
 		public FileReference Output;
 
 		/// <summary>
-		/// Path to a Response File that contains a list of files to add to the pak file, instead of specifying them individually
+		/// Path to a Response File that contains a list of files to add to the pak file -- instead of specifying them individually.
 		/// </summary>
 		[TaskParameter(Optional = true)]
 		public FileReference ResponseFile;
@@ -43,31 +43,31 @@ namespace BuildGraph.Tasks
 		public HashSet<DirectoryReference> RebaseDir;
 
 		/// <summary>
-		/// Script which gives the order of files
+		/// Script that gives the order of files.
 		/// </summary>
 		[TaskParameter(Optional = true)]
 		public FileReference Order;
 
 		/// <summary>
-		/// Encryption keys for this pak file
+		/// Encryption keys for this pak file.
 		/// </summary>
 		[TaskParameter(Optional = true)]
 		public FileReference Sign;
 
 		/// <summary>
-		/// Whether to compress files
+		/// Whether to compress files.
 		/// </summary>
 		[TaskParameter(Optional = true)]
 		public bool Compress = true;
 
 		/// <summary>
-		/// Additional arguments to be passed to UnrealPak
+		/// Additional arguments to pass to UnrealPak.
 		/// </summary>
 		[TaskParameter(Optional = true)]
 		public string Arguments = "";
 
 		/// <summary>
-		/// Tag to be applied to build products of this task
+		/// Tag to be applied to build products of this task.
 		/// </summary>
 		[TaskParameter(Optional = true, ValidationType = TaskParameterValidationType.TagList)]
 		public string Tag;
diff --git a/Engine/Source/Programs/AutomationTool/BuildGraph/Tasks/RenameTask.cs b/Engine/Source/Programs/AutomationTool/BuildGraph/Tasks/RenameTask.cs
--- a/Engine/Source/Programs/AutomationTool/BuildGraph/Tasks/RenameTask.cs
+++ b/Engine/Source/Programs/AutomationTool/BuildGraph/Tasks/RenameTask.cs
@@ -1,4 +1,4 @@
-﻿// Copyright 1998-2019 Epic Games, Inc. All Rights Reserved.
+// Copyright 1998-2019 Epic Games, Inc. All Rights Reserved.
 
 using AutomationTool;
 using System;
@@ -20,13 +20,13 @@ namespace BuildGraph.Tasks
 	public class RenameTaskParameters
 	{
 		/// <summary>
-		/// The file or files to rename
+		/// The file or files to rename.
 		/// </summary>
 		[TaskParameter(ValidationType = TaskParameterValidationType.FileSpec)]
 		public string Files;
 
 		/// <summary>
-		/// The current file name, or pattern to match (eg. *.txt). Should not include any path separators.
+		/// The current file name, or pattern to match (for example, *.txt). Should not include any path separators.
 		/// </summary>
 		[TaskParameter(Optional = true)]
 		public string From;
@@ -38,7 +38,7 @@ namespace BuildGraph.Tasks
 		public string To;
 
 		/// <summary>
-		/// Tag to be applied to the renamed files
+		/// Tag to be applied to the renamed files.
 		/// </summary>
 		[TaskParameter(Optional = true, ValidationType = TaskParameterValidationType.TagList)]
 		public string Tag;
diff --git a/Engine/Source/Programs/AutomationTool/BuildGraph/Tasks/SanitizeReceiptTask.cs b/Engine/Source/Programs/AutomationTool/BuildGraph/Tasks/SanitizeReceiptTask.cs
--- a/Engine/Source/Programs/AutomationTool/BuildGraph/Tasks/SanitizeReceiptTask.cs
+++ b/Engine/Source/Programs/AutomationTool/BuildGraph/Tasks/SanitizeReceiptTask.cs
@@ -30,7 +30,7 @@ namespace AutomationTool.Tasks
 	}
 
 	/// <summary>
-	/// Task which tags build products and/or runtime dependencies by reading from *.target files.
+	/// Task that tags build products and/or runtime dependencies by reading from *.target files.
 	/// </summary>
 	[TaskElement("SanitizeReceipt", typeof(SanitizeReceiptTaskParameters))]
 	class SanitizeReceiptTask : CustomTask
diff --git a/Engine/Source/Programs/AutomationTool/BuildGraph/Tasks/SetVersionTask.cs b/Engine/Source/Programs/AutomationTool/BuildGraph/Tasks/SetVersionTask.cs
--- a/Engine/Source/Programs/AutomationTool/BuildGraph/Tasks/SetVersionTask.cs
+++ b/Engine/Source/Programs/AutomationTool/BuildGraph/Tasks/SetVersionTask.cs
@@ -1,4 +1,4 @@
-﻿// Copyright 1998-2019 Epic Games, Inc. All Rights Reserved.
+// Copyright 1998-2019 Epic Games, Inc. All Rights Reserved.
 
 using System;
 using System.Collections.Generic;
@@ -17,49 +17,49 @@ namespace AutomationTool.Tasks
 	public class SetVersionTaskParameters
 	{
 		/// <summary>
-		/// The changelist to set in the version files
+		/// The changelist to set in the version files.
 		/// </summary>
 		[TaskParameter]
 		public int Change;
 
 		/// <summary>
-		/// The engine compatible changelist to set in the version files
+		/// The engine compatible changelist to set in the version files.
 		/// </summary>
 		[TaskParameter(Optional = true)]
 		public int CompatibleChange;
 
 		/// <summary>
-		/// The branch string
+		/// The branch string.
 		/// </summary>
 		[TaskParameter]
 		public string Branch;
 
 		/// <summary>
-		/// The build version string
+		/// The build version string.
 		/// </summary>
 		[TaskParameter(Optional = true)]
 		public string Build;
 
 		/// <summary>
-		/// Whether to set the IS_LICENSEE_VERSION flag to true
+		/// Whether to set the IS_LICENSEE_VERSION flag to true.
 		/// </summary>
 		[TaskParameter(Optional = true)]
 		public bool Licensee;
 
 		/// <summary>
-		/// Whether to set the ENGINE_IS_PROMOTED_BUILD flag to true
+		/// Whether to set the ENGINE_IS_PROMOTED_BUILD flag to true.
 		/// </summary>
 		[TaskParameter(Optional = true)]
 		public bool Promoted = true;
 
 		/// <summary>
-		/// If set, don't actually write to the files - just return the version files that would be updated. Useful for local builds.
+		/// If set, do not write to the files -- just return the version files that would be updated. Useful for local builds.
 		/// </summary>
 		[TaskParameter(Optional = true)]
 		public bool SkipWrite;
 
 		/// <summary>
-		/// Tag to be applied to build products of this task
+		/// Tag to be applied to build products of this task.
 		/// </summary>
 		[TaskParameter(Optional = true, ValidationType = TaskParameterValidationType.TagList)]
 		public string Tag;
diff --git a/Engine/Source/Programs/AutomationTool/BuildGraph/Tasks/SignTask.cs b/Engine/Source/Programs/AutomationTool/BuildGraph/Tasks/SignTask.cs
--- a/Engine/Source/Programs/AutomationTool/BuildGraph/Tasks/SignTask.cs
+++ b/Engine/Source/Programs/AutomationTool/BuildGraph/Tasks/SignTask.cs
@@ -1,4 +1,4 @@
-﻿// Copyright 1998-2019 Epic Games, Inc. All Rights Reserved.
+// Copyright 1998-2019 Epic Games, Inc. All Rights Reserved.
 
 using AutomationTool;
 using System;
@@ -18,13 +18,13 @@ namespace BuildGraph.Tasks
 	public class SignTaskParameters
 	{
 		/// <summary>
-		/// List of file specifications separated by semicolons (eg. *.cpp;Engine/.../*.bat), or the name of a tag set
+		/// List of file specifications separated by semicolons (for example, *.cpp;Engine/.../*.bat), or the name of a tag set.
 		/// </summary>
 		[TaskParameter(ValidationType = TaskParameterValidationType.FileSpec)]
 		public string Files;
 
 		/// <summary>
-		/// Tag to be applied to build products of this task
+		/// Tag to be applied to build products of this task.
 		/// </summary>
 		[TaskParameter(Optional = true, ValidationType = TaskParameterValidationType.TagList)]
 		public string Tag;
diff --git a/Engine/Source/Programs/AutomationTool/BuildGraph/Tasks/SpawnTask.cs b/Engine/Source/Programs/AutomationTool/BuildGraph/Tasks/SpawnTask.cs
--- a/Engine/Source/Programs/AutomationTool/BuildGraph/Tasks/SpawnTask.cs
+++ b/Engine/Source/Programs/AutomationTool/BuildGraph/Tasks/SpawnTask.cs
@@ -19,19 +19,19 @@ namespace BuildGraph.Tasks
 	public class SpawnTaskParameters
 	{
 		/// <summary>
-		/// Executable to spawn
+		/// Executable to spawn.
 		/// </summary>
 		[TaskParameter]
 		public string Exe;
 
 		/// <summary>
-		/// Arguments for the newly created process
+		/// Arguments for the newly created process.
 		/// </summary>
 		[TaskParameter(Optional = true)]
 		public string Arguments;
 
 		/// <summary>
-		/// The minimum exit code which is treated as an error.
+		/// The minimum exit code, which is treated as an error.
 		/// </summary>
 		[TaskParameter(Optional = true)]
 		public int ErrorLevel = 1;
diff --git a/Engine/Source/Programs/AutomationTool/BuildGraph/Tasks/StageTask.cs b/Engine/Source/Programs/AutomationTool/BuildGraph/Tasks/StageTask.cs
--- a/Engine/Source/Programs/AutomationTool/BuildGraph/Tasks/StageTask.cs
+++ b/Engine/Source/Programs/AutomationTool/BuildGraph/Tasks/StageTask.cs
@@ -18,49 +18,49 @@ namespace BuildGraph.Tasks
 	public class StageTaskParameters
 	{
 		/// <summary>
-		/// The project that this target belongs to
+		/// The project that this target belongs to.
 		/// </summary>
 		[TaskParameter(Optional = true)]
 		public FileReference Project;
 
 		/// <summary>
-		/// Name of the target to stage
+		/// Name of the target to stage.
 		/// </summary>
 		[TaskParameter]
 		public string Target;
 
 		/// <summary>
-		/// Platform to stage
+		/// Platform to stage.
 		/// </summary>
 		[TaskParameter]
 		public UnrealTargetPlatform Platform;
 
 		/// <summary>
-		/// Configuration to be staged
+		/// Configuration to be staged.
 		/// </summary>
 		[TaskParameter]
 		public UnrealTargetConfiguration Configuration;
 
 		/// <summary>
-		/// Architecture to be staged
+		/// Architecture to be staged.
 		/// </summary>
 		[TaskParameter(Optional = true)]
 		public string Architecture;
 
 		/// <summary>
-		/// Directory the receipt files should be staged to
+		/// Directory that the receipt files should be staged to.
 		/// </summary>
 		[TaskParameter]
 		public DirectoryReference ToDir;
 
 		/// <summary>
-		/// Whether to overwrite existing files
+		/// Whether to overwrite existing files.
 		/// </summary>
 		[TaskParameter(Optional = true)]
 		public bool Overwrite;
 
 		/// <summary>
-		/// Tag to be applied to build products of this task
+		/// Tag to be applied to build products of this task.
 		/// </summary>
 		[TaskParameter(Optional = true, ValidationType = TaskParameterValidationType.TagList)]
 		public string Tag;
diff --git a/Engine/Source/Programs/AutomationTool/BuildGraph/Tasks/StripTask.cs b/Engine/Source/Programs/AutomationTool/BuildGraph/Tasks/StripTask.cs
--- a/Engine/Source/Programs/AutomationTool/BuildGraph/Tasks/StripTask.cs
+++ b/Engine/Source/Programs/AutomationTool/BuildGraph/Tasks/StripTask.cs
@@ -18,31 +18,31 @@ namespace BuildGraph.Tasks
 	public class StripTaskParameters
 	{
 		/// <summary>
-		/// The platform toolchain to strip binaries
+		/// The platform toolchain to strip binaries.
 		/// </summary>
 		[TaskParameter]
 		public UnrealTargetPlatform Platform;
 
 		/// <summary>
-		/// The directory to find files in
+		/// The directory to find files in.
 		/// </summary>
 		[TaskParameter(Optional = true)]
 		public DirectoryReference BaseDir;
 
 		/// <summary>
-		/// List of file specifications separated by semicolons (eg. Engine/.../*.pdb), or the name of a tag set
+		/// List of file specifications separated by semicolons (for example, Engine/.../*.pdb), or the name of a tag set.
 		/// </summary>
 		[TaskParameter(ValidationType = TaskParameterValidationType.FileSpec)]
 		public string Files;
 
 		/// <summary>
-		/// Output directory for the stripped files. Defaults to the input path (overwriting the input files).
+		/// Output directory for the stripped files. Defaults to the input path, overwriting the input files.
 		/// </summary>
 		[TaskParameter(Optional = true)]
 		public DirectoryReference OutputDir;
 
 		/// <summary>
-		/// Tag to be applied to build products of this task
+		/// Tag to be applied to build products of this task.
 		/// </summary>
 		[TaskParameter(Optional = true, ValidationType = TaskParameterValidationType.TagList)]
 		public string Tag;
diff --git a/Engine/Source/Programs/AutomationTool/BuildGraph/Tasks/SubmitTask.cs b/Engine/Source/Programs/AutomationTool/BuildGraph/Tasks/SubmitTask.cs
--- a/Engine/Source/Programs/AutomationTool/BuildGraph/Tasks/SubmitTask.cs
+++ b/Engine/Source/Programs/AutomationTool/BuildGraph/Tasks/SubmitTask.cs
@@ -17,19 +17,19 @@ namespace AutomationTool.Tasks
 	public class SubmitTaskParameters
 	{
 		/// <summary>
-		/// The description for the submitted changelist
+		/// The description for the submitted changelist.
 		/// </summary>
 		[TaskParameter]
 		public string Description;
 
 		/// <summary>
-		/// The files to submit
+		/// The files to submit.
 		/// </summary>
 		[TaskParameter(ValidationType = TaskParameterValidationType.FileSpec)]
 		public string Files;
 
 		/// <summary>
-		/// The Perforce file type for the submitted files (eg. binary+FS32)
+		/// The Perforce file type for the submitted files (for example, binary+FS32).
 		/// </summary>
 		[TaskParameter(Optional = true)]
 		public string FileType;
@@ -41,7 +41,7 @@ namespace AutomationTool.Tasks
 		public string Workspace;
 
 		/// <summary>
-		/// The stream for the workspace; defaults to the current stream. Ignored unless the Workspace attribute is also specified.
+		/// The stream for the workspace -- defaults to the current stream. Ignored unless the Workspace attribute is also specified.
 		/// </summary>
 		[TaskParameter(Optional = true)]
 		public string Stream;
@@ -53,13 +53,13 @@ namespace AutomationTool.Tasks
 		public DirectoryReference RootDir;
 
 		/// <summary>
-		/// Whether to revert unchanged files before attempting to submit
+		/// Whether to revert unchanged files before attempting to submit.
 		/// </summary>
 		[TaskParameter(Optional = true)]
 		public bool RevertUnchanged;
 
 		/// <summary>
-		/// Force the submit to happen, even if a resolve is needed (always accept current version)
+		/// Force the submit to happen -- even if a resolve is needed (always accept current version).
 		/// </summary>
 		[TaskParameter(Optional = true)]
 		public bool Force;
diff --git a/Engine/Source/Programs/AutomationTool/BuildGraph/Tasks/SymStoreTask.cs b/Engine/Source/Programs/AutomationTool/BuildGraph/Tasks/SymStoreTask.cs
--- a/Engine/Source/Programs/AutomationTool/BuildGraph/Tasks/SymStoreTask.cs
+++ b/Engine/Source/Programs/AutomationTool/BuildGraph/Tasks/SymStoreTask.cs
@@ -1,4 +1,4 @@
-﻿// Copyright 1998-2019 Epic Games, Inc. All Rights Reserved.
+// Copyright 1998-2019 Epic Games, Inc. All Rights Reserved.
 
 using System;
 using System.Collections.Generic;
@@ -42,14 +42,14 @@ namespace AutomationTool
         public string Product;
 
 		/// <summary>
-		/// BuildVersion associated with these symbols. Used for cleanup in AgeStore by matching this version against a directory name in a build share
+		/// BuildVersion associated with these symbols. Used for clean-up in AgeStore by matching this version against a directory name in a build share.
 		/// </summary>
 		[TaskParameter(Optional = true)]
 		public string BuildVersion;
 	}
 
     /// <summary>
-    /// Task which strips symbols from a set of files
+    /// Task that strips symbols from a set of files.
     /// </summary>
     [TaskElement("SymStore", typeof(SymStoreTaskParameters))]
     public class SymStoreTask : CustomTask
diff --git a/Engine/Source/Programs/AutomationTool/BuildGraph/Tasks/TagReceiptTask.cs b/Engine/Source/Programs/AutomationTool/BuildGraph/Tasks/TagReceiptTask.cs
--- a/Engine/Source/Programs/AutomationTool/BuildGraph/Tasks/TagReceiptTask.cs
+++ b/Engine/Source/Programs/AutomationTool/BuildGraph/Tasks/TagReceiptTask.cs
@@ -30,44 +30,44 @@ namespace AutomationTool.Tasks
 		public DirectoryReference EngineDir;
 
 		/// <summary>
-		/// Path to the project folder, used to expand $(ProjectDir) properties in receipt files. Defaults to the Engine directory for the current workspace. DEPRECATED.
+		/// Path to the project folder, used to expand $(ProjectDir) properties in receipt files. Defaults to the Engine directory for the current workspace -- DEPRECATED.
 		/// </summary>
 		[TaskParameter(Optional = true)]
 		public DirectoryReference ProjectDir;
 
 		/// <summary>
-		/// Whether to tag the Build Products listed in receipts
+		/// Whether to tag the Build Products listed in receipts.
 		/// </summary>
 		[TaskParameter(Optional = true)]
 		public bool BuildProducts;
 
 		/// <summary>
-		/// Which type of Build Products to tag (See TargetReceipt.cs - UnrealBuildTool.BuildProductType for valid values)
+		/// Which type of Build Products to tag (see TargetReceipt.cs - UnrealBuildTool.BuildProductType for valid values).
 		/// </summary>
 		[TaskParameter(Optional = true)]
 		public string BuildProductType;
 
 		/// <summary>
-		/// Whether to tag the Runtime Dependencies listed in receipts
+		/// Whether to tag the Runtime Dependencies listed in receipts.
 		/// </summary>
 		[TaskParameter(Optional = true)]
 		public bool RuntimeDependencies;
 
 		/// <summary>
-		/// Which type of Runtime Dependencies to tag (See TargetReceipt.cs - UnrealBuildTool.StagedFileType for valid values)
+		/// Which type of Runtime Dependencies to tag (see TargetReceipt.cs - UnrealBuildTool.StagedFileType for valid values).
 		/// </summary>
 		[TaskParameter(Optional = true)]
 		public string StagedFileType;
 
 		/// <summary>
-		/// Name of the tag to apply
+		/// Name of the tag to apply.
 		/// </summary>
 		[TaskParameter(ValidationType = TaskParameterValidationType.TagList)]
 		public string With;
 	}
 
 	/// <summary>
-	/// Task which tags build products and/or runtime dependencies by reading from *.target files.
+	/// Task that tags build products and/or runtime dependencies by reading from *.target files.
 	/// </summary>
 	[TaskElement("TagReceipt", typeof(TagReceiptTaskParameters))]
 	class TagReceiptTask : CustomTask
diff --git a/Engine/Source/Programs/AutomationTool/BuildGraph/Tasks/TagTask.cs b/Engine/Source/Programs/AutomationTool/BuildGraph/Tasks/TagTask.cs
--- a/Engine/Source/Programs/AutomationTool/BuildGraph/Tasks/TagTask.cs
+++ b/Engine/Source/Programs/AutomationTool/BuildGraph/Tasks/TagTask.cs
@@ -18,25 +18,25 @@ namespace AutomationTool.Tasks
 	public class TagTaskParameters
 	{
 		/// <summary>
-		/// Set the base directory to resolve relative paths and patterns against. If set, any absolute patterns (eg. /Engine/Build/...) are taken to be relative to this path. If not, they are taken to be truly absolute.
+		/// Set the base directory to resolve relative paths and patterns against. If set, any absolute patterns (for example, /Engine/Build/...) are taken to be relative to this path. If not, they are taken to be truly absolute.
 		/// </summary>
 		[TaskParameter(Optional = true)]
 		public DirectoryReference BaseDir;
 
 		/// <summary>
-		/// Set of files to work from, including wildcards and tag names, separated by semicolons. Resolved relative to BaseDir if set, otherwise to the branch root directory.
+		/// Set of files to work from, including wildcards and tag names, separated by semicolons. If set, resolved relative to BaseDir, otherwise resolved to the branch root directory.
 		/// </summary>
 		[TaskParameter(Optional = true, ValidationType = TaskParameterValidationType.FileSpec)]
 		public string Files;
 
 		/// <summary>
-		/// Set of text files to add additional files from. Each file list should have on file per line.
+		/// Set of text files to add additional files from. Each file list should have one file per line.
 		/// </summary>
 		[TaskParameter(Optional = true, ValidationType = TaskParameterValidationType.FileSpec)]
 		public string FileLists;
 
 		/// <summary>
-		/// Patterns to filter the list of files by, including tag names or wildcards. May include patterns that apply to the base directory if set. Defaults to all files if not specified.
+		/// Patterns to filter the list of files by, including tag names or wildcards. If set, may include patterns that apply to the base directory. If not specified, defaults to all files.
 		/// </summary>
 		[TaskParameter(Optional = true, ValidationType = TaskParameterValidationType.FileSpec)]
 		public string Filter;
@@ -48,7 +48,7 @@ namespace AutomationTool.Tasks
 		public string Except;
 
 		/// <summary>
-		/// Name of the tag to apply
+		/// Name of the tag to apply.
 		/// </summary>
 		[TaskParameter(ValidationType = TaskParameterValidationType.TagList)]
 		public string With;
diff --git a/Engine/Source/Programs/AutomationTool/BuildGraph/Tasks/UnzipTask.cs b/Engine/Source/Programs/AutomationTool/BuildGraph/Tasks/UnzipTask.cs
--- a/Engine/Source/Programs/AutomationTool/BuildGraph/Tasks/UnzipTask.cs
+++ b/Engine/Source/Programs/AutomationTool/BuildGraph/Tasks/UnzipTask.cs
@@ -18,19 +18,19 @@ namespace BuildGraph.Tasks
 	public class UnzipTaskParameters
 	{
 		/// <summary>
-		/// Path to the zip file to extract
+		/// Path to the zip file to extract.
 		/// </summary>
 		[TaskParameter(ValidationType = TaskParameterValidationType.FileSpec)]
 		public string ZipFile;
 
 		/// <summary>
-		/// Output directory for the extracted files
+		/// Output directory for the extracted files.
 		/// </summary>
 		[TaskParameter]
 		public DirectoryReference ToDir;
 
 		/// <summary>
-		/// Tag to be applied to the extracted files
+		/// Tag to be applied to the extracted files.
 		/// </summary>
 		[TaskParameter(Optional = true, ValidationType = TaskParameterValidationType.TagList)]
 		public string Tag;
diff --git a/Engine/Source/Programs/AutomationTool/BuildGraph/Tasks/ZipTask.cs b/Engine/Source/Programs/AutomationTool/BuildGraph/Tasks/ZipTask.cs
--- a/Engine/Source/Programs/AutomationTool/BuildGraph/Tasks/ZipTask.cs
+++ b/Engine/Source/Programs/AutomationTool/BuildGraph/Tasks/ZipTask.cs
@@ -18,25 +18,25 @@ namespace BuildGraph.Tasks
 	public class ZipTaskParameters
 	{
 		/// <summary>
-		/// The directory to read compressed files from
+		/// The directory to read compressed files from.
 		/// </summary>
 		[TaskParameter]
 		public DirectoryReference FromDir;
 
 		/// <summary>
-		/// List of file specifications separated by semicolons (eg. *.cpp;Engine/.../*.bat), or the name of a tag set. Relative paths are taken from FromDir.
+		/// List of file specifications separated by semicolons (for example, *.cpp;Engine/.../*.bat), or the name of a tag set. Relative paths are taken from FromDir.
 		/// </summary>
 		[TaskParameter(Optional = true, ValidationType = TaskParameterValidationType.FileSpec)]
 		public string Files;
 
 		/// <summary>
-		/// The zip file to create
+		/// The zip file to create.
 		/// </summary>
 		[TaskParameter]
 		public FileReference ZipFile;
 
 		/// <summary>
-		/// Tag to be applied to the created zip file
+		/// Tag to be applied to the created zip file.
 		/// </summary>
 		[TaskParameter(Optional = true, ValidationType = TaskParameterValidationType.TagList)]
 		public string Tag;
diff --git a/Engine/Source/Programs/UnrealBuildTool/Configuration/BuildConfiguration.cs b/Engine/Source/Programs/UnrealBuildTool/Configuration/BuildConfiguration.cs
--- a/Engine/Source/Programs/UnrealBuildTool/Configuration/BuildConfiguration.cs
+++ b/Engine/Source/Programs/UnrealBuildTool/Configuration/BuildConfiguration.cs
@@ -14,13 +14,13 @@ namespace UnrealBuildTool
 	/// </summary>
 	class BuildConfiguration
 	{
-		/// <summary>
-		/// Whether to ignore import library files that are out of date when building targets.  Set this to true to improve iteration time.
-		/// By default we don't bother relinking targets if only a dependent .lib has changed, as chances are that
-		/// the import library wasn't actually different unless a dependent header file of this target was also changed,
-		/// in which case the target would be rebuilt automatically.
-		/// </summary>
-		[XmlConfigFile]
+        /// <summary>
+        /// Whether to ignore import library files that are out of date when building targets. Set this to true to improve iteration time.
+        /// By default, we do not bother re-linking targets if only a dependent .lib has changed, as chances are that
+        /// the import library was not actually different unless a dependent header file of this target was also changed,
+        /// in which case the target would automatically be rebuilt.
+        /// </summary>
+        [XmlConfigFile]
 		public bool bIgnoreOutdatedImportLibraries = true;
 
 		/// <summary>
@@ -42,7 +42,7 @@ namespace UnrealBuildTool
 		public bool bLogDetailedActionStats = false;
 
 		/// <summary>
-		/// Whether the hybrid executor will be used (a remote executor and local executor)
+		/// Whether the hybrid executor will be used (a remote executor and local executor).
 		/// </summary>
 		[XmlConfigFile]
 		public bool bAllowHybridExecutor = false;
@@ -61,25 +61,25 @@ namespace UnrealBuildTool
 		public bool bAllowSNDBS = true;
 
 		/// <summary>
-		/// Enables support for very fast iterative builds by caching target data.  Turning this on causes Unreal Build Tool to emit
-		/// 'UBT Makefiles' for targets when they're built the first time.  Subsequent builds will load these Makefiles and begin
-		/// outdatedness checking and build invocation very quickly.  The caveat is that if source files are added or removed to
-		/// the project, UBT will need to gather information about those in order for your build to complete successfully.  Currently,
+		/// Enables support for very fast iterative builds by caching target data. Turning this on causes Unreal Build Tool to emit
+		/// 'UBT Makefiles' for targets when they are built the first time. Subsequent builds will load these Makefiles and begin
+		/// outdatedness checking and build invocation very quickly. The caveat is that if source files are added or removed to
+		/// the project, UBT will need to gather information about those in order for your build to complete successfully. Currently,
 		/// you must run the project file generator after adding/removing source files to tell UBT to re-gather this information.
 		/// 
 		/// Events that can invalidate the 'UBT Makefile':  
 		///		- Adding/removing .cpp files
 		///		- Adding/removing .h files with UObjects
-		///		- Adding new UObject types to a file that didn't previously have any
-		///		- Changing global build settings (most settings in this file qualify.)
+		///		- Adding new UObject types to a file that did not previously have any
+		///		- Changing global build settings (most settings in this file qualify)
 		///		- Changed code that affects how Unreal Header Tool works
 		///	
-		///	You can force regeneration of the 'UBT Makefile' by passing the '-gather' argument, or simply regenerating project files
+		///	You can force regeneration of the 'UBT Makefile' by passing the '-gather' argument, or simply regenerating project files.
 		///
 		///	This also enables the fast include file dependency scanning and caching system that allows Unreal Build Tool to detect out 
-		/// of date dependencies very quickly.  When enabled, a deep C++ include graph does not have to be generated, and instead
-		/// we only scan and cache indirect includes for after a dependent build product was already found to be out of date.  During the
-		/// next build, we'll load those cached indirect includes and check for outdatedness.
+		/// of date dependencies very quickly. When enabled, a deep C++ include graph does not have to be generated, and instead,
+		/// we only scan and cache indirect includes for after a dependent build product was already found to be out of date. During the
+		/// next build, we will load those cached indirect includes and check for outdatedness.
 		/// </summary>
 		[XmlConfigFile]
 		[CommandLine("-NoUBTMakefiles", Value = "false")]
@@ -87,7 +87,7 @@ namespace UnrealBuildTool
 
 		/// <summary>
 		/// Whether DMUCS/Distcc may be used.
-		/// Distcc requires some setup - so by default disable it so we don't break local or remote building
+		/// Distcc requires some setup -- so by default, disable it so that we do not break local or remote building.
 		/// </summary>
 		[XmlConfigFile]
 		public bool bAllowDistcc = false;
@@ -99,14 +99,14 @@ namespace UnrealBuildTool
 		public bool bAllowParallelExecutor = true;
 
 		/// <summary>
-		/// If true, force header regeneration. Intended for the build machine
+		/// If true, force header regeneration. Intended for the build machine.
 		/// </summary>
 		[CommandLine("-ForceHeaderGeneration")]
 		[XmlConfigFile(Category = "UEBuildConfiguration")]
 		public bool bForceHeaderGeneration = false;
 
 		/// <summary>
-		/// If true, do not build UHT, assume it is already built
+		/// If true, do not build UHT, assume it is already built.
 		/// </summary>
 		[CommandLine("-NoBuildUHT")]
 		[XmlConfigFile(Category = "UEBuildConfiguration")]
@@ -120,14 +120,14 @@ namespace UnrealBuildTool
 		public bool bFailIfGeneratedCodeChanges = false;
 
 		/// <summary>
-		/// True if hot-reload from IDE is allowed
+		/// True if hot-reload from IDE is allowed.
 		/// </summary>
 		[CommandLine("-NoHotReloadFromIDE", Value="false")]
 		[XmlConfigFile(Category = "UEBuildConfiguration")]
 		public bool bAllowHotReloadFromIDE = true;
 
 		/// <summary>
-		/// If true, the Debug version of UnrealHeaderTool will be build and run instead of the Development version.
+		/// If true, the Debug version of UnrealHeaderTool will be built and run instead of the Development version.
 		/// </summary>
 		[XmlConfigFile(Category = "UEBuildConfiguration")]
 		public bool bForceDebugUnrealHeaderTool = false;
diff --git a/Engine/Source/Programs/UnrealBuildTool/Executors/Distcc.cs b/Engine/Source/Programs/UnrealBuildTool/Executors/Distcc.cs
--- a/Engine/Source/Programs/UnrealBuildTool/Executors/Distcc.cs
+++ b/Engine/Source/Programs/UnrealBuildTool/Executors/Distcc.cs
@@ -15,7 +15,7 @@ namespace UnrealBuildTool
 	class Distcc : ActionExecutor
 	{
 		/// <summary>
-		/// When enabled allows DMUCS/Distcc to fallback to local compilation when remote compiling fails. Defaults to true as separation of pre-process and compile stages can introduce non-fatal errors.
+		/// When enabled, allows DMUCS/Distcc to fallback to local compilation when remote compiling fails. Defaults to true as separation of pre-process and compile stages can introduce non-fatal errors.
 		/// </summary>
 		[XmlConfigFile(Category = "BuildConfiguration")]
 		public bool bAllowDistccLocalFallback = true;
@@ -33,7 +33,7 @@ namespace UnrealBuildTool
 		public string DistccExecutablesPath = "/usr/local/bin";
 
 		/// <summary>
-		/// DMUCS coordinator hostname or IP address.
+		/// The DMUCS coordinator hostname or IP address.
 		/// </summary>
 		[XmlConfigFile(Category = "BuildConfiguration")]
 		public string DMUCSCoordinator = "localhost";
diff --git a/Engine/Source/Programs/UnrealBuildTool/Executors/HybridExecutor.cs b/Engine/Source/Programs/UnrealBuildTool/Executors/HybridExecutor.cs
--- a/Engine/Source/Programs/UnrealBuildTool/Executors/HybridExecutor.cs
+++ b/Engine/Source/Programs/UnrealBuildTool/Executors/HybridExecutor.cs
@@ -15,7 +15,7 @@ namespace UnrealBuildTool
 	class HybridExecutor : ActionExecutor
 	{
 		/// <summary>
-		/// Maximum number of actions to execute locally
+		/// Maximum number of actions to execute locally.
 		/// </summary>
 		[XmlConfigFile(Category = "HybridExecutor")]
 		int MaxLocalActions;
diff --git a/Engine/Source/Programs/UnrealBuildTool/Modes/BuildMode.cs b/Engine/Source/Programs/UnrealBuildTool/Modes/BuildMode.cs
--- a/Engine/Source/Programs/UnrealBuildTool/Modes/BuildMode.cs
+++ b/Engine/Source/Programs/UnrealBuildTool/Modes/BuildMode.cs
@@ -41,13 +41,13 @@ namespace UnrealBuildTool
 	class BuildMode : ToolMode
 	{
 		/// <summary>
-		/// Specifies the file to use for logging
+		/// Specifies the file to use for logging.
 		/// </summary>
 		[XmlConfigFile(Category = "BuildConfiguration")]
 		public string BaseLogFileName = "../Programs/UnrealBuildTool/Log.txt";
 
 		/// <summary>
-		/// Whether to skip checking for files identified by the junk manifest
+		/// Whether to skip checking for files identified by the junk manifest.
 		/// </summary>
 		[XmlConfigFile]
 		[CommandLine("-IgnoreJunk")]
diff --git a/Engine/Source/Programs/UnrealBuildTool/Modes/GenerateProjectFilesMode.cs b/Engine/Source/Programs/UnrealBuildTool/Modes/GenerateProjectFilesMode.cs
--- a/Engine/Source/Programs/UnrealBuildTool/Modes/GenerateProjectFilesMode.cs
+++ b/Engine/Source/Programs/UnrealBuildTool/Modes/GenerateProjectFilesMode.cs
@@ -40,7 +40,7 @@ namespace UnrealBuildTool
 		HashSet<ProjectFileFormat> ProjectFileFormats = new HashSet<ProjectFileFormat>();
 
 		/// <summary>
-		/// Platforms to disable native project file generators for. Platforms with native project file generators typically require IDE extensions to be installed.
+		/// Disable native project file generators for platforms. Platforms with native project file generators typically require IDE extensions to be installed.
 		/// </summary>
 		[XmlConfigFile(Category = "ProjectFileGenerator")]
 		string[] DisablePlatformProjectGenerators = null;
diff --git a/Engine/Source/Programs/UnrealBuildTool/Platform/IOS/IOSToolChain.cs b/Engine/Source/Programs/UnrealBuildTool/Platform/IOS/IOSToolChain.cs
--- a/Engine/Source/Programs/UnrealBuildTool/Platform/IOS/IOSToolChain.cs
+++ b/Engine/Source/Programs/UnrealBuildTool/Platform/IOS/IOSToolChain.cs
@@ -21,14 +21,14 @@ namespace UnrealBuildTool
 	class IOSToolChainSettings : AppleToolChainSettings
 	{
 		/// <summary>
-		/// Which version of the iOS SDK to target at build time
+		/// The version of the iOS SDK to target at build time.
 		/// </summary>
 		[XmlConfigFile(Category = "IOSToolChain")]
 		public string IOSSDKVersion = "latest";
 		public readonly float IOSSDKVersionFloat = 0.0f;
 
 		/// <summary>
-		/// Which version of the iOS to allow at build time
+		/// The version of the iOS to allow at build time.
 		/// </summary>
 		[XmlConfigFile(Category = "IOSToolChain")]
 		public string BuildIOSVersion = "7.0";
@@ -99,7 +99,7 @@ namespace UnrealBuildTool
 		// If you are looking for where to change the remote compile server name, look in RemoteToolChain.cs
 
 		/// <summary>
-		/// If this is set, then we don't do any post-compile steps except moving the executable into the proper spot on the Mac
+		/// If this is set, then we do not do any post-compile steps -- except moving the executable into the proper spot on Mac.
 		/// </summary>
 		[XmlConfigFile]
 		public static bool bUseDangerouslyFastMode = false;
diff --git a/Engine/Source/Programs/UnrealBuildTool/Platform/IOS/UEBuildIOS.cs b/Engine/Source/Programs/UnrealBuildTool/Platform/IOS/UEBuildIOS.cs
--- a/Engine/Source/Programs/UnrealBuildTool/Platform/IOS/UEBuildIOS.cs
+++ b/Engine/Source/Programs/UnrealBuildTool/Platform/IOS/UEBuildIOS.cs
@@ -18,7 +18,7 @@ namespace UnrealBuildTool
 	public partial class IOSTargetRules
 	{
 		/// <summary>
-		/// Whether to strip iOS symbols or not (implied by Shipping config)
+		/// Whether to strip iOS symbols or not (implied by Shipping config).
 		/// </summary>
 		[XmlConfigFile(Category = "BuildConfiguration")]
 		[CommandLine("-stripsymbols", Value = "true")]
diff --git a/Engine/Source/Programs/UnrealBuildTool/Platform/Mac/UEBuildMac.cs b/Engine/Source/Programs/UnrealBuildTool/Platform/Mac/UEBuildMac.cs
--- a/Engine/Source/Programs/UnrealBuildTool/Platform/Mac/UEBuildMac.cs
+++ b/Engine/Source/Programs/UnrealBuildTool/Platform/Mac/UEBuildMac.cs
@@ -16,28 +16,28 @@ namespace UnrealBuildTool
 	public class MacTargetRules
 	{
 		/// <summary>
-		/// Whether to generate dSYM files
-		/// Lists Architectures that you want to build
+		/// Whether to generate dSYM files.
+		/// Lists Architectures that you want to build.
 		/// </summary>
 		[XmlConfigFile(Category = "BuildConfiguration", Name = "bGeneratedSYMFile")]
 		public bool bGenerateDsymFile = true;
 
 		/// <summary>
-		/// Enables address sanitizer (ASan)
+		/// Enables address sanitizer (ASan).
 		/// </summary>
 		[CommandLine("-EnableASan")]
 		[XmlConfigFile(Category = "BuildConfiguration", Name = "bEnableAddressSanitizer")]
 		public bool bEnableAddressSanitizer = false;
 
 		/// <summary>
-		/// Enables thread sanitizer (TSan)
+		/// Enables thread sanitizer (TSan).
 		/// </summary>
 		[CommandLine("-EnableTSan")]
 		[XmlConfigFile(Category = "BuildConfiguration", Name = "bEnableThreadSanitizer")]
 		public bool bEnableThreadSanitizer = false;
 
 		/// <summary>
-		/// Enables undefined behavior sanitizer (UBSan)
+		/// Enables undefined behavior sanitizer (UBSan).
 		/// </summary>
 		[CommandLine("-EnableUBSan")]
 		[XmlConfigFile(Category = "BuildConfiguration", Name = "bEnableUndefinedBehaviorSanitizer")]
diff --git a/Engine/Source/Programs/UnrealBuildTool/Platform/Windows/UEBuildWindows.cs b/Engine/Source/Programs/UnrealBuildTool/Platform/Windows/UEBuildWindows.cs
--- a/Engine/Source/Programs/UnrealBuildTool/Platform/Windows/UEBuildWindows.cs
+++ b/Engine/Source/Programs/UnrealBuildTool/Platform/Windows/UEBuildWindows.cs
@@ -110,7 +110,7 @@ namespace UnrealBuildTool
 		TargetRules Target;
 
 		/// <summary>
-		/// Version of the compiler toolchain to use on Windows platform. A value of "default" will be changed to a specific version at UBT startup.
+		/// Version of the compiler toolchain to use on Windows platform. A value of "default" will be changed to a specific version at UBT start up.
 		/// </summary>
 		[ConfigFile(ConfigHierarchyType.Engine, "/Script/WindowsTargetPlatform.WindowsTargetSettings", "CompilerVersion")]
 		[XmlConfigFile(Category = "WindowsPlatform")]
@@ -125,15 +125,15 @@ namespace UnrealBuildTool
 		public WindowsArchitecture Architecture = WindowsArchitecture.x64;
 
 		/// <summary>
-		/// The specific toolchain version to use. This may be a specific version number (eg. "14.13.26128") or the string "Latest" to select the newest available version. By default, we use the
-		/// toolchain version indicated by WindowsPlatform.DefaultToolChainVersion if it is available, or the latest version otherwise.
+		/// The specific toolchain version to use. This may be a specific version number (for example, "14.13.26128"), or the string "Latest", to select the newest available version. By default, and if it is available, we use the
+		/// toolchain version indicated by WindowsPlatform.DefaultToolChainVersion (otherwise, we use the latest version).
 		/// </summary>
 		[XmlConfigFile(Category = "WindowsPlatform")]
 		public string CompilerVersion = null;
 
 		/// <summary>
-		/// The specific Windows SDK version to use. This may be a specific version number (eg. "8.1", "10.0", or "10.0.10150.0") or the string "Latest" to select the newest available version.
-		/// By default, we use the Windows SDK version indicated by WindowsPlatform.DefaultWindowsSdkVersion if it is available, or the latest version otherwise.
+		/// The specific Windows SDK version to use. This may be a specific version number (for example, "8.1", "10.0" or "10.0.10150.0"), or the string "Latest", to select the newest available version.
+		/// By default, and if it is available, we use the Windows SDK version indicated by WindowsPlatform.DefaultWindowsSdkVersion (otherwise, we use the latest version).
 		/// </summary>
 		[XmlConfigFile(Category = "WindowsPlatform")]
 		public string WindowsSdkVersion = null;
@@ -174,14 +174,14 @@ namespace UnrealBuildTool
 		public string ProductName;
 
 		/// <summary>
-		/// The static analyzer to use
+		/// The static analyzer to use.
 		/// </summary>
 		[XmlConfigFile(Category = "WindowsPlatform")]
 		[CommandLine("-StaticAnalyzer")]
 		public WindowsStaticAnalyzer StaticAnalyzer = WindowsStaticAnalyzer.None;
 
 		/// <summary>
-		/// Whether we should export a file containing .obj->source file mappings.
+		/// Whether we should export a file containing .obj to source file mappings.
 		/// </summary>
 		[XmlConfigFile]
 		[CommandLine("-ObjSrcMap")]
@@ -268,7 +268,7 @@ namespace UnrealBuildTool
 		public bool bOptimizeGlobalData = true;
 
 		/// <summary>
-		/// (Experimental) Appends the -ftime-trace argument to the command line for Clang, to output a JSON file containing a timeline for the compile. 
+		/// (Experimental) Appends the -ftime-trace argument to the command line for Clang to output a JSON file containing a timeline for the compile. 
 		/// See http://aras-p.info/blog/2019/01/16/time-trace-timeline-flame-chart-profiler-for-Clang/ for more info.
 		/// </summary>
 		[XmlConfigFile(Category = "WindowsPlatform")]
diff --git a/Engine/Source/Programs/UnrealBuildTool/ProjectFiles/ProjectFileGenerator.cs b/Engine/Source/Programs/UnrealBuildTool/ProjectFiles/ProjectFileGenerator.cs
--- a/Engine/Source/Programs/UnrealBuildTool/ProjectFiles/ProjectFileGenerator.cs
+++ b/Engine/Source/Programs/UnrealBuildTool/ProjectFiles/ProjectFileGenerator.cs
@@ -148,7 +148,7 @@ namespace UnrealBuildTool
 	static class ProjectFileGeneratorSettings
 	{
 		/// <summary>
-		/// Default list of project file formats to generate
+		/// Default list of project file formats to generate.
 		/// </summary>
 		[XmlConfigFile(Category = "ProjectFileGenerator", Name = "Format")]
 		public static string Format = null;
@@ -207,100 +207,100 @@ namespace UnrealBuildTool
 		protected string GameProjectName = null;
 
 		/// <summary>
-		/// Whether we should include configurations for "Test" and "Shipping" in generated projects (pass "-NoShippingConfigs" to disable this)
+		/// Whether we should include configurations for "Test" and "Shipping" in generated projects. Pass "-NoShippingConfigs" to disable this.
 		/// </summary>
 		[XmlConfigFile]
 		public static bool bIncludeTestAndShippingConfigs = true;
 
 		/// <summary>
-		/// True if intellisense data should be generated (takes a while longer)
+		/// True if intellisense data should be generated (takes a while longer).
 		/// </summary>
 		[XmlConfigFile]
 		bool bGenerateIntelliSenseData = true;
 
 		/// <summary>
-		/// True if we should include documentation in the generated projects
+		/// True if we should include documentation in the generated projects.
 		/// </summary>
 		[XmlConfigFile]
 		protected bool bIncludeDocumentation = false;
 
 		/// <summary>
-		/// True if all documentation languages should be included in generated projects, otherwise only "INT" will be included
+		/// True if all documentation languages should be included in generated projects, otherwise only INT files will be included.
 		/// </summary>
 		[XmlConfigFile]
 		bool bAllDocumentationLanguages = false;
 
 		/// <summary>
-		/// True if build targets should pass the -useprecompiled argument
+		/// True if build targets should pass the -useprecompiled argument.
 		/// </summary>
 		[XmlConfigFile]
 		public bool bUsePrecompiled = false;
 
 		/// <summary>
-		/// True if we should include engine source in the generated solution
+		/// True if we should include engine source in the generated solution.
 		/// </summary>
 		[XmlConfigFile]
 		protected bool bIncludeEngineSource = true;
 
 		/// <summary>
-		/// Whether to include enterprise source in the generated solution
+		/// Whether to include enterprise source in the generated solution.
 		/// </summary>
 		[XmlConfigFile]
 		protected bool bIncludeEnterpriseSource = true;
 
 		/// <summary>
-		/// True if shader source files should be included in the generated projects
+		/// True if shader source files should be included in generated projects.
 		/// </summary>
 		[XmlConfigFile]
 		protected bool bIncludeShaderSource = true;
 
 		/// <summary>
-		/// True if build system files should be included
+		/// True if build system files should be included.
 		/// </summary>
 		[XmlConfigFile]
 		bool bIncludeBuildSystemFiles = true;
 
 		/// <summary>
-		/// True if we should include config files (.ini files) in the generated project
+		/// True if we should include config (.ini) files in the generated project.
 		/// </summary>
 		[XmlConfigFile]
 		protected bool bIncludeConfigFiles = true;
 
 		/// <summary>
-		/// True if we should include localization files (.int/.kor/etc files) in the generated project
+		/// True if we should include localization files in the generated project.
 		/// </summary>
 		[XmlConfigFile]
 		bool bIncludeLocalizationFiles = false;
 
 		/// <summary>
-		/// True if we should include template files (.template files) in the generated project
+		/// True if we should include template files in the generated project.
 		/// </summary>
 		[XmlConfigFile]
 		protected bool bIncludeTemplateFiles = true;
 
 		/// <summary>
-		/// True if we should include program projects in the generated solution
+		/// True if we should include program projects in the generated solution.
 		/// </summary>
 		[XmlConfigFile]
 		protected bool IncludeEnginePrograms = true;
 
 		/// <summary>
-		/// Whether to include temporary targets generated by UAT to support content only projects with non-default settings
+		/// Whether to include temporary targets generated by UAT to support content only projects with non-default settings.
 		/// </summary>
 		[XmlConfigFile]
 		bool bIncludeTempTargets = false;
 
 		/// <summary>
-		/// True if we should include .NET Core projects in the generated solution
+		/// True if we should include .NET Core projects in the generated solution.
 		/// </summary>
 		[XmlConfigFile]
 		bool bIncludeDotNETCoreProjects = false;
 
-		/// <summary>
-		/// True if we should reflect "Source" sub-directories on disk in the master project as master project directories.
-		/// This arguably adds some visual clutter to the master project, but is truer to the on-disk file organization.
-		/// </summary>
-		[XmlConfigFile]
+        /// <summary>
+        /// True if we should reflect "Source" sub-directories on disk in the master project as master project directories.
+        /// This (arguably) adds some visual clutter to the master project but it is truer to the on-disk file organization.
+        /// </summary>
+        [XmlConfigFile]
 		bool bKeepSourceSubDirectories = true;
 
 		/// <summary>
@@ -353,13 +353,13 @@ namespace UnrealBuildTool
 		bool AllowDotNetCoreProjects = false;
 
 		/// <summary>
-		/// Name of the master project file (e.g. base file name for the solution file for Visual Studio, or the Xcode project file on Mac)
+		/// Name of the master project file -- for example, the base file name for the Visual Studio solution file, or the Xcode project file on Mac.
 		/// </summary>
 		[XmlConfigFile]
 		protected string MasterProjectName = "UE4";
 
 		/// <summary>
-		/// If true, sets the master project name according to the name of the folder it's in
+		/// If true, sets the master project name according to the name of the folder it is in.
 		/// </summary>
 		[XmlConfigFile]
 		protected bool bMasterProjectNameFromFolder = false;
diff --git a/Engine/Source/Programs/UnrealBuildTool/ProjectFiles/VisualStudio/VCProjectFileGenerator.cs b/Engine/Source/Programs/UnrealBuildTool/ProjectFiles/VisualStudio/VCProjectFileGenerator.cs
--- a/Engine/Source/Programs/UnrealBuildTool/ProjectFiles/VisualStudio/VCProjectFileGenerator.cs
+++ b/Engine/Source/Programs/UnrealBuildTool/ProjectFiles/VisualStudio/VCProjectFileGenerator.cs
@@ -58,19 +58,19 @@ namespace UnrealBuildTool
 		protected VCProjectFileFormat ProjectFileFormat = VCProjectFileFormat.Default;
 
 		/// <summary>
-		/// Whether to write a solution option (suo) file for the sln
+		/// Whether to write a solution option (suo) file for the sln.
 		/// </summary>
 		[XmlConfigFile(Category = "BuildConfiguration")]
 		protected bool bWriteSolutionOptionFile = true;
 
 		/// <summary>
-		/// Whether to add the -FastPDB option to build command lines by default
+		/// Whether to add the -FastPDB option to build command lines by default.
 		/// </summary>
 		[XmlConfigFile(Category = "BuildConfiguration")]
 		bool bAddFastPDBToProjects = false;
 
 		/// <summary>
-		/// Whether to generate per-file intellisense data
+		/// Whether to generate per-file intellisense data.
 		/// </summary>
 		[XmlConfigFile(Category = "BuildConfiguration")]
 		bool bUsePerFileIntellisense = false;
diff --git a/Engine/Source/Programs/UnrealBuildTool/System/SourceFileWorkingSet.cs b/Engine/Source/Programs/UnrealBuildTool/System/SourceFileWorkingSet.cs
--- a/Engine/Source/Programs/UnrealBuildTool/System/SourceFileWorkingSet.cs
+++ b/Engine/Source/Programs/UnrealBuildTool/System/SourceFileWorkingSet.cs
@@ -291,19 +291,19 @@ namespace UnrealBuildTool
 		};
 
 		/// <summary>
-		/// Sets the provider to use for determining the working set
+		/// Sets the provider to use for determining the working set.
 		/// </summary>
 		[XmlConfigFile]
 		static ProviderType Provider = ProviderType.Default;
 
 		/// <summary>
-		/// Sets the path to use for the repository. Interpreted relative to the UE root directory (ie. folder above the Engine folder) if relative.
+		/// Sets the path to use for the repository. Interpreted relative to the Unreal Engine root directory (the folder above the Engine folder) -- if relative.
 		/// </summary>
 		[XmlConfigFile]
 		public static string RepositoryPath = null;
 
 		/// <summary>
-		/// Sets the path to use for the Git executable. Defaults to "git" (assuming it's in the PATH).
+		/// Sets the path to use for the Git executable. Defaults to "git" (assuming it is in the PATH).
 		/// </summary>
 		[XmlConfigFile]
 		public static string GitPath = "git";
diff --git a/Engine/Source/Programs/UnrealBuildTool/ToolChain/RemoteMac.cs b/Engine/Source/Programs/UnrealBuildTool/ToolChain/RemoteMac.cs
--- a/Engine/Source/Programs/UnrealBuildTool/ToolChain/RemoteMac.cs
+++ b/Engine/Source/Programs/UnrealBuildTool/ToolChain/RemoteMac.cs
@@ -37,31 +37,31 @@ namespace UnrealBuildTool
 	class RemoteMac
 	{
 		/// <summary>
-		/// These two variables will be loaded from XML config file in XmlConfigLoader.Init()
+		/// These two variables will be loaded from the XML config file in XmlConfigLoader.Init().
 		/// </summary>
 		[XmlConfigFile]
 		private readonly string ServerName;
 
 		/// <summary>
-		/// The remote username
+		/// The remote username.
 		/// </summary>
 		[XmlConfigFile]
 		private readonly string UserName;
 
 		/// <summary>
-		/// Instead of looking for RemoteToolChainPrivate.key in the usual places (Documents/Unreal Engine/UnrealBuildTool/SSHKeys, Engine/Build/SSHKeys), this private key will be used if set
+		/// If set, instead of looking for RemoteToolChainPrivate.key in the usual places (Documents/Unreal, Engine/UnrealBuildTool/SSHKeys or Engine/Build/SSHKeys), this private key will be used.
 		/// </summary>
 		[XmlConfigFile]
 		private FileReference SshPrivateKey;
 
 		/// <summary>
-		/// The authentication used for Rsync (for the -e rsync flag)
+		/// The authentication used for Rsync (for the -e rsync flag).
 		/// </summary>
 		[XmlConfigFile]
 		private string RsyncAuthentication = "ssh -i '${CYGWIN_SSH_PRIVATE_KEY}'";
 
 		/// <summary>
-		/// The authentication used for SSH (probably similar to RsyncAuthentication)
+		/// The authentication used for SSH (probably similar to RsyncAuthentication).
 		/// </summary>
 		[XmlConfigFile]
 		private string SshAuthentication = "-i '${CYGWIN_SSH_PRIVATE_KEY}'";
diff --git a/Engine/Source/Runtime/Core/Public/Stats/Stats.h b/Engine/Source/Runtime/Core/Public/Stats/Stats.h
--- a/Engine/Source/Runtime/Core/Public/Stats/Stats.h
+++ b/Engine/Source/Runtime/Core/Public/Stats/Stats.h
@@ -13,283 +13,7 @@
 #include "Stats/Stats2.h"
 
 /** 
- *	Unreal Engine Stats system
- *	
- *	This is a preliminary version of the documentation, any comments are welcome :)
- *	
- *	This system allows you to collect various performance data and then the data can be used to optimize your game.
- *	There are a few methods how to achieve this. This quick tutorial will describe all of them.
- *	For stats commands check out method PrintStatsHelpToOutputDevice();
- *	
- *	Stats system in the UE4 supports following stats types:
- *		Cycle Counter - a generic cycle counter used to counting the number of cycles during the lifetime of the object
- *		Float/Dword Counter - a counter that is cleared every frame
- *		Float/Dword Accumulator - a counter that is not cleared every frame, persistent stat, but it can be reset
- *		Memory - a special type of counter that is optimized for memory tracking
- *	
- *	Each stat needs to be grouped, this usually corresponds with displaying the specified stat group i.e. 'stat statsystem' which displays stats' related data.
- *	
- *	To define a stat group you need to use one of the following methods:
- *				DECLARE_STATS_GROUP(GroupDesc,GroupId,GroupCat) - declares a stats group which is enabled by default 
- *				DECLARE_STATS_GROUP_VERBOSE(GroupDesc,GroupId,GroupCat) - declares a stats group which is disabled by default
- *				DECLARE_STATS_GROUP_MAYBE_COMPILED_OUT(GroupDesc,GroupId,GroupCat) - declares a stats group which is disabled by default and may be stripped by the compiler
- *	
- *	where
- *		GroupDesc is a text description of the group
- *		GroupId is an UNIQUE id of the group
- *		GroupCat is reserved for future use
- *		CompileIn if set to true, the compiler may strip it out
- *		
- *	It can be done in the source or header file depending the usage scope.
- *	
- *	Examples:
- *				DECLARE_STATS_GROUP(TEXT("Threading"), STATGROUP_Threading, STATCAT_Advanced);
- *				DECLARE_STATS_GROUP_VERBOSE(TEXT("Linker Load"), STATGROUP_LinkerLoad, STATCAT_Advanced);
- *	
- *	Now, you can declare/define a stat.
- *	A stat can be used only in one cpp file, in the function scope, in the module scope or can be used in the whole project.
- *	
- *	For one file scope you need to use one of the following methods depending on the stat type.
- *				DECLARE_CYCLE_STAT(CounterName,StatId,GroupId) - declares a cycle counter stat
- *		
- *				DECLARE_SCOPE_CYCLE_COUNTER(CounterName,StatId,GroupId) - declares a cycle counter stat and uses it at the same time, it is limited to one function scope
- *				QUICK_SCOPE_CYCLE_COUNTER(StatId) - declares a cycle counter stat that will belong to stat group called 'Quick'
- *				RETURN_QUICK_DECLARE_CYCLE_STAT(StatId,GroupId) - returns a cycle counter, used by a few specialized classes, more information later
- *		
- *				DECLARE_FLOAT_COUNTER_STAT(CounterName,StatId,GroupId) - declares a float counter, technically speaking it's based on the double type, 8 bytes
- *				DECLARE_DWORD_COUNTER_STAT(CounterName,StatId,GroupId) - declared a dword counter, technically speaking it's based on the qword type, 8 bytes
- *				DECLARE_FLOAT_ACCUMULATOR_STAT(CounterName,StatId,GroupId) - declares a float accumulator
- *				DECLARE_DWORD_ACCUMULATOR_STAT(CounterName,StatId,GroupId) - declares a dword accumulator
- *				DECLARE_MEMORY_STAT(CounterName,StatId,GroupId) - declares a memory counter, same as the dword accumulator, but will be displayed with memory specific units
- *				DECLARE_MEMORY_STAT_POOL(CounterName,StatId,GroupId,Pool) - declares a memory counter with a pool
- *	
- *		If you want to have these stats accessible in the whole project/or wider range of files you need to use extern version.
- *		These methods are the same as the previously mentioned but with _EXTERN and the end of the name, here is the list:
- *				DECLARE_CYCLE_STAT_EXTERN(CounterName,StatId,GroupId, API)
- *				DECLARE_FLOAT_COUNTER_STAT_EXTERN(CounterName,StatId,GroupId, API)
- *				DECLARE_DWORD_COUNTER_STAT_EXTERN(CounterName,StatId,GroupId, API)
- *				DECLARE_FLOAT_ACCUMULATOR_STAT_EXTERN(CounterName,StatId,GroupId, API)
- *				DECLARE_DWORD_ACCUMULATOR_STAT_EXTERN(CounterName,StatId,GroupId, API)
- *				DECLARE_MEMORY_STAT_EXTERN(CounterName,StatId,GroupId, API)
- *				DECLARE_MEMORY_STAT_POOL_EXTERN(CounterName,StatId,GroupId,Pool, API)
- *				
- *		Then in the source file you need to define those stats.
- *				DEFINE_STAT(CounterName) - defines stats declared with _EXTERN
- *	
- *	where
- *		CounterName is a text description of the stat
- *		StatId is an UNIQUE id of the stat
- *		GroupId is an id of the group that the stat will belong to, the GroupId from DECLARE_STATS_GROUP*
- *		Pool is a platform specific memory pool, more details later
- *		API is the *_API of module, can be empty if the stat will be used only in that module
- *
- *	Examples:
- *		Custom memory stats with pools
- *			First you need to add a new pool to enum EMemoryCounterRegion, it can be global or platform specific.
- *			
- *				enum EMemoryCounterRegion
- *				{
- *					MCR_Invalid,	// not memory
- *					MCR_Physical,	// main system memory
- *					MCR_GPU,		// memory directly a GPU (graphics card, etc)
- *					MCR_GPUSystem,	// system memory directly accessible by a GPU
- *					MCR_TexturePool,// presized texture pools
- *					MCR_MAX
- *				};
- *			
- *			This is an example that will allow using the pools every where, see CORE_API. 
- *			THE NAME OF THE POOL MUST START WITH MCR_
- *			Header file.	
- *				DECLARE_MEMORY_STAT_POOL_EXTERN(TEXT("Physical Memory Pool [Physical]"),	MCR_Physical,		STATGROUP_Memory,  FPlatformMemory::MCR_Physical,	CORE_API);
- *				DECLARE_MEMORY_STAT_POOL_EXTERN(TEXT("GPU Memory Pool [GPU]"),				MCR_GPU,			STATGROUP_Memory,  FPlatformMemory::MCR_GPU,		CORE_API);
- *				DECLARE_MEMORY_STAT_POOL_EXTERN(TEXT("Texture Memory Pool [Texture]"),		MCR_TexturePool,	STATGROUP_Memory,  FPlatformMemory::MCR_TexturePool,CORE_API);
- *				
- *			Source file.
- *				DEFINE_STAT(MCR_Physical);
- *				DEFINE_STAT(MCR_GPU);
- *				DEFINE_STAT(MCR_TexturePool);
- *			
- *			This is a pool, so it needs to be initialized. Usually in the F*PlatformMemory::Init()
- *				SET_MEMORY_STAT(MCR_Physical, PhysicalPoolLimit);
- *				SET_MEMORY_STAT(MCR_GPU, GPUPoolLimit);
- *				SET_MEMORY_STAT(MCR_TexturePool, TexturePoolLimit);
- *			
- *			Now we have pools, so we can setup memory stats for those pools.
- *			Accessible everywhere.
- *				DECLARE_MEMORY_STAT_POOL_EXTERN(TEXT("Index buffer memory"),		STAT_IndexBufferMemory,		STATGROUP_RHI, FPlatformMemory::MCR_GPU, RHI_API);
- *				DECLARE_MEMORY_STAT_POOL_EXTERN(TEXT("Vertex buffer memory"),		STAT_VertexBufferMemory,	STATGROUP_RHI, FPlatformMemory::MCR_GPU, RHI_API);
- *				DECLARE_MEMORY_STAT_POOL_EXTERN(TEXT("Structured buffer memory"),	STAT_StructuredBufferMemory,STATGROUP_RHI, FPlatformMemory::MCR_GPU, RHI_API);
- *				DECLARE_MEMORY_STAT_POOL_EXTERN(TEXT("Pixel buffer memory"),		STAT_PixelBufferMemory,		STATGROUP_RHI, FPlatformMemory::MCR_GPU, RHI_API);
- *				
- *			Accessible only in the module where defined.	
- *				DECLARE_MEMORY_STAT_POOL_EXTERN(TEXT("Pool Memory Size"), STAT_TexturePoolSize,				STATGROUP_Streaming, FPlatformMemory::MCR_TexturePool, );
- *				DECLARE_MEMORY_STAT_POOL_EXTERN(TEXT("Pool Memory Used"), STAT_TexturePoolAllocatedSize,	STATGROUP_Streaming, FPlatformMemory::MCR_TexturePool, );
- *				
- *			And the last thing, updating the memory stats.
- *				INC_MEMORY_STAT_BY(STAT_PixelBufferMemory,NumBytes) - increases a memory stat by the specified value
- *				DEC_MEMORY_STAT_BY(STAT_PixelBufferMemory,NumBytes) - decreases a memory stat by the specified value
- *				SET_MEMORY_STAT(STAT_PixelBufferMemory,NumBytes) - sets a memory stat to the specified value
- *			
- *		Regular memory stats, without pools	
- *				DECLARE_MEMORY_STAT(TEXT("Total Physical"),		STAT_TotalPhysical,		STATGROUP_MemoryPlatform);
- *				DECLARE_MEMORY_STAT(TEXT("Total Virtual"),		STAT_TotalVirtual,		STATGROUP_MemoryPlatform);
- *				DECLARE_MEMORY_STAT(TEXT("Page Size"),			STAT_PageSize,			STATGROUP_MemoryPlatform);
- *				DECLARE_MEMORY_STAT(TEXT("Total Physical GB"),	STAT_TotalPhysicalGB,	STATGROUP_MemoryPlatform);
- *		
- *		Or DECLARE_MEMORY_STAT_EXTERN in the header file and then DEFINE_STAT in the source file.
- *		Updating the memory stats is done the same way as in the version with pools.
- *		
- *		
- *		Performance data using the cycle counters.
- *			First you need to add cycle counters.
- *				DECLARE_CYCLE_STAT(TEXT("Broadcast"),	STAT_StatsBroadcast,STATGROUP_StatSystem);
- *				DECLARE_CYCLE_STAT(TEXT("Condense"),	STAT_StatsCondense,	STATGROUP_StatSystem);
- *				
- *			Or DECLARE_CYCLE_STAT_EXTERN in the header file and then DEFINE_STAT in the source file.
- *			
- *			Now you can grab the performance data.
- *			
- *				Stats::Broadcast()
- *				{
- *					SCOPE_CYCLE_COUNTER(STAT_StatsBroadcast);
- *					...
- *					// a piece of code
- *					...
- *				}
- *
- *			and that's all.
- *			Sometimes you don't want to grab the stats every time the function is called, so you can use conditional cycle counter.
- *			It's not very common, but may be useful.
- *			
- *				Stats::Broadcast(bool bSomeCondition)
- *				{
- *					CONDITIONAL_SCOPE_CYCLE_COUNTER(STAT_StatsBroadcast,bSomeCondition);
- *					...
- *					// a piece of code
- *					...
- *				}
- *
- *			If you want to grab the performance data from one function you can use following construction.
- *			
- *				Stats::Broadcast(bool bSomeCondition)
- *				{
- *					DECLARE_SCOPE_CYCLE_COUNTER(TEXT("Broadcast"), STAT_StatsBroadcast, STATGROUP_StatSystem);
- *					...
- *					// a piece of code
- *					...
- *				}
- *				
- *			or
- *			
- *				Stats::Broadcast(bool bSomeCondition)
- *				{
- *					QUICK_SCOPE_CYCLE_COUNTER(TEXT("Stats::Broadcast"));
- *					...
- *					// a piece of code
- *					...
- *				}
- *				
- *			Mostly used for temporary stats.
- *			
- *			Those all cycle counters are used to generate the hierarchy. So you can get more detailed information about performance data.
- *			There is also an option to set flat cycle counter.
- *			
- *				Stats::Broadcast(bool bSomeCondition)
- *				{
- *					const uint32 BroadcastBeginTime = FPlatformTime::Cycles();
- *					...
- *					// a piece of code
- *					...
- *					const uint32 BroadcastEndTime = FPlatformTime::Cycles();
- *					SET_CYCLE_COUNTER(STAT_StatsBroadcast, BroadcastEndTime-BroadcastBeginTime);
- *				}
- *				
- *		
- *		A few tasks implemented in the UE4 use a different approach in terms of getting the performance data.
- *		They implement method GetStatId(). If there is no GetStatId(), the code will not compile.
- *		Here is an example.
- *		
- *				class FParallelAnimationCompletionTask
- *				{
- *					// ...
- *					// a piece of code
- *					FORCEINLINE TStatId GetStatId() const
- *					{
- *						RETURN_QUICK_DECLARE_CYCLE_STAT(FParallelAnimationCompletionTask, STATGROUP_TaskGraphTasks);
- *					}
- *					// a piece of code
- *					// ...
- *				};
- *				
- *				
- *		Generic data using the float or dword counters.		
- *			First you need to add a few counters.
- *				DECLARE_FLOAT_COUNTER_STAT_EXTERN(STAT_FloatCounter,StatId,STATGROUP_TestGroup, CORE_API)
- *				DECLARE_DWORD_COUNTER_STAT_EXTERN(STAT_DwordCounter,StatId,STATGROUP_TestGroup, CORE_API)
- *				DECLARE_FLOAT_ACCUMULATOR_STAT_EXTERN(STAT_FloatAccumulator,StatId,STATGROUP_TestGroup, CORE_API)
- *				DECLARE_DWORD_ACCUMULATOR_STAT_EXTERN(STAT_DwordAccumulator,StatId,STATGROUP_TestGroup, CORE_API)
- *				 
- *			Updating counters.
- *				INC_DWORD_STAT(StatId) - increases a dword stat by 1
- *				DEC_DWORD_STAT(StatId) - decreases a dword stat by 1
- *				INC_DWORD_STAT_BY(StatId,Amount) - increases a dword stat by the specified value
- *				DEC_DWORD_STAT_BY(StatId,Amount) - decreases a dword stat by the specified value
- *				SET_DWORD_STAT(StatId,Value) - sets a dword stat to the specified value
-
- *				INC_FLOAT_STAT_BY(StatId,Amount) - increases a float stat by the specified value
- *				DEC_FLOAT_STAT_BY(StatId,Amount) - decreases a float stat by the specified value
- *				SET_FLOAT_STAT(StatId,Value) - sets a float stat to the specified value
- *			
- *			
- *		A few helper methods
- *				 GET_STATID(StatId) - returns an instance of the TStatId of the stat, ADVANCED
- *				 GET_STATDESCRIPTION(StatId) - returns a description of the stat
- *
- *		
- *		If you don't want to use the stats system and just log some performance data, there is functionality for this.
- *		
- *				SCOPE_SECONDS_COUNTER(double&Seconds) - captures time passed in seconds, adding delta time to passed in variable
- *		
- *				Stats::Broadcast()
- *				{
- *					double ThisTime = 0;
- *					{
- *						 SCOPE_SECONDS_COUNTER(ThisTime);
-*						...
- *						// a piece of code
- *						...
- *					}
- *					UE_LOG(LogTemp, Log, TEXT("Stats::Broadcast %.2f"), ThisTime );	
- *				}
- *		
- *				FScopeLogTime - utility class to log time passed in seconds, adding cumulative stats to passed in variable, print the performance data to the log in the destructor
- *			
- *				SCOPE_LOG_TIME(Name,CumulativePtr) - using the given name prints the performance data and gathers cumulative stats
- *				SCOPE_LOG_TIME_IN_SECONDS(Name,CumulativePtr) - the same as above, but prints in seconds
- *			
- *				SCOPE_LOG_TIME_FUNC() - using the funcion name prints the performance data, cannot be nested
- *				SCOPE_LOG_TIME_FUNC_WITH_GLOBAL(CumulativePtr), same as above, but gather cumulative stats
- *			
- *			A few examples.
- *			
- *				double GMyBroadcastTime = 0.0;
- *				Stats::Broadcast()
- *				{
- *					SCOPE_LOG_TIME("Stats::Broadcast", &GMyBroadcastTime );
- *					SCOPE_LOG_TIME_IN_SECONDS("Stats::Broadcast (sec)", &GMyBroadcastTime );
- *					...
- *					// a piece of code
- *					...
- *				}
- *				
- *				Stats::Condense()
- *				{
- *					SCOPE_LOG_TIME_FUNC(); // The name should be "Stats::Condense()", may differ across compilers
- *					SCOPE_LOG_TIME_FUNC_WITH_GLOBAL(&GMyBroadcastTime);
- *					...
- *					// a piece of code
- *					...
- *				}
+ *	Learn about the Stats System at docs.unrealengine.com
  */
 
 struct TStatId;
diff --git a/Engine/Source/Runtime/CoreUObject/Private/UObject/UObjectGlobals.cpp b/Engine/Source/Runtime/CoreUObject/Private/UObject/UObjectGlobals.cpp
--- a/Engine/Source/Runtime/CoreUObject/Private/UObject/UObjectGlobals.cpp
+++ b/Engine/Source/Runtime/CoreUObject/Private/UObject/UObjectGlobals.cpp
@@ -2509,7 +2509,7 @@ UObject::UObject()
 	EnsureNotRetrievingVTablePtr();
 
 	FObjectInitializer* ObjectInitializerPtr = FUObjectThreadContext::Get().TopInitializer();
-	UE_CLOG(!ObjectInitializerPtr, LogUObjectGlobals, Fatal, TEXT("%s is not being constructed with either NewObject, NewNamedObject or ConstructObject."), *GetName());
+	UE_CLOG(!ObjectInitializerPtr, LogUObjectGlobals, Fatal, TEXT("%s is not being constructed with NewObject."), *GetName());
 	FObjectInitializer& ObjectInitializer = *ObjectInitializerPtr;
 	UE_CLOG(ObjectInitializer.Obj != nullptr && ObjectInitializer.Obj != this, LogUObjectGlobals, Fatal, TEXT("UObject() constructor called but it's not the object that's currently being constructed with NewObject. Maybe you are trying to construct it on the stack, which is not supported."));
 	const_cast<FObjectInitializer&>(ObjectInitializer).Obj = this;
diff --git a/Engine/Source/Runtime/CoreUObject/Public/UObject/UObjectThreadContext.h b/Engine/Source/Runtime/CoreUObject/Public/UObject/UObjectThreadContext.h
--- a/Engine/Source/Runtime/CoreUObject/Public/UObject/UObjectThreadContext.h
+++ b/Engine/Source/Runtime/CoreUObject/Public/UObject/UObjectThreadContext.h
@@ -62,7 +62,7 @@ public:
 	FObjectInitializer& TopInitializerChecked()
 	{
 		FObjectInitializer* ObjectInitializerPtr = TopInitializer();
-		UE_CLOG(!ObjectInitializerPtr, LogUObjectThreadContext, Fatal, TEXT("Tried to get the current ObjectInitializer, but none is set. Please use NewObject or NewNamedObject to construct new UObject-derived classes."));
+		UE_CLOG(!ObjectInitializerPtr, LogUObjectThreadContext, Fatal, TEXT("Tried to get the current ObjectInitializer, but none is set. Please use NewObject to construct new UObject-derived classes."));
 		return *ObjectInitializerPtr;
 	}
 
diff --git a/Engine/Source/Runtime/Engine/Classes/Kismet/GameplayStatics.h b/Engine/Source/Runtime/Engine/Classes/Kismet/GameplayStatics.h
--- a/Engine/Source/Runtime/Engine/Classes/Kismet/GameplayStatics.h
+++ b/Engine/Source/Runtime/Engine/Classes/Kismet/GameplayStatics.h
@@ -36,10 +36,10 @@ class FMemoryReader;
 struct FDialogueContext;
 
 
-/** Delegate called from AsyncLoadGameFromSlot. First two parameters are passed in SlotName and UserIndex, third parameter is the returned SaveGame, or null if it failed to load */
+/** Delegate called from AsyncLoadGameFromSlot. First two parameters are passed in SlotName and UserIndex, third parameter is a bool indicating success (true) or failure (false). */
 DECLARE_DELEGATE_ThreeParams(FAsyncSaveGameToSlotDelegate, const FString&, const int32, bool);
 
-/** Delegate called from AsyncLoadGameFromSlot. First two parameters are passed in SlotName and UserIndex, third parameter is the returned SaveGame, or null if it failed to load */
+/** Delegate called from AsyncLoadGameFromSlot. First two parameters are passed in SlotName and UserIndex, third parameter is the returned SaveGame, or null if it failed to load. */
 DECLARE_DELEGATE_ThreeParams(FAsyncLoadGameFromSlotDelegate, const FString&, const int32, USaveGame*);
 
 /** Static class with useful gameplay utility functions that can be called from both Blueprint and C++ */
