Remove PLATFORM_LUMIN hack

diff --git a/Engine/Source/Runtime/Core/Private/HAL/MallocBinned2.cpp b/Engine/Source/Runtime/Core/Private/HAL/MallocBinned2.cpp
--- a/Engine/Source/Runtime/Core/Private/HAL/MallocBinned2.cpp
+++ b/Engine/Source/Runtime/Core/Private/HAL/MallocBinned2.cpp
@@ -236,15 +236,7 @@ private:
 		{
 			FirstFreeBlock = (FFreeBlock*)FirstFreeBlock->NextFreeBlock;
 		}
-#if PLATFORM_LUMIN
-		// Temporarily disable fatal check for Lumin since it crashes the ImageTrackerComponent.
-		if (!(!FirstFreeBlock || FirstFreeBlock->GetNumFreeRegularBlocks() != 0))
-		{
-			FPlatformMisc::LowLevelOutputDebugString(TEXT("assert failed - (!FirstFreeBlock || FirstFreeBlock->GetNumFreeRegularBlocks() != 0)"));
-		}
-#else
 		check(!FirstFreeBlock || FirstFreeBlock->GetNumFreeRegularBlocks() != 0);
-#endif // PLATFORM_LUMIN
 	}
 };
 
