Move lumin cvars out of ConsoleManager.cpp

diff --git a/Engine/Plugins/Lumin/MagicLeap/Source/MagicLeap/Private/MagicLeapStereoLayers.cpp b/Engine/Plugins/Lumin/MagicLeap/Source/MagicLeap/Private/MagicLeapStereoLayers.cpp
--- a/Engine/Plugins/Lumin/MagicLeap/Source/MagicLeap/Private/MagicLeapStereoLayers.cpp
+++ b/Engine/Plugins/Lumin/MagicLeap/Source/MagicLeap/Private/MagicLeapStereoLayers.cpp
@@ -6,9 +6,35 @@
 #include "SceneViewExtension.h"
 #include "RHI.h"
 #include "RHIDefinitions.h"
+#include "HAL/IConsoleManager.h"
 
 #include "CoreMinimal.h"
 
+static TAutoConsoleVariable<float> CVarDebugCanvasX(TEXT("r.LuminDebugCanvasX"),
+	11.0f,
+	TEXT("Allows the override of the lumin debug canvas 'x' coordinate."),
+	ECVF_Scalability | ECVF_Default);
+
+static TAutoConsoleVariable<float> CVarDebugCanvasY(TEXT("r.LuminDebugCanvasY"),
+	-18.0f,
+	TEXT("Allows the override of the lumin debug canvas 'y' coordinate."),
+	ECVF_Scalability | ECVF_Default);
+
+static TAutoConsoleVariable<float> CVarDebugCanvasZ(TEXT("r.LuminDebugCanvasZ"),
+	82.0f,
+	TEXT("Allows the override of the lumin debug canvas 'z' coordinate."),
+	ECVF_Scalability | ECVF_Default);
+
+static TAutoConsoleVariable<float> CVarDebugCanvasWidth(TEXT("r.LuminDebugCanvasWidth"),
+	75.0f,
+	TEXT("Allows the override of the lumin debug canvas width."),
+	ECVF_Scalability | ECVF_Default);
+
+static TAutoConsoleVariable<float> CVarDebugCanvasHeight(TEXT("r.LuminDebugCanvasHeight"),
+	40.0f,
+	TEXT("Allows the override of the lumin debug canvas height."),
+	ECVF_Scalability | ECVF_Default);
+
 IStereoLayers* FMagicLeapHMD::GetStereoLayers()
 {
 	if (!DefaultStereoLayers.IsValid())
diff --git a/Engine/Source/Runtime/Core/Private/HAL/ConsoleManager.cpp b/Engine/Source/Runtime/Core/Private/HAL/ConsoleManager.cpp
--- a/Engine/Source/Runtime/Core/Private/HAL/ConsoleManager.cpp
+++ b/Engine/Source/Runtime/Core/Private/HAL/ConsoleManager.cpp
@@ -1976,31 +1976,6 @@ static TAutoConsoleVariable<float> CVarAmbientOcclusionStaticFraction(TEXT("r.Am
 	TEXT(" 1: AO affects the stat lighting"),
 	ECVF_Default);
 
-static TAutoConsoleVariable<float> CVarDebugCanvasX(TEXT("r.LuminDebugCanvasX"),
-	11.0f,
-	TEXT("Allows the override of the lumin debug canvas 'x' coordinate."),
-	ECVF_Scalability | ECVF_Default);
-
-static TAutoConsoleVariable<float> CVarDebugCanvasY(TEXT("r.LuminDebugCanvasY"),
-	-18.0f,
-	TEXT("Allows the override of the lumin debug canvas 'y' coordinate."),
-	ECVF_Scalability | ECVF_Default);
-
-static TAutoConsoleVariable<float> CVarDebugCanvasZ(TEXT("r.LuminDebugCanvasZ"),
-	82.0f,
-	TEXT("Allows the override of the lumin debug canvas 'z' coordinate."),
-	ECVF_Scalability | ECVF_Default);
-
-static TAutoConsoleVariable<float> CVarDebugCanvasWidth(TEXT("r.LuminDebugCanvasWidth"),
-	75.0f,
-	TEXT("Allows the override of the lumin debug canvas width."),
-	ECVF_Scalability | ECVF_Default);
-
-static TAutoConsoleVariable<float> CVarDebugCanvasHeight(TEXT("r.LuminDebugCanvasHeight"),
-	40.0f,
-	TEXT("Allows the override of the lumin debug canvas height."),
-	ECVF_Scalability | ECVF_Default);
-
 static TAutoConsoleVariable<int32> CVarShadowQuality(
 	TEXT("r.ShadowQuality"),
 	5,
@@ -2570,11 +2545,3 @@ static TAutoConsoleVariable<int32> CVarSkyLightingQuality(
 	TEXT("  1: on\n"),
 	ECVF_Scalability | ECVF_RenderThreadSafe
 );
-
-static TAutoConsoleVariable<int32> CVarLuminUseThreadingForPerformance(
-	TEXT("core.Lumin.UseThreadingForPerformance"),
-	1,
-	TEXT("Determine if threading should be used for possible performance gains. Used for FApp::ShouldUseThreadingForPerformance(). (Lumin-only)\n")
-	TEXT("  0 = off\n")
-	TEXT("  1 = on [default]\n"),
-	ECVF_ReadOnly);
