Fixed the additional constructor I added to FQualifiedFrameTime per DaveH's feedback

diff --git a/Engine/Source/Runtime/Core/Public/Misc/QualifiedFrameTime.h b/Engine/Source/Runtime/Core/Public/Misc/QualifiedFrameTime.h
--- a/Engine/Source/Runtime/Core/Public/Misc/QualifiedFrameTime.h
+++ b/Engine/Source/Runtime/Core/Public/Misc/QualifiedFrameTime.h
@@ -30,11 +30,9 @@ struct FQualifiedFrameTime
 	 * User construction from a timecode and its frame rate
 	 */
 	FQualifiedFrameTime(const FTimecode& InTimecode, const FFrameRate& InRate)
-		: Time(InTimecode.Frames)
+		: Time(InTimecode.ToFrameNumber(InRate))
 		, Rate(InRate)
 	{
-		// Check for cases where the user passed invalid data
-		check(InTimecode.Frames < InRate.Numerator);
 	}
 
 public:
