Lumin had a bunch of functions that purported to be case insensitive, but on investigation this was never true. Removing any semblence of the idea that case is not sensitive in Lumin Platform File.

diff --git a/Engine/Source/Runtime/Core/Private/Lumin/LuminPlatformFile.cpp b/Engine/Source/Runtime/Core/Private/Lumin/LuminPlatformFile.cpp
--- a/Engine/Source/Runtime/Core/Private/Lumin/LuminPlatformFile.cpp
+++ b/Engine/Source/Runtime/Core/Private/Lumin/LuminPlatformFile.cpp
@@ -356,7 +356,7 @@ __thread double FFileHandleLumin::AccessTimes[FFileHandleLumin::ACTIVE_HANDLE_CO
 #endif // MANAGE_FILE_HANDLES
 
 /**
- * A class to handle case insensitive file opening. This is a band-aid, non-performant approach,
+ * A class to handle file mapping. This is a band-aid, non-performant approach,
  * without any caching.
  */
 class FLuminFileMapper
@@ -426,7 +426,7 @@ public:
 	}
 
 	/**
-	 * Tries to recursively find (using case-insensitive comparison) and open the file.
+	 * Tries to recursively find and open the file.
 	 * The first file found will be opened.
 	 * 
 	 * @param Filename Original file path as requested (absolute)
@@ -507,12 +507,12 @@ public:
 	}
 
 	/**
-	 * Tries to map a filename (one with a possibly wrong case) to one that exists.
+	 * Tries to map a filename to one that exists.
 	 * 
-	 * @param PossiblyWrongFilename absolute filename (that has possibly a wrong case)
+	 * @param PossiblyWrongFilename absolute filename 
 	 * @param ExistingFilename filename that exists (only valid to use if the function returned success).
 	 */
-	bool MapCaseInsensitiveFile(const FString & PossiblyWrongFilename, FString & ExistingFilename)
+	bool MapFile(const FString & PossiblyWrongFilename, FString & ExistingFilename)
 	{
 		ExistingFilename = PossiblyWrongFilename;
 		return true;
@@ -536,7 +536,7 @@ public:
 		}
 		else
 		{
-			// perform a case-insensitive search from /
+			// perform a search from /
 
 			int MaxPathComponents = CountPathComponents(PossiblyWrongFilename);
 			if (MaxPathComponents > 0)
@@ -554,12 +554,12 @@ public:
 	}
 
 	/**
-	 * Opens a file for reading, disregarding the case.
+	 * Opens a file for reading
 	 * 
 	 * @param Filename absolute filename
 	 * @param MappedToFilename absolute filename that we mapped the Filename to (always filled out on success, even if the same as Filename)
 	 */
-	int32 OpenCaseInsensitiveRead(const FString & Filename, FString & MappedToFilename)
+	int32 OpenReadInternal(const FString & Filename, FString & MappedToFilename)
 	{
 		// We can get some "absolute" filenames like "D:/Blah/" here (e.g. non-Lumin paths to source files embedded in assets).
 		// In that case, fail silently.
@@ -584,9 +584,9 @@ public:
 			}
 			else
 			{
-				// perform a case-insensitive search
+				// perform a search
 				// make sure we get the absolute filename
-				checkf(Filename[0] == TEXT('/'), TEXT("Filename '%s' given to OpenCaseInsensitiveRead is not absolute!"), *Filename);
+				checkf(Filename[0] == TEXT('/'), TEXT("Filename '%s' given to OpenReadInternal is not absolute!"), *Filename);
 				
 				int MaxPathComponents = CountPathComponents(Filename);
 				if (MaxPathComponents > 0)
@@ -611,7 +611,7 @@ public:
 	}
 };
 
-FLuminFileMapper GCaseInsensMapper;
+FLuminFileMapper GLuminFileMapper;
 
 /**
 * Lumin File I/O implementation
@@ -645,7 +645,7 @@ bool FLuminPlatformFile::FileExists(const TCHAR* Filename)
 	FString NormalizedFilename = NormalizeFilename(Filename);
 
 	// Check the read path first, if it doesnt exist, check for the write path instead.
-	return (FileExistsCaseInsensitive(ConvertToLuminPath(NormalizedFilename, false))) ? true : FileExistsCaseInsensitive(ConvertToLuminPath(NormalizedFilename, true));
+	return (FileExistsInternal(ConvertToLuminPath(NormalizedFilename, false))) ? true : FileExistsInternal(ConvertToLuminPath(NormalizedFilename, true));
 }
 
 bool FLuminPlatformFile::FileExists(const TCHAR* Filename, FString& OutLuminPath)
@@ -654,7 +654,7 @@ bool FLuminPlatformFile::FileExists(const TCHAR* Filename, FString& OutLuminPath
 	FString NormalizedFilename = NormalizeFilename(Filename);
 	FString ReadFilePath = ConvertToLuminPath(NormalizedFilename, false);
 	
-	if (FileExistsCaseInsensitive(ReadFilePath))
+	if (FileExistsInternal(ReadFilePath))
 	{
 		OutLuminPath = ReadFilePath;
 		bExists = true;
@@ -662,7 +662,7 @@ bool FLuminPlatformFile::FileExists(const TCHAR* Filename, FString& OutLuminPath
 	else
 	{
 		FString WriteFilePath = ConvertToLuminPath(NormalizedFilename, true);
-		if (FileExistsCaseInsensitive(WriteFilePath))
+		if (FileExistsInternal(WriteFilePath))
 		{
 			OutLuminPath = WriteFilePath;
 			bExists = true;
@@ -679,28 +679,28 @@ int64 FLuminPlatformFile::FileSize(const TCHAR* Filename)
 	FString LuminPath;
 	if (FileExists(Filename, LuminPath))
 	{
-		return FileSizeCaseInsensitive(LuminPath);
+		return FileSizeInterenal(LuminPath);
 	}
 	return -1;
 }
 
 bool FLuminPlatformFile::DeleteFile(const TCHAR* Filename)
 {
-	FString CaseSensitiveFilename;
+	FString MappedFilename;
 	// Only delete from write path
 	FString IntendedFilename(ConvertToLuminPath(NormalizeFilename(Filename), true));
-	if (!GCaseInsensMapper.MapCaseInsensitiveFile(IntendedFilename, CaseSensitiveFilename))
+	if (!GLuminFileMapper.MapFile(IntendedFilename, MappedFilename))
 	{
 		// could not find the file
 		return false;
 	}
 
 	// removing mapped file is too dangerous
-	if (IntendedFilename != CaseSensitiveFilename)
+	if (IntendedFilename != MappedFilename)
 	{
-		UE_LOG(LogLuminPlatformFile, Warning, TEXT("Could not find file '%s', deleting file '%s' instead (for consistency with the rest of file ops)"), *IntendedFilename, *CaseSensitiveFilename);
+		UE_LOG(LogLuminPlatformFile, Warning, TEXT("Could not find file '%s', deleting file '%s' instead (for consistency with the rest of file ops)"), *IntendedFilename, *MappedFilename);
 	}
-	return unlink(TCHAR_TO_UTF8(*CaseSensitiveFilename)) == 0;
+	return unlink(TCHAR_TO_UTF8(*MappedFilename)) == 0;
 }
 
 bool FLuminPlatformFile::IsReadOnly(const TCHAR* Filename)
@@ -710,7 +710,7 @@ bool FLuminPlatformFile::IsReadOnly(const TCHAR* Filename)
 	FString LuminPath;
 	if (FileExists(Filename, LuminPath))
 	{
-		return IsReadOnlyCaseInsensitive(LuminPath);
+		return IsReadOnlyInternal(LuminPath);
 	}
 	return false;
 }
@@ -720,28 +720,28 @@ bool FLuminPlatformFile::MoveFile(const TCHAR* To, const TCHAR* From)
 	// Move to write path only.
 	FString ToLuminFilename = ConvertToLuminPath(NormalizeFilename(To), true);
 	FString FromLuminFilename = ConvertToLuminPath(NormalizeFilename(From), true);
-	FString CaseSensitiveFilename;
-	if (!GCaseInsensMapper.MapCaseInsensitiveFile(FromLuminFilename, CaseSensitiveFilename))
+	FString MappedFilename;
+	if (!GLuminFileMapper.MapFile(FromLuminFilename, MappedFilename))
 	{
 		// could not find the file
 		return false;
 	}
 
-	return rename(TCHAR_TO_UTF8(*CaseSensitiveFilename), TCHAR_TO_UTF8(*ToLuminFilename)) == 0;
+	return rename(TCHAR_TO_UTF8(*MappedFilename), TCHAR_TO_UTF8(*ToLuminFilename)) == 0;
 }
 
 bool FLuminPlatformFile::SetReadOnly(const TCHAR* Filename, bool bNewReadOnlyValue)
 {
 	FString LuminFilename = ConvertToLuminPath(NormalizeFilename(Filename), false);
-	FString CaseSensitiveFilename;
-	if (!GCaseInsensMapper.MapCaseInsensitiveFile(LuminFilename, CaseSensitiveFilename))
+	FString MappedFilename;
+	if (!GLuminFileMapper.MapFile(LuminFilename, MappedFilename))
 	{
 		// could not find the file
 		return false;
 	}
 
 	struct stat FileInfo;
-	if (stat(TCHAR_TO_UTF8(*CaseSensitiveFilename), &FileInfo) == 0)
+	if (stat(TCHAR_TO_UTF8(*MappedFilename), &FileInfo) == 0)
 	{
 		if (bNewReadOnlyValue)
 		{
@@ -751,7 +751,7 @@ bool FLuminPlatformFile::SetReadOnly(const TCHAR* Filename, bool bNewReadOnlyVal
 		{
 			FileInfo.st_mode |= S_IWUSR;
 		}
-		return chmod(TCHAR_TO_UTF8(*CaseSensitiveFilename), FileInfo.st_mode) == 0;
+		return chmod(TCHAR_TO_UTF8(*MappedFilename), FileInfo.st_mode) == 0;
 	}
 	return false;
 }
@@ -763,16 +763,16 @@ FDateTime FLuminPlatformFile::GetTimeStamp(const TCHAR* Filename)
 	FString LuminPath;
 	if (FileExists(Filename, LuminPath))
 	{
-		return GetTimeStampCaseInsensitive(LuminPath);
+		return GetTimeStampInternal(LuminPath);
 	}
 	return FDateTime::MinValue();
 }
 
 void FLuminPlatformFile::SetTimeStamp(const TCHAR* Filename, const FDateTime DateTime)
 {
-	FString CaseSensitiveFilename;
+	FString MappedFilename;
 	// Update timestamp on a file in the write path only.
-	if (!GCaseInsensMapper.MapCaseInsensitiveFile(ConvertToLuminPath(NormalizeFilename(Filename), true), CaseSensitiveFilename))
+	if (!GLuminFileMapper.MapFile(ConvertToLuminPath(NormalizeFilename(Filename), true), MappedFilename))
 	{
 		// could not find the file
 		return;
@@ -780,7 +780,7 @@ void FLuminPlatformFile::SetTimeStamp(const TCHAR* Filename, const FDateTime Dat
 
 	// get file times
 	struct stat FileInfo;
-	if(stat(TCHAR_TO_UTF8(*CaseSensitiveFilename), &FileInfo) != 0)
+	if(stat(TCHAR_TO_UTF8(*MappedFilename), &FileInfo) != 0)
 	{
 		return;
 	}
@@ -789,7 +789,7 @@ void FLuminPlatformFile::SetTimeStamp(const TCHAR* Filename, const FDateTime Dat
 	struct utimbuf Times;
 	Times.actime = FileInfo.st_atime;
 	Times.modtime = (DateTime - UnixEpoch).GetTotalSeconds();
-	utime(TCHAR_TO_UTF8(*CaseSensitiveFilename), &Times);
+	utime(TCHAR_TO_UTF8(*MappedFilename), &Times);
 }
 
 FDateTime FLuminPlatformFile::GetAccessTimeStamp(const TCHAR* Filename)
@@ -799,7 +799,7 @@ FDateTime FLuminPlatformFile::GetAccessTimeStamp(const TCHAR* Filename)
 	FString LuminPath;
 	if (FileExists(Filename, LuminPath))
 	{
-		return GetAccessTimeStampCaseInsensitive(LuminPath);
+		return GetAccessTimeStampInternal(LuminPath);
 	}
 	return FDateTime::MinValue();
 }
@@ -807,15 +807,6 @@ FDateTime FLuminPlatformFile::GetAccessTimeStamp(const TCHAR* Filename)
 FString FLuminPlatformFile::GetFilenameOnDisk(const TCHAR* Filename)
 {
 	return Filename;
-/*
-	FString CaseSensitiveFilename;
-	if (!GCaseInsensMapper.MapCaseInsensitiveFile(NormalizeFilename(Filename), CaseSensitiveFilename))
-	{
-		return Filename;
-	}
-
-	return CaseSensitiveFilename;
-*/
 }
 
 IFileHandle* FLuminPlatformFile::OpenRead(const TCHAR* Filename, bool bAllowWrite)
@@ -823,11 +814,11 @@ IFileHandle* FLuminPlatformFile::OpenRead(const TCHAR* Filename, bool bAllowWrit
 	FString NormalizedFilename = NormalizeFilename(Filename);
 	FString MappedToName;
 	// Check the read path.
-	int32 Handle = GCaseInsensMapper.OpenCaseInsensitiveRead(ConvertToLuminPath(NormalizedFilename, false), MappedToName);
+	int32 Handle = GLuminFileMapper.OpenReadInternal(ConvertToLuminPath(NormalizedFilename, false), MappedToName);
 	if (Handle == -1)
 	{
 		// If not in the read path, check the write path.
-		Handle = GCaseInsensMapper.OpenCaseInsensitiveRead(ConvertToLuminPath(NormalizedFilename, true), MappedToName);
+		Handle = GLuminFileMapper.OpenReadInternal(ConvertToLuminPath(NormalizedFilename, true), MappedToName);
 		if (Handle == -1)
 		{
 			return nullptr;
@@ -896,7 +887,7 @@ bool FLuminPlatformFile::DirectoryExists(const TCHAR* Directory)
 {
 	FString NormalizedFilename = NormalizeFilename(Directory);
 	// Check the read path first, if it doesnt exist, check for the write path instead.
-	return (DirectoryExistsCaseInsensitive(ConvertToLuminPath(NormalizedFilename, false))) ? true : DirectoryExistsCaseInsensitive(ConvertToLuminPath(NormalizedFilename, true));
+	return (DirectoryExistsInternal(ConvertToLuminPath(NormalizedFilename, false))) ? true : DirectoryExistsInternal(ConvertToLuminPath(NormalizedFilename, true));
 }
 
 bool FLuminPlatformFile::CreateDirectory(const TCHAR* Directory)
@@ -909,21 +900,21 @@ bool FLuminPlatformFile::CreateDirectory(const TCHAR* Directory)
 
 bool FLuminPlatformFile::DeleteDirectory(const TCHAR* Directory)
 {
-	FString CaseSensitiveFilename;
+	FString MappedFilename;
 	// Delete directory from write path only.
 	FString IntendedFilename(ConvertToLuminPath(NormalizeFilename(Directory), true));
-	if (!GCaseInsensMapper.MapCaseInsensitiveFile(IntendedFilename, CaseSensitiveFilename))
+	if (!GLuminFileMapper.MapFile(IntendedFilename, MappedFilename))
 	{
 		// could not find the directory
 		return false;
 	}
 
 	// removing mapped directory is too dangerous
-	if (IntendedFilename != CaseSensitiveFilename)
+	if (IntendedFilename != MappedFilename)
 	{
-		UE_LOG(LogLuminPlatformFile, Warning, TEXT("Could not find directory '%s', deleting '%s' instead (for consistency with the rest of file ops)"), *IntendedFilename, *CaseSensitiveFilename);
+		UE_LOG(LogLuminPlatformFile, Warning, TEXT("Could not find directory '%s', deleting '%s' instead (for consistency with the rest of file ops)"), *IntendedFilename, *MappedFilename);
 	}
-	return rmdir(TCHAR_TO_UTF8(*CaseSensitiveFilename)) == 0;
+	return rmdir(TCHAR_TO_UTF8(*MappedFilename)) == 0;
 }
 
 void FLuminPlatformFile::SetSandboxEnabled(bool bInEnabled)
@@ -953,8 +944,8 @@ FFileStatData FLuminPlatformFile::GetStatData(const TCHAR* FilenameOrDirectory)
 	bool found = false;
 
 	// Check the read path first, if it doesnt exist, check for the write path instead.
-	FFileStatData result = GetStatDataCaseInsensitive(ConvertToLuminPath(NormalizedFilename, false), found);
-	return (found) ? result : GetStatDataCaseInsensitive(ConvertToLuminPath(NormalizedFilename, true), found);
+	FFileStatData result = GetStatDataInternal(ConvertToLuminPath(NormalizedFilename, false), found);
+	return (found) ? result : GetStatDataInternal(ConvertToLuminPath(NormalizedFilename, true), found);
 }
 
 bool FLuminPlatformFile::IterateDirectory(const TCHAR* Directory, FDirectoryVisitor& Visitor)
@@ -1161,17 +1152,17 @@ FString FLuminPlatformFile::ConvertToLuminPath(const FString& Filename, bool bFo
 	return Result.ToLower();
 }
 
-bool FLuminPlatformFile::FileExistsCaseInsensitive(const FString& NormalizedFilename) const
+bool FLuminPlatformFile::FileExistsInternal(const FString& NormalizedFilename) const
 {
-	FString CaseSensitiveFilename;
-	if (!GCaseInsensMapper.MapCaseInsensitiveFile(NormalizedFilename, CaseSensitiveFilename))
+	FString MappedFilename;
+	if (!GLuminFileMapper.MapFile(NormalizedFilename, MappedFilename))
 	{
 		// could not find the file
 		return false;
 	}
 
 	struct stat FileInfo;
-	if (stat(TCHAR_TO_UTF8(*CaseSensitiveFilename), &FileInfo) != -1)
+	if (stat(TCHAR_TO_UTF8(*MappedFilename), &FileInfo) != -1)
 	{
 		return S_ISREG(FileInfo.st_mode);
 	}
@@ -1179,10 +1170,10 @@ bool FLuminPlatformFile::FileExistsCaseInsensitive(const FString& NormalizedFile
 	return false;
 }
 
-int64 FLuminPlatformFile::FileSizeCaseInsensitive(const FString& NormalizedFilename) const
+int64 FLuminPlatformFile::FileSizeInterenal(const FString& NormalizedFilename) const
 {
-	FString CaseSensitiveFilename;
-	if (!GCaseInsensMapper.MapCaseInsensitiveFile(NormalizedFilename, CaseSensitiveFilename))
+	FString MappedFilename;
+	if (!GLuminFileMapper.MapFile(NormalizedFilename, MappedFilename))
 	{
 		// could not find the file
 		return -1;
@@ -1190,7 +1181,7 @@ int64 FLuminPlatformFile::FileSizeCaseInsensitive(const FString& NormalizedFilen
 
 	struct stat FileInfo;
 	FileInfo.st_size = -1;
-	if (stat(TCHAR_TO_UTF8(*CaseSensitiveFilename), &FileInfo) != -1)
+	if (stat(TCHAR_TO_UTF8(*MappedFilename), &FileInfo) != -1)
 	{
 		// make sure to return -1 for directories
 		if (S_ISDIR(FileInfo.st_mode))
@@ -1201,28 +1192,28 @@ int64 FLuminPlatformFile::FileSizeCaseInsensitive(const FString& NormalizedFilen
 	return FileInfo.st_size;
 }
 
-bool FLuminPlatformFile::IsReadOnlyCaseInsensitive(const FString& NormalizedFilename) const
+bool FLuminPlatformFile::IsReadOnlyInternal(const FString& NormalizedFilename) const
 {
-	FString CaseSensitiveFilename;
-	if (!GCaseInsensMapper.MapCaseInsensitiveFile(NormalizedFilename, CaseSensitiveFilename))
+	FString MappedFilename;
+	if (!GLuminFileMapper.MapFile(NormalizedFilename, MappedFilename))
 	{
 		// could not find the file
 		return false;
 	}
 
-	// skipping checking F_OK since this is already taken care of by case mapper
+	// skipping checking F_OK since this is already taken care of by file mapper
 
-	if (access(TCHAR_TO_UTF8(*CaseSensitiveFilename), W_OK) == -1)
+	if (access(TCHAR_TO_UTF8(*MappedFilename), W_OK) == -1)
 	{
 		return errno == EACCES;
 	}
 	return false;
 }
 
-FDateTime FLuminPlatformFile::GetTimeStampCaseInsensitive(const FString& NormalizedFilename) const
+FDateTime FLuminPlatformFile::GetTimeStampInternal(const FString& NormalizedFilename) const
 {
-	FString CaseSensitiveFilename;
-	if (!GCaseInsensMapper.MapCaseInsensitiveFile(NormalizedFilename, CaseSensitiveFilename))
+	FString MappedFilename;
+	if (!GLuminFileMapper.MapFile(NormalizedFilename, MappedFilename))
 	{
 		// could not find the file
 		return FDateTime::MinValue();
@@ -1230,7 +1221,7 @@ FDateTime FLuminPlatformFile::GetTimeStampCaseInsensitive(const FString& Normali
 
 	// get file times
 	struct stat FileInfo;
-	if (stat(TCHAR_TO_UTF8(*CaseSensitiveFilename), &FileInfo) == -1)
+	if (stat(TCHAR_TO_UTF8(*MappedFilename), &FileInfo) == -1)
 	{
 		if (errno == EOVERFLOW)
 		{
@@ -1248,10 +1239,10 @@ FDateTime FLuminPlatformFile::GetTimeStampCaseInsensitive(const FString& Normali
 	return UnixEpoch + TimeSinceEpoch;
 }
 
-FDateTime FLuminPlatformFile::GetAccessTimeStampCaseInsensitive(const FString& NormalizedFilename) const
+FDateTime FLuminPlatformFile::GetAccessTimeStampInternal(const FString& NormalizedFilename) const
 {
-	FString CaseSensitiveFilename;
-	if (!GCaseInsensMapper.MapCaseInsensitiveFile(NormalizedFilename, CaseSensitiveFilename))
+	FString MappedFilename;
+	if (!GLuminFileMapper.MapFile(NormalizedFilename, MappedFilename))
 	{
 		// could not find the file
 		return FDateTime::MinValue();
@@ -1259,7 +1250,7 @@ FDateTime FLuminPlatformFile::GetAccessTimeStampCaseInsensitive(const FString& N
 
 	// get file times
 	struct stat FileInfo;
-	if (stat(TCHAR_TO_UTF8(*CaseSensitiveFilename), &FileInfo) == -1)
+	if (stat(TCHAR_TO_UTF8(*MappedFilename), &FileInfo) == -1)
 	{
 		return FDateTime::MinValue();
 	}
@@ -1269,10 +1260,10 @@ FDateTime FLuminPlatformFile::GetAccessTimeStampCaseInsensitive(const FString& N
 	return UnixEpoch + TimeSinceEpoch;
 }
 
-FFileStatData FLuminPlatformFile::GetStatDataCaseInsensitive(const FString& NormalizedFilename, bool& bFound) const
+FFileStatData FLuminPlatformFile::GetStatDataInternal(const FString& NormalizedFilename, bool& bFound) const
 {
-	FString CaseSensitiveFilename;
-	if (!GCaseInsensMapper.MapCaseInsensitiveFile(NormalizedFilename, CaseSensitiveFilename))
+	FString MappedFilename;
+	if (!GLuminFileMapper.MapFile(NormalizedFilename, MappedFilename))
 	{
 		// could not find the file
 		bFound = false;
@@ -1280,7 +1271,7 @@ FFileStatData FLuminPlatformFile::GetStatDataCaseInsensitive(const FString& Norm
 	}
 
 	struct stat FileInfo;
-	if (stat(TCHAR_TO_UTF8(*CaseSensitiveFilename), &FileInfo) == -1)
+	if (stat(TCHAR_TO_UTF8(*MappedFilename), &FileInfo) == -1)
 	{
 		bFound = false;
 		return FFileStatData();
@@ -1290,17 +1281,17 @@ FFileStatData FLuminPlatformFile::GetStatDataCaseInsensitive(const FString& Norm
 	return UnixStatToUEFileData(FileInfo);
 }
 
-bool FLuminPlatformFile::DirectoryExistsCaseInsensitive(const FString& NormalizedFilename) const
+bool FLuminPlatformFile::DirectoryExistsInternal(const FString& NormalizedFilename) const
 {
-	FString CaseSensitiveFilename;
-	if (!GCaseInsensMapper.MapCaseInsensitiveFile(NormalizedFilename, CaseSensitiveFilename))
+	FString MappedFilename;
+	if (!GLuminFileMapper.MapFile(NormalizedFilename, MappedFilename))
 	{
 		// could not find the file
 		return false;
 	}
 
 	struct stat FileInfo;
-	if (stat(TCHAR_TO_UTF8(*CaseSensitiveFilename), &FileInfo) != -1)
+	if (stat(TCHAR_TO_UTF8(*MappedFilename), &FileInfo) != -1)
 	{
 		return S_ISDIR(FileInfo.st_mode);
 	}
diff --git a/Engine/Source/Runtime/Core/Public/Lumin/LuminPlatformFile.h b/Engine/Source/Runtime/Core/Public/Lumin/LuminPlatformFile.h
--- a/Engine/Source/Runtime/Core/Public/Lumin/LuminPlatformFile.h
+++ b/Engine/Source/Runtime/Core/Public/Lumin/LuminPlatformFile.h
@@ -75,11 +75,11 @@ protected:
 	bool bIsSandboxEnabled = true;
 
 private:
-	bool FileExistsCaseInsensitive(const FString& NormalizedFilename) const;
-	int64 FileSizeCaseInsensitive(const FString& NormalizedFilename) const;
-	bool IsReadOnlyCaseInsensitive(const FString& NormalizedFilename) const;
-	FDateTime GetTimeStampCaseInsensitive(const FString& NormalizedFilename) const;
-	FDateTime GetAccessTimeStampCaseInsensitive(const FString& NormalizedFilename) const;
-	FFileStatData GetStatDataCaseInsensitive(const FString& NormalizedFilename, bool& bFound) const;
-	bool DirectoryExistsCaseInsensitive(const FString& NormalizedFilename) const;
+	bool FileExistsInternal(const FString& NormalizedFilename) const;
+	int64 FileSizeInterenal(const FString& NormalizedFilename) const;
+	bool IsReadOnlyInternal(const FString& NormalizedFilename) const;
+	FDateTime GetTimeStampInternal(const FString& NormalizedFilename) const;
+	FDateTime GetAccessTimeStampInternal(const FString& NormalizedFilename) const;
+	FFileStatData GetStatDataInternal(const FString& NormalizedFilename, bool& bFound) const;
+	bool DirectoryExistsInternal(const FString& NormalizedFilename) const;
 };
