Rename of the HoloLens platform files so that we can begin to work through the compile issues due to the platform change

diff --git a/Engine/Source/Programs/UnrealBuildTool/Platform/HoloLens/HoloLensToolChain.cs b/Engine/Source/Programs/UnrealBuildTool/Platform/HoloLens/HoloLensToolChain.cs
--- a/Engine/Source/Programs/UnrealBuildTool/Platform/HoloLens/HoloLensToolChain.cs
+++ b/Engine/Source/Programs/UnrealBuildTool/Platform/HoloLens/HoloLensToolChain.cs
@@ -639,10 +639,18 @@ namespace UnrealBuildTool
 				
 				if (CompileEnvironment.bGenerateDependenciesFile)
 				{
-					CompileAction.DependencyListFile = FileItem.GetItemByFileReference(FileReference.Combine(OutputDir, SourceFile.Location.GetFileName() + ".txt"));
-					CompileAction.CommandArguments = String.Format("{0} -- {1} {2} /showIncludes", Utils.MakePathSafeToUseWithCommandLine(CompileAction.DependencyListFile.Location), Utils.MakePathSafeToUseWithCommandLine(CompileAction.CommandPath), CompileAction.CommandArguments);
-					CompileAction.CommandPath = FileReference.Combine(UnrealBuildTool.EngineDirectory, "Build", "Windows", "cl-filter", "cl-filter.exe");
+					List<string> CommandArguments = new List<string>();
+					CompileAction.DependencyListFile = FileItem.GetItemByFileReference(FileReference.Combine(OutputDir, String.Format("{0}.txt", SourceFile.Location.GetFileName())));
 					CompileAction.ProducedItems.Add(CompileAction.DependencyListFile);
+					CommandArguments.Add(String.Format("-dependencies={0}", Utils.MakePathSafeToUseWithCommandLine(CompileAction.DependencyListFile.Location)));
+
+					CommandArguments.Add(String.Format("-compiler={0}", Utils.MakePathSafeToUseWithCommandLine(CompileAction.CommandPath)));
+					CommandArguments.Add("--");
+					CommandArguments.Add(Utils.MakePathSafeToUseWithCommandLine(CompileAction.CommandPath));
+					CommandArguments.Add(CompileAction.CommandArguments);
+					CommandArguments.Add("/showIncludes");
+					CompileAction.CommandArguments = string.Join(" ", CommandArguments);
+					CompileAction.CommandPath = FileReference.Combine(UnrealBuildTool.EngineDirectory, "Build", "Windows", "cl-filter", "cl-filter.exe");
 				}
 
 				if (CompileEnvironment.PrecompiledHeaderAction == PrecompiledHeaderAction.Create)
diff --git a/Engine/Source/Runtime/Core/Public/HoloLens/HoloLensAtomics.h b/Engine/Source/Runtime/Core/Public/HoloLens/HoloLensAtomics.h
deleted file mode 100644
diff --git a/Engine/Source/Runtime/Core/Public/HoloLens/HoloLensCompilerSetup.h b/Engine/Source/Runtime/Core/Public/HoloLens/HoloLensCompilerSetup.h
deleted file mode 100644
diff --git a/Engine/Source/Runtime/Core/Public/HoloLens/HoloLensFile.h b/Engine/Source/Runtime/Core/Public/HoloLens/HoloLensFile.h
deleted file mode 100644
diff --git a/Engine/Source/Runtime/Core/Public/HoloLens/HoloLensMath.h b/Engine/Source/Runtime/Core/Public/HoloLens/HoloLensMath.h
deleted file mode 100644
diff --git a/Engine/Source/Runtime/Core/Public/HoloLens/HoloLensMemory.h b/Engine/Source/Runtime/Core/Public/HoloLens/HoloLensMemory.h
deleted file mode 100644
diff --git a/Engine/Source/Runtime/Core/Public/HoloLens/HoloLensMisc.h b/Engine/Source/Runtime/Core/Public/HoloLens/HoloLensMisc.h
deleted file mode 100644
diff --git a/Engine/Source/Runtime/Core/Public/HoloLens/HoloLensOutputDevices.h b/Engine/Source/Runtime/Core/Public/HoloLens/HoloLensOutputDevices.h
deleted file mode 100644
diff --git a/Engine/Source/Runtime/Core/Public/HoloLens/HoloLensPlatformAtomics.h b/Engine/Source/Runtime/Core/Public/HoloLens/HoloLensPlatformAtomics.h
new file mode 100644
--- /dev/null
+++ b/Engine/Source/Runtime/Core/Public/HoloLens/HoloLensPlatformAtomics.h
@@ -0,0 +1,316 @@
+// Copyright 1998-2019 Epic Games, Inc. All Rights Reserved.
+
+/*=============================================================================================
+	HoloLensAtomics.h: HoloLens platform Atomics functions
+==============================================================================================*/
+
+#pragma once
+#include "CoreTypes.h"
+#include "GenericPlatform/GenericPlatformAtomics.h"
+#include "HoloLensSystemIncludes.h"
+#include <intrin.h>
+
+
+#pragma intrinsic(_InterlockedExchangeAdd8)
+#pragma intrinsic(_InterlockedIncrement16)
+#pragma intrinsic(_InterlockedIncrement)
+#if PLATFORM_64BITS
+#pragma intrinsic(_InterlockedIncrement64)
+#pragma intrinsic(_InterlockedDecrement64)
+#pragma intrinsic(_InterlockedExchangeAdd64)
+#pragma intrinsic(_InterlockedExchange64)
+#pragma intrinsic(_InterlockedCompareExchange64)
+#else
+#pragma intrinsic(_InterlockedCompareExchange64)
+#endif
+#pragma intrinsic(_InterlockedDecrement16)
+#pragma intrinsic(_InterlockedDecrement)
+#pragma intrinsic(_InterlockedExchangeAdd16)
+#pragma intrinsic(_InterlockedExchangeAdd)
+#pragma intrinsic(_InterlockedExchange8)
+#pragma intrinsic(_InterlockedExchange16)
+#pragma intrinsic(_InterlockedExchange)
+#pragma intrinsic(_InterlockedExchangePointer)
+#pragma intrinsic(_InterlockedCompareExchange8)
+#pragma intrinsic(_InterlockedCompareExchange16)
+#pragma intrinsic(_InterlockedCompareExchange)
+
+/**
+ * HoloLens implementation of the Atomics OS functions
+ */
+struct CORE_API FHoloLensAtomics : public FGenericPlatformAtomics
+{
+
+	static FORCEINLINE int8 InterlockedIncrement(volatile int8* Value)
+	{
+		return (int8)::_InterlockedExchangeAdd8((volatile char*)Value, 1) + 1;
+	}
+
+	static FORCEINLINE int16 InterlockedIncrement(volatile int16* Value)
+	{
+		return (int16)::_InterlockedIncrement16((volatile short*)Value);
+	}
+
+	static FORCEINLINE int32 InterlockedIncrement(volatile int32* Value)
+	{
+		return (int32)::_InterlockedIncrement((volatile long*)Value);
+	}
+
+	static FORCEINLINE int64 InterlockedIncrement (volatile int64* Value)
+	{
+#if PLATFORM_64BITS
+		return (int64)::_InterlockedIncrement64((volatile int64*)Value);
+#else
+		// No explicit instruction for 64-bit atomic increment on 32-bit processors; has to be implemented in terms of CMPXCHG8B
+		for (;;)
+		{
+			int64 OldValue = *Value;
+			if (_InterlockedCompareExchange64(Value, OldValue + 1, OldValue) == OldValue)
+			{
+				return OldValue + 1;
+			}
+		}
+#endif
+	}
+
+	static FORCEINLINE int8 InterlockedDecrement(volatile int8* Value)
+	{
+		return (int8)::_InterlockedExchangeAdd8((volatile char*)Value, -1) - 1;
+	}
+
+	static FORCEINLINE int16 InterlockedDecrement(volatile int16* Value)
+	{
+		return (int16)::_InterlockedDecrement16((volatile short*)Value);
+	}
+
+	static FORCEINLINE int32 InterlockedDecrement(volatile int32* Value)
+	{
+		return (int32)::_InterlockedDecrement((volatile long*)Value);
+	}
+
+	static FORCEINLINE int64 InterlockedDecrement (volatile int64* Value)
+	{
+#if PLATFORM_64BITS
+		return (int64)::_InterlockedDecrement64((volatile int64*)Value);
+#else
+		// No explicit instruction for 64-bit atomic decrement on 32-bit processors; has to be implemented in terms of CMPXCHG8B
+		for (;;)
+		{
+			int64 OldValue = *Value;
+			if (_InterlockedCompareExchange64(Value, OldValue - 1, OldValue) == OldValue)
+			{
+				return OldValue - 1;
+			}
+		}
+#endif
+	}
+
+	static FORCEINLINE int8 InterlockedAdd(volatile int8* Value, int8 Amount)
+	{
+		return (int8)::_InterlockedExchangeAdd8((volatile char*)Value, (char)Amount);
+	}
+
+	static FORCEINLINE int16 InterlockedAdd(volatile int16* Value, int16 Amount)
+	{
+		return (int16)::_InterlockedExchangeAdd16((volatile short*)Value, (short)Amount);
+	}
+
+	static FORCEINLINE int32 InterlockedAdd(volatile int32* Value, int32 Amount)
+	{
+		return (int32)::_InterlockedExchangeAdd((volatile long*)Value, (long)Amount);
+	}
+
+	static FORCEINLINE int64 InterlockedAdd (volatile int64* Value, int64 Amount)
+	{
+#if PLATFORM_64BITS
+		return (int64)::_InterlockedExchangeAdd64((volatile int64*)Value, (int64)Amount);
+#else
+		// No explicit instruction for 64-bit atomic add on 32-bit processors; has to be implemented in terms of CMPXCHG8B
+		for (;;)
+		{
+			int64 OldValue = *Value;
+			if (_InterlockedCompareExchange64(Value, OldValue + Amount, OldValue) == OldValue)
+			{
+				return OldValue + Amount;
+			}
+		}
+#endif
+	}
+
+	static FORCEINLINE int8 InterlockedExchange(volatile int8* Value, int8 Exchange)
+	{
+		return (int8)::_InterlockedExchange8((volatile char*)Value, (char)Exchange);
+	}
+
+	static FORCEINLINE int16 InterlockedExchange(volatile int16* Value, int16 Exchange)
+	{
+		return (int16)::_InterlockedExchange16((volatile short*)Value, (short)Exchange);
+	}
+
+	static FORCEINLINE int32 InterlockedExchange(volatile int32* Value, int32 Exchange)
+	{
+		return (int32)::_InterlockedExchange((volatile long*)Value, (long)Exchange);
+	}
+
+	static FORCEINLINE int64 InterlockedExchange (volatile int64* Value, int64 Exchange)
+	{
+#if PLATFORM_64BITS
+		return ::_InterlockedExchange64(Value, Exchange);
+#else
+		// No explicit instruction for 64-bit atomic exchange on 32-bit processors; has to be implemented in terms of CMPXCHG8B
+		for (;;)
+		{
+			int64 OldValue = *Value;
+			if (_InterlockedCompareExchange64(Value, Exchange, OldValue) == OldValue)
+			{
+				return OldValue;
+			}
+		}
+#endif
+	}
+
+	static FORCEINLINE void* InterlockedExchangePtr( void** Dest, void* Exchange )
+	{
+#if !(UE_BUILD_SHIPPING || UE_BUILD_TEST)
+		if (IsAligned(Dest) == false)
+		{
+			HandleAtomicsFailure(TEXT("InterlockedExchangePointer requires Dest pointer to be aligned to %d bytes"), sizeof(void*));
+		}
+#endif
+
+		return ::_InterlockedExchangePointer(Dest, Exchange);
+	}
+
+	static FORCEINLINE int8 InterlockedCompareExchange(volatile int8* Dest, int8 Exchange, int8 Comparand)
+	{
+		return (int8)::_InterlockedCompareExchange8((volatile char*)Dest, (char)Exchange, (char)Comparand);
+	}
+
+	static FORCEINLINE int16 InterlockedCompareExchange(volatile int16* Dest, int16 Exchange, int16 Comparand)
+	{
+		return (int16)::_InterlockedCompareExchange16((volatile short*)Dest, (short)Exchange, (short)Comparand);
+	}
+
+	static FORCEINLINE int32 InterlockedCompareExchange(volatile int32* Dest,int32 Exchange,int32 Comparand)
+	{
+		return (int32)::_InterlockedCompareExchange((volatile long*)Dest, (long)Exchange, (long)Comparand);
+	}
+
+	static FORCEINLINE int64 InterlockedCompareExchange (volatile int64* Dest, int64 Exchange, int64 Comparand)
+	{
+		#if !(UE_BUILD_SHIPPING || UE_BUILD_TEST)
+			if (IsAligned(Dest, 8) == false)
+			{
+				HandleAtomicsFailure(TEXT("InterlockedCompareExchange64 requires args be aligned to %d bytes"), sizeof(int64));
+			}
+		#endif
+
+			return (int64)::_InterlockedCompareExchange64(Dest, Exchange, Comparand);
+	}
+
+	static FORCEINLINE int8 AtomicRead(volatile const int8* Src)
+	{
+		return InterlockedCompareExchange((volatile int8*)Src, 0, 0);
+	}
+
+	static FORCEINLINE int16 AtomicRead(volatile const int16* Src)
+	{
+		return InterlockedCompareExchange((volatile int16*)Src, 0, 0);
+	}
+
+	static FORCEINLINE int32 AtomicRead(volatile const int32* Src)
+	{
+		return InterlockedCompareExchange((volatile int32*)Src, 0, 0);
+	}
+
+	static FORCEINLINE int64 AtomicRead(volatile const int64* Src)
+	{
+		return InterlockedCompareExchange((volatile int64*)Src, 0, 0);
+	}
+
+	static FORCEINLINE void AtomicStore(volatile int8* Src, int8 Val)
+	{
+		InterlockedExchange(Src, Val);
+	}
+
+	static FORCEINLINE void AtomicStore(volatile int16* Src, int16 Val)
+	{
+		InterlockedExchange(Src, Val);
+	}
+
+	static FORCEINLINE void AtomicStore(volatile int32* Src, int32 Val)
+	{
+		InterlockedExchange(Src, Val);
+	}
+
+	static FORCEINLINE void AtomicStore(volatile int64* Src, int64 Val)
+	{
+		InterlockedExchange(Src, Val);
+	}
+
+	 /**
+	 *	The function compares the Destination value with the Comparand value:
+	 *		- If the Destination value is equal to the Comparand value, the Exchange value is stored in the address specified by Destination, 
+	 *		- Otherwise, the initial value of the Destination parameter is stored in the address specified specified by Comparand.
+	 *	
+	 *	@return true if Comparand equals the original value of the Destination parameter, or false if Comparand does not equal the original value of the Destination parameter.
+	 *	
+	 *	Early AMD64 processors lacked the CMPXCHG16B instruction.
+	 *	To determine whether the processor supports this operation, call the IsProcessorFeaturePresent function with PF_COMPARE64_EXCHANGE128.
+	 */
+#if	PLATFORM_HAS_128BIT_ATOMICS
+	static FORCEINLINE bool InterlockedCompareExchange128( volatile FInt128* Dest, const FInt128& Exchange, FInt128* Comparand )
+	{
+#if !(UE_BUILD_SHIPPING || UE_BUILD_TEST)
+		if (IsAligned(Dest,16) == false)
+		{
+			HandleAtomicsFailure(TEXT("InterlockedCompareExchangePointer requires Dest pointer to be aligned to 16 bytes") );
+		}
+		if (IsAligned(Comparand,16) == false)
+		{
+			HandleAtomicsFailure(TEXT("InterlockedCompareExchangePointer requires Comparand pointer to be aligned to 16 bytes") );
+		}
+#endif
+
+		return ::_InterlockedCompareExchange128((int64 volatile *)Dest, Exchange.High, Exchange.Low, (int64*)Comparand) == 1;
+	}
+	/**
+	* Atomic read of 128 bit value with a memory barrier
+	*/
+	static FORCEINLINE void AtomicRead128(const volatile FInt128* Src, FInt128* OutResult)
+	{
+		FInt128 Zero;
+		Zero.High = 0;
+		Zero.Low = 0;
+		*OutResult = Zero;
+		InterlockedCompareExchange128((FInt128*)Src, Zero, OutResult);
+	}
+
+#endif // PLATFORM_HAS_128BIT_ATOMICS
+
+	/**
+	 * Atomically compares the pointer to comparand and replaces with the exchange
+	 * pointer if they are equal and returns the original value
+	 */
+	static FORCEINLINE void* InterlockedCompareExchangePointer(void** Dest,void* Exchange,void* Comparand)
+	{
+#if !(UE_BUILD_SHIPPING || UE_BUILD_TEST)
+		// Dest needs to be aligned otherwise the function will behave unpredictably 
+		if (IsAligned( Dest ) == false)
+		{
+			HandleAtomicsFailure( TEXT( "InterlockedCompareExchangePointer requires Dest pointer to be aligned to %d bytes" ), sizeof(void*) );
+		}
+#endif // !(UE_BUILD_SHIPPING || UE_BUILD_TEST)
+#if PLATFORM_64BITS
+		return (void*)::_InterlockedCompareExchange64((volatile int64*)Dest, (int64)Exchange, (int64)Comparand);
+#else
+		return (void*)::_InterlockedCompareExchange((volatile long*)Dest, (long)Exchange, (long)Comparand);
+#endif
+	}
+
+private:
+	/** Handles atomics function failure. Since 'check' has not yet been declared here we need to call external function to use it. */
+	static void HandleAtomicsFailure( const TCHAR* InFormat, ... );
+};
+
+typedef FHoloLensAtomics FPlatformAtomics;
diff --git a/Engine/Source/Runtime/Core/Public/HoloLens/HoloLensPlatformCompilerPreSetup.h b/Engine/Source/Runtime/Core/Public/HoloLens/HoloLensPlatformCompilerPreSetup.h
new file mode 100644
--- /dev/null
+++ b/Engine/Source/Runtime/Core/Public/HoloLens/HoloLensPlatformCompilerPreSetup.h
@@ -0,0 +1,3 @@
+// Copyright 1998-2019 Epic Games, Inc. All Rights Reserved.
+
+#pragma once
diff --git a/Engine/Source/Runtime/Core/Public/HoloLens/HoloLensPlatformCompilerSetup.h b/Engine/Source/Runtime/Core/Public/HoloLens/HoloLensPlatformCompilerSetup.h
new file mode 100644
--- /dev/null
+++ b/Engine/Source/Runtime/Core/Public/HoloLens/HoloLensPlatformCompilerSetup.h
@@ -0,0 +1,215 @@
+// Copyright 1998-2019 Epic Games, Inc. All Rights Reserved.
+
+/*=============================================================================================
+	HoloLensCompilerSetup.h: pragmas, version checks and other things for the HoloLens compiler
+==============================================================================================*/
+
+#pragma once
+
+#if !PLATFORM_HOLOLENS
+	#error PLATFORM_HOLOLENS not defined
+#endif
+
+/**
+* We require at least Visual Studio 2015 to compile
+*/
+static_assert(_MSC_VER >= 1900, "Visual Studio 2015 or newer required.");
+
+
+#pragma warning (error:      4001 4002 4003 4004      4006 4007 4008 4009 4010 4011 4012 4013 4014 4015 4016 4017 4018 4019 4020 4021 4022 4023 4024 4025 4026 4027 4028 4029 4030 4031 4032 4033 4034 4035 4036 4037 4038 4039 4040 4041 4042 4043 4044 4045 4046 4047 4048 4049 4050 4051 4052 4053 4054 4055 4056 4057 4058 4059 4060           4063 4064 4065 4066 4067 4068 4069 4070 4071 4072 4073 4074 4075 4076 4077 4078 4079 4080 4081 4082 4083 4084 4085 4086 4087 4088 4089 4090 4091 4092 4093 4094 4095 4096 4097 4098 4099)
+#pragma warning (error: 4100 4101 4102 4103 4104 4105 4106 4107 4108 4109 4110 4111 4112 4113 4114 4115 4116 4117 4118 4119 4120 4121 4122 4123 4124 4125 4126 4127 4128 4129 4130 4131 4132 4133 4134 4135 4136 4137 4138 4139 4140 4141 4142 4143 4144 4145 4146 4147 4148 4149 4150 4151 4152 4153 4154 4155 4156 4157 4158 4159 4160 4161 4162 4163 4164      4166 4167 4168 4169 4170 4171 4172 4173 4174 4175 4176 4177 4178 4179 4180 4181 4182 4183 4184 4185 4186 4187 4188 4189 4190 4191 4192 4193 4194 4195 4196 4197 4198 4199)
+#pragma warning (error: 4200 4201 4202 4203 4204 4205 4206 4207 4208 4209 4210 4211 4212 4213 4214 4215 4216 4217 4218 4219 4220 4221 4222 4223 4224 4225 4226 4227 4228 4229 4230 4231 4232 4233 4234 4235 4236 4237 4238 4239 4240 4241      4243 4244 4245 4246 4247 4248 4249 4250 4251 4252 4253      4255 4256 4257 4258 4259 4260 4261 4262 4263 4264           4267 4268 4269 4270 4271 4272 4273 4274 4275 4276 4277 4278 4279 4280 4281 4282 4283 4284 4285 4286 4287 4288 4289 4290 4291 4292 4293 4294 4295      4297 4298 4299)
+#pragma warning (error: 4300 4301 4302 4303 4304      4306 4307 4308 4309 4310           4313 4314 4315 4316 4317 4318 4319 4320 4321 4322 4323 4324 4325 4326 4327 4328 4329 4330 4331 4332 4333 4334 4335 4336 4337 4338      4340 4341      4343 4344      4346 4347 4348 4349      4351 4352 4353 4354 4355 4356 4357 4358 4359 4360 4361 4362 4363 4364      4366 4367 4368 4369           4372 4373 4374 4375 4376 4377 4378 4379 4380 4381 4382 4383 4384 4385 4386 4387      4389 4390 4391 4392 4393 4394 4395 4396 4397 4398 4399)
+#pragma warning (error: 4400 4401 4402 4403 4404 4405 4406 4407 4408 4409 4410 4411      4413 4414 4415 4416 4417 4418 4419 4420 4421 4422 4423 4424 4425 4426 4427 4428 4429 4430 4431 4432 4433 4434      4436      4438 4439 4440 4441 4442 4443      4445 4446 4447 4448 4449 4450 4451 4452 4453 4454 4455                4459 4460 4461 4462           4465 4466 4467 4468 4469 4470           4473 4474 4475 4476 4477 4478 4479 4480      4482 4483 4484 4485 4486 4487 4488 4489 4490 4491 4492 4493 4494 4495 4496 4497 4498 4499)
+#pragma warning (error: 4500 4501 4502 4503 4504 4505 4506 4507 4508 4509      4511 4512 4513 4514 4515 4516 4517 4518 4519 4520 4521 4522 4523 4524 4525 4526 4527 4528 4529 4530 4531 4532 4533 4534 4535 4536 4537 4538 4539 4540 4541 4542 4543 4544 4545 4546                4550 4551 4552 4553 4554      4556 4557 4558 4559 4560 4561 4562 4563 4564 4565 4566 4567 4568 4569 4570      4572 4573      4575 4576 4577 4578 4579 4580 4581 4582 4583 4584 4585 4586 4587 4588 4589 4590 4591      4593 4594 4595 4596 4597 4598     )
+#pragma warning (error: 4600 4601 4602 4603 4604      4606 4607      4609 4610 4611 4612 4613 4614 4615 4616 4617 4618      4620 4621 4622 4623 4624 4625 4626 4627 4628 4629 4630 4631 4632 4633 4634 4635 4636 4637 4638 4639 4640 4641 4642 4643 4644 4645 4646 4647 4648 4649 4650      4652 4653 4654 4655 4656 4657 4658 4659 4660 4661 4662 4663 4664 4665 4666 4667 4668 4669 4670 4671 4672 4673 4674 4675 4676 4677 4678 4679 4680 4681 4682 4683 4684 4685 4686 4687 4688 4689 4690 4691      4693 4694 4695 4696 4697 4698 4699)
+// @MIXEDREALITY_CHANGE : BEGIN remove 4746
+#pragma warning (error: 4700      4702 4703 4704 4705 4706 4707 4708 4709 4710 4711 4712 4713 4714 4715 4716 4717 4718 4719 4720 4721 4722 4723 4724 4725 4726 4727 4728 4729      4731 4732 4733 4734 4735 4736 4737      4739 4740 4741 4742 4743 4744 4745      4747 4748 4749 4750 4751 4752 4753 4754 4755 4756 4757 4758 4759 4760 4761 4762 4763 4764 4765 4766      4768 4769 4770 4771 4772 4773      4775 4776 4777 4778 4779 4780 4781 4782 4783 4784 4785 4786 4787 4788 4789 4790 4791 4792 4793 4794 4795 4796 4797 4798 4799)
+#pragma warning (error: 4800 4801 4802 4803 4804 4805 4806 4807 4808 4809 4810 4811 4812 4813 4814 4815 4816 4817 4818 4819      4821 4822 4823 4824 4825      4827 4828 4829 4830 4831 4832 4833 4834 4835 4836           4839 4840 4841 4842 4843 4844 4845 4846 4847 4848 4849 4850 4851 4852 4853 4854 4855 4856 4857 4858 4859 4860 4861 4862 4863 4864 4865 4866 4867      4869 4870 4871 4872 4873 4874 4875 4876 4877 4878 4879 4880 4881 4882 4883 4884 4885 4886 4887 4888 4889 4890 4891 4892 4893 4894 4895 4896 4897 4898 4899)
+#pragma warning (error: 4900 4901 4902 4903 4904 4905 4906 4907 4908 4909 4910 4911 4912 4913 4914 4915 4916 4917 4918 4919 4920 4921 4922 4923 4924 4925 4926 4927 4928 4929 4930 4931 4932 4933 4934 4935 4936 4937 4938 4939 4940 4941 4942 4943 4944 4945 4946 4947 4948 4949 4950 4951 4952 4953 4954 4955 4956 4957 4958 4959 4960 4961      4963 4964 4965 4966 4967 4968 4969 4970 4971 4972 4973 4974 4975 4976 4977 4978 4979 4980 4981 4982 4983 4984 4985                4989 4990 4991 4992 4993 4994 4995      4997 4998 4999)
+
+//
+// Skipped warnings, which are not disabled below. Most are disabled by default. Might be useful to look through, re-enable some and fix the code.
+//
+
+// 4005 - 'identifier' : macro redefinition																							http://msdn.microsoft.com/en-us/library/8d10sc3w.aspx
+// 4061 - enumerator 'identifier' in switch of enum 'enumeration' is not explicitly handled by a case label:						http://msdn.microsoft.com/en-us/library/96f5t7fy.aspx
+// 4062 - enumerator 'identifier' in switch of enum 'enumeration' is not handled													http://msdn.microsoft.com/en-us/library/fdt9w8tf.aspx
+// 4165 - 'HRESULT' is being converted to 'bool'; are you sure this is what you want?
+// 4242 - 'identifier' : conversion from 'type1' to 'type2', possible loss of data													http://msdn.microsoft.com/en-us/library/3hca13eh.aspx
+// 4254 - 'operator' : conversion from 'type1' to 'type2', possible loss of data													http://msdn.microsoft.com/en-us/library/3fbf7w04.aspx
+// 4265 - 'class' : class has virtual functions, but destructor is not virtual														http://msdn.microsoft.com/en-us/library/wzxffy8c.aspx
+// 4266 - 'function' : no override available for virtual member function from base 'type'; function is hidden						http://msdn.microsoft.com/en-us/library/4b76ty10.aspx
+// 4305 - 'identifier' : truncation from 'type1' to 'type2'																			http://msdn.microsoft.com/en-us/library/0as1ke3f.aspx
+// 4296 - 'operator' : expression is always false																					http://msdn.microsoft.com/en-us/library/wz2y40yt.aspx 
+// 4311 - 'variable' : pointer truncation from 'type' to 'type'																		http://msdn.microsoft.com/en-us/library/4t91x2k5.aspx
+// 4312 - 'operation' : conversion from 'type1' to 'type2' of greater size															http://msdn.microsoft.com/en-us/library/h97f4b9y.aspx
+// 4339 - 'type' : use of undefined type detected in CLR meta-data - use of this type may lead to a runtime exception				http://msdn.microsoft.com/en-us/library/3fxw8y6x.aspx
+// 4342 - behavior change: 'function' called, but a member operator was called in previous versions									http://msdn.microsoft.com/en-us/library/z8910865.aspx
+// 4345 - behavior change: an object of POD type constructed with an initializer of the form () will be default-initialized			http://msdn.microsoft.com/en-us/library/wewb47ee.aspx
+// 4350 - behavior change: 'member1' called instead of 'member2'																	http://msdn.microsoft.com/en-us/library/0eestyah.aspx
+// 4365 - 'action' : conversion from 'type_1' to 'type_2', signed/unsigned mismatch													http://msdn.microsoft.com/en-us/library/ms173683.aspx
+// 4370 - layout of class has changed from a previous version of the compiler due to better packing
+// 4371 - layout of class may have changed from a previous version of the compiler due to better packing of member
+// 4388 - '==' : signed/unsigned mismatch																							http://msdn.microsoft.com/en-us/library/jj155806.aspx
+// 4412 - 'function': function signature contains type 'type'; C++ objects are unsafe to pass between pure code and mixed or native	http://msdn.microsoft.com/en-us/library/ms235599.aspx
+// 4435 - 'class1' : Object layout under /vd2 will change due to virtual base 'class2'
+// 4437 - dynamic_cast from virtual base 'class1' to 'class2' could fail in some contexts											http://msdn.microsoft.com/en-us/library/23k5d385.aspx
+// 4444 - top level '__unaligned' is not implemented in this context
+// 4456 - declaration of 'LocalVariable' hides previous declaration
+// 4457 - declaration of 'LocalVariable' hides function parameter
+// 4458 - declaration of 'LocalVariable' hides class member
+// 4459 - declaration of 'LocalVariable' hides global declaration
+// 4464 - relative include path contains '..'																						https://msdn.microsoft.com/en-us/library/mt612856
+// 4471 - a forward declaration of an unscoped enumeration must have an underlying type (int assumed)
+// 4472 - 'enum' is a native enum: add an access specifier (private/public) to declare a managed enum
+// 4481 - nonstandard extension used: override specifier 'keyword'																	http://msdn.microsoft.com/en-us/library/ms173703.aspx
+// 4510 - 'class' : default constructor could not be generated																		http://msdn.microsoft.com/en-us/library/2cf74y2b.aspx
+// 4547 - 'operator' : operator before comma has no effect; expected operator with side-effect										http://msdn.microsoft.com/en-us/library/y1724hsf.aspx
+// 4548 - expression before comma has no effect; expected expression with side-effect												http://msdn.microsoft.com/en-us/library/yxyxx8fx.aspx
+// 4549 - 'operator' : operator before comma has no effect; did you intend 'operator'?												http://msdn.microsoft.com/en-us/library/60yhzzeh.aspx
+// 4555 - expression has no effect; expected expression with side-effect															http://msdn.microsoft.com/en-us/library/k64a6he5.aspx
+// 4571 - Informational: catch(...) semantics changed since Visual C++ 7.1; structured exceptions (SEH) are no longer caught		http://msdn.microsoft.com/en-us/library/55s8esw4.aspx
+// 4574 - 'FALSE' is defined to be '0': did you mean to use '#if FALSE'?
+// 4608 - 'symbol1' has already been initialized by another union member in the initializer list, 'symbol2'
+// 4651 - 'definition' specified for precompiled header but not for current compile													http://msdn.microsoft.com/en-us/library/h6dykdte.aspx
+// 4692 - 'function': signature of non - private member contains assembly private native type 'native_type'							http://msdn.microsoft.com/en-us/library/ms173713.aspx
+// 4701 - Potentially uninitialized local variable 'name' used																		http://msdn.microsoft.com/en-us/library/1wea5zwe.aspx
+// 4730 - 'main' : mixing _m64 and floating point expressions may result in incorrect code											http://msdn.microsoft.com/en-us/library/3z3ww2w3.aspx
+// 4738 - storing 32-bit float result in memory, possible loss of performance														http://msdn.microsoft.com/en-us/library/c24hdbz6.aspx
+// 4767 - section name 'symbol' is longer than 8 characters and will be truncated by the linker
+// 4774 - 'sprintf_s' : format string expected in argument 3 is not a string literal
+// 4820 - 'bytes' bytes padding added after construct 'member_name'																	http://msdn.microsoft.com/en-us/library/t7khkyth.aspx
+// 4826 - Conversion from 'type1 ' to 'type_2' is sign-extended. This may cause unexpected runtime behavior.						http://msdn.microsoft.com/en-us/library/ms235307.aspx
+// 4837 - trigraph detected: '??%c' replaced by '%c'																				http://msdn.microsoft.com/en-us/library/cc664919.aspx
+// 4962 - 'function': profile-guided optimizations disabled because optimizations caused profile data to become inconsistent
+// 4986 - 'function': exception specification does not match previous declaration													http://msdn.microsoft.com/en-us/library/jj620898.aspx
+// 4987 - nonstandard extension used: 'throw (...)'
+// 4988 - 'symbol': variable declared outside class/function scope
+#pragma warning (default : 4996)
+
+// Unwanted VC++ level 4 warnings to disable.
+#pragma warning(disable : 4100) // unreferenced formal parameter										
+#pragma warning(disable : 4127) // Conditional expression is constant									
+#pragma warning(disable : 4200) // Zero-length array item at end of structure, a VC-specific extension	
+#pragma warning(disable : 4201) // nonstandard extension used : nameless struct/union	
+#pragma warning(disable : 4244) // 'expression' : conversion from 'type' to 'type', possible loss of data						
+#pragma warning(disable : 4245) // 'initializing': conversion from 'type' to 'type', signed/unsigned mismatch 
+#pragma warning(disable : 4251) // 'type' needs to have dll-interface to be used by clients of 'type'
+#pragma warning(disable : 4291) // typedef-name '' used as synonym for class-name ''                    
+//#pragma warning(disable : 4305) // 'argument' : truncation from 'double' to 'float' --- fp:precise builds only!!!!!!!
+#pragma warning(disable : 4324) // structure was padded due to __declspec(align())						
+#pragma warning(disable : 4355) // this used in base initializer list                                   
+#pragma warning(disable : 4389) // signed/unsigned mismatch                                             
+#pragma warning(disable : 4456) // shadowing\hiding local variable                            
+#pragma warning(disable : 4457) // shadowing\hiding parameter                              
+#pragma warning(disable : 4458) // shadowing\hiding member                              
+#pragma warning(disable : 4459) // shadowing\hiding global                              
+#pragma warning(disable : 4464) // relative include path contains '..' 
+#pragma warning(disable : 4511) // copy constructor could not be generated                              
+#pragma warning(disable : 4512) // assignment operator could not be generated                           
+
+
+#pragma warning(disable : 4514) // unreferenced inline function has been removed						
+#pragma warning(disable : 4699) // creating precompiled header											
+#pragma warning(disable : 4702) // unreachable code in inline expanded function							
+#pragma warning(disable : 4710) // inline function not expanded											
+#pragma warning(disable : 4711) // function selected for automatic inlining								
+#pragma warning(disable : 4714) // __forceinline function not expanded									
+#pragma warning(disable : 4482) // nonstandard extension used: enum 'enum' used in qualified name (having hte enum name helps code readability and should be part of TR1 or TR2)
+
+// @MIXEDREALITY_CHANGE : BEGIN remove 4746
+#pragma warning(disable : 4746)
+
+#pragma warning(disable : 4748)	// /GS can not protect parameters and local variables from local buffer overrun because optimizations are disabled in function
+
+// NOTE: _mm_cvtpu8_ps will generate this falsely if it doesn't get inlined
+#pragma warning(disable : 4799)	// Warning: function 'ident' has no EMMS instruction
+
+//--------------------
+// @HEMI : georgwu - 
+//-----  disable this warning because __is_pod threw a warning saying behavior has changed
+#pragma warning(disable : 4647) //__is_pod operator threw a warning
+//-----
+// @HEMI : END
+//--------------------
+
+// all of the /Wall warnings that we are able to enable
+// @todo:  once we have 2005 working check:  http://msdn2.microsoft.com/library/23k5d385(en-us,vs.80).aspx
+#pragma warning(default : 4191) // 'operator/operation' : unsafe conversion from 'type of expression' to 'type required'
+#pragma warning(disable : 4217) // 'operator' : member template functions cannot be used for copy-assignment or copy-construction
+//#pragma warning(disable : 4242) // 'variable' : conversion from 'type' to 'type', possible loss of data
+//#pragma warning(default : 4254) // 'operator' : conversion from 'type1' to 'type2', possible loss of data
+#pragma warning(default : 4255) // 'function' : no function prototype given: converting '()' to '(void)'
+#pragma warning(default : 4263) // 'function' : member function does not override any base class virtual member function
+#pragma warning(default : 4264) // 'virtual_function' : no override available for virtual member function from base 'class'; function is hidden
+#pragma warning(disable : 4267) // 'argument' : conversion from 'type1' to 'type2', possible loss of data
+#pragma warning(default : 4287) // 'operator' : unsigned/negative constant mismatch
+#pragma warning(default : 4289) // nonstandard extension used : 'var' : loop control variable declared in the for-loop is used outside the for-loop scope
+#pragma warning(default : 4302) // 'conversion' : truncation from 'type 1' to 'type 2'
+#pragma warning(disable : 4315) // 'this' pointer for member may not be aligned 8 as expected by the constructor
+#pragma warning(disable : 4316) // 'type' : object allocated on the heap may not be aligned 16
+//#pragma warning(disable : 4339) // 'type' : use of undefined type detected in CLR meta-data - use of this type may lead to a runtime exception
+#pragma warning(disable : 4347) // behavior change: 'function template' is called instead of 'function
+#pragma warning(disable : 4366) // The result of the unary '&' operator may be unaligned
+#pragma warning(disable : 4514) // unreferenced inline/local function has been removed
+#pragma warning(default : 4529) // 'member_name' : forming a pointer-to-member requires explicit use of the address-of operator ('&') and a qualified name
+#pragma warning(default : 4536) // 'type name' : type-name exceeds meta-data limit of 'limit' characters
+#pragma warning(default : 4545) // expression before comma evaluates to a function which is missing an argument list
+#pragma warning(default : 4546) // function call before comma missing argument list
+//#pragma warning(default : 4547) // 'operator' : operator before comma has no effect; expected operator with side-effect
+//#pragma warning(default : 4548) // expression before comma has no effect; expected expression with side-effect  (needed as xlocale does not compile cleanly)
+//#pragma warning(default : 4549) // 'operator' : operator before comma has no effect; did you intend 'operator'?
+//#pragma warning(disable : 4555) // expression has no effect; expected expression with side-effect
+#pragma warning(default : 4557) // '__assume' contains effect 'effect'
+#pragma warning(disable : 4623) // 'derived class' : default constructor could not be generated because a base class default constructor is inaccessible
+#pragma warning(disable : 4625) // 'derived class' : copy constructor could not be generated because a base class copy constructor is inaccessible
+#pragma warning(disable : 4626) // 'derived class' : assignment operator could not be generated because a base class assignment operator is inaccessible
+#pragma warning(default : 4628) // digraphs not supported with -Ze. Character sequence 'digraph' not interpreted as alternate token for 'char'
+#pragma warning(disable : 4640) // 'instance' : construction of local static object is not thread-safe
+#pragma warning(disable : 4668) // 'symbol' is not defined as a preprocessor macro, replacing with '0' for 'directives'
+#pragma warning(default : 4682) // 'parameter' : no directional parameter attribute specified, defaulting to [in]
+#pragma warning(default : 4686) // 'user-defined type' : possible change in behavior, change in UDT return calling convention
+#pragma warning(disable : 4710) // 'function' : function not inlined / The given function was selected for inline expansion, but the compiler did not perform the inlining.
+#pragma warning(default : 4786) // 'identifier' : identifier was truncated to 'number' characters in the debug information
+#pragma warning(default : 4793) // native code generated for function 'function': 'reason'
+#pragma warning(default : 4905) // wide string literal cast to 'LPSTR'
+#pragma warning(default : 4906) // string literal cast to 'LPWSTR'
+#pragma warning(disable : 4917) // 'declarator' : a GUID cannot only be associated with a class, interface or namespace ( ocid.h breaks this)
+#pragma warning(default : 4931) // we are assuming the type library was built for number-bit pointers
+#pragma warning(default : 4946) // reinterpret_cast used between related classes: 'class1' and 'class2'
+#pragma warning(default : 4928) // illegal copy-initialization; more than one user-defined conversion has been implicitly applied
+#pragma warning(disable : 4180) // qualifier applied to function type has no meaning; ignored
+#pragma warning( disable : 4121 )
+#pragma warning(disable : 4592) // symbol will be dynamically initialized (implementation limitation) // warning broken in VS 2015 Update 1 & 2 - see http://stackoverflow.com/a/34027257
+
+#if UE_BUILD_DEBUG
+	// xstring.h causes this warning in debug builds
+	#pragma warning(disable : 4548) // expression before comma has no effect; expected expression with side-effect
+#endif
+
+// interesting ones to turn on and off at times
+//#pragma warning(disable : 4265) // 'class' : class has virtual functions, but destructor is not virtual
+//#pragma warning(disable : 4266) // '' : no override available for virtual member function from base ''; function is hidden
+
+//#pragma warning(disable : 4296) // 'operator' : expression is always true / false
+
+//#pragma warning(disable : 4820) // 'bytes' bytes padding added after member 'member'
+
+// Mixing MMX/SSE intrinsics will cause this warning, even when it's done correctly.
+//#pragma warning(disable : 4730) //mixing _m64 and floating point expressions may result in incorrect code
+
+// It'd be nice to turn these on, but at the moment they can't be used in DEBUG due to the varargs stuff.	
+#pragma warning(disable : 4189) // local variable is initialized but not referenced 
+#pragma warning(disable : 4505) // unreferenced local function has been removed		
+
+
+// If C++ exception handling is disabled, force guarding to be off.
+#if !defined(_CPPUNWIND) && !defined(__INTELLISENSE__)
+	#error "Bad VCC option: C++ exception handling must be enabled" //lint !e309 suppress as lint doesn't have this defined
+#endif
+
+// Make sure characters are unsigned.
+#ifdef _CHAR_UNSIGNED
+	#error "Bad VC++ option: Characters must be signed" //lint !e309 suppress as lint doesn't have this defined
+#endif
diff --git a/Engine/Source/Runtime/Core/Public/HoloLens/HoloLensPlatformFile.h b/Engine/Source/Runtime/Core/Public/HoloLens/HoloLensPlatformFile.h
new file mode 100644
--- /dev/null
+++ b/Engine/Source/Runtime/Core/Public/HoloLens/HoloLensPlatformFile.h
@@ -0,0 +1,7 @@
+// Copyright 1998-2019 Epic Games, Inc. All Rights Reserved.
+
+/*=============================================================================================
+	HoloLensFile.h: HoloLens platform File functions
+==============================================================================================*/
+
+#pragma once
diff --git a/Engine/Source/Runtime/Core/Public/HoloLens/HoloLensPlatformMath.h b/Engine/Source/Runtime/Core/Public/HoloLens/HoloLensPlatformMath.h
new file mode 100644
--- /dev/null
+++ b/Engine/Source/Runtime/Core/Public/HoloLens/HoloLensPlatformMath.h
@@ -0,0 +1,9 @@
+// Copyright 1998-2019 Epic Games, Inc. All Rights Reserved.
+
+/*=============================================================================================
+	HoloLensMath.h: HoloLens platform Math functions
+==============================================================================================*/
+
+#pragma once
+#include "Windows/WindowsPlatformMath.h"
+
diff --git a/Engine/Source/Runtime/Core/Public/HoloLens/HoloLensPlatformMemory.h b/Engine/Source/Runtime/Core/Public/HoloLens/HoloLensPlatformMemory.h
new file mode 100644
--- /dev/null
+++ b/Engine/Source/Runtime/Core/Public/HoloLens/HoloLensPlatformMemory.h
@@ -0,0 +1,89 @@
+// Copyright 1998-2019 Epic Games, Inc. All Rights Reserved.
+
+
+/*=============================================================================================
+HoloLensPlatformMemory.h: HoloLens platform memory functions
+==============================================================================================*/
+
+#pragma once
+#include "GenericPlatform/GenericPlatformMemory.h"
+#include "HoloLens/HoloLensSystemIncludes.h"
+
+/**
+*	HoloLens implementation of the FGenericPlatformMemoryStats.
+*	At this moment it's just the same as the FGenericPlatformMemoryStats.
+*	Can be extended as shown in the following example.
+*/
+struct FPlatformMemoryStats : public FGenericPlatformMemoryStats
+{
+	/** Default constructor, clears all variables. */
+	FPlatformMemoryStats()
+		: FGenericPlatformMemoryStats()
+		, HoloLensSpecificMemoryStat(0)
+	{}
+
+	/** Memory stat specific only for HoloLens. */
+	SIZE_T HoloLensSpecificMemoryStat;
+};
+
+/**
+* HoloLens implementation of the memory OS functions
+**/
+struct CORE_API FHoloLensPlatformMemory : public FGenericPlatformMemory
+{
+	enum EMemoryCounterRegion
+	{
+		MCR_Invalid, // not memory
+		MCR_Physical, // main system memory
+		MCR_GPU, // memory directly a GPU (graphics card, etc)
+		MCR_GPUSystem, // system memory directly accessible by a GPU
+		MCR_TexturePool, // presized texture pools
+		MCR_StreamingPool, // amount of texture pool available for streaming.
+		MCR_UsedStreamingPool, // amount of texture pool used for streaming.
+		MCR_GPUDefragPool, // presized pool of memory that can be defragmented.
+		MCR_SamplePlatformSpecifcMemoryRegion,
+		MCR_PhysicalLLM, // total physical memory displayed in the LLM stats (on consoles CPU + GPU)
+		MCR_MAX
+	};
+
+	/**
+	* HoloLens representation of a shared memory region
+	*/
+	struct FHoloLensSharedMemoryRegion : public FSharedMemoryRegion
+	{
+		/** Returns the handle to file mapping object. */
+		HANDLE GetMapping() const { return Mapping; }
+
+		FHoloLensSharedMemoryRegion(const FString& InName, uint32 InAccessMode, void* InAddress, SIZE_T InSize, HANDLE InMapping)
+			: FSharedMemoryRegion(InName, InAccessMode, InAddress, InSize)
+			, Mapping(InMapping)
+		{}
+
+	protected:
+
+		/** Handle of a file mapping object */
+		HANDLE				Mapping;
+	};
+
+	//~ Begin FGenericPlatformMemory Interface
+	static void Init();
+	static bool SupportBackupMemoryPool()
+	{
+		return true;
+	}
+	static class FMalloc* BaseAllocator();
+	static FPlatformMemoryStats GetStats();
+	static void GetStatsForMallocProfiler(FGenericMemoryStats& out_Stats);
+	static const FPlatformMemoryConstants& GetConstants();
+	static void* BinnedAllocFromOS(SIZE_T Size);
+	static void BinnedFreeToOS(void* Ptr, SIZE_T Size);
+	static FSharedMemoryRegion* MapNamedSharedMemoryRegion(const FString& InName, bool bCreate, uint32 AccessMode, SIZE_T Size);
+	static bool UnmapNamedSharedMemoryRegion(FSharedMemoryRegion * MemoryRegion);
+protected:
+	friend struct FGenericStatsUpdater;
+
+	static void InternalUpdateStats(const FPlatformMemoryStats& MemoryStats);
+	//~ End FGenericPlatformMemory Interface
+};
+
+typedef FHoloLensPlatformMemory FPlatformMemory;
diff --git a/Engine/Source/Runtime/Core/Public/HoloLens/HoloLensPlatformMisc.h b/Engine/Source/Runtime/Core/Public/HoloLens/HoloLensPlatformMisc.h
new file mode 100644
--- /dev/null
+++ b/Engine/Source/Runtime/Core/Public/HoloLens/HoloLensPlatformMisc.h
@@ -0,0 +1,100 @@
+// Copyright 1998-2019 Epic Games, Inc. All Rights Reserved.
+
+/*=============================================================================================
+	HoloLensMisc.h: HoloLens platform misc functions
+==============================================================================================*/
+
+#pragma once
+#include "GenericPlatform/GenericPlatformMisc.h"
+#include "HAL/Platform.h"
+#include "Misc/Build.h"
+#include "HoloLens/HoloLensSystemIncludes.h"
+
+
+#define UE_DEBUG_BREAK_IMPL() PLATFORM_BREAK()
+
+class FString;
+
+/**
+ * HoloLens implementation of the misc OS functions
+ */
+struct CORE_API FHoloLensMisc : public FGenericPlatformMisc
+{
+	static void PlatformPreInit();
+	static void PlatformInit();
+	static void PlatformPostInit(bool ShowSplashScreen = false);
+	static FString GetEnvironmentVariable(const TCHAR* VariableName);
+	static const TCHAR* GetPlatformFeaturesModuleName();
+
+#if !UE_BUILD_SHIPPING
+	static bool IsDebuggerPresent();
+	FORCEINLINE static void DebugBreak()
+	{
+		if (IsDebuggerPresent())
+		{
+			__debugbreak();
+		}
+	}
+#endif
+
+	/** Break into debugger. Returning false allows this function to be used in conditionals. */
+	FORCEINLINE static bool DebugBreakReturningFalse()
+	{
+#if !UE_BUILD_SHIPPING
+		DebugBreak();
+#endif
+		return false;
+	}
+
+    /** Prompts for remote debugging if debugger is not attached. Regardless of result, breaks into debugger afterwards. Returns false for use in conditionals. */
+    static FORCEINLINE bool DebugBreakAndPromptForRemoteReturningFalse(bool bIsEnsure = false)
+    {
+#if !UE_BUILD_SHIPPING
+        if (!IsDebuggerPresent())
+        {
+            PromptForRemoteDebugging(bIsEnsure);
+        }
+
+        DebugBreak();
+#endif
+        return false;
+    }
+    
+    static void PumpMessages(bool bFromMainLoop);
+	static void LowLevelOutputDebugString(const TCHAR *Message);
+	static void RequestExit(bool Force);
+	static const TCHAR* GetSystemErrorMessage(TCHAR* OutBuffer, int32 BufferCount, int32 Error);
+	static void CreateGuid(struct FGuid& Result);
+	static int32 NumberOfCores();
+	static EAppReturnType::Type MessageBoxExt(EAppMsgType::Type MsgType, const TCHAR* Text, const TCHAR* Caption);
+
+	static bool CoInitialize();
+	static void CoUninitialize();
+
+	/** Get the application root directory. */
+	static const TCHAR* RootDir();
+
+	/**
+	* Function to save the URI string from a protocol activation for use elsewhere in the game.
+	*/
+	static void SetProtocolActivationUri(const FString& NewUriString);
+
+	/**
+	* Function to get the URI string from the most recent protocol activation.
+	*/
+	static const FString& GetProtocolActivationUri();
+
+	static bool SupportsLocalCaching()
+	{
+		return false;
+	}
+
+	static void GetValidTargetPlatforms(TArray<FString>& TargetPlatformNames);
+
+private:
+
+	/** character buffer containing the last protocol activation URI */
+	static FString ProtocolActivationUri;
+};
+
+typedef FHoloLensMisc FPlatformMisc;
diff --git a/Engine/Source/Runtime/Core/Public/HoloLens/HoloLensPlatformOutputDevices.h b/Engine/Source/Runtime/Core/Public/HoloLens/HoloLensPlatformOutputDevices.h
new file mode 100644
--- /dev/null
+++ b/Engine/Source/Runtime/Core/Public/HoloLens/HoloLensPlatformOutputDevices.h
@@ -0,0 +1,15 @@
+// Copyright 1998-2019 Epic Games, Inc. All Rights Reserved.
+
+/*=============================================================================================
+	HoloLensOutputDevices.h: Windows platform OutputDevices functions
+==============================================================================================*/
+
+#pragma once
+#include "GenericPlatform/GenericPlatformOutputDevices.h"
+
+struct CORE_API FHoloLensOutputDevices : public FGenericPlatformOutputDevices
+{
+	static FOutputDevice*			GetEventLog();
+};
+
+typedef FHoloLensOutputDevices FPlatformOutputDevices;
diff --git a/Engine/Source/Runtime/Core/Public/HoloLens/HoloLensPlatformProcess.h b/Engine/Source/Runtime/Core/Public/HoloLens/HoloLensPlatformProcess.h
new file mode 100644
--- /dev/null
+++ b/Engine/Source/Runtime/Core/Public/HoloLens/HoloLensPlatformProcess.h
@@ -0,0 +1,57 @@
+// Copyright 1998-2019 Epic Games, Inc. All Rights Reserved.
+
+/*=============================================================================================
+	HoloLensProcess.h: HoloLens platform Process functions
+==============================================================================================*/
+
+#pragma once
+#include "GenericPlatform/GenericPlatformProcess.h"
+
+/** Dummy process handle for platforms that use generic implementation. */
+struct FProcHandle : public TProcHandle<void*, nullptr>
+{
+public:
+	/** Default constructor. */
+	FORCEINLINE FProcHandle()
+		: TProcHandle()
+	{}
+
+	/** Initialization constructor. */
+	FORCEINLINE explicit FProcHandle( HandleType Other )
+		: TProcHandle( Other )
+	{}
+};
+
+/**
+ * HoloLens implementation of the Process OS functions
+ */
+struct CORE_API FHoloLensProcess : public FGenericPlatformProcess
+{
+	static const TCHAR* BaseDir();
+	static void Sleep( float Seconds );
+	static void SleepNoStats(float Seconds); 
+	static void SleepInfinite();
+	static class FEvent* CreateSynchEvent(bool bIsManualReset = false);
+	static class FRunnableThread* CreateRunnableThread();
+	static const TCHAR* ExecutableName(bool bRemoveExtension = true);
+	static void* GetDllHandle(const TCHAR* Filename);
+	static void FreeDllHandle(void* DllHandle);
+	static void SetCurrentWorkingDirectoryToBaseDir();
+	static FString GetCurrentWorkingDirectory();
+	/** Content saved to the game or engine directories should be rerouted to user directories instead **/
+	static bool ShouldSaveToUserDir();
+
+	static void SetThreadAffinityMask( uint64 AffinityMask );
+	static const TCHAR* UserDir();
+	static const TCHAR* UserSettingsDir();
+	static const TCHAR* UserTempDir();
+	static const TCHAR* ApplicationSettingsDir();
+
+	static const TCHAR* GetLocalAppDataLowLevelPath();
+	static const TCHAR* GetTempAppDataLowLevelPath();
+	static const TCHAR* GetLocalAppDataRedirectPath();
+	static const TCHAR* GetTempAppDataRedirectPath();
+};
+
+typedef FHoloLensProcess FPlatformProcess;
+
diff --git a/Engine/Source/Runtime/Core/Public/HoloLens/HoloLensPlatformProperties.h b/Engine/Source/Runtime/Core/Public/HoloLens/HoloLensPlatformProperties.h
new file mode 100644
--- /dev/null
+++ b/Engine/Source/Runtime/Core/Public/HoloLens/HoloLensPlatformProperties.h
@@ -0,0 +1,56 @@
+// Copyright 1998-2019 Epic Games, Inc. All Rights Reserved.
+
+/*================================================================================
+	HoloLensProperties.h - Basic static properties of a platform
+	These are shared between:
+		the runtime platform - via FPlatformProperties
+		the target platforms - via ITargetPlatform
+==================================================================================*/
+
+#pragma once
+#include "GenericPlatform/GenericPlatformProperties.h"
+#include "HAL/Platform.h"
+
+struct FHoloLensPlatformProperties : public FGenericPlatformProperties
+{
+	static FORCEINLINE const char* PlatformName()
+	{
+		return "HoloLens";
+	}
+	static FORCEINLINE const char* IniPlatformName()
+	{
+		return "HoloLens";
+	}
+	static FORCEINLINE bool HasEditorOnlyData()
+	{
+		return false;
+	}              
+	static FORCEINLINE bool SupportsTessellation()
+	{
+		return true;
+	}
+	static FORCEINLINE bool SupportsWindowedMode()
+	{
+		return true;
+	}
+	static FORCEINLINE bool RequiresCookedData()
+	{
+		return true;
+	}
+	static FORCEINLINE bool SupportsGrayscaleSRGB()
+	{
+		return false; // Requires expand from G8 to RGBA
+	}
+
+	static FORCEINLINE bool HasFixedResolution()
+	{
+		return false;
+	}
+	
+	static FORCEINLINE bool SupportsQuit()
+	{
+		return true;
+	}
+};
+
+
diff --git a/Engine/Source/Runtime/Core/Public/HoloLens/HoloLensPlatformStackWalk.h b/Engine/Source/Runtime/Core/Public/HoloLens/HoloLensPlatformStackWalk.h
new file mode 100644
--- /dev/null
+++ b/Engine/Source/Runtime/Core/Public/HoloLens/HoloLensPlatformStackWalk.h
@@ -0,0 +1,10 @@
+// Copyright 1998-2019 Epic Games, Inc. All Rights Reserved.
+
+/*=============================================================================================
+	HoloLensStackWalk.h: HoloLens platform stack walk functions
+==============================================================================================*/
+
+#pragma once
+#include "GenericPlatform/GenericPlatformStackWalk.h"
+
+typedef FGenericPlatformStackWalk FPlatformStackWalk;
diff --git a/Engine/Source/Runtime/Core/Public/HoloLens/HoloLensPlatformString.h b/Engine/Source/Runtime/Core/Public/HoloLens/HoloLensPlatformString.h
new file mode 100644
--- /dev/null
+++ b/Engine/Source/Runtime/Core/Public/HoloLens/HoloLensPlatformString.h
@@ -0,0 +1,18 @@
+// Copyright 1998-2019 Epic Games, Inc. All Rights Reserved.
+
+/*=============================================================================================
+	HoloLensString.h: HoloLens platform string classes, mostly implemented with ANSI C++
+==============================================================================================*/
+
+#pragma once
+#include "GenericPlatform/MicrosoftPlatformString.h"
+#include "HAL/Platform.h"
+
+/**
+ * HoloLens string implementation
+ */
+struct FHoloLensString : public FMicrosoftPlatformString
+{
+};
+
+typedef FHoloLensString FPlatformString;
diff --git a/Engine/Source/Runtime/Core/Public/HoloLens/HoloLensPlatformTLS.h b/Engine/Source/Runtime/Core/Public/HoloLens/HoloLensPlatformTLS.h
new file mode 100644
--- /dev/null
+++ b/Engine/Source/Runtime/Core/Public/HoloLens/HoloLensPlatformTLS.h
@@ -0,0 +1,64 @@
+// Copyright 1998-2019 Epic Games, Inc. All Rights Reserved.
+
+/*=============================================================================================
+HoloLensTLS.h: HoloLens platform TLS (Thread local storage and thread ID) functions
+==============================================================================================*/
+
+#pragma once
+#include "GenericPlatform/GenericPlatformTLS.h"
+#include "HoloLens/HoloLensSystemIncludes.h"
+
+/**
+* HoloLens implementation of the TLS OS functions
+*/
+struct CORE_API FHoloLensTLS : public FGenericPlatformTLS
+{
+	/**
+	* Returns the currently executing thread's id
+	*/
+	static FORCEINLINE uint32 GetCurrentThreadId(void)
+	{
+		return Windows::GetCurrentThreadId();
+	}
+
+	/**
+	* Allocates a thread local store slot
+	*/
+	static FORCEINLINE uint32 AllocTlsSlot(void)
+	{
+		return Windows::TlsAlloc();
+	}
+
+	/**
+	* Sets a value in the specified TLS slot
+	*
+	* @param SlotIndex the TLS index to store it in
+	* @param Value the value to store in the slot
+	*/
+	static FORCEINLINE void SetTlsValue(uint32 SlotIndex, void* Value)
+	{
+		Windows::TlsSetValue(SlotIndex, Value);
+	}
+
+	/**
+	* Reads the value stored at the specified TLS slot
+	*
+	* @return the value stored in the slot
+	*/
+	static FORCEINLINE void* GetTlsValue(uint32 SlotIndex)
+	{
+		return Windows::TlsGetValue(SlotIndex);
+	}
+
+	/**
+	* Frees a previously allocated TLS slot
+	*
+	* @param SlotIndex the TLS index to store it in
+	*/
+	static FORCEINLINE void FreeTlsSlot(uint32 SlotIndex)
+	{
+		Windows::TlsFree(SlotIndex);
+	}
+};
+
+typedef FHoloLensTLS FPlatformTLS;
diff --git a/Engine/Source/Runtime/Core/Public/HoloLens/HoloLensPlatformTime.h b/Engine/Source/Runtime/Core/Public/HoloLens/HoloLensPlatformTime.h
new file mode 100644
--- /dev/null
+++ b/Engine/Source/Runtime/Core/Public/HoloLens/HoloLensPlatformTime.h
@@ -0,0 +1,56 @@
+// Copyright 1998-2019 Epic Games, Inc. All Rights Reserved.
+
+/*=============================================================================================
+HoloLensTime.h: HoloLens platform Time functions
+==============================================================================================*/
+
+#pragma once
+#include "GenericPlatform/GenericPlatformTime.h"
+#include "HoloLens/HoloLensSystemIncludes.h"
+
+/**
+* HoloLens implementation of the Time OS functions
+*/
+struct CORE_API FHoloLensTime : public FGenericPlatformTime
+{
+	static double InitTiming();
+
+	static FORCEINLINE double Seconds()
+	{
+		Windows::LARGE_INTEGER Cycles;
+		Windows::QueryPerformanceCounter(&Cycles);
+		// Add big number to make bugs apparent where return value is being passed to float
+		return Cycles.QuadPart * GetSecondsPerCycle() + 16777216.0;
+	}
+
+	static FORCEINLINE uint32 Cycles()
+	{
+		Windows::LARGE_INTEGER Cycles;
+		Windows::QueryPerformanceCounter(&Cycles);
+		[REDACTED_EMAIL]: QuadPart is a LONGLONG... can't use it here!
+		//return Cycles.QuadPart;
+		return Cycles.LowPart;
+	}
+
+	static FORCEINLINE uint64 Cycles64()
+	{
+		Windows::LARGE_INTEGER Cycles;
+		Windows::QueryPerformanceCounter(&Cycles);
+		return Cycles.QuadPart;
+	}
+
+	static void SystemTime(int32& Year, int32& Month, int32& DayOfWeek, int32& Day, int32& Hour, int32& Min, int32& Sec, int32& MSec);
+	static void UtcTime(int32& Year, int32& Month, int32& DayOfWeek, int32& Day, int32& Hour, int32& Min, int32& Sec, int32& MSec);
+
+
+	static bool UpdateCPUTime(float DeltaTime);
+	static FCPUTime GetCPUTime();
+
+protected:
+
+	/** Percentage CPU utilization for the last interval relative to one core. */
+	static float CPUTimePctRelative;
+
+};
+
+typedef FHoloLensTime FPlatformTime;
diff --git a/Engine/Source/Runtime/Core/Public/HoloLens/HoloLensProcess.h b/Engine/Source/Runtime/Core/Public/HoloLens/HoloLensProcess.h
deleted file mode 100644
diff --git a/Engine/Source/Runtime/Core/Public/HoloLens/HoloLensProperties.h b/Engine/Source/Runtime/Core/Public/HoloLens/HoloLensProperties.h
deleted file mode 100644
diff --git a/Engine/Source/Runtime/Core/Public/HoloLens/HoloLensStackWalk.h b/Engine/Source/Runtime/Core/Public/HoloLens/HoloLensStackWalk.h
deleted file mode 100644
diff --git a/Engine/Source/Runtime/Core/Public/HoloLens/HoloLensString.h b/Engine/Source/Runtime/Core/Public/HoloLens/HoloLensString.h
deleted file mode 100644
diff --git a/Engine/Source/Runtime/Core/Public/HoloLens/HoloLensTLS.h b/Engine/Source/Runtime/Core/Public/HoloLens/HoloLensTLS.h
deleted file mode 100644
diff --git a/Engine/Source/Runtime/Core/Public/HoloLens/HoloLensTime.h b/Engine/Source/Runtime/Core/Public/HoloLens/HoloLensTime.h
deleted file mode 100644
