Rename of the HoloLens platform files so that we can begin to work through the compile issues due to the platform change
- Missed an update

diff --git a/Engine/Source/Runtime/Core/Public/HoloLens/HoloLensSystemIncludes.h b/Engine/Source/Runtime/Core/Public/HoloLens/HoloLensSystemIncludes.h
--- a/Engine/Source/Runtime/Core/Public/HoloLens/HoloLensSystemIncludes.h
+++ b/Engine/Source/Runtime/Core/Public/HoloLens/HoloLensSystemIncludes.h
@@ -8,7 +8,7 @@
 #define STRICT
 #endif
 
-#include "HoloLens/HoloLensCompilerSetup.h"
+#include "HoloLens/HoloLensPlatformCompilerSetup.h"
 #include "Windows/MinimalWindowsApi.h"
 
 // Macro for releasing COM objects
