Made HoloLensPlatformCompilerPreSetup match the windows version

diff --git a/Engine/Source/Runtime/Core/Public/HoloLens/HoloLensPlatformCompilerPreSetup.h b/Engine/Source/Runtime/Core/Public/HoloLens/HoloLensPlatformCompilerPreSetup.h
--- a/Engine/Source/Runtime/Core/Public/HoloLens/HoloLensPlatformCompilerPreSetup.h
+++ b/Engine/Source/Runtime/Core/Public/HoloLens/HoloLensPlatformCompilerPreSetup.h
@@ -1,3 +1,109 @@
 // Copyright 1998-2019 Epic Games, Inc. All Rights Reserved.
 
 #pragma once
+
+#ifndef DISABLE_DEPRECATION
+	#define DEPRECATED(VERSION, MESSAGE) DEPRECATED_MACRO(4.22, "The DEPRECATED macro has been deprecated in favor of UE_DEPRECATED().") __declspec(deprecated(MESSAGE " Please update your code to the new API before upgrading to the next release, otherwise your project will no longer compile."))
+
+	#define PRAGMA_DISABLE_DEPRECATION_WARNINGS \
+		__pragma (warning(push)) \
+		__pragma (warning(disable: 4995)) /* 'function': name was marked as #pragma deprecated */ \
+		__pragma (warning(disable: 4996)) /* The compiler encountered a deprecated declaration. */
+
+	#define PRAGMA_ENABLE_DEPRECATION_WARNINGS \
+		__pragma (warning(pop))
+#endif // DISABLE_DEPRECATION
+
+#ifndef PRAGMA_DISABLE_SHADOW_VARIABLE_WARNINGS
+	#define PRAGMA_DISABLE_SHADOW_VARIABLE_WARNINGS \
+		__pragma (warning(push)) \
+		__pragma (warning(disable: 4456)) /* declaration of 'LocalVariable' hides previous local declaration */ \
+		__pragma (warning(disable: 4457)) /* declaration of 'LocalVariable' hides function parameter */ \
+		__pragma (warning(disable: 4458)) /* declaration of 'LocalVariable' hides class member */ \
+		__pragma (warning(disable: 4459)) /* declaration of 'LocalVariable' hides global declaration */ \
+		__pragma (warning(disable: 6244)) /* local declaration of <variable> hides previous declaration at <line> of <file> */
+#endif // PRAGMA_DISABLE_SHADOW_VARIABLE_WARNINGS
+
+#ifndef PRAGMA_ENABLE_SHADOW_VARIABLE_WARNINGS
+	#define PRAGMA_ENABLE_SHADOW_VARIABLE_WARNINGS \
+		__pragma(warning(pop))
+#endif // PRAGMA_ENABLE_SHADOW_VARIABLE_WARNINGS
+
+#ifndef PRAGMA_DISABLE_UNDEFINED_IDENTIFIER_WARNINGS
+	#define PRAGMA_DISABLE_UNDEFINED_IDENTIFIER_WARNINGS \
+		__pragma(warning(push)) \
+		__pragma(warning(disable: 4668)) /* 'symbol' is not defined as a preprocessor macro, replacing with '0' for 'directives' */
+#endif // PRAGMA_DISABLE_UNDEFINED_IDENTIFIER_WARNINGS
+
+#ifndef PRAGMA_ENABLE_UNDEFINED_IDENTIFIER_WARNINGS
+	#define PRAGMA_ENABLE_UNDEFINED_IDENTIFIER_WARNINGS \
+		__pragma(warning(pop))
+#endif // PRAGMA_ENABLE_UNDEFINED_IDENTIFIER_WARNINGS
+
+#ifndef PRAGMA_DISABLE_MISSING_VIRTUAL_DESTRUCTOR_WARNINGS
+	#define PRAGMA_DISABLE_MISSING_VIRTUAL_DESTRUCTOR_WARNINGS \
+		__pragma(warning(push)) \
+		__pragma(warning(disable: 4265)) /* class' : class has virtual functions, but destructor is not virtual */
+#endif // PRAGMA_DISABLE_MISSING_VIRTUAL_DESTRUCTOR_WARNINGS
+
+#ifndef PRAGMA_ENABLE_MISSING_VIRTUAL_DESTRUCTOR_WARNINGS
+	#define PRAGMA_ENABLE_MISSING_VIRTUAL_DESTRUCTOR_WARNINGS \
+		__pragma(warning(pop))
+#endif // PRAGMA_ENABLE_MISSING_VIRTUAL_DESTRUCTOR_WARNINGS
+
+#ifndef PRAGMA_DISABLE_REORDER_WARNINGS
+	#define PRAGMA_DISABLE_REORDER_WARNINGS \
+		__pragma(warning(push)) \
+		__pragma(warning(disable: 5038)) /* data member 'member1' will be initialized after data member 'member2' data member 'member' will be initialized after base class 'base_class' */
+#endif // PRAGMA_DISABLE_REORDER_WARNINGS
+
+#ifndef PRAGMA_ENABLE_REORDER_WARNINGS
+	#define PRAGMA_ENABLE_REORDER_WARNINGS \
+		__pragma(warning(pop))
+#endif // PRAGMA_ENABLE_REORDER_WARNINGS
+
+#ifndef PRAGMA_POP
+	#define PRAGMA_POP \
+		__pragma(warning(pop))
+#endif // PRAGMA_POP
+
+// Disable common CA warnings around SDK includes
+#ifndef THIRD_PARTY_INCLUDES_START
+	#define THIRD_PARTY_INCLUDES_START \
+		__pragma(warning(push)) \
+		__pragma(warning(disable: 4510))  /* '<class>': default constructor could not be generated. */ \
+		__pragma(warning(disable: 4610))  /* object '<class>' can never be instantiated - user-defined constructor required. */ \
+		__pragma(warning(disable: 4946))  /* reinterpret_cast used between related classes: '<class1>' and '<class1>' */ \
+		__pragma(warning(disable: 4996))  /* '<obj>' was declared deprecated. */ \
+		__pragma(warning(disable: 6011))  /* Dereferencing NULL pointer '<ptr>'. */ \
+		__pragma(warning(disable: 6101))  /* Returning uninitialized memory '<expr>'.  A successful path through the function does not set the named _Out_ parameter. */ \
+		__pragma(warning(disable: 6287))  /* Redundant code:  the left and right sub-expressions are identical. */ \
+		__pragma(warning(disable: 6308))  /* 'realloc' might return null pointer: assigning null pointer to 'X', which is passed as an argument to 'realloc', will cause the original memory block to be leaked. */ \
+		__pragma(warning(disable: 6326))  /* Potential comparison of a constant with another constant. */ \
+		__pragma(warning(disable: 6340))  /* Mismatch on sign: Incorrect type passed as parameter in call to function. */ \
+		__pragma(warning(disable: 6385))  /* Reading invalid data from '<ptr>':  the readable size is '<num1>' bytes, but '<num2>' bytes may be read. */ \
+		__pragma(warning(disable: 6386))  /* Buffer overrun while writing to '<ptr>':  the writable size is '<num1>' bytes, but '<num2>' bytes might be written. */ \
+		__pragma(warning(disable: 28182)) /* Dereferencing NULL pointer. '<ptr1>' contains the same NULL value as '<ptr2>' did. */ \
+		__pragma(warning(disable: 28251)) /* Inconsistent annotation for '<func>': this instance has no annotations. */ \
+		__pragma(warning(disable: 28252)) /* Inconsistent annotation for '<func>': return/function has '<annotation>' on the prior instance. */ \
+		__pragma(warning(disable: 28253)) /* Inconsistent annotation for '<func>': _Param_(<num>) has '<annotation>' on the prior instance. */ \
+		__pragma(warning(disable: 28301)) /* No annotations for first declaration of '<func>'. */ \
+		PRAGMA_DISABLE_REORDER_WARNINGS \
+		PRAGMA_DISABLE_UNDEFINED_IDENTIFIER_WARNINGS \
+		PRAGMA_DISABLE_SHADOW_VARIABLE_WARNINGS \
+		PRAGMA_DISABLE_MISSING_VIRTUAL_DESTRUCTOR_WARNINGS \
+		PRAGMA_DISABLE_DEPRECATION_WARNINGS
+#endif // THIRD_PARTY_INCLUDES_START
+
+#ifndef THIRD_PARTY_INCLUDES_END
+	#define THIRD_PARTY_INCLUDES_END \
+		PRAGMA_ENABLE_DEPRECATION_WARNINGS \
+		PRAGMA_ENABLE_MISSING_VIRTUAL_DESTRUCTOR_WARNINGS \
+		PRAGMA_ENABLE_SHADOW_VARIABLE_WARNINGS \
+		PRAGMA_ENABLE_UNDEFINED_IDENTIFIER_WARNINGS \
+		PRAGMA_ENABLE_REORDER_WARNINGS \
+		__pragma(warning(pop))
+#endif // THIRD_PARTY_INCLUDES_END
+
+#define EMIT_CUSTOM_WARNING_AT_LINE(Line, Warning) \
+	__pragma(message(WARNING_LOCATION(Line) ": warning C4996: " Warning))
