HoloLensPlatformProperties now declare themselves conditionally correctly

diff --git a/Engine/Source/Runtime/Core/Public/HoloLens/HoloLensPlatformProperties.h b/Engine/Source/Runtime/Core/Public/HoloLens/HoloLensPlatformProperties.h
--- a/Engine/Source/Runtime/Core/Public/HoloLens/HoloLensPlatformProperties.h
+++ b/Engine/Source/Runtime/Core/Public/HoloLens/HoloLensPlatformProperties.h
@@ -53,4 +53,6 @@ struct FHoloLensPlatformProperties : public FGenericPlatformProperties
 	}
 };
 
-
+#ifdef PROPERTY_HEADER_SHOULD_DEFINE_TYPE
+typedef FHoloLensPlatformProperties FPlatformProperties;
+#endif
