Added missing HoloLensPlatformAffinity.h file

diff --git a/Engine/Source/Runtime/Core/Public/HoloLens/HoloLensPlatformAffinity.h b/Engine/Source/Runtime/Core/Public/HoloLens/HoloLensPlatformAffinity.h
new file mode 100644
--- /dev/null
+++ b/Engine/Source/Runtime/Core/Public/HoloLens/HoloLensPlatformAffinity.h
@@ -0,0 +1,6 @@
+// Copyright 1998-2019 Epic Games, Inc. All Rights Reserved.
+
+#pragma once
+
+#include "GenericPlatform/GenericPlatformAffinity.h"
+typedef FGenericPlatformAffinity FPlatformAffinity;
