Rename HoloLensAtomics HoloLensPlatformAtomics

diff --git a/Engine/Source/Runtime/Core/Private/HoloLens/HoloLensAtomics.cpp b/Engine/Source/Runtime/Core/Private/HoloLens/HoloLensAtomics.cpp
deleted file mode 100644
diff --git a/Engine/Source/Runtime/Core/Private/HoloLens/HoloLensPlatformAtomics.cpp b/Engine/Source/Runtime/Core/Private/HoloLens/HoloLensPlatformAtomics.cpp
new file mode 100644
--- /dev/null
+++ b/Engine/Source/Runtime/Core/Private/HoloLens/HoloLensPlatformAtomics.cpp
@@ -0,0 +1,21 @@
+// Copyright 1998-2019 Epic Games, Inc. All Rights Reserved.
+
+#include "HoloLensAtomics.h"
+#include "CoreMinimal.h"
+
+#if PLATFORM_HOLOLENS
+
+void FHoloLensAtomics::HandleAtomicsFailure( const TCHAR* InFormat, ... )
+{	
+	TCHAR TempStr[1024];
+	va_list Ptr;
+
+	va_start( Ptr, InFormat );	
+	FPlatformString::GetVarArgs( TempStr, ARRAY_COUNT(TempStr), InFormat, Ptr );
+	va_end( Ptr );
+
+	UE_LOG(LogTemp, Log, TempStr);
+	check( 0 );
+}
+
+#endif //PLATFORM_HOLOLENS
\ No newline at end of file
