HoloLens platform file renaming fixes (file, atomics)

diff --git a/Engine/Source/Runtime/Core/Private/HoloLens/HoloLensFile.cpp b/Engine/Source/Runtime/Core/Private/HoloLens/HoloLensFile.cpp
deleted file mode 100644
diff --git a/Engine/Source/Runtime/Core/Private/HoloLens/HoloLensPlatformAtomics.cpp b/Engine/Source/Runtime/Core/Private/HoloLens/HoloLensPlatformAtomics.cpp
--- a/Engine/Source/Runtime/Core/Private/HoloLens/HoloLensPlatformAtomics.cpp
+++ b/Engine/Source/Runtime/Core/Private/HoloLens/HoloLensPlatformAtomics.cpp
@@ -1,6 +1,6 @@
 // Copyright 1998-2019 Epic Games, Inc. All Rights Reserved.
 
-#include "HoloLensAtomics.h"
+#include "HoloLensPlatformAtomics.h"
 #include "CoreMinimal.h"
 
 #if PLATFORM_HOLOLENS
diff --git a/Engine/Source/Runtime/Core/Private/HoloLens/HoloLensPlatformFile.cpp b/Engine/Source/Runtime/Core/Private/HoloLens/HoloLensPlatformFile.cpp
new file mode 100644
--- /dev/null
+++ b/Engine/Source/Runtime/Core/Private/HoloLens/HoloLensPlatformFile.cpp
@@ -0,0 +1,478 @@
+// Copyright 1998-2019 Epic Games, Inc. All Rights Reserved.
+
+#include "HoloLensPlatformFile.h"
+#include "CoreMinimal.h"
+#include "GenericPlatform/GenericPlatformFile.h"
+#include "HoloLens/HoloLensProcess.h"
+#include "Misc/Paths.h"
+#include <sys/utime.h>
+
+
+// make an FTimeSpan object that represents the "epoch" for time_t (from a _stat struct)
+const FDateTime HoloLensEpoch(1970, 1, 1);
+
+#include "HoloLens/AllowWindowsPlatformTypes.h"
+namespace FileConstants
+{
+	uint32 WIN_INVALID_SET_FILE_POINTER = INVALID_SET_FILE_POINTER;
+}
+#include "HoloLens/HideWindowsPlatformTypes.h"
+
+namespace
+{
+FORCEINLINE int32 UEDayOfWeekToWindowsSystemTimeDayOfWeek(const EDayOfWeek InDayOfWeek)
+{
+	switch (InDayOfWeek)
+	{
+	case EDayOfWeek::Monday:
+		return 1;
+	case EDayOfWeek::Tuesday:
+		return 2;
+	case EDayOfWeek::Wednesday:
+		return 3;
+	case EDayOfWeek::Thursday:
+		return 4;
+	case EDayOfWeek::Friday:
+		return 5;
+	case EDayOfWeek::Saturday:
+		return 6;
+	case EDayOfWeek::Sunday:
+		return 0;
+	default:
+		break;
+	}
+
+	return 0;
+}
+
+FORCEINLINE FDateTime WindowsFileTimeToUEDateTime(const FILETIME& InFileTime)
+{
+	// This roundabout conversion clamps the precision of the returned time value to match that of time_t (1 second precision)
+	// This avoids issues when sending files over the network via cook-on-the-fly
+	SYSTEMTIME SysTime;
+	if (FileTimeToSystemTime(&InFileTime, &SysTime))
+	{
+		return FDateTime(SysTime.wYear, SysTime.wMonth, SysTime.wDay, SysTime.wHour, SysTime.wMinute, SysTime.wSecond);
+	}
+
+	// Failed to convert
+	return FDateTime::MinValue();
+}
+
+FORCEINLINE FILETIME UEDateTimeToWindowsFileTime(const FDateTime& InDateTime)
+{
+	// This roundabout conversion clamps the precision of the returned time value to match that of time_t (1 second precision)
+	// This avoids issues when sending files over the network via cook-on-the-fly
+	SYSTEMTIME SysTime;
+	SysTime.wYear = InDateTime.GetYear();
+	SysTime.wMonth = InDateTime.GetMonth();
+	SysTime.wDay = InDateTime.GetDay();
+	SysTime.wDayOfWeek = UEDayOfWeekToWindowsSystemTimeDayOfWeek(InDateTime.GetDayOfWeek());
+	SysTime.wHour = InDateTime.GetHour();
+	SysTime.wMinute = InDateTime.GetMinute();
+	SysTime.wSecond = InDateTime.GetSecond();
+	SysTime.wMilliseconds = 0;
+
+	FILETIME FileTime;
+	SystemTimeToFileTime(&SysTime, &FileTime);
+
+	return FileTime;
+}
+}
+
+
+/**
+* HoloLens file handle implementation
+**/
+class CORE_API FFileHandleHoloLens : public IFileHandle
+{
+	enum { READWRITE_SIZE = 1024 * 1024 };
+	HANDLE FileHandle;
+
+	FORCEINLINE int64 FileSeek(int64 Distance, uint32 MoveMethod)
+	{
+		LARGE_INTEGER li, np = { 0 };
+		li.QuadPart = Distance;
+		if (SetFilePointerEx(FileHandle, li, &np, MoveMethod))
+		{
+			return np.QuadPart;
+		}
+		return -1;
+	}
+	FORCEINLINE bool IsValid()
+	{
+		return FileHandle != NULL && FileHandle != INVALID_HANDLE_VALUE;
+	}
+
+public:
+	FFileHandleHoloLens(HANDLE InFileHandle = NULL)
+		: FileHandle(InFileHandle)
+	{
+	}
+	virtual ~FFileHandleHoloLens()
+	{
+		CloseHandle(FileHandle);
+		FileHandle = NULL;
+	}
+	virtual int64 Tell() override
+	{
+		check(IsValid());
+		return FileSeek(0, FILE_CURRENT);
+	}
+	virtual bool Seek(int64 NewPosition) override
+	{
+		check(IsValid());
+		check(NewPosition >= 0);
+		return FileSeek(NewPosition, FILE_BEGIN) != -1;
+	}
+	virtual bool SeekFromEnd(int64 NewPositionRelativeToEnd = 0) override
+	{
+		check(IsValid());
+		check(NewPositionRelativeToEnd <= 0);
+		return FileSeek(NewPositionRelativeToEnd, FILE_END) != -1;
+	}
+	virtual bool Read(uint8* Destination, int64 BytesToRead) override
+	{
+		check(IsValid());
+		while (BytesToRead)
+		{
+			check(BytesToRead >= 0);
+			int64 ThisSize = FMath::Min<int64>(READWRITE_SIZE, BytesToRead);
+			check(Destination);
+			uint32 Result = 0;
+			if (!ReadFile(FileHandle, Destination, uint32(ThisSize), (::DWORD *)&Result, NULL) || Result != uint32(ThisSize))
+			{
+				return false;
+			}
+			Destination += ThisSize;
+			BytesToRead -= ThisSize;
+		}
+		return true;
+	}
+	virtual bool Write(const uint8* Source, int64 BytesToWrite) override
+	{
+		check(IsValid());
+		while (BytesToWrite)
+		{
+			check(BytesToWrite >= 0);
+			int64 ThisSize = FMath::Min<int64>(READWRITE_SIZE, BytesToWrite);
+			check(Source);
+			uint32 Result = 0;
+			if (!WriteFile(FileHandle, Source, uint32(ThisSize), (::DWORD *)&Result, NULL) || Result != uint32(ThisSize))
+			{
+				return false;
+			}
+			Source += ThisSize;
+			BytesToWrite -= ThisSize;
+		}
+		return true;
+	}
+
+	virtual bool Flush(const bool bFullFlush = false) override
+	{
+		check(IsValid());
+		return FlushFileBuffers(FileHandle) != 0;
+	}
+
+	virtual bool Truncate(int64 NewSize) override
+	{
+		check(IsValid());
+		return Seek(NewSize) && SetEndOfFile(FileHandle) != 0;
+	}
+};
+
+/**
+* HoloLens File I/O implementation
+**/
+class CORE_API FHoloLensPlatformFile : public IPhysicalPlatformFile
+{
+protected:
+	virtual FString NormalizeFilename(const TCHAR* Filename)
+	{
+		FString Result(Filename);
+
+#if !UE_BUILD_SHIPPING
+		if (Result.StartsWith(FHoloLensProcess::GetLocalAppDataRedirectPath()))
+		{
+			Result = Result.Replace(FHoloLensProcess::GetLocalAppDataRedirectPath(), FHoloLensProcess::GetLocalAppDataLowLevelPath());
+		}
+		if (Result.StartsWith(FHoloLensProcess::GetTempAppDataRedirectPath()))
+		{
+			Result = Result.Replace(FHoloLensProcess::GetTempAppDataRedirectPath(), FHoloLensProcess::GetTempAppDataLowLevelPath());
+		}
+#endif
+
+		FPaths::NormalizeFilename(Result);
+		if (Result.StartsWith(TEXT("//")))
+		{
+			Result = FString(TEXT("\\\\")) + Result.RightChop(2);
+		}
+
+		return FPaths::ConvertRelativePathToFull(Result);
+	}
+	virtual FString NormalizeDirectory(const TCHAR* Directory)
+	{
+		FString Result(Directory);
+
+#if !UE_BUILD_SHIPPING
+		if (Result.StartsWith(FHoloLensProcess::GetLocalAppDataRedirectPath()))
+		{
+			Result = Result.Replace(FHoloLensProcess::GetLocalAppDataRedirectPath(), FHoloLensProcess::GetLocalAppDataLowLevelPath());
+		}
+		if (Result.StartsWith(FHoloLensProcess::GetTempAppDataRedirectPath()))
+		{
+			Result = Result.Replace(FHoloLensProcess::GetTempAppDataRedirectPath(), FHoloLensProcess::GetTempAppDataLowLevelPath());
+		}
+#endif
+
+		FPaths::NormalizeDirectoryName(Result);
+		if (Result.StartsWith(TEXT("//")))
+		{
+			Result = FString(TEXT("\\\\")) + Result.RightChop(2);
+		}
+		return FPaths::ConvertRelativePathToFull(Result);
+	}
+public:
+	virtual bool FileExists(const TCHAR* Filename) override
+	{
+		WIN32_FILE_ATTRIBUTE_DATA Info;
+		if (GetFileAttributesExW(*NormalizeFilename(Filename), GET_FILEEX_INFO_LEVELS::GetFileExInfoStandard, &Info) &&
+			!(Info.dwFileAttributes & FILE_ATTRIBUTE_DIRECTORY))
+		{
+			return true;
+		}
+		return false;
+	}
+	virtual int64 FileSize(const TCHAR* Filename) override
+	{
+		WIN32_FILE_ATTRIBUTE_DATA Info;
+		if (!!GetFileAttributesExW(*NormalizeFilename(Filename), GET_FILEEX_INFO_LEVELS::GetFileExInfoStandard, &Info))
+		{
+			if ((Info.dwFileAttributes & FILE_ATTRIBUTE_DIRECTORY) == 0)
+			{
+				LARGE_INTEGER li;
+				li.HighPart = Info.nFileSizeHigh;
+				li.LowPart = Info.nFileSizeLow;
+				return li.QuadPart;
+			}
+		}
+		return -1;
+	}
+	virtual bool DeleteFile(const TCHAR* Filename) override
+	{
+		const FString NormalizedFilename = NormalizeFilename(Filename);
+		return !!DeleteFileW(*NormalizedFilename);
+	}
+	virtual bool IsReadOnly(const TCHAR* Filename) override
+	{
+		WIN32_FILE_ATTRIBUTE_DATA Info;
+		if(GetFileAttributesExW(*NormalizeFilename(Filename), GET_FILEEX_INFO_LEVELS::GetFileExInfoStandard, &Info))
+		{
+			return !!(Info.dwFileAttributes & FILE_ATTRIBUTE_READONLY);
+		}
+		return false;
+	}
+	virtual bool MoveFile(const TCHAR* To, const TCHAR* From) override
+	{
+		return !!MoveFileExW(*NormalizeFilename(From), *NormalizeFilename(To), 0ul);
+	}
+	virtual bool SetReadOnly(const TCHAR* Filename, bool bNewReadOnlyValue) override
+	{
+		return !!SetFileAttributesW(*NormalizeFilename(Filename), bNewReadOnlyValue ? FILE_ATTRIBUTE_READONLY : FILE_ATTRIBUTE_NORMAL);
+	}
+
+	virtual FDateTime GetTimeStamp(const TCHAR* Filename) override
+	{
+		// get file times
+		struct _stati64 FileInfo;
+		if (_wstati64(*NormalizeFilename(Filename), &FileInfo))
+		{
+			return FDateTime::MinValue();
+		}
+
+		// convert _stat time to FDateTime
+		FTimespan TimeSinceEpoch(0, 0, FileInfo.st_mtime);
+		return HoloLensEpoch + TimeSinceEpoch;
+	}
+
+	virtual void SetTimeStamp(const TCHAR* Filename, FDateTime DateTime) override
+	{
+		// get file times
+		struct _stati64 FileInfo;
+		if (_wstati64(*NormalizeFilename(Filename), &FileInfo))
+		{
+			return;
+		}
+
+		// change the modification time only
+		struct _utimbuf Times;
+		Times.actime = FileInfo.st_atime;
+		Times.modtime = (DateTime - HoloLensEpoch).GetTotalSeconds();
+		_wutime(Filename, &Times);
+	}
+
+	virtual FDateTime GetAccessTimeStamp(const TCHAR* Filename) override
+	{
+		// get file times
+		struct _stati64 FileInfo;
+		if (_wstati64(*NormalizeFilename(Filename), &FileInfo))
+		{
+			return FDateTime::MinValue();
+		}
+
+		// convert _stat time to FDateTime
+		FTimespan TimeSinceEpoch(0, 0, FileInfo.st_atime);
+		return HoloLensEpoch + TimeSinceEpoch;
+	}
+
+	virtual FString GetFilenameOnDisk(const TCHAR* Filename) override
+	{
+		return NormalizeFilename(Filename);
+	}
+
+	virtual IFileHandle* OpenRead(const TCHAR* Filename, bool bAllowWrite = false) override
+	{
+		uint32  Access = GENERIC_READ | (bAllowWrite ? GENERIC_WRITE : 0);
+		uint32  WinFlags = FILE_SHARE_READ | (bAllowWrite ? FILE_SHARE_WRITE : 0);
+		uint32  Create = OPEN_EXISTING;
+		FString normalizedFileName = NormalizeFilename(Filename);
+		FPlatformMisc::LowLevelOutputDebugStringf(TEXT("HoloLensFile::OpenRead normalized [%s] into [%s]\n"), Filename, *normalizedFileName);
+		HANDLE Handle = CreateFile2(*normalizedFileName, Access, WinFlags, Create, nullptr);
+		if (Handle != INVALID_HANDLE_VALUE)
+		{
+			return new FFileHandleHoloLens(Handle);
+		}
+		return NULL;
+	}
+	virtual IFileHandle* OpenWrite(const TCHAR* Filename, bool bAppend = false, bool bAllowRead = false) override
+	{
+		uint32  Access = GENERIC_WRITE;
+		uint32  WinFlags = bAllowRead ? FILE_SHARE_READ : 0;
+		uint32  Create = bAppend ? OPEN_ALWAYS : CREATE_ALWAYS;
+		HANDLE Handle = CreateFile2(*NormalizeFilename(Filename), Access, WinFlags, Create, nullptr);
+		if (Handle != INVALID_HANDLE_VALUE)
+		{
+			return new FFileHandleHoloLens(Handle);
+		}
+		return NULL;
+	}
+
+	virtual bool DirectoryExists(const TCHAR* Directory) override
+	{
+		// Empty Directory is the current directory so assume it always exists.
+		bool bExists = !FCString::Strlen(Directory);
+		if (!bExists)
+		{
+			WIN32_FILE_ATTRIBUTE_DATA Info;
+			bExists = GetFileAttributesEx(*NormalizeDirectory(Directory), GET_FILEEX_INFO_LEVELS::GetFileExInfoStandard, &Info) &&
+				(Info.dwFileAttributes & FILE_ATTRIBUTE_DIRECTORY);
+		}
+		return bExists;
+	}
+	virtual bool CreateDirectory(const TCHAR* Directory) override
+	{
+		return CreateDirectoryW(*NormalizeDirectory(Directory), NULL) || GetLastError() == ERROR_ALREADY_EXISTS;
+	}
+	virtual bool DeleteDirectory(const TCHAR* Directory) override
+	{
+		RemoveDirectoryW(*NormalizeDirectory(Directory));
+		return !DirectoryExists(Directory);
+	}
+	virtual FFileStatData GetStatData(const TCHAR* FilenameOrDirectory) override
+	{
+		WIN32_FILE_ATTRIBUTE_DATA Info;
+		if (GetFileAttributesExW(*NormalizeFilename(FilenameOrDirectory), GetFileExInfoStandard, &Info))
+		{
+			const bool bIsDirectory = !!(Info.dwFileAttributes & FILE_ATTRIBUTE_DIRECTORY);
+
+			int64 FileSize = -1;
+			if (!bIsDirectory)
+			{
+				LARGE_INTEGER li;
+				li.HighPart = Info.nFileSizeHigh;
+				li.LowPart = Info.nFileSizeLow;
+				FileSize = static_cast<int64>(li.QuadPart);
+			}
+
+			return FFileStatData(
+				WindowsFileTimeToUEDateTime(Info.ftCreationTime),
+				WindowsFileTimeToUEDateTime(Info.ftLastAccessTime),
+				WindowsFileTimeToUEDateTime(Info.ftLastWriteTime),
+				FileSize,
+				bIsDirectory,
+				!!(Info.dwFileAttributes & FILE_ATTRIBUTE_READONLY)
+				);
+		}
+
+		return FFileStatData();
+	}
+	virtual bool IterateDirectory(const TCHAR* Directory, FDirectoryVisitor& Visitor) override
+	{
+		const FString DirectoryStr = Directory;
+		return IterateDirectoryCommon(Directory, [&](const WIN32_FIND_DATAW& InData) -> bool
+		{
+			const bool bIsDirectory = !!(InData.dwFileAttributes & FILE_ATTRIBUTE_DIRECTORY);
+			return Visitor.Visit(*(DirectoryStr / InData.cFileName), bIsDirectory);
+		});
+	}
+	virtual bool IterateDirectoryStat(const TCHAR* Directory, FDirectoryStatVisitor& Visitor) override
+	{
+		const FString DirectoryStr = Directory;
+		return IterateDirectoryCommon(Directory, [&](const WIN32_FIND_DATAW& InData) -> bool
+		{
+			const bool bIsDirectory = !!(InData.dwFileAttributes & FILE_ATTRIBUTE_DIRECTORY);
+
+			int64 FileSize = -1;
+			if (!bIsDirectory)
+			{
+				LARGE_INTEGER li;
+				li.HighPart = InData.nFileSizeHigh;
+				li.LowPart = InData.nFileSizeLow;
+				FileSize = static_cast<int64>(li.QuadPart);
+			}
+
+			return Visitor.Visit(
+				*(DirectoryStr / InData.cFileName),
+				FFileStatData(
+					WindowsFileTimeToUEDateTime(InData.ftCreationTime),
+					WindowsFileTimeToUEDateTime(InData.ftLastAccessTime),
+					WindowsFileTimeToUEDateTime(InData.ftLastWriteTime),
+					FileSize,
+					bIsDirectory,
+					!!(InData.dwFileAttributes & FILE_ATTRIBUTE_READONLY)
+					)
+				);
+		});
+	}
+	bool IterateDirectoryCommon(const TCHAR* Directory, const TFunctionRef<bool(const WIN32_FIND_DATAW&)>& Visitor)
+	{
+		bool Result = false;
+		WIN32_FIND_DATAW Data;
+		HANDLE Handle = FindFirstFileExW(*(NormalizeDirectory(Directory) / TEXT("*.*")), FINDEX_INFO_LEVELS::FindExInfoStandard, &Data, FINDEX_SEARCH_OPS::FindExSearchNameMatch, nullptr, 0);
+		if (Handle != INVALID_HANDLE_VALUE)
+		{
+			Result = true;
+			do
+			{
+				if (FCString::Strcmp(Data.cFileName, TEXT(".")) && FCString::Strcmp(Data.cFileName, TEXT("..")))
+				{
+					Result = Visitor(Data);
+				}
+			} while (Result && FindNextFileW(Handle, &Data));
+			FindClose(Handle);
+		}
+		return Result;
+	}
+
+	virtual bool CreateDirectoryTree(const TCHAR* Directory) override
+	{
+		return IPlatformFile::CreateDirectoryTree(*NormalizeDirectory(Directory));
+	}
+
+};
+
+IPlatformFile& IPlatformFile::GetPlatformPhysical()
+{
+	static FHoloLensPlatformFile Singleton;
+	return Singleton;
+}
diff --git a/Engine/Source/Runtime/Core/Public/HoloLens/HoloLensIncludes.h b/Engine/Source/Runtime/Core/Public/HoloLens/HoloLensIncludes.h
--- a/Engine/Source/Runtime/Core/Public/HoloLens/HoloLensIncludes.h
+++ b/Engine/Source/Runtime/Core/Public/HoloLens/HoloLensIncludes.h
@@ -27,6 +27,4 @@ typedef FGenericPlatformRHIFramePacer FPlatformRHIFramePacer;
 
 typedef FGenericPlatformAffinity FPlatformAffinity;
 
-// include platform properties and typedef it for the runtime
 #include "HoloLensProperties.h"
-typedef FHoloLensPlatformProperties FPlatformProperties;
