HoloLensPlatformProcess include fix

diff --git a/Engine/Source/Runtime/Core/Private/HoloLens/HoloLensPlatformFile.cpp b/Engine/Source/Runtime/Core/Private/HoloLens/HoloLensPlatformFile.cpp
--- a/Engine/Source/Runtime/Core/Private/HoloLens/HoloLensPlatformFile.cpp
+++ b/Engine/Source/Runtime/Core/Private/HoloLens/HoloLensPlatformFile.cpp
@@ -3,7 +3,7 @@
 #include "HoloLensPlatformFile.h"
 #include "CoreMinimal.h"
 #include "GenericPlatform/GenericPlatformFile.h"
-#include "HoloLens/HoloLensProcess.h"
+#include "HoloLens/HoloLensPlatformProcess.h"
 #include "Misc/Paths.h"
 #include <sys/utime.h>
 
