HoloLens FPlatformVirtualMemoryBlock implementation added.

diff --git a/Engine/Source/Runtime/Core/Public/HoloLens/HoloLensPlatformMemory.h b/Engine/Source/Runtime/Core/Public/HoloLens/HoloLensPlatformMemory.h
--- a/Engine/Source/Runtime/Core/Public/HoloLens/HoloLensPlatformMemory.h
+++ b/Engine/Source/Runtime/Core/Public/HoloLens/HoloLensPlatformMemory.h
@@ -77,6 +77,56 @@ struct CORE_API FHoloLensPlatformMemory : public FGenericPlatformMemory
 	static const FPlatformMemoryConstants& GetConstants();
 	static void* BinnedAllocFromOS(SIZE_T Size);
 	static void BinnedFreeToOS(void* Ptr, SIZE_T Size);
+
+	class FPlatformVirtualMemoryBlock : public FBasicVirtualMemoryBlock
+	{
+	public:
+
+		FPlatformVirtualMemoryBlock()
+		{
+		}
+
+		FPlatformVirtualMemoryBlock(void *InPtr, uint32 InVMSizeDivVirtualSizeAlignment)
+			: FBasicVirtualMemoryBlock(InPtr, InVMSizeDivVirtualSizeAlignment)
+		{
+		}
+		FPlatformVirtualMemoryBlock(const FPlatformVirtualMemoryBlock& Other) = default;
+		FPlatformVirtualMemoryBlock& operator=(const FPlatformVirtualMemoryBlock& Other) = default;
+
+		void Commit(size_t InOffset, size_t InSize);
+		void Decommit(size_t InOffset, size_t InSize);
+		void FreeVirtual();
+
+		FORCEINLINE void CommitByPtr(void *InPtr, size_t InSize)
+		{
+			Commit(size_t(((uint8*)InPtr) - ((uint8*)Ptr)), InSize);
+		}
+
+		FORCEINLINE void DecommitByPtr(void *InPtr, size_t InSize)
+		{
+			Decommit(size_t(((uint8*)InPtr) - ((uint8*)Ptr)), InSize);
+		}
+
+		FORCEINLINE void Commit()
+		{
+			Commit(0, GetActualSize());
+		}
+
+		FORCEINLINE void Decommit()
+		{
+			Decommit(0, GetActualSize());
+		}
+
+		FORCEINLINE size_t GetActualSize() const
+		{
+			return VMSizeDivVirtualSizeAlignment * GetVirtualSizeAlignment();
+		}
+
+		static FPlatformVirtualMemoryBlock AllocateVirtual(size_t Size, size_t InAlignment = FPlatformVirtualMemoryBlock::GetVirtualSizeAlignment());
+		static size_t GetCommitAlignment();
+		static size_t GetVirtualSizeAlignment();
+	};
+
 	static FSharedMemoryRegion* MapNamedSharedMemoryRegion(const FString& InName, bool bCreate, uint32 AccessMode, SIZE_T Size);
 	static bool UnmapNamedSharedMemoryRegion(FSharedMemoryRegion * MemoryRegion);
 protected:
