Rename HoloLensMemory HoloLensPlatformMemory and implement FPlatformVirtualMemoryBlock for it
-copied the windows platform implementation

diff --git a/Engine/Source/Runtime/Core/Private/HoloLens/HoloLensMemory.cpp b/Engine/Source/Runtime/Core/Private/HoloLens/HoloLensMemory.cpp
deleted file mode 100644
diff --git a/Engine/Source/Runtime/Core/Private/HoloLens/HoloLensPlatformMemory.cpp b/Engine/Source/Runtime/Core/Private/HoloLens/HoloLensPlatformMemory.cpp
new file mode 100644
--- /dev/null
+++ b/Engine/Source/Runtime/Core/Private/HoloLens/HoloLensPlatformMemory.cpp
@@ -0,0 +1,359 @@
+// Copyright 1998-2019 Epic Games, Inc. All Rights Reserved.
+
+
+#include "HoloLensPlatformMemory.h"
+#include "HAL/MallocTBB.h"
+#include "HAL/MallocAnsi.h"
+#include "GenericPlatform/GenericPlatformMemoryPoolStats.h"
+#include "HAL/MemoryMisc.h"
+#include "CoreGlobals.h"
+#include "Misc/OutputDeviceRedirector.h"
+
+
+#if !FORCE_ANSI_ALLOCATOR
+#include "HAL/MallocBinned.h"
+#endif
+
+#include "HoloLens/AllowWindowsPlatformTypes.h"
+#define PSAPI_VERSION 2
+#include <Psapi.h>
+#pragma comment(lib, "psapi.lib")
+
+DECLARE_MEMORY_STAT(TEXT("HoloLens Specific Memory Stat"), STAT_HoloLensSpecificMemoryStat, STATGROUP_MemoryPlatform);
+
+/** Enable this to track down windows allocations not wrapped by our wrappers
+int WindowsAllocHook(int nAllocType, void *pvData,
+size_t nSize, int nBlockUse, long lRequest,
+const unsigned char * szFileName, int nLine )
+{
+if ((nAllocType == _HOOK_ALLOC || nAllocType == _HOOK_REALLOC) && nSize > 2048)
+{
+static int i = 0;
+i++;
+}
+return true;
+}
+*/
+
+#include "GenericPlatform/GenericPlatformMemoryPoolStats.h"
+
+
+void FHoloLensPlatformMemory::Init()
+{
+	FGenericPlatformMemory::Init();
+
+#if PLATFORM_32BITS
+	const int64 GB(1024 * 1024 * 1024);
+	SET_MEMORY_STAT(MCR_Physical, 2 * GB); //2Gb of physical memory on win32
+#endif
+
+
+	const FPlatformMemoryConstants& MemoryConstants = FPlatformMemory::GetConstants();
+	UE_LOG(LogMemory, Log, TEXT("Memory total: Physical=%.1fGB (%dGB approx) Virtual=%.1fGB"),
+		float(MemoryConstants.TotalPhysical / 1024.0 / 1024.0 / 1024.0),
+		MemoryConstants.TotalPhysicalGB,
+		float(MemoryConstants.TotalVirtual / 1024.0 / 1024.0 / 1024.0));
+
+	DumpStats(*GLog);
+}
+
+FMalloc* FHoloLensPlatformMemory::BaseAllocator()
+{
+#if FORCE_ANSI_ALLOCATOR
+	return new FMallocAnsi();
+#elif (WITH_EDITORONLY_DATA || IS_PROGRAM) && TBB_ALLOCATOR_ALLOWED
+	return new FMallocTBB();
+#else
+	return new FMallocBinned((uint32)(GetConstants().PageSize&MAX_uint32), (uint64)MAX_uint32 + 1);
+#endif
+
+	//	_CrtSetAllocHook(HoloLensAllocHook); // Enable to track down windows allocs not handled by our wrapper
+}
+
+FPlatformMemoryStats FHoloLensPlatformMemory::GetStats()
+{
+	/**
+	*	GetProcessMemoryInfo
+	*	PROCESS_MEMORY_COUNTERS
+	*		WorkingSetSize
+	*		UsedVirtual
+	*		PeakUsedVirtual
+	*
+	*   GetProcessInformation
+	*   APP_MEMORY_INFORMATION
+	*      PrivateCommitUsage
+	*      PeakPrivateCommitUsage
+	*      TotalCommitUsage
+	*      AvailableCommit - remaining memory available for allocation, akin to MEMORYSTATUSEX::ullAvailPhys for desktop titles without memory limits
+	*
+	*	GetSystemInfo 
+	*		SYSTEM_INFO
+	*/
+
+	FPlatformMemoryStats MemoryStats;
+
+	// Gather platform memory stats.
+	APP_MEMORY_INFORMATION AppMemoryInfo = { 0 };
+	::GetProcessInformation(GetCurrentProcess(), ProcessAppMemoryInfo, &AppMemoryInfo, sizeof(AppMemoryInfo));
+	MemoryStats.AvailablePhysical = AppMemoryInfo.AvailableCommit;
+
+	// ATG - Simplified since 32bit 4GB tuned HoloLenss are unlikely to exist
+#if _WIN64
+	MemoryStats.AvailableVirtual = (128ull * 1024 * 1024 * 1024 * 1024) - AppMemoryInfo.TotalCommitUsage;   // 64bit Win8+ 128TB limit, minus currently commited bytes
+#else
+	MemoryStats.AvailableVirtual = (2ull * 1024 * 1024 * 1024) - AppMemoryInfo.TotalCommitUsage;   // 32bit 2GB limit, minus currently commited bytes
+#endif
+
+	// ATG - GetProcessMemoryInfo did not make the cut for app API-set inclusion, removing for now
+	//PROCESS_MEMORY_COUNTERS ProcessMemoryCounters = { 0 };
+	//::GetProcessMemoryInfo(::GetCurrentProcess(), &ProcessMemoryCounters, sizeof(ProcessMemoryCounters));
+	
+	//MemoryStats.UsedPhysical = ProcessMemoryCounters.WorkingSetSize;
+	//MemoryStats.PeakUsedPhysical = ProcessMemoryCounters.PeakWorkingSetSize;
+	//MemoryStats.UsedVirtual = ProcessMemoryCounters.PagefileUsage;
+	//MemoryStats.PeakUsedVirtual = ProcessMemoryCounters.PeakPagefileUsage;
+	MemoryStats.UsedPhysical = AppMemoryInfo.PrivateCommitUsage; // TotalCommitUsage would be more correct, but we have no PeakTotalCommitUsage to go with it
+	MemoryStats.PeakUsedPhysical = AppMemoryInfo.PeakPrivateCommitUsage;
+
+	return MemoryStats;
+}
+
+void FHoloLensPlatformMemory::GetStatsForMallocProfiler(FGenericMemoryStats& out_Stats)
+{
+#if	STATS
+	FGenericPlatformMemory::GetStatsForMallocProfiler(out_Stats);
+
+	FPlatformMemoryStats Stats = GetStats();
+
+	// HoloLens specific stats.
+	out_Stats.Add(GET_STATDESCRIPTION(STAT_HoloLensSpecificMemoryStat), Stats.HoloLensSpecificMemoryStat);
+#endif // STATS
+}
+
+const FPlatformMemoryConstants& FHoloLensPlatformMemory::GetConstants()
+{
+	static FPlatformMemoryConstants MemoryConstants;
+
+	if (MemoryConstants.TotalPhysical == 0)
+	{
+		// Gather platform memory constants.
+		APP_MEMORY_INFORMATION AppMemoryInfo = { 0 };
+		::GetProcessInformation(GetCurrentProcess(), ProcessAppMemoryInfo, &AppMemoryInfo, sizeof(AppMemoryInfo));
+
+		SYSTEM_INFO SystemInformation = { 0 };
+		::GetSystemInfo(&SystemInformation);
+
+		MemoryConstants.TotalPhysical = AppMemoryInfo.TotalCommitUsage + AppMemoryInfo.AvailableCommit;
+		// ATG - Simplified since 32bit 4GB tuned HoloLenss are unlikely to exist
+#if _WIN64
+		MemoryConstants.TotalVirtual = (128ull * 1024 * 1024 * 1024 * 1024);   // 64bit Win8+ 128TB limit
+#else
+		MemoryConstants.TotalVirtual = (2ull * 1024 * 1024 * 1024);   // 32bit 2GB limit
+#endif
+		MemoryConstants.PageSize = SystemInformation.dwPageSize;
+
+		MemoryConstants.TotalPhysicalGB = (MemoryConstants.TotalPhysical + 1024 * 1024 * 1024 - 1) / 1024 / 1024 / 1024;
+	}
+
+	return MemoryConstants;
+}
+
+void FHoloLensPlatformMemory::InternalUpdateStats(const FPlatformMemoryStats& MemoryStats)
+{
+	// HoloLens specific stats.
+	SET_MEMORY_STAT(STAT_HoloLensSpecificMemoryStat, MemoryStats.HoloLensSpecificMemoryStat);
+}
+
+void* FHoloLensPlatformMemory::BinnedAllocFromOS(SIZE_T Size)
+{
+	return VirtualAlloc(NULL, Size, MEM_COMMIT, PAGE_READWRITE);
+}
+
+void FHoloLensPlatformMemory::BinnedFreeToOS(void* Ptr, SIZE_T Size)
+{
+	CA_SUPPRESS(6001)
+	// Windows maintains the size of allocation internally, so Size is unused
+	verify(VirtualFree(Ptr, 0, MEM_RELEASE) != 0);
+}
+
+size_t FHoloLensPlatformMemory::FPlatformVirtualMemoryBlock::GetVirtualSizeAlignment()
+{
+	static SIZE_T OsAllocationGranularity = FPlatformMemory::GetConstants().OsAllocationGranularity;
+	return OsAllocationGranularity;
+}
+
+size_t FHoloLensPlatformMemory::FPlatformVirtualMemoryBlock::GetCommitAlignment()
+{
+	static SIZE_T OSPageSize = FPlatformMemory::GetConstants().PageSize;
+	return OSPageSize;
+}
+
+FHoloLensPlatformMemory::FPlatformVirtualMemoryBlock FHoloLensPlatformMemory::FPlatformVirtualMemoryBlock::AllocateVirtual(size_t InSize, size_t InAlignment)
+{
+	FPlatformVirtualMemoryBlock Result;
+	InSize = Align(InSize, GetVirtualSizeAlignment());
+	Result.VMSizeDivVirtualSizeAlignment = InSize / GetVirtualSizeAlignment();
+
+	size_t Alignment = FMath::Max(InAlignment, GetVirtualSizeAlignment());
+	check(Alignment <= GetVirtualSizeAlignment());
+
+	bool bTopDown = Result.GetActualSize() > 100ll * 1024 * 1024; // this is hacky, but we want to allocate huge VM blocks (like for MB3) top down
+
+	Result.Ptr = VirtualAlloc(NULL, Result.GetActualSize(), MEM_RESERVE | (bTopDown ? MEM_TOP_DOWN : 0), PAGE_NOACCESS);
+
+
+	if (!LIKELY(Result.Ptr))
+	{
+		FPlatformMemory::OnOutOfMemory(Result.GetActualSize(), Alignment);
+	}
+	check(Result.Ptr && IsAligned(Result.Ptr, Alignment));
+	return Result;
+}
+
+
+
+void FHoloLensPlatformMemory::FPlatformVirtualMemoryBlock::FreeVirtual()
+{
+	if (Ptr)
+	{
+		check(GetActualSize() > 0);
+		// this is an iffy assumption, we don't know how much of this memory is really committed, we will assume none of it is
+		//LLM(FLowLevelMemTracker::Get().OnLowLevelFree(ELLMTracker::Platform, Ptr));
+
+		CA_SUPPRESS(6001)
+			// Windows maintains the size of allocation internally, so Size is unused
+			verify(VirtualFree(Ptr, 0, MEM_RELEASE) != 0);
+
+		Ptr = nullptr;
+		VMSizeDivVirtualSizeAlignment = 0;
+	}
+}
+
+void FHoloLensPlatformMemory::FPlatformVirtualMemoryBlock::Commit(size_t InOffset, size_t InSize)
+{
+	check(IsAligned(InOffset, GetCommitAlignment()) && IsAligned(InSize, GetCommitAlignment()));
+	check(InOffset >= 0 && InSize >= 0 && InOffset + InSize <= GetActualSize() && Ptr);
+
+	// There are no guarantees LLM is going to be able to deal with this
+	uint8* UsePtr = ((uint8*)Ptr) + InOffset;
+	LLM(FLowLevelMemTracker::Get().OnLowLevelAlloc(ELLMTracker::Platform, UsePtr, InSize));
+	if (VirtualAlloc(UsePtr, InSize, MEM_COMMIT, PAGE_READWRITE) != UsePtr)
+	{
+		FPlatformMemory::OnOutOfMemory(InSize, 0);
+	}
+}
+
+void FHoloLensPlatformMemory::FPlatformVirtualMemoryBlock::Decommit(size_t InOffset, size_t InSize)
+{
+	check(IsAligned(InOffset, GetCommitAlignment()) && IsAligned(InSize, GetCommitAlignment()));
+	check(InOffset >= 0 && InSize >= 0 && InOffset + InSize <= GetActualSize() && Ptr);
+	uint8* UsePtr = ((uint8*)Ptr) + InOffset;
+	// There are no guarantees LLM is going to be able to deal with this
+	LLM(FLowLevelMemTracker::Get().OnLowLevelFree(ELLMTracker::Platform, UsePtr));
+	VirtualFree(UsePtr, InSize, MEM_DECOMMIT);
+}
+
+FPlatformMemory::FSharedMemoryRegion* FHoloLensPlatformMemory::MapNamedSharedMemoryRegion(const FString& InName, bool bCreate, uint32 AccessMode, SIZE_T Size)
+{
+	FString Name(TEXT("Global\\"));
+	Name += InName;
+
+	DWORD OpenMappingAccess = FILE_MAP_READ;
+	check(AccessMode != 0);
+	if (AccessMode == FPlatformMemory::ESharedMemoryAccess::Write)
+	{
+		OpenMappingAccess = FILE_MAP_WRITE;
+	}
+	else if (AccessMode == (FPlatformMemory::ESharedMemoryAccess::Write | FPlatformMemory::ESharedMemoryAccess::Read))
+	{
+		OpenMappingAccess = FILE_MAP_ALL_ACCESS;
+	}
+
+	HANDLE Mapping = NULL;
+	if (bCreate)
+	{
+		DWORD CreateMappingAccess = PAGE_READONLY;
+		check(AccessMode != 0);
+		if (AccessMode == FPlatformMemory::ESharedMemoryAccess::Write)
+		{
+			CreateMappingAccess = PAGE_WRITECOPY;
+		}
+		else if (AccessMode == (FPlatformMemory::ESharedMemoryAccess::Write | FPlatformMemory::ESharedMemoryAccess::Read))
+		{
+			CreateMappingAccess = PAGE_READWRITE;
+		}
+
+		Mapping = CreateFileMappingFromApp(INVALID_HANDLE_VALUE, NULL, CreateMappingAccess, Size, *Name);
+
+		if (Mapping == NULL)
+		{
+			DWORD ErrNo = GetLastError();
+			UE_LOG(LogHAL, Warning, TEXT("CreateFileMappingFromApp(file=INVALID_HANDLE_VALUE, security=NULL, protect=0x%x, size%I64, name='%s') failed with GetLastError() = %d"),
+				CreateMappingAccess, Size, *Name,
+				ErrNo
+				);
+		}
+	}
+	else
+	{
+		UE_LOG(LogHAL, Warning, TEXT("OpenFileMapping not possible from a packaged process"));
+	}
+
+	if (Mapping == NULL)
+	{
+		return NULL;
+	}
+
+	void* Ptr = MapViewOfFileFromApp(Mapping, OpenMappingAccess, 0, Size);
+	if (Ptr == NULL)
+	{
+		DWORD ErrNo = GetLastError();
+		UE_LOG(LogHAL, Warning, TEXT("MapViewOfFile(mapping=0x%x, access=0x%x, OffsetHigh=0, OffsetLow=0, NumBytes=%u) failed with GetLastError() = %d"),
+			Mapping, OpenMappingAccess, Size,
+			ErrNo
+			);
+
+		CloseHandle(Mapping);
+		return NULL;
+	}
+
+	return new FHoloLensSharedMemoryRegion(Name, AccessMode, Ptr, Size, Mapping);
+}
+
+bool FHoloLensPlatformMemory::UnmapNamedSharedMemoryRegion(FSharedMemoryRegion * MemoryRegion)
+{
+	bool bAllSucceeded = true;
+
+	if (MemoryRegion)
+	{
+		FHoloLensSharedMemoryRegion * HoloLensRegion = static_cast< FHoloLensSharedMemoryRegion* >(MemoryRegion);
+
+		if (!UnmapViewOfFile(HoloLensRegion->GetAddress()))
+		{
+			bAllSucceeded = false;
+
+			int ErrNo = GetLastError();
+			UE_LOG(LogHAL, Warning, TEXT("UnmapViewOfFile(address=%p) failed with GetLastError() = %d"),
+				HoloLensRegion->GetAddress(),
+				ErrNo
+				);
+		}
+
+		if (!CloseHandle(HoloLensRegion->GetMapping()))
+		{
+			bAllSucceeded = false;
+
+			int ErrNo = GetLastError();
+			UE_LOG(LogHAL, Warning, TEXT("CloseHandle(handle=0x%x) failed with GetLastError() = %d"),
+				HoloLensRegion->GetMapping(),
+				ErrNo
+				);
+		}
+
+		// delete the region
+		delete HoloLensRegion;
+	}
+
+	return bAllSucceeded;
+}
+#include "HoloLens/HideWindowsPlatformTypes.h"
