Updated HoloLensAtomics.h to match the API changes

diff --git a/Engine/Source/Runtime/Core/Public/HoloLens/HoloLensPlatformAtomics.h b/Engine/Source/Runtime/Core/Public/HoloLens/HoloLensPlatformAtomics.h
--- a/Engine/Source/Runtime/Core/Public/HoloLens/HoloLensPlatformAtomics.h
+++ b/Engine/Source/Runtime/Core/Public/HoloLens/HoloLensPlatformAtomics.h
@@ -169,7 +169,7 @@ struct CORE_API FHoloLensAtomics : public FGenericPlatformAtomics
 #endif
 	}
 
-	static FORCEINLINE void* InterlockedExchangePtr( void** Dest, void* Exchange )
+	static FORCEINLINE void* InterlockedExchangePtr( void*volatile* Dest, void* Exchange )
 	{
 #if !(UE_BUILD_SHIPPING || UE_BUILD_TEST)
 		if (IsAligned(Dest) == false)
@@ -208,6 +208,102 @@ struct CORE_API FHoloLensAtomics : public FGenericPlatformAtomics
 			return (int64)::_InterlockedCompareExchange64(Dest, Exchange, Comparand);
 	}
 
+	static FORCEINLINE int8 InterlockedAnd(volatile int8* Value, const int8 AndValue)
+	{
+		return (int8)::_InterlockedAnd8((volatile char*)Value, (char)AndValue);
+	}
+
+	static FORCEINLINE int16 InterlockedAnd(volatile int16* Value, const int16 AndValue)
+	{
+		return (int16)::_InterlockedAnd16((volatile short*)Value, (short)AndValue);
+	}
+
+	static FORCEINLINE int32 InterlockedAnd(volatile int32* Value, const int32 AndValue)
+	{
+		return (int32)::_InterlockedAnd((volatile long*)Value, (long)AndValue);
+	}
+
+	static FORCEINLINE int64 InterlockedAnd(volatile int64* Value, const int64 AndValue)
+	{
+#if PLATFORM_64BITS
+		return (int64)::_InterlockedAnd64((volatile long long*)Value, (long long)AndValue);
+#else
+		// No explicit instruction for 64-bit atomic and on 32-bit processors; has to be implemented in terms of CMPXCHG8B
+		for (;;)
+		{
+			const int64 OldValue = *Value;
+			if (_InterlockedCompareExchange64(Value, OldValue & AndValue, OldValue) == OldValue)
+			{
+				return OldValue;
+			}
+		}
+#endif
+	}
+
+	static FORCEINLINE int8 InterlockedOr(volatile int8* Value, const int8 OrValue)
+	{
+		return (int8)::_InterlockedOr8((volatile char*)Value, (char)OrValue);
+	}
+
+	static FORCEINLINE int16 InterlockedOr(volatile int16* Value, const int16 OrValue)
+	{
+		return (int16)::_InterlockedOr16((volatile short*)Value, (short)OrValue);
+	}
+
+	static FORCEINLINE int32 InterlockedOr(volatile int32* Value, const int32 OrValue)
+	{
+		return (int32)::_InterlockedOr((volatile long*)Value, (long)OrValue);
+	}
+
+	static FORCEINLINE int64 InterlockedOr(volatile int64* Value, const int64 OrValue)
+	{
+#if PLATFORM_64BITS
+		return (int64)::_InterlockedOr64((volatile long long*)Value, (long long)OrValue);
+#else
+		// No explicit instruction for 64-bit atomic or on 32-bit processors; has to be implemented in terms of CMPXCHG8B
+		for (;;)
+		{
+			const int64 OldValue = *Value;
+			if (_InterlockedCompareExchange64(Value, OldValue | OrValue, OldValue) == OldValue)
+			{
+				return OldValue;
+			}
+		}
+#endif
+	}
+
+	static FORCEINLINE int8 InterlockedXor(volatile int8* Value, const int8 XorValue)
+	{
+		return (int8)::_InterlockedXor8((volatile char*)Value, (char)XorValue);
+	}
+
+	static FORCEINLINE int16 InterlockedXor(volatile int16* Value, const int16 XorValue)
+	{
+		return (int16)::_InterlockedXor16((volatile short*)Value, (short)XorValue);
+	}
+
+	static FORCEINLINE int32 InterlockedXor(volatile int32* Value, const int32 XorValue)
+	{
+		return (int32)::_InterlockedXor((volatile long*)Value, (int32)XorValue);
+	}
+
+	static FORCEINLINE int64 InterlockedXor(volatile int64* Value, const int64 XorValue)
+	{
+#if PLATFORM_64BITS
+		return (int64)::_InterlockedXor64((volatile long long*)Value, (long long)XorValue);
+#else
+		// No explicit instruction for 64-bit atomic xor on 32-bit processors; has to be implemented in terms of CMPXCHG8B
+		for (;;)
+		{
+			const int64 OldValue = *Value;
+			if (_InterlockedCompareExchange64(Value, OldValue ^ XorValue, OldValue) == OldValue)
+			{
+				return OldValue;
+			}
+		}
+#endif
+	}
+
 	static FORCEINLINE int8 AtomicRead(volatile const int8* Src)
 	{
 		return InterlockedCompareExchange((volatile int8*)Src, 0, 0);
