Updated HoloLensPlatformString.h to have the new format specifiers

diff --git a/Engine/Source/Runtime/Core/Public/HoloLens/HoloLensPlatformString.h b/Engine/Source/Runtime/Core/Public/HoloLens/HoloLensPlatformString.h
--- a/Engine/Source/Runtime/Core/Public/HoloLens/HoloLensPlatformString.h
+++ b/Engine/Source/Runtime/Core/Public/HoloLens/HoloLensPlatformString.h
@@ -16,3 +16,47 @@ struct FHoloLensString : public FMicrosoftPlatformString
 };
 
 typedef FHoloLensString FPlatformString;
+
+// Format specifiers to be able to print values of these types correctly, for example when using UE_LOG.
+
+// SIZE_T format specifier
+#define SIZE_T_FMT "I64u"
+// SIZE_T format specifier for lowercase hexadecimal output
+#define SIZE_T_x_FMT "I64x"
+// SIZE_T format specifier for uppercase hexadecimal output
+#define SIZE_T_X_FMT "I64X"
+
+// SSIZE_T format specifier
+#define SSIZE_T_FMT "I64d"
+// SSIZE_T format specifier for lowercase hexadecimal output
+#define SSIZE_T_x_FMT "I64x"
+// SSIZE_T format specifier for uppercase hexadecimal output
+#define SSIZE_T_X_FMT "I64X"
+
+// PTRINT format specifier for decimal output
+#define PTRINT_FMT "lld"
+// PTRINT format specifier for lowercase hexadecimal output
+#define PTRINT_x_FMT "llx"
+// PTRINT format specifier for uppercase hexadecimal output
+#define PTRINT_X_FMT "llX"
+
+// UPTRINT format specifier for decimal output
+#define UPTRINT_FMT "llu"
+// UPTRINT format specifier for lowercase hexadecimal output
+#define UPTRINT_x_FMT "llx"
+// UPTRINT format specifier for uppercase hexadecimal output
+#define UPTRINT_X_FMT "llX"
+
+// int64 format specifier for decimal output
+#define INT64_FMT "lld"
+// int64 format specifier for lowercase hexadecimal output
+#define INT64_x_FMT "llx"
+// int64 format specifier for uppercase hexadecimal output
+#define INT64_X_FMT "llX"
+
+// uint64 format specifier for decimal output
+#define UINT64_FMT "llu"
+// uint64 format specifier for lowercase hexadecimal output
+#define UINT64_x_FMT "llx"
+// uint64 format specifier for uppercase hexadecimal output
+#define UINT64_X_FMT "llX"
