DR - Delete unused cvar
#rnx

diff --git a/Engine/Source/Runtime/Core/Private/HAL/ConsoleManager.cpp b/Engine/Source/Runtime/Core/Private/HAL/ConsoleManager.cpp
--- a/Engine/Source/Runtime/Core/Private/HAL/ConsoleManager.cpp
+++ b/Engine/Source/Runtime/Core/Private/HAL/ConsoleManager.cpp
@@ -2594,12 +2594,6 @@ static TAutoConsoleVariable<int32> CVarLuminOverrideExternalTextureSupport(
 	TEXT("  4 = force ImageExternalESSL300 (version #300 with GL_OES_EGL_image_external_essl3)"),
 	ECVF_ReadOnly);
 
-static TAutoConsoleVariable<int32> GLSLCvar(
-	TEXT("r.Vulkan.UseGLSL"),
-	0,
-	TEXT("2 to use ES GLSL\n1 to use GLSL\n0 to use SPIRV")
-);
-
 static TAutoConsoleVariable<FString> CVarCustomUnsafeZones(
 	TEXT("r.CustomUnsafeZones"),
 	TEXT(""),
