Default to SF_VULKAN_SM5 for the null shader language over GLSL_150

diff --git a/Engine/Source/Runtime/Core/Private/Unix/UnixPlatformMisc.cpp b/Engine/Source/Runtime/Core/Private/Unix/UnixPlatformMisc.cpp
--- a/Engine/Source/Runtime/Core/Private/Unix/UnixPlatformMisc.cpp
+++ b/Engine/Source/Runtime/Core/Private/Unix/UnixPlatformMisc.cpp
@@ -553,7 +553,7 @@ int32 FUnixPlatformMisc::NumberOfCoresIncludingHyperthreads()
 
 const TCHAR* FUnixPlatformMisc::GetNullRHIShaderFormat()
 {
-	return TEXT("GLSL_150");
+	return TEXT("SF_VULKAN_SM5");
 }
 
 bool FUnixPlatformMisc::HasCPUIDInstruction()
