Depend on the bSendUsage setting, not about sending ensures

diff --git a/Engine/Source/Runtime/Core/Private/Mac/MacPlatformMisc.cpp b/Engine/Source/Runtime/Core/Private/Mac/MacPlatformMisc.cpp
--- a/Engine/Source/Runtime/Core/Private/Mac/MacPlatformMisc.cpp
+++ b/Engine/Source/Runtime/Core/Private/Mac/MacPlatformMisc.cpp
@@ -1888,10 +1888,17 @@ void FMacCrashContext::GenerateCrashInfoAndLaunchReporter() const
 		GConfig->GetBool(TEXT("/Script/UnrealEd.CrashReportsPrivacySettings"), TEXT("bSendUnattendedBugReports"), bSendUnattendedBugReports, GEditorSettingsIni);
 	}
 
+	bool bSendUsageData = true;
+	if (GConfig)
+	{
+		GConfig->GetBool(TEXT("/Script/UnrealEd.AnalyticsPrivacySettings"), TEXT("bSendUsageData"), bSendUsageData, GEditorSettingsIni);
+	}
+
 	if (BuildSettings::IsLicenseeVersion() && !UE_EDITOR)
 	{
 		// do not send unattended reports in licensees' builds except for the editor, where it is governed by the above setting
 		bSendUnattendedBugReports = false;
+		bSendUsageData = false;
 	}
 
 	const bool bUnattended = GMacAppInfo.bIsUnattended || IsRunningDedicatedServer();
@@ -1931,7 +1938,7 @@ void FMacCrashContext::GenerateCrashInfoAndLaunchReporter() const
 			}
 			else
 			{
-				if (bSendUnattendedBugReports)
+				if (bSendUsageData)
 				{
 					execl(GMacAppInfo.CrashReportClient, "CrashReportClient", CrashInfoFolder, NULL);
 				}
@@ -1981,10 +1988,17 @@ void FMacCrashContext::GenerateEnsureInfoAndLaunchReporter() const
 		GConfig->GetBool(TEXT("/Script/UnrealEd.CrashReportsPrivacySettings"), TEXT("bSendUnattendedBugReports"), bSendUnattendedBugReports, GEditorSettingsIni);
 	}
 
+	bool bSendUsageData = true;
+	if (GConfig)
+	{
+		GConfig->GetBool(TEXT("/Script/UnrealEd.AnalyticsPrivacySettings"), TEXT("bSendUsageData"), bSendUsageData, GEditorSettingsIni);
+	}
+
 	if (BuildSettings::IsLicenseeVersion() && !UE_EDITOR)
 	{
 		// do not send unattended reports in licensees' builds except for the editor, where it is governed by the above setting
 		bSendUnattendedBugReports = false;
+		bSendUsageData = false;
 	}
 
 	const bool bUnattended = GMacAppInfo.bIsUnattended || !IsInteractiveEnsureMode() || IsRunningDedicatedServer();
@@ -2019,7 +2033,7 @@ void FMacCrashContext::GenerateEnsureInfoAndLaunchReporter() const
 		}
 
 		// If the editor setting has been disabled to not send analytics extend this to the CRC
-		if (!bSendUnattendedBugReports)
+		if (!bSendUsageData)
 		{
 			Arguments += TEXT(" -NoAnalytics");
 		}
diff --git a/Engine/Source/Runtime/Core/Private/Unix/UnixPlatformCrashContext.cpp b/Engine/Source/Runtime/Core/Private/Unix/UnixPlatformCrashContext.cpp
--- a/Engine/Source/Runtime/Core/Private/Unix/UnixPlatformCrashContext.cpp
+++ b/Engine/Source/Runtime/Core/Private/Unix/UnixPlatformCrashContext.cpp
@@ -428,6 +428,13 @@ void FUnixCrashContext::GenerateCrashInfoAndLaunchReporter(bool bReportingNonCra
 		GConfig->GetBool(TEXT("/Script/UnrealEd.CrashReportsPrivacySettings"), TEXT("bSendUnattendedBugReports"), bSendUnattendedBugReports, GEditorSettingsIni);
 	}
 
+	// Controls if we want analytics in the crash report client
+	bool bSendUsageData = true;
+	if (GConfig)
+	{
+		GConfig->GetBool(TEXT("/Script/UnrealEd.AnalyticsPrivacySettings"), TEXT("bSendUsageData"), bSendUsageData, GEditorSettingsIni);
+	}
+
 	// If we are not an editor but still want to agree to upload for non-licensee check the settings
 	bool bAgreeToCrashUpload = false;
 	if (!UE_EDITOR && GConfig)
@@ -440,6 +447,7 @@ void FUnixCrashContext::GenerateCrashInfoAndLaunchReporter(bool bReportingNonCra
 		// do not send unattended reports in licensees' builds except for the editor, where it is governed by the above setting
 		bSendUnattendedBugReports = false;
 		bAgreeToCrashUpload = false;
+		bSendUsageData = false;
 	}
 
 	bool bSkipCRC = bUnattended && !bSendUnattendedBugReports && !bAgreeToCrashUpload;
@@ -559,7 +567,7 @@ void FUnixCrashContext::GenerateCrashInfoAndLaunchReporter(bool bReportingNonCra
 			CrashReportClientArguments += TEXT(" ");
 
 			// If the editor setting has been disabled to not send analytics extend this to the CRC
-			if (!bSendUnattendedBugReports)
+			if (!bSendUsageData)
 			{
 				CrashReportClientArguments += TEXT(" -NoAnalytics ");
 			}
diff --git a/Engine/Source/Runtime/Core/Private/Windows/WindowsPlatformCrashContext.cpp b/Engine/Source/Runtime/Core/Private/Windows/WindowsPlatformCrashContext.cpp
--- a/Engine/Source/Runtime/Core/Private/Windows/WindowsPlatformCrashContext.cpp
+++ b/Engine/Source/Runtime/Core/Private/Windows/WindowsPlatformCrashContext.cpp
@@ -320,11 +320,19 @@ int32 ReportCrashUsingCrashReportClient(FWindowsPlatformCrashContext& InContext,
 		GConfig->GetBool(TEXT("/Script/UnrealEd.CrashReportsPrivacySettings"), TEXT("bSendUnattendedBugReports"), bSendUnattendedBugReports, GEditorSettingsIni);
 	}
 
+	// Controls if we want analytics in the crash report client
+	bool bSendUsageData = true;
+	if (GConfig)
+	{
+		GConfig->GetBool(TEXT("/Script/UnrealEd.AnalyticsPrivacySettings"), TEXT("bSendUsageData"), bSendUsageData, GEditorSettingsIni);
+	}
+
 #if !UE_EDITOR
 	if (BuildSettings::IsLicenseeVersion())
 	{
 		// do not send unattended reports in licensees' builds except for the editor, where it is governed by the above setting
 		bSendUnattendedBugReports = false;
+		bSendUsageData = false;
 	}
 #endif
 
@@ -428,7 +436,7 @@ int32 ReportCrashUsingCrashReportClient(FWindowsPlatformCrashContext& InContext,
 			FString CrashReportClientArguments = FString::Printf(TEXT("\"%s\""), *CrashFolderAbsolute);
 
 			// If the editor setting has been disabled to not send analytics extend this to the CRC
-			if (!bSendUnattendedBugReports)
+			if (!bSendUsageData)
 			{
 				CrashReportClientArguments += TEXT(" -NoAnalytics ");
 			}
