Fixed incorrect merge

diff --git a/Engine/Source/Runtime/Core/Private/Unix/UnixPlatformCrashContext.cpp b/Engine/Source/Runtime/Core/Private/Unix/UnixPlatformCrashContext.cpp
--- a/Engine/Source/Runtime/Core/Private/Unix/UnixPlatformCrashContext.cpp
+++ b/Engine/Source/Runtime/Core/Private/Unix/UnixPlatformCrashContext.cpp
@@ -579,7 +579,8 @@ void FUnixCrashContext::GenerateCrashInfoAndLaunchReporter(bool bReportingNonCra
 			{
 				CrashReportClientArguments += TEXT(" -NoAnalytics ");
 			}
-			else if (bImplicitSend)
+
+			if (bImplicitSend)
 			{
 				CrashReportClientArguments += TEXT(" -Unattended -ImplicitSend ");
 			}
