Xcode 11 is more strict when calling IMPs. Quick fix to show the compiler the real function signatures.

diff --git a/Engine/Source/Runtime/Core/Private/Apple/AppleLLM.cpp b/Engine/Source/Runtime/Core/Private/Apple/AppleLLM.cpp
--- a/Engine/Source/Runtime/Core/Private/Apple/AppleLLM.cpp
+++ b/Engine/Source/Runtime/Core/Private/Apple/AppleLLM.cpp
@@ -26,8 +26,15 @@ const FLLMTagInfoApple ELLMTagNamesApple[] =
 	{ TEXT("Objective-C"),				GET_STATFNAME(STAT_ObjectiveCLLM),		NAME_None },									// ELLMTagApple::ObjectiveC
 };
 
-static IMP AllocWithZoneOriginal = nullptr;
-static IMP DeallocOriginal = nullptr;
+typedef id (*AllocWithZoneFunc)(id Obj, SEL Sel, struct _NSZone* Zone);
+typedef void (*DeallocFunc)(id Obj, SEL Sel);
+
+static AllocWithZoneFunc AllocWithZoneOriginal = nullptr;
+static DeallocFunc DeallocOriginal = nullptr;
+
+
+static IMP AllocWithZoneOriginalIMP = nullptr;
+static IMP DeallocOriginalIMP = nullptr;
 
 static id AllocWithZoneInterposer(id Obj, SEL Sel, struct _NSZone * Zone)
 {
@@ -117,11 +124,14 @@ void AppleLLM::Initialise()
 	Method AllocZone = class_getClassMethod([NSObject class], @selector(allocWithZone:));
 	Method Dealloc = class_getInstanceMethod([NSObject class], @selector(dealloc));
 	
-	AllocWithZoneOriginal = method_getImplementation(AllocZone);
-	DeallocOriginal = method_getImplementation(Dealloc);
+	AllocWithZoneOriginalIMP = method_getImplementation(AllocZone);
+	DeallocOriginalIMP = method_getImplementation(Dealloc);
 	
-	check(AllocWithZoneOriginal);
-	check(DeallocOriginal);
+	check(AllocWithZoneOriginalIMP);
+	check(DeallocOriginalIMP);
+    
+    AllocWithZoneOriginal = (AllocWithZoneFunc) AllocWithZoneOriginalIMP;
+    DeallocOriginal = (DeallocFunc) DeallocOriginalIMP;
 	
 	method_setImplementation(AllocZone, (IMP)&AllocWithZoneInterposer);
 	method_setImplementation(Dealloc, (IMP)&DeallocInterposer);
