Fixes justified strings on mac.

This is a hack, should probably revisit this entire thing if we ever have free cycles.

diff --git a/Engine/Source/Runtime/Core/Private/GenericPlatform/StandardPlatformString.cpp b/Engine/Source/Runtime/Core/Private/GenericPlatform/StandardPlatformString.cpp
--- a/Engine/Source/Runtime/Core/Private/GenericPlatform/StandardPlatformString.cpp
+++ b/Engine/Source/Runtime/Core/Private/GenericPlatform/StandardPlatformString.cpp
@@ -115,6 +115,18 @@ static int32 GetFormattingInfo(const WIDECHAR* Format, FFormatInfo& OutInfo)
 		OutInfo.Format[OutInfoFormatLength++] = 's';
 	}
 	OutInfo.Format[OutInfoFormatLength] = 0;
+	
+	// HACKHACKHACK
+	// This formatting function expects to understand %s as a string no matter which char width.
+	// On mac (and possibly others) this must be fixed up to %S for widechars.
+	// So we will do the fixup ONLY if this is a widechar system and the format is given as %s.
+	// BUG: This function still doesn't handle char16_t correctly.
+	if (sizeof(WIDECHAR) == sizeof(wchar_t) &&
+		OutInfo.Type == LITERAL(WIDECHAR, 's'))
+	{
+		checkSlow(OutInfo.Format[OutInfoFormatLength-1] == LITERAL(WIDECHAR, 's'));
+		OutInfo.Format[OutInfoFormatLength-1] = LITERAL(WIDECHAR, 'S');
+	}
 
 	return FormatLength;
 }
@@ -170,7 +182,7 @@ static const WIDECHAR* GetFormattedArgument(const FFormatInfo& Info, VA_LIST_REF
 			}
 		}
 		// Is it a plain string?
-		else if (Info.Format[1] == LITERAL(WIDECHAR, 's'))
+		else if (FChar::ToLower(Info.Format[1]) == LITERAL(WIDECHAR, 's'))
 		{
 			const WIDECHAR* String = va_arg(ArgPtr, WIDECHAR*);
 			InOutLength = FCString::Strlen(String);
@@ -179,6 +191,8 @@ static const WIDECHAR* GetFormattedArgument(const FFormatInfo& Info, VA_LIST_REF
 		// Some form of string requiring formatting, such as a left- or right-justified string
 		else
 		{
+			// We call swprintf directly which may expect %S for a widechar string. This will be fixed up
+			// by the time we get here (See above in GetFormattingInfo).
 			const WIDECHAR* String = va_arg(ArgPtr, WIDECHAR*);
 			InOutLength = swprintf(Formatted, InOutLength, Info.Format, String);
 			return Formatted;
