Added TSet and TMap XyzByHash() functions
From CL 6926257 in Dev-Core

diff --git a/Engine/Source/Runtime/Core/Public/Containers/Map.h b/Engine/Source/Runtime/Core/Public/Containers/Map.h
--- a/Engine/Source/Runtime/Core/Public/Containers/Map.h
+++ b/Engine/Source/Runtime/Core/Public/Containers/Map.h
@@ -101,7 +101,15 @@ struct TDefaultMapHashableKeyFuncs : TDefaultMapKeyFuncs<KeyType, ValueType, bIn
 
 /** 
  * The base class of maps from keys to values.  Implemented using a TSet of key-value pairs with a custom KeyFuncs, 
- * with the same O(1) addition, removal, and finding. 
+ * with the same O(1) addition, removal, and finding.
+ *
+ * The ByHash() functions are somewhat dangerous but particularly useful in two scenarios:
+ * -- Heterogeneous lookup to avoid creating expensive keys like FString when looking up by const TCHAR*.
+ *	  You must ensure the hash is calculated in the same way as ElementType is hashed.
+ *    If possible put both ComparableKey and ElementType hash functions next to each other in the same header
+ *    to avoid bugs when the ElementType hash function is changed.
+ * -- Reducing contention around hash tables protected by a lock. It is often important to incur
+ *    the cache misses of reading key data and doing the hashing *before* acquiring the lock.
  **/
 template <typename KeyType, typename ValueType, typename SetAllocator, typename KeyFuncs>
 class TMapBase
@@ -301,6 +309,12 @@ public:
 	FORCEINLINE ValueType& Add(      KeyType&& InKey, const ValueType&  InValue) { return Emplace(MoveTempIfPossible(InKey),                    InValue ); }
 	FORCEINLINE ValueType& Add(      KeyType&& InKey,       ValueType&& InValue) { return Emplace(MoveTempIfPossible(InKey), MoveTempIfPossible(InValue)); }
 
+	/** See Add() and class documentation section on ByHash() functions */
+	FORCEINLINE ValueType& AddByHash(uint32 KeyHash, const KeyType&  InKey, const ValueType&  InValue) { return EmplaceByHash(KeyHash, InKey, InValue); }
+	FORCEINLINE ValueType& AddByHash(uint32 KeyHash, const KeyType&  InKey,		  ValueType&& InValue) { return EmplaceByHash(KeyHash, InKey, MoveTempIfPossible(InValue)); }
+	FORCEINLINE ValueType& AddByHash(uint32 KeyHash,	   KeyType&& InKey, const ValueType&  InValue) { return EmplaceByHash(KeyHash, MoveTempIfPossible(InKey), InValue); }
+	FORCEINLINE ValueType& AddByHash(uint32 KeyHash,	   KeyType&& InKey,		  ValueType&& InValue) { return EmplaceByHash(KeyHash, MoveTempIfPossible(InKey), MoveTempIfPossible(InValue)); }
+
 	/**
 	 * Set a default value associated with a key.
 	 *
@@ -310,6 +324,10 @@ public:
 	FORCEINLINE ValueType& Add(const KeyType&  InKey) { return Emplace(                   InKey ); }
 	FORCEINLINE ValueType& Add(      KeyType&& InKey) { return Emplace(MoveTempIfPossible(InKey)); }
 
+	/** See Add() and class documentation section on ByHash() functions */
+	FORCEINLINE ValueType& AddByHash(uint32 KeyHash, const KeyType&  InKey) { return EmplaceByHash(KeyHash, InKey); }
+	FORCEINLINE ValueType& AddByHash(uint32 KeyHash,	   KeyType&& InKey)	{ return EmplaceByHash(KeyHash, MoveTempIfPossible(InKey)); }
+
 	/**
 	 * Set the value associated with a key.
 	 *
@@ -333,6 +351,15 @@ public:
 		return Pairs[PairId].Value;
 	}
 
+	/** See Emplace() and class documentation section on ByHash() functions */
+	template <typename InitKeyType, typename InitValueType>
+	ValueType& EmplaceByHash(uint32 KeyHash, InitKeyType&& InKey, InitValueType&& InValue)
+	{
+		const FSetElementId PairId = Pairs.EmplaceByHash(KeyHash, TPairInitializer<InitKeyType&&, InitValueType&&>(Forward<InitKeyType>(InKey), Forward<InitValueType>(InValue)));
+
+		return Pairs[PairId].Value;
+	}
+
 	/**
 	 * Set a default value associated with a key.
 	 *
@@ -347,6 +374,15 @@ public:
 		return Pairs[PairId].Value;
 	}
 
+	/** See Emplace() and class documentation section on ByHash() functions */
+	template <typename InitKeyType>
+	ValueType& EmplaceByHash(uint32 KeyHash, InitKeyType&& InKey)
+	{
+		const FSetElementId PairId = Pairs.EmplaceByHash(KeyHash, TKeyInitializer<InitKeyType&&>(Forward<InitKeyType>(InKey)));
+
+		return Pairs[PairId].Value;
+	}
+
 	/**
 	 * Remove all value associations for a key.
 	 *
@@ -359,6 +395,14 @@ public:
 		return NumRemovedPairs;
 	}
 
+	/** See Remove() and class documentation section on ByHash() functions */
+	template<typename ComparableKey>
+	FORCEINLINE int32 RemoveByHash(uint32 KeyHash, const ComparableKey& Key)
+	{
+		const int32 NumRemovedPairs = Pairs.RemoveByHash(KeyHash, Key);
+		return NumRemovedPairs;
+	}
+
 	/**
 	 * Find the key associated with the specified value.
 	 *
@@ -402,7 +446,28 @@ public:
 		return const_cast<TMapBase*>(this)->Find(Key);
 	}
 
+	/** See Find() and class documentation section on ByHash() functions */
+	template<typename ComparableKey>
+	FORCEINLINE ValueType* FindByHash(uint32 KeyHash, const ComparableKey& Key)
+	{
+		if (auto* Pair = Pairs.FindByHash(KeyHash, Key))
+		{
+			return &Pair->Value;
+		}
+
+		return nullptr;
+	}
+	template<typename ComparableKey>
+	FORCEINLINE const ValueType* FindByHash(uint32 KeyHash, const ComparableKey& Key) const
+	{
+		return const_cast<TMapBase*>(this)->FindByHash(KeyHash, Key);
+	}
+
 private:
+	FORCEINLINE static uint32 HashKey(const KeyType& Key)
+	{
+		return KeyFuncs::GetKeyHash(Key);
+	}
 
 	/**
 	 * Find the value associated with a specified key, or if none exists, 
@@ -411,13 +476,15 @@ private:
 	 * @param Key The key to search for.
 	 * @return A reference to the value associated with the specified key.
 	 */
-	template <typename ArgType>
-	FORCEINLINE ValueType& FindOrAddImpl(ArgType&& Arg)
+	template <typename InitKeyType>
+	ValueType& FindOrAddImpl(uint32 KeyHash, InitKeyType&& Key)
 	{
-		if (auto* Pair = Pairs.Find(Arg))
+		if (auto* Pair = Pairs.FindByHash(KeyHash, Key))
+		{
 			return Pair->Value;
+		}
 
-		return Add(Forward<ArgType>(Arg));
+		return AddByHash(KeyHash, Forward<InitKeyType>(Key));
 	}
 
 	/**
@@ -429,12 +496,14 @@ private:
 	 * @return A reference to the value associated with the specified key.
 	 */
 	template <typename InitKeyType, typename InitValueType>
-	FORCEINLINE ValueType& FindOrAddImpl(InitKeyType&& Key, InitValueType&& Value)
+	ValueType& FindOrAddImpl(uint32 KeyHash, InitKeyType&& Key, InitValueType&& Value)
 	{
-		if (auto* Pair = Pairs.Find(Key))
+		if (auto* Pair = Pairs.FindByHash(KeyHash, Key))
+		{
 			return Pair->Value;
+		}
 
-		return Add(Forward<InitKeyType>(Key), Forward<InitValueType>(Value));
+		return AddByHash(KeyHash, Forward<InitKeyType>(Key), Forward<InitValueType>(Value));
 	}
 
 public:
@@ -446,8 +515,12 @@ public:
 	 * @param Key The key to search for.
 	 * @return A reference to the value associated with the specified key.
 	 */
-	FORCEINLINE ValueType& FindOrAdd(const KeyType&  Key) { return FindOrAddImpl(                   Key ); }
-	FORCEINLINE ValueType& FindOrAdd(      KeyType&& Key) { return FindOrAddImpl(MoveTempIfPossible(Key)); }
+	FORCEINLINE ValueType& FindOrAdd(const KeyType&  Key) { return FindOrAddImpl(HashKey(Key),					  Key); }
+	FORCEINLINE ValueType& FindOrAdd(      KeyType&& Key) { return FindOrAddImpl(HashKey(Key), MoveTempIfPossible(Key)); }
+
+	/** See FindOrAdd() and class documentation section on ByHash() functions */
+	FORCEINLINE ValueType& FindOrAddByHash(uint32 KeyHash, const KeyType&  Key) { return FindOrAddImpl(KeyHash,                    Key); }
+	FORCEINLINE ValueType& FindOrAddByHash(uint32 KeyHash,       KeyType&& Key) { return FindOrAddImpl(KeyHash, MoveTempIfPossible(Key)); }
 
 	/**
 	 * Find the value associated with a specified key, or if none exists, 
@@ -457,31 +530,16 @@ public:
 	 * @param Value The value to associate with the key.
 	 * @return A reference to the value associated with the specified key.
 	 */
-	FORCEINLINE ValueType& FindOrAdd(const KeyType&  Key, const ValueType&  Value) { return FindOrAddImpl(                   Key ,                    Value  ); }
-	FORCEINLINE ValueType& FindOrAdd(const KeyType&  Key, ValueType&&       Value) { return FindOrAddImpl(                   Key , MoveTempIfPossible(Value) ); }
-	FORCEINLINE ValueType& FindOrAdd(      KeyType&& Key, const ValueType&  Value) { return FindOrAddImpl(MoveTempIfPossible(Key),                    Value  ); }
-	FORCEINLINE ValueType& FindOrAdd(      KeyType&& Key, ValueType&&       Value) { return FindOrAddImpl(MoveTempIfPossible(Key), MoveTempIfPossible(Value) ); }
+	FORCEINLINE ValueType& FindOrAdd(const KeyType&  Key, const ValueType&  Value) { return FindOrAddImpl(HashKey(Key),						Key,                    Value  ); }
+	FORCEINLINE ValueType& FindOrAdd(const KeyType&  Key, ValueType&&       Value) { return FindOrAddImpl(HashKey(Key),						Key, MoveTempIfPossible(Value) ); }
+	FORCEINLINE ValueType& FindOrAdd(      KeyType&& Key, const ValueType&  Value) { return FindOrAddImpl(HashKey(Key), MoveTempIfPossible(Key),                    Value  ); }
+	FORCEINLINE ValueType& FindOrAdd(      KeyType&& Key, ValueType&&       Value) { return FindOrAddImpl(HashKey(Key), MoveTempIfPossible(Key), MoveTempIfPossible(Value) ); }
 
-	/**
-	 * Find the value associated with a specified key, or if none exists, 
-	 * adds a value using the key as the constructor parameter.
-	 *
-	 * @param Key The key to search for.
-	 * @return A reference to the value associated with the specified key.
-	 */
-	//@todo UE4 merge - this prevents FConfigCacheIni from compiling
-	/*ValueType& FindOrAddKey(KeyInitType Key)
-	{
-		TPair* Pair = Pairs.Find(Key);
-		if( Pair )
-		{
-			return Pair->Value;
-		}
-		else
-		{
-			return Set(Key, ValueType(Key));
-		}
-	}*/
+	/** See FindOrAdd() and class documentation section on ByHash() functions */
+	FORCEINLINE ValueType& FindOrAddByHash(uint32 KeyHash, const KeyType&  Key, const ValueType&  Value) { return FindOrAddImpl(KeyHash,                     Key,                    Value); }
+	FORCEINLINE ValueType& FindOrAddByHash(uint32 KeyHash, const KeyType&  Key,       ValueType&& Value) { return FindOrAddImpl(KeyHash,                     Key, MoveTempIfPossible(Value)); }
+	FORCEINLINE ValueType& FindOrAddByHash(uint32 KeyHash,       KeyType&& Key, const ValueType&  Value) { return FindOrAddImpl(KeyHash, MoveTempIfPossible(Key),                    Value); }
+	FORCEINLINE ValueType& FindOrAddByHash(uint32 KeyHash,       KeyType&& Key,       ValueType&& Value) { return FindOrAddImpl(KeyHash, MoveTempIfPossible(Key), MoveTempIfPossible(Value)); }
 
 	/**
 	 * Find a reference to the value associated with a specified key.
@@ -536,6 +594,13 @@ public:
 		return Pairs.Contains(Key);
 	}
 
+	/** See Contains() and class documentation section on ByHash() functions */
+	template<typename ComparableKey>
+	FORCEINLINE bool ContainsByHash(uint32 KeyHash, const ComparableKey& Key) const
+	{
+		return Pairs.ContainsByHash(KeyHash, Key);
+	}
+
 	/**
 	 * Generate an array from the keys in this map.
 	 *
diff --git a/Engine/Source/Runtime/Core/Public/Containers/Set.h b/Engine/Source/Runtime/Core/Public/Containers/Set.h
--- a/Engine/Source/Runtime/Core/Public/Containers/Set.h
+++ b/Engine/Source/Runtime/Core/Public/Containers/Set.h
@@ -55,7 +55,8 @@ struct DefaultKeyFuncs : BaseKeyFuncs<ElementType,ElementType,bInAllowDuplicateK
 	/**
 	 * @return True if the keys match.
 	 */
-	static FORCEINLINE bool Matches(KeyInitType A,KeyInitType B)
+	template<typename ComparableKey>
+	static FORCEINLINE bool Matches(KeyInitType A, ComparableKey B)
 	{
 		return A == B;
 	}
@@ -198,6 +199,14 @@ public:
  * the element type.  It uses a TSparseArray of the elements, and also links the elements into a hash with a 
  * number of buckets proportional to the number of elements.  Addition, removal, and finding are O(1).
  *
+ * The ByHash() functions are somewhat dangerous but particularly useful in two scenarios:
+ * -- Heterogeneous lookup to avoid creating expensive keys like FString when looking up by const TCHAR*.
+ *	  You must ensure the hash is calculated in the same way as ElementType is hashed.
+ *    If possible put both ComparableKey and ElementType hash functions next to each other in the same header
+ *    to avoid bugs when the ElementType hash function is changed.
+ * -- Reducing contention around hash tables protected by a lock. It is often important to incur
+ *    the cache misses of reading key data and doing the hashing *before* acquiring the lock.
+ *
  **/
 template<
 	typename InElementType,
@@ -487,24 +496,29 @@ public:
 	 * @param	bIsAlreadyInSetPtr	[out]	Optional pointer to bool that will be set depending on whether element is already in set
 	 * @return	A pointer to the element stored in the set.
 	 */
-	FORCEINLINE FSetElementId Add(const InElementType&  InElement, bool* bIsAlreadyInSetPtr = NULL) { return Emplace(                   InElement , bIsAlreadyInSetPtr); }
-	FORCEINLINE FSetElementId Add(      InElementType&& InElement, bool* bIsAlreadyInSetPtr = NULL) { return Emplace(MoveTempIfPossible(InElement), bIsAlreadyInSetPtr); }
+	FORCEINLINE FSetElementId Add(const InElementType&  InElement, bool* bIsAlreadyInSetPtr = nullptr) { return Emplace(                   InElement , bIsAlreadyInSetPtr); }
+	FORCEINLINE FSetElementId Add(      InElementType&& InElement, bool* bIsAlreadyInSetPtr = nullptr) { return Emplace(MoveTempIfPossible(InElement), bIsAlreadyInSetPtr); }
 
 	/**
 	 * Adds an element to the set.
 	 *
-	 * @param	Args						The argument(s) to be forwarded to the set element's constructor.
+	 * @see		Class documentation section on ByHash() functions
+	 * @param	InElement					Element to add to set
 	 * @param	bIsAlreadyInSetPtr	[out]	Optional pointer to bool that will be set depending on whether element is already in set
-	 * @return	A pointer to the element stored in the set.
-	 */
-	template <typename ArgsType>
-	FSetElementId Emplace(ArgsType&& Args,bool* bIsAlreadyInSetPtr = NULL)
+     * @return  A handle to the element stored in the set
+     */
+	FORCEINLINE FSetElementId AddByHash(uint32 KeyHash, const InElementType& InElement, bool* bIsAlreadyInSetPtr = nullptr)
 	{
-		// Create a new element.
-		FSparseArrayAllocationInfo ElementAllocation = Elements.AddUninitialized();
-		FSetElementId ElementId(ElementAllocation.Index);
-		auto& Element = *new(ElementAllocation) SetElementType(Forward<ArgsType>(Args));
+		return EmplaceByHash(KeyHash, InElement, bIsAlreadyInSetPtr);
+	}
+	FORCEINLINE FSetElementId AddByHash(uint32 KeyHash,		 InElementType&& InElement, bool* bIsAlreadyInSetPtr = nullptr)
+	{
+		return EmplaceByHash(KeyHash, MoveTempIfPossible(InElement), bIsAlreadyInSetPtr);
+	}
 
+private:
+	FSetElementId EmplaceImpl(uint32 KeyHash, SetElementType& Element, FSetElementId ElementId, bool* bIsAlreadyInSetPtr)
+	{
 		bool bIsAlreadyInSet = false;
 		if (!KeyFuncs::bAllowDuplicateKeys)
 		{
@@ -513,7 +527,7 @@ public:
 			// Don't bother searching for a duplicate if this is the first element we're adding
 			if (Elements.Num() != 1)
 			{
-				FSetElementId ExistingId = FindId(KeyFuncs::GetSetKey(Element.Value));
+				FSetElementId ExistingId = FindIdByHash(KeyHash, KeyFuncs::GetSetKey(Element.Value));
 				bIsAlreadyInSet = ExistingId.IsValidId();
 				if (bIsAlreadyInSet)
 				{
@@ -532,10 +546,10 @@ public:
 		if (!bIsAlreadyInSet)
 		{
 			// Check if the hash needs to be resized.
-			if(!ConditionalRehash(Elements.Num()))
+			if (!ConditionalRehash(Elements.Num()))
 			{
 				// If the rehash didn't add the new element to the hash, add it.
-				HashElement(ElementId,Element);
+				LinkElement(ElementId, Element, KeyHash);
 			}
 		}
 
@@ -547,6 +561,43 @@ public:
 		return ElementId;
 	}
 
+public:
+	/**
+	 * Adds an element to the set.
+	 *
+	 * @param	Args						The argument(s) to be forwarded to the set element's constructor.
+	 * @param	bIsAlreadyInSetPtr	[out]	Optional pointer to bool that will be set depending on whether element is already in set
+	 * @return	A handle to the element stored in the set.
+	 */
+	template <typename ArgsType>
+	FSetElementId Emplace(ArgsType&& Args, bool* bIsAlreadyInSetPtr = nullptr)
+	{
+		// Create a new element.
+		FSparseArrayAllocationInfo ElementAllocation = Elements.AddUninitialized();
+		SetElementType& Element = *new (ElementAllocation) SetElementType(Forward<ArgsType>(Args));
+
+		uint32 KeyHash = KeyFuncs::GetKeyHash(KeyFuncs::GetSetKey(Element.Value));
+		return EmplaceImpl(KeyHash, Element, ElementAllocation.Index, bIsAlreadyInSetPtr);
+	}
+	
+	/**
+	 * Adds an element to the set.
+	 *
+	 * @see		Class documentation section on ByHash() functions
+	 * @param	Args						The argument(s) to be forwarded to the set element's constructor.
+	 * @param	bIsAlreadyInSetPtr	[out]	Optional pointer to bool that will be set depending on whether element is already in set
+	 * @return	A handle to the element stored in the set.
+	 */
+	template <typename ArgsType>
+	FSetElementId EmplaceByHash(uint32 KeyHash, ArgsType&& Args, bool* bIsAlreadyInSetPtr = nullptr)
+	{
+		// Create a new element.
+		FSparseArrayAllocationInfo ElementAllocation = Elements.AddUninitialized();
+		SetElementType& Element = *new (ElementAllocation) SetElementType(Forward<ArgsType>(Args));
+
+		return EmplaceImpl(KeyHash, Element, ElementAllocation.Index, bIsAlreadyInSetPtr);
+	}
+
 	template<typename ArrayAllocator>
 	void Append(const TArray<ElementType, ArrayAllocator>& InElements)
 	{
@@ -652,6 +703,32 @@ public:
 		return FSetElementId();
 	}
 
+	/**
+	 * Finds an element with a pre-calculated hash and a key that can be compared to KeyType
+	 * @see	Class documentation section on ByHash() functions
+	 * @return The element id that matches the key and hash or an invalid element id
+	 */
+	template<typename ComparableKey>
+	FSetElementId FindIdByHash(uint32 KeyHash, const ComparableKey& Key) const
+	{
+		if (Elements.Num())
+		{
+			checkSlow(KeyHash == KeyFuncs::GetKeyHash(Key));
+
+			for (FSetElementId ElementId = GetTypedHash(KeyHash);
+				ElementId.IsValidId();
+				ElementId = Elements[ElementId].HashNextId)
+			{
+				if (KeyFuncs::Matches(KeyFuncs::GetSetKey(Elements[ElementId].Value), Key))
+				{
+					// Return the first match, regardless of whether the set has multiple matches for the key or not.
+					return ElementId;
+				}
+			}
+		}
+		return FSetElementId();
+	}
+
 	/**
 	 * Finds an element with the given key in the set.
 	 * @param Key - The key to search for.
@@ -666,7 +743,7 @@ public:
 		}
 		else
 		{
-			return NULL;
+			return nullptr;
 		}
 	}
 	
@@ -677,17 +754,67 @@ public:
 	 */
 	FORCEINLINE const ElementType* Find(KeyInitType Key) const
 	{
-		FSetElementId ElementId = FindId(Key);
-		if(ElementId.IsValidId())
+		return const_cast<TSet*>(this)->Find(Key);
+	}
+
+	/**
+	 * Finds an element with a pre-calculated hash and a key that can be compared to KeyType.
+	 * @see	Class documentation section on ByHash() functions
+	 * @return A pointer to the contained element or nullptr.
+	 */
+	template<typename ComparableKey>
+	ElementType* FindByHash(uint32 KeyHash, const ComparableKey& Key)
+	{
+		FSetElementId ElementId = FindIdByHash(KeyHash, Key);
+		if (ElementId.IsValidId())
 		{
 			return &Elements[ElementId].Value;
 		}
 		else
 		{
-			return NULL;
+			return nullptr;
+		}
+	}
+
+	template<typename ComparableKey>
+	const ElementType* FindByHash(uint32 KeyHash, const ComparableKey& Key) const
+	{
+		return const_cast<TSet*>(this)->FindByHash(KeyHash, Key);
+	}
+
+private:
+	template<typename ComparableKey>
+	FORCEINLINE int32 RemoveImpl(uint32 KeyHash, const ComparableKey& Key)
+	{
+		int32 NumRemovedElements = 0;
+
+		FSetElementId* NextElementId = &GetTypedHash(KeyHash);
+		while (NextElementId->IsValidId())
+		{
+			auto& Element = Elements[*NextElementId];
+			if (KeyFuncs::Matches(KeyFuncs::GetSetKey(Element.Value), Key))
+			{
+				// This element matches the key, remove it from the set.  Note that Remove sets *NextElementId to point to the next
+				// element after the removed element in the hash bucket.
+				Remove(*NextElementId);
+				NumRemovedElements++;
+
+				if (!KeyFuncs::bAllowDuplicateKeys)
+				{
+					// If the hash disallows duplicate keys, we're done removing after the first matched key.
+					break;
+				}
+			}
+			else
+			{
+				NextElementId = &Element.HashNextId;
+			}
 		}
+
+		return NumRemovedElements;
 	}
 
+public:
 	/**
 	 * Removes all elements from the set matching the specified key.
 	 * @param Key - The key to match elements against.
@@ -695,35 +822,32 @@ public:
 	 */
 	int32 Remove(KeyInitType Key)
 	{
-		int32 NumRemovedElements = 0;
+		if (Elements.Num())
+		{
+			return RemoveImpl(KeyFuncs::GetKeyHash(Key), Key);
+		}
+
+		return 0;
+	}
+
+	/**
+	 * Removes all elements from the set matching the specified key.
+	 *
+	 * @see		Class documentation section on ByHash() functions
+	 * @param	Key - The key to match elements against.
+	 * @return	The number of elements removed.
+	 */
+	template<typename ComparableKey>
+	int32 RemoveByHash(uint32 KeyHash, const ComparableKey& Key)
+	{
+		checkSlow(KeyHash == KeyFuncs::GetKeyHash(Key));
 
 		if (Elements.Num())
 		{
-			FSetElementId* NextElementId = &GetTypedHash(KeyFuncs::GetKeyHash(Key));
-			while(NextElementId->IsValidId())
-			{
-				auto& Element = Elements[*NextElementId];
-				if(KeyFuncs::Matches(KeyFuncs::GetSetKey(Element.Value),Key))
-				{
-					// This element matches the key, remove it from the set.  Note that Remove sets *NextElementId to point to the next
-					// element after the removed element in the hash bucket.
-					Remove(*NextElementId);
-					NumRemovedElements++;
-
-					if(!KeyFuncs::bAllowDuplicateKeys)
-					{
-						// If the hash disallows duplicate keys, we're done removing after the first matched key.
-						break;
-					}
-				}
-				else
-				{
-					NextElementId = &Element.HashNextId;
-				}
-			}
+			return RemoveImpl(KeyHash, Key);
 		}
 
-		return NumRemovedElements;
+		return 0;
 	}
 
 	/**
@@ -736,6 +860,17 @@ public:
 		return FindId(Key).IsValidId();
 	}
 
+	/**
+	 * Checks if the element contains an element with the given key.
+	 *
+	 * @see	Class documentation section on ByHash() functions
+	 */
+	template<typename ComparableKey>
+	FORCEINLINE bool ContainsByHash(uint32 KeyHash, const ComparableKey& Key) const
+	{
+		return FindIdByHash(KeyHash, Key).IsValidId();
+	}
+
 	/**
 	 * Sorts the set's elements using the provided comparison class.
 	 */
@@ -990,7 +1125,7 @@ private:
 	ElementArrayType Elements;
 
 	mutable HashType Hash;
-	mutable int32    HashSize;
+	mutable int32	 HashSize;
 
 	FORCEINLINE FSetElementId& GetTypedHash(int32 HashIndex) const
 	{
@@ -1019,17 +1154,23 @@ private:
 		return FSetElementId(Index);
 	}
 
-	/** Adds an element to the hash. */
-	FORCEINLINE void HashElement(FSetElementId ElementId,const SetElementType& Element) const
+	/** Links an added element to the hash chain. */
+	FORCEINLINE void LinkElement(FSetElementId ElementId, const SetElementType& Element, uint32 KeyHash) const
 	{
 		// Compute the hash bucket the element goes in.
-		Element.HashIndex = KeyFuncs::GetKeyHash(KeyFuncs::GetSetKey(Element.Value)) & (HashSize-1);
+		Element.HashIndex = KeyHash & (HashSize - 1);
 
 		// Link the element into the hash bucket.
 		Element.HashNextId = GetTypedHash(Element.HashIndex);
 		GetTypedHash(Element.HashIndex) = ElementId;
 	}
 
+	/** Hashes and links an added element to the hash chain. */
+	FORCEINLINE void HashElement(FSetElementId ElementId, const SetElementType& Element) const
+	{
+		LinkElement(ElementId, Element, KeyFuncs::GetKeyHash(KeyFuncs::GetSetKey(Element.Value)));
+	}
+
 	/**
 	 * Checks if the hash has an appropriate number of buckets, and if not resizes it.
 	 * @param NumHashedElements - The number of elements to size the hash for.
@@ -1463,9 +1604,9 @@ public:
 					void* Element = (uint8*)Elements.GetData(Index, Layout.SparseArrayLayout);
 
 					// Compute the hash bucket the element goes in.
-					uint32 ElementHash = GetKeyHash(Element);
-					int32  HashIndex   = ElementHash & (HashSize - 1);
-					GetHashIndexRef(Element, Layout) = ElementHash & (HashSize - 1);
+					uint32 KeyHash = GetKeyHash(Element);
+					int32  HashIndex   = KeyHash & (HashSize - 1);
+					GetHashIndexRef(Element, Layout) = KeyHash & (HashSize - 1);
 
 					// Link the element into the hash bucket.
 					GetHashNextIdRef(Element, Layout) = GetTypedHash(HashIndex);
@@ -1483,8 +1624,8 @@ public:
 	{
 		if (Elements.Num())
 		{
-			const uint32 ElementHash = GetKeyHash(Element);
-			const int32  HashIndex = ElementHash & (HashSize - 1);
+			const uint32 KeyHash = GetKeyHash(Element);
+			const int32  HashIndex = KeyHash & (HashSize - 1);
 
 			uint8* CurrentElement = nullptr;
 			for (FSetElementId ElementId = GetTypedHash(HashIndex);
@@ -1505,7 +1646,7 @@ public:
 	void Add(const void* Element, const FScriptSetLayout& Layout, TFunctionRef<uint32(const void*)> GetKeyHash, TFunctionRef<bool(const void*, const void*)> EqualityFn, TFunctionRef<void(void*)> ConstructFn, TFunctionRef<void(void*)> DestructFn)
 	{
 		// Minor efficiency concern: we hash the element both here and in the FindIndex() call
-		uint32 ElementHash = GetKeyHash(Element);
+		uint32 KeyHash = GetKeyHash(Element);
 
 		int32 NewElementIndex = FindIndex(Element, Layout, GetKeyHash, EqualityFn);
 		if (NewElementIndex != INDEX_NONE)
@@ -1517,7 +1658,7 @@ public:
 
 			// We don't update the hash because we don't need to - the new element
 			// should have the same hash, but let's just check.
-			check(ElementHash == GetKeyHash(ElementPtr));
+			check(KeyHash == GetKeyHash(ElementPtr));
 		}
 		else
 		{
@@ -1535,7 +1676,7 @@ public:
 			else
 			{
 				// link the new element into the set:
-				int32 HashIndex = ElementHash & (HashSize - 1);
+				int32 HashIndex = KeyHash & (HashSize - 1);
 				FSetElementId& TypedHash = GetTypedHash(HashIndex);
 				GetHashIndexRef(ElementPtr, Layout) = HashIndex;
 				GetHashNextIdRef(ElementPtr, Layout) = TypedHash;
