Creates a FSobol::ComputeGPUSpatialSeed()

diff --git a/Engine/Source/Runtime/Core/Private/Math/Sobol.cpp b/Engine/Source/Runtime/Core/Private/Math/Sobol.cpp
--- a/Engine/Source/Runtime/Core/Private/Math/Sobol.cpp
+++ b/Engine/Source/Runtime/Core/Private/Math/Sobol.cpp
@@ -229,3 +229,39 @@ const int32 FSobol::Cell3DGrayNumbers[FSobol::MaxCell3DBits + 1][32][3] = {
 	{ {0x001150, 0x003ba6, 0x0020ba}, {0x0037a4, 0x000367, 0x001a16}, {0x001ac7, 0x002ea7, 0x00020e}, {0x002b1d, 0x003188, 0x000619}, {0x0017a4, 0x000925, 0x001954}, {0x00367e, 0x0015ef, 0x0014e0}, {0x00331d, 0x003dfb, 0x00082c}, {0x0022f4, 0x00302e, 0x0034f1}, {0x001d2e, 0x002aa9, 0x0032f8}, {0x0013b9, 0x0031cc, 0x001b6b}, {0x00165a, 0x0011d8, 0x000fa7}, {0x00185d, 0x000fbd, 0x002d34}, {0x000974, 0x000e68, 0x0015e9}, {0x0036ce, 0x0016ef, 0x0011e0}, {0x0033c5, 0x00397b, 0x000aac}, {0x0033e0, 0x000448, 0x001f0b}, {0x000cd4, 0x0016ef, 0x001ca2}, {0x0024e2, 0x00397b, 0x000d4d}, {0x00105b, 0x000448, 0x0017be}, {0x002cd1, 0x002d49, 0x003847}, {0x002fb2, 0x0002dd, 0x002eb5}, {0x0000ff, 0x00072f, 0x000749}, {0x00007f, 0x000aa7, 0x0008c9}, {0x0011ef, 0x003355, 0x0025b3}, {0x00118f, 0x0036cb, 0x002213}, {0x0037eb, 0x000067, 0x00116f}, {0x0006d9, 0x0018c8, 0x0010d0}, {0x003ab6, 0x0005ee, 0x002a8e}, {0x000681, 0x003563, 0x003fd5}, {0x00016a, 0x001d58, 0x0007c6}, {0x0035bb, 0x00286b, 0x0012b5}, {0x0009ad, 0x000d6a, 0x00308c} }
 };
 
+// static
+uint16 FSobol::ComputeGPUSpatialSeed(int32 x, int32 y, int32 Index)
+{
+	check(x >= 0 && x < 16);
+	check(y >= 0 && y < 16);
+
+	uint16 Result = 0;
+	if (Index == 0)
+	{
+		Result  = (x & 0x001) ? 0xf68e : 0;
+		Result ^= (x & 0x002) ? 0x8e56 : 0;
+		Result ^= (x & 0x004) ? 0x1135 : 0;
+		Result ^= (x & 0x008) ? 0x220a : 0;
+		Result ^= (y & 0x001) ? 0x94c4 : 0;
+		Result ^= (y & 0x002) ? 0x4ac2 : 0;
+		Result ^= (y & 0x004) ? 0xfb57 : 0;
+		Result ^= (y & 0x008) ? 0x0454 : 0;
+	}
+	else if (Index == 1)
+	{
+		Result  = (x & 0x001) ? 0x4414 : 0;
+		Result ^= (x & 0x002) ? 0x8828 : 0;
+		Result ^= (x & 0x004) ? 0xe69e : 0;
+		Result ^= (x & 0x008) ? 0xae76 : 0;
+		Result ^= (y & 0x001) ? 0xa28a : 0;
+		Result ^= (y & 0x002) ? 0x265e : 0;
+		Result ^= (y & 0x004) ? 0xe69e : 0;
+		Result ^= (y & 0x008) ? 0xae76 : 0;
+	}
+	else
+	{
+		check(0);
+	}
+
+	return Result;
+}
diff --git a/Engine/Source/Runtime/Core/Public/Math/Sobol.h b/Engine/Source/Runtime/Core/Public/Math/Sobol.h
--- a/Engine/Source/Runtime/Core/Public/Math/Sobol.h
+++ b/Engine/Source/Runtime/Core/Public/Math/Sobol.h
@@ -106,4 +106,7 @@ public:
 	* @return The Sobol result at given Index
 	*/
 	static FVector Next(int32 Index, int32 CellBits, FVector Value);
+
+	/** Return value for GSystemTextures.SobolSampling */
+	static uint16 ComputeGPUSpatialSeed(int32 x, int32 y, int32 Index);
 };
\ No newline at end of file
diff --git a/Engine/Source/Runtime/Renderer/Private/SystemTextures.cpp b/Engine/Source/Runtime/Renderer/Private/SystemTextures.cpp
--- a/Engine/Source/Runtime/Renderer/Private/SystemTextures.cpp
+++ b/Engine/Source/Runtime/Renderer/Private/SystemTextures.cpp
@@ -182,29 +182,13 @@ void FSystemTextures::InitializeFeatureLevelDependentTextures(FRHICommandListImm
 				// 16x16 block starting at 0,0 = Sobol X,Y from bottom 4 bits of cell X,Y
 				for (int x = 0; x < 16; ++x, ++Dest)
 				{
-					Result  = (x & 0x001) ? 0xf68e : 0;
-					Result ^= (x & 0x002) ? 0x8e56 : 0;
-					Result ^= (x & 0x004) ? 0x1135 : 0;
-					Result ^= (x & 0x008) ? 0x220a : 0;
-					Result ^= (y & 0x001) ? 0x94c4 : 0;
-					Result ^= (y & 0x002) ? 0x4ac2 : 0;
-					Result ^= (y & 0x004) ? 0xfb57 : 0;
-					Result ^= (y & 0x008) ? 0x0454 : 0;
-					*Dest = Result;
+					*Dest = FSobol::ComputeGPUSpatialSeed(x, y, /* Index = */ 0);
 				}
 
 				// 16x16 block starting at 16,0 = Sobol X,Y from 2nd 4 bits of cell X,Y
 				for (int x = 0; x < 16; ++x, ++Dest)
 				{
-					Result =  (x & 0x001) ? 0x4414 : 0;
-					Result ^= (x & 0x002) ? 0x8828 : 0;
-					Result ^= (x & 0x004) ? 0xe69e : 0;
-					Result ^= (x & 0x008) ? 0xae76 : 0;
-					Result ^= (y & 0x001) ? 0xa28a : 0;
-					Result ^= (y & 0x002) ? 0x265e : 0;
-					Result ^= (y & 0x004) ? 0xe69e : 0;
-					Result ^= (y & 0x008) ? 0xae76 : 0;
-					*Dest = Result;
+					*Dest = FSobol::ComputeGPUSpatialSeed(x, y, /* Index = */ 1);
 				}
 			}
 			RHICmdList.UnlockTexture2D((FTexture2DRHIRef&)SobolSampling->GetRenderTargetItem().ShaderResourceTexture, 0, false);
