Trigger out of memory flow when MallocAnsi fails to allocate memory, rather than crashing on a check. This change modifies MallocAnsi to behave like the other allocators (MallocTBB, for example). This change also fixes MallocAnsi to correctly return nullptr on zero sized allocations, matching the behavior of the other allocators.

diff --git a/Engine/Source/Runtime/Core/Private/HAL/MallocAnsi.cpp b/Engine/Source/Runtime/Core/Private/HAL/MallocAnsi.cpp
--- a/Engine/Source/Runtime/Core/Private/HAL/MallocAnsi.cpp
+++ b/Engine/Source/Runtime/Core/Private/HAL/MallocAnsi.cpp
@@ -57,14 +57,17 @@ void* FMallocAnsi::Malloc( SIZE_T Size, uint32 Alignment )
 #elif PLATFORM_USE_ANSI_MEMALIGN
 	void* Result = memalign(Alignment, Size);
 #else
-	void* Ptr = malloc( Size + Alignment + sizeof(void*) + sizeof(SIZE_T) );
-	check(Ptr);
-	void* Result = Align( (uint8*)Ptr + sizeof(void*) + sizeof(SIZE_T), Alignment );
-	*((void**)( (uint8*)Result - sizeof(void*)					))	= Ptr;
-	*((SIZE_T*)( (uint8*)Result - sizeof(void*) - sizeof(SIZE_T)	))	= Size;
+	void* Ptr = malloc(Size + Alignment + sizeof(void*) + sizeof(SIZE_T));
+	void* Result = nullptr;
+	if (Ptr)
+	{
+		Result = Align((uint8*)Ptr + sizeof(void*) + sizeof(SIZE_T), Alignment);
+		 *((void**)((uint8*)Result - sizeof(void*))) = Ptr;
+		*((SIZE_T*)((uint8*)Result - sizeof(void*) - sizeof(SIZE_T))) = Size;
+	}
 #endif
 
-	if (Result == nullptr)
+	if (!Result && Size)
 	{
 		FPlatformMemory::OnOutOfMemory(Size, Alignment);
 	}
diff --git a/Engine/Source/Runtime/Core/Private/HAL/MallocBinned.cpp b/Engine/Source/Runtime/Core/Private/HAL/MallocBinned.cpp
--- a/Engine/Source/Runtime/Core/Private/HAL/MallocBinned.cpp
+++ b/Engine/Source/Runtime/Core/Private/HAL/MallocBinned.cpp
@@ -453,7 +453,7 @@ struct FMallocBinned::Private
 		UPTRINT BasePtr;
 		FPoolInfo* Pool = FindPoolInfo(Allocator, (UPTRINT)Ptr, BasePtr);
 #if PLATFORM_IOS || PLATFORM_MAC
-        if (Pool == NULL)
+        if (Pool == nullptr)
         {
 #if USE_OS_SMALL_BLOCK_ALLOC
 			// check the small block allocation range
@@ -1010,7 +1010,7 @@ void* FMallocBinned::Malloc(SIZE_T Size, uint32 Alignment)
 			Private::TrackStats(Table, Size);
 
 			FPoolInfo* Pool = Table->FirstPool;
-		if( !Pool )
+			if( !Pool )
 			{
 				Pool = Private::AllocatePoolMemory(*this, Table, PageCount*PageSize, BinnedSizeLimit+BinType);
 			}
@@ -1182,7 +1182,7 @@ bool FMallocBinned::GetAllocationSize(void *Original, SIZE_T &SizeOut)
 	FPoolInfo* Pool = Private::FindPoolInfo(*this, (UPTRINT)Original, BasePtr);
 
 #if PLATFORM_IOS || PLATFORM_MAC
-	if (Pool == NULL)
+	if (Pool == nullptr)
 	{
 #if USE_OS_SMALL_BLOCK_ALLOC
 		if ((uint64)Original >= Private::SmallBlockStartPtr && (uint64)Original <= Private::SmallBlockEndPtr)
diff --git a/Engine/Source/Runtime/Core/Private/HAL/MallocTBB.cpp b/Engine/Source/Runtime/Core/Private/HAL/MallocTBB.cpp
--- a/Engine/Source/Runtime/Core/Private/HAL/MallocTBB.cpp
+++ b/Engine/Source/Runtime/Core/Private/HAL/MallocTBB.cpp
@@ -48,7 +48,7 @@ void* FMallocTBB::Malloc( SIZE_T Size, uint32 Alignment )
 
 	MEM_TIME(MemTime -= FPlatformTime::Seconds());
 
-	void* NewPtr = NULL;
+	void* NewPtr = nullptr;
 #if PLATFORM_MAC
 	// macOS expects all allocations to be aligned to 16 bytes, but TBBs default alignment is 8, so on Mac we always have to use scalable_aligned_malloc.
 	// Contrary to scalable_malloc, scalable_aligned_malloc returns nullptr when trying to allocate 0 bytes, which is inconsistent with system malloc, so
@@ -109,7 +109,7 @@ void* FMallocTBB::Realloc( void* Ptr, SIZE_T NewSize, uint32 Alignment )
 		}
 	}
 #endif
-	void* NewPtr = NULL;
+	void* NewPtr = nullptr;
 #if PLATFORM_MAC
 	// macOS expects all allocations to be aligned to 16 bytes, but TBBs default alignment is 8, so on Mac we always have to use scalable_aligned_realloc
 	Alignment = AlignArbitrary(FMath::Max((uint32)16, Alignment), (uint32)16);
