Use a uint64 instead of UPTRINT for FLiteralOrName on 32-bit platforms since high bit of pointer may be in use

#ue4
#android

diff --git a/Engine/Source/Runtime/Core/Public/UObject/NameTypes.h b/Engine/Source/Runtime/Core/Public/UObject/NameTypes.h
--- a/Engine/Source/Runtime/Core/Public/UObject/NameTypes.h
+++ b/Engine/Source/Runtime/Core/Public/UObject/NameTypes.h
@@ -1006,6 +1006,41 @@ public:
 private:
 	struct FLiteralOrName
 	{
+#if PLATFORM_32BITS
+		// NOTE: uint64 instead of UPTRINT here since the high bit of a 32-bit pointer may be in use
+		FLiteralOrName(const TCHAR* Literal)
+			: Int(reinterpret_cast<uint64>(Literal) | LiteralFlag)
+		{}
+
+		explicit FLiteralOrName(FNameEntryId Name)
+			: Int(Name.ToUnstableInt())
+		{}
+
+		static constexpr uint64 LiteralFlag = uint64(1) << (sizeof(uint64) * 8 - 1);
+
+		bool IsName() const
+		{
+			return (LiteralFlag & Int) == 0;
+		}
+
+		bool IsLiteral() const
+		{
+			return LiteralFlag & Int;
+		}
+
+		FNameEntryId AsName() const
+		{
+			return FNameEntryId::FromUnstableInt(static_cast<uint32>(Int));
+		}
+
+		const TCHAR* AsLiteral() const
+		{
+			return reinterpret_cast<const TCHAR*>(Int & ~LiteralFlag);
+		}
+
+		uint64 Int;
+#else
+		// NOTE: uses high bit of pointer for flag; this may be an issue in future when high byte of address may be used for features like hardware ASAN
 		FLiteralOrName(const TCHAR* Literal)
 			: Int(reinterpret_cast<UPTRINT>(Literal) | LiteralFlag)
 		{}
@@ -1037,6 +1072,7 @@ private:
 		}
 
 		UPTRINT Int;
+#endif
 	};
 
 	mutable FLiteralOrName Either;
