Copying //UE4/Partner-MagicLeap-Staging to Dev-VR (//UE4/Dev-VR)

diff --git a/Engine/Source/Runtime/Core/Public/Lumin/CAPIShims/LuminAPIImpl.h b/Engine/Source/Runtime/Core/Public/Lumin/CAPIShims/LuminAPIImpl.h
new file mode 100644
--- /dev/null
+++ b/Engine/Source/Runtime/Core/Public/Lumin/CAPIShims/LuminAPIImpl.h
@@ -0,0 +1,449 @@
+// Copyright 1998-2019 Epic Games, Inc. All Rights Reserved.
+
+#pragma once
+#include "CoreMinimal.h"
+#include "Misc/ConfigCacheIni.h"
+#include "HAL/PlatformProcess.h"
+#include "Misc/Paths.h"
+#include "GenericPlatform/GenericPlatformFile.h"
+#include "Misc/CString.h"
+#include "Templates/Atomic.h"
+// for std::enable_if
+#include <type_traits>
+
+DEFINE_LOG_CATEGORY_STATIC(LogLuminAPI, Display, All);
+
+namespace MLSDK_API
+{
+
+/** Utility class to load the correct MLSDK libs depedning on the path set to the MLSDK package and whether or not we want to use MLremote / Zero Iteration*/
+class LibraryLoader
+{
+public:
+	static LibraryLoader & Ref()
+	{
+		static LibraryLoader LibLoader;
+		return LibLoader;
+	}
+
+	/** Reads the config file and environment variable for the MLSDK package path and sets up the correct environment to load the libraries from. */
+	LibraryLoader()
+	{
+#if !PLATFORM_LUMIN
+		// We search various places for the ML API DLLs to support loading alternate
+		// implementations. For example to use VDZI on PC platforms.
+
+		// Fixing SDKUNREAL-2128: When MLSDK environment variable is not set, we need to check for a possible
+		// config file value for SDK location.
+		FString MLSDKEnvironmentVariableName = TEXT("MLSDK");
+		FString MLSDK = FPlatformMisc::GetEnvironmentVariable(*MLSDKEnvironmentVariableName);
+		if (MLSDK.IsEmpty())
+		{
+			GConfig->GetString(TEXT("/Script/LuminPlatformEditor.MagicLeapSDKSettings"), TEXT("MLSDKPath"), MLSDK, GEngineIni);
+			// Formatted as (Path="C:/Directory")
+			if (!MLSDK.IsEmpty())
+			{
+				auto Index0 = MLSDK.Find(TEXT("\""));
+				if (Index0 != -1)
+				{
+					++Index0;
+					auto Index1 = MLSDK.Find(TEXT("\""), ESearchCase::IgnoreCase, ESearchDir::FromEnd);
+					if (Index1 > Index0)
+					{
+						MLSDK = MLSDK.Mid(Index0, Index1 - Index0);
+					}
+				}
+			}
+		}
+
+		bool bIsVDZIEnabled = false;
+		GConfig->GetBool(TEXT("/Script/MagicLeap.MagicLeapSettings"), TEXT("bEnableZI"), bIsVDZIEnabled, GEngineIni);
+
+		if (bIsVDZIEnabled)
+		{
+			// VDZI search paths: VDZI is only active in PC builds.
+			// This allows repointing MLAPI loading to the VDZI DLLs.
+			FString VDZILibraryPath = TEXT("");
+			GConfig->GetString(TEXT("MLSDK"), TEXT("LibraryPath"), VDZILibraryPath, GEngineIni);
+			if (!VDZILibraryPath.IsEmpty())
+			{
+				DllSearchPaths.Add(VDZILibraryPath);
+			}
+
+			// We also search in the MLSDK VDZI paths for libraries if we have them.
+			if (!MLSDK.IsEmpty())
+			{
+				TArray<FString> ZIShimPath = GetZIShimPath(MLSDK);
+				if (ZIShimPath.Num() > 0)
+				{
+					DllSearchPaths.Append(GetZIShimPath(MLSDK));
+				}
+				else
+				{
+					// Fallback to adding fixed known paths if we fail to get anything from
+					// the configuration data.
+					// The default VDZI dir.
+					DllSearchPaths.Add(FPaths::Combine(*MLSDK, TEXT("VirtualDevice"), TEXT("lib")));
+					// We also need to add the default bin dir as dependent libs are placed there instead
+					// of in the lib directory.
+					DllSearchPaths.Add(FPaths::Combine(*MLSDK, TEXT("VirtualDevice"), TEXT("bin")));
+				}
+			}
+		}
+#endif
+
+		// The MLSDK DLLs are platform specific and are segregated in directories for each platform.
+
+#if PLATFORM_LUMIN
+		// Lumin uses the system path as we are in device.
+		DllSearchPaths.Add(TEXT("/system/lib64"));
+#else
+
+		if (!MLSDK.IsEmpty())
+		{
+#if PLATFORM_WINDOWS
+			DllSearchPaths.Add(FPaths::Combine(*MLSDK, TEXT("lib"), TEXT("win64")));
+#elif PLATFORM_LINUX
+			DllSearchPaths.Add(FPaths::Combine(*MLSDK, TEXT("lib"), TEXT("linux64")));
+#elif PLATFORM_MAC
+			DllSearchPaths.Add(FPaths::Combine(*MLSDK, TEXT("lib"), TEXT("osx")));
+#endif // PLATFORM_WINDOWS
+		}
+
+#endif // PLATFORM_LUMIN
+
+		// Add the search paths to where we will load the DLLs from. 
+		// For all we add to the UE4 dir listing which takes care of the first-level
+		// loading of DLL modules.
+		for (const FString& path : DllSearchPaths)
+		{
+			FPlatformProcess::AddDllDirectory(*path);
+		}
+	}
+
+	/**
+	  Loads the given library from the correct path.
+	  @param Name Name of library to load, without any prefix or extension. e.g."ml_perception_client".
+	  @return True if the library was succesfully loaded. A false value generally indicates that the MLSDK path is not set correctly.
+	*/
+	void * LoadDLL(FString Name)
+	{
+		Name = FString(FPlatformProcess::GetModulePrefix()) + Name + TEXT(".") + FPlatformProcess::GetModuleExtension();
+#if PLATFORM_MAC
+		// FPlatformProcess::GetModulePrefix() for Mac is an empty string in Unreal
+		// whereas MLSDK uses 'lib' as the prefix for its OSX libs.
+		if (FString(FPlatformProcess::GetModulePrefix()).Len() == 0)
+		{
+			Name = FString(TEXT("lib")) + Name;
+		}
+#endif
+		for (const FString& path : DllSearchPaths)
+		{
+			void* dll = FPlatformProcess::GetDllHandle(*FPaths::Combine(*path, *Name));
+			if (dll != nullptr)
+			{
+				UE_LOG(LogLuminAPI, Display, TEXT("Dll loaded: %s"), *FPaths::Combine(*path, *Name));
+				return dll;
+			}
+		}
+
+		return nullptr;
+	}
+
+private:
+	TArray<FString> DllSearchPaths;
+
+	/** Fills the Variables with the evaluated contents of the SDK Shim discovery data. */
+	bool GetZIShimVariables(const FString & MLSDK, TMap<FString, FString> & ResultVariables) const
+	{
+		// The known path to the paths file.
+		FString SDKShimDiscoveryFile = FPaths::Combine(MLSDK, TEXT(".metadata"), TEXT("sdk_shim_discovery.txt"));
+		if (!FPaths::FileExists(SDKShimDiscoveryFile))
+		{
+			return false;
+		}
+		// Map of variable to value for evaluating the content of the file.
+		// On successful parsing and evaluation we copy the data to the result.
+		TMap<FString, FString> Variables;
+		Variables.Add(TEXT("$(MLSDK)"), MLSDK);
+		// TODO: Determine MLSDK version and set MLSDK_VERSION variable.
+#if PLATFORM_WINDOWS
+		Variables.Add(TEXT("$(HOST)"), TEXT("win64"));
+#elif PLATFORM_LINUX
+		Variables.Add(TEXT("$(HOST)"), TEXT("linux64"));
+#elif PLATFORM_MAC
+		Variables.Add(TEXT("$(HOST)"), TEXT("osx"));
+#endif // PLATFORM_WINDOWS
+		// Single pass algo for evaluating the file:
+		// 1. for each line:
+		// a. for each occurance of $(NAME):
+		// i. replace with Variables[NAME], if no NAME var found ignore replacement.
+		// b. split into var=value, and add to Variables.
+		IPlatformFile& PlatformFile = IPlatformFile::GetPlatformPhysical();
+		TUniquePtr<IFileHandle> File(PlatformFile.OpenRead(*SDKShimDiscoveryFile));
+		if (File)
+		{
+			TArray<uint8> Data;
+			Data.SetNumZeroed(File->Size() + 1);
+			File->Read(Data.GetData(), Data.Num() - 1);
+			FUTF8ToTCHAR TextConvert(reinterpret_cast<ANSICHAR*>(Data.GetData()));
+			const TCHAR * Text = TextConvert.Get();
+			while (Text && *Text)
+			{
+				Text += FCString::Strspn(Text, TEXT("\t "));
+				if (Text[0] == '#' || Text[0] == '\r' || Text[0] == '\n' || !*Text)
+				{
+					// Skip comment or empty lines.
+					Text += FCString::Strcspn(Text, TEXT("\r\n"));
+				}
+				else
+				{
+					// Parse variable=value line.
+					FString Variable(FCString::Strcspn(Text, TEXT("\t= ")), Text);
+					Text += Variable.Len();
+					Text += FCString::Strspn(Text, TEXT("\t= "));
+					FString Value(FCString::Strcspn(Text, TEXT("\r\n")), Text);
+					Text += Value.Len();
+					Value.TrimEndInline();
+					// Eval any var references in both variable and value.
+					int32 EvalCount = 0;
+					do
+					{
+						EvalCount = 0;
+						for (auto VarEntry : Variables)
+						{
+							EvalCount += Variable.ReplaceInline(*VarEntry.Key, *VarEntry.Value);
+							EvalCount += Value.ReplaceInline(*VarEntry.Key, *VarEntry.Value);
+						}
+					} while (EvalCount != 0 && (Variable.Contains(TEXT("$(")) || Value.Contains(TEXT("$("))));
+					// Intern the new variable.
+					Variable = TEXT("$(") + Variable + TEXT(")");
+					Variables.Add(Variable, Value);
+				}
+				// Skip over EOL to get to the next line.
+				if (Text[0] == '\r' && Text[1] == '\n')
+				{
+					Text += 2;
+				}
+				else
+				{
+					Text += 1;
+				}
+			}
+		}
+		// We now need to copy the evaled variables to the result and un-munge them for
+		// plain access. We use them munged for the eval for easier eval replacement above.
+		ResultVariables.Empty(Variables.Num());
+		for (auto VarEntry : Variables)
+		{
+			ResultVariables.Add(VarEntry.Key.Mid(2, VarEntry.Key.Len() - 3), VarEntry.Value);
+		}
+		return true;
+	}
+
+	TArray<FString> GetZIShimPath(const FString & MLSDK) const
+	{
+#if PLATFORM_WINDOWS
+		const FString ZIShimPathVar = "ZI_SHIM_PATH_win64";
+#elif PLATFORM_LINUX
+		const FString ZIShimPathVar = "ZI_SHIM_PATH_linux64";
+#elif PLATFORM_MAC
+		const FString ZIShimPathVar = "ZI_SHIM_PATH_osx";
+#else
+		const FString ZIShimPathVar = "";
+#endif // PLATFORM_WINDOWS
+
+		TMap<FString, FString> Variables;
+		if (GetZIShimVariables(MLSDK, Variables) && Variables.Contains(ZIShimPathVar))
+		{
+			// The shim path var we are looking for.
+			FString Value = Variables[ZIShimPathVar];
+			// Since it's a path variable it can have multiple components. Hence
+			// split those out into our result;
+			TArray<FString> Result;
+			for (FString Path; Value.Split(TEXT(";"), &Path, &Value); )
+			{
+				Result.Add(Path);
+			}
+			Result.Add(Value);
+			return Result;
+		}
+		return TArray<FString>();
+	}
+};
+
+/** Manages a single API library to load it on demand when retrieving an entry in that library.
+	The library is designated with a type key to statically bind the loaded instance to
+	only one of these. */
+template <typename LibKey>
+class Library
+{
+public:
+	~Library()
+	{
+		if (DllHandle)
+		{
+			FPlatformProcess::FreeDllHandle(DllHandle);
+			DllHandle = nullptr;
+		}
+	}
+
+	// The singleton for the library.
+	static Library & Ref()
+	{
+		static Library LibraryInstance;
+		return LibraryInstance;
+	}
+
+	// Set the name of the DLLs (or SO, or DYLIB) to load when fetching symbols.
+	void SetName(const char * Name)
+	{
+		if (!LibName) LibName = Name;
+	}
+
+	void * GetEntry(const char * Name)
+	{
+		if (!DllHandle)
+		{
+			// The library name need to be set for us to load it. I.e. someone needs to
+			// call SetName before calling GetEntry. Normally this is done by the
+			// DelayCall class below.
+			check(LibName != nullptr);
+			DllHandle = LibraryLoader::Ref().LoadDLL(ANSI_TO_TCHAR(LibName));
+		}
+		if (DllHandle)
+		{
+			check(Name != nullptr);
+			return FPlatformProcess::GetDllExport(DllHandle, ANSI_TO_TCHAR(Name));
+		}
+		return nullptr;
+	}
+
+private:
+	const char * LibName = nullptr;
+	void * DllHandle = nullptr;
+};
+
+#if MLSDK_API_USE_STUBS
+// Special case for void return type
+template<typename ReturnType> inline typename std::enable_if<std::is_void<ReturnType>::value, ReturnType>::type DefaultReturn()
+{
+}
+// Special case for all pointer return types
+template<typename ReturnType> inline typename std::enable_if<std::is_pointer<ReturnType>::value, ReturnType>::type DefaultReturn()
+{
+	return nullptr;
+}
+// Value type cases
+template<typename ReturnType> inline ReturnType DefaultValue()
+{
+	ReturnType returnVal;
+	FMemory::Memzero(&returnVal, sizeof(ReturnType));
+	return returnVal;
+}
+template<> inline MLResult DefaultValue()
+{
+	return MLResult_NotImplemented;
+}
+template<typename ReturnType> inline typename std::enable_if<!std::is_void<ReturnType>::value && !std::is_pointer<ReturnType>::value, ReturnType>::type DefaultReturn()
+{
+	return DefaultValue<ReturnType>();
+}
+#endif // MLSDK_API_USE_STUBS
+
+/** This is a single delay loaded entry value. The class in keyed on both the library and function,
+	as types. When first created it will try and load the pointer to the named global value. */
+template <typename LibKey, typename Key, typename T>
+class DelayValue
+{
+public:
+	DelayValue(const char* LibName, const char* EntryName)
+	{
+		if (!Value) 
+		{
+			Library<LibKey>::Ref().SetName(LibName);
+			Value = static_cast<T*>(Library<LibKey>::Ref().GetEntry(EntryName));
+		}
+	}
+
+	T Get() 
+	{
+		return Value ? *Value : DefaultReturn<T>();
+	}
+	
+private:
+	static T* Value;
+
+};
+
+template <typename LibKey, typename Key, typename T>
+T* DelayValue<LibKey, Key, T>::Value = nullptr;
+
+/** This is a single delay loaded entry call. The class in keyed on both the library and function,
+	as types. When first used as a function it will attempt to retrieve the foreign entry and
+	call it. Onward the retrieved entry is called directly. */
+template <typename LibKey, typename Key, typename Result, typename... Args>
+class DelayCall
+{
+public:
+	// On construction we use the given LibName to set the library name of that singleton.
+	DelayCall(const char *LibName, const char *EntryName)
+		: EntryName(EntryName)
+	{
+		Library<LibKey>::Ref().SetName(LibName);
+		Self() = this;
+	}
+
+	Result operator()(Args... args)
+	{
+#if MLSDK_API_USE_STUBS
+		if (*Call == nullptr)
+		{
+			return DefaultReturn<Result>();
+		}
+#endif // MLSDK_API_USE_STUBS
+		return (*Call)(args...);
+	}
+
+private:
+	typedef Result (*CallPointer)(Args...);
+
+	const char *EntryName = nullptr;
+	CallPointer Call = &LoadAndCall;
+
+	static DelayCall * & Self()
+	{
+		static DelayCall * DelayCallInstance = nullptr;
+		return DelayCallInstance;
+	}
+
+	// This is the default for a call. After it's called the call destination becomes the
+	// call entry in the loaded library. Which bypasses this call to avoid as much overhead
+	// as possible.
+	static Result LoadAndCall(Args... args)
+	{
+		Self()->Call = (CallPointer)(Library<LibKey>::Ref().GetEntry(Self()->EntryName));
+#if MLSDK_API_USE_STUBS
+		if (Self()->Call == nullptr)
+		{
+			return DefaultReturn<Result>();
+		}
+#endif // MLSDK_API_USE_STUBS
+		return (*Self()->Call)(args...);
+	}
+};
+
+} // namespace MLSDK_API
+
+#if defined(MLSDK_API_NO_DEPRECATION_WARNING)
+#define MLSDK_API_DEPRECATED_MSG(msg)
+#define MLSDK_API_DEPRECATED
+#else
+#if defined(_MSC_VER) && _MSC_VER
+#define MLSDK_API_DEPRECATED_MSG(msg) __declspec(deprecated(msg))
+#define MLSDK_API_DEPRECATED __declspec(deprecated)
+#else
+#define MLSDK_API_DEPRECATED_MSG(msg) __attribute__((deprecated(msg)))
+#define MLSDK_API_DEPRECATED __attribute__((deprecated))
+#endif
+#endif
diff --git a/Engine/Source/Runtime/Core/Public/Lumin/CAPIShims/LuminAPIInput.h b/Engine/Source/Runtime/Core/Public/Lumin/CAPIShims/LuminAPIInput.h
new file mode 100644
--- /dev/null
+++ b/Engine/Source/Runtime/Core/Public/Lumin/CAPIShims/LuminAPIInput.h
@@ -0,0 +1,37 @@
+// Copyright 1998-2019 Epic Games, Inc. All Rights Reserved.
+
+#pragma once
+
+#if !defined(WITH_MLSDK) || WITH_MLSDK
+
+#include "LuminAPI.h"
+
+LUMIN_THIRD_PARTY_INCLUDES_START
+#include <ml_input.h>
+LUMIN_THIRD_PARTY_INCLUDES_END
+
+namespace MLSDK_API
+{
+
+CREATE_FUNCTION_SHIM(ml_input, MLResult, MLInputCreate)
+#define MLInputCreate ::MLSDK_API::MLInputCreateShim
+CREATE_FUNCTION_SHIM(ml_input, MLResult, MLInputSetControllerCallbacks)
+#define MLInputSetControllerCallbacks ::MLSDK_API::MLInputSetControllerCallbacksShim
+CREATE_FUNCTION_SHIM(ml_input, MLResult, MLInputGetControllerState)
+#define MLInputGetControllerState ::MLSDK_API::MLInputGetControllerStateShim
+CREATE_FUNCTION_SHIM(ml_input, MLResult, MLInputStartControllerFeedbackPatternVibe)
+#define MLInputStartControllerFeedbackPatternVibe ::MLSDK_API::MLInputStartControllerFeedbackPatternVibeShim
+CREATE_FUNCTION_SHIM(ml_input, MLResult, MLInputStartControllerFeedbackPatternLED)
+#define MLInputStartControllerFeedbackPatternLED ::MLSDK_API::MLInputStartControllerFeedbackPatternLEDShim
+CREATE_FUNCTION_SHIM(ml_input, MLResult, MLInputStartControllerFeedbackPatternEffectLED)
+#define MLInputStartControllerFeedbackPatternEffectLED ::MLSDK_API::MLInputStartControllerFeedbackPatternEffectLEDShim
+CREATE_FUNCTION_SHIM(ml_input, MLResult, MLInputSetKeyboardCallbacks)
+#define MLInputSetKeyboardCallbacks ::MLSDK_API::MLInputSetKeyboardCallbacksShim
+CREATE_FUNCTION_SHIM(ml_input, MLResult, MLInputGetKeyboardState)
+#define MLInputGetKeyboardState ::MLSDK_API::MLInputGetKeyboardStateShim
+CREATE_FUNCTION_SHIM(ml_input, MLResult, MLInputDestroy)
+#define MLInputDestroy ::MLSDK_API::MLInputDestroyShim
+
+}
+
+#endif // !defined(WITH_MLSDK) || WITH_MLSDK
diff --git a/Engine/Source/Runtime/Core/Public/Lumin/CAPIShims/LuminAPIKeyCodes.h b/Engine/Source/Runtime/Core/Public/Lumin/CAPIShims/LuminAPIKeyCodes.h
new file mode 100644
--- /dev/null
+++ b/Engine/Source/Runtime/Core/Public/Lumin/CAPIShims/LuminAPIKeyCodes.h
@@ -0,0 +1,13 @@
+// Copyright 1998-2019 Epic Games, Inc. All Rights Reserved.
+
+#pragma once
+
+#if !defined(WITH_MLSDK) || WITH_MLSDK
+
+#include "LuminAPI.h"
+
+LUMIN_THIRD_PARTY_INCLUDES_START
+#include <ml_keycodes.h>
+LUMIN_THIRD_PARTY_INCLUDES_END
+
+#endif // !defined(WITH_MLSDK) || WITH_MLSDK
diff --git a/Engine/Source/Runtime/Core/Public/Lumin/CAPIShims/LuminAPILifecycle.h b/Engine/Source/Runtime/Core/Public/Lumin/CAPIShims/LuminAPILifecycle.h
new file mode 100644
--- /dev/null
+++ b/Engine/Source/Runtime/Core/Public/Lumin/CAPIShims/LuminAPILifecycle.h
@@ -0,0 +1,41 @@
+// Copyright 1998-2019 Epic Games, Inc. All Rights Reserved.
+
+#pragma once
+
+#if !defined(WITH_MLSDK) || WITH_MLSDK
+
+#include "LuminAPI.h"
+
+LUMIN_THIRD_PARTY_INCLUDES_START
+#include <ml_lifecycle.h>
+LUMIN_THIRD_PARTY_INCLUDES_END
+
+namespace MLSDK_API
+{
+
+CREATE_FUNCTION_SHIM(ml_lifecycle, MLResult, MLLifecycleInitEx)
+#define MLLifecycleInitEx ::MLSDK_API::MLLifecycleInitExShim
+CREATE_FUNCTION_SHIM(ml_lifecycle, MLResult, MLLifecycleGetSelfInfo)
+#define MLLifecycleGetSelfInfo ::MLSDK_API::MLLifecycleGetSelfInfoShim
+CREATE_FUNCTION_SHIM(ml_lifecycle, MLResult, MLLifecycleFreeSelfInfo)
+#define MLLifecycleFreeSelfInfo ::MLSDK_API::MLLifecycleFreeSelfInfoShim
+CREATE_FUNCTION_SHIM(ml_lifecycle, MLResult, MLLifecycleGetInitArgList)
+#define MLLifecycleGetInitArgList ::MLSDK_API::MLLifecycleGetInitArgListShim
+CREATE_FUNCTION_SHIM(ml_lifecycle, MLResult, MLLifecycleGetInitArgListLength)
+#define MLLifecycleGetInitArgListLength ::MLSDK_API::MLLifecycleGetInitArgListLengthShim
+CREATE_FUNCTION_SHIM(ml_lifecycle, MLResult, MLLifecycleGetInitArgByIndex)
+#define MLLifecycleGetInitArgByIndex ::MLSDK_API::MLLifecycleGetInitArgByIndexShim
+CREATE_FUNCTION_SHIM(ml_lifecycle, MLResult, MLLifecycleGetInitArgUri)
+#define MLLifecycleGetInitArgUri ::MLSDK_API::MLLifecycleGetInitArgUriShim
+CREATE_FUNCTION_SHIM(ml_lifecycle, MLResult, MLLifecycleGetFileInfoListLength)
+#define MLLifecycleGetFileInfoListLength ::MLSDK_API::MLLifecycleGetFileInfoListLengthShim
+CREATE_FUNCTION_SHIM(ml_lifecycle, MLResult, MLLifecycleGetFileInfoByIndex)
+#define MLLifecycleGetFileInfoByIndex ::MLSDK_API::MLLifecycleGetFileInfoByIndexShim
+CREATE_FUNCTION_SHIM(ml_lifecycle, MLResult, MLLifecycleFreeInitArgList)
+#define MLLifecycleFreeInitArgList ::MLSDK_API::MLLifecycleFreeInitArgListShim
+CREATE_FUNCTION_SHIM(ml_lifecycle, MLResult, MLLifecycleSetReadyIndication)
+#define MLLifecycleSetReadyIndication ::MLSDK_API::MLLifecycleSetReadyIndicationShim
+
+}
+
+#endif // !defined(WITH_MLSDK) || WITH_MLSDK
diff --git a/Engine/Source/Runtime/Core/Public/Lumin/CAPIShims/LuminAPILocale.h b/Engine/Source/Runtime/Core/Public/Lumin/CAPIShims/LuminAPILocale.h
new file mode 100644
--- /dev/null
+++ b/Engine/Source/Runtime/Core/Public/Lumin/CAPIShims/LuminAPILocale.h
@@ -0,0 +1,23 @@
+// Copyright 1998-2019 Epic Games, Inc. All Rights Reserved.
+
+#pragma once
+
+#if !defined(WITH_MLSDK) || WITH_MLSDK
+
+#include "LuminAPI.h"
+
+LUMIN_THIRD_PARTY_INCLUDES_START
+#include <ml_locale.h>
+LUMIN_THIRD_PARTY_INCLUDES_END
+
+namespace MLSDK_API
+{
+
+	CREATE_FUNCTION_SHIM(ml_locale, MLResult, MLLocaleGetSystemLanguage)
+	#define MLLocaleGetSystemLanguage ::MLSDK_API::MLLocaleGetSystemLanguageShim
+	CREATE_FUNCTION_SHIM(ml_locale, MLResult, MLLocaleGetSystemCountry)
+	#define MLLocaleGetSystemCountry ::MLSDK_API::MLLocaleGetSystemCountryShim
+
+}
+
+#endif // !defined(WITH_MLSDK) || WITH_MLSDK
diff --git a/Engine/Source/Runtime/Core/Public/Lumin/CAPIShims/LuminAPILogging.h b/Engine/Source/Runtime/Core/Public/Lumin/CAPIShims/LuminAPILogging.h
new file mode 100644
--- /dev/null
+++ b/Engine/Source/Runtime/Core/Public/Lumin/CAPIShims/LuminAPILogging.h
@@ -0,0 +1,27 @@
+// Copyright 1998-2019 Epic Games, Inc. All Rights Reserved.
+
+#pragma once
+
+#if !defined(WITH_MLSDK) || WITH_MLSDK
+
+#include "LuminAPI.h"
+
+LUMIN_THIRD_PARTY_INCLUDES_START
+#include <ml_logging.h>
+LUMIN_THIRD_PARTY_INCLUDES_END
+
+namespace MLSDK_API
+{
+
+CREATE_FUNCTION_SHIM(ml_ext_logging, void, MLLoggingEnableLogLevel)
+#define MLLoggingEnableLogLevel ::MLSDK_API::MLLoggingEnableLogLevelShim
+CREATE_FUNCTION_SHIM(ml_ext_logging, bool, MLLoggingLogLevelIsEnabled)
+#define MLLoggingLogLevelIsEnabled ::MLSDK_API::MLLoggingLogLevelIsEnabledShim
+CREATE_FUNCTION_SHIM(ml_ext_logging, void, MLLoggingLogVargs)
+#define MLLoggingLogVargs ::MLSDK_API::MLLoggingLogVargsShim
+CREATE_FUNCTION_SHIM(ml_ext_logging, void, MLLoggingLog)
+#define MLLoggingLog ::MLSDK_API::MLLoggingLogShim
+
+}
+
+#endif // !defined(WITH_MLSDK) || WITH_MLSDK
diff --git a/Engine/Source/Runtime/Core/Public/Lumin/CAPIShims/LuminAPIPlatform.h b/Engine/Source/Runtime/Core/Public/Lumin/CAPIShims/LuminAPIPlatform.h
new file mode 100644
--- /dev/null
+++ b/Engine/Source/Runtime/Core/Public/Lumin/CAPIShims/LuminAPIPlatform.h
@@ -0,0 +1,21 @@
+// Copyright 1998-2019 Epic Games, Inc. All Rights Reserved.
+
+#pragma once
+
+#if !defined(WITH_MLSDK) || WITH_MLSDK
+
+#include "LuminAPI.h"
+
+LUMIN_THIRD_PARTY_INCLUDES_START
+#include <ml_platform.h>
+LUMIN_THIRD_PARTY_INCLUDES_END
+
+namespace MLSDK_API
+{
+
+CREATE_FUNCTION_SHIM(ml_platform, MLResult, MLPlatformGetAPILevel)
+#define MLPlatformGetAPILevel ::MLSDK_API::MLPlatformGetAPILevelShim
+
+}
+
+#endif // !defined(WITH_MLSDK) || WITH_MLSDK
diff --git a/Engine/Source/Runtime/Core/Public/Lumin/CAPIShims/LuminAPISharedFile.h b/Engine/Source/Runtime/Core/Public/Lumin/CAPIShims/LuminAPISharedFile.h
new file mode 100644
--- /dev/null
+++ b/Engine/Source/Runtime/Core/Public/Lumin/CAPIShims/LuminAPISharedFile.h
@@ -0,0 +1,37 @@
+// Copyright 1998-2019 Epic Games, Inc. All Rights Reserved.
+
+#pragma once
+
+#if !defined(WITH_MLSDK) || WITH_MLSDK
+
+#include "LuminAPI.h"
+
+LUMIN_THIRD_PARTY_INCLUDES_START
+#include <ml_sharedfile.h>
+LUMIN_THIRD_PARTY_INCLUDES_END
+
+namespace MLSDK_API
+{
+
+CREATE_FUNCTION_SHIM(ml_sharedfile, MLResult, MLSharedFileRead)
+#define MLSharedFileRead ::MLSDK_API::MLSharedFileReadShim
+CREATE_FUNCTION_SHIM(ml_sharedfile, MLResult, MLSharedFileWrite)
+#define MLSharedFileWrite ::MLSDK_API::MLSharedFileWriteShim
+CREATE_FUNCTION_SHIM(ml_sharedfile, MLResult, MLSharedFileListAccessibleFiles)
+#define MLSharedFileListAccessibleFiles ::MLSDK_API::MLSharedFileListAccessibleFilesShim
+CREATE_FUNCTION_SHIM(ml_sharedfile, MLResult, MLSharedFilePick)
+#define MLSharedFilePick ::MLSDK_API::MLSharedFilePickShim
+CREATE_FUNCTION_SHIM(ml_sharedfile, MLResult, MLSharedFileGetListLength)
+#define MLSharedFileGetListLength ::MLSDK_API::MLSharedFileGetListLengthShim
+CREATE_FUNCTION_SHIM(ml_sharedfile, MLResult, MLSharedFileGetMLFileInfoByIndex)
+#define MLSharedFileGetMLFileInfoByIndex ::MLSDK_API::MLSharedFileGetMLFileInfoByIndexShim
+CREATE_FUNCTION_SHIM(ml_sharedfile, MLResult, MLSharedFileGetErrorCode)
+#define MLSharedFileGetErrorCode ::MLSDK_API::MLSharedFileGetErrorCodeShim
+CREATE_FUNCTION_SHIM(ml_sharedfile, MLResult, MLSharedFileListRelease)
+#define MLSharedFileListRelease ::MLSDK_API::MLSharedFileListReleaseShim
+CREATE_FUNCTION_SHIM(ml_sharedfile, const char*, MLSharedFileGetResultString)
+#define MLSharedFileGetResultString ::MLSDK_API::MLSharedFileGetResultStringShim
+
+}
+
+#endif // !defined(WITH_MLSDK) || WITH_MLSDK
diff --git a/Engine/Source/Runtime/Core/Public/Lumin/CAPIShims/LuminAPITypes.h b/Engine/Source/Runtime/Core/Public/Lumin/CAPIShims/LuminAPITypes.h
new file mode 100644
--- /dev/null
+++ b/Engine/Source/Runtime/Core/Public/Lumin/CAPIShims/LuminAPITypes.h
@@ -0,0 +1,13 @@
+// Copyright 1998-2019 Epic Games, Inc. All Rights Reserved.
+
+#pragma once
+
+#if !defined(WITH_MLSDK) || WITH_MLSDK
+
+#include "LuminAPI.h"
+
+LUMIN_THIRD_PARTY_INCLUDES_START
+#include <ml_types.h>
+LUMIN_THIRD_PARTY_INCLUDES_END
+
+#endif // !defined(WITH_MLSDK) || WITH_MLSDK
diff --git a/Engine/Source/Runtime/Core/Public/Lumin/CAPIShims/LuminAPIVersion.h b/Engine/Source/Runtime/Core/Public/Lumin/CAPIShims/LuminAPIVersion.h
new file mode 100644
--- /dev/null
+++ b/Engine/Source/Runtime/Core/Public/Lumin/CAPIShims/LuminAPIVersion.h
@@ -0,0 +1,13 @@
+// Copyright 1998-2019 Epic Games, Inc. All Rights Reserved.
+
+#pragma once
+
+#if !defined(WITH_MLSDK) || WITH_MLSDK
+
+#include "LuminAPI.h"
+
+LUMIN_THIRD_PARTY_INCLUDES_START
+#include <ml_version.h>
+LUMIN_THIRD_PARTY_INCLUDES_END
+
+#endif // !defined(WITH_MLSDK) || WITH_MLSDK
diff --git a/Engine/Source/Runtime/Core/Public/Lumin/LuminPlatformDelegates.h b/Engine/Source/Runtime/Core/Public/Lumin/LuminPlatformDelegates.h
new file mode 100644
--- /dev/null
+++ b/Engine/Source/Runtime/Core/Public/Lumin/LuminPlatformDelegates.h
@@ -0,0 +1,25 @@
+// Copyright 1998-2018 Epic Games, Inc. All Rights Reserved.
+
+#pragma once
+
+#include "CoreTypes.h"
+#include "Containers/Array.h"
+#include "Containers/UnrealString.h"
+#include "Delegates/Delegate.h"
+#include "Lumin/LuminPlatformFile.h"
+
+class CORE_API FLuminDelegates
+{
+public:
+	DECLARE_MULTICAST_DELEGATE_TwoParams(FLuminAppStartupArgumentsDelegate, const TArray<FString>&, const TArray<FLuminFileInfo>&);
+
+	// Called with arguments passed to the application on statup, perhaps meta data passed on by another application which launched this one.
+	// TODO: make this private
+	static FLuminAppStartupArgumentsDelegate LuminAppReceivedStartupArgumentsDelegate;
+
+	DECLARE_MULTICAST_DELEGATE(FLuminApplicationLifetimeDelegate);
+
+	static FLuminApplicationLifetimeDelegate DeviceHasReactivatedDelegate;
+	static FLuminApplicationLifetimeDelegate DeviceWillEnterRealityModeDelegate;
+	static FLuminApplicationLifetimeDelegate DeviceWillGoInStandbyDelegate;
+};
diff --git a/Engine/Source/Runtime/Core/Public/Lumin/LuminPlatformFile.h b/Engine/Source/Runtime/Core/Public/Lumin/LuminPlatformFile.h
--- a/Engine/Source/Runtime/Core/Public/Lumin/LuminPlatformFile.h
+++ b/Engine/Source/Runtime/Core/Public/Lumin/LuminPlatformFile.h
@@ -8,6 +8,16 @@ Lumin platform File functions
 #pragma once
 #include "GenericPlatform/GenericPlatformFile.h"
 
+struct CORE_API FLuminFileInfo
+{
+public:
+	FLuminFileInfo();
+
+	FString MimeType;
+	FString FileName;
+	IFileHandle* FileHandle;
+};
+
 /**
  * File I/O implementation
 **/
@@ -70,6 +80,23 @@ public:
 
 	FString ConvertToLuminPath(const FString& Filename, bool bForWrite) const;
 
+	/**
+		Return a IFileHandle pointer to read the user shared file (file descriptor recieved from ml_sharedfile api).
+		@param FileName Name of the shared file to read.
+		@return IFileHandle pointer to read the file. If the application does not have user permission to access the file, nullptr will be returned.
+	*/
+	IFileHandle* SharedFileOpenRead(const TCHAR* Filename);
+
+	/**
+		Return a IFileHandle pointer to write the user shared file (file descriptor recieved from ml_sharedfile api)..
+		@param FileName Name of the shared file to write to. Needs to just be a file name, cannot be a path.
+		@return IFileHandle pointer to read the file. If the application does not have user permission to access the file, nullptr will be returned.
+	*/
+	IFileHandle* SharedFileOpenWrite(const TCHAR* Filename);
+
+	IFileHandle* GetFileHandleForMLFileInfo(const void* FileInfo);
+	bool SetMLFileInfoFD(const IFileHandle* FileHandle, void* FileInfo);
+
 protected:
 	bool IterateDirectoryCommon(const TCHAR* Directory, const TFunctionRef<bool(struct dirent*)>& Visitor);
 	bool bIsSandboxEnabled = true;
diff --git a/Engine/Source/Runtime/Core/Public/Lumin/LuminPlatformMisc.h b/Engine/Source/Runtime/Core/Public/Lumin/LuminPlatformMisc.h
--- a/Engine/Source/Runtime/Core/Public/Lumin/LuminPlatformMisc.h
+++ b/Engine/Source/Runtime/Core/Public/Lumin/LuminPlatformMisc.h
@@ -57,6 +57,9 @@ struct CORE_API FLuminPlatformMisc : public FAndroidMisc
 	static const FString& GetApplicationApplicationPackageName();
 	static const FString& GetApplicationComponentName();
 
+	static FString GetDefaultLocale();
+	static FString GetDefaultLanguage();
+
 private:
 	static void InitApplicationPaths();
 
diff --git a/Engine/Source/Runtime/Core/Public/Lumin/LuminPlatformProperties.h b/Engine/Source/Runtime/Core/Public/Lumin/LuminPlatformProperties.h
--- a/Engine/Source/Runtime/Core/Public/Lumin/LuminPlatformProperties.h
+++ b/Engine/Source/Runtime/Core/Public/Lumin/LuminPlatformProperties.h
@@ -1,7 +1,7 @@
 // Copyright 1998-2019 Epic Games, Inc. All Rights Reserved.
 
 /*================================================================================
-	LinuxPlatformProperties.h - Basic static properties of a platform 
+	LuminPlatformProperties.h - Basic static properties of a platform 
 	These are shared between:
 		the runtime platform - via FPlatformProperties
 		the target platforms - via ITargetPlatform
@@ -39,6 +39,11 @@ struct FLuminPlatformProperties : public FAndroidPlatformProperties
 	{
 		return false;
 	}
+
+	static FORCEINLINE bool SupportsQuit()
+	{
+		return true;
+	}
 };
 
 #ifdef SAVED_PROPERTY_HEADER_SHOULD_DEFINE_TYPE
diff --git a/Engine/Source/Runtime/Launch/Private/Lumin/LaunchLumin.cpp b/Engine/Source/Runtime/Launch/Private/Lumin/LaunchLumin.cpp
--- a/Engine/Source/Runtime/Launch/Private/Lumin/LaunchLumin.cpp
+++ b/Engine/Source/Runtime/Launch/Private/Lumin/LaunchLumin.cpp
@@ -7,7 +7,7 @@
 #include "Misc/CommandLine.h"
 #include "HAL/PlatformProcess.h"
 #include "Lumin/LuminPlatformMisc.h"
-#include <ml_lifecycle.h>
+#include "Lumin/CAPIShims/LuminAPI.h"
 
 #include <locale.h>
 #include <sys/resource.h>
diff --git a/Engine/Source/Runtime/Lumin/LuminRuntimeSettings/Classes/LuminRuntimeSettings.h b/Engine/Source/Runtime/Lumin/LuminRuntimeSettings/Classes/LuminRuntimeSettings.h
--- a/Engine/Source/Runtime/Lumin/LuminRuntimeSettings/Classes/LuminRuntimeSettings.h
+++ b/Engine/Source/Runtime/Lumin/LuminRuntimeSettings/Classes/LuminRuntimeSettings.h
@@ -8,6 +8,63 @@
 #include "Engine/EngineTypes.h"
 #include "LuminRuntimeSettings.generated.h"
 
+UENUM(BlueprintType)
+enum class ELuminComponentSubElementType : uint8
+{
+	FileExtension,
+	MimeType,
+	Mode,
+	MusicAttribute,
+	Schema
+};
+
+UENUM(BlueprintType)
+enum class ELuminComponentType : uint8
+{
+	Universe,
+	Fullscreen,
+	SearchProvider,
+	MusicService,
+	Screens,
+	ScreensImmersive,
+	Console,
+	SystemUI,
+};
+
+USTRUCT(BlueprintType)
+struct FLuminComponentSubElement
+{
+	GENERATED_BODY()
+
+	UPROPERTY(GlobalConfig, EditAnywhere, Category = "Runtime", Meta = (DisplayName = "Component sub-node type"))
+	ELuminComponentSubElementType ElementType;
+
+	UPROPERTY(GlobalConfig, EditAnywhere, Category = "Runtime", Meta = (DisplayName = "Component sub-node value"))
+	FString Value;
+};
+
+USTRUCT(BlueprintType)
+struct FLuminComponentElement
+{
+	GENERATED_BODY()
+
+	UPROPERTY(GlobalConfig, EditAnywhere, Category = "Runtime", Meta = (DisplayName = "Name"))
+	FString Name;
+
+	UPROPERTY(GlobalConfig, EditAnywhere, Category = "Runtime", Meta = (DisplayName = "Visiable name"))
+	FString VisibleName;
+
+	/** Name of the executable for this component. This binary should be packaged into the bin folder of the mpk. Refer to ResonanceAudio_LPL.xml for an example. */
+	UPROPERTY(GlobalConfig, EditAnywhere, Category = "Runtime", Meta = (DisplayName = "Executable name"))
+	FString ExecutableName;
+
+	UPROPERTY(GlobalConfig, EditAnywhere, Category = "Runtime", Meta = (DisplayName = "Component type"))
+	ELuminComponentType ComponentType;
+
+	UPROPERTY(GlobalConfig, EditAnywhere, Category = "Runtime", Meta = (DisplayName = "Extra sub-elements"))
+	TArray<FLuminComponentSubElement> ExtraComponentSubElements;
+};
+
 UENUM(BlueprintType)
 enum class ELuminFrameTimingHint : uint8
 {
@@ -25,7 +82,6 @@ UENUM(BlueprintType)
 enum class ELuminPrivilege : uint8
 {
 	Invalid,
-	AudioRecognizer,
 	BatteryInfo,
 	CameraCapture,
 	WorldReconstruction,
@@ -54,6 +110,11 @@ enum class ELuminPrivilege : uint8
 	ScreensProvider,
 	GesturesSubscribe,
 	GesturesConfig,
+	AddressBookRead,
+	AddressBookWrite,
+	CoarseLocation,
+	HandMesh,
+	WifiStatusRead,
 };
 
 /**
@@ -69,12 +130,6 @@ class LUMINRUNTIMESETTINGS_API ULuminRuntimeSettings : public UObject
 public:
 	GENERATED_BODY()
 
-#if WITH_EDITOR
-	virtual bool CanEditChange(const UProperty* InProperty) const override;
-
-	virtual void PostEditChangeProperty(struct FPropertyChangedEvent& PropertyChangedEvent) override;
-#endif // WITH_EDITOR
-
 	/** The official name of the project. Note: Must have at least 2 sections separated by a period and be unique! */
 	UPROPERTY(GlobalConfig, EditAnywhere, Category = "MPK Packaging", Meta = (DisplayName = "Magic Leap Package Name ('com.Company.Project', [PROJECT] is replaced with project name)"))
 	FString PackageName;
@@ -95,17 +150,13 @@ public:
 	UPROPERTY(GlobalConfig, EditAnywhere, Category = "Runtime", Meta = (DisplayName = "Protected Content"))
 	bool bProtectedContent;
 
-	/** If checked, use Mobile Rendering. Otherwise, use Desktop Rendering [FOR FULL SOURCE GAMES ONLY]. */
-	UPROPERTY(GlobalConfig, EditAnywhere, Category = "Build", Meta = (DisplayName = "Use Mobile Rendering.  Otherwise, Desktop Rendering"))
+	/** If checked, use Mobile Rendering. Otherwise, use Desktop Rendering. */
+	UPROPERTY(GlobalConfig, EditAnywhere, Category = "Build", Meta = (DisplayName = "Use Mobile Rendering (otherwise, Desktop Rendering)"))
 	bool bUseMobileRendering;
 
 	UPROPERTY(GlobalConfig, Meta = (DisplayName = "Use Vulkan (otherwise, OpenGL)"))
 	bool bUseVulkan;
 
-	/** Enable support for NVIDIA Tegra Graphics Debugger [FOR FULL SOURCE GAMES ONLY]. */
-	UPROPERTY(GlobalConfig, EditAnywhere, Category = "Build", Meta = (DisplayName = "Support NVIDIA Tegra Graphics Debugger"))
-	bool bBuildWithNvTegraGfxDebugger;
-
 	/** Certificate File used to sign builds for distribution. */
 	UPROPERTY(GlobalConfig, EditAnywhere, Category = "Distribution Signing", Meta = (DisplayName = "Certificate File Path"))
 	FFilePath Certificate;
@@ -130,13 +181,13 @@ public:
 	UPROPERTY(GlobalConfig, EditAnywhere, Category = "Advanced MPK Packaging", Meta = (DisplayName = "App Privileges"))
 	TArray<ELuminPrivilege> AppPrivileges;
 
-	/** Extra nodes under the <application> node. */
-	UPROPERTY(GlobalConfig, EditAnywhere, Category = "Advanced MPK Packaging", Meta = (DisplayName = "Extra nodes under the <application> node"))
-	TArray<FString> ExtraApplicationNodes;
-
 	/** Extra nodes under the <component> node like <mime-type>, <schema> etc. */
 	UPROPERTY(GlobalConfig, EditAnywhere, Category = "Advanced MPK Packaging", Meta = (DisplayName = "Extra nodes under the <component> node"))
-	TArray<FString> ExtraComponentNodes;
+	TArray<FLuminComponentSubElement> ExtraComponentSubElements;
+
+	/** Extra component elements. */
+	UPROPERTY(GlobalConfig, EditAnywhere, Category = "Advanced MPK Packaging", Meta = (DisplayName = "Extra <component> elements."))
+	TArray<FLuminComponentElement> ExtraComponentElements;
 
 	/** Which of the currently enabled spatialization plugins to use on Lumin. */
 	UPROPERTY(config, EditAnywhere, Category = "Audio")
@@ -153,4 +204,15 @@ public:
 	// Strip debug symbols from packaged builds even if they aren't shipping builds.
 	UPROPERTY(GlobalConfig, EditAnywhere, Category = AdvancedBuild, meta = (DisplayName = "Strip debug symbols from packaged builds even if they aren't shipping builds"))
 	bool bRemoveDebugInfo;
+
+	/** Folder containing the libraries required for vulkan validation layers. Can be found under %NDKROOT%/sources/third_party/vulkan/src/build-android/jniLibs/arm64-v8a */
+	UPROPERTY(GlobalConfig, EditAnywhere, Category = AdvancedBuild, Meta = (DisplayName = "Vulkan Validation Layer libs"))
+	FDirectoryPath VulkanValidationLayerLibs;
+
+	/** Render frame vignette. */
+	UPROPERTY(GlobalConfig, EditAnywhere, Category = "Rendering", Meta = (DisplayName = "Render frame vignette"))
+	bool bFrameVignette;
+
+	TArray<FString> ExtraApplicationNodes_DEPRECATED;
+	TArray<FString> ExtraComponentNodes_DEPRECATED;
 };
diff --git a/Engine/Source/Runtime/Lumin/LuminRuntimeSettings/Private/LuminRuntimeSettings.cpp b/Engine/Source/Runtime/Lumin/LuminRuntimeSettings/Private/LuminRuntimeSettings.cpp
--- a/Engine/Source/Runtime/Lumin/LuminRuntimeSettings/Private/LuminRuntimeSettings.cpp
+++ b/Engine/Source/Runtime/Lumin/LuminRuntimeSettings/Private/LuminRuntimeSettings.cpp
@@ -4,25 +4,3 @@
 #include "Misc/ConfigCacheIni.h"
 #include "CoreGlobals.h"
 #include "Misc/App.h"
-
-#if WITH_EDITOR
-
-bool ULuminRuntimeSettings::CanEditChange(const UProperty* InProperty) const
-{
-	const bool ParentVal = Super::CanEditChange(InProperty);
-
-	if ((InProperty->GetFName() == GET_MEMBER_NAME_CHECKED(ULuminRuntimeSettings, bBuildWithNvTegraGfxDebugger)) ||
-		(InProperty->GetFName() == GET_MEMBER_NAME_CHECKED(ULuminRuntimeSettings, bUseMobileRendering)))
-	{
-		return !FApp::IsEngineInstalled();
-	}
-
-	return ParentVal;
-}
-
-void ULuminRuntimeSettings::PostEditChangeProperty(struct FPropertyChangedEvent& PropertyChangedEvent)
-{
-	GConfig->Flush(1);
-}
-
-#endif
\ No newline at end of file
diff --git a/Engine/Source/Runtime/OpenGLDrv/Private/Lumin/LuminOpenGL.cpp b/Engine/Source/Runtime/OpenGLDrv/Private/Lumin/LuminOpenGL.cpp
--- a/Engine/Source/Runtime/OpenGLDrv/Private/Lumin/LuminOpenGL.cpp
+++ b/Engine/Source/Runtime/OpenGLDrv/Private/Lumin/LuminOpenGL.cpp
@@ -16,6 +16,11 @@ namespace GL_EXT
 	PFNEGLDESTROYSYNCKHRPROC eglDestroySyncKHR = NULL;
 	PFNEGLCLIENTWAITSYNCKHRPROC eglClientWaitSyncKHR = NULL;
 
+	PFNEGLCREATEIMAGEKHRPROC eglCreateImageKHR = NULL;
+	PFNEGLDESTROYIMAGEKHRPROC eglDestroyImageKHR = NULL;
+
+	PFNGLEGLIMAGETARGETTEXTURE2DOESPROC glEGLImageTargetTexture2DOES = NULL;
+
 	// Occlusion Queries
 	PFNGLGENQUERIESEXTPROC 					glGenQueriesEXT = NULL;
 	PFNGLDELETEQUERIESEXTPROC 				glDeleteQueriesEXT = NULL;
@@ -239,6 +244,11 @@ void FPlatformOpenGLDevice::LoadEXT()
 	eglDestroySyncKHR = (PFNEGLDESTROYSYNCKHRPROC)((void*)eglGetProcAddress("eglDestroySyncKHR"));
 	eglClientWaitSyncKHR = (PFNEGLCLIENTWAITSYNCKHRPROC)((void*)eglGetProcAddress("eglClientWaitSyncKHR"));
 
+	eglCreateImageKHR = (PFNEGLCREATEIMAGEKHRPROC)((void*)eglGetProcAddress("eglCreateImageKHR"));
+	eglDestroyImageKHR = (PFNEGLDESTROYIMAGEKHRPROC)((void*)eglGetProcAddress("eglDestroyImageKHR"));
+
+	glEGLImageTargetTexture2DOES = (PFNGLEGLIMAGETARGETTEXTURE2DOESPROC)((void*)eglGetProcAddress("glEGLImageTargetTexture2DOES"));
+
 	glDebugMessageControlKHR = (PFNGLDEBUGMESSAGECONTROLKHRPROC)((void*)eglGetProcAddress("glDebugMessageControlKHR"));
 	glDebugMessageInsertKHR = (PFNGLDEBUGMESSAGEINSERTKHRPROC)((void*)eglGetProcAddress("glDebugMessageInsertKHR"));
 	glDebugMessageCallbackKHR = (PFNGLDEBUGMESSAGECALLBACKKHRPROC)((void*)eglGetProcAddress("glDebugMessageCallbackKHR"));
diff --git a/Engine/Source/Runtime/OpenGLDrv/Private/Lumin/LuminOpenGL.h b/Engine/Source/Runtime/OpenGLDrv/Private/Lumin/LuminOpenGL.h
--- a/Engine/Source/Runtime/OpenGLDrv/Private/Lumin/LuminOpenGL.h
+++ b/Engine/Source/Runtime/OpenGLDrv/Private/Lumin/LuminOpenGL.h
@@ -1,4 +1,5 @@
 // Copyright 1998-2019 Epic Games, Inc. All Rights Reserved.
+#pragma once
 
 #if !PLATFORM_LUMINGL4
 
@@ -168,6 +169,11 @@ namespace GL_EXT
 	extern PFNEGLDESTROYSYNCKHRPROC eglDestroySyncKHR;
 	extern PFNEGLCLIENTWAITSYNCKHRPROC eglClientWaitSyncKHR;
 
+	extern PFNEGLCREATEIMAGEKHRPROC eglCreateImageKHR;
+	extern PFNEGLDESTROYIMAGEKHRPROC eglDestroyImageKHR;
+
+	extern PFNGLEGLIMAGETARGETTEXTURE2DOESPROC glEGLImageTargetTexture2DOES;
+
 	extern PFNGLREADBUFFERPROC glReadBuffer;
 }
 
diff --git a/Engine/Source/Runtime/VulkanRHI/Private/Android/VulkanAndroidPlatform.cpp b/Engine/Source/Runtime/VulkanRHI/Private/Android/VulkanAndroidPlatform.cpp
--- a/Engine/Source/Runtime/VulkanRHI/Private/Android/VulkanAndroidPlatform.cpp
+++ b/Engine/Source/Runtime/VulkanRHI/Private/Android/VulkanAndroidPlatform.cpp
@@ -9,7 +9,7 @@
 #endif
 
 //#todo-Lumin: Remove this define when it becomes untangled from Android
-#if (!defined(PLATFORM_LUMIN) && !defined(PLATFORM_LUMINGL4)) || (!PLATFORM_LUMIN && !PLATFORM_LUMINGL4)
+#if (!defined(PLATFORM_LUMIN) || (!PLATFORM_LUMIN))
 #include "VulkanAndroidPlatform.h"
 #include "../VulkanRHIPrivate.h"
 #include <dlfcn.h>
diff --git a/Engine/Source/Runtime/VulkanRHI/Private/Lumin/VulkanLuminPlatform.cpp b/Engine/Source/Runtime/VulkanRHI/Private/Lumin/VulkanLuminPlatform.cpp
--- a/Engine/Source/Runtime/VulkanRHI/Private/Lumin/VulkanLuminPlatform.cpp
+++ b/Engine/Source/Runtime/VulkanRHI/Private/Lumin/VulkanLuminPlatform.cpp
@@ -164,3 +164,52 @@ bool FVulkanLuminPlatform::RequiresMobileRenderer()
 {
 	return !FLuminPlatformMisc::ShouldUseDesktopVulkan();
 }
+
+VkBool32 FVulkanLuminPlatform::DebugReportFunction(
+	VkDebugReportFlagsEXT			MsgFlags,
+	VkDebugReportObjectTypeEXT		ObjType,
+	uint64_t						SrcObject,
+	size_t							Location,
+	int32							MsgCode,
+	const ANSICHAR*					LayerPrefix,
+	const ANSICHAR*					Msg,
+	void*							UserData)
+{
+	if (MsgFlags & VK_DEBUG_REPORT_ERROR_BIT_EXT)
+	{
+		if (!FCStringAnsi::Strcmp(LayerPrefix, "ParameterValidation"))
+		{
+			// Function called but its required extension has not been enabled.
+			if (MsgCode == 0xa)
+			{
+				// We don't want to disable all messages that fall into this category, just the ones we wont/cant fix.
+
+				// VK_EXT_debug_marker is not available on lumin unless running through the debugger.
+				if (FCStringAnsi::Strfind(Msg, "VK_EXT_debug_marker") != nullptr)
+				{
+					return VK_FALSE;
+				}
+			}
+		}
+	}
+	else if (MsgFlags & VK_DEBUG_REPORT_WARNING_BIT_EXT)
+	{
+		if (!FCStringAnsi::Strcmp(LayerPrefix, "ParameterValidation"))
+		{
+			// Struct from a private extension used.
+			if (MsgCode == 0x9e1c40d)
+			{
+				// vkCreateImage: pCreateInfo->pNext chain includes a structure with unknown VkStructureType (1000027002); 
+				// Allowed structures are [VkDedicatedAllocationImageCreateInfoNV, VkExternalMemoryImageCreateInfoKHR, VkExternalMemoryImageCreateInfoNV, VkImageFormatListCreateInfoKHR, VkImageSwapchainCreateInfoKHX].
+				// The spec valid usage text states 'Each pNext member of any structure (including this one) in the pNext chain must be either NULL or a pointer to a valid instance of
+				// VkDedicatedAllocationImageCreateInfoNV, VkExternalMemoryImageCreateInfoKHR, VkExternalMemoryImageCreateInfoNV, VkImageFormatListCreateInfoKHR, or VkImageSwapchainCreateInfoKHX'
+				// (https://www.khronos.org/registry/vulkan/specs/1.0/html/vkspec.html#VUID-VkImageCreateInfo-pNext-pNext) This warning is based on the Valid Usage documentation for version 69 of the Vulkan header.
+				// It is possible that you are using a struct from a private extension or an extension that was added to a later version of the Vulkan header, in which case your use of pCreateInfo->pNext is perfectly valid
+				// but is not guaranteed to work correctly with validation enabled.
+				return VK_FALSE;
+			}
+		}
+	}
+
+	return VK_TRUE;
+}
diff --git a/Engine/Source/Runtime/VulkanRHI/Private/Lumin/VulkanLuminPlatform.h b/Engine/Source/Runtime/VulkanRHI/Private/Lumin/VulkanLuminPlatform.h
--- a/Engine/Source/Runtime/VulkanRHI/Private/Lumin/VulkanLuminPlatform.h
+++ b/Engine/Source/Runtime/VulkanRHI/Private/Lumin/VulkanLuminPlatform.h
@@ -12,6 +12,7 @@
 #define VULKAN_SUPPORTS_NV_DIAGNOSTIC_CHECKPOINT	0
 #define VULKAN_SUPPORTS_PHYSICAL_DEVICE_PROPERTIES2	0
 #define VULKAN_SUPPORTS_DEDICATED_ALLOCATION		0
+#define VULKAN_SUPPORTS_GPU_CRASH_DUMPS				1
 
 
 #define ENUM_VK_ENTRYPOINTS_PLATFORM_BASE(EnumMacro)
@@ -63,6 +64,8 @@ public:
 
 	static bool RequiresWaitingForFrameCompletionEvent() { return false; }
 
+	static VkBool32 DebugReportFunction(VkDebugReportFlagsEXT MsgFlags, VkDebugReportObjectTypeEXT ObjType, uint64_t SrcObject, size_t Location, int32 MsgCode, const ANSICHAR* LayerPrefix, const ANSICHAR* Msg, void* UserData);
+
 protected:
 	static void* VulkanLib;
 	static bool bAttemptedLoad;
diff --git a/Engine/Source/Runtime/VulkanRHI/Private/VulkanContext.h b/Engine/Source/Runtime/VulkanRHI/Private/VulkanContext.h
--- a/Engine/Source/Runtime/VulkanRHI/Private/VulkanContext.h
+++ b/Engine/Source/Runtime/VulkanRHI/Private/VulkanContext.h
@@ -309,7 +309,7 @@ public:
 
 	inline VkImageLayout GetLayoutForDescriptor(const FVulkanSurface& Surface) const
 	{
-#if PLATFORM_ANDROID && !PLATFORM_LUMIN && !PLATFORM_LUMINGL4
+#if PLATFORM_ANDROID && !PLATFORM_LUMIN
 		// Workaround clang bug; don't use IsDepthOrStencilAspect() directly
 		VkImageAspectFlags AspectFlags = Surface.GetFullAspectMask();
 		if ((AspectFlags & VK_IMAGE_ASPECT_DEPTH_BIT) != 0 || (AspectFlags & VK_IMAGE_ASPECT_STENCIL_BIT) != 0)
diff --git a/Engine/Source/Runtime/VulkanRHI/Private/VulkanDebug.cpp b/Engine/Source/Runtime/VulkanRHI/Private/VulkanDebug.cpp
--- a/Engine/Source/Runtime/VulkanRHI/Private/VulkanDebug.cpp
+++ b/Engine/Source/Runtime/VulkanRHI/Private/VulkanDebug.cpp
@@ -120,7 +120,6 @@ static VkBool32 VKAPI_PTR DebugReportFunction(
 				return VK_FALSE;
 			}
 		}
-
 	}
 	else if (MsgFlags & VK_DEBUG_REPORT_PERFORMANCE_WARNING_BIT_EXT)
 	{
@@ -156,6 +155,13 @@ static VkBool32 VKAPI_PTR DebugReportFunction(
 		ensure(0);
 	}
 
+	const VkBool32 bPlatformPrintLog = FVulkanPlatform::DebugReportFunction(MsgFlags, ObjType, SrcObject, Location, MsgCode, LayerPrefix, Msg, UserData);
+	if (bPlatformPrintLog == VK_FALSE)
+	{
+		// Early out if platform wants to suppress this debug report.
+		return VK_FALSE;
+	}
+
 	FString LayerCode = FString::Printf(TEXT("%s%x"), ANSI_TO_TCHAR(LayerPrefix), MsgCode);
 
 	static TSet<FString> SeenCodes;
diff --git a/Engine/Source/Runtime/VulkanRHI/Private/VulkanDevice.h b/Engine/Source/Runtime/VulkanRHI/Private/VulkanDevice.h
--- a/Engine/Source/Runtime/VulkanRHI/Private/VulkanDevice.h
+++ b/Engine/Source/Runtime/VulkanRHI/Private/VulkanDevice.h
@@ -315,7 +315,7 @@ public:
 
 	void* GetCrashMarkerMappedPointer() const
 	{
-		return CrashMarker.Allocation->GetMappedPointer();
+		return (CrashMarker.Allocation != nullptr) ? CrashMarker.Allocation->GetMappedPointer() : nullptr;
 	}
 #endif
 
diff --git a/Engine/Source/Runtime/VulkanRHI/Private/VulkanGenericPlatform.h b/Engine/Source/Runtime/VulkanRHI/Private/VulkanGenericPlatform.h
--- a/Engine/Source/Runtime/VulkanRHI/Private/VulkanGenericPlatform.h
+++ b/Engine/Source/Runtime/VulkanRHI/Private/VulkanGenericPlatform.h
@@ -116,4 +116,7 @@ public:
 
 	// Gathers a list of pso cache filenames to attempt to load
 	static TArray<FString> GetPSOCacheFilenames();
+
+	// Return VK_FALSE if platform wants to suppress the given debug report from the validation layers, VK_TRUE to print it.
+	static VkBool32 DebugReportFunction(VkDebugReportFlagsEXT MsgFlags, VkDebugReportObjectTypeEXT ObjType, uint64_t SrcObject, size_t Location, int32 MsgCode, const ANSICHAR* LayerPrefix, const ANSICHAR* Msg, void* UserData) { return VK_TRUE; }
 };
diff --git a/Engine/Source/Runtime/VulkanRHI/Private/VulkanLayers.cpp b/Engine/Source/Runtime/VulkanRHI/Private/VulkanLayers.cpp
--- a/Engine/Source/Runtime/VulkanRHI/Private/VulkanLayers.cpp
+++ b/Engine/Source/Runtime/VulkanRHI/Private/VulkanLayers.cpp
@@ -51,7 +51,10 @@ static const ANSICHAR* GIndividualValidationLayers[] =
 	"VK_LAYER_LUNARG_parameter_validation",
 	"VK_LAYER_LUNARG_object_tracker",
 	"VK_LAYER_LUNARG_core_validation",
+#if !PLATFORM_LUMIN
+	// freezes app inside MLGraphicsCreateClientVk() on Lumin if this is enabled.
 	"VK_LAYER_GOOGLE_unique_objects",
+#endif // !PLATFORM_LUMIN
 	nullptr
 };
 
@@ -726,7 +729,7 @@ void FOptionalVulkanDeviceExtensions::Setup(const TArray<const ANSICHAR*>& Devic
 
 	if (GGPUCrashDebuggingEnabled && !bHasAnyCrashExtension)
 	{
-		UE_LOG(LogVulkanRHI, Warning, TEXT("Tried to enable GPU crash debugging but no extension found!"));
+		UE_LOG(LogVulkanRHI, Warning, TEXT("Tried to enable GPU crash debugging but no extension found! Will use local tracepoints."));
 	}
 
 #if VULKAN_SUPPORTS_GOOGLE_DISPLAY_TIMING
diff --git a/Engine/Source/Runtime/VulkanRHI/Private/VulkanMemory.cpp b/Engine/Source/Runtime/VulkanRHI/Private/VulkanMemory.cpp
--- a/Engine/Source/Runtime/VulkanRHI/Private/VulkanMemory.cpp
+++ b/Engine/Source/Runtime/VulkanRHI/Private/VulkanMemory.cpp
@@ -984,7 +984,7 @@ namespace VulkanRHI
 				}
 			}
 #else
-#if PLATFORM_ANDROID && !PLATFORM_LUMIN && !PLATFORM_LUMINGL4
+#if PLATFORM_ANDROID && !PLATFORM_LUMIN
 			// free all pages, as this would keep staging buffers around forever, and they are 64mb in practice
 			for (int32 Index = 0; Index < FreePages.Num(); ++Index)
 #else
@@ -1271,7 +1271,7 @@ namespace VulkanRHI
 					HeapSize = FMath::Min<VkDeviceSize>(VULKAN_FAKE_MEMORY_LIMIT << 30llu, HeapSize);
 				}
 				VkDeviceSize PageSize = FMath::Min<VkDeviceSize>(HeapSize / 8, GPU_ONLY_HEAP_PAGE_SIZE);
-#if PLATFORM_ANDROID && !PLATFORM_LUMIN && !PLATFORM_LUMINGL4
+#if PLATFORM_ANDROID && !PLATFORM_LUMIN
 				PageSize = FMath::Min<VkDeviceSize>(PageSize, ANDROID_MAX_HEAP_PAGE_SIZE);
 #endif
 				ResourceTypeHeaps[TypeIndices[Index]] = new FOldResourceHeap(this, TypeIndices[Index], PageSize);
diff --git a/Engine/Source/Runtime/VulkanRHI/Private/VulkanPipeline.cpp b/Engine/Source/Runtime/VulkanRHI/Private/VulkanPipeline.cpp
--- a/Engine/Source/Runtime/VulkanRHI/Private/VulkanPipeline.cpp
+++ b/Engine/Source/Runtime/VulkanRHI/Private/VulkanPipeline.cpp
@@ -214,7 +214,7 @@ static TAutoConsoleVariable<int32> GEnablePipelineCacheLoadCvar(
 
 static TAutoConsoleVariable<int32> GPipelineCacheFromShaderPipelineCacheCvar(
 	TEXT("r.Vulkan.PipelineCacheFromShaderPipelineCache"),
-	PLATFORM_ANDROID && !(PLATFORM_LUMIN || PLATFORM_LUMINGL4),
+	PLATFORM_ANDROID && !PLATFORM_LUMIN,
 	TEXT("0 look for a pipeline cache in the normal locations with the normal names.")
 	TEXT("1 tie the vulkan pipeline cache to the shader pipeline cache, use the PSOFC guid as part of the filename, etc."),
 	ECVF_ReadOnly
diff --git a/Engine/Source/Runtime/VulkanRHI/Private/VulkanRHI.cpp b/Engine/Source/Runtime/VulkanRHI/Private/VulkanRHI.cpp
--- a/Engine/Source/Runtime/VulkanRHI/Private/VulkanRHI.cpp
+++ b/Engine/Source/Runtime/VulkanRHI/Private/VulkanRHI.cpp
@@ -83,7 +83,7 @@ FDynamicRHI* FVulkanDynamicRHIModule::CreateRHI(ERHIFeatureLevel::Type InRequest
 	else
 	{
 		GMaxRHIFeatureLevel = ERHIFeatureLevel::SM5;
-		GMaxRHIShaderPlatform = (PLATFORM_LUMINGL4 || PLATFORM_LUMIN) ? SP_VULKAN_SM5_LUMIN : SP_VULKAN_SM5;
+		GMaxRHIShaderPlatform = (PLATFORM_LUMIN) ? SP_VULKAN_SM5_LUMIN : SP_VULKAN_SM5;
 	}
 
 	GVulkanRHI = new FVulkanDynamicRHI();
@@ -609,7 +609,7 @@ void FVulkanDynamicRHI::SelectAndInitDevice()
 
 	Device->InitGPU(DeviceIndex);
 
-	if (PLATFORM_ANDROID && !PLATFORM_LUMIN && !PLATFORM_LUMINGL4)
+	if (PLATFORM_ANDROID && !PLATFORM_LUMIN)
 	{
 		GRHIAdapterName.Append(TEXT(" Vulkan"));
 		GRHIAdapterInternalDriverVersion = FString::Printf(TEXT("%d.%d.%d"), VK_VERSION_MAJOR(Props.apiVersion), VK_VERSION_MINOR(Props.apiVersion), VK_VERSION_PATCH(Props.apiVersion));
diff --git a/Engine/Source/Runtime/VulkanRHI/Private/VulkanUtil.cpp b/Engine/Source/Runtime/VulkanRHI/Private/VulkanUtil.cpp
--- a/Engine/Source/Runtime/VulkanRHI/Private/VulkanUtil.cpp
+++ b/Engine/Source/Runtime/VulkanRHI/Private/VulkanUtil.cpp
@@ -349,16 +349,38 @@ float FVulkanEventNode::GetTiming()
 	return Result;
 }
 
+FVulkanGPUProfiler::FVulkanGPUProfiler(FVulkanCommandListContext* InCmd, FVulkanDevice* InDevice)
+	: bCommandlistSubmitted(false)
+	, Device(InDevice)
+	, CmdContext(InCmd)
+	, LocalTracePointsQueryPool(nullptr)
+	, bBeginFrame(false)
+{
+}
+
+FVulkanGPUProfiler::~FVulkanGPUProfiler()
+{
+	if (LocalTracePointsQueryPool != nullptr)
+	{
+		delete LocalTracePointsQueryPool;
+	}
+}
 
 void FVulkanGPUProfiler::BeginFrame()
 {
 #if VULKAN_SUPPORTS_GPU_CRASH_DUMPS
-	if (GGPUCrashDebuggingEnabled && Device->GetOptionalExtensions().HasGPUCrashDumpExtensions())
+	if (GGPUCrashDebuggingEnabled)
 	{
 		static auto* CrashCollectionEnableCvar = IConsoleManager::Get().FindTConsoleVariableDataInt(TEXT("r.gpucrash.collectionenable"));
 		static auto* CrashCollectionDataDepth = IConsoleManager::Get().FindTConsoleVariableDataInt(TEXT("r.gpucrash.datadepth"));
 		bTrackingGPUCrashData = CrashCollectionEnableCvar ? CrashCollectionEnableCvar->GetValueOnRenderThread() != 0 : false;
 		GPUCrashDataDepth = CrashCollectionDataDepth ? CrashCollectionDataDepth->GetValueOnRenderThread() : -1;
+
+		// Use local tracepoints if no extension is available
+		if (!Device->GetOptionalExtensions().HasGPUCrashDumpExtensions() && LocalTracePointsQueryPool == nullptr)
+		{
+			LocalTracePointsQueryPool = new FVulkanTimingQueryPool(Device, GMaxCrashBufferEntries);
+		}
 	}
 #endif
 
@@ -367,6 +389,17 @@ void FVulkanGPUProfiler::BeginFrame()
 	check(!bTrackingEvents);
 	check(!CurrentEventNodeFrame); // this should have already been cleaned up and the end of the previous frame
 
+	if (GGPUCrashDebuggingEnabled && !Device->GetOptionalExtensions().HasGPUCrashDumpExtensions())
+	{
+		VulkanRHI::vkCmdResetQueryPool(Device->GetImmediateContext().GetCommandBufferManager()->GetActiveCmdBuffer()->GetHandle(), LocalTracePointsQueryPool->GetHandle(), 0, GMaxCrashBufferEntries);
+
+		PushPopStack.Reset();
+		CrashMarkers.Reset();
+		CrashMarkers.AddZeroed(GMaxCrashBufferEntries);
+	}
+
+	bBeginFrame = true;
+
 	// latch the bools from the game thread into our private copy
 	bLatchedGProfilingGPU = GTriggerGPUProfile;
 	bLatchedGProfilingGPUHitches = GTriggerGPUHitchProfile;
@@ -463,11 +496,19 @@ void FVulkanGPUProfiler::EndFrame()
 		delete CurrentEventNodeFrame;
 		CurrentEventNodeFrame = nullptr;
 	}
+
+	bBeginFrame = false;
 }
 
 #if VULKAN_SUPPORTS_GPU_CRASH_DUMPS
 void FVulkanGPUProfiler::PushMarkerForCrash(VkCommandBuffer CmdBuffer, VkBuffer DestBuffer, const TCHAR* Name)
 {
+	if (!Device->GetOptionalExtensions().HasGPUCrashDumpExtensions() && !bBeginFrame)
+	{
+		// If using local trace points, ignore any markers pushed before begin frame or after end frame.
+		return;
+	}
+
 	uint32 CRC = 0;
 	if (GPUCrashDataDepth < 0 || PushPopStack.Num() < GPUCrashDataDepth)
 	{
@@ -491,14 +532,32 @@ void FVulkanGPUProfiler::PushMarkerForCrash(VkCommandBuffer CmdBuffer, VkBuffer
 
 	PushPopStack.Push(CRC);
 	FVulkanPlatform::WriteCrashMarker(Device->GetOptionalExtensions(), CmdBuffer, DestBuffer, TArrayView<uint32>(PushPopStack), true);
+
+	if (GGPUCrashDebuggingEnabled && !Device->GetOptionalExtensions().HasGPUCrashDumpExtensions())
+	{
+		VulkanRHI::vkCmdWriteTimestamp(CmdBuffer, VK_PIPELINE_STAGE_TOP_OF_PIPE_BIT, LocalTracePointsQueryPool->GetHandle(), PushPopStack.Num() - 1);
+	}
 }
 
 void FVulkanGPUProfiler::PopMarkerForCrash(VkCommandBuffer CmdBuffer, VkBuffer DestBuffer)
 {
+	if (!Device->GetOptionalExtensions().HasGPUCrashDumpExtensions() && !bBeginFrame)
+	{
+		// If using local trace points, ignore any markers popped before begin frame or after end frame.
+		return;
+	}
+
 	if (PushPopStack.Num() > 0)
 	{
-		PushPopStack.Pop(false);
-		FVulkanPlatform::WriteCrashMarker(Device->GetOptionalExtensions(), CmdBuffer, DestBuffer, TArrayView<uint32>(PushPopStack), false);
+		if (Device->GetOptionalExtensions().HasGPUCrashDumpExtensions())
+		{
+			PushPopStack.Pop(false);
+			FVulkanPlatform::WriteCrashMarker(Device->GetOptionalExtensions(), CmdBuffer, DestBuffer, TArrayView<uint32>(PushPopStack), false);
+		}
+		else if (GGPUCrashDebuggingEnabled)
+		{
+			VulkanRHI::vkGetQueryPoolResults(Device->GetInstanceHandle(), LocalTracePointsQueryPool->GetHandle(), 0, PushPopStack.Num(), sizeof(uint64) * GMaxCrashBufferEntries, CrashMarkers.GetData(), sizeof(uint64), VK_QUERY_RESULT_64_BIT | VK_QUERY_RESULT_WITH_AVAILABILITY_BIT);
+		}
 	}
 }
 
@@ -551,6 +610,22 @@ void FVulkanGPUProfiler::DumpCrashMarkers(void* BufferData)
 		}
 #endif
 	}
+
+	if (!Device->GetOptionalExtensions().HasGPUCrashDumpExtensions())
+	{
+		UE_LOG(LogVulkanRHI, Warning, TEXT("Printing trace points."));
+
+		for (int32 i = 0; i < PushPopStack.Num(); ++i)
+		{
+			const FString* InsertedFrame = CachedStrings.Find(PushPopStack[i]);
+			const FString FrameName = InsertedFrame ? *InsertedFrame : TEXT("<undefined>");
+
+			UE_LOG(LogVulkanRHI, Warning, TEXT("[gpu_crash_markers] %s"), (CrashMarkers[i] != 0) ? *FrameName : TEXT("unavailable"));
+		}
+
+		GLog->PanicFlushThreadedLogs();
+		GLog->Flush();
+	}
 }
 #endif
 
@@ -672,10 +747,7 @@ namespace VulkanRHI
 		if (GIsGPUCrashed && GGPUCrashDebuggingEnabled)
 		{
 			FVulkanDevice* Device = GVulkanRHI->GetDevice();
-			if (Device->GetOptionalExtensions().HasGPUCrashDumpExtensions())
-			{
-				Device->GetImmediateContext().GetGPUProfiler().DumpCrashMarkers(Device->GetCrashMarkerMappedPointer());
-			}
+			Device->GetImmediateContext().GetGPUProfiler().DumpCrashMarkers(Device->GetCrashMarkerMappedPointer());
 		}
 #endif
 
diff --git a/Engine/Source/Runtime/VulkanRHI/Public/VulkanCommon.h b/Engine/Source/Runtime/VulkanRHI/Public/VulkanCommon.h
--- a/Engine/Source/Runtime/VulkanRHI/Public/VulkanCommon.h
+++ b/Engine/Source/Runtime/VulkanRHI/Public/VulkanCommon.h
@@ -9,7 +9,7 @@
 #include "RHIDefinitions.h"
 
 #ifndef VULKAN_SUPPORTS_GEOMETRY_SHADERS
-	#define VULKAN_SUPPORTS_GEOMETRY_SHADERS					(!PLATFORM_ANDROID || PLATFORM_LUMIN || PLATFORM_LUMINGL4) && PLATFORM_SUPPORTS_GEOMETRY_SHADERS
+	#define VULKAN_SUPPORTS_GEOMETRY_SHADERS					(!PLATFORM_ANDROID || PLATFORM_LUMIN) && PLATFORM_SUPPORTS_GEOMETRY_SHADERS
 #endif
 
 // This defines controls shader generation (so will cause a format rebuild)
diff --git a/Engine/Source/Runtime/VulkanRHI/Public/VulkanConfiguration.h b/Engine/Source/Runtime/VulkanRHI/Public/VulkanConfiguration.h
--- a/Engine/Source/Runtime/VulkanRHI/Public/VulkanConfiguration.h
+++ b/Engine/Source/Runtime/VulkanRHI/Public/VulkanConfiguration.h
@@ -165,7 +165,9 @@
 	#define VULKAN_SUPPORTS_NV_DIAGNOSTIC_CHECKPOINT			0
 #endif
 
-#define VULKAN_SUPPORTS_GPU_CRASH_DUMPS							(VULKAN_SUPPORTS_AMD_BUFFER_MARKER || VULKAN_SUPPORTS_NV_DIAGNOSTIC_CHECKPOINT)
+#ifndef VULKAN_SUPPORTS_GPU_CRASH_DUMPS
+	#define VULKAN_SUPPORTS_GPU_CRASH_DUMPS						(VULKAN_SUPPORTS_AMD_BUFFER_MARKER || VULKAN_SUPPORTS_NV_DIAGNOSTIC_CHECKPOINT)
+#endif
 
 #ifndef VULKAN_SUPPORTS_DEBUG_UTILS
 	#ifdef VK_EXT_debug_utils
diff --git a/Engine/Source/Runtime/VulkanRHI/Public/VulkanMemory.h b/Engine/Source/Runtime/VulkanRHI/Public/VulkanMemory.h
--- a/Engine/Source/Runtime/VulkanRHI/Public/VulkanMemory.h
+++ b/Engine/Source/Runtime/VulkanRHI/Public/VulkanMemory.h
@@ -40,7 +40,7 @@ namespace VulkanRHI
 
 	enum
 	{
-#if PLATFORM_ANDROID || PLATFORM_LUMIN || PLATFORM_LUMINGL4
+#if PLATFORM_ANDROID || PLATFORM_LUMIN
 		NUM_FRAMES_TO_WAIT_BEFORE_RELEASING_TO_OS = 3,
 #else
 		NUM_FRAMES_TO_WAIT_BEFORE_RELEASING_TO_OS = 10,
diff --git a/Engine/Source/Runtime/VulkanRHI/Public/VulkanUtil.h b/Engine/Source/Runtime/VulkanRHI/Public/VulkanUtil.h
--- a/Engine/Source/Runtime/VulkanRHI/Public/VulkanUtil.h
+++ b/Engine/Source/Runtime/VulkanRHI/Public/VulkanUtil.h
@@ -157,12 +157,9 @@ struct FVulkanGPUProfiler : public FGPUProfiler
 	/** GPU hitch profile histories */
 	TIndirectArray<FVulkanEventNodeFrame> GPUHitchEventNodeFrames;
 
-	FVulkanGPUProfiler(FVulkanCommandListContext* InCmd, FVulkanDevice* InDevice)
-		: bCommandlistSubmitted(false)
-		, Device(InDevice)
-		, CmdContext(InCmd)
-	{
-	}
+	FVulkanGPUProfiler(FVulkanCommandListContext* InCmd, FVulkanDevice* InDevice);
+
+	virtual ~FVulkanGPUProfiler();
 
 	virtual FGPUProfilerEventNode* CreateEventNode(const TCHAR* InName, FGPUProfilerEventNode* InParent) override final
 	{
@@ -188,6 +185,10 @@ struct FVulkanGPUProfiler : public FGPUProfiler
 	// For crash/marker tracking
 	TMap<uint32, FString> CachedStrings;
 	TArray<uint32> PushPopStack;
+
+	FVulkanTimingQueryPool* LocalTracePointsQueryPool;
+	TArray<uint64> CrashMarkers;
+	bool bBeginFrame;
 };
 
 namespace VulkanRHI
