Fixed Hololens if constexpr error.
-We need to supress this warning on hololens, like we do on windows.

diff --git a/Engine/Source/Runtime/Core/Public/HoloLens/HoloLensPlatformCompilerSetup.h b/Engine/Source/Runtime/Core/Public/HoloLens/HoloLensPlatformCompilerSetup.h
--- a/Engine/Source/Runtime/Core/Public/HoloLens/HoloLensPlatformCompilerSetup.h
+++ b/Engine/Source/Runtime/Core/Public/HoloLens/HoloLensPlatformCompilerSetup.h
@@ -197,6 +197,8 @@ static_assert(_MSC_VER >= 1900, "Visual Studio 2015 or newer required.");
 #pragma warning(disable : 4189) // local variable is initialized but not referenced 
 #pragma warning(disable : 4505) // unreferenced local function has been removed		
 
+#pragma warning(disable: 4984) // 'if constexpr' is a C++17 language extension
+
 // string.h adds this warning
 #pragma warning(disable : 4996)
 
