Deprecated FCustomVersionContainer::GetRegistered() and replaced it with a new thread-safe API

diff --git a/Engine/Plugins/Developer/Concert/ConcertMain/Source/Concert/Private/ConcertVersion.cpp b/Engine/Plugins/Developer/Concert/ConcertMain/Source/Concert/Private/ConcertVersion.cpp
--- a/Engine/Plugins/Developer/Concert/ConcertMain/Source/Concert/Private/ConcertVersion.cpp
+++ b/Engine/Plugins/Developer/Concert/ConcertMain/Source/Concert/Private/ConcertVersion.cpp
@@ -117,7 +117,7 @@ void FConcertSessionVersionInfo::Initialize()
 		}
 	}
 
-	for (const FCustomVersion& EngineCustomVersion : FCustomVersionContainer::GetRegistered().GetAllVersions())
+	for (const FCustomVersion& EngineCustomVersion : FCurrentCustomVersions::GetAll().GetAllVersions())
 	{
 		FConcertCustomVersionInfo& CustomVersion = CustomVersions.AddDefaulted_GetRef();
 		CustomVersion.Initialize(EngineCustomVersion);
diff --git a/Engine/Source/Editor/UnrealEd/Private/CookOnTheFlyServer.cpp b/Engine/Source/Editor/UnrealEd/Private/CookOnTheFlyServer.cpp
--- a/Engine/Source/Editor/UnrealEd/Private/CookOnTheFlyServer.cpp
+++ b/Engine/Source/Editor/UnrealEd/Private/CookOnTheFlyServer.cpp
@@ -4498,8 +4498,7 @@ void GetAdditionalCurrentIniVersionStrings( const ITargetPlatform* TargetPlatfor
 	}
 
 
-	static const FCustomVersionContainer& CustomVersionContainer = FCustomVersionContainer::GetRegistered();
-	for (const auto& CustomVersion : CustomVersionContainer.GetAllVersions())
+	for (const FCustomVersion& CustomVersion : FCurrentCustomVersions::GetAll().GetAllVersions())
 	{
 		FString CustomVersionString = FString::Printf(TEXT("%s:%s"), *CustomVersion.GetFriendlyName().ToString(), *CustomVersion.Key.ToString());
 		FString CustomVersionValue = FString::Printf(TEXT("%d"), CustomVersion.Version);
diff --git a/Engine/Source/Runtime/AssetRegistry/Private/PackageReader.cpp b/Engine/Source/Runtime/AssetRegistry/Private/PackageReader.cpp
--- a/Engine/Source/Runtime/AssetRegistry/Private/PackageReader.cpp
+++ b/Engine/Source/Runtime/AssetRegistry/Private/PackageReader.cpp
@@ -83,17 +83,15 @@ bool FPackageReader::OpenPackageFile(EOpenPackageResult* OutErrorCode)
 	}
 
 	// Check serialized custom versions against latest custom versions.
-	const FCustomVersionContainer& LatestCustomVersions  = FCustomVersionContainer::GetRegistered();
-	const FCustomVersionArray&  PackageCustomVersions = PackageFileSummary.GetCustomVersionContainer().GetAllVersions();
-	for (const FCustomVersion& SerializedCustomVersion : PackageCustomVersions)
+	TArray<FCustomVersionDifference> Diffs = FCurrentCustomVersions::Compare(PackageFileSummary.GetCustomVersionContainer().GetAllVersions());
+	for (FCustomVersionDifference Diff : Diffs)
 	{
-		auto* LatestVersion = LatestCustomVersions.GetVersion(SerializedCustomVersion.Key);
-		if (!LatestVersion)
+		if (Diff.Type == ECustomVersionDifference::Missing)
 		{
 			SetPackageErrorCode(EOpenPackageResult::CustomVersionMissing);
 			return false;
 		}
-		else if (SerializedCustomVersion.Version > LatestVersion->Version)
+		else if (Diff.Type == ECustomVersionDifference::Newer)
 		{
 			SetPackageErrorCode(EOpenPackageResult::VersionTooNew);
 			return false;
diff --git a/Engine/Source/Runtime/Core/Private/Serialization/Archive.cpp b/Engine/Source/Runtime/Core/Private/Serialization/Archive.cpp
--- a/Engine/Source/Runtime/Core/Private/Serialization/Archive.cpp
+++ b/Engine/Source/Runtime/Core/Private/Serialization/Archive.cpp
@@ -441,7 +441,7 @@ const FCustomVersionContainer& FArchive::GetCustomVersions() const
 		// serialization code to use UsingCustomVersion to populate the container.
 		if (this->IsLoading())
 		{
-			*CustomVersionContainer = FCustomVersionContainer::GetRegistered();
+			*CustomVersionContainer = FCurrentCustomVersions::GetAll();
 		}
 		else
 		{
@@ -478,13 +478,8 @@ void FArchive::UsingCustomVersion(const FGuid& Key)
 		return;
 	}
 
-	auto* RegisteredVersion = FCustomVersionContainer::GetRegistered().GetVersion(Key);
-
-	// Ensure that the version has already been registered.
-	// If this fails, you probably don't have an FCustomVersionRegistration variable defined for this GUID.
-	check(RegisteredVersion);
-
-	const_cast<FCustomVersionContainer&>(GetCustomVersions()).SetVersion(Key, RegisteredVersion->Version, RegisteredVersion->GetFriendlyName());
+	FCustomVersion RegisteredVersion = FCurrentCustomVersions::Get(Key).GetValue();
+	const_cast<FCustomVersionContainer&>(GetCustomVersions()).SetVersion(Key, RegisteredVersion.Version, RegisteredVersion.GetFriendlyName());
 }
 
 int32 FArchive::CustomVer(const FGuid& Key) const
diff --git a/Engine/Source/Runtime/Core/Private/Serialization/CustomVersion.cpp b/Engine/Source/Runtime/Core/Private/Serialization/CustomVersion.cpp
--- a/Engine/Source/Runtime/Core/Private/Serialization/CustomVersion.cpp
+++ b/Engine/Source/Runtime/Core/Private/Serialization/CustomVersion.cpp
@@ -8,6 +8,8 @@
 #include "Serialization/StructuredArchive.h"
 #include "Algo/Sort.h"
 #include "Containers/Map.h"
+#include "Misc/ScopeRwLock.h"
+#include "CoreGlobals.h"
 
 namespace
 {
@@ -74,29 +76,8 @@ namespace
 }
 
 /** Defer FName creation and allocations from static FCustomVersionRegistrations that may never be needed */
-class FStaticCustomVersionRegistry
+struct FStaticCustomVersionRegistry
 {
-public:
-	static const FCustomVersionContainer& Get()
-	{
-		FStaticCustomVersionRegistry& StaticVersions = GetInstance();
-		StaticVersions.RegisterQueue();
-		return StaticVersions.Registered;
-	}
-
-	static void Register(FGuid Key, int32 Version, const TCHAR* Name)
-	{
-		RegistrationQueue& Queue = GetInstance().Queue;
-		check(Queue.Find(Key) == nullptr);
-		Queue.Add(Key, { Version, Name });
-	}
-
-	static void Unregister(FGuid Key)
-	{
-		 GetInstance().UnregisterImpl(Key);
-	}
-
-private:
 	struct FPendingRegistration
 	{
 		int32 Version;
@@ -104,14 +85,30 @@ private:
 	};
 	typedef TMap<FGuid, FPendingRegistration, TInlineSetAllocator<64>> RegistrationQueue;
 
+	FRWLock Lock;
 	FCustomVersionContainer Registered;
 	RegistrationQueue Queue;
 
-	static FStaticCustomVersionRegistry& GetInstance()
+	static FStaticCustomVersionRegistry& Get()
 	{
 		static FStaticCustomVersionRegistry Singleton;
 		return Singleton;
-	};
+	}
+
+	TOptional<FCustomVersion> Find(const FGuid& Guid) const
+	{
+		if (const FCustomVersion* RegisteredVersion = Registered.GetVersion(Guid))
+		{
+			return *RegisteredVersion;
+		}
+
+		if (const FPendingRegistration* Pending = Queue.Find(Guid))
+		{
+			return FCustomVersion(Guid, Pending->Version, Pending->FriendlyName);
+		}
+
+		return TOptional<FCustomVersion>();
+	}
 
 	void RegisterQueue()
 	{
@@ -148,7 +145,7 @@ private:
 		}
 	}
 
-	void UnregisterImpl(FGuid Key)
+	void Unregister(FGuid Key)
 	{
 		if (Queue.Remove(Key) == 0)
 		{
@@ -168,18 +165,99 @@ private:
 	}
 };
 
+FCustomVersionContainer FCurrentCustomVersions::GetAll()
+{
+	FStaticCustomVersionRegistry& Registry = FStaticCustomVersionRegistry::Get();
+
+	{
+		FReadScopeLock Scope(Registry.Lock);
+
+		if (Registry.Queue.Num() == 0)
+		{
+			return Registry.Registered;
+		}
+	}
+
+	FWriteScopeLock Scope(Registry.Lock);
+	Registry.RegisterQueue();
+	return Registry.Registered;
+}
+
+TOptional<FCustomVersion> FCurrentCustomVersions::Get(const FGuid& Guid)
+{
+	FStaticCustomVersionRegistry& Registry = FStaticCustomVersionRegistry::Get();
+
+	FReadScopeLock Scope(Registry.Lock);
+	return Registry.Find(Guid);
+}
+
+TArray<FCustomVersionDifference> FCurrentCustomVersions::Compare(const FCustomVersionArray& CompareVersions)
+{
+	TArray<FCustomVersionDifference> Result;
+
+	if (CompareVersions.Num())
+	{
+		FStaticCustomVersionRegistry& Registry = FStaticCustomVersionRegistry::Get();
+
+		FReadScopeLock Scope(Registry.Lock);
+
+		for (const FCustomVersion& CompareVersion : CompareVersions)
+		{
+			if (TOptional<FCustomVersion> CurrentVersion = Registry.Find(CompareVersion.Key))
+			{
+				if (int Delta = CurrentVersion.GetValue().Version - CompareVersion.Version)
+				{
+					Result.Add({ Delta < 0	? ECustomVersionDifference::Newer 
+											: ECustomVersionDifference::Older, &CompareVersion });
+				}
+			}
+			else
+			{
+				Result.Add({ ECustomVersionDifference::Missing, &CompareVersion });
+			}			
+		}
+	}
+
+	return Result;
+}
+
+void FCurrentCustomVersions::Register(const FGuid& Key, int32 Version, const TCHAR* Name)
+{
+	FStaticCustomVersionRegistry& Registry = FStaticCustomVersionRegistry::Get();
+
+	FWriteScopeLock Scope(Registry.Lock);
+	check(Registry.Queue.Find(Key) == nullptr);
+	Registry.Queue.Add(Key, { Version, Name });
+}
+
+void FCurrentCustomVersions::Unregister(const FGuid& Key)
+{
+	FStaticCustomVersionRegistry& Registry = FStaticCustomVersionRegistry::Get();
+
+	FWriteScopeLock Scope(Registry.Lock);
+	Registry.Unregister(Key);
+}
+
 const FName FCustomVersion::GetFriendlyName() const
 {
 	if (FriendlyName == NAME_None)
 	{
-		FriendlyName = FCustomVersionContainer::GetRegistered().GetFriendlyName(Key);
+		if (TOptional<FCustomVersion> CurrentVersion = FCurrentCustomVersions::Get(Key))
+		{
+			FriendlyName = CurrentVersion.GetValue().GetFriendlyName();
+		}
 	}
 	return FriendlyName;
 }
 
 const FCustomVersionContainer& FCustomVersionContainer::GetRegistered()
 {
-	return FStaticCustomVersionRegistry::Get();
+	FStaticCustomVersionRegistry& Registry = FStaticCustomVersionRegistry::Get();
+
+	// Even though returning a reference isn't thread-safe, we can still synchronize access to the queue
+	FWriteScopeLock Scope(Registry.Lock);
+	Registry.RegisterQueue();
+	return Registry.Registered;
 }
 
 void FCustomVersionContainer::Empty()
@@ -307,13 +385,3 @@ void FCustomVersionContainer::SetVersion(FGuid CustomKey, int32 Version, FName F
 		Versions.Add(FCustomVersion(CustomKey, Version, FriendlyName));
 	}
 }
-
-void FCustomVersionRegistration::QueueRegistration(FGuid Key, int32 Version, const TCHAR* Name)
-{
-	FStaticCustomVersionRegistry::Register(Key, Version, Name);
-}
-
-FCustomVersionRegistration::~FCustomVersionRegistration()
-{
-	FStaticCustomVersionRegistry::Unregister(Key);
-}
diff --git a/Engine/Source/Runtime/Core/Private/UObject/DevObjectVersion.cpp b/Engine/Source/Runtime/Core/Private/UObject/DevObjectVersion.cpp
--- a/Engine/Source/Runtime/Core/Private/UObject/DevObjectVersion.cpp
+++ b/Engine/Source/Runtime/Core/Private/UObject/DevObjectVersion.cpp
@@ -45,9 +45,8 @@ void FDevVersionRegistration::DumpVersionsToLog()
 	UE_LOG(LogDevObjectVersion, Log, TEXT("Number of dev versions registered: %d"), GDevVersions.Num());
 	for (FGuid& Guid : GDevVersions)
 	{
-		const FCustomVersion* Version = FCustomVersionContainer::GetRegistered().GetVersion(Guid);
-		check(Version);
-		UE_LOG(LogDevObjectVersion, Log, TEXT("  %s (%s): %d"), *Version->GetFriendlyName().ToString(), *Version->Key.ToString(EGuidFormats::DigitsWithHyphens), Version->Version);
+		FCustomVersion Version = FCurrentCustomVersions::Get(Guid).GetValue();
+		UE_LOG(LogDevObjectVersion, Log, TEXT("  %s (%s): %d"), *Version.GetFriendlyName().ToString(), *Version.Key.ToString(EGuidFormats::DigitsWithHyphens), Version.Version);
 	}
 #endif
 }
diff --git a/Engine/Source/Runtime/Core/Public/Serialization/CustomVersion.h b/Engine/Source/Runtime/Core/Public/Serialization/CustomVersion.h
--- a/Engine/Source/Runtime/Core/Public/Serialization/CustomVersion.h
+++ b/Engine/Source/Runtime/Core/Public/Serialization/CustomVersion.h
@@ -94,7 +94,7 @@ class CORE_API FCustomVersionRegistration;
  */
 class CORE_API FCustomVersionContainer
 {
-	friend class FStaticCustomVersionRegistry;
+	friend struct FStaticCustomVersionRegistry;
 
 public:
 	/** Gets available custom versions in this container. */
@@ -137,6 +137,7 @@ public:
 	 *
 	 * @return The registered version container.
 	 */
+	UE_DEPRECATED(4.24, "Use one of the thread-safe FCurrentCustomVersions methods instead")
 	static const FCustomVersionContainer& GetRegistered();
 
 	/**
@@ -159,6 +160,34 @@ private:
 
 };
 
+enum class ECustomVersionDifference { Missing, Newer, Older };
+
+struct FCustomVersionDifference
+{
+	ECustomVersionDifference Type;
+	const FCustomVersion* Version;
+};
+
+/** Provides access to code-defined custom versions registered via FCustomVersionRegistration. */
+class CORE_API FCurrentCustomVersions
+{
+public:
+	/** Get a copy of all versions that has been statically registered so far in the module loading process. */
+	static FCustomVersionContainer GetAll();
+
+	/** Get a copy of a single statically registered version if it exists. */
+	static TOptional<FCustomVersion> Get(const FGuid& Guid);
+
+	/** Compare a number of versions to current ones and return potential differences. */
+	static TArray<FCustomVersionDifference> Compare(const FCustomVersionArray& CompareVersions);
+
+private:
+	friend class FCustomVersionRegistration;
+
+	static void Register(const FGuid& Key, int32 Version, const TCHAR* FriendlyName);
+	static void Unregister(const FGuid& Key);
+};
+
 
 /**
  * This class will cause the registration of a custom version number and key with the global
@@ -173,14 +202,14 @@ public:
 	FCustomVersionRegistration(FGuid InKey, int32 Version, const TCHAR(&InFriendlyName)[N])
 	: Key(InKey)
 	{
-		QueueRegistration(InKey, Version, InFriendlyName);
+		FCurrentCustomVersions::Register(InKey, Version, InFriendlyName);
 	}
 
-	~FCustomVersionRegistration();
+	~FCustomVersionRegistration()
+	{
+		FCurrentCustomVersions::Unregister(Key);
+	}
 
 private:
 	FGuid Key;
-
-	/** Put registrations in an intermediate queue to avoid allocations and FName creation during static init */
-	static void QueueRegistration(FGuid Key, int32 Version, const TCHAR* FriendlyName);
 };
diff --git a/Engine/Source/Runtime/CoreUObject/Private/UObject/LinkerLoad.cpp b/Engine/Source/Runtime/CoreUObject/Private/UObject/LinkerLoad.cpp
--- a/Engine/Source/Runtime/CoreUObject/Private/UObject/LinkerLoad.cpp
+++ b/Engine/Source/Runtime/CoreUObject/Private/UObject/LinkerLoad.cpp
@@ -1292,35 +1292,26 @@ FLinkerLoad::ELinkerStatus FLinkerLoad::SerializePackageFileSummary()
 FLinkerLoad::ELinkerStatus FLinkerLoad::UpdateFromPackageFileSummary()
 {
 	// When unversioned, pretend we are the latest version
-	bool bCustomVersionIsLatest = Summary.bUnversioned;
-	if (!bCustomVersionIsLatest)
+	bool bCustomVersionIsLatest = true;
+	if (!Summary.bUnversioned)
 	{
-		// Check custom versions.
-		const FCustomVersionContainer& LatestCustomVersions  = FCustomVersionContainer::GetRegistered();
-		bool bAllSavedVersionsMatch = true;
-		const FCustomVersionArray&  PackageCustomVersions = Summary.GetCustomVersionContainer().GetAllVersions();
-		for (auto It = PackageCustomVersions.CreateConstIterator(); It; ++It)
+		TArray<FCustomVersionDifference> Diffs = FCurrentCustomVersions::Compare(Summary.GetCustomVersionContainer().GetAllVersions());
+		for (FCustomVersionDifference Diff : Diffs)
 		{
-			const FCustomVersion& SerializedCustomVersion = *It;
-
-			const FCustomVersion* LatestVersion = LatestCustomVersions.GetVersion(SerializedCustomVersion.Key);
-			if (!LatestVersion)
+			bCustomVersionIsLatest = false;
+			if (Diff.Type == ECustomVersionDifference::Missing)
 			{
 				// Loading a package with custom integration that we don't know about!
 				// Temporarily just warn and continue. @todo: this needs to be fixed properly
-				UE_LOG(LogLinker, Warning, TEXT("Package %s was saved with a custom integration that is not present. Tag %s  Version %d"), *Filename, *SerializedCustomVersion.Key.ToString(), SerializedCustomVersion.Version);
-				bAllSavedVersionsMatch = false;
+				UE_LOG(LogLinker, Warning, TEXT("Package %s was saved with a custom integration that is not present. Tag %s  Version %d"), *Filename, *Diff.Version->Key.ToString(), Diff.Version->Version);
 			}
-			else if (SerializedCustomVersion.Version > LatestVersion->Version)
+			else if (Diff.Type == ECustomVersionDifference::Newer)
 			{
+				FCustomVersion LatestVersion = FCurrentCustomVersions::Get(Diff.Version->Key).GetValue();
 				// Loading a package with a newer custom version than the current one.
-				UE_LOG(LogLinker, Error, TEXT("Package %s was saved with a newer custom version than the current. Tag %s Name '%s' PackageVersion %d  MaxExpected %d"), *Filename, *SerializedCustomVersion.Key.ToString(), *LatestVersion->GetFriendlyName().ToString(), SerializedCustomVersion.Version, LatestVersion->Version);
+				UE_LOG(LogLinker, Error, TEXT("Package %s was saved with a newer custom version than the current. Tag %s Name '%s' PackageVersion %d  MaxExpected %d"), *Filename, *Diff.Version->Key.ToString(), *LatestVersion.GetFriendlyName().ToString(), Diff.Version->Version, LatestVersion.Version);
 				return LINKER_Failed;
 			}
-			else if (SerializedCustomVersion.Version != LatestVersion->Version)
-			{
-				bAllSavedVersionsMatch = false;
-			}
 		}
 	}
 
diff --git a/Engine/Source/Runtime/CoreUObject/Private/UObject/LinkerSave.cpp b/Engine/Source/Runtime/CoreUObject/Private/UObject/LinkerSave.cpp
--- a/Engine/Source/Runtime/CoreUObject/Private/UObject/LinkerSave.cpp
+++ b/Engine/Source/Runtime/CoreUObject/Private/UObject/LinkerSave.cpp
@@ -320,11 +320,10 @@ void FLinkerSave::UsingCustomVersion(const struct FGuid& Guid)
 	// Here we're going to try and dump the callstack that added a new custom version after package summary has been serialized
 	if (Summary.GetCustomVersionContainer().GetVersion(Guid) == nullptr)
 	{
-		const FCustomVersion* RegisteredVersion = FCustomVersionContainer::GetRegistered().GetVersion(Guid);
-		check(RegisteredVersion);
+		FCustomVersion RegisteredVersion = FCurrentCustomVersions::Get(Guid).GetValue();
 
 		FString CustomVersionWarning = FString::Printf(TEXT("Unexpected custom version \"%s\" used after package %s summary has been serialized. Callstack:\n"),
-			*RegisteredVersion->GetFriendlyName().ToString(), *LinkerRoot->GetName());
+			*RegisteredVersion.GetFriendlyName().ToString(), *LinkerRoot->GetName());
 
 		const int32 MaxStackFrames = 100;
 		uint64 StackFrames[MaxStackFrames];
diff --git a/Engine/Source/Runtime/CoreUObject/Private/UObject/ObjectBaseUtility.cpp b/Engine/Source/Runtime/CoreUObject/Private/UObject/ObjectBaseUtility.cpp
--- a/Engine/Source/Runtime/CoreUObject/Private/UObject/ObjectBaseUtility.cpp
+++ b/Engine/Source/Runtime/CoreUObject/Private/UObject/ObjectBaseUtility.cpp
@@ -67,20 +67,18 @@ int32 UObjectBaseUtility::GetLinkerCustomVersion(FGuid CustomVersionKey) const
 	if ( Loader != NULL )
 	{
 		// We have a linker so we can return its version.
-		auto* CustomVersion = Loader->Summary.GetCustomVersionContainer().GetVersion(CustomVersionKey);
+		const FCustomVersion* CustomVersion = Loader->Summary.GetCustomVersionContainer().GetVersion(CustomVersionKey);
 		return CustomVersion ? CustomVersion->Version : -1;
 	}
 	else if (GetOutermost() && GetOutermost()->LinkerCustomVersion.GetAllVersions().Num())
 	{
 		// Get the linker version associated with the package this object lives in
-		auto* CustomVersion = GetOutermost()->LinkerCustomVersion.GetVersion(CustomVersionKey);
+		const FCustomVersion* CustomVersion = GetOutermost()->LinkerCustomVersion.GetVersion(CustomVersionKey);
 		return CustomVersion ? CustomVersion->Version : -1;
 	}
 	// We don't have a linker associated as we e.g. might have been saved or had loaders reset, ...
 	// We must have a current version for this tag.
-	auto* CustomVersion = FCustomVersionContainer::GetRegistered().GetVersion(CustomVersionKey);
-	check(CustomVersion);
-	return CustomVersion->Version;
+	return FCurrentCustomVersions::Get(CustomVersionKey).GetValue().Version;
 }
 
 /**
diff --git a/Engine/Source/Runtime/CoreUObject/Private/UObject/PackageFileSummary.cpp b/Engine/Source/Runtime/CoreUObject/Private/UObject/PackageFileSummary.cpp
--- a/Engine/Source/Runtime/CoreUObject/Private/UObject/PackageFileSummary.cpp
+++ b/Engine/Source/Runtime/CoreUObject/Private/UObject/PackageFileSummary.cpp
@@ -128,7 +128,7 @@ void operator<<(FStructuredArchive::FSlot Slot, FPackageFileSummary& Sum)
 					Sum.bUnversioned = true;
 					Sum.FileVersionUE4 = GPackageFileUE4Version;
 					Sum.FileVersionLicenseeUE4 = GPackageFileLicenseeUE4Version;
-					Sum.CustomVersionContainer = FCustomVersionContainer::GetRegistered();
+					Sum.CustomVersionContainer = FCurrentCustomVersions::GetAll();
 				}
 			}
 			else
diff --git a/Engine/Source/Runtime/Engine/Private/GameplayStatics.cpp b/Engine/Source/Runtime/Engine/Private/GameplayStatics.cpp
--- a/Engine/Source/Runtime/Engine/Private/GameplayStatics.cpp
+++ b/Engine/Source/Runtime/Engine/Private/GameplayStatics.cpp
@@ -105,7 +105,7 @@ FSaveGameHeader::FSaveGameHeader(TSubclassOf<USaveGame> ObjectType)
 	, PackageFileUE4Version(GPackageFileUE4Version)
 	, SavedEngineVersion(FEngineVersion::Current())
 	, CustomVersionFormat(static_cast<int32>(ECustomVersionSerializationFormat::Latest))
-	, CustomVersions(FCustomVersionContainer::GetRegistered())
+	, CustomVersions(FCurrentCustomVersions::GetAll())
 	, SaveGameClassName(ObjectType->GetPathName())
 {}
 
