Linux case-sensitive include fix

diff --git a/Engine/Source/Runtime/Core/Private/Serialization/CustomVersion.cpp b/Engine/Source/Runtime/Core/Private/Serialization/CustomVersion.cpp
--- a/Engine/Source/Runtime/Core/Private/Serialization/CustomVersion.cpp
+++ b/Engine/Source/Runtime/Core/Private/Serialization/CustomVersion.cpp
@@ -8,7 +8,7 @@
 #include "Serialization/StructuredArchive.h"
 #include "Algo/Sort.h"
 #include "Containers/Map.h"
-#include "Misc/ScopeRwLock.h"
+#include "Misc/ScopeRWLock.h"
 #include "CoreGlobals.h"
 
 namespace
