Use errno to implement  FApplePlatformMisc::GetSystemErrorMessage instead of always returning no valuable information.
#tests iOS Client

diff --git a/Engine/Source/Runtime/Core/Private/Apple/ApplePlatformMisc.cpp b/Engine/Source/Runtime/Core/Private/Apple/ApplePlatformMisc.cpp
--- a/Engine/Source/Runtime/Core/Private/Apple/ApplePlatformMisc.cpp
+++ b/Engine/Source/Runtime/Core/Private/Apple/ApplePlatformMisc.cpp
@@ -56,9 +56,15 @@ void FApplePlatformMisc::LocalPrint(const TCHAR* Message)
 
 const TCHAR* FApplePlatformMisc::GetSystemErrorMessage(TCHAR* OutBuffer, int32 BufferCount, int32 Error)
 {
-	// There's no iOS equivalent for GetLastError()
 	check(OutBuffer && BufferCount);
 	*OutBuffer = TEXT('\0');
+	if (Error == 0)
+	{
+		Error = errno;
+	}
+	char ErrorBuffer[1024];
+	strerror_r(Error, ErrorBuffer, 1024);
+	FCString::Strcpy(OutBuffer, BufferCount, UTF8_TO_TCHAR((const ANSICHAR*)ErrorBuffer));
 	return OutBuffer;
 }
 
