Improve FApplePlatformMisc to handle error messages over 1024 chars.
#tests iOS Client

diff --git a/Engine/Source/Runtime/Core/Private/Apple/ApplePlatformMisc.cpp b/Engine/Source/Runtime/Core/Private/Apple/ApplePlatformMisc.cpp
--- a/Engine/Source/Runtime/Core/Private/Apple/ApplePlatformMisc.cpp
+++ b/Engine/Source/Runtime/Core/Private/Apple/ApplePlatformMisc.cpp
@@ -62,8 +62,8 @@ const TCHAR* FApplePlatformMisc::GetSystemErrorMessage(TCHAR* OutBuffer, int32 B
 	{
 		Error = errno;
 	}
-	char ErrorBuffer[1024];
-	strerror_r(Error, ErrorBuffer, 1024);
+	char* ErrorBuffer = (char*)alloca(BufferCount);
+	strerror_r(Error, ErrorBuffer, BufferCount);
 	FCString::Strcpy(OutBuffer, BufferCount, UTF8_TO_TCHAR((const ANSICHAR*)ErrorBuffer));
 	return OutBuffer;
 }
