Allow the CrashReportClient to take window focus after a crash is piped to it

diff --git a/Engine/Source/Runtime/Core/Private/Windows/WindowsPlatformCrashContext.cpp b/Engine/Source/Runtime/Core/Private/Windows/WindowsPlatformCrashContext.cpp
--- a/Engine/Source/Runtime/Core/Private/Windows/WindowsPlatformCrashContext.cpp
+++ b/Engine/Source/Runtime/Core/Private/Windows/WindowsPlatformCrashContext.cpp
@@ -575,6 +575,12 @@ int32 ReportCrashForMonitor(
 		FGenericCrashContext::DumpLog(CrashDirectoryAbsolute);
 	}
 
+	// Allow the monitor process to take window focus
+	if (const DWORD MonitorProcessId = ::GetProcessId(CrashMonitorHandle.Get()))
+	{
+		::AllowSetForegroundWindow(MonitorProcessId);
+	}
+
 	// Write the shared context to the pipe
 	int32 OutDataWritten = 0;
 	FPlatformProcess::WritePipe(WritePipe, (UINT8*)SharedContext, sizeof(FSharedCrashContext), &OutDataWritten);
