FLazyName fixes and improvements

* Fixed number parsing
* Narrow literals were eagerly resolved via implicit FName construction, added support for narrow literals
* Merged 32-bit and 64-bit implementations
* Added additional unit tests
* Moved some code to .cpp

diff --git a/Engine/Source/Runtime/Core/Private/UObject/UnrealNames.cpp b/Engine/Source/Runtime/Core/Private/UObject/UnrealNames.cpp
--- a/Engine/Source/Runtime/Core/Private/UObject/UnrealNames.cpp
+++ b/Engine/Source/Runtime/Core/Private/UObject/UnrealNames.cpp
@@ -1907,6 +1907,19 @@ bool FName::IsValidXName(const FString& InName, const FString& InInvalidChars, F
 	return true;
 }
 
+template <typename CharType, int N>
+void CheckLazyName(const CharType(&Literal)[N])
+{
+	check(FName(Literal) == FLazyName(Literal));
+	check(FLazyName(Literal) == FName(Literal));
+	check(FLazyName(Literal) == FLazyName(Literal));
+	check(FName(Literal) == FLazyName(Literal).Resolve());
+
+	CharType Literal2[N];
+	FMemory::Memcpy(Literal2, Literal);
+	check(FLazyName(Literal) == FLazyName(Literal2));
+}
+
 void FName::AutoTest()
 {
 #if DO_CHECK
@@ -2078,15 +2091,22 @@ void FName::AutoTest()
 	check(Names[5] == "FooC");
 	check(Names[6] == FooWide);
 
-	check(FLazyName("Hej") == FName("Hej"));
-	check(FLazyName("Hej_0") == FName("Hej_0"));
-	check(FLazyName("Hej_00") == FName("Hej_00"));
-	check(FLazyName("Hej_1") == FName("Hej_1"));
-	check(FLazyName("Hej_01") == FName("Hej_01"));
-	check(FLazyName("Hej_-1") == FName("Hej_-1"));
-	check(FLazyName("Hej__0") == FName("Hej__0"));
-	check(FLazyName("Hej_2147483647") == FName("Hej_2147483647"));
-	check(FLazyName("Hej_123") == FLazyName(FName("Hej_123")));
+	
+	CheckLazyName("Hej");
+	CheckLazyName(TEXT("Hej"));
+	CheckLazyName("Hej_0");
+	CheckLazyName("Hej_00");
+	CheckLazyName("Hej_1");
+	CheckLazyName("Hej_01");
+	CheckLazyName("Hej_-1");
+	CheckLazyName("Hej__0");
+	CheckLazyName("Hej_2147483647");
+	CheckLazyName("Hej_123");
+	CheckLazyName("None");
+	CheckLazyName("none");
+	CheckLazyName("None_0");
+	CheckLazyName("None_1");
+
 
 #if 0
 	// Check hash table growth still yields the same unique FName ids
@@ -2230,15 +2250,66 @@ void FName::TearDown()
 	{
 		GetNamePoolPostInit().~FNamePool();
 		bNamePoolInitialized = false;
-	
 	}
 }
 
-uint32 FLazyName::ParseNumber(const TCHAR* Str, int32 Len)
+FName FLazyName::Resolve() const
+{
+	// Make a stack copy to ensure thread-safety
+	FLiteralOrName Copy = Either;
+
+	if (Copy.IsName())
+	{
+		FNameEntryId Id = Copy.AsName();
+		return FName(Id, Id, Number);
+	}
+
+	// Resolve to FName but throw away the number part
+	FNameEntryId Id = bLiteralIsWide ? FName(Copy.AsWideLiteral()).GetComparisonIndex()
+										: FName(Copy.AsAnsiLiteral()).GetComparisonIndex();
+
+	// Deliberately unsynchronized write of word-sized int, ok if multiple threads resolve same lazy name
+	Either = FLiteralOrName(Id);
+
+	return FName(Id, Id, Number);		
+}
+
+uint32 FLazyName::ParseNumber(const ANSICHAR* Str, int32 Len)
 {
 	return FNameHelper::ParseNumber(Str, Len);
 }
 
+uint32 FLazyName::ParseNumber(const WIDECHAR* Str, int32 Len)
+{
+	return FNameHelper::ParseNumber(Str, Len);
+}
+
+bool operator==(const FLazyName& A, const FLazyName& B)
+{
+	// If we have started creating FNames we might as well resolve and cache both lazy names
+	if (A.Either.IsName() || B.Either.IsName())
+	{
+		return A.Resolve() == B.Resolve();
+	}
+
+	// Literal pointer comparison, can ignore width
+	if (A.Either.AsAnsiLiteral() == B.Either.AsAnsiLiteral())
+	{
+		return true;
+	}
+
+	if (A.bLiteralIsWide)
+	{
+		return B.bLiteralIsWide ? FPlatformString::Stricmp(A.Either.AsWideLiteral(), B.Either.AsWideLiteral()) == 0
+								: FPlatformString::Stricmp(A.Either.AsWideLiteral(), B.Either.AsAnsiLiteral()) == 0;
+	}
+	else
+	{
+		return B.bLiteralIsWide ? FPlatformString::Stricmp(A.Either.AsAnsiLiteral(), B.Either.AsWideLiteral()) == 0
+								: FPlatformString::Stricmp(A.Either.AsAnsiLiteral(), B.Either.AsAnsiLiteral()) == 0;	
+	}
+}
+
 #if !UE_BUILD_SHIPPING && !UE_BUILD_TEST
 
 #include "Containers/StackTracker.h"
diff --git a/Engine/Source/Runtime/Core/Public/UObject/NameTypes.h b/Engine/Source/Runtime/Core/Public/UObject/NameTypes.h
--- a/Engine/Source/Runtime/Core/Public/UObject/NameTypes.h
+++ b/Engine/Source/Runtime/Core/Public/UObject/NameTypes.h
@@ -985,47 +985,47 @@ class FLazyName
 public:
 	FLazyName()
 		: Either(FNameEntryId())
-	{
-	}
+	{}
+
+	/** @param Literal must be a string literal */
+	template<int N>
+	FLazyName(const WIDECHAR(&Literal)[N])
+		: Either(Literal)
+		, Number(ParseNumber(Literal, N - 1))
+		, bLiteralIsWide(true)
+	{}
 
 	/** @param Literal must be a string literal */
 	template<int N>
-	FLazyName(const TCHAR(&Literal)[N])
+	FLazyName(const ANSICHAR(&Literal)[N])
 		: Either(Literal)
-		, Number(ParseNumber(Literal, N))
+		, Number(ParseNumber(Literal, N - 1))
+		, bLiteralIsWide(false)
 	{}
 
 	explicit FLazyName(FName Name)
 		: Either(Name.GetComparisonIndex())
 		, Number(Name.GetNumber())
 	{}
-
+	
 	operator FName() const
 	{
-		// Make a stack copy to ensure thread-safety
-		FLiteralOrName Copy = Either;
-
-		if (Copy.IsName())
-		{
-			FNameEntryId Id = Copy.AsName();
-			return FName(Id, Id, Number);
-		}
-
-		// Resolve to FName
-		FNameEntryId Id = FName(Copy.AsLiteral()).GetComparisonIndex();
-
-		// Deliberately unsynchronized write of word-sized int, ok if multiple threads resolve same lazy name
-		Either = FLiteralOrName(Id);
-
-		return FName(Id, Id, Number);
+		return Resolve();
 	}
 
+	CORE_API FName Resolve() const;
+
 private:
 	struct FLiteralOrName
 	{
-#if PLATFORM_32BITS
-		// NOTE: uint64 instead of UPTRINT here since the high bit of a 32-bit pointer may be in use
-		FLiteralOrName(const TCHAR* Literal)
+		// NOTE: uses high bit of pointer for flag; this may be an issue in future when high byte of address may be used for features like hardware ASAN
+		static constexpr uint64 LiteralFlag = UPTRINT(1) << (sizeof(UPTRINT) * 8 - 1);
+
+		explicit FLiteralOrName(const ANSICHAR* Literal)
+			: Int(reinterpret_cast<uint64>(Literal) | LiteralFlag)
+		{}
+		
+		explicit FLiteralOrName(const WIDECHAR* Literal)
 			: Int(reinterpret_cast<uint64>(Literal) | LiteralFlag)
 		{}
 
@@ -1033,8 +1033,6 @@ private:
 			: Int(Name.ToUnstableInt())
 		{}
 
-		static constexpr uint64 LiteralFlag = uint64(1) << (sizeof(uint64) * 8 - 1);
-
 		bool IsName() const
 		{
 			return (LiteralFlag & Int) == 0;
@@ -1049,84 +1047,54 @@ private:
 		{
 			return FNameEntryId::FromUnstableInt(static_cast<uint32>(Int));
 		}
-
-		const TCHAR* AsLiteral() const
-		{
-			return reinterpret_cast<const TCHAR*>(Int & ~LiteralFlag);
-		}
-
-		uint64 Int;
-#else
-		// NOTE: uses high bit of pointer for flag; this may be an issue in future when high byte of address may be used for features like hardware ASAN
-		FLiteralOrName(const TCHAR* Literal)
-			: Int(reinterpret_cast<UPTRINT>(Literal) | LiteralFlag)
-		{}
-
-		explicit FLiteralOrName(FNameEntryId Name)
-			: Int(Name.ToUnstableInt())
-		{}
-
-		static constexpr UPTRINT LiteralFlag = UPTRINT(1) << (sizeof(UPTRINT) * 8 - 1);
-
-		bool IsName() const
-		{
-			return (LiteralFlag & Int) == 0;
-		}
-
-		bool IsLiteral() const
-		{
-			return (LiteralFlag & Int) != 0;
-		}
-
-		FNameEntryId AsName() const
+		
+		const ANSICHAR* AsAnsiLiteral() const
 		{
-			return FNameEntryId::FromUnstableInt(static_cast<uint32>(Int));
+			return reinterpret_cast<const ANSICHAR*>(Int & ~LiteralFlag);
 		}
 
-		const TCHAR* AsLiteral() const
+		const WIDECHAR* AsWideLiteral() const
 		{
-			return reinterpret_cast<const TCHAR*>(Int & ~LiteralFlag);
+			return reinterpret_cast<const WIDECHAR*>(Int & ~LiteralFlag);
 		}
 
-		UPTRINT Int;
-#endif
+		uint64 Int;
 	};
 
 	mutable FLiteralOrName Either;
 	mutable uint32 Number = 0;
 
-	CORE_API static uint32 ParseNumber(const TCHAR* Literal, int32 Len);
+	// Distinguishes WIDECHAR* and ANSICHAR* literals, doesn't indicate if literal contains any wide characters 
+	bool bLiteralIsWide = false;
+	
+	CORE_API static uint32 ParseNumber(const WIDECHAR* Literal, int32 Len);
+	CORE_API static uint32 ParseNumber(const ANSICHAR* Literal, int32 Len);
 
 public:
-	friend bool operator==(FName Name, FLazyName Lazy)
+
+	friend bool operator==(FName Name, const FLazyName& Lazy)
 	{
-		if (Lazy.Either.IsName())
+		// If !Name.IsNone(), we have started creating FNames
+		// and might as well resolve and cache Lazy
+		if (Lazy.Either.IsName() || !Name.IsNone())
 		{
-			return Name == FName(Lazy);
+			return Name == Lazy.Resolve();
 		}
-		else
-		{
-			return Name == Lazy.Either.AsLiteral();
-		}
-	}
-
-	friend bool operator==(FLazyName A, FLazyName B)
-	{
-		if (A.Either.IsName())
+		else if (!Lazy.bLiteralIsWide)
 		{
-			return FName(A) == B;
+			return Name == Lazy.Either.AsAnsiLiteral();
 		}
-		else if (B.Either.IsName())
+		else
 		{
-			return FName(B) == A;
+			return Name == Lazy.Either.AsWideLiteral();
 		}
-
-		return A.Either.AsLiteral() == B.Either.AsLiteral() ||
-			FPlatformString::Stricmp(A.Either.AsLiteral(), B.Either.AsLiteral()) == 0;
 	}
-	
-	friend bool operator==(FLazyName Lazy, FName Name)
+
+	friend bool operator==(const FLazyName& Lazy, FName Name)
 	{
 		return Name == Lazy;
 	}
+
+	CORE_API friend bool operator==(const FLazyName& A, const FLazyName& B);
+
 };
