- Fixed file handle leak and file handle registry corruption on Linux. Preveted file opened in 'write' mode form being handled with the file handle registry reserved for files opened in 'read' mode.

diff --git a/Engine/Source/Runtime/Core/Private/Unix/UnixPlatformFile.cpp b/Engine/Source/Runtime/Core/Private/Unix/UnixPlatformFile.cpp
--- a/Engine/Source/Runtime/Core/Private/Unix/UnixPlatformFile.cpp
+++ b/Engine/Source/Runtime/Core/Private/Unix/UnixPlatformFile.cpp
@@ -171,10 +171,23 @@ public:
 
 	virtual bool Read(uint8* Destination, int64 BytesToRead) override
 	{
-		int64 BytesRead = 0;
-		// handle virtual file handles
-		GFileRegistry.TrackStartRead(this);
+		struct FScopedReadTracker
+		{
+			FScopedReadTracker(FFileHandleUnix* InHandle) : Handle(InHandle)
+			{
+				if (Handle) { GFileRegistry.TrackStartRead(Handle); }
+			}
+			~FScopedReadTracker()
+			{
+				if (Handle) { GFileRegistry.TrackEndRead(Handle); }
+			}
+			FFileHandleUnix* Handle;
+		};
 
+		// Handle virtual file handles (only in read mode, write mode doesn't use the file handle registry)
+		FScopedReadTracker ScopedReadTracker(FileOpenAsWrite ? nullptr : this);
+
+		int64 BytesRead = 0;
 		{
 			FScopedDiskUtilizationTracker Tracker(BytesToRead, FileOffset);
 			// seek to the offset on seek? this matches console behavior more closely
@@ -186,8 +199,6 @@ public:
 			BytesRead += ReadInternal(Destination, BytesToRead);
 		}
 
-		// handle virtual file handles
-		GFileRegistry.TrackEndRead(this);
 		FileOffset += BytesRead;
 		return BytesRead == BytesToRead;
 	}
