PID 0 is invalid, getpriority does not give a bad errno for PID 0

This was causing huge hangs in Editor Analytics when trying to process a large amount of session data that was not being cleaned up
Due to it thinking that session was still running with a default PID of 0

diff --git a/Engine/Source/Runtime/Core/Private/Mac/MacPlatformProcess.cpp b/Engine/Source/Runtime/Core/Private/Mac/MacPlatformProcess.cpp
--- a/Engine/Source/Runtime/Core/Private/Mac/MacPlatformProcess.cpp
+++ b/Engine/Source/Runtime/Core/Private/Mac/MacPlatformProcess.cpp
@@ -684,6 +684,12 @@ bool FMacPlatformProcess::GetProcReturnCode( FProcHandle& ProcessHandle, int32*
 
 bool FMacPlatformProcess::IsApplicationRunning( uint32 ProcessId )
 {
+	// PID 0 is not a valid user application so lets ignore it as valid
+	if (ProcessId == 0)
+	{
+		return false;
+	}
+
 	errno = 0;
 	getpriority(PRIO_PROCESS, ProcessId);
 	return errno == 0;
diff --git a/Engine/Source/Runtime/Core/Private/Unix/UnixPlatformProcess.cpp b/Engine/Source/Runtime/Core/Private/Unix/UnixPlatformProcess.cpp
--- a/Engine/Source/Runtime/Core/Private/Unix/UnixPlatformProcess.cpp
+++ b/Engine/Source/Runtime/Core/Private/Unix/UnixPlatformProcess.cpp
@@ -1558,6 +1558,12 @@ bool FUnixPlatformProcess::Daemonize()
 
 bool FUnixPlatformProcess::IsApplicationRunning( uint32 ProcessId )
 {
+	// PID 0 is not a valid user application so lets ignore it as valid
+	if (ProcessId == 0)
+	{
+		return false;
+	}
+
 	errno = 0;
 	getpriority(PRIO_PROCESS, ProcessId);
 	return errno == 0;
