Read shader compiler version directly from dll and remove code writing hashed version to file.  Fixes race condition where two shader compile workers attempt to create the version file simultaneously.
#rnx

diff --git a/Engine/Source/Runtime/Core/Private/Windows/WindowsPlatformMisc.cpp b/Engine/Source/Runtime/Core/Private/Windows/WindowsPlatformMisc.cpp
--- a/Engine/Source/Runtime/Core/Private/Windows/WindowsPlatformMisc.cpp
+++ b/Engine/Source/Runtime/Core/Private/Windows/WindowsPlatformMisc.cpp
@@ -2758,3 +2758,23 @@ void FWindowsPlatformMisc::PumpMessagesOutsideMainLoop()
 	PeekMessage(&Msg, NULL, 0, 0, PM_NOREMOVE | PM_QS_SENDMESSAGE);
 	return;
 }
+
+uint64 FWindowsPlatformMisc::GetFileVersion(const FString &FileName)
+{
+	::DWORD VersionInfoSize = GetFileVersionInfoSize(*FileName, NULL);
+	if (VersionInfoSize != 0)
+	{
+		TArray<uint8> VersionInfo;
+		VersionInfo.AddUninitialized(VersionInfoSize);
+		if (GetFileVersionInfo(*FileName, NULL, VersionInfoSize, VersionInfo.GetData()))
+		{
+			VS_FIXEDFILEINFO *FileInfo;
+			::UINT FileInfoLen;
+			if (VerQueryValue(VersionInfo.GetData(), TEXT("\\"), (LPVOID*)&FileInfo, &FileInfoLen))
+			{
+				return (uint64(FileInfo->dwFileVersionMS) << 32) | uint64(FileInfo->dwFileVersionLS);
+			}
+		}
+	}
+	return 0;
+}
\ No newline at end of file
diff --git a/Engine/Source/Runtime/Core/Public/Windows/WindowsPlatformMisc.h b/Engine/Source/Runtime/Core/Public/Windows/WindowsPlatformMisc.h
--- a/Engine/Source/Runtime/Core/Public/Windows/WindowsPlatformMisc.h
+++ b/Engine/Source/Runtime/Core/Public/Windows/WindowsPlatformMisc.h
@@ -253,6 +253,8 @@ struct CORE_API FWindowsPlatformMisc
 
 	static void PumpMessagesOutsideMainLoop();
 
+	static uint64 GetFileVersion(const FString &FileName);
+
 };
 
 
