Only set GIsRequestingExit on the start of an Engine Tick. If its done in the middle of a tick bad things can happen!

diff --git a/Engine/Source/Runtime/Core/Private/Android/AndroidPlatformMisc.cpp b/Engine/Source/Runtime/Core/Private/Android/AndroidPlatformMisc.cpp
--- a/Engine/Source/Runtime/Core/Private/Android/AndroidPlatformMisc.cpp
+++ b/Engine/Source/Runtime/Core/Private/Android/AndroidPlatformMisc.cpp
@@ -2309,3 +2309,10 @@ uint32 FAndroidMisc::GetCoreFrequency(int32 CoreIndex, ECoreFrequencyProperty Co
 	}
 	return ReturnFrequency;
 }
+
+void FAndroidMisc::NonReentrentRequestExit()
+{
+	// Cheating here to grab access to this. This function should only be used in exterme cases in which non-reentrent functions are needed (ie. crash handling/signal handler)
+	extern bool GShouldRequestExit;
+	GShouldRequestExit = true;
+}
diff --git a/Engine/Source/Runtime/Core/Private/Misc/CoreGlobals.cpp b/Engine/Source/Runtime/Core/Private/Misc/CoreGlobals.cpp
--- a/Engine/Source/Runtime/Core/Private/Misc/CoreGlobals.cpp
+++ b/Engine/Source/Runtime/Core/Private/Misc/CoreGlobals.cpp
@@ -254,24 +254,30 @@ static void appNoop()
 {
 }
 
+// This should be left non static to allow *edge* cases only in Core to extern an set this.
+bool GShouldRequestExit = false;
+
+void CORE_API HandleRequestExitIfSetDuringTick()
+{
+	if (GShouldRequestExit)
+	{
+PRAGMA_DISABLE_DEPRECATION_WARNINGS
+		GIsRequestingExit = true;
+PRAGMA_ENABLE_DEPRECATION_WARNINGS
+	}
+}
+
 void CORE_API RequestEngineExit(const TCHAR* ReasonString)
 {
 	ensureMsgf(ReasonString && FCString::Strlen(ReasonString) > 4, TEXT("RequestEngineExit must be given a valid reason (reason \"%s\""), ReasonString);
 
-PRAGMA_DISABLE_DEPRECATION_WARNINGS
-	UE_LOG(LogCore, Log, TEXT("Engine exit requested (reason: %s%s)"), ReasonString, GIsRequestingExit ? TEXT("; note: exit was already requested") : TEXT(""));
-	GIsRequestingExit = true;
-PRAGMA_ENABLE_DEPRECATION_WARNINGS
+	UE_LOG(LogCore, Log, TEXT("Engine exit requested (reason: %s%s)"), ReasonString, GShouldRequestExit ? TEXT("; note: exit was already requested") : TEXT(""));
+	GShouldRequestExit = true;
 }
 
 void CORE_API RequestEngineExit(const FString& ReasonString)
 {
-	ensureMsgf(ReasonString.Len() > 4, TEXT("RequestEngineExit must be given a valid reason (reason \"%s\""), *ReasonString);
-
-PRAGMA_DISABLE_DEPRECATION_WARNINGS
-	UE_LOG(LogCore, Log, TEXT("Engine exit requested (reason: %s%s)"), *ReasonString, GIsRequestingExit ? TEXT("; note: exit was already requested") : TEXT(""));
-	GIsRequestingExit = true;
-PRAGMA_ENABLE_DEPRECATION_WARNINGS
+	RequestEngineExit(*ReasonString);
 }
 
 /** Exec handler for game debugging tool, allowing commands like "editactor", ...							*/
diff --git a/Engine/Source/Runtime/Core/Private/Unix/UnixPlatformMisc.cpp b/Engine/Source/Runtime/Core/Private/Unix/UnixPlatformMisc.cpp
--- a/Engine/Source/Runtime/Core/Private/Unix/UnixPlatformMisc.cpp
+++ b/Engine/Source/Runtime/Core/Private/Unix/UnixPlatformMisc.cpp
@@ -356,9 +356,9 @@ void FUnixPlatformMisc::RequestExit(bool Force)
 	{
 		// Lets set our selfs to request exit as the generic platform request exit could log
 		// This is deprecated but one of the few excpetions to leave around for now as we dont want to UE_LOG as that may allocate memory
-PRAGMA_DISABLE_DEPRECATION_WARNINGS
-		GIsRequestingExit = 1;
-PRAGMA_ENABLE_DEPRECATION_WARNINGS
+		// ShouldRequestExit *should* only be used this way in cases where non-reentrent code is required
+		extern bool GShouldRequestExit;
+		GShouldRequestExit = 1;
 	}
 	else
 	{
diff --git a/Engine/Source/Runtime/Core/Public/Android/AndroidPlatformMisc.h b/Engine/Source/Runtime/Core/Public/Android/AndroidPlatformMisc.h
--- a/Engine/Source/Runtime/Core/Public/Android/AndroidPlatformMisc.h
+++ b/Engine/Source/Runtime/Core/Public/Android/AndroidPlatformMisc.h
@@ -266,8 +266,11 @@ public:
 	// Window access is locked by the game thread before preinit and unlocked here after RHIInit (PlatformCreateDynamicRHI). 
 	static void UnlockAndroidWindow();
 
+	// Android specific requesting of exit, *ONLY* use this function in signal handling code. Otherwise normal RequestExit functions
+	static void NonReentrentRequestExit();
+
 private:
-    static EDeviceScreenOrientation DeviceOrientation;
+	static EDeviceScreenOrientation DeviceOrientation;
 };
 
 #if !PLATFORM_LUMIN
diff --git a/Engine/Source/Runtime/Core/Public/CoreGlobals.h b/Engine/Source/Runtime/Core/Public/CoreGlobals.h
--- a/Engine/Source/Runtime/Core/Public/CoreGlobals.h
+++ b/Engine/Source/Runtime/Core/Public/CoreGlobals.h
@@ -256,6 +256,11 @@ extern CORE_API bool GIsGuarded;
 UE_DEPRECATED(4.24, "Please use IsEngineExitRequested()/RequestEngineExit(const FString&)")
 extern CORE_API bool GIsRequestingExit;
 
+/**
+ * This will check if a RequestExit has come in, if it has will set GIsRequestingExit.
+ */
+extern CORE_API void HandleRequestExitIfSetDuringTick();
+
 FORCEINLINE bool IsEngineExitRequested()
 {
 PRAGMA_DISABLE_DEPRECATION_WARNINGS
diff --git a/Engine/Source/Runtime/Launch/Private/Android/LaunchAndroid.cpp b/Engine/Source/Runtime/Launch/Private/Android/LaunchAndroid.cpp
--- a/Engine/Source/Runtime/Launch/Private/Android/LaunchAndroid.cpp
+++ b/Engine/Source/Runtime/Launch/Private/Android/LaunchAndroid.cpp
@@ -1449,9 +1449,7 @@ static void OnAppCommandCB(struct android_app* app, int32_t cmd)
 				FCoreDelegates::ApplicationWillTerminateDelegate.Broadcast();
 			}, TStatId(), NULL, ENamedThreads::GameThread);
 			FTaskGraphInterface::Get().WaitUntilTaskCompletes(WillTerminateTask);
-PRAGMA_DISABLE_DEPRECATION_WARNINGS
-			GIsRequestingExit = true; //destroy immediately. Game will shutdown.
-PRAGMA_ENABLE_DEPRECATION_WARNINGS
+			FAndroidMisc::NonReentrentRequestExit();
 		}));
 
 
diff --git a/Engine/Source/Runtime/Launch/Private/LaunchEngineLoop.cpp b/Engine/Source/Runtime/Launch/Private/LaunchEngineLoop.cpp
--- a/Engine/Source/Runtime/Launch/Private/LaunchEngineLoop.cpp
+++ b/Engine/Source/Runtime/Launch/Private/LaunchEngineLoop.cpp
@@ -4228,8 +4228,8 @@ static inline void EndFrameRenderThread(FRHICommandListImmediate& RHICmdList)
 
 void FEngineLoop::Tick()
 {
-    // make sure to catch any FMemStack uses outside of UWorld::Tick
-    FMemMark MemStackMark(FMemStack::Get());
+	// make sure to catch any FMemStack uses outside of UWorld::Tick
+	FMemMark MemStackMark(FMemStack::Get());
 
 #if !UE_BUILD_SHIPPING && !UE_BUILD_TEST && MALLOC_GT_HOOKS
 	FScopedSampleMallocChurn ChurnTracker;
@@ -4239,6 +4239,8 @@ void FEngineLoop::Tick()
 
 	LLM_SCOPE(ELLMTag::EngineMisc);
 
+	HandleRequestExitIfSetDuringTick();
+
 	// Send a heartbeat for the diagnostics thread
 	FThreadHeartBeat::Get().HeartBeat(true);
 	FGameThreadHitchHeartBeat::Get().FrameStart();
