Fix file helper utilities not returning failure when closing/flushing fails.  This caused a corner case bug in [redacted]

diff --git a/Engine/Source/Runtime/Core/Private/Misc/FileHelper.cpp b/Engine/Source/Runtime/Core/Private/Misc/FileHelper.cpp
--- a/Engine/Source/Runtime/Core/Private/Misc/FileHelper.cpp
+++ b/Engine/Source/Runtime/Core/Private/Misc/FileHelper.cpp
@@ -202,7 +202,11 @@ bool FFileHelper::SaveArrayToFile(TArrayView<const uint8> Array, const TCHAR* Fi
 	{
 		return false;
 	}
-	Ar->Serialize(const_cast<uint8*>(Array.GetData()), Array.Num());	
+	Ar->Serialize(const_cast<uint8*>(Array.GetData()), Array.Num());
+
+	// Always explicitly close to catch errors from flush/close
+	Ar->Close();
+
 	return !Ar->IsError() && !Ar->IsCriticalError();
 }
 
@@ -249,6 +253,9 @@ bool FFileHelper::SaveStringToFile( const FString& String, const TCHAR* Filename
 		Ar->Serialize( (ANSICHAR*)Src.Get(), Src.Length() * sizeof(ANSICHAR) );
 	}
 
+	// Always explicitly close to catch errors from flush/close
+	Ar->Close();
+
 	return !Ar->IsError() && !Ar->IsCriticalError();
 }
 
