Reducing stack size usage by LogTrace in debug builds

(ushell-p4-cherrypick of 9630599 by [REDACTED_PERSON])

diff --git a/Engine/Source/Runtime/Core/Public/Logging/LogTrace.h b/Engine/Source/Runtime/Core/Public/Logging/LogTrace.h
--- a/Engine/Source/Runtime/Core/Public/Logging/LogTrace.h
+++ b/Engine/Source/Runtime/Core/Public/Logging/LogTrace.h
@@ -24,7 +24,7 @@ struct FLogTrace
 	CORE_API static void OutputLogMessageSpec(const void* LogPoint, const FLogCategoryBase* Category, ELogVerbosity::Type Verbosity, const ANSICHAR* File, int32 Line, const TCHAR* Format);
 
 	template <typename... Types>
-	static void OutputLogMessage(const void* LogPoint, Types... FormatArgs)
+	FORCENOINLINE static void OutputLogMessage(const void* LogPoint, Types... FormatArgs)
 	{
 		uint8 FormatArgsBuffer[3072];
 		uint16 FormatArgsSize = FFormatArgsTrace::EncodeArguments(FormatArgsBuffer, FormatArgs...);
