We now spoof a crash report and send it from the CrashReportClientApp when detecting an abnormal shutdown has occurred. That is to say, a shutdown where, as far as we can tell, no known exit path was followed.

Reverted change to where FCrashReportAnalytics was initialized now that we get those settings from the UECrashContext file.

Added DelayedSend analytics attribute that determines whether or not the process that is sending an analytics event was the same one that created it.

diff --git a/Engine/Source/Developer/EditorAnalyticsSession/Private/EditorSessionSummarySender.cpp b/Engine/Source/Developer/EditorAnalyticsSession/Private/EditorSessionSummarySender.cpp
--- a/Engine/Source/Developer/EditorAnalyticsSession/Private/EditorSessionSummarySender.cpp
+++ b/Engine/Source/Developer/EditorAnalyticsSession/Private/EditorSessionSummarySender.cpp
@@ -32,7 +32,6 @@ FEditorSessionSummarySender::FEditorSessionSummarySender(IAnalyticsProvider& InA
 	, Sender(InSenderName)
 	, CurrentSessionProcessId(InCurrentSessionProcessId)
 {
-	SendStoredSessions();
 }
 
 FEditorSessionSummarySender::~FEditorSessionSummarySender()
@@ -62,6 +61,40 @@ void FEditorSessionSummarySender::SetCurrentSessionExitCode(const int32 InCurren
 	CurrentSessionExitCode = InExitCode;
 }
 
+bool FEditorSessionSummarySender::FindCurrentSession(FEditorAnalyticsSession& OutSession) const
+{
+	if (FPlatformProcess::IsApplicationRunning(CurrentSessionProcessId))
+	{
+		// still running, can't be abnormal termination
+		return false;
+	}
+
+	bool bFound = false;
+
+	if (FEditorAnalyticsSession::Lock(FTimespan::FromMilliseconds(100)))
+	{
+		TArray<FEditorAnalyticsSession> ExistingSessions;
+		FEditorAnalyticsSession::LoadAllStoredSessions(ExistingSessions);
+
+		const int32 ProcessID = CurrentSessionProcessId;
+		FEditorAnalyticsSession* CurrentSession = ExistingSessions.FindByPredicate(
+			[ProcessID](const FEditorAnalyticsSession& Session)
+			{
+				return Session.PlatformProcessID == ProcessID;
+			});
+
+		if (CurrentSession != nullptr)
+		{
+			OutSession = *CurrentSession;
+			bFound = true;
+		}
+
+		FEditorAnalyticsSession::Unlock();
+	}
+
+	return bFound;
+}
+
 void FEditorSessionSummarySender::SendStoredSessions(const bool bForceSendCurrentSession) const
 {
 	TArray<FEditorAnalyticsSession> SessionsToReport;
@@ -179,12 +212,16 @@ void FEditorSessionSummarySender::SendSessionSummaryEvent(const FEditorAnalytics
 	AnalyticsAttributes.Emplace(TEXT("IsInVRMode"), Session.bIsInVRMode);
 	AnalyticsAttributes.Emplace(TEXT("SentFrom"), Sender);
 
-	AnalyticsProvider.RecordEvent(TEXT("SessionSummary"), AnalyticsAttributes);
+	// was this sent from some other process than itself or the out-of-process monitor for that run?
+	AnalyticsAttributes.Emplace(TEXT("DelayedSend"), Session.PlatformProcessID != CurrentSessionProcessId);
 
 	if (Session.PlatformProcessID == CurrentSessionProcessId && CurrentSessionExitCode.IsSet())
 	{
 		AnalyticsAttributes.Emplace(TEXT("ExitCode"), CurrentSessionExitCode.GetValue());
 	}
 
+	// Send the event.
+	AnalyticsProvider.RecordEvent(TEXT("SessionSummary"), AnalyticsAttributes);
+
 	UE_LOG(LogEditorSessionSummary, Log, TEXT("EditorSessionSummary sent report. Type=%s, SessionId=%s"), *ShutdownTypeString, *SessionIdString);
 }
diff --git a/Engine/Source/Developer/EditorAnalyticsSession/Public/EditorSessionSummarySender.h b/Engine/Source/Developer/EditorAnalyticsSession/Public/EditorSessionSummarySender.h
--- a/Engine/Source/Developer/EditorAnalyticsSession/Public/EditorSessionSummarySender.h
+++ b/Engine/Source/Developer/EditorAnalyticsSession/Public/EditorSessionSummarySender.h
@@ -22,6 +22,7 @@ public:
 	void Shutdown();
 
 	void SetCurrentSessionExitCode(const int32 InCurrentSessionProcessId, const int32 InExitCode);
+	bool FindCurrentSession(FEditorAnalyticsSession& OutSession) const;
 
 private:
 	/** Send any stored Sessions. */
diff --git a/Engine/Source/Editor/DetailCustomizations/Private/ImportantToggleSettingCustomization.cpp b/Engine/Source/Editor/DetailCustomizations/Private/ImportantToggleSettingCustomization.cpp
--- a/Engine/Source/Editor/DetailCustomizations/Private/ImportantToggleSettingCustomization.cpp
+++ b/Engine/Source/Editor/DetailCustomizations/Private/ImportantToggleSettingCustomization.cpp
@@ -97,13 +97,14 @@ void FImportantToggleSettingCustomization::CustomizeDetails(IDetailLayoutBuilder
 	if (Objects.Num() == 1)
 	{
 		ToggleSettingObject = Objects[0];
-		IImportantToggleSettingInterface* ToogleSettingInterface = Cast<IImportantToggleSettingInterface>(ToggleSettingObject.Get());
+		IImportantToggleSettingInterface* ToggleSettingInterface = Cast<IImportantToggleSettingInterface>(ToggleSettingObject.Get());
 
-		if (ToogleSettingInterface != nullptr)
+		if (ToggleSettingInterface != nullptr)
 		{
 			FName CategoryName;
 			FName PropertyName;
-			ToogleSettingInterface->GetToogleCategoryAndPropertyNames(CategoryName, PropertyName);
+			ToggleSettingInterface->GetToggleCategoryAndPropertyNames(CategoryName, PropertyName);
+
 			IDetailCategoryBuilder& Category = DetailBuilder.EditCategory(CategoryName);
 			TogglePropertyHandle = DetailBuilder.GetProperty(PropertyName);
 
@@ -127,8 +128,8 @@ void FImportantToggleSettingCustomization::CustomizeDetails(IDetailLayoutBuilder
 					[
 						SNew(SImportantToggleButton)
 						.CheckBoxStyle(FEditorStyle::Get(), "Property.ToggleButton.Start")
-						.Text(ToogleSettingInterface->GetFalseStateLabel())
-						.ToolTipText(ToogleSettingInterface->GetFalseStateTooltip())
+						.Text(ToggleSettingInterface->GetFalseStateLabel())
+						.ToolTipText(ToggleSettingInterface->GetFalseStateTooltip())
 						.IsSet(this, &FImportantToggleSettingCustomization::IsToggleValue, false)
 						.OnToggled(this, &FImportantToggleSettingCustomization::OnToggledTo, false)
 					]
@@ -137,8 +138,8 @@ void FImportantToggleSettingCustomization::CustomizeDetails(IDetailLayoutBuilder
 					[
 						SNew(SImportantToggleButton)
 						.CheckBoxStyle(FEditorStyle::Get(), "Property.ToggleButton.End")
-						.Text(ToogleSettingInterface->GetTrueStateLabel())
-						.ToolTipText(ToogleSettingInterface->GetTrueStateTooltip())
+						.Text(ToggleSettingInterface->GetTrueStateLabel())
+						.ToolTipText(ToggleSettingInterface->GetTrueStateTooltip())
 						.IsSet(this, &FImportantToggleSettingCustomization::IsToggleValue, true)
 						.OnToggled(this, &FImportantToggleSettingCustomization::OnToggledTo, true)
 					]
@@ -151,8 +152,8 @@ void FImportantToggleSettingCustomization::CustomizeDetails(IDetailLayoutBuilder
 						.VAlign(VAlign_Center)
 						[
 							SNew(SHyperlink)
-							.Text(ToogleSettingInterface->GetAdditionalInfoUrlLabel())
-							.OnNavigate(this, &FImportantToggleSettingCustomization::OnNavigateHyperlink, ToogleSettingInterface->GetAdditionalInfoUrl())
+							.Text(ToggleSettingInterface->GetAdditionalInfoUrlLabel())
+							.OnNavigate(this, &FImportantToggleSettingCustomization::OnNavigateHyperlink, ToggleSettingInterface->GetAdditionalInfoUrl())
 						]
 					]
 				]
diff --git a/Engine/Source/Editor/UnrealEd/Classes/Analytics/AnalyticsPrivacySettings.h b/Engine/Source/Editor/UnrealEd/Classes/Analytics/AnalyticsPrivacySettings.h
--- a/Engine/Source/Editor/UnrealEd/Classes/Analytics/AnalyticsPrivacySettings.h
+++ b/Engine/Source/Editor/UnrealEd/Classes/Analytics/AnalyticsPrivacySettings.h
@@ -27,7 +27,7 @@ class UAnalyticsPrivacySettings : public UObject, public IImportantToggleSetting
 public:
 
 	// BEGIN IImportantToggleSettingInterface
-	virtual void GetToogleCategoryAndPropertyNames(FName& OutCategory, FName& OutProperty) const override;
+	virtual void GetToggleCategoryAndPropertyNames(FName& OutCategory, FName& OutProperty) const override;
 	virtual FText GetFalseStateLabel() const override;
 	virtual FText GetFalseStateTooltip() const override;
 	virtual FText GetFalseStateDescription() const override;
diff --git a/Engine/Source/Editor/UnrealEd/Classes/Analytics/CrashReportsPrivacySettings.h b/Engine/Source/Editor/UnrealEd/Classes/Analytics/CrashReportsPrivacySettings.h
--- a/Engine/Source/Editor/UnrealEd/Classes/Analytics/CrashReportsPrivacySettings.h
+++ b/Engine/Source/Editor/UnrealEd/Classes/Analytics/CrashReportsPrivacySettings.h
@@ -27,7 +27,7 @@ class UCrashReportsPrivacySettings : public UObject, public IImportantToggleSett
 public:
 
 	// BEGIN IImportantToggleSettingInterface
-	virtual void GetToogleCategoryAndPropertyNames(FName& OutCategory, FName& OutProperty) const override;
+	virtual void GetToggleCategoryAndPropertyNames(FName& OutCategory, FName& OutProperty) const override;
 	virtual FText GetFalseStateLabel() const override;
 	virtual FText GetFalseStateTooltip() const override;
 	virtual FText GetFalseStateDescription() const override;
@@ -37,4 +37,10 @@ public:
 	virtual FString GetAdditionalInfoUrl() const override;
 	virtual FText GetAdditionalInfoUrlLabel() const override;
 	// END IImportantToggleSettingInterface
+
+#if WITH_EDITOR
+	//~ Begin UObject Interface
+	virtual void PostEditChangeProperty(struct FPropertyChangedEvent& PropertyChangedEvent) override;
+	//~ End UObject Interface
+#endif
 };
diff --git a/Engine/Source/Editor/UnrealEd/Private/Analytics/AnalyticsPrivacySettings.cpp b/Engine/Source/Editor/UnrealEd/Private/Analytics/AnalyticsPrivacySettings.cpp
--- a/Engine/Source/Editor/UnrealEd/Private/Analytics/AnalyticsPrivacySettings.cpp
+++ b/Engine/Source/Editor/UnrealEd/Private/Analytics/AnalyticsPrivacySettings.cpp
@@ -4,6 +4,7 @@
 #include "UObject/UnrealType.h"
 #include "Interfaces/IAnalyticsProvider.h"
 #include "EngineAnalytics.h"
+#include "Misc/CoreDelegates.h"
 
 #define LOCTEXT_NAMESPACE "AnalyticsPrivacySettings"
 
@@ -13,7 +14,7 @@ UAnalyticsPrivacySettings::UAnalyticsPrivacySettings(const FObjectInitializer& O
 {
 }
 
-void UAnalyticsPrivacySettings::GetToogleCategoryAndPropertyNames(FName& OutCategory, FName& OutProperty) const
+void UAnalyticsPrivacySettings::GetToggleCategoryAndPropertyNames(FName& OutCategory, FName& OutProperty) const
 {
 	OutCategory = FName("Options");
 	OutProperty = FName("bSendUsageData");
@@ -64,10 +65,17 @@ void UAnalyticsPrivacySettings::PostEditChangeProperty(struct FPropertyChangedEv
 {
 	Super::PostEditChangeProperty(PropertyChangedEvent);
 
-	const FName PropertyName = (PropertyChangedEvent.Property != NULL) ? PropertyChangedEvent.Property->GetFName() : NAME_None;
+	const FName PropertyName = (PropertyChangedEvent.Property != nullptr) ? PropertyChangedEvent.Property->GetFName() : NAME_None;
 	if (PropertyName == GET_MEMBER_NAME_CHECKED(UAnalyticsPrivacySettings, bSendUsageData))
 	{
 		OnSendFullUsageDataChanged();
+
+		FCrashOverrideParameters Params;
+		Params.bSetCrashReportClientMessageText = false;
+		Params.bSetGameNameSuffix = false;
+		Params.SendUsageData = bSendUsageData;
+
+		FCoreDelegates::CrashOverrideParamsChanged.Broadcast(Params);
 	}
 }
 
diff --git a/Engine/Source/Editor/UnrealEd/Private/Analytics/CrashReportsPrivacySettings.cpp b/Engine/Source/Editor/UnrealEd/Private/Analytics/CrashReportsPrivacySettings.cpp
--- a/Engine/Source/Editor/UnrealEd/Private/Analytics/CrashReportsPrivacySettings.cpp
+++ b/Engine/Source/Editor/UnrealEd/Private/Analytics/CrashReportsPrivacySettings.cpp
@@ -5,6 +5,7 @@
 #include "Interfaces/IAnalyticsProvider.h"
 #include "EngineAnalytics.h"
 #include "Misc/ConfigCacheIni.h"
+#include "Misc/CoreDelegates.h"
 
 #define LOCTEXT_NAMESPACE "CrashReportsPrivacySettings"
 
@@ -18,7 +19,7 @@ UCrashReportsPrivacySettings::UCrashReportsPrivacySettings(const FObjectInitiali
 	}
 }
 
-void UCrashReportsPrivacySettings::GetToogleCategoryAndPropertyNames(FName& OutCategory, FName& OutProperty) const
+void UCrashReportsPrivacySettings::GetToggleCategoryAndPropertyNames(FName& OutCategory, FName& OutProperty) const
 {
 	OutCategory = FName("Options");
 	OutProperty = FName("bSendUnattendedBugReports");
@@ -64,4 +65,20 @@ FText UCrashReportsPrivacySettings::GetAdditionalInfoUrlLabel() const
 	return LOCTEXT("HyperlinkLabel", "Epic Games Privacy Policy");
 };
 
+#if WITH_EDITOR
+void UCrashReportsPrivacySettings::PostEditChangeProperty(struct FPropertyChangedEvent& PropertyChangedEvent)
+{
+	const FName PropertyName = (PropertyChangedEvent.Property != nullptr) ? PropertyChangedEvent.Property->GetFName() : NAME_None;
+	if (PropertyName == GET_MEMBER_NAME_CHECKED(UCrashReportsPrivacySettings, bSendUnattendedBugReports))
+	{
+		FCrashOverrideParameters Params;
+		Params.bSetCrashReportClientMessageText = false;
+		Params.bSetGameNameSuffix = false;
+		Params.SendUnattendedBugReports = bSendUnattendedBugReports;
+
+		FCoreDelegates::CrashOverrideParamsChanged.Broadcast(Params);
+	}
+}
+#endif
+
 #undef LOCTEXT_NAMESPACE
diff --git a/Engine/Source/Programs/CrashReportClient/Private/CrashReportClientApp.cpp b/Engine/Source/Programs/CrashReportClient/Private/CrashReportClientApp.cpp
--- a/Engine/Source/Programs/CrashReportClient/Private/CrashReportClientApp.cpp
+++ b/Engine/Source/Programs/CrashReportClient/Private/CrashReportClientApp.cpp
@@ -18,6 +18,7 @@
 #include "HAL/PlatformApplicationMisc.h"
 #include "HAL/PlatformCrashContext.h"
 #include "IAnalyticsProviderET.h"
+#include "XmlParser.h"
 
 #if !CRASH_REPORT_UNATTENDED_ONLY
 	#include "SCrashReportClient.h"
@@ -38,6 +39,7 @@
 #include "RecoveryService.h"
 
 #if CRASH_REPORT_WITH_MTBF
+#include "EditorAnalyticsSession.h"
 #include "EditorSessionSummarySender.h"
 #endif
 
@@ -452,7 +454,6 @@ FPlatformErrorReport CollectErrorReport(FRecoveryService* RecoveryService, uint3
 	{
 		DirectoryExists = FGenericCrashContext::CreateCrashReportDirectory(
 			SharedCrashContext.SessionContext.CrashGUIDRoot,
-			SharedCrashContext.SessionContext.GameName,
 			0,
 			ReportDirectoryAbsolutePath);
 	}
@@ -474,7 +475,10 @@ FPlatformErrorReport CollectErrorReport(FRecoveryService* RecoveryService, uint3
 	CrashContext.SerializeAsXML(*CrashContextXMLPath);
 
 #if CRASH_REPORT_WITH_RECOVERY
-	if (RecoveryService && DirectoryExists && SharedCrashContext.bSendUsageData && SharedCrashContext.CrashType != ECrashContextType::Ensure && FCrashReportCoreConfig::Get().GetSendLogFile())
+	if (RecoveryService && 
+		DirectoryExists && 
+		SharedCrashContext.UserSettings.bSendUsageData && 
+		SharedCrashContext.CrashType != ECrashContextType::Ensure)
 	{
 		RecoveryService->CollectFiles(ReportDirectoryAbsolutePath);
 	}
@@ -553,6 +557,159 @@ bool IsCrashReportAvailable(uint32 WatchedProcess, FSharedCrashContext& CrashCon
 	return false;
 }
 
+static void DeleteTempCrashContextFile(uint64 ProcessID)
+{
+	const FString SessionContextFile = FGenericCrashContext::GetTempSessionContextFilePath(ProcessID);
+	FPlatformFileManager::Get().GetPlatformFile().DeleteFile(*SessionContextFile);
+}
+
+#if CRASH_REPORT_WITH_MTBF
+
+template <typename Type>
+bool FindAndParseValue(const TMap<FString, FString>& Map, const FString& Key, Type& OutValue)
+{
+	const FString* ValueString = Map.Find(Key);
+	if (ValueString != nullptr)
+	{
+		TTypeFromString<Type>::FromString(OutValue, **ValueString);
+		return true;
+	}
+
+	return false;
+}
+
+template <size_t Size>
+bool FindAndCopyValue(const TMap<FString, FString>& Map, const FString& Key, TCHAR (&OutValue)[Size])
+{
+	const FString* ValueString = Map.Find(Key);
+	if (ValueString != nullptr)
+	{
+		FCString::Strncpy(OutValue, **ValueString, Size);
+		return true;
+	}
+
+	return false;
+}
+
+static bool LoadTempCrashContextFromFile(FSharedCrashContext& CrashContext, uint64 ProcessID)
+{
+	const FString TempContextFilePath = FGenericCrashContext::GetTempSessionContextFilePath(ProcessID);
+
+	FXmlFile File;
+	if (!File.LoadFile(TempContextFilePath))
+	{
+		return false;
+	}
+
+	TMap<FString, FString> ContextProperties;
+	for (FXmlNode* Node : File.GetRootNode()->GetChildrenNodes())
+	{
+		ContextProperties.Add(Node->GetTag(), Node->GetContent());
+	}
+
+	FSessionContext& SessionContext = CrashContext.SessionContext;
+
+	FindAndParseValue(ContextProperties, TEXT("SecondsSinceStart"), SessionContext.SecondsSinceStart);
+	FindAndParseValue(ContextProperties, TEXT("IsInternalBuild"), SessionContext.bIsInternalBuild);
+	FindAndParseValue(ContextProperties, TEXT("IsPerforceBuild"), SessionContext.bIsPerforceBuild);
+	FindAndParseValue(ContextProperties, TEXT("IsSourceDistribution"), SessionContext.bIsSourceDistribution);
+	FindAndCopyValue(ContextProperties, TEXT("GameName"), SessionContext.GameName);
+	FindAndCopyValue(ContextProperties, TEXT("ExecutableName"), SessionContext.ExecutableName);
+	FindAndCopyValue(ContextProperties, TEXT("GameSessionID"), SessionContext.GameSessionID);
+	FindAndCopyValue(ContextProperties, TEXT("EngineMode"), SessionContext.EngineMode);
+	FindAndCopyValue(ContextProperties, TEXT("EngineModeEx"), SessionContext.EngineModeEx);
+	FindAndCopyValue(ContextProperties, TEXT("DeploymentName"), SessionContext.DeploymentName);
+	FindAndCopyValue(ContextProperties, TEXT("CommandLine"), SessionContext.CommandLine);
+	FindAndParseValue(ContextProperties, TEXT("LanguageLCID"), SessionContext.LanguageLCID);
+	FindAndCopyValue(ContextProperties, TEXT("AppDefaultLocale"), SessionContext.DefaultLocale);
+	FindAndParseValue(ContextProperties, TEXT("IsUE4Release"), SessionContext.bIsUE4Release);
+	FindAndCopyValue(ContextProperties, TEXT("UserName"), SessionContext.UserName);
+	FindAndCopyValue(ContextProperties, TEXT("BaseDir"), SessionContext.BaseDir);
+	FindAndCopyValue(ContextProperties, TEXT("RootDir"), SessionContext.RootDir);
+	FindAndCopyValue(ContextProperties, TEXT("LoginId"), SessionContext.LoginIdStr);
+	FindAndCopyValue(ContextProperties, TEXT("EpicAccountId"), SessionContext.EpicAccountId);
+	FindAndCopyValue(ContextProperties, TEXT("UserActivityHint"), SessionContext.UserActivityHint);
+	FindAndParseValue(ContextProperties, TEXT("CrashDumpMode"), SessionContext.CrashDumpMode);
+	FindAndCopyValue(ContextProperties, TEXT("GameStateName"), SessionContext.GameStateName);
+	FindAndParseValue(ContextProperties, TEXT("Misc.NumberOfCores"), SessionContext.NumberOfCores);
+	FindAndParseValue(ContextProperties, TEXT("Misc.NumberOfCoresIncludingHyperthreads"), SessionContext.NumberOfCoresIncludingHyperthreads);
+	FindAndCopyValue(ContextProperties, TEXT("Misc.CPUVendor"), SessionContext.CPUVendor);
+	FindAndCopyValue(ContextProperties, TEXT("Misc.CPUBrand"), SessionContext.CPUBrand);
+	FindAndCopyValue(ContextProperties, TEXT("Misc.PrimaryGPUBrand"), SessionContext.PrimaryGPUBrand);
+	FindAndCopyValue(ContextProperties, TEXT("Misc.OSVersionMajor"), SessionContext.OsVersion);
+	FindAndCopyValue(ContextProperties, TEXT("Misc.OSVersionMinor"), SessionContext.OsSubVersion);
+	FindAndParseValue(ContextProperties, TEXT("MemoryStats.AvailablePhysical"), SessionContext.MemoryStats.AvailablePhysical);
+	FindAndParseValue(ContextProperties, TEXT("MemoryStats.AvailableVirtual"), SessionContext.MemoryStats.AvailableVirtual);
+	FindAndParseValue(ContextProperties, TEXT("MemoryStats.UsedPhysical"), SessionContext.MemoryStats.UsedPhysical);
+	FindAndParseValue(ContextProperties, TEXT("MemoryStats.PeakUsedPhysical"), SessionContext.MemoryStats.PeakUsedPhysical);
+	FindAndParseValue(ContextProperties, TEXT("MemoryStats.UsedVirtual"), SessionContext.MemoryStats.UsedVirtual);
+	FindAndParseValue(ContextProperties, TEXT("MemoryStats.PeakUsedVirtual"), SessionContext.MemoryStats.PeakUsedVirtual);
+	FindAndParseValue(ContextProperties, TEXT("MemoryStats.bIsOOM"), SessionContext.bIsOOM);
+	FindAndParseValue(ContextProperties, TEXT("MemoryStats.OOMAllocationSize"), SessionContext.OOMAllocationSize);
+	FindAndParseValue(ContextProperties, TEXT("MemoryStats.OOMAllocationAlignment"), SessionContext.OOMAllocationAlignment);
+
+	// user settings
+	FUserSettingsContext& UserSettings = CrashContext.UserSettings;
+
+	FindAndParseValue(ContextProperties, TEXT("NoDialog"), UserSettings.bNoDialog);
+	FindAndParseValue(ContextProperties, TEXT("SendUnattendedBugReports"), UserSettings.bSendUnattendedBugReports);
+	FindAndParseValue(ContextProperties, TEXT("SendUsageData"), UserSettings.bSendUsageData);
+	FindAndCopyValue(ContextProperties, TEXT("LogFilePath"), UserSettings.LogFilePath);
+
+	return true;
+}
+
+static void HandleAbnormalShutdown(FSharedCrashContext& CrashContext, uint64 ProcessID, void* WritePipe, const TSharedPtr<FRecoveryService>& RecoveryService)
+{
+	CrashContext.CrashType = ECrashContextType::AbnormalShutdown;
+
+	IPlatformFile& PlatformFile = FPlatformFileManager::Get().GetPlatformFile();
+
+	// create a temporary crash directory
+	const FString TempCrashDirectory = FPlatformProcess::UserTempDir() / FString::Printf(TEXT("UECrashContext-%d"), ProcessID);
+	FCString::Strcpy(CrashContext.CrashFilesDirectory, *TempCrashDirectory);
+
+	if (PlatformFile.CreateDirectory(CrashContext.CrashFilesDirectory))
+	{
+		// copy the log file to the temporary directory
+		const FString LogDestination = TempCrashDirectory / FPaths::GetCleanFilename(CrashContext.UserSettings.LogFilePath);
+		PlatformFile.CopyFile(*LogDestination, CrashContext.UserSettings.LogFilePath);
+
+		FPlatformErrorReport ErrorReport = CollectErrorReport(RecoveryService.Get(), ProcessID, CrashContext, WritePipe);
+		SubmitCrashReportResult Result = SendErrorReport(ErrorReport, /*bNoDialog*/ true);
+
+		// delete the temporary directory
+		PlatformFile.DeleteDirectoryRecursively(*TempCrashDirectory);
+
+		if (CrashContext.UserSettings.bSendUsageData)
+		{
+			// If analytics is enabled make sure they are submitted now.
+			FCrashReportAnalytics::GetProvider().BlockUntilFlushed(5.0f);
+		}
+	}
+}
+
+static bool WasAbnormalShutdown(const FEditorSessionSummarySender& SessionSummarySender)
+{
+	FEditorAnalyticsSession AnalyticsSession;
+	if (SessionSummarySender.FindCurrentSession(AnalyticsSession))
+	{
+		// check if this was an abnormal shutdown (aka. none of the known shutdown types, and not debugged)
+		if (AnalyticsSession.bCrashed == false &&
+			AnalyticsSession.bGPUCrashed == false &&
+			AnalyticsSession.bWasShutdown == false &&
+			AnalyticsSession.bIsTerminating == false &&
+			AnalyticsSession.bWasEverDebugger == false)
+		{
+			return true;
+		}
+	}
+
+	return false;
+}
+
+#endif 
+
 void RunCrashReportClient(const TCHAR* CommandLine)
 {
 	// Override the stack size for the thread pool.
@@ -618,22 +775,6 @@ void RunCrashReportClient(const TCHAR* CommandLine)
 		RecoveryServicePtr = MakeShared<FRecoveryService>(MonitorPid);
 #endif
 
-		FCrashReportAnalytics::Initialize();
-
-#if CRASH_REPORT_WITH_MTBF
-		TUniquePtr<FEditorSessionSummarySender> EditorSessionSummarySender;
-		if (FCrashReportCoreConfig::Get().GetAllowToBeContacted())
-		{
-			EditorSessionSummarySender = MakeUnique<FEditorSessionSummarySender>(FCrashReportAnalytics::GetProvider(), TEXT("CrashReportClient"), MonitorPid);
-
-			FTicker::GetCoreTicker().AddTicker(TEXT("EditorSessionSummarySender"), 0, [&EditorSessionSummarySender](float DeltaTime)
-			{
-				EditorSessionSummarySender->Tick(DeltaTime);
-				return true;
-			});
-		}
-#endif
-
 		FProcHandle MonitoredProcess = FPlatformProcess::OpenProcess(MonitorPid);
 		if (!MonitoredProcess.IsValid())
 		{
@@ -660,9 +801,15 @@ void RunCrashReportClient(const TCHAR* CommandLine)
 				FSharedCrashContext CrashContext;
 				if (IsCrashReportAvailable(MonitorPid, CrashContext, MonitorReadPipe))
 				{
+					const bool bReportCrashAnalyticInfo = CrashContext.UserSettings.bSendUsageData;
+					if (bReportCrashAnalyticInfo)
+					{
+						FCrashReportAnalytics::Initialize();
+					}
+
 					// Build error report in memory.
 					FPlatformErrorReport ErrorReport = CollectErrorReport(RecoveryServicePtr.Get(), MonitorPid, CrashContext, MonitorWritePipe);
-
+					
 #if CRASH_REPORT_WITH_RECOVERY
 					if (RecoveryServicePtr && !FPrimaryCrashProperties::Get()->bIsEnsure)
 					{
@@ -670,7 +817,9 @@ void RunCrashReportClient(const TCHAR* CommandLine)
 						RecoveryServicePtr.Reset();
 					}
 #endif
-					const SubmitCrashReportResult Result = SendErrorReport(ErrorReport, CrashContext.bNoDialog && CrashContext.bSendUnattenededBugReports);
+
+					const bool bNoDialog = CrashContext.UserSettings.bNoDialog && CrashContext.UserSettings.bSendUnattendedBugReports;
+					const SubmitCrashReportResult Result = SendErrorReport(ErrorReport, bNoDialog);
 
 					// At this point the game can continue execution. It is important this happens
 					// as soon as thread state and minidump has been created, so that ensures cause
@@ -678,7 +827,6 @@ void RunCrashReportClient(const TCHAR* CommandLine)
 					uint8 ResponseCode[] = { 0xd, 0xe, 0xa, 0xd };
 					FPlatformProcess::WritePipe(MonitorWritePipe, ResponseCode, sizeof(ResponseCode));
 
-					bool bReportCrashAnalyticInfo = CrashContext.bSendUsageData;
 					if (bReportCrashAnalyticInfo)
 					{
 						// If analytics is enabled make sure they are submitted now.
@@ -713,17 +861,47 @@ void RunCrashReportClient(const TCHAR* CommandLine)
 		}
 
 #if CRASH_REPORT_WITH_MTBF
-		if (EditorSessionSummarySender.IsValid())
 		{
-			// Query this again, as the crash reporting loop above may have exited before setting this information (via IsEngineExitRequested)
-			bApplicationAlive = IsMonitoredProcessAlive(ApplicationReturnCode);
-			if (!bApplicationAlive)
+			// load our temporary crash context file
+			FSharedCrashContext TempCrashContext;
+			FMemory::Memzero(TempCrashContext);
+			if (LoadTempCrashContextFromFile(TempCrashContext, MonitorPid))
 			{
-				EditorSessionSummarySender->SetCurrentSessionExitCode(MonitorPid, ApplicationReturnCode);
+				TempCrashContext.SessionContext.ProcessId = MonitorPid;
+
+				if (TempCrashContext.UserSettings.bSendUsageData)
+				{
+					// Query this again, as the crash reporting loop above may have exited before setting this information (via IsEngineExitRequested)
+					bApplicationAlive = IsMonitoredProcessAlive(ApplicationReturnCode);
+					if (!bApplicationAlive)
+					{
+						if (!FCrashReportAnalytics::IsAvailable())
+						{
+							FCrashReportAnalytics::Initialize();
+						}
+
+						// initialize analytics
+						TUniquePtr<FEditorSessionSummarySender> EditorSessionSummarySender = MakeUnique<FEditorSessionSummarySender>(FCrashReportAnalytics::GetProvider(), TEXT("CrashReportClient"), MonitorPid);
+						EditorSessionSummarySender->SetCurrentSessionExitCode(MonitorPid, ApplicationReturnCode);
+
+						if (TempCrashContext.UserSettings.bSendUnattendedBugReports)
+						{
+							// Send a spoofed crash report in the case that we detect an abnormal shutdown has occurred
+							if (WasAbnormalShutdown(*EditorSessionSummarySender.Get()))
+							{
+								HandleAbnormalShutdown(TempCrashContext, MonitorPid, MonitorWritePipe, RecoveryServicePtr);
+							}
+						}
+
+						// send analytics and shutdown
+						EditorSessionSummarySender->Shutdown();
+					}
+				}
 			}
-			EditorSessionSummarySender->Shutdown();
 		}
 #endif
+		// clean up the context file
+		DeleteTempCrashContextFile(MonitorPid);
 
 		FPlatformProcess::CloseProc(MonitoredProcess);
 
diff --git a/Engine/Source/Runtime/Core/Private/GenericPlatform/GenericPlatformCrashContext.cpp b/Engine/Source/Runtime/Core/Private/GenericPlatform/GenericPlatformCrashContext.cpp
--- a/Engine/Source/Runtime/Core/Private/GenericPlatform/GenericPlatformCrashContext.cpp
+++ b/Engine/Source/Runtime/Core/Private/GenericPlatform/GenericPlatformCrashContext.cpp
@@ -61,6 +61,7 @@ const TCHAR* const FGenericCrashContext::CrashTypeAssert = TEXT("Assert");
 const TCHAR* const FGenericCrashContext::CrashTypeEnsure = TEXT("Ensure");
 const TCHAR* const FGenericCrashContext::CrashTypeGPU = TEXT("GPUCrash");
 const TCHAR* const FGenericCrashContext::CrashTypeHang = TEXT("Hang");
+const TCHAR* const FGenericCrashContext::CrashTypeAbnormalShutdown = TEXT("AbnormalShutdown");
 
 const TCHAR* const FGenericCrashContext::EngineModeExUnknown = TEXT("Unset");
 const TCHAR* const FGenericCrashContext::EngineModeExDirty = TEXT("Dirty");
@@ -75,6 +76,7 @@ const FGuid FGenericCrashContext::ExecutionGuid = FGuid::NewGuid();
 namespace NCached
 {
 	static FSessionContext Session;
+	static FUserSettingsContext UserSettings;
 	static TArray<FString> EnabledPluginsList;
 	static TMap<FString, FString> EngineData;
 	static TMap<FString, FString> GameData;
@@ -169,7 +171,7 @@ void FGenericCrashContext::Initialize()
 	const float PollingInterval = 1.0f;
 	FTicker::GetCoreTicker().AddTicker( FTickerDelegate::CreateLambda( []( float DeltaTime )
 	{
-        QUICK_SCOPE_CYCLE_COUNTER(STAT_NCachedCrashContextProperties_LambdaTicker);
+    QUICK_SCOPE_CYCLE_COUNTER(STAT_NCachedCrashContextProperties_LambdaTicker);
 
 		NCached::Session.SecondsSinceStart = int32(FPlatformTime::Seconds() - GStartTime);
 		return true;
@@ -201,6 +203,14 @@ PRAGMA_DISABLE_DEPRECATION_WARNINGS
 		{
 			FCString::Strcpy(NCached::Session.GameName, *(FString(TEXT("UE4-")) + FApp::GetProjectName() + InParams.GameNameSuffix));
 		}
+		if (InParams.SendUnattendedBugReports.IsSet())
+		{
+			NCached::UserSettings.bSendUnattendedBugReports = InParams.SendUnattendedBugReports.GetValue();
+		}
+		if (InParams.SendUsageData.IsSet())
+		{
+			NCached::UserSettings.bSendUsageData = InParams.SendUsageData.GetValue();
+		}
 	});
 PRAGMA_ENABLE_DEPRECATION_WARNINGS
 
@@ -211,6 +221,12 @@ PRAGMA_ENABLE_DEPRECATION_WARNINGS
 
 	FCoreDelegates::ConfigReadyForUse.AddStatic(FGenericCrashContext::InitializeFromConfig);
 
+	FCString::Strcpy(NCached::UserSettings.LogFilePath, *FPlatformOutputDevices::GetAbsoluteLogFilename());
+
+	NCached::UserSettings.bNoDialog = FApp::IsUnattended() || IsRunningDedicatedServer();
+
+	SerializeTempCrashContextToFile();
+
 	bIsInitialized = true;
 #endif	// !NOINITCRASHREPORTER
 }
@@ -258,6 +274,8 @@ void FGenericCrashContext::InitializeFromContext(const FSessionContext& Session,
 		}
 	}
 
+	SerializeTempCrashContextToFile();
+
 	bIsInitialized = true;
 }
 
@@ -265,6 +283,7 @@ void FGenericCrashContext::CopySharedCrashContext(FSharedCrashContext& Dst)
 {
 	//Copy the session
 	FMemory::Memcpy(Dst.SessionContext, NCached::Session);
+	FMemory::Memcpy(Dst.UserSettings, NCached::UserSettings);
 
 	TCHAR* DynamicDataStart = &Dst.DynamicData[0];
 	TCHAR* DynamicDataPtr = DynamicDataStart;
@@ -309,6 +328,8 @@ void FGenericCrashContext::SetMemoryStats(const FPlatformMemoryStats& InMemorySt
 	NCached::Session.bIsOOM = FPlatformMemory::bIsOOM;
 	NCached::Session.OOMAllocationSize = FPlatformMemory::OOMAllocationSize;
 	NCached::Session.OOMAllocationAlignment = FPlatformMemory::OOMAllocationAlignment;
+
+	SerializeTempCrashContextToFile();
 }
 
 void FGenericCrashContext::InitializeFromConfig()
@@ -336,8 +357,15 @@ void FGenericCrashContext::InitializeFromConfig()
 	// Read the initial un-localized crash context text
 	UpdateLocalizedStrings();
 
+	// Set privacy settings
+	GConfig->GetBool(TEXT("/Script/UnrealEd.CrashReportsPrivacySettings"), TEXT("bSendUnattendedBugReports"), NCached::UserSettings.bSendUnattendedBugReports, GEditorSettingsIni);
+	GConfig->GetBool(TEXT("/Script/UnrealEd.AnalyticsPrivacySettings"), TEXT("bSendUsageData"), NCached::UserSettings.bSendUsageData, GEditorSettingsIni);
+
 	// Make sure we get updated text once the localized version is loaded
 	FTextLocalizationManager::Get().OnTextRevisionChangedEvent.AddStatic(&UpdateLocalizedStrings);
+
+	SerializeTempCrashContextToFile();
+
 #endif	// !NOINITCRASHREPORTER
 }
 
@@ -363,140 +391,176 @@ FGenericCrashContext::FGenericCrashContext(ECrashContextType InType, const TCHAR
 	CrashContextIndex = StaticCrashContextIndex++;
 }
 
-void FGenericCrashContext::SerializeContentToBuffer() const
+FString FGenericCrashContext::GetTempSessionContextFilePath(uint64 ProcessID)
 {
-	TCHAR CrashGUID[CrashGUIDLength];
-	GetUniqueCrashName(CrashGUID, CrashGUIDLength);
+	return FPlatformProcess::UserTempDir() / FString::Printf(TEXT("UECrashContext-%u.xml"), ProcessID);
+}
 
-	// Must conform against:
-	// https://www.securecoding.cert.org/confluence/display/seccode/SIG30-C.+Call+only+asynchronous-safe+functions+within+signal+handlers
-	AddHeader();
+void FGenericCrashContext::SerializeTempCrashContextToFile()
+{
+	if (!IsOutOfProcessCrashReporter())
+	{
+		return;
+	}
 
-	BeginSection( RuntimePropertiesTag );
-	AddCrashProperty( TEXT( "CrashVersion" ), (int32)ECrashDescVersions::VER_3_CrashContext );
-	AddCrashProperty( TEXT( "ExecutionGuid" ), *ExecutionGuid.ToString() );
-	AddCrashProperty( TEXT( "CrashGUID" ), (const TCHAR*)CrashGUID);
-	AddCrashProperty( TEXT( "ProcessId" ), NCached::Session.ProcessId );
-	AddCrashProperty( TEXT( "IsInternalBuild" ), NCached::Session.bIsInternalBuild );
-	AddCrashProperty( TEXT( "IsPerforceBuild" ), NCached::Session.bIsPerforceBuild );
-	AddCrashProperty( TEXT( "IsSourceDistribution" ), NCached::Session.bIsSourceDistribution );
-	AddCrashProperty( TEXT( "IsEnsure" ), (Type == ECrashContextType::Ensure) );
-	AddCrashProperty( TEXT( "IsAssert" ), (Type == ECrashContextType::Assert) );
-	AddCrashProperty( TEXT( "CrashType" ), GetCrashTypeString(Type) );
+	FString SessionBuffer;
+	SessionBuffer.Reserve(32 * 1024);
 
-	AddCrashProperty( TEXT( "SecondsSinceStart" ), NCached::Session.SecondsSinceStart );
+	AddHeader(SessionBuffer);
+
+	SerializeSessionContext(SessionBuffer);
+	SerializeUserSettings(SessionBuffer);
+
+	AddFooter(SessionBuffer);
+
+	const FString SessionFilePath = GetTempSessionContextFilePath(NCached::Session.ProcessId);
+	FFileHelper::SaveStringToFile(SessionBuffer, *SessionFilePath);
+}
+
+void FGenericCrashContext::SerializeSessionContext(FString& Buffer)
+{
+	AddCrashPropertyInternal(Buffer, TEXT("ProcessId"), NCached::Session.ProcessId);
+	AddCrashPropertyInternal(Buffer, TEXT("SecondsSinceStart"), NCached::Session.SecondsSinceStart);
+
+	AddCrashPropertyInternal(Buffer, TEXT("IsInternalBuild"), NCached::Session.bIsInternalBuild);
+	AddCrashPropertyInternal(Buffer, TEXT("IsPerforceBuild"), NCached::Session.bIsPerforceBuild);
+	AddCrashPropertyInternal(Buffer, TEXT("IsSourceDistribution"), NCached::Session.bIsSourceDistribution);
 
-	// Add common crash properties.
 	if (FCString::Strlen(NCached::Session.GameName) > 0)
 	{
-		AddCrashProperty(TEXT("GameName"), NCached::Session.GameName);
+		AddCrashPropertyInternal(Buffer, TEXT("GameName"), NCached::Session.GameName);
 	}
 	else
 	{
 		const TCHAR* ProjectName = FApp::GetProjectName();
 		if (ProjectName != nullptr && ProjectName[0] != 0)
 		{
-			AddCrashProperty(TEXT("GameName"), *FString::Printf(TEXT("UE4-%s"), ProjectName));
+			AddCrashPropertyInternal(Buffer, TEXT("GameName"), *FString::Printf(TEXT("UE4-%s"), ProjectName));
 		}
 		else
 		{
-			AddCrashProperty(TEXT("GameName"), TEXT(""));
+			AddCrashPropertyInternal(Buffer, TEXT("GameName"), TEXT(""));
 		}
 	}
-	AddCrashProperty( TEXT( "ExecutableName" ), NCached::Session.ExecutableName );
-	AddCrashProperty( TEXT( "BuildConfiguration" ), LexToString( FApp::GetBuildConfiguration() ) );
-	AddCrashProperty( TEXT( "GameSessionID" ), NCached::Session.GameSessionID );
-	
+	AddCrashPropertyInternal(Buffer, TEXT("ExecutableName"), NCached::Session.ExecutableName);
+	AddCrashPropertyInternal(Buffer, TEXT("BuildConfiguration"), LexToString(FApp::GetBuildConfiguration()));
+	AddCrashPropertyInternal(Buffer, TEXT("GameSessionID"), NCached::Session.GameSessionID);
+
 	// Unique string specifying the symbols to be used by CrashReporter
-	FString Symbols = FString::Printf( TEXT( "%s" ), FApp::GetBuildVersion());
+	FString Symbols = FString::Printf(TEXT("%s"), FApp::GetBuildVersion());
 #ifdef UE_APP_FLAVOR
-	Symbols = FString::Printf(TEXT( "%s-%s" ), *Symbols, *FString(UE_APP_FLAVOR));
+	Symbols = FString::Printf(TEXT("%s-%s"), *Symbols, *FString(UE_APP_FLAVOR));
 #endif
-	Symbols = FString::Printf(TEXT("%s-%s-%s"), *Symbols, FPlatformMisc::GetUBTPlatform(), LexToString(FApp::GetBuildConfiguration())).Replace( TEXT( "+" ), TEXT( "*" ));
+	Symbols = FString::Printf(TEXT("%s-%s-%s"), *Symbols, FPlatformMisc::GetUBTPlatform(), LexToString(FApp::GetBuildConfiguration())).Replace(TEXT("+"), TEXT("*"));
 #ifdef UE_BUILD_FLAVOR
-	Symbols = FString::Printf(TEXT( "%s-%s" ), *Symbols, *FString(UE_BUILD_FLAVOR));
-#endif
-#ifdef UE_APP_FLAVOR
-	Symbols = FString::Printf(TEXT( "%s-%s" ), *Symbols, *FString(UE_APP_FLAVOR));
+	Symbols = FString::Printf(TEXT("%s-%s"), *Symbols, *FString(UE_BUILD_FLAVOR));
 #endif
 
-	AddCrashProperty( TEXT( "Symbols" ), Symbols);
+	AddCrashPropertyInternal(Buffer, TEXT("Symbols"), Symbols);
 
-	AddCrashProperty( TEXT( "PlatformName" ), FPlatformProperties::PlatformName() );
-	AddCrashProperty( TEXT( "PlatformNameIni" ), FPlatformProperties::IniPlatformName());
-	AddCrashProperty( TEXT( "EngineMode" ), NCached::Session.EngineMode);
-	AddCrashProperty( TEXT( "EngineModeEx" ), NCached::Session.EngineModeEx);
+	AddCrashPropertyInternal(Buffer, TEXT("PlatformName"), FPlatformProperties::PlatformName());
+	AddCrashPropertyInternal(Buffer, TEXT("PlatformNameIni"), FPlatformProperties::IniPlatformName());
+	AddCrashPropertyInternal(Buffer, TEXT("EngineMode"), NCached::Session.EngineMode);
+	AddCrashPropertyInternal(Buffer, TEXT("EngineModeEx"), NCached::Session.EngineModeEx);
 
-	AddCrashProperty( TEXT( "DeploymentName"), NCached::Session.DeploymentName );
+	AddCrashPropertyInternal(Buffer, TEXT("DeploymentName"), NCached::Session.DeploymentName);
 
-	AddCrashProperty( TEXT( "EngineVersion" ), *FEngineVersion::Current().ToString() );
-	AddCrashProperty( TEXT( "CommandLine" ), NCached::Session.CommandLine );
-	AddCrashProperty( TEXT( "LanguageLCID" ), NCached::Session.LanguageLCID );
-	AddCrashProperty( TEXT( "AppDefaultLocale" ), NCached::Session.DefaultLocale );
-	AddCrashProperty( TEXT( "BuildVersion" ), FApp::GetBuildVersion() );
-	AddCrashProperty( TEXT( "IsUE4Release" ), NCached::Session.bIsUE4Release );
-	AddCrashProperty( TEXT( "IsRequestingExit" ), IsEngineExitRequested() );
+	AddCrashPropertyInternal(Buffer, TEXT("EngineVersion"), *FEngineVersion::Current().ToString());
+	AddCrashPropertyInternal(Buffer, TEXT("CommandLine"), NCached::Session.CommandLine);
+	AddCrashPropertyInternal(Buffer, TEXT("LanguageLCID"), NCached::Session.LanguageLCID);
+	AddCrashPropertyInternal(Buffer, TEXT("AppDefaultLocale"), NCached::Session.DefaultLocale);
+	AddCrashPropertyInternal(Buffer, TEXT("BuildVersion"), FApp::GetBuildVersion());
+	AddCrashPropertyInternal(Buffer, TEXT("IsUE4Release"), NCached::Session.bIsUE4Release);
+	AddCrashPropertyInternal(Buffer, TEXT("IsRequestingExit"), IsEngineExitRequested());
 
 	// Remove periods from user names to match AutoReporter user names
 	// The name prefix is read by CrashRepository.AddNewCrash in the website code
 	const bool bSendUserName = NCached::Session.bIsInternalBuild;
 	FString SanitizedUserName = FString(NCached::Session.UserName).Replace(TEXT("."), TEXT(""));
-	AddCrashProperty( TEXT( "UserName" ), bSendUserName ? *SanitizedUserName : TEXT(""));
+	AddCrashPropertyInternal(Buffer, TEXT("UserName"), bSendUserName ? *SanitizedUserName : TEXT(""));
 
-	AddCrashProperty( TEXT( "BaseDir" ), NCached::Session.BaseDir );
-	AddCrashProperty( TEXT( "RootDir" ), NCached::Session.RootDir );
-	AddCrashProperty( TEXT( "MachineId" ), *FString(NCached::Session.LoginIdStr).ToUpper() );
-	AddCrashProperty( TEXT( "LoginId" ), NCached::Session.LoginIdStr );
-	AddCrashProperty( TEXT( "EpicAccountId" ), NCached::Session.EpicAccountId );
+	AddCrashPropertyInternal(Buffer, TEXT("BaseDir"), NCached::Session.BaseDir);
+	AddCrashPropertyInternal(Buffer, TEXT("RootDir"), NCached::Session.RootDir);
+	AddCrashPropertyInternal(Buffer, TEXT("MachineId"), *FString(NCached::Session.LoginIdStr).ToUpper());
+	AddCrashPropertyInternal(Buffer, TEXT("LoginId"), NCached::Session.LoginIdStr);
+	AddCrashPropertyInternal(Buffer, TEXT("EpicAccountId"), NCached::Session.EpicAccountId);
 
-	// Legacy callstack element for current crash reporter
-	AddCrashProperty( TEXT( "NumMinidumpFramesToIgnore"), NumMinidumpFramesToIgnore );
-	AddCrashProperty( TEXT( "CallStack" ), TEXT("") );
-
-	// Add new portable callstack element with crash stack
-	AddPortableCallStack();
-	AddPortableCallStackHash();
-
-	AddCrashProperty( TEXT( "SourceContext" ), TEXT( "" ) );
-	AddCrashProperty( TEXT( "UserDescription" ), TEXT( "" ) );
-	AddCrashProperty( TEXT( "UserActivityHint" ), NCached::Session.UserActivityHint );
-	AddCrashProperty( TEXT( "ErrorMessage" ), ErrorMessage );
-	AddCrashProperty( TEXT( "CrashDumpMode" ), NCached::Session.CrashDumpMode );
-	AddCrashProperty( TEXT( "CrashReporterMessage" ), NCached::Session.CrashReportClientRichText );
+	AddCrashPropertyInternal(Buffer, TEXT("SourceContext"), TEXT(""));
+	AddCrashPropertyInternal(Buffer, TEXT("UserDescription"), TEXT(""));
+	AddCrashPropertyInternal(Buffer, TEXT("UserActivityHint"), NCached::Session.UserActivityHint);
+	AddCrashPropertyInternal(Buffer, TEXT("CrashDumpMode"), NCached::Session.CrashDumpMode);
+	AddCrashPropertyInternal(Buffer, TEXT("GameStateName"), NCached::Session.GameStateName);
 
 	// Add misc stats.
-	AddCrashProperty( TEXT( "Misc.NumberOfCores" ), NCached::Session.NumberOfCores );
-	AddCrashProperty( TEXT( "Misc.NumberOfCoresIncludingHyperthreads" ), NCached::Session.NumberOfCoresIncludingHyperthreads );
-	AddCrashProperty( TEXT( "Misc.Is64bitOperatingSystem" ), (int32)FPlatformMisc::Is64bitOperatingSystem() );
+	AddCrashPropertyInternal(Buffer, TEXT("Misc.NumberOfCores"), NCached::Session.NumberOfCores);
+	AddCrashPropertyInternal(Buffer, TEXT("Misc.NumberOfCoresIncludingHyperthreads"), NCached::Session.NumberOfCoresIncludingHyperthreads);
+	AddCrashPropertyInternal(Buffer, TEXT("Misc.Is64bitOperatingSystem"), (int32) FPlatformMisc::Is64bitOperatingSystem());
 
-	AddCrashProperty( TEXT( "Misc.CPUVendor" ), NCached::Session.CPUVendor );
-	AddCrashProperty( TEXT( "Misc.CPUBrand" ), NCached::Session.CPUBrand );
-	AddCrashProperty( TEXT( "Misc.PrimaryGPUBrand" ), NCached::Session.PrimaryGPUBrand );
-	AddCrashProperty( TEXT( "Misc.OSVersionMajor" ), NCached::Session.OsVersion );
-	AddCrashProperty( TEXT( "Misc.OSVersionMinor" ), NCached::Session.OsSubVersion );
-
-	AddCrashProperty(TEXT("GameStateName"), NCached::Session.GameStateName);
+	AddCrashPropertyInternal(Buffer, TEXT("Misc.CPUVendor"), NCached::Session.CPUVendor);
+	AddCrashPropertyInternal(Buffer, TEXT("Misc.CPUBrand"), NCached::Session.CPUBrand);
+	AddCrashPropertyInternal(Buffer, TEXT("Misc.PrimaryGPUBrand"), NCached::Session.PrimaryGPUBrand);
+	AddCrashPropertyInternal(Buffer, TEXT("Misc.OSVersionMajor"), NCached::Session.OsVersion);
+	AddCrashPropertyInternal(Buffer, TEXT("Misc.OSVersionMinor"), NCached::Session.OsSubVersion);
 
 	// FPlatformMemory::GetConstants is called in the GCreateMalloc, so we can assume it is always valid.
 	{
 		// Add memory stats.
 		const FPlatformMemoryConstants& MemConstants = FPlatformMemory::GetConstants();
 
-		AddCrashProperty( TEXT( "MemoryStats.TotalPhysical" ), (uint64)MemConstants.TotalPhysical );
-		AddCrashProperty( TEXT( "MemoryStats.TotalVirtual" ), (uint64)MemConstants.TotalVirtual );
-		AddCrashProperty( TEXT( "MemoryStats.PageSize" ), (uint64)MemConstants.PageSize );
-		AddCrashProperty( TEXT( "MemoryStats.TotalPhysicalGB" ), MemConstants.TotalPhysicalGB );
+		AddCrashPropertyInternal(Buffer, TEXT("MemoryStats.TotalPhysical"), (uint64) MemConstants.TotalPhysical);
+		AddCrashPropertyInternal(Buffer, TEXT("MemoryStats.TotalVirtual"), (uint64) MemConstants.TotalVirtual);
+		AddCrashPropertyInternal(Buffer, TEXT("MemoryStats.PageSize"), (uint64) MemConstants.PageSize);
+		AddCrashPropertyInternal(Buffer, TEXT("MemoryStats.TotalPhysicalGB"), MemConstants.TotalPhysicalGB);
 	}
 
-	AddCrashProperty( TEXT( "MemoryStats.AvailablePhysical" ), (uint64)NCached::Session.MemoryStats.AvailablePhysical );
-	AddCrashProperty( TEXT( "MemoryStats.AvailableVirtual" ), (uint64)NCached::Session.MemoryStats.AvailableVirtual );
-	AddCrashProperty( TEXT( "MemoryStats.UsedPhysical" ), (uint64)NCached::Session.MemoryStats.UsedPhysical );
-	AddCrashProperty( TEXT( "MemoryStats.PeakUsedPhysical" ), (uint64)NCached::Session.MemoryStats.PeakUsedPhysical );
-	AddCrashProperty( TEXT( "MemoryStats.UsedVirtual" ), (uint64)NCached::Session.MemoryStats.UsedVirtual );
-	AddCrashProperty( TEXT( "MemoryStats.PeakUsedVirtual" ), (uint64)NCached::Session.MemoryStats.PeakUsedVirtual );
-	AddCrashProperty( TEXT( "MemoryStats.bIsOOM" ), (int) NCached::Session.bIsOOM );
-	AddCrashProperty( TEXT( "MemoryStats.OOMAllocationSize"), NCached::Session.OOMAllocationSize );
-	AddCrashProperty( TEXT( "MemoryStats.OOMAllocationAlignment"), NCached::Session.OOMAllocationAlignment );
+	AddCrashPropertyInternal(Buffer, TEXT("MemoryStats.AvailablePhysical"), (uint64) NCached::Session.MemoryStats.AvailablePhysical);
+	AddCrashPropertyInternal(Buffer, TEXT("MemoryStats.AvailableVirtual"), (uint64) NCached::Session.MemoryStats.AvailableVirtual);
+	AddCrashPropertyInternal(Buffer, TEXT("MemoryStats.UsedPhysical"), (uint64) NCached::Session.MemoryStats.UsedPhysical);
+	AddCrashPropertyInternal(Buffer, TEXT("MemoryStats.PeakUsedPhysical"), (uint64) NCached::Session.MemoryStats.PeakUsedPhysical);
+	AddCrashPropertyInternal(Buffer, TEXT("MemoryStats.UsedVirtual"), (uint64) NCached::Session.MemoryStats.UsedVirtual);
+	AddCrashPropertyInternal(Buffer, TEXT("MemoryStats.PeakUsedVirtual"), (uint64) NCached::Session.MemoryStats.PeakUsedVirtual);
+	AddCrashPropertyInternal(Buffer, TEXT("MemoryStats.bIsOOM"), (int) NCached::Session.bIsOOM);
+	AddCrashPropertyInternal(Buffer, TEXT("MemoryStats.OOMAllocationSize"), NCached::Session.OOMAllocationSize);
+	AddCrashPropertyInternal(Buffer, TEXT("MemoryStats.OOMAllocationAlignment"), NCached::Session.OOMAllocationAlignment);
+}
+
+void FGenericCrashContext::SerializeUserSettings(FString& Buffer)
+{
+	AddCrashPropertyInternal(Buffer, TEXT("NoDialog"), NCached::UserSettings.bNoDialog);
+	AddCrashPropertyInternal(Buffer, TEXT("SendUnattendedBugReports"), NCached::UserSettings.bSendUnattendedBugReports);
+	AddCrashPropertyInternal(Buffer, TEXT("SendUsageData"), NCached::UserSettings.bSendUsageData);
+	AddCrashPropertyInternal(Buffer, TEXT("LogFilePath"), NCached::UserSettings.LogFilePath);
+}
+
+void FGenericCrashContext::SerializeContentToBuffer() const
+{
+	TCHAR CrashGUID[CrashGUIDLength];
+	GetUniqueCrashName(CrashGUID, CrashGUIDLength);
+
+	// Must conform against:
+	// https://www.securecoding.cert.org/confluence/display/seccode/SIG30-C.+Call+only+asynchronous-safe+functions+within+signal+handlers
+	AddHeader(CommonBuffer);
+
+	BeginSection( CommonBuffer, RuntimePropertiesTag );
+	AddCrashProperty( TEXT( "CrashVersion" ), (int32)ECrashDescVersions::VER_3_CrashContext );
+	AddCrashProperty( TEXT( "ExecutionGuid" ), *ExecutionGuid.ToString() );
+	AddCrashProperty( TEXT( "CrashGUID" ), (const TCHAR*)CrashGUID);
+
+	AddCrashProperty( TEXT( "IsEnsure" ), (Type == ECrashContextType::Ensure) );
+	AddCrashProperty( TEXT( "IsAssert" ), (Type == ECrashContextType::Assert) );
+	AddCrashProperty( TEXT( "CrashType" ), GetCrashTypeString(Type) );
+	AddCrashProperty( TEXT( "ErrorMessage" ), ErrorMessage );
+	AddCrashProperty( TEXT( "CrashReporterMessage" ), NCached::Session.CrashReportClientRichText );
+
+	SerializeSessionContext(CommonBuffer);
+
+	// Legacy callstack element for current crash reporter
+	AddCrashProperty( TEXT( "NumMinidumpFramesToIgnore"), NumMinidumpFramesToIgnore );
+	AddCrashProperty( TEXT( "CallStack" ), TEXT("") );
+
+	// Add new portable callstack element with crash stack
+	AddPortableCallStack();
+	AddPortableCallStackHash();
 
 	{
 		FString AllThreadStacks;
@@ -509,48 +573,48 @@ void FGenericCrashContext::SerializeContentToBuffer() const
 		}
 	}
 
-	EndSection( RuntimePropertiesTag );
+	EndSection( CommonBuffer, RuntimePropertiesTag );
 
 	// Add platform specific properties.
-	BeginSection( PlatformPropertiesTag );
+	BeginSection( CommonBuffer, PlatformPropertiesTag );
 	AddPlatformSpecificProperties();
 	// The name here is a bit cryptic, but we keep it to avoid breaking backend stuff.
 	AddCrashProperty(TEXT("PlatformCallbackResult"), NCached::Session.CrashType);
-	EndSection( PlatformPropertiesTag );
+	EndSection( CommonBuffer, PlatformPropertiesTag );
 
 	// Add the engine data
-	BeginSection( EngineDataTag );
+	BeginSection( CommonBuffer, EngineDataTag );
 	for (const TPair<FString, FString>& Pair : NCached::EngineData)
 	{
-		AddCrashProperty(*Pair.Key, *Pair.Value);
+		AddCrashProperty( *Pair.Key, *Pair.Value);
 	}
-	EndSection( EngineDataTag );
+	EndSection( CommonBuffer, EngineDataTag );
 
 	// Add the game data
-	BeginSection( GameDataTag );
+	BeginSection( CommonBuffer, GameDataTag );
 	for (const TPair<FString, FString>& Pair : NCached::GameData)
 	{
-		AddCrashProperty(*Pair.Key, *Pair.Value);
+		AddCrashProperty( *Pair.Key, *Pair.Value);
 	}
-	EndSection( GameDataTag );
+	EndSection( CommonBuffer, GameDataTag );
 
 	// Writing out the list of plugin JSON descriptors causes us to run out of memory
 	// in GMallocCrash on console, so enable this only for desktop platforms.
 #if PLATFORM_DESKTOP
-	if(NCached::EnabledPluginsList.Num() > 0)
+	if (NCached::EnabledPluginsList.Num() > 0)
 	{
-		BeginSection(EnabledPluginsTag);
+		BeginSection(CommonBuffer, EnabledPluginsTag);
 
 		for (const FString& Str : NCached::EnabledPluginsList)
 		{
 			AddCrashProperty(TEXT("Plugin"), *Str);
 		}
 
-		EndSection(EnabledPluginsTag);
+		EndSection(CommonBuffer, EnabledPluginsTag);
 	}
 #endif // PLATFORM_DESKTOP
 
-	AddFooter();
+	AddFooter(CommonBuffer);
 }
 
 void FGenericCrashContext::SetNumMinidumpFramesToIgnore(int InNumMinidumpFramesToIgnore)
@@ -593,19 +657,18 @@ void FGenericCrashContext::SerializeAsXML( const TCHAR* Filename ) const
 	FFileHelper::SaveStringToFile( CommonBuffer, Filename, FFileHelper::EEncodingOptions::AutoDetect );
 }
 
-void FGenericCrashContext::AddCrashProperty( const TCHAR* PropertyName, const TCHAR* PropertyValue ) const
+void FGenericCrashContext::AddCrashPropertyInternal(FString& Buffer, const TCHAR* PropertyName, const TCHAR* PropertyValue)
 {
-	CommonBuffer += TEXT( "<" );
-	CommonBuffer += PropertyName;
-	CommonBuffer += TEXT( ">" );
-
+	Buffer += TEXT( "<" );
+	Buffer += PropertyName;
+	Buffer += TEXT( ">" );
 
-	AppendEscapedXMLString(CommonBuffer, PropertyValue );
+	AppendEscapedXMLString(Buffer, PropertyValue);
 
-	CommonBuffer += TEXT( "</" );
-	CommonBuffer += PropertyName;
-	CommonBuffer += TEXT( ">" );
-	CommonBuffer += LINE_TERMINATOR;
+	Buffer += TEXT( "</" );
+	Buffer += PropertyName;
+	Buffer += TEXT( ">" );
+	Buffer += LINE_TERMINATOR;
 }
 
 void FGenericCrashContext::AddPlatformSpecificProperties() const
@@ -654,7 +717,6 @@ void FGenericCrashContext::AddPortableCallStackHash() const
 
 void FGenericCrashContext::AddPortableCallStack() const
 {	
-
 	if (CallStack.Num() == 0)
 	{
 		AddCrashProperty(TEXT("PCallStack"), TEXT(""));
@@ -683,31 +745,31 @@ void FGenericCrashContext::AddPortableCallStack() const
 	AddCrashProperty(TEXT("PCallStack"), *EscapedStackBuffer);
 }
 
-void FGenericCrashContext::AddHeader() const
+void FGenericCrashContext::AddHeader(FString& Buffer)
 {
-	CommonBuffer += TEXT( "<?xml version=\"1.0\" encoding=\"UTF-8\"?>" ) LINE_TERMINATOR;
-	BeginSection( TEXT("FGenericCrashContext") );
+	Buffer += TEXT("<?xml version=\"1.0\" encoding=\"UTF-8\"?>") LINE_TERMINATOR;
+	BeginSection(Buffer, TEXT("FGenericCrashContext") );
 }
 
-void FGenericCrashContext::AddFooter() const
+void FGenericCrashContext::AddFooter(FString& Buffer)
 {
-	EndSection( TEXT( "FGenericCrashContext" ) );
+	EndSection(Buffer, TEXT( "FGenericCrashContext" ));
 }
 
-void FGenericCrashContext::BeginSection( const TCHAR* SectionName ) const
+void FGenericCrashContext::BeginSection(FString& Buffer, const TCHAR* SectionName)
 {
-	CommonBuffer += TEXT( "<" );
-	CommonBuffer += SectionName;
-	CommonBuffer += TEXT( ">" );
-	CommonBuffer += LINE_TERMINATOR;
+	Buffer += TEXT( "<" );
+	Buffer += SectionName;
+	Buffer += TEXT(">");
+	Buffer += LINE_TERMINATOR;
 }
 
-void FGenericCrashContext::EndSection( const TCHAR* SectionName ) const
+void FGenericCrashContext::EndSection(FString& Buffer, const TCHAR* SectionName)
 {
-	CommonBuffer += TEXT( "</" );
-	CommonBuffer += SectionName;
-	CommonBuffer += TEXT( ">" );
-	CommonBuffer += LINE_TERMINATOR;
+	Buffer += TEXT( "</" );
+	Buffer += SectionName;
+	Buffer += TEXT( ">" );
+	Buffer += LINE_TERMINATOR;
 }
 
 void FGenericCrashContext::AppendEscapedXMLString(FString& OutBuffer, const TCHAR* Text)
@@ -773,6 +835,8 @@ const TCHAR* FGenericCrashContext::GetCrashTypeString(ECrashContextType Type)
 		return CrashTypeEnsure;
 	case ECrashContextType::Assert:
 		return CrashTypeAssert;
+	case ECrashContextType::AbnormalShutdown:
+		return CrashTypeAbnormalShutdown;
 	default:
 		return CrashTypeCrash;
 	}
@@ -941,8 +1005,6 @@ void FGenericCrashContext::DumpLog(const FString& CrashFolderAbsolute)
 #endif // !NO_LOGGING
 }
 
-
-
 FORCENOINLINE void FGenericCrashContext::CapturePortableCallStack(int32 NumStackFramesToIgnore, void* Context)
 {
 	// If the callstack is for the executing thread, ignore this function
@@ -1029,7 +1091,7 @@ void FGenericCrashContext::CopyPlatformSpecificFiles(const TCHAR* OutputDirector
 /**
  * Attempts to create the output report directory.
  */
-bool FGenericCrashContext::CreateCrashReportDirectory(const TCHAR* CrashGUIDRoot, const TCHAR* AppName, int32 CrashIndex, FString& OutCrashDirectoryAbsolute)
+bool FGenericCrashContext::CreateCrashReportDirectory(const TCHAR* CrashGUIDRoot, int32 CrashIndex, FString& OutCrashDirectoryAbsolute)
 {
 	// Generate Crash GUID
 	TCHAR CrashGUID[FGenericCrashContext::CrashGUIDLength];
diff --git a/Engine/Source/Runtime/Core/Private/Windows/WindowsPlatformCrashContext.cpp b/Engine/Source/Runtime/Core/Private/Windows/WindowsPlatformCrashContext.cpp
--- a/Engine/Source/Runtime/Core/Private/Windows/WindowsPlatformCrashContext.cpp
+++ b/Engine/Source/Runtime/Core/Private/Windows/WindowsPlatformCrashContext.cpp
@@ -469,20 +469,29 @@ int32 ReportCrashForMonitor(
 	void* ReadPipe,
 	EErrorReportUI ReportUI)
 {
+	//Initialize all the static data from crash context
+	FGenericCrashContext::Initialize();
+	FGenericCrashContext::CopySharedCrashContext(*SharedContext);
+
+	// Set the platform specific crash context, so that we can stack walk and minidump from
+	// the crash reporter client.
+	SharedContext->PlatformCrashContext = (void*)ExceptionInfo;
+
+	// Setup up the shared memory area so that the crash report
+	SharedContext->CrashType = Type;
+	SharedContext->CrashingThreadId = CrashingThreadId;
+	SharedContext->NumStackFramesToIgnore = NumStackFramesToIgnore;
+
 	// Determine UI settings for the crash report. Suppress the user input dialog if we're running in unattended mode
 	// Usage data controls if we want analytics in the crash report client
 	// Finally we cannot call some of these functions if we crash during static init, so check if they are initialized.
-	bool bNoDialog = ReportUI == EErrorReportUI::ReportInUnattendedMode;
+	bool bNoDialog = ReportUI == EErrorReportUI::ReportInUnattendedMode || IsRunningDedicatedServer();
 	bool bSendUnattendedBugReports = true;
 	bool bSendUsageData = true;
+
 	if (FCommandLine::IsInitialized())
 	{
-		bNoDialog |= FApp::IsUnattended() || IsRunningDedicatedServer();
-	}
-	// If we don't have a command line, but know we are dedicated server, also run in no dialog mode.
-	else if (IsRunningDedicatedServer())
-	{
-		bNoDialog = true;
+		bNoDialog |= FApp::IsUnattended();
 	}
 
 	if (GConfig)
@@ -490,22 +499,11 @@ int32 ReportCrashForMonitor(
 		GConfig->GetBool(TEXT("/Script/UnrealEd.CrashReportsPrivacySettings"), TEXT("bSendUnattendedBugReports"), bSendUnattendedBugReports, GEditorSettingsIni);
 		GConfig->GetBool(TEXT("/Script/UnrealEd.AnalyticsPrivacySettings"), TEXT("bSendUsageData"), bSendUsageData, GEditorSettingsIni);
 	}
-	
-	//Initialize all the static data from crash context
-	FGenericCrashContext::Initialize();
-	FGenericCrashContext::CopySharedCrashContext(*SharedContext);
 
-	// Set the platform specific crash context, so that we can stack walk and minidump from
-	// the crash reporter client.
-	SharedContext->PlatformCrashContext = (void*)ExceptionInfo;
+	SharedContext->UserSettings.bNoDialog = bNoDialog;
+	SharedContext->UserSettings.bSendUnattendedBugReports = bSendUnattendedBugReports;
+	SharedContext->UserSettings.bSendUsageData = bSendUsageData;
 
-	// Setup up the shared memory area so that the crash report
-	SharedContext->CrashType = Type;
-	SharedContext->CrashingThreadId = CrashingThreadId;
-	SharedContext->NumStackFramesToIgnore = NumStackFramesToIgnore;
-	SharedContext->bNoDialog = bNoDialog;
-	SharedContext->bSendUnattenededBugReports = bSendUnattendedBugReports;
-	SharedContext->bSendUsageData = bSendUsageData;
 	FCString::Strcpy(SharedContext->ErrorMessage, CR_MAX_ERROR_MESSAGE_CHARS-1, ErrorMessage);
 
 	if (GLog)
@@ -568,7 +566,7 @@ int32 ReportCrashForMonitor(
 	CloseHandle(ThreadSnapshot);
 
 	FString CrashDirectoryAbsolute;
-	if (FGenericCrashContext::CreateCrashReportDirectory(SharedContext->SessionContext.CrashGUIDRoot, SharedContext->SessionContext.GameName, ReportCrashCallCount, CrashDirectoryAbsolute))
+	if (FGenericCrashContext::CreateCrashReportDirectory(SharedContext->SessionContext.CrashGUIDRoot, ReportCrashCallCount, CrashDirectoryAbsolute))
 	{
 		FCString::Strcpy(SharedContext->CrashFilesDirectory, *CrashDirectoryAbsolute);
 		// Copy the log file to output
diff --git a/Engine/Source/Runtime/Core/Public/GenericPlatform/GenericPlatformCrashContext.h b/Engine/Source/Runtime/Core/Public/GenericPlatform/GenericPlatformCrashContext.h
--- a/Engine/Source/Runtime/Core/Public/GenericPlatform/GenericPlatformCrashContext.h
+++ b/Engine/Source/Runtime/Core/Public/GenericPlatform/GenericPlatformCrashContext.h
@@ -98,6 +98,7 @@ enum class ECrashContextType
 	GPUCrash,
 	Hang,
 	OutOfMemory,
+	AbnormalShutdown,
 
 	Max
 };
@@ -165,6 +166,15 @@ struct FSessionContext
 	TOptional<bool>			bIsVanilla;
 };
 
+/** Additional user settings to be communicated to crash reporting client. */
+struct FUserSettingsContext
+{
+	bool					bNoDialog = false;
+	bool					bSendUnattendedBugReports = false;
+	bool					bSendUsageData = false;
+	TCHAR					LogFilePath[CR_MAX_DIRECTORY_CHARS];
+};
+
 /**
  * Fixed size struct holds crash information and session specific state. It is designed
  * to shared between processes (e.g. Game and CrashReporterClient).
@@ -179,14 +189,12 @@ struct FSharedCrashContext
 	uint32					CrashingThreadId;
 	uint32					NumStackFramesToIgnore;
 	ECrashContextType		CrashType;
+
+	// Additional user settings.
+	FUserSettingsContext	UserSettings;
+
 	// Platform specific crash context (must be portable)
 	void*					PlatformCrashContext;
-	// Platform specific memory stats
-	FPlatformMemoryStats	MemoryStats;
-	// User settings info
-	bool					bNoDialog;
-	bool					bSendUnattenededBugReports;
-	bool					bSendUsageData;
 	// Directory for dumped files
 	TCHAR					CrashFilesDirectory[CR_MAX_DIRECTORY_CHARS];
 	// Game/Engine information not possible to catch out of process
@@ -237,6 +245,7 @@ public:
 	static const TCHAR* const CrashTypeEnsure;
 	static const TCHAR* const CrashTypeGPU;
 	static const TCHAR* const CrashTypeHang;
+	static const TCHAR* const CrashTypeAbnormalShutdown;
 
 	static const TCHAR* const EngineModeExUnknown;
 	static const TCHAR* const EngineModeExDirty;
@@ -278,6 +287,9 @@ public:
 
 	virtual ~FGenericCrashContext() { }
 
+	/** Get the file path to the temporary session context file that we create for the given process. */
+	static FString GetTempSessionContextFilePath(uint64 ProcessID);
+
 	/** Serializes all data to the buffer. */
 	void SerializeContentToBuffer() const;
 
@@ -302,14 +314,16 @@ public:
 	/** Serializes crash's informations to the specified filename. Should be overridden for platforms where using FFileHelper is not safe, all POSIX platforms. */
 	virtual void SerializeAsXML( const TCHAR* Filename ) const;
 
-	/** Writes a common property to the buffer. */
-	void AddCrashProperty( const TCHAR* PropertyName, const TCHAR* PropertyValue ) const;
-
-	/** Writes a common property to the buffer. */
+	/** 
+	 * Serializes session context to the given buffer. 
+	 * NOTE: Assumes that the buffer already has a header and section open.
+	 */
+	static void SerializeSessionContext(FString& Buffer);
+	
 	template <typename Type>
-	void AddCrashProperty( const TCHAR* PropertyName, const Type& Value ) const
+	void AddCrashProperty(const TCHAR* PropertyName, const Type& Value) const
 	{
-		AddCrashProperty( PropertyName, *TTypeToString<Type>::ToString( Value ) );
+		AddCrashPropertyInternal(CommonBuffer, PropertyName, Value);
 	}
 
 	/** Escapes and appends specified text to XML string */
@@ -363,7 +377,7 @@ public:
 	static void SetMemoryStats(const FPlatformMemoryStats& MemoryStats);
 
 	/** Attempts to create the output report directory. */
-	static bool CreateCrashReportDirectory(const TCHAR* CrashGUIDRoot, const TCHAR* AppName, int32 CrashIndex, FString& OutCrashDirectoryAbsolute);
+	static bool CreateCrashReportDirectory(const TCHAR* CrashGUIDRoot, int32 CrashIndex, FString& OutCrashDirectoryAbsolute);
 
 	/** Sets the process id to that has crashed. On supported platforms this will analyze the given process rather than current. Default is current process. */
 	void SetCrashedProcess(const FProcHandle& Process) { ProcessHandle = Process; }
@@ -421,6 +435,22 @@ protected:
 
 private:
 
+	/** Serializes the session context section of the crash context to a temporary file. */
+	static void SerializeTempCrashContextToFile();
+
+	/** Serializes the current user setting struct to a buffer. */
+	static void SerializeUserSettings(FString& Buffer);
+
+	/** Writes a common property to the buffer. */
+	static void AddCrashPropertyInternal(FString& Buffer, const TCHAR* PropertyName, const TCHAR* PropertyValue);
+
+	/** Writes a common property to the buffer. */
+	template <typename Type>
+	static void AddCrashPropertyInternal(FString& Buffer, const TCHAR* PropertyName, const Type& Value)
+	{
+		AddCrashPropertyInternal(Buffer, PropertyName, *TTypeToString<Type>::ToString(Value));
+	}
+
 	/** Serializes platform specific properties to the buffer. */
 	virtual void AddPlatformSpecificProperties() const;
 
@@ -431,13 +461,13 @@ private:
 	void AddPortableCallStackHash() const;
 
 	/** Writes header information to the buffer. */
-	void AddHeader() const;
+	static void AddHeader(FString& Buffer);
 
 	/** Writes footer to the buffer. */
-	void AddFooter() const;
+	static void AddFooter(FString& Buffer);
 
-	void BeginSection( const TCHAR* SectionName ) const;
-	void EndSection( const TCHAR* SectionName ) const;
+	static void BeginSection(FString& Buffer, const TCHAR* SectionName);
+	static void EndSection(FString& Buffer, const TCHAR* SectionName);
 
 	/** Called once when GConfig is initialized. Opportunity to cache values from config. */
 	static void InitializeFromConfig();
diff --git a/Engine/Source/Runtime/Core/Public/Misc/CoreDelegates.h b/Engine/Source/Runtime/Core/Public/Misc/CoreDelegates.h
--- a/Engine/Source/Runtime/Core/Public/Misc/CoreDelegates.h
+++ b/Engine/Source/Runtime/Core/Public/Misc/CoreDelegates.h
@@ -42,6 +42,8 @@ struct FCrashOverrideParameters
 	/** Default this to true for backward compatibility before these bools were added. */
 	bool bSetCrashReportClientMessageText = true;
 	bool bSetGameNameSuffix = false;
+	TOptional<bool> SendUnattendedBugReports;
+	TOptional<bool> SendUsageData;
 
 	CORE_API ~FCrashOverrideParameters();
 };
diff --git a/Engine/Source/Runtime/Engine/Classes/Engine/ImportantToggleSettingInterface.h b/Engine/Source/Runtime/Engine/Classes/Engine/ImportantToggleSettingInterface.h
--- a/Engine/Source/Runtime/Engine/Classes/Engine/ImportantToggleSettingInterface.h
+++ b/Engine/Source/Runtime/Engine/Classes/Engine/ImportantToggleSettingInterface.h
@@ -25,7 +25,7 @@ class IImportantToggleSettingInterface
 {
 	GENERATED_IINTERFACE_BODY()
 
-	virtual void GetToogleCategoryAndPropertyNames(FName& OutCategory, FName& OutProperty) const { OutCategory = NAME_None; OutProperty = NAME_None; }
+	virtual void GetToggleCategoryAndPropertyNames(FName& OutCategory, FName& OutProperty) const { OutCategory = NAME_None; OutProperty = NAME_None; }
 
 	virtual FText GetFalseStateLabel() const { return FText::GetEmpty(); }
 	virtual FText GetFalseStateTooltip() const { return FText::GetEmpty(); }
