Replace convex hull algorithm to fix bugs and performance issues
  - From O(nh) complexity to O(n log n)
  - Fix suboptimal alignment problem when orienting charts for UV packing caused by erroneous convex hull results.
  - Fix suboptimal performance of the convex hull computation.
  - Results in improved UV packing efficiency due to better charts alignment (no more diagonally aligned charts)
  - Expected to reduce jagginess in static lighting computation due to less diagonally oriented charts.

  - For DATASET-0008a
     - 42.9s -> 4s on all cores for computing convex hulls of charts when importing
     - 1m50s -> 1m41s on all cores of total UV packing computation time due to better charts alignment (no more diagonally aligned charts)
     - 63% -> 83%  Avg. UV packing efficiency
     - 1m28s -> 1m13s for total import time

From dev-enterprise[REDACTED_HANDLE]

diff --git a/Engine/Plugins/Enterprise/StaticMeshEditorExtension/Source/StaticMeshEditorExtension/Private/UVTools/UVGenerationFlattenMapping.cpp b/Engine/Plugins/Enterprise/StaticMeshEditorExtension/Source/StaticMeshEditorExtension/Private/UVTools/UVGenerationFlattenMapping.cpp
--- a/Engine/Plugins/Enterprise/StaticMeshEditorExtension/Source/StaticMeshEditorExtension/Private/UVTools/UVGenerationFlattenMapping.cpp
+++ b/Engine/Plugins/Enterprise/StaticMeshEditorExtension/Source/StaticMeshEditorExtension/Private/UVTools/UVGenerationFlattenMapping.cpp
@@ -411,17 +411,17 @@ void FUVGenerationFlattenMappingInternal::AlignIslandOnAxes(TArray< FaceStruct*
 {
 	TRACE_CPUPROFILER_EVENT_SCOPE(FDatasmithFlattenMappingInternal::AlignIslandOnAxes)
 
-	TSet< FVector > Vertices;
+	TSet< FVector2D > Vertices;
 	Vertices.Reserve( Island.Num() * 3 );
 
 	for ( const FaceStruct* Face : Island )
 	{
-		Vertices.Add( FVector( Face->Uvs[0], 0.f ) );
-		Vertices.Add( FVector( Face->Uvs[1], 0.f ) );
-		Vertices.Add( FVector( Face->Uvs[2], 0.f ) );
+		Vertices.Add( Face->Uvs[0] );
+		Vertices.Add( Face->Uvs[1] );
+		Vertices.Add( Face->Uvs[2] );
 	}
 
-	TArray< FVector > VerticesArray = Vertices.Array();
+	TArray< FVector2D > VerticesArray = Vertices.Array();
 
 	TArray< int32 > ConvexHullIndices;
 	ConvexHull2D::ComputeConvexHull( VerticesArray, ConvexHullIndices );
@@ -431,7 +431,7 @@ void FUVGenerationFlattenMappingInternal::AlignIslandOnAxes(TArray< FaceStruct*
 
 	for ( int32 Index : ConvexHullIndices )
 	{
-		ConvexHull.Emplace( VerticesArray[ Index ].X, VerticesArray[ Index ].Y );
+		ConvexHull.Emplace( VerticesArray[ Index ] );
 	}
 
 	FVector IslandTangent(ForceInit);
@@ -461,17 +461,7 @@ void FUVGenerationFlattenMappingInternal::AlignIslandOnAxes(TArray< FaceStruct*
 	{
 		IslandTangent.Normalize();
 
-		float DotForward = FMath::Abs( IslandTangent | FVector::ForwardVector );
-		float DotSide = FMath::Abs( IslandTangent | FVector::RightVector );
-
-		FVector AxisDirection = DotForward > DotSide ? FVector::ForwardVector : FVector::RightVector;
-
-		if ( ( IslandTangent | AxisDirection ) < 0.f )
-		{
-			AxisDirection *= -1.f;
-		}
-
-		IslandRotation = FQuat::FindBetweenNormals( IslandTangent, AxisDirection );
+		IslandRotation = FQuat::FindBetweenNormals( IslandTangent, FVector::ForwardVector );
 	}
 
 	for ( int32 Face = 0; Face < Island.Num(); Face++ )
diff --git a/Engine/Source/Runtime/Core/Public/Math/ConvexHull2d.h b/Engine/Source/Runtime/Core/Public/Math/ConvexHull2d.h
--- a/Engine/Source/Runtime/Core/Public/Math/ConvexHull2d.h
+++ b/Engine/Source/Runtime/Core/Public/Math/ConvexHull2d.h
@@ -8,6 +8,87 @@
 
 namespace ConvexHull2D
 {
+	/**
+	 * Andrew's monotone chain convex hull algorithm for 2-dimensional points. O(N log N).
+	 *
+	 * Not the fastest algorithm out there, but definitely the simplest one to understand.
+	 *
+	 * 1 - Sort O(N log N)
+	 * 2 - Scan sorted vertex from left to right to compute lower hull  O(N)
+	 * 3 - Scan sorted vertex from right to left to compute upper hull. O(N)
+	 *
+	 * If this is slow for some reason, O(N log H) also exists where H is the number of outputted
+	 * hull vertices which is normally a lot lower than N and helps reduce the overall complexity.
+	 */
+	template<typename VectorType, typename Allocator>
+	void ComputeConvexHull(const TArray<VectorType, Allocator>& Points, TArray<int32, Allocator>& OutIndices)
+	{
+		// Handle case too simple for the algorithm to handle
+		int32 PointsNum = Points.Num();
+		if (PointsNum <= 3)
+		{
+			for (int32 Index = 0; Index < PointsNum; ++Index)
+			{
+				OutIndices.Add(Index);
+			}
+
+			return;
+		}
+
+		// Simple sorted index lookup table into immutable Points array.
+		TArray<uint32, Allocator> SortedIndices;
+		SortedIndices.SetNumUninitialized(PointsNum);
+		for (int32 Index = 0; Index < PointsNum; ++Index)
+		{
+			SortedIndices[Index] = Index;
+		}
+	
+		// Get rid of costly RangeCheck during sort by using pointer directly
+		Algo::Sort(
+			SortedIndices,
+			[P = Points.GetData()](uint32 a, uint32 b)
+			{
+				// Sort in Y if X are equal
+				return P[a].X == P[b].X ? P[a].Y < P[b].Y : P[a].X < P[b].X;
+			}
+		);
+
+		auto IsClockwise =
+			[](const VectorType& O, const VectorType& A, const VectorType& B)
+			{
+				return ((A.X - O.X) * (B.Y - O.Y) - (A.Y - O.Y) * (B.X - O.X)) <= 0;
+			};
+
+		int32 HullIndex = 0;
+		OutIndices.SetNum(PointsNum*2);
+
+		// Build lower hull
+		for (int32 Index = 0; Index < PointsNum; ++Index) 
+		{
+			const VectorType& B = Points[SortedIndices[Index]];
+			while (HullIndex >= 2 && IsClockwise(Points[OutIndices[HullIndex - 2]], Points[OutIndices[HullIndex - 1]], B))
+			{
+				--HullIndex;
+			}
+
+			OutIndices[HullIndex++] = SortedIndices[Index];
+		}
+
+		// Build upper hull
+		for (int32 Index = PointsNum - 1, StartIndex = HullIndex + 1; Index > 0; --Index) 
+		{
+			const VectorType& B = Points[SortedIndices[Index - 1]];
+			while (HullIndex >= StartIndex && IsClockwise(Points[OutIndices[HullIndex - 2]], Points[OutIndices[HullIndex - 1]], B))
+			{
+				--HullIndex;
+			}
+
+			OutIndices[HullIndex++] = SortedIndices[Index - 1];
+		}
+
+		OutIndices.SetNum(HullIndex - 1, false);
+	}
+
 	/** Returns <0 if C is left of A-B */
 	inline float ComputeDeterminant(const FVector& A, const FVector& B, const FVector& C)
 	{
@@ -48,9 +129,11 @@ namespace ConvexHull2D
 	/** 
 	 * Calculates convex hull on xy-plane of points on 'Points' and stores the indices of the resulting hull in 'OutIndices'.
 	 * This code was fixed to work with duplicated vertices and precision issues.
+	 * 
+	 * Should be replaced by ComputeConvexHull and tested properly. Keep for backward compatibility until then.
 	 */
 	template<typename Allocator>
-	void ComputeConvexHull(const TArray<FVector, Allocator>& Points, TArray<int32, Allocator>& OutIndices)
+	void ComputeConvexHullLegacy(const TArray<FVector, Allocator>& Points, TArray<int32, Allocator>& OutIndices)
 	{
 		if (Points.Num() == 0)
 		{
@@ -175,9 +258,11 @@ namespace ConvexHull2D
 
 	/** 
 	 * Alternate simple implementation that was found to work correctly for points that are very close together (inside the 0-1 range).
+	 * 
+	 * Should be replaced by ComputeConvexHull and tested properly. Keep for backward compatibility until then.
 	 */
 	template<typename Allocator>
-	void ComputeConvexHull2(const TArray<FVector2D, Allocator>& Points, TArray<int32, Allocator>& OutIndices)
+	void ComputeConvexHullLegacy2(const TArray<FVector2D, Allocator>& Points, TArray<int32, Allocator>& OutIndices)
 	{
 		if (Points.Num() == 0)
 		{
@@ -222,7 +307,7 @@ namespace ConvexHull2D
 	}
 
 /*
-	static Test()
+	static void Test()
 	{
 		{
 			TArray<FVector, TInlineAllocator<8> > In;
diff --git a/Engine/Source/Runtime/Engine/Private/AI/Navigation/NavigationModifier.cpp b/Engine/Source/Runtime/Engine/Private/AI/Navigation/NavigationModifier.cpp
--- a/Engine/Source/Runtime/Engine/Private/AI/Navigation/NavigationModifier.cpp
+++ b/Engine/Source/Runtime/Engine/Private/AI/Navigation/NavigationModifier.cpp
@@ -528,7 +528,7 @@ void FAreaNavModifier::FillConvexNavAreaData(const FVector* InPoints, const int3
 	TArray<int32, TInlineAllocator<MaxConvexPoints>> HullIndices;
 	HullIndices.Empty(MaxConvexPoints);
 
-	ConvexHull2D::ComputeConvexHull(HullVertices, HullIndices);
+	ConvexHull2D::ComputeConvexHullLegacy(HullVertices, HullIndices);
 	
 	// ConvexHull implementation requires at least 3 vertices  (i.e. GrowConvexHull)
 	const int32 MIN_NUM_POINTS = 3;
diff --git a/Engine/Source/Runtime/Engine/Private/KismetMathLibrary.cpp b/Engine/Source/Runtime/Engine/Private/KismetMathLibrary.cpp
--- a/Engine/Source/Runtime/Engine/Private/KismetMathLibrary.cpp
+++ b/Engine/Source/Runtime/Engine/Private/KismetMathLibrary.cpp
@@ -967,7 +967,7 @@ void UKismetMathLibrary::MinimumAreaRectangle(class UObject* WorldContextObject,
 	OutRectCenter /= InVerts.Num();
 
 	// Compute the convex hull of the sample points
-	ConvexHull2D::ComputeConvexHull(TransformedVerts, PolyVertIndices);
+	ConvexHull2D::ComputeConvexHullLegacy(TransformedVerts, PolyVertIndices);
 
 	// Minimum area rectangle as computed by http://www.geometrictools.com/Documentation/MinimumAreaRectangle.pdf
 	for (int32 Idx = 1; Idx < PolyVertIndices.Num() - 1; ++Idx)
diff --git a/Engine/Source/Runtime/Engine/Private/Particles/SubUVAnimation.cpp b/Engine/Source/Runtime/Engine/Private/Particles/SubUVAnimation.cpp
--- a/Engine/Source/Runtime/Engine/Private/Particles/SubUVAnimation.cpp
+++ b/Engine/Source/Runtime/Engine/Private/Particles/SubUVAnimation.cpp
@@ -613,7 +613,7 @@ void FSubUVDerivedData::Build(UTexture2D* SubUVTexture, int32 SubImages_Horizont
 					{
 						TArray<int32> ConvexHullIndices;
 						// Compute the 2d convex hull of texels with non-zero alpha
-						ConvexHull2D::ComputeConvexHull2(PotentialHullVertices, ConvexHullIndices);
+						ConvexHull2D::ComputeConvexHullLegacy2(PotentialHullVertices, ConvexHullIndices);
 
 						bSubImageSuccess = ConvexHullIndices.Num() >= 3;
 
