- Fixed an issue in Linux HAL preventing a seek() followed by a read() to read at the correct position when a file is opened in 'write' mode.

diff --git a/Engine/Source/Runtime/Core/Private/Unix/UnixPlatformFile.cpp b/Engine/Source/Runtime/Core/Private/Unix/UnixPlatformFile.cpp
--- a/Engine/Source/Runtime/Core/Private/Unix/UnixPlatformFile.cpp
+++ b/Engine/Source/Runtime/Core/Private/Unix/UnixPlatformFile.cpp
@@ -173,34 +173,32 @@ public:
 	{
 		struct FScopedReadTracker
 		{
-			FScopedReadTracker(FFileHandleUnix* InHandle) : Handle(InHandle)
-			{
-				if (Handle) { GFileRegistry.TrackStartRead(Handle); }
-			}
-			~FScopedReadTracker()
-			{
-				if (Handle) { GFileRegistry.TrackEndRead(Handle); }
-			}
-			FFileHandleUnix* Handle;
+			FScopedReadTracker(FFileHandleUnix& InHandle) : Handle(InHandle) { GFileRegistry.TrackStartRead(&Handle); }
+			~FScopedReadTracker() { GFileRegistry.TrackEndRead(&Handle); }
+			FFileHandleUnix& Handle;
 		};
 
-		// Handle virtual file handles (only in read mode, write mode doesn't use the file handle registry)
-		FScopedReadTracker ScopedReadTracker(FileOpenAsWrite ? nullptr : this);
-
-		int64 BytesRead = 0;
+		check(IsValid());
+		if (!FileOpenAsWrite)
 		{
+			// Handle virtual file handles (only in read mode, write mode doesn't use the file handle registry)
+			FScopedReadTracker ScopedReadTracker(*this);
 			FScopedDiskUtilizationTracker Tracker(BytesToRead, FileOffset);
+
 			// seek to the offset on seek? this matches console behavior more closely
-			check(IsValid());
 			if (lseek(FileHandle, FileOffset, SEEK_SET) == -1)
 			{
 				return false;
 			}
-			BytesRead += ReadInternal(Destination, BytesToRead);
+			int64 BytesRead = ReadInternal(Destination, BytesToRead);
+			FileOffset += BytesRead;
+			return BytesRead == BytesToRead;
+		}
+		else // FileOffset is invalid in 'read/write' mode, i.e. not updated by Write(), Seek(), SeekFronEnd(). Read from the current location.
+		{
+			FScopedDiskUtilizationTracker Tracker(BytesToRead, Tell());
+			return ReadInternal(Destination, BytesToRead) == BytesToRead;
 		}
-
-		FileOffset += BytesRead;
-		return BytesRead == BytesToRead;
 	}
 
 	virtual bool Write(const uint8* Source, int64 BytesToWrite) override
