Fix memory stomp in Trace writer trying to use the size of GErrorHist + Overhead to decrement a buffer by providing a format buffer

This caused a crash while handling a crash so no crash reporter was starting

diff --git a/Engine/Source/Runtime/Core/Private/Mac/MacErrorOutputDevice.cpp b/Engine/Source/Runtime/Core/Private/Mac/MacErrorOutputDevice.cpp
--- a/Engine/Source/Runtime/Core/Private/Mac/MacErrorOutputDevice.cpp
+++ b/Engine/Source/Runtime/Core/Private/Mac/MacErrorOutputDevice.cpp
@@ -90,7 +90,7 @@ void FMacErrorOutputDevice::HandleError()
 
 	// Dump the error and flush the log.
 	UE_LOG(LogMac, Log, TEXT("=== Critical error: ===") LINE_TERMINATOR TEXT("%s") LINE_TERMINATOR, GErrorExceptionDescription);
-	UE_LOG(LogMac, Log, GErrorHist);
+	UE_LOG(LogMac, Log, TEXT("%s"), GErrorHist);
 
 	GLog->Flush();
 
diff --git a/Engine/Source/Runtime/Core/Private/Misc/OutputDeviceAnsiError.cpp b/Engine/Source/Runtime/Core/Private/Misc/OutputDeviceAnsiError.cpp
--- a/Engine/Source/Runtime/Core/Private/Misc/OutputDeviceAnsiError.cpp
+++ b/Engine/Source/Runtime/Core/Private/Misc/OutputDeviceAnsiError.cpp
@@ -82,7 +82,7 @@ void FOutputDeviceAnsiError::HandleError()
 	{
 		// print to log and flush it
 		UE_LOG( LogHAL, Log, TEXT( "=== Critical error: ===" ) LINE_TERMINATOR LINE_TERMINATOR TEXT( "%s" ) LINE_TERMINATOR, GErrorExceptionDescription );
-		UE_LOG(LogHAL, Log, GErrorHist);
+		UE_LOG(LogHAL, Log, TEXT("%s"), GErrorHist);
 
 		GLog->Flush();
 	}
diff --git a/Engine/Source/Runtime/Core/Private/Unix/UnixErrorOutputDevice.cpp b/Engine/Source/Runtime/Core/Private/Unix/UnixErrorOutputDevice.cpp
--- a/Engine/Source/Runtime/Core/Private/Unix/UnixErrorOutputDevice.cpp
+++ b/Engine/Source/Runtime/Core/Private/Unix/UnixErrorOutputDevice.cpp
@@ -104,7 +104,7 @@ void FUnixErrorOutputDevice::HandleError()
 
 		// Dump the error and flush the log.
 		UE_LOG(LogCore, Log, TEXT("=== Critical error: ===") LINE_TERMINATOR TEXT("%s") LINE_TERMINATOR, GErrorExceptionDescription);
-		UE_LOG(LogCore, Log, GErrorHist);
+		UE_LOG(LogCore, Log, TEXT("%s"), GErrorHist);
 
 		GLog->Flush();
 
