Fix for iOS Launch Crash in MallocBinned.
#tests iOS Client (QA Game)

diff --git a/Engine/Source/Runtime/Core/Private/HAL/MallocBinned.cpp b/Engine/Source/Runtime/Core/Private/HAL/MallocBinned.cpp
--- a/Engine/Source/Runtime/Core/Private/HAL/MallocBinned.cpp
+++ b/Engine/Source/Runtime/Core/Private/HAL/MallocBinned.cpp
@@ -31,7 +31,7 @@ DEFINE_STAT(STAT_Binned_NanoMallocPages_WastePeak);
 #if PLATFORM_IOS
 #define PLAT_PAGE_SIZE_LIMIT 16384
 #define PLAT_BINNED_ALLOC_POOLSIZE 16384
-#define PLAT_SMALL_BLOCK_POOL_SIZE 256 //224
+#define PLAT_SMALL_BLOCK_POOL_SIZE 256
 #else
 #define PLAT_PAGE_SIZE_LIMIT 65536
 #define PLAT_BINNED_ALLOC_POOLSIZE 65536
@@ -446,12 +446,10 @@ struct FMallocBinned::Private
 		return Indirect;
 	}
 
-	/** 
-	* Gets the FPoolInfo for a memory address. If no valid info exists one is created. 
-	* NOTE: This function requires a mutex across threads, but its is the callers responsibility to 
-	* acquire the mutex before calling
+	/**
+	* Initializes tables for HashBuckets if they haven't already been initialized.
 	*/
-	static FORCEINLINE FPoolInfo* GetPoolInfo(FMallocBinned& Allocator, UPTRINT Ptr)
+	static FORCEINLINE void InitializeHashBuckets(FMallocBinned& Allocator)
 	{
 		if (!Allocator.HashBuckets)
 		{
@@ -459,11 +457,24 @@ struct FMallocBinned::Private
             LLM_PLATFORM_SCOPE(ELLMTag::FMalloc);
 			Allocator.HashBuckets = (PoolHashBucket*)FPlatformMemory::BinnedAllocFromOS(Align(Allocator.MaxHashBuckets * sizeof(PoolHashBucket), Allocator.PageSize));
 
-			for (uint32 i = 0; i < Allocator.MaxHashBuckets; ++i) 
+			for (uint32 i = 0; i < Allocator.MaxHashBuckets; ++i)
 			{
 				new (Allocator.HashBuckets + i) PoolHashBucket();
 			}
 		}
+	}
+	
+	/** 
+	* Gets the FPoolInfo for a memory address. If no valid info exists one is created. 
+	* NOTE: This function requires a mutex across threads, but its is the callers responsibility to 
+	* acquire the mutex before calling
+	*/
+	static FORCEINLINE FPoolInfo* GetPoolInfo(FMallocBinned& Allocator, UPTRINT Ptr)
+	{
+		if (!Allocator.HashBuckets)
+		{
+			InitializeHashBuckets(Allocator);
+		}
 
 		UPTRINT Key       = Ptr >> Allocator.HashKeyShift;
 		UPTRINT Hash      = Key & (Allocator.MaxHashBuckets - 1);
@@ -1321,6 +1332,10 @@ void* FMallocBinned::Malloc(SIZE_T Size, uint32 Alignment)
 #if USE_OS_SMALL_BLOCK_ALLOC && !USE_OS_SMALL_BLOCK_GRAB_MEMORY_FROM_OS
 	if (Size <= Private::SMALL_BLOCK_POOL_SIZE)
 	{
+		//Make sure we have initialized our hash buckets even if we are using the NANO_MALLOC grabber, as otherwise we can end
+		//up making bad assumptions and trying to grab invalid data during a Realloc of this data.
+		Private::InitializeHashBuckets(*this);
+		
 		bUsePools = false;
 		UPTRINT AlignedSize = Align(Size, Alignment);
 		SIZE_T ActualPoolSize; //TODO: use this to reduce waste?
