Appease static analysis
#rnx

diff --git a/Engine/Source/Runtime/Core/Private/HAL/MallocBinned.cpp b/Engine/Source/Runtime/Core/Private/HAL/MallocBinned.cpp
--- a/Engine/Source/Runtime/Core/Private/HAL/MallocBinned.cpp
+++ b/Engine/Source/Runtime/Core/Private/HAL/MallocBinned.cpp
@@ -475,6 +475,7 @@ struct FMallocBinned::Private
 		{
 			InitializeHashBuckets(Allocator);
 		}
+		checkSlow(Allocator.HashBuckets);
 
 		UPTRINT Key       = Ptr >> Allocator.HashKeyShift;
 		UPTRINT Hash      = Key & (Allocator.MaxHashBuckets - 1);
