Fix for wrong type for LiteralFlag in FLiteralOrName

#ue4

diff --git a/Engine/Source/Runtime/Core/Public/UObject/NameTypes.h b/Engine/Source/Runtime/Core/Public/UObject/NameTypes.h
--- a/Engine/Source/Runtime/Core/Public/UObject/NameTypes.h
+++ b/Engine/Source/Runtime/Core/Public/UObject/NameTypes.h
@@ -1019,7 +1019,7 @@ private:
 	struct FLiteralOrName
 	{
 		// NOTE: uses high bit of pointer for flag; this may be an issue in future when high byte of address may be used for features like hardware ASAN
-		static constexpr uint64 LiteralFlag = UPTRINT(1) << (sizeof(UPTRINT) * 8 - 1);
+		static constexpr uint64 LiteralFlag = uint64(1) << (sizeof(uint64) * 8 - 1);
 
 		explicit FLiteralOrName(const ANSICHAR* Literal)
 			: Int(reinterpret_cast<uint64>(Literal) | LiteralFlag)
