Changed FQueuedThread so that it doesn't wake up every 10ms while waiting for work. This wakeup is a waste of cpu time from having to context switch for no reason.

This saves 0.3ms on switch.

#rnx

diff --git a/Engine/Source/Runtime/Core/Private/HAL/ThreadingBase.cpp b/Engine/Source/Runtime/Core/Private/HAL/ThreadingBase.cpp
--- a/Engine/Source/Runtime/Core/Private/HAL/ThreadingBase.cpp
+++ b/Engine/Source/Runtime/Core/Private/HAL/ThreadingBase.cpp
@@ -8,6 +8,7 @@
 #include "Misc/EventPool.h"
 #include "Misc/LazySingleton.h"
 #include "Templates/Atomic.h"
+#include "HAL/IConsoleManager.h"
 #include "HAL/PlatformStackWalk.h"
 #include "ProfilingDebugging/MiscTrace.h"
 
@@ -16,6 +17,14 @@ DEFINE_STAT( STAT_EventTriggerWithId );
 
 DECLARE_DWORD_COUNTER_STAT( TEXT( "ThreadPoolDummyCounter" ), STAT_ThreadPoolDummyCounter, STATGROUP_ThreadPoolAsyncTasks );
 
+static bool GDoPooledThreadWaitTimeouts = false;
+static FAutoConsoleVariableRef CVarDoPooledThreadWaitTimeouts(
+	TEXT("DoPooledThreadWaitTimeouts"),
+	GDoPooledThreadWaitTimeouts,
+	TEXT("If enabled, uses the old behaviour for waking up pool threads every 10ms. Otherwise, lets pooled threads sleep until data arrives."),
+	ECVF_Default
+);
+
 /** The global thread pool */
 FQueuedThreadPool* GThreadPool = nullptr;
 
@@ -507,15 +516,7 @@ protected:
 				DECLARE_SCOPE_CYCLE_COUNTER( TEXT( "FQueuedThread::Run.WaitForWork" ), STAT_FQueuedThread_Run_WaitForWork, STATGROUP_ThreadPoolAsyncTasks );
 
 				// Wait for some work to do
-
-				bContinueWaiting = !DoWorkEvent->Wait( 10 );
-			}
-			}
-#endif
-
-			if (bContinueWaiting)
-			{
-				DoWorkEvent->Wait();
+				bContinueWaiting = !DoWorkEvent->Wait(GDoPooledThreadWaitTimeouts ? 10 : MAX_uint32);
 			}
 
 			IQueuedWork* LocalQueuedWork = QueuedWork;
