Do not set realtime signals to be ignored on Linux.

- This is particularly problematic for VTune.

diff --git a/Engine/Source/Runtime/Core/Private/Unix/UnixPlatformCrashContext.cpp b/Engine/Source/Runtime/Core/Private/Unix/UnixPlatformCrashContext.cpp
--- a/Engine/Source/Runtime/Core/Private/Unix/UnixPlatformCrashContext.cpp
+++ b/Engine/Source/Runtime/Core/Private/Unix/UnixPlatformCrashContext.cpp
@@ -848,8 +848,9 @@ void FUnixPlatformMisc::SetCrashHandler(void (* CrashHandler)(const FGenericCras
 	Action.sa_flags = SA_SIGINFO | SA_RESTART | SA_ONSTACK;
 	Action.sa_handler = SIG_IGN;
 
-	// set all the signals except ones we know we are handling to be ignored
-	for (int Signal = 1; Signal < NSIG; ++Signal)
+	// Set all the signals except ones we know we are handling to be ignored
+	// Exempt realtime signals as well as they are used by third party libs and VTune
+	for (int Signal = 1; Signal < SIGRTMIN; ++Signal)
 	{
 		bool bSignalShouldBeIgnored = true;
 		for (int HandledSignal : HandledSignals)
