Force GNameBlocksDebug to be used so LTO does not optimize it out

diff --git a/Engine/Source/Runtime/Core/Private/Unix/UnixPlatformMisc.cpp b/Engine/Source/Runtime/Core/Private/Unix/UnixPlatformMisc.cpp
--- a/Engine/Source/Runtime/Core/Private/Unix/UnixPlatformMisc.cpp
+++ b/Engine/Source/Runtime/Core/Private/Unix/UnixPlatformMisc.cpp
@@ -238,6 +238,14 @@ void FUnixPlatformMisc::PlatformInit()
 		// If "-norandomguids" specified, don't use SYS_getrandom syscall
 		SysGetRandomSupported = 0;
 	}
+
+	// This symbol is used for debugging but with LTO enabled it gets stripped as nothing is using it
+	// Lets use it here just to log if its valid under VeryVerbose settings
+	extern uint8** GNameBlocksDebug;
+	if (GNameBlocksDebug)
+	{
+		UE_LOG(LogInit, VeryVerbose, TEXT("GNameBlocksDebug Valid - %i"), !!GNameBlocksDebug);
+	}
 }
 
 extern void CORE_API UnixPlatformStackWalk_UnloadPreloadedModuleSymbol();
