Fix bad merge
#rnx

diff --git a/Engine/Source/Runtime/Core/Private/HAL/ThreadingBase.cpp b/Engine/Source/Runtime/Core/Private/HAL/ThreadingBase.cpp
--- a/Engine/Source/Runtime/Core/Private/HAL/ThreadingBase.cpp
+++ b/Engine/Source/Runtime/Core/Private/HAL/ThreadingBase.cpp
@@ -511,12 +511,19 @@ protected:
 #if STATS
 			if (FThreadStats::IsCollectingData())
 			{
-			while( bContinueWaiting )
-			{				
-				DECLARE_SCOPE_CYCLE_COUNTER( TEXT( "FQueuedThread::Run.WaitForWork" ), STAT_FQueuedThread_Run_WaitForWork, STATGROUP_ThreadPoolAsyncTasks );
+				while (bContinueWaiting)
+				{
+					DECLARE_SCOPE_CYCLE_COUNTER(TEXT("FQueuedThread::Run.WaitForWork"), STAT_FQueuedThread_Run_WaitForWork, STATGROUP_ThreadPoolAsyncTasks);
+
+					// Wait for some work to do
+					bContinueWaiting = !DoWorkEvent->Wait(GDoPooledThreadWaitTimeouts ? 10 : MAX_uint32);
+				}
+			}
+#endif
 
-				// Wait for some work to do
-				bContinueWaiting = !DoWorkEvent->Wait(GDoPooledThreadWaitTimeouts ? 10 : MAX_uint32);
+			if (bContinueWaiting)
+			{
+				DoWorkEvent->Wait();
 			}
 
 			IQueuedWork* LocalQueuedWork = QueuedWork;
