MN - Implementing a robin hood hash table as faster alternative.

diff --git a/Engine/Source/Runtime/Core/Public/Experimental/Containers/RobinHoodHashTable.h b/Engine/Source/Runtime/Core/Public/Experimental/Containers/RobinHoodHashTable.h
new file mode 100644
--- /dev/null
+++ b/Engine/Source/Runtime/Core/Public/Experimental/Containers/RobinHoodHashTable.h
@@ -0,0 +1,876 @@
+#pragma once
+
+// Implementation of Robin Hood hash table.
+#include "CoreMinimal.h"
+#include "Containers/Map.h"
+#include "Templates/UnrealTemplate.h"
+
+namespace Experimental
+{
+
+namespace RobinHoodHashTable_Private
+{
+	template<typename, typename, typename, typename>
+	class TRobinHoodHashTable;
+}
+
+class FHashType
+{
+public:
+	explicit FHashType() : Hash(InvalidHash) {}
+
+private:
+	template<typename, typename, typename, typename>
+	friend class RobinHoodHashTable_Private::TRobinHoodHashTable;
+
+	using IntType = uint32;
+
+	inline explicit FHashType(IntType InHash) : Hash(InHash)
+	{
+		checkSlow(!(InHash & InvalidHash));
+	}
+
+	inline bool IsOccupied() const
+	{
+		return Hash != InvalidHash;
+	}
+
+	inline bool IsFree() const
+	{
+		return Hash == InvalidHash;
+	}
+
+	inline IntType AsUInt() const
+	{
+		return Hash;
+	}
+
+	inline bool operator == (FHashType Other) const
+	{
+		return Hash == Other.Hash;
+	}
+
+	static constexpr const IntType InvalidHash = (IntType(1)) << (sizeof(IntType) * 8 - 1);
+	IntType Hash;
+};
+
+class FHashElementId
+{
+public:
+	FHashElementId() : Index(INDEX_NONE) {}
+	inline FHashElementId(int InIndex) : Index(InIndex) {}
+
+	inline int GetIndex() const
+	{
+		return Index;
+	}
+
+	inline bool IsValid() const
+	{
+		return Index != INDEX_NONE;
+	}
+
+private:
+	int Index;
+};
+
+namespace RobinHoodHashTable_Private
+{
+	struct FUnitType
+	{
+	};
+
+	template<typename KeyType, typename ValueType>
+	class TKeyValue
+	{
+		using FindValueType = ValueType*;
+		using ElementType = TPair<const KeyType, ValueType>;
+
+		template<typename, typename, typename, typename>
+		friend class TRobinHoodHashTable;
+
+		template<typename DeducedKeyType, typename DeducedValueType>
+		inline TKeyValue(DeducedKeyType&& InKey, DeducedValueType&& InVal) : Pair(Forward<DeducedKeyType>(InKey), Forward<DeducedValueType>(InVal)) {}
+
+		ElementType Pair;
+
+		inline FindValueType FindImpl()
+		{
+			return &Pair.Value;
+		}
+
+		inline ElementType& GetElement()
+		{
+			return Pair;
+		}
+
+		inline const ElementType& GetElement() const
+		{
+			return Pair;
+		}
+
+		inline const KeyType& GetKey() const
+		{
+			return Pair.Key;
+		}
+	};
+
+	template<typename KeyType>
+	class TKeyValue<KeyType, FUnitType>
+	{
+		using FindValueType = const KeyType*;
+		using ElementType = const KeyType;
+
+		template<typename, typename, typename, typename>
+		friend class TRobinHoodHashTable;
+
+		template<typename DeducedKeyType>
+		inline TKeyValue(DeducedKeyType&& InKey, FUnitType) : Key(Forward<DeducedKeyType>(InKey)) {}
+
+		ElementType Key;
+
+		inline FindValueType FindImpl() const
+		{
+			return &Key;
+		}
+
+		inline ElementType& GetElement() const
+		{
+			return Key;
+		}
+
+		inline const KeyType& GetKey() const
+		{
+			return Key;
+		}
+	};
+
+	template<typename KeyType, typename ValueType, typename Hasher, typename HashMapAllocator>
+	class TRobinHoodHashTable
+	{
+	protected:
+		using KeyValueType = RobinHoodHashTable_Private::TKeyValue<KeyType, ValueType>;
+		using FindValueType = typename KeyValueType::FindValueType;
+		using ElementType = typename KeyValueType::ElementType;
+		using IndexType = uint32;
+
+		static constexpr const IndexType LoadFactorDivisor = 3;
+		static constexpr const IndexType LoadFactorQuotient = 5;
+
+
+		static constexpr const IndexType InvalidIndex = ~IndexType(0);
+
+		struct FData
+		{
+			FData() = default;
+			FData(const FData&) = default;
+			FData& operator=(const FData&) = default;
+			~FData()
+			{
+				Empty();
+			}
+
+			SIZE_T GetAllocatedSize() const
+			{
+				return KeyVals.GetAllocatedSize() + FreeList.GetAllocatedSize();
+			}
+
+			template<typename DeducedKeyType, typename DeducedValueType>
+			inline IndexType Allocate(DeducedKeyType&& Key, DeducedValueType&& Val)
+			{
+				IndexType Index;
+				if (FreeList.Num() > 0)
+				{
+					Index = FreeList.Pop();
+					new (&KeyVals[Index]) KeyValueType{ Forward<DeducedKeyType>(Key), Forward<DeducedValueType>(Val) };
+				}
+				else
+				{
+					Index = KeyVals.Num();
+					KeyVals.Push(KeyValueType{ Forward<DeducedKeyType>(Key), Forward<DeducedValueType>(Val) });
+				}
+				checkSlow(Index != InvalidIndex);
+				return Index;
+			}
+
+			inline void Deallocate(IndexType Index)
+			{
+				int InsertIndex = 0;
+				for (; InsertIndex < FreeList.Num(); InsertIndex++)
+				{
+					if (FreeList[InsertIndex] < Index)
+					{
+						break;
+					}
+				}
+				FreeList.EmplaceAt(InsertIndex, Index);
+				KeyVals[Index].~KeyValueType();
+				//TODO shrink KeyValue Array.
+				//As the FreeList is sorted backwards this would work as follows:
+				// - go though the FreeList front to back as long as the FreeList entry is the last
+				// - than pop-off the last KeyValue WITHOUT calling its destructor again
+				// - remove the entry from the FreeList (probably remove the entire span at the end of iteration, as its tail would need to be moved)
+			}
+
+			inline const KeyValueType& Get(IndexType Index) const
+			{
+#if DO_CHECK
+				for (IndexType Idx : FreeList)
+				{
+					check(Idx != Index);
+				}
+#endif
+				return KeyVals[Index];
+			}
+
+			inline KeyValueType& Get(IndexType Index)
+			{
+#if DO_CHECK
+				for (IndexType Idx : FreeList)
+				{
+					check(Idx != Index);
+				}
+#endif
+				return KeyVals[Index];
+			}
+
+			inline IndexType Num() const
+			{
+				return KeyVals.Num() - FreeList.Num();
+			}
+
+			inline IndexType GetMaxIndex() const
+			{
+				return KeyVals.Num();
+			}
+
+			struct FIteratorState
+			{
+				IndexType Index;
+				int FreeListIndex;
+
+				inline bool operator ==(const FIteratorState& Rhs) const
+				{
+					return Index == Rhs.Index && FreeListIndex == Rhs.FreeListIndex;
+				}
+
+				inline bool operator !=(const FIteratorState& Rhs) const
+				{
+					return Index != Rhs.Index || FreeListIndex != Rhs.FreeListIndex;
+				}
+			};
+
+			inline FIteratorState Next(FIteratorState State) const
+			{
+				for (;;)
+				{
+					State.Index = (State.Index + 1) & InvalidIndex;
+					if (State.FreeListIndex < 0 || State.Index < FreeList[State.FreeListIndex])
+					{
+						return FIteratorState{ State.Index, State.FreeListIndex };
+					}
+					else
+					{
+						checkSlow(State.Index == FreeList[State.FreeListIndex]);
+						State.FreeListIndex--;
+					}
+				}
+			}
+
+			inline FIteratorState Start() const
+			{
+				return Next(FIteratorState{ InvalidIndex, FreeList.Num() - 1 });
+			}
+
+			inline FIteratorState End() const
+			{
+				return FIteratorState{ IndexType(KeyVals.Num()), INDEX_NONE };
+			}
+
+			void Empty()
+			{
+				FIteratorState Iter = Start();
+				FIteratorState EndIter = End();
+				while (Iter != EndIter)
+				{
+					KeyVals[Iter.Index].~KeyValueType();
+					Iter = Next(Iter);
+				}
+				KeyVals.SetNumUnsafeInternal(0);
+				KeyVals.Empty();
+				FreeList.Empty();
+			}
+
+			TArray<KeyValueType, HashMapAllocator> KeyVals;
+			TArray<IndexType, HashMapAllocator> FreeList;
+		};
+
+		inline IndexType ModTableSize(IndexType HashValue) const
+		{
+			return HashValue & SizePow2Minus1;
+		}
+
+		inline void InsertIntoTable(IndexType Index, FHashType Hash)
+		{
+			IndexType InsertIndex = Index;
+			FHashType InsertHash = Hash;
+			IndexType CurrentBucket = ModTableSize(Hash.AsUInt());
+			IndexType InsertDistance = 0;
+			for (;;)
+			{
+				IndexType OtherDistance = ModTableSize(CurrentBucket - HashData[CurrentBucket].AsUInt());
+
+				checkSlow(HashData[CurrentBucket].IsFree() || OtherDistance <= MaximumDistance);
+				checkSlow(CurrentBucket == (ModTableSize(ModTableSize(HashData[CurrentBucket].AsUInt()) + OtherDistance)));
+
+				if (HashData[CurrentBucket].IsFree())
+				{
+					if (InsertDistance > MaximumDistance)
+					{
+						MaximumDistance = InsertDistance;
+					}
+
+					IndexData[CurrentBucket] = InsertIndex;
+					HashData[CurrentBucket] = InsertHash;
+					break;
+				}
+				else if (OtherDistance < InsertDistance)
+				{
+					if (InsertDistance > MaximumDistance)
+					{
+						MaximumDistance = InsertDistance;
+					}
+
+					IndexType OtherIndex = IndexData[CurrentBucket];
+					FHashType OtherHash = HashData[CurrentBucket];
+					IndexData[CurrentBucket] = InsertIndex;
+					HashData[CurrentBucket] = InsertHash;
+
+					InsertDistance = OtherDistance;
+					InsertIndex = OtherIndex;
+					InsertHash = OtherHash;
+				}
+				InsertDistance++;
+				CurrentBucket = ModTableSize(CurrentBucket + 1);
+			}
+		}
+
+	protected:
+		template<typename DeducedKeyType, typename DeducedValueType>
+		inline FHashElementId FindOrAddIdByHash(FHashType HashValue, DeducedKeyType&& Key, DeducedValueType&& Val)
+		{
+			checkSlow(HashValue == ComputeHash(Key));
+			IndexType BucketIndex = ModTableSize(HashValue.AsUInt());
+			const IndexType EndBucketIndex = ModTableSize(HashValue.AsUInt() + MaximumDistance + 1);
+			do
+			{
+				if (HashValue == HashData[BucketIndex])
+				{
+					if (Hasher::Matches(Key, KeyValueData.Get(IndexData[BucketIndex]).GetKey()))
+					{
+						return IndexData[BucketIndex];
+					}
+				}
+
+				BucketIndex = ModTableSize(BucketIndex + 1);
+			} while (BucketIndex != EndBucketIndex);
+
+			if ((KeyValueData.Num() * LoadFactorQuotient) >= (SizePow2Minus1 * LoadFactorDivisor))
+			{
+				TArray<IndexType> IndexDataOld = MoveTemp(IndexData);
+				TArray<FHashType> HashDataOld = MoveTemp(HashData);
+				IndexType OldSizePow2Minus1 = SizePow2Minus1;
+				SizePow2Minus1 = SizePow2Minus1 * 2 + 1;
+				MaximumDistance = 0;
+				IndexData.Reserve(SizePow2Minus1 + 1); IndexData.AddUninitialized(SizePow2Minus1 + 1);
+				HashData.Reserve(SizePow2Minus1 + 1);  HashData.AddDefaulted(SizePow2Minus1 + 1);
+
+				for (IndexType Index = 0; Index <= OldSizePow2Minus1; Index++)
+				{
+					if (HashDataOld[Index].IsOccupied())
+					{
+						InsertIntoTable(IndexDataOld[Index], HashDataOld[Index]);
+					}
+				}
+			}
+
+			IndexType InsertIndex = KeyValueData.Allocate(Forward<DeducedKeyType>(Key), Forward<DeducedValueType>(Val));
+			InsertIntoTable(InsertIndex, HashValue);
+			return FHashElementId(InsertIndex);
+		}
+
+		template<typename DeducedKeyType, typename DeducedValueType>
+		inline FHashElementId FindOrAddId(DeducedKeyType&& Key, DeducedValueType&& Val)
+		{
+			FHashType HashValue = ComputeHash(Key);
+			return FindOrAddIdByHash(HashValue, Forward<DeducedKeyType>(Key), Forward<DeducedValueType>(Val));
+		}
+
+		template<typename DeducedKeyType, typename DeducedValueType>
+		inline FindValueType FindOrAdd(DeducedKeyType&& Key, DeducedValueType&& Val)
+		{
+			FHashElementId Id = FindOrAddId(Forward<DeducedKeyType>(Key), Forward<DeducedValueType>(Val));
+			return KeyValueData.Get(Id.GetIndex()).FindImpl();
+		}
+
+		TRobinHoodHashTable()
+		{
+			IndexData.Reserve(1); IndexData.AddUninitialized();
+			HashData.Reserve(1);  HashData.AddDefaulted();
+		}
+
+		TRobinHoodHashTable(const TRobinHoodHashTable& Other)
+		{
+			SizePow2Minus1 = Other.SizePow2Minus1;
+			MaximumDistance = Other.MaximumDistance;
+			KeyValueData = Other.KeyValueData;
+
+			IndexData.Reserve(SizePow2Minus1 + 1); 
+			IndexData.AddUninitialized(SizePow2Minus1 + 1);
+			HashData.Reserve(SizePow2Minus1 + 1);  
+			HashData.AddUninitialized(SizePow2Minus1 + 1);
+
+			for (int Idx = 0; Idx <= SizePow2Minus1; Idx++)
+			{
+				IndexData[Idx] = Other.IndexData[Idx];
+				HashData[Idx] = Other.HashData[Idx];
+			}
+		}
+
+		TRobinHoodHashTable& operator=(const TRobinHoodHashTable& Other)
+		{
+			if (this != &Other)
+			{
+				SizePow2Minus1 = Other.SizePow2Minus1;
+				MaximumDistance = Other.MaximumDistance;
+				KeyValueData = Other.KeyValueData;
+
+				IndexData.Reserve(SizePow2Minus1 + 1); 
+				IndexData.AddUninitialized(SizePow2Minus1 + 1);
+				HashData.Reserve(SizePow2Minus1 + 1);  
+				HashData.AddUninitialized(SizePow2Minus1 + 1);
+
+				for (int Idx = 0; Idx <= SizePow2Minus1; Idx++)
+				{
+					IndexData[Idx] = Other.IndexData[Idx];
+					HashData[Idx] = Other.HashData[Idx];
+				}
+			}
+
+			return *this;
+		}
+
+		TRobinHoodHashTable(TRobinHoodHashTable&& Other)
+		{
+			SizePow2Minus1 = Other.SizePow2Minus1;
+			MaximumDistance = Other.MaximumDistance;
+			KeyValueData = MoveTemp(Other.KeyValueData);
+
+			IndexData = MoveTemp(Other.IndexData);
+			HashData = MoveTemp(Other.HashData);
+
+			Other.Empty();
+		}
+
+		TRobinHoodHashTable& operator=(TRobinHoodHashTable&& Other)
+		{
+			if (this != &Other)
+			{
+				SizePow2Minus1 = Other.SizePow2Minus1;
+				MaximumDistance = Other.MaximumDistance;
+				KeyValueData = MoveTemp(Other.KeyValueData);
+
+				IndexData = MoveTemp(Other.IndexData);
+				HashData = MoveTemp(Other.HashData);
+
+				Other.Empty();
+			}
+
+			return *this;
+		}
+
+	public:
+		inline FHashType ComputeHash(const KeyType& Key) const
+		{
+			typename FHashType::IntType HashValue = Hasher::GetKeyHash(Key);
+			constexpr typename FHashType::IntType HashBits = (~(1 << (sizeof(typename FHashType::IntType) * 8 - 1)));
+			return FHashType(HashValue & HashBits);
+		}
+
+		class FIteratorType
+		{
+			typename FData::FIteratorState State;
+			FData& Data;
+
+			template<typename, typename, typename, typename>
+			friend class TRobinHoodHashTable;
+
+			inline FIteratorType(FData& InData, bool bIsStartIterator) : Data(InData)
+			{
+				if (bIsStartIterator)
+				{
+					State = Data.Start();
+				}
+				else
+				{
+					State = Data.End();
+				}
+			}
+
+		public:
+			inline bool operator ==(const FIteratorType& Rhs) const
+			{
+				return State == Rhs.State && &Data == &Rhs.Data;
+			}
+
+			inline bool operator !=(const FIteratorType& Rhs) const
+			{
+				return State != Rhs.State || &Data != &Rhs.Data;
+			}
+
+			inline ElementType& operator*() const
+			{
+				return Data.Get(State.Index).GetElement();
+			}
+
+			inline FIteratorType& operator++()
+			{
+				State = Data.Next(State);
+				return *this;
+			}
+		};
+
+		inline FIteratorType begin()
+		{
+			return FIteratorType(KeyValueData, true);
+		}
+
+		inline FIteratorType end()
+		{
+			return FIteratorType(KeyValueData, false);
+		}
+
+		class FConstIteratorType
+		{
+			typename FData::FIteratorState State;
+			const FData& Data;
+
+			template<typename, typename, typename, typename>
+			friend class TRobinHoodHashTable;
+
+			inline FConstIteratorType(const FData& InData, bool bIsStartIterator) : Data(InData)
+			{
+				if (bIsStartIterator)
+				{
+					State = Data.Start();
+				}
+				else
+				{
+					State = Data.End();
+				}
+			}
+
+		public:
+			inline bool operator ==(const FConstIteratorType& Rhs) const
+			{
+				return State == Rhs.State && &Data == &Rhs.Data;
+			}
+
+			inline bool operator !=(const FConstIteratorType& Rhs) const
+			{
+				return State != Rhs.State || &Data != &Rhs.Data;
+			}
+
+			inline const ElementType& operator*() const
+			{
+				return Data.Get(State.Index).GetElement();
+			}
+
+			inline FConstIteratorType& operator++()
+			{
+				State = Data.Next(State);
+				return *this;
+			}
+		};
+
+		inline FConstIteratorType begin() const
+		{
+			return FConstIteratorType(KeyValueData, true);
+		}
+
+		inline FConstIteratorType end() const
+		{
+			return FConstIteratorType(KeyValueData, false);
+		}
+
+		SIZE_T GetAllocatedSize() const
+		{
+			return KeyValueData.GetAllocatedSize() + IndexData.GetAllocatedSize() + HashData.GetAllocatedSize();
+		}
+
+		inline int32 Num() const
+		{
+			return KeyValueData.Num();
+		}
+
+		inline IndexType GetMaxIndex() const
+		{
+			return KeyValueData.GetMaxIndex();
+		}
+
+		inline ElementType& GetByElementId(FHashElementId Id)
+		{
+			return KeyValueData.Get(Id.GetIndex()).GetElement();
+		}
+
+		inline const ElementType& GetByElementId(FHashElementId Id) const
+		{
+			return KeyValueData.Get(Id.GetIndex()).GetElement();
+		}
+
+		inline FHashElementId FindId(const KeyType& Key) const
+		{
+			FHashType HashValue = ComputeHash(Key);
+			IndexType BucketIndex = ModTableSize(HashValue.AsUInt());
+			const IndexType EndBucketIndex = ModTableSize(HashValue.AsUInt() + MaximumDistance + 1);
+			do
+			{
+				if (HashValue == HashData[BucketIndex])
+				{
+					if (Hasher::Matches(Key, KeyValueData.Get(IndexData[BucketIndex]).GetKey()))
+					{
+						return FHashElementId(IndexData[BucketIndex]);
+					}
+				}
+
+				BucketIndex = ModTableSize(BucketIndex + 1);
+			} while (BucketIndex != EndBucketIndex);
+
+			return FHashElementId();
+		}
+
+		FindValueType Find(const KeyType& Key)
+		{
+			FHashElementId Id = FindId(Key);
+			if (Id.IsValid())
+			{
+				return KeyValueData.GetByElementId(Id).FindImpl();
+			}
+
+			return nullptr;
+		}
+
+		const FindValueType Find(const KeyType& Key) const
+		{
+			FHashElementId Id = FindId(Key);
+			if (Id.IsValid())
+			{
+				return KeyValueData.GetByElementId(Id).FindImpl();
+			}
+
+			return nullptr;
+		}
+
+		bool Remove(const KeyType& Key)
+		{
+			bool bRetValue = false;
+			FHashType HashValue = ComputeHash(Key);
+			IndexType BucketIndex = ModTableSize(HashValue.AsUInt());
+			const IndexType EndBucketIndex = ModTableSize(HashValue.AsUInt() + MaximumDistance + 1);
+			do
+			{
+				if (HashValue == HashData[BucketIndex])
+				{
+					if (Hasher::Matches(Key, KeyValueData.Get(IndexData[BucketIndex]).GetKey()))
+					{
+						KeyValueData.Deallocate(IndexData[BucketIndex]);
+						HashData[BucketIndex] = FHashType();
+						bRetValue = true;
+						break;
+					}
+				}
+
+				BucketIndex = ModTableSize(BucketIndex + 1);
+			} while (BucketIndex != EndBucketIndex);
+
+			if (bRetValue && (KeyValueData.Num() * LoadFactorQuotient * 4) < (SizePow2Minus1 * LoadFactorDivisor))
+			{
+				TArray<IndexType> IndexDataOld = MoveTemp(IndexData);
+				TArray<FHashType> HashDataOld = MoveTemp(HashData);
+				IndexType OldSizePow2Minus1 = SizePow2Minus1;
+				SizePow2Minus1 = SizePow2Minus1 / 2;
+				MaximumDistance = 0;
+				IndexData.Reserve(SizePow2Minus1 + 1); 
+				IndexData.AddUninitialized(SizePow2Minus1 + 1);
+				HashData.Reserve(SizePow2Minus1 + 1);  
+				HashData.AddDefaulted(SizePow2Minus1 + 1);
+
+				for (IndexType Index = 0; Index <= OldSizePow2Minus1; Index++)
+				{
+					if (HashDataOld[Index].IsOccupied())
+					{
+						InsertIntoTable(IndexDataOld[Index], HashDataOld[Index]);
+					}
+				}
+			}
+
+			return bRetValue;
+		}
+
+		inline bool RemoveByElementId(FHashElementId Id)
+		{
+			return Remove(KeyValueData.Get(Id.GetIndex()).GetKey());
+		}
+
+		void Empty()
+		{
+			IndexData.Empty(1); 
+			IndexData.AddUninitialized();
+			HashData.Empty(1); 
+			HashData.AddDefaulted();
+			KeyValueData.Empty();
+			SizePow2Minus1 = 0;
+			MaximumDistance = 0;
+		}
+
+	private:
+		FData KeyValueData;
+
+		TArray<IndexType, HashMapAllocator> IndexData;
+		TArray<FHashType, HashMapAllocator> HashData;
+
+		IndexType SizePow2Minus1 = 0;
+		IndexType MaximumDistance = 0;
+	};
+}
+
+template<typename KeyType, typename ValueType, typename Hasher = TDefaultMapHashableKeyFuncs<KeyType, ValueType, false>, typename HashMapAllocator = FDefaultAllocator>
+class TRobinHoodHashMap : public RobinHoodHashTable_Private::TRobinHoodHashTable<KeyType, ValueType, Hasher, HashMapAllocator>
+{
+	using Base = typename RobinHoodHashTable_Private::TRobinHoodHashTable<KeyType, ValueType, Hasher, HashMapAllocator>;
+	using IndexType = typename Base::IndexType;
+	using FindValueType = typename Base::FindValueType;
+
+public:
+	TRobinHoodHashMap() : Base()
+	{
+		static_assert(sizeof(Base) == sizeof(TRobinHoodHashMap), "This class should only limit the interface and not implement anything");
+	}
+
+	TRobinHoodHashMap(const TRobinHoodHashMap& Other) = default;
+	TRobinHoodHashMap& operator=(const TRobinHoodHashMap& Other) = default;
+	TRobinHoodHashMap(TRobinHoodHashMap&& Other) = default;
+	TRobinHoodHashMap& operator=(TRobinHoodHashMap&& Other) = default;
+
+	FHashElementId FindOrAddIdByHash(FHashType HashValue, const KeyType& Key, const ValueType& Val)
+	{
+		return Base::FindOrAddIdByHash(HashValue, Key, Val);
+	}
+
+	FHashElementId FindOrAddIdByHash(FHashType HashValue, const KeyType& Key, ValueType&& Val)
+	{
+		return Base::FindOrAddIdByHash(HashValue, Key, MoveTemp(Val));
+	}
+
+	FHashElementId FindOrAddIdByHash(FHashType HashValue, KeyType&& Key, const ValueType& Val)
+	{
+		return Base::FindOrAddIdByHash(HashValue, MoveTemp(Key), Val);
+	}
+
+	FHashElementId FindOrAddIdByHash(FHashType HashValue, KeyType&& Key, ValueType&& Val)
+	{
+		return Base::FindOrAddIdByHash(HashValue, MoveTemp(Key), MoveTemp(Val));
+	}
+
+	FHashElementId FindOrAddId(const KeyType& Key, const ValueType& Val)
+	{
+		return Base::FindOrAddId(Key, Val);
+	}
+
+	FHashElementId FindOrAddId(const KeyType& Key, ValueType&& Val)
+	{
+		return Base::FindOrAddId(Key, MoveTemp(Val));
+	}
+
+	FHashElementId FindOrAddId(KeyType&& Key, const ValueType& Val)
+	{
+		return Base::FindOrAddId(MoveTemp(Key), Val);
+	}
+
+	FHashElementId FindOrAddId(KeyType&& Key, ValueType&& Val)
+	{
+		return Base::FindOrAddId(MoveTemp(Key), MoveTemp(Val));
+	}
+
+	FindValueType FindOrAdd(const KeyType& Key, const ValueType& Val)
+	{
+		return Base::FindOrAdd(Key, Val);
+	}
+
+	FindValueType FindOrAdd(const KeyType& Key, ValueType&& Val)
+	{
+		return Base::FindOrAdd(Key, MoveTemp(Val));
+	}
+
+	FindValueType FindOrAdd(KeyType&& Key, const ValueType& Val)
+	{
+		return Base::FindOrAdd(MoveTemp(Key), Val);
+	}
+
+	FindValueType FindOrAdd(KeyType&& Key, ValueType&& Val)
+	{
+		return Base::FindOrAdd(MoveTemp(Key), MoveTemp(Val));
+	}
+};
+
+template<typename KeyType, typename Hasher = DefaultKeyFuncs<KeyType, false>, typename HashMapAllocator = FDefaultAllocator>
+class TRobinHoodHashSet : public RobinHoodHashTable_Private::TRobinHoodHashTable<KeyType, RobinHoodHashTable_Private::FUnitType, Hasher, HashMapAllocator>
+{
+	using Unit = RobinHoodHashTable_Private::FUnitType;
+	using Base = typename RobinHoodHashTable_Private::TRobinHoodHashTable<KeyType, Unit, Hasher, HashMapAllocator>;
+	using IndexType = typename Base::IndexType;
+	using FindValueType = typename Base::FindValueType;
+
+public:
+	TRobinHoodHashSet() : Base()
+	{
+		static_assert(sizeof(Base) == sizeof(TRobinHoodHashSet), "This class should only limit the interface and not implement anything");
+	}
+
+	TRobinHoodHashSet(const TRobinHoodHashSet& Other) = default;
+	TRobinHoodHashSet& operator=(const TRobinHoodHashSet& Other) = default;
+	TRobinHoodHashSet(TRobinHoodHashSet&& Other) = default;
+	TRobinHoodHashSet& operator=(TRobinHoodHashSet&& Other) = default;
+
+	FHashElementId FindOrAddIdByHash(FHashType HashValue, const KeyType& Key)
+	{
+		return Base::template FindOrAddIdByHash<const KeyType&, Unit>(HashValue, Key, Unit());
+	}
+
+	FHashElementId FindOrAddIdByHash(FHashType HashValue, KeyType&& Key)
+	{
+		return Base::template FindOrAddIdByHash<KeyType&&, Unit>(HashValue, MoveTemp(Key), Unit());
+	}
+
+	FHashElementId FindOrAddId(const KeyType& Key)
+	{
+		return Base::template FindOrAddId<const KeyType&, Unit>(Key, Unit());
+	}
+
+	FHashElementId FindOrAddId(KeyType&& Key)
+	{
+		return Base::template FindOrAddId<KeyType&&, Unit>(MoveTemp(Key), Unit());
+	}
+
+	FindValueType FindOrAdd(const KeyType& Key)
+	{
+		return Base::template FindOrAdd<const KeyType&, Unit>(Key, Unit());
+	}
+
+	FindValueType FindOrAdd(KeyType&& Key)
+	{
+		return Base::template FindOrAdd<KeyType&&, Unit>(MoveTemp(Key), Unit());
+	}
+};
+
+};
\ No newline at end of file
