Add copyright
#rnx

diff --git a/Engine/Source/Runtime/Core/Public/Experimental/Containers/RobinHoodHashTable.h b/Engine/Source/Runtime/Core/Public/Experimental/Containers/RobinHoodHashTable.h
--- a/Engine/Source/Runtime/Core/Public/Experimental/Containers/RobinHoodHashTable.h
+++ b/Engine/Source/Runtime/Core/Public/Experimental/Containers/RobinHoodHashTable.h
@@ -1,3 +1,4 @@
+// Copyright 1998-2019 Epic Games, Inc. All Rights Reserved.
 #pragma once
 
 // Implementation of Robin Hood hash table.
