Fixed a number of log warnings seen in the Mac Client cook relating to bad file reads.

Removed the optimization added in CL 7474219 due to an issue where the internal buffer and the current position of the actual file handle would get out of sync. This code will be fixed properly and resubmitted for the next release.

diff --git a/Engine/Source/Runtime/Core/Private/HAL/FileManagerGeneric.cpp b/Engine/Source/Runtime/Core/Private/HAL/FileManagerGeneric.cpp
--- a/Engine/Source/Runtime/Core/Private/HAL/FileManagerGeneric.cpp
+++ b/Engine/Source/Runtime/Core/Private/HAL/FileManagerGeneric.cpp
@@ -654,25 +654,16 @@ void FArchiveFileReaderGeneric::Seek( int64 InPos )
 	checkf(InPos >= 0, TEXT("Attempted to seek to a negative location (%lld/%lld), file: %s. The file is most likely corrupt."), InPos, Size, *Filename);
 	checkf(InPos <= Size, TEXT("Attempted to seek past the end of file (%lld/%lld), file: %s. The file is most likely corrupt."), InPos, Size, *Filename);
 
-	if (InPos < BufferBase || InPos >= BufferBase + BufferCount)
+	if (!SeekLowLevel(InPos))
 	{
-		// The seek is outside of the cached buffer so we actually need to seek
-		if (!SeekLowLevel(InPos))
-		{
-			TCHAR ErrorBuffer[1024];
-			ArIsError = true;
-			UE_LOG(LogFileManager, Error, TEXT("SetFilePointer on %s Failed %lld/%lld: %lld %s"), *Filename, InPos, Size, Pos, FPlatformMisc::GetSystemErrorMessage(ErrorBuffer, 1024, 0));
-		}
-
-		Pos = InPos;
-		BufferBase = Pos;
-		BufferCount = 0;
-	}
-	else
-	{
-		// We already have the seek position cached so we can just move the current reading position
-		Pos = InPos;
+		TCHAR ErrorBuffer[1024];
+		ArIsError = true;
+		UE_LOG(LogFileManager, Error, TEXT("SetFilePointer on %s Failed %lld/%lld: %lld %s"), *Filename, InPos, Size, Pos, FPlatformMisc::GetSystemErrorMessage(ErrorBuffer, 1024, 0));
 	}
+
+	Pos = InPos;
+	BufferBase = Pos;
+	BufferCount = 0;
 }
 
 FArchiveFileReaderGeneric::~FArchiveFileReaderGeneric()
