Fix for build configuration and isExisting incorrect in CrashContext.runtime-xml

(ushell-p4-cherrypick of 9868284 by [REDACTED_PERSON])

diff --git a/Engine/Source/Runtime/Core/Private/GenericPlatform/GenericPlatformCrashContext.cpp b/Engine/Source/Runtime/Core/Private/GenericPlatform/GenericPlatformCrashContext.cpp
--- a/Engine/Source/Runtime/Core/Private/GenericPlatform/GenericPlatformCrashContext.cpp
+++ b/Engine/Source/Runtime/Core/Private/GenericPlatform/GenericPlatformCrashContext.cpp
@@ -94,6 +94,7 @@ void FGenericCrashContext::Initialize()
 	FCString::Strcpy(NCached::Session.GameSessionID, TEXT(""));
 	FCString::Strcpy(NCached::Session.GameStateName, TEXT(""));
 	FCString::Strcpy(NCached::Session.UserActivityHint, TEXT(""));
+	FCString::Strcpy(NCached::Session.BuildConfigurationName, LexToString(FApp::GetBuildConfiguration()));
 	FCString::Strcpy(NCached::Session.ExecutableName, FPlatformProcess::ExecutableName());
 	FCString::Strcpy(NCached::Session.BaseDir, FPlatformProcess::BaseDir());
 	FCString::Strcpy(NCached::Session.RootDir, FPlatformMisc::RootDir());
@@ -443,7 +444,7 @@ void FGenericCrashContext::SerializeSessionContext(FString& Buffer)
 		}
 	}
 	AddCrashPropertyInternal(Buffer, TEXT("ExecutableName"), NCached::Session.ExecutableName);
-	AddCrashPropertyInternal(Buffer, TEXT("BuildConfiguration"), LexToString(FApp::GetBuildConfiguration()));
+	AddCrashPropertyInternal(Buffer, TEXT("BuildConfiguration"), NCached::Session.BuildConfigurationName);
 	AddCrashPropertyInternal(Buffer, TEXT("GameSessionID"), NCached::Session.GameSessionID);
 
 	// Unique string specifying the symbols to be used by CrashReporter
@@ -451,7 +452,7 @@ void FGenericCrashContext::SerializeSessionContext(FString& Buffer)
 #ifdef UE_APP_FLAVOR
 	Symbols = FString::Printf(TEXT("%s-%s"), *Symbols, *FString(UE_APP_FLAVOR));
 #endif
-	Symbols = FString::Printf(TEXT("%s-%s-%s"), *Symbols, FPlatformMisc::GetUBTPlatform(), LexToString(FApp::GetBuildConfiguration())).Replace(TEXT("+"), TEXT("*"));
+	Symbols = FString::Printf(TEXT("%s-%s-%s"), *Symbols, FPlatformMisc::GetUBTPlatform(), NCached::Session.BuildConfigurationName).Replace(TEXT("+"), TEXT("*"));
 #ifdef UE_BUILD_FLAVOR
 	Symbols = FString::Printf(TEXT("%s-%s"), *Symbols, *FString(UE_BUILD_FLAVOR));
 #endif
@@ -471,7 +472,7 @@ void FGenericCrashContext::SerializeSessionContext(FString& Buffer)
 	AddCrashPropertyInternal(Buffer, TEXT("AppDefaultLocale"), NCached::Session.DefaultLocale);
 	AddCrashPropertyInternal(Buffer, TEXT("BuildVersion"), FApp::GetBuildVersion());
 	AddCrashPropertyInternal(Buffer, TEXT("IsUE4Release"), NCached::Session.bIsUE4Release);
-	AddCrashPropertyInternal(Buffer, TEXT("IsRequestingExit"), IsEngineExitRequested());
+	AddCrashPropertyInternal(Buffer, TEXT("IsRequestingExit"), NCached::Session.bIsExitRequested);
 
 	// Remove periods from user names to match AutoReporter user names
 	// The name prefix is read by CrashRepository.AddNewCrash in the website code
diff --git a/Engine/Source/Runtime/Core/Private/Windows/WindowsPlatformCrashContext.cpp b/Engine/Source/Runtime/Core/Private/Windows/WindowsPlatformCrashContext.cpp
--- a/Engine/Source/Runtime/Core/Private/Windows/WindowsPlatformCrashContext.cpp
+++ b/Engine/Source/Runtime/Core/Private/Windows/WindowsPlatformCrashContext.cpp
@@ -504,6 +504,7 @@ int32 ReportCrashForMonitor(
 	SharedContext->UserSettings.bSendUnattendedBugReports = bSendUnattendedBugReports;
 	SharedContext->UserSettings.bSendUsageData = bSendUsageData;
 
+	SharedContext->SessionContext.bIsExitRequested = IsEngineExitRequested();
 	FCString::Strcpy(SharedContext->ErrorMessage, CR_MAX_ERROR_MESSAGE_CHARS-1, ErrorMessage);
 
 	if (GLog)
diff --git a/Engine/Source/Runtime/Core/Public/GenericPlatform/GenericPlatformCrashContext.h b/Engine/Source/Runtime/Core/Public/GenericPlatform/GenericPlatformCrashContext.h
--- a/Engine/Source/Runtime/Core/Public/GenericPlatform/GenericPlatformCrashContext.h
+++ b/Engine/Source/Runtime/Core/Public/GenericPlatform/GenericPlatformCrashContext.h
@@ -130,6 +130,7 @@ struct FSessionContext
 	bool 					bIsSourceDistribution;
 	bool 					bIsUE4Release;
 	bool					bIsOOM;
+	bool					bIsExitRequested;
 	uint32					ProcessId;
 	int32 					LanguageLCID;
 	int32 					NumberOfCores;
@@ -143,6 +144,7 @@ struct FSessionContext
 	TCHAR					EngineMode[CR_MAX_GENERIC_FIELD_CHARS];
 	TCHAR					EngineModeEx[CR_MAX_GENERIC_FIELD_CHARS];
 	TCHAR 					ExecutableName[CR_MAX_GENERIC_FIELD_CHARS];
+	TCHAR					BuildConfigurationName[CR_MAX_GENERIC_FIELD_CHARS];
 	TCHAR 					DeploymentName[CR_MAX_GENERIC_FIELD_CHARS];
 	TCHAR 					BaseDir[CR_MAX_DIRECTORY_CHARS];
 	TCHAR 					RootDir[CR_MAX_DIRECTORY_CHARS];
