Added a const qualifier to FMatrix::ApplyScale. This function was already const however it was not marked as such

diff --git a/Engine/Source/Runtime/Core/Public/Math/Matrix.h b/Engine/Source/Runtime/Core/Public/Math/Matrix.h
--- a/Engine/Source/Runtime/Core/Public/Math/Matrix.h
+++ b/Engine/Source/Runtime/Core/Public/Math/Matrix.h
@@ -205,7 +205,7 @@ public:
 	inline float GetMaximumAxisScale() const;
 
 	/** Apply Scale to this matrix **/
-	inline FMatrix ApplyScale(float Scale);
+	inline FMatrix ApplyScale(float Scale) const;
 
 	// @return the origin of the co-ordinate system
 	inline FVector GetOrigin() const;
diff --git a/Engine/Source/Runtime/Core/Public/Math/Matrix.inl b/Engine/Source/Runtime/Core/Public/Math/Matrix.inl
--- a/Engine/Source/Runtime/Core/Public/Math/Matrix.inl
+++ b/Engine/Source/Runtime/Core/Public/Math/Matrix.inl
@@ -749,7 +749,7 @@ inline void FMatrix::Mirror(EAxis::Type MirrorAxis, EAxis::Type FlipAxis)
 /** 
  * Apply Scale to this matrix
  */
-inline FMatrix FMatrix::ApplyScale(float Scale)
+inline FMatrix FMatrix::ApplyScale(float Scale) const
 {
 	FMatrix ScaleMatrix(
 		FPlane(Scale, 0.0f, 0.0f, 0.0f),
