Add version of LoadFileToString() that takes a specific IPlatformFile.

diff --git a/Engine/Source/Runtime/Core/Private/Misc/FileHelper.cpp b/Engine/Source/Runtime/Core/Private/Misc/FileHelper.cpp
--- a/Engine/Source/Runtime/Core/Private/Misc/FileHelper.cpp
+++ b/Engine/Source/Runtime/Core/Private/Misc/FileHelper.cpp
@@ -12,6 +12,7 @@
 #include "ProfilingDebugging/ProfilingHelpers.h"
 #include "Misc/ConfigCacheIni.h"
 #include "Misc/SecureHash.h"
+#include "HAL/FileManagerGeneric.h"
 
 #define LOCTEXT_NAMESPACE "FileHelper"
 
@@ -122,33 +123,55 @@ void FFileHelper::BufferToString( FString& Result, const uint8* Buffer, int32 Si
  * @param VerifyFlags flags controlling the hash verification behavior ( see EHashOptions )
  */
 bool FFileHelper::LoadFileToString( FString& Result, const TCHAR* Filename, EHashOptions VerifyFlags, uint32 ReadFlags)
+{
+	return LoadFileToString(Result, nullptr, Filename, VerifyFlags, ReadFlags);
+}
+
+bool FFileHelper::LoadFileToString(FString& Result, IPlatformFile* PlatformFile, const TCHAR* Filename, EHashOptions VerifyFlags /*= EHashOptions::None*/, uint32 ReadFlags)
 {
 	FScopedLoadingState ScopedLoadingState(Filename);
 
-	TUniquePtr<FArchive> Reader( IFileManager::Get().CreateFileReader( Filename, ReadFlags) );
-	if( !Reader )
+	TUniquePtr<FArchive> Reader;
+	if (PlatformFile)
+	{
+		IFileHandle* File = PlatformFile->OpenRead(Filename);
+		if (!File)
+		{
+			UE_LOG(LogStreaming, Warning, TEXT("Failed to read file '%s' error."), Filename);
+		}
+		else
+		{
+			Reader = MakeUnique<FArchiveFileReaderGeneric>(File, Filename, File->Size());
+		}
+	}
+	else
+	{
+		Reader.Reset(IFileManager::Get().CreateFileReader(Filename, ReadFlags));
+	}
+
+	if (!Reader)
 	{
 		return false;
 	}
-	
+
 	int32 Size = Reader->TotalSize();
-	if( !Size )
+	if (!Size)
 	{
 		Result.Empty();
 		return true;
 	}
 
 	uint8* Ch = (uint8*)FMemory::Malloc(Size);
-	Reader->Serialize( Ch, Size );
+	Reader->Serialize(Ch, Size);
 	bool Success = Reader->Close();
 	Reader = nullptr;
-	BufferToString( Result, Ch, Size );
+	BufferToString(Result, Ch, Size);
 
 	// handle SHA verify of the file
-	if( EnumHasAnyFlags(VerifyFlags, EHashOptions::EnableVerify) && ( EnumHasAnyFlags(VerifyFlags, EHashOptions::ErrorMissingHash) || FSHA1::GetFileSHAHash(Filename, NULL) ) )
+	if (EnumHasAnyFlags(VerifyFlags, EHashOptions::EnableVerify) && (EnumHasAnyFlags(VerifyFlags, EHashOptions::ErrorMissingHash) || FSHA1::GetFileSHAHash(Filename, NULL)))
 	{
 		// kick off SHA verify task. this frees the buffer on close
-		FBufferReaderWithSHA Ar( Ch, Size, true, Filename, false, true );
+		FBufferReaderWithSHA Ar(Ch, Size, true, Filename, false, true);
 	}
 	else
 	{
diff --git a/Engine/Source/Runtime/Core/Public/Misc/FileHelper.h b/Engine/Source/Runtime/Core/Public/Misc/FileHelper.h
--- a/Engine/Source/Runtime/Core/Public/Misc/FileHelper.h
+++ b/Engine/Source/Runtime/Core/Public/Misc/FileHelper.h
@@ -55,6 +55,16 @@ struct CORE_API FFileHelper
 	 */
 	static bool LoadFileToString( FString& Result, const TCHAR* Filename, EHashOptions VerifyFlags = EHashOptions::None, uint32 ReadFlags = 0 );
 
+	/**
+	 * Load a text file to an FString. Supports all combination of ANSI/Unicode files and platforms.
+	 *
+	 * @param Result       String representation of the loaded file
+	 * @param PlatformFile PlatformFile interface to use
+	 * @param Filename     Name of the file to load
+	 * @param VerifyFlags  Flags controlling the hash verification behavior ( see EHashOptions )
+	 */
+	static bool LoadFileToString(FString& Result, IPlatformFile* PlatformFile, const TCHAR* Filename, EHashOptions VerifyFlags = EHashOptions::None, uint32 ReadFlags = 0 );
+
 	/**
 	 * Load a text file to an array of strings. Supports all combination of ANSI/Unicode files and platforms.
 	 *
