Use different delimited and equals characters to encode crash context pairs.

Using ',' and '=' to encode key/value pairs in the dynamic data part of the crash context didn't work well when the values contained more complex data, like json objects. Use character codes 1 and 2 instead.

diff --git a/Engine/Source/Runtime/Core/Private/GenericPlatform/GenericPlatformCrashContext.cpp b/Engine/Source/Runtime/Core/Private/GenericPlatform/GenericPlatformCrashContext.cpp
--- a/Engine/Source/Runtime/Core/Private/GenericPlatform/GenericPlatformCrashContext.cpp
+++ b/Engine/Source/Runtime/Core/Private/GenericPlatform/GenericPlatformCrashContext.cpp
@@ -232,9 +232,15 @@ PRAGMA_ENABLE_DEPRECATION_WARNINGS
 #endif	// !NOINITCRASHREPORTER
 }
 
+// When encoding the plugins list and engine and game data key/value pairs into the dynamic data segment
+// we use 1 and 2 to denote delimiter and equals respectively. This is necessary since the values could 
+// contain any characters normally used for delimiting.
+const TCHAR* CR_PAIR_DELIM = TEXT("\x01");
+const TCHAR* CR_PAIR_EQ = TEXT("\x02");
+
 void FGenericCrashContext::InitializeFromContext(const FSessionContext& Session, const TCHAR* EnabledPluginsStr, const TCHAR* EngineDataStr, const TCHAR* GameDataStr)
 {
-	static const TCHAR* TokenDelim[] = { TEXT(","), TEXT("=") };
+	static const TCHAR* TokenDelim[] = { CR_PAIR_DELIM, CR_PAIR_EQ };
 
 	// Copy the session struct which should be all pod types and fixed size buggers
 	FMemory::Memcpy(NCached::Session, Session);
@@ -294,7 +300,7 @@ void FGenericCrashContext::CopySharedCrashContext(FSharedCrashContext& Dst)
 	for (const FString& Plugin : NCached::EnabledPluginsList)
 	{
 		FCString::Strcat(DynamicDataPtr, Plugin.Len(), *Plugin);
-		FCString::Strcat(DynamicDataPtr, 1, TEXT(","));
+		FCString::Strcat(DynamicDataPtr, 1, CR_PAIR_DELIM);
 	}
 	DynamicDataPtr += FCString::Strlen(DynamicDataPtr) + 1;
 
@@ -303,9 +309,9 @@ void FGenericCrashContext::CopySharedCrashContext(FSharedCrashContext& Dst)
 	for (const TPair<FString, FString>& Pair : NCached::EngineData)
 	{
 		FCString::Strcat(DynamicDataPtr, Pair.Key.Len(), *Pair.Key);
-		FCString::Strcat(DynamicDataPtr, 1, TEXT("="));
+		FCString::Strcat(DynamicDataPtr, 1, CR_PAIR_EQ);
 		FCString::Strcat(DynamicDataPtr, Pair.Value.Len(), *Pair.Value);
-		FCString::Strcat(DynamicDataPtr, 1, TEXT(","));
+		FCString::Strcat(DynamicDataPtr, 1, CR_PAIR_DELIM);
 	}
 	DynamicDataPtr += FCString::Strlen(DynamicDataPtr) + 1;
 
@@ -314,9 +320,9 @@ void FGenericCrashContext::CopySharedCrashContext(FSharedCrashContext& Dst)
 	for (const TPair<FString, FString>& Pair : NCached::GameData)
 	{
 		FCString::Strcat(DynamicDataPtr, Pair.Key.Len(), *Pair.Key);
-		FCString::Strcat(DynamicDataPtr, 1, TEXT("="));
+		FCString::Strcat(DynamicDataPtr, 1, CR_PAIR_EQ);
 		FCString::Strcat(DynamicDataPtr, Pair.Value.Len(), *Pair.Value);
-		FCString::Strcat(DynamicDataPtr, 1, TEXT(","));
+		FCString::Strcat(DynamicDataPtr, 1, CR_PAIR_DELIM);
 	}
 	DynamicDataPtr += FCString::Strlen(DynamicDataPtr) + 1;
 }
