Add a version of LoadFileToString() that takes an FArchive.  Other versions call that one.

diff --git a/Engine/Source/Runtime/Core/Private/Misc/FileHelper.cpp b/Engine/Source/Runtime/Core/Private/Misc/FileHelper.cpp
--- a/Engine/Source/Runtime/Core/Private/Misc/FileHelper.cpp
+++ b/Engine/Source/Runtime/Core/Private/Misc/FileHelper.cpp
@@ -115,63 +115,34 @@ void FFileHelper::BufferToString( FString& Result, const uint8* Buffer, int32 Si
 	}
 }
 
-/**
- * Load a text file to an FString.
- * Supports all combination of ANSI/Unicode files and platforms.
- * @param Result string representation of the loaded file
- * @param Filename name of the file to load
- * @param VerifyFlags flags controlling the hash verification behavior ( see EHashOptions )
- */
-bool FFileHelper::LoadFileToString( FString& Result, const TCHAR* Filename, EHashOptions VerifyFlags, uint32 ReadFlags)
+bool FFileHelper::LoadFileToString(FString& Result, FArchive& Reader, EHashOptions VerifyFlags /*= EHashOptions::None*/)
 {
-	return LoadFileToString(Result, nullptr, Filename, VerifyFlags, ReadFlags);
-}
-
-bool FFileHelper::LoadFileToString(FString& Result, IPlatformFile* PlatformFile, const TCHAR* Filename, EHashOptions VerifyFlags /*= EHashOptions::None*/, uint32 ReadFlags)
-{
-	FScopedLoadingState ScopedLoadingState(Filename);
+	FScopedLoadingState ScopedLoadingState(*Reader.GetArchiveName());
 
-	TUniquePtr<FArchive> Reader;
-	if (PlatformFile)
-	{
-		IFileHandle* File = PlatformFile->OpenRead(Filename);
-		if (!File)
-		{
-			UE_LOG(LogStreaming, Warning, TEXT("Failed to read file '%s' error."), Filename);
-		}
-		else
-		{
-			Reader = MakeUnique<FArchiveFileReaderGeneric>(File, Filename, File->Size());
-		}
-	}
-	else
+	int32 Size = Reader.TotalSize();
+	if (!Size)
 	{
-		Reader.Reset(IFileManager::Get().CreateFileReader(Filename, ReadFlags));
+		Result.Empty();
+		return true;
 	}
 
-	if (!Reader)
+	if (Reader.Tell() != 0)
 	{
+		UE_LOG(LogStreaming, Warning, TEXT("Archive '%s' has already been read from."), *Reader.GetArchiveName());
 		return false;
 	}
 
-	int32 Size = Reader->TotalSize();
-	if (!Size)
-	{
-		Result.Empty();
-		return true;
-	}
-
 	uint8* Ch = (uint8*)FMemory::Malloc(Size);
-	Reader->Serialize(Ch, Size);
-	bool Success = Reader->Close();
-	Reader = nullptr;
+	Reader.Serialize(Ch, Size);
+	bool Success = Reader.Close();
+
 	BufferToString(Result, Ch, Size);
 
 	// handle SHA verify of the file
-	if (EnumHasAnyFlags(VerifyFlags, EHashOptions::EnableVerify) && (EnumHasAnyFlags(VerifyFlags, EHashOptions::ErrorMissingHash) || FSHA1::GetFileSHAHash(Filename, NULL)))
+	if (EnumHasAnyFlags(VerifyFlags, EHashOptions::EnableVerify) && (EnumHasAnyFlags(VerifyFlags, EHashOptions::ErrorMissingHash) || FSHA1::GetFileSHAHash(*Reader.GetArchiveName(), nullptr)))
 	{
 		// kick off SHA verify task. this frees the buffer on close
-		FBufferReaderWithSHA Ar(Ch, Size, true, Filename, false, true);
+		FBufferReaderWithSHA Ar(Ch, Size, true, *Reader.GetArchiveName(), false, true);
 	}
 	else
 	{
@@ -182,6 +153,47 @@ bool FFileHelper::LoadFileToString(FString& Result, IPlatformFile* PlatformFile,
 	return Success;
 }
 
+/**
+ * Load a text file to an FString.
+ * Supports all combination of ANSI/Unicode files and platforms.
+ * @param Result string representation of the loaded file
+ * @param Filename name of the file to load
+ * @param VerifyFlags flags controlling the hash verification behavior ( see EHashOptions )
+ */
+bool FFileHelper::LoadFileToString( FString& Result, const TCHAR* Filename, EHashOptions VerifyFlags, uint32 ReadFlags)
+{
+	TUniquePtr<FArchive> Reader(IFileManager::Get().CreateFileReader(Filename, ReadFlags));
+	if (!Reader)
+	{
+		return false;
+	}
+
+	return LoadFileToString(Result, *Reader.Get(), VerifyFlags);
+}
+
+bool FFileHelper::LoadFileToString(FString& Result, IPlatformFile* PlatformFile, const TCHAR* Filename, EHashOptions VerifyFlags /*= EHashOptions::None*/)
+{
+	if (!PlatformFile)
+	{
+		return false;
+	}
+
+	IFileHandle* File = PlatformFile->OpenRead(Filename);
+	if (!File)
+	{
+		UE_LOG(LogStreaming, Warning, TEXT("Failed to read file '%s' error."), Filename);
+		return false;
+	}
+
+	TUniquePtr<FArchive> Reader = MakeUnique<FArchiveFileReaderGeneric>(File, Filename, File->Size());
+	if (!Reader)
+	{
+		return false;
+	}
+
+	return LoadFileToString(Result, *Reader.Get(), VerifyFlags);
+}
+
 bool FFileHelper::LoadFileToStringArray( TArray<FString>& Result, const TCHAR* Filename, EHashOptions VerifyFlags )
 {
 	Result.Empty();
diff --git a/Engine/Source/Runtime/Core/Public/Misc/FileHelper.h b/Engine/Source/Runtime/Core/Public/Misc/FileHelper.h
--- a/Engine/Source/Runtime/Core/Public/Misc/FileHelper.h
+++ b/Engine/Source/Runtime/Core/Public/Misc/FileHelper.h
@@ -46,6 +46,15 @@ struct CORE_API FFileHelper
 	*/
 	static bool LoadFileToArray( TArray<uint8>& Result, const TCHAR* Filename, uint32 Flags = 0 );
 
+	/**
+	 * Load a text file to an FString. Supports all combination of ANSI/Unicode files and platforms.
+	 *
+	 * @param Result       String representation of the loaded file
+	 * @param Archive      Name of the archive to load from
+	 * @param VerifyFlags  Flags controlling the hash verification behavior ( see EHashOptions )
+	 */
+	static bool LoadFileToString(FString& Result, FArchive& Reader, EHashOptions VerifyFlags = EHashOptions::None);
+
 	/**
 	 * Load a text file to an FString. Supports all combination of ANSI/Unicode files and platforms.
 	 *
@@ -63,7 +72,7 @@ struct CORE_API FFileHelper
 	 * @param Filename     Name of the file to load
 	 * @param VerifyFlags  Flags controlling the hash verification behavior ( see EHashOptions )
 	 */
-	static bool LoadFileToString(FString& Result, IPlatformFile* PlatformFile, const TCHAR* Filename, EHashOptions VerifyFlags = EHashOptions::None, uint32 ReadFlags = 0 );
+	static bool LoadFileToString(FString& Result, IPlatformFile* PlatformFile, const TCHAR* Filename, EHashOptions VerifyFlags = EHashOptions::None);
 
 	/**
 	 * Load a text file to an array of strings. Supports all combination of ANSI/Unicode files and platforms.
@@ -118,7 +127,7 @@ struct CORE_API FFileHelper
 	 * @return true if success
 	 */
 	static bool GenerateNextBitmapFilename(const FString& Pattern, const FString& Extension, FString& OutFilename, IFileManager* FileManager = &IFileManager::Get());
-
+	
 	/**
 	 * Generates the next unique bitmap filename with a specified extension
 	 *
