Start/Stop csv on event + add frameCount inside csv profiler

example:
-csvStartOnEvent="My Event1" -csvStopOnEvent="My Event2"
-csvStartOnEvent="My Event1" -csvCaptureOnEventFrameCount=1500

#rsx

diff --git a/Engine/Source/Runtime/Core/Private/ProfilingDebugging/CsvProfiler.cpp b/Engine/Source/Runtime/Core/Private/ProfilingDebugging/CsvProfiler.cpp
--- a/Engine/Source/Runtime/Core/Private/ProfilingDebugging/CsvProfiler.cpp
+++ b/Engine/Source/Runtime/Core/Private/ProfilingDebugging/CsvProfiler.cpp
@@ -118,7 +118,8 @@ static bool GCsvUseProcessingThread = true;
 static int32 GCsvRepeatCount = 0;
 static int32 GCsvRepeatFrameCount = 0;
 static bool GCsvStatCounts = false;
-
+static FString* GStartOnEvent = nullptr;
+static FString* GStopOnEvent = nullptr;
 static uint32 GCsvProcessingThreadId = 0;
 static bool GGameThreadIsCsvProcessingThread = true;
 
@@ -2442,6 +2443,7 @@ FCsvProfiler* FCsvProfiler::Get()
 FCsvProfiler::FCsvProfiler()
 	: NumFramesToCapture(-1)
 	, CaptureFrameNumber(0)
+	, CaptureOnEventFrameCount(-1)
 	, bInsertEndFrameAtFrameStart(false)
 	, LastEndFrameTimestamp(0)
 	, CaptureEndFrameCount(0)
@@ -2473,6 +2475,18 @@ FCsvProfiler::~FCsvProfiler()
 		FPlatformProcess::ReturnSynchEventToPool(FileWriteBlockingEvent);
 		FileWriteBlockingEvent = nullptr;
 	}
+
+	if (GStartOnEvent)
+	{
+		delete GStartOnEvent;
+		GStartOnEvent = nullptr;
+	}
+
+	if (GStopOnEvent)
+	{
+		delete GStopOnEvent;
+		GStopOnEvent = nullptr;
+	}
 }
 
 /** Per-frame update */
@@ -2970,14 +2984,31 @@ void FCsvProfiler::EndWait()
 
 void FCsvProfiler::RecordEventfInternal(int32 CategoryIndex, const TCHAR* Fmt, ...)
 {
-	if (GCsvProfilerIsCapturing && GCsvCategoriesEnabled[CategoryIndex])
+	bool bIsCsvRecording = GCsvProfilerIsCapturing && GCsvCategoriesEnabled[CategoryIndex];
+	if (bIsCsvRecording || GStartOnEvent)
 	{
 		LLM_SCOPE(ELLMTag::CsvProfiler);
 		TCHAR Buffer[256];
 		GET_VARARGS(Buffer, UE_ARRAY_COUNT(Buffer), UE_ARRAY_COUNT(Buffer) - 1, Fmt, Fmt);
 		Buffer[255] = '\0';
 		FString Str = Buffer;
-		RecordEvent(CategoryIndex, Str);
+
+		if (bIsCsvRecording)
+		{
+			RecordEvent(CategoryIndex, Str);
+
+			if (GStopOnEvent && GStopOnEvent->Equals(Str, ESearchCase::IgnoreCase))
+			{
+				FCsvProfiler::Get()->EndCapture();
+			}
+		}
+		else
+		{
+			if (GStartOnEvent && GStartOnEvent->Equals(Str, ESearchCase::IgnoreCase))
+			{
+				FCsvProfiler::Get()->BeginCapture(FCsvProfiler::Get()->GetNumFrameToCaptureOnEvent());
+			}
+		}
 	}
 }
 
@@ -3063,8 +3094,26 @@ void FCsvProfiler::RecordCustomStat(const FName& StatName, uint32 CategoryIndex,
 void FCsvProfiler::Init()
 {
 #if CSV_PROFILER_ALLOW_DEBUG_FEATURES
+	FParse::Value(FCommandLine::Get(), TEXT("csvCaptureOnEventFrameCount="), CaptureOnEventFrameCount);
+
+	GStartOnEvent = new FString();
+	FParse::Value(FCommandLine::Get(), TEXT("csvStartOnEvent="), *GStartOnEvent);
+
+	if (GStartOnEvent->IsEmpty())
+	{
+		delete GStartOnEvent;
+		GStartOnEvent = nullptr;
+	}
+
+	GStopOnEvent = new FString();
+	FParse::Value(FCommandLine::Get(), TEXT("csvStopOnEvent="), *GStopOnEvent);
+
+	if (GStopOnEvent->IsEmpty())
+	{
+		delete GStopOnEvent;
+		GStopOnEvent = nullptr;
+	}
 
-	int32 NumCsvFrames = 0;
 	if (FParse::Param(FCommandLine::Get(), TEXT("csvGpuStats")))
 	{
 		IConsoleVariable* CVarGPUCsvStatsEnabled = IConsoleManager::Get().FindConsoleVariable(TEXT("r.GPUCsvStatsEnabled"));
@@ -3123,6 +3172,7 @@ void FCsvProfiler::Init()
 	{
 		CVarCsvStatCounts.AsVariable()->Set(1);
 	}
+	int32 NumCsvFrames = 0;
 	if (FParse::Value(FCommandLine::Get(), TEXT("csvCaptureFrames="), NumCsvFrames))
 	{
 		check(IsInGameThread());
@@ -3171,11 +3221,20 @@ bool FCsvProfiler::IsWritingFile()
 	return GCsvProfilerIsWritingFile;
 }
 
+/*Get the current frame capture count*/
 int32 FCsvProfiler::GetCaptureFrameNumber()
 {
 	return CaptureFrameNumber;
 }
 
+//Get the total frame to capture when we are capturing on event. 
+//Example:  -csvStartOnEvent="My Event"
+//			-csvCaptureOnEventFrameCount=2500
+int32 FCsvProfiler::GetNumFrameToCaptureOnEvent()
+{
+	return CaptureOnEventFrameCount;
+}
+
 bool FCsvProfiler::EnableCategoryByString(const FString& CategoryName) const
 {
 	int32 Category = FCsvCategoryData::Get()->GetCategoryIndex(CategoryName);
diff --git a/Engine/Source/Runtime/Core/Public/ProfilingDebugging/CsvProfiler.h b/Engine/Source/Runtime/Core/Public/ProfilingDebugging/CsvProfiler.h
--- a/Engine/Source/Runtime/Core/Public/ProfilingDebugging/CsvProfiler.h
+++ b/Engine/Source/Runtime/Core/Public/ProfilingDebugging/CsvProfiler.h
@@ -261,6 +261,7 @@ public:
 	CORE_API bool IsWritingFile();
 
 	CORE_API int32 GetCaptureFrameNumber();
+	CORE_API int32 GetNumFrameToCaptureOnEvent();
 
 	CORE_API bool EnableCategoryByString(const FString& CategoryName) const;
 
@@ -310,6 +311,7 @@ private:
 
 	int32 NumFramesToCapture;
 	int32 CaptureFrameNumber;
+	int32 CaptureOnEventFrameCount;
 
 	bool bInsertEndFrameAtFrameStart;
 
