Do not show crash reporter dialog for ensures when the user has selected to not automatically submit bug reports.

diff --git a/Engine/Source/Runtime/Core/Private/Windows/WindowsPlatformCrashContext.cpp b/Engine/Source/Runtime/Core/Private/Windows/WindowsPlatformCrashContext.cpp
--- a/Engine/Source/Runtime/Core/Private/Windows/WindowsPlatformCrashContext.cpp
+++ b/Engine/Source/Runtime/Core/Private/Windows/WindowsPlatformCrashContext.cpp
@@ -488,6 +488,7 @@ int32 ReportCrashForMonitor(
 	bool bNoDialog = ReportUI == EErrorReportUI::ReportInUnattendedMode || IsRunningDedicatedServer();
 	bool bSendUnattendedBugReports = true;
 	bool bSendUsageData = true;
+	bool bCanSendCrashReport = true;
 
 	if (FCommandLine::IsInitialized())
 	{
@@ -500,6 +501,18 @@ int32 ReportCrashForMonitor(
 		GConfig->GetBool(TEXT("/Script/UnrealEd.AnalyticsPrivacySettings"), TEXT("bSendUsageData"), bSendUsageData, GEditorSettingsIni);
 	}
 
+	if (bNoDialog && !bSendUnattendedBugReports)
+	{
+		// If we shouldn't display a dialog (like for ensures) and the user
+		// does not allow unattended bug reports we cannot send the report.
+		bCanSendCrashReport = false;
+	}
+
+	if (!bSendUnattendedBugReports)
+	{
+		return EXCEPTION_CONTINUE_EXECUTION;
+	}
+
 	SharedContext->UserSettings.bNoDialog = bNoDialog;
 	SharedContext->UserSettings.bSendUnattendedBugReports = bSendUnattendedBugReports;
 	SharedContext->UserSettings.bSendUsageData = bSendUsageData;
