MN - implement Reserve for RobinHoodHashTable

diff --git a/Engine/Source/Runtime/Core/Public/Experimental/Containers/RobinHoodHashTable.h b/Engine/Source/Runtime/Core/Public/Experimental/Containers/RobinHoodHashTable.h
--- a/Engine/Source/Runtime/Core/Public/Experimental/Containers/RobinHoodHashTable.h
+++ b/Engine/Source/Runtime/Core/Public/Experimental/Containers/RobinHoodHashTable.h
@@ -126,7 +126,7 @@ namespace RobinHoodHashTable_Private
 		friend class TRobinHoodHashTable;
 
 		template<typename DeducedKeyType>
-		inline TKeyValue(DeducedKeyType&& InKey, FUnitType) : Key(Forward<DeducedKeyType>(InKey)) {}
+		inline TKeyValue(DeducedKeyType&& InKey, FUnitType&&) : Key(Forward<DeducedKeyType>(InKey)) {}
 
 		ElementType Key;
 
@@ -154,6 +154,7 @@ namespace RobinHoodHashTable_Private
 		using FindValueType = typename KeyValueType::FindValueType;
 		using ElementType = typename KeyValueType::ElementType;
 		using IndexType = uint32;
+		using SizeType = SIZE_T;
 
 		static constexpr const IndexType LoadFactorDivisor = 3;
 		static constexpr const IndexType LoadFactorQuotient = 5;
@@ -171,7 +172,7 @@ namespace RobinHoodHashTable_Private
 				Empty();
 			}
 
-			SIZE_T GetAllocatedSize() const
+			SizeType GetAllocatedSize() const
 			{
 				return KeyVals.GetAllocatedSize() + FreeList.GetAllocatedSize();
 			}
@@ -235,9 +236,9 @@ namespace RobinHoodHashTable_Private
 				return KeyVals[Index];
 			}
 
-			inline IndexType Num() const
+			inline SizeType Num() const
 			{
-				return KeyVals.Num() - FreeList.Num();
+				return SizeType(KeyVals.Num()) - SizeType(FreeList.Num());
 			}
 
 			inline IndexType GetMaxIndex() const
@@ -302,6 +303,11 @@ namespace RobinHoodHashTable_Private
 				FreeList.Empty();
 			}
 
+			void Reserve(SizeType ReserveNum)
+			{
+				KeyVals.Reserve(ReserveNum);
+			}
+
 			TArray<KeyValueType, HashMapAllocator> KeyVals;
 			TArray<IndexType, HashMapAllocator> FreeList;
 		};
@@ -383,8 +389,10 @@ namespace RobinHoodHashTable_Private
 				IndexType OldSizePow2Minus1 = SizePow2Minus1;
 				SizePow2Minus1 = SizePow2Minus1 * 2 + 1;
 				MaximumDistance = 0;
-				IndexData.Reserve(SizePow2Minus1 + 1); IndexData.AddUninitialized(SizePow2Minus1 + 1);
-				HashData.Reserve(SizePow2Minus1 + 1);  HashData.AddDefaulted(SizePow2Minus1 + 1);
+				IndexData.Reserve(SizePow2Minus1 + 1); 
+				IndexData.AddUninitialized(SizePow2Minus1 + 1);
+				HashData.Reserve(SizePow2Minus1 + 1);  
+				HashData.AddDefaulted(SizePow2Minus1 + 1);
 
 				for (IndexType Index = 0; Index <= OldSizePow2Minus1; Index++)
 				{
@@ -604,7 +612,7 @@ namespace RobinHoodHashTable_Private
 			return FConstIteratorType(KeyValueData, false);
 		}
 
-		SIZE_T GetAllocatedSize() const
+		SizeType GetAllocatedSize() const
 		{
 			return KeyValueData.GetAllocatedSize() + IndexData.GetAllocatedSize() + HashData.GetAllocatedSize();
 		}
@@ -734,6 +742,41 @@ namespace RobinHoodHashTable_Private
 			MaximumDistance = 0;
 		}
 
+		void Reserve(SizeType ReserveNum)
+		{
+			if (ReserveNum > KeyValueData.Num())
+			{
+				KeyValueData.Reserve(ReserveNum);
+
+				IndexType NewSizePow2Minus1 = SizePow2Minus1;
+				while ((ReserveNum * LoadFactorQuotient) >= (NewSizePow2Minus1 * LoadFactorDivisor))
+				{
+					NewSizePow2Minus1 = NewSizePow2Minus1 * 2 + 1;
+				}
+
+				if (NewSizePow2Minus1 > SizePow2Minus1)
+				{
+					TArray<IndexType> IndexDataOld = MoveTemp(IndexData);
+					TArray<FHashType> HashDataOld = MoveTemp(HashData);
+					IndexType OldSizePow2Minus1 = SizePow2Minus1;
+					SizePow2Minus1 = NewSizePow2Minus1;
+					MaximumDistance = 0;
+					IndexData.Reserve(SizePow2Minus1 + 1); 
+					IndexData.AddUninitialized(SizePow2Minus1 + 1);
+					HashData.Reserve(SizePow2Minus1 + 1);  
+					HashData.AddDefaulted(SizePow2Minus1 + 1);
+
+					for (IndexType Index = 0; Index <= OldSizePow2Minus1; Index++)
+					{
+						if (HashDataOld[Index].IsOccupied())
+						{
+							InsertIntoTable(IndexDataOld[Index], HashDataOld[Index]);
+						}
+					}
+				}
+			}
+		}
+
 	private:
 		FData KeyValueData;
 
@@ -845,32 +888,32 @@ public:
 
 	FHashElementId FindOrAddIdByHash(FHashType HashValue, const KeyType& Key)
 	{
-		return Base::template FindOrAddIdByHash<const KeyType&, Unit>(HashValue, Key, Unit());
+		return Base::FindOrAddIdByHash(HashValue, Key, Unit());
 	}
 
 	FHashElementId FindOrAddIdByHash(FHashType HashValue, KeyType&& Key)
 	{
-		return Base::template FindOrAddIdByHash<KeyType&&, Unit>(HashValue, MoveTemp(Key), Unit());
+		return Base::FindOrAddIdByHash(HashValue, MoveTemp(Key), Unit());
 	}
 
 	FHashElementId FindOrAddId(const KeyType& Key)
 	{
-		return Base::template FindOrAddId<const KeyType&, Unit>(Key, Unit());
+		return Base::FindOrAddId(Key, Unit());
 	}
 
 	FHashElementId FindOrAddId(KeyType&& Key)
 	{
-		return Base::template FindOrAddId<KeyType&&, Unit>(MoveTemp(Key), Unit());
+		return Base::FindOrAddId(MoveTemp(Key), Unit());
 	}
 
 	FindValueType FindOrAdd(const KeyType& Key)
 	{
-		return Base::template FindOrAdd<const KeyType&, Unit>(Key, Unit());
+		return Base::FindOrAdd(Key, Unit());
 	}
 
 	FindValueType FindOrAdd(KeyType&& Key)
 	{
-		return Base::template FindOrAdd<KeyType&&, Unit>(MoveTemp(Key), Unit());
+		return Base::FindOrAdd(MoveTemp(Key), Unit());
 	}
 };
 
