Do not show crash reporter dialog for ensures when the user has selected to not automatically submit bug reports. Second attempt.

diff --git a/Engine/Source/Runtime/Core/Private/Windows/WindowsPlatformCrashContext.cpp b/Engine/Source/Runtime/Core/Private/Windows/WindowsPlatformCrashContext.cpp
--- a/Engine/Source/Runtime/Core/Private/Windows/WindowsPlatformCrashContext.cpp
+++ b/Engine/Source/Runtime/Core/Private/Windows/WindowsPlatformCrashContext.cpp
@@ -508,7 +508,7 @@ int32 ReportCrashForMonitor(
 		bCanSendCrashReport = false;
 	}
 
-	if (!bSendUnattendedBugReports)
+	if (!bCanSendCrashReport)
 	{
 		return EXCEPTION_CONTINUE_EXECUTION;
 	}
