[NO NON-PII COMMIT MESSAGE]

diff --git a/Engine/Source/Runtime/Core/Private/Windows/WindowsPlatformCrashContext.cpp b/Engine/Source/Runtime/Core/Private/Windows/WindowsPlatformCrashContext.cpp
--- a/Engine/Source/Runtime/Core/Private/Windows/WindowsPlatformCrashContext.cpp
+++ b/Engine/Source/Runtime/Core/Private/Windows/WindowsPlatformCrashContext.cpp
@@ -501,6 +501,15 @@ int32 ReportCrashForMonitor(
 		GConfig->GetBool(TEXT("/Script/UnrealEd.AnalyticsPrivacySettings"), TEXT("bSendUsageData"), bSendUsageData, GEditorSettingsIni);
 	}
 
+#if !UE_EDITOR
+	if (BuildSettings::IsLicenseeVersion())
+	{
+		// do not send unattended reports in licensees' builds except for the editor, where it is governed by the above setting
+		bSendUnattendedBugReports = false;
+		bSendUsageData = false;
+	}
+#endif
+
 	if (bNoDialog && !bSendUnattendedBugReports)
 	{
 		// If we shouldn't display a dialog (like for ensures) and the user
