Use UsedPhysical memory stats in LLM captures for memory allocated like on mobile devices to provide more accurate data.

diff --git a/Engine/Source/Runtime/Core/Private/HAL/LowLevelMemTracker.cpp b/Engine/Source/Runtime/Core/Private/HAL/LowLevelMemTracker.cpp
--- a/Engine/Source/Runtime/Core/Private/HAL/LowLevelMemTracker.cpp
+++ b/Engine/Source/Runtime/Core/Private/HAL/LowLevelMemTracker.cpp
@@ -637,7 +637,7 @@ void FLowLevelMemTracker::UpdateStatsPerFrame(const TCHAR* LogName)
 
 	// calculate memory the platform thinks we have allocated, compared to what we have tracked, including the program memory
 	FPlatformMemoryStats PlatformStats = FPlatformMemory::GetStats();
-#if PLATFORM_ANDROID || PLATFORM_IOS
+#if PLATFORM_ANDROID || PLATFORM_IOS || WITH_SERVER_CODE
 	uint64 PlatformProcessMemory = PlatformStats.UsedPhysical;
 #else
 	uint64 PlatformProcessMemory = PlatformStats.TotalPhysical - PlatformStats.AvailablePhysical;
@@ -2007,7 +2007,11 @@ void FLLMCsvWriter::WriteGraph(FLLMCustomTag* CustomTags, const int32* ParentTag
 		
 		const TCHAR* TrackerName = GetTrackerCsvName(Tracker);
 		const FDateTime FileDate = FDateTime::Now();
+#if WITH_SERVER_CODE
+		FString Filename = FString::Printf(TEXT("%s/%s_Pid%d_%s.csv"), *Directory, TrackerName, FPlatformProcess::GetCurrentProcessId(), *FileDate.ToString());
+#else
 		FString Filename = FString::Printf(TEXT("%s/%s_%s.csv"), *Directory, TrackerName, *FileDate.ToString());
+#endif
 		Archive = IFileManager::Get().CreateFileWriter(*Filename, FILEWRITE_AllowRead);
 		LLMCheck(Archive);
 
