Reduce the mutex scope for GetHighPerformanceEnableForStat to avoid a dead lock when accessing a TLS static varaible in FCpuProfilerTrace::OutputEventType while another module is being loaded

diff --git a/Engine/Source/Runtime/Core/Private/Stats/Stats2.cpp b/Engine/Source/Runtime/Core/Private/Stats/Stats2.cpp
--- a/Engine/Source/Runtime/Core/Private/Stats/Stats2.cpp
+++ b/Engine/Source/Runtime/Core/Private/Stats/Stats2.cpp
@@ -508,99 +508,103 @@ public:
 	{
 		LLM_SCOPE(ELLMTag::Stats);
 
-		FScopeLock ScopeLock(&SynchronizationObject);
+		TStatIdData* Result;
+		FString StatDescription;
+		{
+			FScopeLock ScopeLock(&SynchronizationObject);
 
-		FStatNameAndInfo LongName(StatShortName, InGroup, InCategory, InDescription, InStatType, bShouldClearEveryFrame, bCycleStat, bSortByName, MemoryRegion);
+			FStatNameAndInfo LongName(StatShortName, InGroup, InCategory, InDescription, InStatType, bShouldClearEveryFrame, bCycleStat, bSortByName, MemoryRegion);
 
-		FName Stat = LongName.GetEncodedName();
+			FName Stat = LongName.GetEncodedName();
 
-		FName Group(InGroup);
-		FGroupEnable* Found = HighPerformanceEnable.Find(Group);
-		if (Found)
-		{
-			if (Found->DefaultEnable != bDefaultEnable)
-			{
-				UE_LOG(LogStatGroupEnableManager, Fatal, TEXT("Stat group %s was was defined both on and off by default."), *Group.ToString());
-			}
-			TStatIdData** StatFound = Found->NamesInThisGroup.Find( Stat );
-			TStatIdData** StatFoundAlways = Found->AlwaysEnabledNamesInThisGroup.Find( Stat );
-			if( StatFound )
+			FName Group(InGroup);
+			FGroupEnable* Found = HighPerformanceEnable.Find(Group);
+			if (Found)
 			{
-				if( StatFoundAlways )
+				if (Found->DefaultEnable != bDefaultEnable)
+				{
+					UE_LOG(LogStatGroupEnableManager, Fatal, TEXT("Stat group %s was was defined both on and off by default."), *Group.ToString());
+				}
+				TStatIdData** StatFound = Found->NamesInThisGroup.Find( Stat );
+				TStatIdData** StatFoundAlways = Found->AlwaysEnabledNamesInThisGroup.Find( Stat );
+				if( StatFound )
 				{
-					UE_LOG( LogStatGroupEnableManager, Fatal, TEXT( "Stat %s is both always enabled and not always enabled, so it was used for two different things." ), *Stat.ToString() );
+					if( StatFoundAlways )
+					{
+						UE_LOG( LogStatGroupEnableManager, Fatal, TEXT( "Stat %s is both always enabled and not always enabled, so it was used for two different things." ), *Stat.ToString() );
+					}
+					return TStatId( *StatFound );
+				}
+				else if( StatFoundAlways )
+				{
+					return TStatId( *StatFoundAlways );
 				}
-				return TStatId( *StatFound );
 			}
-			else if( StatFoundAlways )
+			else
 			{
-				return TStatId( *StatFoundAlways );
-			}
-		}
-		else
-		{
-			Found = &HighPerformanceEnable.Add( Group, FGroupEnable( bDefaultEnable || !bShouldClearEveryFrame ) );
+				Found = &HighPerformanceEnable.Add( Group, FGroupEnable( bDefaultEnable || !bShouldClearEveryFrame ) );
 
-			// this was set up before we saw the group, so set the enable now
-			if (EnableForNewGroup.Contains(Group))
-			{
-				Found->CurrentEnable = EnableForNewGroup.FindChecked(Group);
-				EnableForNewGroup.Remove(Group); // by definition, we will never need this again
+				// this was set up before we saw the group, so set the enable now
+				if (EnableForNewGroup.Contains(Group))
+				{
+					Found->CurrentEnable = EnableForNewGroup.FindChecked(Group);
+					EnableForNewGroup.Remove(Group); // by definition, we will never need this again
+				}
+				else if (UseEnableForNewGroups)
+				{
+					Found->CurrentEnable = EnableForNewGroups;
+				}
 			}
-			else if (UseEnableForNewGroups)
+			if (StatIDs.Num() == 0 || StatIDs.Last().Num() == NUM_PER_BLOCK)
 			{
-				Found->CurrentEnable = EnableForNewGroups;
+				TArray<TStatIdData>& NewBlock = StatIDs.AddDefaulted_GetRef();
+				NewBlock.Reserve(NUM_PER_BLOCK);
 			}
-		}
-		if (StatIDs.Num() == 0 || StatIDs.Last().Num() == NUM_PER_BLOCK)
-		{
-			TArray<TStatIdData>& NewBlock = StatIDs.AddDefaulted_GetRef();
-			NewBlock.Reserve(NUM_PER_BLOCK);
-		}
-		TStatIdData* Result = &StatIDs.Last().AddDefaulted_GetRef();
+			Result = &StatIDs.Last().AddDefaulted_GetRef();
 
-		const FString StatDescription = InDescription ? InDescription : StatShortName.GetPlainNameString();
+			StatDescription = InDescription ? InDescription : StatShortName.GetPlainNameString();
 
-		const int32 StatDescLen = StatDescription.Len() + 1;
+			const int32 StatDescLen = StatDescription.Len() + 1;
 
-		int32 MemoryAllocated = 0;
+			int32 MemoryAllocated = 0;
 
-		// Get the wide stat description.
-		{
-			auto StatDescriptionWide = StringCast<WIDECHAR>(*StatDescription, StatDescLen);
-			int32 StatDescriptionWideLength = StatDescriptionWide.Length();
-			TUniquePtr<WIDECHAR[]> StatDescWide = MakeUnique<WIDECHAR[]>(StatDescriptionWideLength + 1);
-			TCString<WIDECHAR>::Strcpy(StatDescWide.Get(), StatDescriptionWideLength + 1, StatDescriptionWide.Get());
-			Result->StatDescriptionWide = MoveTemp(StatDescWide);
+			// Get the wide stat description.
+			{
+				auto StatDescriptionWide = StringCast<WIDECHAR>(*StatDescription, StatDescLen);
+				int32 StatDescriptionWideLength = StatDescriptionWide.Length();
+				TUniquePtr<WIDECHAR[]> StatDescWide = MakeUnique<WIDECHAR[]>(StatDescriptionWideLength + 1);
+				TCString<WIDECHAR>::Strcpy(StatDescWide.Get(), StatDescriptionWideLength + 1, StatDescriptionWide.Get());
+				Result->StatDescriptionWide = MoveTemp(StatDescWide);
 
-			MemoryAllocated += StatDescriptionWideLength * sizeof(WIDECHAR);
-		}
+				MemoryAllocated += StatDescriptionWideLength * sizeof(WIDECHAR);
+			}
 
-		// Get the ansi stat description.
-		{
-			auto StatDescriptionAnsi = StringCast<ANSICHAR>(*StatDescription, StatDescLen);
-			int32 StatDescriptionAnsiLength = StatDescriptionAnsi.Length();
-			TUniquePtr<ANSICHAR[]> StatDescAnsi = MakeUnique<ANSICHAR[]>(StatDescriptionAnsiLength + 1);
-			TCString<ANSICHAR>::Strcpy(StatDescAnsi.Get(), StatDescriptionAnsiLength + 1, StatDescriptionAnsi.Get());
-			Result->StatDescriptionAnsi = MoveTemp(StatDescAnsi);
+			// Get the ansi stat description.
+			{
+				auto StatDescriptionAnsi = StringCast<ANSICHAR>(*StatDescription, StatDescLen);
+				int32 StatDescriptionAnsiLength = StatDescriptionAnsi.Length();
+				TUniquePtr<ANSICHAR[]> StatDescAnsi = MakeUnique<ANSICHAR[]>(StatDescriptionAnsiLength + 1);
+				TCString<ANSICHAR>::Strcpy(StatDescAnsi.Get(), StatDescriptionAnsiLength + 1, StatDescriptionAnsi.Get());
+				Result->StatDescriptionAnsi = MoveTemp(StatDescAnsi);
 
-			MemoryAllocated += StatDescriptionAnsiLength * sizeof(ANSICHAR);
-		}
+				MemoryAllocated += StatDescriptionAnsiLength * sizeof(ANSICHAR);
+			}
 
-		MemoryCounter.Add( MemoryAllocated );
+			MemoryCounter.Add( MemoryAllocated );
 
-		if( Found->CurrentEnable )
-		{
-			EnableStat( Stat, Result );
-		}
+			if( Found->CurrentEnable )
+			{
+				EnableStat( Stat, Result );
+			}
 
-		if( bShouldClearEveryFrame )
-		{
-			Found->NamesInThisGroup.Add( Stat, Result );
-		}
-		else
-		{
-			Found->AlwaysEnabledNamesInThisGroup.Add( Stat, Result );
+			if( bShouldClearEveryFrame )
+			{
+				Found->NamesInThisGroup.Add( Stat, Result );
+			}
+			else
+			{
+				Found->AlwaysEnabledNamesInThisGroup.Add( Stat, Result );
+			}
 		}
 
 #if CPUPROFILERTRACE_ENABLED
