Added support for RHI name filter on driver blacklist, e.g. D3D11 or D3D12.

diff --git a/Engine/Source/Runtime/Core/Public/GenericPlatform/GenericPlatformDriver.h b/Engine/Source/Runtime/Core/Public/GenericPlatform/GenericPlatformDriver.h
--- a/Engine/Source/Runtime/Core/Public/GenericPlatform/GenericPlatformDriver.h
+++ b/Engine/Source/Runtime/Core/Public/GenericPlatform/GenericPlatformDriver.h
@@ -265,6 +265,8 @@ struct FGPUDriverInfo
 	FString UserDriverVersion;
 	// e.g. 3-13-2015
 	FString DriverDate;
+	// e.g. D3D11, D3D12
+	FString RHIName;
 
 	bool IsValid() const
 	{
@@ -328,6 +330,8 @@ struct FBlackListEntry
 	FString DriverVersionString;
 	// optional, e.g. "<=MM-DD-YYYY"
 	FString DriverDateString;
+	// optional, e.g. "D3D11", "D3D12"
+	FString RHIName;
 	// required
 	FString Reason;
 
@@ -339,6 +343,8 @@ struct FBlackListEntry
 		FParse::Value(In, TEXT("DriverDate="), DriverDateString);
 		ensure(!DriverVersionString.IsEmpty() || !DriverDateString.IsEmpty());
 
+		FParse::Value(In, TEXT("RHI="), RHIName);
+		
 		// later:
 //		FParse::Value(In, TEXT("DeviceId="), DeviceId);
 //		FParse::Value(In, TEXT("OS="), OS);
@@ -355,6 +361,12 @@ struct FBlackListEntry
 	{
 		if (IsValid())
 		{
+			// If RHI specified, ignore if mismatched
+			if (!RHIName.IsEmpty() && RHIName != Info.RHIName)
+			{
+				return false;
+			}
+
 			if (!DriverVersionString.IsEmpty())
 			{
 				return CompareStringOp(*DriverVersionString, *Info.GetUnifiedDriverVersion());
diff --git a/Engine/Source/Runtime/RHI/Private/DynamicRHI.cpp b/Engine/Source/Runtime/RHI/Private/DynamicRHI.cpp
--- a/Engine/Source/Runtime/RHI/Private/DynamicRHI.cpp
+++ b/Engine/Source/Runtime/RHI/Private/DynamicRHI.cpp
@@ -78,7 +78,7 @@ static void RHIDetectAndWarnOfBadDrivers(bool bHasEditorToken)
 	DriverInfo.InternalDriverVersion = GRHIAdapterInternalDriverVersion;
 	DriverInfo.UserDriverVersion = GRHIAdapterUserDriverVersion;
 	DriverInfo.DriverDate = GRHIAdapterDriverDate;
-
+	DriverInfo.RHIName = GDynamicRHI ? GDynamicRHI->GetName() : FString();
 
 #if !(UE_BUILD_SHIPPING || UE_BUILD_TEST)
 	// for testing
