#ue4 - Fix cases where FMod() and VectorMod() would fail to return a result within the expected output range for very large input values. This could in turn result in NaN/Inf propagation if for example those values are used with SinCos during Rotator/Vector/Quat conversion.

Added some additional ensures during development to catch similar overflow and uninitialized variable issues sooner.

diff --git a/Engine/Source/Runtime/Core/Private/GenericPlatform/GenericPlatformMath.cpp b/Engine/Source/Runtime/Core/Private/GenericPlatform/GenericPlatformMath.cpp
--- a/Engine/Source/Runtime/Core/Private/GenericPlatform/GenericPlatformMath.cpp
+++ b/Engine/Source/Runtime/Core/Private/GenericPlatform/GenericPlatformMath.cpp
@@ -77,14 +77,15 @@ float FGenericPlatformMath::Atan2(float Y, float X)
 
 /*FORCENOINLINE*/ float FGenericPlatformMath::Fmod(float X, float Y)
 {
-	if (fabsf(Y) <= 1.e-8f)
+	const float AbsY = fabsf(Y);
+	if (AbsY <= 1.e-8f)
 	{
 		FmodReportError(X, Y);
 		return 0.f;
 	}
 	const float Div = (X / Y);
 	// All floats where abs(f) >= 2^23 (8388608) are whole numbers so do not need truncation, and avoid overflow in TruncToFloat as they get even larger.
-	const float Quotient = fabsf(Div) < 8388608.f ? TruncToFloat(Div) : Div;
+	const float Quotient = fabsf(Div) < FLOAT_NON_FRACTIONAL ? TruncToFloat(Div) : Div;
 	float IntPortion = Y * Quotient;
 
 	// Rounding and imprecision could cause IntPortion to exceed X and cause the result to be outside the expected range.
@@ -95,7 +96,9 @@ float FGenericPlatformMath::Atan2(float Y, float X)
 	}
 
 	const float Result = X - IntPortion;
-	return Result;
+	// Clamp to [-AbsY, AbsY] because of possible failures for very large numbers (>1e10) due to precision loss.
+	// We could instead fall back to stock fmodf() for large values, however this would diverge from the SIMD VectorMod() which has no similar fallback with reasonable performance.
+	return FMath::Clamp(Result, -AbsY, AbsY);
 }
 
 void FGenericPlatformMath::FmodReportError(float X, float Y)
diff --git a/Engine/Source/Runtime/Core/Private/Math/UnrealMath.cpp b/Engine/Source/Runtime/Core/Private/Math/UnrealMath.cpp
--- a/Engine/Source/Runtime/Core/Private/Math/UnrealMath.cpp
+++ b/Engine/Source/Runtime/Core/Private/Math/UnrealMath.cpp
@@ -203,7 +203,7 @@ FRotator FVector::ToOrientationRotator() const
 	// Find roll.
 	R.Roll = 0;
 
-#if ENABLE_NAN_DIAGNOSTIC
+#if ENABLE_NAN_DIAGNOSTIC || (DO_CHECK && !UE_BUILD_SHIPPING)
 	if (R.ContainsNaN())
 	{
 		logOrEnsureNanError(TEXT("FVector::Rotation(): Rotator result %s contains NaN! Input FVector = %s"), *R.ToString(), *this->ToString());
@@ -232,7 +232,7 @@ FRotator FVector4::ToOrientationRotator() const
 	// Find roll.
 	R.Roll = 0;
 
-#if ENABLE_NAN_DIAGNOSTIC
+#if ENABLE_NAN_DIAGNOSTIC || (DO_CHECK && !UE_BUILD_SHIPPING)
 	if (R.ContainsNaN())
 	{
 		logOrEnsureNanError(TEXT("FVector4::Rotation(): Rotator result %s contains NaN! Input FVector4 = %s"), *R.ToString(), *this->ToString());
@@ -379,11 +379,34 @@ FRotator::FRotator(const FQuat& Quat)
 
 CORE_API FVector FRotator::Vector() const
 {
+	// Extremely large but valid values (or invalid values from uninitialized vars) can cause SinCos to return NaN/Inf, so catch that here. Similar to what is done in FRotator::Quaternion().
+#if ENABLE_NAN_DIAGNOSTIC || (DO_CHECK && !UE_BUILD_SHIPPING)
+	if (FMath::Abs(Pitch) > FLOAT_NON_FRACTIONAL ||
+		FMath::Abs(Yaw  ) > FLOAT_NON_FRACTIONAL ||
+		FMath::Abs(Roll ) > FLOAT_NON_FRACTIONAL)
+	{
+		logOrEnsureNanError(TEXT("FRotator::Vector() provided with unreasonably large input values (%s), possible use of uninitialized variable?"), *ToString());
+	}
+#endif
+	
+	// Remove winding and clamp to [-360, 360]
+	const float PitchNoWinding = FMath::Fmod(Pitch, 360.0f);
+	const float YawNoWinding = FMath::Fmod(Yaw, 360.0f);
+
 	float CP, SP, CY, SY;
-	FMath::SinCos( &SP, &CP, FMath::DegreesToRadians(Pitch) );
-	FMath::SinCos( &SY, &CY, FMath::DegreesToRadians(Yaw) );
+	FMath::SinCos( &SP, &CP, FMath::DegreesToRadians(PitchNoWinding) );
+	FMath::SinCos( &SY, &CY, FMath::DegreesToRadians(YawNoWinding) );
 	FVector V = FVector( CP*CY, CP*SY, SP );
 
+	// Error checking
+#if ENABLE_NAN_DIAGNOSTIC || (DO_CHECK && !UE_BUILD_SHIPPING)
+	if (V.ContainsNaN())
+	{
+		logOrEnsureNanError(TEXT("FRotator::Vector() resulted in NaN/Inf with input: %s output: %s"), *ToString(), *V.ToString());
+		V = FVector::ForwardVector;
+	}
+#endif
+
 	return V;
 }
 
diff --git a/Engine/Source/Runtime/Core/Public/Math/UnrealMathDirectX.h b/Engine/Source/Runtime/Core/Public/Math/UnrealMathDirectX.h
--- a/Engine/Source/Runtime/Core/Public/Math/UnrealMathDirectX.h
+++ b/Engine/Source/Runtime/Core/Public/Math/UnrealMathDirectX.h
@@ -1061,7 +1061,10 @@ FORCEINLINE VectorRegister VectorFractional(const VectorRegister& X)
 
 FORCEINLINE VectorRegister VectorMod(const VectorRegister& X, const VectorRegister& Y)
 {
-	return DirectX::XMVectorMod(X, Y);
+	VectorRegister AbsY = VectorAbs(Y);
+	VectorRegister Result = DirectX::XMVectorMod(X, Y);
+	// Clamp to [-AbsY, AbsY] because of possible failures for very large numbers (>1e10) due to precision loss.
+	return DirectX::XMVectorClamp(Result, VectorNegate(AbsY), AbsY);
 }
 
 //TODO: Vectorize
diff --git a/Engine/Source/Runtime/Core/Public/Math/UnrealMathSSE.h b/Engine/Source/Runtime/Core/Public/Math/UnrealMathSSE.h
--- a/Engine/Source/Runtime/Core/Public/Math/UnrealMathSSE.h
+++ b/Engine/Source/Runtime/Core/Public/Math/UnrealMathSSE.h
@@ -1206,7 +1206,10 @@ FORCEINLINE VectorRegister VectorMod(const VectorRegister& X, const VectorRegist
 	// Floats where abs(f) >= 2^23 have no fractional portion, and larger values would overflow VectorTruncate.
 	VectorRegister NoFractionMask = VectorCompareGE(VectorAbs(Div), GlobalVectorConstants::FloatNonFractional);
 	VectorRegister Temp = VectorSelect(NoFractionMask, Div, VectorTruncate(Div));
-	return VectorSubtract(X, VectorMultiply(Y, Temp));
+	VectorRegister Result = VectorSubtract(X, VectorMultiply(Y, Temp));
+	// Clamp to [-AbsY, AbsY] because of possible failures for very large numbers (>1e10) due to precision loss.
+	VectorRegister AbsY = VectorAbs(Y);
+	return VectorMax(VectorNegate(AbsY), VectorMin(Result, AbsY));
 }
 
 FORCEINLINE VectorRegister VectorSign(const VectorRegister& X)
diff --git a/Engine/Source/Runtime/Core/Public/Math/UnrealMathUtility.h b/Engine/Source/Runtime/Core/Public/Math/UnrealMathUtility.h
--- a/Engine/Source/Runtime/Core/Public/Math/UnrealMathUtility.h
+++ b/Engine/Source/Runtime/Core/Public/Math/UnrealMathUtility.h
@@ -47,6 +47,7 @@ struct FLinearColor;
 #define BIG_NUMBER			(3.4e+38f)
 #define EULERS_NUMBER       (2.71828182845904523536f)
 #define UE_GOLDEN_RATIO		(1.6180339887498948482045868343656381f)	/* Also known as divine proportion, golden mean, or golden section - related to the Fibonacci Sequence = (1 + sqrt(5)) / 2 */
+#define FLOAT_NON_FRACTIONAL (8388608.f) /* All single-precision floating point numbers greater than or equal to this have no fractional value. */
 
 // Copied from float.h
 #define MAX_FLT 3.402823466e+38F
diff --git a/Engine/Source/Runtime/Core/Public/Math/UnrealMathVectorConstants.h b/Engine/Source/Runtime/Core/Public/Math/UnrealMathVectorConstants.h
--- a/Engine/Source/Runtime/Core/Public/Math/UnrealMathVectorConstants.h
+++ b/Engine/Source/Runtime/Core/Public/Math/UnrealMathVectorConstants.h
@@ -66,7 +66,7 @@ namespace GlobalVectorConstants
 	static const VectorRegister Float180 = MakeVectorRegister(180.f, 180.f, 180.f, 180.f);
 
 	// All float numbers greater than or equal to this have no fractional value.
-	static const VectorRegister FloatNonFractional = MakeVectorRegister(8388608.f, 8388608.f, 8388608.f, 8388608.f);
+	static const VectorRegister FloatNonFractional = MakeVectorRegister(FLOAT_NON_FRACTIONAL, FLOAT_NON_FRACTIONAL, FLOAT_NON_FRACTIONAL, FLOAT_NON_FRACTIONAL);
 
 	static const VectorRegister FloatTwo = MakeVectorRegister(2.0f, 2.0f, 2.0f, 2.0f);
 	static const uint32 AlmostTwoBits = 0x3fffffff;
