Update crash context file when changing crash reporting user settings

diff --git a/Engine/Source/Runtime/Core/Private/GenericPlatform/GenericPlatformCrashContext.cpp b/Engine/Source/Runtime/Core/Private/GenericPlatform/GenericPlatformCrashContext.cpp
--- a/Engine/Source/Runtime/Core/Private/GenericPlatform/GenericPlatformCrashContext.cpp
+++ b/Engine/Source/Runtime/Core/Private/GenericPlatform/GenericPlatformCrashContext.cpp
@@ -212,6 +212,7 @@ PRAGMA_DISABLE_DEPRECATION_WARNINGS
 		{
 			NCached::UserSettings.bSendUsageData = InParams.SendUsageData.GetValue();
 		}
+		SerializeTempCrashContextToFile();
 	});
 PRAGMA_ENABLE_DEPRECATION_WARNINGS
 
