Fix crash when selecting locales with numeric symbols outside the BMP

diff --git a/Engine/Source/Runtime/Core/Private/Internationalization/ICUCulture.cpp b/Engine/Source/Runtime/Core/Private/Internationalization/ICUCulture.cpp
--- a/Engine/Source/Runtime/Core/Private/Internationalization/ICUCulture.cpp
+++ b/Engine/Source/Runtime/Core/Private/Internationalization/ICUCulture.cpp
@@ -838,16 +838,16 @@ FDecimalNumberFormattingRules ExtractNumberFormattingRulesFromICUDecimalFormatte
 	// This allows us to use the correct groupings if we should ever force grouping for a number, rather than use the culture default
 	InICUDecimalFormat.setGroupingUsed(true);
 
-	auto ICUStringToTCHAR = [](const icu::UnicodeString& InICUString) -> TCHAR
-	{
-		check(InICUString.length() == 1);
-		return static_cast<TCHAR>(InICUString.charAt(0));
-	};
-
-	auto ExtractFormattingSymbolAsCharacter = [&](icu::DecimalFormatSymbols::ENumberFormatSymbol InSymbolToExtract) -> TCHAR
+	auto ExtractFormattingSymbolAsCharacter = [&](icu::DecimalFormatSymbols::ENumberFormatSymbol InSymbolToExtract, const TCHAR InFallbackChar) -> TCHAR
 	{
 		const icu::UnicodeString& ICUSymbolString = InICUDecimalFormat.getDecimalFormatSymbols()->getConstSymbol(InSymbolToExtract);
-		return ICUStringToTCHAR(ICUSymbolString); // For efficiency we assume that these symbols are always a single character
+		check(ICUSymbolString.length() <= 2);
+
+		// Some cultures use characters outside of the BMP which present as a surrogate pair on platforms using UTF-16 TCHAR
+		// We need to update this code to use FString or UTF32CHAR for these symbols (see UE-83143), but for now we just use the fallback if we find a surrogate pair
+		return ICUSymbolString.length() == 1
+			? static_cast<TCHAR>(ICUSymbolString.charAt(0))
+			: InFallbackChar;
 	};
 
 	icu::UnicodeString ScratchICUString;
@@ -860,23 +860,23 @@ FDecimalNumberFormattingRules ExtractNumberFormattingRulesFromICUDecimalFormatte
 	NewUEDecimalNumberFormattingRules.PositiveSuffixString			= ICUUtilities::ConvertString(InICUDecimalFormat.getPositiveSuffix(ScratchICUString));
 	NewUEDecimalNumberFormattingRules.PlusString					= ICUUtilities::ConvertString(InICUDecimalFormat.getDecimalFormatSymbols()->getConstSymbol(icu::DecimalFormatSymbols::kPlusSignSymbol));
 	NewUEDecimalNumberFormattingRules.MinusString					= ICUUtilities::ConvertString(InICUDecimalFormat.getDecimalFormatSymbols()->getConstSymbol(icu::DecimalFormatSymbols::kMinusSignSymbol));
-	NewUEDecimalNumberFormattingRules.GroupingSeparatorCharacter	= ExtractFormattingSymbolAsCharacter(icu::DecimalFormatSymbols::kGroupingSeparatorSymbol);
-	NewUEDecimalNumberFormattingRules.DecimalSeparatorCharacter		= ExtractFormattingSymbolAsCharacter(icu::DecimalFormatSymbols::kDecimalSeparatorSymbol);
+	NewUEDecimalNumberFormattingRules.GroupingSeparatorCharacter	= ExtractFormattingSymbolAsCharacter(icu::DecimalFormatSymbols::kGroupingSeparatorSymbol, TEXT(','));
+	NewUEDecimalNumberFormattingRules.DecimalSeparatorCharacter		= ExtractFormattingSymbolAsCharacter(icu::DecimalFormatSymbols::kDecimalSeparatorSymbol,  TEXT('.'));
 	NewUEDecimalNumberFormattingRules.PrimaryGroupingSize			= static_cast<uint8>(InICUDecimalFormat.getGroupingSize());
 	NewUEDecimalNumberFormattingRules.SecondaryGroupingSize			= (InICUDecimalFormat.getSecondaryGroupingSize() < 1) 
 																		? NewUEDecimalNumberFormattingRules.PrimaryGroupingSize 
 																		: static_cast<uint8>(InICUDecimalFormat.getSecondaryGroupingSize());
 
-	NewUEDecimalNumberFormattingRules.DigitCharacters[0]			= ExtractFormattingSymbolAsCharacter(icu::DecimalFormatSymbols::kZeroDigitSymbol);
-	NewUEDecimalNumberFormattingRules.DigitCharacters[1]			= ExtractFormattingSymbolAsCharacter(icu::DecimalFormatSymbols::kOneDigitSymbol);
-	NewUEDecimalNumberFormattingRules.DigitCharacters[2]			= ExtractFormattingSymbolAsCharacter(icu::DecimalFormatSymbols::kTwoDigitSymbol);
-	NewUEDecimalNumberFormattingRules.DigitCharacters[3]			= ExtractFormattingSymbolAsCharacter(icu::DecimalFormatSymbols::kThreeDigitSymbol);
-	NewUEDecimalNumberFormattingRules.DigitCharacters[4]			= ExtractFormattingSymbolAsCharacter(icu::DecimalFormatSymbols::kFourDigitSymbol);
-	NewUEDecimalNumberFormattingRules.DigitCharacters[5]			= ExtractFormattingSymbolAsCharacter(icu::DecimalFormatSymbols::kFiveDigitSymbol);
-	NewUEDecimalNumberFormattingRules.DigitCharacters[6]			= ExtractFormattingSymbolAsCharacter(icu::DecimalFormatSymbols::kSixDigitSymbol);
-	NewUEDecimalNumberFormattingRules.DigitCharacters[7]			= ExtractFormattingSymbolAsCharacter(icu::DecimalFormatSymbols::kSevenDigitSymbol);
-	NewUEDecimalNumberFormattingRules.DigitCharacters[8]			= ExtractFormattingSymbolAsCharacter(icu::DecimalFormatSymbols::kEightDigitSymbol);
-	NewUEDecimalNumberFormattingRules.DigitCharacters[9]			= ExtractFormattingSymbolAsCharacter(icu::DecimalFormatSymbols::kNineDigitSymbol);
+	NewUEDecimalNumberFormattingRules.DigitCharacters[0]			= ExtractFormattingSymbolAsCharacter(icu::DecimalFormatSymbols::kZeroDigitSymbol,	TEXT('0'));
+	NewUEDecimalNumberFormattingRules.DigitCharacters[1]			= ExtractFormattingSymbolAsCharacter(icu::DecimalFormatSymbols::kOneDigitSymbol,	TEXT('1'));
+	NewUEDecimalNumberFormattingRules.DigitCharacters[2]			= ExtractFormattingSymbolAsCharacter(icu::DecimalFormatSymbols::kTwoDigitSymbol,	TEXT('2'));
+	NewUEDecimalNumberFormattingRules.DigitCharacters[3]			= ExtractFormattingSymbolAsCharacter(icu::DecimalFormatSymbols::kThreeDigitSymbol,	TEXT('3'));
+	NewUEDecimalNumberFormattingRules.DigitCharacters[4]			= ExtractFormattingSymbolAsCharacter(icu::DecimalFormatSymbols::kFourDigitSymbol,	TEXT('4'));
+	NewUEDecimalNumberFormattingRules.DigitCharacters[5]			= ExtractFormattingSymbolAsCharacter(icu::DecimalFormatSymbols::kFiveDigitSymbol,	TEXT('5'));
+	NewUEDecimalNumberFormattingRules.DigitCharacters[6]			= ExtractFormattingSymbolAsCharacter(icu::DecimalFormatSymbols::kSixDigitSymbol,	TEXT('6'));
+	NewUEDecimalNumberFormattingRules.DigitCharacters[7]			= ExtractFormattingSymbolAsCharacter(icu::DecimalFormatSymbols::kSevenDigitSymbol,	TEXT('7'));
+	NewUEDecimalNumberFormattingRules.DigitCharacters[8]			= ExtractFormattingSymbolAsCharacter(icu::DecimalFormatSymbols::kEightDigitSymbol,	TEXT('8'));
+	NewUEDecimalNumberFormattingRules.DigitCharacters[9]			= ExtractFormattingSymbolAsCharacter(icu::DecimalFormatSymbols::kNineDigitSymbol,	TEXT('9'));
 
 	return NewUEDecimalNumberFormattingRules;
 }
