Fix string wildcard matching. For example, matching "abcdef" against "abc*def" would return false.

diff --git a/Engine/Source/Runtime/Core/Private/Containers/String.cpp b/Engine/Source/Runtime/Core/Private/Containers/String.cpp
--- a/Engine/Source/Runtime/Core/Private/Containers/String.cpp
+++ b/Engine/Source/Runtime/Core/Private/Containers/String.cpp
@@ -81,12 +81,12 @@ namespace UE4String_Private
 			}
 
 			--TargetLength;
+			--WildcardLength;
+
 			if (TargetLength == 0)
 			{
-				return false;
+				break;
 			}
-
-			--WildcardLength;
 		}
 
 		// Match * against anything and ? against single (and zero?) chars
diff --git a/Engine/Source/Runtime/Core/Private/HAL/FileManagerGeneric.cpp b/Engine/Source/Runtime/Core/Private/HAL/FileManagerGeneric.cpp
--- a/Engine/Source/Runtime/Core/Private/HAL/FileManagerGeneric.cpp
+++ b/Engine/Source/Runtime/Core/Private/HAL/FileManagerGeneric.cpp
@@ -406,10 +406,13 @@ void FFileManagerGeneric::FindFiles( TArray<FString>& Result, const TCHAR* InFil
 		}
 		virtual bool Visit( const TCHAR* FilenameOrDirectory, bool bIsDirectory )
 		{
-			if (((bIsDirectory && bDirectories) || (!bIsDirectory && bFiles))
-				&& FPaths::GetCleanFilename(FilenameOrDirectory).MatchesWildcard(WildCard))
+			if ((bIsDirectory && bDirectories) || (!bIsDirectory && bFiles))
 			{
-				new( Result ) FString( FPaths::GetCleanFilename(FilenameOrDirectory) );
+				const FString Filename = FPaths::GetCleanFilename(FilenameOrDirectory);
+				if (Filename.MatchesWildcard(WildCard))
+				{
+					new( Result ) FString(Filename);
+				}
 			}
 			return true;
 		}
