Fixed some stats not displaying with the budget command until toggled off and back on again

diff --git a/Engine/Source/Runtime/Core/Private/Stats/StatsCommand.cpp b/Engine/Source/Runtime/Core/Private/Stats/StatsCommand.cpp
--- a/Engine/Source/Runtime/Core/Private/Stats/StatsCommand.cpp
+++ b/Engine/Source/Runtime/Core/Private/Stats/StatsCommand.cpp
@@ -175,7 +175,7 @@ struct FBudgetData
 };
 
 FCriticalSection BudgetStatMapCS;
-TMap<FString, FBudgetData> BudgetStatMapping;
+TMap<FName, FBudgetData> BudgetStatMapping;
 
 /** Holds parameters used by the 'stat hier' or 'stat group ##' command. */
 struct FStatParams
@@ -209,7 +209,7 @@ struct FStatParams
 	/** -root=[name]. */
 	TParsedValueWithDefault<FName> Root;
 
-	FString BudgetSection;
+	FName BudgetSection;
 
 	/**
 	 *	Maximum number of frames to be included in the history. 
@@ -1005,7 +1005,7 @@ struct FHUDGroupManager
 					EnabledGroups.Add( MaybeGroupFName, FInternalGroup( MaybeGroupFName, NAME_None, EStatDisplayMode::Hierarchical, false, EmptySet, TEXT( "Hierarchy for game and render" ) ) );
 				}			
 			}
-			else if(!Params.BudgetSection.IsEmpty())
+			else if(Params.BudgetSection != NAME_None)
 			{
 				const bool bEnabledBudgetMode = EnabledGroups.Num() > 0;
 				if (bEnabledBudgetMode)
@@ -1038,10 +1038,9 @@ struct FHUDGroupManager
 					}
 					
 					{
-						TSet<FName> StatSet;
-						GetStatsForNames(StatSet, StatShortNames);
-						FName BudgetGroupName(*Params.BudgetSection);
-						EnabledGroups.Add(BudgetGroupName, FInternalGroup(*Params.BudgetSection, NAME_None, EStatDisplayMode::Flat, false, StatSet, TEXT("Budget"), &ThreadBudgetMap, &NonAccumulatingStats));
+						TSet<FName> EmptyStatSet; // see LazyInitBudgetGroup
+						FName BudgetGroupName(Params.BudgetSection);
+						EnabledGroups.Add(BudgetGroupName, FInternalGroup(Params.BudgetSection, NAME_None, EStatDisplayMode::Flat, false, EmptyStatSet, TEXT("Budget"), &ThreadBudgetMap, &NonAccumulatingStats));
 						HandleToggleCommandBroadcast( BudgetGroupName, bCurrentEnabled, bOthersEnabled );
 					}
 				}
@@ -1072,6 +1071,30 @@ struct FHUDGroupManager
 		}
 	}
 
+	/** Convert short name stats to long name. Delayed so the ShortNameToLongName map gets populated with more stats. */
+	void LazyInitBudgetGroup()
+	{
+		FName BudgetGroupName(Params.BudgetSection);
+		if (FInternalGroup* InternalGroup = EnabledGroups.Find(BudgetGroupName))
+		{
+			if (InternalGroup->EnabledItems.Num() > 0)
+			{				
+				return; // skip if already initialized
+			}
+
+			TArray<FName> StatShortNames;
+			FScopeLock BudgetLock(&BudgetStatMapCS);
+			if (FBudgetData* BudgetData = BudgetStatMapping.Find(Params.BudgetSection))
+			{
+				for (const FString& StatEntry : BudgetData->Stats)
+				{
+					StatShortNames.Add(FName(*StatEntry));
+				}
+			}
+			GetStatsForNames(InternalGroup->EnabledItems, StatShortNames);
+		}
+	}
+
 	void ResizeFramesHistory( int32 MaxFrames )
 	{
 		History.Empty(MaxFrames + 1);
@@ -1122,7 +1145,7 @@ struct FHUDGroupManager
 		FString RootString = RootName == NAME_None ? FString() : RootName.ToString();
 
 		const bool bUseSlowMode = Params.bSlowMode;
-		const bool bUseBudgetMode = !Params.BudgetSection.IsEmpty();
+		const bool bUseBudgetMode = Params.BudgetSection != NAME_None;
 
 		if (bUseSlowMode)
 		{
@@ -1148,6 +1171,11 @@ struct FHUDGroupManager
 		}
 		else
 		{
+			if (bUseBudgetMode)
+			{
+				LazyInitBudgetGroup();
+			}
+
 			TSet<FName> HierEnabledItems;
 			for( auto It = EnabledGroups.CreateIterator(); It; ++It )
 			{
@@ -2125,7 +2153,7 @@ static void StatCmd(FString InCmd, bool bStatCommand, FOutputDevice* Ar /*= null
 #if STATS
 			// Try to parse.
 			FStatParams Params(Cmd);
-			Params.BudgetSection = MaybeBudget;
+			Params.BudgetSection = *MaybeBudget;
 			Params.Group.Set(FName("Budget"));
 			FStatsThreadState& Stats = FStatsThreadState::GetLocalState();
 			FHUDGroupManager::Get(Stats).HandleCommand(Params, false);
@@ -2279,7 +2307,7 @@ bool DirectStatsCommand(const TCHAR* Cmd, bool bBlockForCompletion /*= false*/,
 			{
 				FScopeLock BudgetINILock(&BudgetStatMapCS);   //Make sure stats thread isn't currently reading from this data
 				{
-					FBudgetData& BudgetData = BudgetStatMapping.FindOrAdd(BudgetSection);
+					FBudgetData& BudgetData = BudgetStatMapping.FindOrAdd(*BudgetSection);
 					BudgetData = FBudgetData();
 					GConfig->GetArray(*BudgetSection, TEXT("Stats"), BudgetData.Stats, GEngineIni);
 					
