Expand the parents config path if its outside of the PE layer

diff --git a/Engine/Source/Programs/UnrealBuildTool/System/ConfigHierarchy.cs b/Engine/Source/Programs/UnrealBuildTool/System/ConfigHierarchy.cs
--- a/Engine/Source/Programs/UnrealBuildTool/System/ConfigHierarchy.cs
+++ b/Engine/Source/Programs/UnrealBuildTool/System/ConfigHierarchy.cs
@@ -817,65 +817,86 @@ namespace UnrealBuildTool
 			return PersonalConfigFolder;
 		}
 
-		/// <summary>
-		/// Returns a list of INI filenames for the given project
-		/// </summary>
-		public static IEnumerable<FileReference> EnumerateConfigFileLocations(ConfigHierarchyType Type, DirectoryReference ProjectDir, UnrealTargetPlatform Platform)
+		private static string GetLayerPath(ConfigLayer Layer, string PlatformExtensionName, string IniPlatformName, DirectoryReference ProjectDir, string BaseIniName,
+			out bool bHasPlatformTag, out bool bHasProjectTag, out bool bHasExpansionTag)
 		{
-			string BaseIniName = Enum.GetName(typeof(ConfigHierarchyType), Type);
-			string PlatformName = GetIniPlatformName(Platform);
-
 			// cache some platform extension information that can be used inside the loops
-			string PlatformExtensionEngineConfigDir = DirectoryReference.Combine(UnrealBuildTool.EnginePlatformExtensionsDirectory, Platform.ToString()).FullName;
-			string PlatformExtensionProjectConfigDir = ProjectDir != null ? DirectoryReference.Combine(UnrealBuildTool.ProjectPlatformExtensionsDirectory(ProjectDir), Platform.ToString()).FullName : null;
+			string PlatformExtensionEngineConfigDir = DirectoryReference.Combine(UnrealBuildTool.EnginePlatformExtensionsDirectory, PlatformExtensionName).FullName;
+			string PlatformExtensionProjectConfigDir = ProjectDir != null ? DirectoryReference.Combine(UnrealBuildTool.ProjectPlatformExtensionsDirectory(ProjectDir), PlatformExtensionName).FullName : null;
 			bool bHasPlatformExtensionEngineConfigDir = Directory.Exists(PlatformExtensionEngineConfigDir);
 			bool bHasPlatformExtensionProjectConfigDir = PlatformExtensionProjectConfigDir != null && Directory.Exists(PlatformExtensionProjectConfigDir);
 
+			bHasPlatformTag = Layer.Path.Contains("{PLATFORM}");
+			bHasProjectTag = Layer.Path.Contains("{PROJECT}");
+			bHasExpansionTag = Layer.Path.Contains("{ED}") || Layer.Path.Contains("{EF}");
+			bool bHasUserTag = Layer.Path.Contains("{USER}");
 
-			foreach (ConfigLayer Layer in ConfigLayers)
+			// skip platform layers if we are "platform-less", or user layers without a user dir
+			if ((bHasPlatformTag && IniPlatformName == "None") ||
+				(bHasProjectTag && ProjectDir == null) ||
+				(bHasUserTag && GetUserDir() == null))
 			{
-				bool bHasPlatformTag = Layer.Path.Contains("{PLATFORM}");
-				bool bHasProjectTag = Layer.Path.Contains("{PROJECT}");
-				bool bHasExpansionTag = Layer.Path.Contains("{ED}") || Layer.Path.Contains("{EF}");
-				bool bHasUserTag = Layer.Path.Contains("{USER}");
+				return null;
+			}
 
-				// skip platform layers if we are "platform-less", or user layers without a user dir
-				if (bHasPlatformTag && Platform == null ||
-					bHasProjectTag && ProjectDir == null ||
-					bHasUserTag && GetUserDir() == null)
+			// basic replacements
+			string LayerPath;
+			// you can only have PROJECT or ENGINE, not both
+			if (bHasProjectTag)
+			{
+				if (bHasPlatformTag && bHasPlatformExtensionProjectConfigDir)
 				{
-					continue;
+					LayerPath = Layer.ExtProjectPath.Replace("{EXTPROJECT}", PlatformExtensionProjectConfigDir);
 				}
-
-				// basic replacements
-				string LayerPath;
-				// you can only have PROJECT or ENGINE, not both
-				if (bHasProjectTag)
+				else
 				{
-					if (bHasPlatformTag && bHasPlatformExtensionProjectConfigDir)
-					{
-						LayerPath = Layer.ExtProjectPath.Replace("{EXTPROJECT}", PlatformExtensionProjectConfigDir);
-					}
-					else
-					{
-						LayerPath = Layer.Path.Replace("{PROJECT}", ProjectDir.FullName);
-					}
+					LayerPath = Layer.Path.Replace("{PROJECT}", ProjectDir.FullName);
+				}
+			}
+			else
+			{
+				if (bHasPlatformTag && bHasPlatformExtensionEngineConfigDir)
+				{
+					LayerPath = Layer.ExtEnginePath.Replace("{EXTENGINE}", PlatformExtensionEngineConfigDir);
 				}
 				else
 				{
-					if (bHasPlatformTag && bHasPlatformExtensionEngineConfigDir)
-					{
-						LayerPath = Layer.ExtEnginePath.Replace("{EXTENGINE}", PlatformExtensionEngineConfigDir);
-					}
-					else
-					{
-						LayerPath = Layer.Path.Replace("{ENGINE}", UnrealBuildTool.EngineDirectory.FullName);
-					}
+					LayerPath = Layer.Path.Replace("{ENGINE}", UnrealBuildTool.EngineDirectory.FullName);
 				}
-				LayerPath = LayerPath.Replace("{TYPE}", BaseIniName);
-				LayerPath = LayerPath.Replace("{USERSETTINGS}", Utils.GetUserSettingDirectory().FullName);
-				if (bHasUserTag) LayerPath = LayerPath.Replace("{USER}", GetUserDir());
+			}
+			LayerPath = LayerPath.Replace("{TYPE}", BaseIniName);
+			LayerPath = LayerPath.Replace("{USERSETTINGS}", Utils.GetUserSettingDirectory().FullName);
+			LayerPath = LayerPath.Replace("{USER}", GetUserDir());
+
+			return LayerPath;
+		}
+
+		private static string GetExpansionPath(ConfigLayerExpansion Expansion, string LayerPath)
+		{
+			string ExpansionPath = LayerPath.Replace("{ED}", Expansion.DirectoryPrefix);
+			ExpansionPath = ExpansionPath.Replace("{EF}", Expansion.FilePrefix);
+
+			return ExpansionPath;
+		}
+
+		/// <summary>
+		/// Returns a list of INI filenames for the given project
+		/// </summary>
+		public static IEnumerable<FileReference> EnumerateConfigFileLocations(ConfigHierarchyType Type, DirectoryReference ProjectDir, UnrealTargetPlatform Platform)
+		{
+			string BaseIniName = Enum.GetName(typeof(ConfigHierarchyType), Type);
+			string PlatformName = GetIniPlatformName(Platform);
 
+			foreach (ConfigLayer Layer in ConfigLayers)
+			{
+				bool bHasPlatformTag, bHasProjectTag, bHasExpansionTag;
+				string LayerPath = GetLayerPath(Layer, Platform.ToString(), PlatformName, ProjectDir, BaseIniName, out bHasPlatformTag, out bHasProjectTag, out bHasExpansionTag);
+
+				// skip the layer if we aren't going to use it
+				if (LayerPath == null)
+				{
+					continue;
+				}
 
 				// handle expansion (and platform - the C++ code will validate that only expansion layers have platforms)
 				if (bHasExpansionTag)
@@ -883,8 +904,7 @@ namespace UnrealBuildTool
 					foreach (ConfigLayerExpansion Expansion in ConfigLayerExpansions)
 					{
 						// expansion replacements
-						string ExpansionPath = LayerPath.Replace("{ED}", Expansion.DirectoryPrefix);
-						ExpansionPath = ExpansionPath.Replace("{EF}", Expansion.FilePrefix);
+						string ExpansionPath = GetExpansionPath(Expansion, LayerPath);
 
 						// now go up the ini parent chain
 						if (bHasPlatformTag)
@@ -895,7 +915,13 @@ namespace UnrealBuildTool
 								// the IniParentChain
 								foreach (string ParentPlatform in Info.IniParentChain)
 								{
-									yield return new FileReference(ExpansionPath.Replace("{PLATFORM}", ParentPlatform));
+									// @note: We are using the ParentPlatform as both PlatformExtensionName _and_ IniPlatformName. This is because the parent
+									// may not even exist as a UnrealTargetPlatform, and all we have is a string to look up, and it would just get the same
+									// string back, if we did look it up. This could become an issue if Win64 becomes a PlatformExtension, and wants to have 
+									// a parent Platform, of ... something. This is likely to never be an issue, but leaving this note here just in case.
+									string LocalLayerPath = GetLayerPath(Layer, ParentPlatform, ParentPlatform, ProjectDir, BaseIniName, out bHasPlatformTag, out bHasProjectTag, out bHasExpansionTag);
+									string LocalExpansionPath = GetExpansionPath(Expansion, LocalLayerPath);
+									yield return new FileReference(LocalExpansionPath.Replace("{PLATFORM}", ParentPlatform));
 								}
 							}
 							// always yield the active platform last 
diff --git a/Engine/Source/Runtime/Core/Private/Misc/ConfigCacheIni.cpp b/Engine/Source/Runtime/Core/Private/Misc/ConfigCacheIni.cpp
--- a/Engine/Source/Runtime/Core/Private/Misc/ConfigCacheIni.cpp
+++ b/Engine/Source/Runtime/Core/Private/Misc/ConfigCacheIni.cpp
@@ -3319,6 +3319,89 @@ static void ConditionalOverrideIniFilename(FString& IniFilename, const TCHAR* Ba
 }
 
 
+static FString GetLayerPath(const FConfigLayer& Layer, const FString& PlatformName, const TCHAR* EngineConfigDir, const TCHAR* SourceConfigDir, const TCHAR* BaseIniName, bool& bHasPlatformTag, bool& bHasProjectTag)
+{
+	static FString LastPlatform;
+	static FString PlatformExtensionEngineConfigDir;
+	static FString PlatformExtensionProjectConfigDir;
+	static bool bHasPlatformExtensionEngineConfigDir;
+	static bool bHasPlatformExtensionProjectConfigDir;
+
+	if (LastPlatform != PlatformName)
+	{
+		LastPlatform = PlatformName;
+		PlatformExtensionEngineConfigDir = FPaths::Combine(*FPaths::EnginePlatformExtensionsDir(), *PlatformName, TEXT("Config"));
+		PlatformExtensionProjectConfigDir = FPaths::Combine(*FPaths::ProjectPlatformExtensionsDir(), *PlatformName, TEXT("Config"));
+		bHasPlatformExtensionEngineConfigDir = FPaths::DirectoryExists(*PlatformExtensionEngineConfigDir);
+		bHasPlatformExtensionProjectConfigDir = FPaths::DirectoryExists(*PlatformExtensionProjectConfigDir);
+	}
+
+	// cache some searches
+	bHasPlatformTag = FCString::Strstr(Layer.Path, TEXT("{PLATFORM}")) != nullptr;
+	bHasProjectTag = FCString::Strstr(Layer.Path, TEXT("{PROJECT}")) != nullptr;
+
+	// you can only have PROJECT or ENGINE, not both
+	FString LayerPath;
+	if (bHasProjectTag)
+	{
+		if (bHasPlatformTag && bHasPlatformExtensionProjectConfigDir)
+		{
+			LayerPath = FString(Layer.PlatformExtensionPath).Replace(TEXT("{EXTPROJECT}"), *PlatformExtensionProjectConfigDir);
+		}
+		else
+		{
+			LayerPath = FString(Layer.Path).Replace(TEXT("{PROJECT}"), SourceConfigDir);
+		}
+	}
+	else
+	{
+		if (bHasPlatformTag && bHasPlatformExtensionEngineConfigDir)
+		{
+			LayerPath = FString(Layer.PlatformExtensionPath).Replace(TEXT("{EXTENGINE}"), *PlatformExtensionEngineConfigDir);
+		}
+		else
+		{
+			LayerPath = FString(Layer.Path).Replace(TEXT("{ENGINE}"), EngineConfigDir);
+		}
+	}
+
+	LayerPath = LayerPath.Replace(TEXT("{TYPE}"), BaseIniName, ESearchCase::CaseSensitive);
+	LayerPath = LayerPath.Replace(TEXT("{USERSETTINGS}"), FPlatformProcess::UserSettingsDir(), ESearchCase::CaseSensitive);
+	LayerPath = LayerPath.Replace(TEXT("{USER}"), FPlatformProcess::UserDir(), ESearchCase::CaseSensitive);
+
+	return LayerPath;
+}
+
+static FString GetExpansionPath(const FConfigLayerExpansion& Expansion, const FString& LayerPath, bool bHasPlatformTag)
+{
+	// replace the expansion tags
+	FString ExpansionPath = LayerPath.Replace(TEXT("{ED}"), Expansion.DirectoryPrefix, ESearchCase::CaseSensitive);
+	ExpansionPath = ExpansionPath.Replace(TEXT("{EF}"), Expansion.FilePrefix, ESearchCase::CaseSensitive);
+
+	// check for dedicated server expansion
+	if (EnumHasAnyFlags(Expansion.Flag, EConfigLayerFlags::DedicatedServerOnly))
+	{
+		// it's unclear how a platform DS ini would be named, so for now, not supported
+		if (bHasPlatformTag)
+		{
+			return TEXT("");
+		}
+		else if (IsRunningDedicatedServer())
+		{
+			ExpansionPath = ExpansionPath.Replace(TEXT("Base"), TEXT(""));
+			ExpansionPath = ExpansionPath.Replace(TEXT("Default"), TEXT(""));
+			// ExpansionPath = ExpansionPath.Replace(TEXT("{PLATFORM}"), TEXT(""));
+		}
+		else
+		{
+			// skip this expansion if not DS
+			return TEXT("");
+		}
+	}
+
+	return ExpansionPath;
+}
+
 /**
  * Creates a chain of ini filenames to load and combine.
  *
@@ -3330,52 +3413,20 @@ void FConfigFile::AddStaticLayersToHierarchy(const TCHAR* InBaseIniName, const T
 {
 	SourceEngineConfigDir = EngineConfigDir;
 	SourceProjectConfigDir = SourceConfigDir;
-	
+
 	// get the platform name
 	const FString PlatformName(InPlatformName ? InPlatformName : ANSI_TO_TCHAR(FPlatformProperties::IniPlatformName()));
 
 	// cache some platform extension information that can be used inside the loops
-	FString PlatformExtensionEngineConfigDir = FPaths::Combine(*FPaths::EnginePlatformExtensionsDir(), *PlatformName, TEXT("Config"));
-	FString PlatformExtensionProjectConfigDir = FPaths::Combine(*FPaths::ProjectPlatformExtensionsDir(), *PlatformName, TEXT("Config"));
-	bool bHasPlatformExtensionEngineConfigDir = FPaths::DirectoryExists(*PlatformExtensionEngineConfigDir);
-	bool bHasPlatformExtensionProjectConfigDir = FPaths::DirectoryExists(*PlatformExtensionProjectConfigDir);
 
 	// go over all the config layers
 	for (int32 LayerIndex = 0; LayerIndex < UE_ARRAY_COUNT(GConfigLayers); LayerIndex++)
 	{
 		const FConfigLayer& Layer = GConfigLayers[LayerIndex];
-		const bool bHasPlatformTag = FCString::Strstr(Layer.Path, TEXT("{PLATFORM}")) != nullptr;
-		const bool bHasProjectTag = FCString::Strstr(Layer.Path, TEXT("{PROJECT}")) != nullptr;
 
 		// start replacing basic variables
-		FString LayerPath;
-		// you can only have PROJECT or ENGINE, not both
-		if (bHasProjectTag)
-		{
-			if (bHasPlatformTag && bHasPlatformExtensionProjectConfigDir)
-			{
-				LayerPath = FString(Layer.PlatformExtensionPath).Replace(TEXT("{EXTPROJECT}"), *PlatformExtensionProjectConfigDir);
-			}
-			else
-			{
-				LayerPath = FString(Layer.Path).Replace(TEXT("{PROJECT}"), SourceConfigDir);
-			}
-		}
-		else
-		{
-			if (bHasPlatformTag && bHasPlatformExtensionEngineConfigDir)
-			{
-				LayerPath = FString(Layer.PlatformExtensionPath).Replace(TEXT("{EXTENGINE}"), *PlatformExtensionEngineConfigDir);
-			}
-			else
-			{
-				LayerPath = FString(Layer.Path).Replace(TEXT("{ENGINE}"), EngineConfigDir);
-			}
-		}
-
-		LayerPath = LayerPath.Replace(TEXT("{TYPE}"), InBaseIniName, ESearchCase::CaseSensitive);
-		LayerPath = LayerPath.Replace(TEXT("{USERSETTINGS}"), FPlatformProcess::UserSettingsDir(), ESearchCase::CaseSensitive);
-		LayerPath = LayerPath.Replace(TEXT("{USER}"), FPlatformProcess::UserDir(), ESearchCase::CaseSensitive);
+		bool bHasPlatformTag, bHasProjectTag;
+		FString LayerPath = GetLayerPath(Layer, PlatformName, EngineConfigDir, SourceConfigDir, InBaseIniName, bHasPlatformTag, bHasProjectTag);
 
 		// PROGRAMs don't require any ini files
 #if IS_PROGRAM
@@ -3397,29 +3448,14 @@ void FConfigFile::AddStaticLayersToHierarchy(const TCHAR* InBaseIniName, const T
 			{
 				const FConfigLayerExpansion& Expansion = GConfigLayerExpansions[ExpansionIndex];
 
-				// replace the expansion tags
-				FString ExpansionPath = LayerPath.Replace(TEXT("{ED}"), Expansion.DirectoryPrefix, ESearchCase::CaseSensitive);
-				ExpansionPath = ExpansionPath.Replace(TEXT("{EF}"), Expansion.FilePrefix, ESearchCase::CaseSensitive);
+				bool bSkipExpansion = false;
+				FString ExpansionPath = GetExpansionPath(Expansion, LayerPath, bHasPlatformTag);
 
-				// check for dedicated server expansion
-				if (EnumHasAnyFlags(Expansion.Flag, EConfigLayerFlags::DedicatedServerOnly))
+				FPlatformMisc::LowLevelOutputDebugStringf(TEXT("  Expansion: %s\n"), *ExpansionPath);
+				// skip this expansion if we didn't want it
+				if (ExpansionPath.Len() == 0)
 				{
-					// it's unclear how a platform DS ini would be named, so for now, not supported
-					if (bHasPlatformTag)
-					{
-						continue;
-					}
-					if (IsRunningDedicatedServer())
-					{
-						ExpansionPath = ExpansionPath.Replace(TEXT("Base"), TEXT(""));
-						ExpansionPath = ExpansionPath.Replace(TEXT("Default"), TEXT(""));
-						// ExpansionPath = ExpansionPath.Replace(TEXT("{PLATFORM}"), TEXT(""));
-					}
-					else
-					{
-						// skip this expansion if not DS
-						continue;
-					}
+					continue;
 				}
 
 				// allow for override, only on BASE EXPANSION!
@@ -3436,20 +3472,36 @@ void FConfigFile::AddStaticLayersToHierarchy(const TCHAR* InBaseIniName, const T
 
 				const FDataDrivenPlatformInfoRegistry::FPlatformInfo& Info = FDataDrivenPlatformInfoRegistry::GetPlatformInfo(PlatformName);
 
-				// go over parents, and then this platform, unless there's no platform tag, then we simply want to run through the loop one time to add it to the 
+				// go over parents, and then this platform, unless there's no platform tag, then we simply want to run through the loop one time to add it to the
 				int NumPlatforms = bHasPlatformTag ? Info.IniParentChain.Num() + 1 : 1;
 				check(NumPlatforms < MaxPlatformIndex);
 				for (int PlatformIndex = 0; PlatformIndex < NumPlatforms; PlatformIndex++)
 				{
 					// the platform to work on (active platform is always last)
 					const FString& CurrentPlatform = (PlatformIndex == NumPlatforms - 1) ? PlatformName : Info.IniParentChain[PlatformIndex];
+					FString PlatformPath;
+
+					// ini parent may not be a PlatformExtension even if this was, so update based on the current platform after the first
+					if (PlatformIndex < NumPlatforms - 1)
+					{
+						// overriding the bHas*Tag fields is fine here since it will evaluate to the same
+						FString LocalLayerPath = GetLayerPath(Layer, CurrentPlatform, EngineConfigDir, SourceConfigDir, InBaseIniName, bHasPlatformTag, bHasProjectTag);
+
+						// re-expand the new path
+						FString LocalExpansionPath = GetExpansionPath(Expansion, LocalLayerPath, bHasPlatformTag);
 
-					// replace
-					FString PlatformPath = ExpansionPath.Replace(TEXT("{PLATFORM}"), *CurrentPlatform, ESearchCase::CaseSensitive);
+						// modify get the final platform path from the parent ini's location
+						PlatformPath = LocalExpansionPath.Replace(TEXT("{PLATFORM}"), *CurrentPlatform, ESearchCase::CaseSensitive);
+					}
+					else
+					{
+						// replace
+						PlatformPath = ExpansionPath.Replace(TEXT("{PLATFORM}"), *CurrentPlatform, ESearchCase::CaseSensitive);
+					}
 
 					// add this to the list!
 					SourceIniHierarchy.AddStaticLayer(
-						FIniFilename(PlatformPath, bIsRequired, bGenerateCacheKey ? 
+						FIniFilename(PlatformPath, bIsRequired, bGenerateCacheKey ?
 							GenerateHierarchyCacheKey(SourceIniHierarchy, PlatformPath, InBaseIniName) :
 							FString(TEXT(""))),
 						LayerIndex, ExpansionIndex, PlatformIndex);
@@ -3464,7 +3516,7 @@ void FConfigFile::AddStaticLayersToHierarchy(const TCHAR* InBaseIniName, const T
 
 			// final layer needs to generate a hierarchy cache
 			SourceIniHierarchy.AddStaticLayer(
-				FIniFilename(LayerPath, bIsRequired, EnumHasAnyFlags(Layer.Flag, EConfigLayerFlags::GenerateCacheKey) ? 
+				FIniFilename(LayerPath, bIsRequired, EnumHasAnyFlags(Layer.Flag, EConfigLayerFlags::GenerateCacheKey) ?
 					GenerateHierarchyCacheKey(SourceIniHierarchy, LayerPath, InBaseIniName) :
 					FString(TEXT(""))),
 				LayerIndex);
