Missed removing debug low level print

diff --git a/Engine/Source/Runtime/Core/Private/Misc/ConfigCacheIni.cpp b/Engine/Source/Runtime/Core/Private/Misc/ConfigCacheIni.cpp
--- a/Engine/Source/Runtime/Core/Private/Misc/ConfigCacheIni.cpp
+++ b/Engine/Source/Runtime/Core/Private/Misc/ConfigCacheIni.cpp
@@ -3451,7 +3451,6 @@ void FConfigFile::AddStaticLayersToHierarchy(const TCHAR* InBaseIniName, const T
 				bool bSkipExpansion = false;
 				FString ExpansionPath = GetExpansionPath(Expansion, LayerPath, bHasPlatformTag);
 
-				FPlatformMisc::LowLevelOutputDebugStringf(TEXT("  Expansion: %s\n"), *ExpansionPath);
 				// skip this expansion if we didn't want it
 				if (ExpansionPath.Len() == 0)
 				{
