Set human readable error message for unhandled exceptions.

diff --git a/Engine/Source/Runtime/Core/Private/Windows/WindowsPlatformCrashContext.cpp b/Engine/Source/Runtime/Core/Private/Windows/WindowsPlatformCrashContext.cpp
--- a/Engine/Source/Runtime/Core/Private/Windows/WindowsPlatformCrashContext.cpp
+++ b/Engine/Source/Runtime/Core/Private/Windows/WindowsPlatformCrashContext.cpp
@@ -1117,6 +1117,12 @@ private:
 			ErrorMessage = Info.ErrorMessage;
 			NumStackFramesToIgnore += Info.NumStackFramesToIgnore;
 		}
+		// Generic exception description is stored in GErrorExceptionDescription
+		else if (ExceptionInfo->ExceptionRecord->ExceptionCode != 1)
+		{
+			CreateExceptionInfoString(ExceptionInfo->ExceptionRecord);
+			ErrorMessage = GErrorExceptionDescription;
+		}
 
 #if USE_CRASH_REPORTER_MONITOR
 		if (CrashClientHandle.IsValid() && FPlatformProcess::IsProcRunning(CrashClientHandle))
