Basic Android support for crashdebughelper, crashreportcore.
Android Arm64 stackwalk uses libunwind.
Crash handler support for Android platform.
#rnx

diff --git a/Engine/Source/Developer/CrashDebugHelper/Private/Android/CrashDebugHelperAndroid.cpp b/Engine/Source/Developer/CrashDebugHelper/Private/Android/CrashDebugHelperAndroid.cpp
new file mode 100644
--- /dev/null
+++ b/Engine/Source/Developer/CrashDebugHelper/Private/Android/CrashDebugHelperAndroid.cpp
@@ -0,0 +1,8 @@
+// Copyright 1998-2019 Epic Games, Inc. All Rights Reserved.
+
+#include "CrashDebugHelperAndroid.h"
+
+bool FCrashDebugHelperAndroid::CreateMinidumpDiagnosticReport( const FString& InCrashDumpName )
+{
+	return true;
+}
diff --git a/Engine/Source/Developer/CrashDebugHelper/Private/Android/CrashDebugHelperAndroid.h b/Engine/Source/Developer/CrashDebugHelper/Private/Android/CrashDebugHelperAndroid.h
new file mode 100644
--- /dev/null
+++ b/Engine/Source/Developer/CrashDebugHelper/Private/Android/CrashDebugHelperAndroid.h
@@ -0,0 +1,23 @@
+// Copyright 1998-2019 Epic Games, Inc. All Rights Reserved.
+
+#pragma once
+
+#include "CoreMinimal.h"
+#include "CrashDebugHelper.h"
+
+class FString;
+
+class FCrashDebugHelperAndroid : public ICrashDebugHelper
+{
+public:
+	/**
+	 *	Parse the given crash dump, and generate a report. 
+	 *
+	 *	@param	InCrashDumpName		The crash dump file to process
+	 *
+	 *	@return	bool				true if successful, false if not
+	 */
+	virtual bool CreateMinidumpDiagnosticReport( const FString& InCrashDumpName ) override;
+};
+
+typedef FCrashDebugHelperAndroid FCrashDebugHelper;
diff --git a/Engine/Source/Developer/CrashDebugHelper/Private/CrashDebugHelperModule.cpp b/Engine/Source/Developer/CrashDebugHelper/Private/CrashDebugHelperModule.cpp
--- a/Engine/Source/Developer/CrashDebugHelper/Private/CrashDebugHelperModule.cpp
+++ b/Engine/Source/Developer/CrashDebugHelper/Private/CrashDebugHelperModule.cpp
@@ -20,6 +20,9 @@ DEFINE_LOG_CATEGORY(LogCrashDebugHelper);
 #elif PLATFORM_IOS
 	#include "CrashDebugHelperIOS.h"
 
+#elif PLATFORM_ANDROID
+	#include "Android/CrashDebugHelperAndroid.h"
+
 #else
 	#error "Unknown platform"
 #endif
diff --git a/Engine/Source/Runtime/Core/Core.Build.cs b/Engine/Source/Runtime/Core/Core.Build.cs
--- a/Engine/Source/Runtime/Core/Core.Build.cs
+++ b/Engine/Source/Runtime/Core/Core.Build.cs
@@ -124,7 +124,8 @@ public class Core : ModuleRules
 		{
 			AddEngineThirdPartyPrivateStaticDependencies(Target,
 				"cxademangle",
-				"zlib"
+				"zlib",
+				"libunwind"
 				);
 		}
 		else if (Target.IsInPlatformGroup(UnrealPlatformGroup.Unix))
diff --git a/Engine/Source/Runtime/Core/Private/Android/AndroidPlatformCrashContext.cpp b/Engine/Source/Runtime/Core/Private/Android/AndroidPlatformCrashContext.cpp
new file mode 100644
--- /dev/null
+++ b/Engine/Source/Runtime/Core/Private/Android/AndroidPlatformCrashContext.cpp
@@ -0,0 +1,287 @@
+// Copyright 1998-2019 Epic Games, Inc. All Rights Reserved.
+
+/*=============================================================================
+	AndroidPlatformCrashContext.cpp: implementations of Android platform crash context.
+=============================================================================*/
+#include "Android/AndroidPlatformCrashContext.h"
+#include <sys/types.h>
+#include <unistd.h>
+#include <stdlib.h>
+#include <limits.h>
+#include <dlfcn.h>
+#include "Containers/StringConv.h"
+#include "CoreGlobals.h"
+#include "Misc/App.h"
+#include "Misc/Guid.h"
+#include "HAL/PlatformStackWalk.h"
+#include "Misc/Paths.h"
+#include "GenericPlatform/GenericPlatformOutputDevices.h"
+#include "HAL/FileManager.h"
+#include "HAL/IConsoleManager.h"
+
+static int64 GetAndroidLibraryBaseAddress();
+
+extern FString AndroidThunkCpp_GetMetaDataString(const FString& Key);
+
+extern FString AndroidRelativeToAbsolutePath(bool bUseInternalBasePath, FString RelPath);
+
+struct FAndroidCrashInfo
+{
+	FAndroidCrashInfo()
+	{
+	}
+
+	void Init()
+	{
+		if (!bInitialized)
+		{
+			FGuid RunGUID = FGuid::NewGuid();
+			FCStringAnsi::Strcpy(AppName, TCHAR_TO_UTF8(FApp::GetProjectName()));
+			FString LogPath = FGenericPlatformOutputDevices::GetAbsoluteLogFilename();
+			LogPath = AndroidRelativeToAbsolutePath(false, LogPath);
+			FCStringAnsi::Strcpy(AppLogPath, TCHAR_TO_UTF8(*LogPath));
+
+			// Cache & create the crash report folder.
+			FString ReportPath = FPaths::GameAgnosticSavedDir() / TEXT("Crashes");
+			ReportPath = AndroidRelativeToAbsolutePath(true, ReportPath);
+			IFileManager::Get().MakeDirectory(*ReportPath, true);
+			FCStringAnsi::Strcpy(AndroidCrashReportPath, TCHAR_TO_UTF8(*ReportPath));
+
+			FCStringAnsi::Strcpy(ProjectNameUTF8, TCHAR_TO_UTF8(FApp::GetProjectName()));
+			FAndroidCrashContext::GenerateReportDirectoryName(TargetDirectory);
+			bInitialized = true;
+		}
+	}
+
+	static const int32 MaxAppNameSize = 128;
+	char AppName[MaxAppNameSize] = { 0 };
+	char AndroidCrashReportPath[FAndroidCrashContext::CrashReportMaxPathSize] = { 0 };
+	char AppLogPath[FAndroidCrashContext::CrashReportMaxPathSize] = { 0 };
+	char JavaLog[FAndroidCrashContext::CrashReportMaxPathSize] = { 0 };
+	char TargetDirectory[FAndroidCrashContext::CrashReportMaxPathSize] = { 0 };
+	char ProjectNameUTF8[FAndroidCrashContext::CrashReportMaxPathSize] = { 0 };
+	bool bInitialized = false;
+} GAndroidCrashInfo;
+
+const FString FAndroidCrashContext::GetCrashDirectoryName()
+{
+	return FString(GAndroidCrashInfo.TargetDirectory);
+}
+
+void FAndroidCrashContext::GetCrashDirectoryName(char(&DirectoryNameOUT)[CrashReportMaxPathSize])
+{
+	FCStringAnsi::Strncpy(DirectoryNameOUT, GAndroidCrashInfo.TargetDirectory, CrashReportMaxPathSize);
+}
+
+static ANSICHAR* ItoANSI(uint64 Val, uint64 Base, uint32 Len = 0)
+{
+	static ANSICHAR InternalBuffer[64] = { 0 };
+
+	uint64 i = 62;
+	int32 pad = Len;
+
+	if (Val)
+	{
+		for (; Val && i; --i, Val /= Base, --pad)
+		{
+			InternalBuffer[i] = "0123456789abcdef"[Val % Base];
+		}
+	}
+	else
+	{
+		InternalBuffer[i--] = '0';
+		--pad;
+	}
+
+	while (pad > 0)
+	{
+		InternalBuffer[i--] = '0';
+		--pad;
+	}
+
+	return &InternalBuffer[i + 1];
+}
+
+
+void FAndroidCrashContext::GenerateReportDirectoryName(char(&DirectoryNameOUT)[CrashReportMaxPathSize])
+{
+	FGuid ReportGUID = FGuid::NewGuid();
+	//FCStringAnsi::Snprintf(DirectoryNameOUT, CrashReportMaxPathSize, ("%s/CrashReport-UE4-%s-pid-%d-%08X%08X%08X%08X"), GAndroidCrashInfo.AndroidCrashReportPath, GAndroidCrashInfo.ProjectNameUTF8, (uint32)getpid(), ReportGUID.A, ReportGUID.B, ReportGUID.C, ReportGUID.D);
+	FCStringAnsi::Strncpy(DirectoryNameOUT, GAndroidCrashInfo.AndroidCrashReportPath, CrashReportMaxPathSize);
+	FCStringAnsi::Strcat(DirectoryNameOUT, CrashReportMaxPathSize, "/CrashReport-UE4-");
+	FCStringAnsi::Strcat(DirectoryNameOUT, CrashReportMaxPathSize, GAndroidCrashInfo.ProjectNameUTF8);
+	FCStringAnsi::Strcat(DirectoryNameOUT, CrashReportMaxPathSize, "-pid-");
+	FCStringAnsi::Strcat(DirectoryNameOUT, CrashReportMaxPathSize, ItoANSI((uint32)getpid(), 10));
+	FCStringAnsi::Strcat(DirectoryNameOUT, CrashReportMaxPathSize, "-");
+	FCStringAnsi::Strcat(DirectoryNameOUT, CrashReportMaxPathSize, ItoANSI(ReportGUID.A, 16, 8));
+	FCStringAnsi::Strcat(DirectoryNameOUT, CrashReportMaxPathSize, ItoANSI(ReportGUID.B, 16, 8));
+	FCStringAnsi::Strcat(DirectoryNameOUT, CrashReportMaxPathSize, ItoANSI(ReportGUID.C, 16, 8));
+	FCStringAnsi::Strcat(DirectoryNameOUT, CrashReportMaxPathSize, ItoANSI(ReportGUID.D, 16, 8));
+}
+
+static void CrashReportFileCopy(const char* DestPath, const char* SourcePath)
+{
+	int SourceHandle = open(SourcePath, O_RDONLY);
+	int DestHandle = open(DestPath, O_CREAT | O_WRONLY, S_IRUSR | S_IWUSR);
+
+	char Data[PATH_MAX] = {};
+	int Bytes = 0;
+	while ((Bytes = read(SourceHandle, Data, PATH_MAX)) > 0)
+	{
+		write(DestHandle, Data, Bytes);
+	}
+
+	close(DestHandle);
+	close(SourceHandle);
+}
+
+void FAndroidCrashContext::StoreCrashInfo() const
+{
+	char CrashDirectoryName[CrashReportMaxPathSize] = { 0 };
+	char FilePath[CrashReportMaxPathSize] = { 0 };
+	if (GetType() == ECrashContextType::Ensure)
+	{
+		// create a new report folder.
+		GenerateReportDirectoryName(CrashDirectoryName);
+	}
+	else
+	{
+		GetCrashDirectoryName(CrashDirectoryName);
+	}
+
+	FCStringAnsi::Strcpy(FilePath, CrashDirectoryName);
+	FCStringAnsi::Strcat(FilePath, "/");
+	FCStringAnsi::Strcat(FilePath, FGenericCrashContext::CrashContextRuntimeXMLNameA);
+	SerializeAsXML(*FString(FilePath)); // CreateFileWriter will also create destination directory.
+
+	// copy log:
+	FCStringAnsi::Strcpy(FilePath, CrashDirectoryName);
+	FCStringAnsi::Strcat(FilePath, "/");
+	FCStringAnsi::Strcat(FilePath, FCStringAnsi::Strlen(GAndroidCrashInfo.AppName) ? GAndroidCrashInfo.AppName : "UE4");
+	FCStringAnsi::Strcat(FilePath, ".log");
+	CrashReportFileCopy(FilePath, GAndroidCrashInfo.AppLogPath);
+}
+
+void FAndroidCrashContext::Initialize()
+{
+	GAndroidCrashInfo.Init();
+}
+
+FAndroidCrashContext::FAndroidCrashContext(ECrashContextType InType, const TCHAR* InErrorMessage)
+:	FGenericCrashContext(InType, InErrorMessage)
+, Signal(0)
+, Info(NULL)
+, Context(NULL)
+{
+}
+
+void FAndroidCrashContext::CaptureCrashInfo()
+{
+	CapturePortableCallStack(0, Context);
+}
+
+static int64 GetAndroidLibraryBaseAddress()
+{
+	int64 BaseAddress = 0;
+
+	const char *LibraryName = "libUE4.so";
+	int32 LibraryLength = strlen(LibraryName);
+
+	// try to open process map file
+	FILE *file = fopen("/proc/self/maps", "r");
+	if (file == NULL)
+	{
+		return BaseAddress;
+	}
+
+	char LineBuffer[256];
+	LineBuffer[255] = 0;
+	while (fgets(LineBuffer, 255, file) != NULL)
+	{
+		int32 BufferLength = strlen(LineBuffer);
+		if (BufferLength > 0 && LineBuffer[BufferLength - 1] == '\n')
+		{
+			BufferLength--;
+			LineBuffer[BufferLength] = 0;
+		}
+
+		// does it end with library name?
+		if (BufferLength < LibraryLength || memcmp(LineBuffer + BufferLength - LibraryLength, LibraryName, LibraryLength))
+		{
+			continue;
+		}
+
+		// parse the line
+		int64 StartAddress, EndAddress, Offset;
+		char flags[4];
+		if (sscanf(LineBuffer, "%llx-%llx %c%c%c%c %llx", &StartAddress, &EndAddress, &flags[0], &flags[1], &flags[2], &flags[3], &Offset) != 7)
+		{
+			continue;
+		}
+
+		// needs to be r-x (ignore 4th)
+		if (flags[0] == 'r' && flags[1] == '-' && flags[2] == 'x')
+		{
+			BaseAddress = StartAddress - Offset;
+			break;
+		}
+	}
+	fclose(file);
+	return BaseAddress;
+}
+
+void FAndroidCrashContext::AddAndroidCrashProperty(const FString& Key, const FString& Value)
+{
+	AdditionalProperties.Add(Key, Value);
+}
+
+void FAndroidCrashContext::AddPlatformSpecificProperties() const
+{
+	for (TMap<FString, FString>::TConstIterator Iter(AdditionalProperties); Iter; ++Iter)
+	{
+		AddCrashProperty(*Iter.Key(), *Iter.Value());
+	}
+}
+
+void FAndroidCrashContext::GetPortableCallStack(const uint64* StackFrames, int32 NumStackFrames, TArray<FCrashStackFrame>& OutCallStack)
+{
+	// Get all the modules in the current process
+	uint32 NumModules = (uint32)FPlatformStackWalk::GetProcessModuleCount();
+
+	TArray<FStackWalkModuleInfo> Modules;
+	Modules.AddUninitialized(NumModules);
+
+	NumModules = FPlatformStackWalk::GetProcessModuleSignatures(Modules.GetData(), NumModules);
+	Modules.SetNum(NumModules);
+
+	// Update the callstack with offsets from each module
+	OutCallStack.Reset(NumStackFrames);
+	for (int32 Idx = 0; Idx < NumStackFrames; Idx++)
+	{
+		const uint64 StackFrame = StackFrames[Idx];
+
+		// Try to find the module containing this stack frame
+		const FStackWalkModuleInfo* FoundModule = nullptr;
+
+		Dl_info DylibInfo;
+		int32 Result = dladdr((const void*)StackFrame, &DylibInfo);
+		if (Result != 0)
+		{
+			ANSICHAR* DylibPath = (ANSICHAR*)DylibInfo.dli_fname;
+			ANSICHAR* DylibName = FCStringAnsi::Strrchr(DylibPath, '/');
+			if (DylibName)
+			{
+				DylibName += 1;
+			}
+			else
+			{
+				DylibName = DylibPath;
+			}
+			OutCallStack.Add(FCrashStackFrame(FPaths::GetBaseFilename(DylibName), (uint64)DylibInfo.dli_fbase, StackFrame - (uint64)DylibInfo.dli_fbase));
+		}
+		else
+		{
+			OutCallStack.Add(FCrashStackFrame(TEXT("Unknown"), 0, StackFrame));
+		}
+	}
+}
diff --git a/Engine/Source/Runtime/Core/Private/Android/AndroidPlatformFile.cpp b/Engine/Source/Runtime/Core/Private/Android/AndroidPlatformFile.cpp
--- a/Engine/Source/Runtime/Core/Private/Android/AndroidPlatformFile.cpp
+++ b/Engine/Source/Runtime/Core/Private/Android/AndroidPlatformFile.cpp
@@ -119,6 +119,21 @@ const FString &GetFileBasePath()
 	return BasePath;
 }
 
+
+FString AndroidRelativeToAbsolutePath(bool bUseInternalBasePath, FString RelPath)
+{
+	if (RelPath.StartsWith(TEXT("../")))
+	{
+		while (RelPath.StartsWith(TEXT("../")))
+		{
+			RelPath = RelPath.RightChop(3);
+		}
+
+		return (bUseInternalBasePath ? GInternalFilePath : GetFileBasePath()) / RelPath;
+	}
+	return RelPath;
+}
+
 /**
 	Android file handle implementation for partial (i.e. parcels) files.
 	This can handle all the variations of accessing data for assets and files.
diff --git a/Engine/Source/Runtime/Core/Private/Android/AndroidPlatformMisc.cpp b/Engine/Source/Runtime/Core/Private/Android/AndroidPlatformMisc.cpp
--- a/Engine/Source/Runtime/Core/Private/Android/AndroidPlatformMisc.cpp
+++ b/Engine/Source/Runtime/Core/Private/Android/AndroidPlatformMisc.cpp
@@ -41,6 +41,9 @@
 #include "Async/TaskGraphInterfaces.h"
 
 #include "FramePro/FrameProProfiler.h"
+#include <sys/mman.h>
+#include "Templates/AlignmentTemplates.h"
+#include "GenericPlatform/GenericPlatformOutputDevices.h"
 
 #if USE_ANDROID_JNI
 extern AAssetManager * AndroidThunkCpp_GetAssetManager();
@@ -171,6 +174,7 @@ namespace FAndroidAppEntry
 
 void FAndroidMisc::PlatformPreInit()
 {
+	FAndroidCrashContext::Initialize();
 	FGenericPlatformMisc::PlatformPreInit();
 	EstablishVulkanDeviceSupport();
 	FAndroidAppEntry::PlatformInit();
@@ -902,6 +906,38 @@ const int32 NumTargetSignals = UE_ARRAY_COUNT(TargetSignals);
 struct sigaction PrevActions[NumTargetSignals];
 static bool PreviousSignalHandlersValid = false;
 
+const int32 GAndroidFakeSignalStackSize = 256 * 1024; // Desired stack size
+const uint32 GAndroidFakeSignalStackAllocationSize = GAndroidFakeSignalStackSize + PAGE_SIZE + PAGE_SIZE; // Total stack allocation (including guard page bookends)
+const uint32 GAndroidFakeSignalStackLocation = GAndroidFakeSignalStackAllocationSize - PAGE_SIZE; // Beginning of the desired stack.
+uint8* GAndroidFakeSignalStack = nullptr;
+
+void AllocateFakeSignalStack()
+{
+	check(GAndroidFakeSignalStack == nullptr);
+
+	uint8* BaseStackAllocation = (uint8*)mmap(nullptr, GAndroidFakeSignalStackAllocationSize, PROT_READ | PROT_WRITE, MAP_ANONYMOUS | MAP_PRIVATE, -1, 0);
+	// protect top of the stack to catch overruns.
+	if (mprotect(BaseStackAllocation, PAGE_SIZE, PROT_NONE) != 0)
+	{
+		UE_LOG(LogAndroid, Fatal, TEXT("Failed to set alt stack protection"));
+	}
+	if (mprotect(BaseStackAllocation + GAndroidFakeSignalStackSize + PAGE_SIZE, PAGE_SIZE, PROT_NONE) != 0)
+	{
+		UE_LOG(LogAndroid, Fatal, TEXT("Failed to set alt stack protection"));
+	}
+
+	GAndroidFakeSignalStack = BaseStackAllocation;
+
+	check(IsAligned(&GAndroidFakeSignalStack[GAndroidFakeSignalStackLocation], 16));
+}
+
+void ReleaseFakeSignalStack()
+{
+	check(GAndroidFakeSignalStack);
+	munmap(GAndroidFakeSignalStack, GAndroidFakeSignalStackAllocationSize);
+	GAndroidFakeSignalStack = nullptr;
+}
+
 static void RestorePreviousSignalHandlers()
 {
 	if (PreviousSignalHandlersValid)
@@ -914,20 +950,60 @@ static void RestorePreviousSignalHandlers()
 	}
 }
 
-/** True system-specific crash handler that gets called first */
-void PlatformCrashHandler(int32 Signal, siginfo* Info, void* Context)
+static void SetDefaultSignalHandlers()
 {
-	// Switch to malloc crash.
-	//FGenericPlatformMallocCrash::Get().SetAsGMalloc(); @todo uncomment after verification
+	struct sigaction Action;
+	FMemory::Memzero(&Action, sizeof(struct sigaction));
+	Action.sa_handler = SIG_DFL;
+	sigemptyset(&Action.sa_mask);
 
-	FPlatformMisc::LowLevelOutputDebugStringf(TEXT("Signal %d caught!"), Signal);
+	for (int32 i = 0; i < NumTargetSignals; ++i)
+	{
+		sigaction(TargetSignals[i], &Action, NULL);
+	}
+}
 
-	// Restore system handlers so Android could catch this signal after we are done with crashreport
-	RestorePreviousSignalHandlers();
+void AndroidCrashHandlerSubstituteStack(int Signal, siginfo* Info, void* Context);
 
-	FAndroidCrashContext CrashContext(ECrashContextType::Crash, TEXT("Caught signal"));
-	CrashContext.InitFromSignal(Signal, Info, Context);
+void HookSignals()
+{
+	// hook our signals and record current set.
+	struct sigaction Action;
+	FMemory::Memzero(&Action, sizeof(struct sigaction));
+	Action.sa_sigaction = AndroidCrashHandlerSubstituteStack;
+	// sigfillset will block all other signals whilst the signal handler is processing.
+	sigfillset(&Action.sa_mask);
+	Action.sa_flags = SA_SIGINFO | SA_RESTART | SA_ONSTACK;
+
+	for (int32 i = 0; i < NumTargetSignals; ++i)
+	{
+		sigaction(TargetSignals[i], &Action, &PrevActions[i]);
+	}
+	PreviousSignalHandlersValid = true;
+}
+
+// params to pass from AndroidCrashHandlerSubstituteStack to PlatformCrashHandler.
+struct FAndroidSignalParams
+{
+	int Signal;
+	siginfo* Info;
+	void* Context;
+};
+static FAndroidSignalParams GAndroidSignalParams;
+static volatile sig_atomic_t handling_fatal_signal = 0;
+
+bool FAndroidMisc::IsInSignalHandler()
+{
+	return FPlatformAtomics::AtomicRead(&handling_fatal_signal) > 0;
+}
 
+void FAndroidMisc::TriggerNonFatalCrashHandler(ECrashContextType InType, const FString& Message)
+{
+	check(InType == ECrashContextType::Ensure);
+
+	FAndroidCrashContext CrashContext(InType, *Message);
+
+	CrashContext.CaptureCrashInfo();
 	if (GCrashHandlerPointer)
 	{
 		GCrashHandlerPointer(CrashContext);
@@ -939,12 +1015,113 @@ void PlatformCrashHandler(int32 Signal, siginfo* Info, void* Context)
 	}
 }
 
-void FAndroidMisc::SetCrashHandler(void (* CrashHandler)(const FGenericCrashContext& Context))
+void FAndroidMisc::TriggerCrashHandler(const TCHAR* InErrorMessage, const TMap<FString, FString>& AdditionalProperties)
 {
+	// we are a fatal signal, we *can* only handle one at a time. So avoid allow multiple fatal signals going through
+	if (FPlatformAtomics::InterlockedIncrement(&handling_fatal_signal) != 1)
+	{
+		FPlatformProcess::Sleep(60.0f);
+		exit(0);
+	}
+
+	FAndroidCrashContext CrashContext(ECrashContextType::Crash, InErrorMessage);
+	CrashContext.CaptureCrashInfo();
+	for (TMap<FString, FString>::TConstIterator Iter(AdditionalProperties); Iter; ++Iter)
+	{
+		CrashContext.AddAndroidCrashProperty(*Iter.Key(), *Iter.Value());
+	}
+
+	if (GCrashHandlerPointer)
+	{
+		GCrashHandlerPointer(CrashContext);
+	}
+	else
+	{
+		// call default one
+		DefaultCrashHandler(CrashContext);
+	}
+}
+
+
+void PlatformCrashHandler()
+{
+	// Switch to malloc crash.
+	//FGenericPlatformMallocCrash::Get().SetAsGMalloc(); @todo uncomment after verification
+
+	FAndroidCrashContext CrashContext(ECrashContextType::Crash, TEXT("Caught signal"));
+	CrashContext.InitFromSignal(GAndroidSignalParams.Signal, GAndroidSignalParams.Info, GAndroidSignalParams.Context);
+	CrashContext.CaptureCrashInfo();
+	if (GCrashHandlerPointer)
+	{
+		GCrashHandlerPointer(CrashContext);
+	}
+	else
+	{
+		// call default one
+		DefaultCrashHandler(CrashContext);
+	}
+}
+
+PRAGMA_DISABLE_OPTIMIZATION
+#if PLATFORM_ANDROID_ARM64
+/** True system-specific crash handler that replaces stack pointer,  */
+void AndroidCrashHandlerSubstituteStack(int Signal, siginfo* Info, void* Context)
+{
+	// we are a fatal signal, we *can* only handle one at a time. So avoid allow multiple fatal signals going through
+	if (FPlatformAtomics::InterlockedIncrement(&handling_fatal_signal) != 1)
+	{
+		FPlatformProcess::Sleep(60.0f);
+		exit(0);
+	}
+
+	// prepare our signal parameters:
+	GAndroidSignalParams.Signal = Signal;
+	GAndroidSignalParams.Info = Info;
+	GAndroidSignalParams.Context = Context;
+
+	// Arm64, save current stack frame, swap in new stack and call our signal handler.
+	uint8* FakeStackBase = &GAndroidFakeSignalStack[GAndroidFakeSignalStackLocation];	
+	__asm__("mov x8, %0" : : "r"(FakeStackBase));					// Put the fake stack addr into x8
+	__asm__("stp x29, x30, [x8, -32]" : : );						// push the existing frame and lr on to fake stack
+	__asm__("mov x9, sp" : : );										// copy existing stack ptr to x9
+	__asm__("str x9, [x8, -16]" : : );								// push the existing stack on to fake stack.
+
+	__asm__("sub sp, x8, 32" : : );									// set stack ptr to new stack addr, after our pushed regs
+	__asm__("sub fp, x8, 32" : : );									// set frame ptr to new stack addr, after our pushed regs
+
+    __asm__("bl %0" : : "i"(PlatformCrashHandler));					// Call our actual crash handler.
+
+	__asm__("ldp x29,x30,[sp]" : : );								// restore the original frame and lr.
+	__asm__("ldr x8, [sp, 16]" : : );								// restore the previously pushed stack pointer into x8
+	__asm__("mov sp, x8" : : );										// restore the original stack pointer
+
+	// Restore system handlers so Android could catch this signal after we are done with crashreport
+	RestorePreviousSignalHandlers();
+
+	FPlatformAtomics::InterlockedDecrement(&handling_fatal_signal);
+}
+#else
+// TODO
+void AndroidCrashHandlerSubstituteStack(int Signal, siginfo* Info, void* Context)
+{	
+	PlatformCrashHandler();
+}
+#endif 
+PRAGMA_ENABLE_OPTIMIZATION
+
+void FAndroidMisc::SetCrashHandler(void(*CrashHandler)(const FGenericCrashContext& Context))
+{
+#if PLATFORM_ANDROID_ARM64
+	UE_LOG(LogAndroid, Log, TEXT("Setting Crash Handler = %p"), CrashHandler);
+
 	GCrashHandlerPointer = CrashHandler;
 
 	RestorePreviousSignalHandlers();
 	FMemory::Memzero(&PrevActions, sizeof(PrevActions));
+	if(GAndroidFakeSignalStack)
+	{
+		ReleaseFakeSignalStack();
+	}
 
 	// Passing -1 will leave these restored and won't trap them
 	if ((PTRINT)CrashHandler == -1)
@@ -952,17 +1129,11 @@ void FAndroidMisc::SetCrashHandler(void (* CrashHandler)(const FGenericCrashCont
 		return;
 	}
 
-	struct sigaction Action;
-	FMemory::Memzero(&Action, sizeof(struct sigaction));
-	Action.sa_sigaction = PlatformCrashHandler;
-	sigemptyset(&Action.sa_mask);
-	Action.sa_flags = SA_SIGINFO | SA_RESTART | SA_ONSTACK;
+	AllocateFakeSignalStack();
 
-	for (int32 i = 0; i < NumTargetSignals; ++i)
-	{
-		sigaction(TargetSignals[i],	&Action, &PrevActions[i]);
-	}
-	PreviousSignalHandlersValid = true;
+	HookSignals();
+
+#endif //PLATFORM_ANDROID_ARM64
 }
 
 bool FAndroidMisc::GetUseVirtualJoysticks()
diff --git a/Engine/Source/Runtime/Core/Private/Android/AndroidPlatformStackWalk.cpp b/Engine/Source/Runtime/Core/Private/Android/AndroidPlatformStackWalk.cpp
--- a/Engine/Source/Runtime/Core/Private/Android/AndroidPlatformStackWalk.cpp
+++ b/Engine/Source/Runtime/Core/Private/Android/AndroidPlatformStackWalk.cpp
@@ -12,6 +12,13 @@
 #include <cxxabi.h>
 #include <stdio.h>
 
+#define HAS_LIBUNWIND PLATFORM_ANDROID_ARM64 && !PLATFORM_LUMIN
+
+#if HAS_LIBUNWIND
+#define UNW_LOCAL_ONLY
+#include "libunwind.h"
+#endif
+
 void FAndroidPlatformStackWalk::ProgramCounterToSymbolInfo(uint64 ProgramCounter, FProgramCounterSymbolInfo& out_SymbolInfo)
 {
 	Dl_info DylibInfo;
@@ -95,6 +102,38 @@ namespace AndroidStackWalkHelpers
 	}
 }
 
+#if HAS_LIBUNWIND
+// code based on unw_backtrace using signal context for the walk, note that this code was originally intended to walk the current stack.
+// Since it walks a signal context it includes the first frame.
+static int backtrace_signal(void* sigcontext, void **buffer, int size)
+{
+	unw_cursor_t cursor;
+	unw_word_t ip;
+	int n = 0;
+
+	if (unw_init_local2(&cursor, (unw_context_t *)sigcontext, 1) < 0)
+	{
+		return 0;
+	}
+
+	do
+	{
+		if (n >= size)
+		{
+			return n;
+		}
+
+		if (unw_get_reg(&cursor, UNW_REG_IP, &ip) < 0)
+		{
+			return n;
+		}
+		buffer[n++] = (void *)(uintptr_t)ip;
+	} while (unw_step(&cursor) > 0);
+
+	return n;
+}
+#endif
+
 extern int32 unwind_backtrace_signal(void* sigcontext, uint64* Backtrace, int32 MaxDepth);
 
 uint32 FAndroidPlatformStackWalk::CaptureStackBackTrace(uint64* BackTrace, uint32 MaxDepth, void* Context)
@@ -117,8 +156,20 @@ uint32 FAndroidPlatformStackWalk::CaptureStackBackTrace(uint64* BackTrace, uint3
 		// Code taken from https://android.googlesource.com/platform/system/core/+/jb-dev/libcorkscrew/arch-arm/backtrace-arm.c
 		return unwind_backtrace_signal(Context, BackTrace, MaxDepth);
 	}
-#endif //PLATFORM_ANDROID_ARM
-	
+#elif HAS_LIBUNWIND 
+	if (Context)
+	{
+		// Android signal handlers always catch signals before user handlers and passes it down to user later
+		// unw_backtrace does not use signal context and will produce wrong callstack in this case
+		// We use code from libunwind to unwind backtrace using actual signal context
+		return backtrace_signal(Context, (void**)BackTrace, MaxDepth);
+	}
+	else
+	{
+		return unw_backtrace((void**)BackTrace, MaxDepth);
+	}
+#endif 
+
 	AndroidStackWalkHelpers::BackTrace = BackTrace;
 	AndroidStackWalkHelpers::MaxDepth = MaxDepth;
 	uint32 Depth = 0;
diff --git a/Engine/Source/Runtime/Core/Public/Android/AndroidPlatformCrashContext.h b/Engine/Source/Runtime/Core/Public/Android/AndroidPlatformCrashContext.h
--- a/Engine/Source/Runtime/Core/Public/Android/AndroidPlatformCrashContext.h
+++ b/Engine/Source/Runtime/Core/Public/Android/AndroidPlatformCrashContext.h
@@ -3,6 +3,7 @@
 #pragma once
 
 #include "GenericPlatform/GenericPlatformCrashContext.h"
+#include "Containers/Map.h"
 
 struct CORE_API FAndroidCrashContext : public FGenericCrashContext
 {
@@ -15,13 +16,7 @@ struct CORE_API FAndroidCrashContext : public FGenericCrashContext
 	/** Thread context */
 	void* Context;
 
-	FAndroidCrashContext(ECrashContextType InType, const TCHAR* InErrorMessage)
-		:	FGenericCrashContext(InType, InErrorMessage)
-		,	Signal(0)
-		,	Info(NULL)
-		,	Context(NULL)
-	{
-	}
+	FAndroidCrashContext(ECrashContextType InType, const TCHAR* InErrorMessage);
 
 	~FAndroidCrashContext()
 	{
@@ -40,6 +35,27 @@ struct CORE_API FAndroidCrashContext : public FGenericCrashContext
 		Info = InInfo;
 		Context = InContext;
 	}
+
+	virtual void GetPortableCallStack(const uint64* StackFrames, int32 NumStackFrames, TArray<FCrashStackFrame>& OutCallStack) override;
+	virtual void AddPlatformSpecificProperties() const override;
+
+	void CaptureCrashInfo();
+	void StoreCrashInfo() const;
+
+	static const int32 CrashReportMaxPathSize = 512;
+
+	static void Initialize();
+	static const FString GetCrashDirectoryName();
+	static void GetCrashDirectoryName(char(&DirectoryNameOUT)[CrashReportMaxPathSize]);
+
+	void AddAndroidCrashProperty(const FString& Key, const FString& Value);
+
+	// generate an absolute path to a crash report folder.
+	static void GenerateReportDirectoryName(char(&DirectoryNameOUT)[CrashReportMaxPathSize]);
+
+private:
+	TMap<FString, FString> AdditionalProperties;
+
 };
 
 typedef FAndroidCrashContext FPlatformCrashContext;
\ No newline at end of file
diff --git a/Engine/Source/Runtime/Core/Public/Android/AndroidPlatformMisc.h b/Engine/Source/Runtime/Core/Public/Android/AndroidPlatformMisc.h
--- a/Engine/Source/Runtime/Core/Public/Android/AndroidPlatformMisc.h
+++ b/Engine/Source/Runtime/Core/Public/Android/AndroidPlatformMisc.h
@@ -211,6 +211,10 @@ public:
 	typedef TFunction<void()> OnPauseCallBackType;
 	static OnPauseCallBackType GetOnPauseCallback();
 	static void SetOnPauseCallback(OnPauseCallBackType InOnPauseCallback);
+	static void TriggerCrashHandler(const TCHAR* InErrorMessage, const TMap<FString, FString>& AdditionalProperties);
+	static void TriggerNonFatalCrashHandler(enum class ECrashContextType InType, const FString& Message);
+
+	static bool IsInSignalHandler();
 
 #if !UE_BUILD_SHIPPING
 	static bool IsDebuggerPresent();
diff --git a/Engine/Source/Runtime/CrashReportCore/Private/Android/AndroidErrorReport.cpp b/Engine/Source/Runtime/CrashReportCore/Private/Android/AndroidErrorReport.cpp
new file mode 100644
--- /dev/null
+++ b/Engine/Source/Runtime/CrashReportCore/Private/Android/AndroidErrorReport.cpp
@@ -0,0 +1,30 @@
+// Copyright 1998-2019 Epic Games, Inc. All Rights Reserved.
+
+#include "Android/AndroidErrorReport.h"
+#include "../CrashReportUtil.h"
+#include "CrashReportCoreConfig.h"
+#include "Modules/ModuleManager.h"
+#include "CrashDebugHelperModule.h"
+#include "CrashDebugHelper.h"
+#include "Misc/FileHelper.h"
+#include "HAL/PlatformFilemanager.h"
+
+#define LOCTEXT_NAMESPACE "CrashReport"
+
+namespace
+{
+	/** Pointer to dynamically loaded crash diagnosis module */
+	FCrashDebugHelperModule* CrashHelperModule;
+}
+
+void FAndroidErrorReport::Init()
+{
+	CrashHelperModule = &FModuleManager::LoadModuleChecked<FCrashDebugHelperModule>(FName("CrashDebugHelper"));
+}
+
+void FAndroidErrorReport::ShutDown()
+{
+	CrashHelperModule->ShutdownModule();
+}
+
+#undef LOCTEXT_NAMESPACE
diff --git a/Engine/Source/Runtime/CrashReportCore/Public/Android/AndroidErrorReport.h b/Engine/Source/Runtime/CrashReportCore/Public/Android/AndroidErrorReport.h
new file mode 100644
--- /dev/null
+++ b/Engine/Source/Runtime/CrashReportCore/Public/Android/AndroidErrorReport.h
@@ -0,0 +1,70 @@
+// Copyright 1998-2019 Epic Games, Inc. All Rights Reserved.
+
+#pragma once
+
+#include "CoreMinimal.h"
+#include "Containers/UnrealString.h"
+#include "Internationalization/Text.h"
+#include "GenericErrorReport.h"
+
+struct FTimespan;
+
+/**
+ * Helper that works with Android Error Reports
+ */
+class FAndroidErrorReport : public FGenericErrorReport
+{
+public:
+	/**
+	 * Default constructor: creates a report with no files
+	 */
+	FAndroidErrorReport()
+	{
+	}
+
+	/**
+	 * Discover all files in the crash report directory
+	 * @param Directory Full path to directory containing the report
+	 */
+	explicit FAndroidErrorReport(const FString& Directory)
+		: FGenericErrorReport(Directory)
+	{
+	}
+
+	/**
+	 * Load helper modules
+	 */
+	static void Init();
+
+	/**
+	 * Unload helper modules
+	 */
+	static void ShutDown();
+
+	/**
+	 * Do nothing - shouldn't be called on Android
+	 * @return Dummy text
+	 */
+	FText DiagnoseReport() const	
+	{
+		return FText::FromString("No local diagnosis on Android");
+	}
+
+	/**
+	 * Do nothing - shouldn't be called on Android
+	 * @return Empty string
+	 */
+	static void FindMostRecentErrorReports(TArray<FString>& ErrorReportPaths, const FTimespan& MaxCrashReportAge)
+	{
+		// The report folder is currently always sent on the command-line on Android
+	}
+
+	/**
+	 * Get the full path of the crashed app from the report
+	 */
+	FString FindCrashedAppPath() const
+	{
+		FString AppPath = FPaths::Combine(FPrimaryCrashProperties::Get()->BaseDir, FPrimaryCrashProperties::Get()->ExecutableName);
+		return AppPath;
+	}
+};
diff --git a/Engine/Source/Runtime/CrashReportCore/Public/PlatformErrorReport.h b/Engine/Source/Runtime/CrashReportCore/Public/PlatformErrorReport.h
--- a/Engine/Source/Runtime/CrashReportCore/Public/PlatformErrorReport.h
+++ b/Engine/Source/Runtime/CrashReportCore/Public/PlatformErrorReport.h
@@ -28,6 +28,12 @@ typedef FMacErrorReport FPlatformErrorReport;
 
 typedef FIOSErrorReport FPlatformErrorReport;
 
+#elif PLATFORM_ANDROID
+
+#include "Android/AndroidErrorReport.h"
+
+typedef FAndroidErrorReport FPlatformErrorReport;
+
 #else
 
 typedef FGenericErrorReport FPlatformErrorReport;
