Replicated TArray::AddUninitialized() changes from CL# 6631168.

diff --git a/Engine/Source/Runtime/Core/Public/Containers/Array.h b/Engine/Source/Runtime/Core/Public/Containers/Array.h
--- a/Engine/Source/Runtime/Core/Public/Containers/Array.h
+++ b/Engine/Source/Runtime/Core/Public/Containers/Array.h
@@ -1198,6 +1198,28 @@ public:
 		Ar.CountBytes(ArrayNum*sizeof(ElementType), ArrayMax*sizeof(ElementType));
 	}
 
+	/**
+	 * Adds an uninitialized element into the array.
+	 *
+	 * Caution, AddUninitialized() will create elements without calling
+	 * the constructor and this is not appropriate for element types that
+	 * require a constructor to function properly.
+	 *
+	 * @returns Number of elements in array before addition.
+	 */
+	FORCEINLINE SizeType AddUninitialized()
+	{
+		CheckInvariants();
+
+		const SizeType OldNum = ArrayNum;
+		const SizeType NewNum = ArrayNum += 1;
+		if (NewNum > ArrayMax)
+		{
+			ResizeGrow(OldNum);
+		}
+		return OldNum;
+	}
+
 	/**
 	 * Adds a given number of uninitialized elements into the array.
 	 *
@@ -1208,15 +1230,25 @@ public:
 	 * @param Count Number of elements to add.
 	 * @returns Number of elements in array before addition.
 	 */
-	FORCEINLINE SizeType AddUninitialized(SizeType Count = 1)
+	FORCEINLINE SizeType AddUninitialized(SizeType Count)
 	{
 		CheckInvariants();
 		checkSlow(Count >= 0);
 
 		const SizeType OldNum = ArrayNum;
-		if ((ArrayNum += Count) > ArrayMax)
+		const SizeType NewNum = ArrayNum + Count;
+		if (OldNum == 0)
 		{
-			ResizeGrow(OldNum);
+			Reserve(Count);
+			ArrayNum = NewNum;
+		}
+		else
+		{
+			ArrayNum = NewNum;
+			if (NewNum > ArrayMax)
+			{
+				ResizeGrow(OldNum);
+			}
 		}
 		return OldNum;
 	}
@@ -1854,7 +1886,7 @@ public:
 	template <typename... ArgsType>
 	FORCEINLINE SizeType Emplace(ArgsType&&... Args)
 	{
-		const SizeType Index = AddUninitialized(1);
+		const SizeType Index = AddUninitialized();
 		new(GetData() + Index) ElementType(Forward<ArgsType>(Args)...);
 		return Index;
 	}
@@ -1868,7 +1900,7 @@ public:
 	template <typename... ArgsType>
 	FORCEINLINE ElementType& Emplace_GetRef(ArgsType&&... Args)
 	{
-		const SizeType Index = AddUninitialized(1);
+		const SizeType Index = AddUninitialized();
 		ElementType* Ptr = GetData() + Index;
 		new(Ptr) ElementType(Forward<ArgsType>(Args)...);
 		return *Ptr;
@@ -1991,7 +2023,7 @@ public:
 	 */
 	ElementType& AddZeroed_GetRef()
 	{
-		const SizeType Index = AddUninitialized(1);
+		const SizeType Index = AddUninitialized();
 		ElementType* Ptr = GetData() + Index;
 		FMemory::Memzero(Ptr, sizeof(ElementType));
 		return *Ptr;
@@ -2021,7 +2053,7 @@ public:
 	 */
 	ElementType& AddDefaulted_GetRef()
 	{
-		const SizeType Index = AddUninitialized(1);
+		const SizeType Index = AddUninitialized();
 		ElementType* Ptr = GetData() + Index;
 		DefaultConstructItems<ElementType>(Ptr, 1);
 		return *Ptr;
@@ -2863,7 +2895,7 @@ template <typename InElementType, typename Allocator> struct TIsTArray<const vol
 template <typename T,typename Allocator> void* operator new( size_t Size, TArray<T,Allocator>& Array )
 {
 	check(Size == sizeof(T));
-	const auto Index = Array.AddUninitialized(1);
+	const auto Index = Array.AddUninitialized();
 	return &Array[Index];
 }
 template <typename T,typename Allocator> void* operator new( size_t Size, TArray<T,Allocator>& Array, typename TArray<T, Allocator>::SizeType Index )
