Fix missing AndroidRelativeToAbsolutePath

diff --git a/Engine/Source/Runtime/Core/Private/Lumin/LuminPlatformFile.cpp b/Engine/Source/Runtime/Core/Private/Lumin/LuminPlatformFile.cpp
--- a/Engine/Source/Runtime/Core/Private/Lumin/LuminPlatformFile.cpp
+++ b/Engine/Source/Runtime/Core/Private/Lumin/LuminPlatformFile.cpp
@@ -21,6 +21,34 @@ FLuminFileInfo::FLuminFileInfo()
 // make an FTimeSpan object that represents the "epoch" for time_t (from a stat struct)
 const FDateTime UnixEpoch(1970, 1, 1);
 
+FString AndroidRelativeToAbsolutePath(bool bUseInternalBasePath, FString RelPath)
+{
+	FString Result = MoveTemp(RelPath);
+
+	if (Result.StartsWith(TEXT("../")))
+	{
+		while (Result.StartsWith(TEXT("../")))
+		{
+			Result = Result.RightChop(3);
+		}
+	}
+
+	// Remove the base app path if present, we will prepend it the correct base path as needed.
+	Result.ReplaceInline(*FLuminPlatformMisc::GetApplicationPackageDirectoryPath(), TEXT(""));
+	// Remove the writable path if present, we will prepend it the correct base path as needed.
+	Result.ReplaceInline(*FLuminPlatformMisc::GetApplicationWritableDirectoryPath(), TEXT(""));
+
+	// Then add it to the app writable directory path.
+	FString lhs = FLuminPlatformMisc::GetApplicationWritableDirectoryPath();
+	FString rhs = MoveTemp(Result);
+	lhs.RemoveFromEnd(TEXT("/"));
+	rhs.RemoveFromStart(TEXT("/"));
+	Result = lhs / rhs;
+
+	// always use lower case ... always
+	return Result.ToLower();
+}
+
 namespace
 {
 	FFileStatData UnixStatToUEFileData(struct stat& FileInfo)
