Fixed potential disaster recovery service name collision on the network that could happen if two Editors instances on two different machines gets the same Editor process ID.

diff --git a/Engine/Source/Runtime/Core/Private/GenericPlatform/GenericPlatformCrashContext.cpp b/Engine/Source/Runtime/Core/Private/GenericPlatform/GenericPlatformCrashContext.cpp
--- a/Engine/Source/Runtime/Core/Private/GenericPlatform/GenericPlatformCrashContext.cpp
+++ b/Engine/Source/Runtime/Core/Private/GenericPlatform/GenericPlatformCrashContext.cpp
@@ -1131,5 +1131,7 @@ FProgramCounterSymbolInfoEx::FProgramCounterSymbolInfoEx( FString InModuleName,
 
 FString RecoveryService::GetRecoveryServerName()
 {
-	return FString::Printf(TEXT("RecoverySvr_%d"), FPlatformProcess::GetCurrentProcessId());
+	// Requirement: The name must be unique on the local machine (multiple instances) and across the local network (multiple users).
+	static FGuid RecoverySessionGuid = FGuid::NewGuid();
+	return FString::Printf(TEXT("RecoverySvr_%s"), *RecoverySessionGuid.ToString());
 }
