Removed utility method for determining EngineModeEx.

On Linux/Mac EngineModeEx was not properly set. This was because FGenericPlatformCrashContext::Initialize was called during PreInit before all the information needed had been setup. EngineMode and EngineModeEx are now updated by a callback. WindowsPlatformCrashContext was calling Initialize to support crashes during static initialization, which was refactored to only call it when actually needed.

diff --git a/Engine/Source/Runtime/Core/Private/GenericPlatform/GenericPlatformCrashContext.cpp b/Engine/Source/Runtime/Core/Private/GenericPlatform/GenericPlatformCrashContext.cpp
--- a/Engine/Source/Runtime/Core/Private/GenericPlatform/GenericPlatformCrashContext.cpp
+++ b/Engine/Source/Runtime/Core/Private/GenericPlatform/GenericPlatformCrashContext.cpp
@@ -91,9 +91,9 @@ void FGenericCrashContext::Initialize()
 	NCached::Session.ProcessId = FPlatformProcess::GetCurrentProcessId();
 
 	FCString::Strcpy(NCached::Session.GameName, *FString::Printf( TEXT("UE4-%s"), FApp::GetProjectName() ));
-	FCString::Strcpy(NCached::Session.GameSessionID, TEXT(""));
-	FCString::Strcpy(NCached::Session.GameStateName, TEXT(""));
-	FCString::Strcpy(NCached::Session.UserActivityHint, TEXT(""));
+	FCString::Strcpy(NCached::Session.GameSessionID, TEXT("")); // Updated by callback
+	FCString::Strcpy(NCached::Session.GameStateName, TEXT("")); // Updated by callback
+	FCString::Strcpy(NCached::Session.UserActivityHint, TEXT("")); // Updated by callback
 	FCString::Strcpy(NCached::Session.BuildConfigurationName, LexToString(FApp::GetBuildConfiguration()));
 	FCString::Strcpy(NCached::Session.ExecutableName, FPlatformProcess::ExecutableName());
 	FCString::Strcpy(NCached::Session.BaseDir, FPlatformProcess::BaseDir());
@@ -121,7 +121,7 @@ void FGenericCrashContext::Initialize()
 		NCached::Session.bIsUE4Release = FApp::IsEngineInstalled();
 		FCString::Strcpy(NCached::Session.CommandLine, (FCommandLine::IsInitialized() ? FCommandLine::GetOriginalForLogging() : TEXT("")));
 		FCString::Strcpy(NCached::Session.EngineMode, FGenericPlatformMisc::GetEngineMode());
-		FCString::Strcpy(NCached::Session.EngineModeEx, EngineModeExString());
+		FCString::Strcpy(NCached::Session.EngineModeEx, FGenericCrashContext::EngineModeExUnknown); // Updated from callback
 
 		// Use -epicapp value from the commandline to start. This will also be set by the game
 		FParse::Value(FCommandLine::Get(), TEXT("EPICAPP="), NCached::Session.DeploymentName, CR_MAX_GENERIC_FIELD_CHARS, true);
@@ -216,9 +216,15 @@ PRAGMA_DISABLE_DEPRECATION_WARNINGS
 	});
 PRAGMA_ENABLE_DEPRECATION_WARNINGS
 
+	FCoreDelegates::OnPostEngineInit.AddLambda([] 
+	{
+		// IsEditor global have been properly initialized here.
+		FCString::Strcpy(NCached::Session.EngineMode, FGenericPlatformMisc::GetEngineMode());
+	});
+
 	FCoreDelegates::IsVanillaProductChanged.AddLambda([](bool bIsVanilla)
 	{
-		NCached::Session.bIsVanilla = bIsVanilla;
+		FCString::Strcpy(NCached::Session.EngineModeEx, bIsVanilla ? FGenericCrashContext::EngineModeExVanilla : FGenericCrashContext::EngineModeExDirty);
 	});
 
 	FCoreDelegates::ConfigReadyForUse.AddStatic(FGenericCrashContext::InitializeFromConfig);
@@ -850,12 +856,6 @@ const TCHAR* FGenericCrashContext::GetCrashTypeString(ECrashContextType Type)
 	}
 }
 
-const TCHAR* FGenericCrashContext::EngineModeExString()
-{
-	return !NCached::Session.bIsVanilla.IsSet() ? FGenericCrashContext::EngineModeExUnknown :
-		(NCached::Session.bIsVanilla.GetValue() ? FGenericCrashContext::EngineModeExVanilla : FGenericCrashContext::EngineModeExDirty);
-}
-
 const TCHAR* FGenericCrashContext::GetCrashConfigFilePath()
 {
 	if (FCString::Strlen(NCached::Session.CrashConfigFilePath) == 0)
diff --git a/Engine/Source/Runtime/Core/Private/Windows/WindowsPlatformCrashContext.cpp b/Engine/Source/Runtime/Core/Private/Windows/WindowsPlatformCrashContext.cpp
--- a/Engine/Source/Runtime/Core/Private/Windows/WindowsPlatformCrashContext.cpp
+++ b/Engine/Source/Runtime/Core/Private/Windows/WindowsPlatformCrashContext.cpp
@@ -469,8 +469,6 @@ int32 ReportCrashForMonitor(
 	void* ReadPipe,
 	EErrorReportUI ReportUI)
 {
-	//Initialize all the static data from crash context
-	FGenericCrashContext::Initialize();
 	FGenericCrashContext::CopySharedCrashContext(*SharedContext);
 
 	// Set the platform specific crash context, so that we can stack walk and minidump from
@@ -1059,6 +1057,11 @@ public:
 			const int NumStackFramesToIgnore = 1;
 			const TCHAR* ErrorMessage = TEXT("Crash during static initialization");
 
+			if (!FPlatformCrashContext::IsInitalized())
+			{
+				FPlatformCrashContext::Initialize();
+			}
+
 			return ReportCrashForMonitor(
 				InExceptionInfo,
 				Type,
diff --git a/Engine/Source/Runtime/Core/Public/GenericPlatform/GenericPlatformCrashContext.h b/Engine/Source/Runtime/Core/Public/GenericPlatform/GenericPlatformCrashContext.h
--- a/Engine/Source/Runtime/Core/Public/GenericPlatform/GenericPlatformCrashContext.h
+++ b/Engine/Source/Runtime/Core/Public/GenericPlatform/GenericPlatformCrashContext.h
@@ -165,7 +165,6 @@ struct FSessionContext
 	TCHAR 					GameStateName[CR_MAX_GENERIC_FIELD_CHARS];
 	TCHAR 					CrashConfigFilePath[CR_MAX_DIRECTORY_CHARS];
 	FPlatformMemoryStats	MemoryStats;
-	TOptional<bool>			bIsVanilla;
 };
 
 /** Additional user settings to be communicated to crash reporting client. */
@@ -340,9 +339,6 @@ public:
 	/** Get the Game Name of the crash */
 	static FString GetCrashGameName();
 
-	/** Gets the "vanilla" status string. */
-	static const TCHAR* EngineModeExString();
-
 	/** Helper to get the crash report client config filepath saved by this instance and copied to each crash report folder. */
 	static const TCHAR* GetCrashConfigFilePath();
 
