Enable FPlatformMallocCrash with android crash handler.
#rnx

diff --git a/Engine/Source/Runtime/Core/Private/Android/AndroidPlatformMisc.cpp b/Engine/Source/Runtime/Core/Private/Android/AndroidPlatformMisc.cpp
--- a/Engine/Source/Runtime/Core/Private/Android/AndroidPlatformMisc.cpp
+++ b/Engine/Source/Runtime/Core/Private/Android/AndroidPlatformMisc.cpp
@@ -1046,7 +1046,7 @@ void FAndroidMisc::TriggerCrashHandler(const TCHAR* InErrorMessage, const TMap<F
 void PlatformCrashHandler()
 {
 	// Switch to malloc crash.
-	//FGenericPlatformMallocCrash::Get().SetAsGMalloc(); @todo uncomment after verification
+	FPlatformMallocCrash::Get().SetAsGMalloc();
 
 	FAndroidCrashContext CrashContext(ECrashContextType::Crash, TEXT("Caught signal"));
 	CrashContext.InitFromSignal(GAndroidSignalParams.Signal, GAndroidSignalParams.Info, GAndroidSignalParams.Context);
