Clamp Android crash context's ItoANSI base parameter to within supported values.
#rnx

diff --git a/Engine/Source/Runtime/Core/Private/Android/AndroidPlatformCrashContext.cpp b/Engine/Source/Runtime/Core/Private/Android/AndroidPlatformCrashContext.cpp
--- a/Engine/Source/Runtime/Core/Private/Android/AndroidPlatformCrashContext.cpp
+++ b/Engine/Source/Runtime/Core/Private/Android/AndroidPlatformCrashContext.cpp
@@ -79,7 +79,7 @@ static ANSICHAR* ItoANSI(uint64 Val, uint64 Base, uint32 Len = 0)
 
 	uint64 i = 62;
 	int32 pad = Len;
-
+	Base = FMath::Clamp<uint64>(Base, 2, 16);
 	if (Val)
 	{
 		for (; Val && i; --i, Val /= Base, --pad)
@@ -245,15 +245,6 @@ void FAndroidCrashContext::AddPlatformSpecificProperties() const
 
 void FAndroidCrashContext::GetPortableCallStack(const uint64* StackFrames, int32 NumStackFrames, TArray<FCrashStackFrame>& OutCallStack)
 {
-	// Get all the modules in the current process
-	uint32 NumModules = (uint32)FPlatformStackWalk::GetProcessModuleCount();
-
-	TArray<FStackWalkModuleInfo> Modules;
-	Modules.AddUninitialized(NumModules);
-
-	NumModules = FPlatformStackWalk::GetProcessModuleSignatures(Modules.GetData(), NumModules);
-	Modules.SetNum(NumModules);
-
 	// Update the callstack with offsets from each module
 	OutCallStack.Reset(NumStackFrames);
 	for (int32 Idx = 0; Idx < NumStackFrames; Idx++)
