Made FBitReader's default copy constructor use a const ref.

#rnx

diff --git a/Engine/Source/Runtime/Core/Public/Serialization/BitReader.h b/Engine/Source/Runtime/Core/Public/Serialization/BitReader.h
--- a/Engine/Source/Runtime/Core/Public/Serialization/BitReader.h
+++ b/Engine/Source/Runtime/Core/Public/Serialization/BitReader.h
@@ -24,7 +24,7 @@ struct CORE_API FBitReader : public FBitArchive
 public:
 	FBitReader( uint8* Src = nullptr, int64 CountBits = 0 );
 
-    FBitReader(FBitReader&) = default;
+	FBitReader(const FBitReader&) = default;
     FBitReader& operator=(const FBitReader&) = default;
     FBitReader(FBitReader&&) = default;
     FBitReader& operator=(FBitReader&&) = default;
