Static analysis fix.

diff --git a/Engine/Source/Runtime/Core/Private/HAL/MallocBinned.cpp b/Engine/Source/Runtime/Core/Private/HAL/MallocBinned.cpp
--- a/Engine/Source/Runtime/Core/Private/HAL/MallocBinned.cpp
+++ b/Engine/Source/Runtime/Core/Private/HAL/MallocBinned.cpp
@@ -449,19 +449,18 @@ struct FMallocBinned::Private
 	/**
 	* Initializes tables for HashBuckets if they haven't already been initialized.
 	*/
-	static FORCEINLINE void InitializeHashBuckets(FMallocBinned& Allocator)
+	static FORCEINLINE PoolHashBucket* CreateHashBuckets(const FMallocBinned& Allocator)
 	{
-		if (!Allocator.HashBuckets)
-		{
-			// Init tables.
-            LLM_PLATFORM_SCOPE(ELLMTag::FMalloc);
-			Allocator.HashBuckets = (PoolHashBucket*)FPlatformMemory::BinnedAllocFromOS(Align(Allocator.MaxHashBuckets * sizeof(PoolHashBucket), Allocator.PageSize));
+		// Init tables.
+	`	LLM_PLATFORM_SCOPE(ELLMTag::FMalloc);
+		PoolHashBucket* Result = (PoolHashBucket*)FPlatformMemory::BinnedAllocFromOS(Align(Allocator.MaxHashBuckets * sizeof(PoolHashBucket), Allocator.PageSize));
 
-			for (uint32 i = 0; i < Allocator.MaxHashBuckets; ++i)
-			{
-				new (Allocator.HashBuckets + i) PoolHashBucket();
-			}
+		for (uint32 i = 0; i < Allocator.MaxHashBuckets; ++i)
+		{
+			new (Result + i) PoolHashBucket();
 		}
+
+		return Result;
 	}
 	
 	/** 
@@ -473,7 +472,7 @@ struct FMallocBinned::Private
 	{
 		if (!Allocator.HashBuckets)
 		{
-			InitializeHashBuckets(Allocator);
+			Allocator.HashBuckets = CreateHashBuckets(Allocator);
 		}
 		checkSlow(Allocator.HashBuckets);
 
@@ -1335,7 +1334,10 @@ void* FMallocBinned::Malloc(SIZE_T Size, uint32 Alignment)
 	{
 		//Make sure we have initialized our hash buckets even if we are using the NANO_MALLOC grabber, as otherwise we can end
 		//up making bad assumptions and trying to grab invalid data during a Realloc of this data.
-		Private::InitializeHashBuckets(*this);
+		if (!HashBuckets)
+		{
+			HashBuckets = Private::CreateHashBuckets(*this);
+		}
 		
 		bUsePools = false;
 		UPTRINT AlignedSize = Align(Size, Alignment);
