Typo fix.

diff --git a/Engine/Source/Runtime/Core/Private/HAL/MallocBinned.cpp b/Engine/Source/Runtime/Core/Private/HAL/MallocBinned.cpp
--- a/Engine/Source/Runtime/Core/Private/HAL/MallocBinned.cpp
+++ b/Engine/Source/Runtime/Core/Private/HAL/MallocBinned.cpp
@@ -452,7 +452,7 @@ struct FMallocBinned::Private
 	static FORCEINLINE PoolHashBucket* CreateHashBuckets(const FMallocBinned& Allocator)
 	{
 		// Init tables.
-	`	LLM_PLATFORM_SCOPE(ELLMTag::FMalloc);
+		LLM_PLATFORM_SCOPE(ELLMTag::FMalloc);
 		PoolHashBucket* Result = (PoolHashBucket*)FPlatformMemory::BinnedAllocFromOS(Align(Allocator.MaxHashBuckets * sizeof(PoolHashBucket), Allocator.PageSize));
 
 		for (uint32 i = 0; i < Allocator.MaxHashBuckets; ++i)
