Add inline versions of Left, LeftChop, Right, RightChop, Mid, and ConvertTabsToSpaces
Fix Mid to gracefully handle negative count and negative starts.
Fix RemoveAt to not leave an unterminated string.

diff --git a/Engine/Source/Runtime/Core/Private/Containers/String.cpp b/Engine/Source/Runtime/Core/Private/Containers/String.cpp
--- a/Engine/Source/Runtime/Core/Private/Containers/String.cpp
+++ b/Engine/Source/Runtime/Core/Private/Containers/String.cpp
@@ -1248,36 +1248,32 @@ FString FString::ReplaceEscapedCharWithChar( const TArray<TCHAR>* Chars/*=nullpt
  * Replaces all instances of '\t' with TabWidth number of spaces
  * @param InSpacesPerTab - Number of spaces that a tab represents
  */
-FString FString::ConvertTabsToSpaces (const int32 InSpacesPerTab)
+void FString::ConvertTabsToSpacesInline(const int32 InSpacesPerTab)
 {
 	//must call this with at least 1 space so the modulus operation works
 	check(InSpacesPerTab > 0);
 
-	FString FinalString = *this;
 	int32 TabIndex;
-	while ((TabIndex = FinalString.Find(TEXT("\t"))) != INDEX_NONE )
+	while ((TabIndex = Find(TEXT("\t"), ESearchCase::CaseSensitive)) != INDEX_NONE )
 	{
-		FString LeftSide = FinalString.Left(TabIndex);
-		FString RightSide = FinalString.Mid(TabIndex+1);
+		FString RightSide = Mid(TabIndex+1);
+		LeftInline(TabIndex, false);
 
-		FinalString = LeftSide;
 		//for a tab size of 4, 
-		int32 LineBegin = LeftSide.Find(TEXT("\n"), ESearchCase::IgnoreCase, ESearchDir::FromEnd, TabIndex);
+		int32 LineBegin = Find(TEXT("\n"), ESearchCase::CaseSensitive, ESearchDir::FromEnd, TabIndex);
 		if (LineBegin == INDEX_NONE)
 		{
 			LineBegin = 0;
 		}
-		int32 CharactersOnLine = (LeftSide.Len()-LineBegin);
+		const int32 CharactersOnLine = (Len()-LineBegin);
 
 		int32 NumSpacesForTab = InSpacesPerTab - (CharactersOnLine % InSpacesPerTab);
 		for (int32 i = 0; i < NumSpacesForTab; ++i)
 		{
-			FinalString.AppendChar(' ');
+			AppendChar(' ');
 		}
-		FinalString += RightSide;
+		Append(RightSide);
 	}
-
-	return FinalString;
 }
 
 // This starting size catches 99.97% of printf calls - there are about 700k printf calls per level
diff --git a/Engine/Source/Runtime/Core/Private/Containers/StringTest.cpp b/Engine/Source/Runtime/Core/Private/Containers/StringTest.cpp
--- a/Engine/Source/Runtime/Core/Private/Containers/StringTest.cpp
+++ b/Engine/Source/Runtime/Core/Private/Containers/StringTest.cpp
@@ -238,4 +238,70 @@ bool FLexTryParseStringTest::RunTest(const FString& Parameters)
 	return true;
 }
 
+IMPLEMENT_SIMPLE_AUTOMATION_TEST(FStringSubstringTest, "System.Core.String.Substring", EAutomationTestFlags::ApplicationContextMask | EAutomationTestFlags::SmokeFilter)
+bool FStringSubstringTest::RunTest(const FString& Parameters)
+{
+	auto DoTest = [this](const TCHAR* Call, const FString& Result, const TCHAR* InExpected)
+	{
+		if (!Result.Equals(InExpected, ESearchCase::CaseSensitive))
+		{
+			AddError(FString::Printf(TEXT("'%s' failure: result '%s' (expected '%s')"), Call, *Result, InExpected));
+		}
+	};
+
+	const FString TestString(TEXT("0123456789"));
+
+#define SUBSTRINGTEST(TestName, ExpectedResult, Operation, ...) \
+	FString TestName = TestString.Operation(__VA_ARGS__); \
+	DoTest(TEXT(#TestName), TestName, ExpectedResult); \
+	\
+	FString Inline##TestName = TestString; \
+	Inline##TestName.Operation##Inline(__VA_ARGS__); \
+	DoTest(TEXT("Inline" #TestName), Inline##TestName, ExpectedResult); 
+
+	// Left
+	SUBSTRINGTEST(Left, TEXT("0123"), Left, 4);
+	SUBSTRINGTEST(ExactLengthLeft, *TestString, Left, 10);
+	SUBSTRINGTEST(LongerThanLeft, *TestString, Left, 20);
+	SUBSTRINGTEST(ZeroLeft, TEXT(""), Left, 0);
+	SUBSTRINGTEST(NegativeLeft, TEXT(""), Left, -1);
+
+	// LeftChop
+	SUBSTRINGTEST(LeftChop, TEXT("012345"), LeftChop, 4);
+	SUBSTRINGTEST(ExactLengthLeftChop, TEXT(""), LeftChop, 10);
+	SUBSTRINGTEST(LongerThanLeftChop, TEXT(""), LeftChop, 20);
+	SUBSTRINGTEST(ZeroLeftChop, *TestString, LeftChop, 0);
+	SUBSTRINGTEST(NegativeLeftChop, *TestString, LeftChop, -1);
+
+	// Right
+	SUBSTRINGTEST(Right, TEXT("6789"), Right, 4);
+	SUBSTRINGTEST(ExactLengthRight, *TestString, Right, 10);
+	SUBSTRINGTEST(LongerThanRight, *TestString, Right, 20);
+	SUBSTRINGTEST(ZeroRight, TEXT(""), Right, 0);
+	SUBSTRINGTEST(NegativeRight, TEXT(""), Right, -1);
+
+	// RightChop
+	SUBSTRINGTEST(RightChop, TEXT("456789"), RightChop, 4);
+	SUBSTRINGTEST(ExactLengthRightChop, TEXT(""), RightChop, 10);
+	SUBSTRINGTEST(LongerThanRightChop, TEXT(""), RightChop, 20);
+	SUBSTRINGTEST(ZeroRightChop, *TestString, RightChop, 0);
+	SUBSTRINGTEST(NegativeRightChop, *TestString, RightChop, -1);
+
+	// Mid
+	SUBSTRINGTEST(Mid, TEXT("456789"), Mid, 4);
+	SUBSTRINGTEST(MidCount, TEXT("4567"), Mid, 4, 4);
+	SUBSTRINGTEST(MidCountFullLength, *TestString, Mid, 0, 10);
+	SUBSTRINGTEST(MidCountOffEnd, TEXT("89"), Mid, 8, 4);
+	SUBSTRINGTEST(MidStartAfterEnd, TEXT(""), Mid, 20);
+	SUBSTRINGTEST(MidZeroCount, TEXT(""), Mid, 5, 0);
+	SUBSTRINGTEST(MidNegativeCount, TEXT(""), Mid, 5, -1);
+	SUBSTRINGTEST(MidNegativeStartNegativeEnd, TEXT(""), Mid, -5, 1);
+	SUBSTRINGTEST(MidNegativeStartPositiveEnd, TEXT("012"), Mid, -1, 4);
+	SUBSTRINGTEST(MidNegativeStartBeyondEnd, *TestString, Mid, -1, 15);
+
+#undef SUBSTRINGTEST
+
+	return true;
+}
+
 #endif // WITH_DEV_AUTOMATION_TESTS
diff --git a/Engine/Source/Runtime/Core/Public/Containers/UnrealString.h b/Engine/Source/Runtime/Core/Public/Containers/UnrealString.h
--- a/Engine/Source/Runtime/Core/Public/Containers/UnrealString.h
+++ b/Engine/Source/Runtime/Core/Public/Containers/UnrealString.h
@@ -481,7 +481,7 @@ public:
 	 */
 	FORCEINLINE void RemoveAt(int32 Index, int32 Count = 1, bool bAllowShrinking = true)
 	{
-		Data.RemoveAt(Index, Count, bAllowShrinking);
+		Data.RemoveAt(Index, FMath::Clamp(Count, 0, Len()-Index), bAllowShrinking);
 	}
 
 	FORCEINLINE void InsertAt(int32 Index, TCHAR Character)
@@ -1147,33 +1147,75 @@ public:
 		return FString( FMath::Clamp(Count,0,Len()), **this );
 	}
 
+	/** Modifies the string such that it is now the left most given number of characters */
+	FORCEINLINE void LeftInline(int32 Count, bool bAllowShrinking = true)
+	{
+		const int32 Length = Len();
+		Count = FMath::Clamp(Count, 0, Length);
+		RemoveAt(Count, Length-Count, bAllowShrinking);
+	}
+
 	/** Returns the left most characters from the string chopping the given number of characters from the end */
 	FORCEINLINE FString LeftChop( int32 Count ) const
 	{
-		return FString( FMath::Clamp(Len()-Count,0,Len()), **this );
+		const int32 Length = Len();
+		return FString( FMath::Clamp(Length-Count,0, Length), **this );
+	}
+
+	/** Modifies the string such that it is now the left most characters chopping the given number of characters from the end */
+	FORCEINLINE void LeftChopInline(int32 Count, bool bAllowShrinking = true)
+	{
+		const int32 Length = Len();
+		RemoveAt(FMath::Clamp(Length-Count, 0, Length), Count, bAllowShrinking);
 	}
 
 	/** Returns the string to the right of the specified location, counting back from the right (end of the word). */
 	FORCEINLINE FString Right( int32 Count ) const
 	{
-		return FString( **this + Len()-FMath::Clamp(Count,0,Len()) );
+		const int32 Length = Len();
+		return FString( **this + Length-FMath::Clamp(Count,0,Length) );
+	}
+
+	/** Modifies the string such that it is now the right most given number of characters */
+	FORCEINLINE void RightInline(int32 Count, bool bAllowShrinking = true)
+	{
+		const int32 Length = Len();
+		RemoveAt(0, Length-FMath::Clamp(Count,0,Length), bAllowShrinking);
 	}
 
 	/** Returns the string to the right of the specified location, counting forward from the left (from the beginning of the word). */
 	FORCEINLINE FString RightChop( int32 Count ) const
 	{
-		return FString( **this + Len()-FMath::Clamp(Len()-Count,0,Len()) );
+		const int32 Length = Len();
+		return FString( **this + Length-FMath::Clamp(Length-Count,0, Length) );
+	}
+
+	/** Modifies the string such that it is now the string to the right of the specified location, counting forward from the left (from the beginning of the word). */
+	FORCEINLINE void RightChopInline(int32 Count, bool bAllowShrinking = true)
+	{
+		RemoveAt(0, Count, bAllowShrinking);
 	}
 
 	/** Returns the substring from Start position for Count characters. */
 	FORCEINLINE FString Mid( int32 Start, int32 Count=MAX_int32 ) const
 	{
-		check(Count >= 0);
-		uint32 End = Count;
-		End += Start;
-		Start    = FMath::Clamp( (uint32)Start, (uint32)0,     (uint32)Len() );
-		End      = FMath::Clamp( (uint32)End,   (uint32)Start, (uint32)Len() );
-		return FString( End-Start, **this + Start );
+		FString Result;
+		if (Count >= 0)
+		{
+			const int32 Length = Len();
+			const int32 RequestedStart = Start;
+			Start = FMath::Clamp(Start, 0, Length);
+			const int32 End = (int32)FMath::Clamp((int64)Count + RequestedStart, (int64)Start, (int64)Length);
+			Result = FString(End-Start, **this + Start);
+		}
+		return Result;
+	}
+
+	/** Modifies the string such that it is now the substring from Start position for Count characters. */
+	FORCEINLINE void MidInline(int32 Start, int32 Count = MAX_int32, bool bAllowShrinking = true)
+	{
+		LeftInline((int32)FMath::Min((int64)Count+Start, (int64)MAX_int32), false);
+		RightChopInline(Start, bAllowShrinking);
 	}
 
 	/**
@@ -1183,7 +1225,7 @@ public:
 	 * @param SubStr			The string array of TCHAR to search for
 	 * @param StartPosition		The start character position to search from
 	 * @param SearchCase		Indicates whether the search is case sensitive or not
-	 * @param SearchDir			Indicates whether the search starts at the begining or at the end.
+	 * @param SearchDir			Indicates whether the search starts at the beginning or at the end.
 	 */
 	int32 Find( const TCHAR* SubStr, ESearchCase::Type SearchCase = ESearchCase::IgnoreCase, 
 				ESearchDir::Type SearchDir = ESearchDir::FromStart, int32 StartPosition=INDEX_NONE ) const;
@@ -1195,7 +1237,7 @@ public:
 	 * @param SubStr			The string to search for
 	 * @param StartPosition		The start character position to search from
 	 * @param SearchCase		Indicates whether the search is case sensitive or not ( defaults to ESearchCase::IgnoreCase )
-	 * @param SearchDir			Indicates whether the search starts at the begining or at the end ( defaults to ESearchDir::FromStart )
+	 * @param SearchDir			Indicates whether the search starts at the beginning or at the end ( defaults to ESearchDir::FromStart )
 	 */
 	FORCEINLINE int32 Find( const FString& SubStr, ESearchCase::Type SearchCase = ESearchCase::IgnoreCase, 
 							ESearchDir::Type SearchDir = ESearchDir::FromStart, int32 StartPosition=INDEX_NONE ) const
@@ -1208,7 +1250,7 @@ public:
 	 *
 	 * @param SubStr			Find to search for
 	 * @param SearchCase		Indicates whether the search is case sensitive or not ( defaults to ESearchCase::IgnoreCase )
-	 * @param SearchDir			Indicates whether the search starts at the begining or at the end ( defaults to ESearchDir::FromStart )
+	 * @param SearchDir			Indicates whether the search starts at the beginning or at the end ( defaults to ESearchDir::FromStart )
 	 * @return					Returns whether the string contains the substring
 	 **/
 	FORCEINLINE bool Contains(const TCHAR* SubStr, ESearchCase::Type SearchCase = ESearchCase::IgnoreCase, 
@@ -1222,7 +1264,7 @@ public:
 	 *
 	 * @param SubStr			Find to search for
 	 * @param SearchCase		Indicates whether the search is case sensitive or not ( defaults to ESearchCase::IgnoreCase )
-	 * @param SearchDir			Indicates whether the search starts at the begining or at the end ( defaults to ESearchDir::FromStart )
+	 * @param SearchDir			Indicates whether the search starts at the beginning or at the end ( defaults to ESearchDir::FromStart )
 	 * @return					Returns whether the string contains the substring
 	 **/
 	FORCEINLINE bool Contains(const FString& SubStr, ESearchCase::Type SearchCase = ESearchCase::IgnoreCase, 
@@ -1329,7 +1371,7 @@ public:
 	 * @param LeftS out the string to the left of InStr, not updated if return is false
 	 * @param RightS out the string to the right of InStr, not updated if return is false
 	 * @param SearchCase		Indicates whether the search is case sensitive or not ( defaults to ESearchCase::IgnoreCase )
-	 * @param SearchDir			Indicates whether the search starts at the begining or at the end ( defaults to ESearchDir::FromStart )
+	 * @param SearchDir			Indicates whether the search starts at the beginning or at the end ( defaults to ESearchDir::FromStart )
 	 * @return true if string is split, otherwise false
 	 */
 	bool Split(const FString& InS, FString* LeftS, FString* RightS, ESearchCase::Type SearchCase = ESearchCase::IgnoreCase,
@@ -1730,12 +1772,23 @@ public:
 	 */
 	FString ReplaceEscapedCharWithChar( const TArray<TCHAR>* Chars = nullptr ) const;
 
+	/**
+	 * Replaces all instances of '\t' with TabWidth number of spaces
+	 * @param InSpacesPerTab - Number of spaces that a tab represents
+	 */
+	void ConvertTabsToSpacesInline(const int32 InSpacesPerTab);
+
 	/** 
 	 * Replaces all instances of '\t' with TabWidth number of spaces
 	 * @param InSpacesPerTab - Number of spaces that a tab represents
 	 * @return copy of this string with replacement made
 	 */
-	FString ConvertTabsToSpaces(const int32 InSpacesPerTab);
+	FString ConvertTabsToSpaces(const int32 InSpacesPerTab) const
+	{
+		FString FinalString(*this);
+		FinalString.ConvertTabsToSpacesInline(InSpacesPerTab);
+		return FinalString;
+	}
 
 	// Takes the number passed in and formats the string in comma format ( 12345 becomes "12,345")
 	static FString FormatAsNumber( int32 InNumber );
