Use new inline versions of substring functions

diff --git a/Engine/Source/Editor/UnrealEd/Private/PlayLevel.cpp b/Engine/Source/Editor/UnrealEd/Private/PlayLevel.cpp
--- a/Engine/Source/Editor/UnrealEd/Private/PlayLevel.cpp
+++ b/Engine/Source/Editor/UnrealEd/Private/PlayLevel.cpp
@@ -1605,22 +1605,24 @@ void UEditorEngine::HandleStageStarted(const FString& InStage, TWeakPtr<SNotific
 	FText NotificationText;
 	if (InStage.Contains(TEXT("Cooking")) || InStage.Contains(TEXT("Cook Task")))
 	{
-		FString PlatformName = PlayUsingLauncherDeviceId.Left(PlayUsingLauncherDeviceId.Find(TEXT("@")));
-		if (PlatformName.Contains(TEXT("NoEditor")))
+		FString PlatformName = PlayUsingLauncherDeviceId.Left(PlayUsingLauncherDeviceId.Find(TEXT("@"), ESearchCase::CaseSensitive));
+		const int32 NoEditorIdx = PlatformName.Find(TEXT("NoEditor"));
+		if (NoEditorIdx != INDEX_NONE)
 		{
-			PlatformName = PlatformName.Left(PlatformName.Find(TEXT("NoEditor")));
+			PlatformName.LeftInline(NoEditorIdx, false);
 		}
-		Arguments.Add(TEXT("PlatformName"), FText::FromString(PlatformName));
+		Arguments.Add(TEXT("PlatformName"), FText::FromString(MoveTemp(PlatformName)));
 		NotificationText = FText::Format(LOCTEXT("LauncherTaskProcessingNotification", "Processing Assets for {PlatformName}..."), Arguments);
 	}
 	else if (InStage.Contains(TEXT("Build Task")))
 	{
-		FString PlatformName = PlayUsingLauncherDeviceId.Left(PlayUsingLauncherDeviceId.Find(TEXT("@")));
-		if (PlatformName.Contains(TEXT("NoEditor")))
+		FString PlatformName = PlayUsingLauncherDeviceId.Left(PlayUsingLauncherDeviceId.Find(TEXT("@"), ESearchCase::CaseSensitive));
+		const int32 NoEditorIdx = PlatformName.Find(TEXT("NoEditor"));
+		if (NoEditorIdx != INDEX_NONE)
 		{
-			PlatformName = PlatformName.Left(PlatformName.Find(TEXT("NoEditor")));
+			PlatformName.LeftInline(NoEditorIdx, false);
 		}
-		Arguments.Add(TEXT("PlatformName"), FText::FromString(PlatformName));
+		Arguments.Add(TEXT("PlatformName"), FText::FromString(MoveTemp(PlatformName)));
 		if (!bPlayUsingLauncherBuild)
 		{
 			NotificationText = FText::Format(LOCTEXT("LauncherTaskValidateNotification", "Validating Executable for {PlatformName}..."), Arguments);
diff --git a/Engine/Source/Editor/UnrealEd/Private/SkeletalMeshEdit.cpp b/Engine/Source/Editor/UnrealEd/Private/SkeletalMeshEdit.cpp
--- a/Engine/Source/Editor/UnrealEd/Private/SkeletalMeshEdit.cpp
+++ b/Engine/Source/Editor/UnrealEd/Private/SkeletalMeshEdit.cpp
@@ -1393,7 +1393,7 @@ bool UnFbx::FFbxImporter::ImportAnimation(USkeleton* Skeleton, UAnimSequence * D
 							// Also avoid to endup with a empty name, we prefer having the Blendshapename instead of nothing
 							if(ChannelName.StartsWith(BlendShapeName) && ChannelName.Len() > BlendShapeName.Len())
 							{
-								ChannelName = ChannelName.Right(ChannelName.Len() - (BlendShapeName.Len()+1));
+								ChannelName.RightInline(ChannelName.Len() - (BlendShapeName.Len()+1), false);
 							}
 							
 							if (bMightBeBadMAXFile)
diff --git a/Engine/Source/Editor/UnrealEd/Private/SourceCodeNavigation.cpp b/Engine/Source/Editor/UnrealEd/Private/SourceCodeNavigation.cpp
--- a/Engine/Source/Editor/UnrealEd/Private/SourceCodeNavigation.cpp
+++ b/Engine/Source/Editor/UnrealEd/Private/SourceCodeNavigation.cpp
@@ -613,7 +613,7 @@ void FSourceCodeNavigationImpl::NavigateToFunctionSource( const FString& Functio
 								int32 ArgumentIndex = -1;
 								if(SymbolName.FindLastChar(TCHAR('('), ArgumentIndex))
 								{
-									SymbolName = SymbolName.Left(ArgumentIndex);
+									SymbolName.LeftInline(ArgumentIndex, false);
 									int32 TemplateNesting = 0;
 									
 									int32 Pos = SymbolName.Len();
@@ -630,7 +630,7 @@ void FSourceCodeNavigationImpl::NavigateToFunctionSource( const FString& Functio
 										TCHAR Character = SymbolName[Pos - 1];
 										if(Character == TCHAR(' ') && TemplateNesting == 0)
 										{
-											SymbolName = SymbolName.Mid(Pos);
+											SymbolName.MidInline(Pos, MAX_int32, false);
 											break;
 										}
 										else if(Character == TCHAR('>'))
@@ -1082,7 +1082,7 @@ void FSourceCodeNavigationImpl::GatherFunctions( const FString& ModuleName, cons
 									int32 ArgumentIndex = -1;
 									if(FunctionSymbolName.FindLastChar(TCHAR('('), ArgumentIndex))
 									{
-										FunctionSymbolName = FunctionSymbolName.Left(ArgumentIndex);
+										FunctionSymbolName.LeftInline(ArgumentIndex, false);
 										int32 TemplateNesting = 0;
 										
 										int32 Pos = FunctionSymbolName.Len();
@@ -1099,7 +1099,7 @@ void FSourceCodeNavigationImpl::GatherFunctions( const FString& ModuleName, cons
 											TCHAR Character = FunctionSymbolName[Pos - 1];
 											if(Character == TCHAR(' ') && TemplateNesting == 0)
 											{
-												FunctionSymbolName = FunctionSymbolName.Mid(Pos);
+												FunctionSymbolName.MidInline(Pos, MAX_int32, false);
 												break;
 											}
 											else if(Character == TCHAR('>'))
diff --git a/Engine/Source/Editor/UnrealEd/Private/UnrealEdSrv.cpp b/Engine/Source/Editor/UnrealEd/Private/UnrealEdSrv.cpp
--- a/Engine/Source/Editor/UnrealEd/Private/UnrealEdSrv.cpp
+++ b/Engine/Source/Editor/UnrealEd/Private/UnrealEdSrv.cpp
@@ -1270,9 +1270,9 @@ bool UUnrealEdEngine::Exec( UWorld* InWorld, const TCHAR* Stream, FOutputDevice&
 						FString CurrentWindowName = OpenWindows[WindowId]->GetTitle().ToString();
 
 						//Strip off the * from the end if it exists
-						if( CurrentWindowName.EndsWith(TEXT("*")) )
+						if( CurrentWindowName.EndsWith(TEXT("*"), ESearchCase::CaseSensitive) )
 						{
-							CurrentWindowName = CurrentWindowName.LeftChop(1);
+							CurrentWindowName.LeftChopInline(1, false);
 						}
 
 						if( CurrentWindowName == WindowNameStr )
diff --git a/Engine/Source/Editor/WorldBrowser/Private/SPropertyEditorLevelPackage.cpp b/Engine/Source/Editor/WorldBrowser/Private/SPropertyEditorLevelPackage.cpp
--- a/Engine/Source/Editor/WorldBrowser/Private/SPropertyEditorLevelPackage.cpp
+++ b/Engine/Source/Editor/WorldBrowser/Private/SPropertyEditorLevelPackage.cpp
@@ -56,7 +56,7 @@ FText SPropertyEditorLevelPackage::GetDisplayPackageName() const
 			FString LongPackageName = PropertyValue.ToString();
 			if (LongPackageName.StartsWith(RootPath))
 			{
-				LongPackageName = LongPackageName.RightChop(RootPath.Len()-1); // do not chop front '/' from display name
+				LongPackageName.RightChopInline(RootPath.Len()-1, false); // do not chop front '/' from display name
 				return FText::FromString(LongPackageName);
 			}
 		}
diff --git a/Engine/Source/Editor/WorldBrowser/Private/Tiles/WorldTileCollectionModel.cpp b/Engine/Source/Editor/WorldBrowser/Private/Tiles/WorldTileCollectionModel.cpp
--- a/Engine/Source/Editor/WorldBrowser/Private/Tiles/WorldTileCollectionModel.cpp
+++ b/Engine/Source/Editor/WorldBrowser/Private/Tiles/WorldTileCollectionModel.cpp
@@ -1608,7 +1608,7 @@ void FWorldTileCollectionModel::ImportTiledLandscape_Executed()
 		// Extract tile prefix
 		FString FolderName = FPaths::GetBaseFilename(ImportSettings.HeightmapFileList[0]);
 		int32 PrefixEnd = FolderName.Find(TEXT("_x"), ESearchCase::IgnoreCase, ESearchDir::FromEnd);
-		FolderName = FolderName.Left(PrefixEnd);
+		FolderName.LeftInline(PrefixEnd, false);
 		WorldRootPath+= FolderName;
 		WorldRootPath+= TEXT("/");
 
diff --git a/Engine/Source/Runtime/AIModule/Private/BehaviorTree/BTCompositeNode.cpp b/Engine/Source/Runtime/AIModule/Private/BehaviorTree/BTCompositeNode.cpp
--- a/Engine/Source/Runtime/AIModule/Private/BehaviorTree/BTCompositeNode.cpp
+++ b/Engine/Source/Runtime/AIModule/Private/BehaviorTree/BTCompositeNode.cpp
@@ -411,7 +411,7 @@ static bool UpdateOperationStack(const UBehaviorTreeComponent& OwnerComp, FStrin
 		UE_VLOG(OwnerComp.GetOwner(), LogBehaviorTree, Verbose, TEXT("%s%s finished: %s"), *Indent,
 			*DescribeLogicOp(CurrentOp.Op),
 			bTestResult ? TEXT("allowed") : TEXT("forbidden"));
-		Indent = Indent.LeftChop(2);
+		Indent.LeftChopInline(2, false);
 
 		Stack.RemoveAt(Stack.Num() - 1);
 		return UpdateOperationStack(OwnerComp, Indent, Stack, bTestResult, FailedDecoratorIdx, NodeDecoratorIdx, bShouldStoreNodeIndex);
diff --git a/Engine/Source/Runtime/AIModule/Private/BehaviorTree/BehaviorTreeTypes.cpp b/Engine/Source/Runtime/AIModule/Private/BehaviorTree/BehaviorTreeTypes.cpp
--- a/Engine/Source/Runtime/AIModule/Private/BehaviorTree/BehaviorTreeTypes.cpp
+++ b/Engine/Source/Runtime/AIModule/Private/BehaviorTree/BehaviorTreeTypes.cpp
@@ -479,10 +479,10 @@ FString UBehaviorTreeTypes::GetShortTypeName(const UObject* Ob)
 	}
 
 	FString TypeDesc = Ob->GetClass()->GetName();
-	const int32 ShortNameIdx = TypeDesc.Find(TEXT("_"));
+	const int32 ShortNameIdx = TypeDesc.Find(TEXT("_"), ESearchCase::CaseSensitive);
 	if (ShortNameIdx != INDEX_NONE)
 	{
-		TypeDesc = TypeDesc.Mid(ShortNameIdx + 1);
+		TypeDesc.MidInline(ShortNameIdx + 1, MAX_int32, false);
 	}
 
 	return TypeDesc;
diff --git a/Engine/Source/Runtime/AIModule/Private/BehaviorTree/BlueprintNodeHelpers.cpp b/Engine/Source/Runtime/AIModule/Private/BehaviorTree/BlueprintNodeHelpers.cpp
--- a/Engine/Source/Runtime/AIModule/Private/BehaviorTree/BlueprintNodeHelpers.cpp
+++ b/Engine/Source/Runtime/AIModule/Private/BehaviorTree/BlueprintNodeHelpers.cpp
@@ -108,7 +108,7 @@ namespace BlueprintNodeHelpers
 			const FString CategoryLimit = StructProp->GetMetaData(TEXT("Categories"));
 			if (!CategoryLimit.IsEmpty() && ExportedStringValue.StartsWith(CategoryLimit))
 			{
-				ExportedStringValue = ExportedStringValue.Mid(CategoryLimit.Len());
+				ExportedStringValue.MidInline(CategoryLimit.Len(), MAX_int32, false);
 			}
 #endif
 		}
diff --git a/Engine/Source/Runtime/AIModule/Private/DataProviders/AIDataProvider.cpp b/Engine/Source/Runtime/AIModule/Private/DataProviders/AIDataProvider.cpp
--- a/Engine/Source/Runtime/AIModule/Private/DataProviders/AIDataProvider.cpp
+++ b/Engine/Source/Runtime/AIModule/Private/DataProviders/AIDataProvider.cpp
@@ -153,7 +153,7 @@ FString UAIDataProvider::ToString(FName PropName) const
 	const bool bFound = ProviderName.FindChar(TEXT('_'), SplitIdx);
 	if (bFound)
 	{
-		ProviderName = ProviderName.Mid(SplitIdx + 1);
+		ProviderName.MidInline(SplitIdx + 1, MAX_int32, false);
 	}
 
 	ProviderName += TEXT('.');
diff --git a/Engine/Source/Runtime/AIModule/Private/EnvironmentQuery/EnvQueryTypes.cpp b/Engine/Source/Runtime/AIModule/Private/EnvironmentQuery/EnvQueryTypes.cpp
--- a/Engine/Source/Runtime/AIModule/Private/EnvironmentQuery/EnvQueryTypes.cpp
+++ b/Engine/Source/Runtime/AIModule/Private/EnvironmentQuery/EnvQueryTypes.cpp
@@ -108,10 +108,10 @@ FText UEnvQueryTypes::GetShortTypeName(const UObject* Ob)
 	}
 
 	FString TypeDesc = ObClass->GetName();
-	const int32 ShortNameIdx = TypeDesc.Find(TEXT("_"));
+	const int32 ShortNameIdx = TypeDesc.Find(TEXT("_"), ESearchCase::CaseSensitive);
 	if (ShortNameIdx != INDEX_NONE)
 	{
-		TypeDesc = TypeDesc.Mid(ShortNameIdx + 1);
+		TypeDesc.MidInline(ShortNameIdx + 1, MAX_int32, false);
 	}
 
 	return FText::FromString(TypeDesc);
diff --git a/Engine/Source/Runtime/AIModule/Private/Perception/AISense.cpp b/Engine/Source/Runtime/AIModule/Private/Perception/AISense.cpp
--- a/Engine/Source/Runtime/AIModule/Private/Perception/AISense.cpp
+++ b/Engine/Source/Runtime/AIModule/Private/Perception/AISense.cpp
@@ -131,7 +131,7 @@ FString UAISenseConfig::GetSenseName() const
 		const bool bHasSeparator = CachedSenseName.FindLastChar(TEXT('_'), SeparatorIdx);
 		if (bHasSeparator)
 		{
-			CachedSenseName = CachedSenseName.Mid(SeparatorIdx + 1);
+			CachedSenseName.MidInline(SeparatorIdx + 1, MAX_int32, false);
 		}
 	}
 
diff --git a/Engine/Source/Runtime/Analytics/Analytics/Public/AnalyticsConversion.h b/Engine/Source/Runtime/Analytics/Analytics/Public/AnalyticsConversion.h
--- a/Engine/Source/Runtime/Analytics/Analytics/Public/AnalyticsConversion.h
+++ b/Engine/Source/Runtime/Analytics/Analytics/Public/AnalyticsConversion.h
@@ -46,7 +46,7 @@ FString AnalyticsConversionToString(const TArray<T, AllocatorType>& ValueArray)
 		Result += TEXT(",");
 	}
 	// Remove the trailing comma (LeftChop will ensure an empty container won't crash here).
-	Result = Result.LeftChop(1);
+	Result.LeftChopInline(1, false);
 	return Result;
 }
 
@@ -64,7 +64,7 @@ FString AnalyticsConversionToString(const TMap<KeyType, ValueType, Allocator, Ke
 		Result += TEXT(",");
 	}
 	// Remove the trailing comma (LeftChop will ensure an empty container won't crash here).
-	Result = Result.LeftChop(1);
+	Result.LeftChopInline(1, false);
 	return Result;
 }
 
diff --git a/Engine/Source/Runtime/AppFramework/Private/Widgets/Testing/STestSuite.cpp b/Engine/Source/Runtime/AppFramework/Private/Widgets/Testing/STestSuite.cpp
--- a/Engine/Source/Runtime/AppFramework/Private/Widgets/Testing/STestSuite.cpp
+++ b/Engine/Source/Runtime/AppFramework/Private/Widgets/Testing/STestSuite.cpp
@@ -3487,7 +3487,7 @@ protected:
 		{
 			if(FChar::IsLinebreak(SelectedText[SelectedTextIndex]))
 			{
-				SelectedText = SelectedText.Left(SelectedTextIndex);
+				SelectedText.LeftInline(SelectedTextIndex, false);
 				break;
 			}
 		}
diff --git a/Engine/Source/Runtime/ApplicationCore/Private/IOS/IOSAppDelegate.cpp b/Engine/Source/Runtime/ApplicationCore/Private/IOS/IOSAppDelegate.cpp
--- a/Engine/Source/Runtime/ApplicationCore/Private/IOS/IOSAppDelegate.cpp
+++ b/Engine/Source/Runtime/ApplicationCore/Private/IOS/IOSAppDelegate.cpp
@@ -1365,10 +1365,10 @@ extern EDeviceScreenOrientation ConvertFromUIInterfaceOrientation(UIInterfaceOri
 	// "MyGame://arg1 arg2 arg3 ..."
 	// So, we're going to make it look like:
 	// "arg1 arg2 arg3 ..."
-	int32 URLTerminator = CommandLineParameters.Find( TEXT("://"));
+	int32 URLTerminator = CommandLineParameters.Find( TEXT("://"), ESearchCase::CaseSensitive);
 	if ( URLTerminator > -1 )
 	{
-		CommandLineParameters = CommandLineParameters.RightChop(URLTerminator + 3);
+		CommandLineParameters.RightChopInline(URLTerminator + 3, false);
 	}
 
 	FIOSCommandLineHelper::InitCommandArgs(CommandLineParameters);
diff --git a/Engine/Source/Runtime/ApplicationCore/Private/Windows/WindowsApplication.cpp b/Engine/Source/Runtime/ApplicationCore/Private/Windows/WindowsApplication.cpp
--- a/Engine/Source/Runtime/ApplicationCore/Private/Windows/WindowsApplication.cpp
+++ b/Engine/Source/Runtime/ApplicationCore/Private/Windows/WindowsApplication.cpp
@@ -585,7 +585,7 @@ inline bool GetSizeForDevID(const FString& TargetDevID, int32& Width, int32& Hei
 			if (CM_Get_Device_ID(DevInfoData.DevInst, Buffer, MAX_DEVICE_ID_LEN, 0) == CR_SUCCESS)
 			{
 				FString DevID(Buffer);
-				DevID = DevID.Mid(8, DevID.Find(TEXT("\\"), ESearchCase::CaseSensitive, ESearchDir::FromStart, 9) - 8);
+				DevID.MidInline(8, DevID.Find(TEXT("\\"), ESearchCase::CaseSensitive, ESearchDir::FromStart, 9) - 8, false);
 				if (DevID == TargetDevID)
 				{
 					HKEY hDevRegKey = SetupDiOpenDevRegKey(DevInfo, &DevInfoData, DICS_FLAG_GLOBAL, 0, DIREG_DEV, KEY_READ);
diff --git a/Engine/Source/Runtime/AssetRegistry/Private/AssetRegistry.cpp b/Engine/Source/Runtime/AssetRegistry/Private/AssetRegistry.cpp
--- a/Engine/Source/Runtime/AssetRegistry/Private/AssetRegistry.cpp
+++ b/Engine/Source/Runtime/AssetRegistry/Private/AssetRegistry.cpp
@@ -337,11 +337,11 @@ void UAssetRegistryImpl::InitializeSerializationOptionsFromIni(FAssetRegistrySer
 		TrimmedFilterlistItem.TrimStartAndEndInline();
 		if (TrimmedFilterlistItem.Left(1) == TEXT("("))
 		{
-			TrimmedFilterlistItem = TrimmedFilterlistItem.RightChop(1);
+			TrimmedFilterlistItem.RightChopInline(1, false);
 		}
 		if (TrimmedFilterlistItem.Right(1) == TEXT(")"))
 		{
-			TrimmedFilterlistItem = TrimmedFilterlistItem.LeftChop(1);
+			TrimmedFilterlistItem.LeftChopInline(1, false);
 		}
 
 		TArray<FString> Tokens;
@@ -2702,10 +2702,10 @@ void UAssetRegistryImpl::OnContentPathDismounted(const FString& InAssetPath, con
 {
 	// Sanitize
 	FString AssetPath = InAssetPath;
-	if (AssetPath.EndsWith(TEXT("/")))
+	if (AssetPath.EndsWith(TEXT("/"), ESearchCase::CaseSensitive))
 	{
 		// We don't want a trailing slash here as it could interfere with RemoveAssetPath
-		AssetPath = AssetPath.LeftChop(1);
+		AssetPath.LeftChopInline(1, false);
 	}
 
 	// Remove all cached assets found at this location
diff --git a/Engine/Source/Runtime/AutomationWorker/Private/AutomationWorkerModule.cpp b/Engine/Source/Runtime/AutomationWorker/Private/AutomationWorkerModule.cpp
--- a/Engine/Source/Runtime/AutomationWorker/Private/AutomationWorkerModule.cpp
+++ b/Engine/Source/Runtime/AutomationWorker/Private/AutomationWorkerModule.cpp
@@ -486,7 +486,7 @@ void FAutomationWorkerModule::SendAnalyticsEvents(TArray<FString>& InAnalyticsIt
 		if( EventString.EndsWith( TEXT( ",PERF" ) ) )
 		{
 			// Chop the ",PERF" off the end
-			EventString = EventString.Left( EventString.Len() - 5 );
+			EventString.LeftInline( EventString.Len() - 5, false );
 
 			FAutomationPerformanceSnapshot PerfSnapshot;
 			PerfSnapshot.FromCommaDelimitedString( EventString );
diff --git a/Engine/Source/Runtime/Core/Private/Containers/StackTracker.cpp b/Engine/Source/Runtime/Core/Private/Containers/StackTracker.cpp
--- a/Engine/Source/Runtime/Core/Private/Containers/StackTracker.cpp
+++ b/Engine/Source/Runtime/Core/Private/Containers/StackTracker.cpp
@@ -54,10 +54,10 @@ void FStackTracker::CaptureStackTrace(int32 EntriesToIgnore, void* UserData, int
 						AddressInformation[0] = 0;
 						FPlatformStackWalk::ProgramCounterToHumanReadableString( 1, BackTrace[Index], AddressInformation, UE_ARRAY_COUNT(AddressInformation)-1 );
 						FString Symbol(AddressInformation);
-						int32 Spot = Symbol.Find(TEXT(" - "));
+						int32 Spot = Symbol.Find(TEXT(" - "), ESearchCase::CaseSensitive);
 						if (Spot != INDEX_NONE)
 						{
-							Symbol = Symbol.RightChop(Spot + 3);
+							Symbol.RightChopInline(Spot + 3, false);
 						}
 						Existing = StringAliasMap.Find(Symbol);
 						if (Existing)
diff --git a/Engine/Source/Runtime/Core/Private/GenericPlatform/GenericPlatformMisc.cpp b/Engine/Source/Runtime/Core/Private/GenericPlatform/GenericPlatformMisc.cpp
--- a/Engine/Source/Runtime/Core/Private/GenericPlatform/GenericPlatformMisc.cpp
+++ b/Engine/Source/Runtime/Core/Private/GenericPlatform/GenericPlatformMisc.cpp
@@ -757,7 +757,7 @@ const TCHAR* FGenericPlatformMisc::RootDir()
 		int32 chopPos = TempPath.Find(TEXT("/Engine"), ESearchCase::IgnoreCase, ESearchDir::FromEnd);
 		if (chopPos != INDEX_NONE)
 		{
-			TempPath = TempPath.Left(chopPos + 1);
+			TempPath.LeftInline(chopPos + 1, false);
 		}
 		else
 		{
@@ -766,7 +766,7 @@ const TCHAR* FGenericPlatformMisc::RootDir()
 			// if the path ends in a separator, remove it
 			if (TempPath.Right(1) == TEXT("/"))
 			{
-				TempPath = TempPath.LeftChop(1);
+				TempPath.LeftChopInline(1, false);
 			}
 
 			// keep going until we've removed Binaries
@@ -777,7 +777,7 @@ const TCHAR* FGenericPlatformMisc::RootDir()
 #endif
 			if (pos != INDEX_NONE)
 			{
-				TempPath = TempPath.Left(pos + 1);
+				TempPath.LeftInline(pos + 1, false);
 			}
 			else
 			{
@@ -790,7 +790,7 @@ const TCHAR* FGenericPlatformMisc::RootDir()
 				{
 					while (TempPath.Len() && TempPath.Right(1) != TEXT("/"))
 					{
-						TempPath = TempPath.LeftChop(1);
+						TempPath.LeftChopInline(1, false);
 					}
 				}
 			}
diff --git a/Engine/Source/Runtime/Core/Private/HAL/ConsoleManager.cpp b/Engine/Source/Runtime/Core/Private/HAL/ConsoleManager.cpp
--- a/Engine/Source/Runtime/Core/Private/HAL/ConsoleManager.cpp
+++ b/Engine/Source/Runtime/Core/Private/HAL/ConsoleManager.cpp
@@ -1335,7 +1335,7 @@ bool FConsoleManager::ProcessUserConsoleInput(const TCHAR* InInput, FOutputDevic
 	const bool bCommandEndedInQuestion = Param1.EndsWith(TEXT("?"), ESearchCase::CaseSensitive);
 	if (bCommandEndedInQuestion)
 	{
-		Param1 = Param1.Mid(0, Param1.Len() - 1);
+		Param1.MidInline(0, Param1.Len() - 1, false);
 	}
 
 	IConsoleObject* CObj = FindConsoleObject(*Param1);
@@ -1402,13 +1402,13 @@ bool FConsoleManager::ProcessUserConsoleInput(const TCHAR* InInput, FOutputDevic
 			{
 				if(Param2[0] == (TCHAR)'\"' && Param2[Param2.Len() - 1] == (TCHAR)'\"')
 				{
-					Param2 = Param2.Mid(1, Param2.Len() - 2);
+					Param2.MidInline(1, Param2.Len() - 2, false);
 				}
 				// this is assumed to be unintended e.g. copy and paste accident from ini file
 				if(Param2.Len() > 0 && Param2[0] == (TCHAR)'=')
 				{
 					Ar.Logf(TEXT("Warning: Processing the console input parameters the leading '=' is ignored (only needed for ini files)."));
-					Param2 = Param2.Mid(1, Param2.Len() - 1);
+					Param2.MidInline(1, Param2.Len() - 1, false);
 				}
 			}
 
diff --git a/Engine/Source/Runtime/Core/Private/HAL/FileManagerGeneric.cpp b/Engine/Source/Runtime/Core/Private/HAL/FileManagerGeneric.cpp
--- a/Engine/Source/Runtime/Core/Private/HAL/FileManagerGeneric.cpp
+++ b/Engine/Source/Runtime/Core/Private/HAL/FileManagerGeneric.cpp
@@ -580,7 +580,7 @@ FString FFileManagerGeneric::DefaultConvertToRelativePath( const TCHAR* Filename
 			//move up a directory and on an extra .. TEXT("/")
 			// the +1 from "InStr" moves to include the "\" at the end of the directory name
 			NumberOfDirectoriesToGoUp++;
-			RootDirectory = RootDirectory.Left( PositionOfNextSlash + 1 );
+			RootDirectory.LeftInline( PositionOfNextSlash + 1, false );
 		}
 		else
 		{
diff --git a/Engine/Source/Runtime/Core/Private/HAL/LowLevelMemTracker.cpp b/Engine/Source/Runtime/Core/Private/HAL/LowLevelMemTracker.cpp
--- a/Engine/Source/Runtime/Core/Private/HAL/LowLevelMemTracker.cpp
+++ b/Engine/Source/Runtime/Core/Private/HAL/LowLevelMemTracker.cpp
@@ -2082,14 +2082,14 @@ FString FLLMCsvWriter::GetTagName(int64 Tag, FLLMCustomTag* CustomTags, const in
 		FString Name = TagToFName(Tag).ToString();
 
 		// if it has a trible slash assume it is a Stat string and extract the descriptive name
-		int32 StartIndex = Name.Find(TEXT("///"));
+		int32 StartIndex = Name.Find(TEXT("///"), ESearchCase::CaseSensitive);
 		if (StartIndex != -1)
 		{
 			StartIndex += 3;
-			int32 EndIndex = Name.Find(TEXT("///"), ESearchCase::IgnoreCase, ESearchDir::FromStart, StartIndex);
+			int32 EndIndex = Name.Find(TEXT("///"), ESearchCase::CaseSensitive, ESearchDir::FromStart, StartIndex);
 			if (EndIndex != -1)
 			{
-				Name = Name.Mid(StartIndex, EndIndex - StartIndex);
+				Name.MidInline(StartIndex, EndIndex - StartIndex, false);
 			}
 		}
 
diff --git a/Engine/Source/Runtime/Core/Private/Logging/LogSuppressionInterface.cpp b/Engine/Source/Runtime/Core/Private/Logging/LogSuppressionInterface.cpp
--- a/Engine/Source/Runtime/Core/Private/Logging/LogSuppressionInterface.cpp
+++ b/Engine/Source/Runtime/Core/Private/Logging/LogSuppressionInterface.cpp
@@ -136,7 +136,7 @@ class FLogSuppressionImplementation: public FLogSuppressionInterface, private FS
 			FString Command = SubCmds[Index].TrimStart();
 			if (Command.StartsWith(LogString))
 			{
-				Command = Command.Right(Command.Len() - (UE_ARRAY_COUNT(LogString) - 1));
+				Command.RightInline(Command.Len() - (UE_ARRAY_COUNT(LogString) - 1), false);
 			}
 			TArray<FString> CommandParts;
 			Command.ParseIntoArrayWS(CommandParts);
@@ -518,10 +518,10 @@ class FLogSuppressionImplementation: public FLogSuppressionInterface, private FS
 				{
 					break;
 				}
-				CmdLineEnv = CmdLineEnv.Mid(Index + LogCmds.Len());
+				CmdLineEnv.MidInline(Index + LogCmds.Len(), MAX_int32, false);
 			}
 			// now strip off the environment arg part
-			CmdLine = CmdLine.Mid(0, IndexOfEnv);
+			CmdLine.MidInline(0, IndexOfEnv, false);
 		}
 		while (1)
 		{
@@ -538,7 +538,7 @@ class FLogSuppressionImplementation: public FLogSuppressionInterface, private FS
 			{
 				break;
 			}
-			CmdLine = CmdLine.Mid(Index + LogCmds.Len());
+			CmdLine.MidInline(Index + LogCmds.Len(), MAX_int32, false);
 		}
 #endif // !UE_BUILD_SHIPPING
 
diff --git a/Engine/Source/Runtime/Core/Private/Lumin/LuminPlatformFile.cpp b/Engine/Source/Runtime/Core/Private/Lumin/LuminPlatformFile.cpp
--- a/Engine/Source/Runtime/Core/Private/Lumin/LuminPlatformFile.cpp
+++ b/Engine/Source/Runtime/Core/Private/Lumin/LuminPlatformFile.cpp
@@ -25,12 +25,12 @@ FString AndroidRelativeToAbsolutePath(bool bUseInternalBasePath, FString RelPath
 {
 	FString Result = MoveTemp(RelPath);
 
-	if (Result.StartsWith(TEXT("../")))
+	if (Result.StartsWith(TEXT("../"), ESearchCase::CaseSensitive))
 	{
-		while (Result.StartsWith(TEXT("../")))
+		do 
 		{
-			Result = Result.RightChop(3);
-		}
+			Result.RightChopInline(3, false);
+		} while (Result.StartsWith(TEXT("../"), ESearchCase::CaseSensitive))
 	}
 
 	// Remove the base app path if present, we will prepend it the correct base path as needed.
@@ -41,8 +41,8 @@ FString AndroidRelativeToAbsolutePath(bool bUseInternalBasePath, FString RelPath
 	// Then add it to the app writable directory path.
 	FString lhs = FLuminPlatformMisc::GetApplicationWritableDirectoryPath();
 	FString rhs = MoveTemp(Result);
-	lhs.RemoveFromEnd(TEXT("/"));
-	rhs.RemoveFromStart(TEXT("/"));
+	lhs.RemoveFromEnd(TEXT("/"), ESearchCase::CaseSensitive);
+	rhs.RemoveFromStart(TEXT("/"), ESearchCase::CaseSensitive);
 	Result = lhs / rhs;
 
 	// always use lower case ... always
diff --git a/Engine/Source/Runtime/Core/Private/Math/UnrealMath.cpp b/Engine/Source/Runtime/Core/Private/Math/UnrealMath.cpp
--- a/Engine/Source/Runtime/Core/Private/Math/UnrealMath.cpp
+++ b/Engine/Source/Runtime/Core/Private/Math/UnrealMath.cpp
@@ -2947,7 +2947,7 @@ FString FMath::FormatIntToHumanReadable(int32 Val)
 	while (Src.Len() > 3 && Src[Src.Len() - 4] != TEXT('-'))
 	{
 		Dst = FString::Printf(TEXT(",%s%s"), *Src.Right(3), *Dst);
-		Src = Src.Left(Src.Len() - 3);
+		Src.LeftInline(Src.Len() - 3, false);
 	}
 
 	Dst = Src + Dst;
diff --git a/Engine/Source/Runtime/Core/Private/Misc/AutomationTest.cpp b/Engine/Source/Runtime/Core/Private/Misc/AutomationTest.cpp
--- a/Engine/Source/Runtime/Core/Private/Misc/AutomationTest.cpp
+++ b/Engine/Source/Runtime/Core/Private/Misc/AutomationTest.cpp
@@ -67,7 +67,7 @@ void FAutomationTestFramework::FAutomationTestFeedbackContext::Serialize( const
 			if (LogString.StartsWith(*AnalyticsString))
 			{
 				//Remove "analytics" from the string
-				LogString = LogString.Right(LogString.Len() - (AnalyticsString.Len() + 1));
+				LogString.RightInline(LogString.Len() - (AnalyticsString.Len() + 1), false);
 
 				CurTest->AddAnalyticsItem(LogString);
 			}
diff --git a/Engine/Source/Runtime/Core/Private/Misc/ConfigCacheIni.cpp b/Engine/Source/Runtime/Core/Private/Misc/ConfigCacheIni.cpp
--- a/Engine/Source/Runtime/Core/Private/Misc/ConfigCacheIni.cpp
+++ b/Engine/Source/Runtime/Core/Private/Misc/ConfigCacheIni.cpp
@@ -3980,9 +3980,9 @@ private:
 	void ClearTrailingWhitespace(FString& InStr)
 	{
 		const FString Endl(LINE_TERMINATOR);
-		while (InStr.EndsWith(LINE_TERMINATOR))
+		while (InStr.EndsWith(LINE_TERMINATOR, ESearchCase::CaseSensitive))
 		{
-			InStr = InStr.LeftChop(Endl.Len());
+			InStr.LeftChopInline(Endl.Len(), false);
 		}
 	}
 	
diff --git a/Engine/Source/Runtime/Core/Private/Misc/FeedbackContextMarkup.cpp b/Engine/Source/Runtime/Core/Private/Misc/FeedbackContextMarkup.cpp
--- a/Engine/Source/Runtime/Core/Private/Misc/FeedbackContextMarkup.cpp
+++ b/Engine/Source/Runtime/Core/Private/Misc/FeedbackContextMarkup.cpp
@@ -68,14 +68,14 @@ bool FFeedbackContextMarkup::PipeProcessOutput(const FText& Description, const F
 			while(BufferedText.FindChar('\n', EndOfLineIdx))
 			{
 				FString Line = BufferedText.Left(EndOfLineIdx);
-				Line.RemoveFromEnd(TEXT("\r"));
+				Line.RemoveFromEnd(TEXT("\r"), ESearchCase::CaseSensitive);
 
 				if(!ParseCommand(Line, Warn))
 				{
 					Warn->Log(*Line);
 				}
 
-				BufferedText = BufferedText.Mid(EndOfLineIdx + 1);
+				BufferedText.MidInline(EndOfLineIdx + 1, MAX_int32, false);
 			}
 
 			FPlatformProcess::Sleep(0.1f);
diff --git a/Engine/Source/Runtime/Core/Private/Misc/MonitoredProcess.cpp b/Engine/Source/Runtime/Core/Private/Misc/MonitoredProcess.cpp
--- a/Engine/Source/Runtime/Core/Private/Misc/MonitoredProcess.cpp
+++ b/Engine/Source/Runtime/Core/Private/Misc/MonitoredProcess.cpp
@@ -121,7 +121,7 @@ void FMonitoredProcess::ProcessOutput( const FString& Output )
 	}
 
 	// Remove all the complete lines from the buffer
-	OutputBuffer = OutputBuffer.Mid(LineStartIdx);
+	OutputBuffer.MidInline(LineStartIdx, MAX_int32, false);
 }
 
 void FMonitoredProcess::TickInternal()
diff --git a/Engine/Source/Runtime/Core/Private/Misc/Paths.cpp b/Engine/Source/Runtime/Core/Private/Misc/Paths.cpp
--- a/Engine/Source/Runtime/Core/Private/Misc/Paths.cpp
+++ b/Engine/Source/Runtime/Core/Private/Misc/Paths.cpp
@@ -725,7 +725,7 @@ FString FPaths::GetBaseFilename( const FString& InPath, bool bRemovePath )
 
 	if (ExtPos != INDEX_NONE && (LeafPos == INDEX_NONE || ExtPos > LeafPos))
 	{
-		Wk = Wk.Left(ExtPos);
+		Wk.LeftInline(ExtPos);
 	}
 
 	return Wk;
@@ -903,11 +903,11 @@ bool FPaths::IsDrive(const FString& InPath)
 				int32 SlashIndex = CheckPath.Find(TEXT("\\"), ESearchCase::CaseSensitive);
 				if (SlashIndex != INDEX_NONE)
 				{
-					CheckPath = CheckPath.Right(CheckPath.Len() - SlashIndex  - 1);
+					CheckPath.RightInline(CheckPath.Len() - SlashIndex  - 1, false);
 				}
 				else
 				{
-					CheckPath = TEXT("");
+					CheckPath.Reset();
 				}
 			}
 		}
@@ -932,7 +932,7 @@ bool FPaths::IsDrive(const FString& InPath)
 					// It's a real folder, so add one to the count
 					CheckCount++;
 				}
-				CheckPath = CheckPath.Right(CheckPath.Len() - SlashIndex  - 1);
+				CheckPath.RightInline(CheckPath.Len() - SlashIndex  - 1, false);
 				SlashIndex = CheckPath.Find(TEXT("\\"), ESearchCase::CaseSensitive);
 			}
 
@@ -1308,7 +1308,7 @@ bool FPaths::ValidatePath( const FString& InPath, FText* OutReason )
 	// The loop below requires that the path not end with a /
 	if(Standardized.EndsWith(TEXT("/"), ESearchCase::CaseSensitive))
 	{
-		Standardized = Standardized.LeftChop(1);
+		Standardized.LeftChopInline(1, false);
 	}
 
 	// Walk each part of the path looking for name errors
diff --git a/Engine/Source/Runtime/Core/Private/Misc/RemoteConfigIni.cpp b/Engine/Source/Runtime/Core/Private/Misc/RemoteConfigIni.cpp
--- a/Engine/Source/Runtime/Core/Private/Misc/RemoteConfigIni.cpp
+++ b/Engine/Source/Runtime/Core/Private/Misc/RemoteConfigIni.cpp
@@ -523,7 +523,7 @@ void MakeLocalCopy(const TCHAR* Filename)
 		FString FilenameStr = Filename;
 		if (FCString::Stristr(*FilenameStr, TEXT(".ini")))
 		{
-			FilenameStr = FilenameStr.LeftChop(4);
+			FilenameStr.LeftChopInline(4, false);
 		} 
 		else
 		{
diff --git a/Engine/Source/Runtime/Core/Private/Misc/Timespan.cpp b/Engine/Source/Runtime/Core/Private/Misc/Timespan.cpp
--- a/Engine/Source/Runtime/Core/Private/Misc/Timespan.cpp
+++ b/Engine/Source/Runtime/Core/Private/Misc/Timespan.cpp
@@ -158,7 +158,7 @@ bool FTimespan::Parse(const FString& TimespanString, FTimespan& OutTimespan)
 
 		if (FractionalLen > 9)
 		{
-			Tokens[4] = Tokens[4].Left(9);
+			Tokens[4].LeftInline(9, false);
 		}
 		else if (FractionalLen < 9)
 		{
diff --git a/Engine/Source/Runtime/Core/Private/Modules/BuildVersion.cpp b/Engine/Source/Runtime/Core/Private/Modules/BuildVersion.cpp
--- a/Engine/Source/Runtime/Core/Private/Modules/BuildVersion.cpp
+++ b/Engine/Source/Runtime/Core/Private/Modules/BuildVersion.cpp
@@ -29,7 +29,7 @@ FString FBuildVersion::GetFileNameForCurrentExecutable()
 #if PLATFORM_WINDOWS || PLATFORM_MAC
 	if(AppExecutableName.EndsWith(TEXT("-Cmd")))
 	{
-		AppExecutableName = AppExecutableName.Left(AppExecutableName.Len() - 4);
+		AppExecutableName.LeftInline(AppExecutableName.Len() - 4, false);
 	}
 #endif
 #if UE_BUILD_DEVELOPMENT
@@ -43,7 +43,7 @@ FString FBuildVersion::GetFileNameForCurrentExecutable()
 			{
 				break;
 			}
-			AppExecutableName = AppExecutableName.Left(EndIdx);
+			AppExecutableName.LeftInline(EndIdx, false);
 		}
 	}
 #endif
diff --git a/Engine/Source/Runtime/Core/Private/ProfilingDebugging/CsvProfiler.cpp b/Engine/Source/Runtime/Core/Private/ProfilingDebugging/CsvProfiler.cpp
--- a/Engine/Source/Runtime/Core/Private/ProfilingDebugging/CsvProfiler.cpp
+++ b/Engine/Source/Runtime/Core/Private/ProfilingDebugging/CsvProfiler.cpp
@@ -206,7 +206,7 @@ public:
 				{
 					FString Val = ValueStr.Mid(0, CommaIndex);
 					CVarValues.Add(FCString::Atof(*Val));
-					ValueStr = ValueStr.Mid(CommaIndex + 1);
+					ValueStr.MidInline(CommaIndex + 1, MAX_int32, false);
 					continue;
 				}
 			}
@@ -223,7 +223,7 @@ public:
 				{
 					FString Val = ValueStr.Mid(0, SemiColonIndex);
 					CVarValues.Add(FCString::Atof(*Val));
-					ValueStr = ValueStr.Mid(SemiColonIndex);
+					ValueStr.MidInline(SemiColonIndex, MAX_int32, false);
 					continue;
 				}
 			}
diff --git a/Engine/Source/Runtime/Core/Private/ProfilingDebugging/ProfilingHelpers.cpp b/Engine/Source/Runtime/Core/Private/ProfilingDebugging/ProfilingHelpers.cpp
--- a/Engine/Source/Runtime/Core/Private/ProfilingDebugging/ProfilingHelpers.cpp
+++ b/Engine/Source/Runtime/Core/Private/ProfilingDebugging/ProfilingHelpers.cpp
@@ -163,7 +163,7 @@ FString CreateProfileFilename( const FString& InFilename, const FString& InFileE
 	{
 		// now create the string
 		FolderName = FString::Printf(TEXT("%s-%s-%s"), *MapNameStr, *PlatformStr, *FDateTime::Now().ToString(TEXT("%m.%d-%H.%M.%S")));
-		FolderName = FolderName.Right(MaxFilenameLen);
+		FolderName.RightInline(MaxFilenameLen, false);
 
 		ProfilingSessionFolderName = FolderName;
 		bSetProfilingSessionFolderName = true;
@@ -176,7 +176,7 @@ FString CreateProfileFilename( const FString& InFilename, const FString& InFileE
 	// now create the string
 	// NOTE: due to the changelist this is implicitly using the same directory
 	FString FolderNameOfProfileNoDate = FString::Printf( TEXT("%s-%s-%i"), *MapNameStr, *PlatformStr, GetChangeListNumberForPerfTesting() );
-	FolderNameOfProfileNoDate = FolderNameOfProfileNoDate.Right(MaxFilenameLen);
+	FolderNameOfProfileNoDate.RightInline(MaxFilenameLen, false);
 
 
 	FString NameOfProfile;
@@ -188,10 +188,10 @@ FString CreateProfileFilename( const FString& InFilename, const FString& InFileE
 	{
 		NameOfProfile = InFilename;
 	}
-	NameOfProfile = NameOfProfile.Right(MaxFilenameLen);
+	NameOfProfile.RightInline(MaxFilenameLen, false);
 
 	FString FileNameWithExtension = FString::Printf( TEXT("%s%s"), *NameOfProfile, *InFileExtension );
-	FileNameWithExtension = FileNameWithExtension.Right(MaxFilenameLen);
+	FileNameWithExtension.RightInline(MaxFilenameLen, false);
 
 	FString Filename;
 	if( bIncludeDateForDirectoryName == true )
@@ -228,7 +228,7 @@ FString CreateProfileDirectoryAndFilename( const FString& InSubDirectoryName, co
 
 	// create the directory name of this profile
 	FString NameOfProfile = FString::Printf(TEXT("%s-%s-%s"), *MapNameStr, *PlatformStr, *FDateTime::Now().ToString(TEXT("%m.%d-%H.%M")));	
-	NameOfProfile = NameOfProfile.Right(MaxFilenameLen);
+	NameOfProfile.RightInline(MaxFilenameLen, false);
 
 	IFileManager::Get().MakeDirectory( *(PathName+NameOfProfile) );
 	//UE_LOG(LogProfilingDebugging, Warning, TEXT( "CreateProfileDirectoryAndFilename: %s"), *(PathName+NameOfProfile) );
@@ -236,7 +236,7 @@ FString CreateProfileDirectoryAndFilename( const FString& InSubDirectoryName, co
 
 	// create the actual file name
 	FString FileNameWithExtension = FString::Printf( TEXT("%s%s"), *NameOfProfile, *InFileExtension );
-	FileNameWithExtension = FileNameWithExtension.Left(MaxFilenameLen);
+	FileNameWithExtension.LeftInline(MaxFilenameLen, false);
 	//UE_LOG(LogProfilingDebugging, Warning, TEXT( "CreateProfileDirectoryAndFilename: %s"), *FileNameWithExtension );
 
 
diff --git a/Engine/Source/Runtime/Core/Private/Stats/Stats2.cpp b/Engine/Source/Runtime/Core/Private/Stats/Stats2.cpp
--- a/Engine/Source/Runtime/Core/Private/Stats/Stats2.cpp
+++ b/Engine/Source/Runtime/Core/Private/Stats/Stats2.cpp
@@ -764,29 +764,29 @@ FName FStatNameAndInfo::GetShortNameFrom(FName InLongName)
 
 	if (Input.StartsWith(TEXT("//"), ESearchCase::CaseSensitive))
 	{
-		Input = Input.RightChop(2);
+		Input.RightChopInline(2, false);
 		const int32 IndexEnd = Input.Find(TEXT("//"), ESearchCase::CaseSensitive);
 		if (IndexEnd == INDEX_NONE)
 		{
 			checkStats(0);
 			return InLongName;
 		}
-		Input = Input.RightChop(IndexEnd + 2);
+		Input.RightChopInline(IndexEnd + 2, false);
 	}
 	const int32 DescIndexEnd = Input.Find(TEXT("///"), ESearchCase::CaseSensitive);
 	if (DescIndexEnd != INDEX_NONE)
 	{
-		Input = Input.Left(DescIndexEnd);
+		Input.LeftInline(DescIndexEnd, false);
 	}
 	const int32 CategoryIndexEnd = Input.Find( TEXT( "####" ), ESearchCase::CaseSensitive );
 	if( DescIndexEnd == INDEX_NONE && CategoryIndexEnd != INDEX_NONE )
 	{
-		Input = Input.Left(CategoryIndexEnd);
+		Input.LeftInline(CategoryIndexEnd, false);
 	}
 	const int32 SortByNameIndexEnd = Input.Find( TEXT( "/#/#" ), ESearchCase::CaseSensitive );
 	if( DescIndexEnd == INDEX_NONE && CategoryIndexEnd == INDEX_NONE && SortByNameIndexEnd != INDEX_NONE )
 	{
-		Input = Input.Left(SortByNameIndexEnd);
+		Input.LeftInline(SortByNameIndexEnd, false);
 	}
 	return FName(*Input);
 }
@@ -797,10 +797,10 @@ FName FStatNameAndInfo::GetGroupNameFrom(FName InLongName)
 
 	if (Input.StartsWith(TEXT("//"), ESearchCase::CaseSensitive))
 	{
-		Input = Input.RightChop(2);
+		Input.RightChopInline(2, false);
 		if (Input.StartsWith(TEXT("Groups//")))
 		{
-			Input = Input.RightChop(8);
+			Input.RightChopInline(8, false);
 		}
 		const int32 IndexEnd = Input.Find(TEXT("//"), ESearchCase::CaseSensitive);
 		if (IndexEnd != INDEX_NONE)
@@ -819,7 +819,7 @@ FString FStatNameAndInfo::GetDescriptionFrom(FName InLongName)
 	const int32 IndexStart = Input.Find(TEXT("///"), ESearchCase::CaseSensitive);
 	if (IndexStart != INDEX_NONE)
 	{
-		Input = Input.RightChop(IndexStart + 3);
+		Input.RightChopInline(IndexStart + 3, false);
 		const int32 IndexEnd = Input.Find(TEXT("///"), ESearchCase::CaseSensitive);
 		if (IndexEnd != INDEX_NONE)
 		{
@@ -836,7 +836,7 @@ FName FStatNameAndInfo::GetGroupCategoryFrom(FName InLongName)
 	const int32 IndexStart = Input.Find(TEXT("####"), ESearchCase::CaseSensitive);
 	if (IndexStart != INDEX_NONE)
 	{
-		Input = Input.RightChop(IndexStart + 4);
+		Input.RightChopInline(IndexStart + 4, false);
 		const int32 IndexEnd = Input.Find(TEXT("####"), ESearchCase::CaseSensitive);
 		if (IndexEnd != INDEX_NONE)
 		{
@@ -854,7 +854,7 @@ bool FStatNameAndInfo::GetSortByNameFrom(FName InLongName)
 	const int32 IndexStart = Input.Find(TEXT("/#/#"), ESearchCase::CaseSensitive);
 	if (IndexStart != INDEX_NONE)
 	{
-		Input = Input.RightChop(IndexStart + 4);
+		Input.RightChopInline(IndexStart + 4, false);
 		const int32 IndexEnd = Input.Find(TEXT("/#/#"), ESearchCase::CaseSensitive);
 		if (IndexEnd != INDEX_NONE)
 		{
@@ -1321,7 +1321,7 @@ void FThreadStats::CheckForCollectingStartupStats()
 		{
 			break;
 		}
-		CmdLine = CmdLine.Mid(Index + StatCmds.Len());
+		CmdLine.MidInline(Index + StatCmds.Len(), MAX_int32, false);
 	}
 
 	if (FParse::Param( FCommandLine::Get(), TEXT( "LoadTimeStats" ) ))
diff --git a/Engine/Source/Runtime/Core/Private/Stats/StatsData.cpp b/Engine/Source/Runtime/Core/Private/Stats/StatsData.cpp
--- a/Engine/Source/Runtime/Core/Private/Stats/StatsData.cpp
+++ b/Engine/Source/Runtime/Core/Private/Stats/StatsData.cpp
@@ -2057,7 +2057,7 @@ FString FStatsUtils::FromEscapedFString(const TCHAR* Escaped)
 				break;
 			}
 			Result += Input.Left(Index);
-			Input = Input.RightChop(Index + 1);
+			Input.RightChopInline(Index + 1, false);
 
 		}
 		{
@@ -2069,7 +2069,7 @@ FString FStatsUtils::FromEscapedFString(const TCHAR* Escaped)
 				break;
 			}
 			FString Number = Input.Left(IndexEnd);
-			Input = Input.RightChop(IndexEnd + 1);
+			Input.RightChopInline(IndexEnd + 1, false);
 			Result.AppendChar(TCHAR(uint32(FCString::Atoi64(*Number))));
 		}
 	}
diff --git a/Engine/Source/Runtime/Core/Private/Unix/UnixPlatformProcess.cpp b/Engine/Source/Runtime/Core/Private/Unix/UnixPlatformProcess.cpp
--- a/Engine/Source/Runtime/Core/Private/Unix/UnixPlatformProcess.cpp
+++ b/Engine/Source/Runtime/Core/Private/Unix/UnixPlatformProcess.cpp
@@ -1826,9 +1826,9 @@ bool FUnixPlatformProcess::IsFirstInstance()
 		{
 			FString LockFileName(TEXT("/tmp/"));
 			FString ExecPath(FPlatformProcess::ExecutableName());
-			ExecPath.ReplaceInline(TEXT("/"), TEXT("-"));
+			ExecPath.ReplaceInline(TEXT("/"), TEXT("-"), ESearchCase::CaseSensitive);
 			// [RCL] 2015-09-20: can run out of filename limits (256 bytes) due to a long path, be conservative and assume 4-char UTF-8 name like e.g. Japanese
-			ExecPath = ExecPath.Right(80);
+			ExecPath.RightInline(80, false);
 
 			LockFileName += ExecPath;
 
diff --git a/Engine/Source/Runtime/Core/Public/GenericPlatform/GenericPlatformDriver.h b/Engine/Source/Runtime/Core/Public/GenericPlatform/GenericPlatformDriver.h
--- a/Engine/Source/Runtime/Core/Public/GenericPlatform/GenericPlatformDriver.h
+++ b/Engine/Source/Runtime/Core/Public/GenericPlatform/GenericPlatformDriver.h
@@ -387,13 +387,13 @@ struct FBlackListEntry
 
 				FString DriverDateComparisonOp = DriverDateString.Left(OpLength);
 				FString TrimmedDriverDateString = DriverDateString.RightChop(OpLength);
-				TrimmedDriverDateString.Split(TEXT("-"), &TempMonthDay, &TempYear, ESearchCase::IgnoreCase, ESearchDir::FromEnd);
-				TempMonthDay.Split(TEXT("-"), &TempMonth, &TempDay, ESearchCase::IgnoreCase, ESearchDir::FromEnd);
+				TrimmedDriverDateString.Split(TEXT("-"), &TempMonthDay, &TempYear, ESearchCase::CaseSensitive, ESearchDir::FromEnd);
+				TempMonthDay.Split(TEXT("-"), &TempMonth, &TempDay, ESearchCase::CaseSensitive, ESearchDir::FromEnd);
 				
 				FDateTime TestDate(FCString::Atoi(*TempYear), FCString::Atoi(*TempMonth), FCString::Atoi(*TempDay));
 
-				Info.DriverDate.Split(TEXT("-"), &TempMonthDay, &TempYear, ESearchCase::IgnoreCase, ESearchDir::FromEnd);
-				TempMonthDay.Split(TEXT("-"), &TempMonth, &TempDay, ESearchCase::IgnoreCase, ESearchDir::FromEnd);
+				Info.DriverDate.Split(TEXT("-"), &TempMonthDay, &TempYear, ESearchCase::CaseSensitive, ESearchDir::FromEnd);
+				TempMonthDay.Split(TEXT("-"), &TempMonth, &TempDay, ESearchCase::CaseSensitive, ESearchDir::FromEnd);
 
 				FDateTime InfoDate(FCString::Atoi(*TempYear), FCString::Atoi(*TempMonth), FCString::Atoi(*TempDay));
 
diff --git a/Engine/Source/Runtime/CoreUObject/Private/Misc/PackageName.cpp b/Engine/Source/Runtime/CoreUObject/Private/Misc/PackageName.cpp
--- a/Engine/Source/Runtime/CoreUObject/Private/Misc/PackageName.cpp
+++ b/Engine/Source/Runtime/CoreUObject/Private/Misc/PackageName.cpp
@@ -1088,7 +1088,7 @@ FString FPackageName::GetDelegateResolvedPackagePath(const FString& InSourcePack
 		if (PathName.FindChar('.', DotIndex))
 		{
 			ObjectName = PathName.Mid(DotIndex + 1);
-			PathName = PathName.Left(DotIndex);
+			PathName.LeftInline(DotIndex, false);
 		}
 
 		for (auto Delegate : FCoreDelegates::PackageNameResolvers)
@@ -1310,7 +1310,7 @@ bool FPackageName::ParseExportTextPath(const FString& InExportTextPath, FString*
 			FString& OutObjectPathRef = *OutObjectPath;
 			if ( OutObjectPathRef.EndsWith(TEXT("'"), ESearchCase::CaseSensitive) )
 			{
-				OutObjectPathRef = OutObjectPathRef.LeftChop(1);
+				OutObjectPathRef.LeftChopInline(1, false);
 			}
 		}
 
diff --git a/Engine/Source/Runtime/CoreUObject/Private/Serialization/ArchiveStackTrace.cpp b/Engine/Source/Runtime/CoreUObject/Private/Serialization/ArchiveStackTrace.cpp
--- a/Engine/Source/Runtime/CoreUObject/Private/Serialization/ArchiveStackTrace.cpp
+++ b/Engine/Source/Runtime/CoreUObject/Private/Serialization/ArchiveStackTrace.cpp
@@ -138,7 +138,7 @@ FString FArchiveStackTrace::FCallstackData::ToString(const TCHAR* CallstackCutof
 			int32 CutoffIndex = StackLine.Find(TEXT(" "), ESearchCase::CaseSensitive);
 			if (CutoffIndex >= -1 && CutoffIndex < StackLine.Len() - 2)
 			{
-				StackLine = StackLine.Mid(CutoffIndex + 1);
+				StackLine.MidInline(CutoffIndex + 1, MAX_int32, false);
 			}
 		}
 		HumanReadableString += Indent;
diff --git a/Engine/Source/Runtime/CoreUObject/Private/UObject/Enum.cpp b/Engine/Source/Runtime/CoreUObject/Private/UObject/Enum.cpp
--- a/Engine/Source/Runtime/CoreUObject/Private/UObject/Enum.cpp
+++ b/Engine/Source/Runtime/CoreUObject/Private/UObject/Enum.cpp
@@ -111,7 +111,7 @@ FString UEnum::GetBaseEnumNameOnDuplication() const
 	check(DoubleColonPos != INDEX_NONE);
 
 	// Get actual base name.
-	BaseEnumName = BaseEnumName.LeftChop(BaseEnumName.Len() - DoubleColonPos);
+	BaseEnumName.LeftChopInline(BaseEnumName.Len() - DoubleColonPos, false);
 
 	return BaseEnumName;
 }
@@ -308,7 +308,7 @@ FString UEnum::GenerateEnumPrefix() const
 			}
 
 			// Trim the prefix to the length of the common prefix.
-			Prefix = Prefix.Left(PrefixIdx);
+			Prefix.LeftInline(PrefixIdx, false);
 		}
 
 		// Find the index of the rightmost underscore in the prefix.
@@ -317,7 +317,7 @@ FString UEnum::GenerateEnumPrefix() const
 		// If an underscore was found, trim the prefix so only the part before the rightmost underscore is included.
 		if (UnderscoreIdx > 0)
 		{
-			Prefix = Prefix.Left(UnderscoreIdx);
+			Prefix.LeftInline(UnderscoreIdx, false);
 		}
 		else
 		{
@@ -668,7 +668,7 @@ FString UEnum::GetMetaData( const TCHAR* Key, int32 NameIndex/*=INDEX_NONE*/, bo
 		if (!GConfig->GetString(TEXT("EnumRemap"), *KeyString, ResultString, GEngineIni))
 		{
 			// if this fails, then use what's after the ini:
-			ResultString = ResultString.Mid(4);
+			ResultString.MidInline(4, MAX_int32, false);
 		}
 	}
 
diff --git a/Engine/Source/Runtime/CoreUObject/Private/UObject/Obj.cpp b/Engine/Source/Runtime/CoreUObject/Private/UObject/Obj.cpp
--- a/Engine/Source/Runtime/CoreUObject/Private/UObject/Obj.cpp
+++ b/Engine/Source/Runtime/CoreUObject/Private/UObject/Obj.cpp
@@ -3323,13 +3323,13 @@ bool StaticExec( UWorld* InWorld, const TCHAR* Cmd, FOutputDevice& Ar )
 					if (AsteriskPos != INDEX_NONE && (QuestionPos == INDEX_NONE || QuestionPos > AsteriskPos))
 					{
 						new(WildcardPieces) FListPropsWildcardPiece(PropWildcard.Left(AsteriskPos), true);
-						PropWildcard = PropWildcard.Right(PropWildcard.Len() - AsteriskPos - 1);
+						PropWildcard.RightInline(PropWildcard.Len() - AsteriskPos - 1, false);
 						bFound = true;
 					}
 					else if (QuestionPos != INDEX_NONE)
 					{
 						new(WildcardPieces) FListPropsWildcardPiece(PropWildcard.Left(QuestionPos), false);
-						PropWildcard = PropWildcard.Right(PropWildcard.Len() - QuestionPos - 1);
+						PropWildcard.RightInline(PropWildcard.Len() - QuestionPos - 1, false);
 						bFound = true;
 					}
 				}
@@ -3373,7 +3373,7 @@ bool StaticExec( UWorld* InWorld, const TCHAR* Cmd, FOutputDevice& Ar )
 							break;
 						}
 
-						Match = Match.Right(Match.Len() - Pos - WildcardPieces[i].Str.Len());
+						Match.RightInline(Match.Len() - Pos - WildcardPieces[i].Str.Len(), false);
 					}
 				}
 				if (bResult)
diff --git a/Engine/Source/Runtime/CoreUObject/Private/UObject/UObjectBase.cpp b/Engine/Source/Runtime/CoreUObject/Private/UObject/UObjectBase.cpp
--- a/Engine/Source/Runtime/CoreUObject/Private/UObject/UObjectBase.cpp
+++ b/Engine/Source/Runtime/CoreUObject/Private/UObject/UObjectBase.cpp
@@ -704,10 +704,10 @@ FString RemoveClassPrefix(const TCHAR* ClassName)
 {
 	const TCHAR* DeprecatedPrefix = TEXT("DEPRECATED_");
 	FString NameWithoutPrefix(ClassName);
-	NameWithoutPrefix = NameWithoutPrefix.Mid(1);
+	NameWithoutPrefix.MidInline(1, MAX_int32, false);
 	if (NameWithoutPrefix.StartsWith(DeprecatedPrefix))
 	{
-		NameWithoutPrefix = NameWithoutPrefix.Mid(FCString::Strlen(DeprecatedPrefix));
+		NameWithoutPrefix.MidInline(FCString::Strlen(DeprecatedPrefix), MAX_int32, false);
 	}
 	return NameWithoutPrefix;
 }
diff --git a/Engine/Source/Runtime/CoreUObject/Private/UObject/UObjectGlobals.cpp b/Engine/Source/Runtime/CoreUObject/Private/UObject/UObjectGlobals.cpp
--- a/Engine/Source/Runtime/CoreUObject/Private/UObject/UObjectGlobals.cpp
+++ b/Engine/Source/Runtime/CoreUObject/Private/UObject/UObjectGlobals.cpp
@@ -649,7 +649,7 @@ FString ResolveIniObjectsReference(const FString& ObjectReference, const FString
 	if (i != -1)
 	{
 		Key = Section.Mid(i + 1);
-		Section = Section.Left(i);
+		Section.LeftInline(i, false);
 	}
 
 	FString Output;
@@ -3259,7 +3259,7 @@ void ConstructorHelpers::StripObjectClass( FString& PathName, bool bAssertOnBadP
 		PathName.FindLastChar( TCHAR('\''), NameEndIndex );
 		if(NameEndIndex > NameStartIndex)
 		{
-			PathName = PathName.Mid( NameStartIndex+1, NameEndIndex-NameStartIndex-1 );
+			PathName.MidInline( NameStartIndex+1, NameEndIndex-NameStartIndex-1, false );
 		}
 		else
 		{
diff --git a/Engine/Source/Runtime/CrashReportCore/Private/CrashDescription.cpp b/Engine/Source/Runtime/CrashReportCore/Private/CrashDescription.cpp
--- a/Engine/Source/Runtime/CrashReportCore/Private/CrashDescription.cpp
+++ b/Engine/Source/Runtime/CrashReportCore/Private/CrashDescription.cpp
@@ -586,12 +586,12 @@ FCrashWERContext::FCrashWERContext( const FString& WERXMLFilepath )
 
 			if (ParsedParameters9.Num() > 0)
 			{
-				BranchName = ParsedParameters9[0].Replace( TEXT( "+" ), TEXT( "/" ) );
+				BranchName = ParsedParameters9[0].Replace( TEXT( "+" ), TEXT( "/" ), ESearchCase::CaseSensitive);
 
 				const FString DepotRoot = TEXT( "//depot/" );
 				if (BranchName.StartsWith( DepotRoot ))
 				{
-					BranchName = BranchName.Mid( DepotRoot.Len() );
+					BranchName.MidInline( DepotRoot.Len(), MAX_int32, false );
 				}
 				EngineVersionComponents++;
 			}
diff --git a/Engine/Source/Runtime/CrashReportCore/Private/IOS/IOSErrorReport.cpp b/Engine/Source/Runtime/CrashReportCore/Private/IOS/IOSErrorReport.cpp
--- a/Engine/Source/Runtime/CrashReportCore/Private/IOS/IOSErrorReport.cpp
+++ b/Engine/Source/Runtime/CrashReportCore/Private/IOS/IOSErrorReport.cpp
@@ -50,11 +50,11 @@ FString FIOSErrorReport::FindCrashedAppPath() const
 			int32 LineEnd = -1;
 			if(PathData.FindChar( TCHAR('\r'), LineEnd ))
 			{
-				PathData = PathData.Left(LineEnd);
+				PathData.LeftInline(LineEnd, false);
 			}
 			if(PathData.FindChar( TCHAR('\n'), LineEnd ))
 			{
-				PathData = PathData.Left(LineEnd);
+				PathData.LeftInline(LineEnd, false);
 			}
 			return PathData;
 		}
diff --git a/Engine/Source/Runtime/CrashReportCore/Private/Mac/MacErrorReport.cpp b/Engine/Source/Runtime/CrashReportCore/Private/Mac/MacErrorReport.cpp
--- a/Engine/Source/Runtime/CrashReportCore/Private/Mac/MacErrorReport.cpp
+++ b/Engine/Source/Runtime/CrashReportCore/Private/Mac/MacErrorReport.cpp
@@ -60,11 +60,11 @@ FString FMacErrorReport::FindCrashedAppPath() const
 			int32 LineEnd = -1;
 			if(PathData.FindChar( TCHAR('\r'), LineEnd ))
 			{
-				PathData = PathData.Left(LineEnd);
+				PathData.LeftInline(LineEnd, false);
 			}
 			if(PathData.FindChar( TCHAR('\n'), LineEnd ))
 			{
-				PathData = PathData.Left(LineEnd);
+				PathData.LeftInline(LineEnd, false);
 			}
 			return PathData;
 		}
diff --git a/Engine/Source/Runtime/Engine/Private/ActorConstruction.cpp b/Engine/Source/Runtime/Engine/Private/ActorConstruction.cpp
--- a/Engine/Source/Runtime/Engine/Private/ActorConstruction.cpp
+++ b/Engine/Source/Runtime/Engine/Private/ActorConstruction.cpp
@@ -1210,7 +1210,7 @@ void AActor::CheckComponentInstanceName(const FName InName)
 			if (CharIndex < ConflictingObjectName.Len() - 1)
 			{
 				Counter = FCString::Atoi(*ConflictingObjectName.RightChop(CharIndex + 1));
-				ConflictingObjectName = ConflictingObjectName.Left(CharIndex + 1);
+				ConflictingObjectName.LeftInline(CharIndex + 1, false);
 			}
 			FString NewObjectName;
 			do
diff --git a/Engine/Source/Runtime/Engine/Private/Collision/CollisionProfile.cpp b/Engine/Source/Runtime/Engine/Private/Collision/CollisionProfile.cpp
--- a/Engine/Source/Runtime/Engine/Private/Collision/CollisionProfile.cpp
+++ b/Engine/Source/Runtime/Engine/Private/Collision/CollisionProfile.cpp
@@ -320,7 +320,7 @@ void UCollisionProfile::LoadProfileConfig(bool bForceInit)
 	for ( int32 EnumIndex=0; EnumIndex<NumEnum; ++EnumIndex )
 	{
 		FString EnumName = Enum->GetNameStringByIndex(EnumIndex);
-		EnumName = EnumName.RightChop(Prefix.Len());
+		EnumName.RightChopInline(Prefix.Len(), false);
 		FName DisplayName = FName(*EnumName);
 
 		if ( IS_VALID_COLLISIONCHANNEL(EnumIndex) )
diff --git a/Engine/Source/Runtime/Engine/Private/GameInstance.cpp b/Engine/Source/Runtime/Engine/Private/GameInstance.cpp
--- a/Engine/Source/Runtime/Engine/Private/GameInstance.cpp
+++ b/Engine/Source/Runtime/Engine/Private/GameInstance.cpp
@@ -1302,7 +1302,7 @@ AGameModeBase* UGameInstance::CreateGameModeForURL(FURL InURL, UWorld* InWorld)
 		if (MapNameNoPath.StartsWith(PLAYWORLD_PACKAGE_PREFIX))
 		{
 			const int32 PrefixLen = UWorld::BuildPIEPackagePrefix(WorldContext->PIEInstance).Len();
-			MapNameNoPath = MapNameNoPath.Mid(PrefixLen);
+			MapNameNoPath.MidInline(PrefixLen, MAX_int32, false);
 		}
 
 		FString const GameClassName = UGameMapsSettings::GetGameModeForMapName(FString(MapNameNoPath));
diff --git a/Engine/Source/Runtime/Engine/Private/GameplayStatics.cpp b/Engine/Source/Runtime/Engine/Private/GameplayStatics.cpp
--- a/Engine/Source/Runtime/Engine/Private/GameplayStatics.cpp
+++ b/Engine/Source/Runtime/Engine/Private/GameplayStatics.cpp
@@ -2820,16 +2820,17 @@ bool UGameplayStatics::GrabOption( FString& Options, FString& Result )
 	{
 		// Get result.
 		Result = Options.Mid(1, MAX_int32);
-		if (Result.Contains(QuestionMark, ESearchCase::CaseSensitive))
+		const int32 QMIdx = Result.Find(QuestionMark, ESearchCase::CaseSensitive);
+		if (QMIdx != INDEX_NONE)
 		{
-			Result = Result.Left(Result.Find(QuestionMark, ESearchCase::CaseSensitive));
+			Result.LeftInline(QMIdx, false);
 		}
 
 		// Update options.
-		Options = Options.Mid(1, MAX_int32);
+		Options.MidInline(1, MAX_int32, false);
 		if (Options.Contains(QuestionMark, ESearchCase::CaseSensitive))
 		{
-			Options = Options.Mid(Options.Find(QuestionMark, ESearchCase::CaseSensitive), MAX_int32);
+			Options.MidInline(Options.Find(QuestionMark, ESearchCase::CaseSensitive), MAX_int32, false);
 		}
 		else
 		{
diff --git a/Engine/Source/Runtime/Engine/Private/Internationalization/StringTable.cpp b/Engine/Source/Runtime/Engine/Private/Internationalization/StringTable.cpp
--- a/Engine/Source/Runtime/Engine/Private/Internationalization/StringTable.cpp
+++ b/Engine/Source/Runtime/Engine/Private/Internationalization/StringTable.cpp
@@ -273,7 +273,7 @@ private:
 			int32 DotIndex = INDEX_NONE;
 			if (StringTablePackageName.FindChar(TEXT('.'), DotIndex))
 			{
-				StringTablePackageName = StringTablePackageName.Left(DotIndex);
+				StringTablePackageName.LeftInline(DotIndex, false);
 			}
 		}
 
diff --git a/Engine/Source/Runtime/Engine/Private/Interpolation.cpp b/Engine/Source/Runtime/Engine/Private/Interpolation.cpp
--- a/Engine/Source/Runtime/Engine/Private/Interpolation.cpp
+++ b/Engine/Source/Runtime/Engine/Private/Interpolation.cpp
@@ -321,13 +321,13 @@ namespace InterpTools
 		// Check to see if there is a period in the name, which is the case 
 		// for structs and components that own interp variables. In these  
 		// cases, we want to cut off the preceeding text up and the period.
-		int32 PeriodPosition = PropertyString.Find(TEXT("."));
+		int32 PeriodPosition = PropertyString.Find(TEXT("."), ESearchCase::CaseSensitive);
 
 		if(PeriodPosition != INDEX_NONE)
 		{
 			// We found a period; Only capture the text after the 
 			// period, which represents the actual property name.
-			PropertyString = PropertyString.Mid( PeriodPosition + 1 );
+			PropertyString.MidInline( PeriodPosition + 1, MAX_int32, false );
 		}
 
 		return FName(*PropertyString);
diff --git a/Engine/Source/Runtime/Engine/Private/LevelStreaming.cpp b/Engine/Source/Runtime/Engine/Private/LevelStreaming.cpp
--- a/Engine/Source/Runtime/Engine/Private/LevelStreaming.cpp
+++ b/Engine/Source/Runtime/Engine/Private/LevelStreaming.cpp
@@ -1751,7 +1751,7 @@ ULevelStreamingDynamic* ULevelStreamingDynamic::LoadLevelInstance_Internal(UWorl
 
 	if (ShortPackageName.StartsWith(World->StreamingLevelsPrefix))
 	{
-		ShortPackageName = ShortPackageName.RightChop(World->StreamingLevelsPrefix.Len());
+		ShortPackageName.RightChopInline(World->StreamingLevelsPrefix.Len(), false);
 	}
 
 	// Remove PIE prefix if it's there before we actually load the level
diff --git a/Engine/Source/Runtime/Engine/Private/Materials/MaterialExpressions.cpp b/Engine/Source/Runtime/Engine/Private/Materials/MaterialExpressions.cpp
--- a/Engine/Source/Runtime/Engine/Private/Materials/MaterialExpressions.cpp
+++ b/Engine/Source/Runtime/Engine/Private/Materials/MaterialExpressions.cpp
@@ -1137,7 +1137,7 @@ bool UMaterialExpression::MatchesSearchQuery( const TCHAR* SearchQuery )
 	if (FCString::Stristr(SearchQuery, TEXT("NAME=")) != nullptr)
 	{
 		FString SearchString(SearchQuery);
-		SearchString = SearchString.Right(SearchString.Len() - 5);
+		SearchString.RightInline(SearchString.Len() - 5, false);
 		return (GetName().Contains(SearchString) );
 	}
 	return Desc.Contains(SearchQuery);
diff --git a/Engine/Source/Runtime/Engine/Private/StreamableManager.cpp b/Engine/Source/Runtime/Engine/Private/StreamableManager.cpp
--- a/Engine/Source/Runtime/Engine/Private/StreamableManager.cpp
+++ b/Engine/Source/Runtime/Engine/Private/StreamableManager.cpp
@@ -927,10 +927,10 @@ FStreamable* FStreamableManager::StreamInternal(const FSoftObjectPath& InTargetN
 		{
 			// We always queue a new request in case the existing one gets cancelled
 			FString Package = TargetName.ToString();
-			int32 FirstDot = Package.Find(TEXT("."));
+			int32 FirstDot = Package.Find(TEXT("."), ESearchCase::CaseSensitive);
 			if (FirstDot != INDEX_NONE)
 			{
-				Package = Package.Left(FirstDot);
+				Package.LeftInline(FirstDot,false);
 			}
 
 			Existing->bAsyncLoadRequestOutstanding = true;
diff --git a/Engine/Source/Runtime/Engine/Private/Tests/AutomationCommon.cpp b/Engine/Source/Runtime/Engine/Private/Tests/AutomationCommon.cpp
--- a/Engine/Source/Runtime/Engine/Private/Tests/AutomationCommon.cpp
+++ b/Engine/Source/Runtime/Engine/Private/Tests/AutomationCommon.cpp
@@ -73,7 +73,7 @@ namespace AutomationCommon
 		if ( HardwareDetailsString.Len() > 0 )
 		{
 			//Get rid of the leading "_"
-			HardwareDetailsString = HardwareDetailsString.RightChop(1);
+			HardwareDetailsString.RightChopInline(1, false);
 		}
 
 		return HardwareDetailsString;
diff --git a/Engine/Source/Runtime/Engine/Private/UnrealClient.cpp b/Engine/Source/Runtime/Engine/Private/UnrealClient.cpp
--- a/Engine/Source/Runtime/Engine/Private/UnrealClient.cpp
+++ b/Engine/Source/Runtime/Engine/Private/UnrealClient.cpp
@@ -2168,7 +2168,7 @@ ENGINE_API bool GetHighResScreenShotInput(const TCHAR* Cmd, FOutputDevice& Ar, u
 	while (CmdString.FindChar(TCHAR(' '), SeperatorPos))
 	{
 		Arguments.Add(CmdString.Mid(LastSeperatorPos, SeperatorPos));
-		CmdString = CmdString.Mid(SeperatorPos + 1);
+		CmdString.MidInline(SeperatorPos + 1, MAX_int32, false);
 	}
 
 	if (CmdString.Len() > 0)
diff --git a/Engine/Source/Runtime/Engine/Private/UnrealEngine.cpp b/Engine/Source/Runtime/Engine/Private/UnrealEngine.cpp
--- a/Engine/Source/Runtime/Engine/Private/UnrealEngine.cpp
+++ b/Engine/Source/Runtime/Engine/Private/UnrealEngine.cpp
@@ -646,7 +646,8 @@ bool ParseResolution(const TCHAR* InResolution, uint32& OutX, uint32& OutY, int3
 					StringTripLen = 1;
 				}
 
-				YString = YString.Left(YString.Len() - StringTripLen).TrimStartAndEnd();
+				YString.LeftInline(YString.Len() - StringTripLen, false);
+				YString.TrimStartAndEndInline();
 			}
 
 			if (YString.IsNumeric())
@@ -6506,7 +6507,7 @@ bool UEngine::HandleMergeMeshCommand( const TCHAR* Cmd, FOutputDevice& Ar, UWorl
 		const TCHAR* LocalCmd = *CmdCopy;
 		FString Token = FParse::Token( LocalCmd, true );
 		Tokens.Add( Token );
-		CmdCopy = CmdCopy.Right( CmdCopy.Len() - Token.Len() - 1 );
+		CmdCopy.RightInline( CmdCopy.Len() - Token.Len() - 1, false);
 	}
 
 	// array of source meshes that will be merged
@@ -9593,11 +9594,11 @@ static void DrawProperty(UCanvas* CanvasObject, UObject* Obj, const FDebugDispla
 	do
 	{
 		FString Str = ValueText;
-		CommaIdx = ValueText.Find( TEXT(",") );
+		CommaIdx = ValueText.Find( TEXT(","), ESearchCase::CaseSensitive);
 		if( CommaIdx >= 0 )
 		{
 			Str = ValueText.Left(CommaIdx);
-			ValueText = ValueText.Mid( CommaIdx+1 );
+			ValueText.MidInline( CommaIdx+1, MAX_int32, false );
 		}
 
 		int32 XL, YL;
diff --git a/Engine/Source/Runtime/Engine/Private/UserInterface/Console.cpp b/Engine/Source/Runtime/Engine/Private/UserInterface/Console.cpp
--- a/Engine/Source/Runtime/Engine/Private/UserInterface/Console.cpp
+++ b/Engine/Source/Runtime/Engine/Private/UserInterface/Console.cpp
@@ -741,7 +741,7 @@ void UConsole::AppendInputText(const FString& Text)
 	while (TextMod.Len() > 0)
 	{
 		int32 Character = **TextMod.Left(1);
-		TextMod = TextMod.Mid(1);
+		TextMod.MidInline(1, MAX_int32, false);
 
 		if (Character >= 0x20 && Character < 0x100)
 		{
diff --git a/Engine/Source/Runtime/Engine/Private/UserInterface/PlayerInput.cpp b/Engine/Source/Runtime/Engine/Private/UserInterface/PlayerInput.cpp
--- a/Engine/Source/Runtime/Engine/Private/UserInterface/PlayerInput.cpp
+++ b/Engine/Source/Runtime/Engine/Private/UserInterface/PlayerInput.cpp
@@ -2098,7 +2098,7 @@ void UPlayerInput::SetBind(FName BindName, const FString& Command)
 		FString CommandMod = Command;
 		if ( CommandMod.Left(1) == TEXT("\"") && CommandMod.Right(1) == ("\"") )
 		{
-			CommandMod = CommandMod.Mid(1, CommandMod.Len() - 2);
+			CommandMod.MidInline(1, CommandMod.Len() - 2, false);
 		}
 
 		for(int32 BindIndex = DebugExecBindings.Num()-1;BindIndex >= 0;BindIndex--)
diff --git a/Engine/Source/Runtime/Engine/Private/UserInterface/VisualizeRT.cpp b/Engine/Source/Runtime/Engine/Private/UserInterface/VisualizeRT.cpp
--- a/Engine/Source/Runtime/Engine/Private/UserInterface/VisualizeRT.cpp
+++ b/Engine/Source/Runtime/Engine/Private/UserInterface/VisualizeRT.cpp
@@ -61,12 +61,12 @@ struct FRTInfo : public FRefCountedObject
 		Type = TEXT("-");
 
 		// Current Desc Info format: (DIM W[xH[xD]] FMT[ RT]) NUM NAME SIZEkB
-		if ( !Text.IsEmpty() && Text.Split(TEXT(" "), &Dimensions, &Text))
+		if ( !Text.IsEmpty() && Text.Split(TEXT(" "), &Dimensions, &Text, ESearchCase::CaseSensitive))
 		{
-			Dimensions = Dimensions.Mid(1);
+			Dimensions.MidInline(1, MAX_int32, false);
 			if (Dimensions.StartsWith(TEXT("Cube")))
 			{
-				if (!Text.Split(TEXT(" "), &Width, &Text))
+				if (!Text.Split(TEXT(" "), &Width, &Text, ESearchCase::CaseSensitive))
 				{
 					return false;
 				}
@@ -84,14 +84,14 @@ struct FRTInfo : public FRefCountedObject
 							return false;
 						}
 
-						if (!Text.Split(TEXT(" "), &Depth, &Text))
+						if (!Text.Split(TEXT(" "), &Depth, &Text, ESearchCase::CaseSensitive))
 						{
 							return false;
 						}
 					}
 					else
 					{
-						if (!Text.Split(TEXT(" "), &Height, &Text))
+						if (!Text.Split(TEXT(" "), &Height, &Text, ESearchCase::CaseSensitive))
 						{
 							return false;
 						}
@@ -100,18 +100,18 @@ struct FRTInfo : public FRefCountedObject
 			}
 
 			// Might or not have Type
-			if (Text.Split(TEXT(") "), &Format, &Text))
+			if (Text.Split(TEXT(") "), &Format, &Text, ESearchCase::CaseSensitive))
 			{
-				if (Text.Split(TEXT(" "), &Number, &Text))
+				if (Text.Split(TEXT(" "), &Number, &Text, ESearchCase::CaseSensitive))
 				{
-					if (Text.Split(TEXT(" "), &Name, &Size))
+					if (Text.Split(TEXT(" "), &Name, &Size, ESearchCase::CaseSensitive))
 					{
 						// Get the type if found
-						int32 Found = Format.Find(TEXT(" "));
+						int32 Found = Format.Find(TEXT(" "), ESearchCase::CaseSensitive);
 						if (Found > 0)
 						{
 							Type = Format.Mid(Found + 1);
-							Format = Format.Mid(0, Found);
+							Format.MidInline(0, Found, false);
 						}
 						return true;
 					}
diff --git a/Engine/Source/Runtime/GameplayTags/Private/BlueprintGameplayTagLibrary.cpp b/Engine/Source/Runtime/GameplayTags/Private/BlueprintGameplayTagLibrary.cpp
--- a/Engine/Source/Runtime/GameplayTags/Private/BlueprintGameplayTagLibrary.cpp
+++ b/Engine/Source/Runtime/GameplayTags/Private/BlueprintGameplayTagLibrary.cpp
@@ -231,15 +231,15 @@ bool UBlueprintGameplayTagLibrary::NotEqual_TagContainerTagContainer(FGameplayTa
 	FString TagString = MoveTemp(B);
 	if (TagString.StartsWith(OpenParenthesesStr, ESearchCase::CaseSensitive) && TagString.EndsWith(CloseParenthesesStr, ESearchCase::CaseSensitive))
 	{
-		TagString = TagString.LeftChop(1);
-		TagString = TagString.RightChop(1);
+		TagString.LeftChopInline(1, false);
+		TagString.RightChopInline(1, false);
 
 		const FString EqualStr(TEXT("="));
 
 		TagString.Split(EqualStr, nullptr, &TagString, ESearchCase::CaseSensitive);
 
-		TagString = TagString.LeftChop(1);
-		TagString = TagString.RightChop(1);
+		TagString.LeftChopInline(1, false);
+		TagString.RightChopInline(1, false);
 
 		FString ReadTag;
 		FString Remainder;
@@ -252,11 +252,11 @@ bool UBlueprintGameplayTagLibrary::NotEqual_TagContainerTagContainer(FGameplayTa
 			ReadTag.Split(EqualStr, nullptr, &ReadTag, ESearchCase::CaseSensitive);
 			if (ReadTag.EndsWith(CloseParenthesesStr, ESearchCase::CaseSensitive))
 			{
-				ReadTag = ReadTag.LeftChop(1);
+				ReadTag.LeftChopInline(1, false);
 				if (ReadTag.StartsWith(QuoteStr, ESearchCase::CaseSensitive) && ReadTag.EndsWith(QuoteStr, ESearchCase::CaseSensitive))
 				{
-					ReadTag = ReadTag.LeftChop(1);
-					ReadTag = ReadTag.RightChop(1);
+					ReadTag.LeftChopInline(1, false);
+					ReadTag.RightChopInline(1, false);
 				}
 			}
 			TagString = Remainder;
@@ -273,11 +273,11 @@ bool UBlueprintGameplayTagLibrary::NotEqual_TagContainerTagContainer(FGameplayTa
 			Remainder.Split(EqualStr, nullptr, &Remainder, ESearchCase::CaseSensitive);
 			if (Remainder.EndsWith(CloseParenthesesStr, ESearchCase::CaseSensitive))
 			{
-				Remainder = Remainder.LeftChop(1);
+				Remainder.LeftChopInline(1, false);
 				if (Remainder.StartsWith(QuoteStr, ESearchCase::CaseSensitive) && Remainder.EndsWith(QuoteStr, ESearchCase::CaseSensitive))
 				{
-					Remainder = Remainder.LeftChop(1);
-					Remainder = Remainder.RightChop(1);
+					Remainder.LeftChopInline(1, false);
+					Remainder.RightChopInline(1, false);
 				}
 			}
 			const FGameplayTag Tag = FGameplayTag::RequestGameplayTag(FName(*Remainder));
diff --git a/Engine/Source/Runtime/Launch/Private/LaunchEngineLoop.cpp b/Engine/Source/Runtime/Launch/Private/LaunchEngineLoop.cpp
--- a/Engine/Source/Runtime/Launch/Private/LaunchEngineLoop.cpp
+++ b/Engine/Source/Runtime/Launch/Private/LaunchEngineLoop.cpp
@@ -569,7 +569,7 @@ bool LaunchSetGameName(const TCHAR *InCmdLine, FString& OutGameProjectFilePathUn
 			int32 FirstCharToRemove = INDEX_NONE;
 			if (LocalGameName.FindChar(TCHAR('-'), FirstCharToRemove))
 			{
-				LocalGameName = LocalGameName.Left(FirstCharToRemove);
+				LocalGameName.LeftInline(FirstCharToRemove, false);
 			}
 			FApp::SetProjectName(*LocalGameName);
 
@@ -1588,7 +1588,7 @@ int32 FEngineLoop::PreInitPreStartupScreen(const TCHAR* CmdLine)
 #endif
 			if (Token.StartsWith(TEXT("run=")))
 			{
-				Token = Token.RightChop(4);
+				Token.RightChopInline(4, false);
 				if (!Token.EndsWith(TEXT("Commandlet")))
 				{
 					Token += TEXT("Commandlet");
@@ -1647,7 +1647,7 @@ int32 FEngineLoop::PreInitPreStartupScreen(const TCHAR* CmdLine)
 #endif
 		if (Token.StartsWith(TEXT("run=")))
 		{
-			Token = Token.RightChop(4);
+			Token.RightChopInline(4, false);
 			if (!Token.EndsWith(TEXT("Commandlet")))
 			{
 				Token += TEXT("Commandlet");
@@ -2062,7 +2062,7 @@ int32 FEngineLoop::PreInitPreStartupScreen(const TCHAR* CmdLine)
 		{
 			if (Token.StartsWith(TEXT("run=")))
 			{
-				Token = Token.RightChop(4);
+				Token.RightChopInline(4, false);
 				bDefinitelyCommandlet = true;
 				if (!Token.EndsWith(TEXT("Commandlet")))
 				{
diff --git a/Engine/Source/Runtime/NetworkFileSystem/Private/NetworkFileServerConnection.cpp b/Engine/Source/Runtime/NetworkFileSystem/Private/NetworkFileServerConnection.cpp
--- a/Engine/Source/Runtime/NetworkFileSystem/Private/NetworkFileServerConnection.cpp
+++ b/Engine/Source/Runtime/NetworkFileSystem/Private/NetworkFileServerConnection.cpp
@@ -1006,7 +1006,7 @@ bool FNetworkFileServerClientConnection::ProcessGetFileList( FArchive& In, FArch
 				int32 GameDirOffset = ConnectedContentFolder.Find(ConnectedProjectDir, ESearchCase::IgnoreCase, ESearchDir::FromEnd);
 				if (GameDirOffset != INDEX_NONE)
 				{
-					ConnectedContentFolder = ConnectedContentFolder.RightChop(GameDirOffset);
+					ConnectedContentFolder.RightChopInline(GameDirOffset, false);
 				}
 			}
 
diff --git a/Engine/Source/Runtime/Online/BuildPatchServices/Private/Generation/DirectoryBuildStreamer.cpp b/Engine/Source/Runtime/Online/BuildPatchServices/Private/Generation/DirectoryBuildStreamer.cpp
--- a/Engine/Source/Runtime/Online/BuildPatchServices/Private/Generation/DirectoryBuildStreamer.cpp
+++ b/Engine/Source/Runtime/Online/BuildPatchServices/Private/Generation/DirectoryBuildStreamer.cpp
@@ -454,11 +454,11 @@ namespace BuildPatchServices
 		}
 		for (FString& Filename : IgnoreFiles)
 		{
-			int32 TabLocation = Filename.Find(TEXT("\t"));
+			int32 TabLocation = Filename.Find(TEXT("\t"), ESearchCase::CaseSensitive);
 			if (TabLocation != INDEX_NONE)
 			{
 				// Strip tab separated timestamp if it exists
-				Filename = Filename.Left(TabLocation);
+				Filename.LeftInline(TabLocation, false);
 			}
 			Filename = Config.BuildRoot / Filename;
 			FPaths::NormalizeFilename(Filename);
diff --git a/Engine/Source/Runtime/Online/HTTP/Private/GenericPlatform/GenericPlatformHttp.cpp b/Engine/Source/Runtime/Online/HTTP/Private/GenericPlatform/GenericPlatformHttp.cpp
--- a/Engine/Source/Runtime/Online/HTTP/Private/GenericPlatform/GenericPlatformHttp.cpp
+++ b/Engine/Source/Runtime/Online/HTTP/Private/GenericPlatform/GenericPlatformHttp.cpp
@@ -203,7 +203,7 @@ FString FGenericPlatformHttp::GetUrlDomain(const FString& Url)
 	const int32 Idx = Domain.GetCharArray().IndexOfByPredicate([](TCHAR Character) { return Character == TEXT('/') || Character == TEXT('?') || Character == TEXT(':'); });
 	if (Idx != INDEX_NONE)
 	{
-		Domain = Domain.Left(Idx);
+		Domain.LeftInline(Idx, false);
 	}
 	return Domain;
 }
diff --git a/Engine/Source/Runtime/Online/HTTPServer/Private/HttpConnectionRequestReadContext.cpp b/Engine/Source/Runtime/Online/HTTPServer/Private/HttpConnectionRequestReadContext.cpp
--- a/Engine/Source/Runtime/Online/HTTPServer/Private/HttpConnectionRequestReadContext.cpp
+++ b/Engine/Source/Runtime/Online/HTTPServer/Private/HttpConnectionRequestReadContext.cpp
@@ -261,7 +261,7 @@ TSharedPtr<FHttpServerRequest> FHttpConnectionRequestReadContext::BuildRequest(c
 	if (RequestHttpPath.FindChar(TCHAR('?'), QueryParamsIndex))
 	{
 		FString QueryParamsStr = RequestHttpPath.Mid(QueryParamsIndex+1);
-		RequestHttpPath = RequestHttpPath.Mid(0, QueryParamsIndex);
+		RequestHttpPath.MidInline(0, QueryParamsIndex, false);
 
 		// Split query params
 		TArray<FString> QueryParamPairs;
diff --git a/Engine/Source/Runtime/OpenGLDrv/Private/OpenGLShaders.cpp b/Engine/Source/Runtime/OpenGLDrv/Private/OpenGLShaders.cpp
--- a/Engine/Source/Runtime/OpenGLDrv/Private/OpenGLShaders.cpp
+++ b/Engine/Source/Runtime/OpenGLDrv/Private/OpenGLShaders.cpp
@@ -2899,9 +2899,9 @@ static void VerifyUniformLayout(const TCHAR* UniformName, const UniformData& GLS
 	*/
 	FString RequestedUniformName(UniformName);
 	RequestedUniformName = RequestedUniformName.Replace(*BlockName, TEXT(""));
-	if(RequestedUniformName.StartsWith(TEXT(".")))
+	if(RequestedUniformName.StartsWith(TEXT("."), ESearchCase::CaseSensitive))
 	{
-		RequestedUniformName = RequestedUniformName.RightChop(1);
+		RequestedUniformName.RightChopInline(1, false);
 	}
 #else
 	FString RequestedUniformName = UniformName;
diff --git a/Engine/Source/Runtime/PakFile/Private/IPlatformFilePak.cpp b/Engine/Source/Runtime/PakFile/Private/IPlatformFilePak.cpp
--- a/Engine/Source/Runtime/PakFile/Private/IPlatformFilePak.cpp
+++ b/Engine/Source/Runtime/PakFile/Private/IPlatformFilePak.cpp
@@ -4958,11 +4958,11 @@ void FPakFile::LoadIndex(FArchive* Reader)
 				// add the parent directories up to the mount point
 				while (MountPoint != Path)
 				{
-					Path = Path.Left(Path.Len() - 1);
+					Path.LeftInline(Path.Len() - 1, false);
 					int32 Offset = 0;
 					if (Path.FindLastChar('/', Offset))
 					{
-						Path = Path.Left(Offset);
+						Path.LeftInline(Offset, false);
 						MakeDirectoryFromPath(Path);
 						if (Index.Find(Path) == NULL)
 						{
@@ -6789,11 +6789,11 @@ void FPakFile::AddSpecialFile(FPakEntry Entry, const FString& Filename)
 		// add the parent directories up to the mount point
 		while (MountPoint != Path)
 		{
-			Path = Path.Left(Path.Len() - 1);
+			Path.LeftInline(Path.Len() - 1, false);
 			int32 Offset = 0;
 			if (Path.FindLastChar('/', Offset))
 			{
-				Path = Path.Left(Offset);
+				Path.LeftInline(Offset, false);
 				MakeDirectoryFromPath(Path);
 				if (Index.Find(Path) == NULL)
 				{
diff --git a/Engine/Source/Runtime/Projects/Private/PluginDescriptor.cpp b/Engine/Source/Runtime/Projects/Private/PluginDescriptor.cpp
--- a/Engine/Source/Runtime/Projects/Private/PluginDescriptor.cpp
+++ b/Engine/Source/Runtime/Projects/Private/PluginDescriptor.cpp
@@ -105,9 +105,9 @@ bool FPluginDescriptor::Read(const FJsonObject& Object, FText& OutFailReason)
         
 	// Due to a difference in command line parsing between Windows and Mac, we shipped a few Mac samples containing
 	// a category name with escaped quotes. Remove them here to make sure we can list them in the right category.
-	if (Category.Len() >= 2 && Category.StartsWith(TEXT("\"")) && Category.EndsWith(TEXT("\"")))
+	if (Category.Len() >= 2 && Category.StartsWith(TEXT("\""), ESearchCase::CaseSensitive) && Category.EndsWith(TEXT("\""), ESearchCase::CaseSensitive))
 	{
-		Category = Category.Mid(1, Category.Len() - 2);
+		Category.MidInline(1, Category.Len() - 2, false);
 	}
 
 	Object.TryGetStringField(TEXT("CreatedBy"), CreatedBy);
diff --git a/Engine/Source/Runtime/RenderCore/Private/ShaderCore.cpp b/Engine/Source/Runtime/RenderCore/Private/ShaderCore.cpp
--- a/Engine/Source/Runtime/RenderCore/Private/ShaderCore.cpp
+++ b/Engine/Source/Runtime/RenderCore/Private/ShaderCore.cpp
@@ -611,18 +611,18 @@ static FString GetShaderSourceFilePath(const FString& VirtualFilePath, TArray<FS
 FString ParseVirtualShaderFilename(const FString& InFilename)
 {
 	FString ShaderDir = FString(FPlatformProcess::ShaderDir());
-	ShaderDir.ReplaceInline(TEXT("\\"), TEXT("/"));
+	ShaderDir.ReplaceInline(TEXT("\\"), TEXT("/"), ESearchCase::CaseSensitive);
 	int32 CharIndex = ShaderDir.Find(TEXT("/"), ESearchCase::CaseSensitive, ESearchDir::FromEnd, ShaderDir.Len() - 1);
 	if (CharIndex != INDEX_NONE)
 	{
-		ShaderDir = ShaderDir.Right(ShaderDir.Len() - CharIndex);
+		ShaderDir.RightInline(ShaderDir.Len() - CharIndex, false);
 	}
 
-	FString RelativeFilename = InFilename.Replace(TEXT("\\"), TEXT("/"));
+	FString RelativeFilename = InFilename.Replace(TEXT("\\"), TEXT("/"), ESearchCase::CaseSensitive);
 	// remove leading "/" because this makes path absolute on Linux (and Mac).
 	if (RelativeFilename.Len() > 0 && RelativeFilename[0] == TEXT('/'))
 	{
-		RelativeFilename = RelativeFilename.Right(RelativeFilename.Len() - 1);
+		RelativeFilename.RightInline(RelativeFilename.Len() - 1, false);
 	}
 	RelativeFilename = IFileManager::Get().ConvertToRelativePath(*RelativeFilename);
 	CharIndex = RelativeFilename.Find(ShaderDir);
@@ -646,7 +646,7 @@ FString ParseVirtualShaderFilename(const FString& InFilename)
 			}
 			while (NewCharIndex != INDEX_NONE && ++NumDirsSkipped < NumDirsToSkip);
 		}
-		RelativeFilename = RelativeFilename.Mid(CharIndex, RelativeFilename.Len() - CharIndex);
+		RelativeFilename.MidInline(CharIndex, RelativeFilename.Len() - CharIndex, false);
 	}
 
 	// add leading "/" to the relative filename because that's what virtual shader path expects
diff --git a/Engine/Source/Runtime/Renderer/Private/LightRendering.cpp b/Engine/Source/Runtime/Renderer/Private/LightRendering.cpp
--- a/Engine/Source/Runtime/Renderer/Private/LightRendering.cpp
+++ b/Engine/Source/Runtime/Renderer/Private/LightRendering.cpp
@@ -707,7 +707,7 @@ void FSceneRenderer::GetLightNameForDrawEvent(const FLightSceneProxy* LightProxy
 		{
 			// Trim the leading path before the level name to make it more readable
 			// The level FName was taken directly from the Outermost UObject, otherwise we would do this operation on the game thread
-			FullLevelName = FullLevelName.Mid(LastSlashIndex + 1, FullLevelName.Len() - (LastSlashIndex + 1));
+			FullLevelName.MidInline(LastSlashIndex + 1, FullLevelName.Len() - (LastSlashIndex + 1), false);
 		}
 
 		LightNameWithLevel = FullLevelName + TEXT(".") + LightProxy->GetComponentName().ToString();
diff --git a/Engine/Source/Runtime/Renderer/Private/Renderer.cpp b/Engine/Source/Runtime/Renderer/Private/Renderer.cpp
--- a/Engine/Source/Runtime/Renderer/Private/Renderer.cpp
+++ b/Engine/Source/Runtime/Renderer/Private/Renderer.cpp
@@ -549,13 +549,13 @@ static void VisualizeTextureExec( const TCHAR* Cmd, FOutputDevice &Ar )
 		// e.g. mip2 or mip0
 		else if(Parameter.Left(3) == TEXT("mip"))
 		{
-			Parameter = Parameter.Right(Parameter.Len() - 3);
+			Parameter.RightInline(Parameter.Len() - 3, false);
 			GVisualizeTexture.CustomMip = FCString::Atoi(*Parameter);
 		}
 		// e.g. [0] or [2]
 		else if(Parameter.Left(5) == TEXT("index"))
 		{
-			Parameter = Parameter.Right(Parameter.Len() - 5);
+			Parameter.RightInline(Parameter.Len() - 5, false);
 			GVisualizeTexture.ArrayIndex = FCString::Atoi(*Parameter);
 		}
 		// e.g. RGB*6, A, *22, /2.7, A*7
@@ -571,7 +571,7 @@ static void VisualizeTextureExec( const TCHAR* Cmd, FOutputDevice &Ar )
 
 			if(Parameter.Left(3) == TEXT("rgb"))
 			{
-				Parameter = Parameter.Right(Parameter.Len() - 3);
+				Parameter.RightInline(Parameter.Len() - 3, false);
 			}
 			else if(Parameter.Left(1) == TEXT("r")) SingleChannel = 0;
 			else if(Parameter.Left(1) == TEXT("g")) SingleChannel = 1;
@@ -579,7 +579,7 @@ static void VisualizeTextureExec( const TCHAR* Cmd, FOutputDevice &Ar )
 			else if(Parameter.Left(1) == TEXT("a")) SingleChannel = 3;
 			if ( SingleChannel >= 0 )
 			{
-				Parameter = Parameter.Right(Parameter.Len() - 1);
+				Parameter.RightInline(Parameter.Len() - 1, false);
 				GVisualizeTexture.SingleChannel = SingleChannel;
 				GVisualizeTexture.SingleChannelMul = 1;
 				GVisualizeTexture.RGBMul = 0;
@@ -590,12 +590,12 @@ static void VisualizeTextureExec( const TCHAR* Cmd, FOutputDevice &Ar )
 			// * or /
 			if(Parameter.Left(1) == TEXT("*"))
 			{
-				Parameter = Parameter.Right(Parameter.Len() - 1);
+				Parameter.RightInline(Parameter.Len() - 1, false);
 				Mul = FCString::Atof(*Parameter);
 			}
 			else if(Parameter.Left(1) == TEXT("/"))
 			{
-				Parameter = Parameter.Right(Parameter.Len() - 1);
+				Parameter.RightInline(Parameter.Len() - 1, false);
 				Mul = 1.0f / FCString::Atof(*Parameter);
 			}
 			GVisualizeTexture.RGBMul *= Mul;
diff --git a/Engine/Source/Runtime/SandboxFile/Private/IPlatformFileSandboxWrapper.cpp b/Engine/Source/Runtime/SandboxFile/Private/IPlatformFileSandboxWrapper.cpp
--- a/Engine/Source/Runtime/SandboxFile/Private/IPlatformFileSandboxWrapper.cpp
+++ b/Engine/Source/Runtime/SandboxFile/Private/IPlatformFileSandboxWrapper.cpp
@@ -72,7 +72,7 @@ bool FSandboxPlatformFile::Initialize(IPlatformFile* Inner, const TCHAR* CmdLine
 			if( FPaths::IsDrive(DriveCheck) == false )
 			{
 				FString Command( CommandLineDirectory.Mid( CommandIndex + 1 ) );
-				CommandLineDirectory = CommandLineDirectory.Left( CommandIndex );
+				CommandLineDirectory.LeftInline( CommandIndex, false);
 		
 				if( Command == TEXT("wipe") )
 				{
diff --git a/Engine/Source/Runtime/SandboxFile/Public/IPlatformFileSandboxWrapper.h b/Engine/Source/Runtime/SandboxFile/Public/IPlatformFileSandboxWrapper.h
--- a/Engine/Source/Runtime/SandboxFile/Public/IPlatformFileSandboxWrapper.h
+++ b/Engine/Source/Runtime/SandboxFile/Public/IPlatformFileSandboxWrapper.h
@@ -490,14 +490,14 @@ public:
 				VisitedSandboxFiles.Add( *LocalFilename );
 				// Now convert the sandbox path back to engine path because the sandbox folder should not be exposed
 				// to the engine and remain transparent.
-				LocalFilename = LocalFilename.Mid(SandboxFile.GetSandboxDirectory().Len());
+				LocalFilename.MidInline(SandboxFile.GetSandboxDirectory().Len(), MAX_int32, false);
 				if (LocalFilename.StartsWith(TEXT("Engine/")) || (FCString::Stricmp( *LocalFilename, TEXT("Engine") ) == 0))
 				{
 					LocalFilename = SandboxFile.GetAbsoluteRootDirectory() / LocalFilename;
 				}
 				else
 				{
-					LocalFilename = LocalFilename.Mid(SandboxFile.GetGameSandboxDirectoryName().Len());
+					LocalFilename.MidInline(SandboxFile.GetGameSandboxDirectoryName().Len(), MAX_int32, false);
 					LocalFilename = SandboxFile.GetAbsoluteGameDirectory() / LocalFilename;
 				}
 			}
@@ -562,7 +562,7 @@ public:
 				VisitedSandboxFiles.Add( *LocalFilename );
 				// Now convert the sandbox path back to engine path because the sandbox folder should not be exposed
 				// to the engine and remain transparent.
-				LocalFilename = LocalFilename.Mid(SandboxFile.GetSandboxDirectory().Len());
+				LocalFilename.MidInline(SandboxFile.GetSandboxDirectory().Len(), MAX_int32, false);
 				if (LocalFilename.StartsWith(TEXT("Engine/")))
 				{
 					LocalFilename = SandboxFile.GetAbsoluteRootDirectory() / LocalFilename;
diff --git a/Engine/Source/Runtime/SessionServices/Private/SessionManager.cpp b/Engine/Source/Runtime/SessionServices/Private/SessionManager.cpp
--- a/Engine/Source/Runtime/SessionServices/Private/SessionManager.cpp
+++ b/Engine/Source/Runtime/SessionServices/Private/SessionManager.cpp
@@ -22,14 +22,14 @@ FSessionManager::FSessionManager(const TSharedRef<IMessageBus, ESPMode::ThreadSa
 	if (FParse::Value(FCommandLine::Get(), TEXT("SessionFilter="), Filter))
 	{
 		// Allow support for -SessionFilter=Filter1+Filter2+Filter3
-		int32 PlusIdx = Filter.Find(TEXT("+"));
+		int32 PlusIdx = Filter.Find(TEXT("+"), ESearchCase::CaseSensitive);
 
 		while (PlusIdx != INDEX_NONE)
 		{
 			FString Owner = Filter.Left(PlusIdx);
 			FilteredOwners.Add(Owner);
-			Filter = Filter.Right(Filter.Len() - (PlusIdx + 1));
-			PlusIdx = Filter.Find(TEXT("+"));
+			Filter.RightInline(Filter.Len() - (PlusIdx + 1), false);
+			PlusIdx = Filter.Find(TEXT("+"), ESearchCase::CaseSensitive);
 		}
 
 		FilteredOwners.Add(Filter);
diff --git a/Engine/Source/Runtime/Sockets/Private/BSDIPv6Sockets/IPAddressBSDIPv6.cpp b/Engine/Source/Runtime/Sockets/Private/BSDIPv6Sockets/IPAddressBSDIPv6.cpp
--- a/Engine/Source/Runtime/Sockets/Private/BSDIPv6Sockets/IPAddressBSDIPv6.cpp
+++ b/Engine/Source/Runtime/Sockets/Private/BSDIPv6Sockets/IPAddressBSDIPv6.cpp
@@ -91,7 +91,7 @@ void FInternetAddrBSDIPv6::SetIp(const TCHAR* InAddr, bool& bIsValid)
 		if (AddressString.Contains("]:") || (!bIsIPv6 && LastColonIndex != INDEX_NONE))
 		{
 			Port = AddressString.RightChop(LastColonIndex + 1);
-			AddressString = AddressString.Left(LastColonIndex);
+			AddressString.LeftInline(LastColonIndex, false);
 		}
 
 		AddressString.RemoveFromStart("[");
diff --git a/Engine/Source/Runtime/Sockets/Private/BSDSockets/IPAddressBSD.cpp b/Engine/Source/Runtime/Sockets/Private/BSDSockets/IPAddressBSD.cpp
--- a/Engine/Source/Runtime/Sockets/Private/BSDSockets/IPAddressBSD.cpp
+++ b/Engine/Source/Runtime/Sockets/Private/BSDSockets/IPAddressBSD.cpp
@@ -151,7 +151,7 @@ void FInternetAddrBSD::SetIp(const TCHAR* InAddr, bool& bIsValid)
 	if (AddressString.Contains("]:") || (FirstColonIndex == LastColonIndex && LastColonIndex != INDEX_NONE))
 	{
 		Port = AddressString.RightChop(LastColonIndex + 1);
-		AddressString = AddressString.Left(LastColonIndex);
+		AddressString.LeftInline(LastColonIndex, false);
 	}
 
 	// Strip these for backwards compatibility.
@@ -436,10 +436,10 @@ FString FInternetAddrBSD::ToString(bool bAppendPort) const
 		{
 			FString IPv6Str(ANSI_TO_TCHAR(IPStr));
 			// Remove the scope interface if it exists.
-			const int32 InterfaceMarkerIndex = IPv6Str.Find("%", ESearchCase::IgnoreCase, ESearchDir::FromEnd);
+			const int32 InterfaceMarkerIndex = IPv6Str.Find("%", ESearchCase::CaseSensitive, ESearchDir::FromEnd);
 			if (InterfaceMarkerIndex != INDEX_NONE)
 			{
-				IPv6Str = IPv6Str.Left(InterfaceMarkerIndex);
+				IPv6Str.LeftInline(InterfaceMarkerIndex, false);
 			}
 
 			// Using dynamic formatting strings are deprecated.
diff --git a/Engine/Source/Runtime/XmlParser/Private/XmlFile.cpp b/Engine/Source/Runtime/XmlParser/Private/XmlFile.cpp
--- a/Engine/Source/Runtime/XmlParser/Private/XmlFile.cpp
+++ b/Engine/Source/Runtime/XmlParser/Private/XmlFile.cpp
@@ -162,7 +162,7 @@ void FXmlFile::PreProcessInput(TArray<FString>& Input)
 
 		if(NumWhiteSpace > 0)
 		{
-			Input[i] = Input[i].Mid(NumWhiteSpace);
+			Input[i].MidInline(NumWhiteSpace, MAX_int32, false);
 		}
 	}
 
