Clean up code that was using LeftChop incorrectly and modifications to use new inline functions
#rnx

diff --git a/Engine/Source/Runtime/Core/Private/Android/AndroidPlatformFile.cpp b/Engine/Source/Runtime/Core/Private/Android/AndroidPlatformFile.cpp
--- a/Engine/Source/Runtime/Core/Private/Android/AndroidPlatformFile.cpp
+++ b/Engine/Source/Runtime/Core/Private/Android/AndroidPlatformFile.cpp
@@ -122,12 +122,12 @@ const FString &GetFileBasePath()
 
 FString AndroidRelativeToAbsolutePath(bool bUseInternalBasePath, FString RelPath)
 {
-	if (RelPath.StartsWith(TEXT("../")))
+	if (RelPath.StartsWith(TEXT("../"), ESearchCase::CaseSensitive))
 	{
-		while (RelPath.StartsWith(TEXT("../")))
-		{
-			RelPath = RelPath.RightChop(3);
-		}
+		
+		do {
+			RelPath.RightChopInline(3, false);
+		} while (RelPath.StartsWith(TEXT("../")))
 
 		return (bUseInternalBasePath ? GInternalFilePath : GetFileBasePath()) / RelPath;
 	}
@@ -1917,17 +1917,13 @@ private:
 	FString NormalizePath(const TCHAR* Path)
 	{
 		FString Result(Path);
-		Result.ReplaceInline(TEXT("\\"), TEXT("/"));
+		Result.ReplaceInline(TEXT("\\"), TEXT("/"), ESearchCase::CaseSensitive);
 		// This replacement addresses a "bug" where some callers
 		// pass in paths that are badly composed with multiple
 		// subdir separators.
-		Result.ReplaceInline(TEXT("//"), TEXT("/"));
-		if (!Result.IsEmpty() && Result[Result.Len() - 1] == TEXT('/'))
-		{
-			Result.LeftChop(1);
-		}
+		Result.ReplaceInline(TEXT("//"), TEXT("/"), ESearchCase::CaseSensitive);
 		// Remove redundant current-dir references.
-		Result.ReplaceInline(TEXT("/./"), TEXT("/"));
+		Result.ReplaceInline(TEXT("/./"), TEXT("/"), ESearchCase::CaseSensitive);
 		return Result;
 	}
 
@@ -1942,7 +1938,7 @@ private:
 #endif
 		if (!AndroidPath.IsEmpty())
 		{
-			if (AndroidPath.StartsWith(TEXT("/")) ||	// (AllowLocal && AndroidPath.StartsWith(TEXT("/"))) ||		// requiring AllowLocal isn't really useful; system will do permission checks anyway
+			if (AndroidPath.StartsWith(TEXT("/"), ESearchCase::CaseSensitive) ||	// (AllowLocal && AndroidPath.StartsWith(TEXT("/"), ESearchCase::CaseSensitive)) ||		// requiring AllowLocal isn't really useful; system will do permission checks anyway
 				AndroidPath.StartsWith(GFontPathBase) ||
 				AndroidPath.StartsWith(TEXT("/system/etc/")) ||
 				AndroidPath.StartsWith(GInternalFilePath.Left(AndroidPath.Len())) ||
@@ -1954,12 +1950,12 @@ private:
 			}
 			else
 			{
-				while (AndroidPath.StartsWith(TEXT("../")))
+				while (AndroidPath.StartsWith(TEXT("../"), ESearchCase::CaseSensitive))
 				{
-					AndroidPath = AndroidPath.RightChop(3);
+					AndroidPath.RightChopInline(3, false);
 				}
 				AndroidPath.ReplaceInline(FPlatformProcess::BaseDir(), TEXT(""));
-				if (AndroidPath.Equals(TEXT("..")))
+				if (AndroidPath.Equals(TEXT(".."), ESearchCase::CaseSensitive))
 				{
 					AndroidPath = TEXT("");
 				}
diff --git a/Engine/Source/Runtime/Core/Private/Apple/ApplePlatformFile.cpp b/Engine/Source/Runtime/Core/Private/Apple/ApplePlatformFile.cpp
--- a/Engine/Source/Runtime/Core/Private/Apple/ApplePlatformFile.cpp
+++ b/Engine/Source/Runtime/Core/Private/Apple/ApplePlatformFile.cpp
@@ -379,17 +379,13 @@ __thread double FFileHandleApple::AccessTimes[ FFileHandleApple::ACTIVE_HANDLE_C
 FString FApplePlatformFile::NormalizeFilename(const TCHAR* Filename)
 {
 	FString Result(Filename);
-	Result.ReplaceInline(TEXT("\\"), TEXT("/"));
+	Result.ReplaceInline(TEXT("\\"), TEXT("/"), ESearchCase::CaseSensitive);
 	return Result;
 }
 FString FApplePlatformFile::NormalizeDirectory(const TCHAR* Directory)
 {
 	FString Result(Directory);
-	Result.ReplaceInline(TEXT("\\"), TEXT("/"));
-	if (Result.EndsWith(TEXT("/")))
-	{
-		Result.LeftChop(1);
-	}
+	Result.ReplaceInline(TEXT("\\"), TEXT("/"), ESearchCase::CaseSensitive);
 	return Result;
 }
 
diff --git a/Engine/Source/Runtime/Core/Private/IOS/IOSPlatformFile.cpp b/Engine/Source/Runtime/Core/Private/IOS/IOSPlatformFile.cpp
--- a/Engine/Source/Runtime/Core/Private/IOS/IOSPlatformFile.cpp
+++ b/Engine/Source/Runtime/Core/Private/IOS/IOSPlatformFile.cpp
@@ -527,18 +527,14 @@ bool Initialize(IPlatformFile* Inner, const TCHAR* CommandLineParam)
 FString FIOSPlatformFile::NormalizeFilename(const TCHAR* Filename)
 {
 	FString Result(Filename);
-	Result.ReplaceInline(TEXT("\\"), TEXT("/"));
+	Result.ReplaceInline(TEXT("\\"), TEXT("/"), ESearchCase::CaseSensitive);
 	return Result;
 }
 
 FString FIOSPlatformFile::NormalizeDirectory(const TCHAR* Directory)
 {
 	FString Result(Directory);
-	Result.ReplaceInline(TEXT("\\"), TEXT("/"));
-	if (Result.EndsWith(TEXT("/")))
-	{
-		Result.LeftChop(1);
-	}
+	Result.ReplaceInline(TEXT("\\"), TEXT("/"), ESearchCase::CaseSensitive);
 	return Result;
 }
 
