Fix missing semi-colon.
Use case sensitive search
#rnx

diff --git a/Engine/Source/Runtime/Core/Private/Android/AndroidPlatformFile.cpp b/Engine/Source/Runtime/Core/Private/Android/AndroidPlatformFile.cpp
--- a/Engine/Source/Runtime/Core/Private/Android/AndroidPlatformFile.cpp
+++ b/Engine/Source/Runtime/Core/Private/Android/AndroidPlatformFile.cpp
@@ -127,7 +127,7 @@ FString AndroidRelativeToAbsolutePath(bool bUseInternalBasePath, FString RelPath
 		
 		do {
 			RelPath.RightChopInline(3, false);
-		} while (RelPath.StartsWith(TEXT("../")))
+		} while (RelPath.StartsWith(TEXT("../"), ESearchCase::CaseSensitive));
 
 		return (bUseInternalBasePath ? GInternalFilePath : GetFileBasePath()) / RelPath;
 	}
