Implemented FPlatformMisc::GetCPUBrand () for Mac, based on Linux implementation. Also, switched FMacPlatformMisc to use __cpuid intrinsic.

diff --git a/Engine/Source/Runtime/Core/Private/Mac/MacPlatformMisc.cpp b/Engine/Source/Runtime/Core/Private/Mac/MacPlatformMisc.cpp
--- a/Engine/Source/Runtime/Core/Private/Mac/MacPlatformMisc.cpp
+++ b/Engine/Source/Runtime/Core/Private/Mac/MacPlatformMisc.cpp
@@ -45,6 +45,7 @@
 #include <libproc.h>
 #include <notify.h>
 #include <uuid/uuid.h>
+#include <cpuid.h>
 #include "Apple/PostAppleSystemHeaders.h"
 
 extern CORE_API bool GIsGPUCrashed;
@@ -1380,23 +1381,57 @@ FString FMacPlatformMisc::GetCPUVendor()
 	} VendorResult;
 
 
-	int32 Args[4];
-	asm( "cpuid" : "=a" (Args[0]), "=b" (Args[1]), "=c" (Args[2]), "=d" (Args[3]) : "a" (0));
+	int32 CPUInfo[4];
+	__cpuid(0, CPUInfo[0], CPUInfo[1], CPUInfo[2], CPUInfo[3]);
 
-	VendorResult.Dw.dw0 = Args[1];
-	VendorResult.Dw.dw1 = Args[3];
-	VendorResult.Dw.dw2 = Args[2];
+	VendorResult.Dw.dw0 = CPUInfo[1];
+	VendorResult.Dw.dw1 = CPUInfo[3];
+	VendorResult.Dw.dw2 = CPUInfo[2];
 	VendorResult.Buffer[12] = 0;
 
 	return ANSI_TO_TCHAR(VendorResult.Buffer);
 }
 
+FString FMacPlatformMisc::GetCPUBrand()
+{
+	static FString Result = FGenericPlatformMisc::GetCPUBrand();
+	static bool bHaveResult = false;
+
+	if (!bHaveResult)
+	{
+		// @see for more information http://msdn.microsoft.com/en-us/library/vstudio/hskdteyh(v=vs.100).aspx
+		ANSICHAR BrandString[0x40] = { 0 };
+		int32 CPUInfo[4] = { -1 };
+		const SIZE_T CPUInfoSize = sizeof(CPUInfo);
+
+		__cpuid(0x80000000, CPUInfo[0], CPUInfo[1], CPUInfo[2], CPUInfo[3]);
+		const uint32 MaxExtIDs = CPUInfo[0];
+
+		if (MaxExtIDs >= 0x80000004)
+		{
+			const uint32 FirstBrandString = 0x80000002;
+			const uint32 NumBrandStrings = 3;
+			for (uint32 Index = 0; Index < NumBrandStrings; ++Index)
+			{
+				__cpuid(FirstBrandString + Index, CPUInfo[0], CPUInfo[1], CPUInfo[2], CPUInfo[3]);
+				FPlatformMemory::Memcpy(BrandString + CPUInfoSize * Index, CPUInfo, CPUInfoSize);
+			}
+		}
+
+		Result = BrandString;
+
+		bHaveResult = true;
+	}
+
+	return FString(Result);
+}
+
 uint32 FMacPlatformMisc::GetCPUInfo()
 {
-	uint32 Args[4];
-	asm( "cpuid" : "=a" (Args[0]), "=b" (Args[1]), "=c" (Args[2]), "=d" (Args[3]) : "a" (1));
+	uint32 CPUInfo[4];
+	__cpuid(1, CPUInfo[0], CPUInfo[1], CPUInfo[2], CPUInfo[3]);
 
-	return Args[0];
+	return CPUInfo[0];
 }
 
 FText FMacPlatformMisc::GetFileManagerName()
diff --git a/Engine/Source/Runtime/Core/Public/Mac/MacPlatformMisc.h b/Engine/Source/Runtime/Core/Public/Mac/MacPlatformMisc.h
--- a/Engine/Source/Runtime/Core/Public/Mac/MacPlatformMisc.h
+++ b/Engine/Source/Runtime/Core/Public/Mac/MacPlatformMisc.h
@@ -63,6 +63,13 @@ struct CORE_API FMacPlatformMisc : public FApplePlatformMisc
 	 */
 	static FString GetCPUVendor();
 
+	/**
+	 * Uses cpuid instruction to get the CPU brand string
+	 *
+	 * @return    CPU brand string
+	 */
+	static FString GetCPUBrand();
+
 	/**
 	 * Uses cpuid instruction to get the vendor string
 	 *
