Minor optimizations inside FPaths::MakePathRelativeTo
ProjectSavedDir returns const FString& now

diff --git a/Engine/Source/Runtime/Core/Private/Misc/Paths.cpp b/Engine/Source/Runtime/Core/Private/Misc/Paths.cpp
--- a/Engine/Source/Runtime/Core/Private/Misc/Paths.cpp
+++ b/Engine/Source/Runtime/Core/Private/Misc/Paths.cpp
@@ -267,7 +267,7 @@ FString FPaths::ProjectConfigDir()
 	return FPaths::ProjectDir() + TEXT("Config/");
 }
 
-FString FPaths::ProjectSavedDir()
+const FString& FPaths::ProjectSavedDir()
 {
 	FStaticData& StaticData = TLazySingleton<FStaticData>::Get();
 	if (!StaticData.bGameSavedDirInitialized)
@@ -1099,9 +1099,8 @@ void FPaths::MakePlatformFilename( FString& InPath )
 bool FPaths::MakePathRelativeTo( FString& InPath, const TCHAR* InRelativeTo )
 {
 	FString Target = FPaths::ConvertRelativePathToFull(InPath);
-	FString Source = FPaths::ConvertRelativePathToFull(InRelativeTo);
-	
-	Source = FPaths::GetPath(Source);
+	FString Source = FPaths::GetPath(FPaths::ConvertRelativePathToFull(InRelativeTo));
+
 	Source.ReplaceInline(TEXT("\\"), TEXT("/"), ESearchCase::CaseSensitive);
 	Target.ReplaceInline(TEXT("\\"), TEXT("/"), ESearchCase::CaseSensitive);
 
@@ -1142,7 +1141,7 @@ bool FPaths::MakePathRelativeTo( FString& InPath, const TCHAR* InRelativeTo )
 		}
 	}
 	
-	InPath = Result;
+	InPath = MoveTemp(Result);
 	return true;
 }
 
diff --git a/Engine/Source/Runtime/Core/Public/Misc/Paths.h b/Engine/Source/Runtime/Core/Public/Misc/Paths.h
--- a/Engine/Source/Runtime/Core/Public/Misc/Paths.h
+++ b/Engine/Source/Runtime/Core/Public/Misc/Paths.h
@@ -193,7 +193,7 @@ public:
 	 *
 	 * @return saved directory
 	 */
-	static FString ProjectSavedDir();
+	static const FString& ProjectSavedDir();
 
 	UE_DEPRECATED(4.18, "FPaths::GameSavedDir() has been superseded by FPaths::ProjectSavedDir().")
 	static FORCEINLINE FString GameSavedDir() { return ProjectSavedDir(); }
