fix for crash when trying to remove a pointer from the LLM map when it wasn't added in the first place

diff --git a/Engine/Source/Runtime/Core/Private/HAL/LowLevelMemTracker.cpp b/Engine/Source/Runtime/Core/Private/HAL/LowLevelMemTracker.cpp
--- a/Engine/Source/Runtime/Core/Private/HAL/LowLevelMemTracker.cpp
+++ b/Engine/Source/Runtime/Core/Private/HAL/LowLevelMemTracker.cpp
@@ -1316,6 +1316,10 @@ void FLLMTracker::TrackFree(const void* Ptr, ELLMTracker Tracker, ELLMAllocType
 	}
 
 	// look up the pointer in the tracking map
+	if (!GetAllocationMap().HasKey(Ptr))
+	{
+		return;
+	}
 	LLMMap::Values Values = GetAllocationMap().Remove(Ptr);
 	uint64 Size = Values.Value1;
 	FLLMTracker::FLowLevelAllocInfo AllocInfo = Values.Value2;
