CIS fix

diff --git a/Engine/Source/Runtime/Core/Private/Mac/MacPlatformMisc.cpp b/Engine/Source/Runtime/Core/Private/Mac/MacPlatformMisc.cpp
--- a/Engine/Source/Runtime/Core/Private/Mac/MacPlatformMisc.cpp
+++ b/Engine/Source/Runtime/Core/Private/Mac/MacPlatformMisc.cpp
@@ -45,7 +45,6 @@
 #include <libproc.h>
 #include <notify.h>
 #include <uuid/uuid.h>
-#include <cpuid.h>
 #include "Apple/PostAppleSystemHeaders.h"
 
 extern CORE_API bool GIsGPUCrashed;
@@ -1381,12 +1380,12 @@ FString FMacPlatformMisc::GetCPUVendor()
 	} VendorResult;
 
 
-	int32 CPUInfo[4];
-	__cpuid(0, CPUInfo[0], CPUInfo[1], CPUInfo[2], CPUInfo[3]);
+	int32 Args[4];
+	asm( "cpuid" : "=a" (Args[0]), "=b" (Args[1]), "=c" (Args[2]), "=d" (Args[3]) : "a" (0));
 
-	VendorResult.Dw.dw0 = CPUInfo[1];
-	VendorResult.Dw.dw1 = CPUInfo[3];
-	VendorResult.Dw.dw2 = CPUInfo[2];
+	VendorResult.Dw.dw0 = Args[1];
+	VendorResult.Dw.dw1 = Args[3];
+	VendorResult.Dw.dw2 = Args[2];
 	VendorResult.Buffer[12] = 0;
 
 	return ANSI_TO_TCHAR(VendorResult.Buffer);
@@ -1404,7 +1403,7 @@ FString FMacPlatformMisc::GetCPUBrand()
 		int32 CPUInfo[4] = { -1 };
 		const SIZE_T CPUInfoSize = sizeof(CPUInfo);
 
-		__cpuid(0x80000000, CPUInfo[0], CPUInfo[1], CPUInfo[2], CPUInfo[3]);
+		asm( "cpuid" : "=a" (CPUInfo[0]), "=b" (CPUInfo[1]), "=c" (CPUInfo[2]), "=d" (CPUInfo[3]) : "a" (0x80000000));
 		const uint32 MaxExtIDs = CPUInfo[0];
 
 		if (MaxExtIDs >= 0x80000004)
@@ -1413,7 +1412,7 @@ FString FMacPlatformMisc::GetCPUBrand()
 			const uint32 NumBrandStrings = 3;
 			for (uint32 Index = 0; Index < NumBrandStrings; ++Index)
 			{
-				__cpuid(FirstBrandString + Index, CPUInfo[0], CPUInfo[1], CPUInfo[2], CPUInfo[3]);
+				asm( "cpuid" : "=a" (CPUInfo[0]), "=b" (CPUInfo[1]), "=c" (CPUInfo[2]), "=d" (CPUInfo[3]) : "a" (FirstBrandString + Index));
 				FPlatformMemory::Memcpy(BrandString + CPUInfoSize * Index, CPUInfo, CPUInfoSize);
 			}
 		}
@@ -1428,10 +1427,10 @@ FString FMacPlatformMisc::GetCPUBrand()
 
 uint32 FMacPlatformMisc::GetCPUInfo()
 {
-	uint32 CPUInfo[4];
-	__cpuid(1, CPUInfo[0], CPUInfo[1], CPUInfo[2], CPUInfo[3]);
+	uint32 Args[4];
+	asm( "cpuid" : "=a" (Args[0]), "=b" (Args[1]), "=c" (Args[2]), "=d" (Args[3]) : "a" (1));
 
-	return CPUInfo[0];
+	return Args[0];
 }
 
 FText FMacPlatformMisc::GetFileManagerName()
