Fix missing ;
#rnx

diff --git a/Engine/Source/Runtime/Core/Private/Lumin/LuminPlatformFile.cpp b/Engine/Source/Runtime/Core/Private/Lumin/LuminPlatformFile.cpp
--- a/Engine/Source/Runtime/Core/Private/Lumin/LuminPlatformFile.cpp
+++ b/Engine/Source/Runtime/Core/Private/Lumin/LuminPlatformFile.cpp
@@ -30,7 +30,7 @@ FString AndroidRelativeToAbsolutePath(bool bUseInternalBasePath, FString RelPath
 		do 
 		{
 			Result.RightChopInline(3, false);
-		} while (Result.StartsWith(TEXT("../"), ESearchCase::CaseSensitive))
+		} while (Result.StartsWith(TEXT("../"), ESearchCase::CaseSensitive));
 	}
 
 	// Remove the base app path if present, we will prepend it the correct base path as needed.
