Properly handle dispatch of the serialized property on FArchiveProxy to stop crashes in the cooker when doing non-determinism checks on arrays which differ.

diff --git a/Engine/Source/Runtime/Core/Public/Serialization/Archive.h b/Engine/Source/Runtime/Core/Public/Serialization/Archive.h
--- a/Engine/Source/Runtime/Core/Public/Serialization/Archive.h
+++ b/Engine/Source/Runtime/Core/Public/Serialization/Archive.h
@@ -1209,7 +1209,7 @@ public:
 	 * 
 	 * @param InProperty Pointer to the property that is currently being serialized
 	 */
-	FORCEINLINE void SetSerializedProperty(UProperty* InProperty)
+	FORCEINLINE virtual void SetSerializedProperty(UProperty* InProperty)
 	{
 		SerializedProperty = InProperty;
 	}
@@ -1242,7 +1242,7 @@ public:
 	/**
 	 * Set the raw serialized property chain for this archive, optionally overriding the serialized property too (or null to use the head of the property chain)
 	 */
-	void SetSerializedPropertyChain(const FArchiveSerializedPropertyChain* InSerializedPropertyChain, class UProperty* InSerializedPropertyOverride = nullptr);
+	virtual void SetSerializedPropertyChain(const FArchiveSerializedPropertyChain* InSerializedPropertyChain, class UProperty* InSerializedPropertyOverride = nullptr);
 
 	/**
 	 * Push a property that is currently being serialized onto the property stack
diff --git a/Engine/Source/Runtime/Core/Public/Serialization/ArchiveProxy.h b/Engine/Source/Runtime/Core/Public/Serialization/ArchiveProxy.h
--- a/Engine/Source/Runtime/Core/Public/Serialization/ArchiveProxy.h
+++ b/Engine/Source/Runtime/Core/Public/Serialization/ArchiveProxy.h
@@ -236,14 +236,28 @@ public:
 	}
 #endif
 
+	FORCEINLINE void SetSerializedProperty(UProperty* InProperty) override
+	{
+		FArchive::SetSerializedProperty(InProperty);
+		InnerArchive.SetSerializedProperty(InProperty);
+	}
+
+	void SetSerializedPropertyChain(const FArchiveSerializedPropertyChain* InSerializedPropertyChain, class UProperty* InSerializedPropertyOverride = nullptr) override
+	{
+		FArchive::SetSerializedPropertyChain(InSerializedPropertyChain, InSerializedPropertyOverride);
+		InnerArchive.SetSerializedPropertyChain(InSerializedPropertyChain, InSerializedPropertyOverride);
+	}
+
 	/** Pushes editor-only marker to the stack of currently serialized properties */
 	virtual FORCEINLINE void PushSerializedProperty(class UProperty* InProperty, const bool bIsEditorOnlyProperty)
 	{
+		FArchive::PushSerializedProperty(InProperty, bIsEditorOnlyProperty);
 		InnerArchive.PushSerializedProperty(InProperty, bIsEditorOnlyProperty);
 	}
 	/** Pops editor-only marker from the stack of currently serialized properties */
 	virtual FORCEINLINE void PopSerializedProperty(class UProperty* InProperty, const bool bIsEditorOnlyProperty)
 	{
+		FArchive::PopSerializedProperty(InProperty, bIsEditorOnlyProperty);
 		InnerArchive.PopSerializedProperty(InProperty, bIsEditorOnlyProperty);
 	}
 #if WITH_EDITORONLY_DATA
