Memory savings from FString, in lieu of CL# 10207395 being backed out.

diff --git a/Engine/Source/Runtime/Core/Public/Containers/UnrealString.h b/Engine/Source/Runtime/Core/Public/Containers/UnrealString.h
--- a/Engine/Source/Runtime/Core/Public/Containers/UnrealString.h
+++ b/Engine/Source/Runtime/Core/Public/Containers/UnrealString.h
@@ -113,6 +113,7 @@ public:
 		{
 			int32 SrcLen  = TCString<CharType>::Strlen(Src) + 1;
 			int32 DestLen = FPlatformString::ConvertedLength<TCHAR>(Src, SrcLen);
+			Data.Reserve(DestLen);
 			Data.AddUninitialized(DestLen);
 
 			FPlatformString::Convert(Data.GetData(), DestLen, Src, SrcLen);
@@ -136,6 +137,7 @@ public:
 			int32 DestLen = FPlatformString::ConvertedLength<TCHAR>(InSrc, InCount);
 			if (DestLen > 0)
 			{
+				Data.Reserve(DestLen + 1);
 				Data.AddUninitialized(DestLen + 1);
 
 				FPlatformString::Convert(Data.GetData(), DestLen, InSrc, InCount);
@@ -163,6 +165,7 @@ public:
 			if (CFStringGetBytes((__bridge CFStringRef)In, Range, Encoding, '?', false, NULL, 0, &BytesNeeded) > 0)
 			{
 				const size_t Length = BytesNeeded / sizeof(TCHAR);
+				Data.Reserve(Length + 1);
 				Data.AddUninitialized(Length + 1);
 				CFStringGetBytes((__bridge CFStringRef)In, Range, Encoding, '?', false, (uint8*)Data.GetData(), Length * sizeof(TCHAR) + 1, NULL);
 				Data[Length] = 0;
@@ -645,6 +648,7 @@ private:
 		// the memory could be reused here without constructing a new object.  However, until there is proof otherwise,
 		// I believe this will be relatively rare and isn't worth making the code a lot more complex right now.
 		FString Result;
+		Result.Data.Reserve(LhsLen + RhsLen + 1);
 		Result.Data.AddUninitialized(LhsLen + RhsLen + 1);
 
 		TCHAR* ResultData = Result.Data.GetData();
