Fix FPaths::IsUnderDirectory() not working correctly for paths which are not under the engine root, causing failure to find correct editor target at startup.

#rnx

diff --git a/Engine/Source/Runtime/Core/Private/Misc/Paths.cpp b/Engine/Source/Runtime/Core/Private/Misc/Paths.cpp
--- a/Engine/Source/Runtime/Core/Private/Misc/Paths.cpp
+++ b/Engine/Source/Runtime/Core/Private/Misc/Paths.cpp
@@ -1443,12 +1443,9 @@ bool FPaths::IsSamePath(const FString& PathA, const FString& PathB)
 
 bool FPaths::IsUnderDirectory(const FString& InPath, const FString& InDirectory)
 {
-	FString Path = InPath;
-	MakeStandardFilename(Path);
-
-	FString Directory = InDirectory;
-	MakeStandardFilename(Directory);
+	FString Path = FPaths::ConvertRelativePathToFull(InPath);
 
+	FString Directory = FPaths::ConvertRelativePathToFull(InDirectory);
 	if (Directory.EndsWith(TEXT("/")))
 	{
 		Directory.RemoveAt(Directory.Len() - 1);
