fix for crash during startup on iPone 6+
turned off MALLOC_NANO usage for small block allocations
fixed an issue where we can try to remove a pointer from LLM that we aren't tracking
fixed a default alignment issue

diff --git a/Engine/Source/Runtime/Core/Public/HAL/MallocBinned.h b/Engine/Source/Runtime/Core/Public/HAL/MallocBinned.h
--- a/Engine/Source/Runtime/Core/Public/HAL/MallocBinned.h
+++ b/Engine/Source/Runtime/Core/Public/HAL/MallocBinned.h
@@ -20,7 +20,7 @@
 #endif
 
 #ifndef USE_OS_SMALL_BLOCK_ALLOC
-#define USE_OS_SMALL_BLOCK_ALLOC PLATFORM_IOS
+#define USE_OS_SMALL_BLOCK_ALLOC 0
 #endif //USE_OS_SMALL_BLOCK_ALLOC
 
 #ifndef USE_OS_SMALL_BLOCK_GRAB_MEMORY_FROM_OS
