Revert a one liner introduced to see the progress as soon as it is created to fix a crash when using the progress while recompiling all materials

#rnx

diff --git a/Engine/Source/Runtime/Core/Private/Misc/SlowTask.cpp b/Engine/Source/Runtime/Core/Private/Misc/SlowTask.cpp
--- a/Engine/Source/Runtime/Core/Private/Misc/SlowTask.cpp
+++ b/Engine/Source/Runtime/Core/Private/Misc/SlowTask.cpp
@@ -42,8 +42,6 @@ void FSlowTask::Initialize()
 	if (bEnabled)
 	{
 		Context.ScopeStack.Push(this);
-
-		Context.RequestUpdateUI();
 	}
 }
 
