Add version of InaccurateFindNearest which has a segment out parameter.

diff --git a/Engine/Source/Runtime/Core/Public/Math/InterpCurve.h b/Engine/Source/Runtime/Core/Public/Math/InterpCurve.h
--- a/Engine/Source/Runtime/Core/Public/Math/InterpCurve.h
+++ b/Engine/Source/Runtime/Core/Public/Math/InterpCurve.h
@@ -98,6 +98,16 @@ public:
 	 */
 	float InaccurateFindNearest( const T &PointInSpace, float& OutDistanceSq ) const;
 
+	/**
+	* Find the nearest point on spline to the given point.
+	 *
+	* @param PointInSpace - the given point
+	* @param OutDistanceSq - output - the squared distance between the given point and the closest found point.
+	* @param OutSegment - output - the nearest segment to the given point.
+	* @return The key (the 't' parameter) of the nearest point.
+	*/
+	float InaccurateFindNearest( const T &PointInSpace, float& OutDistanceSq, float& OutSegment ) const;
+
 	/** 
 	 * Find the nearest point (to the given point) on segment between Points[PtIdx] and Points[PtIdx+1]
 	 *
@@ -473,6 +483,13 @@ T FInterpCurve<T>::EvalSecondDerivative(const float InVal, const T& Default) con
 
 template< class T >
 float FInterpCurve<T>::InaccurateFindNearest(const T &PointInSpace, float& OutDistanceSq) const
+{
+	float OutSegment;
+	return InaccurateFindNearest(PointInSpace, OutDistanceSq, OutSegment);
+}
+
+template< class T >
+float FInterpCurve<T>::InaccurateFindNearest(const T &PointInSpace, float& OutDistanceSq, float& OutSegment) const
 {
 	const int32 NumPoints = Points.Num();
 	const int32 NumSegments = bIsLooped ? NumPoints : NumPoints - 1;
@@ -481,6 +498,7 @@ float FInterpCurve<T>::InaccurateFindNearest(const T &PointInSpace, float& OutDi
 	{
 		float BestDistanceSq;
 		float BestResult = InaccurateFindNearestOnSegment(PointInSpace, 0, BestDistanceSq);
+		float BestSegment = 0;
 		for (int32 Segment = 1; Segment < NumSegments; ++Segment)
 		{
 			float LocalDistanceSq;
@@ -489,15 +507,18 @@ float FInterpCurve<T>::InaccurateFindNearest(const T &PointInSpace, float& OutDi
 			{
 				BestDistanceSq = LocalDistanceSq;
 				BestResult = LocalResult;
+				BestSegment = Segment;
 			}
 		}
 		OutDistanceSq = BestDistanceSq;
+		OutSegment = BestSegment;
 		return BestResult;
 	}
 
 	if (NumPoints == 1)
 	{
 		OutDistanceSq = (PointInSpace - Points[0].OutVal).SizeSquared();
+		OutSegment = 0;
 		return Points[0].InVal;
 	}
 
