Move module manager extra search path configuration out of FModuleManager::Get() and into a seperate function which is called from the launch engine loop AFTER platform files are created. Ensures that editor builds that use pak files have the paks mounted before we start scanning for data driven platform ini files

josh.adams, luke.thatcher, jeff.newquist

diff --git a/Engine/Source/Runtime/Core/Private/Modules/ModuleManager.cpp b/Engine/Source/Runtime/Core/Private/Modules/ModuleManager.cpp
--- a/Engine/Source/Runtime/Core/Private/Modules/ModuleManager.cpp
+++ b/Engine/Source/Runtime/Core/Private/Modules/ModuleManager.cpp
@@ -73,22 +73,6 @@ FModuleManager::FModuleManager()
 	: bCanProcessNewlyLoadedObjects(false)
 {
 	check(IsInGameThread());
-
-#if PLATFORM_DESKTOP && !IS_MONOLITHIC
-	// Ensure that dependency dlls can be found in restricted sub directories
-	TArray<FString> RestrictedFolderNames = { TEXT("NoRedist"), TEXT("NotForLicensees"), TEXT("CarefullyRedist") };
-	RestrictedFolderNames.Append(FDataDrivenPlatformInfoRegistry::GetConfidentialPlatforms());
-
-	FString ModuleDir = FPlatformProcess::GetModulesDirectory();
-	for (const FString& RestrictedFolderName : RestrictedFolderNames)
-	{
-		FString RestrictedFolder = ModuleDir / RestrictedFolderName;
-		if (FPaths::DirectoryExists(RestrictedFolder))
-		{
-			AddBinariesDirectory(*RestrictedFolder, false);
-		}
-	}
-#endif
 }
 
 FModuleManager::~FModuleManager()
@@ -1098,6 +1082,28 @@ bool FModuleManager::LoadModuleWithCallback( const FName InModuleName, FOutputDe
 	return true;
 }
 
+void FModuleManager::AddExtraBinarySearchPaths()
+{
+	static bool bAdded = false;
+	if (!bAdded)
+	{
+		// Ensure that dependency dlls can be found in restricted sub directories
+		TArray<FString> RestrictedFolderNames = { TEXT("NoRedist"), TEXT("NotForLicensees"), TEXT("CarefullyRedist") };
+		RestrictedFolderNames.Append(FDataDrivenPlatformInfoRegistry::GetConfidentialPlatforms());
+
+		FString ModuleDir = FPlatformProcess::GetModulesDirectory();
+		for (const FString& RestrictedFolderName : RestrictedFolderNames)
+		{
+			FString RestrictedFolder = ModuleDir / RestrictedFolderName;
+			if (FPaths::DirectoryExists(RestrictedFolder))
+			{
+				AddBinariesDirectory(*RestrictedFolder, false);
+			}
+		}
+
+		bAdded = true;
+	}
+}
 
 void FModuleManager::MakeUniqueModuleFilename( const FName InModuleName, FString& UniqueSuffix, FString& UniqueModuleFileName ) const
 {
diff --git a/Engine/Source/Runtime/Core/Public/Modules/ModuleManager.h b/Engine/Source/Runtime/Core/Public/Modules/ModuleManager.h
--- a/Engine/Source/Runtime/Core/Public/Modules/ModuleManager.h
+++ b/Engine/Source/Runtime/Core/Public/Modules/ModuleManager.h
@@ -247,6 +247,11 @@ public:
 	 */
 	void AbandonModuleWithCallback( const FName InModuleName );
 
+	/**
+	 * Add any extra search paths that may be required
+	 */
+	void AddExtraBinarySearchPaths();
+
 	/**
 	  * Gets a module by name, checking to ensure it exists.
 	  *
diff --git a/Engine/Source/Runtime/Launch/Private/LaunchEngineLoop.cpp b/Engine/Source/Runtime/Launch/Private/LaunchEngineLoop.cpp
--- a/Engine/Source/Runtime/Launch/Private/LaunchEngineLoop.cpp
+++ b/Engine/Source/Runtime/Launch/Private/LaunchEngineLoop.cpp
@@ -1431,6 +1431,11 @@ int32 FEngineLoop::PreInitPreStartupScreen(const TCHAR* CmdLine)
 		}
 	}
 
+	// 
+#if PLATFORM_DESKTOP && !IS_MONOLITHIC
+	FModuleManager::Get().AddExtraBinarySearchPaths();
+#endif
+
 	// Initialize file manager
 	{
 		SCOPED_BOOT_TIMING("IFileManager::Get().ProcessCommandLineOptions");
