Fix collapse relative directories beginning with dots

diff --git a/Engine/Source/Runtime/Core/Private/Misc/Paths.cpp b/Engine/Source/Runtime/Core/Private/Misc/Paths.cpp
--- a/Engine/Source/Runtime/Core/Private/Misc/Paths.cpp
+++ b/Engine/Source/Runtime/Core/Private/Misc/Paths.cpp
@@ -999,16 +999,42 @@ bool FPaths::CollapseRelativeDirectories(FString& InPath)
 	{
 		// An empty path is finished
 		if (InPath.IsEmpty())
+		{
 			break;
+		}
 
 		// Consider empty paths or paths which start with .. or /.. as invalid
 		if (InPath.StartsWith(TEXT(".."), ESearchCase::CaseSensitive) || InPath.StartsWith(ParentDir))
+		{
 			return false;
+		}
 
 		// If there are no "/.."s left then we're done
-		const int32 Index = InPath.Find(ParentDir, ESearchCase::CaseSensitive);
+		int32 Index = InPath.Find(ParentDir, ESearchCase::CaseSensitive);
+		if (Index == -1)
+		{
+			break;
+		}
+
+		// Ignore folders beginning with dots
+		for (;;)
+		{
+			if (InPath.Len() <= Index + ParentDirLength || InPath[Index + ParentDirLength] == TEXT('/'))
+			{
+				break;
+			}
+
+			Index = InPath.Find(ParentDir, ESearchCase::CaseSensitive, ESearchDir::FromStart, Index + ParentDirLength);
+			if (Index == -1)
+			{
+				break;
+			}
+		}
+
 		if (Index == -1)
+		{
 			break;
+		}
 
 		int32 PreviousSeparatorIndex = Index;
 		for (;;)
@@ -1018,17 +1044,23 @@ bool FPaths::CollapseRelativeDirectories(FString& InPath)
 
 			// Stop if we've hit the start of the string
 			if (PreviousSeparatorIndex == 0)
+			{
 				break;
+			}
 
 			// Stop if we've found a directory that isn't "/./"
 			if ((Index - PreviousSeparatorIndex) > 1 && (InPath[PreviousSeparatorIndex + 1] != TEXT('.') || InPath[PreviousSeparatorIndex + 2] != TEXT('/')))
+			{
 				break;
+			}
 		}
 
 		// If we're attempting to remove the drive letter, that's illegal
 		int32 Colon = InPath.Find(TEXT(":"), ESearchCase::CaseSensitive, ESearchDir::FromStart, PreviousSeparatorIndex);
 		if (Colon >= 0 && Colon < Index)
+		{
 			return false;
+		}
 
 		InPath.RemoveAt(PreviousSeparatorIndex, Index - PreviousSeparatorIndex + ParentDirLength, false);
 	}
diff --git a/Engine/Source/Runtime/Core/Private/Tests/Misc/PathsTest.cpp b/Engine/Source/Runtime/Core/Private/Tests/Misc/PathsTest.cpp
--- a/Engine/Source/Runtime/Core/Private/Tests/Misc/PathsTest.cpp
+++ b/Engine/Source/Runtime/Core/Private/Tests/Misc/PathsTest.cpp
@@ -92,6 +92,11 @@ bool FPathTests::RunTest( const FString& Parameters )
 		RunCollapseRelativeDirectoriesTest(TEXT("./Folder/.svn"),                                        TEXT("Folder/.svn"));
 		RunCollapseRelativeDirectoriesTest(TEXT("./.svn/../.svn"),                                       TEXT(".svn"));
 		RunCollapseRelativeDirectoriesTest(TEXT(".svn/./.svn/.././../.svn"),                             TEXT("/.svn"));
+		RunCollapseRelativeDirectoriesTest(TEXT("C:/Folder1/./Folder2/..Folder3"),						 TEXT("C:/Folder1/Folder2/..Folder3"));
+		RunCollapseRelativeDirectoriesTest(TEXT("C:/Folder1/./Folder2/..Folder3/Folder4"),				 TEXT("C:/Folder1/Folder2/..Folder3/Folder4"));
+		RunCollapseRelativeDirectoriesTest(TEXT("C:/Folder1/./Folder2/..Folder3/..Folder4"),			 TEXT("C:/Folder1/Folder2/..Folder3/..Folder4"));
+		RunCollapseRelativeDirectoriesTest(TEXT("C:/Folder1/./Folder2/..Folder3/Folder4/../Folder5"),	 TEXT("C:/Folder1/Folder2/..Folder3/Folder5"));
+		RunCollapseRelativeDirectoriesTest(TEXT("C:/Folder1/..Folder2/Folder3/..Folder4/../Folder5"),	 TEXT("C:/Folder1/..Folder2/Folder3/Folder5"));
 	}
 
 	// Extension texts
