Add thread callstacks to android crash reports
 - Add per thread stack walk functionality.
 - Add functionality to android crash context to use all thread manager thread's callstacks.
 - Before report upload, add all thread xml data into crash report
 - tidy up signal handler's separate stack behaviour. it is now shared by both fatal signal and thread backtrace signal.
#rnx

diff --git a/Engine/Source/Runtime/Core/Private/Android/AndroidPlatformCrashContext.cpp b/Engine/Source/Runtime/Core/Private/Android/AndroidPlatformCrashContext.cpp
--- a/Engine/Source/Runtime/Core/Private/Android/AndroidPlatformCrashContext.cpp
+++ b/Engine/Source/Runtime/Core/Private/Android/AndroidPlatformCrashContext.cpp
@@ -18,6 +18,8 @@
 #include "GenericPlatform/GenericPlatformOutputDevices.h"
 #include "HAL/FileManager.h"
 #include "HAL/IConsoleManager.h"
+#include "HAL/ThreadManager.h"
+#include "HAL/RunnableThread.h"
 
 static int64 GetAndroidLibraryBaseAddress();
 
@@ -137,31 +139,103 @@ static void CrashReportFileCopy(const char* DestPath, const char* SourcePath)
 
 void FAndroidCrashContext::StoreCrashInfo() const
 {
-	char CrashDirectoryName[CrashReportMaxPathSize] = { 0 };
 	char FilePath[CrashReportMaxPathSize] = { 0 };
-	if (GetType() == ECrashContextType::Ensure)
-	{
-		// create a new report folder.
-		GenerateReportDirectoryName(CrashDirectoryName);
-	}
-	else
-	{
-		GetCrashDirectoryName(CrashDirectoryName);
-	}
-
-	FCStringAnsi::Strcpy(FilePath, CrashDirectoryName);
+	FCStringAnsi::Strcpy(FilePath, ReportDirectory);
 	FCStringAnsi::Strcat(FilePath, "/");
 	FCStringAnsi::Strcat(FilePath, FGenericCrashContext::CrashContextRuntimeXMLNameA);
 	SerializeAsXML(*FString(FilePath)); // CreateFileWriter will also create destination directory.
 
 	// copy log:
-	FCStringAnsi::Strcpy(FilePath, CrashDirectoryName);
+	FCStringAnsi::Strcpy(FilePath, ReportDirectory);
 	FCStringAnsi::Strcat(FilePath, "/");
 	FCStringAnsi::Strcat(FilePath, FCStringAnsi::Strlen(GAndroidCrashInfo.AppName) ? GAndroidCrashInfo.AppName : "UE4");
 	FCStringAnsi::Strcat(FilePath, ".log");
 	CrashReportFileCopy(FilePath, GAndroidCrashInfo.AppLogPath);
 }
 
+
+//Create a separate file containing thread context info (callstacks etc) in xml form
+// This is added to the crash report xml at during pre-processing time.
+void FAndroidCrashContext::DumpAllThreadCallstacks() const
+{
+	char FilePath[FAndroidCrashContext::CrashReportMaxPathSize] = { 0 };
+	FCStringAnsi::Strcpy(FilePath, ReportDirectory);
+	FCStringAnsi::Strcat(FilePath, FAndroidCrashContext::CrashReportMaxPathSize, "/AllThreads.txt");
+	TArray<FCrashStackFrame> CrashStackFrames;
+	CrashStackFrames.Empty(32);
+	uint32 CallstacksRecorded = 0;
+	int DestHandle = open(FilePath, O_CREAT | O_WRONLY, S_IRUSR | S_IWUSR);
+	if (DestHandle >= 0)
+	{
+		auto Write = [](int FileHandle, const ANSICHAR* Buffer)
+		{
+			write(FileHandle, Buffer, FCStringAnsi::Strlen(Buffer));
+		};
+		auto Writeln = [](int FileHandle, const ANSICHAR* Buffer)
+		{
+			write(FileHandle, Buffer, FCStringAnsi::Strlen(Buffer));
+			write(FileHandle, "\n", 1);
+		};
+
+		// For each thread append it's info to the file.
+		FThreadManager::Get().ForEachThread([Writeln, Write, this, DestHandle, &CrashStackFrames, &CallstacksRecorded](uint32 ThreadID, FRunnableThread* Runnable)
+		{
+			// Capture the stack trace
+			static const int StackTraceMaxDepth = 100;
+			uint64 StackTrace[StackTraceMaxDepth];
+			FMemory::Memzero(StackTrace);
+			uint32 Depth = FPlatformStackWalk::CaptureThreadStackBackTrace(ThreadID, StackTrace, StackTraceMaxDepth);
+
+			if (Depth)
+			{
+				ANSICHAR Line[256];
+				if (CallstacksRecorded++ == 0)
+				{
+					Writeln(DestHandle, "<Threads>");
+				}
+				Writeln(DestHandle, "<Thread>");
+				Write(DestHandle, "<CallStack>");
+				// Write stack
+				GetPortableCallStack(StackTrace, Depth, CrashStackFrames);
+				for (const FCrashStackFrame& CrashStackFrame : CrashStackFrames)
+				{
+					FCStringAnsi::Strncpy(Line, TCHAR_TO_UTF8(*CrashStackFrame.ModuleName), UE_ARRAY_COUNT(Line));
+					FCStringAnsi::Strcat(Line, " 0x");
+					FCStringAnsi::Strcat(Line, ItoANSI(CrashStackFrame.BaseAddress, (uint64)16, 16));
+					FCStringAnsi::Strcat(Line, " + ");
+					FCStringAnsi::Strcat(Line, ItoANSI(CrashStackFrame.Offset, (uint64)16, 16));
+					Writeln(DestHandle, Line);
+				}
+				Writeln(DestHandle, "</CallStack>");
+				Writeln(DestHandle, "<IsCrashed>false</IsCrashed>");
+				Writeln(DestHandle, "<Registers/>");
+				// write Thread Id
+				FCStringAnsi::Strncpy(Line, ItoANSI((uint64)ThreadID, (uint64)10), UE_ARRAY_COUNT(Line));
+				Write(DestHandle, "<ThreadID>");
+				Write(DestHandle, Line);
+				Writeln(DestHandle, "</ThreadID>");
+
+				// write Thread Name
+				Write(DestHandle, "<ThreadName>");
+				FCStringAnsi::Strncpy(Line, TCHAR_TO_UTF8(*Runnable->GetThreadName()), UE_ARRAY_COUNT(Line));
+				Write(DestHandle, Line);
+				Writeln(DestHandle, "</ThreadName>");
+
+				Writeln(DestHandle, "</Thread>");
+			}
+		});
+		if (CallstacksRecorded)
+		{
+			Writeln(DestHandle, "</Threads>");
+		}
+		close(DestHandle);
+		if (CallstacksRecorded == 0)
+		{
+			unlink(FilePath); // remove the file if nothing was written.
+		}
+	}
+}
+
 void FAndroidCrashContext::Initialize()
 {
 	GAndroidCrashInfo.Init();
@@ -173,6 +247,15 @@ FAndroidCrashContext::FAndroidCrashContext(ECrashContextType InType, const TCHAR
 , Info(NULL)
 , Context(NULL)
 {
+	if (GetType() == ECrashContextType::Ensure)
+	{
+		// create a new report folder.
+		GenerateReportDirectoryName(ReportDirectory);
+	}
+	else
+	{
+		GetCrashDirectoryName(ReportDirectory);
+	}
 }
 
 void FAndroidCrashContext::CaptureCrashInfo()
@@ -243,7 +326,7 @@ void FAndroidCrashContext::AddPlatformSpecificProperties() const
 	}
 }
 
-void FAndroidCrashContext::GetPortableCallStack(const uint64* StackFrames, int32 NumStackFrames, TArray<FCrashStackFrame>& OutCallStack)
+void FAndroidCrashContext::GetPortableCallStack(const uint64* StackFrames, int32 NumStackFrames, TArray<FCrashStackFrame>& OutCallStack) const
 {
 	// Update the callstack with offsets from each module
 	OutCallStack.Reset(NumStackFrames);
diff --git a/Engine/Source/Runtime/Core/Private/Android/AndroidPlatformMisc.cpp b/Engine/Source/Runtime/Core/Private/Android/AndroidPlatformMisc.cpp
--- a/Engine/Source/Runtime/Core/Private/Android/AndroidPlatformMisc.cpp
+++ b/Engine/Source/Runtime/Core/Private/Android/AndroidPlatformMisc.cpp
@@ -45,6 +45,8 @@
 #include "Templates/AlignmentTemplates.h"
 #include "GenericPlatform/GenericPlatformOutputDevices.h"
 
+#include "Android/AndroidPlatformStackWalk.h"
+
 #if USE_ANDROID_JNI
 extern AAssetManager * AndroidThunkCpp_GetAssetManager();
 extern int32 GAndroidPackageVersion;
@@ -891,7 +893,7 @@ void DefaultCrashHandler(const FAndroidCrashContext& Context)
 /** Global pointer to crash handler */
 void (* GCrashHandlerPointer)(const FGenericCrashContext& Context) = NULL;
 
-const int32 TargetSignals[] =
+static constexpr int32 TargetSignals[] =
 {
 	SIGQUIT, // SIGQUIT is a user-initiated "crash".
 	SIGILL,
@@ -901,100 +903,279 @@ const int32 TargetSignals[] =
 	SIGSYS
 };
 
-const int32 NumTargetSignals = UE_ARRAY_COUNT(TargetSignals);
+static constexpr int32 NumTargetSignals = UE_ARRAY_COUNT(TargetSignals);
+
+PRAGMA_DISABLE_OPTIMIZATION
+struct FAndroidSignalParams
+{
+	int Signal;
+	siginfo* Info;
+	void* Context;
+};
+template<uint32 SignalStackSize, typename Derived>
+class FSignalSubstituteStack
+{
+public:
+	static void Init()
+	{
+		AllocateSignalStack();
+	}
+
+	static void Release()
+	{
+		if (FakeSignalStack)
+		{
+			ReleaseSignalStack();
+		}
+	}
 
-struct sigaction PrevActions[NumTargetSignals];
-static bool PreviousSignalHandlersValid = false;
+protected:
+#if PLATFORM_ANDROID_ARM64
+	template< void(*ActualHandlerFun)()>
+	static void CallActualHandler()
+	{
+		// Arm64, save current stack frame, swap in new stack and call our signal handler.
+		__asm__("mov x8, %0" : : "r"(&FakeSignalStack[SignalStackLocation]));					// Put the fake stack addr into x8
+		__asm__("stp x29, x30, [x8, -32]" : : );						// push the existing frame and lr on to fake stack
+		__asm__("mov x9, sp" : : );										// copy existing stack ptr to x9
+		__asm__("str x9, [x8, -16]" : : );								// push the existing stack on to fake stack.
 
-const int32 GAndroidFakeSignalStackSize = 256 * 1024; // Desired stack size
-const uint32 GAndroidFakeSignalStackAllocationSize = GAndroidFakeSignalStackSize + PAGE_SIZE + PAGE_SIZE; // Total stack allocation (including guard page bookends)
-const uint32 GAndroidFakeSignalStackLocation = GAndroidFakeSignalStackAllocationSize - PAGE_SIZE; // Beginning of the desired stack.
-uint8* GAndroidFakeSignalStack = nullptr;
+		__asm__("sub sp, x8, 32" : : );									// set stack ptr to new stack addr, after our pushed regs
+		__asm__("sub fp, x8, 32" : : );									// set frame ptr to new stack addr, after our pushed regs
 
-void AllocateFakeSignalStack()
-{
-	check(GAndroidFakeSignalStack == nullptr);
+		__asm__("bl %0" : : "i"(ActualHandlerFun));						// Call our actual crash handler.
 
-	uint8* BaseStackAllocation = (uint8*)mmap(nullptr, GAndroidFakeSignalStackAllocationSize, PROT_READ | PROT_WRITE, MAP_ANONYMOUS | MAP_PRIVATE, -1, 0);
-	// protect top of the stack to catch overruns.
-	if (mprotect(BaseStackAllocation, PAGE_SIZE, PROT_NONE) != 0)
+		__asm__("ldp x29,x30,[sp]" : : );								// restore the original frame and lr.
+		__asm__("ldr x8, [sp, 16]" : : );								// restore the previously pushed stack pointer into x8
+		__asm__("mov sp, x8" : : );										// restore the original stack pointer
+	}
+#else
+	template< void(*ActualHandlerFunc)()>
+	static void CallActualHandler()
 	{
-		UE_LOG(LogAndroid, Fatal, TEXT("Failed to set alt stack protection"));
+		ActualHandlerFunc();
 	}
-	if (mprotect(BaseStackAllocation + GAndroidFakeSignalStackSize + PAGE_SIZE, PAGE_SIZE, PROT_NONE) != 0)
+#endif
+
+	static void AllocateSignalStack()
 	{
-		UE_LOG(LogAndroid, Fatal, TEXT("Failed to set alt stack protection"));
+		check(FakeSignalStack == nullptr);
+
+		uint8* BaseStackAllocation = (uint8*)mmap(nullptr, SignalStackAllocationSize, PROT_READ | PROT_WRITE, MAP_ANONYMOUS | MAP_PRIVATE, -1, 0);
+		// protect top of the stack to catch overruns.
+		if (mprotect(BaseStackAllocation, PAGE_SIZE, PROT_NONE) != 0)
+		{
+			UE_LOG(LogAndroid, Fatal, TEXT("Failed to set alt stack protection"));
+		}
+		if (mprotect(BaseStackAllocation + SignalStackSize + PAGE_SIZE, PAGE_SIZE, PROT_NONE) != 0)
+		{
+			UE_LOG(LogAndroid, Fatal, TEXT("Failed to set alt stack protection"));
+		}
+
+		FakeSignalStack = BaseStackAllocation;
+
+		check(IsAligned(&FakeSignalStack[SignalStackLocation], 16));
 	}
 
-	GAndroidFakeSignalStack = BaseStackAllocation;
+	static void ReleaseSignalStack()
+	{
+		check(FakeSignalStack);
+		munmap(FakeSignalStack, SignalStackAllocationSize);
+		FakeSignalStack = nullptr;
+	}
 
-	check(IsAligned(&GAndroidFakeSignalStack[GAndroidFakeSignalStackLocation], 16));
-}
+	static const uint32 SignalStackAllocationSize = SignalStackSize + PAGE_SIZE + PAGE_SIZE;
+	static const uint32 SignalStackLocation = SignalStackAllocationSize - PAGE_SIZE;
 
-void ReleaseFakeSignalStack()
-{
-	check(GAndroidFakeSignalStack);
-	munmap(GAndroidFakeSignalStack, GAndroidFakeSignalStackAllocationSize);
-	GAndroidFakeSignalStack = nullptr;
-}
+	static uint8* FakeSignalStack;
+};
 
-static void RestorePreviousSignalHandlers()
+template<uint32 SignalStackSize, typename Derived >
+uint8* FSignalSubstituteStack<SignalStackSize, Derived>::FakeSignalStack = nullptr;
+#if ANDROID_HAS_THREADBACKTRACE
+class FThreadCallstackSignalSubstituteStack : public FSignalSubstituteStack<256 * 1024, FThreadCallstackSignalSubstituteStack >
 {
-	if (PreviousSignalHandlersValid)
+public:
+	static void Init()
 	{
-		for (int32 i = 0; i < NumTargetSignals; ++i)
+		FSignalSubstituteStack<256 * 1024, FThreadCallstackSignalSubstituteStack >::Init();
+		struct sigaction ActionForThread;
+		FMemory::Memzero(ActionForThread);
+		sigfillset(&ActionForThread.sa_mask);
+		ActionForThread.sa_flags = SA_SIGINFO | SA_RESTART | SA_ONSTACK;
+		ActionForThread.sa_sigaction = &OnSignal;
+		sigaction(THREAD_CALLSTACK_GENERATOR, &ActionForThread, nullptr);
+	}
+private:
+	static void OnSignal(int Signal, siginfo* Info, void* Context)
+	{
+		while (FPlatformAtomics::InterlockedCompareExchange(&handling_signal, 1, 0) != 0)
 		{
-			sigaction(TargetSignals[i], &PrevActions[i], NULL);
+			FPlatformProcess::Sleep(0.0f);
 		}
-		PreviousSignalHandlersValid = false;
+
+		// prepare our signal parameters:
+		StackWalkSignalParams.Signal = Signal;
+		StackWalkSignalParams.Info = Info;
+		StackWalkSignalParams.Context = Context;
+
+		ThreadStackData = static_cast<ThreadStackUserData*>(Info->si_value.sival_ptr);
+
+		FSignalSubstituteStack::CallActualHandler< &ThreadStackWalk >();
+
+		FPlatformAtomics::AtomicStore(&handling_signal, 0);
 	}
-}
 
-static void SetDefaultSignalHandlers()
+	static void ThreadStackWalk()
+	{
+		if (ThreadStackData)
+		{
+			ThreadStackData->BackTraceCount = FPlatformStackWalk::CaptureStackBackTrace(ThreadStackData->BackTrace, ThreadStackData->CallStackSize, StackWalkSignalParams.Context);
+			ThreadStackData->bDone = true;
+		}
+	}
+	static FAndroidSignalParams StackWalkSignalParams;
+	static volatile sig_atomic_t handling_signal;
+	static ThreadStackUserData* ThreadStackData;
+};
+volatile sig_atomic_t FThreadCallstackSignalSubstituteStack::handling_signal = 0;
+ThreadStackUserData* FThreadCallstackSignalSubstituteStack::ThreadStackData = nullptr;
+FAndroidSignalParams FThreadCallstackSignalSubstituteStack::StackWalkSignalParams;
+#else
+	// per thread backtrace unavailable
+	class FThreadCallstackSignalSubstituteStack
+	{
+	public:
+		static void Init() {};
+		static void Release() {};
+	};
+#endif
+
+class FFatalSignalHandler : protected FSignalSubstituteStack<256 * 1024, FFatalSignalHandler>
 {
-	struct sigaction Action;
-	FMemory::Memzero(&Action, sizeof(struct sigaction));
-	Action.sa_handler = SIG_DFL;
-	sigemptyset(&Action.sa_mask);
+public:
+	static void Init()
+	{
+		FSignalSubstituteStack<256 * 1024, FFatalSignalHandler>::Init();
+		HookSignals();
+	}
 
-	for (int32 i = 0; i < NumTargetSignals; ++i)
+	static void Release()
 	{
-		sigaction(TargetSignals[i], &Action, NULL);
+		RestorePreviousSignalHandlers();
+		FSignalSubstituteStack<256 * 1024, FFatalSignalHandler>::Release();
+	}
+
+	static bool IsInFatalSignalHandler()
+	{
+		return FPlatformAtomics::AtomicRead(&handling_fatal_signal) > 0;
+	}
+
+	// EnterFatalCrash, called before a crash is reported
+	// Only the first occurrence succeeds.
+	static void EnterFatalCrash()
+	{
+		// we are a fatal signal, we *can* only handle one at a time. So avoid allow multiple fatal signals going through
+		if (FPlatformAtomics::InterlockedIncrement(&handling_fatal_signal) != 1)
+		{
+			FPlatformProcess::Sleep(60.0f);
+			exit(0);
+		}
+	}
+
+private:
+	static void OnHandleSignal(int Signal, siginfo* Info, void* Context)
+	{
+		EnterFatalCrash();
+
+		// prepare our signal parameters:
+		FatalSignalParams.Signal = Signal;
+		FatalSignalParams.Info = Info;
+		FatalSignalParams.Context = Context;
+
+		FFatalSignalHandler::CallActualHandler< &FatalSignalHandler >();
+
+		// Restore system handlers so Android could catch this signal after we are done with crashreport
+		RestorePreviousSignalHandlers();
+	}
+
+	static void FatalSignalHandler()
+	{
+		// Switch to malloc crash.
+		FPlatformMallocCrash::Get().SetAsGMalloc();
+
+		FAndroidCrashContext CrashContext(ECrashContextType::Crash, TEXT("Caught signal"));
+		CrashContext.InitFromSignal(FatalSignalParams.Signal, FatalSignalParams.Info, FatalSignalParams.Context);
+		CrashContext.CaptureCrashInfo();
+		if (GCrashHandlerPointer)
+		{
+			GCrashHandlerPointer(CrashContext);
+		}
+		else
+		{
+			// call default one
+			DefaultCrashHandler(CrashContext);
+		}
 	}
-}
 
-void AndroidCrashHandlerSubstituteStack(int Signal, siginfo* Info, void* Context);
+	static void HookSignals()
+	{
+		// hook our signals and record current set.
+		struct sigaction Action;
+		FMemory::Memzero(&Action, sizeof(struct sigaction));
+		Action.sa_sigaction = &OnHandleSignal;
+		// sigfillset will block all other signals whilst the signal handler is processing.
+		sigfillset(&Action.sa_mask);
+		Action.sa_flags = SA_SIGINFO | SA_RESTART | SA_ONSTACK;
+
+		for (int32 i = 0; i < NumTargetSignals; ++i)
+		{
+			sigaction(TargetSignals[i], &Action, &PrevActions[i]);
+		}
+		PreviousSignalHandlersValid = true;
+	}
 
-void HookSignals()
+	static void RestorePreviousSignalHandlers()
+	{
+		if (PreviousSignalHandlersValid)
+		{
+			for (int32 i = 0; i < NumTargetSignals; ++i)
+			{
+				sigaction(TargetSignals[i], &PrevActions[i], NULL);
+			}
+			PreviousSignalHandlersValid = false;
+		}
+	}
+
+	// params to pass from FFatalSignalHandler::OnHandleSignal to FatalSignalHandler.
+	static FAndroidSignalParams FatalSignalParams;
+	static volatile sig_atomic_t handling_fatal_signal;
+	static struct sigaction PrevActions[NumTargetSignals];
+	static bool PreviousSignalHandlersValid;
+};
+FAndroidSignalParams FFatalSignalHandler::FatalSignalParams;
+volatile sig_atomic_t FFatalSignalHandler::handling_fatal_signal = 0;
+struct sigaction FFatalSignalHandler::PrevActions[NumTargetSignals];
+bool FFatalSignalHandler::PreviousSignalHandlersValid = false;
+
+PRAGMA_ENABLE_OPTIMIZATION
+
+static void SetDefaultSignalHandlers()
 {
-	// hook our signals and record current set.
 	struct sigaction Action;
 	FMemory::Memzero(&Action, sizeof(struct sigaction));
-	Action.sa_sigaction = AndroidCrashHandlerSubstituteStack;
-	// sigfillset will block all other signals whilst the signal handler is processing.
-	sigfillset(&Action.sa_mask);
-	Action.sa_flags = SA_SIGINFO | SA_RESTART | SA_ONSTACK;
+	Action.sa_handler = SIG_DFL;
+	sigemptyset(&Action.sa_mask);
 
 	for (int32 i = 0; i < NumTargetSignals; ++i)
 	{
-		sigaction(TargetSignals[i], &Action, &PrevActions[i]);
+		sigaction(TargetSignals[i], &Action, NULL);
 	}
-	PreviousSignalHandlersValid = true;
 }
 
-// params to pass from AndroidCrashHandlerSubstituteStack to PlatformCrashHandler.
-struct FAndroidSignalParams
-{
-	int Signal;
-	siginfo* Info;
-	void* Context;
-};
-static FAndroidSignalParams GAndroidSignalParams;
-static volatile sig_atomic_t handling_fatal_signal = 0;
-
 bool FAndroidMisc::IsInSignalHandler()
 {
-	return FPlatformAtomics::AtomicRead(&handling_fatal_signal) > 0;
+	return FFatalSignalHandler::IsInFatalSignalHandler();
 }
 
 void FAndroidMisc::TriggerNonFatalCrashHandler(ECrashContextType InType, const FString& Message)
@@ -1017,12 +1198,7 @@ void FAndroidMisc::TriggerNonFatalCrashHandler(ECrashContextType InType, const F
 
 void FAndroidMisc::TriggerCrashHandler(const TCHAR* InErrorMessage, const TMap<FString, FString>& AdditionalProperties)
 {
-	// we are a fatal signal, we *can* only handle one at a time. So avoid allow multiple fatal signals going through
-	if (FPlatformAtomics::InterlockedIncrement(&handling_fatal_signal) != 1)
-	{
-		FPlatformProcess::Sleep(60.0f);
-		exit(0);
-	}
+	FFatalSignalHandler::EnterFatalCrash();
 
 	FAndroidCrashContext CrashContext(ECrashContextType::Crash, InErrorMessage);
 	CrashContext.CaptureCrashInfo();
@@ -1042,73 +1218,6 @@ void FAndroidMisc::TriggerCrashHandler(const TCHAR* InErrorMessage, const TMap<F
 	}
 }
 
-
-void PlatformCrashHandler()
-{
-	// Switch to malloc crash.
-	FPlatformMallocCrash::Get().SetAsGMalloc();
-
-	FAndroidCrashContext CrashContext(ECrashContextType::Crash, TEXT("Caught signal"));
-	CrashContext.InitFromSignal(GAndroidSignalParams.Signal, GAndroidSignalParams.Info, GAndroidSignalParams.Context);
-	CrashContext.CaptureCrashInfo();
-	if (GCrashHandlerPointer)
-	{
-		GCrashHandlerPointer(CrashContext);
-	}
-	else
-	{
-		// call default one
-		DefaultCrashHandler(CrashContext);
-	}
-}
-
-PRAGMA_DISABLE_OPTIMIZATION
-#if PLATFORM_ANDROID_ARM64
-/** True system-specific crash handler that replaces stack pointer,  */
-void AndroidCrashHandlerSubstituteStack(int Signal, siginfo* Info, void* Context)
-{
-	// we are a fatal signal, we *can* only handle one at a time. So avoid allow multiple fatal signals going through
-	if (FPlatformAtomics::InterlockedIncrement(&handling_fatal_signal) != 1)
-	{
-		FPlatformProcess::Sleep(60.0f);
-		exit(0);
-	}
-
-	// prepare our signal parameters:
-	GAndroidSignalParams.Signal = Signal;
-	GAndroidSignalParams.Info = Info;
-	GAndroidSignalParams.Context = Context;
-
-	// Arm64, save current stack frame, swap in new stack and call our signal handler.
-	uint8* FakeStackBase = &GAndroidFakeSignalStack[GAndroidFakeSignalStackLocation];	
-	__asm__("mov x8, %0" : : "r"(FakeStackBase));					// Put the fake stack addr into x8
-	__asm__("stp x29, x30, [x8, -32]" : : );						// push the existing frame and lr on to fake stack
-	__asm__("mov x9, sp" : : );										// copy existing stack ptr to x9
-	__asm__("str x9, [x8, -16]" : : );								// push the existing stack on to fake stack.
-
-	__asm__("sub sp, x8, 32" : : );									// set stack ptr to new stack addr, after our pushed regs
-	__asm__("sub fp, x8, 32" : : );									// set frame ptr to new stack addr, after our pushed regs
-
-    __asm__("bl %0" : : "i"(PlatformCrashHandler));					// Call our actual crash handler.
-
-	__asm__("ldp x29,x30,[sp]" : : );								// restore the original frame and lr.
-	__asm__("ldr x8, [sp, 16]" : : );								// restore the previously pushed stack pointer into x8
-	__asm__("mov sp, x8" : : );										// restore the original stack pointer
-
-	// Restore system handlers so Android could catch this signal after we are done with crashreport
-	RestorePreviousSignalHandlers();
-
-	FPlatformAtomics::InterlockedDecrement(&handling_fatal_signal);
-}
-#else
-// TODO
-void AndroidCrashHandlerSubstituteStack(int Signal, siginfo* Info, void* Context)
-{	
-	PlatformCrashHandler();
-}
-#endif 
-PRAGMA_ENABLE_OPTIMIZATION
-
 void FAndroidMisc::SetCrashHandler(void(*CrashHandler)(const FGenericCrashContext& Context))
 {
 #if PLATFORM_ANDROID_ARM64
@@ -1116,23 +1225,16 @@ void FAndroidMisc::SetCrashHandler(void(*CrashHandler)(const FGenericCrashContex
 
 	GCrashHandlerPointer = CrashHandler;
 
-	RestorePreviousSignalHandlers();
-	FMemory::Memzero(&PrevActions, sizeof(PrevActions));
-	if(GAndroidFakeSignalStack)
-	{
-		ReleaseFakeSignalStack();
-	}
-
+	FFatalSignalHandler::Release();
+	FThreadCallstackSignalSubstituteStack::Release();
 	// Passing -1 will leave these restored and won't trap them
 	if ((PTRINT)CrashHandler == -1)
 	{
 		return;
 	}
 
-	AllocateFakeSignalStack();
-
-	HookSignals();
-
+	FFatalSignalHandler::Init();
+	FThreadCallstackSignalSubstituteStack::Init();
 #endif //PLATFORM_ANDROID_ARM64
 }
 
diff --git a/Engine/Source/Runtime/Core/Private/Android/AndroidPlatformStackWalk.cpp b/Engine/Source/Runtime/Core/Private/Android/AndroidPlatformStackWalk.cpp
--- a/Engine/Source/Runtime/Core/Private/Android/AndroidPlatformStackWalk.cpp
+++ b/Engine/Source/Runtime/Core/Private/Android/AndroidPlatformStackWalk.cpp
@@ -18,6 +18,11 @@
 #define UNW_LOCAL_ONLY
 #include "libunwind.h"
 #endif
+#if ANDROID_HAS_THREADBACKTRACE
+#include <syscall.h>
+#include "HAL/PlatformTime.h"
+#endif
+#include "HAL/IConsoleManager.h"
 
 void FAndroidPlatformStackWalk::ProgramCounterToSymbolInfo(uint64 ProgramCounter, FProgramCounterSymbolInfo& out_SymbolInfo)
 {
@@ -252,3 +257,67 @@ bool FAndroidPlatformStackWalk::SymbolInfoToHumanReadableString(const FProgramCo
 	}
 	return false;
 }
+
+static TAutoConsoleVariable<float> CVarAndroidPlatformThreadCallStackMaxWait(
+	TEXT("AndroidPlatformThreadStackWalk.MaxWait"),
+	60.0f,
+	TEXT("The number of seconds allowed to spin before killing the process, with the assumption the signal handler has hung."));
+
+#if ANDROID_HAS_THREADBACKTRACE
+uint32 FAndroidPlatformStackWalk::CaptureThreadStackBackTrace(uint64 ThreadId, uint64* BackTrace, uint32 MaxDepth)
+{
+	auto GatherCallstackFromThread = [](ThreadStackUserData& ThreadStack, uint64 TargetThreadId)
+	{
+		auto WaitForSignalHandlerToFinishOrCrash = [](ThreadStackUserData& WaitThreadStack)
+		{
+			float EndWaitTimestamp = FPlatformTime::Seconds() + CVarAndroidPlatformThreadCallStackMaxWait.AsVariable()->GetFloat();
+			float CurrentTimestamp = FPlatformTime::Seconds();
+
+			while (!WaitThreadStack.bDone)
+			{
+				if (CurrentTimestamp > EndWaitTimestamp)
+				{
+					// We have waited for as long as we should for the signal handler to finish. Assume it has hang and we need to kill our selfs
+					*(int*)0x10 = 0x0;
+				}
+
+				CurrentTimestamp = FPlatformTime::Seconds();
+			}
+		};
+
+		sigval UserData;
+		UserData.sival_ptr = &ThreadStack;
+
+		siginfo_t info;
+		memset(&info, 0, sizeof(siginfo_t));
+		info.si_signo = THREAD_CALLSTACK_GENERATOR;
+		info.si_code = SI_QUEUE;
+		info.si_pid = syscall(SYS_getpid);
+		info.si_uid = syscall(SYS_getuid);
+		info.si_value = UserData;
+
+		// Avoid using sigqueue here as if the ThreadId is already blocked and in a signal handler
+		// sigqueue will try a different thread signal handler and report the wrong callstack
+		if (syscall(SYS_rt_tgsigqueueinfo, info.si_pid, TargetThreadId, THREAD_CALLSTACK_GENERATOR, &info) == 0)
+		{
+			WaitForSignalHandlerToFinishOrCrash(ThreadStack);
+		}
+	};
+
+	ThreadStackUserData ThreadBackTrace;
+	ThreadBackTrace.CallStackSize = MaxDepth;
+	ThreadBackTrace.BackTrace = BackTrace;
+	ThreadBackTrace.BackTraceCount = 0;
+	ThreadBackTrace.bDone = false;
+
+	GatherCallstackFromThread(ThreadBackTrace, ThreadId);
+
+	// The signal handler will set this value, we just have to make sure we wait for the signal handler we raised to finish
+	return ThreadBackTrace.BackTraceCount;
+}
+#else
+uint32 FAndroidPlatformStackWalk::CaptureThreadStackBackTrace(uint64 ThreadId, uint64* BackTrace, uint32 MaxDepth)
+{
+	return 0;
+}
+#endif //ANDROID_HAS_THREADBACKTRACE
\ No newline at end of file
diff --git a/Engine/Source/Runtime/Core/Private/GenericPlatform/GenericPlatformCrashContext.cpp b/Engine/Source/Runtime/Core/Private/GenericPlatform/GenericPlatformCrashContext.cpp
--- a/Engine/Source/Runtime/Core/Private/GenericPlatform/GenericPlatformCrashContext.cpp
+++ b/Engine/Source/Runtime/Core/Private/GenericPlatform/GenericPlatformCrashContext.cpp
@@ -1039,7 +1039,7 @@ void FGenericCrashContext::SetPortableCallStack(const uint64* StackFrames, int32
 	GetPortableCallStack(StackFrames, NumStackFrames, CallStack);
 }
 
-void FGenericCrashContext::GetPortableCallStack(const uint64* StackFrames, int32 NumStackFrames, TArray<FCrashStackFrame>& OutCallStack)
+void FGenericCrashContext::GetPortableCallStack(const uint64* StackFrames, int32 NumStackFrames, TArray<FCrashStackFrame>& OutCallStack) const
 {
 	// Get all the modules in the current process
 	uint32 NumModules = (uint32)FPlatformStackWalk::GetProcessModuleCount();
diff --git a/Engine/Source/Runtime/Core/Private/HAL/ThreadingBase.cpp b/Engine/Source/Runtime/Core/Private/HAL/ThreadingBase.cpp
--- a/Engine/Source/Runtime/Core/Private/HAL/ThreadingBase.cpp
+++ b/Engine/Source/Runtime/Core/Private/HAL/ThreadingBase.cpp
@@ -267,6 +267,15 @@ void FThreadManager::GetAllThreadStackBackTraces(TArray<FThreadStackBackTrace>&
 }
 #endif
 
+void FThreadManager::ForEachThread(TFunction<void(uint32, class FRunnableThread*)> Func)
+{
+	FScopeLock Lock(&ThreadsCritical);
+	for (const TPair<uint32, FRunnableThread*>& Pair : Threads)
+	{
+		Func(Pair.Key, Pair.Value);
+	}
+}
+
 FThreadManager& FThreadManager::Get()
 {
 	static FThreadManager Singleton;
diff --git a/Engine/Source/Runtime/Core/Public/Android/AndroidPlatformCrashContext.h b/Engine/Source/Runtime/Core/Public/Android/AndroidPlatformCrashContext.h
--- a/Engine/Source/Runtime/Core/Public/Android/AndroidPlatformCrashContext.h
+++ b/Engine/Source/Runtime/Core/Public/Android/AndroidPlatformCrashContext.h
@@ -36,7 +36,7 @@ struct CORE_API FAndroidCrashContext : public FGenericCrashContext
 		Context = InContext;
 	}
 
-	virtual void GetPortableCallStack(const uint64* StackFrames, int32 NumStackFrames, TArray<FCrashStackFrame>& OutCallStack) override;
+	virtual void GetPortableCallStack(const uint64* StackFrames, int32 NumStackFrames, TArray<FCrashStackFrame>& OutCallStack) const override;
 	virtual void AddPlatformSpecificProperties() const override;
 
 	void CaptureCrashInfo();
@@ -53,9 +53,13 @@ struct CORE_API FAndroidCrashContext : public FGenericCrashContext
 	// generate an absolute path to a crash report folder.
 	static void GenerateReportDirectoryName(char(&DirectoryNameOUT)[CrashReportMaxPathSize]);
 
+	void DumpAllThreadCallstacks() const;
+
 private:
 	TMap<FString, FString> AdditionalProperties;
 
+	// The path used by this instance to store the report.
+	char ReportDirectory[CrashReportMaxPathSize];
 };
 
 typedef FAndroidCrashContext FPlatformCrashContext;
\ No newline at end of file
diff --git a/Engine/Source/Runtime/Core/Public/Android/AndroidPlatformStackWalk.h b/Engine/Source/Runtime/Core/Public/Android/AndroidPlatformStackWalk.h
--- a/Engine/Source/Runtime/Core/Public/Android/AndroidPlatformStackWalk.h
+++ b/Engine/Source/Runtime/Core/Public/Android/AndroidPlatformStackWalk.h
@@ -18,6 +18,23 @@ struct CORE_API FAndroidPlatformStackWalk : public FGenericPlatformStackWalk
 	static void ProgramCounterToSymbolInfo(uint64 ProgramCounter, FProgramCounterSymbolInfo& out_SymbolInfo);
 	static uint32 CaptureStackBackTrace(uint64* BackTrace, uint32 MaxDepth, void* Context = nullptr);
 	static bool SymbolInfoToHumanReadableString(const FProgramCounterSymbolInfo& SymbolInfo, ANSICHAR* HumanReadableString, SIZE_T HumanReadableStringSize);
+
+	static uint32 CaptureThreadStackBackTrace(uint64 ThreadId, uint64* BackTrace, uint32 MaxDepth);
+};
+
+#define ANDROID_HAS_THREADBACKTRACE !PLATFORM_LUMIN && PLATFORM_USED_NDK_VERSION_INTEGER >= 21
+
+#if ANDROID_HAS_THREADBACKTRACE
+
+/** Passed in through sigqueue for gathering of a callstack from a signal */
+struct ThreadStackUserData
+{
+	uint64* BackTrace;
+	int32 BackTraceCount;
+	SIZE_T CallStackSize;
+	TAtomic<bool> bDone;
 };
+#define THREAD_CALLSTACK_GENERATOR SIGRTMIN + 5
+#endif // ANDROID_HAS_THREADBACKTRACE
 
 typedef FAndroidPlatformStackWalk FPlatformStackWalk;
diff --git a/Engine/Source/Runtime/Core/Public/GenericPlatform/GenericPlatformCrashContext.h b/Engine/Source/Runtime/Core/Public/GenericPlatform/GenericPlatformCrashContext.h
--- a/Engine/Source/Runtime/Core/Public/GenericPlatform/GenericPlatformCrashContext.h
+++ b/Engine/Source/Runtime/Core/Public/GenericPlatform/GenericPlatformCrashContext.h
@@ -393,7 +393,7 @@ public:
 	virtual void SetPortableCallStack(const uint64* StackFrames, int32 NumStackFrames);
 
 	/** Gets the portable callstack to a specified stack and puts it into OutCallStack */
-	virtual void GetPortableCallStack(const uint64* StackFrames, int32 NumStackFrames, TArray<FCrashStackFrame>& OutCallStack);
+	virtual void GetPortableCallStack(const uint64* StackFrames, int32 NumStackFrames, TArray<FCrashStackFrame>& OutCallStack) const;
 
 	/** Adds a portable callstack for a thread */
 	virtual void AddPortableThreadCallStack(uint32 ThreadId, const TCHAR* ThreadName, const uint64* StackFrames, int32 NumStackFrames);
diff --git a/Engine/Source/Runtime/Core/Public/HAL/ThreadManager.h b/Engine/Source/Runtime/Core/Public/HAL/ThreadManager.h
--- a/Engine/Source/Runtime/Core/Public/HAL/ThreadManager.h
+++ b/Engine/Source/Runtime/Core/Public/HAL/ThreadManager.h
@@ -60,6 +60,12 @@ public:
 	void GetAllThreadStackBackTraces(TArray<FThreadStackBackTrace>& StackTraces);
 #endif
 
+	/**
+	 * Enumerate each thread.
+	 *
+	 */
+	void ForEachThread(TFunction<void(uint32, class FRunnableThread*)> Func);
+
 	/**
 	 * Access to the singleton object.
 	 *
diff --git a/Engine/Source/Runtime/CrashReportCore/Private/Android/AndroidErrorReport.cpp b/Engine/Source/Runtime/CrashReportCore/Private/Android/AndroidErrorReport.cpp
--- a/Engine/Source/Runtime/CrashReportCore/Private/Android/AndroidErrorReport.cpp
+++ b/Engine/Source/Runtime/CrashReportCore/Private/Android/AndroidErrorReport.cpp
@@ -27,4 +27,32 @@ void FAndroidErrorReport::ShutDown()
 	AndroidErrorReport::CrashHelperModule->ShutdownModule();
 }
 
+static void AddThreadContexts(const FString& ReportDirectory)
+{
+	// Try to load the callstacks file.
+	const FString ThreadContextsFile(ReportDirectory / TEXT("AllThreads.txt"));
+	IPlatformFile& PlatformFile = FPlatformFileManager::Get().GetPlatformFile();
+
+	if (PlatformFile.FileExists(*ThreadContextsFile))
+	{
+		// rename the threads file, if anything goes wrong during processing it will be skipped on the next run.
+		const FString ThreadContextsFileTemp(ReportDirectory / TEXT("AllThreads.tmp"));
+		PlatformFile.MoveFile(*ThreadContextsFileTemp, *ThreadContextsFile);
+
+		FXmlFile ThreadsNode(ThreadContextsFileTemp);
+		if (ThreadsNode.IsValid())
+		{			
+			FPrimaryCrashProperties::Get()->Threads = ThreadsNode.GetRootNode();
+			// delete the file as it has been added to the primary report.
+			PlatformFile.DeleteFile(*ThreadContextsFileTemp);
+		}
+	}
+}
+
+FText FAndroidErrorReport::DiagnoseReport() const
+{
+	AddThreadContexts(ReportDirectory);
+	return FText();
+}
+
 #undef LOCTEXT_NAMESPACE
diff --git a/Engine/Source/Runtime/CrashReportCore/Private/CrashDescription.cpp b/Engine/Source/Runtime/CrashReportCore/Private/CrashDescription.cpp
--- a/Engine/Source/Runtime/CrashReportCore/Private/CrashDescription.cpp
+++ b/Engine/Source/Runtime/CrashReportCore/Private/CrashDescription.cpp
@@ -93,6 +93,20 @@ int64 FCrashProperty::AsInt64() const
 	return Value;
 }
 
+FCrashPropertyXmlNode::FCrashPropertyXmlNode(const FString& InMainCategory, const FString& InSecondCategory, FPrimaryCrashProperties* InOwner)
+	: Owner(InOwner)
+	, MainCategory(InMainCategory)
+	, SecondCategory(InSecondCategory)
+	, bSet(false)
+{ }
+
+FCrashPropertyXmlNode& FCrashPropertyXmlNode::operator=(const FXmlNode* Node)
+{
+	bSet = true;
+	Owner->SetCrashProperty(MainCategory, SecondCategory, Node);
+	return *this;
+}
+
 /*-----------------------------------------------------------------------------
 	FPrimaryCrashProperties
 -----------------------------------------------------------------------------*/
@@ -122,6 +136,7 @@ FPrimaryCrashProperties::FPrimaryCrashProperties()
 	, PlatformCallbackResult(FGenericCrashContext::PlatformPropertiesTag, TEXT("PlatformCallbackResult"), this)
 	, CrashReportClientVersion(FGenericCrashContext::RuntimePropertiesTag, TEXT("CrashReportClientVersion"), this)
 	, CPUBrand(FGenericCrashContext::RuntimePropertiesTag, TEXT("CPUBrand"), this)
+	, Threads(FGenericCrashContext::RuntimePropertiesTag, TEXT("Threads"), this)	
 	, bIsOOM(false)
 	, bLowMemoryWarning(false)
 	, bInBackground(false)
diff --git a/Engine/Source/Runtime/CrashReportCore/Public/Android/AndroidErrorReport.h b/Engine/Source/Runtime/CrashReportCore/Public/Android/AndroidErrorReport.h
--- a/Engine/Source/Runtime/CrashReportCore/Public/Android/AndroidErrorReport.h
+++ b/Engine/Source/Runtime/CrashReportCore/Public/Android/AndroidErrorReport.h
@@ -42,13 +42,8 @@ public:
 	static void ShutDown();
 
 	/**
-	 * Do nothing - shouldn't be called on Android
-	 * @return Dummy text
 	 */
-	FText DiagnoseReport() const	
-	{
-		return FText::FromString("No local diagnosis on Android");
-	}
+	FText DiagnoseReport() const;
 
 	/**
 	 * Do nothing - shouldn't be called on Android
diff --git a/Engine/Source/Runtime/CrashReportCore/Public/CrashDescription.h b/Engine/Source/Runtime/CrashReportCore/Public/CrashDescription.h
--- a/Engine/Source/Runtime/CrashReportCore/Public/CrashDescription.h
+++ b/Engine/Source/Runtime/CrashReportCore/Public/CrashDescription.h
@@ -143,10 +143,33 @@ protected:
 	mutable bool bSet;
 };
 
+struct FCrashPropertyXmlNode
+{
+	friend struct FPrimaryCrashProperties;
+	/** Assignment operator for . */
+	FCrashPropertyXmlNode& operator=(const FXmlNode* Node);
+protected:
+	/** Initialization constructor. */
+	FCrashPropertyXmlNode(const FString& InMainCategory, const FString& InSecondCategory, FPrimaryCrashProperties* InOwner);
+
+protected:
+	/** Owner of the property. */
+	FPrimaryCrashProperties* Owner;
+
+	/** Main category in the crash context. */
+	FString MainCategory;
+
+	/** Second category in the crash context. */
+	FString SecondCategory;
+
+	mutable bool bSet;
+};
+
 /** Primary crash properties required by the crash report system. */
 struct FPrimaryCrashProperties
 {
 	friend struct FCrashProperty;
+	friend struct FCrashPropertyXmlNode;
 
 	/** Version. */
 	ECrashDescVersions CrashVersion;
@@ -388,6 +411,11 @@ struct FPrimaryCrashProperties
 	 */
 	FCrashProperty CPUBrand;
 
+	/**
+	 *	Thread contexts, XML elements containing info specific to an active thread, e.g. callstacks
+	 */
+	FCrashPropertyXmlNode Threads;
+
 	/**
 	 * Whether it was an OOM or not
 	 */
@@ -512,6 +540,37 @@ protected:
 		}
 	}
 
+	/** Sets a crash property to a new value. */
+	void SetCrashProperty(const FString& MainCategory, const FString& SecondCategory, const FXmlNode* NewNode)
+	{
+		if (XmlFile->IsValid())
+		{
+			FXmlNode* MainNode = XmlFile->GetRootNode()->FindChildNode(MainCategory);
+			if (MainNode)
+			{
+				FXmlNode* CategoryNode = MainNode->FindChildNode(SecondCategory);
+				const FString& Content = NewNode->GetContent();
+
+				if(!CategoryNode) // can only add.
+				{
+					TFunction<void(FXmlNode*, const FXmlNode*)> AppendNode;
+					AppendNode = [&AppendNode](FXmlNode* DestNode, const FXmlNode* SourceNode)
+					{
+						DestNode->AppendChildNode(SourceNode->GetTag(), SourceNode->GetContent());
+						for (const FXmlNode *SourceChild : SourceNode->GetChildrenNodes())
+						{
+							const TArray<FXmlNode *>& DestChildren = DestNode->GetChildrenNodes();
+							AppendNode(DestChildren.Last(), SourceChild);
+						}
+					};
+
+					AppendNode(MainNode, NewNode);
+				}
+			}
+		}
+	}
+
+
 	/** Encodes multi line property to be saved as single line. */
 	FString EncodeArrayStringAsXMLString( const TArray<FString>& ArrayString ) const;
 
