Fix engine not compiling since CL 10451570

diff --git a/Engine/Source/Runtime/Core/Public/UObject/UnrealNames.inl b/Engine/Source/Runtime/Core/Public/UObject/UnrealNames.inl
--- a/Engine/Source/Runtime/Core/Public/UObject/UnrealNames.inl
+++ b/Engine/Source/Runtime/Core/Public/UObject/UnrealNames.inl
@@ -202,5 +202,4 @@ REGISTER_NAME(601,EditorKeyBindings)
 REGISTER_NAME(602,GameUserSettings)
 
 // Commonly multicased names - defined globally to prevent non-determinism in package name tables
-REGISTER_NAME(700,ID)
-REGISTER_NAME(701,Filename)
+REGISTER_NAME(700,Filename)
