Iteration on StringView

- Fix use of concrete StringView types as contiguous containers.
- Fix StringView::Right[Chop] to work without null termination.
- Fix StringView/Range [in]equality operator ambiguity by disabling the Range == StringView and Range != StringView overloads when Range is a string view.
- Fix TIsDerivedFrom to provide a Value member for use with TAnd/TOr/TNot.
- Fix the default GetNum for contiguous containers to be disabled if the type does not contain a Num function. (Suggested by [REDACTED_PERSON])
- Change Compare and Equals on StringView to use ESearchCase::CaseSensitive by default for consistency with FString.

diff --git a/Engine/Source/Runtime/Core/Private/Tests/Misc/StringViewTest.cpp b/Engine/Source/Runtime/Core/Private/Tests/Misc/StringViewTest.cpp
--- a/Engine/Source/Runtime/Core/Private/Tests/Misc/StringViewTest.cpp
+++ b/Engine/Source/Runtime/Core/Private/Tests/Misc/StringViewTest.cpp
@@ -7,6 +7,10 @@
 
 #if WITH_DEV_AUTOMATION_TESTS 
 
+static_assert(TIsContiguousContainer<FStringView>::Value, "FStringView must be a contiguous container.");
+static_assert(TIsContiguousContainer<FAnsiStringView>::Value, "FAnsiStringView must be a contiguous container.");
+static_assert(TIsContiguousContainer<FWideStringView>::Value, "FWideStringView must be a contiguous container.");
+
 #define TEST_NAME_ROOT "System.Core.Misc.StringView"
 constexpr const uint32 TestFlags = EAutomationTestFlags::ApplicationContextMask | EAutomationTestFlags::EngineFilter;
 
@@ -162,8 +166,10 @@ bool FStringViewTestEquality::RunTest(const FString& Parameters)
 
 	FStringView IdenticalView(WideStringLiteralSrc);
 
-	TestTrue(TEXT("Equality(40)"), WideView == IdenticalView);
-	TestTrue(TEXT("Equality(41)"), IdenticalView == WideView);
+	TestTrue(TEXT("Equality(40a)"), WideView == IdenticalView);
+	TestFalse(TEXT("Equality(40b)"), WideView != IdenticalView);
+	TestTrue(TEXT("Equality(41a)"), IdenticalView == WideView);
+	TestFalse(TEXT("Equality(41b)"), IdenticalView != WideView);
 
 	// Views without null termination
 
@@ -480,7 +486,7 @@ bool FStringViewTestLeft::RunTest(const FString& Parameters)
 	}
 	
 	{
-		FStringView View(TEXT("A test string"));
+		FStringView View(TEXT("A test string padded"), 13); // "A test string" without null termination
 		FStringView Result = View.Left(8);
 
 		TestTrue(TEXT("FStringView::Left"), FCString::Strncmp(Result.GetData(), TEXT("A test s"), Result.Len()) == 0);
@@ -505,7 +511,7 @@ bool FStringViewTestLeftChop::RunTest(const FString& Parameters)
 	}
 
 	{
-		FStringView View(TEXT("A test string"));
+		FStringView View(TEXT("A test string padded"), 13); // "A test string" without null termination
 		FStringView Result = View.LeftChop(5);
 
 		TestTrue(TEXT("FStringView::LeftChop"), FCString::Strncmp(Result.GetData(), TEXT("A test s"), Result.Len()) == 0);
@@ -530,7 +536,7 @@ bool FStringViewTestRight::RunTest(const FString& Parameters)
 	}
 
 	{
-		FStringView View(TEXT("A test string"));
+		FStringView View(TEXT("A test string padded"), 13); // "A test string" without null termination
 		FStringView Result = View.Right(8);
 
 		TestTrue(TEXT("FStringView::Right"), FCString::Strncmp(Result.GetData(), TEXT("t string"), Result.Len()) == 0);
@@ -555,7 +561,7 @@ bool FStringViewTestRightChop::RunTest(const FString& Parameters)
 	}
 
 	{
-		FStringView View(TEXT("A test string"));
+		FStringView View(TEXT("A test string padded"), 13); // "A test string" without null termination
 		FStringView Result = View.RightChop(3);
 
 		TestTrue(TEXT("FStringView::RightChop"), FCString::Strncmp(Result.GetData(), TEXT("est string"), Result.Len()) == 0);
@@ -595,6 +601,7 @@ bool FStringViewTestMid::RunTest(const FString& Parameters)
 		TestTrue(TEXT("FStringView::Mid(4)"), FCString::Strncmp(Word2.GetData(), TEXT("string"), Word2.Len()) == 0);
 		TestTrue(TEXT("FStringView::Mid(5)"), NullTerminatorResult.IsEmpty());
 		TestTrue(TEXT("FStringView::Mid(6)"), View == OutofBoundsResult);
+		TestTrue(TEXT("FStringView::Mid(7)"), View.Mid(512, 1024).IsEmpty());
 	}
 
 	return true;
diff --git a/Engine/Source/Runtime/Core/Public/Misc/StringView.h b/Engine/Source/Runtime/Core/Public/Misc/StringView.h
--- a/Engine/Source/Runtime/Core/Public/Misc/StringView.h
+++ b/Engine/Source/Runtime/Core/Public/Misc/StringView.h
@@ -3,6 +3,7 @@
 #pragma once
 
 #include "CoreTypes.h"
+#include "Math/NumericLimits.h"
 #include "Math/UnrealMathUtility.h"
 #include "Misc/CString.h"
 #include "Templates/AndOrNot.h"
@@ -81,19 +82,19 @@ public:
 
 	TStringViewImpl() = default;
 
-	TStringViewImpl(const CharType* InData)
+	inline TStringViewImpl(const CharType* InData)
 		: DataPtr(InData)
 		, Size(TCString<CharType>::Strlen(InData))
 	{
 	}
 
-	TStringViewImpl(const CharType* InData, SizeType InSize)
+	inline TStringViewImpl(const CharType* InData, SizeType InSize)
 		: DataPtr(InData)
 		, Size(InSize)
 	{
 	}
 
-	inline const CharType& operator[](SizeType Pos) const { return DataPtr[Pos]; }
+	inline const CharType& operator[](SizeType Index) const;
 
 	inline const CharType* GetData() const { return DataPtr; }
 
@@ -113,24 +114,24 @@ public:
 	// Operations
 
 	inline SizeType CopyString(CharType* Dest, SizeType CharCount, SizeType Position = 0) const;
-	inline ViewType SubStr(SizeType Position, SizeType CharCount) const						{ check(Position <= Len()); return ViewType(DataPtr + Position, FMath::Min(Size - Position, CharCount)); }
+	inline ViewType SubStr(SizeType Position, SizeType CharCount) const { return Mid(Position, CharCount); }
 
 	// Maintain compatibility with FString 
-	inline ViewType Left(SizeType CharCount) const											{ return ViewType(GetData(), FMath::Clamp(CharCount, 0, Len())); }
-	inline ViewType LeftChop(SizeType CharCount) const										{ return ViewType(GetData(), FMath::Clamp(Len() - CharCount, 0, Len())); }
-	inline ViewType Right(SizeType CharCount) const											{ return ViewType(GetData() + Len() - FMath::Clamp(CharCount, 0, Len())); }
-	inline ViewType RightChop(SizeType CharCount) const										{ return ViewType(GetData() + Len() - FMath::Clamp(Len() - CharCount, 0, Len())); }
-	inline ViewType Mid(SizeType Position, SizeType CharCount) const						{ return SubStr(Position, CharCount); /*This is just a wrapper around SubStr to keep compatibility with the FString interface*/} 
+	inline ViewType Left(SizeType CharCount) const;
+	inline ViewType LeftChop(SizeType CharCount) const;
+	inline ViewType Right(SizeType CharCount) const;
+	inline ViewType RightChop(SizeType CharCount) const;
+	inline ViewType Mid(SizeType Position, SizeType CharCount = TNumericLimits<SizeType>::Max()) const;
 
 	// Comparison
-	inline bool				Equals(const TStringViewImpl& Other, ESearchCase::Type SearchCase) const { return Len() == Other.Len() && Compare(Other, SearchCase) == 0; }
-	CORE_API int32			Compare(const TStringViewImpl& Other, ESearchCase::Type SearchCase) const;
+	bool Equals(const TStringViewImpl& Other, ESearchCase::Type SearchCase = ESearchCase::CaseSensitive) const;
+	CORE_API int32 Compare(const TStringViewImpl& Other, ESearchCase::Type SearchCase = ESearchCase::CaseSensitive) const;
 
-	inline bool StartsWith(CharType Prefix) const										{ return Size >= 1 && DataPtr[0] == Prefix; }
-	inline bool StartsWith(const TStringViewImpl& Prefix) const							{ return Len() >= Prefix.Len() && TCString<CharType>::Strnicmp(GetData(), Prefix.GetData(), Prefix.Len()) == 0; }
+	inline bool StartsWith(CharType Prefix) const { return Size >= 1 && DataPtr[0] == Prefix; }
+	inline bool StartsWith(const TStringViewImpl& Prefix) const;
 
-	inline bool EndsWith(CharType Suffix) const											{ return Size >= 1 && DataPtr[Size-1] == Suffix; }
-	inline bool EndsWith(const TStringViewImpl& Suffix) const							{ return Len() >= Suffix.Len() && TCString<CharType>::Strnicmp(GetData() + Len() - Suffix.Len(), Suffix.GetData(), Suffix.Len()) == 0; }
+	inline bool EndsWith(CharType Suffix) const { return Size >= 1 && DataPtr[Size-1] == Suffix; }
+	inline bool EndsWith(const TStringViewImpl& Suffix) const;
 
 	// Searching/Finding
 	CORE_API bool FindChar(CharType InChar, SizeType& OutIndex) const;
@@ -181,4 +182,12 @@ public:
 	using TStringViewImpl<ElementType, FWideStringView>::TStringViewImpl;
 };
 
+template <> struct TIsContiguousContainer<FStringView> { enum { Value = true }; };
+template <> struct TIsContiguousContainer<FAnsiStringView> { enum { Value = true }; };
+template <> struct TIsContiguousContainer<FWideStringView> { enum { Value = true }; };
+
+inline SIZE_T GetNum(const FStringView& String) { return String.Len(); }
+inline SIZE_T GetNum(const FAnsiStringView& String) { return String.Len(); }
+inline SIZE_T GetNum(const FWideStringView& String) { return String.Len(); }
+
 #include "StringView.inl"
diff --git a/Engine/Source/Runtime/Core/Public/Misc/StringView.inl b/Engine/Source/Runtime/Core/Public/Misc/StringView.inl
--- a/Engine/Source/Runtime/Core/Public/Misc/StringView.inl
+++ b/Engine/Source/Runtime/Core/Public/Misc/StringView.inl
@@ -2,53 +2,106 @@
 
 #include "HAL/PlatformString.h"
 #include "HAL/UnrealMemory.h"
+#include "Templates/Decay.h"
+#include "Templates/UnrealTypeTraits.h"
+#include "Traits/IntType.h"
+
+template <typename CharType, typename ViewType>
+inline const CharType& TStringViewImpl<CharType, ViewType>::operator[](SizeType Index) const
+{
+	checkf(Index >= 0 && Index < Size, TEXT("Index out of bounds on StringView: index %i on a view with a length of %i"), Index, Size);
+	return DataPtr[Index];
+}
 
 template <typename CharType, typename ViewType>
 inline typename TStringViewImpl<CharType, ViewType>::SizeType TStringViewImpl<CharType, ViewType>::CopyString(CharType* Dest, SizeType CharCount, SizeType Position) const
 {
-	SizeType CopyCount = FMath::Min(Size - Position, CharCount);
+	const  SizeType CopyCount = FMath::Min(Size - Position, CharCount);
 	FMemory::Memcpy(Dest, DataPtr + Position, CopyCount);
 	return CopyCount;
 }
 
-// Case-insensitive comparison operators
+template <typename CharType, typename ViewType>
+inline ViewType TStringViewImpl<CharType, ViewType>::Left(SizeType CharCount) const
+{
+	return ViewType(DataPtr, FMath::Clamp(CharCount, 0, Size));
+}
 
 template <typename CharType, typename ViewType>
-inline bool operator==(const TStringViewImpl<CharType, ViewType>& Lhs, const TStringViewImpl<CharType, ViewType>& Rhs)
+inline ViewType TStringViewImpl<CharType, ViewType>::LeftChop(SizeType CharCount) const
 {
-	return Lhs.Len() == Rhs.Len() && FPlatformString::Strnicmp(Lhs.GetData(), Rhs.GetData(), Lhs.Len()) == 0;
+	return ViewType(DataPtr, FMath::Clamp(Size - CharCount, 0, Size));
 }
 
 template <typename CharType, typename ViewType>
-inline bool operator!=(const TStringViewImpl<CharType, ViewType>& Lhs, const TStringViewImpl<CharType, ViewType>& Rhs)
+inline ViewType TStringViewImpl<CharType, ViewType>::Right(SizeType CharCount) const
 {
-	return !operator==(Lhs, Rhs);
+	const SizeType OutLen = FMath::Clamp(CharCount, 0, Size);
+	return ViewType(DataPtr + Size - OutLen, OutLen);
 }
 
-// Case-insensitive range comparison operators
+template <typename CharType, typename ViewType>
+inline ViewType TStringViewImpl<CharType, ViewType>::RightChop(SizeType CharCount) const
+{
+	const SizeType OutLen = FMath::Clamp(Size - CharCount, 0, Size);
+	return ViewType(DataPtr + Size - OutLen, OutLen);
+}
+
+template <typename CharType, typename ViewType>
+inline ViewType TStringViewImpl<CharType, ViewType>::Mid(SizeType Position, SizeType CharCount) const
+{
+	check(CharCount >= 0);
+	using USizeType = TUnsignedIntType_T<sizeof(SizeType)>;
+	Position = FMath::Clamp<USizeType>(Position, 0, Size);
+	CharCount = FMath::Clamp<USizeType>(CharCount, 0, Size - Position);
+	return ViewType(DataPtr + Position, CharCount);
+}
+
+template <typename CharType, typename ViewType>
+inline bool TStringViewImpl<CharType, ViewType>::Equals(const TStringViewImpl& Other, ESearchCase::Type SearchCase) const
+{
+	return Size == Other.Size && Compare(Other, SearchCase) == 0;
+}
+
+template <typename CharType, typename ViewType>
+inline bool TStringViewImpl<CharType, ViewType>::StartsWith(const TStringViewImpl& Prefix) const
+{
+	return Size >= Prefix.Size && TCString<CharType>::Strnicmp(DataPtr, Prefix.DataPtr, Prefix.Size) == 0;
+}
+
+template <typename CharType, typename ViewType>
+inline bool TStringViewImpl<CharType, ViewType>::EndsWith(const TStringViewImpl& Suffix) const
+{
+	return Size >= Suffix.Size && TCString<CharType>::Strnicmp(DataPtr + Size - Suffix.Size, Suffix.DataPtr, Suffix.Size) == 0;
+}
+
+// Case-insensitive comparison operators
 
-template <typename CharType, typename ViewType, typename RangeType, typename = typename TStringViewImpl<CharType, ViewType>::template TEnableIfCompatibleRangeType<RangeType>>
-inline bool operator==(const TStringViewImpl<CharType, ViewType>& Lhs, const RangeType& Rhs)
+template <typename CharType, typename ViewType, typename RangeType,
+	typename = typename TStringViewImpl<CharType, ViewType>::template TEnableIfCompatibleRangeType<RangeType>>
+inline bool operator==(const TStringViewImpl<CharType, ViewType>& Lhs, RangeType&& Rhs)
 {
-	return Lhs == TStringViewImpl<CharType, ViewType>(Rhs);
+	return Lhs.Equals(TStringViewImpl<CharType, ViewType>(Forward<RangeType>(Rhs)), ESearchCase::IgnoreCase);
 }
 
-template <typename CharType, typename ViewType, typename RangeType, typename = typename TStringViewImpl<CharType, ViewType>::template TEnableIfCompatibleRangeType<RangeType>>
-inline bool operator==(const RangeType& Lhs, const TStringViewImpl<CharType, ViewType>& Rhs)
+template <typename CharType, typename ViewType, typename RangeType,
+	typename = typename TEnableIf<TNot<TIsDerivedFrom<typename TDecay<RangeType>::Type, TStringViewImpl<CharType, ViewType>>>::Value>::Type>
+inline auto operator==(RangeType&& Lhs, const TStringViewImpl<CharType, ViewType>& Rhs) -> decltype(Rhs == Forward<RangeType>(Lhs))
 {
-	return TStringViewImpl<CharType, ViewType>(Lhs) == Rhs;
+	return Rhs == Forward<RangeType>(Lhs);
 }
 
-template <typename CharType, typename ViewType, typename RangeType, typename = typename TStringViewImpl<CharType, ViewType>::template TEnableIfCompatibleRangeType<RangeType>>
-inline bool operator!=(const TStringViewImpl<CharType, ViewType>& Lhs, const RangeType& Rhs)
+template <typename CharType, typename ViewType, typename RangeType>
+inline auto operator!=(const TStringViewImpl<CharType, ViewType>& Lhs, RangeType&& Rhs) -> decltype(!(Lhs == Forward<RangeType>(Rhs)))
 {
-	return !operator==(Lhs, Rhs);
+	return !(Lhs == Forward<RangeType>(Rhs));
 }
 
-template <typename CharType, typename ViewType, typename RangeType, typename = typename TStringViewImpl<CharType, ViewType>::template TEnableIfCompatibleRangeType<RangeType>>
-inline bool operator!=(const RangeType& Lhs, const TStringViewImpl<CharType, ViewType>& Rhs)
+template <typename CharType, typename ViewType, typename RangeType,
+	typename = typename TEnableIf<TNot<TIsDerivedFrom<typename TDecay<RangeType>::Type, TStringViewImpl<CharType, ViewType>>>::Value>::Type>
+inline auto operator!=(RangeType&& Lhs, const TStringViewImpl<CharType, ViewType>& Rhs) -> decltype(!(Rhs == Forward<RangeType>(Lhs)))
 {
-	return !operator==(Lhs, Rhs);
+	return !(Rhs == Forward<RangeType>(Lhs));
 }
 
 // Case-insensitive C-string comparison operators
@@ -62,17 +115,17 @@ inline bool operator==(const TStringViewImpl<CharType, ViewType>& Lhs, const Cha
 template <typename CharType, typename ViewType>
 inline bool operator==(const CharType* Lhs, const TStringViewImpl<CharType, ViewType>& Rhs)
 {
-	return TCString<CharType>::Strnicmp(Lhs, Rhs.GetData(), Rhs.Len()) == 0 && !Lhs[Rhs.Len()];
+	return Rhs == Lhs;
 }
 
 template <typename CharType, typename ViewType>
 inline bool operator!=(const TStringViewImpl<CharType, ViewType>& Lhs, const CharType* Rhs)
 {
-	return !operator==(Lhs, Rhs);
+	return !(Lhs == Rhs);
 }
 
 template <typename CharType, typename ViewType>
 inline bool operator!=(const CharType* Lhs, const TStringViewImpl<CharType, ViewType>& Rhs)
 {
-	return !operator==(Lhs, Rhs);
+	return !(Lhs == Rhs);
 }
diff --git a/Engine/Source/Runtime/Core/Public/Templates/UnrealTemplate.h b/Engine/Source/Runtime/Core/Public/Templates/UnrealTemplate.h
--- a/Engine/Source/Runtime/Core/Public/Templates/UnrealTemplate.h
+++ b/Engine/Source/Runtime/Core/Public/Templates/UnrealTemplate.h
@@ -90,7 +90,7 @@ CONSTEXPR const T* GetData(std::initializer_list<T> List)
 * Generically gets the number of items in a contiguous container
 */
 template<typename T, typename = typename TEnableIf<TIsContiguousContainer<T>::Value>::Type>
-SIZE_T GetNum(T&& Container)
+auto GetNum(T&& Container) -> decltype((SIZE_T)Container.Num())
 {
 	return (SIZE_T)Container.Num();
 }
diff --git a/Engine/Source/Runtime/Core/Public/Templates/UnrealTypeTraits.h b/Engine/Source/Runtime/Core/Public/Templates/UnrealTypeTraits.h
--- a/Engine/Source/Runtime/Core/Public/Templates/UnrealTypeTraits.h
+++ b/Engine/Source/Runtime/Core/Public/Templates/UnrealTypeTraits.h
@@ -54,7 +54,9 @@ struct TIsDerivedFrom
 	public:
 	// Test the derived type pointer. If it inherits from BaseType, the Test( BaseType* ) 
 	// will be chosen. If it does not, Test( ... ) will be chosen.
-	static const bool IsDerived = sizeof(Test( DerivedTypePtr() )) == sizeof(Yes);
+	static const bool Value = sizeof(Test( DerivedTypePtr() )) == sizeof(Yes);
+
+	static const bool IsDerived = Value;
 };
 
 /**
