Add rvalue overloads for Left, LeftChop, Right, RightChop, Mind
Add inline version of TrimQuotes

diff --git a/Engine/Source/Runtime/Core/Private/Containers/String.cpp b/Engine/Source/Runtime/Core/Private/Containers/String.cpp
--- a/Engine/Source/Runtime/Core/Private/Containers/String.cpp
+++ b/Engine/Source/Runtime/Core/Private/Containers/String.cpp
@@ -590,7 +590,7 @@ FString FString::TrimEnd() &&
 	return Result;
 }
 
-FString FString::TrimQuotes( bool* bQuotesRemoved ) const
+void FString::TrimQuotesInline(bool* bQuotesRemoved)
 {
 	bool bQuotesWereRemoved=false;
 	int32 Start = 0, Count = Len();
@@ -614,7 +614,14 @@ FString FString::TrimQuotes( bool* bQuotesRemoved ) const
 	{
 		*bQuotesRemoved = bQuotesWereRemoved;
 	}
-	return Mid(Start, Count);
+	MidInline(Start, Count, false);
+}
+
+FString FString::TrimQuotes(bool* bQuotesRemoved) const
+{
+	FString Result(*this);
+	Result.TrimQuotesInline(bQuotesRemoved);
+	return Result;
 }
 
 int32 FString::CullArray( TArray<FString>* InArray )
diff --git a/Engine/Source/Runtime/Core/Public/Containers/UnrealString.h b/Engine/Source/Runtime/Core/Public/Containers/UnrealString.h
--- a/Engine/Source/Runtime/Core/Public/Containers/UnrealString.h
+++ b/Engine/Source/Runtime/Core/Public/Containers/UnrealString.h
@@ -1144,11 +1144,18 @@ public:
 	}
 
 	/** Returns the left most given number of characters */
-	FORCEINLINE FString Left( int32 Count ) const
+	FORCEINLINE FString Left( int32 Count ) const &
 	{
 		return FString( FMath::Clamp(Count,0,Len()), **this );
 	}
 
+	FORCEINLINE FString Left(int32 Count) &&
+	{
+		FString Str(MoveTemp(*this));
+		Str.LeftInline(Count, false);
+		return Str;
+	}
+
 	/** Modifies the string such that it is now the left most given number of characters */
 	FORCEINLINE void LeftInline(int32 Count, bool bAllowShrinking = true)
 	{
@@ -1158,12 +1165,19 @@ public:
 	}
 
 	/** Returns the left most characters from the string chopping the given number of characters from the end */
-	FORCEINLINE FString LeftChop( int32 Count ) const
+	FORCEINLINE FString LeftChop( int32 Count ) const &
 	{
 		const int32 Length = Len();
 		return FString( FMath::Clamp(Length-Count,0, Length), **this );
 	}
 
+	FORCEINLINE FString LeftChop(int32 Count)&&
+	{
+		FString Str(MoveTemp(*this));
+		Str.LeftChopInline(Count, false);
+		return Str;
+	}
+
 	/** Modifies the string such that it is now the left most characters chopping the given number of characters from the end */
 	FORCEINLINE void LeftChopInline(int32 Count, bool bAllowShrinking = true)
 	{
@@ -1172,12 +1186,19 @@ public:
 	}
 
 	/** Returns the string to the right of the specified location, counting back from the right (end of the word). */
-	FORCEINLINE FString Right( int32 Count ) const
+	FORCEINLINE FString Right( int32 Count ) const &
 	{
 		const int32 Length = Len();
 		return FString( **this + Length-FMath::Clamp(Count,0,Length) );
 	}
 
+	FORCEINLINE FString Right(int32 Count) &&
+	{
+		FString Str(MoveTemp(*this));
+		Str.RightInline(Count, false);
+		return Str;
+	}
+
 	/** Modifies the string such that it is now the right most given number of characters */
 	FORCEINLINE void RightInline(int32 Count, bool bAllowShrinking = true)
 	{
@@ -1186,12 +1207,19 @@ public:
 	}
 
 	/** Returns the string to the right of the specified location, counting forward from the left (from the beginning of the word). */
-	FORCEINLINE FString RightChop( int32 Count ) const
+	FORCEINLINE FString RightChop( int32 Count ) const &
 	{
 		const int32 Length = Len();
 		return FString( **this + Length-FMath::Clamp(Length-Count,0, Length) );
 	}
 
+	FORCEINLINE FString RightChop(int32 Count) &&
+	{
+		FString Str(MoveTemp(*this));
+		Str.RightChopInline(Count, false);
+		return Str;
+	}
+
 	/** Modifies the string such that it is now the string to the right of the specified location, counting forward from the left (from the beginning of the word). */
 	FORCEINLINE void RightChopInline(int32 Count, bool bAllowShrinking = true)
 	{
@@ -1199,7 +1227,7 @@ public:
 	}
 
 	/** Returns the substring from Start position for Count characters. */
-	FORCEINLINE FString Mid( int32 Start, int32 Count=MAX_int32 ) const
+	FORCEINLINE FString Mid( int32 Start, int32 Count=MAX_int32 ) const &
 	{
 		FString Result;
 		if (Count >= 0)
@@ -1213,6 +1241,13 @@ public:
 		return Result;
 	}
 
+	FORCEINLINE FString Mid(int32 Start, int32 Count = MAX_int32) &&
+	{
+		FString Str(MoveTemp(*this));
+		Str.MidInline(Start, Count, false);
+		return Str;
+	}
+
 	/** Modifies the string such that it is now the substring from Start position for Count characters. */
 	FORCEINLINE void MidInline(int32 Start, int32 Count = MAX_int32, bool bAllowShrinking = true)
 	{
@@ -1636,6 +1671,12 @@ public:
 	 */
 	void TrimToNullTerminator();
 
+
+	/**
+	 * Trims wrapping quotation marks from this string.
+	 */
+	void TrimQuotesInline(bool* bQuotesRemoved = nullptr);
+
 	/**
 	 * Returns a copy of this string with wrapping quotation marks removed.
 	 */
