#UE4 You can now specify -ini overrides on the commandline more than once

diff --git a/Engine/Source/Runtime/Core/Private/Misc/ConfigCacheIni.cpp b/Engine/Source/Runtime/Core/Private/Misc/ConfigCacheIni.cpp
--- a/Engine/Source/Runtime/Core/Private/Misc/ConfigCacheIni.cpp
+++ b/Engine/Source/Runtime/Core/Private/Misc/ConfigCacheIni.cpp
@@ -1137,7 +1137,8 @@ void FConfigFile::OverrideFromCommandline(FConfigFile* File, const FString& File
 	// for example:
 	//		-ini:Engine:[/Script/Engine.Engine]:bSmoothFrameRate=False,[TextureStreaming]:PoolSize=100
 	//			(will update the cache after the final combined engine.ini)
-	if (FParse::Value(FCommandLine::Get(), *FString::Printf(TEXT("%s%s"), CommandlineOverrideSpecifiers::IniSwitchIdentifier, *FPaths::GetBaseFilename(Filename)), Settings, false))
+	const TCHAR* CommandlineStream = FCommandLine::Get();
+	while(FParse::Value(CommandlineStream, *FString::Printf(TEXT("%s%s"), CommandlineOverrideSpecifiers::IniSwitchIdentifier, *FPaths::GetBaseFilename(Filename)), Settings, false))
 	{
 		// break apart on the commas
 		TArray<FString> SettingPairs;
@@ -1173,6 +1174,11 @@ void FConfigFile::OverrideFromCommandline(FConfigFile* File, const FString& File
 				File->SetString(*CommandlineOption.Section, *CommandlineOption.PropertyKey, *CommandlineOption.PropertyValue);
 			}
 		}
+
+		// Keep searching for more instances of -ini
+		CommandlineStream = FCString::Stristr(CommandlineStream, CommandlineOverrideSpecifiers::IniSwitchIdentifier);
+		check(CommandlineStream);
+		CommandlineStream++;
 	}
 #endif
 }
