Attempted CIS fix

#rnx

diff --git a/Engine/Source/Runtime/Core/Public/Misc/EnumClassFlags.h b/Engine/Source/Runtime/Core/Public/Misc/EnumClassFlags.h
--- a/Engine/Source/Runtime/Core/Public/Misc/EnumClassFlags.h
+++ b/Engine/Source/Runtime/Core/Public/Misc/EnumClassFlags.h
@@ -14,13 +14,13 @@
 	inline constexpr Enum  operator~ (Enum  E)             { return (Enum)~(__underlying_type(Enum))E; }
 
 template<typename Enum>
-constexpr bool EnumHasAllFlags(Enum Flags, Enum Contains)
+static constexpr bool EnumHasAllFlags(Enum Flags, Enum Contains)
 {
 	return ( ( ( __underlying_type(Enum) )Flags ) & ( __underlying_type(Enum) )Contains ) == ( ( __underlying_type(Enum) )Contains );
 }
 
 template<typename Enum>
-constexpr bool EnumHasAnyFlags(Enum Flags, Enum Contains)
+static constexpr bool EnumHasAnyFlags(Enum Flags, Enum Contains)
 {
 	return ( ( ( __underlying_type(Enum) )Flags ) & ( __underlying_type(Enum) )Contains ) != 0;
 }
