Added the ability to disable localization targets

This may be needed if you've gathered Engine data into your Game and want to change the translations, as loading the Engine data would cause a conflict warning between the two LocRes files.

diff --git a/Engine/Source/Runtime/Core/Private/Internationalization/LocalizationResourceTextSource.cpp b/Engine/Source/Runtime/Core/Private/Internationalization/LocalizationResourceTextSource.cpp
--- a/Engine/Source/Runtime/Core/Private/Internationalization/LocalizationResourceTextSource.cpp
+++ b/Engine/Source/Runtime/Core/Private/Internationalization/LocalizationResourceTextSource.cpp
@@ -119,6 +119,16 @@ void FLocalizationResourceTextSource::LoadLocalizedResources(const ELocalization
 	static const FString PlatformName = ANSI_TO_TCHAR(FPlatformProperties::IniPlatformName());
 	auto LoadLocalizationResourcesForCulture = [](FTextLocalizationResource& InOutLocRes, const FString& InLocalizationPath, const FString& InCulture, const FString& InLocResFilename, const int32 InLocResPriority)
 	{
+		const TArray<FString>& DisabledLocalizationTargets = TextLocalizationResourceUtil::GetDisabledLocalizationTargets();
+		if (DisabledLocalizationTargets.Num() > 0)
+		{
+			const FString LocalizationTargetName = FPaths::GetBaseFilename(InLocResFilename);
+			if (DisabledLocalizationTargets.Contains(LocalizationTargetName))
+			{
+				return;
+			}
+		}
+
 		const FString PlatformAgnosticLocResFilename = InLocalizationPath / InCulture / InLocResFilename;
 		InOutLocRes.LoadFromFile(PlatformAgnosticLocResFilename, InLocResPriority);
 
diff --git a/Engine/Source/Runtime/Core/Private/Internationalization/TextLocalizationResource.cpp b/Engine/Source/Runtime/Core/Private/Internationalization/TextLocalizationResource.cpp
--- a/Engine/Source/Runtime/Core/Private/Internationalization/TextLocalizationResource.cpp
+++ b/Engine/Source/Runtime/Core/Private/Internationalization/TextLocalizationResource.cpp
@@ -4,9 +4,11 @@
 #include "Internationalization/TextLocalizationResourceVersion.h"
 #include "Internationalization/Culture.h"
 #include "HAL/FileManager.h"
+#include "Misc/App.h"
 #include "Misc/Parse.h"
 #include "Misc/Paths.h"
 #include "Misc/Optional.h"
+#include "Misc/ConfigCacheIni.h"
 #include "Templates/UniquePtr.h"
 #include "Internationalization/Internationalization.h"
 #include "Misc/FileHelper.h"
@@ -635,5 +637,39 @@ TArray<FString> TextLocalizationResourceUtil::GetLocalizedCultureNames(const TAr
 	return CultureNames;
 }
 
+const TArray<FString>& TextLocalizationResourceUtil::GetDisabledLocalizationTargets()
+{
+	static TArray<FString> DisabledLocalizationTargets;
+	static bool bHasInitializedDisabledLocalizationTargets = false;
+
+	if (!bHasInitializedDisabledLocalizationTargets)
+	{
+		check(GConfig && GConfig->IsReadyForUse());
+
+		const bool bShouldLoadEditor = GIsEditor;
+		const bool bShouldLoadGame = FApp::IsGame();
+
+		GConfig->GetArray(TEXT("Internationalization"), TEXT("DisabledLocalizationTargets"), DisabledLocalizationTargets, GEngineIni);
+
+		if (bShouldLoadEditor)
+		{
+			TArray<FString> EditorArray;
+			GConfig->GetArray(TEXT("Internationalization"), TEXT("DisabledLocalizationTargets"), EditorArray, GEditorIni);
+			DisabledLocalizationTargets.Append(MoveTemp(EditorArray));
+		}
+
+		if (bShouldLoadGame)
+		{
+			TArray<FString> GameArray;
+			GConfig->GetArray(TEXT("Internationalization"), TEXT("DisabledLocalizationTargets"), GameArray, GGameIni);
+			DisabledLocalizationTargets.Append(MoveTemp(GameArray));
+		}
+
+		bHasInitializedDisabledLocalizationTargets = true;
+	}
+
+	return DisabledLocalizationTargets;
+}
+
 #undef PRELOAD_LOCMETA_FILES
 #undef PRELOAD_LOCRES_FILES
diff --git a/Engine/Source/Runtime/Core/Public/Internationalization/TextLocalizationResource.h b/Engine/Source/Runtime/Core/Public/Internationalization/TextLocalizationResource.h
--- a/Engine/Source/Runtime/Core/Public/Internationalization/TextLocalizationResource.h
+++ b/Engine/Source/Runtime/Core/Public/Internationalization/TextLocalizationResource.h
@@ -153,4 +153,9 @@ CORE_API void ClearNativeEditorCultureName();
  */
 CORE_API TArray<FString> GetLocalizedCultureNames(const TArray<FString>& InLocalizationPaths);
 
+/**
+ * Get the array of localization targets that have been disabled for the current configuration.
+ */
+CORE_API const TArray<FString>& GetDisabledLocalizationTargets();
+
 }
