Apply codereview feedback to rvalue changes
#rnx

diff --git a/Engine/Source/Runtime/Core/Private/Containers/String.cpp b/Engine/Source/Runtime/Core/Private/Containers/String.cpp
--- a/Engine/Source/Runtime/Core/Private/Containers/String.cpp
+++ b/Engine/Source/Runtime/Core/Private/Containers/String.cpp
@@ -537,9 +537,8 @@ FString FString::TrimStartAndEnd() const &
 
 FString FString::TrimStartAndEnd() &&
 {
-	FString Result(MoveTemp(*this));
-	Result.TrimStartAndEndInline();
-	return Result;
+	TrimStartAndEndInline();
+	return MoveTemp(*this);
 }
 
 void FString::TrimStartInline()
@@ -561,9 +560,8 @@ FString FString::TrimStart() const &
 
 FString FString::TrimStart() &&
 {
-	FString Result(MoveTemp(*this));
-	Result.TrimStartInline();
-	return Result;
+	TrimStartInline();
+	return MoveTemp(*this);
 }
 
 void FString::TrimEndInline()
@@ -585,9 +583,8 @@ FString FString::TrimEnd() const &
 
 FString FString::TrimEnd() &&
 {
-	FString Result(MoveTemp(*this));
-	Result.TrimEndInline();
-	return Result;
+	TrimEndInline();
+	return MoveTemp(*this);
 }
 
 void FString::TrimQuotesInline(bool* bQuotesRemoved)
@@ -617,13 +614,19 @@ void FString::TrimQuotesInline(bool* bQuotesRemoved)
 	MidInline(Start, Count, false);
 }
 
-FString FString::TrimQuotes(bool* bQuotesRemoved) const
+FString FString::TrimQuotes(bool* bQuotesRemoved) const &
 {
 	FString Result(*this);
 	Result.TrimQuotesInline(bQuotesRemoved);
 	return Result;
 }
 
+FString FString::TrimQuotes(bool* bQuotesRemoved) &&
+{
+	TrimQuotesInline(bQuotesRemoved);
+	return MoveTemp(*this);
+}
+
 int32 FString::CullArray( TArray<FString>* InArray )
 {
 	check(InArray);
diff --git a/Engine/Source/Runtime/Core/Public/Containers/UnrealString.h b/Engine/Source/Runtime/Core/Public/Containers/UnrealString.h
--- a/Engine/Source/Runtime/Core/Public/Containers/UnrealString.h
+++ b/Engine/Source/Runtime/Core/Public/Containers/UnrealString.h
@@ -1151,9 +1151,8 @@ public:
 
 	FORCEINLINE FString Left(int32 Count) &&
 	{
-		FString Str(MoveTemp(*this));
-		Str.LeftInline(Count, false);
-		return Str;
+		LeftInline(Count, false);
+		return MoveTemp(*this);
 	}
 
 	/** Modifies the string such that it is now the left most given number of characters */
@@ -1173,9 +1172,8 @@ public:
 
 	FORCEINLINE FString LeftChop(int32 Count)&&
 	{
-		FString Str(MoveTemp(*this));
-		Str.LeftChopInline(Count, false);
-		return Str;
+		LeftChopInline(Count, false);
+		return MoveTemp(*this);
 	}
 
 	/** Modifies the string such that it is now the left most characters chopping the given number of characters from the end */
@@ -1194,9 +1192,8 @@ public:
 
 	FORCEINLINE FString Right(int32 Count) &&
 	{
-		FString Str(MoveTemp(*this));
-		Str.RightInline(Count, false);
-		return Str;
+		RightInline(Count, false);
+		return MoveTemp(*this);
 	}
 
 	/** Modifies the string such that it is now the right most given number of characters */
@@ -1215,9 +1212,8 @@ public:
 
 	FORCEINLINE FString RightChop(int32 Count) &&
 	{
-		FString Str(MoveTemp(*this));
-		Str.RightChopInline(Count, false);
-		return Str;
+		RightChopInline(Count, false);
+		return MoveTemp(*this);
 	}
 
 	/** Modifies the string such that it is now the string to the right of the specified location, counting forward from the left (from the beginning of the word). */
@@ -1243,9 +1239,8 @@ public:
 
 	FORCEINLINE FString Mid(int32 Start, int32 Count = MAX_int32) &&
 	{
-		FString Str(MoveTemp(*this));
-		Str.MidInline(Start, Count, false);
-		return Str;
+		MidInline(Start, Count, false);
+		return MoveTemp(*this);
 	}
 
 	/** Modifies the string such that it is now the substring from Start position for Count characters. */
@@ -1680,7 +1675,12 @@ public:
 	/**
 	 * Returns a copy of this string with wrapping quotation marks removed.
 	 */
-	FString TrimQuotes( bool* bQuotesRemoved = nullptr ) const;
+	FString TrimQuotes( bool* bQuotesRemoved = nullptr ) const &;
+
+	/**
+	 * Returns this string with wrapping quotation marks removed.
+	 */
+	FString TrimQuotes(bool* bQuotesRemoved = nullptr) &&;
 
 	/**
 	 * Breaks up a delimited string into elements of a string array.
