Adding UELogDisplay to EmbeddedCommunication to support the additional Display verbosity
#tests iOS Client

diff --git a/Engine/Source/Runtime/Core/Private/Misc/EmbeddedCommunication.cpp b/Engine/Source/Runtime/Core/Private/Misc/EmbeddedCommunication.cpp
--- a/Engine/Source/Runtime/Core/Private/Misc/EmbeddedCommunication.cpp
+++ b/Engine/Source/Runtime/Core/Private/Misc/EmbeddedCommunication.cpp
@@ -233,6 +233,16 @@ void FEmbeddedCommunication::UELogWarning(const TCHAR* String)
 #endif
 }
 
+void FEmbeddedCommunication::UELogDisplay(const TCHAR* String)
+{
+#if BUILD_EMBEDDED_APP
+	if (GLog && UE_LOG_ACTIVE(LogBridge, Display))
+	{
+		GLog->Log("LogBridge", ELogVerbosity::Display, String);
+	}
+#endif
+}
+
 void FEmbeddedCommunication::UELogLog(const TCHAR* String)
 {
 #if BUILD_EMBEDDED_APP
diff --git a/Engine/Source/Runtime/Core/Public/Misc/EmbeddedCommunication.h b/Engine/Source/Runtime/Core/Public/Misc/EmbeddedCommunication.h
--- a/Engine/Source/Runtime/Core/Public/Misc/EmbeddedCommunication.h
+++ b/Engine/Source/Runtime/Core/Public/Misc/EmbeddedCommunication.h
@@ -83,6 +83,7 @@ public:
 	
 	static void UELogError(const TCHAR* String);
 	static void UELogWarning(const TCHAR* String);
+	static void UELogDisplay(const TCHAR* String);
 	static void UELogLog(const TCHAR* String);
 	static void UELogVerbose(const TCHAR* String);
 
