Android Crash reporter
 - Add gamethread to all threads callstack lists
 - Fixed issue which could leave a malformed threads callstack file on disk after report upload
 - Added detailed signal information to crash report message.
 - Fixed a hang that could occur when generating all threads callstack file.
#rnx

diff --git a/Engine/Source/Runtime/Core/Private/Android/AndroidPlatformCrashContext.cpp b/Engine/Source/Runtime/Core/Private/Android/AndroidPlatformCrashContext.cpp
--- a/Engine/Source/Runtime/Core/Private/Android/AndroidPlatformCrashContext.cpp
+++ b/Engine/Source/Runtime/Core/Private/Android/AndroidPlatformCrashContext.cpp
@@ -75,7 +75,7 @@ void FAndroidCrashContext::GetCrashDirectoryName(char(&DirectoryNameOUT)[CrashRe
 	FCStringAnsi::Strncpy(DirectoryNameOUT, GAndroidCrashInfo.TargetDirectory, CrashReportMaxPathSize);
 }
 
-static ANSICHAR* ItoANSI(uint64 Val, uint64 Base, uint32 Len = 0)
+const ANSICHAR* FAndroidCrashContext::ItoANSI(uint64 Val, uint64 Base, uint32 Len)
 {
 	static ANSICHAR InternalBuffer[64] = { 0 };
 
@@ -167,6 +167,7 @@ void FAndroidCrashContext::DumpAllThreadCallstacks() const
 	int DestHandle = open(FilePath, O_CREAT | O_WRONLY, S_IRUSR | S_IWUSR);
 	if (DestHandle >= 0)
 	{
+		int32 CurrentThreadID = FPlatformTLS::GetCurrentThreadId();
 		auto Write = [](int FileHandle, const ANSICHAR* Buffer)
 		{
 			write(FileHandle, Buffer, FCStringAnsi::Strlen(Buffer));
@@ -177,22 +178,25 @@ void FAndroidCrashContext::DumpAllThreadCallstacks() const
 			write(FileHandle, "\n", 1);
 		};
 
-		// For each thread append it's info to the file.
-		FThreadManager::Get().ForEachThread([Writeln, Write, this, DestHandle, &CrashStackFrames, &CallstacksRecorded](uint32 ThreadID, FRunnableThread* Runnable)
+		auto WriteThreadEntry = [Writeln, Write, this, DestHandle, &CrashStackFrames, &CallstacksRecorded, CurrentThreadID](uint32 ThreadID, const char* ThreadName)
 		{
 			// Capture the stack trace
 			static const int StackTraceMaxDepth = 100;
 			uint64 StackTrace[StackTraceMaxDepth];
 			FMemory::Memzero(StackTrace);
-			uint32 Depth = FPlatformStackWalk::CaptureThreadStackBackTrace(ThreadID, StackTrace, StackTraceMaxDepth);
+			uint32 Depth = 0;
+			if (CurrentThreadID == ThreadID)
+			{
+				Depth = FPlatformStackWalk::CaptureStackBackTrace(StackTrace, StackTraceMaxDepth);
+			}
+			else
+			{
+				Depth = FPlatformStackWalk::CaptureThreadStackBackTrace(ThreadID, StackTrace, StackTraceMaxDepth);
+			}
 
 			if (Depth)
 			{
 				ANSICHAR Line[256];
-				if (CallstacksRecorded++ == 0)
-				{
-					Writeln(DestHandle, "<Threads>");
-				}
 				Writeln(DestHandle, "<Thread>");
 				Write(DestHandle, "<CallStack>");
 				// Write stack
@@ -217,18 +221,32 @@ void FAndroidCrashContext::DumpAllThreadCallstacks() const
 
 				// write Thread Name
 				Write(DestHandle, "<ThreadName>");
-				FCStringAnsi::Strncpy(Line, TCHAR_TO_UTF8(*Runnable->GetThreadName()), UE_ARRAY_COUNT(Line));
+				FCStringAnsi::Strncpy(Line, ThreadName, UE_ARRAY_COUNT(Line));
 				Write(DestHandle, Line);
 				Writeln(DestHandle, "</ThreadName>");
 
 				Writeln(DestHandle, "</Thread>");
+				CallstacksRecorded++;
 			}
-		});
-		if (CallstacksRecorded)
+		};
+
+		Writeln(DestHandle, "<Threads>");
+
+		if( GIsGameThreadIdInitialized )
 		{
-			Writeln(DestHandle, "</Threads>");
+			// On android the Game thread is that which calls android_main entry point. 
+			// Explicitly call it here as the thread manager is not aware of it.
+			WriteThreadEntry(GGameThreadId, "GameThread");
 		}
+		// For each thread append it's info to the file.
+		FThreadManager::Get().ForEachThread([WriteThreadEntry, this, DestHandle, &CrashStackFrames, &CallstacksRecorded](uint32 ThreadID, FRunnableThread* Runnable)
+		{
+			WriteThreadEntry(ThreadID, TCHAR_TO_UTF8(*Runnable->GetThreadName()));
+		});
+
+		Writeln(DestHandle, "</Threads>");
 		close(DestHandle);
+
 		if (CallstacksRecorded == 0)
 		{
 			unlink(FilePath); // remove the file if nothing was written.
diff --git a/Engine/Source/Runtime/Core/Private/Android/AndroidPlatformMisc.cpp b/Engine/Source/Runtime/Core/Private/Android/AndroidPlatformMisc.cpp
--- a/Engine/Source/Runtime/Core/Private/Android/AndroidPlatformMisc.cpp
+++ b/Engine/Source/Runtime/Core/Private/Android/AndroidPlatformMisc.cpp
@@ -902,9 +902,29 @@ static constexpr int32 TargetSignals[] =
 	SIGSEGV,
 	SIGSYS
 };
-
 static constexpr int32 NumTargetSignals = UE_ARRAY_COUNT(TargetSignals);
 
+static const char* SignalToString(int32 Signal)
+{
+	switch (Signal)
+	{
+		case SIGQUIT:
+			return "SIGQUIT";
+		case SIGILL:
+			return "SIGILL";
+		case SIGFPE:
+			return "SIGFPE";
+		case SIGBUS:
+			return "SIGBUS";
+		case SIGSEGV:
+			return "SIGSEGV";
+		case SIGSYS:
+			return "SIGSYS";
+		default:
+			return FAndroidCrashContext::ItoANSI(Signal,16, 16);
+	}
+}
+
 PRAGMA_DISABLE_OPTIMIZATION
 struct FAndroidSignalParams
 {
@@ -1099,12 +1119,101 @@ private:
 		RestorePreviousSignalHandlers();
 	}
 
+	static const ANSICHAR* CodeToString(int Signal, int si_code)
+	{
+		switch (Signal)
+		{
+			case SIGILL:
+			{
+				switch (si_code)
+				{
+					// SIGILL
+					case ILL_ILLOPC: return "ILL_ILLOPC";
+					case ILL_ILLOPN: return "ILL_ILLOPN";
+					case ILL_ILLADR: return "ILL_ILLADR";
+					case ILL_ILLTRP: return "ILL_ILLTRP";
+					case ILL_PRVOPC: return "ILL_PRVOPC";
+					case ILL_PRVREG: return "ILL_PRVREG";
+					case ILL_COPROC: return "ILL_COPROC";
+					case ILL_BADSTK: return "ILL_BADSTK";
+				}
+			}
+			break;
+			case SIGFPE:
+			{
+				switch (si_code)
+				{
+					// SIGFPE
+					case FPE_INTDIV: return "FPE_INTDIV";
+					case FPE_INTOVF: return "FPE_INTOVF";
+					case FPE_FLTDIV: return "FPE_FLTDIV";
+					case FPE_FLTOVF: return "FPE_FLTOVF";
+					case FPE_FLTUND: return "FPE_FLTUND";
+					case FPE_FLTRES: return "FPE_FLTRES";
+					case FPE_FLTINV: return "FPE_FLTINV";
+					case FPE_FLTSUB: return "FPE_FLTSUB";
+				}
+			}
+			break;
+			case SIGBUS:
+			{
+				switch (si_code)
+				{
+					// SIGBUS
+					case BUS_ADRALN: return "BUS_ADRALN";
+					case BUS_ADRERR: return "BUS_ADRERR";
+					case BUS_OBJERR: return "BUS_OBJERR";
+				}
+			}
+			break;
+			case SIGSEGV:
+			{
+				switch (si_code)
+				{
+					// SIGSEGV
+					case SEGV_MAPERR: return "SEGV_MAPERR";
+					case SEGV_ACCERR: return "SEGV_ACCERR";
+				}
+			}
+			break;
+		}
+		return FAndroidCrashContext::ItoANSI(si_code, 10, 0);
+	}
+
+	static FString GetFatalSignalMessage(int Signal, siginfo* Info)
+	{
+		const int MessageSize = 255;
+		char AnsiMessage[MessageSize];
+		FCStringAnsi::Strncpy(AnsiMessage, "Caught signal : ", MessageSize);
+		FCStringAnsi::Strcat(AnsiMessage, SignalToString(Signal));
+		FCStringAnsi::Strcat(AnsiMessage, " (");
+		FCStringAnsi::Strcat(AnsiMessage, CodeToString(Signal, Info->si_code));
+		FCStringAnsi::Strcat(AnsiMessage, ")");
+		switch (Signal)
+		{
+			case SIGILL:
+			case SIGFPE:
+			case SIGSEGV:
+			case SIGBUS:
+			case SIGTRAP:
+			{
+				FCStringAnsi::Strcat(AnsiMessage, " fault address 0x");
+				FCStringAnsi::Strcat(AnsiMessage, FAndroidCrashContext::ItoANSI((uintptr_t)Info->si_addr, 16, 16));
+				break;
+			}
+		}
+
+		return ANSI_TO_TCHAR(AnsiMessage);
+	}
+
 	static void FatalSignalHandler()
 	{
 		// Switch to malloc crash.
 		FPlatformMallocCrash::Get().SetAsGMalloc();
 
-		FAndroidCrashContext CrashContext(ECrashContextType::Crash, TEXT("Caught signal"));
+		FString Message = GetFatalSignalMessage(FatalSignalParams.Signal, FatalSignalParams.Info);	
+		FAndroidCrashContext CrashContext(ECrashContextType::Crash, *Message);
+
 		CrashContext.InitFromSignal(FatalSignalParams.Signal, FatalSignalParams.Info, FatalSignalParams.Context);
 		CrashContext.CaptureCrashInfo();
 		if (GCrashHandlerPointer)
diff --git a/Engine/Source/Runtime/Core/Private/Android/AndroidPlatformStackWalk.cpp b/Engine/Source/Runtime/Core/Private/Android/AndroidPlatformStackWalk.cpp
--- a/Engine/Source/Runtime/Core/Private/Android/AndroidPlatformStackWalk.cpp
+++ b/Engine/Source/Runtime/Core/Private/Android/AndroidPlatformStackWalk.cpp
@@ -260,7 +260,7 @@ bool FAndroidPlatformStackWalk::SymbolInfoToHumanReadableString(const FProgramCo
 
 static TAutoConsoleVariable<float> CVarAndroidPlatformThreadCallStackMaxWait(
 	TEXT("AndroidPlatformThreadStackWalk.MaxWait"),
-	60.0f,
+	5.0f,
 	TEXT("The number of seconds allowed to spin before killing the process, with the assumption the signal handler has hung."));
 
 #if ANDROID_HAS_THREADBACKTRACE
diff --git a/Engine/Source/Runtime/Core/Public/Android/AndroidPlatformCrashContext.h b/Engine/Source/Runtime/Core/Public/Android/AndroidPlatformCrashContext.h
--- a/Engine/Source/Runtime/Core/Public/Android/AndroidPlatformCrashContext.h
+++ b/Engine/Source/Runtime/Core/Public/Android/AndroidPlatformCrashContext.h
@@ -55,6 +55,9 @@ struct CORE_API FAndroidCrashContext : public FGenericCrashContext
 
 	void DumpAllThreadCallstacks() const;
 
+	/** Async-safe ItoA */
+	static const ANSICHAR* ItoANSI(uint64 Val, uint64 Base, uint32 Len = 0);
+
 private:
 	TMap<FString, FString> AdditionalProperties;
 
diff --git a/Engine/Source/Runtime/CrashReportCore/Private/Android/AndroidErrorReport.cpp b/Engine/Source/Runtime/CrashReportCore/Private/Android/AndroidErrorReport.cpp
--- a/Engine/Source/Runtime/CrashReportCore/Private/Android/AndroidErrorReport.cpp
+++ b/Engine/Source/Runtime/CrashReportCore/Private/Android/AndroidErrorReport.cpp
@@ -27,31 +27,47 @@ void FAndroidErrorReport::ShutDown()
 	AndroidErrorReport::CrashHelperModule->ShutdownModule();
 }
 
-static void AddThreadContexts(const FString& ReportDirectory)
+ FAndroidErrorReport::FAndroidErrorReport(const FString& Directory)
+	: FGenericErrorReport(Directory)
 {
-	// Try to load the callstacks file.
-	const FString ThreadContextsFile(ReportDirectory / TEXT("AllThreads.txt"));
-	IPlatformFile& PlatformFile = FPlatformFileManager::Get().GetPlatformFile();
+	// Check for allthreads.txt, if it exists rename it.
+	// This way if anything goes wrong during processing it will not continue to cause issues on subsequent runs.
 
-	if (PlatformFile.FileExists(*ThreadContextsFile))
+	const FString StartingThreadContextsFileName(TEXT("AllThreads.txt"));
+	const FString StartingThreadContextsFilePath(ReportDirectory / StartingThreadContextsFileName);
+	IPlatformFile& PlatformFile = FPlatformFileManager::Get().GetPlatformFile();
+	bool bHasFoundThreadsFile = PlatformFile.FileExists(*StartingThreadContextsFilePath);
+	if (bHasFoundThreadsFile)
 	{
-		// rename the threads file, if anything goes wrong during processing it will be skipped on the next run.
-		const FString ThreadContextsFileTemp(ReportDirectory / TEXT("AllThreads.tmp"));
-		PlatformFile.MoveFile(*ThreadContextsFileTemp, *ThreadContextsFile);
+		FString ThreadContextsFileName(TEXT("AllThreads.tmp"));
+		ThreadContextsPathName = ReportDirectory / ThreadContextsFileName;
+		PlatformFile.MoveFile(*ThreadContextsPathName, *StartingThreadContextsFilePath);
 
-		FXmlFile ThreadsNode(ThreadContextsFileTemp);
+		// mirror the renaming in ReportFilenames.
+		ReportFilenames.Add(ThreadContextsFileName);
+		ReportFilenames.RemoveSingle(StartingThreadContextsFileName);
+	}
+}
+
+static void AddThreadContexts(const FString& ThreadContextsPathName)
+{
+	// Try to load the callstacks file.
+	if (!ThreadContextsPathName.IsEmpty())
+	{
+		FXmlFile ThreadsNode(ThreadContextsPathName);
 		if (ThreadsNode.IsValid())
 		{			
 			FPrimaryCrashProperties::Get()->Threads = ThreadsNode.GetRootNode();
 			// delete the file as it has been added to the primary report.
-			PlatformFile.DeleteFile(*ThreadContextsFileTemp);
+			IPlatformFile& PlatformFile = FPlatformFileManager::Get().GetPlatformFile();
+			PlatformFile.DeleteFile(*ThreadContextsPathName);
 		}
 	}
 }
 
 FText FAndroidErrorReport::DiagnoseReport() const
 {
-	AddThreadContexts(ReportDirectory);
+	AddThreadContexts(ThreadContextsPathName);
 	return FText();
 }
 
diff --git a/Engine/Source/Runtime/CrashReportCore/Public/Android/AndroidErrorReport.h b/Engine/Source/Runtime/CrashReportCore/Public/Android/AndroidErrorReport.h
--- a/Engine/Source/Runtime/CrashReportCore/Public/Android/AndroidErrorReport.h
+++ b/Engine/Source/Runtime/CrashReportCore/Public/Android/AndroidErrorReport.h
@@ -26,10 +26,7 @@ public:
 	 * Discover all files in the crash report directory
 	 * @param Directory Full path to directory containing the report
 	 */
-	explicit FAndroidErrorReport(const FString& Directory)
-		: FGenericErrorReport(Directory)
-	{
-	}
+	explicit FAndroidErrorReport(const FString& Directory);
 
 	/**
 	 * Load helper modules
@@ -62,4 +59,10 @@ public:
 		FString AppPath = FPaths::Combine(FPrimaryCrashProperties::Get()->BaseDir, FPrimaryCrashProperties::Get()->ExecutableName);
 		return AppPath;
 	}
+protected:
+
+	/**
+	 * Filename of the renamed threads context file. Empty if no such file exists.
+	 */
+	FString ThreadContextsPathName;
 };
