Add a new API for querying extra development memory.

In the CSVProfiler, this gets reported as metadata and by default we subtract the dev memory from the free memory stat (this is controllable via the cvar csv.AdjustPhysicalFreeForExtraDevMemory)

This change fixes a bug with IsExtraDevelopmentMemoryAvailable() always returning false on xbox, causing the onscreen warning not to display.

diff --git a/Engine/Source/Runtime/Core/Private/GenericPlatform/GenericPlatformMemory.cpp b/Engine/Source/Runtime/Core/Private/GenericPlatform/GenericPlatformMemory.cpp
--- a/Engine/Source/Runtime/Core/Private/GenericPlatform/GenericPlatformMemory.cpp
+++ b/Engine/Source/Runtime/Core/Private/GenericPlatform/GenericPlatformMemory.cpp
@@ -458,6 +458,11 @@ bool FGenericPlatformMemory::IsExtraDevelopmentMemoryAvailable()
 	return false;
 }
 
+uint64 FGenericPlatformMemory::GetExtraDevelopmentMemorySize()
+{
+	return 0;
+}
+
 bool FGenericPlatformMemory::GetLLMAllocFunctions(void*(*&OutAllocFunction)(size_t), void(*&OutFreeFunction)(void*, size_t), int32& OutAlignment)
 {
 	return false;
diff --git a/Engine/Source/Runtime/Core/Private/ProfilingDebugging/CsvProfiler.cpp b/Engine/Source/Runtime/Core/Private/ProfilingDebugging/CsvProfiler.cpp
--- a/Engine/Source/Runtime/Core/Private/ProfilingDebugging/CsvProfiler.cpp
+++ b/Engine/Source/Runtime/Core/Private/ProfilingDebugging/CsvProfiler.cpp
@@ -89,6 +89,16 @@ TAutoConsoleVariable<int32> CVarCsvShippingContinuousWrites(
 );
 #endif
 
+#if !UE_BUILD_SHIPPING
+TAutoConsoleVariable<int32> CVarCsvAdjustPhysicalFreeForExtraDevMemory(
+	TEXT("csv.AdjustPhysicalFreeForExtraDevMemory"),
+	1,
+	TEXT("If 1, extra development memory is subtracted from the physicalFreeMB stat before writing.\r\n")
+	TEXT("This can result in negative values in cases which would have crashed OOM"),
+	ECVF_Default
+);
+#endif
+
 TAutoConsoleVariable<int32> CVarCsvCompressionMode(
 	TEXT("csv.CompressionMode"),
 	-1,
@@ -2586,6 +2596,11 @@ void FCsvProfiler::BeginFrame()
 					}
 					SetMetadata(TEXT("TargetFramerate"), *FString::FromInt(TargetFPS));
 
+#if !UE_BUILD_SHIPPING
+					uint64 ExtraDevelopmentMemoryMB = FPlatformMemory::GetExtraDevelopmentMemorySize()/1024ull/1024ull;
+					SetMetadata(TEXT("ExtraDevelopmentMemoryMB"), *FString::FromInt(ExtraDevelopmentMemoryMB)); 
+#endif
+
 					GCsvStatCounts = !!CVarCsvStatCounts.GetValueOnGameThread();
 
 					// Initialize tls before setting the capturing flag to true.
@@ -2639,6 +2654,13 @@ void FCsvProfiler::EndFrame()
 
 		FPlatformMemoryStats MemoryStats = FPlatformMemory::GetStats();
 		float PhysicalMBFree = float(MemoryStats.AvailablePhysical) / (1024.0f * 1024.0f);
+
+#if !UE_BUILD_SHIPPING
+		if (CVarCsvAdjustPhysicalFreeForExtraDevMemory.GetValueOnGameThread())
+		{
+			PhysicalMBFree -= float(FPlatformMemory::GetExtraDevelopmentMemorySize() / 1024ull / 1024ull);
+		}
+#endif
 		float PhysicalMBUsed = float(MemoryStats.UsedPhysical) / (1024.0f * 1024.0f);
 		float VirtualMBUsed  = float(MemoryStats.UsedVirtual) / (1024.0f * 1024.0f);
 		CSV_CUSTOM_STAT_GLOBAL(MemoryFreeMB, PhysicalMBFree, ECsvCustomStatOp::Set);
diff --git a/Engine/Source/Runtime/Core/Public/GenericPlatform/GenericPlatformMemory.h b/Engine/Source/Runtime/Core/Public/GenericPlatform/GenericPlatformMemory.h
--- a/Engine/Source/Runtime/Core/Public/GenericPlatform/GenericPlatformMemory.h
+++ b/Engine/Source/Runtime/Core/Public/GenericPlatform/GenericPlatformMemory.h
@@ -629,10 +629,16 @@ public:
 
 	/**
 	* Returns true if debug memory has been assigned to the title for general use.
-	* Only applies to consoles with fixed memory and no paging.
+	* Only applies to platforms with fixed memory and no paging.
 	*/
 	static bool IsExtraDevelopmentMemoryAvailable();
 
+	/**
+	* Returns >0 if debug memory has been assigned to the title for general use.
+	* Only applies to platforms with fixed memory and no paging.
+	*/
+	static uint64 GetExtraDevelopmentMemorySize();
+
 	/**
 	* This function sets AllocFunction and FreeFunction and returns true, or just returns false.
 	* These functions are the platform dependant low low low level functions that LLM uses to allocate memory.
diff --git a/Engine/Source/Runtime/Engine/Private/UnrealEngine.cpp b/Engine/Source/Runtime/Engine/Private/UnrealEngine.cpp
--- a/Engine/Source/Runtime/Engine/Private/UnrealEngine.cpp
+++ b/Engine/Source/Runtime/Engine/Private/UnrealEngine.cpp
@@ -1718,6 +1718,10 @@ void UEngine::Init(IEngineLoop* InEngineLoop)
 
 	// Record the analytics for any attached HMD devices
 	RecordHMDAnalytics();
+
+#if !UE_BUILD_SHIPPING
+	UE_CLOG(FPlatformMemory::IsExtraDevelopmentMemoryAvailable(), LogInit, Warning, TEXT("Running with %dMB of extra development memory!"),FPlatformMemory::GetExtraDevelopmentMemorySize()/1024ull/1024ull);
+#endif
 }
 
 void UEngine::Start()
