Allow per thread callstack walking to time out and fail silently.
 - reduces likelyhood of a hang during thread callstack walking.
#rnx

diff --git a/Engine/Source/Runtime/Core/Private/Android/AndroidPlatformMisc.cpp b/Engine/Source/Runtime/Core/Private/Android/AndroidPlatformMisc.cpp
--- a/Engine/Source/Runtime/Core/Private/Android/AndroidPlatformMisc.cpp
+++ b/Engine/Source/Runtime/Core/Private/Android/AndroidPlatformMisc.cpp
@@ -1039,8 +1039,6 @@ private:
 		StackWalkSignalParams.Info = Info;
 		StackWalkSignalParams.Context = Context;
 
-		ThreadStackData = static_cast<ThreadStackUserData*>(Info->si_value.sival_ptr);
-
 		FSignalSubstituteStack::CallActualHandler< &ThreadStackWalk >();
 
 		FPlatformAtomics::AtomicStore(&handling_signal, 0);
@@ -1048,18 +1046,13 @@ private:
 
 	static void ThreadStackWalk()
 	{
-		if (ThreadStackData)
-		{
-			ThreadStackData->BackTraceCount = FPlatformStackWalk::CaptureStackBackTrace(ThreadStackData->BackTrace, ThreadStackData->CallStackSize, StackWalkSignalParams.Context);
-			ThreadStackData->bDone = true;
-		}
+		FPlatformStackWalk::HandleBackTraceSignal(StackWalkSignalParams.Info, StackWalkSignalParams.Context);
 	}
 	static FAndroidSignalParams StackWalkSignalParams;
 	static volatile sig_atomic_t handling_signal;
 	static ThreadStackUserData* ThreadStackData;
 };
 volatile sig_atomic_t FThreadCallstackSignalSubstituteStack::handling_signal = 0;
-ThreadStackUserData* FThreadCallstackSignalSubstituteStack::ThreadStackData = nullptr;
 FAndroidSignalParams FThreadCallstackSignalSubstituteStack::StackWalkSignalParams;
 #else
 	// per thread backtrace unavailable
diff --git a/Engine/Source/Runtime/Core/Private/Android/AndroidPlatformStackWalk.cpp b/Engine/Source/Runtime/Core/Private/Android/AndroidPlatformStackWalk.cpp
--- a/Engine/Source/Runtime/Core/Private/Android/AndroidPlatformStackWalk.cpp
+++ b/Engine/Source/Runtime/Core/Private/Android/AndroidPlatformStackWalk.cpp
@@ -23,6 +23,8 @@
 #include "HAL/PlatformTime.h"
 #endif
 #include "HAL/IConsoleManager.h"
+#include "HAL/PlatformProcess.h"
+#include "Misc/ScopeExit.h"
 
 void FAndroidPlatformStackWalk::ProgramCounterToSymbolInfo(uint64 ProgramCounter, FProgramCounterSymbolInfo& out_SymbolInfo)
 {
@@ -258,35 +260,86 @@ bool FAndroidPlatformStackWalk::SymbolInfoToHumanReadableString(const FProgramCo
 	return false;
 }
 
+
+static float GThreadCallStackRequestMaxWait = 0.5f;
+static FAutoConsoleVariableRef CVarAndroidPlatformThreadCallStackRequestMaxWait(
+	TEXT("AndroidPlatformThreadStackWalk.RequestMaxWait"),
+	GThreadCallStackRequestMaxWait,
+	TEXT("The number of seconds to spin before an individual back trace has timed out."));
+
+static float GThreadCallStackMaxWait = 5.0f;
 static TAutoConsoleVariable<float> CVarAndroidPlatformThreadCallStackMaxWait(
 	TEXT("AndroidPlatformThreadStackWalk.MaxWait"),
-	5.0f,
-	TEXT("The number of seconds allowed to spin before killing the process, with the assumption the signal handler has hung."));
+	GThreadCallStackMaxWait,
+	TEXT("The number of seconds allowed to spin before killing the process, with the assumption the back trace handler has hung."));
 
 #if ANDROID_HAS_THREADBACKTRACE
+
+int32 ThreadStackBackTraceStatus = 0;
+static const int32 ThreadStackBackTraceCurrentStatus_RUNNING = -2;
+static const int32 ThreadStackBackTraceCurrentStatus_DONE = -3;
+
+static ThreadStackUserData SignalThreadStackUserData;
+
+// the callback when THREAD_CALLSTACK_GENERATOR is being processed.
+void FAndroidPlatformStackWalk::HandleBackTraceSignal(siginfo* Info, void* Context)
+{
+	if (FPlatformAtomics::InterlockedCompareExchange(&ThreadStackBackTraceStatus, ThreadStackBackTraceCurrentStatus_RUNNING, Info->si_value.sival_int) == Info->si_value.sival_int)
+	{
+		SignalThreadStackUserData.BackTraceCount = FPlatformStackWalk::CaptureStackBackTrace(SignalThreadStackUserData.BackTrace, SignalThreadStackUserData.CallStackSize, Context);
+		FPlatformAtomics::AtomicStore(&ThreadStackBackTraceStatus, ThreadStackBackTraceCurrentStatus_DONE);
+	}
+}
+
+// Sends a signal to ThreadId, wait AndroidPlatformThreadStackWalk.RequestMaxWait seconds for result or time out and return 0.
+// if callstack capture begins, but takes > AndroidPlatformThreadStackWalk.MaxWait the process will be killed.
+// Is not thread safe, returns 0 if a CaptureThreadStackBackTrace is occurring on another thread.
 uint32 FAndroidPlatformStackWalk::CaptureThreadStackBackTrace(uint64 ThreadId, uint64* BackTrace, uint32 MaxDepth)
 {
-	auto GatherCallstackFromThread = [](ThreadStackUserData& ThreadStack, uint64 TargetThreadId)
+	static TAtomic<bool> bHasReentered(false);
+	bool bExpected = false;
+	if (!bHasReentered.CompareExchange(bExpected, true))
+	{
+		return 0;
+	}
+	ON_SCOPE_EXIT
+	{
+		bHasReentered = false;
+	};
+
+	auto GatherCallstackFromThread = [](uint64 TargetThreadId)
 	{
-		auto WaitForSignalHandlerToFinishOrCrash = [](ThreadStackUserData& WaitThreadStack)
+		static int32 ThreadStackBackTraceNextRequest = 0;
+		int32 CurrentThreadStackBackTrace = ThreadStackBackTraceNextRequest++;
+
+		auto WaitForSignalHandlerToFinishOrCrash = [CurrentThreadStackBackTrace]()
 		{
-			float EndWaitTimestamp = FPlatformTime::Seconds() + CVarAndroidPlatformThreadCallStackMaxWait.AsVariable()->GetFloat();
-			float CurrentTimestamp = FPlatformTime::Seconds();
+			const float PollTime = 0.001f;
 
-			while (!WaitThreadStack.bDone)
+			for (float CurrentTime = 0; CurrentTime <= GThreadCallStackMaxWait; CurrentTime += PollTime)
 			{
-				if (CurrentTimestamp > EndWaitTimestamp)
+				if (FPlatformAtomics::InterlockedCompareExchange(&ThreadStackBackTraceStatus, ThreadStackBackTraceNextRequest, ThreadStackBackTraceCurrentStatus_DONE) == ThreadStackBackTraceCurrentStatus_DONE)
 				{
-					// We have waited for as long as we should for the signal handler to finish. Assume it has hang and we need to kill our selfs
-					*(int*)0x10 = 0x0;
+					// success 
+					return SignalThreadStackUserData.BackTraceCount;
 				}
 
-				CurrentTimestamp = FPlatformTime::Seconds();
+				// signal timed out
+				if (CurrentTime > GThreadCallStackRequestMaxWait && FPlatformAtomics::InterlockedCompareExchange(&ThreadStackBackTraceStatus, ThreadStackBackTraceNextRequest, CurrentThreadStackBackTrace) == CurrentThreadStackBackTrace)
+				{
+					// request not yet started, skip it.
+					return 0;
+				}
+				FPlatformProcess::Sleep(PollTime);
 			}
+
+			// We have waited for as long as we should for the signal handler to finish. Assume it has hang and we need to kill our selfs
+			*(int*)0x10 = 0x0;
+			return 0;
 		};
 
 		sigval UserData;
-		UserData.sival_ptr = &ThreadStack;
+		UserData.sival_int = CurrentThreadStackBackTrace;
 
 		siginfo_t info;
 		memset(&info, 0, sizeof(siginfo_t));
@@ -300,20 +353,16 @@ uint32 FAndroidPlatformStackWalk::CaptureThreadStackBackTrace(uint64 ThreadId, u
 		// sigqueue will try a different thread signal handler and report the wrong callstack
 		if (syscall(SYS_rt_tgsigqueueinfo, info.si_pid, TargetThreadId, THREAD_CALLSTACK_GENERATOR, &info) == 0)
 		{
-			WaitForSignalHandlerToFinishOrCrash(ThreadStack);
+			return WaitForSignalHandlerToFinishOrCrash();
 		}
+		return 0;
 	};
 
-	ThreadStackUserData ThreadBackTrace;
-	ThreadBackTrace.CallStackSize = MaxDepth;
-	ThreadBackTrace.BackTrace = BackTrace;
-	ThreadBackTrace.BackTraceCount = 0;
-	ThreadBackTrace.bDone = false;
-
-	GatherCallstackFromThread(ThreadBackTrace, ThreadId);
+	SignalThreadStackUserData.CallStackSize = MaxDepth;
+	SignalThreadStackUserData.BackTrace = BackTrace;
+	SignalThreadStackUserData.BackTraceCount = 0;
 
-	// The signal handler will set this value, we just have to make sure we wait for the signal handler we raised to finish
-	return ThreadBackTrace.BackTraceCount;
+	return GatherCallstackFromThread(ThreadId);
 }
 #else
 uint32 FAndroidPlatformStackWalk::CaptureThreadStackBackTrace(uint64 ThreadId, uint64* BackTrace, uint32 MaxDepth)
diff --git a/Engine/Source/Runtime/Core/Public/Android/AndroidPlatformStackWalk.h b/Engine/Source/Runtime/Core/Public/Android/AndroidPlatformStackWalk.h
--- a/Engine/Source/Runtime/Core/Public/Android/AndroidPlatformStackWalk.h
+++ b/Engine/Source/Runtime/Core/Public/Android/AndroidPlatformStackWalk.h
@@ -20,6 +20,8 @@ struct CORE_API FAndroidPlatformStackWalk : public FGenericPlatformStackWalk
 	static bool SymbolInfoToHumanReadableString(const FProgramCounterSymbolInfo& SymbolInfo, ANSICHAR* HumanReadableString, SIZE_T HumanReadableStringSize);
 
 	static uint32 CaptureThreadStackBackTrace(uint64 ThreadId, uint64* BackTrace, uint32 MaxDepth);
+
+	static void HandleBackTraceSignal(siginfo* Info, void* Context);
 };
 
 #define ANDROID_HAS_THREADBACKTRACE !PLATFORM_LUMIN && PLATFORM_USED_NDK_VERSION_INTEGER >= 21
@@ -32,7 +34,6 @@ struct ThreadStackUserData
 	uint64* BackTrace;
 	int32 BackTraceCount;
 	SIZE_T CallStackSize;
-	TAtomic<bool> bDone;
 };
 #define THREAD_CALLSTACK_GENERATOR SIGRTMIN + 5
 #endif // ANDROID_HAS_THREADBACKTRACE
