Copy over header guard from Dev-Editor CL# 10310169
#rnx

diff --git a/Engine/Source/Runtime/Core/Public/Math/sse_mathfun.h b/Engine/Source/Runtime/Core/Public/Math/sse_mathfun.h
--- a/Engine/Source/Runtime/Core/Public/Math/sse_mathfun.h
+++ b/Engine/Source/Runtime/Core/Public/Math/sse_mathfun.h
@@ -40,6 +40,7 @@
   2. It looks like some linked object files have conflicting names (log_ps, exp_ps, etc). So adding the 
      prefix SseMath_ to all the functions.
 */
+#pragma once
 
 #include <xmmintrin.h>
 
