Remove unnecessary forward declare causing include tool warning
#rnx

diff --git a/Engine/Source/Runtime/Core/Public/Math/UnrealMathSSE.h b/Engine/Source/Runtime/Core/Public/Math/UnrealMathSSE.h
--- a/Engine/Source/Runtime/Core/Public/Math/UnrealMathSSE.h
+++ b/Engine/Source/Runtime/Core/Public/Math/UnrealMathSSE.h
@@ -2,8 +2,6 @@
 
 #pragma once
 
-struct FMath;
-
 #if defined(__cplusplus_cli) && !PLATFORM_HOLOLENS
 // there are compile issues with this file in managed mode, so use the FPU version
 #include "Math/UnrealMathFPU.h"
