NumUsed takes into consideration free blocks, while TotalNumUsed does not. So we were *saying* we were out of blocks free while we still had blocks free

diff --git a/Engine/Source/Runtime/Core/Private/GenericPlatform/GenericPlatformMallocCrash.cpp b/Engine/Source/Runtime/Core/Private/GenericPlatform/GenericPlatformMallocCrash.cpp
--- a/Engine/Source/Runtime/Core/Private/GenericPlatform/GenericPlatformMallocCrash.cpp
+++ b/Engine/Source/Runtime/Core/Private/GenericPlatform/GenericPlatformMallocCrash.cpp
@@ -484,7 +484,7 @@ FMallocCrashPool* FGenericPlatformMallocCrash::FindPoolFromSize( uint32 Allocati
 		FMallocCrashPool* Pool = &FMallocCrashPool::GetPool(Index);
 
 		// If we have used up this pools allocation limit lets try the next one
-		if ( Pool->TotalNumUsed >= Pool->MaxNumAllocations )
+		if ( Pool->NumUsed >= Pool->MaxNumAllocations )
 		{
 			continue;
 		}
