Fix free memory not getting adjusted for extra dev memory. Turns out we report free memory in two places if we're not a server. This change fixes up the other spot (we should fix this properly at some point, but I don't want to break the server codepath).
Also remove  the cvar csv.AdjustPhysicalFreeForExtraDevMemory. We always do this now.

diff --git a/Engine/Source/Runtime/Core/Private/ProfilingDebugging/CsvProfiler.cpp b/Engine/Source/Runtime/Core/Private/ProfilingDebugging/CsvProfiler.cpp
--- a/Engine/Source/Runtime/Core/Private/ProfilingDebugging/CsvProfiler.cpp
+++ b/Engine/Source/Runtime/Core/Private/ProfilingDebugging/CsvProfiler.cpp
@@ -89,16 +89,6 @@ TAutoConsoleVariable<int32> CVarCsvShippingContinuousWrites(
 );
 #endif
 
-#if !UE_BUILD_SHIPPING
-TAutoConsoleVariable<int32> CVarCsvAdjustPhysicalFreeForExtraDevMemory(
-	TEXT("csv.AdjustPhysicalFreeForExtraDevMemory"),
-	1,
-	TEXT("If 1, extra development memory is subtracted from the physicalFreeMB stat before writing.\r\n")
-	TEXT("This can result in negative values in cases which would have crashed OOM"),
-	ECVF_Default
-);
-#endif
-
 TAutoConsoleVariable<int32> CVarCsvCompressionMode(
 	TEXT("csv.CompressionMode"),
 	-1,
@@ -2656,10 +2646,8 @@ void FCsvProfiler::EndFrame()
 		float PhysicalMBFree = float(MemoryStats.AvailablePhysical) / (1024.0f * 1024.0f);
 
 #if !UE_BUILD_SHIPPING
-		if (CVarCsvAdjustPhysicalFreeForExtraDevMemory.GetValueOnGameThread())
-		{
-			PhysicalMBFree -= float(FPlatformMemory::GetExtraDevelopmentMemorySize() / 1024ull / 1024ull);
-		}
+		// Subtract any extra development memory from physical free. This can result in negative values in cases where we would have crashed OOM
+		PhysicalMBFree -= float(FPlatformMemory::GetExtraDevelopmentMemorySize() / 1024ull / 1024ull);
 #endif
 		float PhysicalMBUsed = float(MemoryStats.UsedPhysical) / (1024.0f * 1024.0f);
 		float VirtualMBUsed  = float(MemoryStats.UsedVirtual) / (1024.0f * 1024.0f);
diff --git a/Engine/Source/Runtime/Launch/Private/LaunchEngineLoop.cpp b/Engine/Source/Runtime/Launch/Private/LaunchEngineLoop.cpp
--- a/Engine/Source/Runtime/Launch/Private/LaunchEngineLoop.cpp
+++ b/Engine/Source/Runtime/Launch/Private/LaunchEngineLoop.cpp
@@ -1237,6 +1237,10 @@ static void UpdateCoreCsvStats_EndFrame()
 	}
 	FPlatformMemoryStats MemoryStats = FPlatformMemory::GetStats();
 	float PhysicalMBFree = float(MemoryStats.AvailablePhysical / 1024) / 1024.0f;
+#if !UE_BUILD_SHIPPING
+	// Subtract any extra development memory from physical free. This can result in negative values in cases where we would have crashed OOM
+	PhysicalMBFree -= float(FPlatformMemory::GetExtraDevelopmentMemorySize() / 1024ull / 1024ull);
+#endif
 	CSV_CUSTOM_STAT_GLOBAL(MemoryFreeMB, PhysicalMBFree, ECsvCustomStatOp::Set);
 }
 #endif // WITH_ENGINE && CSV_PROFILER
