Add a way for a title to register an area of extra memory

diff --git a/Engine/Source/Runtime/Core/Private/HAL/UnrealMemory.cpp b/Engine/Source/Runtime/Core/Private/HAL/UnrealMemory.cpp
--- a/Engine/Source/Runtime/Core/Private/HAL/UnrealMemory.cpp
+++ b/Engine/Source/Runtime/Core/Private/HAL/UnrealMemory.cpp
@@ -616,6 +616,73 @@ void FUseSystemMallocForNew::operator delete[](void* Ptr)
 	FMemory::SystemFree(Ptr);
 }
 
+static bool GPersistentAuxiliaryEnabled = true;
+static uint8 * GPersistentAuxiliary = nullptr;
+static uint8 * GPersistentAuxiliaryEnd = nullptr;
+static TAtomic<SIZE_T> GPersistentAuxiliaryCurrentOffset;
+static SIZE_T GPersistentAuxiliarySize = 0;
+
+
+void FMemory::RegisterPersistentAuxiliary(void* InMemory, SIZE_T InSize)
+{
+	check(GPersistentAuxiliary == nullptr);
+	GPersistentAuxiliaryCurrentOffset = 0;
+	GPersistentAuxiliarySize = InSize;
+	GPersistentAuxiliary = (uint8 *)InMemory;
+	GPersistentAuxiliaryEnd = GPersistentAuxiliary + InSize;
+}
+void* FMemory::MallocPersistentAuxiliary(SIZE_T InSize, SIZE_T InAlignment)
+{
+	if (GPersistentAuxiliary != nullptr && GPersistentAuxiliaryEnabled)
+	{
+		SIZE_T Alignment = FMath::Max(InAlignment, (SIZE_T)16);
+		Alignment = FMath::Max(Alignment, (SIZE_T)16);
+		SIZE_T AlignedSize = Align(InSize, Alignment);
+		// 1st check if there is room, this is atomic but could still fail when actually incrementing the offset.
+		if (GPersistentAuxiliaryCurrentOffset + AlignedSize <= GPersistentAuxiliarySize)
+		{
+			SIZE_T OldOffset = GPersistentAuxiliaryCurrentOffset.AddExchange(AlignedSize);
+			if (OldOffset + AlignedSize <= GPersistentAuxiliarySize)
+			{
+				// we were able to increment the offset and it's still within the bounds of the aux memory.
+				return &GPersistentAuxiliary[OldOffset];
+			}
+			// we've gone over the end of the aux memory, this could waste some space, if it's a problem protect with a critical section.
+		}
+	}
+	return FMemory::Malloc(InSize, InAlignment);
+}
+void FMemory::FreePersistentAuxiliary(void* InPtr)
+{
+	if (GPersistentAuxiliary != nullptr)
+	{
+		uint8* Ptr = (uint8*)InPtr;
+		if (Ptr >= GPersistentAuxiliary && Ptr < GPersistentAuxiliaryEnd)
+		{
+			// it is part of the GPersistentAuxiliary
+			return;
+		}
+	}
+	return FMemory::Free(InPtr);
+}
+bool FMemory::IsPersistentAuxiliaryActive()
+{
+	return GPersistentAuxiliary != nullptr && GPersistentAuxiliaryEnabled;
+}
+void FMemory::DisablePersistentAuxiliary()
+{
+	GPersistentAuxiliaryEnabled = false;
+}
+void FMemory::EnablePersistentAuxiliary()
+{
+	GPersistentAuxiliaryEnabled = true;
+}
+SIZE_T FMemory::GetUsedPersistentAuxiliary()
+{
+	return GPersistentAuxiliaryCurrentOffset;
+}
+
+
 #if !INLINE_FMEMORY_OPERATION && !PLATFORM_USES_FIXED_GMalloc_CLASS
 #include "HAL/FMemory.inl"
 #endif
diff --git a/Engine/Source/Runtime/Core/Public/HAL/UnrealMemory.h b/Engine/Source/Runtime/Core/Public/HAL/UnrealMemory.h
--- a/Engine/Source/Runtime/Core/Public/HAL/UnrealMemory.h
+++ b/Engine/Source/Runtime/Core/Public/HAL/UnrealMemory.h
@@ -208,12 +208,23 @@ struct CORE_API FMemory
 	* This uses the purgatory malloc proxy to check if things are writing to stale pointers.
 	*/
 	static void EnablePoisonTests();
-
 	/**
 	* Set global allocator instead of creating it lazily on first allocation.
 	* Must only be called once and only if lazy init is disabled via a macro.
 	*/
 	static void ExplicitInit(FMalloc& Allocator);
+
+	/**
+	* Functions to handle special memory given to the title from the platform
+	* This memory is allocated like a stack, it's never really freed
+	*/
+	static void RegisterPersistentAuxiliary(void* InMemory, SIZE_T InSize);
+	static void* MallocPersistentAuxiliary(SIZE_T InSize, SIZE_T InAlignment = 0);
+	static void FreePersistentAuxiliary(void* InPtr);
+	static bool IsPersistentAuxiliaryActive();
+	static void DisablePersistentAuxiliary();
+	static void EnablePersistentAuxiliary();
+	static SIZE_T GetUsedPersistentAuxiliary();
 private:
 	static void GCreateMalloc();
 	// These versions are called either at startup or in the event of a crash
