Take advantage of extra memory if a platform supports it

diff --git a/Engine/Source/Runtime/Core/Private/UObject/UnrealNames.cpp b/Engine/Source/Runtime/Core/Private/UObject/UnrealNames.cpp
--- a/Engine/Source/Runtime/Core/Private/UObject/UnrealNames.cpp
+++ b/Engine/Source/Runtime/Core/Private/UObject/UnrealNames.cpp
@@ -265,7 +265,7 @@ public:
 	FNameEntryAllocator()
 	{
 		LLM_SCOPE(ELLMTag::FName);
-		Blocks[0] = (uint8*)FMemory::Malloc(BlockSizeBytes, FPlatformMemory::GetConstants().PageSize);
+		Blocks[0] = (uint8*)FMemory::MallocPersistentAuxiliary(BlockSizeBytes, FPlatformMemory::GetConstants().PageSize);
 	}
 
 	~FNameEntryAllocator()
@@ -393,7 +393,7 @@ private:
 		check(CurrentBlock < FNameMaxBlocks);
 		check(Blocks[CurrentBlock] == nullptr);
 
-		Blocks[CurrentBlock] = (uint8*)FMemory::Malloc(BlockSizeBytes, FPlatformMemory::GetConstants().PageSize);
+		Blocks[CurrentBlock] = (uint8*)FMemory::MallocPersistentAuxiliary(BlockSizeBytes, FPlatformMemory::GetConstants().PageSize);
 	}
 
 	mutable FRWLock Lock;
diff --git a/Engine/Source/Runtime/Core/Public/Async/TaskGraphInterfaces.h b/Engine/Source/Runtime/Core/Public/Async/TaskGraphInterfaces.h
--- a/Engine/Source/Runtime/Core/Public/Async/TaskGraphInterfaces.h
+++ b/Engine/Source/Runtime/Core/Public/Async/TaskGraphInterfaces.h
@@ -422,7 +422,7 @@ public:
 		/** Total size in bytes for a small task that will use the custom allocator **/
 		SMALL_TASK_SIZE = 256
 	};
-	typedef TLockFreeFixedSizeAllocator_TLSCache<SMALL_TASK_SIZE, PLATFORM_CACHE_LINE_SIZE> TSmallTaskAllocator;
+	typedef TLockFreeFixedSizeAllocator_TLSCache<SMALL_TASK_SIZE, PLATFORM_CACHE_LINE_SIZE, FNoopCounter, true> TSmallTaskAllocator;
 protected:
 	/** 
 	 *	Constructor
diff --git a/Engine/Source/Runtime/Core/Public/Containers/LockFreeFixedSizeAllocator.h b/Engine/Source/Runtime/Core/Public/Containers/LockFreeFixedSizeAllocator.h
--- a/Engine/Source/Runtime/Core/Public/Containers/LockFreeFixedSizeAllocator.h
+++ b/Engine/Source/Runtime/Core/Public/Containers/LockFreeFixedSizeAllocator.h
@@ -17,7 +17,7 @@
 #define USE_NAIVE_TLockFreeFixedSizeAllocator_TLSCacheBase (0) // this is useful for find who really leaked
 #endif
 
-template<int32 SIZE, typename TBundleRecycler, typename TTrackingCounter = FNoopCounter>
+template<int32 SIZE, typename TBundleRecycler, typename TTrackingCounter = FNoopCounter, bool AllowDisablingOfTrim = false>
 class TLockFreeFixedSizeAllocator_TLSCacheBase : public FNoncopyable
 {
 	enum
@@ -65,7 +65,7 @@ public:
 				TLS.PartialBundle = GlobalFreeListBundles.Pop();
 				if (!TLS.PartialBundle)
 				{
-					TLS.PartialBundle = (void**)FMemory::Malloc(SIZE * NUM_PER_BUNDLE);
+					TLS.PartialBundle = (void**)FMemory::MallocPersistentAuxiliary(SIZE * NUM_PER_BUNDLE);
 					void **Next = TLS.PartialBundle;
 					for (int32 Index = 0; Index < NUM_PER_BUNDLE - 1; Index++)
 					{
@@ -189,7 +189,7 @@ private:
  * never returns free space until program shutdown.
  * alignment isn't handled, assumes FMemory::Malloc will work
  */
-template<int32 SIZE, int TPaddingForCacheContention, typename TTrackingCounter = FNoopCounter>
+template<int32 SIZE, int TPaddingForCacheContention, typename TTrackingCounter = FNoopCounter, bool AllowDisablingOfTrim = false>
 class TLockFreeFixedSizeAllocator
 {
 public:
@@ -217,7 +217,7 @@ public:
 		}
 		else
 		{
-			Memory = FMemory::Malloc(SIZE);
+			Memory = FMemory::MallocPersistentAuxiliary(SIZE);
 		}
 		return Memory;
 	}
@@ -240,9 +240,16 @@ public:
 	*/
 	void Trim()
 	{
+		if (AllowDisablingOfTrim)
+		{
+			if (FMemory::IsPersistentAuxiliaryActive())
+			{
+				return;
+			}
+		}
 		while (void* Mem = FreeList.Pop())
 		{
-			FMemory::Free(Mem);
+			FMemory::FreePersistentAuxiliary(Mem);
 			NumFree.Decrement();
 		}
 	}
@@ -286,8 +293,8 @@ private:
  * never returns free space, even at shutdown
  * alignment isn't handled, assumes FMemory::Malloc will work
  */
-template<int32 SIZE, int TPaddingForCacheContention, typename TTrackingCounter = FNoopCounter>
-class TLockFreeFixedSizeAllocator_TLSCache : public TLockFreeFixedSizeAllocator_TLSCacheBase<SIZE, TLockFreePointerListUnordered<void*, TPaddingForCacheContention>, TTrackingCounter>
+template<int32 SIZE, int TPaddingForCacheContention, typename TTrackingCounter = FNoopCounter, bool AllowDisablingOfTrim = false>
+class TLockFreeFixedSizeAllocator_TLSCache : public TLockFreeFixedSizeAllocator_TLSCacheBase<SIZE, TLockFreePointerListUnordered<void*, TPaddingForCacheContention>, TTrackingCounter, AllowDisablingOfTrim>
 {
 };
 
@@ -296,8 +303,8 @@ class TLockFreeFixedSizeAllocator_TLSCache : public TLockFreeFixedSizeAllocator_
  *
  * Never returns free space until program shutdown.
  */
-template<class T, int TPaddingForCacheContention>
-class TLockFreeClassAllocator : private TLockFreeFixedSizeAllocator<sizeof(T), TPaddingForCacheContention>
+template<class T, int TPaddingForCacheContention, bool AllowDisablingOfTrim = false>
+class TLockFreeClassAllocator : private TLockFreeFixedSizeAllocator<sizeof(T), TPaddingForCacheContention, FNoopCounter, AllowDisablingOfTrim>
 {
 public:
 	/**
@@ -340,8 +347,8 @@ public:
  *
  * Never returns free space until program shutdown.
  */
-template<class T, int TPaddingForCacheContention>
-class TLockFreeClassAllocator_TLSCache : private TLockFreeFixedSizeAllocator_TLSCache<sizeof(T), TPaddingForCacheContention>
+template<class T, int TPaddingForCacheContention, bool AllowDisablingOfTrim = false>
+class TLockFreeClassAllocator_TLSCache : private TLockFreeFixedSizeAllocator_TLSCache<sizeof(T), TPaddingForCacheContention, FNoopCounter , AllowDisablingOfTrim>
 {
 public:
 	/**
diff --git a/Engine/Source/Runtime/CoreUObject/Private/UObject/UObjectAllocator.cpp b/Engine/Source/Runtime/CoreUObject/Private/UObject/UObjectAllocator.cpp
--- a/Engine/Source/Runtime/CoreUObject/Private/UObject/UObjectAllocator.cpp
+++ b/Engine/Source/Runtime/CoreUObject/Private/UObject/UObjectAllocator.cpp
@@ -20,7 +20,7 @@ COREUOBJECT_API FUObjectAllocator GUObjectAllocator;
 void FUObjectAllocator::AllocatePermanentObjectPool(int32 InPermanentObjectPoolSize)
 {
 	PermanentObjectPoolSize	= InPermanentObjectPoolSize;
-	PermanentObjectPool		= (uint8*) FMemory::Malloc( PermanentObjectPoolSize );
+	PermanentObjectPool		= (uint8*) FMemory::MallocPersistentAuxiliary( PermanentObjectPoolSize );
 	PermanentObjectPoolTail	= PermanentObjectPool;
 	PermanentObjectPoolExceededTail = PermanentObjectPoolTail;
 }
