Use sched api to determine number of cpu cores on Android.
 (fixes incorrect core count from AndroidMisc::NumberOfCores)
#rnx

diff --git a/Engine/Source/Runtime/Core/Private/Android/AndroidPlatformMisc.cpp b/Engine/Source/Runtime/Core/Private/Android/AndroidPlatformMisc.cpp
--- a/Engine/Source/Runtime/Core/Private/Android/AndroidPlatformMisc.cpp
+++ b/Engine/Source/Runtime/Core/Private/Android/AndroidPlatformMisc.cpp
@@ -726,31 +726,13 @@ int32 FAndroidMisc::NumberOfCores()
 	int32 NumberOfCores = android_getCpuCount();
 
 	static int CalculatedNumberOfCores = 0;
-
-#ifndef CPU_SETSIZE
-#if PLATFORM_64BITS
-	#define CPU_SETSIZE 1024
-#else
-	#define CPU_SETSIZE 32
-#endif 
-#endif
-
-	char cpuset[CPU_SETSIZE / 8];
-
 	if (CalculatedNumberOfCores == 0)
 	{
 		pid_t ThreadId = gettid();
-		syscall(__NR_sched_getaffinity, ThreadId, sizeof(cpuset), &cpuset);
-
-		char *coreptr = cpuset;
-		int32 CoreSets = CPU_SETSIZE / 8;;
-		while (CoreSets--)
+		cpu_set_t cpuset;
+		if (sched_getaffinity(ThreadId, sizeof(cpuset), &cpuset) != -1)
 		{
-			char coremask = *coreptr++;
-			for (int i = 0; i < 8; i++)
-			{
-				CalculatedNumberOfCores += ((coremask & (1 << i)) != 0);
-			}
+			CalculatedNumberOfCores = CPU_COUNT(&cpuset);
 		}
 
 		UE_LOG(LogTemp, Log, TEXT("%d cores and %d assignable cores"), NumberOfCores, CalculatedNumberOfCores);
