Fix per thread stack walking code, failed signal send will also update the last processed thread request.
#rnx

diff --git a/Engine/Source/Runtime/Core/Private/Android/AndroidPlatformStackWalk.cpp b/Engine/Source/Runtime/Core/Private/Android/AndroidPlatformStackWalk.cpp
--- a/Engine/Source/Runtime/Core/Private/Android/AndroidPlatformStackWalk.cpp
+++ b/Engine/Source/Runtime/Core/Private/Android/AndroidPlatformStackWalk.cpp
@@ -355,6 +355,12 @@ uint32 FAndroidPlatformStackWalk::CaptureThreadStackBackTrace(uint64 ThreadId, u
 		{
 			return WaitForSignalHandlerToFinishOrCrash();
 		}
+		else
+		{
+			// we failed to send the signal, update the current status.
+			FPlatformAtomics::AtomicStore(&ThreadStackBackTraceStatus, ThreadStackBackTraceNextRequest);
+		}
+
 		return 0;
 	};
 
