AndroidMisc::NumberOfCores - Clear cpu set variable before use.
#rnx

diff --git a/Engine/Source/Runtime/Core/Private/Android/AndroidPlatformMisc.cpp b/Engine/Source/Runtime/Core/Private/Android/AndroidPlatformMisc.cpp
--- a/Engine/Source/Runtime/Core/Private/Android/AndroidPlatformMisc.cpp
+++ b/Engine/Source/Runtime/Core/Private/Android/AndroidPlatformMisc.cpp
@@ -730,6 +730,7 @@ int32 FAndroidMisc::NumberOfCores()
 	{
 		pid_t ThreadId = gettid();
 		cpu_set_t cpuset;
+		CPU_ZERO(&cpuset);
 		if (sched_getaffinity(ThreadId, sizeof(cpuset), &cpuset) != -1)
 		{
 			CalculatedNumberOfCores = CPU_COUNT(&cpuset);
