More detailed chaos LLM numbers

diff --git a/Engine/Source/Runtime/Core/Private/HAL/LowLevelMemTracker.cpp b/Engine/Source/Runtime/Core/Private/HAL/LowLevelMemTracker.cpp
--- a/Engine/Source/Runtime/Core/Private/HAL/LowLevelMemTracker.cpp
+++ b/Engine/Source/Runtime/Core/Private/HAL/LowLevelMemTracker.cpp
@@ -88,6 +88,7 @@ DECLARE_LLM_MEMORY_STAT(TEXT("UI"), STAT_UILLM, STATGROUP_LLMFULL);
 DECLARE_LLM_MEMORY_STAT(TEXT("Physics"), STAT_PhysicsLLM, STATGROUP_LLMFULL);
 DECLARE_LLM_MEMORY_STAT(TEXT("PhysX"), STAT_PhysXLLM, STATGROUP_LLMFULL);
 DECLARE_LLM_MEMORY_STAT(TEXT("Chaos"), STAT_ChaosLLM, STATGROUP_LLMFULL);
+DECLARE_LLM_MEMORY_STAT(TEXT("ChaosGeometry"), STAT_ChaosGeometryLLM, STATGROUP_LLMFULL);
 DECLARE_LLM_MEMORY_STAT(TEXT("EnginePreInit"), STAT_EnginePreInitLLM, STATGROUP_LLMFULL);
 DECLARE_LLM_MEMORY_STAT(TEXT("EngineInit"), STAT_EngineInitLLM, STATGROUP_LLMFULL);
 DECLARE_LLM_MEMORY_STAT(TEXT("Rendering Thread"), STAT_RenderingThreadLLM, STATGROUP_LLMFULL);
diff --git a/Engine/Source/Runtime/Core/Public/HAL/LowLevelMemTracker.h b/Engine/Source/Runtime/Core/Public/HAL/LowLevelMemTracker.h
--- a/Engine/Source/Runtime/Core/Public/HAL/LowLevelMemTracker.h
+++ b/Engine/Source/Runtime/Core/Public/HAL/LowLevelMemTracker.h
@@ -187,6 +187,7 @@ enum class ELLMTagSet : uint8
 	macro(Physics,								"Physics",						GET_STATFNAME(STAT_PhysicsLLM),								GET_STATFNAME(STAT_PhysicsSummaryLLM),			-1)\
 	macro(PhysX,								"PhysX",						GET_STATFNAME(STAT_PhysXLLM),								GET_STATFNAME(STAT_PhysXSummaryLLM),			ELLMTag::Physics)\
 	macro(Chaos,								"Chaos",						GET_STATFNAME(STAT_ChaosLLM),								GET_STATFNAME(STAT_ChaosSummaryLLM),			ELLMTag::Physics)\
+	macro(ChaosGeometry,						"ChaosGeometry",				GET_STATFNAME(STAT_ChaosGeometryLLM),						GET_STATFNAME(STAT_ChaosSummaryLLM),			ELLMTag::Physics)\
 	macro(EnginePreInitMemory,					"EnginePreInit",				GET_STATFNAME(STAT_EnginePreInitLLM),						GET_STATFNAME(STAT_EngineSummaryLLM),			-1)\
 	macro(EngineInitMemory,						"EngineInit",					GET_STATFNAME(STAT_EngineInitLLM),							GET_STATFNAME(STAT_EngineSummaryLLM),			-1)\
 	macro(RenderingThreadMemory,				"RenderingThread",				GET_STATFNAME(STAT_RenderingThreadLLM),						GET_STATFNAME(STAT_EngineSummaryLLM),			-1)\
diff --git a/Engine/Source/Runtime/Engine/Private/Components/ActorComponent.cpp b/Engine/Source/Runtime/Engine/Private/Components/ActorComponent.cpp
--- a/Engine/Source/Runtime/Engine/Private/Components/ActorComponent.cpp
+++ b/Engine/Source/Runtime/Engine/Private/Components/ActorComponent.cpp
@@ -1324,7 +1324,12 @@ void UActorComponent::OnDestroyPhysicsState()
 
 void UActorComponent::CreatePhysicsState()
 {
+#if WITH_CHAOS
+	LLM_SCOPE(ELLMTag::Chaos);
+#else
 	LLM_SCOPE(ELLMTag::PhysX);
+#endif
+
 	SCOPE_CYCLE_COUNTER(STAT_ComponentCreatePhysicsState);
 
 	if (!bPhysicsStateCreated && WorldPrivate->GetPhysicsScene() && ShouldCreatePhysicsState())
diff --git a/Engine/Source/Runtime/Engine/Private/PhysicsEngine/BodyInstance.cpp b/Engine/Source/Runtime/Engine/Private/PhysicsEngine/BodyInstance.cpp
--- a/Engine/Source/Runtime/Engine/Private/PhysicsEngine/BodyInstance.cpp
+++ b/Engine/Source/Runtime/Engine/Private/PhysicsEngine/BodyInstance.cpp
@@ -1171,7 +1171,11 @@ struct FInitBodiesHelper
 
 	void InitBodies()
 	{
+#if WITH_CHAOS
+		LLM_SCOPE(ELLMTag::Chaos);
+#else
 		LLM_SCOPE(ELLMTag::PhysX);
+#endif
 
 		check(IsInGameThread());
 
diff --git a/Engine/Source/Runtime/Engine/Private/PhysicsEngine/ConstraintInstance.cpp b/Engine/Source/Runtime/Engine/Private/PhysicsEngine/ConstraintInstance.cpp
--- a/Engine/Source/Runtime/Engine/Private/PhysicsEngine/ConstraintInstance.cpp
+++ b/Engine/Source/Runtime/Engine/Private/PhysicsEngine/ConstraintInstance.cpp
@@ -348,7 +348,11 @@ bool GetActorRefs(FBodyInstance* Body1, FBodyInstance* Body2, FPhysicsActorHandl
 
 bool FConstraintInstance::CreateJoint_AssumesLocked(const FPhysicsActorHandle& InActorRef1, const FPhysicsActorHandle& InActorRef2)
 {
+#if WITH_CHAOS
+	LLM_SCOPE(ELLMTag::Chaos);
+#else
 	LLM_SCOPE(ELLMTag::PhysX);
+#endif
 
 	FTransform Local1 = GetRefFrame(EConstraintFrame::Frame1);
 	if(FPhysicsInterface::IsValid(InActorRef1))
diff --git a/Engine/Source/Runtime/Engine/Private/PhysicsEngine/Experimental/ChaosDerivedDataReader.cpp b/Engine/Source/Runtime/Engine/Private/PhysicsEngine/Experimental/ChaosDerivedDataReader.cpp
--- a/Engine/Source/Runtime/Engine/Private/PhysicsEngine/Experimental/ChaosDerivedDataReader.cpp
+++ b/Engine/Source/Runtime/Engine/Private/PhysicsEngine/Experimental/ChaosDerivedDataReader.cpp
@@ -26,6 +26,7 @@ FChaosDerivedDataReader<T, d>::FChaosDerivedDataReader(FUntypedBulkData* InBulkD
 	}
 	else
 	{
+		LLM_SCOPE(ELLMTag::ChaosGeometry);
 		ChaosAr << ConvexImplicitObjects << TrimeshImplicitObjects << UVInfo << FaceRemap;
 
 		bReadSuccessful = true;
diff --git a/Engine/Source/Runtime/Engine/Private/PhysicsEngine/Experimental/ChaosInterfaceUtils.cpp b/Engine/Source/Runtime/Engine/Private/PhysicsEngine/Experimental/ChaosInterfaceUtils.cpp
--- a/Engine/Source/Runtime/Engine/Private/PhysicsEngine/Experimental/ChaosInterfaceUtils.cpp
+++ b/Engine/Source/Runtime/Engine/Private/PhysicsEngine/Experimental/ChaosInterfaceUtils.cpp
@@ -136,6 +136,7 @@ namespace ChaosInterface
 
 	void CreateGeometry(const FGeometryAddParams& InParams, TArray<TUniquePtr<Chaos::FImplicitObject>>& OutGeoms, Chaos::TShapesArray<float, 3>& OutShapes)
 	{
+		LLM_SCOPE(ELLMTag::ChaosGeometry);
 		const FVector& Scale = InParams.Scale;
 		TArray<TUniquePtr<Chaos::FImplicitObject>>& Geoms = OutGeoms;
 		Chaos::TShapesArray<float, 3>& Shapes = OutShapes;
diff --git a/Engine/Source/Runtime/Engine/Private/PhysicsEngine/Experimental/PhysInterface_Chaos.cpp b/Engine/Source/Runtime/Engine/Private/PhysicsEngine/Experimental/PhysInterface_Chaos.cpp
--- a/Engine/Source/Runtime/Engine/Private/PhysicsEngine/Experimental/PhysInterface_Chaos.cpp
+++ b/Engine/Source/Runtime/Engine/Private/PhysicsEngine/Experimental/PhysInterface_Chaos.cpp
@@ -1299,6 +1299,7 @@ const FBodyInstance* FPhysInterface_Chaos::ShapeToOriginalBodyInstance(const FBo
 
 void FPhysInterface_Chaos::AddGeometry(FPhysicsActorHandle& InActor, const FGeometryAddParams& InParams, TArray<FPhysicsShapeHandle>* OutOptShapes)
 {
+	LLM_SCOPE(ELLMTag::ChaosGeometry);
 	TArray<TUniquePtr<Chaos::FImplicitObject>> Geoms;
 	Chaos::TShapesArray<float, 3> Shapes;
 	ChaosInterface::CreateGeometry(InParams, Geoms, Shapes);
diff --git a/Engine/Source/Runtime/Experimental/Chaos/Private/Chaos/PBDRigidsEvolution.cpp b/Engine/Source/Runtime/Experimental/Chaos/Private/Chaos/PBDRigidsEvolution.cpp
--- a/Engine/Source/Runtime/Experimental/Chaos/Private/Chaos/PBDRigidsEvolution.cpp
+++ b/Engine/Source/Runtime/Experimental/Chaos/Private/Chaos/PBDRigidsEvolution.cpp
@@ -224,6 +224,8 @@ namespace Chaos
 	template<class FPBDRigidsEvolution, class FPBDCollisionConstraint, class T, int d>
 	void TPBDRigidsEvolutionBase<FPBDRigidsEvolution, FPBDCollisionConstraint, T, d>::FChaosAccelerationStructureTask::DoTask(ENamedThreads::Type CurrentThread, const FGraphEventRef& MyCompletionGraphEvent)
 	{
+		LLM_SCOPE(ELLMTag::Chaos);
+
 		uint8 ActiveBucketsMask = SpatialCollectionFactory.GetActiveBucketsMask();
 		TArray<TSOAView<TSpatialAccelerationCache<T, d>>> ViewsPerBucket[8];
 
diff --git a/Engine/Source/Runtime/Experimental/ChaosSolvers/Private/PBDRigidsSolver.cpp b/Engine/Source/Runtime/Experimental/ChaosSolvers/Private/PBDRigidsSolver.cpp
--- a/Engine/Source/Runtime/Experimental/ChaosSolvers/Private/PBDRigidsSolver.cpp
+++ b/Engine/Source/Runtime/Experimental/ChaosSolvers/Private/PBDRigidsSolver.cpp
@@ -50,6 +50,7 @@ namespace Chaos
 
 		void DoWork()
 		{
+			LLM_SCOPE(ELLMTag::Chaos);
 			UE_LOG(LogPBDRigidsSolverSolver, Verbose, TEXT("AdvanceOneTimeStepTask::DoWork()"));
 
 			{
diff --git a/Engine/Source/Runtime/Landscape/Private/LandscapeCollision.cpp b/Engine/Source/Runtime/Landscape/Private/LandscapeCollision.cpp
--- a/Engine/Source/Runtime/Landscape/Private/LandscapeCollision.cpp
+++ b/Engine/Source/Runtime/Landscape/Private/LandscapeCollision.cpp
@@ -553,6 +553,9 @@ void ULandscapeHeightfieldCollisionComponent::ApplyWorldOffset(const FVector& In
 
 void ULandscapeHeightfieldCollisionComponent::CreateCollisionObject()
 {
+#if WITH_CHAOS
+	LLM_SCOPE(ELLMTag::ChaosGeometry);
+#endif
 	// If we have not created a heightfield yet - do it now.
 	if (!IsValidRef(HeightfieldRef))
 	{
