fix for iPad7 device detection and profile

diff --git a/Engine/Source/Runtime/Core/Private/IOS/IOSPlatformMisc.cpp b/Engine/Source/Runtime/Core/Private/IOS/IOSPlatformMisc.cpp
--- a/Engine/Source/Runtime/Core/Private/IOS/IOSPlatformMisc.cpp
+++ b/Engine/Source/Runtime/Core/Private/IOS/IOSPlatformMisc.cpp
@@ -488,7 +488,7 @@ FIOSPlatformMisc::EIOSDevice FIOSPlatformMisc::GetIOSDeviceType()
 			{
 				DeviceType = IOS_IPad6;
 			}
-			else if (Minor == 7 || Minor == 8)
+			else if (Minor == 11 || Minor == 12)
 			{
 				DeviceType = IOS_IPad7;
 			}
