Build Fix Non-Unity

#rnx

diff --git a/Engine/Source/Runtime/Core/Private/Misc/FileHelper.cpp b/Engine/Source/Runtime/Core/Private/Misc/FileHelper.cpp
--- a/Engine/Source/Runtime/Core/Private/Misc/FileHelper.cpp
+++ b/Engine/Source/Runtime/Core/Private/Misc/FileHelper.cpp
@@ -581,7 +581,7 @@ bool FFileHelper::CreateBitmap( const TCHAR* Pattern, int32 SourceWidth, int32 S
 				{
 					const uint8 Max = 255;
 					uint8 ChannelValue = 0;
-					// When using Channel mask write the masked channel to blue channel.
+					// When using Channel mask write the masked channel to all channels (except alpha).
 					switch (ChannelMask)
 					{
 					case EChannelMask::B:
@@ -596,9 +596,6 @@ bool FFileHelper::CreateBitmap( const TCHAR* Pattern, int32 SourceWidth, int32 S
 					case EChannelMask::A:
 						ChannelValue = Data[i * SourceWidth + j].A;
 						break;
-					default:
-						check(0);
-						break;
 					}
 										
 					// replicate Channel in B, G, R
diff --git a/Engine/Source/Runtime/Core/Public/Misc/FileHelper.h b/Engine/Source/Runtime/Core/Public/Misc/FileHelper.h
--- a/Engine/Source/Runtime/Core/Public/Misc/FileHelper.h
+++ b/Engine/Source/Runtime/Core/Public/Misc/FileHelper.h
@@ -7,6 +7,7 @@
 #include "HAL/FileManager.h"
 #include "Containers/ArrayView.h"
 #include "Misc/EnumClassFlags.h"
+#include "Math/Color.h"
 
 /*-----------------------------------------------------------------------------
 	FFileHelper
