Chaos - Serialize PerShapeData bSimulate flag

diff --git a/Engine/Source/Runtime/Core/Public/UObject/ExternalPhysicsCustomObjectVersion.h b/Engine/Source/Runtime/Core/Public/UObject/ExternalPhysicsCustomObjectVersion.h
--- a/Engine/Source/Runtime/Core/Public/UObject/ExternalPhysicsCustomObjectVersion.h
+++ b/Engine/Source/Runtime/Core/Public/UObject/ExternalPhysicsCustomObjectVersion.h
@@ -86,6 +86,9 @@ struct CORE_API FExternalPhysicsCustomObjectVersion
 
 		//Replace TBox with TAABB in many places
 		TBoxReplacedWithTAABB,
+
+		// Serialize bSimulate on PerShapeData
+		SerializePerShapeDataSimulateFlag,
 		
 		// -----<new versions can be added above this line>-------------------------------------------------
 		VersionPlusOne,
diff --git a/Engine/Source/Runtime/Experimental/Chaos/Private/Chaos/GeometryParticles.cpp b/Engine/Source/Runtime/Experimental/Chaos/Private/Chaos/GeometryParticles.cpp
--- a/Engine/Source/Runtime/Experimental/Chaos/Private/Chaos/GeometryParticles.cpp
+++ b/Engine/Source/Runtime/Experimental/Chaos/Private/Chaos/GeometryParticles.cpp
@@ -117,6 +117,11 @@ namespace Chaos
 		{
 			Ar << bDisable;
 		}
+
+		if (Ar.CustomVer(FExternalPhysicsCustomObjectVersion::GUID) >= FExternalPhysicsCustomObjectVersion::SerializePerShapeDataSimulateFlag)
+		{
+			Ar << bSimulate;
+		}
 	}
 
 
