Fix non-unity
#rnx

diff --git a/Engine/Source/Runtime/Core/Public/Math/MovingWindowAverageFast.h b/Engine/Source/Runtime/Core/Public/Math/MovingWindowAverageFast.h
--- a/Engine/Source/Runtime/Core/Public/Math/MovingWindowAverageFast.h
+++ b/Engine/Source/Runtime/Core/Public/Math/MovingWindowAverageFast.h
@@ -1,6 +1,7 @@
 // Copyright 1998-2019 Epic Games, Inc. All Rights Reserved.
 #pragma once
 #include "CoreMinimal.h"
+#include "Containers/StaticArray.h"
 
 /**
  * This class calculates a moving window average. Its designed to be used with floats or doubles and
