Changed the default min scope time for framepro captures from 25us to 50us.
This is the default for framepro windows client live captures.

Significantly reduces hitching from scope "FramePro Start Frame" when recording captures.

diff --git a/Engine/Source/Runtime/Core/Private/FramePro/FrameProProfiler.cpp b/Engine/Source/Runtime/Core/Private/FramePro/FrameProProfiler.cpp
--- a/Engine/Source/Runtime/Core/Private/FramePro/FrameProProfiler.cpp
+++ b/Engine/Source/Runtime/Core/Private/FramePro/FrameProProfiler.cpp
@@ -341,12 +341,12 @@ void FFrameProProfiler::StartFrameProRecordingFromCommand(const TArray< FString
 		FilenameRoot = Args[0];
 	}
 
-	StartFrameProRecording(FilenameRoot, 25);
+	StartFrameProRecording(FilenameRoot, 50);
 }
 
 void FFrameProProfiler::StartFrameProRecordingScopeOverrideFromCommand(const TArray< FString >& Args)
 {
-	int32 MinScopeTime = 25;
+	int32 MinScopeTime = 50;
 	if (Args.Num() > 0 && Args[0].Len() > 0)
 	{
 		MinScopeTime = FCString::Atoi(*Args[0]);
