Android crash reporter
Use a parked thread to process signals, gives crash handler larger stack and removes Arm64 requirement.
#rnx

diff --git a/Engine/Source/Runtime/Core/Private/Android/AndroidPlatformMisc.cpp b/Engine/Source/Runtime/Core/Private/Android/AndroidPlatformMisc.cpp
--- a/Engine/Source/Runtime/Core/Private/Android/AndroidPlatformMisc.cpp
+++ b/Engine/Source/Runtime/Core/Private/Android/AndroidPlatformMisc.cpp
@@ -46,6 +46,7 @@
 #include "GenericPlatform/GenericPlatformOutputDevices.h"
 
 #include "Android/AndroidPlatformStackWalk.h"
+#include "Android/AndroidSignals.h"
 
 #if USE_ANDROID_JNI
 extern AAssetManager * AndroidThunkCpp_GetAssetManager();
@@ -908,158 +909,104 @@ static const char* SignalToString(int32 Signal)
 	}
 }
 
-PRAGMA_DISABLE_OPTIMIZATION
-struct FAndroidSignalParams
-{
-	int Signal;
-	siginfo* Info;
-	void* Context;
-};
-template<uint32 SignalStackSize, typename Derived>
-class FSignalSubstituteStack
+#if ANDROID_HAS_RTSIGNALS
+
+float GAndroidSignalTimeOut = 20.0f;
+static FAutoConsoleVariableRef CAndroidSignalTimeout(
+	TEXT("android.SignalTimeout"),
+	GAndroidSignalTimeOut,
+	TEXT("Time in seconds to wait for the signal handler to complete before timing out and terminating the process."),
+	ECVF_Default
+);
+
+template<typename Derived >
+typename FSignalHandler< Derived >::FSignalParams FSignalHandler< Derived >::SignalParams;
+template<typename Derived >
+int32 FSignalHandler< Derived >::SignalThreadStatus = (int32)FSignalHandler< Derived >::ESignalThreadStatus::NotInitialized;
+template<typename Derived >
+uint32 FSignalHandler< Derived >::ForwardingThreadID = 0xffffffff;
+template<typename Derived >
+int32 FSignalHandler< Derived >::ForwardingSignalType = -1;
+template<typename Derived >
+struct sigaction FSignalHandler< Derived >::PreviousActionForForwardSignal;
+
+
+class FThreadCallstackSignalHandler : FSignalHandler<FThreadCallstackSignalHandler>
 {
+	friend class FSignalHandler<FThreadCallstackSignalHandler>;
 public:
 	static void Init()
 	{
-		AllocateSignalStack();
+		FSignalHandler<FThreadCallstackSignalHandler>::Init(THREADBACKTRACE_SIGNAL_FWD);
+		HookTargetSignal();
 	}
 
 	static void Release()
 	{
-		if (FakeSignalStack)
-		{
-			ReleaseSignalStack();
-		}
-	}
-
-protected:
-#if PLATFORM_ANDROID_ARM64
-	template< void(*ActualHandlerFun)()>
-	static void CallActualHandler()
-	{
-		// Arm64, save current stack frame, swap in new stack and call our signal handler.
-		__asm__("mov x8, %0" : : "r"(&FakeSignalStack[SignalStackLocation]));					// Put the fake stack addr into x8
-		__asm__("stp x29, x30, [x8, -32]" : : );						// push the existing frame and lr on to fake stack
-		__asm__("mov x9, sp" : : );										// copy existing stack ptr to x9
-		__asm__("str x9, [x8, -16]" : : );								// push the existing stack on to fake stack.
-
-		__asm__("sub sp, x8, 32" : : );									// set stack ptr to new stack addr, after our pushed regs
-		__asm__("sub fp, x8, 32" : : );									// set frame ptr to new stack addr, after our pushed regs
-
-		__asm__("bl %0" : : "i"(ActualHandlerFun));						// Call our actual crash handler.
-
-		__asm__("ldp x29,x30,[sp]" : : );								// restore the original frame and lr.
-		__asm__("ldr x8, [sp, 16]" : : );								// restore the previously pushed stack pointer into x8
-		__asm__("mov sp, x8" : : );										// restore the original stack pointer
-	}
-#else
-	template< void(*ActualHandlerFunc)()>
-	static void CallActualHandler()
-	{
-		ActualHandlerFunc();
+		RestorePreviousTargetSignalHandler();
+		FSignalHandler<FThreadCallstackSignalHandler>::Release();
 	}
-#endif
 
-	static void AllocateSignalStack()
+private:
+	static void OnTargetSignal(int Signal, siginfo* Info, void* Context)
 	{
-		check(FakeSignalStack == nullptr);
-
-		uint8* BaseStackAllocation = (uint8*)mmap(nullptr, SignalStackAllocationSize, PROT_READ | PROT_WRITE, MAP_ANONYMOUS | MAP_PRIVATE, -1, 0);
-		// protect top of the stack to catch overruns.
-		if (mprotect(BaseStackAllocation, PAGE_SIZE, PROT_NONE) != 0)
-		{
-			UE_LOG(LogAndroid, Fatal, TEXT("Failed to set alt stack protection"));
-		}
-		if (mprotect(BaseStackAllocation + SignalStackSize + PAGE_SIZE, PAGE_SIZE, PROT_NONE) != 0)
+		while (FPlatformAtomics::InterlockedCompareExchange(&handling_signal, 1, 0) != 0)
 		{
-			UE_LOG(LogAndroid, Fatal, TEXT("Failed to set alt stack protection"));
+			FPlatformProcess::SleepNoStats(0.0f);
 		}
-
-		FakeSignalStack = BaseStackAllocation;
-
-		check(IsAligned(&FakeSignalStack[SignalStackLocation], 16));
+		FSignalHandler<FThreadCallstackSignalHandler>::ForwardSignal(Signal, Info, Context);
+		FPlatformAtomics::AtomicStore(&handling_signal, 0);
 	}
 
-	static void ReleaseSignalStack()
+	static void HandleTargetSignal(int Signal, siginfo* Info, void* Context)
 	{
-		check(FakeSignalStack);
-		munmap(FakeSignalStack, SignalStackAllocationSize);
-		FakeSignalStack = nullptr;
+		FPlatformStackWalk::HandleBackTraceSignal(Info, Context);
 	}
 
-	static const uint32 SignalStackAllocationSize = SignalStackSize + PAGE_SIZE + PAGE_SIZE;
-	static const uint32 SignalStackLocation = SignalStackAllocationSize - PAGE_SIZE;
-
-	static uint8* FakeSignalStack;
-};
-
-template<uint32 SignalStackSize, typename Derived >
-uint8* FSignalSubstituteStack<SignalStackSize, Derived>::FakeSignalStack = nullptr;
-#if ANDROID_HAS_THREADBACKTRACE
-class FThreadCallstackSignalSubstituteStack : public FSignalSubstituteStack<256 * 1024, FThreadCallstackSignalSubstituteStack >
-{
-public:
-	static void Init()
+	static void HookTargetSignal()
 	{
-		FSignalSubstituteStack<256 * 1024, FThreadCallstackSignalSubstituteStack >::Init();
+		check(bSignalHooked == false);
 		struct sigaction ActionForThread;
 		FMemory::Memzero(ActionForThread);
 		sigfillset(&ActionForThread.sa_mask);
 		ActionForThread.sa_flags = SA_SIGINFO | SA_RESTART | SA_ONSTACK;
-		ActionForThread.sa_sigaction = &OnSignal;
-		sigaction(THREAD_CALLSTACK_GENERATOR, &ActionForThread, nullptr);
+		ActionForThread.sa_sigaction = &OnTargetSignal;
+		sigaction(THREAD_CALLSTACK_GENERATOR, &ActionForThread, &PreviousActionForThreadGenerator);
+		bSignalHooked = true;
 	}
-private:
-	static void OnSignal(int Signal, siginfo* Info, void* Context)
+
+	static void RestorePreviousTargetSignalHandler()
 	{
-		while (FPlatformAtomics::InterlockedCompareExchange(&handling_signal, 1, 0) != 0)
+		if (bSignalHooked)
 		{
-			FPlatformProcess::Sleep(0.0f);
+			bSignalHooked = false;
+			sigaction(THREAD_CALLSTACK_GENERATOR, &PreviousActionForThreadGenerator, nullptr);
 		}
-
-		// prepare our signal parameters:
-		StackWalkSignalParams.Signal = Signal;
-		StackWalkSignalParams.Info = Info;
-		StackWalkSignalParams.Context = Context;
-
-		FSignalSubstituteStack::CallActualHandler< &ThreadStackWalk >();
-
-		FPlatformAtomics::AtomicStore(&handling_signal, 0);
 	}
 
-	static void ThreadStackWalk()
-	{
-		FPlatformStackWalk::HandleBackTraceSignal(StackWalkSignalParams.Info, StackWalkSignalParams.Context);
-	}
-	static FAndroidSignalParams StackWalkSignalParams;
+	static bool bSignalHooked;
 	static volatile sig_atomic_t handling_signal;
-	static ThreadStackUserData* ThreadStackData;
+	static struct sigaction PreviousActionForThreadGenerator;
 };
-volatile sig_atomic_t FThreadCallstackSignalSubstituteStack::handling_signal = 0;
-FAndroidSignalParams FThreadCallstackSignalSubstituteStack::StackWalkSignalParams;
-#else
-	// per thread backtrace unavailable
-	class FThreadCallstackSignalSubstituteStack
-	{
-	public:
-		static void Init() {};
-		static void Release() {};
-	};
-#endif
 
-class FFatalSignalHandler : protected FSignalSubstituteStack<256 * 1024, FFatalSignalHandler>
+volatile sig_atomic_t FThreadCallstackSignalHandler::handling_signal = 0;
+bool FThreadCallstackSignalHandler::bSignalHooked = false;
+struct sigaction FThreadCallstackSignalHandler::PreviousActionForThreadGenerator;
+
+class FFatalSignalHandler : public FSignalHandler<FFatalSignalHandler>
 {
+	friend class FSignalHandler<FFatalSignalHandler>;
 public:
 	static void Init()
 	{
-		FSignalSubstituteStack<256 * 1024, FFatalSignalHandler>::Init();
-		HookSignals();
+		FSignalHandler<FFatalSignalHandler>::Init(FATAL_SIGNAL_FWD);
+		HookTargetSignals();
 	}
 
 	static void Release()
 	{
-		RestorePreviousSignalHandlers();
-		FSignalSubstituteStack<256 * 1024, FFatalSignalHandler>::Release();
+		RestorePreviousTargetSignalHandlers();
+		FSignalHandler<FFatalSignalHandler>::Release();
 	}
 
 	static bool IsInFatalSignalHandler()
@@ -1067,32 +1014,23 @@ public:
 		return FPlatformAtomics::AtomicRead(&handling_fatal_signal) > 0;
 	}
 
-	// EnterFatalCrash, called before a crash is reported
-	// Only the first occurrence succeeds.
+protected:
+
 	static void EnterFatalCrash()
 	{
-		// we are a fatal signal, we *can* only handle one at a time. So avoid allow multiple fatal signals going through
+		// we are a fatal signal, we can only handle one at a time. So avoid allow multiple fatal signals going through
 		if (FPlatformAtomics::InterlockedIncrement(&handling_fatal_signal) != 1)
 		{
-			FPlatformProcess::Sleep(60.0f);
+			FPlatformProcess::SleepNoStats(60.0f);
 			exit(0);
 		}
 	}
 
-private:
-	static void OnHandleSignal(int Signal, siginfo* Info, void* Context)
+	static void OnTargetSignal(int Signal, siginfo* Info, void* Context)
 	{
 		EnterFatalCrash();
-
-		// prepare our signal parameters:
-		FatalSignalParams.Signal = Signal;
-		FatalSignalParams.Info = Info;
-		FatalSignalParams.Context = Context;
-
-		FFatalSignalHandler::CallActualHandler< &FatalSignalHandler >();
-
-		// Restore system handlers so Android could catch this signal after we are done with crashreport
-		RestorePreviousSignalHandlers();
+		FSignalHandler<FFatalSignalHandler>::ForwardSignal(Signal, Info, Context);
+		RestorePreviousTargetSignalHandlers();
 	}
 
 	static const ANSICHAR* CodeToString(int Signal, int si_code)
@@ -1182,15 +1120,15 @@ private:
 		return ANSI_TO_TCHAR(AnsiMessage);
 	}
 
-	static void FatalSignalHandler()
+	static void HandleTargetSignal(int Signal, siginfo* Info, void* Context)
 	{
 		// Switch to malloc crash.
 		FPlatformMallocCrash::Get().SetAsGMalloc();
 
-		FString Message = GetFatalSignalMessage(FatalSignalParams.Signal, FatalSignalParams.Info);	
+		FString Message = GetFatalSignalMessage(Signal, Info);
 		FAndroidCrashContext CrashContext(ECrashContextType::Crash, *Message);
 
-		CrashContext.InitFromSignal(FatalSignalParams.Signal, FatalSignalParams.Info, FatalSignalParams.Context);
+		CrashContext.InitFromSignal(Signal, Info, Context);
 		CrashContext.CaptureCrashInfo();
 		if (GCrashHandlerPointer)
 		{
@@ -1202,13 +1140,13 @@ private:
 			DefaultCrashHandler(CrashContext);
 		}
 	}
-
-	static void HookSignals()
+	static void HookTargetSignals()
 	{
+		check(PreviousSignalHandlersValid == false);
 		// hook our signals and record current set.
 		struct sigaction Action;
 		FMemory::Memzero(&Action, sizeof(struct sigaction));
-		Action.sa_sigaction = &OnHandleSignal;
+		Action.sa_sigaction = &OnTargetSignal;
 		// sigfillset will block all other signals whilst the signal handler is processing.
 		sigfillset(&Action.sa_mask);
 		Action.sa_flags = SA_SIGINFO | SA_RESTART | SA_ONSTACK;
@@ -1220,7 +1158,7 @@ private:
 		PreviousSignalHandlersValid = true;
 	}
 
-	static void RestorePreviousSignalHandlers()
+	static void RestorePreviousTargetSignalHandlers()
 	{
 		if (PreviousSignalHandlersValid)
 		{
@@ -1232,18 +1170,15 @@ private:
 		}
 	}
 
-	// params to pass from FFatalSignalHandler::OnHandleSignal to FatalSignalHandler.
-	static FAndroidSignalParams FatalSignalParams;
 	static volatile sig_atomic_t handling_fatal_signal;
 	static struct sigaction PrevActions[NumTargetSignals];
 	static bool PreviousSignalHandlersValid;
 };
-FAndroidSignalParams FFatalSignalHandler::FatalSignalParams;
+
 volatile sig_atomic_t FFatalSignalHandler::handling_fatal_signal = 0;
 struct sigaction FFatalSignalHandler::PrevActions[NumTargetSignals];
 bool FFatalSignalHandler::PreviousSignalHandlersValid = false;
-
-PRAGMA_ENABLE_OPTIMIZATION
+#endif// ANDROID_HAS_RTSIGNALS
 
 static void SetDefaultSignalHandlers()
 {
@@ -1260,7 +1195,11 @@ static void SetDefaultSignalHandlers()
 
 bool FAndroidMisc::IsInSignalHandler()
 {
+#if ANDROID_HAS_RTSIGNALS
 	return FFatalSignalHandler::IsInFatalSignalHandler();
+#else
+	return false;
+#endif
 }
 
 void FAndroidMisc::TriggerNonFatalCrashHandler(ECrashContextType InType, const FString& Message)
@@ -1283,8 +1222,6 @@ void FAndroidMisc::TriggerNonFatalCrashHandler(ECrashContextType InType, const F
 
 void FAndroidMisc::TriggerCrashHandler(const TCHAR* InErrorMessage, const TMap<FString, FString>& AdditionalProperties)
 {
-	FFatalSignalHandler::EnterFatalCrash();
-
 	FAndroidCrashContext CrashContext(ECrashContextType::Crash, InErrorMessage);
 	CrashContext.CaptureCrashInfo();
 	for (TMap<FString, FString>::TConstIterator Iter(AdditionalProperties); Iter; ++Iter)
@@ -1305,13 +1242,13 @@ void FAndroidMisc::TriggerCrashHandler(const TCHAR* InErrorMessage, const TMap<F
 
 void FAndroidMisc::SetCrashHandler(void(*CrashHandler)(const FGenericCrashContext& Context))
 {
-#if PLATFORM_ANDROID_ARM64
+#if ANDROID_HAS_RTSIGNALS
 	UE_LOG(LogAndroid, Log, TEXT("Setting Crash Handler = %p"), CrashHandler);
 
 	GCrashHandlerPointer = CrashHandler;
 
 	FFatalSignalHandler::Release();
-	FThreadCallstackSignalSubstituteStack::Release();
+	FThreadCallstackSignalHandler::Release();
 	// Passing -1 will leave these restored and won't trap them
 	if ((PTRINT)CrashHandler == -1)
 	{
@@ -1319,8 +1256,8 @@ void FAndroidMisc::SetCrashHandler(void(*CrashHandler)(const FGenericCrashContex
 	}
 
 	FFatalSignalHandler::Init();
-	FThreadCallstackSignalSubstituteStack::Init();
-#endif //PLATFORM_ANDROID_ARM64
+	FThreadCallstackSignalHandler::Init();
+#endif
 }
 
 bool FAndroidMisc::GetUseVirtualJoysticks()
diff --git a/Engine/Source/Runtime/Core/Private/Android/AndroidPlatformStackWalk.cpp b/Engine/Source/Runtime/Core/Private/Android/AndroidPlatformStackWalk.cpp
--- a/Engine/Source/Runtime/Core/Private/Android/AndroidPlatformStackWalk.cpp
+++ b/Engine/Source/Runtime/Core/Private/Android/AndroidPlatformStackWalk.cpp
@@ -11,6 +11,7 @@
 #include <dlfcn.h>
 #include <cxxabi.h>
 #include <stdio.h>
+#include "Android/AndroidSignals.h"
 
 #define HAS_LIBUNWIND PLATFORM_ANDROID_ARM64 && !PLATFORM_LUMIN
 
@@ -18,9 +19,8 @@
 #define UNW_LOCAL_ONLY
 #include "libunwind.h"
 #endif
-#if ANDROID_HAS_THREADBACKTRACE
+#if ANDROID_HAS_RTSIGNALS
 #include <syscall.h>
-#include "HAL/PlatformTime.h"
 #endif
 #include "HAL/IConsoleManager.h"
 #include "HAL/PlatformProcess.h"
@@ -273,7 +273,14 @@ static TAutoConsoleVariable<float> CVarAndroidPlatformThreadCallStackMaxWait(
 	GThreadCallStackMaxWait,
 	TEXT("The number of seconds allowed to spin before killing the process, with the assumption the back trace handler has hung."));
 
-#if ANDROID_HAS_THREADBACKTRACE
+#if ANDROID_HAS_RTSIGNALS
+/** Passed in through sigqueue for gathering of a callstack from a signal */
+struct ThreadStackUserData
+{
+	uint64* BackTrace;
+	int32 BackTraceCount;
+	SIZE_T CallStackSize;
+};
 
 int32 ThreadStackBackTraceStatus = 0;
 static const int32 ThreadStackBackTraceCurrentStatus_RUNNING = -2;
@@ -330,7 +337,7 @@ uint32 FAndroidPlatformStackWalk::CaptureThreadStackBackTrace(uint64 ThreadId, u
 					// request not yet started, skip it.
 					return 0;
 				}
-				FPlatformProcess::Sleep(PollTime);
+				FPlatformProcess::SleepNoStats(PollTime);
 			}
 
 			// We have waited for as long as we should for the signal handler to finish. Assume it has hang and we need to kill our selfs
@@ -375,4 +382,4 @@ uint32 FAndroidPlatformStackWalk::CaptureThreadStackBackTrace(uint64 ThreadId, u
 {
 	return 0;
 }
-#endif //ANDROID_HAS_THREADBACKTRACE
\ No newline at end of file
+#endif //ANDROID_HAS_RTSIGNALS
\ No newline at end of file
diff --git a/Engine/Source/Runtime/Core/Private/Android/AndroidSignals.h b/Engine/Source/Runtime/Core/Private/Android/AndroidSignals.h
new file mode 100644
--- /dev/null
+++ b/Engine/Source/Runtime/Core/Private/Android/AndroidSignals.h
@@ -0,0 +1,189 @@
+// Copyright 1998-2019 Epic Games, Inc. All Rights Reserved.
+
+/*=============================================================================
+	AndroidSignals.h
+=============================================================================*/
+#include "Android/AndroidPlatformMisc.h"
+#include "CoreGlobals.h"
+#include "HAL/PlatformProcess.h"
+#include <sys/syscall.h>
+
+#pragma once
+
+#if ANDROID_HAS_RTSIGNALS
+
+#define THREAD_CALLSTACK_GENERATOR + 5
+#define FATAL_SIGNAL_FWD SIGRTMIN + 6
+#define THREADBACKTRACE_SIGNAL_FWD SIGRTMIN + 7
+
+/*
+	This is a helper class that attempts to avoid android's limited signal stack size by
+	using a thread that infinite sleeps and is woken to only process signals.
+	'Target' signals are caught and forwarded to the sleeping thread.
+	The forwarded signal runs on the sleeping thread calling Derived::HandleTargetSignal() with the originating signal's parameters.
+
+	The crash handler therefore, does not run on the same thread that originally received the target signal.
+
+	Note, to avoid the requirement of suspending the thread while backgrounding (and therefore missing signals) the thread sleeps infinitely.
+*/
+
+extern float GAndroidSignalTimeOut;
+
+template<typename Derived>
+class FSignalHandler
+{
+protected:
+	enum class ESignalType : int32 { Signal, Exit };
+	enum class ESignalThreadStatus : int32 { NotInitialized, Ready, Starting, Busy, Complete, Exited };
+
+	static void Init(int ForwardingSignal)
+	{
+		check(ForwardingSignalType == -1);
+		check(SignalThreadStatus == (int32)ESignalThreadStatus::NotInitialized);
+
+		ForwardingSignalType = ForwardingSignal;
+		pthread_t SignalHandlerThread;
+
+		pthread_attr_t otherAttr;
+		pthread_attr_init(&otherAttr);
+		pthread_attr_setdetachstate(&otherAttr, PTHREAD_CREATE_DETACHED);
+		const size_t StackSize = 256 * 1024;
+		bool bStackSizeSuccess = (pthread_attr_setstacksize(&otherAttr, StackSize) == 0);
+		UE_CLOG(!bStackSizeSuccess, LogAndroid, Error, TEXT("Failed to set signal thread stack size."));
+
+		bool bThreadCreated = pthread_create(&SignalHandlerThread, &otherAttr, ThreadFunc, nullptr) == 0;
+		UE_CLOG(!bThreadCreated, LogAndroid, Fatal, TEXT("Failed to create signal handler"));
+
+		while (FPlatformAtomics::AtomicRead(&SignalThreadStatus) != (int32)ESignalThreadStatus::Ready)
+		{
+			FPlatformProcess::SleepNoStats(0);
+		}
+
+		struct sigaction ActionForForwardSignal;
+
+		FMemory::Memzero(ActionForForwardSignal);
+		sigfillset(&ActionForForwardSignal.sa_mask);
+		ActionForForwardSignal.sa_flags = SA_SIGINFO | SA_RESTART;
+		ActionForForwardSignal.sa_sigaction = &OnForwardedTargetSignal;
+		sigaction(ForwardingSignalType, &ActionForForwardSignal, &PreviousActionForForwardSignal);
+	}
+
+	static void Release()
+	{
+		if (ForwardingSignalType != -1)
+		{
+			SendExitSignal();
+			while (FPlatformAtomics::AtomicRead(&SignalThreadStatus) != (int32)ESignalThreadStatus::Exited)
+			{
+				FPlatformProcess::SleepNoStats(0);
+			}
+
+			sigaction(ForwardingSignalType, &PreviousActionForForwardSignal, nullptr);
+
+			SignalThreadStatus = (int32)ESignalThreadStatus::NotInitialized;
+			ForwardingThreadID = 0xffffffff;
+			ForwardingSignalType = -1;
+		}
+	}
+
+	static void ForwardSignal(int InSignal, siginfo* InInfo, void* InContext)
+	{
+		while (FPlatformAtomics::InterlockedCompareExchange(&SignalThreadStatus, (int32)ESignalThreadStatus::Starting, (int32)ESignalThreadStatus::Ready) != (int32)ESignalThreadStatus::Ready)
+		{
+			FPlatformProcess::SleepNoStats(0.0f);
+		}
+		SignalParams.Signal = InSignal;
+		SignalParams.Info = InInfo;
+		SignalParams.Context = InContext;
+
+		FPlatformMisc::MemoryBarrier();
+
+		if (SendSignal((int)ESignalType::Signal))
+		{
+			WaitForSignalHandlerToFinishOrExit();
+		}
+		else
+		{
+			// did not send, thread not available... set status to ready.
+			FPlatformAtomics::AtomicStore(&SignalThreadStatus, (int32)ESignalThreadStatus::Ready);
+		}
+	}
+
+private:
+	static void* ThreadFunc(void* param)
+	{
+		ForwardingThreadID = FPlatformTLS::GetCurrentThreadId();
+		FPlatformAtomics::AtomicStore(&SignalThreadStatus, (int32)ESignalThreadStatus::Ready);
+		while (true)
+		{
+			FPlatformProcess::SleepInfinite();
+		}
+		return nullptr;
+	}
+
+	static bool SendSignal(int SignalType)
+	{
+		siginfo_t info;
+		memset(&info, 0, sizeof(siginfo_t));
+		info.si_signo = ForwardingSignalType;
+		info.si_code = SI_QUEUE;
+		info.si_pid = syscall(SYS_getpid);
+		info.si_uid = syscall(SYS_getuid);
+		info.si_value.sival_int = SignalType;
+
+		// Avoid using sigqueue here as if the ThreadId is already blocked and in a signal handler
+		// sigqueue will try a different thread signal handler and report the wrong callstack
+		return syscall(SYS_rt_tgsigqueueinfo, info.si_pid, ForwardingThreadID, ForwardingSignalType, &info) == 0;
+	}
+
+	static bool SendExitSignal()
+	{
+		// enqueue a signal to exit the signal thread.
+		return SendSignal((int)ESignalType::Exit);
+	}
+
+	static void WaitForSignalHandlerToFinishOrExit()
+	{
+		const float PollInterval = 0.01f;
+		float CurrentWaitTime = 0.0f;
+
+		while (FPlatformAtomics::InterlockedCompareExchange(&SignalThreadStatus, (int32)ESignalThreadStatus::Ready, (int32)ESignalThreadStatus::Complete) != (int32)ESignalThreadStatus::Complete)
+		{
+			// have we timed out?
+			if (CurrentWaitTime > GAndroidSignalTimeOut)
+			{
+				exit(0);
+			}
+			FPlatformProcess::SleepNoStats(PollInterval);
+			CurrentWaitTime += PollInterval;
+		}
+	}
+
+	static void OnForwardedTargetSignal(int Signal, siginfo* Info, void* Context)
+	{
+		FPlatformAtomics::AtomicStore(&SignalThreadStatus, (int32)ESignalThreadStatus::Busy);
+
+		if (Info->si_value.sival_int == (int)ESignalType::Exit)
+		{
+			FPlatformAtomics::AtomicStore(&SignalThreadStatus, (int32)ESignalThreadStatus::Exited);
+			pthread_exit(0);
+			return;
+		}
+		Derived::HandleTargetSignal(SignalParams.Signal, SignalParams.Info, SignalParams.Context);
+		FPlatformAtomics::AtomicStore(&SignalThreadStatus, (int32)ESignalThreadStatus::Complete);
+	}
+
+	struct FSignalParams
+	{
+		int Signal;
+		siginfo* Info;
+		void* Context;
+	};
+
+	static FSignalParams SignalParams;
+	static int32 SignalThreadStatus;
+	static uint32 ForwardingThreadID;
+	static int32 ForwardingSignalType;
+	static struct sigaction PreviousActionForForwardSignal;
+};
+#endif
\ No newline at end of file
diff --git a/Engine/Source/Runtime/Core/Public/Android/AndroidPlatformMisc.h b/Engine/Source/Runtime/Core/Public/Android/AndroidPlatformMisc.h
--- a/Engine/Source/Runtime/Core/Public/Android/AndroidPlatformMisc.h
+++ b/Engine/Source/Runtime/Core/Public/Android/AndroidPlatformMisc.h
@@ -24,6 +24,8 @@ class TFunction;
 
 #define UE_DEBUG_BREAK_IMPL()	PLATFORM_BREAK()
 
+#define ANDROID_HAS_RTSIGNALS !PLATFORM_LUMIN && PLATFORM_USED_NDK_VERSION_INTEGER >= 21
+
 /**
  * Android implementation of the misc OS functions
  */
diff --git a/Engine/Source/Runtime/Core/Public/Android/AndroidPlatformStackWalk.h b/Engine/Source/Runtime/Core/Public/Android/AndroidPlatformStackWalk.h
--- a/Engine/Source/Runtime/Core/Public/Android/AndroidPlatformStackWalk.h
+++ b/Engine/Source/Runtime/Core/Public/Android/AndroidPlatformStackWalk.h
@@ -24,18 +24,4 @@ struct CORE_API FAndroidPlatformStackWalk : public FGenericPlatformStackWalk
 	static void HandleBackTraceSignal(siginfo* Info, void* Context);
 };
 
-#define ANDROID_HAS_THREADBACKTRACE !PLATFORM_LUMIN && PLATFORM_USED_NDK_VERSION_INTEGER >= 21
-
-#if ANDROID_HAS_THREADBACKTRACE
-
-/** Passed in through sigqueue for gathering of a callstack from a signal */
-struct ThreadStackUserData
-{
-	uint64* BackTrace;
-	int32 BackTraceCount;
-	SIZE_T CallStackSize;
-};
-#define THREAD_CALLSTACK_GENERATOR SIGRTMIN + 5
-#endif // ANDROID_HAS_THREADBACKTRACE
-
 typedef FAndroidPlatformStackWalk FPlatformStackWalk;
