Add operator* for TOptional.

diff --git a/Engine/Source/Runtime/Core/Public/Misc/Optional.h b/Engine/Source/Runtime/Core/Public/Misc/Optional.h
--- a/Engine/Source/Runtime/Core/Public/Misc/Optional.h
+++ b/Engine/Source/Runtime/Core/Public/Misc/Optional.h
@@ -187,10 +187,13 @@ public:
 
 	/** @return The optional value; undefined when IsSet() returns false. */
 	const OptionalType& GetValue() const { checkf(IsSet(), TEXT("It is an error to call GetValue() on an unset TOptional. Please either check IsSet() or use Get(DefaultValue) instead.")); return *(OptionalType*)&Value; }
-	      OptionalType& GetValue()       { checkf(IsSet(), TEXT("It is an error to call GetValue() on an unset TOptional. Please either check IsSet() or use Get(DefaultValue) instead.")); return *(OptionalType*)&Value; }
+		  OptionalType& GetValue()		 { checkf(IsSet(), TEXT("It is an error to call GetValue() on an unset TOptional. Please either check IsSet() or use Get(DefaultValue) instead.")); return *(OptionalType*)&Value; }
 
 	const OptionalType* operator->() const { return &GetValue(); }
-	      OptionalType* operator->()       { return &GetValue(); }
+		  OptionalType* operator->()	   { return &GetValue(); }
+
+	const OptionalType& operator*() const { return GetValue(); }
+		  OptionalType& operator*()		  { return GetValue(); }
 
 	/** @return The optional value when set; DefaultValue otherwise. */
 	const OptionalType& Get(const OptionalType& DefaultValue) const { return IsSet() ? *(OptionalType*)&Value : DefaultValue; }
