Culture - No longer need to create a hack to determine if cultures have are Left to Right or Right to Left, we now call the ICU function to determine it.

diff --git a/Engine/Source/Runtime/Core/Private/Internationalization/Culture.cpp b/Engine/Source/Runtime/Core/Private/Internationalization/Culture.cpp
--- a/Engine/Source/Runtime/Core/Private/Internationalization/Culture.cpp
+++ b/Engine/Source/Runtime/Core/Private/Internationalization/Culture.cpp
@@ -28,6 +28,7 @@ FCulture::FCulture(TUniquePtr<FCultureImplementation>&& InImplementation)
 	, CachedNativeRegion(Implementation->GetNativeRegion())
 	, CachedScript(Implementation->GetScript())
 	, CachedVariant(Implementation->GetVariant())
+	, CachedIsRightToLeft(Implementation->IsRightToLeft())
 { 
 }
 
@@ -177,6 +178,11 @@ const FString& FCulture::GetVariant() const
 	return CachedVariant;
 }
 
+bool FCulture::IsRightToLeft() const
+{
+	return CachedIsRightToLeft;
+}
+
 const FDecimalNumberFormattingRules& FCulture::GetDecimalNumberFormattingRules() const
 {
 	return Implementation->GetDecimalNumberFormattingRules();
diff --git a/Engine/Source/Runtime/Core/Private/Internationalization/CultureImplementation.h b/Engine/Source/Runtime/Core/Private/Internationalization/CultureImplementation.h
--- a/Engine/Source/Runtime/Core/Private/Internationalization/CultureImplementation.h
+++ b/Engine/Source/Runtime/Core/Private/Internationalization/CultureImplementation.h
@@ -33,6 +33,7 @@ public:
 	virtual FString GetRegion() const = 0;
 	virtual FString GetScript() const = 0;
 	virtual FString GetVariant() const = 0;
+	virtual bool IsRightToLeft() const = 0;
 	virtual const FDecimalNumberFormattingRules& GetDecimalNumberFormattingRules() = 0;
 	virtual const FDecimalNumberFormattingRules& GetPercentFormattingRules() = 0;
 	virtual const FDecimalNumberFormattingRules& GetCurrencyFormattingRules(const FString& InCurrencyCode) = 0;
diff --git a/Engine/Source/Runtime/Core/Private/Internationalization/Cultures/InvariantCulture.h b/Engine/Source/Runtime/Core/Private/Internationalization/Cultures/InvariantCulture.h
--- a/Engine/Source/Runtime/Core/Private/Internationalization/Cultures/InvariantCulture.h
+++ b/Engine/Source/Runtime/Core/Private/Internationalization/Cultures/InvariantCulture.h
@@ -60,8 +60,9 @@ public:
 			FString(TEXT("iv")),																//const FString TwoLetterISOLanguageName
 			DecimalNumberFormattingRules,														//const FDecimalNumberFormattingRules InDecimalNumberFormattingRules
 			PercentFormattingRules,																//const FDecimalNumberFormattingRules InPercentFormattingRules
-			BaseCurrencyFormattingRules															//const FDecimalNumberFormattingRules InBaseCurrencyFormattingRules
-			));
+			BaseCurrencyFormattingRules,														//const FDecimalNumberFormattingRules InBaseCurrencyFormattingRules
+			false																				//const bool IsRightToLeft
+		));
 
 		return Culture;
 	}
diff --git a/Engine/Source/Runtime/Core/Private/Internationalization/Cultures/LeetCulture.h b/Engine/Source/Runtime/Core/Private/Internationalization/Cultures/LeetCulture.h
--- a/Engine/Source/Runtime/Core/Private/Internationalization/Cultures/LeetCulture.h
+++ b/Engine/Source/Runtime/Core/Private/Internationalization/Cultures/LeetCulture.h
@@ -28,6 +28,7 @@ public:
 	virtual FString GetRegion() const override { return FString(); }
 	virtual FString GetScript() const override { return FString(); }
 	virtual FString GetVariant() const override { return FString(); }
+	virtual bool IsRightToLeft() const override { return false; }
 
 private:
 	FCultureRef InvariantCulture;
diff --git a/Engine/Source/Runtime/Core/Private/Internationalization/CustomCultureImplementation.h b/Engine/Source/Runtime/Core/Private/Internationalization/CustomCultureImplementation.h
--- a/Engine/Source/Runtime/Core/Private/Internationalization/CustomCultureImplementation.h
+++ b/Engine/Source/Runtime/Core/Private/Internationalization/CustomCultureImplementation.h
@@ -31,7 +31,8 @@ public:
 			InCustomCulture->GetBaseCulture()->GetTwoLetterISOLanguageName(),
 			InCustomCulture->GetBaseCulture()->GetDecimalNumberFormattingRules(),
 			InCustomCulture->GetBaseCulture()->GetPercentFormattingRules(),
-			InCustomCulture->GetBaseCulture()->GetCurrencyFormattingRules(FString())
+			InCustomCulture->GetBaseCulture()->GetCurrencyFormattingRules(FString()),
+			InCustomCulture->GetBaseCulture()->IsRightToLeft()
 		)
 #endif
 		, CustomCulture(InCustomCulture)
@@ -53,6 +54,7 @@ public:
 	virtual FString GetRegion() const override { return CustomCulture->GetRegion(); }
 	virtual FString GetScript() const override { return CustomCulture->GetScript(); }
 	virtual FString GetVariant() const override { return CustomCulture->GetVariant(); }
+	virtual bool IsRightToLeft() const override { return CustomCulture->IsRightToLeft(); }
 
 private:
 	TSharedRef<ICustomCulture> CustomCulture;
diff --git a/Engine/Source/Runtime/Core/Private/Internationalization/ICUCulture.cpp b/Engine/Source/Runtime/Core/Private/Internationalization/ICUCulture.cpp
--- a/Engine/Source/Runtime/Core/Private/Internationalization/ICUCulture.cpp
+++ b/Engine/Source/Runtime/Core/Private/Internationalization/ICUCulture.cpp
@@ -618,6 +618,15 @@ FString FICUCultureImplementation::GetVariant() const
 	return ICULocale.getVariant();
 }
 
+bool FICUCultureImplementation::IsRightToLeft() const
+{
+#if WITH_ICU_V64
+	return ICULocale.isRightToLeft();
+#else
+	return false;
+#endif
+}
+
 TSharedRef<const icu::BreakIterator> FICUCultureImplementation::GetBreakIterator(const EBreakIteratorType Type)
 {
 	TSharedPtr<const icu::BreakIterator> Result;
diff --git a/Engine/Source/Runtime/Core/Private/Internationalization/ICUCulture.h b/Engine/Source/Runtime/Core/Private/Internationalization/ICUCulture.h
--- a/Engine/Source/Runtime/Core/Private/Internationalization/ICUCulture.h
+++ b/Engine/Source/Runtime/Core/Private/Internationalization/ICUCulture.h
@@ -186,6 +186,7 @@ public:
 	virtual FString GetRegion() const override;
 	virtual FString GetScript() const override;
 	virtual FString GetVariant() const override;
+	virtual bool IsRightToLeft() const override;
 	virtual const FDecimalNumberFormattingRules& GetDecimalNumberFormattingRules() override;
 	virtual const FDecimalNumberFormattingRules& GetPercentFormattingRules() override;
 	virtual const FDecimalNumberFormattingRules& GetCurrencyFormattingRules(const FString& InCurrencyCode) override;
diff --git a/Engine/Source/Runtime/Core/Private/Internationalization/LegacyCulture.cpp b/Engine/Source/Runtime/Core/Private/Internationalization/LegacyCulture.cpp
--- a/Engine/Source/Runtime/Core/Private/Internationalization/LegacyCulture.cpp
+++ b/Engine/Source/Runtime/Core/Private/Internationalization/LegacyCulture.cpp
@@ -18,7 +18,8 @@ FLegacyCultureImplementation::FLegacyCultureImplementation(
 	const FString& InTwoLetterISOLanguageName,
 	const FDecimalNumberFormattingRules& InDecimalNumberFormattingRules,
 	const FDecimalNumberFormattingRules& InPercentFormattingRules,
-	const FDecimalNumberFormattingRules& InBaseCurrencyFormattingRules
+	const FDecimalNumberFormattingRules& InBaseCurrencyFormattingRules,
+	bool InIsRightToLeft
 	)
 	: DisplayName(InDisplayName)
 	, EnglishName(InEnglishName)
@@ -32,6 +33,7 @@ FLegacyCultureImplementation::FLegacyCultureImplementation(
 	, DecimalNumberFormattingRules( InDecimalNumberFormattingRules )
 	, PercentFormattingRules( InPercentFormattingRules )
 	, BaseCurrencyFormattingRules( InBaseCurrencyFormattingRules )
+	, bIsRightToLeft(InIsRightToLeft)
 { 
 }
 
@@ -107,6 +109,11 @@ FString FLegacyCultureImplementation::GetVariant() const
 	return FString();
 }
 
+bool FLegacyCultureImplementation::IsRightToLeft() const
+{
+	return bIsRightToLeft;
+}
+
 FString FLegacyCultureImplementation::GetUnrealLegacyThreeLetterISOLanguageName() const
 {
 	return UnrealLegacyThreeLetterISOLanguageName;
diff --git a/Engine/Source/Runtime/Core/Private/Internationalization/LegacyCulture.h b/Engine/Source/Runtime/Core/Private/Internationalization/LegacyCulture.h
--- a/Engine/Source/Runtime/Core/Private/Internationalization/LegacyCulture.h
+++ b/Engine/Source/Runtime/Core/Private/Internationalization/LegacyCulture.h
@@ -27,7 +27,8 @@ public:
 		const FString& InTwoLetterISOLanguageName,
 		const FDecimalNumberFormattingRules& InDecimalNumberFormattingRules,
 		const FDecimalNumberFormattingRules& InPercentFormattingRules,
-		const FDecimalNumberFormattingRules& InBaseCurrencyFormattingRules
+		const FDecimalNumberFormattingRules& InBaseCurrencyFormattingRules,
+		bool InIsRightToLeft
 		);
 	virtual ~FLegacyCultureImplementation() = default;
 
@@ -46,6 +47,7 @@ public:
 	virtual FString GetRegion() const override;
 	virtual FString GetScript() const override;
 	virtual FString GetVariant() const override;
+	virtual bool IsRightToLeft() const override;
 	virtual const FDecimalNumberFormattingRules& GetDecimalNumberFormattingRules() override;
 	virtual const FDecimalNumberFormattingRules& GetPercentFormattingRules() override;
 	virtual const FDecimalNumberFormattingRules& GetCurrencyFormattingRules(const FString& InCurrencyCode) override;
@@ -92,6 +94,9 @@ private:
 	// Rules for formatting currency numbers in this culture
 	const FDecimalNumberFormattingRules BaseCurrencyFormattingRules;
 
+	// Is this culture right to left?
+	const bool bIsRightToLeft;
+
 	// Rules for formatting alternate currencies in this culture
 	TMap<FString, TSharedPtr<const FDecimalNumberFormattingRules>> UEAlternateCurrencyFormattingRules;
 	FCriticalSection UEAlternateCurrencyFormattingRulesCS;
diff --git a/Engine/Source/Runtime/Core/Public/Internationalization/Culture.h b/Engine/Source/Runtime/Core/Public/Internationalization/Culture.h
--- a/Engine/Source/Runtime/Core/Public/Internationalization/Culture.h
+++ b/Engine/Source/Runtime/Core/Public/Internationalization/Culture.h
@@ -68,6 +68,8 @@ public:
 
 	const FString& GetVariant() const;
 
+	bool IsRightToLeft() const;
+
 	const FDecimalNumberFormattingRules& GetDecimalNumberFormattingRules() const;
 
 	const FDecimalNumberFormattingRules& GetPercentFormattingRules() const;
@@ -115,4 +117,5 @@ private:
 	FString CachedNativeRegion;
 	FString CachedScript;
 	FString CachedVariant;
+	bool CachedIsRightToLeft;
 };
diff --git a/Engine/Source/Runtime/Core/Public/Internationalization/ICustomCulture.h b/Engine/Source/Runtime/Core/Public/Internationalization/ICustomCulture.h
--- a/Engine/Source/Runtime/Core/Public/Internationalization/ICustomCulture.h
+++ b/Engine/Source/Runtime/Core/Public/Internationalization/ICustomCulture.h
@@ -27,4 +27,5 @@ public:
 	virtual FString GetRegion() const = 0;
 	virtual FString GetScript() const = 0;
 	virtual FString GetVariant() const = 0;
+	virtual bool IsRightToLeft() const = 0;
 };
diff --git a/Engine/Source/Runtime/SlateCore/Private/Layout/FlowDirection.cpp b/Engine/Source/Runtime/SlateCore/Private/Layout/FlowDirection.cpp
--- a/Engine/Source/Runtime/SlateCore/Private/Layout/FlowDirection.cpp
+++ b/Engine/Source/Runtime/SlateCore/Private/Layout/FlowDirection.cpp
@@ -35,8 +35,7 @@ EFlowDirection FLayoutLocalization::GetLocalizedLayoutDirection()
 	static TFrameValue<EFlowDirection> FrameCoherentDirection;
 	if (!FrameCoherentDirection.IsSet())
 	{
-		//HACK: Normally we'd get this from the culture, but our cultures don't tell us if they're RightToLeft.  Newer versions of ICU tell you, but we need to upgrade.
-		FrameCoherentDirection = FInternationalization::Get().GetCurrentCulture()->GetTwoLetterISOLanguageName().StartsWith(TEXT("ar")) ? EFlowDirection::RightToLeft : EFlowDirection::LeftToRight;
+		FrameCoherentDirection = FInternationalization::Get().GetCurrentCulture()->IsRightToLeft() ? EFlowDirection::RightToLeft : EFlowDirection::LeftToRight;
 	}
 
 	return FrameCoherentDirection.GetValue();
