Fix implict ubool to bool cast for vs2019.

diff --git a/Engine/Source/Runtime/Core/Private/Internationalization/ICUCulture.cpp b/Engine/Source/Runtime/Core/Private/Internationalization/ICUCulture.cpp
--- a/Engine/Source/Runtime/Core/Private/Internationalization/ICUCulture.cpp
+++ b/Engine/Source/Runtime/Core/Private/Internationalization/ICUCulture.cpp
@@ -621,7 +621,7 @@ FString FICUCultureImplementation::GetVariant() const
 bool FICUCultureImplementation::IsRightToLeft() const
 {
 #if WITH_ICU_V64
-	return ICULocale.isRightToLeft();
+	return ICULocale.isRightToLeft() != 0;
 #else
 	return false;
 #endif
