#UE4 Disable CrashReportClient monitor in non-editor builds due to not respecting a few commandline parameters that are used in runtime

diff --git a/Engine/Source/Runtime/Core/Private/Windows/WindowsPlatformCrashContext.cpp b/Engine/Source/Runtime/Core/Private/Windows/WindowsPlatformCrashContext.cpp
--- a/Engine/Source/Runtime/Core/Private/Windows/WindowsPlatformCrashContext.cpp
+++ b/Engine/Source/Runtime/Core/Private/Windows/WindowsPlatformCrashContext.cpp
@@ -36,7 +36,7 @@
 	#define UE_LOG_CRASH_CALLSTACK 1
 #endif
 
-#if !IS_PROGRAM
+#if WITH_EDITOR
 	#define USE_CRASH_REPORTER_MONITOR 1
 #else 
 	#define USE_CRASH_REPORTER_MONITOR 0
