Add FRay type to UnrealMath

diff --git a/Engine/Source/Runtime/Core/Public/Math/Ray.h b/Engine/Source/Runtime/Core/Public/Math/Ray.h
new file mode 100644
--- /dev/null
+++ b/Engine/Source/Runtime/Core/Public/Math/Ray.h
@@ -0,0 +1,112 @@
+// Copyright 1998-2019 Epic Games, Inc. All Rights Reserved.
+
+#pragma once
+
+#include "Math/Vector.h"
+
+
+/**
+ * 3D Ray represented by Origin and (normalized) Direction
+ */
+class FRay
+{
+public:
+
+	/** Ray origin point */
+	FVector Origin;
+
+	/** Ray direction vector (always normalized) */
+	FVector Direction;
+
+public:
+
+	/** Default constructor initializes ray to Zero origin and Z-axis direction */
+	FRay()
+	{
+		Origin = FVector::ZeroVector;
+		Direction = FVector(0, 0, 1);
+	}
+
+	/** 
+	  * Initialize Ray with origin and direction
+	  *
+	  * @param Origin Ray Origin Point
+	  * @param Direction Ray Direction Vector
+	  * @param bDirectionIsNormalized Direction will be normalized unless this is passed as true (default false)
+	  */
+	FRay(const FVector& Origin, const FVector& Direction, bool bDirectionIsNormalized = false)
+	{
+		this->Origin = Origin;
+		this->Direction = Direction;
+		if (bDirectionIsNormalized == false)
+		{
+			this->Direction.Normalize();    // is this a full-accuracy sqrt?
+		}
+	}
+
+
+public:
+
+	/** 
+	 * Calculate position on ray at given distance/parameter
+	 *
+	 * @param RayParameter Scalar distance along Ray
+	 * @return Point on Ray
+	 */
+	FVector PointAt(float RayParameter) const
+	{
+		return Origin + RayParameter * Direction;
+	}
+
+	/**
+	 * Calculate ray parameter (distance from origin to closest point) for query Point
+	 *
+	 * @param Point query Point
+	 * @return distance along ray from origin to closest point
+	 */
+	float GetParameter(const FVector& Point) const
+	{
+		return FVector::DotProduct((Point - Origin), Direction);
+	}
+
+	/**
+	 * Find minimum squared distance from query point to ray
+	 *
+	 * @param Point query Point
+	 * @return squared distance to Ray
+	 */
+	float DistSquared(const FVector& Point) const
+	{
+		float RayParameter = FVector::DotProduct((Point - Origin), Direction);
+		if (RayParameter < 0)
+		{
+			return FVector::DistSquared(Origin, Point);
+		}
+		else 
+		{
+			FVector ProjectionPt = Origin + RayParameter * Direction;
+			return FVector::DistSquared(ProjectionPt, Point);
+		}
+	}
+
+	/**
+	 * Find closest point on ray to query point
+	 * @param Point query point
+	 * @return closest point on Ray
+	 */
+	FVector ClosestPoint(const FVector& Point) const
+	{
+		float RayParameter = FVector::DotProduct((Point - Origin), Direction);
+		if (RayParameter < 0) 
+		{
+			return Origin;
+		}
+		else 
+		{
+			return Origin + RayParameter * Direction;
+		}
+	}
+
+
+};
+
diff --git a/Engine/Source/Runtime/Core/Public/Math/UnrealMath.h b/Engine/Source/Runtime/Core/Public/Math/UnrealMath.h
--- a/Engine/Source/Runtime/Core/Public/Math/UnrealMath.h
+++ b/Engine/Source/Runtime/Core/Public/Math/UnrealMath.h
@@ -14,6 +14,7 @@
 #include "Math/Vector4.h"
 #include "Math/VectorRegister.h"
 #include "Math/TwoVectors.h"
+#include "Math/Ray.h"
 #include "Math/Edge.h"
 #include "Math/Plane.h"
 #include "Math/Sphere.h"
