compile fix for MeshEditor Extension

diff --git a/Engine/Source/Runtime/Core/Public/Misc/Change.h b/Engine/Source/Runtime/Core/Public/Misc/Change.h
--- a/Engine/Source/Runtime/Core/Public/Misc/Change.h
+++ b/Engine/Source/Runtime/Core/Public/Misc/Change.h
@@ -78,7 +78,7 @@ class CORE_API FSwapChange : public FChange
 {
 
 public:
-	virtual EChangeStyle GetChangeType()
+	virtual EChangeStyle GetChangeType() override
 	{
 		return FChange::EChangeStyle::InPlaceSwap;
 	}
@@ -106,7 +106,7 @@ class CORE_API FCommandChange : public FChange
 {
 
 public:
-	virtual EChangeStyle GetChangeType()
+	virtual EChangeStyle GetChangeType() override
 	{
 		return FChange::EChangeStyle::CommandPattern;
 	}
