Fixed String Table crash when attempting to load an asset during a save or GC

diff --git a/Engine/Source/Editor/Kismet/Private/ImaginaryBlueprintData.cpp b/Engine/Source/Editor/Kismet/Private/ImaginaryBlueprintData.cpp
--- a/Engine/Source/Editor/Kismet/Private/ImaginaryBlueprintData.cpp
+++ b/Engine/Source/Editor/Kismet/Private/ImaginaryBlueprintData.cpp
@@ -35,7 +35,10 @@ FText FSearchableValueInfo::GetDisplayText(const TMap<int32, FText>& InLookupTab
 			AsyncTask(ENamedThreads::GameThread, [TableId, &Promise]()
 			{
 				FName ResolvedTableId = TableId;
-				IStringTableEngineBridge::FullyLoadStringTableAsset(ResolvedTableId); // Trigger the asset load
+				if (IStringTableEngineBridge::CanFindOrLoadStringTableAsset())
+				{
+					IStringTableEngineBridge::FullyLoadStringTableAsset(ResolvedTableId); // Trigger the asset load
+				}
 				Promise.SetValue(true); // Signal completion
 			});
 
diff --git a/Engine/Source/Runtime/Core/Private/Internationalization/StringTableCore.cpp b/Engine/Source/Runtime/Core/Private/Internationalization/StringTableCore.cpp
--- a/Engine/Source/Runtime/Core/Private/Internationalization/StringTableCore.cpp
+++ b/Engine/Source/Runtime/Core/Private/Internationalization/StringTableCore.cpp
@@ -577,7 +577,10 @@ void FStringTableRedirects::RedirectTableId(FName& InOutTableId)
 	}
 
 	// Process the asset redirect (only works if the asset is loaded)
-	IStringTableEngineBridge::RedirectStringTableAsset(InOutTableId);
+	if (IStringTableEngineBridge::CanFindOrLoadStringTableAsset())
+	{
+		IStringTableEngineBridge::RedirectStringTableAsset(InOutTableId);
+	}
 }
 
 void FStringTableRedirects::RedirectKey(const FName InTableId, FString& InOutKey)
diff --git a/Engine/Source/Runtime/Core/Private/Internationalization/TextHistory.cpp b/Engine/Source/Runtime/Core/Private/Internationalization/TextHistory.cpp
--- a/Engine/Source/Runtime/Core/Private/Internationalization/TextHistory.cpp
+++ b/Engine/Source/Runtime/Core/Private/Internationalization/TextHistory.cpp
@@ -2577,7 +2577,7 @@ void FTextHistory_StringTableEntry::FStringTableReferenceData::Initialize(uint16
 		LoadingPhase = EStringTableLoadingPhase::Loaded;
 		ResolveStringTableEntry();
 	}
-	else if (InLoadingPolicy == EStringTableLoadingPolicy::FindOrFullyLoad && IsInGameThread())
+	else if (InLoadingPolicy == EStringTableLoadingPolicy::FindOrFullyLoad && IStringTableEngineBridge::CanFindOrLoadStringTableAsset())
 	{
 		// Forced synchronous load
 		LoadingPhase = EStringTableLoadingPhase::Loaded;
@@ -2667,7 +2667,7 @@ FStringTableEntryConstPtr FTextHistory_StringTableEntry::FStringTableReferenceDa
 
 void FTextHistory_StringTableEntry::FStringTableReferenceData::ConditionalBeginAssetLoad()
 {
-	if (!IsInGameThread())
+	if (!IStringTableEngineBridge::CanFindOrLoadStringTableAsset())
 	{
 		return;
 	}
diff --git a/Engine/Source/Runtime/Core/Public/Internationalization/StringTableCore.h b/Engine/Source/Runtime/Core/Public/Internationalization/StringTableCore.h
--- a/Engine/Source/Runtime/Core/Public/Internationalization/StringTableCore.h
+++ b/Engine/Source/Runtime/Core/Public/Internationalization/StringTableCore.h
@@ -178,6 +178,15 @@ public:
 	 */
 	typedef TFunction<void(FName, FName)> FLoadStringTableAssetCallback;
 
+	/** 
+	 * Check to see whether it is currently safe to attempt to find or load a string table asset.
+	 * @return True if it is safe to attempt to find or load a string table asset, false otherwise.
+	 */
+	static bool CanFindOrLoadStringTableAsset()
+	{
+		return !InstancePtr || InstancePtr->CanFindOrLoadStringTableAssetImpl();
+	}
+
 	/**
 	 * Load a string table asset by its name, potentially doing so asynchronously. 
 	 * @note If the string table is already loaded, or loading is perform synchronously, then the callback will be called before this function returns.
@@ -185,7 +194,7 @@ public:
 	 */
 	static int32 LoadStringTableAsset(const FName InTableId, FLoadStringTableAssetCallback InLoadedCallback = FLoadStringTableAssetCallback())
 	{
-		check(IsInGameThread());
+		check(CanFindOrLoadStringTableAsset());
 
 		if (InstancePtr)
 		{
@@ -206,7 +215,7 @@ public:
 	 */
 	static void FullyLoadStringTableAsset(FName& InOutTableId)
 	{
-		check(IsInGameThread());
+		check(CanFindOrLoadStringTableAsset());
 
 		if (InstancePtr)
 		{
@@ -217,6 +226,8 @@ public:
 	/** Redirect string table asset by its name */
 	static void RedirectStringTableAsset(FName& InOutTableId)
 	{
+		check(CanFindOrLoadStringTableAsset());
+
 		if (InstancePtr)
 		{
 			InstancePtr->RedirectStringTableAssetImpl(InOutTableId);
@@ -247,6 +258,7 @@ public:
 protected:
 	virtual ~IStringTableEngineBridge() {}
 
+	virtual bool CanFindOrLoadStringTableAssetImpl() = 0;
 	virtual int32 LoadStringTableAssetImpl(const FName InTableId, FLoadStringTableAssetCallback InLoadedCallback) = 0;
 	virtual void FullyLoadStringTableAssetImpl(FName& InOutTableId) = 0;
 	virtual void RedirectStringTableAssetImpl(FName& InOutTableId) = 0;
diff --git a/Engine/Source/Runtime/Engine/Private/Internationalization/StringTable.cpp b/Engine/Source/Runtime/Engine/Private/Internationalization/StringTable.cpp
--- a/Engine/Source/Runtime/Engine/Private/Internationalization/StringTable.cpp
+++ b/Engine/Source/Runtime/Engine/Private/Internationalization/StringTable.cpp
@@ -124,6 +124,11 @@ private:
 	}
 
 	//~ IStringTableEngineBridge interface
+	virtual bool CanFindOrLoadStringTableAssetImpl() override
+	{
+		return IsInGameThread() && !IsGarbageCollecting() && !GIsSavingPackage;
+	}
+
 	virtual int32 LoadStringTableAssetImpl(const FName InTableId, FLoadStringTableAssetCallback InLoadedCallback) override
 	{
 		const FSoftObjectPath StringTableAssetReference = GetAssetReference(InTableId);
