Move FTextRange to Core

diff --git a/Engine/Plugins/Tests/RuntimeTests/Source/RuntimeTests/Private/Slate/RichTextMarkupProcessingTest.cpp b/Engine/Plugins/Tests/RuntimeTests/Source/RuntimeTests/Private/Slate/RichTextMarkupProcessingTest.cpp
--- a/Engine/Plugins/Tests/RuntimeTests/Source/RuntimeTests/Private/Slate/RichTextMarkupProcessingTest.cpp
+++ b/Engine/Plugins/Tests/RuntimeTests/Source/RuntimeTests/Private/Slate/RichTextMarkupProcessingTest.cpp
@@ -2,7 +2,6 @@
 
 #include "CoreMinimal.h"
 #include "Misc/AutomationTest.h"
-#include "Framework/Text/TextRange.h"
 #include "Framework/Text/ITextDecorator.h"
 
 #include "Framework/Text/RichTextMarkupProcessing.h"
diff --git a/Engine/Source/Developer/OutputLog/Private/SOutputLog.cpp b/Engine/Source/Developer/OutputLog/Private/SOutputLog.cpp
--- a/Engine/Source/Developer/OutputLog/Private/SOutputLog.cpp
+++ b/Engine/Source/Developer/OutputLog/Private/SOutputLog.cpp
@@ -1,7 +1,6 @@
 // Copyright 1998-2019 Epic Games, Inc. All Rights Reserved.
 
 #include "SOutputLog.h"
-#include "Framework/Text/TextRange.h"
 #include "Framework/Text/IRun.h"
 #include "Framework/Text/TextLayout.h"
 #include "Misc/ConfigCacheIni.h"
diff --git a/Engine/Source/Editor/UnrealEd/Private/UnrealEdPrivatePCH.h b/Engine/Source/Editor/UnrealEd/Private/UnrealEdPrivatePCH.h
--- a/Engine/Source/Editor/UnrealEd/Private/UnrealEdPrivatePCH.h
+++ b/Engine/Source/Editor/UnrealEd/Private/UnrealEdPrivatePCH.h
@@ -392,7 +392,6 @@
 #include "Framework/Commands/InputBindingManager.h"
 #include "Framework/Commands/Commands.h"
 #include "Framework/Text/TextLayout.h"
-#include "Framework/Text/TextRange.h"
 #include "Framework/Text/TextRunRenderer.h"
 #include "Framework/Text/TextLineHighlight.h"
 #include "Framework/Text/IRun.h"
diff --git a/Engine/Source/Editor/UnrealEd/Public/UnrealEdSharedPCH.h b/Engine/Source/Editor/UnrealEd/Public/UnrealEdSharedPCH.h
--- a/Engine/Source/Editor/UnrealEd/Public/UnrealEdSharedPCH.h
+++ b/Engine/Source/Editor/UnrealEd/Public/UnrealEdSharedPCH.h
@@ -382,7 +382,6 @@
 #include "Framework/Commands/UICommandList.h"
 #include "Framework/MultiBox/MultiBoxExtender.h"
 #include "Framework/Text/TextLayout.h"
-#include "Framework/Text/TextRange.h"
 #include "Framework/Text/TextRunRenderer.h"
 #include "Framework/Text/TextLineHighlight.h"
 #include "Framework/Text/IRun.h"
diff --git a/Engine/Source/Runtime/AppFramework/Private/Widgets/Testing/STestSuite.cpp b/Engine/Source/Runtime/AppFramework/Private/Widgets/Testing/STestSuite.cpp
--- a/Engine/Source/Runtime/AppFramework/Private/Widgets/Testing/STestSuite.cpp
+++ b/Engine/Source/Runtime/AppFramework/Private/Widgets/Testing/STestSuite.cpp
@@ -42,7 +42,6 @@
 #include "Framework/Commands/InputChord.h"
 #include "Framework/Commands/Commands.h"
 #include "Framework/Commands/UICommandList.h"
-#include "Framework/Text/TextRange.h"
 #include "Framework/Text/IRun.h"
 #include "Framework/Text/TextLayout.h"
 #include "Framework/Text/ISlateRun.h"
diff --git a/Engine/Source/Runtime/Core/Private/Containers/String.cpp b/Engine/Source/Runtime/Core/Private/Containers/String.cpp
--- a/Engine/Source/Runtime/Core/Private/Containers/String.cpp
+++ b/Engine/Source/Runtime/Core/Private/Containers/String.cpp
@@ -1555,4 +1555,35 @@ FString SlugStringForValidName(const FString& DisplayString, const TCHAR* Replac
 	}
 
 	return GeneratedName;
-}
\ No newline at end of file
+}
+
+void FTextRange::CalculateLineRangesFromString(const FString& Input, TArray<FTextRange>& LineRanges)
+{
+	int32 LineBeginIndex = 0;
+
+	// Loop through splitting at new-lines
+	const TCHAR* const InputStart = *Input;
+	for (const TCHAR* CurrentChar = InputStart; CurrentChar && *CurrentChar; ++CurrentChar)
+	{
+		// Handle a chain of \r\n slightly differently to stop the FChar::IsLinebreak adding two separate new-lines
+		const bool bIsWindowsNewLine = (*CurrentChar == '\r' && *(CurrentChar + 1) == '\n');
+		if (bIsWindowsNewLine || FChar::IsLinebreak(*CurrentChar))
+		{
+			const int32 LineEndIndex = (CurrentChar - InputStart);
+			check(LineEndIndex >= LineBeginIndex);
+			LineRanges.Emplace(FTextRange(LineBeginIndex, LineEndIndex));
+
+			if (bIsWindowsNewLine)
+			{
+				++CurrentChar; // skip the \n of the \r\n chain
+			}
+			LineBeginIndex = (CurrentChar - InputStart) + 1; // The next line begins after the end of the current line
+		}
+	}
+
+	// Process any remaining string after the last new-line
+	if (LineBeginIndex <= Input.Len())
+	{
+		LineRanges.Emplace(FTextRange(LineBeginIndex, Input.Len()));
+	}
+}
diff --git a/Engine/Source/Runtime/Core/Public/Containers/UnrealString.h b/Engine/Source/Runtime/Core/Public/Containers/UnrealString.h
--- a/Engine/Source/Runtime/Core/Public/Containers/UnrealString.h
+++ b/Engine/Source/Runtime/Core/Public/Containers/UnrealString.h
@@ -22,6 +22,7 @@
 #include "Templates/IsValidVariadicFunctionArg.h"
 #include "Templates/AndOrNot.h"
 #include "Templates/IsArrayOrRefOfType.h"
+#include "Templates/TypeHash.h"
 
 struct FStringFormatArg;
 template<typename KeyType,typename ValueType,typename SetAllocator ,typename KeyFuncs > class TMap;
@@ -2363,4 +2364,65 @@ CORE_API int32 FindMatchingClosingParenthesis(const FString& TargetString, const
 */
 CORE_API FString SlugStringForValidName(const FString& DisplayString, const TCHAR* ReplaceWith = TEXT(""));
 
+struct CORE_API FTextRange
+{
+	FTextRange()
+		: BeginIndex(INDEX_NONE)
+		, EndIndex(INDEX_NONE)
+	{
+
+	}
+
+	FTextRange(int32 InBeginIndex, int32 InEndIndex)
+		: BeginIndex(InBeginIndex)
+		, EndIndex(InEndIndex)
+	{
+
+	}
+
+	FORCEINLINE bool operator==(const FTextRange& Other) const
+	{
+		return BeginIndex == Other.BeginIndex
+			&& EndIndex == Other.EndIndex;
+	}
+
+	FORCEINLINE bool operator!=(const FTextRange& Other) const
+	{
+		return !(*this == Other);
+	}
+
+	friend inline uint32 GetTypeHash(const FTextRange& Key)
+	{
+		uint32 KeyHash = 0;
+		KeyHash = HashCombine(KeyHash, GetTypeHash(Key.BeginIndex));
+		KeyHash = HashCombine(KeyHash, GetTypeHash(Key.EndIndex));
+		return KeyHash;
+	}
+
+	int32 Len() const { return EndIndex - BeginIndex; }
+	bool IsEmpty() const { return (EndIndex - BeginIndex) <= 0; }
+	void Offset(int32 Amount) { BeginIndex += Amount; BeginIndex = FMath::Max(0, BeginIndex);  EndIndex += Amount; EndIndex = FMath::Max(0, EndIndex); }
+	bool Contains(int32 Index) const { return Index >= BeginIndex && Index < EndIndex; }
+	bool InclusiveContains(int32 Index) const { return Index >= BeginIndex && Index <= EndIndex; }
+
+	FTextRange Intersect(const FTextRange& Other) const
+	{
+		FTextRange Intersected(FMath::Max(BeginIndex, Other.BeginIndex), FMath::Min(EndIndex, Other.EndIndex));
+		if (Intersected.EndIndex <= Intersected.BeginIndex)
+		{
+			return FTextRange(0, 0);
+		}
+
+		return Intersected;
+	}
+
+	/**
+	 * Produce an array of line ranges from the given text, breaking at any new-line characters
+	 */
+	static void CalculateLineRangesFromString(const FString& Input, TArray<FTextRange>& LineRanges);
+
+	int32 BeginIndex;
+	int32 EndIndex;
+};
+
 #include "Misc/StringFormatArg.h"
diff --git a/Engine/Source/Runtime/Engine/Public/EngineSharedPCH.h b/Engine/Source/Runtime/Engine/Public/EngineSharedPCH.h
--- a/Engine/Source/Runtime/Engine/Public/EngineSharedPCH.h
+++ b/Engine/Source/Runtime/Engine/Public/EngineSharedPCH.h
@@ -398,7 +398,6 @@
 #include "Widgets/Input/IVirtualKeyboardEntry.h"
 #include "Widgets/Layout/SBorder.h"
 #include "Framework/Text/TextLayout.h"
-#include "Framework/Text/TextRange.h"
 #include "Framework/Text/TextRunRenderer.h"
 #include "Framework/Text/TextLineHighlight.h"
 #include "Framework/Text/IRun.h"
diff --git a/Engine/Source/Runtime/Slate/Private/Framework/Text/PlainTextLayoutMarshaller.cpp b/Engine/Source/Runtime/Slate/Private/Framework/Text/PlainTextLayoutMarshaller.cpp
--- a/Engine/Source/Runtime/Slate/Private/Framework/Text/PlainTextLayoutMarshaller.cpp
+++ b/Engine/Source/Runtime/Slate/Private/Framework/Text/PlainTextLayoutMarshaller.cpp
@@ -1,7 +1,6 @@
 // Copyright 1998-2019 Epic Games, Inc. All Rights Reserved.
 
 #include "Framework/Text/PlainTextLayoutMarshaller.h"
-#include "Framework/Text/TextRange.h"
 #include "Framework/Text/TextLineHighlight.h"
 #include "Framework/Text/IRun.h"
 #include "Framework/Text/TextLayout.h"
diff --git a/Engine/Source/Runtime/Slate/Private/Framework/Text/TextRange.cpp b/Engine/Source/Runtime/Slate/Private/Framework/Text/TextRange.cpp
deleted file mode 100644
diff --git a/Engine/Source/Runtime/Slate/Public/Framework/Text/DefaultLayoutBlock.h b/Engine/Source/Runtime/Slate/Public/Framework/Text/DefaultLayoutBlock.h
--- a/Engine/Source/Runtime/Slate/Public/Framework/Text/DefaultLayoutBlock.h
+++ b/Engine/Source/Runtime/Slate/Public/Framework/Text/DefaultLayoutBlock.h
@@ -2,7 +2,6 @@
 #pragma once
 
 #include "CoreMinimal.h"
-#include "Framework/Text/TextRange.h"
 #include "Framework/Text/IRun.h"
 #include "Framework/Text/ILayoutBlock.h"
 
diff --git a/Engine/Source/Runtime/Slate/Public/Framework/Text/ILayoutBlock.h b/Engine/Source/Runtime/Slate/Public/Framework/Text/ILayoutBlock.h
--- a/Engine/Source/Runtime/Slate/Public/Framework/Text/ILayoutBlock.h
+++ b/Engine/Source/Runtime/Slate/Public/Framework/Text/ILayoutBlock.h
@@ -2,7 +2,6 @@
 #pragma once
 
 #include "CoreMinimal.h"
-#include "Framework/Text/TextRange.h"
 #include "Framework/Text/IRun.h"
 #include "Framework/Text/IRunRenderer.h"
 
diff --git a/Engine/Source/Runtime/Slate/Public/Framework/Text/IRun.h b/Engine/Source/Runtime/Slate/Public/Framework/Text/IRun.h
--- a/Engine/Source/Runtime/Slate/Public/Framework/Text/IRun.h
+++ b/Engine/Source/Runtime/Slate/Public/Framework/Text/IRun.h
@@ -4,7 +4,6 @@
 #include "CoreMinimal.h"
 #include "Framework/Text/ShapedTextCacheFwd.h"
 #include "Misc/EnumClassFlags.h"
-#include "Framework/Text/TextRange.h"
 
 class ILayoutBlock;
 class IRunRenderer;
diff --git a/Engine/Source/Runtime/Slate/Public/Framework/Text/ITextDecorator.h b/Engine/Source/Runtime/Slate/Public/Framework/Text/ITextDecorator.h
--- a/Engine/Source/Runtime/Slate/Public/Framework/Text/ITextDecorator.h
+++ b/Engine/Source/Runtime/Slate/Public/Framework/Text/ITextDecorator.h
@@ -3,7 +3,6 @@
 
 #include "CoreMinimal.h"
 #include "SlateGlobals.h"
-#include "Framework/Text/TextRange.h"
 #include "Framework/Text/IRun.h"
 #include "Framework/Text/ISlateRun.h"
 
diff --git a/Engine/Source/Runtime/Slate/Public/Framework/Text/RichTextMarkupProcessing.h b/Engine/Source/Runtime/Slate/Public/Framework/Text/RichTextMarkupProcessing.h
--- a/Engine/Source/Runtime/Slate/Public/Framework/Text/RichTextMarkupProcessing.h
+++ b/Engine/Source/Runtime/Slate/Public/Framework/Text/RichTextMarkupProcessing.h
@@ -3,7 +3,6 @@
 
 #include "CoreMinimal.h"
 #include "SlateGlobals.h"
-#include "Framework/Text/TextRange.h"
 #include "Framework/Text/ITextDecorator.h"
 #include "Internationalization/Regex.h"
 #include "Framework/Text/IRichTextMarkupParser.h"
diff --git a/Engine/Source/Runtime/Slate/Public/Framework/Text/RunUtils.h b/Engine/Source/Runtime/Slate/Public/Framework/Text/RunUtils.h
--- a/Engine/Source/Runtime/Slate/Public/Framework/Text/RunUtils.h
+++ b/Engine/Source/Runtime/Slate/Public/Framework/Text/RunUtils.h
@@ -2,7 +2,6 @@
 #pragma once
 
 #include "CoreMinimal.h"
-#include "Framework/Text/TextRange.h"
 
 enum class ETextHitPoint : uint8;
 
diff --git a/Engine/Source/Runtime/Slate/Public/Framework/Text/ShapedTextCache.h b/Engine/Source/Runtime/Slate/Public/Framework/Text/ShapedTextCache.h
--- a/Engine/Source/Runtime/Slate/Public/Framework/Text/ShapedTextCache.h
+++ b/Engine/Source/Runtime/Slate/Public/Framework/Text/ShapedTextCache.h
@@ -5,7 +5,6 @@
 #include "CoreMinimal.h"
 #include "Fonts/ShapedTextFwd.h"
 #include "Fonts/SlateFontInfo.h"
-#include "Framework/Text/TextRange.h"
 #include "Framework/Text/IRun.h"
 
 class FSlateFontCache;
diff --git a/Engine/Source/Runtime/Slate/Public/Framework/Text/SlateHyperlinkRun.h b/Engine/Source/Runtime/Slate/Public/Framework/Text/SlateHyperlinkRun.h
--- a/Engine/Source/Runtime/Slate/Public/Framework/Text/SlateHyperlinkRun.h
+++ b/Engine/Source/Runtime/Slate/Public/Framework/Text/SlateHyperlinkRun.h
@@ -5,7 +5,6 @@
 #include "SlateGlobals.h"
 #include "Widgets/SWidget.h"
 #include "Styling/SlateTypes.h"
-#include "Framework/Text/TextRange.h"
 #include "Framework/Text/IRun.h"
 #include "Framework/Text/TextLayout.h"
 #include "Widgets/IToolTip.h"
diff --git a/Engine/Source/Runtime/Slate/Public/Framework/Text/SlateImageRun.h b/Engine/Source/Runtime/Slate/Public/Framework/Text/SlateImageRun.h
--- a/Engine/Source/Runtime/Slate/Public/Framework/Text/SlateImageRun.h
+++ b/Engine/Source/Runtime/Slate/Public/Framework/Text/SlateImageRun.h
@@ -4,7 +4,6 @@
 #include "CoreMinimal.h"
 #include "SlateGlobals.h"
 #include "Widgets/SWidget.h"
-#include "Framework/Text/TextRange.h"
 #include "Framework/Text/IRun.h"
 #include "Framework/Text/TextLayout.h"
 #include "Framework/Text/ILayoutBlock.h"
diff --git a/Engine/Source/Runtime/Slate/Public/Framework/Text/SlateTextRun.h b/Engine/Source/Runtime/Slate/Public/Framework/Text/SlateTextRun.h
--- a/Engine/Source/Runtime/Slate/Public/Framework/Text/SlateTextRun.h
+++ b/Engine/Source/Runtime/Slate/Public/Framework/Text/SlateTextRun.h
@@ -4,7 +4,6 @@
 #include "CoreMinimal.h"
 #include "Widgets/SWidget.h"
 #include "Styling/SlateTypes.h"
-#include "Framework/Text/TextRange.h"
 #include "Framework/Text/IRun.h"
 #include "Framework/Text/TextLayout.h"
 #include "Framework/Text/ILayoutBlock.h"
diff --git a/Engine/Source/Runtime/Slate/Public/Framework/Text/SlateWidgetRun.h b/Engine/Source/Runtime/Slate/Public/Framework/Text/SlateWidgetRun.h
--- a/Engine/Source/Runtime/Slate/Public/Framework/Text/SlateWidgetRun.h
+++ b/Engine/Source/Runtime/Slate/Public/Framework/Text/SlateWidgetRun.h
@@ -4,7 +4,6 @@
 #include "CoreMinimal.h"
 #include "SlateGlobals.h"
 #include "Widgets/SWidget.h"
-#include "Framework/Text/TextRange.h"
 #include "Framework/Text/IRun.h"
 #include "Framework/Text/TextLayout.h"
 #include "Framework/Text/ILayoutBlock.h"
diff --git a/Engine/Source/Runtime/Slate/Public/Framework/Text/SyntaxTokenizer.h b/Engine/Source/Runtime/Slate/Public/Framework/Text/SyntaxTokenizer.h
--- a/Engine/Source/Runtime/Slate/Public/Framework/Text/SyntaxTokenizer.h
+++ b/Engine/Source/Runtime/Slate/Public/Framework/Text/SyntaxTokenizer.h
@@ -2,7 +2,6 @@
 #pragma once
 
 #include "CoreMinimal.h"
-#include "Framework/Text/TextRange.h"
 
 /**
  * Tokenize the text based upon the given rule set
diff --git a/Engine/Source/Runtime/Slate/Public/Framework/Text/TextLayout.h b/Engine/Source/Runtime/Slate/Public/Framework/Text/TextLayout.h
--- a/Engine/Source/Runtime/Slate/Public/Framework/Text/TextLayout.h
+++ b/Engine/Source/Runtime/Slate/Public/Framework/Text/TextLayout.h
@@ -5,7 +5,6 @@
 #include "UObject/ObjectMacros.h"
 #include "Misc/Attribute.h"
 #include "Layout/Margin.h"
-#include "Framework/Text/TextRange.h"
 #include "Framework/Text/TextRunRenderer.h"
 #include "Framework/Text/TextLineHighlight.h"
 #include "Framework/Text/IRun.h"
diff --git a/Engine/Source/Runtime/Slate/Public/Framework/Text/TextLayoutEngine.h b/Engine/Source/Runtime/Slate/Public/Framework/Text/TextLayoutEngine.h
--- a/Engine/Source/Runtime/Slate/Public/Framework/Text/TextLayoutEngine.h
+++ b/Engine/Source/Runtime/Slate/Public/Framework/Text/TextLayoutEngine.h
@@ -5,7 +5,6 @@
 MONOLITHIC_HEADER_BOILERPLATE()
 
 // Generic Text Layout
-#include "Framework/Text/TextRange.h"
 #include "Framework/Text/TextRunRenderer.h"
 #include "Framework/Text/TextLineHighlight.h"
 
diff --git a/Engine/Source/Runtime/Slate/Public/Framework/Text/TextLineHighlight.h b/Engine/Source/Runtime/Slate/Public/Framework/Text/TextLineHighlight.h
--- a/Engine/Source/Runtime/Slate/Public/Framework/Text/TextLineHighlight.h
+++ b/Engine/Source/Runtime/Slate/Public/Framework/Text/TextLineHighlight.h
@@ -2,7 +2,6 @@
 #pragma once
 
 #include "CoreMinimal.h"
-#include "Framework/Text/TextRange.h"
 
 class ILineHighlighter;
 
diff --git a/Engine/Source/Runtime/Slate/Public/Framework/Text/TextRange.h b/Engine/Source/Runtime/Slate/Public/Framework/Text/TextRange.h
--- a/Engine/Source/Runtime/Slate/Public/Framework/Text/TextRange.h
+++ b/Engine/Source/Runtime/Slate/Public/Framework/Text/TextRange.h
@@ -1,65 +1,14 @@
 // Copyright 1998-2019 Epic Games, Inc. All Rights Reserved.
 #pragma once
 
-#include "CoreMinimal.h"
+#include "Containers/UnrealString.h"
 
-struct SLATE_API FTextRange
+UE_DEPRECATED(4.23, "Including TextRange.h has been deprecated. Please include UnrealString.h in the Core module instead.")
+inline void TextRangeHeaderDeprecatedWarning()
 {
-	FTextRange()
-		: BeginIndex(INDEX_NONE)
-		, EndIndex(INDEX_NONE)
-	{
+}
 
-	}
-
-	FTextRange( int32 InBeginIndex, int32 InEndIndex )
-		: BeginIndex( InBeginIndex )
-		, EndIndex( InEndIndex )
-	{
-
-	}
-
-	FORCEINLINE bool operator==(const FTextRange& Other) const
-	{
-		return BeginIndex == Other.BeginIndex 
-			&& EndIndex == Other.EndIndex;
-	}
-
-	FORCEINLINE bool operator!=(const FTextRange& Other) const
-	{
-		return !(*this == Other);
-	}
-
-	friend inline uint32 GetTypeHash(const FTextRange& Key)
-	{
-		uint32 KeyHash = 0;
-		KeyHash = HashCombine(KeyHash, GetTypeHash(Key.BeginIndex));
-		KeyHash = HashCombine(KeyHash, GetTypeHash(Key.EndIndex));
-		return KeyHash;
-	}
-
-	int32 Len() const { return EndIndex - BeginIndex; }
-	bool IsEmpty() const { return (EndIndex - BeginIndex) <= 0; }
-	void Offset(int32 Amount) { BeginIndex += Amount; BeginIndex = FMath::Max(0, BeginIndex);  EndIndex += Amount; EndIndex = FMath::Max(0, EndIndex); }
-	bool Contains(int32 Index) const { return Index >= BeginIndex && Index < EndIndex; }
-	bool InclusiveContains(int32 Index) const { return Index >= BeginIndex && Index <= EndIndex; }
-
-	FTextRange Intersect(const FTextRange& Other) const
-	{
-		FTextRange Intersected(FMath::Max(BeginIndex, Other.BeginIndex), FMath::Min(EndIndex, Other.EndIndex));
-		if (Intersected.EndIndex <= Intersected.BeginIndex)
-		{
-			return FTextRange(0, 0);
-		}
-
-		return Intersected;
-	}
-
-	/**
-	 * Produce an array of line ranges from the given text, breaking at any new-line characters
-	 */
-	static void CalculateLineRangesFromString(const FString& Input, TArray<FTextRange>& LineRanges);
-
-	int32 BeginIndex;
-	int32 EndIndex;
-};
+inline void TriggerTextRangeHeaderDeprecatedWarning()
+{
+	TextRangeHeaderDeprecatedWarning();
+}
diff --git a/Engine/Source/Runtime/Slate/Public/Framework/Text/TextRunRenderer.h b/Engine/Source/Runtime/Slate/Public/Framework/Text/TextRunRenderer.h
--- a/Engine/Source/Runtime/Slate/Public/Framework/Text/TextRunRenderer.h
+++ b/Engine/Source/Runtime/Slate/Public/Framework/Text/TextRunRenderer.h
@@ -2,7 +2,6 @@
 #pragma once
 
 #include "CoreMinimal.h"
-#include "Framework/Text/TextRange.h"
 
 class IRunRenderer;
 
diff --git a/Engine/Source/Runtime/Slate/Public/Framework/Text/WidgetLayoutBlock.h b/Engine/Source/Runtime/Slate/Public/Framework/Text/WidgetLayoutBlock.h
--- a/Engine/Source/Runtime/Slate/Public/Framework/Text/WidgetLayoutBlock.h
+++ b/Engine/Source/Runtime/Slate/Public/Framework/Text/WidgetLayoutBlock.h
@@ -3,7 +3,6 @@
 
 #include "CoreMinimal.h"
 #include "Widgets/SWidget.h"
-#include "Framework/Text/TextRange.h"
 #include "Framework/Text/IRun.h"
 #include "Framework/Text/ILayoutBlock.h"
 
diff --git a/Engine/Source/Runtime/Slate/Public/SlateBasics.h b/Engine/Source/Runtime/Slate/Public/SlateBasics.h
--- a/Engine/Source/Runtime/Slate/Public/SlateBasics.h
+++ b/Engine/Source/Runtime/Slate/Public/SlateBasics.h
@@ -36,7 +36,6 @@ MONOLITHIC_HEADER_BOILERPLATE()
 
 // Legacy
 #include "Widgets/SWeakWidget.h"
-#include "Framework/Text/TextRange.h"
 #include "Framework/Text/TextRunRenderer.h"
 #include "Framework/Text/TextLineHighlight.h"
 #include "Framework/Text/TextHitPoint.h"
diff --git a/Engine/Source/Runtime/Slate/Public/SlateSharedPCH.h b/Engine/Source/Runtime/Slate/Public/SlateSharedPCH.h
--- a/Engine/Source/Runtime/Slate/Public/SlateSharedPCH.h
+++ b/Engine/Source/Runtime/Slate/Public/SlateSharedPCH.h
@@ -333,7 +333,6 @@
 #include "Framework/SlateDelegates.h"
 #include "Widgets/Layout/SBorder.h"
 #include "Framework/Text/TextLayout.h"
-#include "Framework/Text/TextRange.h"
 #include "Framework/Text/TextRunRenderer.h"
 #include "Framework/Text/TextLineHighlight.h"
 #include "Framework/Text/IRun.h"
diff --git a/Engine/Source/Runtime/Slate/Public/Widgets/Text/SlateEditableTextLayout.h b/Engine/Source/Runtime/Slate/Public/Widgets/Text/SlateEditableTextLayout.h
--- a/Engine/Source/Runtime/Slate/Public/Widgets/Text/SlateEditableTextLayout.h
+++ b/Engine/Source/Runtime/Slate/Public/Widgets/Text/SlateEditableTextLayout.h
@@ -13,7 +13,6 @@
 #include "Widgets/Input/IVirtualKeyboardEntry.h"
 #include "Widgets/Text/ISlateEditableTextWidget.h"
 #include "Framework/Text/ITextLayoutMarshaller.h"
-#include "Framework/Text/TextRange.h"
 #include "Framework/Text/TextLineHighlight.h"
 #include "Framework/Text/IRun.h"
 #include "Framework/Text/TextLayout.h"
