Stopped culture invariant text from being keyed during save, as it can cause non-deterministic cooking issues

Culture invariant text is now just serialized directly as its source string, and re-generated on load.

This also fixes some more places where an empty string could be used to generate an FText instance, rather than using the shared empty instance.

diff --git a/Engine/Source/Runtime/Core/Private/Internationalization/Text.cpp b/Engine/Source/Runtime/Core/Private/Internationalization/Text.cpp
--- a/Engine/Source/Runtime/Core/Private/Internationalization/Text.cpp
+++ b/Engine/Source/Runtime/Core/Private/Internationalization/Text.cpp
@@ -244,10 +244,10 @@ FText::FText( TSharedRef<ITextData, ESPMode::ThreadSafe> InTextData )
 }
 
 FText::FText( FString&& InSourceString )
-	: TextData(new TGeneratedTextData<FTextHistory_Base>(FString(InSourceString)))
+	: TextData(new TGeneratedTextData<FTextHistory_Base>(CopyTemp(InSourceString))) // Copy the source string as the live display string
 	, Flags(0)
 {
-	TextData->SetTextHistory(FTextHistory_Base(MoveTemp(InSourceString)));
+	TextData->SetTextHistory(FTextHistory_Base(MoveTemp(InSourceString))); // Move the source string as the historic source string
 }
 
 FText::FText( FName InTableId, FString InKey, const EStringTableLoadingPolicy InLoadingPolicy )
@@ -320,119 +320,91 @@ FText FText::ToUpper() const
 
 FText FText::TrimPreceding( const FText& InText )
 {
-	FString TrimmedString = InText.ToString();
+	const FString& CurrentString = InText.ToString();
+
+	int32 StartPos = 0;
+	while (StartPos < CurrentString.Len())
 	{
-		int32 StartPos = 0;
-		while ( StartPos < TrimmedString.Len() )
+		if (!FText::IsWhitespace(CurrentString[StartPos]))
 		{
-			if( !FText::IsWhitespace( TrimmedString[StartPos] ) )
-			{
-				break;
-			}
-
-			++StartPos;
+			break;
 		}
 
-		TrimmedString = TrimmedString.Right( TrimmedString.Len() - StartPos );
+		++StartPos;
 	}
 
-	FText NewText = FText( MoveTemp( TrimmedString ) );
-
-	if (!GIsEditor)
+	if (StartPos == 0)
 	{
-		if( (NewText.Flags & ETextFlag::CultureInvariant) != 0 )
-		{
-			NewText.Flags |= ETextFlag::Transient;
-		}
-		else
-		{
-			NewText.Flags |= ETextFlag::CultureInvariant;
-		}
+		// Nothing to trim!
+		return InText;
 	}
 
-	return NewText;
+	// Trim the string, preserving culture invariance if set
+	FString TrimmedString = CurrentString.Right(TrimmedString.Len() - StartPos);
+	return InText.IsCultureInvariant() ? FText::AsCultureInvariant(MoveTemp(TrimmedString)) : FText::FromString(MoveTemp(TrimmedString));
 }
 
 FText FText::TrimTrailing( const FText& InText )
 {
-	FString TrimmedString = InText.ToString();
+	const FString& CurrentString = InText.ToString();
+
+	int32 EndPos = CurrentString.Len() - 1;
+	while (EndPos >= 0)
 	{
-		int32 EndPos = TrimmedString.Len() - 1;
-		while( EndPos >= 0 )
+		if (!FText::IsWhitespace(CurrentString[EndPos]))
 		{
-			if( !FText::IsWhitespace( TrimmedString[EndPos] ) )
-			{
-				break;
-			}
-
-			EndPos--;
+			break;
 		}
 
-		TrimmedString = TrimmedString.Left( EndPos + 1 );
+		--EndPos;
 	}
 
-	FText NewText = FText( MoveTemp ( TrimmedString ) );
-
-	if (!GIsEditor)
+	if (EndPos == CurrentString.Len() - 1)
 	{
-		if( (NewText.Flags & ETextFlag::CultureInvariant) != 0 )
-		{
-			NewText.Flags |= ETextFlag::Transient;
-		}
-		else
-		{
-			NewText.Flags |= ETextFlag::CultureInvariant;
-		}
+		// Nothing to trim!
+		return InText;
 	}
 
-	return NewText;
+	// Trim the string, preserving culture invariance if set
+	FString TrimmedString = CurrentString.Left(EndPos + 1);
+	return InText.IsCultureInvariant() ? FText::AsCultureInvariant(MoveTemp(TrimmedString)) : FText::FromString(MoveTemp(TrimmedString));
 }
 
 FText FText::TrimPrecedingAndTrailing( const FText& InText )
 {
-	FString TrimmedString = InText.ToString();
+	const FString& CurrentString = InText.ToString();
+
+	int32 StartPos = 0;
+	while (StartPos < CurrentString.Len())
 	{
-		int32 StartPos = 0;
-		while ( StartPos < TrimmedString.Len() )
+		if (!FText::IsWhitespace(CurrentString[StartPos]))
 		{
-			if( !FText::IsWhitespace( TrimmedString[StartPos] ) )
-			{
-				break;
-			}
-
-			++StartPos;
+			break;
 		}
 
-		int32 EndPos = TrimmedString.Len();
-		while( EndPos > StartPos )
-		{
-			if( !FText::IsWhitespace( TrimmedString[EndPos - 1] ) )
-			{
-				break;
-			}
+		++StartPos;
+	}
 
-			--EndPos;
+	int32 EndPos = CurrentString.Len() - 1;
+	while (EndPos > StartPos)
+	{
+		if (!FText::IsWhitespace(CurrentString[EndPos]))
+		{
+			break;
 		}
 
-		const int32 Len = EndPos - StartPos;
-		TrimmedString = TrimmedString.Mid( StartPos, Len );
+		--EndPos;
 	}
 
-	FText NewText = FText( MoveTemp( TrimmedString ) );
-
-	if (!GIsEditor)
+	if (StartPos == 0 && EndPos == CurrentString.Len() - 1)
 	{
-		if( (NewText.Flags & ETextFlag::CultureInvariant) != 0 )
-		{
-			NewText.Flags |= ETextFlag::Transient;
-		}
-		else
-		{
-			NewText.Flags |= ETextFlag::CultureInvariant;
-		}
+		// Nothing to trim!
+		return InText;
 	}
 
-	return NewText;
+	// Trim the string, preserving culture invariance if set
+	FString TrimmedString = CurrentString.Mid(StartPos, EndPos - StartPos + 1);
+	return InText.IsCultureInvariant() ? FText::AsCultureInvariant(MoveTemp(TrimmedString)) : FText::FromString(MoveTemp(TrimmedString));
 }
 
 void FText::GetFormatPatternParameters(const FTextFormat& Fmt, TArray<FString>& ParameterNames)
@@ -771,6 +743,8 @@ void FText::SerializeText(FStructuredArchive::FSlot Slot, FText& Value)
 	FArchive& UnderlyingArchive = Slot.GetUnderlyingArchive();
 	FStructuredArchive::FRecord Record = Slot.EnterRecord();
 
+	UnderlyingArchive.UsingCustomVersion(FEditorObjectVersion::GUID);
+
 	//When duplicating, the CDO is used as the template, then values for the instance are assigned.
 	//If we don't duplicate the string, the CDO and the instance are both pointing at the same thing.
 	//This would result in all subsequently duplicated objects stamping over formerly duplicated ones.
@@ -847,12 +821,20 @@ void FText::SerializeText(FStructuredArchive::FSlot Slot, FText& Value)
 		if (UnderlyingArchive.IsSaving())
 		{
 			// Skip the history for empty texts
-			bSerializeHistory = !Value.IsEmpty();
+			bSerializeHistory = !Value.IsEmpty() && !Value.IsCultureInvariant();
 
 			if (!bSerializeHistory)
 			{
 				int8 HistoryType = (int8)ETextHistoryType::None;
 				Record << NAMED_FIELD(HistoryType);
+
+				bool bHasCultureInvariantString = !Value.IsEmpty() && Value.IsCultureInvariant();
+				Record << NAMED_FIELD(bHasCultureInvariantString);
+				if (bHasCultureInvariantString)
+				{
+					FString CultureInvariantString = Value.GetSourceString();
+					Record << NAMED_FIELD(CultureInvariantString);
+				}
 			}
 		}
 		else if (UnderlyingArchive.IsLoading())
@@ -933,6 +915,18 @@ void FText::SerializeText(FStructuredArchive::FSlot Slot, FText& Value)
 				{
 					bSerializeHistory = false;
 					Value.TextData = FText::GetEmpty().TextData;
+
+					if (UnderlyingArchive.CustomVer(FEditorObjectVersion::GUID) >= FEditorObjectVersion::CultureInvariantTextSerializationKeyStability)
+					{
+						bool bHasCultureInvariantString = false;
+						Record << NAMED_FIELD(bHasCultureInvariantString);
+						if (bHasCultureInvariantString)
+						{
+							FString CultureInvariantString;
+							Record << NAMED_FIELD(CultureInvariantString);
+							Value.TextData = FText(MoveTemp(CultureInvariantString)).TextData;
+						}
+					}
 				}
 			}
 		}
@@ -985,7 +979,7 @@ FText FText::FromName( const FName& Val)
 
 FText FText::FromString( const FString& String )
 {
-	FText NewText = String.IsEmpty() ? FText::GetEmpty() : FText( FString(String) );
+	FText NewText = String.IsEmpty() ? FText::GetEmpty() : FText(CopyTemp(String));
 
 	if (!GIsEditor)
 	{
@@ -998,7 +992,7 @@ FText FText::FromString( const FString& String )
 
 FText FText::FromString( FString&& String )
 {
-	FText NewText = String.IsEmpty() ? FText::GetEmpty() : FText( MoveTemp(String) );
+	FText NewText = String.IsEmpty() ? FText::GetEmpty() : FText(MoveTemp(String));
 
 	if (!GIsEditor)
 	{
@@ -1011,7 +1005,7 @@ FText FText::FromString( FString&& String )
 
 FText FText::AsCultureInvariant( const FString& String )
 {
-	FText NewText = String.IsEmpty() ? FText::GetEmpty() : FText( FString(String) );
+	FText NewText = String.IsEmpty() ? FText::GetEmpty() : FText(CopyTemp(String));
 	NewText.Flags |= ETextFlag::CultureInvariant;
 
 	return NewText;
@@ -1019,7 +1013,7 @@ FText FText::AsCultureInvariant( const FString& String )
 
 FText FText::AsCultureInvariant( FString&& String )
 {
-	FText NewText = String.IsEmpty() ? FText() : FText( MoveTemp(String) );
+	FText NewText = String.IsEmpty() ? FText::GetEmpty() : FText(MoveTemp(String));
 	NewText.Flags |= ETextFlag::CultureInvariant;
 
 	return NewText;
@@ -1027,7 +1021,7 @@ FText FText::AsCultureInvariant( FString&& String )
 
 FText FText::AsCultureInvariant( FText Text )
 {
-	FText NewText = FText( MoveTemp(Text) );
+	FText NewText = FText(MoveTemp(Text));
 	NewText.Flags |= ETextFlag::CultureInvariant;
 
 	return NewText;
diff --git a/Engine/Source/Runtime/Core/Private/Internationalization/TextData.h b/Engine/Source/Runtime/Core/Private/Internationalization/TextData.h
--- a/Engine/Source/Runtime/Core/Private/Internationalization/TextData.h
+++ b/Engine/Source/Runtime/Core/Private/Internationalization/TextData.h
@@ -158,7 +158,7 @@ public:
 			if (!this->LocalizedString.IsValid())
 			{
 				// We copy (rather than move) DisplayString here, as other threads may currently be accessing it
-				this->LocalizedString = MakeShareable(new FString(DisplayString));
+				this->LocalizedString = MakeShared<FString, ESPMode::ThreadSafe>(DisplayString);
 			}
 		}
 	}
diff --git a/Engine/Source/Runtime/Core/Private/Internationalization/TextHistory.cpp b/Engine/Source/Runtime/Core/Private/Internationalization/TextHistory.cpp
--- a/Engine/Source/Runtime/Core/Private/Internationalization/TextHistory.cpp
+++ b/Engine/Source/Runtime/Core/Private/Internationalization/TextHistory.cpp
@@ -836,34 +836,36 @@ void FTextHistory_Base::SerializeForDisplayString(FStructuredArchive::FRecord Re
 			TextNamespaceUtil::StripPackageNamespaceInline(Namespace);
 		}
 		else
-#if USE_STABLE_LOCALIZATION_KEYS
-		// Make sure the package namespace for this text property is up-to-date
-		if (GIsEditor && !BaseArchive.HasAnyPortFlags(PPF_DuplicateVerbatim | PPF_DuplicateForPIE))
 		{
-			const FString PackageNamespace = TextNamespaceUtil::GetPackageNamespace(BaseArchive);
-			if (!PackageNamespace.IsEmpty())
+#if USE_STABLE_LOCALIZATION_KEYS
+			// Make sure the package namespace for this text property is up-to-date
+			if (GIsEditor && !BaseArchive.HasAnyPortFlags(PPF_DuplicateVerbatim | PPF_DuplicateForPIE))
 			{
-				const FString FullNamespace = TextNamespaceUtil::BuildFullNamespace(Namespace, PackageNamespace);
-				if (!Namespace.Equals(FullNamespace, ESearchCase::CaseSensitive))
+				const FString PackageNamespace = TextNamespaceUtil::GetPackageNamespace(BaseArchive);
+				if (!PackageNamespace.IsEmpty())
 				{
-					// We may assign a new key when saving if we don't have the correct package namespace in order to avoid identity conflicts when instancing (which duplicates without any special flags)
-					// This can happen if an asset was duplicated (and keeps the same keys) but later both assets are instanced into the same world (causing them to both take the worlds package id, and conflict with each other)
-					Namespace = FullNamespace;
-					Key = FGuid::NewGuid().ToString();
+					const FString FullNamespace = TextNamespaceUtil::BuildFullNamespace(Namespace, PackageNamespace);
+					if (!Namespace.Equals(FullNamespace, ESearchCase::CaseSensitive))
+					{
+						// We may assign a new key when saving if we don't have the correct package namespace in order to avoid identity conflicts when instancing (which duplicates without any special flags)
+						// This can happen if an asset was duplicated (and keeps the same keys) but later both assets are instanced into the same world (causing them to both take the worlds package id, and conflict with each other)
+						Namespace = FullNamespace;
+						Key = FGuid::NewGuid().ToString();
+					}
 				}
 			}
-		}
 #endif // USE_STABLE_LOCALIZATION_KEYS
 
-		// If this has no key, give it a GUID for a key
-		if (!bFoundNamespaceAndKey && GIsEditor && (BaseArchive.IsPersistent() && !BaseArchive.HasAnyPortFlags(PPF_Duplicate)))
-		{
-			Key = FGuid::NewGuid().ToString();
-			if (!FTextLocalizationManager::Get().AddDisplayString(InOutDisplayString.ToSharedRef(), Namespace, Key))
+			// If this has no key, give it a GUID for a key
+			if (!bFoundNamespaceAndKey && GIsEditor && (BaseArchive.IsPersistent() && !BaseArchive.HasAnyPortFlags(PPF_Duplicate)))
 			{
-				// Could not add display string, reset namespace and key.
-				Namespace.Empty();
-				Key.Empty();
+				Key = FGuid::NewGuid().ToString();
+				if (!FTextLocalizationManager::Get().AddDisplayString(InOutDisplayString.ToSharedRef(), Namespace, Key))
+				{
+					// Could not add display string, reset namespace and key.
+					Namespace.Empty();
+					Key.Empty();
+				}
 			}
 		}
 
diff --git a/Engine/Source/Runtime/Core/Public/Internationalization/Text.h b/Engine/Source/Runtime/Core/Public/Internationalization/Text.h
--- a/Engine/Source/Runtime/Core/Public/Internationalization/Text.h
+++ b/Engine/Source/Runtime/Core/Public/Internationalization/Text.h
@@ -498,17 +498,17 @@ public:
 	FText ToUpper() const;
 
 	/**
-	 * Removes whitespace characters from the front of the string.
+	 * Removes any whitespace characters from the start of the text.
 	 */
 	static FText TrimPreceding( const FText& );
 
 	/**
-	 * Removes trailing whitespace characters
+	 * Removes any whitespace characters from the end of the text.
 	 */
 	static FText TrimTrailing( const FText& );
 
 	/**
-	 * Does both of the above without needing to create an additional FText in the interim.
+	 * Removes any whitespace characters from the start and end of the text.
 	 */
 	static FText TrimPrecedingAndTrailing( const FText& );
 
diff --git a/Engine/Source/Runtime/Core/Public/UObject/EditorObjectVersion.h b/Engine/Source/Runtime/Core/Public/UObject/EditorObjectVersion.h
--- a/Engine/Source/Runtime/Core/Public/UObject/EditorObjectVersion.h
+++ b/Engine/Source/Runtime/Core/Public/UObject/EditorObjectVersion.h
@@ -73,6 +73,8 @@ struct CORE_API FEditorObjectVersion
 		MeshDescriptionRemovedHoles,
 		// Change to the WidgetCompoent WindowVisibilty default value
 		ChangedWidgetComponentWindowVisibilityDefault,
+		// Avoid keying culture invariant display strings during serialization to avoid non-deterministic cooking issues
+		CultureInvariantTextSerializationKeyStability,
 		// -----<new versions can be added above this line>-------------------------------------------------
 		VersionPlusOne,
 		LatestVersion = VersionPlusOne - 1
