Fix FText::TrimPreceding bug.

#rnx

diff --git a/Engine/Source/Runtime/Core/Private/Internationalization/Text.cpp b/Engine/Source/Runtime/Core/Private/Internationalization/Text.cpp
--- a/Engine/Source/Runtime/Core/Private/Internationalization/Text.cpp
+++ b/Engine/Source/Runtime/Core/Private/Internationalization/Text.cpp
@@ -340,7 +340,7 @@ FText FText::TrimPreceding( const FText& InText )
 	}
 
 	// Trim the string, preserving culture invariance if set
-	FString TrimmedString = CurrentString.Right(TrimmedString.Len() - StartPos);
+	FString TrimmedString = CurrentString.Right(CurrentString.Len() - StartPos);
 	return InText.IsCultureInvariant() ? FText::AsCultureInvariant(MoveTemp(TrimmedString)) : FText::FromString(MoveTemp(TrimmedString));
 }
 
