Implement ThreadStackWalkAndDump and CaptureThreadStackBackTrace for Unix Platform

diff --git a/Engine/Source/Programs/TestPAL/Private/Main.cpp b/Engine/Source/Programs/TestPAL/Private/Main.cpp
--- a/Engine/Source/Programs/TestPAL/Private/Main.cpp
+++ b/Engine/Source/Programs/TestPAL/Private/Main.cpp
@@ -36,6 +36,7 @@ IMPLEMENT_APPLICATION(TestPAL, "TestPAL");
 #define ARG_INLINE_CALLSTACK_TEST			"inline"
 #define ARG_STRINGS_ALLOCATION_TEST			"stringsallocation"
 #define ARG_CREATEGUID_TEST					"createguid"
+#define ARG_THREADSTACK_TEST				"threadstack"
 
 namespace TestPAL
 {
@@ -1468,6 +1469,86 @@ int32 InlineCallstacksTest(const TCHAR* CommandLine)
 	return 0;
 }
 
+struct FThreadStackTest : public FRunnable
+{
+	FThreadStackTest(uint64 InThreadId = ~0) :
+		ThreadId(InThreadId)
+	{
+	}
+
+	uint64 ThreadId;
+
+	virtual uint32 Run()
+	{
+		// If we dont have a valid ThreadId lets use the threads id as default
+		if (ThreadId == ~0)
+		{
+			ThreadId = syscall(SYS_gettid);
+		}
+
+		// Hopefully this wont be to much extra stack space for a testing thread
+		SIZE_T StackTraceSize = 65536;
+		ANSICHAR StackTrace[StackTraceSize] = {0};
+		FPlatformStackWalk::ThreadStackWalkAndDump(StackTrace, StackTraceSize, 0, ThreadId);
+
+		UE_LOG(LogTestPAL, Warning, TEXT("***** ThreadStackWalkAndDump for ThreadId(%lu) ******\n%s"), ThreadId, ANSI_TO_TCHAR(StackTrace));
+
+		SIZE_T BackTraceSize = 100;
+		uint64 BackTrace[BackTraceSize];
+		int32 BackTraceCount = FPlatformStackWalk::CaptureThreadStackBackTrace(ThreadId, BackTrace, BackTraceSize);
+
+		UE_LOG(LogTestPAL, Warning, TEXT("***** CaptureThreadStackBackTrace for ThreadId(%lu) ******"), ThreadId);
+		for (int i = 0; i < BackTraceCount; i++)
+		{
+			UE_LOG(LogTestPAL, Warning, TEXT("0x%llx"), BackTrace[i]);
+		}
+		UE_LOG(LogTestPAL, Warning, TEXT("\n\n"));
+
+		UE_LOG(LogTestPAL, Warning, TEXT("***** ProgramCounterToHumanReadableString for BackTrace for ThreadId(%lu) ******"), ThreadId);
+		for (int i = 0; i < BackTraceCount; i++)
+		{
+			ANSICHAR TempString[1024] = {0};
+			FPlatformStackWalk::ProgramCounterToHumanReadableString(i, BackTrace[i], TempString, ARRAY_COUNT(TempString));
+
+			UE_LOG(LogTestPAL, Warning, TEXT("%s"), ANSI_TO_TCHAR(TempString));
+		}
+
+		return 0;
+	}
+};
+
+
+int32 ThreadTraceTest(const TCHAR* CommandLine)
+{
+	FPlatformMisc::SetCrashHandler(NULL);
+	FPlatformMisc::SetGracefulTerminationHandler();
+
+	GEngineLoop.PreInit(CommandLine);
+
+	// Dump the callstack/backtrace of the thread that is running
+	{
+		FThreadStackTest* ThreadStackTest = new FThreadStackTest;
+		FRunnableThread* Runnable = FRunnableThread::Create(ThreadStackTest, ANSI_TO_TCHAR("ThreadStackTest"), 0);
+
+		Runnable->WaitForCompletion();
+		delete ThreadStackTest;
+	}
+
+	// Dump the GGameThreadId callstack/backtrace from the thread
+	{
+		FThreadStackTest* ThreadStackTest = new FThreadStackTest(GGameThreadId);
+		FRunnableThread* Runnable = FRunnableThread::Create(ThreadStackTest, ANSI_TO_TCHAR("ThreadStackTest"), 0);
+
+		Runnable->WaitForCompletion();
+		delete ThreadStackTest;
+	}
+
+	FEngineLoop::AppPreExit();
+	FEngineLoop::AppExit();
+
+	return 0;
+}
+
 /**
  * Selects and runs one of test cases.
  *
@@ -1553,6 +1634,10 @@ int32 MultiplexedMain(int32 ArgC, char* ArgV[])
 		{
 			return CreateGuidTest(*TestPAL::CommandLine);
 		}
+		else if (!FCStringAnsi::Strcmp(ArgV[IdxArg], ARG_THREADSTACK_TEST))
+		{
+			return ThreadTraceTest(*TestPAL::CommandLine);
+		}
 	}
 
 	FPlatformMisc::SetCrashHandler(NULL);
@@ -1580,6 +1665,7 @@ int32 MultiplexedMain(int32 ArgC, char* ArgV[])
 	UE_LOG(LogTestPAL, Warning, TEXT("  %s: test inline callstacks through ensures and a final crash."), UTF8_TO_TCHAR(ARG_INLINE_CALLSTACK_TEST));
 	UE_LOG(LogTestPAL, Warning, TEXT("  %s: test string allocations."), UTF8_TO_TCHAR(ARG_STRINGS_ALLOCATION_TEST));
 	UE_LOG(LogTestPAL, Warning, TEXT("  %s: test CreateGuid."), UTF8_TO_TCHAR(ARG_CREATEGUID_TEST));
+	UE_LOG(LogTestPAL, Warning, TEXT("  %s: test ThreadWalkStackAndDump and CaptureThreadBackTrace."), UTF8_TO_TCHAR(ARG_THREADSTACK_TEST));
 	UE_LOG(LogTestPAL, Warning, TEXT(""));
 	UE_LOG(LogTestPAL, Warning, TEXT("Pass one of those to run an appropriate test."));
 
diff --git a/Engine/Source/Runtime/Core/Private/Unix/UnixPlatformCrashContext.cpp b/Engine/Source/Runtime/Core/Private/Unix/UnixPlatformCrashContext.cpp
--- a/Engine/Source/Runtime/Core/Private/Unix/UnixPlatformCrashContext.cpp
+++ b/Engine/Source/Runtime/Core/Private/Unix/UnixPlatformCrashContext.cpp
@@ -24,6 +24,7 @@
 #include "Misc/App.h"
 #include "Misc/EngineVersion.h"
 #include "HAL/PlatformMallocCrash.h"
+#include "Unix/UnixPlatformRealTimeSignals.h"
 #include "Unix/UnixPlatformRunnableThread.h"
 #include "HAL/ExceptionHandling.h"
 #include "Stats/Stats.h"
@@ -720,6 +721,27 @@ void PlatformCrashHandler(int32 Signal, siginfo_t* Info, void* Context)
 	}
 }
 
+void ThreadStackWalker(int32 Signal, siginfo_t* Info, void* Context)
+{
+	ThreadStackUserData* ThreadStackData = static_cast<ThreadStackUserData*>(Info->si_value.sival_ptr);
+
+	if (ThreadStackData)
+	{
+		if (ThreadStackData->bCaptureCallStack)
+		{
+			// One for the pthread frame and one for siqueue
+			int32 IgnoreCount = 2;
+			FPlatformStackWalk::StackWalkAndDump(ThreadStackData->CallStack, ThreadStackData->CallStackSize, IgnoreCount);
+		}
+		else
+		{
+			ThreadStackData->BackTraceCount = FPlatformStackWalk::CaptureStackBackTrace(ThreadStackData->BackTrace, ThreadStackData->CallStackSize);
+		}
+
+		ThreadStackData->bDone = true;
+	}
+}
+
 void FUnixPlatformMisc::SetGracefulTerminationHandler()
 {
 	struct sigaction Action;
@@ -810,6 +832,13 @@ void FUnixPlatformMisc::SetCrashHandler(void (* CrashHandler)(const FGenericCras
 		}
 	}
 
+	struct sigaction ActionForThread;
+	FMemory::Memzero(ActionForThread);
+	sigfillset(&ActionForThread.sa_mask);
+	ActionForThread.sa_flags = SA_SIGINFO | SA_RESTART | SA_ONSTACK;
+	ActionForThread.sa_sigaction = ThreadStackWalker;
+	sigaction(THREAD_CALLSTACK_GENERATOR, &ActionForThread, nullptr);
+
 	checkf(IsInGameThread(), TEXT("Crash handler for the game thread should be set from the game thread only."));
 
 	FRunnableThreadUnix::SetupSignalHandlerStack(FRunnableThreadUnix::MainThreadSignalHandlerStack, sizeof(FRunnableThreadUnix::MainThreadSignalHandlerStack), nullptr);
diff --git a/Engine/Source/Runtime/Core/Private/Unix/UnixPlatformStackWalk.cpp b/Engine/Source/Runtime/Core/Private/Unix/UnixPlatformStackWalk.cpp
--- a/Engine/Source/Runtime/Core/Private/Unix/UnixPlatformStackWalk.cpp
+++ b/Engine/Source/Runtime/Core/Private/Unix/UnixPlatformStackWalk.cpp
@@ -11,11 +11,20 @@
 #include "Misc/ScopeLock.h"
 #include "Misc/CommandLine.h"
 #include "Unix/UnixPlatformCrashContext.h"
+#include "Unix/UnixPlatformRealTimeSignals.h"
 #include "HAL/ExceptionHandling.h"
 #include "HAL/PlatformProcess.h"
 #include "HAL/PlatformTime.h"
 
 #include <link.h>
+#include <signal.h>
+
+#include "HAL/IConsoleManager.h"
+
+static TAutoConsoleVariable<float> CVarUnixPlatformThreadCallStackMaxWait(
+	TEXT("UnixPlatformThreadStackWalk.MaxWait"),
+	60.0f,
+	TEXT("The number of seconds allowed to spin before killing the process, with the assumption the signal handler has hung."));
 
 // Init'ed in UnixPlatformMemory. Once this is tested more we can remove this fallback flag
 extern bool CORE_API GFullCrashCallstack;
@@ -494,74 +503,89 @@ bool FUnixPlatformStackWalk::ProgramCounterToHumanReadableString( int32 CurrentC
 
 			// Get filename, source file and line number
 			FUnixCrashContext* UnixContext = static_cast< FUnixCrashContext* >( Context );
-			if (UnixContext)
-			{
-				// for ensure, use the fast path - do not even attempt to get detailed info as it will result in long hitch
-				bool bAddDetailedInfo = UnixContext->GetType() != ECrashContextType::Ensure;
 
-				// Program counters in the backtrace point to the location from where the execution will be resumed (in all frames except the one where we crashed),
-				// which results in callstack pointing to the next lines in code. In order to determine the source line where the actual call happened, we need to go
-				// back to the line that had the "call" instruction. Since x86(-64) instructions vary in length, we cannot do it reliably without disassembling,
-				// just go back one byte - even if it's not the actual address of the call site.
-				int OffsetToCallsite = CurrentCallDepth > 0 ? 1 : 0;
+			// for ensure, use the fast path - do not even attempt to get detailed info as it will result in long hitch
+			bool bAddDetailedInfo = UnixContext ? UnixContext->GetType() != ECrashContextType::Ensure : false;
 
-				FProgramCounterSymbolInfo TempSymbolInfo;
+			// Program counters in the backtrace point to the location from where the execution will be resumed (in all frames except the one where we crashed),
+			// which results in callstack pointing to the next lines in code. In order to determine the source line where the actual call happened, we need to go
+			// back to the line that had the "call" instruction. Since x86(-64) instructions vary in length, we cannot do it reliably without disassembling,
+			// just go back one byte - even if it's not the actual address of the call site.
+			int OffsetToCallsite = CurrentCallDepth > 0 ? 1 : 0;
 
-				// We can print detail info out during ensures, the only reason not to is if fail to populate the symbol info all the way
-				bAddDetailedInfo = PopulateProgramCounterSymbolInfoFromSymbolFile(ProgramCounter - OffsetToCallsite, TempSymbolInfo);
+			FProgramCounterSymbolInfo TempSymbolInfo;
 
-				if (bAddDetailedInfo)
-				{
-					// append Module!FunctionName() [Source.cpp:X] to HumanReadableString
-					FCStringAnsi::Strncat(HumanReadableString, TempSymbolInfo.ModuleName, HumanReadableStringSize);
-					FCStringAnsi::Strncat(HumanReadableString, "!", HumanReadableStringSize);
-					FCStringAnsi::Strncat(HumanReadableString, TempSymbolInfo.FunctionName, HumanReadableStringSize);
-					FCStringAnsi::Sprintf(TempArray, " [%s:%d]", TempSymbolInfo.Filename, TempSymbolInfo.LineNumber);
-					FCStringAnsi::Strncat(HumanReadableString, TempArray, HumanReadableStringSize);
+			// We can print detail info out during ensures, the only reason not to is if fail to populate the symbol info all the way
+			bAddDetailedInfo = PopulateProgramCounterSymbolInfoFromSymbolFile(ProgramCounter - OffsetToCallsite, TempSymbolInfo);
 
+			if (bAddDetailedInfo)
+			{
+				// append Module!FunctionName() [Source.cpp:X] to HumanReadableString
+				FCStringAnsi::Strncat(HumanReadableString, TempSymbolInfo.ModuleName, HumanReadableStringSize);
+				FCStringAnsi::Strncat(HumanReadableString, "!", HumanReadableStringSize);
+				FCStringAnsi::Strncat(HumanReadableString, TempSymbolInfo.FunctionName, HumanReadableStringSize);
+				FCStringAnsi::Sprintf(TempArray, " [%s:%d]", TempSymbolInfo.Filename, TempSymbolInfo.LineNumber);
+				FCStringAnsi::Strncat(HumanReadableString, TempArray, HumanReadableStringSize);
+
+				if (UnixContext)
+				{
 					// append Module!FunctioName [Source.cpp:X] to MinidumpCallstackInfo
 					FCStringAnsi::Strncat(UnixContext->MinidumpCallstackInfo, TempSymbolInfo.ModuleName, ARRAY_COUNT( UnixContext->MinidumpCallstackInfo ) - 1);
 					FCStringAnsi::Strncat(UnixContext->MinidumpCallstackInfo, "!", ARRAY_COUNT( UnixContext->MinidumpCallstackInfo ) - 1);
 					FCStringAnsi::Strncat(UnixContext->MinidumpCallstackInfo, TempSymbolInfo.FunctionName, ARRAY_COUNT( UnixContext->MinidumpCallstackInfo ) - 1);
 					FCStringAnsi::Strncat(UnixContext->MinidumpCallstackInfo, TempArray, ARRAY_COUNT( UnixContext->MinidumpCallstackInfo ) - 1);
 				}
-				else
-				{
-					const char* ModuleName   = nullptr;
-					const char* FunctionName = nullptr;
+			}
+			else
+			{
+				const char* ModuleName   = nullptr;
+				const char* FunctionName = nullptr;
 
-					// We have failed to fully populate the SymbolInfo, but we could still have basic information. Lets try to print as much info as possible
-					if (TempSymbolInfo.ModuleName[0] != '\0')
-					{
-						ModuleName = TempSymbolInfo.ModuleName;
-					}
+				// We have failed to fully populate the SymbolInfo, but we could still have basic information. Lets try to print as much info as possible
+				if (TempSymbolInfo.ModuleName[0] != '\0')
+				{
+					ModuleName = TempSymbolInfo.ModuleName;
+				}
 
-					if (TempSymbolInfo.FunctionName[0] != '\0')
-					{
-						FunctionName = TempSymbolInfo.FunctionName;
-					}
+				if (TempSymbolInfo.FunctionName[0] != '\0')
+				{
+					FunctionName = TempSymbolInfo.FunctionName;
+				}
 
-					FCStringAnsi::Strncat(HumanReadableString, ModuleName != nullptr ? ModuleName : "", HumanReadableStringSize);
-					FCStringAnsi::Strncat(HumanReadableString, "!", HumanReadableStringSize);
-					FCStringAnsi::Strncat(HumanReadableString, FunctionName != nullptr ? FunctionName : "UnknownFunction", HumanReadableStringSize);
-					FCStringAnsi::Strncat(HumanReadableString, FunctionName && TempSymbolInfo.SymbolDisplacement ? "(+": "(", HumanReadableStringSize);
+				FCStringAnsi::Strncat(HumanReadableString, ModuleName != nullptr ? ModuleName : "", HumanReadableStringSize);
+				FCStringAnsi::Strncat(HumanReadableString, "!", HumanReadableStringSize);
+				FCStringAnsi::Strncat(HumanReadableString, FunctionName != nullptr ? FunctionName : "UnknownFunction", HumanReadableStringSize);
+				FCStringAnsi::Strncat(HumanReadableString, FunctionName && TempSymbolInfo.SymbolDisplacement ? "(+": "(", HumanReadableStringSize);
 
+				if (UnixContext)
+				{
 					FCStringAnsi::Strncat(UnixContext->MinidumpCallstackInfo, ModuleName != nullptr ? ModuleName : "Unknown", ARRAY_COUNT( UnixContext->MinidumpCallstackInfo ) - 1);
 					FCStringAnsi::Strncat(UnixContext->MinidumpCallstackInfo, "!", ARRAY_COUNT( UnixContext->MinidumpCallstackInfo ) - 1);
 					FCStringAnsi::Strncat(UnixContext->MinidumpCallstackInfo, FunctionName != nullptr ? FunctionName : "UnknownFunction", ARRAY_COUNT( UnixContext->MinidumpCallstackInfo ) - 1);
 					FCStringAnsi::Strncat(UnixContext->MinidumpCallstackInfo, FunctionName && TempSymbolInfo.SymbolDisplacement ? "(+": "(", ARRAY_COUNT( UnixContext->MinidumpCallstackInfo ) - 1);
+				}
+
+				if (TempSymbolInfo.SymbolDisplacement > 0x0)
+				{
+					FCStringAnsi::Sprintf(TempArray, "%p", TempSymbolInfo.SymbolDisplacement);
+					FCStringAnsi::Strncat(HumanReadableString, TempArray, HumanReadableStringSize);
 
-					if (TempSymbolInfo.SymbolDisplacement > 0x0)
+					if (UnixContext)
 					{
-						FCStringAnsi::Sprintf(TempArray, "%p", TempSymbolInfo.SymbolDisplacement);
-						FCStringAnsi::Strncat(HumanReadableString, TempArray, HumanReadableStringSize);
 						FCStringAnsi::Strncat(UnixContext->MinidumpCallstackInfo, TempArray, ARRAY_COUNT( UnixContext->MinidumpCallstackInfo ) - 1);
 					}
+				}
+
+				FCStringAnsi::Strncat(HumanReadableString, ")", HumanReadableStringSize);
 
-					FCStringAnsi::Strncat(HumanReadableString, ")", HumanReadableStringSize);
+				if (UnixContext)
+				{
 					FCStringAnsi::Strncat(UnixContext->MinidumpCallstackInfo, ")", ARRAY_COUNT( UnixContext->MinidumpCallstackInfo ) - 1);
 				}
+			}
 
+			if (UnixContext)
+			{
 				FCStringAnsi::Strncat(UnixContext->MinidumpCallstackInfo, "\r\n", ARRAY_COUNT( UnixContext->MinidumpCallstackInfo ) - 1);	// this one always uses Windows line terminators
 			}
 		}
@@ -674,21 +698,74 @@ uint32 FUnixPlatformStackWalk::CaptureStackBackTrace( uint64* BackTrace, uint32
 	return (uint32)Size;
 }
 
-void FUnixPlatformStackWalk::ThreadStackWalkAndDump(ANSICHAR* HumanReadableString, SIZE_T HumanReadableStringSize, int32 IgnoreCount, uint32 ThreadId)
+namespace
 {
-	// This is intentional on servers. Right now we cannot symbolicate the other thread, so we crash it instead, which also helps in identifying lock ups
-	if (UE_SERVER)
+	void WaitForSignalHandlerToFinishOrCrash(ThreadStackUserData& ThreadStack)
 	{
-		if (kill(ThreadId, SIGQUIT) == 0)
+		float EndWaitTimestamp = FPlatformTime::Seconds() + CVarUnixPlatformThreadCallStackMaxWait.AsVariable()->GetFloat();
+		float CurrentTimestamp = FPlatformTime::Seconds();
+
+		while (!ThreadStack.bDone)
 		{
-			// do not exit, crash is imminent anyway (signals are delivered asynchronously)
-			for(;;)
+			if (CurrentTimestamp > EndWaitTimestamp)
 			{
+				// We have waited for as long as we should for the signal handler to finish. Assume it has hang and we need to kill our selfs
+				*(int*)0x10 = 0x0;
 			}
+
+			CurrentTimestamp = FPlatformTime::Seconds();
+		}
+	}
+
+	void GatherCallstackFromThread(ThreadStackUserData& ThreadStack, uint64 ThreadId)
+	{
+		sigval UserData;
+		UserData.sival_ptr = &ThreadStack;
+
+		siginfo_t info;
+		memset (&info, 0, sizeof (siginfo_t));
+		info.si_signo = THREAD_CALLSTACK_GENERATOR;
+		info.si_code  = SI_QUEUE;
+		info.si_pid   = syscall(SYS_getpid);
+		info.si_uid   = syscall(SYS_getuid);
+		info.si_value = UserData;
+
+		// Avoid using sigqueue here as if the ThreadId is already blocked and in a signal handler
+		// sigqueue will try a different thread signal handler and report the wrong callstack
+		if (syscall(SYS_rt_tgsigqueueinfo, info.si_pid, ThreadId, THREAD_CALLSTACK_GENERATOR, &info) == 0)
+		{
+			WaitForSignalHandlerToFinishOrCrash(ThreadStack);
 		}
 	}
 }
 
+void FUnixPlatformStackWalk::ThreadStackWalkAndDump(ANSICHAR* HumanReadableString, SIZE_T HumanReadableStringSize, int32 IgnoreCount, uint32 ThreadId)
+{
+	ThreadStackUserData ThreadCallStack;
+	ThreadCallStack.bCaptureCallStack = true;
+	ThreadCallStack.CallStackSize     = HumanReadableStringSize;
+	ThreadCallStack.CallStack         = HumanReadableString;
+	ThreadCallStack.BackTraceCount    = 0;
+	ThreadCallStack.bDone             = false;
+
+	GatherCallstackFromThread(ThreadCallStack, ThreadId);
+}
+
+uint32 FUnixPlatformStackWalk::CaptureThreadStackBackTrace(uint64 ThreadId, uint64* BackTrace, uint32 MaxDepth)
+{
+	ThreadStackUserData ThreadBackTrace;
+	ThreadBackTrace.bCaptureCallStack = false;
+	ThreadBackTrace.CallStackSize     = MaxDepth;
+	ThreadBackTrace.BackTrace         = BackTrace;
+	ThreadBackTrace.BackTraceCount    = 0;
+	ThreadBackTrace.bDone             = false;
+
+	GatherCallstackFromThread(ThreadBackTrace, ThreadId);
+
+	// The signal handler will set this value, we just have to make sure we wait for the signal handler we raised to finish
+	return ThreadBackTrace.BackTraceCount;
+}
+
 namespace
 {
 	int NumberOfDynamicLibrariesCallback(struct dl_phdr_info* Info, size_t /*size*/, void* Data)
diff --git a/Engine/Source/Runtime/Core/Public/Unix/UnixPlatformCrashContext.h b/Engine/Source/Runtime/Core/Public/Unix/UnixPlatformCrashContext.h
--- a/Engine/Source/Runtime/Core/Public/Unix/UnixPlatformCrashContext.h
+++ b/Engine/Source/Runtime/Core/Public/Unix/UnixPlatformCrashContext.h
@@ -5,6 +5,23 @@
 #include "CoreTypes.h"
 #include "GenericPlatform/GenericPlatformCrashContext.h"
 
+/** Passed in through sigqueue for gathering of a callstack from a signal */
+struct ThreadStackUserData
+{
+	// If we want a backtrace or a callstack
+	// Backtrace is just a list of program counters and callstack is a symbolicated backtrace
+	bool bCaptureCallStack;
+	union 
+	{
+		ANSICHAR* CallStack;
+		uint64* BackTrace;
+	};
+
+	int32 BackTraceCount;
+	SIZE_T CallStackSize;
+	TAtomic<bool> bDone;
+};
+
 struct CORE_API FUnixCrashContext : public FGenericCrashContext
 {
 	/** Signal number */
diff --git a/Engine/Source/Runtime/Core/Public/Unix/UnixPlatformRealTimeSignals.h b/Engine/Source/Runtime/Core/Public/Unix/UnixPlatformRealTimeSignals.h
--- a/Engine/Source/Runtime/Core/Public/Unix/UnixPlatformRealTimeSignals.h
+++ b/Engine/Source/Runtime/Core/Public/Unix/UnixPlatformRealTimeSignals.h
@@ -10,3 +10,8 @@
 
 // Used in UnixSignalGameHitchHeartBeat.cpp for a hitch signal handler
 #define HEART_BEAT_SIGNAL SIGRTMIN + 3
+
+// Skip SIGRTMIN + 4 default signal used in VTune
+
+// Using in UnixPlatformCrashContext.cpp/UnixPlatformStackWalk.cpp for gathering current callstack from a thread
+#define THREAD_CALLSTACK_GENERATOR SIGRTMIN + 5
diff --git a/Engine/Source/Runtime/Core/Public/Unix/UnixPlatformStackWalk.h b/Engine/Source/Runtime/Core/Public/Unix/UnixPlatformStackWalk.h
--- a/Engine/Source/Runtime/Core/Public/Unix/UnixPlatformStackWalk.h
+++ b/Engine/Source/Runtime/Core/Public/Unix/UnixPlatformStackWalk.h
@@ -18,6 +18,7 @@ struct CORE_API FUnixPlatformStackWalk : public FGenericPlatformStackWalk
 	static uint32 CaptureStackBackTrace( uint64* BackTrace, uint32 MaxDepth, void* Context = nullptr );
 	static void StackWalkAndDump( ANSICHAR* HumanReadableString, SIZE_T HumanReadableStringSize, int32 IgnoreCount, void* Context = nullptr );
 	static void StackWalkAndDumpEx(ANSICHAR* HumanReadableString, SIZE_T HumanReadableStringSize, int32 IgnoreCount, uint32 Flags, void* Context = nullptr);
+	static uint32 CaptureThreadStackBackTrace(uint64 ThreadId, uint64* BackTrace, uint32 MaxDepth);
 	static void ThreadStackWalkAndDump(ANSICHAR* HumanReadableString, SIZE_T HumanReadableStringSize, int32 IgnoreCount, uint32 ThreadId);
 	static int32 GetProcessModuleCount();
 	static int32 GetProcessModuleSignatures(FStackWalkModuleInfo *ModuleSignatures, const int32 ModuleSignaturesSize);
