Fix missing vtable

diff --git a/Engine/Source/Runtime/Core/Public/HAL/PThreadEvent.h b/Engine/Source/Runtime/Core/Public/HAL/PThreadEvent.h
--- a/Engine/Source/Runtime/Core/Public/HAL/PThreadEvent.h
+++ b/Engine/Source/Runtime/Core/Public/HAL/PThreadEvent.h
@@ -8,7 +8,7 @@
 /**
  * This is the PThreads version of FEvent.
  */
-class FPThreadEvent
+class CORE_VTABLE FPThreadEvent
 	: public FEvent
 {
 	// This is a little complicated, in an attempt to match Win32 Event semantics...
