Avoid signed int overflow when adding two signed ints, this is undefined behaviour

Was found from -EnableUBSan

diff --git a/Engine/Source/Runtime/Core/Public/Containers/UnrealString.h b/Engine/Source/Runtime/Core/Public/Containers/UnrealString.h
--- a/Engine/Source/Runtime/Core/Public/Containers/UnrealString.h
+++ b/Engine/Source/Runtime/Core/Public/Containers/UnrealString.h
@@ -1168,7 +1168,8 @@ public:
 	FORCEINLINE FString Mid( int32 Start, int32 Count=MAX_int32 ) const
 	{
 		check(Count >= 0);
-		uint32 End = Start+Count;
+		uint32 End = Count;
+		End += Start;
 		Start    = FMath::Clamp( (uint32)Start, (uint32)0,     (uint32)Len() );
 		End      = FMath::Clamp( (uint32)End,   (uint32)Start, (uint32)Len() );
 		return FString( End-Start, **this + Start );
