Corrected platform implementations of VectorSign and VectorStep so they are all consistent with the SSE implementation. Added a basic unit test.

diff --git a/Engine/Source/Runtime/Core/Private/Tests/Math/UnrealMathTest.cpp b/Engine/Source/Runtime/Core/Private/Tests/Math/UnrealMathTest.cpp
--- a/Engine/Source/Runtime/Core/Private/Tests/Math/UnrealMathTest.cpp
+++ b/Engine/Source/Runtime/Core/Private/Tests/Math/UnrealMathTest.cpp
@@ -1071,6 +1071,18 @@ bool FVectorRegisterAbstractionTest::RunTest(const FString& Parameters)
 	V3 = VectorMod(V0, V1);
 	LogTest( TEXT("VectorMod negative"), TestVectorsEqual(V2, V3));
 
+	// VectorSign
+	V0 = MakeVectorRegister(2.0f, -2.0f, 0.0f, -3.0f);
+	V2 = MakeVectorRegister(1.0f, -1.0f, 1.0f, -1.0f);
+	V3 = VectorSign(V0);
+	LogTest(TEXT("VectorSign"), TestVectorsEqual(V2, V3));
+
+	// VectorStep
+	V0 = MakeVectorRegister(2.0f, -2.0f, 0.0f, -3.0f);
+	V2 = MakeVectorRegister(1.0f, 0.0f, 1.0f, 0.0f);
+	V3 = VectorStep(V0);
+	LogTest(TEXT("VectorStep"), TestVectorsEqual(V2, V3));
+
 	FMatrix	M0, M1, M2, M3;
 	FVector Eye, LookAt, Up;	
 	// Create Look at Matrix
diff --git a/Engine/Source/Runtime/Core/Public/Math/UnrealMathDirectX.h b/Engine/Source/Runtime/Core/Public/Math/UnrealMathDirectX.h
--- a/Engine/Source/Runtime/Core/Public/Math/UnrealMathDirectX.h
+++ b/Engine/Source/Runtime/Core/Public/Math/UnrealMathDirectX.h
@@ -1059,10 +1059,10 @@ FORCEINLINE VectorRegister VectorSign(const VectorRegister& X)
 FORCEINLINE VectorRegister VectorStep(const VectorRegister& X)
 {
 	return MakeVectorRegister(
-		(float)(VectorGetComponent(X, 0) >= 0.0f ? 1.0f : -1.0f),
-		(float)(VectorGetComponent(X, 1) >= 0.0f ? 1.0f : -1.0f),
-		(float)(VectorGetComponent(X, 2) >= 0.0f ? 1.0f : -1.0f),
-		(float)(VectorGetComponent(X, 3) >= 0.0f ? 1.0f : -1.0f));
+		(float)(VectorGetComponent(X, 0) >= 0.0f ? 1.0f : 0.0f),
+		(float)(VectorGetComponent(X, 1) >= 0.0f ? 1.0f : 0.0f),
+		(float)(VectorGetComponent(X, 2) >= 0.0f ? 1.0f : 0.0f),
+		(float)(VectorGetComponent(X, 3) >= 0.0f ? 1.0f : 0.0f));
 }
 
 //////////////////////////////////////////////////////////////////////////
diff --git a/Engine/Source/Runtime/Core/Public/Math/UnrealMathFPU.h b/Engine/Source/Runtime/Core/Public/Math/UnrealMathFPU.h
--- a/Engine/Source/Runtime/Core/Public/Math/UnrealMathFPU.h
+++ b/Engine/Source/Runtime/Core/Public/Math/UnrealMathFPU.h
@@ -1279,20 +1279,20 @@ FORCEINLINE VectorRegister VectorMod(const VectorRegister& X, const VectorRegist
 FORCEINLINE VectorRegister VectorSign(const VectorRegister& X)
 {
 	return MakeVectorRegister(
-		(float)(VectorGetComponent(X, 0) >= 0.0f ? 1.0f : 0.0f),
-		(float)(VectorGetComponent(X, 1) >= 0.0f ? 1.0f : 0.0f),
-		(float)(VectorGetComponent(X, 2) >= 0.0f ? 1.0f : 0.0f),
-		(float)(VectorGetComponent(X, 3) >= 0.0f ? 1.0f : 0.0f));
+		(float)(VectorGetComponent(X, 0) >= 0.0f ? 1.0f : -1.0f),
+		(float)(VectorGetComponent(X, 1) >= 0.0f ? 1.0f : -1.0f),
+		(float)(VectorGetComponent(X, 2) >= 0.0f ? 1.0f : -1.0f),
+		(float)(VectorGetComponent(X, 3) >= 0.0f ? 1.0f : -1.0f));
 }
 
 //TODO: Vectorize
 FORCEINLINE VectorRegister VectorStep(const VectorRegister& X)
 {
 	return MakeVectorRegister(
-		(float)(VectorGetComponent(X, 0) >= 0.0f ? 1.0f : -1.0f),
-		(float)(VectorGetComponent(X, 1) >= 0.0f ? 1.0f : -1.0f),
-		(float)(VectorGetComponent(X, 2) >= 0.0f ? 1.0f : -1.0f),
-		(float)(VectorGetComponent(X, 3) >= 0.0f ? 1.0f : -1.0f));
+		(float)(VectorGetComponent(X, 0) >= 0.0f ? 1.0f : 0.0f),
+		(float)(VectorGetComponent(X, 1) >= 0.0f ? 1.0f : 0.0f),
+		(float)(VectorGetComponent(X, 2) >= 0.0f ? 1.0f : 0.0f),
+		(float)(VectorGetComponent(X, 3) >= 0.0f ? 1.0f : 0.0f));
 }
 
 /**
diff --git a/Engine/Source/Runtime/Core/Public/Math/UnrealMathNeon.h b/Engine/Source/Runtime/Core/Public/Math/UnrealMathNeon.h
--- a/Engine/Source/Runtime/Core/Public/Math/UnrealMathNeon.h
+++ b/Engine/Source/Runtime/Core/Public/Math/UnrealMathNeon.h
@@ -1294,20 +1294,20 @@ FORCEINLINE VectorRegister VectorMod(const VectorRegister& X, const VectorRegist
 FORCEINLINE VectorRegister VectorSign(const VectorRegister& X)
 {
 	return MakeVectorRegister(
-		(float)(VectorGetComponent(X, 0) >= 0.0f ? 1.0f : 0.0f),
-		(float)(VectorGetComponent(X, 1) >= 0.0f ? 1.0f : 0.0f),
-		(float)(VectorGetComponent(X, 2) >= 0.0f ? 1.0f : 0.0f),
-		(float)(VectorGetComponent(X, 3) >= 0.0f ? 1.0f : 0.0f));
+		(float)(VectorGetComponent(X, 0) >= 0.0f ? 1.0f : -1.0f),
+		(float)(VectorGetComponent(X, 1) >= 0.0f ? 1.0f : -1.0f),
+		(float)(VectorGetComponent(X, 2) >= 0.0f ? 1.0f : -1.0f),
+		(float)(VectorGetComponent(X, 3) >= 0.0f ? 1.0f : -1.0f));
 }
 
 //TODO: Vectorize
 FORCEINLINE VectorRegister VectorStep(const VectorRegister& X)
 {
 	return MakeVectorRegister(
-		(float)(VectorGetComponent(X, 0) >= 0.0f ? 1.0f : -1.0f),
-		(float)(VectorGetComponent(X, 1) >= 0.0f ? 1.0f : -1.0f),
-		(float)(VectorGetComponent(X, 2) >= 0.0f ? 1.0f : -1.0f),
-		(float)(VectorGetComponent(X, 3) >= 0.0f ? 1.0f : -1.0f));
+		(float)(VectorGetComponent(X, 0) >= 0.0f ? 1.0f : 0.0f),
+		(float)(VectorGetComponent(X, 1) >= 0.0f ? 1.0f : 0.0f),
+		(float)(VectorGetComponent(X, 2) >= 0.0f ? 1.0f : 0.0f),
+		(float)(VectorGetComponent(X, 3) >= 0.0f ? 1.0f : 0.0f));
 }
 
 /**
