Added begin()/end() methods to TStaticArray so that it can be treated as an iterable type, used in Algo:: operations, and as part of a range-for.

diff --git a/Engine/Source/Runtime/Core/Public/Containers/StaticArray.h b/Engine/Source/Runtime/Core/Public/Containers/StaticArray.h
--- a/Engine/Source/Runtime/Core/Public/Containers/StaticArray.h
+++ b/Engine/Source/Runtime/Core/Public/Containers/StaticArray.h
@@ -87,6 +87,12 @@ public:
 		return Result;
 	}
 
+	/** STL-like iterators to enable range-based for loop support. */
+	FORCEINLINE TElement*		begin()			{ return &Storage.Elements[0].Element; }
+	FORCEINLINE const TElement*	begin() const	{ return &Storage.Elements[0].Element; }
+	FORCEINLINE TElement*		end()			{ return &Storage.Elements[NumElements].Element; }
+	FORCEINLINE const TElement*	end() const		{ return &Storage.Elements[NumElements].Element; }
+
 private:
 
 	struct alignas(Alignment) TArrayStorageElementAligned
