Added constructors to TArrayView allowing it to view a TStaticArray.

diff --git a/Engine/Source/Runtime/Core/Public/Containers/ArrayView.h b/Engine/Source/Runtime/Core/Public/Containers/ArrayView.h
--- a/Engine/Source/Runtime/Core/Public/Containers/ArrayView.h
+++ b/Engine/Source/Runtime/Core/Public/Containers/ArrayView.h
@@ -9,6 +9,7 @@
 #include "Templates/UnrealTypeTraits.h"
 #include "Math/NumericLimits.h"
 #include "Containers/Array.h"
+#include "Containers/StaticArray.h"
 
 namespace ArrayViewPrivate
 {
@@ -133,6 +134,27 @@ public:
 	{
 	}
 
+	/**
+	 * Construct a view of a TStaticArray with a compatible element type
+	 *
+	 * @param Other The source array to view.
+	 */
+	template <typename OtherElementType, uint32 Size,
+		typename = typename TEnableIf<TIsCompatibleElementType<OtherElementType>::Value>::Type>
+	FORCEINLINE TArrayView(TStaticArray<OtherElementType, Size>& Other)
+		: DataPtr(&Other[0])
+		, ArrayNum((int32)Size)
+	{
+	}
+
+	template <typename OtherElementType, uint32 Size,
+		typename = typename TEnableIf<TIsCompatibleElementType<const OtherElementType>::Value>::Type>
+	FORCEINLINE TArrayView(const TStaticArray<OtherElementType, Size>& Other)
+		: DataPtr(&Other[0])
+		, ArrayNum((int32)Size)
+	{
+	}
+
 	/**
 	 * Construct a view from an initializer list with a compatible element type
 	 *
