Fixes to Algo::Find* to correctly forward the range type, so that rvalue references compile correctly.

diff --git a/Engine/Source/Runtime/Core/Public/Algo/Find.h b/Engine/Source/Runtime/Core/Public/Algo/Find.h
--- a/Engine/Source/Runtime/Core/Public/Algo/Find.h
+++ b/Engine/Source/Runtime/Core/Public/Algo/Find.h
@@ -10,9 +10,9 @@
 namespace AlgoImpl
 {
 	template <typename RangeType, typename ValueType, typename ProjectionType>
-	typename TRangePointerType<RangeType>::Type FindBy(RangeType& Range, const ValueType& Value, ProjectionType Proj)
+	typename TRangePointerType<typename TRemoveReference<RangeType>::Type>::Type FindBy(RangeType&& Range, const ValueType& Value, ProjectionType Proj)
 	{
-		for (auto& Elem : Range)
+		for (auto&& Elem : Forward<RangeType>(Range))
 		{
 			if (Invoke(Proj, Elem) == Value)
 			{
@@ -24,9 +24,9 @@ namespace AlgoImpl
 	}
 
 	template <typename RangeType, typename PredicateType>
-	typename TRangePointerType<RangeType>::Type FindByPredicate(RangeType& Range, PredicateType Pred)
+	typename TRangePointerType<typename TRemoveReference<RangeType>::Type>::Type FindByPredicate(RangeType&& Range, PredicateType Pred)
 	{
-		for (auto& Elem : Range)
+		for (auto&& Elem : Forward<RangeType>(Range))
 		{
 			if (Invoke(Pred, Elem))
 			{
@@ -49,10 +49,10 @@ namespace Algo
 	 * @return A pointer to the first element found, or nullptr if none was found.
 	 */
 	template <typename RangeType, typename ValueType>
-	FORCEINLINE auto Find(RangeType& Range, const ValueType& Value)
-		-> decltype(AlgoImpl::FindBy(Range, Value, FIdentityFunctor()))
+	FORCEINLINE auto Find(RangeType&& Range, const ValueType& Value)
+		-> decltype(AlgoImpl::FindBy(Forward<RangeType>(Range), Value, FIdentityFunctor()))
 	{
-		return AlgoImpl::FindBy(Range, Value, FIdentityFunctor());
+		return AlgoImpl::FindBy(Forward<RangeType>(Range), Value, FIdentityFunctor());
 	}
 
 	/**
@@ -65,10 +65,10 @@ namespace Algo
 	 * @return A pointer to the first element found, or nullptr if none was found.
 	 */
 	template <typename RangeType, typename ValueType, typename ProjectionType>
-	FORCEINLINE auto FindBy(RangeType& Range, const ValueType& Value, ProjectionType Proj)
-		-> decltype(AlgoImpl::FindBy(Range, Value, MoveTemp(Proj)))
+	FORCEINLINE auto FindBy(RangeType&& Range, const ValueType& Value, ProjectionType Proj)
+		-> decltype(AlgoImpl::FindBy(Forward<RangeType>(Range), Value, MoveTemp(Proj)))
 	{
-		return AlgoImpl::FindBy(Range, Value, MoveTemp(Proj));
+		return AlgoImpl::FindBy(Forward<RangeType>(Range), Value, MoveTemp(Proj));
 	}
 
 	/**
@@ -80,9 +80,9 @@ namespace Algo
 	 * @return A pointer to the first element found, or nullptr if none was found.
 	 */
 	template <typename RangeType, typename PredicateType>
-	FORCEINLINE auto FindByPredicate(RangeType& Range, PredicateType Pred)
-		-> decltype(AlgoImpl::FindByPredicate(Range, MoveTemp(Pred)))
+	FORCEINLINE auto FindByPredicate(RangeType&& Range, PredicateType Pred)
+		-> decltype(AlgoImpl::FindByPredicate(Forward<RangeType>(Range), MoveTemp(Pred)))
 	{
-		return AlgoImpl::FindByPredicate(Range, MoveTemp(Pred));
+		return AlgoImpl::FindByPredicate(Forward<RangeType>(Range), MoveTemp(Pred));
 	}
 }
