Fixed TStaticArray ranged-for support in the case where the element type's natural alignment is different to the alignment specified in the template parameter. Implemented a basic iterator class to manage ranged-for iteration which assumes the stride of the specified alignment. Should also fix static analysis errors.

diff --git a/Engine/Source/Runtime/Core/Public/Containers/StaticArray.h b/Engine/Source/Runtime/Core/Public/Containers/StaticArray.h
--- a/Engine/Source/Runtime/Core/Public/Containers/StaticArray.h
+++ b/Engine/Source/Runtime/Core/Public/Containers/StaticArray.h
@@ -87,11 +87,6 @@ public:
 		return Result;
 	}
 
-	/** STL-like iterators to enable range-based for loop support. */
-	FORCEINLINE TElement*		begin()			{ return &Storage.Elements[0].Element; }
-	FORCEINLINE const TElement*	begin() const	{ return &Storage.Elements[0].Element; }
-	FORCEINLINE TElement*		end()			{ return &Storage.Elements[NumElements].Element; }
-	FORCEINLINE const TElement*	end() const		{ return &Storage.Elements[NumElements].Element; }
 
 private:
 
@@ -123,6 +118,46 @@ private:
 	};
 
 	TArrayStorage Storage;
+
+
+public:
+
+	template <typename StorageElementType>
+	struct FRangedForIterator
+	{
+		explicit FRangedForIterator(StorageElementType* InPtr)
+			: Ptr(InPtr)
+		{}
+
+		auto& operator*() const
+		{
+			return Ptr->Element;
+		}
+
+		FRangedForIterator& operator++()
+		{
+			++Ptr;
+			return *this;
+		}
+
+		friend bool operator!=(const FRangedForIterator& A, const FRangedForIterator& B)
+		{
+			return A.Ptr != B.Ptr;
+		}
+
+	private:
+		StorageElementType* Ptr;
+	};
+
+	using RangedForIteratorType = FRangedForIterator<TArrayStorageElementAligned>;
+	using RangedForConstIteratorType = FRangedForIterator<const TArrayStorageElementAligned>;
+
+	/** STL-like iterators to enable range-based for loop support. */
+	FORCEINLINE RangedForIteratorType		begin()			{ return RangedForIteratorType(Storage.Elements); }
+	FORCEINLINE RangedForConstIteratorType	begin() const	{ return RangedForConstIteratorType(Storage.Elements); }
+	FORCEINLINE RangedForIteratorType		end()			{ return RangedForIteratorType(Storage.Elements + NumElements); }
+	FORCEINLINE RangedForConstIteratorType	end() const		{ return RangedForConstIteratorType(Storage.Elements + NumElements); }
+
 };
 
 /** Creates a static array filled with the specified value. */
