PR #5357: Separate abbreviations in Display Name (Contributed by [REDACTED_PERSON])

diff --git a/Engine/Source/Runtime/Core/Private/UObject/UnrealNames.cpp b/Engine/Source/Runtime/Core/Private/UObject/UnrealNames.cpp
--- a/Engine/Source/Runtime/Core/Private/UObject/UnrealNames.cpp
+++ b/Engine/Source/Runtime/Core/Private/UObject/UnrealNames.cpp
@@ -1295,6 +1295,14 @@ FString FName::NameToDisplayString( const FString& InDisplayName, const bool bIs
 			bInARun = false;
 		}
 
+		// We were running on uppercase letters before and still do, but the next character is a lowercase letter,
+		// so we should break the run here, like "3DWidget" should be "3D Widget"
+		if (bInARun && !bWasSpace && !bWasOpenParen && CharIndex < Chars.Num() - 1 && FChar::IsLower(Chars[CharIndex + 1]))
+		{
+			OutDisplayName += TEXT(' ');
+			bWasSpace = true;
+		}
+
 		// An underscore denotes a space, so replace it and continue the run
 		if( bIsUnderscore )
 		{
