Fixed bug where automatically-generated display names were getting an additional space prepended to them. Added unit test to try to prevent this from occurring again.

Introduced in CL 7252282.

diff --git a/Engine/Source/Runtime/Core/Private/Tests/Misc/DisplayNameTest.cpp b/Engine/Source/Runtime/Core/Private/Tests/Misc/DisplayNameTest.cpp
new file mode 100644
--- /dev/null
+++ b/Engine/Source/Runtime/Core/Private/Tests/Misc/DisplayNameTest.cpp
@@ -0,0 +1,38 @@
+// Copyright 1998-2019 Epic Games, Inc. All Rights Reserved.
+
+#if WITH_DEV_AUTOMATION_TESTS 
+
+#include "UObject/UnrealNames.h"
+#include "Misc/AutomationTest.h"
+
+IMPLEMENT_SIMPLE_AUTOMATION_TEST(TDisplayNameTest, "System.Core.Misc.DisplayName", EAutomationTestFlags::ApplicationContextMask | EAutomationTestFlags::SmokeFilter)
+
+bool TDisplayNameTest::RunTest(const FString& Parameters)
+{
+	TestEqual(TEXT("Boolean"), FName::NameToDisplayString(TEXT("bTest"), true), TEXT("Test"));
+	TestEqual(TEXT("Boolean Lower"), FName::NameToDisplayString(TEXT("bTwoWords"), true), TEXT("Two Words"));
+	TestEqual(TEXT("Lower Boolean"), FName::NameToDisplayString(TEXT("boolean"), true), TEXT("Boolean"));
+	TestEqual(TEXT("Almost Boolean"), FName::NameToDisplayString(TEXT("bNotBoolean"), false), TEXT("B Not Boolean"));
+	TestEqual(TEXT("Boolean No Prefix"), FName::NameToDisplayString(TEXT("NonprefixBoolean"), true), TEXT("Nonprefix Boolean"));
+	TestEqual(TEXT("Lower Boolean No Prefix"), FName::NameToDisplayString(TEXT("lowerNonprefixBoolean"), true), TEXT("Lower Nonprefix Boolean"));
+	TestEqual(TEXT("Lower Camel Case"), FName::NameToDisplayString(TEXT("lowerCase"), false), TEXT("Lower Case"));
+	TestEqual(TEXT("3D Widget"), FName::NameToDisplayString(TEXT("3DWidget"), false), TEXT("3D Widget"));
+	TestEqual(TEXT("FBX Editor"), FName::NameToDisplayString(TEXT("FBXEditor"), false), TEXT("FBX Editor"));
+	TestEqual(TEXT("With Underscores"), FName::NameToDisplayString(TEXT("With_Underscores"), false), TEXT("With Underscores"));
+	TestEqual(TEXT("Lower Underscores"), FName::NameToDisplayString(TEXT("lower_underscores"), false), TEXT("Lower Underscores"));
+	TestEqual(TEXT("Mixed Underscores"), FName::NameToDisplayString(TEXT("mixed_Underscores"), false), TEXT("Mixed Underscores"));
+	TestEqual(TEXT("Mixed Underscores"), FName::NameToDisplayString(TEXT("Mixed_underscores"), false), TEXT("Mixed Underscores"));
+	TestEqual(TEXT("Article in String"), FName::NameToDisplayString(TEXT("ArticleInString"), false), TEXT("Article in String"));
+	TestEqual(TEXT("One or Two"), FName::NameToDisplayString(TEXT("OneOrTwo"), false), TEXT("One or Two"));
+	TestEqual(TEXT("One and Two"), FName::NameToDisplayString(TEXT("OneAndTwo"), false), TEXT("One and Two"));
+	TestEqual(TEXT("123 Numbers"), FName::NameToDisplayString(TEXT("123Numbers"), false), TEXT("123 Numbers"));
+	TestEqual(TEXT("-1.5"), FName::NameToDisplayString(TEXT("-1.5"), false), TEXT("-1.5"));
+	TestEqual(TEXT("1234"), FName::NameToDisplayString(TEXT("1234"), false), TEXT("1234"));
+	TestEqual(TEXT("1234.5"), FName::NameToDisplayString(TEXT("1234.5"), false), TEXT("1234.5"));
+	TestEqual(TEXT("-1234.5"), FName::NameToDisplayString(TEXT("-1234.5"), false), TEXT("-1234.5"));
+	TestEqual(TEXT("Text (In Parens)"), FName::NameToDisplayString(TEXT("Text (in parens)"), false), TEXT("Text (In Parens)"));
+
+	return true;
+}
+
+#endif //WITH_DEV_AUTOMATION_TESTS
\ No newline at end of file
diff --git a/Engine/Source/Runtime/Core/Private/UObject/UnrealNames.cpp b/Engine/Source/Runtime/Core/Private/UObject/UnrealNames.cpp
--- a/Engine/Source/Runtime/Core/Private/UObject/UnrealNames.cpp
+++ b/Engine/Source/Runtime/Core/Private/UObject/UnrealNames.cpp
@@ -1299,8 +1299,11 @@ FString FName::NameToDisplayString( const FString& InDisplayName, const bool bIs
 		// so we should break the run here, like "3DWidget" should be "3D Widget"
 		if (bInARun && !bWasSpace && !bWasOpenParen && CharIndex < Chars.Num() - 1 && FChar::IsLower(Chars[CharIndex + 1]))
 		{
-			OutDisplayName += TEXT(' ');
-			bWasSpace = true;
+			if (!bWasSpace && OutDisplayName.Len() > 0)
+			{
+				OutDisplayName += TEXT(' ');
+				bWasSpace = true;
+			}
 		}
 
 		// An underscore denotes a space, so replace it and continue the run
