Fixed execute process corrupting strings on windows platform

diff --git a/Engine/Source/Runtime/Core/Private/Windows/WindowsPlatformProcess.cpp b/Engine/Source/Runtime/Core/Private/Windows/WindowsPlatformProcess.cpp
--- a/Engine/Source/Runtime/Core/Private/Windows/WindowsPlatformProcess.cpp
+++ b/Engine/Source/Runtime/Core/Private/Windows/WindowsPlatformProcess.cpp
@@ -734,13 +734,38 @@ bool FWindowsPlatformProcess::ExecProcess(const TCHAR* URL, const TCHAR* Params,
 		{
 			HANDLE ReadablePipes[2] = { hStdOutRead, hStdErrRead };
 			FString* OutStrings[2] = { OutStdOut, OutStdErr };
+			TArray<uint8> PipeBytes[2];
+
+			auto ReadPipes = [&]()
+			{
+				for (int32 PipeIndex = 0; PipeIndex < 2; ++PipeIndex)
+				{
+					if (ReadablePipes[PipeIndex] && OutStrings[PipeIndex])
+					{
+						TArray<uint8> BinaryData;
+						ReadPipeToArray(ReadablePipes[PipeIndex], BinaryData);
+						PipeBytes[PipeIndex].Append(BinaryData);
+					}
+				}
+			};
+
 			FProcHandle ProcHandle(ProcInfo.hProcess);
 			do 
 			{
-				ReadFromPipes(OutStrings, ReadablePipes, 2);
+				ReadPipes();
 				FPlatformProcess::Sleep(0);
 			} while (IsProcRunning(ProcHandle));
-			ReadFromPipes(OutStrings, ReadablePipes, 2);
+			ReadPipes();
+
+			// Convert only after all bytes are available to prevent string corruption
+			for (int32 PipeIndex = 0; PipeIndex < 2; ++PipeIndex)
+			{
+				if (OutStrings[PipeIndex] && PipeBytes[PipeIndex].Num() > 0)
+				{
+					PipeBytes[PipeIndex].Add('\0');
+					*OutStrings[PipeIndex] = FUTF8ToTCHAR((const ANSICHAR*)PipeBytes[PipeIndex].GetData()).Get();
+				}
+			}
 		}
 		else
 		{
@@ -1329,6 +1354,7 @@ FString FWindowsPlatformProcess::ReadPipe( void* ReadPipe )
 {
 	FString Output;
 
+	// Note: String becomes corrupted when more than one byte per character and all bytes are not available
 	uint32 BytesAvailable = 0;
 	if (::PeekNamedPipe(ReadPipe, NULL, 0, NULL, (::DWORD*)&BytesAvailable, NULL) && (BytesAvailable > 0))
 	{
