Added a LexToString implementation for FFrameTime to increase consistency in logging display format and to shorten the amount of work needed to print them. Usage:
UE_LOG(LogTemp, Log, TEXT("Frame Time: %s"), *LexToString(MyFrameTime))

diff --git a/Engine/Source/Runtime/Core/Public/Misc/FrameTime.h b/Engine/Source/Runtime/Core/Public/Misc/FrameTime.h
--- a/Engine/Source/Runtime/Core/Public/Misc/FrameTime.h
+++ b/Engine/Source/Runtime/Core/Public/Misc/FrameTime.h
@@ -5,6 +5,7 @@
 #include "Misc/FrameNumber.h"
 #include "Misc/AssertionMacros.h"
 #include "Templates/EnableIf.h"
+#include "Containers/UnrealString.h"
 
 /**
  * Structure representing a time by a context-free frame number, plus a sub frame value in the range [0:1)
@@ -269,4 +270,10 @@ FORCEINLINE_DEBUGGABLE FFrameTime FFrameTime::FromDecimal(double InDecimalFrame)
 	// Ensure fractional parts above the highest sub frame float precision do not round to 0.0
 	double Fraction = InDecimalFrame - FMath::FloorToDouble(InDecimalFrame);
 	return FFrameTime(NewFrame, FMath::Clamp(Fraction, 0.0, (double)MaxSubframe));
+}
+
+/** Convert a FFrameTime into a string */
+inline FString LexToString(const FFrameTime InTime)
+{
+	return FString::Printf(TEXT("Frame: %d Subframe: %f"), InTime.GetFrame().Value, InTime.GetSubFrame());
 }
\ No newline at end of file
