If our dir already exists return true from CreateDirectoriesFromPath

diff --git a/Engine/Source/Runtime/Core/Private/Unix/UnixPlatformFile.cpp b/Engine/Source/Runtime/Core/Private/Unix/UnixPlatformFile.cpp
--- a/Engine/Source/Runtime/Core/Private/Unix/UnixPlatformFile.cpp
+++ b/Engine/Source/Runtime/Core/Private/Unix/UnixPlatformFile.cpp
@@ -1127,10 +1127,19 @@ bool FUnixPlatformFile::CreateDirectoriesFromPath(const TCHAR* Path)
 				if (mkdir(SubPath, 0755) == -1)
 				{
 					int ErrNo = errno;
-					UE_LOG_UNIX_FILE(Warning, TEXT( "create dir('%s') failed: errno=%d (%s)" ), UTF8_TO_TCHAR(DirPath), ErrNo, UTF8_TO_TCHAR(strerror(ErrNo)));
+
+					// We happen to have gotten into a position after stat that the folder was created so dont treat this as an error
+					bool bErrnoDirExist = ErrNo == EEXIST;
+
+					if (!bErrnoDirExist)
+					{
+						UE_LOG_UNIX_FILE(Warning, TEXT( "create dir('%s') failed: errno=%d (%s)" ), UTF8_TO_TCHAR(DirPath), ErrNo, UTF8_TO_TCHAR(strerror(ErrNo)));
+					}
+
 					FMemory::Free(DirPath);
 					FMemory::Free(SubPath);
-					return false;
+
+					return bErrnoDirExist;
 				}
 			}
 		}
