Add mod function for FFrameNumber

diff --git a/Engine/Source/Runtime/Core/Public/Misc/FrameNumber.h b/Engine/Source/Runtime/Core/Public/Misc/FrameNumber.h
--- a/Engine/Source/Runtime/Core/Public/Misc/FrameNumber.h
+++ b/Engine/Source/Runtime/Core/Public/Misc/FrameNumber.h
@@ -40,6 +40,7 @@ struct FFrameNumber
 
 	FFrameNumber& operator+=(FFrameNumber RHS)                    { Value += RHS.Value; return *this; }
 	FFrameNumber& operator-=(FFrameNumber RHS)                    { Value -= RHS.Value; return *this; }
+	FFrameNumber& operator%=(FFrameNumber RHS)                    { Value %= RHS.Value; return *this; }
 
 	FFrameNumber& operator++()                                    { ++Value; return *this; }
 	FFrameNumber& operator--()                                    { --Value; return *this; }
@@ -57,6 +58,7 @@ struct FFrameNumber
 
 	friend FFrameNumber operator+(FFrameNumber A, FFrameNumber B) { return FFrameNumber(A.Value + B.Value); }
 	friend FFrameNumber operator-(FFrameNumber A, FFrameNumber B) { return FFrameNumber(A.Value - B.Value); }
+	friend FFrameNumber operator%(FFrameNumber A, FFrameNumber B) { return FFrameNumber(A.Value % B.Value); }
 
 	friend FFrameNumber operator-(FFrameNumber A)                 { return FFrameNumber(-A.Value); }
 
