Added ability to pass flags in to LoadFileTostring
Updated CSV reader to be able to read write-locked files

diff --git a/Engine/Source/Editor/UnrealEd/Private/Factories/CSVImportFactory.cpp b/Engine/Source/Editor/UnrealEd/Private/Factories/CSVImportFactory.cpp
--- a/Engine/Source/Editor/UnrealEd/Private/Factories/CSVImportFactory.cpp
+++ b/Engine/Source/Editor/UnrealEd/Private/Factories/CSVImportFactory.cpp
@@ -348,7 +348,7 @@ EReimportResult::Type UCSVImportFactory::Reimport(UObject* Obj, const FString& P
 		FString FilePath = IFileManager::Get().ConvertToRelativePath(*Path);
 
 		FString Data;
-		if ( FFileHelper::LoadFileToString( Data, *FilePath) )
+		if ( FFileHelper::LoadFileToString( Data, *FilePath, FFileHelper::EHashOptions::None, FILEREAD_AllowWrite) )
 		{
 			const TCHAR* Ptr = *Data;
 			CurrentFilename = FilePath; //not thread safe but seems to be how it is done..
diff --git a/Engine/Source/Runtime/Core/Private/Misc/FileHelper.cpp b/Engine/Source/Runtime/Core/Private/Misc/FileHelper.cpp
--- a/Engine/Source/Runtime/Core/Private/Misc/FileHelper.cpp
+++ b/Engine/Source/Runtime/Core/Private/Misc/FileHelper.cpp
@@ -112,11 +112,11 @@ void FFileHelper::BufferToString( FString& Result, const uint8* Buffer, int32 Si
  * @param Filename name of the file to load
  * @param VerifyFlags flags controlling the hash verification behavior ( see EHashOptions )
  */
-bool FFileHelper::LoadFileToString( FString& Result, const TCHAR* Filename, EHashOptions VerifyFlags )
+bool FFileHelper::LoadFileToString( FString& Result, const TCHAR* Filename, EHashOptions VerifyFlags, uint32 ReadFlags)
 {
 	FScopedLoadingState ScopedLoadingState(Filename);
 
-	TUniquePtr<FArchive> Reader( IFileManager::Get().CreateFileReader( Filename ) );
+	TUniquePtr<FArchive> Reader( IFileManager::Get().CreateFileReader( Filename, ReadFlags) );
 	if( !Reader )
 	{
 		return false;
diff --git a/Engine/Source/Runtime/Core/Public/Misc/FileHelper.h b/Engine/Source/Runtime/Core/Public/Misc/FileHelper.h
--- a/Engine/Source/Runtime/Core/Public/Misc/FileHelper.h
+++ b/Engine/Source/Runtime/Core/Public/Misc/FileHelper.h
@@ -53,7 +53,7 @@ struct CORE_API FFileHelper
 	 * @param Filename     Name of the file to load
 	 * @param VerifyFlags  Flags controlling the hash verification behavior ( see EHashOptions )
 	 */
-	static bool LoadFileToString( FString& Result, const TCHAR* Filename, EHashOptions VerifyFlags = EHashOptions::None );
+	static bool LoadFileToString( FString& Result, const TCHAR* Filename, EHashOptions VerifyFlags = EHashOptions::None, uint32 ReadFlags = 0 );
 
 	/**
 	 * Load a text file to an array of strings. Supports all combination of ANSI/Unicode files and platforms.
