FName::NameToDisplayString functionality reverted to its state prior to CL 7252282.

This has some poor behaviour around abbreviations followed by words (eg. "FBXEditor" -> "FBXEditor"), but behaves better for plurals of abbreviations (eg. "ImportLODs" -> "Import LODs").

The former can be worked around by adding an underscore to the variable name ("FBX_Editor"), or introducing an explicit meta DisplayName=.

Added a test to ensure that the "ImportLODs" behaviour is retained.

diff --git a/Engine/Source/Runtime/Core/Private/Tests/Misc/DisplayNameTest.cpp b/Engine/Source/Runtime/Core/Private/Tests/Misc/DisplayNameTest.cpp
--- a/Engine/Source/Runtime/Core/Private/Tests/Misc/DisplayNameTest.cpp
+++ b/Engine/Source/Runtime/Core/Private/Tests/Misc/DisplayNameTest.cpp
@@ -16,8 +16,6 @@ bool TDisplayNameTest::RunTest(const FString& Parameters)
 	TestEqual(TEXT("Boolean No Prefix"), FName::NameToDisplayString(TEXT("NonprefixBoolean"), true), TEXT("Nonprefix Boolean"));
 	TestEqual(TEXT("Lower Boolean No Prefix"), FName::NameToDisplayString(TEXT("lowerNonprefixBoolean"), true), TEXT("Lower Nonprefix Boolean"));
 	TestEqual(TEXT("Lower Camel Case"), FName::NameToDisplayString(TEXT("lowerCase"), false), TEXT("Lower Case"));
-	TestEqual(TEXT("3D Widget"), FName::NameToDisplayString(TEXT("3DWidget"), false), TEXT("3D Widget"));
-	TestEqual(TEXT("FBX Editor"), FName::NameToDisplayString(TEXT("FBXEditor"), false), TEXT("FBX Editor"));
 	TestEqual(TEXT("With Underscores"), FName::NameToDisplayString(TEXT("With_Underscores"), false), TEXT("With Underscores"));
 	TestEqual(TEXT("Lower Underscores"), FName::NameToDisplayString(TEXT("lower_underscores"), false), TEXT("Lower Underscores"));
 	TestEqual(TEXT("Mixed Underscores"), FName::NameToDisplayString(TEXT("mixed_Underscores"), false), TEXT("Mixed Underscores"));
@@ -25,13 +23,17 @@ bool TDisplayNameTest::RunTest(const FString& Parameters)
 	TestEqual(TEXT("Article in String"), FName::NameToDisplayString(TEXT("ArticleInString"), false), TEXT("Article in String"));
 	TestEqual(TEXT("One or Two"), FName::NameToDisplayString(TEXT("OneOrTwo"), false), TEXT("One or Two"));
 	TestEqual(TEXT("One and Two"), FName::NameToDisplayString(TEXT("OneAndTwo"), false), TEXT("One and Two"));
-	TestEqual(TEXT("123 Numbers"), FName::NameToDisplayString(TEXT("123Numbers"), false), TEXT("123 Numbers"));
 	TestEqual(TEXT("-1.5"), FName::NameToDisplayString(TEXT("-1.5"), false), TEXT("-1.5"));
 	TestEqual(TEXT("1234"), FName::NameToDisplayString(TEXT("1234"), false), TEXT("1234"));
 	TestEqual(TEXT("1234.5"), FName::NameToDisplayString(TEXT("1234.5"), false), TEXT("1234.5"));
 	TestEqual(TEXT("-1234.5"), FName::NameToDisplayString(TEXT("-1234.5"), false), TEXT("-1234.5"));
 	TestEqual(TEXT("Text (In Parens)"), FName::NameToDisplayString(TEXT("Text (in parens)"), false), TEXT("Text (In Parens)"));
 
+	TestEqual(TEXT("Text 3D"), FName::NameToDisplayString(TEXT("Text3D"), false), TEXT("Text 3D"));
+	TestEqual(TEXT("Plural CAPs"), FName::NameToDisplayString(TEXT("PluralCAPs"), false), TEXT("Plural CAPs"));
+	TestEqual(TEXT("FBXEditor"), FName::NameToDisplayString(TEXT("FBXEditor"), false), TEXT("FBXEditor"));
+	TestEqual(TEXT("FBX_Editor"), FName::NameToDisplayString(TEXT("FBX_Editor"), false), TEXT("FBX Editor"));
+
 	return true;
 }
 
diff --git a/Engine/Source/Runtime/Core/Private/UObject/UnrealNames.cpp b/Engine/Source/Runtime/Core/Private/UObject/UnrealNames.cpp
--- a/Engine/Source/Runtime/Core/Private/UObject/UnrealNames.cpp
+++ b/Engine/Source/Runtime/Core/Private/UObject/UnrealNames.cpp
@@ -1313,17 +1313,6 @@ FString FName::NameToDisplayString( const FString& InDisplayName, const bool bIs
 			bInARun = false;
 		}
 
-		// We were running on uppercase letters before and still do, but the next character is a lowercase letter,
-		// so we should break the run here, like "3DWidget" should be "3D Widget"
-		if (bInARun && !bWasSpace && !bWasOpenParen && CharIndex < Chars.Num() - 1 && FChar::IsLower(Chars[CharIndex + 1]))
-		{
-			if (!bWasSpace && OutDisplayName.Len() > 0)
-			{
-				OutDisplayName += TEXT(' ');
-				bWasSpace = true;
-			}
-		}
-
 		// An underscore denotes a space, so replace it and continue the run
 		if( bIsUnderscore )
 		{
