PR #6224: Shared pointer test fixes (Contributed by [REDACTED_PERSON])

diff --git a/Engine/Source/Runtime/Core/Public/Templates/SharedPointer.h b/Engine/Source/Runtime/Core/Public/Templates/SharedPointer.h
--- a/Engine/Source/Runtime/Core/Public/Templates/SharedPointer.h
+++ b/Engine/Source/Runtime/Core/Public/Templates/SharedPointer.h
@@ -51,7 +51,7 @@
  *
  *
  *	Examples:
- *		- Please see 'SharedPointerTesting.h' for various examples of shared pointers and references!
+ *		- Please see 'SharedPointerTesting.inl' for various examples of shared pointers and references!
  *
  *
  *	Tips:
diff --git a/Engine/Source/Runtime/Core/Public/Templates/SharedPointerTesting.inl b/Engine/Source/Runtime/Core/Public/Templates/SharedPointerTesting.inl
--- a/Engine/Source/Runtime/Core/Public/Templates/SharedPointerTesting.inl
+++ b/Engine/Source/Runtime/Core/Public/Templates/SharedPointerTesting.inl
@@ -511,10 +511,12 @@ namespace SharedPointerTesting
 			const bool Value = SharedRefHash.FindRef( FooRef );
 
 			// Map a type to shared refs
-			TMap< int32, TSharedRef< int32, Mode > > SharedRefValueHash;
-			SharedRefValueHash.Add( 10, FooRef );
-			const int32* FoundKey = SharedRefValueHash.FindKey( FooRef );
-			check( FoundKey != nullptr && *FoundKey == 10 );
+			{
+				TMap< int32, TSharedRef< int32, Mode > > SharedRefValueHash;
+				SharedRefValueHash.Add(10, FooRef);
+				const int32* FoundKey = SharedRefValueHash.FindKey(FooRef);
+				check(FoundKey != nullptr && *FoundKey == 10);
+			}
 
 			// Maps and 'const'
 			{
