Core: Changed default sort predicate for TSortedMap to use its ConstPointerType rather than KeyType directly to ensure compatability for non-const-pointer keys

diff --git a/Engine/Source/Runtime/Core/Public/Containers/ContainersFwd.h b/Engine/Source/Runtime/Core/Public/Containers/ContainersFwd.h
--- a/Engine/Source/Runtime/Core/Public/Containers/ContainersFwd.h
+++ b/Engine/Source/Runtime/Core/Public/Containers/ContainersFwd.h
@@ -17,7 +17,8 @@ template<typename KeyType, typename ValueType, bool bInAllowDuplicateKeys> struc
 template<typename KeyType, typename ValueType, typename SetAllocator = FDefaultSetAllocator, typename KeyFuncs = TDefaultMapHashableKeyFuncs<KeyType, ValueType, false> > class TMap;
 template<typename KeyType, typename ValueType, typename SetAllocator = FDefaultSetAllocator, typename KeyFuncs = TDefaultMapHashableKeyFuncs<KeyType, ValueType, true > > class TMultiMap;
 template <typename T = void > struct TLess;
-template<typename KeyType, typename ValueType, typename ArrayAllocator = FDefaultAllocator, typename SortPredicate = TLess<KeyType> > class TSortedMap;
+template <typename> struct TTypeTraits;
+template<typename KeyType, typename ValueType, typename ArrayAllocator = FDefaultAllocator, typename SortPredicate = TLess<typename TTypeTraits<KeyType>::ConstPointerType> > class TSortedMap;
 template<typename ElementType,bool bInAllowDuplicateKeys = false> struct DefaultKeyFuncs;
 template<typename InElementType, typename KeyFuncs = DefaultKeyFuncs<InElementType>, typename Allocator = FDefaultSetAllocator> class TSet;
 /// @endcond
