Fixed LEET failing to load native LocRes data if set during boot

This makes the behavior of -leet and -culture=leet on the command line consistent with what would happen when switching to LEET at runtime.

diff --git a/Engine/Source/Runtime/Core/Private/Internationalization/TextLocalizationManager.cpp b/Engine/Source/Runtime/Core/Private/Internationalization/TextLocalizationManager.cpp
--- a/Engine/Source/Runtime/Core/Private/Internationalization/TextLocalizationManager.cpp
+++ b/Engine/Source/Runtime/Core/Private/Internationalization/TextLocalizationManager.cpp
@@ -886,12 +886,21 @@ void FTextLocalizationManager::LoadLocalizationResourcesForPrioritizedCultures(T
 		return;
 	}
 
+	// Leet-ify always needs the native text to operate on, so force native data if we're loading for LEET
+	ELocalizationLoadFlags FinalLocLoadFlags = LocLoadFlags;
+#if ENABLE_LOC_TESTING
+	if (PrioritizedCultureNames[0] == FLeetCulture::StaticGetName())
+	{
+		FinalLocLoadFlags |= ELocalizationLoadFlags::Native;
+	}
+#endif
+
 	// Load the resources from each text source
 	FTextLocalizationResource NativeResource;
 	FTextLocalizationResource LocalizedResource;
 	for (const TSharedPtr<ILocalizedTextSource>& LocalizedTextSource : LocalizedTextSources)
 	{
-		LocalizedTextSource->LoadLocalizedResources(LocLoadFlags, PrioritizedCultureNames, NativeResource, LocalizedResource);
+		LocalizedTextSource->LoadLocalizedResources(FinalLocLoadFlags, PrioritizedCultureNames, NativeResource, LocalizedResource);
 	}
 
 	// When loc testing is enabled, UpdateFromNative also takes care of restoring non-localized text which is why the condition below is gated
