Message dialog text and results are now logged

diff --git a/Engine/Source/Editor/UnrealEd/Private/Dialogs/Dialogs.cpp b/Engine/Source/Editor/UnrealEd/Private/Dialogs/Dialogs.cpp
--- a/Engine/Source/Editor/UnrealEd/Private/Dialogs/Dialogs.cpp
+++ b/Engine/Source/Editor/UnrealEd/Private/Dialogs/Dialogs.cpp
@@ -343,10 +343,10 @@ EAppReturnType::Type OpenMsgDlgInt(EAppMsgType::Type InMessageType, const FText&
 
 EAppReturnType::Type OpenMsgDlgInt(EAppMsgType::Type InMessageType, EAppReturnType::Type InDefaultValue, const FText& InMessage, const FText& InTitle)
 {
+	EAppReturnType::Type Response = InDefaultValue;
 	if (FApp::IsUnattended() == true || GIsRunningUnattendedScript)
 	{
 		UE_LOG(LogDialogs, Log, TEXT("Message Dialog was triggered in unattended script mode without a default value. %d will be used."), (int32)InDefaultValue);
-		return InDefaultValue;
 	}
 	else
 	{
@@ -357,10 +357,11 @@ EAppReturnType::Type OpenMsgDlgInt(EAppMsgType::Type InMessageType, EAppReturnTy
 
 		GEditor->EditorAddModalWindow(MsgWindow.ToSharedRef());
 
-		EAppReturnType::Type Response = MsgDialog->GetResponse();
-
-		return Response;
+		Response = MsgDialog->GetResponse();
 	}
+
+	UE_LOG(LogDialogs, Log, TEXT("Message dialog closed, result: %s, title: %s, text: %s"), LexToString(Response), *InTitle.ToString(), *InMessage.ToString());
+	return Response;
 }
 
 TSharedRef<SWindow> OpenMsgDlgInt_NonModal(EAppMsgType::Type InMessageType, const FText& InMessage, const FText& InTitle,
diff --git a/Engine/Source/Runtime/Core/Private/GenericPlatform/GenericPlatformMisc.cpp b/Engine/Source/Runtime/Core/Private/GenericPlatform/GenericPlatformMisc.cpp
--- a/Engine/Source/Runtime/Core/Private/GenericPlatform/GenericPlatformMisc.cpp
+++ b/Engine/Source/Runtime/Core/Private/GenericPlatform/GenericPlatformMisc.cpp
@@ -716,6 +716,31 @@ void FGenericPlatformMisc::CreateGuid(FGuid& Guid)
 	Guid = FGuid(RandBits | (SequentialBits << 16), EstimatedCurrentDateTime.GetTicks() >> 32, EstimatedCurrentDateTime.GetTicks() & 0xffffffff, FPlatformTime::Cycles());
 }
 
+const TCHAR* LexToString( EAppReturnType::Type Value )
+{
+	switch (Value)
+	{
+	case EAppReturnType::No:
+		return TEXT("No");
+	case EAppReturnType::Yes:
+		return TEXT("Yes");
+	case EAppReturnType::YesAll:
+		return TEXT("YesAll");
+	case EAppReturnType::NoAll:
+		return TEXT("NoAll");
+	case EAppReturnType::Cancel:
+		return TEXT("Cancel");
+	case EAppReturnType::Ok:
+		return TEXT("Ok");
+	case EAppReturnType::Retry:
+		return TEXT("Retry");
+	case EAppReturnType::Continue:
+		return TEXT("Continue");
+	default:
+		return TEXT("Unknown");
+	}
+}
+
 EAppReturnType::Type FGenericPlatformMisc::MessageBoxExt( EAppMsgType::Type MsgType, const TCHAR* Text, const TCHAR* Caption )
 {
 	if (GWarn)
diff --git a/Engine/Source/Runtime/Core/Private/Misc/MessageDialog.cpp b/Engine/Source/Runtime/Core/Private/Misc/MessageDialog.cpp
--- a/Engine/Source/Runtime/Core/Private/Misc/MessageDialog.cpp
+++ b/Engine/Source/Runtime/Core/Private/Misc/MessageDialog.cpp
@@ -112,25 +112,27 @@ EAppReturnType::Type FMessageDialog::Open( EAppMsgType::Type MessageType, const
 
 EAppReturnType::Type FMessageDialog::Open(EAppMsgType::Type MessageType, EAppReturnType::Type DefaultValue, const FText& Message, const FText* OptTitle)
 {
-	if ( FApp::IsUnattended() == true || GIsRunningUnattendedScript )
-	{
-		if (GWarn)
-		{
-			GWarn->Logf( TEXT("%s"), *Message.ToString() );
-		}
+	EAppReturnType::Type Result = DefaultValue;
+	const FText Title = OptTitle ? *OptTitle : GetDefaultMessageTitle();
+	bool bAlreadyLogged = false;
 
-		return DefaultValue;
-	}
-	else
+	if (!FApp::IsUnattended() && !GIsRunningUnattendedScript)
 	{
-		FText Title = OptTitle ? *OptTitle : GetDefaultMessageTitle();
 		if ( GIsEditor && !IsRunningCommandlet() && FCoreDelegates::ModalErrorMessage.IsBound() )
 		{
-			return FCoreDelegates::ModalErrorMessage.Execute( MessageType, Message, Title );
+			bAlreadyLogged = true;
+			Result = FCoreDelegates::ModalErrorMessage.Execute( MessageType, Message, Title );
 		}
 		else
 		{
-			return FPlatformMisc::MessageBoxExt( MessageType, *Message.ToString(), *Title.ToString() );
+			Result = FPlatformMisc::MessageBoxExt( MessageType, *Message.ToString(), *Title.ToString() );
 		}
 	}
+
+	if (GWarn && !bAlreadyLogged)
+	{
+		GWarn->Logf(TEXT("Message dialog closed, result: %s, title: %s, text: %s"), LexToString(Result), *Title.ToString(), *Message.ToString());
+	}
+
+	return Result;
 }
diff --git a/Engine/Source/Runtime/Core/Public/GenericPlatform/GenericPlatformMisc.h b/Engine/Source/Runtime/Core/Public/GenericPlatform/GenericPlatformMisc.h
--- a/Engine/Source/Runtime/Core/Public/GenericPlatform/GenericPlatformMisc.h
+++ b/Engine/Source/Runtime/Core/Public/GenericPlatform/GenericPlatformMisc.h
@@ -278,6 +278,14 @@ namespace EAppReturnType
 	};
 }
 
+/**
+ * Returns the string representation of the specified EAppReturnType::Type value.
+ *
+ * @param Value The value to get the string for.
+ * @return The string representation.
+ */
+CORE_API const TCHAR* LexToString(EAppReturnType::Type Value);
+
 /*
  * Holds a computed SHA256 hash.
  */
