Require LocRes to consider a language usable by the engine

#rnx

diff --git a/Engine/Source/Runtime/Core/Private/Internationalization/TextLocalizationResource.cpp b/Engine/Source/Runtime/Core/Private/Internationalization/TextLocalizationResource.cpp
--- a/Engine/Source/Runtime/Core/Private/Internationalization/TextLocalizationResource.cpp
+++ b/Engine/Source/Runtime/Core/Private/Internationalization/TextLocalizationResource.cpp
@@ -612,14 +612,19 @@ TArray<FString> TextLocalizationResourceUtil::GetLocalizedCultureNames(const TAr
 	const FString PlatformFolderName = FPaths::GetPlatformLocalizationFolderName();
 	for (const FString& LocalizationPath : InLocalizationPaths)
 	{
-		IFileManager::Get().IterateDirectory(*LocalizationPath, [&CultureNames, &PlatformFolderName](const TCHAR* FilenameOrDirectory, bool bIsDirectory) -> bool
+		const FString LocResFilename = FPaths::GetBaseFilename(LocalizationPath) + TEXT(".locres");
+		IFileManager::Get().IterateDirectory(*LocalizationPath, [&CultureNames, &PlatformFolderName, &LocResFilename](const TCHAR* FilenameOrDirectory, bool bIsDirectory) -> bool
 		{
 			if (bIsDirectory && FCString::Stricmp(FilenameOrDirectory, *PlatformFolderName) != 0)
 			{
-				// UE localization resource folders use "en-US" style while ICU uses "en_US"
-				const FString LocalizationFolder = FPaths::GetCleanFilename(FilenameOrDirectory);
-				const FString CanonicalName = FCulture::GetCanonicalName(LocalizationFolder);
-				CultureNames.AddUnique(CanonicalName);
+				const FString LocResPath = FilenameOrDirectory / LocResFilename;
+				if (FPaths::FileExists(LocResPath))
+				{
+					// UE localization resource folders use "en-US" style while ICU uses "en_US"
+					const FString LocalizationFolder = FPaths::GetCleanFilename(FilenameOrDirectory);
+					const FString CanonicalName = FCulture::GetCanonicalName(LocalizationFolder);
+					CultureNames.AddUnique(CanonicalName);
+				}
 			}
 			return true;
 		});
