Sequencer: add modulo operator to frame times.

diff --git a/Engine/Source/Runtime/Core/Public/Misc/FrameTime.h b/Engine/Source/Runtime/Core/Public/Misc/FrameTime.h
--- a/Engine/Source/Runtime/Core/Public/Misc/FrameTime.h
+++ b/Engine/Source/Runtime/Core/Public/Misc/FrameTime.h
@@ -102,6 +102,7 @@ private:
 	friend FFrameTime& operator+=(FFrameTime& LHS, FFrameTime RHS);
 	friend FFrameTime  operator+(FFrameTime A, FFrameTime B);
 	friend FFrameTime  operator-(FFrameTime A, FFrameTime B);
+	friend FFrameTime  operator%(FFrameTime A, FFrameTime B);
 
 	friend FFrameTime  operator-(FFrameTime A);
 
@@ -220,6 +221,26 @@ FORCEINLINE_DEBUGGABLE FFrameTime operator-(FFrameTime A, FFrameTime B)
 }
 
 
+FORCEINLINE_DEBUGGABLE FFrameTime operator%(FFrameTime A, FFrameTime B)
+{
+	check(B.FrameNumber.Value != 0 || B.GetSubFrame() != 0.f);
+
+	if (A.SubFrame == 0.f && B.SubFrame == 0.f)
+	{
+		return FFrameTime(A.FrameNumber % B.FrameNumber);
+	}
+	else
+	{
+		FFrameTime Result = A;
+		while (Result >= B)
+		{
+			Result = Result - B;
+		}
+		return Result;
+	}
+}
+
+
 FORCEINLINE_DEBUGGABLE FFrameTime operator-(FFrameTime A)
 {
 	return A.GetSubFrame() == 0.f
@@ -276,4 +297,4 @@ FORCEINLINE_DEBUGGABLE FFrameTime FFrameTime::FromDecimal(double InDecimalFrame)
 inline FString LexToString(const FFrameTime InTime)
 {
 	return FString::Printf(TEXT("Frame: %d Subframe: %f"), InTime.GetFrame().Value, InTime.GetSubFrame());
-}
\ No newline at end of file
+}
