Improved GIsPIEUsingPlayerStart deprecation message to specify what to use instead.

#rnx

diff --git a/Engine/Source/Runtime/Core/Public/CoreGlobals.h b/Engine/Source/Runtime/Core/Public/CoreGlobals.h
--- a/Engine/Source/Runtime/Core/Public/CoreGlobals.h
+++ b/Engine/Source/Runtime/Core/Public/CoreGlobals.h
@@ -346,7 +346,7 @@ extern CORE_API bool GIsPlayInEditorWorld;
 extern CORE_API int32 GPlayInEditorID;
 
 /** Whether or not PIE was attempting to play from PlayerStart */
-UE_DEPRECATED(4.25, "This variable is unused.")
+UE_DEPRECATED(4.25, "This variable is no longer set. Use !GEditor->GetPlayInEditorSessionInfo()->OriginalRequestParams.HasPlayWorldPlacement() instead.")
 extern CORE_API bool GIsPIEUsingPlayerStart;
 
 /** true if the runtime needs textures to be powers of two */
