Renamed A and B parameters in Test* methods in FAutomationTestBase to Actual and Expected to make it simpler to implement unit tests.

diff --git a/Engine/Source/Runtime/Core/Public/Misc/AutomationTest.h b/Engine/Source/Runtime/Core/Public/Misc/AutomationTest.h
--- a/Engine/Source/Runtime/Core/Public/Misc/AutomationTest.h
+++ b/Engine/Source/Runtime/Core/Public/Misc/AutomationTest.h
@@ -1279,8 +1279,8 @@ public:
 	void TestEqual(const TCHAR* What, double Actual, double Expected, double Tolerance = KINDA_SMALL_NUMBER);
 	void TestEqual(const TCHAR* What, FVector Actual, FVector Expected, float Tolerance = KINDA_SMALL_NUMBER);
 	void TestEqual(const TCHAR* What, FColor Actual, FColor Expected);
-	void TestEqual(const TCHAR* What, const TCHAR* A, const TCHAR* B);
-	void TestEqualInsensitive(const TCHAR* What, const TCHAR* A, const TCHAR* B);
+	void TestEqual(const TCHAR* What, const TCHAR* Actual, const TCHAR* Expected);
+	void TestEqualInsensitive(const TCHAR* What, const TCHAR* Actual, const TCHAR* Expected);
 
 	void TestEqual(const FString& What, int32 Actual, int32 Expected)
 	{
@@ -1352,18 +1352,18 @@ public:
 	 * @see TestNotEqual
 	 */
 	template<typename ValueType> 
-	void TestEqual(const TCHAR* What, const ValueType& A, const ValueType& B)
+	void TestEqual(const TCHAR* What, const ValueType& Actual, const ValueType& Expected)
 	{
-		if (A != B)
+		if (Actual != Expected)
 		{
 			AddError(FString::Printf(TEXT("%s: The two values are not equal."), What), 1);
 		}
 	}
 
 	template<typename ValueType>
-	void TestEqual(const FString& What, const ValueType& A, const ValueType& B)
+	void TestEqual(const FString& What, const ValueType& Actual, const ValueType& Expected)
 	{
-		TestEqual(*What, A, B);
+		TestEqual(*What, Actual, Expected);
 	}
 
 
@@ -1412,17 +1412,17 @@ public:
 	 *
 	 * @see TestEqual
 	 */
-	template<typename ValueType> void TestNotEqual(const TCHAR* Description, const ValueType& A, const ValueType& B)
+	template<typename ValueType> void TestNotEqual(const TCHAR* Description, const ValueType& Actual, const ValueType& Expected)
 	{
-		if (A == B)
+		if (Actual == Expected)
 		{
 			AddError(FString::Printf(TEXT("%s: The two values are equal."), Description), 1);
 		}
 	}
 
-	template<typename ValueType> void TestNotEqual(const FString& Description, const ValueType& A, const ValueType& B)
+	template<typename ValueType> void TestNotEqual(const FString& Description, const ValueType& Actual, const ValueType& Expected)
 	{
-		TestNotEqual(*Description, A, B);
+		TestNotEqual(*Description, Actual, Expected);
 	}
 
 	/**
@@ -1455,17 +1455,17 @@ public:
 	 *
 	 * @see TestSame
 	 */
-	template<typename ValueType> void TestNotSame(const TCHAR* Description, const ValueType& A, const ValueType& B)
+	template<typename ValueType> void TestNotSame(const TCHAR* Description, const ValueType& Actual, const ValueType& Expected)
 	{
-		if (&A == &B)
+		if (&Actual == &Expected)
 		{
 			AddError(FString::Printf(TEXT("%s: The two values are the same."), Description), 1);
 		}
 	}
 
-	template<typename ValueType> void TestNotSame(const FString& Description, const ValueType& A, const ValueType& B)
+	template<typename ValueType> void TestNotSame(const FString& Description, const ValueType& Actual, const ValueType& Expected)
 	{
-		TestNotSame(*Description, A, B);
+		TestNotSame(*Description, Actual, Expected);
 	}
 
 	/**
@@ -1487,22 +1487,22 @@ public:
 	 * Logs an error if the two values are not the same object in memory.
 	 *
 	 * @param Description - Description text for the test.
-	 * @param A - The first value.
-	 * @param B - The second value.
+	 * @param Actual - The actual value.
+	 * @param Expected - The expected value.
 	 *
 	 * @see TestNotSame
 	 */
-	template<typename ValueType> void TestSame(const TCHAR* Description, const ValueType& A, const ValueType& B)
+	template<typename ValueType> void TestSame(const TCHAR* Description, const ValueType& Actual, const ValueType& Expected)
 	{
-		if (&A != &B)
+		if (&Actual != &Expected)
 		{
 			AddError(FString::Printf(TEXT("%s: The two values are not the same."), Description), 1);
 		}
 	}
 
-	template<typename ValueType> void TestSame(const FString& Description, const ValueType& A, const ValueType& B)
+	template<typename ValueType> void TestSame(const FString& Description, const ValueType& Actual, const ValueType& Expected)
 	{
-		TestSame(*Description, A, B);
+		TestSame(*Description, Actual, Expected);
 	}
 
 	/**
