Fix TIndirectArray::operator=(TIndirectArray&&)

The defaulted version of the move assignment operator fails to free the pointed at data in the existing array before assigning, and thus the current contents are all leaked

#rnx

diff --git a/Engine/Source/Runtime/Core/Public/Containers/IndirectArray.h b/Engine/Source/Runtime/Core/Public/Containers/IndirectArray.h
--- a/Engine/Source/Runtime/Core/Public/Containers/IndirectArray.h
+++ b/Engine/Source/Runtime/Core/Public/Containers/IndirectArray.h
@@ -25,8 +25,6 @@ public:
 	/** Default constructors. */
 	TIndirectArray() = default;
 	TIndirectArray(TIndirectArray&&) = default;
-	TIndirectArray& operator=(TIndirectArray&&) = default;
-
 	/**
 	 * Copy constructor.
 	 *
@@ -41,7 +39,7 @@ public:
 	}
 
 	/**
-	 * Assignment operator.
+	 * Assignment operators.
 	 *
 	 * @param Other Other array to assign with.
 	 */
@@ -58,6 +56,17 @@ public:
 
 		return *this;
 	}
+	TIndirectArray &operator=(TIndirectArray && Other)
+	{
+		if (&Other != this)
+		{
+			Empty();
+			Array = MoveTemp(Other.Array);
+		}
+
+		return *this;
+	}
+
 
 	/** Destructor. */
 	~TIndirectArray()
