Automation: Add convenience macros for testing state in one line.

All the test methods now return a boolean to indicate if the test passed or fail.
New macros run those methods and, if they return false, insert a "return false" in the code, so that it exits the current test immediately.
See comments and example in code.

diff --git a/Engine/Source/Runtime/Core/Private/Misc/AutomationTest.cpp b/Engine/Source/Runtime/Core/Private/Misc/AutomationTest.cpp
--- a/Engine/Source/Runtime/Core/Private/Misc/AutomationTest.cpp
+++ b/Engine/Source/Runtime/Core/Private/Misc/AutomationTest.cpp
@@ -1199,93 +1199,115 @@ void FAutomationTestBase::GenerateTestNames(TArray<FAutomationTestInfo>& TestInf
 
 // --------------------------------------------------------------------------------------
 
-void FAutomationTestBase::TestEqual(const TCHAR* What, const int32 Actual, const int32 Expected)
+bool FAutomationTestBase::TestEqual(const TCHAR* What, const int32 Actual, const int32 Expected)
 {
 	if (Actual != Expected)
 	{
 		AddError(FString::Printf(TEXT("Expected '%s' to be %d, but it was %d."), What, Expected, Actual), 1);
+		return false;
 	}
+	return true;
 }
 
 
-void FAutomationTestBase::TestEqual(const TCHAR* What, const int64 Actual, const int64 Expected)
+bool FAutomationTestBase::TestEqual(const TCHAR* What, const int64 Actual, const int64 Expected)
 {
 	if (Actual != Expected)
 	{
 		AddError(FString::Printf(TEXT("Expected '%s' to be %" PRId64 ", but it was %" PRId64 "."), What, Expected, Actual), 1);
+		return false;
 	}
+	return true;
 }
 
-void FAutomationTestBase::TestEqual(const TCHAR* What, const float Actual, const float Expected, float Tolerance)
+bool FAutomationTestBase::TestEqual(const TCHAR* What, const float Actual, const float Expected, float Tolerance)
 {
 	if (!FMath::IsNearlyEqual(Actual, Expected, Tolerance))
 	{
 		AddError(FString::Printf(TEXT("Expected '%s' to be %f, but it was %f within tolerance %f."), What, Expected, Actual, Tolerance), 1);
+		return false;
 	}
+	return true;
 }
 
-void FAutomationTestBase::TestEqual(const TCHAR* What, const double Actual, const double Expected, double Tolerance)
+bool FAutomationTestBase::TestEqual(const TCHAR* What, const double Actual, const double Expected, double Tolerance)
 {
 	if (!FMath::IsNearlyEqual(Actual, Expected, Tolerance))
 	{
 		AddError(FString::Printf(TEXT("Expected '%s' to be %f, but it was %f within tolerance %f."), What, Expected, Actual, Tolerance), 1);
+		return false;
 	}
+	return true;
 }
 
-void FAutomationTestBase::TestEqual(const TCHAR* What, const FVector Actual, const FVector Expected, float Tolerance)
+bool FAutomationTestBase::TestEqual(const TCHAR* What, const FVector Actual, const FVector Expected, float Tolerance)
 {
 	if (!Expected.Equals(Actual, Tolerance))
 	{
 		AddError(FString::Printf(TEXT("Expected '%s' to be %s, but it was %s within tolerance %f."), What, *Expected.ToString(), *Actual.ToString(), Tolerance), 1);
+		return false;
 	}
+	return true;
 }
 
-void FAutomationTestBase::TestEqual(const TCHAR* What, const FColor Actual, const FColor Expected)
+bool FAutomationTestBase::TestEqual(const TCHAR* What, const FColor Actual, const FColor Expected)
 {
 	if (Expected != Actual)
 	{
 		AddError(FString::Printf(TEXT("Expected '%s' to be %s, but it was %s."), What, *Expected.ToString(), *Actual.ToString()), 1);
+		return false;
 	}
+	return true;
 }
 
-void FAutomationTestBase::TestEqual(const TCHAR* What, const TCHAR* Actual, const TCHAR* Expected)
+bool FAutomationTestBase::TestEqual(const TCHAR* What, const TCHAR* Actual, const TCHAR* Expected)
 {
 	if (FCString::Strcmp(Actual, Expected) != 0)
 	{
 		AddError(FString::Printf(TEXT("Expected '%s' to be %s, but it was %s."), What, Expected, Actual), 1);
+		return false;
 	}
+	return true;
 }
 
-void FAutomationTestBase::TestEqualInsensitive(const TCHAR* What, const TCHAR* Actual, const TCHAR* Expected)
+bool FAutomationTestBase::TestEqualInsensitive(const TCHAR* What, const TCHAR* Actual, const TCHAR* Expected)
 {
 	if (FCString::Stricmp(Actual, Expected) != 0)
 	{
 		AddError(FString::Printf(TEXT("Expected '%s' to be %s, but it was %s."), What, Expected, Actual), 1);
+		return false;
 	}
+	return true;
 }
 
-void FAutomationTestBase::TestFalse(const TCHAR* What, bool Value)
+bool FAutomationTestBase::TestFalse(const TCHAR* What, bool Value)
 {
 	if (Value)
 	{
 		AddError(FString::Printf(TEXT("Expected '%s' to be false."), What), 1);
+		return false;
 	}
+	return true;
 }
 
-void FAutomationTestBase::TestTrue(const TCHAR* What, bool Value)
+bool FAutomationTestBase::TestTrue(const TCHAR* What, bool Value)
 {
 	if (!Value)
 	{
 		AddError(FString::Printf(TEXT("Expected '%s' to be true."), What), 1);
+		return false;
 	}
+	return true;
 }
 
-void FAutomationTestBase::TestNull(const TCHAR* What, const void* Pointer)
+bool FAutomationTestBase::TestNull(const TCHAR* What, const void* Pointer)
 {
 	if ( Pointer != nullptr )
 	{
 		AddError(FString::Printf(TEXT("Expected '%s' to be null."), *What), 1);
+		return false;
 	}
+	return true;
 }
 
 bool FAutomationTestBase::IsExpectedError(const FString& Error)
diff --git a/Engine/Source/Runtime/Core/Public/Misc/AutomationTest.h b/Engine/Source/Runtime/Core/Public/Misc/AutomationTest.h
--- a/Engine/Source/Runtime/Core/Public/Misc/AutomationTest.h
+++ b/Engine/Source/Runtime/Core/Public/Misc/AutomationTest.h
@@ -1273,73 +1273,73 @@ public:
 
 public:
 
-	void TestEqual(const TCHAR* What, int32 Actual, int32 Expected);
-	void TestEqual(const TCHAR* What, int64 Actual, int64 Expected);
-	void TestEqual(const TCHAR* What, float Actual, float Expected, float Tolerance = KINDA_SMALL_NUMBER);
-	void TestEqual(const TCHAR* What, double Actual, double Expected, double Tolerance = KINDA_SMALL_NUMBER);
-	void TestEqual(const TCHAR* What, FVector Actual, FVector Expected, float Tolerance = KINDA_SMALL_NUMBER);
-	void TestEqual(const TCHAR* What, FColor Actual, FColor Expected);
-	void TestEqual(const TCHAR* What, const TCHAR* Actual, const TCHAR* Expected);
-	void TestEqualInsensitive(const TCHAR* What, const TCHAR* Actual, const TCHAR* Expected);
+	bool TestEqual(const TCHAR* What, int32 Actual, int32 Expected);
+	bool TestEqual(const TCHAR* What, int64 Actual, int64 Expected);
+	bool TestEqual(const TCHAR* What, float Actual, float Expected, float Tolerance = KINDA_SMALL_NUMBER);
+	bool TestEqual(const TCHAR* What, double Actual, double Expected, double Tolerance = KINDA_SMALL_NUMBER);
+	bool TestEqual(const TCHAR* What, FVector Actual, FVector Expected, float Tolerance = KINDA_SMALL_NUMBER);
+	bool TestEqual(const TCHAR* What, FColor Actual, FColor Expected);
+	bool TestEqual(const TCHAR* What, const TCHAR* Actual, const TCHAR* Expected);
+	bool TestEqualInsensitive(const TCHAR* What, const TCHAR* Actual, const TCHAR* Expected);
 
-	void TestEqual(const FString& What, int32 Actual, int32 Expected)
+	bool TestEqual(const FString& What, int32 Actual, int32 Expected)
 	{
-		TestEqual(*What, Actual, Expected);
+		return TestEqual(*What, Actual, Expected);
 	}
 
-	void TestEqual(const FString& What, float Actual, float Expected, float Tolerance = KINDA_SMALL_NUMBER)
+	bool TestEqual(const FString& What, float Actual, float Expected, float Tolerance = KINDA_SMALL_NUMBER)
 	{
-		TestEqual(*What, Actual, Expected, Tolerance);
+		return TestEqual(*What, Actual, Expected, Tolerance);
 	}
 
-	void TestEqual(const FString& What, double Actual, double Expected, double Tolerance = KINDA_SMALL_NUMBER)
+	bool TestEqual(const FString& What, double Actual, double Expected, double Tolerance = KINDA_SMALL_NUMBER)
 	{
-		TestEqual(*What, Actual, Expected, Tolerance);
+		return TestEqual(*What, Actual, Expected, Tolerance);
 	}
 
-	void TestEqual(const FString& What, FVector Actual, FVector Expected, float Tolerance = KINDA_SMALL_NUMBER)
+	bool TestEqual(const FString& What, FVector Actual, FVector Expected, float Tolerance = KINDA_SMALL_NUMBER)
 	{
-		TestEqual(*What, Actual, Expected, Tolerance);
+		return TestEqual(*What, Actual, Expected, Tolerance);
 	}
 
-	void TestEqual(const FString& What, FColor Actual, FColor Expected)
+	bool TestEqual(const FString& What, FColor Actual, FColor Expected)
 	{
-		TestEqual(*What, Actual, Expected);
+		return TestEqual(*What, Actual, Expected);
 	}
 
-	void TestEqual(const FString& What, const TCHAR* Actual, const TCHAR* Expected)
+	bool TestEqual(const FString& What, const TCHAR* Actual, const TCHAR* Expected)
 	{
-		TestEqual(*What, Actual, Expected);
+		return TestEqual(*What, Actual, Expected);
 	}
 
-	void TestEqual(const TCHAR* What, const FString& Actual, const TCHAR* Expected)
+	bool TestEqual(const TCHAR* What, const FString& Actual, const TCHAR* Expected)
 	{
-		TestEqualInsensitive(What, *Actual, Expected);
+		return TestEqualInsensitive(What, *Actual, Expected);
 	}
 
-	void TestEqual(const FString& What, const FString& Actual, const TCHAR* Expected)
+	bool TestEqual(const FString& What, const FString& Actual, const TCHAR* Expected)
 	{
-		TestEqualInsensitive(*What, *Actual, Expected);
+		return TestEqualInsensitive(*What, *Actual, Expected);
 	}
 
-	void TestEqual(const TCHAR* What, const TCHAR* Actual, const FString& Expected)
+	bool TestEqual(const TCHAR* What, const TCHAR* Actual, const FString& Expected)
 	{
-		TestEqualInsensitive(What, Actual, *Expected);
+		return TestEqualInsensitive(What, Actual, *Expected);
 	}
 
-	void TestEqual(const FString& What, const TCHAR* Actual, const FString& Expected)
+	bool TestEqual(const FString& What, const TCHAR* Actual, const FString& Expected)
 	{
-		TestEqualInsensitive(*What, Actual, *Expected);
+		return TestEqualInsensitive(*What, Actual, *Expected);
 	}
 
-	void TestEqual(const TCHAR* What, const FString& Actual, const FString& Expected)
+	bool TestEqual(const TCHAR* What, const FString& Actual, const FString& Expected)
 	{
-		TestEqualInsensitive(What, *Actual, *Expected);
+		return TestEqualInsensitive(What, *Actual, *Expected);
 	}
 
-	void TestEqual(const FString& What, const FString& Actual, const FString& Expected)
+	bool TestEqual(const FString& What, const FString& Actual, const FString& Expected)
 	{
-		TestEqualInsensitive(*What, *Actual, *Expected);
+		return TestEqualInsensitive(*What, *Actual, *Expected);
 	}
 
 	/**
@@ -1352,18 +1352,20 @@ public:
 	 * @see TestNotEqual
 	 */
 	template<typename ValueType> 
-	void TestEqual(const TCHAR* What, const ValueType& Actual, const ValueType& Expected)
+	bool TestEqual(const TCHAR* What, const ValueType& Actual, const ValueType& Expected)
 	{
 		if (Actual != Expected)
 		{
 			AddError(FString::Printf(TEXT("%s: The two values are not equal."), What), 1);
+			return false;
 		}
+		return true;
 	}
 
 	template<typename ValueType>
-	void TestEqual(const FString& What, const ValueType& Actual, const ValueType& Expected)
+	bool TestEqual(const FString& What, const ValueType& Actual, const ValueType& Expected)
 	{
-		TestEqual(*What, Actual, Expected);
+		return TestEqual(*What, Actual, Expected);
 	}
 
 
@@ -1375,11 +1377,11 @@ public:
 	 *
 	 * @see TestFalse
 	 */
-	void TestFalse(const TCHAR* What, bool Value);
+	bool TestFalse(const TCHAR* What, bool Value);
 
-	void TestFalse(const FString& What, bool Value)
+	bool TestFalse(const FString& What, bool Value)
 	{
-		TestFalse(*What, Value);
+		return TestFalse(*What, Value);
 	}
 
 	/**
@@ -1390,17 +1392,19 @@ public:
 	 *
 	 * @see TestValid
 	 */
-	template<typename ValueType> void TestInvalid(const TCHAR* Description, const TSharedPtr<ValueType>& SharedPointer)
+	template<typename ValueType> bool TestInvalid(const TCHAR* Description, const TSharedPtr<ValueType>& SharedPointer)
 	{
 		if (SharedPointer.IsValid())
 		{
 			AddError(FString::Printf(TEXT("%s: The shared pointer is valid."), Description), 1);
+			return false;
 		}
+		return true;
 	}
 
-	template<typename ValueType> void TestInvalid(const FString& Description, const TSharedPtr<ValueType>& SharedPointer)
+	template<typename ValueType> bool TestInvalid(const FString& Description, const TSharedPtr<ValueType>& SharedPointer)
 	{
-		TestInvalid(*Description, SharedPointer);
+		return TestInvalid(*Description, SharedPointer);
 	}
 
 	/**
@@ -1412,17 +1416,19 @@ public:
 	 *
 	 * @see TestEqual
 	 */
-	template<typename ValueType> void TestNotEqual(const TCHAR* Description, const ValueType& Actual, const ValueType& Expected)
+	template<typename ValueType> bool TestNotEqual(const TCHAR* Description, const ValueType& Actual, const ValueType& Expected)
 	{
 		if (Actual == Expected)
 		{
 			AddError(FString::Printf(TEXT("%s: The two values are equal."), Description), 1);
+			return false;
 		}
+		return true;
 	}
 
-	template<typename ValueType> void TestNotEqual(const FString& Description, const ValueType& Actual, const ValueType& Expected)
+	template<typename ValueType> bool TestNotEqual(const FString& Description, const ValueType& Actual, const ValueType& Expected)
 	{
-		TestNotEqual(*Description, Actual, Expected);
+		return TestNotEqual(*Description, Actual, Expected);
 	}
 
 	/**
@@ -1433,17 +1439,19 @@ public:
 	 *
 	 * @see TestNull
 	 */
-	template<typename ValueType> void TestNotNull(const TCHAR* What, ValueType* Pointer)
+	template<typename ValueType> bool TestNotNull(const TCHAR* What, ValueType* Pointer)
 	{
 		if (Pointer == nullptr)
 		{
 			AddError(FString::Printf(TEXT("Expected '%s' to be not null."), What), 1);
+			return false;
 		}
+		return true;
 	}
 
-	template<typename ValueType> void TestNotNull(const FString& What, ValueType* Pointer)
+	template<typename ValueType> bool TestNotNull(const FString& What, ValueType* Pointer)
 	{
-		TestNotNull(*What, Pointer);
+		return TestNotNull(*What, Pointer);
 	}
 
 	/**
@@ -1455,17 +1463,19 @@ public:
 	 *
 	 * @see TestSame
 	 */
-	template<typename ValueType> void TestNotSame(const TCHAR* Description, const ValueType& Actual, const ValueType& Expected)
+	template<typename ValueType> bool TestNotSame(const TCHAR* Description, const ValueType& Actual, const ValueType& Expected)
 	{
 		if (&Actual == &Expected)
 		{
 			AddError(FString::Printf(TEXT("%s: The two values are the same."), Description), 1);
+			return false;
 		}
+		return true;
 	}
 
-	template<typename ValueType> void TestNotSame(const FString& Description, const ValueType& Actual, const ValueType& Expected)
+	template<typename ValueType> bool TestNotSame(const FString& Description, const ValueType& Actual, const ValueType& Expected)
 	{
-		TestNotSame(*Description, Actual, Expected);
+		return TestNotSame(*Description, Actual, Expected);
 	}
 
 	/**
@@ -1476,11 +1486,11 @@ public:
 	 *
 	 * @see TestNotNull
 	 */
-	void TestNull(const TCHAR* What, const void* Pointer);
+	bool TestNull(const TCHAR* What, const void* Pointer);
 
-	void TestNull(const FString& What, const void* Pointer)
+	bool TestNull(const FString& What, const void* Pointer)
 	{
-		TestNull(*What, Pointer);
+		return TestNull(*What, Pointer);
 	}
 
 	/**
@@ -1492,17 +1502,19 @@ public:
 	 *
 	 * @see TestNotSame
 	 */
-	template<typename ValueType> void TestSame(const TCHAR* Description, const ValueType& Actual, const ValueType& Expected)
+	template<typename ValueType> bool TestSame(const TCHAR* Description, const ValueType& Actual, const ValueType& Expected)
 	{
 		if (&Actual != &Expected)
 		{
 			AddError(FString::Printf(TEXT("%s: The two values are not the same."), Description), 1);
+			return false;
 		}
+		return true;
 	}
 
-	template<typename ValueType> void TestSame(const FString& Description, const ValueType& Actual, const ValueType& Expected)
+	template<typename ValueType> bool TestSame(const FString& Description, const ValueType& Actual, const ValueType& Expected)
 	{
-		TestSame(*Description, Actual, Expected);
+		return TestSame(*Description, Actual, Expected);
 	}
 
 	/**
@@ -1513,11 +1525,11 @@ public:
 	 *
 	 * @see TestFalse
 	 */
-	void TestTrue(const TCHAR* What, bool Value);
+	bool TestTrue(const TCHAR* What, bool Value);
 
-	void TestTrue(const FString& What, bool Value)
+	bool TestTrue(const FString& What, bool Value)
 	{
-		TestTrue(*What, Value);
+		return TestTrue(*What, Value);
 	}
 
 	/** Macro version of above, uses the passed in expression as the description as well */
@@ -1531,17 +1543,19 @@ public:
 	 *
 	 * @see TestInvalid
 	 */
-	template<typename ValueType> void TestValid(const TCHAR* Description, const TSharedPtr<ValueType>& SharedPointer)
+	template<typename ValueType> bool TestValid(const TCHAR* Description, const TSharedPtr<ValueType>& SharedPointer)
 	{
 		if (!SharedPointer.IsValid())
 		{
 			AddError(FString::Printf(TEXT("%s: The shared pointer is not valid."), Description), 1);
+			return false;
 		}
+		return true;
 	}
 
-	template<typename ValueType> void TestValid(const FString& Description, const TSharedPtr<ValueType>& SharedPointer)
+	template<typename ValueType> bool TestValid(const FString& Description, const TSharedPtr<ValueType>& SharedPointer)
 	{
-		TestValid(*Description, SharedPointer);
+		return TestValid(*Description, SharedPointer);
 	}
 
 protected:
@@ -3082,6 +3096,99 @@ public: \
 #endif // #if WITH_AUTOMATION_WORKER
 
 
+/**
+ * Macros to make it easy to test state with one-liners: they will run the appropriate
+ * test method and, if the test fail, with execute `return false;`, which (if placed in
+ * the main test case method) will stop the test immediately.
+ *
+ * The error logging is already handled by the test method being called.
+ *
+ * As a result, you can easily test things that, if wrong, would potentially crash the test:
+ *
+ *		bool FMyEasyTest::RunTest(const FString& Parameters)
+ *		{
+ *			TArray<float> Data = GetSomeData();
+ *			int32 Index = GetSomeIndex();
+ *			UTEST_TRUE("Check valid index", Index < Data.Num());
+ *			float DataItem = Data[Index];   // Won't crash, the test exited on the previous 
+ *										    // line if index was invalid.
+ *			UTEST_TRUE("Check valid item", DataItem > 0.f);
+ *		}
+ *
+ */
+
+#define UTEST_EQUAL(What, Actual, Expected)\
+	if (!TestEqual(What, Actual, Expected))\
+	{\
+		return false;\
+	}
+
+#define UTEST_EQUAL_TOLERANCE(What, Actual, Expected, Tolerance)\
+	if (!TestEqual(What, Actual, Expected, Tolerance))\
+	{\
+		return false;\
+	}
+
+#define UTEST_EQUAL_INSENSITIVE(What, Actual, Expected)\
+	if (!TestEqualInsensitive(What, Actual, Expected))\
+	{\
+		return false;\
+	}
+
+#define UTEST_NOT_EQUAL(What, Actual, Expected)\
+	if (!TestNotEqual(What, Actual, Expected))\
+	{\
+		return false;\
+	}
+
+#define UTEST_SAME(What, Actual, Expected)\
+	if (!TestSame(What, Actual, Expected))\
+	{\
+		return false;\
+	}
+
+#define UTEST_NOT_SAME(What, Actual, Expected)\
+	if (!TestNotSame(What, Actual, Expected))\
+	{\
+		return false;\
+	}
+
+#define UTEST_TRUE(What, Value)\
+	if (!TestTrue(What, Value))\
+	{\
+		return false;\
+	}
+
+#define UTEST_FALSE(What, Value)\
+	if (!TestFalse(What, Value))\
+	{\
+		return false;\
+	}
+
+#define UTEST_VALID(What, SharedPointer)\
+	if (!TestValid(What, SharedPointer))\
+	{\
+		return false;\
+	}
+
+#define UTEST_INVALID(What, SharedPointer)\
+	if (!TestInvalid(What, SharedPointer))\
+	{\
+		return false;\
+	}
+
+#define UTEST_NULL(What, Pointer)\
+	if (!TestNull(What, Pointer))\
+	{\
+		return false;\
+	}
+
+#define UTEST_NOT_NULL(What, Pointer)\
+	if (!TestNotNull(What, Pointer))\
+	{\
+		return false;\
+	}
+
 //////////////////////////////////////////////////////////////////////////
 // Basic Latent Commands
 
diff --git a/Engine/Source/Runtime/MovieScene/Private/Tests/MovieSceneCompilerTests.cpp b/Engine/Source/Runtime/MovieScene/Private/Tests/MovieSceneCompilerTests.cpp
--- a/Engine/Source/Runtime/MovieScene/Private/Tests/MovieSceneCompilerTests.cpp
+++ b/Engine/Source/Runtime/MovieScene/Private/Tests/MovieSceneCompilerTests.cpp
@@ -333,21 +333,6 @@ bool FMovieSceneCompilerEmptySpaceOnTheFlyTest::RunTest(const FString& Parameter
 }
 
 
-#define _TEST_ASSERT(cond, msgfmt, ...)\
-	if (!ensure((cond)))\
-	{\
-		AddError(FString::Printf(TEXT(msgfmt), __VA_ARGS__));\
-		return false;\
-	}
-
-// This should use TestEqual methods on the test if those methods actually returned the true/false outcome...
-#define _TEST_ASSERT_EQUAL(actual, expected)\
-	if (!ensure(actual == expected))\
-	{\
-		AddError(FString::Printf(TEXT("Expected '%s' but actual value was '%s'."), *LexToString(expected), *LexToString(actual)));\
-		return false;\
-	}
-
 IMPLEMENT_SIMPLE_AUTOMATION_TEST(FMovieSceneCompilerSubSequencesTest, "System.Engine.Sequencer.Compiler.SubSequences", EAutomationTestFlags::EditorContext | EAutomationTestFlags::EngineFilter)
 bool FMovieSceneCompilerSubSequencesTest::RunTest(const FString& Parameters)
 {
@@ -405,16 +390,13 @@ bool FMovieSceneCompilerSubSequencesTest::RunTest(const FString& Parameters)
 		FMovieSceneCompiler::Compile(*RootSequence, Store);
 
 		const FMovieSceneEvaluationTemplate& RootTemplate = Store.AccessTemplate(*RootSequence);
-		_TEST_ASSERT_EQUAL(RootTemplate.EvaluationField.GetRanges().Num(), 3);
-		_TEST_ASSERT_EQUAL(RootTemplate.EvaluationField.GetRange(1), TRange<FFrameNumber>(0, 60));
+		UTEST_EQUAL("Ranges count", RootTemplate.EvaluationField.GetRanges().Num(), 3);
+		UTEST_EQUAL("First range", RootTemplate.EvaluationField.GetRange(1), TRange<FFrameNumber>(0, 60));
 		const FMovieSceneEvaluationGroup& EvalGroup = RootTemplate.EvaluationField.GetGroup(1);
-		_TEST_ASSERT_EQUAL(EvalGroup.SegmentPtrLUT[0].SequenceID, Shot1SequenceID);
+		UTEST_EQUAL("Sequence ID", EvalGroup.SegmentPtrLUT[0].SequenceID, Shot1SequenceID);
 	}
 
 	return true;
 }
 
-#undef _TEST_ASSERT
-#undef _TEST_ASSERT_EQUAL
-
 #endif // WITH_DEV_AUTOMATION_TESTS
