Adding Epic boilerplate license to modified SseMath header file.

diff --git a/Engine/Source/Runtime/Core/Public/Math/sse_mathfun.h b/Engine/Source/Runtime/Core/Public/Math/sse_mathfun.h
--- a/Engine/Source/Runtime/Core/Public/Math/sse_mathfun.h
+++ b/Engine/Source/Runtime/Core/Public/Math/sse_mathfun.h
@@ -1,3 +1,5 @@
+// Copyright 1998-2019 Epic Games, Inc. All Rights Reserved.
+
 /* SIMD (SSE1+MMX or SSE2) implementation of sin, cos, exp and log
 
    Inspired by Intel Approximate Math library, and based on the
