Removed check() which is always true

diff --git a/Engine/Source/Runtime/Core/Public/Async/MappedFileHandle.h b/Engine/Source/Runtime/Core/Public/Async/MappedFileHandle.h
--- a/Engine/Source/Runtime/Core/Public/Async/MappedFileHandle.h
+++ b/Engine/Source/Runtime/Core/Public/Async/MappedFileHandle.h
@@ -89,7 +89,6 @@ public:
 	IMappedFileHandle(size_t InFileSize)
 		: MappedFileSize(InFileSize)
 	{
-		check(MappedFileSize >= 0);
 		INC_DWORD_STAT(STAT_MappedFileHandles);
 	}
 	/** Destructor, also the only way to close the file handle. It is not legal to delete an async file with outstanding requests. You must always call WaitCompletion before deleting a request. **/
