Added initialization to satisfy PVS

diff --git a/Engine/Source/Runtime/Core/Public/Containers/IndirectArray.h b/Engine/Source/Runtime/Core/Public/Containers/IndirectArray.h
--- a/Engine/Source/Runtime/Core/Public/Containers/IndirectArray.h
+++ b/Engine/Source/Runtime/Core/Public/Containers/IndirectArray.h
@@ -188,7 +188,7 @@ public:
 		if (Ar.IsLoading())
 		{
 			// Load array.
-			int32 NewNum;
+			int32 NewNum = 0;
 			Ar << NewNum;
 			Empty(NewNum);
 			for (int32 Index = 0; Index < NewNum; Index++)
