Refactor blacklist classes
#rnx

diff --git a/Engine/Source/Runtime/Core/Private/Misc/BlacklistNames.cpp b/Engine/Source/Runtime/Core/Private/Misc/BlacklistNames.cpp
new file mode 100644
--- /dev/null
+++ b/Engine/Source/Runtime/Core/Private/Misc/BlacklistNames.cpp
@@ -0,0 +1,295 @@
+// Copyright 1998-2019 Epic Games, Inc. All Rights Reserved.
+
+#include "Misc/BlacklistNames.h"
+
+FBlacklistNames::FBlacklistNames() :
+	bSuppressOnFilterChanged(false)
+{
+}
+
+bool FBlacklistNames::PassesFilter(const FName Item) const
+{
+	if (Whitelist.Num() > 0 && !Whitelist.Contains(Item))
+	{
+		return false;
+	}
+
+	if (Blacklist.Contains(Item))
+	{
+		return false;
+	}
+
+	if (BlacklistAll.Num() > 0)
+	{
+		return false;
+	}
+
+	return true;
+}
+
+void FBlacklistNames::AddBlacklistItem(const FName OwnerName, const FName Item)
+{
+	Blacklist.FindOrAdd(Item).AddUnique(OwnerName);
+	
+	if (!bSuppressOnFilterChanged)
+	{
+		OnFilterChanged().Broadcast();
+	}
+}
+
+void FBlacklistNames::AddWhitelistItem(const FName OwnerName, const FName Item)
+{
+	Whitelist.FindOrAdd(Item).AddUnique(OwnerName);	
+	
+	if (!bSuppressOnFilterChanged)
+	{
+		OnFilterChanged().Broadcast();
+	}
+}
+
+void FBlacklistNames::AddBlacklistAll(const FName OwnerName)
+{
+	BlacklistAll.AddUnique(OwnerName);
+
+	if (!bSuppressOnFilterChanged)
+	{
+		OnFilterChanged().Broadcast();
+	}
+}
+
+bool FBlacklistNames::HasFiltering() const
+{
+	return Blacklist.Num() > 0 || Whitelist.Num() > 0 || BlacklistAll.Num() > 0;
+}
+
+void FBlacklistNames::UnregisterOwner(const FName OwnerName)
+{
+	for (auto It = Blacklist.CreateIterator(); It; ++It)
+	{
+		It->Value.Remove(OwnerName);
+		if (It->Value.Num() == 0)
+		{
+			It.RemoveCurrent();
+		}
+	}
+
+	for (auto It = Whitelist.CreateIterator(); It; ++It)
+	{
+		It->Value.Remove(OwnerName);
+		if (It->Value.Num() == 0)
+		{
+			It.RemoveCurrent();
+		}
+	}
+
+	BlacklistAll.Remove(OwnerName);
+
+	if (!bSuppressOnFilterChanged)
+	{
+		OnFilterChanged().Broadcast();
+	}
+}
+
+void FBlacklistNames::Append(const FBlacklistNames& Other)
+{
+	{
+		TGuardValue<bool> Guard(bSuppressOnFilterChanged, true);
+
+		for (const auto& It : Other.Blacklist)
+		{
+			for (const auto& OwnerName : It.Value)
+			{
+				AddBlacklistItem(OwnerName, It.Key);
+			}
+		}
+
+		for (const auto& It : Other.Whitelist)
+		{
+			for (const auto& OwnerName : It.Value)
+			{
+				AddWhitelistItem(OwnerName, It.Key);
+			}
+		}
+
+		for (const auto& OwnerName : Other.BlacklistAll)
+		{
+			AddBlacklistAll(OwnerName);
+		}
+	}
+
+	if (!bSuppressOnFilterChanged)
+	{
+		OnFilterChanged().Broadcast();
+	}
+}
+
+/** FBlacklistPaths */
+FBlacklistPaths::FBlacklistPaths() :
+	bSuppressOnFilterChanged(false)
+{
+}
+
+bool FBlacklistPaths::PassesFilter(const FString& Item) const
+{
+	if (Whitelist.Num() > 0 && !Whitelist.Contains(Item))
+	{
+		return false;
+	}
+
+	if (Blacklist.Contains(Item))
+	{
+		return false;
+	}
+
+	if (BlacklistAll.Num() > 0)
+	{
+		return false;
+	}
+
+	return true;
+}
+
+bool FBlacklistPaths::PassesFilter(const FName Item) const
+{
+	return PassesFilter(Item.ToString());
+}
+
+bool FBlacklistPaths::PassesStartsWithFilter(const FString& Item) const
+{
+	if (Whitelist.Num() > 0)
+	{
+		bool bPassedWhitelist = false;
+		for (const auto& Other : Whitelist)
+		{
+			if (Item.StartsWith(Other.Key) && (Item.Len() <= Other.Key.Len() || Item[Other.Key.Len()] == TEXT('/')))
+			{
+				bPassedWhitelist = true;
+				break;
+			}
+		}
+
+		if (!bPassedWhitelist)
+		{
+			return false;
+		}
+	}
+
+	if (Blacklist.Num() > 0)
+	{
+		for (const auto& Other : Blacklist)
+		{
+			if (Item.StartsWith(Other.Key) && (Item.Len() <= Other.Key.Len() || Item[Other.Key.Len()] == TEXT('/')))
+			{
+				return false;
+			}
+		}
+	}
+
+	if (BlacklistAll.Num() > 0)
+	{
+		return false;
+	}
+
+	return true;
+}
+
+bool FBlacklistPaths::PassesStartsWithFilter(const FName Item) const
+{
+	return PassesStartsWithFilter(Item.ToString());
+}
+
+void FBlacklistPaths::AddBlacklistItem(const FName OwnerName, const FString& Item)
+{
+	Blacklist.FindOrAdd(Item).AddUnique(OwnerName);
+	
+	if (!bSuppressOnFilterChanged)
+	{
+		OnFilterChanged().Broadcast();
+	}
+}
+
+void FBlacklistPaths::AddWhitelistItem(const FName OwnerName, const FString& Item)
+{
+	Whitelist.FindOrAdd(Item).AddUnique(OwnerName);	
+	
+	if (!bSuppressOnFilterChanged)
+	{
+		OnFilterChanged().Broadcast();
+	}
+}
+
+void FBlacklistPaths::AddBlacklistAll(const FName OwnerName)
+{
+	BlacklistAll.AddUnique(OwnerName);
+
+	if (!bSuppressOnFilterChanged)
+	{
+		OnFilterChanged().Broadcast();
+	}
+}
+
+bool FBlacklistPaths::HasFiltering() const
+{
+	return Blacklist.Num() > 0 || Whitelist.Num() > 0 || BlacklistAll.Num() > 0;
+}
+
+void FBlacklistPaths::UnregisterOwner(const FName OwnerName)
+{
+	for (auto It = Blacklist.CreateIterator(); It; ++It)
+	{
+		It->Value.Remove(OwnerName);
+		if (It->Value.Num() == 0)
+		{
+			It.RemoveCurrent();
+		}
+	}
+
+	for (auto It = Whitelist.CreateIterator(); It; ++It)
+	{
+		It->Value.Remove(OwnerName);
+		if (It->Value.Num() == 0)
+		{
+			It.RemoveCurrent();
+		}
+	}
+
+	BlacklistAll.Remove(OwnerName);
+
+	if (!bSuppressOnFilterChanged)
+	{
+		OnFilterChanged().Broadcast();
+	}
+}
+
+void FBlacklistPaths::Append(const FBlacklistPaths& Other)
+{
+	{
+		TGuardValue<bool> Guard(bSuppressOnFilterChanged, true);
+
+		for (const auto& It : Other.Blacklist)
+		{
+			for (const auto& OwnerName : It.Value)
+			{
+				AddBlacklistItem(OwnerName, It.Key);
+			}
+		}
+
+		for (const auto& It : Other.Whitelist)
+		{
+			for (const auto& OwnerName : It.Value)
+			{
+				AddWhitelistItem(OwnerName, It.Key);
+			}
+		}
+
+		for (const auto& OwnerName : Other.BlacklistAll)
+		{
+			AddBlacklistAll(OwnerName);
+		}
+	}
+
+	if (!bSuppressOnFilterChanged)
+	{
+		OnFilterChanged().Broadcast();
+	}
+}
diff --git a/Engine/Source/Runtime/Core/Public/Misc/BlacklistNames.h b/Engine/Source/Runtime/Core/Public/Misc/BlacklistNames.h
--- a/Engine/Source/Runtime/Core/Public/Misc/BlacklistNames.h
+++ b/Engine/Source/Runtime/Core/Public/Misc/BlacklistNames.h
@@ -3,169 +3,131 @@
 #pragma once
 
 #include "CoreTypes.h"
+#include "Templates/SharedPointer.h"
+#include "Delegates/Delegate.h"
 
-/** 
- * Filter for blacklisting or whitelisting items with ownership to allow change cleanup
- */
-class FNamedBlacklist
+/** List of owner names that requested a specific item filtered, allows unregistering specific set of changes by a given plugin or system */
+typedef TArray<FName> FBlacklistOwners;
+
+class CORE_API FBlacklistNames : public TSharedFromThis<FBlacklistNames>
 {
 public:
+	FBlacklistNames();
+	~FBlacklistNames() {}
+
+	/** Returns true if passes filter restrictions using exact match */
+	bool PassesFilter(const FName Item) const;
 
-	/** Returns true if passes filter restrictions */
-	bool PassesFilter(const FName ItemName) const
-	{
-		if (BlacklistAll.Num() > 0)
-		{
-			return false;
-		}
-		else if (Blacklist.Contains(ItemName))
-		{
-			return false;
-		}
-		else if (Whitelist.Num() > 0 && !Whitelist.Contains(ItemName))
-		{
-			return false;
-		}
-
-		return true;
-	}
-	
 	/** Add item to blacklist, this specific item will be filtered out */
-	void AddBlacklistItem(const FName OwnerName, const FName ItemName)
-	{
-		Blacklist.FindOrAdd(ItemName).AddUnique(OwnerName);
-	}
+	void AddBlacklistItem(const FName OwnerName, const FName Item);
 
 	/** Add item to whitelist after which all items not in the whitelist will be filtered out */
-	void AddWhitelistItem(const FName OwnerName, const FName ItemName)
-	{
-		Whitelist.FindOrAdd(ItemName).AddUnique(OwnerName);
-	}
+	void AddWhitelistItem(const FName OwnerName, const FName Item);
 
 	/** Set to filter out all items */
-	void AddBlacklistAll(const FName OwnerName)
-	{
-		BlacklistAll.AddUnique(OwnerName);
-	}
+	void AddBlacklistAll(const FName OwnerName);
+	
+	/** True if has filters active */
+	bool HasFiltering() const;
 
 	/** Removes all filtering changes associated with a specific owner name */
-	void UnregisterOwner(const FName OwnerName)
-	{
-		for (auto It = Blacklist.CreateIterator(); It; ++It)
-		{
-			It->Value.Remove(OwnerName);
-			if (It->Value.Num() == 0)
-			{
-				It.RemoveCurrent();
-			}
-		}
-
-		for (auto It = Whitelist.CreateIterator(); It; ++It)
-		{
-			It->Value.Remove(OwnerName);
-			if (It->Value.Num() == 0)
-			{
-				It.RemoveCurrent();
-			}
-		}
-
-		BlacklistAll.Remove(OwnerName);
-	}
-
-	/** True if has filters active */
-	bool HasFiltering() const
-	{
-		return Blacklist.Num() > 0 || Whitelist.Num() > 0 || BlacklistAll.Num() > 0;
-	}
+	void UnregisterOwner(const FName OwnerName);
 
 	/** Combine two filters together */
-	void Append(const FNamedBlacklist& Other)
-	{
-		for (const auto& It : Other.Blacklist)
-		{
-			for (const auto& OwnerName : It.Value)
-			{
-				AddBlacklistItem(OwnerName, It.Key);
-			}
-		}
-
-		for (const auto& It : Other.Whitelist)
-		{
-			for (const auto& OwnerName : It.Value)
-			{
-				AddWhitelistItem(OwnerName, It.Key);
-			}
-		}
-
-		for (const auto& OwnerName : Other.BlacklistAll)
-		{
-			AddBlacklistAll(OwnerName);
-		}
-	}
-
-private:
-
-	/** List of owner names that requested a specific item filtered, allows unregistering specific set of changes by a given plugin or system */
-	typedef TArray<FName> FNamedBlacklistOwners;
+	void Append(const FBlacklistNames& Other);
+
+	/** Get raw blacklist */
+	const TMap<FName, FBlacklistOwners>& GetBlacklist() const { return Blacklist; }
+	
+	/** Get raw whitelist */
+	const TMap<FName, FBlacklistOwners>& GetWhitelist() const { return Whitelist; }
+
+	/** Are all items set to be filtered out */
+	bool IsBlacklistAll() const { return Blacklist.Num() > 0; }
+
+	/** Triggered when filter changes */
+	FSimpleMulticastDelegate& OnFilterChanged() { return OnFilterChangedDelegate; }
+
+protected:
 
 	/** List if items to filter out */
-	TMap<FName, FNamedBlacklistOwners> Blacklist;
+	TMap<FName, FBlacklistOwners> Blacklist;
 
 	/** List of items to allow, if not empty all items will be filtered out unless they are in the list */
-	TMap<FName, FNamedBlacklistOwners> Whitelist;
+	TMap<FName, FBlacklistOwners> Whitelist;
 
 	/** List of owner names that requested all items to be filtered out */
-	FNamedBlacklistOwners BlacklistAll;
+	FBlacklistOwners BlacklistAll;
+
+	/** Triggered when filter changes */
+	FSimpleMulticastDelegate OnFilterChangedDelegate;
+
+	/** Temporarily prevent delegate from being triggered */
+	bool bSuppressOnFilterChanged;
 };
 
-/**
- * Collection of blacklisting filters
- */
-class FNamedBlacklistCollection
+class CORE_API FBlacklistPaths : public TSharedFromThis<FBlacklistPaths>
 {
 public:
+	FBlacklistPaths();
+	~FBlacklistPaths() {}
+	
+	/** Returns true if passes filter restrictions using exact match */
+	bool PassesFilter(const FString& Item) const;
+
+	/** Returns true if passes filter restrictions using exact match */
+	bool PassesFilter(const FName Item) const;
+
+	/** Returns true if passes filter restrictions for path */
+	bool PassesStartsWithFilter(const FString& Item) const;
+
+	/** Returns true if passes filter restrictions for path */
+	bool PassesStartsWithFilter(const FName Item) const;
+
+	/** Add item to blacklist, this specific item will be filtered out */
+	void AddBlacklistItem(const FName OwnerName, const FString& Item);
+
+	/** Add item to whitelist after which all items not in the whitelist will be filtered out */
+	void AddWhitelistItem(const FName OwnerName, const FString& Item);
+
+	/** Set to filter out all items */
+	void AddBlacklistAll(const FName OwnerName);
+	
+	/** True if has filters active */
+	bool HasFiltering() const;
+
+	/** Removes all filtering changes associated with a specific owner name */
+	void UnregisterOwner(const FName OwnerName);
+
+	/** Combine two filters together */
+	void Append(const FBlacklistPaths& Other);
+
+	/** Get raw blacklist */
+	const TMap<FString, FBlacklistOwners>& GetBlacklist() const { return Blacklist; }
+	
+	/** Get raw whitelist */
+	const TMap<FString, FBlacklistOwners>& GetWhitelist() const { return Whitelist; }
+
+	/** Are all items set to be filtered out */
+	bool IsBlacklistAll() const { return Blacklist.Num() > 0; }
+	
+	/** Triggered when filter changes */
+	FSimpleMulticastDelegate& OnFilterChanged() { return OnFilterChangedDelegate; }
+
+protected:
+
+	/** List if items to filter out */
+	TMap<FString, FBlacklistOwners> Blacklist;
+
+	/** List of items to allow, if not empty all items will be filtered out unless they are in the list */
+	TMap<FString, FBlacklistOwners> Whitelist;
+
+	/** List of owner names that requested all items to be filtered out */
+	FBlacklistOwners BlacklistAll;
+	
+	/** Triggered when filter changes */
+	FSimpleMulticastDelegate OnFilterChangedDelegate;
 
-	/** Returns true if passes filter restrictions */
-	bool PassesFilter(const FName GroupName, const FName ItemName) const
-	{
-		const FNamedBlacklist* Filter = Filters.Find(GroupName);
-		return !Filter || Filter->PassesFilter(ItemName);
-	}
-
-	/** Add item to blacklist */
-	void AddBlacklistItem(const FName OwnerName, const FName GroupName, const FName ItemName)
-	{
-		Filters.FindOrAdd(GroupName).AddBlacklistItem(OwnerName, ItemName);
-	}
-
-	/** Add item to whitelist */
-	void AddWhitelistItem(const FName OwnerName, const FName GroupName, const FName ItemName)
-	{
-		Filters.FindOrAdd(GroupName).AddWhitelistItem(OwnerName, ItemName);
-	}
-
-	/** Sets group to fully blacklisted */
-	void AddBlacklistAll(const FName OwnerName, const FName GroupName)
-	{
-		Filters.FindOrAdd(GroupName).AddBlacklistAll(OwnerName);
-	}
-
-	/** Remove filtering associated with owner name */
-	void UnregisterOwner(const FName OwnerName)
-	{
-		for (auto GroupIt = Filters.CreateIterator(); GroupIt; ++GroupIt)
-		{
-			FNamedBlacklist& Filter = GroupIt->Value;
-			Filter.UnregisterOwner(OwnerName);
-			if (!Filter.HasFiltering())
-			{
-				GroupIt.RemoveCurrent();
-			}
-		}
-	}
-
-private:
-
-	/** List of filters */
-	TMap<FName, FNamedBlacklist> Filters;
+	/** Temporarily prevent delegate from being triggered */
+	bool bSuppressOnFilterChanged;
 };
diff --git a/Engine/Source/Runtime/Slate/Private/Framework/Docking/TabManager.cpp b/Engine/Source/Runtime/Slate/Private/Framework/Docking/TabManager.cpp
--- a/Engine/Source/Runtime/Slate/Private/Framework/Docking/TabManager.cpp
+++ b/Engine/Source/Runtime/Slate/Private/Framework/Docking/TabManager.cpp
@@ -1624,7 +1624,7 @@ bool FTabManager::HasTabSpawner(FName TabId) const
 	return Spawner != nullptr;
 }
 
-FNamedBlacklist& FTabManager::GetTabBlacklist()
+FBlacklistNames& FTabManager::GetTabBlacklist()
 {
 	return TabBlacklist;
 }
diff --git a/Engine/Source/Runtime/Slate/Public/Framework/Docking/TabManager.h b/Engine/Source/Runtime/Slate/Public/Framework/Docking/TabManager.h
--- a/Engine/Source/Runtime/Slate/Public/Framework/Docking/TabManager.h
+++ b/Engine/Source/Runtime/Slate/Public/Framework/Docking/TabManager.h
@@ -803,7 +803,7 @@ class SLATE_API FTabManager : public TSharedFromThis<FTabManager>
 		TSharedPtr<SDockTab> GetOwnerTab() { return OwnerTabPtr.Pin(); }
 
 		/** Returns filter for additional control over available tabs */
-		FNamedBlacklist& GetTabBlacklist();
+		FBlacklistNames& GetTabBlacklist();
 
 	protected:
 		void InvokeTabForMenu( FName TabId );
@@ -972,7 +972,7 @@ class SLATE_API FTabManager : public TSharedFromThis<FTabManager>
 		bool bCanDoDragOperation;
 
 		/** Allow systems to dynamically hide tabs */
-		FNamedBlacklist TabBlacklist;
+		FBlacklistNames TabBlacklist;
 };
 
 
diff --git a/Engine/Source/Runtime/Slate/Public/Framework/MultiBox/ToolMenuBase.h b/Engine/Source/Runtime/Slate/Public/Framework/MultiBox/ToolMenuBase.h
--- a/Engine/Source/Runtime/Slate/Public/Framework/MultiBox/ToolMenuBase.h
+++ b/Engine/Source/Runtime/Slate/Public/Framework/MultiBox/ToolMenuBase.h
@@ -74,7 +74,7 @@ struct SLATE_API FCustomizedToolMenu
 	UPROPERTY()
 	TArray<FName> SectionOrder;
 
-	FNamedBlacklist BlacklistFilter;
+	FBlacklistNames BlacklistFilter;
 
 	FCustomizedToolMenuEntry* FindEntry(const FName InEntryName);
 	const FCustomizedToolMenuEntry* FindEntry(const FName InEntryName) const;
