Fix a linker error with Clang in debug configurations when using TStringBuilderImpl from another module

#rnx

diff --git a/Engine/Source/Runtime/Core/Public/Misc/StringBuilder.h b/Engine/Source/Runtime/Core/Public/Misc/StringBuilder.h
--- a/Engine/Source/Runtime/Core/Public/Misc/StringBuilder.h
+++ b/Engine/Source/Runtime/Core/Public/Misc/StringBuilder.h
@@ -360,8 +360,6 @@ constexpr inline SIZE_T GetNum(const TStringBuilderImpl<CharType>& Builder)
 
 //////////////////////////////////////////////////////////////////////////
 
-extern template class TStringBuilderImpl<ANSICHAR>;
-
 class FAnsiStringBuilderBase : public TStringBuilderImpl<ANSICHAR>
 {
 protected:
@@ -415,8 +413,6 @@ struct TIsContiguousContainer<TAnsiStringBuilder<N>> { enum { Value = true }; };
 
 //////////////////////////////////////////////////////////////////////////
 
-extern template class TStringBuilderImpl<TCHAR>;
-
 class FStringBuilderBase : public TStringBuilderImpl<TCHAR>
 {
 protected:
