Backing out PR (submitted as CL10610795 ) for TInterval::ToString() because it interacts poorly with some as-yet-unidentified change which was robomerged later from another stream

Nothing in the tree currently depends on this functionality (it came in as a PR)

diff --git a/Engine/Source/Runtime/Core/Public/Math/Interval.h b/Engine/Source/Runtime/Core/Public/Math/Interval.h
--- a/Engine/Source/Runtime/Core/Public/Math/Interval.h
+++ b/Engine/Source/Runtime/Core/Public/Math/Interval.h
@@ -7,7 +7,6 @@
 #include "Templates/UnrealTypeTraits.h"
 #include "Math/NumericLimits.h"
 #include "Math/UnrealMathUtility.h"
-#include "Containers/UnrealString.h"
 
 /**
  * Type traits for Arithmetic interval.
@@ -180,16 +179,6 @@ public:
 		return ElementType();
 	}
 
-	/**
-	 * Get a textual representation of the interval.
-	 *
-	 * @return Text describing the interval.
-	 */
-	FString ToString() const
-	{
-		return FString::Printf(TEXT("Min=%s Max=%s"), *LexToString(Min), *LexToString(Max));
-	}
-
 public:
 
 	/**
@@ -278,17 +267,3 @@ public:
 
 DEFINE_INTERVAL_WRAPPER_STRUCT(FFloatInterval, float)
 DEFINE_INTERVAL_WRAPPER_STRUCT(FInt32Interval, int32)
-
-template<>
-inline FString TInterval<float>::ToString() const
-{
-	// This is more efficient than going through LexToString
-	return FString::Printf(TEXT("Min=%f Max=%f"), Min, Max);
-}
-
-template<>
-inline FString TInterval<int32>::ToString() const
-{
-	// This is more efficient than going through LexToString
-	return FString::Printf(TEXT("Min=%d Max=%d"), Min, Max);
-}
