Deprecated FCoreDelegates::PakFileMountedCallback in favor of FCoreDelegates::OnPakFileMounted, which takes a second argument that specifies the chunk ID of the package (or INDEX_NONE).

diff --git a/Engine/Source/Runtime/Core/Private/Misc/CoreDelegates.cpp b/Engine/Source/Runtime/Core/Private/Misc/CoreDelegates.cpp
--- a/Engine/Source/Runtime/Core/Private/Misc/CoreDelegates.cpp
+++ b/Engine/Source/Runtime/Core/Private/Misc/CoreDelegates.cpp
@@ -31,7 +31,10 @@ FCoreDelegates::FOnMountPak FCoreDelegates::OnMountPak;
 FCoreDelegates::FOnUnmountPak FCoreDelegates::OnUnmountPak;
 FCoreDelegates::FOnOptimizeMemoryUsageForMountedPaks FCoreDelegates::OnOptimizeMemoryUsageForMountedPaks;
 
+FCoreDelegates::FOnPakFileMounted FCoreDelegates::OnPakFileMounted;
+PRAGMA_DISABLE_DEPRECATION_WARNINGS
 FCoreDelegates::FPakFileMountedDelegate FCoreDelegates::PakFileMountedCallback;
+PRAGMA_ENABLE_DEPRECATION_WARNINGS
 FCoreDelegates::FNewFileAddedDelegate FCoreDelegates::NewFileAddedDelegate;
 FCoreDelegates::FNoPakFilesMountedDelegate FCoreDelegates::NoPakFilesMountedDelegate;
 FCoreDelegates::FOnFileOpenedForReadFromPakFile FCoreDelegates::OnFileOpenedForReadFromPakFile;
diff --git a/Engine/Source/Runtime/Core/Public/Misc/CoreDelegates.h b/Engine/Source/Runtime/Core/Public/Misc/CoreDelegates.h
--- a/Engine/Source/Runtime/Core/Public/Misc/CoreDelegates.h
+++ b/Engine/Source/Runtime/Core/Public/Misc/CoreDelegates.h
@@ -72,6 +72,9 @@ public:
 	// delegate for handling when a new pak file is successfully mounted passes in the name of the mounted pak file
 	DECLARE_MULTICAST_DELEGATE_OneParam(FPakFileMountedDelegate, const TCHAR*);
 
+	// delegate for handling when a new pak file is successfully mounted passes in the name of the pak file and its chunk ID (or INDEX_NONE)
+	DECLARE_MULTICAST_DELEGATE_TwoParams(FOnPakFileMounted, const TCHAR*, const int32);
+
 	// delegate to let other systems no that no paks were mounted, in case something wants to handle that case
 	DECLARE_MULTICAST_DELEGATE(FNoPakFilesMountedDelegate);
 
@@ -160,7 +163,9 @@ public:
 	// Callback to optimize memeory for currently mounted paks
 	static FOnOptimizeMemoryUsageForMountedPaks OnOptimizeMemoryUsageForMountedPaks;
 
-	// After a pakfile is mounted this callback is called for each new file
+	// After a pakfile is mounted this is called
+	static FOnPakFileMounted OnPakFileMounted;
+	UE_DEPRECATED(4.25, "FCoreDelegates::PakFileMountedCallback is deprecated. Use FCoreDelegates::OnPakFileMounted instead.")
 	static FPakFileMountedDelegate PakFileMountedCallback;
 
 	// After a file is added this is called
diff --git a/Engine/Source/Runtime/Engine/Private/Streaming/StreamingManagerTexture.cpp b/Engine/Source/Runtime/Engine/Private/Streaming/StreamingManagerTexture.cpp
--- a/Engine/Source/Runtime/Engine/Private/Streaming/StreamingManagerTexture.cpp
+++ b/Engine/Source/Runtime/Engine/Private/Streaming/StreamingManagerTexture.cpp
@@ -134,7 +134,7 @@ FRenderAssetStreamingManager::FRenderAssetStreamingManager()
 
 	FCoreUObjectDelegates::GetPreGarbageCollectDelegate().AddRaw(this, &FRenderAssetStreamingManager::OnPreGarbageCollect);
 
-	FCoreDelegates::PakFileMountedCallback.AddRaw(this, &FRenderAssetStreamingManager::OnPakFileMounted);
+	FCoreDelegates::OnPakFileMounted.AddRaw(this, &FRenderAssetStreamingManager::OnPakFileMounted);
 }
 
 FRenderAssetStreamingManager::~FRenderAssetStreamingManager()
@@ -189,7 +189,7 @@ void FRenderAssetStreamingManager::OnPreGarbageCollect()
 
 
 
-void FRenderAssetStreamingManager::OnPakFileMounted(const TCHAR* PakFilename)
+void FRenderAssetStreamingManager::OnPakFileMounted(const TCHAR* PakFilename, const int32 ChunkId)
 {
 	// clear the cached file exists checks which failed as they may now be loaded
 	bNewFilesLoaded = true;
diff --git a/Engine/Source/Runtime/Engine/Private/Streaming/StreamingManagerTexture.h b/Engine/Source/Runtime/Engine/Private/Streaming/StreamingManagerTexture.h
--- a/Engine/Source/Runtime/Engine/Private/Streaming/StreamingManagerTexture.h
+++ b/Engine/Source/Runtime/Engine/Private/Streaming/StreamingManagerTexture.h
@@ -408,7 +408,7 @@ private:
 	TArray<int32> InflightRenderAssets;
 
 	TMap<FString, bool> CachedFileExistsChecks;
-	void OnPakFileMounted(const TCHAR* PakFilename);
+	void OnPakFileMounted(const TCHAR* PakFilename, const int32 ChunkId);
 
 	// A critical section use around code that could be called in parallel with NotifyPrimitiveUpdated() or NotifyPrimitiveUpdated_Concurrent().
 	FCriticalSection CriticalSection;
diff --git a/Engine/Source/Runtime/PakFile/Private/IPlatformFilePak.cpp b/Engine/Source/Runtime/PakFile/Private/IPlatformFilePak.cpp
--- a/Engine/Source/Runtime/PakFile/Private/IPlatformFilePak.cpp
+++ b/Engine/Source/Runtime/PakFile/Private/IPlatformFilePak.cpp
@@ -6233,10 +6233,10 @@ bool FPakPlatformFile::Mount(const TCHAR* InPakFilename, uint32 PakOrder, const
 
 		if (bSuccess)
 		{
-			if (FCoreDelegates::PakFileMountedCallback.IsBound())
-			{
-				FCoreDelegates::PakFileMountedCallback.Broadcast(InPakFilename);
-			}
+			PRAGMA_DISABLE_DEPRECATION_WARNINGS
+			FCoreDelegates::PakFileMountedCallback.Broadcast(InPakFilename);
+			PRAGMA_ENABLE_DEPRECATION_WARNINGS
+			FCoreDelegates::OnPakFileMounted.Broadcast(InPakFilename, Pak->ChunkID);
 			if (FCoreDelegates::NewFileAddedDelegate.IsBound())
 			{
 				TArray<FString> Filenames;
