Copy the metadata at the endCapture since the FinalizeCsvFile can happen a few frames later. We need to ensure that the metadata correspond to what happening BeginCapture/Endcapture. We enqueue the metadata, the same way we enqueue the start/stop command, so the will always match.

#rnx

diff --git a/Engine/Source/Runtime/Core/Private/ProfilingDebugging/CsvProfiler.cpp b/Engine/Source/Runtime/Core/Private/ProfilingDebugging/CsvProfiler.cpp
--- a/Engine/Source/Runtime/Core/Private/ProfilingDebugging/CsvProfiler.cpp
+++ b/Engine/Source/Runtime/Core/Private/ProfilingDebugging/CsvProfiler.cpp
@@ -66,7 +66,7 @@ const char * GIgnoreWaitStatName =  "[IGNORE]";
 
 TAutoConsoleVariable<int32> CVarCsvBlockOnCaptureEnd(
 	TEXT("csv.BlockOnCaptureEnd"), 
-	0,
+	1,
 	TEXT("When 1, blocks the game thread until the CSV file has been written completely when the capture is ended.\r\n")
 	TEXT("When 0, the game thread is not blocked whilst the file is written."),
 	ECVF_Default
@@ -2458,6 +2458,24 @@ FCsvProfiler::FCsvProfiler()
 	FCoreDelegates::OnBeginFrameRT.AddStatic(CsvProfilerBeginFrameRT);
 	FCoreDelegates::OnEndFrameRT.AddStatic(CsvProfilerEndFrameRT);
 #endif
+
+	// add constant metadata
+	FString PlatformStr = FString::Printf(TEXT("%s"), ANSI_TO_TCHAR(FPlatformProperties::IniPlatformName()));
+	FString BuildConfigurationStr = EBuildConfigurations::ToString(FApp::GetBuildConfiguration());
+	FString CommandlineStr = FString("\"") + FCommandLine::Get() + FString("\"");
+	// Strip newlines
+	CommandlineStr.ReplaceInline(TEXT("\n"), TEXT(""));
+	CommandlineStr.ReplaceInline(TEXT("\r"), TEXT(""));
+	FString BuildVersionString = FApp::GetBuildVersion();
+	FString EngineVersionString = FEngineVersion::Current().ToString();
+
+	MetadataMap.FindOrAdd(TEXT("Platform")) = PlatformStr;
+	MetadataMap.FindOrAdd(TEXT("Config")) = BuildConfigurationStr;
+	MetadataMap.FindOrAdd(TEXT("BuildVersion")) = BuildVersionString;
+	MetadataMap.FindOrAdd(TEXT("EngineVersion")) = EngineVersionString;
+
+	// Commandline has to be last for parsing, since it might include commas
+	MetadataMap.FindOrAdd(TEXT("Commandline")) = CommandlineStr;
 }
 
 FCsvProfiler::~FCsvProfiler()
@@ -2703,7 +2721,7 @@ void FCsvProfiler::EndFrame()
 					FinalizeCsvFile();
 					bCaptureComplete = true;
 				}
-				else if (CVarCsvBlockOnCaptureEnd.GetValueOnGameThread() == 0)
+				else if (CVarCsvBlockOnCaptureEnd.GetValueOnGameThread() == 1)
 				{
 					// Suspend the hang and hitch heartbeats, as this is a long running task.
 					FSlowHeartBeatScope SuspendHeartBeat;
@@ -2822,6 +2840,14 @@ TSharedFuture<FString> FCsvProfiler::EndCapture(FGraphEventRef EventToSignal)
 		}
 	});
 
+	// Copy the metadata array for the next FinalizeCsvFile
+	TMap<FString, FString> CopyMetadataMap;
+	{
+		FScopeLock Lock(&MetadataCS);
+		CopyMetadataMap = MetadataMap;
+	}
+	MetadataQueue.Enqueue(MoveTemp(CopyMetadataMap));
+
 	TSharedFuture<FString> Future = Completion->GetFuture().Share();
 	CommandQueue.Enqueue(FCsvCaptureCommand(ECsvCommandType::Stop, GCsvProfilerFrameNumber, Completion, Future));
 
@@ -2840,35 +2866,14 @@ void FCsvProfiler::FinalizeCsvFile()
 	// Do a final process of the stat data
 	ProcessStatData();
 
-	// Copy the metadata array into a local one
-	TMap<FString, FString> LocalMetadata;
-	{
-		FScopeLock Lock(&MetadataCS);
-		LocalMetadata = MetadataMap;
-	}
-
-	// Add metadata
-	FString PlatformStr = FString::Printf(TEXT("%s"), ANSI_TO_TCHAR(FPlatformProperties::IniPlatformName()));
-	FString BuildConfigurationStr = LexToString(FApp::GetBuildConfiguration());
-	FString CommandlineStr = FString("\"") + FCommandLine::Get() + FString("\"");
-	// Strip newlines
-	CommandlineStr.ReplaceInline(TEXT("\n"), TEXT(""));
-	CommandlineStr.ReplaceInline(TEXT("\r"), TEXT(""));
-	FString BuildVersionString = FApp::GetBuildVersion();
-	FString EngineVersionString = FEngineVersion::Current().ToString();
-
-	LocalMetadata.FindOrAdd(TEXT("Platform")) = PlatformStr;
-	LocalMetadata.FindOrAdd(TEXT("Config")) = BuildConfigurationStr;
-	LocalMetadata.FindOrAdd(TEXT("DeviceProfile")) = DeviceProfileName;
-	LocalMetadata.FindOrAdd(TEXT("BuildVersion")) = BuildVersionString;
-	LocalMetadata.FindOrAdd(TEXT("EngineVersion")) = EngineVersionString;
-
-	// Commandline has to be last for parsing, since it might include commas
-	LocalMetadata.FindOrAdd(TEXT("Commandline")) = CommandlineStr;
-
 	uint64 MemoryBytesAtEndOfCapture = CsvWriter->GetAllocatedSize();
+	
+	// Get the queued metadata for the next csv finalize
+	TMap<FString, FString> CurrentMetadata;
+	MetadataQueue.Dequeue(CurrentMetadata);
 
-	CsvWriter->Finalize(LocalMetadata);
+
+	CsvWriter->Finalize(CurrentMetadata);
 
 	delete CsvWriter;
 	CsvWriter = nullptr;
@@ -2890,8 +2895,7 @@ void FCsvProfiler::FinalizeCsvFile()
 
 void FCsvProfiler::SetDeviceProfileName(FString InDeviceProfileName)
 {
-	LLM_SCOPE(ELLMTag::CsvProfiler);
-	DeviceProfileName = InDeviceProfileName;
+	CSV_METADATA(TEXT("DeviceProfile"), *InDeviceProfileName);
 }
 
 /** Push/pop events */
diff --git a/Engine/Source/Runtime/Core/Public/ProfilingDebugging/CsvProfiler.h b/Engine/Source/Runtime/Core/Public/ProfilingDebugging/CsvProfiler.h
--- a/Engine/Source/Runtime/Core/Public/ProfilingDebugging/CsvProfiler.h
+++ b/Engine/Source/Runtime/Core/Public/ProfilingDebugging/CsvProfiler.h
@@ -323,12 +323,10 @@ private:
 	FCsvProfilerProcessingThread* ProcessingThread;
 
 	FEvent* FileWriteBlockingEvent;
-
-	FString DeviceProfileName;
-
 	FThreadSafeCounter IsShuttingDown;
 
 	TMap<FString, FString> MetadataMap;
+	TQueue<TMap<FString, FString>> MetadataQueue;
 	FCriticalSection MetadataCS;
 
 	class FCsvStreamWriter* CsvWriter;
diff --git a/Engine/Source/Runtime/Engine/Private/DeviceProfiles/DeviceProfileManager.cpp b/Engine/Source/Runtime/Engine/Private/DeviceProfiles/DeviceProfileManager.cpp
--- a/Engine/Source/Runtime/Engine/Private/DeviceProfiles/DeviceProfileManager.cpp
+++ b/Engine/Source/Runtime/Engine/Private/DeviceProfiles/DeviceProfileManager.cpp
@@ -751,7 +751,7 @@ void UDeviceProfileManager::SetActiveDeviceProfile( UDeviceProfile* DeviceProfil
 	ActiveDeviceProfile = DeviceProfile;
 
 #if CSV_PROFILER
-	FCsvProfiler::Get()->SetDeviceProfileName(ActiveDeviceProfile->GetName());
+	CSV_METADATA(TEXT("DeviceProfile"), *GetActiveDeviceProfileName());
 #endif
 }
 
