Fix deprecation
#rnx

diff --git a/Engine/Source/Runtime/Core/Private/ProfilingDebugging/CsvProfiler.cpp b/Engine/Source/Runtime/Core/Private/ProfilingDebugging/CsvProfiler.cpp
--- a/Engine/Source/Runtime/Core/Private/ProfilingDebugging/CsvProfiler.cpp
+++ b/Engine/Source/Runtime/Core/Private/ProfilingDebugging/CsvProfiler.cpp
@@ -2461,7 +2461,7 @@ FCsvProfiler::FCsvProfiler()
 
 	// add constant metadata
 	FString PlatformStr = FString::Printf(TEXT("%s"), ANSI_TO_TCHAR(FPlatformProperties::IniPlatformName()));
-	FString BuildConfigurationStr = EBuildConfigurations::ToString(FApp::GetBuildConfiguration());
+	FString BuildConfigurationStr = LexToString(FApp::GetBuildConfiguration());
 	FString CommandlineStr = FString("\"") + FCommandLine::Get() + FString("\"");
 	// Strip newlines
 	CommandlineStr.ReplaceInline(TEXT("\n"), TEXT(""));
