Add event tag supression to OutputDeviceMemory
Fix ring buffer writing when bounds crossed

diff --git a/Engine/Source/Runtime/Core/Private/Misc/OutputDeviceMemory.cpp b/Engine/Source/Runtime/Core/Private/Misc/OutputDeviceMemory.cpp
--- a/Engine/Source/Runtime/Core/Private/Misc/OutputDeviceMemory.cpp
+++ b/Engine/Source/Runtime/Core/Private/Misc/OutputDeviceMemory.cpp
@@ -14,12 +14,13 @@
 
 #define DUMP_LOG_ON_EXIT (!NO_LOGGING && PLATFORM_DESKTOP && (!UE_BUILD_SHIPPING || USE_LOGGING_IN_SHIPPING))
 
-FOutputDeviceMemory::FOutputDeviceMemory(int32 InPreserveSize /*= 256 * 1024*/, int32 InBufferSize /*= 2048 * 1024*/)
+FOutputDeviceMemory::FOutputDeviceMemory(int32 InPreserveSize /*= 256 * 1024*/, int32 InBufferSize /*= 2048 * 1024*/, bool bInSuppressEventTag /*= false */)
 : ArchiveProxy(*this)
 , BufferStartPos(0)
 , BufferLength(0)
 , PreserveSize(InPreserveSize)
 {
+	bSuppressEventTag = bInSuppressEventTag;
 #if DUMP_LOG_ON_EXIT
 	const FString LogFileName = FPlatformOutputDevices::GetAbsoluteLogFilename();
 	FOutputDeviceFile::CreateBackupCopy(*LogFileName);
@@ -29,13 +30,18 @@ FOutputDeviceMemory::FOutputDeviceMemory(int32 InPreserveSize /*= 256 * 1024*/,
 	checkf(InBufferSize >= InPreserveSize * 2, TEXT("FOutputDeviceMemory buffer size should be >= 2x PreserveSize"));
 
 	Buffer.AddUninitialized(InBufferSize);
-	Logf(TEXT("Log file open, %s"), FPlatformTime::StrTimestamp());
+	if (!bSuppressEventTag)
+	{
+		Logf(TEXT("Log file open, %s"), FPlatformTime::StrTimestamp());
+	}
 }
 
 void FOutputDeviceMemory::TearDown() 
 {
-	Logf(TEXT("Log file closed, %s"), FPlatformTime::StrTimestamp());
-
+	if (!bSuppressEventTag)
+	{
+		Logf(TEXT("Log file closed, %s"), FPlatformTime::StrTimestamp());
+	}
 	// Dump on exit
 #if DUMP_LOG_ON_EXIT
 	const FString LogFileName = FPlatformOutputDevices::GetAbsoluteLogFilename();
@@ -100,6 +106,7 @@ void FOutputDeviceMemory::SerializeToBuffer(ANSICHAR* Data, int32 Length)
 		}
 
 		FMemory::Memcpy(Buffer.GetData() + WritePos, Data, WriteLength * sizeof(ANSICHAR));
+		Data += WriteLength;
 		Length -= WriteLength;
 	}
 }
diff --git a/Engine/Source/Runtime/Core/Public/Misc/OutputDeviceMemory.h b/Engine/Source/Runtime/Core/Public/Misc/OutputDeviceMemory.h
--- a/Engine/Source/Runtime/Core/Public/Misc/OutputDeviceMemory.h
+++ b/Engine/Source/Runtime/Core/Public/Misc/OutputDeviceMemory.h
@@ -33,7 +33,7 @@ public:
 	 * @param InPreserveSize	Bytes of the rung buffer not to overwrite (startup info etc)
 	 * @param InBufferSize		Maximum size of the memory ring buffer
 	 */
-	FOutputDeviceMemory(int32 InPreserveSize = 256 * 1024, int32 InBufferSize = 2048 * 1024);
+	FOutputDeviceMemory(int32 InPreserveSize = 256 * 1024, int32 InBufferSize = 2048 * 1024, bool bInSuppressEventTag = false);
 
 	/** Dumps the contents of the buffer to an archive */
 	virtual void Dump(FArchive& Ar) override;
