Back out changelist 10806937 to fix broken CSV metadata

diff --git a/Engine/Source/Runtime/Core/Private/ProfilingDebugging/CsvProfiler.cpp b/Engine/Source/Runtime/Core/Private/ProfilingDebugging/CsvProfiler.cpp
--- a/Engine/Source/Runtime/Core/Private/ProfilingDebugging/CsvProfiler.cpp
+++ b/Engine/Source/Runtime/Core/Private/ProfilingDebugging/CsvProfiler.cpp
@@ -66,7 +66,7 @@ const char * GIgnoreWaitStatName =  "[IGNORE]";
 
 TAutoConsoleVariable<int32> CVarCsvBlockOnCaptureEnd(
 	TEXT("csv.BlockOnCaptureEnd"), 
-	1,
+	0,
 	TEXT("When 1, blocks the game thread until the CSV file has been written completely when the capture is ended.\r\n")
 	TEXT("When 0, the game thread is not blocked whilst the file is written."),
 	ECVF_Default
@@ -2458,24 +2458,6 @@ FCsvProfiler::FCsvProfiler()
 	FCoreDelegates::OnBeginFrameRT.AddStatic(CsvProfilerBeginFrameRT);
 	FCoreDelegates::OnEndFrameRT.AddStatic(CsvProfilerEndFrameRT);
 #endif
-
-	// add constant metadata
-	FString PlatformStr = FString::Printf(TEXT("%s"), ANSI_TO_TCHAR(FPlatformProperties::IniPlatformName()));
-	FString BuildConfigurationStr = LexToString(FApp::GetBuildConfiguration());
-	FString CommandlineStr = FString("\"") + FCommandLine::Get() + FString("\"");
-	// Strip newlines
-	CommandlineStr.ReplaceInline(TEXT("\n"), TEXT(""));
-	CommandlineStr.ReplaceInline(TEXT("\r"), TEXT(""));
-	FString BuildVersionString = FApp::GetBuildVersion();
-	FString EngineVersionString = FEngineVersion::Current().ToString();
-
-	MetadataMap.FindOrAdd(TEXT("Platform")) = PlatformStr;
-	MetadataMap.FindOrAdd(TEXT("Config")) = BuildConfigurationStr;
-	MetadataMap.FindOrAdd(TEXT("BuildVersion")) = BuildVersionString;
-	MetadataMap.FindOrAdd(TEXT("EngineVersion")) = EngineVersionString;
-
-	// Commandline has to be last for parsing, since it might include commas
-	MetadataMap.FindOrAdd(TEXT("Commandline")) = CommandlineStr;
 }
 
 FCsvProfiler::~FCsvProfiler()
@@ -2721,7 +2703,7 @@ void FCsvProfiler::EndFrame()
 					FinalizeCsvFile();
 					bCaptureComplete = true;
 				}
-				else if (CVarCsvBlockOnCaptureEnd.GetValueOnGameThread() == 1)
+				else if (CVarCsvBlockOnCaptureEnd.GetValueOnGameThread() == 0)
 				{
 					// Suspend the hang and hitch heartbeats, as this is a long running task.
 					FSlowHeartBeatScope SuspendHeartBeat;
@@ -2840,14 +2822,6 @@ TSharedFuture<FString> FCsvProfiler::EndCapture(FGraphEventRef EventToSignal)
 		}
 	});
 
-	// Copy the metadata array for the next FinalizeCsvFile
-	TMap<FString, FString> CopyMetadataMap;
-	{
-		FScopeLock Lock(&MetadataCS);
-		CopyMetadataMap = MetadataMap;
-	}
-	MetadataQueue.Enqueue(MoveTemp(CopyMetadataMap));
-
 	TSharedFuture<FString> Future = Completion->GetFuture().Share();
 	CommandQueue.Enqueue(FCsvCaptureCommand(ECsvCommandType::Stop, GCsvProfilerFrameNumber, Completion, Future));
 
@@ -2866,14 +2840,35 @@ void FCsvProfiler::FinalizeCsvFile()
 	// Do a final process of the stat data
 	ProcessStatData();
 
-	uint64 MemoryBytesAtEndOfCapture = CsvWriter->GetAllocatedSize();
-	
-	// Get the queued metadata for the next csv finalize
-	TMap<FString, FString> CurrentMetadata;
-	MetadataQueue.Dequeue(CurrentMetadata);
+	// Copy the metadata array into a local one
+	TMap<FString, FString> LocalMetadata;
+	{
+		FScopeLock Lock(&MetadataCS);
+		LocalMetadata = MetadataMap;
+	}
 
+	// Add metadata
+	FString PlatformStr = FString::Printf(TEXT("%s"), ANSI_TO_TCHAR(FPlatformProperties::IniPlatformName()));
+	FString BuildConfigurationStr = EBuildConfigurations::ToString(FApp::GetBuildConfiguration());
+	FString CommandlineStr = FString("\"") + FCommandLine::Get() + FString("\"");
+	// Strip newlines
+	CommandlineStr.ReplaceInline(TEXT("\n"), TEXT(""));
+	CommandlineStr.ReplaceInline(TEXT("\r"), TEXT(""));
+	FString BuildVersionString = FApp::GetBuildVersion();
+	FString EngineVersionString = FEngineVersion::Current().ToString();
+
+	LocalMetadata.FindOrAdd(TEXT("Platform")) = PlatformStr;
+	LocalMetadata.FindOrAdd(TEXT("Config")) = BuildConfigurationStr;
+	LocalMetadata.FindOrAdd(TEXT("DeviceProfile")) = DeviceProfileName;
+	LocalMetadata.FindOrAdd(TEXT("BuildVersion")) = BuildVersionString;
+	LocalMetadata.FindOrAdd(TEXT("EngineVersion")) = EngineVersionString;
+
+	// Commandline has to be last for parsing, since it might include commas
+	LocalMetadata.FindOrAdd(TEXT("Commandline")) = CommandlineStr;
 
-	CsvWriter->Finalize(CurrentMetadata);
+	uint64 MemoryBytesAtEndOfCapture = CsvWriter->GetAllocatedSize();
+	
+	CsvWriter->Finalize(LocalMetadata);
 
 	delete CsvWriter;
 	CsvWriter = nullptr;
@@ -2895,7 +2890,8 @@ void FCsvProfiler::FinalizeCsvFile()
 
 void FCsvProfiler::SetDeviceProfileName(FString InDeviceProfileName)
 {
-	CSV_METADATA(TEXT("DeviceProfile"), *InDeviceProfileName);
+	LLM_SCOPE(ELLMTag::CsvProfiler);
+	DeviceProfileName = InDeviceProfileName;
 }
 
 /** Push/pop events */
diff --git a/Engine/Source/Runtime/Core/Public/ProfilingDebugging/CsvProfiler.h b/Engine/Source/Runtime/Core/Public/ProfilingDebugging/CsvProfiler.h
--- a/Engine/Source/Runtime/Core/Public/ProfilingDebugging/CsvProfiler.h
+++ b/Engine/Source/Runtime/Core/Public/ProfilingDebugging/CsvProfiler.h
@@ -323,10 +323,12 @@ private:
 	FCsvProfilerProcessingThread* ProcessingThread;
 
 	FEvent* FileWriteBlockingEvent;
+
+	FString DeviceProfileName;
+
 	FThreadSafeCounter IsShuttingDown;
 
 	TMap<FString, FString> MetadataMap;
-	TQueue<TMap<FString, FString>> MetadataQueue;
 	FCriticalSection MetadataCS;
 
 	class FCsvStreamWriter* CsvWriter;
diff --git a/Engine/Source/Runtime/Engine/Private/DeviceProfiles/DeviceProfileManager.cpp b/Engine/Source/Runtime/Engine/Private/DeviceProfiles/DeviceProfileManager.cpp
--- a/Engine/Source/Runtime/Engine/Private/DeviceProfiles/DeviceProfileManager.cpp
+++ b/Engine/Source/Runtime/Engine/Private/DeviceProfiles/DeviceProfileManager.cpp
@@ -751,7 +751,7 @@ void UDeviceProfileManager::SetActiveDeviceProfile( UDeviceProfile* DeviceProfil
 	ActiveDeviceProfile = DeviceProfile;
 
 #if CSV_PROFILER
-	CSV_METADATA(TEXT("DeviceProfile"), *GetActiveDeviceProfileName());
+	FCsvProfiler::Get()->SetDeviceProfileName(ActiveDeviceProfile->GetName());
 #endif
 }
 
