FPaths::GetBaseFilename now takes advantage of move semantics
#rnx

diff --git a/Engine/Source/Runtime/Core/Private/Misc/Paths.cpp b/Engine/Source/Runtime/Core/Private/Misc/Paths.cpp
--- a/Engine/Source/Runtime/Core/Private/Misc/Paths.cpp
+++ b/Engine/Source/Runtime/Core/Private/Misc/Paths.cpp
@@ -709,23 +709,26 @@ FString FPaths::GetCleanFilename(FString&& InPath)
 	return MoveTemp(InPath);
 }
 
-FString FPaths::GetBaseFilename( const FString& InPath, bool bRemovePath )
+FString FPaths::GetBaseFilename(FString InPath, bool bRemovePath )
 {
-	FString Wk = bRemovePath ? GetCleanFilename(InPath) : InPath;
+	FString Wk = bRemovePath ? GetCleanFilename(MoveTemp(InPath)) : MoveTemp(InPath);
 
 	// remove the extension
-	int32 ExtPos = Wk.Find(TEXT("."), ESearchCase::CaseSensitive, ESearchDir::FromEnd);
+	const int32 ExtPos = Wk.Find(TEXT("."), ESearchCase::CaseSensitive, ESearchDir::FromEnd);
 	
-	// determine the position of the path/leaf separator
-	int32 LeafPos = INDEX_NONE;
-	if (!bRemovePath)
+	if (ExtPos != INDEX_NONE)
 	{
-		LeafPos = Wk.FindLastCharByPredicate(UE4Paths_Private::IsSlashOrBackslash);
-	}
+		// determine the position of the path/leaf separator
+		int32 LeafPos = INDEX_NONE;
+		if (!bRemovePath)
+		{
+			LeafPos = Wk.FindLastCharByPredicate(UE4Paths_Private::IsSlashOrBackslash);
+		}
 
-	if (ExtPos != INDEX_NONE && (LeafPos == INDEX_NONE || ExtPos > LeafPos))
-	{
-		Wk.LeftInline(ExtPos);
+		if (LeafPos == INDEX_NONE || ExtPos > LeafPos)
+		{
+			Wk.LeftInline(ExtPos);
+		}
 	}
 
 	return Wk;
diff --git a/Engine/Source/Runtime/Core/Public/Misc/Paths.h b/Engine/Source/Runtime/Core/Public/Misc/Paths.h
--- a/Engine/Source/Runtime/Core/Public/Misc/Paths.h
+++ b/Engine/Source/Runtime/Core/Public/Misc/Paths.h
@@ -393,7 +393,7 @@ public:
 	static FString GetCleanFilename(FString&& InPath);
 
 	// Returns the same thing as GetCleanFilename, but without the extension
-	static FString GetBaseFilename( const FString& InPath, bool bRemovePath=true );
+	static FString GetBaseFilename(FString InPath, bool bRemovePath=true );
 
 	// Returns the path in front of the filename
 	static FString GetPath(const FString& InPath);
