Fix physics guid being incremented accidentally (which broke chaos serialization)

diff --git a/Engine/Source/Runtime/Core/Private/UObject/DevObjectVersion.cpp b/Engine/Source/Runtime/Core/Private/UObject/DevObjectVersion.cpp
--- a/Engine/Source/Runtime/Core/Private/UObject/DevObjectVersion.cpp
+++ b/Engine/Source/Runtime/Core/Private/UObject/DevObjectVersion.cpp
@@ -166,7 +166,7 @@ const FGuid FDestructionObjectVersion::GUID(0x174F1F0B, 0xB4C645A5, 0xB13F2EE8,
 FDevVersionRegistration GRegisterDestructionObjectVersion(FDestructionObjectVersion::GUID, FDestructionObjectVersion::LatestVersion, TEXT("Dev-Destruction"));
 
 // Unique Physics Object version id
-const FGuid FExternalPhysicsCustomObjectVersion::GUID(0x35F94A84, 0xE258406D, 0xA31809F6, 0x9610247D);
+const FGuid FExternalPhysicsCustomObjectVersion::GUID(0x35F94A83, 0xE258406C, 0xA31809F5, 0x9610247C);
 // Register Physics custom version with Core
 FDevVersionRegistration GRegisterExternalPhysicsCustomVersion(FExternalPhysicsCustomObjectVersion::GUID, FExternalPhysicsCustomObjectVersion::LatestVersion, TEXT("Dev-Physics-Ext"));
 
diff --git a/Engine/Source/Runtime/Engine/Private/PhysicsEngine/Experimental/ChaosDerivedData.cpp b/Engine/Source/Runtime/Engine/Private/PhysicsEngine/Experimental/ChaosDerivedData.cpp
--- a/Engine/Source/Runtime/Engine/Private/PhysicsEngine/Experimental/ChaosDerivedData.cpp
+++ b/Engine/Source/Runtime/Engine/Private/PhysicsEngine/Experimental/ChaosDerivedData.cpp
@@ -21,7 +21,7 @@ const TCHAR* FChaosDerivedDataCooker::GetPluginName() const
 
 const TCHAR* FChaosDerivedDataCooker::GetVersionString() const
 {
-	return TEXT("385FFCA8CED24783AD05635A6835CF95");
+	return TEXT("41954A2DFA074735B0B409CF01F3598A");
 }
 
 FString FChaosDerivedDataCooker::GetPluginSpecificCacheKeySuffix() const
