Added more LLM tags for chaos

diff --git a/Engine/Source/Runtime/Core/Private/HAL/LowLevelMemTracker.cpp b/Engine/Source/Runtime/Core/Private/HAL/LowLevelMemTracker.cpp
--- a/Engine/Source/Runtime/Core/Private/HAL/LowLevelMemTracker.cpp
+++ b/Engine/Source/Runtime/Core/Private/HAL/LowLevelMemTracker.cpp
@@ -91,6 +91,9 @@ DECLARE_LLM_MEMORY_STAT(TEXT("Chaos"), STAT_ChaosLLM, STATGROUP_LLMFULL);
 DECLARE_LLM_MEMORY_STAT(TEXT("ChaosGeometry"), STAT_ChaosGeometryLLM, STATGROUP_LLMFULL);
 DECLARE_LLM_MEMORY_STAT(TEXT("ChaosAcceleration"), STAT_ChaosAccelerationLLM, STATGROUP_LLMFULL);
 DECLARE_LLM_MEMORY_STAT(TEXT("ChaosParticles"), STAT_ChaosParticlesLLM, STATGROUP_LLMFULL);
+DECLARE_LLM_MEMORY_STAT(TEXT("ChaosLandscape"), STAT_ChaosLandscapeLLM, STATGROUP_LLMFULL);
+DECLARE_LLM_MEMORY_STAT(TEXT("ChaosTrimesh"), STAT_ChaosTrimeshLLM, STATGROUP_LLMFULL);
+DECLARE_LLM_MEMORY_STAT(TEXT("ChaosConvex"), STAT_ChaosConvexLLM, STATGROUP_LLMFULL);
 DECLARE_LLM_MEMORY_STAT(TEXT("EnginePreInit"), STAT_EnginePreInitLLM, STATGROUP_LLMFULL);
 DECLARE_LLM_MEMORY_STAT(TEXT("EngineInit"), STAT_EngineInitLLM, STATGROUP_LLMFULL);
 DECLARE_LLM_MEMORY_STAT(TEXT("Rendering Thread"), STAT_RenderingThreadLLM, STATGROUP_LLMFULL);
diff --git a/Engine/Source/Runtime/Core/Public/HAL/LowLevelMemTracker.h b/Engine/Source/Runtime/Core/Public/HAL/LowLevelMemTracker.h
--- a/Engine/Source/Runtime/Core/Public/HAL/LowLevelMemTracker.h
+++ b/Engine/Source/Runtime/Core/Public/HAL/LowLevelMemTracker.h
@@ -190,6 +190,9 @@ enum class ELLMTagSet : uint8
 	macro(ChaosGeometry,						"ChaosGeometry",				GET_STATFNAME(STAT_ChaosGeometryLLM),						GET_STATFNAME(STAT_ChaosSummaryLLM),			ELLMTag::Physics)\
 	macro(ChaosAcceleration,					"ChaosAcceleration",			GET_STATFNAME(STAT_ChaosAccelerationLLM),					GET_STATFNAME(STAT_ChaosSummaryLLM),			ELLMTag::Physics)\
 	macro(ChaosParticles,						"ChaosParticles",				GET_STATFNAME(STAT_ChaosParticlesLLM),						GET_STATFNAME(STAT_ChaosSummaryLLM),			ELLMTag::Physics)\
+	macro(ChaosLandscape,						"ChaosLandscape",				GET_STATFNAME(STAT_ChaosLandscapeLLM),						GET_STATFNAME(STAT_ChaosSummaryLLM),			ELLMTag::Physics)\
+	macro(ChaosTrimesh,							"ChaosTrimesh",					GET_STATFNAME(STAT_ChaosTrimeshLLM),						GET_STATFNAME(STAT_ChaosSummaryLLM),			ELLMTag::Physics)\
+	macro(ChaosConvex,							"ChaosConvex",					GET_STATFNAME(STAT_ChaosConvexLLM),							GET_STATFNAME(STAT_ChaosSummaryLLM),			ELLMTag::Physics)\
 	macro(EnginePreInitMemory,					"EnginePreInit",				GET_STATFNAME(STAT_EnginePreInitLLM),						GET_STATFNAME(STAT_EngineSummaryLLM),			-1)\
 	macro(EngineInitMemory,						"EngineInit",					GET_STATFNAME(STAT_EngineInitLLM),							GET_STATFNAME(STAT_EngineSummaryLLM),			-1)\
 	macro(RenderingThreadMemory,				"RenderingThread",				GET_STATFNAME(STAT_RenderingThreadLLM),						GET_STATFNAME(STAT_EngineSummaryLLM),			-1)\
diff --git a/Engine/Source/Runtime/Engine/Private/PhysicsEngine/Experimental/ChaosDerivedDataReader.cpp b/Engine/Source/Runtime/Engine/Private/PhysicsEngine/Experimental/ChaosDerivedDataReader.cpp
--- a/Engine/Source/Runtime/Engine/Private/PhysicsEngine/Experimental/ChaosDerivedDataReader.cpp
+++ b/Engine/Source/Runtime/Engine/Private/PhysicsEngine/Experimental/ChaosDerivedDataReader.cpp
@@ -26,8 +26,16 @@ FChaosDerivedDataReader<T, d>::FChaosDerivedDataReader(FBulkDataInterface* InBul
 	}
 	else
 	{
-		LLM_SCOPE(ELLMTag::ChaosGeometry);
-		ChaosAr << ConvexImplicitObjects << TrimeshImplicitObjects << UVInfo << FaceRemap;
+		{
+			LLM_SCOPE(ELLMTag::ChaosConvex);
+			ChaosAr << ConvexImplicitObjects;
+		}
+
+		{
+			LLM_SCOPE(ELLMTag::ChaosTrimesh);
+			ChaosAr << TrimeshImplicitObjects << UVInfo << FaceRemap;
+		}
+		
 
 		bReadSuccessful = true;
 	}
diff --git a/Engine/Source/Runtime/Landscape/Private/LandscapeCollision.cpp b/Engine/Source/Runtime/Landscape/Private/LandscapeCollision.cpp
--- a/Engine/Source/Runtime/Landscape/Private/LandscapeCollision.cpp
+++ b/Engine/Source/Runtime/Landscape/Private/LandscapeCollision.cpp
@@ -549,7 +549,7 @@ void ULandscapeHeightfieldCollisionComponent::ApplyWorldOffset(const FVector& In
 void ULandscapeHeightfieldCollisionComponent::CreateCollisionObject()
 {
 #if WITH_CHAOS
-	LLM_SCOPE(ELLMTag::ChaosGeometry);
+	LLM_SCOPE(ELLMTag::ChaosLandscape);
 #endif
 	// If we have not created a heightfield yet - do it now.
 	if (!IsValidRef(HeightfieldRef))
