Have explicit move and const ref versions of FPaths::GetBaseFilename
#rnx

diff --git a/Engine/Source/Runtime/Core/Private/Misc/Paths.cpp b/Engine/Source/Runtime/Core/Private/Misc/Paths.cpp
--- a/Engine/Source/Runtime/Core/Private/Misc/Paths.cpp
+++ b/Engine/Source/Runtime/Core/Private/Misc/Paths.cpp
@@ -709,13 +709,14 @@ FString FPaths::GetCleanFilename(FString&& InPath)
 	return MoveTemp(InPath);
 }
 
-FString FPaths::GetBaseFilename(FString InPath, bool bRemovePath )
+template<typename T>
+FString GetBaseFilenameImpl(T&& InPath, bool bRemovePath)
 {
-	FString Wk = bRemovePath ? GetCleanFilename(MoveTemp(InPath)) : MoveTemp(InPath);
+	FString Wk = bRemovePath ? FPaths::GetCleanFilename(Forward<T>(InPath)) : Forward<T>(InPath);
 
 	// remove the extension
 	const int32 ExtPos = Wk.Find(TEXT("."), ESearchCase::CaseSensitive, ESearchDir::FromEnd);
-	
+
 	if (ExtPos != INDEX_NONE)
 	{
 		// determine the position of the path/leaf separator
@@ -734,6 +735,16 @@ FString FPaths::GetBaseFilename(FString InPath, bool bRemovePath )
 	return Wk;
 }
 
+FString FPaths::GetBaseFilename(const FString& InPath, bool bRemovePath)
+{
+	return GetBaseFilenameImpl(InPath, bRemovePath);
+}
+
+FString FPaths::GetBaseFilename(FString&& InPath, bool bRemovePath)
+{
+	return GetBaseFilenameImpl(MoveTemp(InPath), bRemovePath);
+}
+
 FString FPaths::GetPath(const FString& InPath)
 {
 	int32 Pos = InPath.FindLastCharByPredicate(UE4Paths_Private::IsSlashOrBackslash);
diff --git a/Engine/Source/Runtime/Core/Public/Misc/Paths.h b/Engine/Source/Runtime/Core/Public/Misc/Paths.h
--- a/Engine/Source/Runtime/Core/Public/Misc/Paths.h
+++ b/Engine/Source/Runtime/Core/Public/Misc/Paths.h
@@ -393,7 +393,10 @@ public:
 	static FString GetCleanFilename(FString&& InPath);
 
 	// Returns the same thing as GetCleanFilename, but without the extension
-	static FString GetBaseFilename(FString InPath, bool bRemovePath=true );
+	static FString GetBaseFilename(const FString& InPath, bool bRemovePath=true );
+
+	// Returns the same thing as GetCleanFilename, but without the extension
+	static FString GetBaseFilename(FString&& InPath, bool bRemovePath = true);
 
 	// Returns the path in front of the filename
 	static FString GetPath(const FString& InPath);
