Fix for inconsistent function signatures between the Windows SDK and our minimal API wrappers.

diff --git a/Engine/Source/Runtime/Core/Public/Windows/MinimalWindowsApi.h b/Engine/Source/Runtime/Core/Public/Windows/MinimalWindowsApi.h
--- a/Engine/Source/Runtime/Core/Public/Windows/MinimalWindowsApi.h
+++ b/Engine/Source/Runtime/Core/Public/Windows/MinimalWindowsApi.h
@@ -41,6 +41,7 @@ struct HWND__;
 struct HKEY__;
 struct HDC__;
 struct HICON__;
+struct _RTL_SRWLOCK;
 
 // Other types
 struct tagPROCESSENTRY32W;
@@ -88,12 +89,14 @@ namespace Windows
 	typedef _OVERLAPPED* LPOVERLAPPED;
 	typedef _LARGE_INTEGER* LPLARGE_INTEGER;
 
-	typedef struct _RTL_SRWLOCK
+	typedef _RTL_SRWLOCK RTL_SRWLOCK, *PRTL_SRWLOCK;
+	typedef RTL_SRWLOCK *PSRWLOCK;
+
+	// Opaque SRWLOCK structure
+	struct SRWLOCK
 	{
 		void* Ptr;
-	} RTL_SRWLOCK, *PRTL_SRWLOCK;
-
-	typedef RTL_SRWLOCK SRWLOCK, *PSRWLOCK;
+	};
 
 	// Constants
 	static const BOOL TRUE = 1;
@@ -118,6 +121,31 @@ namespace Windows
 	MINIMAL_WINDOWS_API void WINAPI AcquireSRWLockExclusive(PSRWLOCK SRWLock);
 	MINIMAL_WINDOWS_API void WINAPI ReleaseSRWLockExclusive(PSRWLOCK SRWLock);
 
+	FORCEINLINE void WINAPI InitializeSRWLock(SRWLOCK* SRWLock)
+	{
+		InitializeSRWLock((PSRWLOCK)SRWLock);
+	}
+
+	FORCEINLINE void WINAPI AcquireSRWLockShared(SRWLOCK* SRWLock)
+	{
+		AcquireSRWLockShared((PSRWLOCK)SRWLock);
+	}
+
+	FORCEINLINE void WINAPI ReleaseSRWLockShared(SRWLOCK* SRWLock)
+	{
+		ReleaseSRWLockShared((PSRWLOCK)SRWLock);
+	}
+
+	FORCEINLINE void WINAPI AcquireSRWLockExclusive(SRWLOCK* SRWLock)
+	{
+		AcquireSRWLockExclusive((PSRWLOCK)SRWLock);
+	}
+
+	FORCEINLINE void WINAPI ReleaseSRWLockExclusive(SRWLOCK* SRWLock)
+	{
+		ReleaseSRWLockExclusive((PSRWLOCK)SRWLock);
+	}
+
 	// I/O
 	MINIMAL_WINDOWS_API BOOL WINAPI ConnectNamedPipe(HANDLE hNamedPipe, LPOVERLAPPED lpOverlapped);
 	MINIMAL_WINDOWS_API BOOL WINAPI GetOverlappedResult(HANDLE hFile, LPOVERLAPPED lpOverlapped, LPDWORD lpNumberOfBytesTransferred, BOOL bWait);
