Clean up config file array code appending empty strings if the destination array already contained values.
General code cleanup for config arrays.

diff --git a/Engine/Source/Runtime/Core/Private/Misc/ConfigCacheIni.cpp b/Engine/Source/Runtime/Core/Private/Misc/ConfigCacheIni.cpp
--- a/Engine/Source/Runtime/Core/Private/Misc/ConfigCacheIni.cpp
+++ b/Engine/Source/Runtime/Core/Private/Misc/ConfigCacheIni.cpp
@@ -1819,18 +1819,11 @@ bool FConfigFile::GetBool(const TCHAR* Section, const TCHAR* Key, bool& Value )
 
 int32 FConfigFile::GetArray(const TCHAR* Section, const TCHAR* Key, TArray<FString>& Value) const
 {
+	Value.Empty();
 	const FConfigSection* Sec = Find(Section);
 	if (Sec != nullptr)
 	{
-		TArray<FConfigValue> RemapArray;
-		Sec->MultiFind(Key, RemapArray);
-
-		// TMultiMap::MultiFind will return the results in reverse order
-		Value.AddZeroed(RemapArray.Num());
-		for (int32 RemapIndex = RemapArray.Num() - 1, Index = 0; RemapIndex >= 0; RemapIndex--, Index++)
-		{
-			Value[Index] = RemapArray[RemapIndex].GetValue();
-		}
+		Sec->MultiFind(Key, Value, true);
 	}
 #if !UE_BUILD_SHIPPING
 	else
@@ -2758,25 +2751,7 @@ int32 FConfigCacheIni::GetArray
 	FConfigFile* File = Find( Filename, 0 );
 	if ( File != nullptr )
 	{
-		FConfigSection* Sec = File->Find( Section );
-		if ( Sec != nullptr )
-		{
-			TArray<FConfigValue> RemapArray;
-			Sec->MultiFind(Key, RemapArray);
-
-			// TMultiMap::MultiFind will return the results in reverse order
-			out_Arr.AddZeroed(RemapArray.Num());
-			for ( int32 RemapIndex = RemapArray.Num() - 1, Index = 0; RemapIndex >= 0; RemapIndex--, Index++ )
-			{
-				out_Arr[Index] = RemapArray[RemapIndex].GetValue();
-			}
-		}
-#if !UE_BUILD_SHIPPING
-		else
-		{
-			CheckLongSectionNames( Section, File );
-		}
-#endif
+		File->GetArray(Section, Key, out_Arr);
 	}
 
 	if (out_Arr.Num())
diff --git a/Engine/Source/Runtime/Core/Public/Misc/ConfigCacheIni.h b/Engine/Source/Runtime/Core/Public/Misc/ConfigCacheIni.h
--- a/Engine/Source/Runtime/Core/Public/Misc/ConfigCacheIni.h
+++ b/Engine/Source/Runtime/Core/Public/Misc/ConfigCacheIni.h
@@ -236,12 +236,9 @@ public:
 	template<typename Allocator> 
 	void MultiFind(const FName Key, TArray<FString, Allocator>& OutValues, const bool bMaintainOrder = false) const
 	{
-		for (const TPair<FName, FConfigValue>& Pair : Pairs)
+		for (typename ElementSetType::TConstKeyIterator It(Pairs, Key); It; ++It)
 		{
-			if (Pair.Key == Key)
-			{
-				OutValues.Add(Pair.Value.GetValue());
-			}
+			OutValues.Add(It->Value.GetValue());
 		}
 
 		if (bMaintainOrder)
