Add an allocator that tracks all allocations and free's.
Keeps a map of current allocations, and a map of allocations that were freed.

diff --git a/Engine/Source/Runtime/Core/Private/HAL/MallocDoubleFreeFinder.cpp b/Engine/Source/Runtime/Core/Private/HAL/MallocDoubleFreeFinder.cpp
new file mode 100644
--- /dev/null
+++ b/Engine/Source/Runtime/Core/Private/HAL/MallocDoubleFreeFinder.cpp
@@ -0,0 +1,337 @@
+// Copyright 1998-2020 Epic Games, Inc. All Rights Reserved.
+
+/*=============================================================================
+	FMallocDoubleFreeFinder.cpp: Memoory tracking allocator
+=============================================================================*/
+#include "HAL/MallocDoubleFreeFinder.h"
+#include "Logging/LogMacros.h"
+#include "Misc/ScopeLock.h"
+//#include "Templates/Function.h"
+#include "GenericPlatform/GenericPlatformProcess.h"
+#include "HAL/IConsoleManager.h"
+#include "HAL/PlatformMisc.h"
+#include "HAL/PlatformStackWalk.h"
+
+CORE_API FMallocDoubleFreeFinder* GMallocDoubleFreeFinder;
+CORE_API bool GMallocDoubleFreeFinderEnabled = false;
+
+FMallocDoubleFreeFinder::FMallocDoubleFreeFinder(FMalloc* InMalloc)
+	: UsedMalloc(InMalloc)
+	, Initialized(false)
+{
+	DisabledTLS = FPlatformTLS::AllocTlsSlot();
+	uint64_t Count = 0;
+	FPlatformTLS::SetTlsValue(DisabledTLS, (void*)Count);
+}
+
+void FMallocDoubleFreeFinder::Init()
+{
+	if (Initialized)
+	{
+		return;
+	}
+	CallStackInfoArray.Reserve(1250000);	/* Needs to be big enough to never resize! */
+	CallStackMapKeyToCallStackIndexMap.Reserve(1250000);
+	TrackedFreeAllocations.Reserve(6000000);
+	TrackedCurrentAllocations.Reserve(8000000);
+	Initialized = true;
+}
+
+
+
+//***********************************
+// FMalloc
+//***********************************
+
+void* FMallocDoubleFreeFinder::Malloc(SIZE_T Size, uint32 Alignment)
+{
+	if (IsDisabled())
+	{
+		return UsedMalloc->Malloc(Size, Alignment);
+	}
+
+	FScopeDisableDoubleFreeFinder Disable;
+
+	int32 CallStackIndex = GetCallStackIndex();
+
+	FScopeLock Lock(&CriticalSection);
+
+	void* Ptr = UsedMalloc->Malloc(Size, Alignment);
+
+	SIZE_T AllocatedSize = Size;
+	if (UsedMalloc->GetAllocationSize(Ptr, AllocatedSize))
+	{
+		TrackMalloc(Ptr, (uint32)AllocatedSize, CallStackIndex);
+	}
+	else
+	{
+		TrackMalloc(Ptr, (uint32)Size, CallStackIndex);
+	}
+	return Ptr;
+}
+
+void* FMallocDoubleFreeFinder::Realloc(void* OldPtr, SIZE_T NewSize, uint32 Alignment)
+{
+	if (IsDisabled())
+	{
+		return UsedMalloc->Realloc(OldPtr, NewSize, Alignment);
+	}
+
+	FScopeDisableDoubleFreeFinder Disable;
+
+	int32 CallStackIndex = GetCallStackIndex();
+	SIZE_T OldSize = 0;
+
+	FScopeLock Lock(&CriticalSection);
+
+	UsedMalloc->GetAllocationSize(OldPtr, OldSize);
+
+	void* NewPtr = UsedMalloc->Realloc(OldPtr, NewSize, Alignment);
+
+	SIZE_T AllocatedSize = NewSize;
+	if (UsedMalloc->GetAllocationSize(NewPtr, AllocatedSize))
+	{
+		TrackRealloc(OldPtr, NewPtr, (uint32)AllocatedSize, (uint32)OldSize, CallStackIndex);
+	}
+	else
+	{
+		TrackRealloc(OldPtr, NewPtr, (uint32)NewSize, (uint32)OldSize, CallStackIndex);
+	}
+	return NewPtr;
+}
+
+void FMallocDoubleFreeFinder::Free(void* Ptr)
+{
+	if (IsDisabled() || Ptr == nullptr)
+	{
+		return UsedMalloc->Free(Ptr);
+	}
+
+	FScopeDisableDoubleFreeFinder Disable;
+
+	int32 CallStackIndex = GetCallStackIndex();
+
+	FScopeLock Lock(&CriticalSection);
+	SIZE_T OldSize = 0;
+	UsedMalloc->GetAllocationSize(Ptr, OldSize);
+	UsedMalloc->Free(Ptr);
+	TrackFree(Ptr, OldSize, CallStackIndex);
+}
+
+void FMallocDoubleFreeFinder::TrackMalloc(void* Ptr, uint32 Size, int32 CallStackIndex)
+{
+	if (Ptr != nullptr)
+	{
+		TrackedAllocationData* AlreadyThere = TrackedCurrentAllocations.Find(Ptr);
+		if (AlreadyThere != nullptr)
+		{
+			static TrackedAllocationData* AlreadyThereStatic = AlreadyThere;
+			TrackSpecial(Ptr);
+			PLATFORM_BREAK();
+		}
+		TrackedCurrentAllocations.Add(Ptr, TrackedAllocationData(Size, CallStackIndex));
+	}
+}
+
+void FMallocDoubleFreeFinder::TrackFree(void* Ptr, uint32 OldSize, int32 CallStackIndex)
+{
+	TrackedAllocationData Removed;
+	if (!TrackedCurrentAllocations.RemoveAndCopyValue(Ptr, Removed))
+	{
+		static TrackedAllocationData WhatHaveWeHere;
+		//memory we don't know about
+		TrackedAllocationData*AlreadyThere = TrackedFreeAllocations.Find(Ptr);
+		WhatHaveWeHere = *AlreadyThere;
+		DumpStackTraceToLog(AlreadyThere->CallStackIndex);
+		PLATFORM_BREAK();
+	}
+	else
+	{
+		if (OldSize != 0 && OldSize != Removed.Size)
+		{
+			PLATFORM_BREAK();
+		}
+		TrackedFreeAllocations.Add(Ptr, TrackedAllocationData(OldSize, CallStackIndex));
+	}
+}
+
+void FMallocDoubleFreeFinder::TrackRealloc(void* OldPtr, void* NewPtr, uint32 NewSize, uint32 OldSize, int32 CallStackIndex)
+{
+	if (OldPtr == nullptr)
+	{
+		TrackMalloc(NewPtr, NewSize, CallStackIndex);
+	}
+	else
+	{
+		if (OldPtr != NewPtr)
+		{
+			if (OldPtr)
+			{
+				TrackFree(OldPtr, OldSize, CallStackIndex);
+			}
+			if (NewPtr)
+			{
+				TrackMalloc(NewPtr, NewSize, CallStackIndex);
+			}
+		}
+	}
+}
+
+int32 FMallocDoubleFreeFinder::GetCallStackIndex()
+{
+	// Index of the callstack in CallStackInfoArray.
+	int32 Index = INDEX_NONE;
+
+	// Capture callstack and create FCallStackMapKey.
+	uint64 FullCallStack[MallocDoubleFreeFinderMaxCallStackDepth + MallocDoubleFreeFinderCallStackEntriesToSkipCount] = { 0 };
+	// CRC is filled in by the CaptureStackBackTrace, not all platforms calculate this for you.
+	uint32 CRC;
+	FPlatformStackWalk::CaptureStackBackTrace(&FullCallStack[0], MallocDoubleFreeFinderMaxCallStackDepth + MallocDoubleFreeFinderCallStackEntriesToSkipCount, &CRC);
+	// Skip first n entries as they are inside the allocator.
+	uint64* CallStack = &FullCallStack[MallocDoubleFreeFinderCallStackEntriesToSkipCount];
+	FCallStackMapKey CallStackMapKey(CRC, CallStack);
+	RWLock.ReadLock();
+	int32* IndexPtr = CallStackMapKeyToCallStackIndexMap.Find(CallStackMapKey);
+	if (IndexPtr)
+	{
+		// Use index if found
+		Index = *IndexPtr;
+		RWLock.ReadUnlock();
+	}
+	else
+	{
+		// new call stack, add to array and set index.
+		RWLock.ReadUnlock();
+		FCallStackInfoDoublleFreeFinder CallStackInfo;
+		CallStackInfo.Count = MallocDoubleFreeFinderMaxCallStackDepth;
+		for (int32 i = 0; i < MallocDoubleFreeFinderMaxCallStackDepth; i++)
+		{
+			if (!CallStack[i] && CallStackInfo.Count == MallocDoubleFreeFinderMaxCallStackDepth)
+			{
+				CallStackInfo.Count = i;
+			}
+			CallStackInfo.FramePointers[i] = CallStack[i];
+		}
+
+		RWLock.WriteLock();
+		Index = CallStackInfoArray.Num();
+		CallStackInfoArray.Append(&CallStackInfo, 1);
+		CallStackMapKey.CallStack = &CallStackInfoArray[Index].FramePointers[0];
+		CallStackMapKeyToCallStackIndexMap.Add(CallStackMapKey, Index);
+		RWLock.WriteUnlock();
+	}
+	return Index;
+}
+
+// This can be set externally, if it is we try and find what freed it before.
+void * GTrackFreeSpecialPtr = nullptr;
+
+// Can be called to find out what freed something last
+void FMallocDoubleFreeFinder::TrackSpecial(void* Ptr)
+{
+	FScopeDisableDoubleFreeFinder Disable;
+	FScopeLock Lock(&CriticalSection);
+	static TrackedAllocationData WhatHaveWeHere;	// Made static so it should be visible in the debugger
+	TrackedAllocationData*AlreadyThere;
+	TrackedAllocationData Removed;
+	if (GTrackFreeSpecialPtr != nullptr)
+	{
+		if (!TrackedCurrentAllocations.RemoveAndCopyValue(GTrackFreeSpecialPtr, Removed))
+		{
+			// Untracked memory!!
+			AlreadyThere = TrackedFreeAllocations.Find(GTrackFreeSpecialPtr);
+			WhatHaveWeHere = *AlreadyThere;
+			DumpStackTraceToLog(AlreadyThere->CallStackIndex);
+			PLATFORM_BREAK();
+		}
+
+	}
+	if (!TrackedCurrentAllocations.RemoveAndCopyValue(Ptr, Removed))
+	{
+		// Untracked memory!!
+		AlreadyThere = TrackedFreeAllocations.Find(Ptr);
+		WhatHaveWeHere = *AlreadyThere;
+		DumpStackTraceToLog(AlreadyThere->CallStackIndex);
+		PLATFORM_BREAK();
+	}
+
+	// Untracked memory!!
+	AlreadyThere = TrackedFreeAllocations.Find(Ptr);
+	if (AlreadyThere)
+	{
+		// found an exact match
+		WhatHaveWeHere = *AlreadyThere;
+		DumpStackTraceToLog(AlreadyThere->CallStackIndex);
+		PLATFORM_BREAK();
+	}
+	else
+	{
+		// look for the pointer within another allocation that was previously freed
+		auto MyIterator = TrackedFreeAllocations.CreateIterator();
+		for (; MyIterator; ++MyIterator)
+		{
+			intptr_t MyKey = (intptr_t)MyIterator.Key();
+			intptr_t MyPtr = (intptr_t)Ptr;
+			TrackedAllocationData *AlreadyThere2 = &MyIterator.Value();
+			if (MyPtr >= MyKey)
+			{
+				if (MyPtr < (MyKey + (intptr_t)AlreadyThere2->Size))
+				{
+					WhatHaveWeHere = *AlreadyThere2;
+					DumpStackTraceToLog(AlreadyThere2->CallStackIndex);
+					PLATFORM_BREAK();
+				}
+			}
+		}
+	}
+}
+
+FORCENOINLINE void FMallocDoubleFreeFinder::DumpStackTraceToLog(int32 StackIndex)
+{
+#if !NO_LOGGING
+	// Walk the stack and dump it to the allocated memory.
+	const SIZE_T StackTraceStringSize = 16384;
+	ANSICHAR StackTraceString[StackTraceStringSize];
+	{
+		StackTraceString[0] = 0;
+		uint32 CurrentDepth = 0;
+		while (CurrentDepth < MallocDoubleFreeFinderMaxCallStackDepth && CallStackInfoArray[StackIndex].FramePointers[CurrentDepth] != 0)
+		{
+			FPlatformStackWalk::ProgramCounterToHumanReadableString(CurrentDepth, CallStackInfoArray[StackIndex].FramePointers[CurrentDepth], &StackTraceString[0], StackTraceStringSize, reinterpret_cast<FGenericCrashContext*>(0));
+			FCStringAnsi::Strncat(StackTraceString, LINE_TERMINATOR_ANSI, StackTraceStringSize);
+			CurrentDepth++;
+		}
+	}
+	// Dump the error and flush the log.
+	// ELogVerbosity::Error to make sure it gets printed in log for convenience.
+	FDebug::LogFormattedMessageWithCallstack(LogOutputDevice.GetCategoryName(), __FILE__, __LINE__, L"FMallocDoubleFreeFinder::DumpStackTraceToLog", ANSI_TO_TCHAR(&StackTraceString[0]), ELogVerbosity::Error);
+	GLog->Flush();
+#endif
+}
+
+
+bool FMallocDoubleFreeFinder::Exec(UWorld* InWorld, const TCHAR* Cmd, FOutputDevice& Ar)
+{
+	if (FParse::Command(&Cmd, TEXT("DoubleFreeFinderCrash")))
+	{
+		void * test;
+		test = FMemory::Malloc(128);
+		FMemory::Free(test);
+		FMemory::Free(test);
+		return true;
+	}
+
+	return UsedMalloc->Exec(InWorld, Cmd, Ar);
+}
+
+
+FMalloc* FMallocDoubleFreeFinder::OverrideIfEnabled(FMalloc*InUsedAlloc)
+{
+	if (GMallocDoubleFreeFinderEnabled)
+	{
+		GMallocDoubleFreeFinder = new FMallocDoubleFreeFinder(InUsedAlloc);
+		GMallocDoubleFreeFinder->Init();
+		return GMallocDoubleFreeFinder;
+	}
+	return InUsedAlloc;
+}
diff --git a/Engine/Source/Runtime/Core/Private/HAL/UnrealMemory.cpp b/Engine/Source/Runtime/Core/Private/HAL/UnrealMemory.cpp
--- a/Engine/Source/Runtime/Core/Private/HAL/UnrealMemory.cpp
+++ b/Engine/Source/Runtime/Core/Private/HAL/UnrealMemory.cpp
@@ -24,6 +24,7 @@
 #include "HAL/MallocLeakDetectionProxy.h"
 #include "HAL/PlatformMallocCrash.h"
 #include "HAL/MallocPoisonProxy.h"
+#include "HAL/MallocDoubleFreeFinder.h"
 
 #if MALLOC_GT_HOOKS
 
@@ -382,6 +383,7 @@ static int FMemory_GCreateMalloc_ThreadUnsafe()
 	UE_LOG(LogMemory, Display, TEXT("Used memory before allocating anything was %.2fMB"), SizeInMb);
 #endif
 
+	GMalloc = FMallocDoubleFreeFinder::OverrideIfEnabled(GMalloc);
 	return 0;
 }
 
diff --git a/Engine/Source/Runtime/Core/Public/HAL/MallocDoubleFreeFinder.h b/Engine/Source/Runtime/Core/Public/HAL/MallocDoubleFreeFinder.h
new file mode 100644
--- /dev/null
+++ b/Engine/Source/Runtime/Core/Public/HAL/MallocDoubleFreeFinder.h
@@ -0,0 +1,283 @@
+// Copyright 1998-2020 Epic Games, Inc. All Rights Reserved.
+
+#pragma once
+
+#include "CoreTypes.h"
+#include "HAL/MemoryBase.h"
+#include "Containers/Array.h"
+
+class CORE_API FMallocDoubleFreeFinder final : public FMalloc
+{
+	/** Malloc we're based on, aka using under the hood												*/
+	FMalloc*								UsedMalloc;
+	bool									Initialized;
+
+public:
+	FMallocDoubleFreeFinder(FMalloc* InMalloc);
+
+	/**
+	 * Malloc
+	 */
+	virtual void* Malloc(SIZE_T Count, uint32 Alignment = DEFAULT_ALIGNMENT) override;
+
+	/**
+	 * Realloc
+	 */
+	virtual void* Realloc(void* Original, SIZE_T Count, uint32 Alignment = DEFAULT_ALIGNMENT)  override;
+
+	/**
+	 * Free
+	 */
+	virtual void Free(void* Original) override;
+
+	/**
+	* For some allocators this will return the actual size that should be requested to eliminate
+	* internal fragmentation. The return value will always be >= Count. This can be used to grow
+	* and shrink containers to optimal sizes.
+	* This call is always fast and thread safe with no locking.
+	*/
+	virtual SIZE_T QuantizeSize(SIZE_T Count, uint32 Alignment)  override
+	{
+		return UsedMalloc->QuantizeSize(Count, Alignment);
+	}
+
+	/**
+	* If possible determine the size of the memory allocated at the given address
+	*
+	* @param Original - Pointer to memory we are checking the size of
+	* @param SizeOut - If possible, this value is set to the size of the passed in pointer
+	* @return true if succeeded
+	*/
+	virtual bool GetAllocationSize(void *Original, SIZE_T &SizeOut)  override
+	{
+		return UsedMalloc->GetAllocationSize(Original, SizeOut);
+	}
+
+	/**
+	* Releases as much memory as possible. Must be called from the main thread.
+	*/
+	virtual void Trim(bool bTrimThreadCaches) override
+	{
+		UsedMalloc->Trim(bTrimThreadCaches);
+	}
+
+	/**
+	* Set up TLS caches on the current thread. These are the threads that we can trim.
+	*/
+	virtual void SetupTLSCachesOnCurrentThread() override
+	{
+		UsedMalloc->SetupTLSCachesOnCurrentThread();
+	}
+
+	/**
+	* Clears the TLS caches on the current thread and disables any future caching.
+	*/
+	virtual void ClearAndDisableTLSCachesOnCurrentThread() override
+	{
+		UsedMalloc->ClearAndDisableTLSCachesOnCurrentThread();
+	}
+
+	/**
+	*	Initializes stats metadata. We need to do this as soon as possible, but cannot be done in the constructor
+	*	due to the FName::StaticInit
+	*/
+	virtual void InitializeStatsMetadata() override
+	{
+		UsedMalloc->InitializeStatsMetadata();
+	}
+
+	/**
+	 * Handles any commands passed in on the command line
+	 */
+	virtual bool Exec(UWorld* InWorld, const TCHAR* Cmd, FOutputDevice& Ar) override;
+
+	/** Called once per frame, gathers and sets all memory allocator statistics into the corresponding stats. MUST BE THREAD SAFE. */
+	virtual void UpdateStats() override
+	{
+		UsedMalloc->UpdateStats();
+	}
+
+	/** Writes allocator stats from the last update into the specified destination. */
+	virtual void GetAllocatorStats(FGenericMemoryStats& out_Stats) override
+	{
+		UsedMalloc->GetAllocatorStats(out_Stats);
+	}
+
+	/** Dumps current allocator stats to the log. */
+	virtual void DumpAllocatorStats(class FOutputDevice& Ar) override
+	{
+		UsedMalloc->DumpAllocatorStats(Ar);
+	}
+
+	/**
+	 * Returns if the allocator is guaranteed to be thread-safe and therefore
+	 * doesn't need a unnecessary thread-safety wrapper around it.
+	 */
+	virtual bool IsInternallyThreadSafe() const override
+	{
+		return true;
+	}
+
+	/**
+	 * Validates the allocator's heap
+	 */
+	virtual bool ValidateHeap() override
+	{
+		return UsedMalloc->ValidateHeap();
+	}
+
+	/**
+	 * Gets descriptive name for logging purposes.
+	 *
+	 * @return pointer to human-readable malloc name
+	 */
+	virtual const TCHAR* GetDescriptiveName() override
+	{
+		return UsedMalloc->GetDescriptiveName();
+	}
+
+	static const SIZE_T MallocDoubleFreeFinderMaxCallStackDepth = 64;
+	static const SIZE_T MallocDoubleFreeFinderCallStackEntriesToSkipCount = 2;
+
+	struct FCallStackInfoDoublleFreeFinder
+	{
+		uint32 Count;
+		uint64 FramePointers[MallocDoubleFreeFinderMaxCallStackDepth];
+	};
+
+	/** Used as a key in our current allocations/freed allocations maps*/
+	struct FCallStackMapKey
+	{
+		uint32 CRC; /* Generated by the platforms capture call stack function, if the platform you are working with doesn't generate this for you crc32 the call stack array manually */
+		uint64 *CallStack;
+
+		FCallStackMapKey(uint32 InCRC, uint64* InCallStack)
+			: CRC(InCRC),
+			CallStack(InCallStack)
+		{
+		}
+
+		friend bool operator==(const FCallStackMapKey& A, const FCallStackMapKey& B)
+		{
+			if (A.CRC != B.CRC)
+			{
+				return false;
+			}
+			for (int i = 0; i < MallocDoubleFreeFinderMaxCallStackDepth; ++i)
+			{
+				uint64 APtr = A.CallStack[i];
+				uint64 BPtr = B.CallStack[i];
+				if (APtr != BPtr)
+				{
+					return false;
+				}
+				if (APtr == 0)
+					break;
+			}
+			return true;
+		}
+
+		friend inline uint32 GetTypeHash(const FCallStackMapKey& InKey)
+		{
+			return InKey.CRC;
+		}
+	};
+
+	/**
+	 * If you get an allocation/memory error outside of the allocator you can call this directly
+	 * It will dump a callstack of the last allocator free most likely to have caused the problem to the log, if you have symbols loaded
+	 * Might be useful to pass an access violation ptr to this!
+	 */
+	void TrackSpecial(void* Ptr);
+
+	void Init();
+	void DumpStackTraceToLog(int32 StackIndex);
+
+
+	static FMalloc* OverrideIfEnabled(FMalloc*InUsedAlloc);
+
+	friend class FMallocDoubleFreeFinder;
+	friend class FScopeDisableDoubleFreeFinder;
+
+protected:
+	FCriticalSection CriticalSection;
+	uint32 DisabledTLS;
+
+	FORCEINLINE void IncDisabled()
+	{
+		uint64_t DisabledCount = (uint64_t)FPlatformTLS::GetTlsValue(DisabledTLS);
+		++DisabledCount;
+		FPlatformTLS::SetTlsValue(DisabledTLS, (void*)DisabledCount);
+	}
+
+	FORCEINLINE void DecDisabled()
+	{
+		uint64_t DisabledCount = (uint64_t)FPlatformTLS::GetTlsValue(DisabledTLS);
+		--DisabledCount;
+		FPlatformTLS::SetTlsValue(DisabledTLS, (void*)DisabledCount);
+	}
+	bool IsDisabled()
+	{
+		return FPlatformTLS::GetTlsValue(DisabledTLS) != 0;
+	}
+	void TrackMalloc(void* Ptr, uint32 Size, int32 CallStackIndex);
+	void TrackFree(void* Ptr, uint32 OldSize, int32 CallStackIndex);
+	void TrackRealloc(void* OldPtr, void* NewPtr, uint32 NewSize, uint32 OldSize, int32 CallStackIndex);
+
+	FRWLock RWLock;
+
+	TMap<FCallStackMapKey, int32> CallStackMapKeyToCallStackIndexMap;
+	TArray<FCallStackInfoDoublleFreeFinder> CallStackInfoArray;
+
+
+	struct TrackedAllocationData
+	{
+		SIZE_T Size;
+		int32 CallStackIndex;
+		TrackedAllocationData() :
+			Size(0),
+			CallStackIndex(-1)
+		{
+		};
+		TrackedAllocationData(SIZE_T InRequestedSize, int32 InCallStackIndex)
+		{
+			Size = InRequestedSize;
+			CallStackIndex = InCallStackIndex;
+		};
+		~TrackedAllocationData()
+		{
+			Size = 0;
+			CallStackIndex = -1;
+		};
+	};
+	TMap<const void* const, TrackedAllocationData> TrackedCurrentAllocations;	// Pointer as a key to a call stack for all the current allocations we known about
+	TMap<const void* const, TrackedAllocationData> TrackedFreeAllocations;		// Pointer as a key to a call stack for all allocations that have been freed
+
+	int32 GetCallStackIndex();
+};
+
+extern CORE_API FMallocDoubleFreeFinder* GMallocDoubleFreeFinder;
+extern CORE_API bool GMallocDoubleFreeFinderEnabled;
+
+
+/**
+ * Disables the double free finder for the current thread
+ * Need to do this as we might allocate memory for the allocators tracking data, that can't be tracked!
+ */
+
+class FScopeDisableDoubleFreeFinder
+{
+public:
+
+	FScopeDisableDoubleFreeFinder()
+	{
+		GMallocDoubleFreeFinder->IncDisabled();
+	}
+
+	~FScopeDisableDoubleFreeFinder()
+	{
+		GMallocDoubleFreeFinder->DecDisabled();
+	}
+};
+
+
