Change unsafe c-cast (multiple inheritance) in BufferArchive to static_cast

diff --git a/Engine/Source/Runtime/Core/Public/Serialization/BufferArchive.h b/Engine/Source/Runtime/Core/Public/Serialization/BufferArchive.h
--- a/Engine/Source/Runtime/Core/Public/Serialization/BufferArchive.h
+++ b/Engine/Source/Runtime/Core/Public/Serialization/BufferArchive.h
@@ -15,7 +15,7 @@ class FBufferArchive : public FMemoryWriter, public TArray<uint8>
 {
 public:
 	FBufferArchive( bool bIsPersistent = false, const FName InArchiveName = NAME_None )
-	: FMemoryWriter( (TArray<uint8>&)*this, bIsPersistent, false, InArchiveName )
+	: FMemoryWriter( *static_cast<TArray<uint8>*>(this), bIsPersistent, false, InArchiveName )
 	{}
 	/**
   	 * Returns the name of the Archive.  Useful for getting the name of the package a struct or object
