Boot instrumentation for trace - allows visualization of very early activity which was previously only available in a separate log file

(peafour-cherrypick of 6543969 by [REDACTED_PERSON])

diff --git a/Engine/Source/Runtime/Core/Private/Misc/CoreGlobals.cpp b/Engine/Source/Runtime/Core/Private/Misc/CoreGlobals.cpp
--- a/Engine/Source/Runtime/Core/Private/Misc/CoreGlobals.cpp
+++ b/Engine/Source/Runtime/Core/Private/Misc/CoreGlobals.cpp
@@ -362,6 +362,7 @@ FScopedBootTiming::~FScopedBootTiming()
 }
 void BootTimingPoint(const ANSICHAR *Message)
 {
+	TRACE_BOOKMARK(TEXT("%s"), *FString(Message));
 }
 void DumpBootTiming()
 {
@@ -398,6 +399,8 @@ void DumpBootTiming()
 
 static void BootTimingPoint(const TCHAR *Message, const TCHAR *Prefix = nullptr, int32 Depth = 0, double TookTime = 0.0)
 {
+	TRACE_BOOKMARK(TEXT("%s"), Message);
+
 	static double LastTime = 0.0;
 	static TArray<FString> MessageStack;
 	static FString LastGapMessage;
diff --git a/Engine/Source/Runtime/Core/Public/CoreGlobals.h b/Engine/Source/Runtime/Core/Public/CoreGlobals.h
--- a/Engine/Source/Runtime/Core/Public/CoreGlobals.h
+++ b/Engine/Source/Runtime/Core/Public/CoreGlobals.h
@@ -68,7 +68,7 @@ struct CORE_API FScopedBootTiming
 };
 
 
-#define SCOPED_BOOT_TIMING(x) FScopedBootTiming ANONYMOUS_VARIABLE(BootTiming_)(x);
+#define SCOPED_BOOT_TIMING(x) TRACE_CPUPROFILER_EVENT_SCOPE_TEXT(L ## x); FScopedBootTiming ANONYMOUS_VARIABLE(BootTiming_)(x);
 
 #define GLog GetGlobalLogSingleton()
 extern CORE_API FConfigCacheIni* GConfig;
