Fixed CIS issue for some configs caused by missing

#include "ProfilingDebugging/CpuProfilerTrace.h"

diff --git a/Engine/Source/Runtime/Core/Public/CoreGlobals.h b/Engine/Source/Runtime/Core/Public/CoreGlobals.h
--- a/Engine/Source/Runtime/Core/Public/CoreGlobals.h
+++ b/Engine/Source/Runtime/Core/Public/CoreGlobals.h
@@ -7,6 +7,7 @@
 #include "Logging/LogMacros.h"
 #include "HAL/PlatformTLS.h"
 #include "Templates/Atomic.h"
+#include "ProfilingDebugging/CpuProfilerTrace.h"
 
 class Error;
 class FConfigCacheIni;
