UnrealInsights:
Release readiness sweep

diff --git a/Engine/Source/Developer/TraceInsights/Private/Insights/Widgets/STimersView.cpp b/Engine/Source/Developer/TraceInsights/Private/Insights/Widgets/STimersView.cpp
--- a/Engine/Source/Developer/TraceInsights/Private/Insights/Widgets/STimersView.cpp
+++ b/Engine/Source/Developer/TraceInsights/Private/Insights/Widgets/STimersView.cpp
@@ -1480,7 +1480,7 @@ void STimersView::RebuildTree(bool bRebuildEvenIfNotChanged)
 					for (uint64 TimerIndex = 0; TimerIndex < TimersCount; ++TimerIndex)
 					{
 						const Trace::FTimingProfilerTimer& Timer = Timers[TimerIndex];
-						FName Name(*Timer.Name);// +TEXT(" [GPU]")));
+						FName Name(Timer.Name);// +TEXT(" [GPU]")));
 						FName Group(Timer.IsGpuTimer ? TEXT("GPU") : TEXT("CPU"));
 						ETimerNodeType Type = Timer.IsGpuTimer ? ETimerNodeType::GpuScope : ETimerNodeType::CpuScope;
 						FTimerNode* TimerPtr = new FTimerNode(Timer.Id, Name, Group, Type);
diff --git a/Engine/Source/Developer/TraceInsights/Private/Insights/Widgets/STimingView.cpp b/Engine/Source/Developer/TraceInsights/Private/Insights/Widgets/STimingView.cpp
--- a/Engine/Source/Developer/TraceInsights/Private/Insights/Widgets/STimingView.cpp
+++ b/Engine/Source/Developer/TraceInsights/Private/Insights/Widgets/STimingView.cpp
@@ -1015,14 +1015,14 @@ void STimingView::DrawCpuGpuTimelineTrack(FTimingViewDrawHelper& Helper, FTiming
 							const double SecondsPerPixel = 1.0 / Helper.GetViewport().ScaleX;
 							Timeline.EnumerateEventsDownSampled(Helper.GetViewport().StartTime, Helper.GetViewport().EndTime, SecondsPerPixel, [this, &Helper, Timers](double StartTime, double EndTime, uint32 Depth, const Trace::FTimingProfilerEvent& Event)
 							{
-								Helper.AddEvent(StartTime, EndTime, Depth, *Timers[Event.TimerIndex].Name);
+								Helper.AddEvent(StartTime, EndTime, Depth, Timers[Event.TimerIndex].Name);
 							});
 						}
 						else
 						{
 							Timeline.EnumerateEvents(Helper.GetViewport().StartTime, Helper.GetViewport().EndTime, [this, &Helper, Timers](double StartTime, double EndTime, uint32 Depth, const Trace::FTimingProfilerEvent& Event)
 							{
-								Helper.AddEvent(StartTime, EndTime, Depth, *Timers[Event.TimerIndex].Name);
+								Helper.AddEvent(StartTime, EndTime, Depth, Timers[Event.TimerIndex].Name);
 							});
 						}
 					});
diff --git a/Engine/Source/Developer/TraceServices/Private/AnalysisService.cpp b/Engine/Source/Developer/TraceServices/Private/AnalysisService.cpp
--- a/Engine/Source/Developer/TraceServices/Private/AnalysisService.cpp
+++ b/Engine/Source/Developer/TraceServices/Private/AnalysisService.cpp
@@ -74,9 +74,9 @@ FAnalysisSession::FAnalysisSession(const TCHAR* SessionName)
 	, ClassInfos(Allocator, 4096)
 	, BookmarkProvider(MakeShared<FBookmarkProvider>(Lock, StringStore))
 	, LogProvider(MakeShared<FLogProvider>(Allocator, Lock, StringStore))
-	, ThreadProvider(MakeShared<FThreadProvider>(Lock))
+	, ThreadProvider(MakeShared<FThreadProvider>(Lock, StringStore))
 	, FramesProvider(MakeShared<FFrameProvider>(Allocator, Lock))
-	, TimingProfilerProvider(MakeShared<FTimingProfilerProvider>(Allocator, Lock))
+	, TimingProfilerProvider(MakeShared<FTimingProfilerProvider>(Allocator, Lock, StringStore))
 	, FileActivityProvider(MakeShared<FFileActivityProvider>(Allocator, Lock))
 	, LoadTimeProfilerProvider(MakeShared<FLoadTimeProfilerProvider>(Allocator, Lock))
 	, CountersProvider(MakeShared<FCountersProvider>(Allocator, Lock))
@@ -215,7 +215,7 @@ void FAnalysisService::FMockGenerator::CreateMockData()
 		int32 ThreadId = 4;
 		while (ThreadId <= 10)
 		{
-			ThreadProvider->AddThread(ThreadId, FString::Printf(TEXT("Thread %d"), ThreadId), TPri_SlightlyBelowNormal);
+			ThreadProvider->AddThread(ThreadId, *FString::Printf(TEXT("Thread %d"), ThreadId), TPri_SlightlyBelowNormal);
 			ThreadProvider->SetThreadGroup(ThreadId, TraceThreadGroup_ThreadPool);
 			Timelines.Add(TimingProfilerProvider->EditCpuThreadTimeline(ThreadId));
 			++ThreadId;
diff --git a/Engine/Source/Developer/TraceServices/Private/Analyzers/CpuProfilerTraceAnalysis.cpp b/Engine/Source/Developer/TraceServices/Private/Analyzers/CpuProfilerTraceAnalysis.cpp
--- a/Engine/Source/Developer/TraceServices/Private/Analyzers/CpuProfilerTraceAnalysis.cpp
+++ b/Engine/Source/Developer/TraceServices/Private/Analyzers/CpuProfilerTraceAnalysis.cpp
@@ -17,7 +17,6 @@ void FCpuProfilerAnalyzer::OnAnalysisBegin(const FOnAnalysisContext& Context)
 	auto& Builder = Context.InterfaceBuilder;
 
 	Builder.RouteEvent(RouteId_EventSpec, "CpuProfiler", "EventSpec");
-	Builder.RouteEvent(RouteId_DynamicEventName, "CpuProfiler", "DynamicEventName");
 	Builder.RouteEvent(RouteId_EventBatch, "CpuProfiler", "EventBatch");
 }
 
@@ -31,31 +30,17 @@ void FCpuProfilerAnalyzer::OnEvent(uint16 RouteId, const FOnEventContext& Contex
 	case RouteId_EventSpec:
 	{
 		uint16 Id = EventData.GetValue("Id").As<uint16>();
-		ScopeIdToEventIdMap.Add(Id, TimingProfilerProvider->AddCpuTimer(*FString((const char*)EventData.GetAttachment(), EventData.GetValue("NameSize").As<uint16>())));
-		break;
-	}
-	case RouteId_DynamicEventName:
-	{
-		uint16 Id = EventData.GetValue("Id").As<uint16>();
-		if (DynamicNameToEventIdMap.Contains(Id))
-		{
-			DynamicNameToEventIdMap[Id] = TimingProfilerProvider->AddCpuTimer(*FString((const char*)EventData.GetAttachment(), EventData.GetValue("NameSize").As<uint16>()));
-		}
-		else
-		{
-			DynamicNameToEventIdMap.Add(Id, TimingProfilerProvider->AddCpuTimer(*FString((const char*)EventData.GetAttachment(), EventData.GetValue("NameSize").As<uint16>())));
-		}
+		ScopeIdToEventIdMap.Add(Id, TimingProfilerProvider->AddCpuTimer(reinterpret_cast<const TCHAR*>(EventData.GetAttachment())));
 		break;
 	}
 	case RouteId_EventBatch:
 	{
 		uint32 ThreadId = EventData.GetValue("ThreadId").As<uint32>();
 		TSharedRef<FThreadState> ThreadState = GetThreadState(ThreadId);
-		uint64 LastCycle = EventData.GetValue("CycleBase").As<uint64>();
-		uint64 BufferSize = EventData.GetValue("BufferSize").As<uint8>();
+		uint64 LastCycle = ThreadState->LastCycle;
+		uint64 BufferSize = EventData.GetAttachmentSize();
 		const uint8* BufferPtr = EventData.GetAttachment();
 		const uint8* BufferEnd = BufferPtr + BufferSize;
-		uint64 LatestEndCycle = 0;
 		while (BufferPtr < BufferEnd)
 		{
 			uint64 DecodedCycle = FTraceAnalyzerUtils::Decode7bit(BufferPtr);
@@ -65,18 +50,8 @@ void FCpuProfilerAnalyzer::OnEvent(uint16 RouteId, const FOnEventContext& Contex
 			{
 				EventScopeState& ScopeState = ThreadState->ScopeStack.AddDefaulted_GetRef();
 				ScopeState.StartCycle = ActualCycle;
-				uint16 SpecId = DecodeSpecId(BufferPtr);
-				if (SpecId == 0)
-				{
-					uint16 DynamicEventNameId = DecodeSpecId(BufferPtr);
-					check(DynamicNameToEventIdMap.Contains(DynamicEventNameId));
-					ScopeState.EventTypeId = DynamicNameToEventIdMap[DynamicEventNameId];
-				}
-				else
-				{
-					check(ScopeIdToEventIdMap.Contains(SpecId));
-					ScopeState.EventTypeId = ScopeIdToEventIdMap[SpecId];
-				}
+				uint16 SpecId = FTraceAnalyzerUtils::Decode7bit(BufferPtr);
+				ScopeState.EventTypeId = ScopeIdToEventIdMap[SpecId];
 				Trace::FTimingProfilerEvent Event;
 				Event.TimerIndex = ScopeState.EventTypeId;
 				ThreadState->Timeline->AppendBeginEvent(Context.SessionContext.TimestampFromCycle(ScopeState.StartCycle), Event);
@@ -86,14 +61,14 @@ void FCpuProfilerAnalyzer::OnEvent(uint16 RouteId, const FOnEventContext& Contex
 				check(ThreadState->ScopeStack.Num());
 				ThreadState->ScopeStack.Pop();
 				ThreadState->Timeline->AppendEndEvent(Context.SessionContext.TimestampFromCycle(ActualCycle));
-				LatestEndCycle = ActualCycle;
 			}
 		}
 		check(BufferPtr == BufferEnd);
-		if (LatestEndCycle)
+		if (LastCycle)
 		{
-			Session->UpdateDuration(Context.SessionContext.TimestampFromCycle(LatestEndCycle));
+			Session->UpdateDuration(Context.SessionContext.TimestampFromCycle(LastCycle));
 		}
+		ThreadState->LastCycle = LastCycle;
 		break;
 	}
 	}
@@ -115,10 +90,3 @@ TSharedRef<FCpuProfilerAnalyzer::FThreadState> FCpuProfilerAnalyzer::GetThreadSt
 	}
 }
 
-uint16 FCpuProfilerAnalyzer::DecodeSpecId(const uint8*& BufferPtr)
-{
-	uint16 SpecId = 0;
-	SpecId |= *BufferPtr++;
-	SpecId |= (uint16(*BufferPtr++) << 8);
-	return SpecId;
-}
diff --git a/Engine/Source/Developer/TraceServices/Private/Analyzers/CpuProfilerTraceAnalysis.h b/Engine/Source/Developer/TraceServices/Private/Analyzers/CpuProfilerTraceAnalysis.h
--- a/Engine/Source/Developer/TraceServices/Private/Analyzers/CpuProfilerTraceAnalysis.h
+++ b/Engine/Source/Developer/TraceServices/Private/Analyzers/CpuProfilerTraceAnalysis.h
@@ -34,18 +34,15 @@ private:
 	{
 		TArray<EventScopeState> ScopeStack;
 		TSharedPtr<Trace::FTimingProfilerProvider::TimelineInternal> Timeline;
+		double LastCycle = 0.0;
 	};
 
 	TSharedRef<FThreadState> GetThreadState(uint32 ThreadId);
-	uint16 DecodeSpecId(const uint8*& BufferPtr);
 
 	enum : uint16
 	{
 		RouteId_EventSpec,
-		RouteId_DynamicEventName,
 		RouteId_EventBatch,
-		RouteId_BeginGameFrame,
-		RouteId_BeginRenderingFrame
 	};
 
 	TSharedRef<Trace::FAnalysisSession> Session;
@@ -53,5 +50,4 @@ private:
 	TSharedRef<Trace::FTimingProfilerProvider> TimingProfilerProvider;
 	TMap<uint32, TSharedRef<FThreadState>> ThreadStatesMap;
 	TMap<uint16, uint32> ScopeIdToEventIdMap;
-	TMap<uint16, uint32> DynamicNameToEventIdMap;
 };
diff --git a/Engine/Source/Developer/TraceServices/Private/Analyzers/MiscTraceAnalysis.cpp b/Engine/Source/Developer/TraceServices/Private/Analyzers/MiscTraceAnalysis.cpp
--- a/Engine/Source/Developer/TraceServices/Private/Analyzers/MiscTraceAnalysis.cpp
+++ b/Engine/Source/Developer/TraceServices/Private/Analyzers/MiscTraceAnalysis.cpp
@@ -94,9 +94,8 @@ void FMiscTraceAnalyzer::OnEvent(uint16 RouteId, const FOnEventContext& Context)
 	case RouteId_CreateThread:
 	{
 		uint32 CreatedThreadId = EventData.GetValue("CreatedThreadId").As<uint32>();
-		FString Name = FString((const char*)EventData.GetAttachment(), EventData.GetValue("NameSize").As<uint16>());
 		EThreadPriority Priority = static_cast<EThreadPriority>(EventData.GetValue("Priority").As<uint32>());
-		ThreadProvider->AddThread(CreatedThreadId, Name, Priority);
+		ThreadProvider->AddThread(CreatedThreadId, reinterpret_cast<const TCHAR*>(EventData.GetAttachment()), Priority);
 		uint32 CurrentThreadId = EventData.GetValue("CurrentThreadId").As<uint32>();
 		FThreadState* ThreadState = GetThreadState(CurrentThreadId);
 		if (ThreadState->ThreadGroupStack.Num())
@@ -107,7 +106,7 @@ void FMiscTraceAnalyzer::OnEvent(uint16 RouteId, const FOnEventContext& Context)
 	}
 	case RouteId_SetThreadGroup:
 	{
-		const char* GroupName = (const char*)EventData.GetAttachment();
+		const char* GroupName = reinterpret_cast<const char*>(EventData.GetAttachment());
 		ETraceThreadGroup ThreadGroup = GetThreadGroupFromName(GroupName);
 		uint32 ThreadId = EventData.GetValue("ThreadId").As<uint32>();
 		ThreadProvider->SetThreadGroup(ThreadId, ThreadGroup);
@@ -115,7 +114,7 @@ void FMiscTraceAnalyzer::OnEvent(uint16 RouteId, const FOnEventContext& Context)
 	}
 	case RouteId_BeginThreadGroupScope:
 	{
-		const char* GroupName = (const char*)EventData.GetAttachment();
+		const char* GroupName = reinterpret_cast<const char*>(EventData.GetAttachment());
 		ETraceThreadGroup ThreadGroup = GetThreadGroupFromName(GroupName);
 		uint32 CurrentThreadId = EventData.GetValue("CurrentThreadId").As<uint32>();
 		FThreadState* ThreadState = GetThreadState(CurrentThreadId);
diff --git a/Engine/Source/Developer/TraceServices/Private/Common/FormatArgs.cpp b/Engine/Source/Developer/TraceServices/Private/Common/FormatArgs.cpp
--- a/Engine/Source/Developer/TraceServices/Private/Common/FormatArgs.cpp
+++ b/Engine/Source/Developer/TraceServices/Private/Common/FormatArgs.cpp
@@ -252,15 +252,7 @@ int32 FFormatArgsHelper::FormatArgument(TCHAR* Out, uint64 MaxOut, const FFormat
 	case FFormatArgsTrace::FormatArgTypeCode_CategoryInteger:
 		if (ArgSpec.NothingPrinted)
 		{
-			if (ArgSpec.AdditionalIntegerArgumentCount == 2)
-			{
-				ExtractIntegerArgument(ArgStream);
-			}
-			if (ArgSpec.AdditionalIntegerArgumentCount >= 1)
-			{
-				ExtractIntegerArgument(ArgStream);
-			}
-			else
+			for (uint8 IntegerArgIndex = 0; IntegerArgIndex < ArgSpec.AdditionalIntegerArgumentCount + 1; ++IntegerArgIndex)
 			{
 				ExtractIntegerArgument(ArgStream);
 			}
diff --git a/Engine/Source/Developer/TraceServices/Private/Common/SlabAllocator.h b/Engine/Source/Developer/TraceServices/Private/Common/SlabAllocator.h
--- a/Engine/Source/Developer/TraceServices/Private/Common/SlabAllocator.h
+++ b/Engine/Source/Developer/TraceServices/Private/Common/SlabAllocator.h
@@ -3,6 +3,7 @@
 #pragma once
 
 #include "HAL/UnrealMemory.h"
+#include "Containers/Array.h"
 
 class FSlabAllocator
 {
diff --git a/Engine/Source/Developer/TraceServices/Private/Model/Threads.cpp b/Engine/Source/Developer/TraceServices/Private/Model/Threads.cpp
--- a/Engine/Source/Developer/TraceServices/Private/Model/Threads.cpp
+++ b/Engine/Source/Developer/TraceServices/Private/Model/Threads.cpp
@@ -3,23 +3,33 @@
 #include "Model/Threads.h"
 #include "Misc/ScopeLock.h"
 #include "AnalysisServicePrivate.h"
+#include "Common/StringStore.h"
 
 namespace Trace
 {
 
-FThreadProvider::FThreadProvider(const FAnalysisSessionLock& InSessionLock)
+FThreadProvider::FThreadProvider(const FAnalysisSessionLock& InSessionLock, FStringStore& InStringStore)
 	: SessionLock(InSessionLock)
+	, StringStore(InStringStore)
 {
 
 }
 
+FThreadProvider::~FThreadProvider()
+{
+	for (FThreadInfoInternal* Thread : SortedThreads)
+	{
+		delete Thread;
+	}
+}
+
 void FThreadProvider::EnsureThreadExists(uint32 Id)
 {
 	SessionLock.WriteAccessCheck();
 
 	if (!ThreadMap.Contains(Id))
 	{
-		TSharedRef<FThreadInfoInternal> ThreadInfo = MakeShared<FThreadInfoInternal>();
+		FThreadInfoInternal* ThreadInfo = new FThreadInfoInternal();
 		ThreadInfo->Id = Id;
 		ThreadInfo->FallbackSortOrder = SortedThreads.Num();
 		SortedThreads.Add(ThreadInfo);
@@ -34,10 +44,10 @@ void FThreadProvider::AddGameThread(uint32 Id)
 	SessionLock.WriteAccessCheck();
 
 	check(!ThreadMap.Contains(Id));
-	TSharedRef<FThreadInfoInternal> ThreadInfo = MakeShared<FThreadInfoInternal>();
+	FThreadInfoInternal* ThreadInfo = new FThreadInfoInternal();
 	ThreadInfo->Id = Id;
 	ThreadInfo->PrioritySortOrder = GetPrioritySortOrder(TPri_Normal);
-	ThreadInfo->Name = *FName(NAME_GameThread).GetPlainNameString();
+	ThreadInfo->Name = StringStore.Store(*FName(NAME_GameThread).GetPlainNameString());
 	ThreadInfo->FallbackSortOrder = SortedThreads.Num();
 	ThreadInfo->IsGameThread = true;
 	SortedThreads.Add(ThreadInfo);
@@ -45,25 +55,25 @@ void FThreadProvider::AddGameThread(uint32 Id)
 	++ModCount;
 }
 
-void FThreadProvider::AddThread(uint32 Id, const FString& Name, EThreadPriority Priority)
+void FThreadProvider::AddThread(uint32 Id, const TCHAR* Name, EThreadPriority Priority)
 {
 	SessionLock.WriteAccessCheck();
 
-	TSharedPtr<FThreadInfoInternal> ThreadInfo;
+	FThreadInfoInternal* ThreadInfo;
 	if (!ThreadMap.Contains(Id))
 	{
-		ThreadInfo = MakeShared<FThreadInfoInternal>();
+		ThreadInfo = new FThreadInfoInternal();
 		ThreadInfo->Id = Id;
 		ThreadInfo->FallbackSortOrder = SortedThreads.Num();
-		SortedThreads.Add(ThreadInfo.ToSharedRef());
-		ThreadMap.Add(Id, ThreadInfo.ToSharedRef());
+		SortedThreads.Add(ThreadInfo);
+		ThreadMap.Add(Id, ThreadInfo);
 	}
 	else
 	{
 		ThreadInfo = ThreadMap[Id];
 	}
 	ThreadInfo->PrioritySortOrder = GetPrioritySortOrder(Priority);
-	ThreadInfo->Name = Name;
+	ThreadInfo->Name = StringStore.Store(Name);
 	SortThreads();
 	++ModCount;
 }
@@ -73,7 +83,7 @@ void FThreadProvider::SetThreadPriority(uint32 Id, EThreadPriority Priority)
 	SessionLock.WriteAccessCheck();
 
 	check(ThreadMap.Contains(Id));
-	TSharedRef<FThreadInfoInternal> ThreadInfo = ThreadMap[Id];
+	FThreadInfoInternal* ThreadInfo = ThreadMap[Id];
 	ThreadInfo->PrioritySortOrder = GetPrioritySortOrder(Priority);
 	SortThreads();
 	++ModCount;
@@ -84,7 +94,7 @@ void FThreadProvider::SetThreadGroup(uint32 Id, ETraceThreadGroup Group)
 	SessionLock.WriteAccessCheck();
 
 	check(ThreadMap.Contains(Id));
-	TSharedRef<FThreadInfoInternal> ThreadInfo = ThreadMap[Id];
+	FThreadInfoInternal* ThreadInfo = ThreadMap[Id];
 	ThreadInfo->GroupSortOrder = GetGroupSortOrder(Group);
 	switch (Group)
 	{
@@ -121,11 +131,11 @@ void FThreadProvider::EnumerateThreads(TFunctionRef<void(const FThreadInfo &)> C
 {
 	SessionLock.ReadAccessCheck();
 
-	for (const TSharedRef<FThreadInfoInternal>& Thread : SortedThreads)
+	for (const FThreadInfoInternal* Thread : SortedThreads)
 	{
 		FThreadInfo ThreadInfo;
 		ThreadInfo.Id = Thread->Id;
-		ThreadInfo.Name = *Thread->Name;
+		ThreadInfo.Name = Thread->Name;
 		ThreadInfo.GroupName = Thread->GroupName;
 		Callback(ThreadInfo);
 	}
@@ -133,9 +143,9 @@ void FThreadProvider::EnumerateThreads(TFunctionRef<void(const FThreadInfo &)> C
 
 void FThreadProvider::SortThreads()
 {
-	Algo::SortBy(SortedThreads, [](const TSharedRef<FThreadInfoInternal>& In) -> const FThreadInfoInternal&
+	Algo::SortBy(SortedThreads, [](const FThreadInfoInternal* In) -> const FThreadInfoInternal&
 	{
-		return In.Get();
+		return *In;
 	});
 }
 
diff --git a/Engine/Source/Developer/TraceServices/Private/Model/Threads.h b/Engine/Source/Developer/TraceServices/Private/Model/Threads.h
--- a/Engine/Source/Developer/TraceServices/Private/Model/Threads.h
+++ b/Engine/Source/Developer/TraceServices/Private/Model/Threads.h
@@ -8,15 +8,17 @@
 namespace Trace
 {
 class FAnalysisSessionLock;
+class FStringStore;
 
 class FThreadProvider
 	: public IThreadProvider
 {
 public:
-	FThreadProvider(const FAnalysisSessionLock& SessionLock);
+	FThreadProvider(const FAnalysisSessionLock& SessionLock, FStringStore& StringStore);
+	~FThreadProvider();
 	void EnsureThreadExists(uint32 Id);
 	void AddGameThread(uint32 Id);
-	void AddThread(uint32 Id, const FString& Name, EThreadPriority Priority);
+	void AddThread(uint32 Id, const TCHAR* Name, EThreadPriority Priority);
 	void SetThreadPriority(uint32 Id, EThreadPriority Priority);
 	void SetThreadGroup(uint32 Id, ETraceThreadGroup Group);
 	virtual uint64 GetModCount() const override { return ModCount; }
@@ -26,7 +28,7 @@ private:
 	struct FThreadInfoInternal
 	{
 		uint32 Id = 0;
-		FString Name;
+		const TCHAR* Name = nullptr;
 		uint32 GroupSortOrder = ~0u;
 		uint32 PrioritySortOrder = ~0u;
 		uint32 FallbackSortOrder = ~0u;
@@ -41,9 +43,10 @@ private:
 	static uint32 GetGroupSortOrder(ETraceThreadGroup ThreadGroup);
 
 	const FAnalysisSessionLock& SessionLock;
+	FStringStore& StringStore;
 	uint64 ModCount = 0;
-	TMap<uint32, TSharedRef<FThreadInfoInternal>> ThreadMap;
-	TArray<TSharedRef<FThreadInfoInternal>> SortedThreads;
+	TMap<uint32, FThreadInfoInternal*> ThreadMap;
+	TArray<FThreadInfoInternal*> SortedThreads;
 };
 
 }
\ No newline at end of file
diff --git a/Engine/Source/Developer/TraceServices/Private/Model/TimingProfiler.cpp b/Engine/Source/Developer/TraceServices/Private/Model/TimingProfiler.cpp
--- a/Engine/Source/Developer/TraceServices/Private/Model/TimingProfiler.cpp
+++ b/Engine/Source/Developer/TraceServices/Private/Model/TimingProfiler.cpp
@@ -2,13 +2,15 @@
 
 #include "Model/TimingProfiler.h"
 #include "AnalysisServicePrivate.h"
+#include "Common/StringStore.h"
 
 namespace Trace
 {
 
-FTimingProfilerProvider::FTimingProfilerProvider(FSlabAllocator& InAllocator, FAnalysisSessionLock& InSessionLock)
+FTimingProfilerProvider::FTimingProfilerProvider(FSlabAllocator& InAllocator, FAnalysisSessionLock& InSessionLock, FStringStore& InStringStore)
 	: Allocator(InAllocator)
 	, SessionLock(InSessionLock)
+	, StringStore(InStringStore)
 {
 	Timelines.Add(MakeShared<TimelineInternal>(Allocator));
 }
@@ -37,7 +39,7 @@ FTimingProfilerTimer& FTimingProfilerProvider::AddTimerInternal(const TCHAR* Nam
 {
 	FTimingProfilerTimer& Timer = Timers.AddDefaulted_GetRef();
 	Timer.Id = Timers.Num() - 1;
-	Timer.Name = Name;
+	Timer.Name = StringStore.Store(Name);
 	uint32 NameHash = 0;
 	for (const TCHAR* c = Name; *c; ++c)
 	{
diff --git a/Engine/Source/Developer/TraceServices/Private/Model/TimingProfiler.h b/Engine/Source/Developer/TraceServices/Private/Model/TimingProfiler.h
--- a/Engine/Source/Developer/TraceServices/Private/Model/TimingProfiler.h
+++ b/Engine/Source/Developer/TraceServices/Private/Model/TimingProfiler.h
@@ -10,6 +10,7 @@ namespace Trace
 {
 
 class FAnalysisSessionLock;
+class FStringStore;
 
 class FTimingProfilerProvider
 	: public ITimingProfilerProvider
@@ -17,7 +18,7 @@ class FTimingProfilerProvider
 public:
 	typedef TMonotonicTimeline<FTimingProfilerEvent> TimelineInternal;
 
-	FTimingProfilerProvider(FSlabAllocator& Allocator, FAnalysisSessionLock& InSessionLock);
+	FTimingProfilerProvider(FSlabAllocator& Allocator, FAnalysisSessionLock& InSessionLock, FStringStore& StringStore);
 	virtual ~FTimingProfilerProvider();
 	uint32 AddCpuTimer(const TCHAR* Name);
 	uint32 AddGpuTimer(const TCHAR* Name);
@@ -35,6 +36,7 @@ private:
 
 	FSlabAllocator& Allocator;
 	FAnalysisSessionLock& SessionLock;
+	FStringStore& StringStore;
 	TArray<FTimingProfilerTimer> Timers;
 	TArray<TSharedRef<TimelineInternal>> Timelines;
 	TMap<uint32, uint32> CpuThreadTimelineIndexMap;
diff --git a/Engine/Source/Developer/TraceServices/Public/TraceServices/AnalysisService.h b/Engine/Source/Developer/TraceServices/Public/TraceServices/AnalysisService.h
--- a/Engine/Source/Developer/TraceServices/Public/TraceServices/AnalysisService.h
+++ b/Engine/Source/Developer/TraceServices/Public/TraceServices/AnalysisService.h
@@ -88,7 +88,7 @@ public:
 
 struct FTimingProfilerTimer
 {
-	FString Name;
+	const TCHAR* Name;
 	uint32 Id;
 	uint32 NameHash;
 	bool IsGpuTimer;
diff --git a/Engine/Source/Runtime/Core/Private/ProfilingDebugging/CpuProfilerTrace.cpp b/Engine/Source/Runtime/Core/Private/ProfilingDebugging/CpuProfilerTrace.cpp
--- a/Engine/Source/Runtime/Core/Private/ProfilingDebugging/CpuProfilerTrace.cpp
+++ b/Engine/Source/Runtime/Core/Private/ProfilingDebugging/CpuProfilerTrace.cpp
@@ -12,21 +12,10 @@
 UE_TRACE_EVENT_BEGIN(CpuProfiler, EventSpec)
 	UE_TRACE_EVENT_FIELD(uint16, Id)
 	UE_TRACE_EVENT_FIELD(uint16, Group)
-	UE_TRACE_EVENT_FIELD(uint16, NameSize)
-	//UE_TRACE_EVENT_FIELD(uint8[], Name)
-UE_TRACE_EVENT_END()
-
-UE_TRACE_EVENT_BEGIN(CpuProfiler, DynamicEventName)
-	UE_TRACE_EVENT_FIELD(uint16, Id)
-	UE_TRACE_EVENT_FIELD(uint16, NameSize)
-	//UE_TRACE_EVENT_FIELD(uint8[], Name)
 UE_TRACE_EVENT_END()
 
 UE_TRACE_EVENT_BEGIN(CpuProfiler, EventBatch)
-	UE_TRACE_EVENT_FIELD(uint64, CycleBase)
 	UE_TRACE_EVENT_FIELD(uint32, ThreadId)
-	UE_TRACE_EVENT_FIELD(uint8, BufferSize)
-	//UE_TRACE_EVENT_FIELD(uint8[], Buffer)
 UE_TRACE_EVENT_END()
 
 struct FCpuProfilerTraceInternal
@@ -35,15 +24,12 @@ public:
 	enum
 	{
 		MaxBufferSize = 255,
-		MaxEncodedEventSize = 11, // 9 + 2
-		MaxEncodedDynamicEventSize = 13, // 9 + 2 + 2,
+		MaxEncodedEventSize = 12, // 9 + 3
 		FullBufferThreshold = MaxBufferSize - MaxEncodedEventSize,
-		DynamicEventFullBufferThreshold = MaxBufferSize - MaxEncodedDynamicEventSize
 	};
 
 	struct FThreadState
 	{
-		TMap<uint32, uint16> KnownDynamicStrings;
 		uint64 CycleBase;
 		uint64 LastCycle;
 		uint32 ThreadId;
@@ -51,188 +37,79 @@ public:
 		uint8 Buffer[MaxBufferSize];
 	};
 
-	FORCENOINLINE static FThreadState* InitThreadState(uint64 CycleBase);
-	FORCENOINLINE static void FlushThreadBuffer(FThreadState* ThreadState, uint64 NewCycleBase);
-	FORCENOINLINE static uint16 NewDynamicEventName(const TCHAR* Name);
-	static void EncodeSpecId(uint16 SpecId, uint8*& BufferPtr);
-
-	static uint32 TlsSlot;
+	static FThreadState* GetThreadState() { return ThreadLocalThreadState; }
+	FORCENOINLINE static FThreadState* InitThreadState();
+	FORCENOINLINE static void FlushThreadBuffer(FThreadState* ThreadState);
 
 private:
-	friend struct FCpuProfilerEventSpec;
-
-	FCpuProfilerTraceInternal();
-
-	TAtomic<uint32> NextSpecId;
-	TAtomic<uint16> NextDynamicEventNameId;
-
-	static FCpuProfilerTraceInternal& Instance();
+	static thread_local FThreadState* ThreadLocalThreadState;
 };
 
+thread_local FCpuProfilerTraceInternal::FThreadState* FCpuProfilerTraceInternal::ThreadLocalThreadState = nullptr;
 
-uint32 FCpuProfilerTraceInternal::TlsSlot;
-
-FCpuProfilerTraceInternal::FCpuProfilerTraceInternal()
-	: NextSpecId(1)
-	, NextDynamicEventNameId(0)
-{
-	TlsSlot = FPlatformTLS::AllocTlsSlot();
-}
-
-void FCpuProfilerTraceInternal::EncodeSpecId(uint16 SpecId, uint8*& BufferPtr)
-{
-	*BufferPtr++ = uint8(SpecId);
-	*BufferPtr++ = uint8(SpecId >> 8);
-}
-
-FCpuProfilerTraceInternal::FThreadState* FCpuProfilerTraceInternal::InitThreadState(uint64 CycleBase)
+FCpuProfilerTraceInternal::FThreadState* FCpuProfilerTraceInternal::InitThreadState()
 {
-	FThreadState* ThreadState = new FThreadState();
-	ThreadState->BufferSize = 0;
-	ThreadState->ThreadId = FPlatformTLS::GetCurrentThreadId();
-	ThreadState->CycleBase = CycleBase;
-	ThreadState->LastCycle = CycleBase;
-	FPlatformTLS::SetTlsValue(TlsSlot, ThreadState);
-	return ThreadState;
+	ThreadLocalThreadState = new FThreadState();
+	ThreadLocalThreadState->BufferSize = 0;
+	ThreadLocalThreadState->ThreadId = FPlatformTLS::GetCurrentThreadId();
+	ThreadLocalThreadState->LastCycle = 0;
+	return ThreadLocalThreadState;
 }
 
-void FCpuProfilerTraceInternal::FlushThreadBuffer(FThreadState* ThreadState, uint64 NewCycleBase)
+void FCpuProfilerTraceInternal::FlushThreadBuffer(FThreadState* ThreadState)
 {
 	UE_TRACE_LOG(CpuProfiler, EventBatch, ThreadState->BufferSize)
-		<< EventBatch.CycleBase(ThreadState->CycleBase)
 		<< EventBatch.ThreadId(ThreadState->ThreadId)
-		<< EventBatch.BufferSize(ThreadState->BufferSize)
 		<< EventBatch.Attachment(ThreadState->Buffer, ThreadState->BufferSize);
-
 	ThreadState->BufferSize = 0;
-	ThreadState->CycleBase = NewCycleBase;
-	ThreadState->LastCycle = NewCycleBase;
-}
-
-uint16 FCpuProfilerTraceInternal::NewDynamicEventName(const TCHAR* Name)
-{
-	uint16 NameLength = FCString::Strlen(Name);
-	uint16 Id = FCpuProfilerTraceInternal::Instance().NextDynamicEventNameId++;
-	auto NameCopy = [NameLength, Name](uint8* Out) {
-		int32 i = 0;
-		for (; i < NameLength; ++i)
-		{
-			Out[i] = Name[i] & 0x7f;
-		}
-		Out[i] = 0;
-	};
-	UE_TRACE_LOG(CpuProfiler, DynamicEventName, NameLength + 1)
-		<< DynamicEventName.Id(Id)
-		<< DynamicEventName.NameSize(NameLength)
-		<< DynamicEventName.Attachment(NameCopy);
-	return Id;
-}
-
-FCpuProfilerTraceInternal& FCpuProfilerTraceInternal::Instance()
-{
-	static FCpuProfilerTraceInternal I;
-	return I;
 }
 
 void FCpuProfilerTrace::OutputBeginEvent(uint16 SpecId)
 {
 	uint64 Cycle = FPlatformTime::Cycles64();
-	FCpuProfilerTraceInternal::FThreadState* ThreadState = (FCpuProfilerTraceInternal::FThreadState*)FPlatformTLS::GetTlsValue(FCpuProfilerTraceInternal::TlsSlot);
+	FCpuProfilerTraceInternal::FThreadState* ThreadState = FCpuProfilerTraceInternal::GetThreadState();
 	if (!ThreadState)
 	{
-		ThreadState = FCpuProfilerTraceInternal::InitThreadState(Cycle);
+		ThreadState = FCpuProfilerTraceInternal::InitThreadState();
 	}
 	uint64 CycleDiff = Cycle - ThreadState->LastCycle;
 	ThreadState->LastCycle = Cycle;
 	if (ThreadState->BufferSize >= FCpuProfilerTraceInternal::FullBufferThreshold)
 	{
-		FCpuProfilerTraceInternal::FlushThreadBuffer(ThreadState, Cycle);
-		CycleDiff = 0;
+		FCpuProfilerTraceInternal::FlushThreadBuffer(ThreadState);
 	}
-
 	uint8* BufferPtr = ThreadState->Buffer + ThreadState->BufferSize;
 	FTraceUtils::Encode7bit((CycleDiff << 1) | 1ull, BufferPtr);
-	FCpuProfilerTraceInternal::EncodeSpecId(SpecId, BufferPtr);
-	ThreadState->BufferSize = BufferPtr - ThreadState->Buffer;
-}
-
-void FCpuProfilerTrace::OutputBeginDynamicEvent(const TCHAR* Name)
-{
-	uint64 Cycle = FPlatformTime::Cycles64();
-	FCpuProfilerTraceInternal::FThreadState* ThreadState = (FCpuProfilerTraceInternal::FThreadState*)FPlatformTLS::GetTlsValue(FCpuProfilerTraceInternal::TlsSlot);
-	if (!ThreadState)
-	{
-		ThreadState = FCpuProfilerTraceInternal::InitThreadState(Cycle);
-	}
-	uint64 CycleDiff = Cycle - ThreadState->LastCycle;
-	ThreadState->LastCycle = Cycle;
-	if (ThreadState->BufferSize >= FCpuProfilerTraceInternal::DynamicEventFullBufferThreshold)
-	{
-		FCpuProfilerTraceInternal::FlushThreadBuffer(ThreadState, Cycle);
-		CycleDiff = 0;
-	}
-
-	uint8* BufferPtr = ThreadState->Buffer + ThreadState->BufferSize;
-	FTraceUtils::Encode7bit((CycleDiff << 1) | 1ull, BufferPtr);
-	FCpuProfilerTraceInternal::EncodeSpecId(0, BufferPtr);
-
-	uint32 NameHash = FCrc::StrCrc32(Name);
-	uint16* NameId = ThreadState->KnownDynamicStrings.Find(NameHash);
-	if (!NameId)
-	{
-		uint16 NewNameId = FCpuProfilerTraceInternal::NewDynamicEventName(Name);
-		FCpuProfilerTraceInternal::EncodeSpecId(NewNameId, BufferPtr);
-		ThreadState->KnownDynamicStrings.Add(NameHash, NewNameId);
-	}
-	else
-	{
-		FCpuProfilerTraceInternal::EncodeSpecId(*NameId, BufferPtr);
-	}
+	FTraceUtils::Encode7bit(SpecId, BufferPtr);
 	ThreadState->BufferSize = BufferPtr - ThreadState->Buffer;
 }
 
 void FCpuProfilerTrace::OutputEndEvent()
 {
 	uint64 Cycle = FPlatformTime::Cycles64();
-	FCpuProfilerTraceInternal::FThreadState* ThreadState = (FCpuProfilerTraceInternal::FThreadState*)FPlatformTLS::GetTlsValue(FCpuProfilerTraceInternal::TlsSlot);
+	FCpuProfilerTraceInternal::FThreadState* ThreadState = FCpuProfilerTraceInternal::GetThreadState();
 	uint64 CycleDiff = Cycle - ThreadState->LastCycle;
 	ThreadState->LastCycle = Cycle;
 	if (ThreadState->BufferSize >= FCpuProfilerTraceInternal::FullBufferThreshold)
 	{
-		FCpuProfilerTraceInternal::FlushThreadBuffer(ThreadState, Cycle);
-		CycleDiff = 0;
+		FCpuProfilerTraceInternal::FlushThreadBuffer(ThreadState);
 	}
 	uint8* BufferPtr = ThreadState->Buffer + ThreadState->BufferSize;
 	FTraceUtils::Encode7bit(CycleDiff << 1, BufferPtr);
 	ThreadState->BufferSize = BufferPtr - ThreadState->Buffer;
 }
 
-uint64 FCpuProfilerTrace::GetCurrentCycle()
+uint16 FCpuProfilerTrace::OutputEventType(const TCHAR* Name, ECpuProfilerGroup Group)
 {
-	return FPlatformTime::Cycles64();
-}
-
-uint16 FCpuProfilerEventSpec::AssignId(const TCHAR* Name, ECpuProfilerGroup Group)
-{
-	uint32 WideId = FCpuProfilerTraceInternal::Instance().NextSpecId++;
+	static TAtomic<uint32> NextSpecId(1);
+	uint32 WideId = NextSpecId++;
 	check(WideId <= 0xFFFF);
 	uint16 Id = uint16(WideId);
-	uint16 NameLength = FCString::Strlen(Name);
-
-	auto NameCopy = [NameLength, Name](uint8* Out) {
-		int32 i = 0;
-		for (; i < NameLength; ++i)
-		{
-			Out[i] = Name[i] & 0x7f;
-		}
-		Out[i] = 0;
-	};
-
-	UE_TRACE_LOG(CpuProfiler, EventSpec, NameLength + 1)
+	uint16 NameSize = (FCString::Strlen(Name) + 1) * sizeof(TCHAR);
+	UE_TRACE_LOG(CpuProfiler, EventSpec, NameSize)
 		<< EventSpec.Id(Id)
 		<< EventSpec.Group(uint16(Group))
-		<< EventSpec.NameSize(NameLength)
-		<< EventSpec.Attachment(NameCopy);
+		<< EventSpec.Attachment(Name, NameSize);
 	return Id;
 }
 
diff --git a/Engine/Source/Runtime/Core/Private/ProfilingDebugging/MiscTrace.cpp b/Engine/Source/Runtime/Core/Private/ProfilingDebugging/MiscTrace.cpp
--- a/Engine/Source/Runtime/Core/Private/ProfilingDebugging/MiscTrace.cpp
+++ b/Engine/Source/Runtime/Core/Private/ProfilingDebugging/MiscTrace.cpp
@@ -21,12 +21,10 @@ UE_TRACE_EVENT_END()
 
 UE_TRACE_EVENT_BEGIN(Misc, SetThreadGroup)
 	UE_TRACE_EVENT_FIELD(uint32, ThreadId)
-	UE_TRACE_EVENT_FIELD(uint16, GroupNameLength)
 UE_TRACE_EVENT_END()
 
 UE_TRACE_EVENT_BEGIN(Misc, BeginThreadGroupScope)
 	UE_TRACE_EVENT_FIELD(uint32, CurrentThreadId)
-	UE_TRACE_EVENT_FIELD(uint16, GroupNameLength)
 UE_TRACE_EVENT_END()
 
 UE_TRACE_EVENT_BEGIN(Misc, EndThreadGroupScope)
@@ -61,21 +59,12 @@ void FMiscTrace::OutputRegisterGameThread(uint32 Id)
 
 void FMiscTrace::OutputCreateThread(uint32 Id, const TCHAR* Name, uint32 Priority)
 {
-	uint16 NameLength = FCString::Strlen(Name);
-	auto NameCopy = [NameLength, Name](uint8* Out) {
-		int32 i = 0;
-		for (; i < NameLength; ++i)
-		{
-			Out[i] = Name[i] & 0x7f;
-		}
-		Out[i] = 0;
-	};
-	UE_TRACE_LOG(Misc, CreateThread, NameLength + 1)
+	uint16 NameSize = (FCString::Strlen(Name) + 1) * sizeof(TCHAR);
+	UE_TRACE_LOG(Misc, CreateThread, NameSize)
 		<< CreateThread.CurrentThreadId(FPlatformTLS::GetCurrentThreadId())
 		<< CreateThread.CreatedThreadId(Id)
 		<< CreateThread.Priority(Priority)
-		<< CreateThread.NameSize(NameLength)
-		<< CreateThread.Attachment(NameCopy);
+		<< CreateThread.Attachment(Name, NameSize);
 }
 
 static const char* GetThreadGroupName(ETraceThreadGroup Group)
@@ -108,36 +97,19 @@ static const char* GetThreadGroupName(ETraceThreadGroup Group)
 void FMiscTrace::OutputSetThreadGroup(uint32 Id, ETraceThreadGroup Group)
 {
 	const char* GroupName = GetThreadGroupName(Group);
-	uint16 NameLength = strlen(GroupName);
-	auto NameCopy = [NameLength, GroupName](uint8* Out) {
-		int32 i = 0;
-		for (; i < NameLength; ++i)
-		{
-			Out[i] = GroupName[i];
-		}
-		Out[i] = 0;
-	};
-	UE_TRACE_LOG(Misc, SetThreadGroup, NameLength + 1)
+	uint16 NameSize = strlen(GroupName) + 1;
+	UE_TRACE_LOG(Misc, SetThreadGroup, NameSize)
 		<< SetThreadGroup.ThreadId(Id)
-		<< SetThreadGroup.GroupNameLength(NameLength)
-		<< SetThreadGroup.Attachment(NameCopy);
+		<< SetThreadGroup.Attachment(GroupName, NameSize);
 }
 
 void FMiscTrace::OutputBeginThreadGroupScope(ETraceThreadGroup Group)
 {
 	const char* GroupName = GetThreadGroupName(Group);
-	uint16 NameLength = strlen(GroupName);
-	auto NameCopy = [NameLength, GroupName](uint8* Out) {
-		int32 i = 0;
-		for (; i < NameLength; ++i)
-		{
-			Out[i] = GroupName[i];
-		}
-		Out[i] = 0;
-	};
-	UE_TRACE_LOG(Misc, BeginThreadGroupScope, NameLength + 1)
+	uint16 NameSize = strlen(GroupName) + 1;
+	UE_TRACE_LOG(Misc, BeginThreadGroupScope, NameSize)
 		<< BeginThreadGroupScope.CurrentThreadId(FPlatformTLS::GetCurrentThreadId())
-		<< BeginThreadGroupScope.Attachment(NameCopy);
+		<< BeginThreadGroupScope.Attachment(GroupName, NameSize);
 }
 
 void FMiscTrace::OutputEndThreadGroupScope()
diff --git a/Engine/Source/Runtime/Core/Private/Stats/StatsTrace.cpp b/Engine/Source/Runtime/Core/Private/Stats/StatsTrace.cpp
--- a/Engine/Source/Runtime/Core/Private/Stats/StatsTrace.cpp
+++ b/Engine/Source/Runtime/Core/Private/Stats/StatsTrace.cpp
@@ -27,7 +27,7 @@ public:
 	enum
 	{
 		MaxBufferSize = 512,
-		MaxEncodedEventSize = 21, // 3 + 9 + 9
+		MaxEncodedEventSize = 27, // 9 + 9 + 9
 		FullBufferThreshold = MaxBufferSize - MaxEncodedEventSize,
 	};
 
@@ -49,32 +49,25 @@ public:
 		uint8 Buffer[MaxBufferSize];
 	};
 
+	static FThreadState* GetThreadState() { return ThreadLocalThreadState; }
 	FORCENOINLINE static FThreadState* InitThreadState();
 	FORCENOINLINE static void FlushThreadBuffer(FThreadState* ThreadState);
 	static void BeginEncodeOp(const FName& Stat, EOpType Op, FStatsTraceInternal::FThreadState*& OutThreadState, uint8*& OutBufferPtr);
 	static void EndEncodeOp(FStatsTraceInternal::FThreadState* ThreadState, uint8* BufferPtr);
 
-	static uint32 TlsSlot;
-
 private:
-	FStatsTraceInternal();
+	static thread_local FThreadState* ThreadLocalThreadState;
 };
 
-uint32 FStatsTraceInternal::TlsSlot;
-
-FStatsTraceInternal::FStatsTraceInternal()
-{
-	TlsSlot = FPlatformTLS::AllocTlsSlot();
-}
+thread_local FStatsTraceInternal::FThreadState* FStatsTraceInternal::ThreadLocalThreadState = nullptr;
 
 FStatsTraceInternal::FThreadState* FStatsTraceInternal::InitThreadState()
 {
-	FThreadState* ThreadState = new FThreadState();
-	ThreadState->BufferSize = 0;
-	ThreadState->ThreadId = FPlatformTLS::GetCurrentThreadId();
-	ThreadState->LastCycle = 0;
-	FPlatformTLS::SetTlsValue(TlsSlot, ThreadState);
-	return ThreadState;
+	ThreadLocalThreadState = new FThreadState();
+	ThreadLocalThreadState->BufferSize = 0;
+	ThreadLocalThreadState->ThreadId = FPlatformTLS::GetCurrentThreadId();
+	ThreadLocalThreadState->LastCycle = 0;
+	return ThreadLocalThreadState;
 }
 
 void FStatsTraceInternal::FlushThreadBuffer(FThreadState* ThreadState)
@@ -88,7 +81,7 @@ void FStatsTraceInternal::FlushThreadBuffer(FThreadState* ThreadState)
 void FStatsTraceInternal::BeginEncodeOp(const FName& Stat, EOpType Op, FThreadState*& OutThreadState, uint8*& OutBufferPtr)
 {
 	uint64 Cycle = FPlatformTime::Cycles64();
-	OutThreadState = (FThreadState*)FPlatformTLS::GetTlsValue(TlsSlot);
+	OutThreadState = GetThreadState();
 	if (!OutThreadState)
 	{
 		OutThreadState = InitThreadState();
diff --git a/Engine/Source/Runtime/Core/Public/HAL/PlatformMisc.h b/Engine/Source/Runtime/Core/Public/HAL/PlatformMisc.h
--- a/Engine/Source/Runtime/Core/Public/HAL/PlatformMisc.h
+++ b/Engine/Source/Runtime/Core/Public/HAL/PlatformMisc.h
@@ -137,19 +137,13 @@ public:
 	TRACE_CPUPROFILER_EVENT_SCOPE(Name) \
 	FScopedNamedEventStatic ANONYMOUS_VARIABLE(NamedEvent_##Name##_)(Color, NAMED_EVENT_STR(#Name));
 
-#define SCOPED_NAMED_EVENT_FSTRING(Text, Color) \
-	TRACE_CPUPROFILER_DYNAMIC_EVENT_SCOPE(*Text) \
-	FScopedNamedEvent       ANONYMOUS_VARIABLE(NamedEvent_)         (Color, *Text);
+#define SCOPED_NAMED_EVENT_FSTRING(Text, Color)  FScopedNamedEvent       ANONYMOUS_VARIABLE(NamedEvent_)         (Color, *Text);
 
 #define SCOPED_NAMED_EVENT_TEXT(Text, Color) \
-	TRACE_CPUPROFILER_EVENT_SCOPE_TEXT(L##Text) \
+	TRACE_CPUPROFILER_EVENT_SCOPE_TEXT(TEXT(Text)) \
 	FScopedNamedEventStatic ANONYMOUS_VARIABLE(NamedEvent_)         (Color, NAMED_EVENT_STR(Text));
 
-#define SCOPED_NAMED_EVENT_F(Format, Color, ...) \
-	FString __NamedEventFormattedName(FString::Printf(Format, __VA_ARGS__)); \
-	TRACE_CPUPROFILER_DYNAMIC_EVENT_SCOPE(*__NamedEventFormattedName) \
-	FScopedNamedEvent       ANONYMOUS_VARIABLE(NamedEvent_)         (Color, *__NamedEventFormattedName);
-
+#define SCOPED_NAMED_EVENT_F(Format, Color, ...) FScopedNamedEvent       ANONYMOUS_VARIABLE(NamedEvent_)         (Color, *FString::Printf(Format, __VA_ARGS__));
 #define SCOPED_PROFILER_COLOR(Color)			 FScopedProfilerColor    ANONYMOUS_VARIABLE(ProfilerColor_##Name##_)(Color);
 
 
diff --git a/Engine/Source/Runtime/Core/Public/ProfilingDebugging/CpuProfilerTrace.h b/Engine/Source/Runtime/Core/Public/ProfilingDebugging/CpuProfilerTrace.h
--- a/Engine/Source/Runtime/Core/Public/ProfilingDebugging/CpuProfilerTrace.h
+++ b/Engine/Source/Runtime/Core/Public/ProfilingDebugging/CpuProfilerTrace.h
@@ -22,49 +22,30 @@ enum ECpuProfilerGroup
 
 struct FCpuProfilerTrace
 {
+	CORE_API static uint16 OutputEventType(const TCHAR* Name, ECpuProfilerGroup Group);
 	CORE_API static void OutputBeginEvent(uint16 SpecId);
-	CORE_API static void OutputBeginDynamicEvent(const TCHAR* Name);
 	CORE_API static void OutputEndEvent();
-	CORE_API static uint64 GetCurrentCycle();
-};
-
-struct FCpuProfilerEventSpec
-{
-	CORE_API static uint16 AssignId(const TCHAR* Name, ECpuProfilerGroup Group);
-};
-
-struct FCpuProfilerEventScope
-{
-	FCpuProfilerEventScope(uint16 InSpecId)
-	{
-		FCpuProfilerTrace::OutputBeginEvent(InSpecId);
-	}
 
-	~FCpuProfilerEventScope()
+	struct FEventScope
 	{
-		FCpuProfilerTrace::OutputEndEvent();
-	}
-};
-
-struct FCpuProfilerDynamicEventScope
-{
-	FCpuProfilerDynamicEventScope(const TCHAR* InName)
-	{
-		FCpuProfilerTrace::OutputBeginDynamicEvent(InName);
-	}
-
-	~FCpuProfilerDynamicEventScope()
-	{
-		FCpuProfilerTrace::OutputEndEvent();
-	}
+		FEventScope(uint16 InSpecId)
+		{
+			OutputBeginEvent(InSpecId);
+		}
+
+		~FEventScope()
+		{
+			OutputEndEvent();
+		}
+	};
 };
 
 #define TRACE_CPUPROFILER_EVENT_SCOPE_TEXT_GROUP(Name, Group) \
 	static uint16 PREPROCESSOR_JOIN(__CpuProfilerEventSpecId, __LINE__); \
 	if (PREPROCESSOR_JOIN(__CpuProfilerEventSpecId, __LINE__) == 0) { \
-		PREPROCESSOR_JOIN(__CpuProfilerEventSpecId, __LINE__) = FCpuProfilerEventSpec::AssignId(Name, Group); \
+		PREPROCESSOR_JOIN(__CpuProfilerEventSpecId, __LINE__) = FCpuProfilerTrace::OutputEventType(Name, Group); \
 	} \
-	FCpuProfilerEventScope PREPROCESSOR_JOIN(__CpuProfilerEventScope, __LINE__)(PREPROCESSOR_JOIN(__CpuProfilerEventSpecId, __LINE__));
+	FCpuProfilerTrace::FEventScope PREPROCESSOR_JOIN(__CpuProfilerEventScope, __LINE__)(PREPROCESSOR_JOIN(__CpuProfilerEventSpecId, __LINE__));
 
 #define TRACE_CPUPROFILER_EVENT_SCOPE_TEXT(Name) \
 	TRACE_CPUPROFILER_EVENT_SCOPE_TEXT_GROUP(Name, CpuProfilerGroup_Default)
@@ -75,15 +56,11 @@ struct FCpuProfilerDynamicEventScope
 #define TRACE_CPUPROFILER_EVENT_SCOPE(Name) \
 	TRACE_CPUPROFILER_EVENT_SCOPE_TEXT_GROUP(TEXT(#Name), CpuProfilerGroup_Default);
 
-#define TRACE_CPUPROFILER_DYNAMIC_EVENT_SCOPE(Name) \
-	FCpuProfilerDynamicEventScope ANONYMOUS_VARIABLE(__Scope)(Name);
-
 #else
 
 #define TRACE_CPUPROFILER_EVENT_SCOPE_GROUP(Name, Group)
 #define TRACE_CPUPROFILER_EVENT_SCOPE(Name)
 #define TRACE_CPUPROFILER_EVENT_SCOPE_TEXT_GROUP(Name, Group)
 #define TRACE_CPUPROFILER_EVENT_SCOPE_TEXT(Name)
-#define TRACE_CPUPROFILER_DYNAMIC_EVENT_SCOPE(Name)
 
 #endif
