Compile fix for CPUPROFILERTRACE_ENABLED=1

diff --git a/Engine/Source/Runtime/Core/Private/Stats/Stats2.cpp b/Engine/Source/Runtime/Core/Private/Stats/Stats2.cpp
--- a/Engine/Source/Runtime/Core/Private/Stats/Stats2.cpp
+++ b/Engine/Source/Runtime/Core/Private/Stats/Stats2.cpp
@@ -606,7 +606,7 @@ public:
 #if CPUPROFILERTRACE_ENABLED
 		if (bCycleStat)
 		{
-			Result->TraceCpuProfilerSpecId = FCpuProfilerEventSpec::AssignId(*StatDescription, CpuProfilerGroup_Stats);
+			Result->TraceCpuProfilerSpecId = FCpuProfilerTrace::OutputEventType(*StatDescription, CpuProfilerGroup_Stats);
 		}
 #endif
 #if STATSTRACE_ENABLED
