FName optimizations and improvements

* Added new ToString() with a TCHAR buffer that avoids heap allocations
* Other smallish optimizations
* Improved some comments

diff --git a/Engine/Source/Runtime/Core/Private/UObject/UnrealNames.cpp b/Engine/Source/Runtime/Core/Private/UObject/UnrealNames.cpp
--- a/Engine/Source/Runtime/Core/Private/UObject/UnrealNames.cpp
+++ b/Engine/Source/Runtime/Core/Private/UObject/UnrealNames.cpp
@@ -54,45 +54,41 @@ static bool operator==(FNameEntryHeader A, FNameEntryHeader B)
 	return (uint16&)A == (uint16&)B;
 }
 
-union FNameBuffer
-{
-	ANSICHAR AnsiName[NAME_SIZE];
-	WIDECHAR WideName[NAME_SIZE];
-
-	template<typename FromCharType, typename ToCharType>
-	ToCharType* ConvertInPlace(uint32 Len);
-};
-
-template<> ANSICHAR* FNameBuffer::ConvertInPlace<ANSICHAR, ANSICHAR>(uint32 Len)
-{
-	return AnsiName;
-}
-
-template<> WIDECHAR* FNameBuffer::ConvertInPlace<WIDECHAR, WIDECHAR>(uint32 Len)
+template<typename FromCharType, typename ToCharType>
+ToCharType* ConvertInPlace(FromCharType* Str, uint32 Len)
 {
-	return WideName;
+	static_assert(TIsSame<FromCharType, ToCharType>::Value, "Unsupported conversion");
+	return Str;
 }
 
-template<> WIDECHAR* FNameBuffer::ConvertInPlace<ANSICHAR, WIDECHAR>(uint32 Len)
+template<>
+WIDECHAR* ConvertInPlace<ANSICHAR, WIDECHAR>(ANSICHAR* Str, uint32 Len)
 {
 	for (uint32 Index = Len; Index--; )
 	{
-		WideName[Index] = AnsiName[Index];
+		reinterpret_cast<WIDECHAR*>(Str)[Index] = Str[Index];
 	}
 
-	return WideName;
+	return reinterpret_cast<WIDECHAR*>(Str);
 }
 
-template<> ANSICHAR* FNameBuffer::ConvertInPlace<WIDECHAR, ANSICHAR>(uint32 Len)
+template<>
+ANSICHAR* ConvertInPlace<WIDECHAR, ANSICHAR>(WIDECHAR* Str, uint32 Len)
 {
 	for (uint32 Index = 0; Index < Len; ++Index)
 	{
-		AnsiName[Index] = WideName[Index];
+		reinterpret_cast<ANSICHAR*>(Str)[Index] = Str[Index];
 	}
 
-	return AnsiName;
+	return reinterpret_cast<ANSICHAR*>(Str);
 }
 
+union FNameBuffer
+{
+	ANSICHAR AnsiName[NAME_SIZE];
+	WIDECHAR WideName[NAME_SIZE];
+};
+
 struct FNameStringView
 {
 	FNameStringView() : Data(nullptr), Len(0), bIsWide(false) {}
@@ -719,6 +715,7 @@ private:
 
 	// Put constant lookup on separate cache line to avoid it being constantly invalidated by insertion
 	alignas(PLATFORM_CACHE_LINE_SIZE) FNameEntryId ENameToEntry[NAME_MaxHardcodedNameIndex] = {};
+	uint32 LargestEnameUnstableId;
 	TMap<FNameEntryId, EName, TInlineSetAllocator<MaxENames>> EntryToEName;
 };
 
@@ -744,11 +741,13 @@ FNamePool::FNamePool()
 #undef REGISTER_NAME
 
 	// Make reverse mapping
+	LargestEnameUnstableId = 0;
 	for (uint32 ENameIndex = 0; ENameIndex < NAME_MaxHardcodedNameIndex; ++ENameIndex)
 	{
 		if (ENameIndex == NAME_None || ENameToEntry[ENameIndex])
 		{
 			EntryToEName.Add(ENameToEntry[ENameIndex], (EName)ENameIndex);
+			LargestEnameUnstableId = FMath::Max(LargestEnameUnstableId, ENameToEntry[ENameIndex].ToUnstableInt());
 		}
 	}
 
@@ -877,7 +876,7 @@ bool FNamePool::IsValid(FNameEntryHandle Handle) const
 
 const EName* FNamePool::FindEName(FNameEntryId Id) const
 {
-	return EntryToEName.Find(Id);
+	return Id.ToUnstableInt() > LargestEnameUnstableId ? nullptr : EntryToEName.Find(Id);
 }
 
 static bool bNamePoolInitialized;
@@ -898,9 +897,16 @@ static FNamePool& GetNamePool()
 	return *Singleton;
 }
 
+// Only call from functions guaranteed to run after FName lazy initialization
+static FNamePool& GetNamePoolPostInit()
+{
+	checkSlow(bNamePoolInitialized);
+	return (FNamePool&)NamePoolData;
+}
+
 bool operator==(FNameEntryId Id, EName Ename)
 {
-	return Id == GetNamePool().Find(Ename);
+	return Id == GetNamePoolPostInit().Find(Ename);
 }
 
 static int32 CompareDifferentIdsAlphabetically(FNameEntryId AId, FNameEntryId BId)
@@ -919,11 +925,11 @@ static int32 CompareDifferentIdsAlphabetically(FNameEntryId AId, FNameEntryId BI
 		FNameBuffer& AnsiBuffer =	AView.bIsWide ? BBuffer : ABuffer;
 
 #ifndef WITH_CUSTOM_NAME_ENCODING
-		memcpy(AnsiBuffer.AnsiName, AnsiView.Ansi, AnsiView.Len * sizeof(ANSICHAR));
+		FPlatformMemory::Memcpy(AnsiBuffer.AnsiName, AnsiView.Ansi, AnsiView.Len * sizeof(ANSICHAR));
 		AnsiView.Ansi = AnsiBuffer.AnsiName;
 #endif
 
-		AnsiBuffer.ConvertInPlace<ANSICHAR, WIDECHAR>(AnsiView.Len);
+		ConvertInPlace<ANSICHAR, WIDECHAR>(AnsiBuffer.AnsiName, AnsiView.Len);
 		AnsiView.bIsWide = true;
 	}
 
@@ -1013,25 +1019,25 @@ static uint16 GetRawNonCasePreservingHash(const TCharType* Source)
 
 void FNameEntry::StoreName(const ANSICHAR* InName, uint32 Len)
 {
-	memcpy(AnsiName, InName, sizeof(ANSICHAR) * Len);
+	FPlatformMemory::Memcpy(AnsiName, InName, sizeof(ANSICHAR) * Len);
 	Encode(AnsiName, Len);
 }
 
 void FNameEntry::StoreName(const WIDECHAR* InName, uint32 Len)
 {
-	memcpy(WideName, InName, sizeof(WIDECHAR) * Len);
+	FPlatformMemory::Memcpy(WideName, InName, sizeof(WIDECHAR) * Len);
 	Encode(WideName, Len);
 }
 
 void FNameEntry::CopyUnterminatedName(ANSICHAR* Out) const
 {
-	memcpy(Out, AnsiName, sizeof(ANSICHAR) * Header.Len);
+	FPlatformMemory::Memcpy(Out, AnsiName, sizeof(ANSICHAR) * Header.Len);
 	Decode(Out, Header.Len);
 }
 
 void FNameEntry::CopyUnterminatedName(WIDECHAR* Out) const
 {
-	memcpy(Out, WideName, sizeof(WIDECHAR) * Header.Len);
+	FPlatformMemory::Memcpy(Out, WideName, sizeof(WIDECHAR) * Header.Len);
 	Decode(Out, Header.Len);
 }
 
@@ -1061,6 +1067,39 @@ FORCEINLINE FNameStringView FNameEntry::MakeView(FNameBuffer& OptionalDecodeBuff
 					: FNameStringView(GetUnterminatedName(OptionalDecodeBuffer.AnsiName), GetNameLength());
 }
 
+void FNameEntry::GetUnterminatedName(TCHAR* OutName, uint32 OutLen) const
+{
+	check(static_cast<int32>(OutLen) >= GetNameLength());
+	CopyAndConvertUnterminatedName(OutName);
+}
+
+void FNameEntry::GetName(TCHAR(&OutName)[NAME_SIZE]) const
+{
+	CopyAndConvertUnterminatedName(OutName);
+	OutName[GetNameLength()] = '\0';
+}
+
+void FNameEntry::CopyAndConvertUnterminatedName(TCHAR* OutName) const
+{
+	if (sizeof(TCHAR) < sizeof(WIDECHAR) && IsWide()) // Normally compiled out
+	{
+		FNameBuffer Temp;
+		CopyUnterminatedName(Temp.WideName);
+		ConvertInPlace<WIDECHAR, TCHAR>(Temp.WideName, Header.Len);
+		FPlatformMemory::Memcpy(OutName, Temp.AnsiName, Header.Len * sizeof(TCHAR));
+	}
+	else if (IsWide())
+	{
+		CopyUnterminatedName((WIDECHAR*)OutName);
+		ConvertInPlace<WIDECHAR, TCHAR>((WIDECHAR*)OutName, Header.Len);
+	}
+	else
+	{
+		CopyUnterminatedName((ANSICHAR*)OutName);
+		ConvertInPlace<ANSICHAR, TCHAR>((ANSICHAR*)OutName, Header.Len);
+	}
+}
+
 void FNameEntry::GetAnsiName(ANSICHAR(&Out)[NAME_SIZE]) const
 {
 	check(!IsWide());
@@ -1081,12 +1120,12 @@ static const TCHAR* EntryToCString(const FNameEntry& Entry, FNameBuffer& Temp)
 	if (Entry.IsWide())
 	{
 		Entry.GetWideName(Temp.WideName);
-		return Temp.ConvertInPlace<WIDECHAR, TCHAR>(Entry.GetNameLength() + 1);
+		return ConvertInPlace<WIDECHAR, TCHAR>(Temp.WideName, Entry.GetNameLength() + 1);
 	}
 	else
 	{
 		Entry.GetAnsiName(Temp.AnsiName);
-		return Temp.ConvertInPlace<ANSICHAR, TCHAR>(Entry.GetNameLength() + 1);
+		return ConvertInPlace<ANSICHAR, TCHAR>(Temp.AnsiName, Entry.GetNameLength() + 1);
 	}
 }
 
@@ -1327,12 +1366,12 @@ FString FName::NameToDisplayString( const FString& InDisplayName, const bool bIs
 
 const EName* FName::ToEName() const
 {
-	return GetNamePool().FindEName(ComparisonIndex);
+	return GetNamePoolPostInit().FindEName(ComparisonIndex);
 }
 
 bool FName::IsWithinBounds(FNameEntryId Id)
 {
-	return GetNamePool().IsValid(Id);
+	return GetNamePoolPostInit().IsValid(Id);
 }
 
 /*-----------------------------------------------------------------------------
@@ -1593,14 +1632,12 @@ struct FNameHelper
 };
 
 
+#if WITH_CASE_PRESERVING_NAME
 FNameEntryId FName::GetComparisonIdFromDisplayId(FNameEntryId DisplayId)
 {
-#if WITH_CASE_PRESERVING_NAME
 	return GetEntry(DisplayId)->ComparisonId;
-#else
-	return DisplayId;
-#endif
 }
+#endif
 
 FName::FName(const WIDECHAR* Name, EFindName FindType)
 	: FName(FNameHelper::MakeDetectNumber(MakeUnconvertedView(Name), FindType))
@@ -1682,14 +1719,16 @@ int32 FName::Compare( const FName& Other ) const
 	return CompareDifferentIdsAlphabetically(ComparisonIndex, Other.ComparisonIndex);
 }
 
-FString FName::GetPlainNameString() const
+uint32 FName::GetPlainNameString(TCHAR(&OutName)[NAME_SIZE]) const
 {
-	if (const FNameEntry* CurEntry = GetDisplayNameEntry())
-	{
-		return CurEntry->GetPlainNameString();
-	}
+	const FNameEntry& Entry = *GetDisplayNameEntry();
+	Entry.GetName(OutName);
+	return Entry.GetNameLength();
+}
 
-	return TEXT("*INVALID*");
+FString FName::GetPlainNameString() const
+{
+	return GetDisplayNameEntry()->GetPlainNameString();
 }
 
 void FName::GetPlainANSIString(ANSICHAR(&AnsiName)[NAME_SIZE]) const
@@ -1716,11 +1755,8 @@ FString FName::ToString() const
 {
 	if (GetNumber() == NAME_NO_NUMBER_INTERNAL)
 	{
-		if (const FNameEntry* const DisplayEntry = GetDisplayNameEntry())
-		{
-			// Avoids some extra allocations in non-number case
-			return DisplayEntry->GetPlainNameString();
-		}
+		// Avoids some extra allocations in non-number case
+		return GetDisplayNameEntry()->GetPlainNameString();
 	}
 	
 	FString Out;	
@@ -1733,11 +1769,7 @@ void FName::ToString(FString& Out) const
 	// A version of ToString that saves at least one string copy
 	const FNameEntry* const NameEntry = GetDisplayNameEntry();
 
-	if (NameEntry == nullptr)
-	{
-		Out = TEXT("*INVALID*");
-	}
-	else if (GetNumber() == NAME_NO_NUMBER_INTERNAL)
+	if (GetNumber() == NAME_NO_NUMBER_INTERNAL)
 	{
 		Out.Empty(NameEntry->GetNameLength());
 		NameEntry->AppendNameToString(Out);
@@ -1752,22 +1784,58 @@ void FName::ToString(FString& Out) const
 	}
 }
 
-void FName::AppendString(FString& Out) const
+uint32 FName::GetStringLength() const
 {
-	const FNameEntry* const NameEntry = GetDisplayNameEntry();
+	const FNameEntry& Entry = *GetDisplayNameEntry();
+	uint32 NameLen = Entry.GetNameLength();
 
-	if (NameEntry == nullptr)
+	if (GetNumber() == NAME_NO_NUMBER_INTERNAL)
 	{
-		Out += TEXT("*INVALID*");
+		return NameLen;
 	}
 	else
 	{
-		NameEntry->AppendNameToString( Out );
-		if (GetNumber() != NAME_NO_NUMBER_INTERNAL)
-		{
-			Out += TEXT('_');
-			Out.AppendInt(NAME_INTERNAL_TO_EXTERNAL(GetNumber()));
-		}
+		TCHAR NumberSuffixStr[16];
+		int32 SuffixLen = FCString::Sprintf(NumberSuffixStr, TEXT("_%d"), NAME_INTERNAL_TO_EXTERNAL(GetNumber()));
+		check(SuffixLen > 0);
+
+		return NameLen + SuffixLen;
+	}
+}
+
+uint32 FName::ToString(TCHAR* Out, uint32 OutSize) const
+{
+	const FNameEntry& Entry = *GetDisplayNameEntry();
+	uint32 NameLen = Entry.GetNameLength();
+	Entry.GetUnterminatedName(Out, OutSize);
+
+	if (GetNumber() == NAME_NO_NUMBER_INTERNAL)
+	{
+		Out[NameLen] = '\0';
+		return NameLen;
+	}
+	else
+	{
+		TCHAR NumberSuffixStr[16];
+		int32 SuffixLen = FCString::Sprintf(NumberSuffixStr, TEXT("_%d"), NAME_INTERNAL_TO_EXTERNAL(GetNumber()));
+		uint32 TotalLen = NameLen + SuffixLen;
+		check(SuffixLen > 0 && OutSize > TotalLen);
+
+		FPlatformMemory::Memcpy(Out + NameLen, NumberSuffixStr, SuffixLen * sizeof(TCHAR));
+		Out[TotalLen] = '\0';
+		return TotalLen;
+	}
+}
+
+void FName::AppendString(FString& Out) const
+{
+	const FNameEntry* const NameEntry = GetDisplayNameEntry();
+
+	NameEntry->AppendNameToString( Out );
+	if (GetNumber() != NAME_NO_NUMBER_INTERNAL)
+	{
+		Out += TEXT('_');
+		Out.AppendInt(NAME_INTERNAL_TO_EXTERNAL(GetNumber()));
 	}
 }
 
@@ -1817,6 +1885,7 @@ bool FName::IsValidXName(const FString& InName, const FString& InInvalidChars, F
 
 void FName::AutoTest()
 {
+#if DO_CHECK
 	check(FNameHash::IsAnsiNone("None", 4) == 1);
 	check(FNameHash::IsAnsiNone("none", 4) == 1);
 	check(FNameHash::IsAnsiNone("NONE", 4) == 1);
@@ -1835,12 +1904,22 @@ void FName::AutoTest()
 	check(AutoTest_1 != AutoTest_2);
 	check(AutoTest_1 == autoTest_1);
 	check(AutoTest_1 == autoTeSt_1);
+
+	TCHAR Buffer[FName::StringBufferSize];
+
 #if WITH_CASE_PRESERVING_NAME
 	check(!FCString::Strcmp(*AutoTest_1.ToString(), TEXT("AutoTest_1")));
 	check(!FCString::Strcmp(*autoTest_1.ToString(), TEXT("autoTest_1")));
 	check(!FCString::Strcmp(*autoTeSt_1.ToString(), TEXT("autoTeSt_1")));
 	check(!FCString::Strcmp(*AutoTestB_2.ToString(), TEXT("AutoTestB_2")));
+	
+	check(FName("ABC").ToString(Buffer) == 3 &&			!FCString::Strcmp(Buffer, TEXT("ABC")));
+	check(FName("abc").ToString(Buffer) == 3 &&			!FCString::Strcmp(Buffer, TEXT("abc")));
+	check(FName(TEXT("abc")).ToString(Buffer) == 3 &&	!FCString::Strcmp(Buffer, TEXT("abc")));
+	check(FName("ABC_0").ToString(Buffer) == 5 &&		!FCString::Strcmp(Buffer, TEXT("ABC_0")));
+	check(FName("ABC_10").ToString(Buffer) == 6 &&		!FCString::Strcmp(Buffer, TEXT("ABC_10")));	
 #endif
+
 	check(autoTest_1.GetComparisonIndex() == AutoTest_2.GetComparisonIndex());
 	check(autoTest_1.GetPlainNameString() == AutoTest_1.GetPlainNameString());
 	check(autoTest_1.GetPlainNameString() == AutoTest_2.GetPlainNameString());
@@ -1848,6 +1927,13 @@ void FName::AutoTest()
 	check(AutoTestB_2.GetNumber() == AutoTest_2.GetNumber());
 	check(autoTest_1.GetNumber() != AutoTest_2.GetNumber());
 
+	check(FCStringAnsi::Strlen("None") == FName().GetStringLength());
+	check(FCStringAnsi::Strlen("ABC") == FName("ABC").GetStringLength());
+	check(FCStringAnsi::Strlen("ABC_0") == FName("ABC_0").GetStringLength());
+	check(FCStringAnsi::Strlen("ABC_9") == FName("ABC_9").GetStringLength());
+	check(FCStringAnsi::Strlen("ABC_10") == FName("ABC_10").GetStringLength());
+	check(FCStringAnsi::Strlen("ABC_2000000000") == FName("ABC_2000000000").GetStringLength());
+
 	const FName NullName(static_cast<ANSICHAR*>(nullptr));
 	check(NullName.IsNone());
 	check(NullName == FName(NAME_None));
@@ -1875,6 +1961,9 @@ void FName::AutoTest()
 	Wide[4] = 60000;
 	FName WideName(*Wide);
 	check(WideName.GetPlainNameString() == Wide);
+	check(FName(*Wide).GetPlainNameString() == Wide);
+	check(FName(*Wide).ToString(Buffer) == 5 && !FCString::Strcmp(Buffer, *Wide));
+	check(Wide.Len() == WideName.GetStringLength());
 	FString WideLong = FString::ChrN(1000, 60000);
 	check(FName(*WideLong).GetPlainNameString() == WideLong);
 
@@ -1960,7 +2049,7 @@ void FName::AutoTest()
 
 #if 0
 	// Check hash table growth still yields the same unique FName ids
-	static int32 OverflowAtLeastTwiceCount = 4 * FNamePoolInitialSlots * FNamePoolShards;
+	static int32 OverflowAtLeastTwiceCount = 4 * FNamePoolInitialSlotsPerShard * FNamePoolShards;
 	TArray<FNameEntryId> Ids;
 	for (int I = 0; I < OverflowAtLeastTwiceCount; ++I)
 	{
@@ -1980,8 +2069,10 @@ void FName::AutoTest()
 		check(Id == OldId);
 	}
 #endif
+#endif // DO_CHECK
 }
 
+
 /*-----------------------------------------------------------------------------
 	FNameEntry implementation.
 -----------------------------------------------------------------------------*/
@@ -2071,17 +2162,15 @@ FArchive& operator<<(FArchive& Ar, FNameEntrySerialized& E)
 			ANSICHAR* AnsiName = const_cast<ANSICHAR*>(E.GetAnsiName());
 			Ar.Serialize(AnsiName, StringLen);
 		}
-		if (Ar.UE4Ver() >= VER_UE4_NAME_HASHES_SERIALIZED)
-		{
-			// These hashes are no longer used. They're only kept to maintain serialization format.
-			// Please remove them if you ever change serialization format.
-			Ar << E.NonCasePreservingHash;
-			Ar << E.CasePreservingHash;
-			E.bWereHashesLoaded = true;
-		}
+
+		uint16 DummyHashes[2];
+		uint32 SkipPastHashBytes = (Ar.UE4Ver() >= VER_UE4_NAME_HASHES_SERIALIZED) * sizeof(DummyHashes);
+		Ar.Serialize(&DummyHashes, SkipPastHashBytes);
 	}
 	else
 	{
+		// These hashes are no longer used. They're only kept to maintain serialization format.
+		// Please remove them if you ever change serialization format.
 		FString Str = E.GetPlainNameString();
 		Ar << Str;
 		Ar << E.NonCasePreservingHash;
@@ -2106,8 +2195,6 @@ void operator<<(FStructuredArchive::FSlot Slot, FNameEntrySerialized& E)
 			// get the pointer to the wide array 
 			WIDECHAR* WideName = const_cast<WIDECHAR*>(E.GetWideName());
 			FCString::Strcpy(WideName, 1024, *Str);
-
-			E.bWereHashesLoaded = false;
 		}
 	}
 	else
diff --git a/Engine/Source/Runtime/Core/Public/UObject/NameTypes.h b/Engine/Source/Runtime/Core/Public/UObject/NameTypes.h
--- a/Engine/Source/Runtime/Core/Public/UObject/NameTypes.h
+++ b/Engine/Source/Runtime/Core/Public/UObject/NameTypes.h
@@ -74,6 +74,8 @@ private:
 CORE_API uint32 GetTypeHash(FNameEntryId Id);
 CORE_API bool operator==(FNameEntryId Id, EName Ename);
 inline bool operator==(EName Ename, FNameEntryId Id) { return Id == Ename; }
+inline bool operator!=(EName Ename, FNameEntryId Id) { return !(Id == Ename); }
+inline bool operator!=(FNameEntryId Id, EName Ename) { return !(Id == Ename); }
 
 /** Serialize as process specific unstable int */
 CORE_API FArchive& operator<<(FArchive& Ar, FNameEntryId& InId);
@@ -163,7 +165,7 @@ struct FNameEntryHeader
 };
 
 /**
- * A global name, as stored in the global name table.
+ * A global deduplicated name stored in the global name table.
  */
 struct FNameEntry
 {
@@ -184,49 +186,42 @@ private:
 	FNameEntry& operator=(FNameEntry&&) = delete;
 
 public:
-	/**
-	 * Returns whether this name entry is represented via WIDECHAR or ANSICHAR
-	 */
+	/** Returns whether this name entry is represented via WIDECHAR or ANSICHAR. */
 	FORCEINLINE bool IsWide() const
 	{
 		return Header.bIsWide;
 	}
-
-	/**
-	 * @return FString of name portion minus number.
-	 */
-	CORE_API FString GetPlainNameString() const;	
-
-	/**
-	 * Appends this name entry to the passed in string.
-	 *
-	 * @param	String	String to append this name to
-	 */
-	CORE_API void AppendNameToString( FString& String ) const;
-
-	/**
-	* Appends this name entry to the passed in string, adding path separator between strings (with FString operator/).
-	*
-	* @param	String	String to append this name to
-	*/
-	CORE_API void AppendNameToPathString( FString& String ) const;
-
-	int32 GetNameLength() const
+	
+	FORCEINLINE int32 GetNameLength() const
 	{
 		return Header.Len;
 	}
 
 	/**
-	 * Copy name to a buffer to avoid allocation
+	 * Copy unterminated name to TCHAR buffer without allocating.
+	 *
+	 * @param OutSize must be at least GetNameLength()
 	 */
+	void GetUnterminatedName(TCHAR* OutName, uint32 OutSize) const;
+
+	/** Copy null-terminated name to TCHAR buffer without allocating. */
+	void GetName(TCHAR(&OutName)[NAME_SIZE]) const;
+
+	/** Copy null-terminated name to ANSICHAR buffer without allocating. Entry must not be wide. */
 	CORE_API void GetAnsiName(ANSICHAR(&OutName)[NAME_SIZE]) const;
 
-	/**
-	 * Copy name to a buffer to avoid allocation
-	 */
+	/** Copy null-terminated name to WIDECHAR buffer without allocating. Entry must be wide. */
 	CORE_API void GetWideName(WIDECHAR(&OutName)[NAME_SIZE]) const;
 
-	static CORE_API int32 GetSize( const TCHAR* Name );
+	/** Copy name to a dynamically allocated FString. */
+	CORE_API FString GetPlainNameString() const;
+
+	/** Appends name to string. May allocate. */
+	CORE_API void AppendNameToString(FString& OutString) const;
+
+	/** Appends name to string with path separator using FString::PathAppend(). */
+	CORE_API void AppendNameToPathString(FString& OutString) const;
+
 
 	/**
 	 * Returns the size in bytes for FNameEntry structure. This is != sizeof(FNameEntry) as we only allocated as needed.
@@ -236,10 +231,10 @@ public:
 	 * @return	required size of FNameEntry structure to hold this string (might be wide or ansi)
 	 */
 	static int32 GetSize(int32 Length, bool bIsPureAnsi);
+	static CORE_API int32 GetSize(const TCHAR* Name);
 
 	CORE_API int32 GetSizeInBytes() const;
 
-	// Functions.
 	CORE_API void Write(FArchive& Ar) const;
 	CORE_API void Write(FStructuredArchive::FSlot Slot) const;
 
@@ -260,6 +255,7 @@ private:
 	void StoreName(const WIDECHAR* InName, uint32 Len);
 	void CopyUnterminatedName(ANSICHAR* OutName) const;
 	void CopyUnterminatedName(WIDECHAR* OutName) const;
+	void CopyAndConvertUnterminatedName(TCHAR* OutName) const;
 	const ANSICHAR* GetUnterminatedName(ANSICHAR(&OptionalDecodeBuffer)[NAME_SIZE]) const;
 	const WIDECHAR* GetUnterminatedName(WIDECHAR(&OptionalDecodeBuffer)[NAME_SIZE]) const;
 };
@@ -281,7 +277,6 @@ struct FNameEntrySerialized
 	// These are not used anymore but recalculated on save to maintain serialization format
 	uint16 NonCasePreservingHash = 0;
 	uint16 CasePreservingHash = 0;
-	bool bWereHashesLoaded = false;
 
 	FNameEntrySerialized(const FNameEntry& NameEntry);
 	FNameEntrySerialized(enum ELinkerNameTableConstructor) {}
@@ -406,17 +401,16 @@ public:
 		Number = NewNumber;
 	}
 	
-	/** Returns the pure name string without any trailing numbers */
+	/** Get name without number part as a dynamically allocated string */
 	FString GetPlainNameString() const;
 
-	/**
-	 * Returns the underlying ANSI string pointer.  No allocations.  Will fail if this is actually a wide name.
-	 */
+	/** Convert name without number part into TCHAR buffer and returns string length. Doesn't allocate. */
+	uint32 GetPlainNameString(TCHAR(&OutName)[NAME_SIZE]) const;
+
+	/** Copy ANSI name without number part. Must *only* be used for ANSI FNames. Doesn't allocate. */
 	void GetPlainANSIString(ANSICHAR(&AnsiName)[NAME_SIZE]) const;
 
-	/**
-	 * Returns the underlying WIDE string pointer.  No allocations.  Will fail if this is actually an ANSI name.
-	 */
+	/** Copy wide name without number part. Must *only* be used for wide FNames. Doesn't allocate. */
 	void GetPlainWIDEString(WIDECHAR(&WideName)[NAME_SIZE]) const;
 
 	const FNameEntry* GetComparisonNameEntry() const;
@@ -436,6 +430,31 @@ public:
 	 */
 	void ToString(FString& Out) const;
 
+	/**
+	 * Get the number of characters, excluding null-terminator, that ToString() would yield
+	 */
+	uint32 GetStringLength() const;
+
+	/**
+	 * Buffer size required for any null-terminated FName string, i.e. [name] '_' [digits] '\0'
+	 */
+	static constexpr uint32 StringBufferSize = NAME_SIZE + 1 + 10; // NAME_SIZE includes null-terminator
+
+	/**
+	 * Convert to string buffer to avoid dynamic allocations and returns string length
+	 *
+	 * Fails hard if OutLen < GetStringLength() + 1. StringBufferSize guarantees success.
+	 *
+	 * Note that a default constructed FName returns "None" instead of ""
+	 */
+	uint32 ToString(TCHAR* Out, uint32 OutSize) const;
+
+	template<int N>
+	uint32 ToString(TCHAR (&Out)[N]) const
+	{
+		return ToString(Out, N);
+	}
+
 	/**
 	 * Converts an FName to a readable format, in place, appending to an existing string (ala GetFullName)
 	 * 
@@ -462,6 +481,16 @@ public:
 		return !(*this == Other);
 	}
 
+	FORCEINLINE bool operator==(EName Ename) const
+	{
+		return (ComparisonIndex == Ename) & (GetNumber() == 0);
+	}
+	
+	FORCEINLINE bool operator!=(EName Ename) const
+	{
+		return (ComparisonIndex != Ename) | (GetNumber() != 0);
+	}
+
 	UE_DEPRECATED(4.23, "Please use FastLess() / FNameFastLess or LexicalLess() / FNameLexicalLess instead. "
 		"Default lexical sort order is deprecated to avoid unintended expensive sorting. ")
 	FORCEINLINE bool operator<( const FName& Other ) const
@@ -636,7 +665,11 @@ public:
 	{
 	}
 
+#if WITH_CASE_PRESERVING_NAME
 	static FNameEntryId GetComparisonIdFromDisplayId(FNameEntryId DisplayId);
+#else
+	static FNameEntryId GetComparisonIdFromDisplayId(FNameEntryId DisplayId) { return DisplayId; }
+#endif
 
 	/**
 	 * Only call this if you *really* know what you're doing
