Added FAsciiSet::FindLastOrEnd() and renamed Find() to FindFirstOrEnd()

diff --git a/Engine/Source/Runtime/Core/Private/Tests/Misc/AsciiSetTest.cpp b/Engine/Source/Runtime/Core/Private/Tests/Misc/AsciiSetTest.cpp
--- a/Engine/Source/Runtime/Core/Private/Tests/Misc/AsciiSetTest.cpp
+++ b/Engine/Source/Runtime/Core/Private/Tests/Misc/AsciiSetTest.cpp
@@ -24,7 +24,12 @@ bool TAsciiSetTest::RunTest(const FString& Parameters)
 	TestEqual(TEXT("Num"), WhitespaceNum, 6);
 
 	TestEqual(TEXT("Skip"), FAsciiSet::Skip(TEXT("  \t\tHello world!"), Whitespaces), TEXT("Hello world!"));
-	TestEqual(TEXT("Find"), *FAsciiSet::Find("NonWhitespace\t NonWhitespace", Whitespaces), '\t');
+	TestEqual(TEXT("Skip"), FAsciiSet::Skip(TEXT("Hello world!"), Whitespaces), TEXT("Hello world!"));
+	TestEqual(TEXT("AdvanceToFirst"),	*FAsciiSet::FindFirstOrEnd("NonWhitespace\t \nNonWhitespace", Whitespaces), '\t');
+	TestEqual(TEXT("AdvanceToLast"),	*FAsciiSet::FindLastOrEnd("NonWhitespace\t \nNonWhitespace", Whitespaces), '\n');
+	TestEqual(TEXT("AdvanceToLast"),	*FAsciiSet::FindLastOrEnd("NonWhitespace\t NonWhitespace\n", Whitespaces), '\n');
+	TestEqual(TEXT("AdvanceToFirst"),	*FAsciiSet::FindFirstOrEnd("NonWhitespaceNonWhitespace", Whitespaces), '\0');
+	TestEqual(TEXT("AdvanceToLast"),	*FAsciiSet::FindLastOrEnd("NonWhitespaceNonWhitespace", Whitespaces), '\0');
 
 	constexpr FAsciiSet XmlEscapeChars("&<>\"'");
 	TestTrue(TEXT("None"), FAsciiSet::HasNone("No escape chars", XmlEscapeChars));
diff --git a/Engine/Source/Runtime/Core/Public/Misc/AsciiSet.h b/Engine/Source/Runtime/Core/Public/Misc/AsciiSet.h
--- a/Engine/Source/Runtime/Core/Public/Misc/AsciiSet.h
+++ b/Engine/Source/Runtime/Core/Public/Misc/AsciiSet.h
@@ -12,11 +12,19 @@
  * Example use cases:
  *
  *   constexpr FAsciiSet WhitespaceCharacters(" \v\f\t\r\n");
- *   bool bIsWhitespace = WhitespaceCharacters.Contains(MyChar);
- *   const char* HelloWorld = FAsciiSet::Skip("  \t\tHello world!", WhitespaceCharacters;
- *
+ *   bool bIsWhitespace = WhitespaceCharacters.Test(MyChar);
+ *   const char* HelloWorld = FAsciiSet::Skip("  \t\tHello world!", WhitespaceCharacters);
+ *   
  *   constexpr FAsciiSet XmlEscapeChars("&<>\"'");
  *   check(FAsciiSet::HasNone(EscapedXmlString, XmlEscapeChars));
+ *   
+ *   constexpr FAsciiSet Delimiters(".:;");
+ *   const TCHAR* DelimiterOrEnd = FAsciiSet::FindFirstOrEnd(PrefixedName, Delimiters);
+ *   FString Prefix(PrefixedName, DelimiterOrEnd - PrefixedName);
+ *   
+ *   constexpr FAsciiSet Slashes("/\\");
+ *   const TCHAR* SlashOrEnd = FAsciiSet::FindLastOrEnd(PathName, Slashes);
+ *   const TCHAR* FileName = *SlashOrEnd ? SlashOrEnd + 1 : PathName;
  */
 class FAsciiSet
 {
@@ -41,7 +49,7 @@ public:
 		return FAsciiSet(Bitset);
 	}
 
-	/** Create new set containg inverse set of characters - likely including null-terminator */
+	/** Create new set containing inverse set of characters - likely including null-terminator */
 	constexpr FORCEINLINE FAsciiSet operator~() const
 	{
 		return FAsciiSet(~LoMask, ~HiMask);
@@ -49,34 +57,48 @@ public:
 
 	// C string util functions
 
-	/** Scan string for first character in a set or null-terminator */
+	/** Find first character of string inside set or end pointer. Never returns null. */
 	template<class CharType>
-	static constexpr const CharType* Find(const CharType* Str, FAsciiSet Set)
+	static constexpr const CharType* FindFirstOrEnd(const CharType* Str, FAsciiSet Set)
 	{
 		for (FAsciiSet SetOrNil(Set.LoMask | NilMask, Set.HiMask); !SetOrNil.Test(*Str); ++Str);
 
 		return Str;
 	}
 
-	/** Scan string for first character not in a set or null-terminator */
+	/** Find last character of string inside set or end pointer. Never returns null. */
+	template<class CharType>
+	static constexpr const CharType* FindLastOrEnd(const CharType* Str, FAsciiSet Set)
+	{
+		const CharType* Last = FindFirstOrEnd(Str, Set);
+
+		for (const CharType* It = Last; *It; It = FindFirstOrEnd(It + 1, Set))
+		{
+			Last = It;
+		}
+
+		return Last;
+	}
+
+	/** Find first character of string outside set or end pointer. Never returns null. */
 	template<typename CharType>
 	static constexpr const CharType* Skip(const CharType* Str, FAsciiSet Set)
 	{
-		return Find(Str, ~Set);
+		return FindFirstOrEnd(Str, ~Set);
 	}
 
 	/** Test if string contains any character in set */
 	template<typename CharType>
 	static constexpr bool HasAny(const CharType* Str, FAsciiSet Set)
 	{
-		return *Find(Str, Set) != '\0';
+		return *FindFirstOrEnd(Str, Set) != '\0';
 	}
 
 	/** Test if string contains no character in set */
 	template<typename CharType>
 	static constexpr bool HasNone(const CharType* Str, FAsciiSet Set)
 	{
-		return *Find(Str, Set) == '\0';
+		return *FindFirstOrEnd(Str, Set) == '\0';
 	}
 
 	/** Test if string contains any character outside of set */
diff --git a/Engine/Source/Runtime/CoreUObject/Private/UObject/UObjectHash.cpp b/Engine/Source/Runtime/CoreUObject/Private/UObject/UObjectHash.cpp
--- a/Engine/Source/Runtime/CoreUObject/Private/UObject/UObjectHash.cpp
+++ b/Engine/Source/Runtime/CoreUObject/Private/UObject/UObjectHash.cpp
@@ -541,17 +541,14 @@ UObject* StaticFindObjectFastInternalThreadSafe(FUObjectHashTables& ThreadHash,
 		TCHAR PlainObjectName[NAME_SIZE];
 		int32 PlainObjectNameLen = ObjectName.GetPlainNameString(PlainObjectName);
 		
-		// Reverse scan for first '.' or ':'
+		// Drop part prefixed by '.' or ':' and keep the suffix part
 		constexpr FAsciiSet DotColon(".:");
-		for (const TCHAR* PlainIt = PlainObjectName + PlainObjectNameLen - 1; PlainIt > PlainObjectName; --PlainIt)
+		const TCHAR* DelimiterOrEnd = FAsciiSet::FindLastOrEnd(PlainObjectName, DotColon);
+		if (*DelimiterOrEnd)
 		{
-			if (DotColon.Test(*PlainIt))
-			{
-				ActualObjectName = FName(PlainIt + 1, ObjectName.GetNumber());
-				int32 OuterLen = static_cast<int32>(PlainIt - PlainObjectName);
-				VerifyOuterName = FName(OuterLen, PlainObjectName);
-				break;
-			}
+			ActualObjectName = FName(DelimiterOrEnd + 1, ObjectName.GetNumber());
+			int32 OuterLen = static_cast<int32>(DelimiterOrEnd - PlainObjectName);
+			VerifyOuterName = FName(OuterLen, PlainObjectName);
 		}
 
 		const int32 Hash = GetObjectHash(ActualObjectName);
