Fix non-determinism in custom version serialization by sorting custom the custom version container by key

diff --git a/Engine/Source/Runtime/Core/Private/Serialization/CustomVersion.cpp b/Engine/Source/Runtime/Core/Private/Serialization/CustomVersion.cpp
--- a/Engine/Source/Runtime/Core/Private/Serialization/CustomVersion.cpp
+++ b/Engine/Source/Runtime/Core/Private/Serialization/CustomVersion.cpp
@@ -6,6 +6,7 @@
 
 #include "Serialization/CustomVersion.h"
 #include "Serialization/StructuredArchiveFromArchive.h"
+#include "Algo/Sort.h"
 
 namespace
 {
@@ -84,6 +85,11 @@ void FCustomVersionContainer::Empty()
 	Versions.Empty();
 }
 
+void FCustomVersionContainer::SortByKey()
+{
+	Algo::SortBy(Versions, &FCustomVersion::Key);
+}
+
 FString FCustomVersionContainer::ToString(const FString& Indent) const
 {
 	FString VersionsAsString;
diff --git a/Engine/Source/Runtime/Core/Public/Serialization/CustomVersion.h b/Engine/Source/Runtime/Core/Public/Serialization/CustomVersion.h
--- a/Engine/Source/Runtime/Core/Public/Serialization/CustomVersion.h
+++ b/Engine/Source/Runtime/Core/Public/Serialization/CustomVersion.h
@@ -144,6 +144,11 @@ public:
 	 */
 	void Empty();
 
+	/**
+	 * Sorts the custom version container by key.
+	 */
+	void SortByKey();
+
 	/** Return a string representation of custom versions. Used for debug. */
 	FString ToString(const FString& Indent) const;
 
diff --git a/Engine/Source/Runtime/CoreUObject/Private/UObject/PackageFileSummary.cpp b/Engine/Source/Runtime/CoreUObject/Private/UObject/PackageFileSummary.cpp
--- a/Engine/Source/Runtime/CoreUObject/Private/UObject/PackageFileSummary.cpp
+++ b/Engine/Source/Runtime/CoreUObject/Private/UObject/PackageFileSummary.cpp
@@ -350,3 +350,9 @@ FArchive& operator<<( FArchive& Ar, FPackageFileSummary& Sum )
 	FStructuredArchiveFromArchive(Ar).GetSlot() << Sum;
 	return Ar;
 }
+
+void FPackageFileSummary::SetCustomVersionContainer(const FCustomVersionContainer& InContainer)
+{
+	CustomVersionContainer = InContainer;
+	CustomVersionContainer.SortByKey();
+}
diff --git a/Engine/Source/Runtime/CoreUObject/Public/UObject/PackageFileSummary.h b/Engine/Source/Runtime/CoreUObject/Public/UObject/PackageFileSummary.h
--- a/Engine/Source/Runtime/CoreUObject/Public/UObject/PackageFileSummary.h
+++ b/Engine/Source/Runtime/CoreUObject/Public/UObject/PackageFileSummary.h
@@ -224,10 +224,7 @@ public:
 		return CustomVersionContainer;
 	}
 
-	void SetCustomVersionContainer(const FCustomVersionContainer& InContainer)
-	{
-		CustomVersionContainer = InContainer;
-	}
+	void SetCustomVersionContainer(const FCustomVersionContainer& InContainer);
 
 	void SetFileVersions(const int32 EpicUE4, const int32 LicenseeUE4, const bool bInSaveUnversioned = false)
 	{
