Fixed Linux compile error

diff --git a/Engine/Source/Runtime/Core/Public/CoreGlobals.h b/Engine/Source/Runtime/Core/Public/CoreGlobals.h
--- a/Engine/Source/Runtime/Core/Public/CoreGlobals.h
+++ b/Engine/Source/Runtime/Core/Public/CoreGlobals.h
@@ -69,7 +69,7 @@ struct CORE_API FScopedBootTiming
 };
 
 
-#define SCOPED_BOOT_TIMING(x) TRACE_CPUPROFILER_EVENT_SCOPE_TEXT(L ## x); FScopedBootTiming ANONYMOUS_VARIABLE(BootTiming_)(x);
+#define SCOPED_BOOT_TIMING(x) TRACE_CPUPROFILER_EVENT_SCOPE_TEXT(TEXT(x)); FScopedBootTiming ANONYMOUS_VARIABLE(BootTiming_)(x);
 
 #define GLog GetGlobalLogSingleton()
 extern CORE_API FConfigCacheIni* GConfig;
