Deprecation (later to be made a compile error) binding delegates with a non-const member function and a const object pointer, as this is const-incorrect.
A bunch of const_casts to fix warnings of this - hopefully owners will fix them as they find them.

diff --git a/Engine/Plugins/Online/OnlineFramework/Source/Party/Private/SocialToolkit.cpp b/Engine/Plugins/Online/OnlineFramework/Source/Party/Private/SocialToolkit.cpp
--- a/Engine/Plugins/Online/OnlineFramework/Source/Party/Private/SocialToolkit.cpp
+++ b/Engine/Plugins/Online/OnlineFramework/Source/Party/Private/SocialToolkit.cpp
@@ -247,7 +247,7 @@ void USocialToolkit::TrySendFriendInvite(const FString& DisplayNameOrEmail) cons
 	IOnlineUserPtr UserInterface = PrimaryOSS ? PrimaryOSS->GetUserInterface() : nullptr;
 	if (UserInterface.IsValid())
 	{
-		IOnlineUser::FOnQueryUserMappingComplete QueryCompleteDelegate = IOnlineUser::FOnQueryUserMappingComplete::CreateUObject(this, &USocialToolkit::HandleQueryPrimaryUserIdMappingComplete);
+		IOnlineUser::FOnQueryUserMappingComplete QueryCompleteDelegate = IOnlineUser::FOnQueryUserMappingComplete::CreateUObject(const_cast<USocialToolkit*>(this), &USocialToolkit::HandleQueryPrimaryUserIdMappingComplete);
 		UserInterface->QueryUserIdMapping(*GetLocalUserNetId(ESocialSubsystem::Primary), DisplayNameOrEmail, QueryCompleteDelegate);
 	}
 }
@@ -414,7 +414,7 @@ bool USocialToolkit::TrySendFriendInvite(USocialUser& SocialUser, ESocialSubsyst
 
 		if (FriendsInterface && SubsystemId.IsValid() && !bIsFriendshipRestricted)
 		{
-			return FriendsInterface->SendInvite(GetLocalUserNum(), *SubsystemId, FriendListToQuery, FOnSendInviteComplete::CreateUObject(this, &USocialToolkit::HandleFriendInviteSent, SubsystemType, SocialUser.GetDisplayName()));
+			return FriendsInterface->SendInvite(GetLocalUserNum(), *SubsystemId, FriendListToQuery, FOnSendInviteComplete::CreateUObject(const_cast<USocialToolkit*>(this), &USocialToolkit::HandleFriendInviteSent, SubsystemType, SocialUser.GetDisplayName()));
 		}
 	}
 	return false;
diff --git a/Engine/Plugins/Runtime/GameplayAbilities/Source/GameplayAbilities/Private/Abilities/Tasks/AbilityTask_WaitTargetData.cpp b/Engine/Plugins/Runtime/GameplayAbilities/Source/GameplayAbilities/Private/Abilities/Tasks/AbilityTask_WaitTargetData.cpp
--- a/Engine/Plugins/Runtime/GameplayAbilities/Source/GameplayAbilities/Private/Abilities/Tasks/AbilityTask_WaitTargetData.cpp
+++ b/Engine/Plugins/Runtime/GameplayAbilities/Source/GameplayAbilities/Private/Abilities/Tasks/AbilityTask_WaitTargetData.cpp
@@ -139,8 +139,8 @@ void UAbilityTask_WaitTargetData::InitializeTargetActor(AGameplayAbilityTargetAc
 	SpawnedActor->MasterPC = Ability->GetCurrentActorInfo()->PlayerController.Get();
 
 	// If we spawned the target actor, always register the callbacks for when the data is ready.
-	SpawnedActor->TargetDataReadyDelegate.AddUObject(this, &UAbilityTask_WaitTargetData::OnTargetDataReadyCallback);
-	SpawnedActor->CanceledDelegate.AddUObject(this, &UAbilityTask_WaitTargetData::OnTargetDataCancelledCallback);
+	SpawnedActor->TargetDataReadyDelegate.AddUObject(const_cast<UAbilityTask_WaitTargetData*>(this), &UAbilityTask_WaitTargetData::OnTargetDataReadyCallback);
+	SpawnedActor->CanceledDelegate.AddUObject(const_cast<UAbilityTask_WaitTargetData*>(this), &UAbilityTask_WaitTargetData::OnTargetDataCancelledCallback);
 }
 
 void UAbilityTask_WaitTargetData::FinalizeTargetActor(AGameplayAbilityTargetActor* SpawnedActor) const
diff --git a/Engine/Source/Developer/MessageLog/Private/UserInterface/SMessageLogListing.cpp b/Engine/Source/Developer/MessageLog/Private/UserInterface/SMessageLogListing.cpp
--- a/Engine/Source/Developer/MessageLog/Private/UserInterface/SMessageLogListing.cpp
+++ b/Engine/Source/Developer/MessageLog/Private/UserInterface/SMessageLogListing.cpp
@@ -419,7 +419,7 @@ TSharedRef<SWidget> SMessageLogListing::OnGetPageMenuContent() const
 				MessageLogListingViewModel->GetPageTitle(PageIndex), 
 				FText::Format(LOCTEXT("PageMenuEntry_Tooltip", "View page: {PageName}"), Arguments), 
 				FSlateIcon(),
-				FExecuteAction::CreateSP(this, &SMessageLogListing::OnPageSelected, PageIndex));
+				FExecuteAction::CreateSP(const_cast<SMessageLogListing*>(this), &SMessageLogListing::OnPageSelected, PageIndex));
 		}
 
 		return MenuBuilder.MakeWidget();
diff --git a/Engine/Source/Developer/SlateFileDialogs/Private/SlateFileDlgWindow.cpp b/Engine/Source/Developer/SlateFileDialogs/Private/SlateFileDlgWindow.cpp
--- a/Engine/Source/Developer/SlateFileDialogs/Private/SlateFileDlgWindow.cpp
+++ b/Engine/Source/Developer/SlateFileDialogs/Private/SlateFileDlgWindow.cpp
@@ -1020,7 +1020,7 @@ TSharedPtr<SWidget> SSlateFileOpenDlg::OnGetCrumbDelimiterContent(const FString&
 						FText::FromString(Drive),
 						FText::GetEmpty(),
 						FSlateIcon(),
-						FUIAction(FExecuteAction::CreateSP(this, &SSlateFileOpenDlg::OnPathMenuItemClicked, Drive + TEXT("/"))));
+						FUIAction(FExecuteAction::CreateSP(const_cast<SSlateFileOpenDlg*>(this), &SSlateFileOpenDlg::OnPathMenuItemClicked, Drive + TEXT("/"))));
 				}
 
 				DrivesMask >>= 1;
@@ -1051,7 +1051,7 @@ TSharedPtr<SWidget> SSlateFileOpenDlg::OnGetCrumbDelimiterContent(const FString&
 				FText::FromString(SubDir),
 				FText::GetEmpty(),
 				FSlateIcon(),
-				FUIAction(FExecuteAction::CreateSP(this, &SSlateFileOpenDlg::OnPathMenuItemClicked, CrumbData + SubDir + TEXT("/"))));
+				FUIAction(FExecuteAction::CreateSP(const_cast<SSlateFileOpenDlg*>(this), &SSlateFileOpenDlg::OnPathMenuItemClicked, CrumbData + SubDir + TEXT("/"))));
 		}
 
 		Widget = 
diff --git a/Engine/Source/Editor/BlueprintGraph/Private/K2Node_CommutativeAssociativeBinaryOperator.cpp b/Engine/Source/Editor/BlueprintGraph/Private/K2Node_CommutativeAssociativeBinaryOperator.cpp
--- a/Engine/Source/Editor/BlueprintGraph/Private/K2Node_CommutativeAssociativeBinaryOperator.cpp
+++ b/Engine/Source/Editor/BlueprintGraph/Private/K2Node_CommutativeAssociativeBinaryOperator.cpp
@@ -238,7 +238,7 @@ void UK2Node_CommutativeAssociativeBinaryOperator::GetContextMenuActions(const F
 					LOCTEXT("RemovePinTooltip", "Remove this input pin"),
 					FSlateIcon(),
 					FUIAction(
-						FExecuteAction::CreateUObject(this, &UK2Node_CommutativeAssociativeBinaryOperator::RemoveInputPin, const_cast<UEdGraphPin*>(Context.Pin))
+						FExecuteAction::CreateUObject(const_cast<UK2Node_CommutativeAssociativeBinaryOperator*>(this), &UK2Node_CommutativeAssociativeBinaryOperator::RemoveInputPin, const_cast<UEdGraphPin*>(Context.Pin))
 					)
 				);
 				Context.MenuBuilder->EndSection();
@@ -252,7 +252,7 @@ void UK2Node_CommutativeAssociativeBinaryOperator::GetContextMenuActions(const F
 				LOCTEXT("AddPinTooltip", "Add another input pin"),
 				FSlateIcon(),
 				FUIAction(
-					FExecuteAction::CreateUObject(this, &UK2Node_CommutativeAssociativeBinaryOperator::AddInputPin)
+					FExecuteAction::CreateUObject(const_cast<UK2Node_CommutativeAssociativeBinaryOperator*>(this), &UK2Node_CommutativeAssociativeBinaryOperator::AddInputPin)
 				)
 			);
 			Context.MenuBuilder->EndSection();
diff --git a/Engine/Source/Editor/BlueprintGraph/Private/K2Node_DoOnceMultiInput.cpp b/Engine/Source/Editor/BlueprintGraph/Private/K2Node_DoOnceMultiInput.cpp
--- a/Engine/Source/Editor/BlueprintGraph/Private/K2Node_DoOnceMultiInput.cpp
+++ b/Engine/Source/Editor/BlueprintGraph/Private/K2Node_DoOnceMultiInput.cpp
@@ -286,7 +286,7 @@ void UK2Node_DoOnceMultiInput::GetContextMenuActions(const FGraphNodeContextMenu
 					LOCTEXT("RemovePinTooltip", "Remove this input pin"),
 					FSlateIcon(),
 					FUIAction(
-						FExecuteAction::CreateUObject(this, &UK2Node_DoOnceMultiInput::RemoveInputPin, const_cast<UEdGraphPin*>(Context.Pin))
+						FExecuteAction::CreateUObject(const_cast<UK2Node_DoOnceMultiInput*>(this), &UK2Node_DoOnceMultiInput::RemoveInputPin, const_cast<UEdGraphPin*>(Context.Pin))
 					)
 				);
 				Context.MenuBuilder->EndSection();
@@ -300,7 +300,7 @@ void UK2Node_DoOnceMultiInput::GetContextMenuActions(const FGraphNodeContextMenu
 				LOCTEXT("AddPinTooltip", "Add another input pin"),
 				FSlateIcon(),
 				FUIAction(
-					FExecuteAction::CreateUObject(this, &UK2Node_DoOnceMultiInput::AddInputPin)
+					FExecuteAction::CreateUObject(const_cast<UK2Node_DoOnceMultiInput*>(this), &UK2Node_DoOnceMultiInput::AddInputPin)
 				)
 			);
 			Context.MenuBuilder->EndSection();
diff --git a/Engine/Source/Editor/BlueprintGraph/Private/K2Node_DynamicCast.cpp b/Engine/Source/Editor/BlueprintGraph/Private/K2Node_DynamicCast.cpp
--- a/Engine/Source/Editor/BlueprintGraph/Private/K2Node_DynamicCast.cpp
+++ b/Engine/Source/Editor/BlueprintGraph/Private/K2Node_DynamicCast.cpp
@@ -154,7 +154,7 @@ void UK2Node_DynamicCast::GetContextMenuActions(const FGraphNodeContextMenuBuild
 				MenuEntryTooltip,
 				FSlateIcon(),
 				FUIAction(
-					FExecuteAction::CreateUObject(this, &UK2Node_DynamicCast::TogglePurity),
+					FExecuteAction::CreateUObject(const_cast<UK2Node_DynamicCast*>(this), &UK2Node_DynamicCast::TogglePurity),
 					FCanExecuteAction::CreateStatic(CanExecutePurityToggle, bCanTogglePurity),
 					FIsActionChecked()
 				)
diff --git a/Engine/Source/Editor/BlueprintGraph/Private/K2Node_EaseFunction.cpp b/Engine/Source/Editor/BlueprintGraph/Private/K2Node_EaseFunction.cpp
--- a/Engine/Source/Editor/BlueprintGraph/Private/K2Node_EaseFunction.cpp
+++ b/Engine/Source/Editor/BlueprintGraph/Private/K2Node_EaseFunction.cpp
@@ -548,7 +548,7 @@ void UK2Node_EaseFunction::GetContextMenuActions(const FGraphNodeContextMenuBuil
 					LOCTEXT("AddPinTooltip", "Resets A, B and Results pins to its default wildcard state"),
 					FSlateIcon(),
 					FUIAction(
-					FExecuteAction::CreateUObject(this, &UK2Node_EaseFunction::ResetToWildcards)
+					FExecuteAction::CreateUObject(const_cast<UK2Node_EaseFunction*>(this), &UK2Node_EaseFunction::ResetToWildcards)
 					)
 				);
 			}
diff --git a/Engine/Source/Editor/BlueprintGraph/Private/K2Node_GetArrayItem.cpp b/Engine/Source/Editor/BlueprintGraph/Private/K2Node_GetArrayItem.cpp
--- a/Engine/Source/Editor/BlueprintGraph/Private/K2Node_GetArrayItem.cpp
+++ b/Engine/Source/Editor/BlueprintGraph/Private/K2Node_GetArrayItem.cpp
@@ -194,7 +194,7 @@ void UK2Node_GetArrayItem::GetContextMenuActions(const FGraphNodeContextMenuBuil
 		bReturnIsRef ? LOCTEXT("ChangeNodeToRef", "Change to return a copy") : LOCTEXT("ChangeNodeToVal", "Change to return a reference"),
 		ToggleTooltip,
 		FSlateIcon(),
-		FUIAction(FExecuteAction::CreateUObject(this, &UK2Node_GetArrayItem::ToggleReturnPin),
+		FUIAction(FExecuteAction::CreateUObject(const_cast<UK2Node_GetArrayItem*>(this), &UK2Node_GetArrayItem::ToggleReturnPin),
 			FCanExecuteAction::CreateLambda([bCannotReturnRef]()->bool { return !bCannotReturnRef; }))
 		);
 	Context.MenuBuilder->EndSection();
diff --git a/Engine/Source/Editor/BlueprintGraph/Private/K2Node_MakeArray.cpp b/Engine/Source/Editor/BlueprintGraph/Private/K2Node_MakeArray.cpp
--- a/Engine/Source/Editor/BlueprintGraph/Private/K2Node_MakeArray.cpp
+++ b/Engine/Source/Editor/BlueprintGraph/Private/K2Node_MakeArray.cpp
@@ -89,7 +89,7 @@ void UK2Node_MakeArray::GetContextMenuActions(const FGraphNodeContextMenuBuilder
 					LOCTEXT("RemovePinTooltip", "Remove this array element pin"),
 					FSlateIcon(),
 					FUIAction(
-						FExecuteAction::CreateUObject(this, &UK2Node_MakeArray::RemoveInputPin, const_cast<UEdGraphPin*>(Context.Pin))
+						FExecuteAction::CreateUObject(const_cast<UK2Node_MakeArray*>(this), &UK2Node_MakeArray::RemoveInputPin, const_cast<UEdGraphPin*>(Context.Pin))
 					)
 				);
 			}
@@ -101,7 +101,7 @@ void UK2Node_MakeArray::GetContextMenuActions(const FGraphNodeContextMenuBuilder
 				LOCTEXT("AddPinTooltip", "Add another array element pin"),
 				FSlateIcon(),
 				FUIAction(
-					FExecuteAction::CreateUObject(this, &UK2Node_MakeArray::InteractiveAddInputPin)
+					FExecuteAction::CreateUObject(const_cast<UK2Node_MakeArray*>(this), &UK2Node_MakeArray::InteractiveAddInputPin)
 				)
 			);
 		}
@@ -111,7 +111,7 @@ void UK2Node_MakeArray::GetContextMenuActions(const FGraphNodeContextMenuBuilder
 			LOCTEXT("ResetToWildcardTooltip", "Reset the node to have wildcard input/outputs. Requires no pins are connected."),
 			FSlateIcon(),
 			FUIAction(
-				FExecuteAction::CreateUObject(this, &UK2Node_MakeArray::ClearPinTypeToWildcard),
+				FExecuteAction::CreateUObject(const_cast<UK2Node_MakeArray*>(this), &UK2Node_MakeArray::ClearPinTypeToWildcard),
 				FCanExecuteAction::CreateUObject(this, &UK2Node_MakeArray::CanResetToWildcard)
 			)
 		);
diff --git a/Engine/Source/Editor/BlueprintGraph/Private/K2Node_MakeMap.cpp b/Engine/Source/Editor/BlueprintGraph/Private/K2Node_MakeMap.cpp
--- a/Engine/Source/Editor/BlueprintGraph/Private/K2Node_MakeMap.cpp
+++ b/Engine/Source/Editor/BlueprintGraph/Private/K2Node_MakeMap.cpp
@@ -161,7 +161,7 @@ void UK2Node_MakeMap::GetContextMenuActions(const FGraphNodeContextMenuBuilder&
 					LOCTEXT("RemovePinTooltip", "Remove this pin and its corresponding key/value pin"),
 					FSlateIcon(),
 					FUIAction(
-						FExecuteAction::CreateUObject(this, &UK2Node_MakeMap::RemoveInputPin, const_cast<UEdGraphPin*>(Context.Pin))
+						FExecuteAction::CreateUObject(const_cast<UK2Node_MakeMap*>(this), &UK2Node_MakeMap::RemoveInputPin, const_cast<UEdGraphPin*>(Context.Pin))
 					)
 				);
 			}
@@ -173,7 +173,7 @@ void UK2Node_MakeMap::GetContextMenuActions(const FGraphNodeContextMenuBuilder&
 				LOCTEXT("AddPinTooltip", "Add another pair of key/value pins"),
 				FSlateIcon(),
 				FUIAction(
-					FExecuteAction::CreateUObject(this, &UK2Node_MakeMap::InteractiveAddInputPin)
+					FExecuteAction::CreateUObject(const_cast<UK2Node_MakeMap*>(this), &UK2Node_MakeMap::InteractiveAddInputPin)
 				)
 			);
 		}
@@ -183,7 +183,7 @@ void UK2Node_MakeMap::GetContextMenuActions(const FGraphNodeContextMenuBuilder&
 			LOCTEXT("ResetToWildcardTooltip", "Reset the node to have wildcard input/outputs. Requires no pins are connected."),
 			FSlateIcon(),
 			FUIAction(
-				FExecuteAction::CreateUObject(this, &UK2Node_MakeMap::ClearPinTypeToWildcard),
+				FExecuteAction::CreateUObject(const_cast<UK2Node_MakeMap*>(this), &UK2Node_MakeMap::ClearPinTypeToWildcard),
 				FCanExecuteAction::CreateUObject(this, &UK2Node_MakeMap::CanResetToWildcard)
 			)
 		);
diff --git a/Engine/Source/Editor/BlueprintGraph/Private/K2Node_MakeSet.cpp b/Engine/Source/Editor/BlueprintGraph/Private/K2Node_MakeSet.cpp
--- a/Engine/Source/Editor/BlueprintGraph/Private/K2Node_MakeSet.cpp
+++ b/Engine/Source/Editor/BlueprintGraph/Private/K2Node_MakeSet.cpp
@@ -89,7 +89,7 @@ void UK2Node_MakeSet::GetContextMenuActions(const FGraphNodeContextMenuBuilder&
 					LOCTEXT("RemovePinTooltip", "Remove this set element pin"),
 					FSlateIcon(),
 					FUIAction(
-						FExecuteAction::CreateUObject(this, &UK2Node_MakeSet::RemoveInputPin, const_cast<UEdGraphPin*>(Context.Pin))
+						FExecuteAction::CreateUObject(const_cast<UK2Node_MakeSet*>(this), &UK2Node_MakeSet::RemoveInputPin, const_cast<UEdGraphPin*>(Context.Pin))
 					)
 				);
 			}
@@ -101,7 +101,7 @@ void UK2Node_MakeSet::GetContextMenuActions(const FGraphNodeContextMenuBuilder&
 				LOCTEXT("AddPinTooltip", "Add another set element pin"),
 				FSlateIcon(),
 				FUIAction(
-					FExecuteAction::CreateUObject(this, &UK2Node_MakeSet::InteractiveAddInputPin)
+					FExecuteAction::CreateUObject(const_cast<UK2Node_MakeSet*>(this), &UK2Node_MakeSet::InteractiveAddInputPin)
 				)
 			);
 		}
@@ -111,7 +111,7 @@ void UK2Node_MakeSet::GetContextMenuActions(const FGraphNodeContextMenuBuilder&
 			LOCTEXT("ResetToWildcardTooltip", "Reset the node to have wildcard input/outputs. Requires no pins are connected."),
 			FSlateIcon(),
 			FUIAction(
-				FExecuteAction::CreateUObject(this, &UK2Node_MakeSet::ClearPinTypeToWildcard),
+				FExecuteAction::CreateUObject(const_cast<UK2Node_MakeSet*>(this), &UK2Node_MakeSet::ClearPinTypeToWildcard),
 				FCanExecuteAction::CreateUObject(this, &UK2Node_MakeSet::CanResetToWildcard)
 			)
 		);
diff --git a/Engine/Source/Editor/BlueprintGraph/Private/K2Node_VariableGet.cpp b/Engine/Source/Editor/BlueprintGraph/Private/K2Node_VariableGet.cpp
--- a/Engine/Source/Editor/BlueprintGraph/Private/K2Node_VariableGet.cpp
+++ b/Engine/Source/Editor/BlueprintGraph/Private/K2Node_VariableGet.cpp
@@ -361,7 +361,7 @@ void UK2Node_VariableGet::GetContextMenuActions(const FGraphNodeContextMenuBuild
 				MenuEntryTooltip,
 				FSlateIcon(),
 				FUIAction(
-				FExecuteAction::CreateUObject(this, &UK2Node_VariableGet::TogglePurity),
+				FExecuteAction::CreateUObject(const_cast<UK2Node_VariableGet*>(this), &UK2Node_VariableGet::TogglePurity),
 				FCanExecuteAction::CreateStatic(CanExecutePurityToggle, bCanTogglePurity && !Context.bIsDebugging),
 				FIsActionChecked()
 				)
diff --git a/Engine/Source/Editor/ContentBrowser/Private/SContentBrowser.cpp b/Engine/Source/Editor/ContentBrowser/Private/SContentBrowser.cpp
--- a/Engine/Source/Editor/ContentBrowser/Private/SContentBrowser.cpp
+++ b/Engine/Source/Editor/ContentBrowser/Private/SContentBrowser.cpp
@@ -1798,7 +1798,7 @@ TSharedPtr<SWidget> SContentBrowser::OnGetCrumbDelimiterContent(const FString& C
 						FText::FromName(ChildCollection.Name),
 						FText::GetEmpty(),
 						FSlateIcon(FEditorStyle::GetStyleSetName(), ECollectionShareType::GetIconStyleName(ChildCollection.Type)),
-						FUIAction(FExecuteAction::CreateSP(this, &SContentBrowser::OnPathMenuItemClicked, ChildCollectionCrumbData))
+						FUIAction(FExecuteAction::CreateSP(const_cast<SContentBrowser*>(this), &SContentBrowser::OnPathMenuItemClicked, ChildCollectionCrumbData))
 						);
 				}
 
@@ -1842,7 +1842,7 @@ TSharedPtr<SWidget> SContentBrowser::OnGetCrumbDelimiterContent(const FString& C
 					FText::FromString(PathWithoutParent),
 					FText::GetEmpty(),
 					FSlateIcon(FEditorStyle::GetStyleSetName(), "ContentBrowser.BreadcrumbPathPickerFolder"),
-					FUIAction(FExecuteAction::CreateSP(this, &SContentBrowser::OnPathMenuItemClicked, SubPath))
+					FUIAction(FExecuteAction::CreateSP(const_cast<SContentBrowser*>(this), &SContentBrowser::OnPathMenuItemClicked, SubPath))
 					);
 			}
 
diff --git a/Engine/Source/Editor/ContentBrowser/Private/SFilterList.cpp b/Engine/Source/Editor/ContentBrowser/Private/SFilterList.cpp
--- a/Engine/Source/Editor/ContentBrowser/Private/SFilterList.cpp
+++ b/Engine/Source/Editor/ContentBrowser/Private/SFilterList.cpp
@@ -978,7 +978,7 @@ void SFilterList::CreateFiltersMenuCategory(FMenuBuilder& MenuBuilder, const TAr
 					FText::Format( LOCTEXT("FilterByTooltipPrefix", "Filter by {0}"), LabelText ),
 					FSlateIcon(),
 					FUIAction(
-						FExecuteAction::CreateSP( this, &SFilterList::FilterByTypeClicked, WeakTypeActions ),
+						FExecuteAction::CreateSP( const_cast<SFilterList*>(this), &SFilterList::FilterByTypeClicked, WeakTypeActions ),
 						FCanExecuteAction(),
 						FIsActionChecked::CreateSP(this, &SFilterList::IsAssetTypeActionsInUse, WeakTypeActions ) ),
 					NAME_None,
@@ -1002,7 +1002,7 @@ void SFilterList::CreateOtherFiltersMenuCategory(FMenuBuilder& MenuBuilder, TSha
 				FrontendFilter->GetToolTipText(),
 				FSlateIcon(FEditorStyle::GetStyleSetName(), FrontendFilter->GetIconName()),
 				FUIAction(
-				FExecuteAction::CreateSP( this, &SFilterList::FrontendFilterClicked, FrontendFilter ),
+				FExecuteAction::CreateSP( const_cast<SFilterList*>(this), &SFilterList::FrontendFilterClicked, FrontendFilter ),
 				FCanExecuteAction(),
 				FIsActionChecked::CreateSP(this, &SFilterList::IsFrontendFilterInUse, FrontendFilter ) ),
 				NAME_None,
diff --git a/Engine/Source/Editor/DetailCustomizations/Private/CameraDetails.cpp b/Engine/Source/Editor/DetailCustomizations/Private/CameraDetails.cpp
--- a/Engine/Source/Editor/DetailCustomizations/Private/CameraDetails.cpp
+++ b/Engine/Source/Editor/DetailCustomizations/Private/CameraDetails.cpp
@@ -213,7 +213,7 @@ TSharedRef<SWidget> FCameraDetails::OnGetComboContent() const
 	for (auto ItemIter = Items.CreateConstIterator(); ItemIter; ++ItemIter)
 	{
 		FText ItemText = *ItemIter;
-		FUIAction ItemAction( FExecuteAction::CreateSP( this, &FCameraDetails::CommitAspectRatioText, ItemText ) );
+		FUIAction ItemAction( FExecuteAction::CreateSP( const_cast<FCameraDetails*>(this), &FCameraDetails::CommitAspectRatioText, ItemText ) );
 		MenuBuilder.AddMenuEntry(ItemText, TAttribute<FText>(), FSlateIcon(), ItemAction);
 	}
 
diff --git a/Engine/Source/Editor/Kismet/Private/WorkflowOrientedApp/WorkflowTabManager.cpp b/Engine/Source/Editor/Kismet/Private/WorkflowOrientedApp/WorkflowTabManager.cpp
--- a/Engine/Source/Editor/Kismet/Private/WorkflowOrientedApp/WorkflowTabManager.cpp
+++ b/Engine/Source/Editor/Kismet/Private/WorkflowOrientedApp/WorkflowTabManager.cpp
@@ -199,7 +199,7 @@ TSharedRef<SWidget> FTabInfo::CreateHistoryMenu(bool bInBackHistory) const
 			{
 				MenuBuilder.AddMenuEntry(History[HistoryIdx]->GetHistoryTitle().Get(), FText(), FSlateIcon(), 
 					FUIAction(
-					FExecuteAction::CreateRaw(this, &FTabInfo::GoToHistoryIndex, HistoryIdx)
+					FExecuteAction::CreateRaw(const_cast<FTabInfo*>(this), &FTabInfo::GoToHistoryIndex, HistoryIdx)
 					), 
 					NAME_None, EUserInterfaceActionType::Button);
 			}
@@ -216,7 +216,7 @@ TSharedRef<SWidget> FTabInfo::CreateHistoryMenu(bool bInBackHistory) const
 			{
 				MenuBuilder.AddMenuEntry(History[HistoryIdx]->GetHistoryTitle().Get(), FText(), FSlateIcon(), 
 					FUIAction(
-					FExecuteAction::CreateRaw(this, &FTabInfo::GoToHistoryIndex, HistoryIdx)
+					FExecuteAction::CreateRaw(const_cast<FTabInfo*>(this), &FTabInfo::GoToHistoryIndex, HistoryIdx)
 					), 
 					NAME_None, EUserInterfaceActionType::Button);
 			}
diff --git a/Engine/Source/Editor/MovieSceneTools/Private/K2Node_GetSequenceBinding.cpp b/Engine/Source/Editor/MovieSceneTools/Private/K2Node_GetSequenceBinding.cpp
--- a/Engine/Source/Editor/MovieSceneTools/Private/K2Node_GetSequenceBinding.cpp
+++ b/Engine/Source/Editor/MovieSceneTools/Private/K2Node_GetSequenceBinding.cpp
@@ -260,7 +260,7 @@ void UK2Node_GetSequenceBinding::GetContextMenuActions(const FGraphNodeContextMe
 						AllowedClasses,
 						PropertyCustomizationHelpers::GetNewAssetFactoriesForClasses(AllowedClasses),
 						FOnShouldFilterAsset(),
-						FOnAssetSelected::CreateUObject(this, &UK2Node_GetSequenceBinding::SetSequence),
+						FOnAssetSelected::CreateUObject(const_cast<UK2Node_GetSequenceBinding*>(this), &UK2Node_GetSequenceBinding::SetSequence),
 						FSimpleDelegate());
 					
 					SubMenuBuilder.AddWidget(MenuContent, FText::GetEmpty(), false);
diff --git a/Engine/Source/Editor/MovieSceneTools/Private/K2Node_GetSequenceBindings.cpp b/Engine/Source/Editor/MovieSceneTools/Private/K2Node_GetSequenceBindings.cpp
--- a/Engine/Source/Editor/MovieSceneTools/Private/K2Node_GetSequenceBindings.cpp
+++ b/Engine/Source/Editor/MovieSceneTools/Private/K2Node_GetSequenceBindings.cpp
@@ -176,7 +176,7 @@ void UDEPRECATED_K2Node_GetSequenceBindings::GetContextMenuActions(const FGraphN
 						AllowedClasses,
 						PropertyCustomizationHelpers::GetNewAssetFactoriesForClasses(AllowedClasses),
 						FOnShouldFilterAsset(),
-						FOnAssetSelected::CreateUObject(this, &UDEPRECATED_K2Node_GetSequenceBindings::SetSequence),
+						FOnAssetSelected::CreateUObject(const_cast<UDEPRECATED_K2Node_GetSequenceBindings*>(this), &UDEPRECATED_K2Node_GetSequenceBindings::SetSequence),
 						FSimpleDelegate());
 					
 					SubMenuBuilder.AddWidget(MenuContent, FText::GetEmpty(), false);
@@ -193,7 +193,7 @@ void UDEPRECATED_K2Node_GetSequenceBindings::GetContextMenuActions(const FGraphN
 				LOCTEXT("Refresh_ToolTip", "Refresh this node's bindings"),
 				FSlateIcon(),
 				FUIAction(
-						FExecuteAction::CreateUObject(this, &UDEPRECATED_K2Node_GetSequenceBindings::ReconstructNode)
+						FExecuteAction::CreateUObject(const_cast<UDEPRECATED_K2Node_GetSequenceBindings*>(this), &UDEPRECATED_K2Node_GetSequenceBindings::ReconstructNode)
 					)
 				);
 		}
diff --git a/Engine/Source/Editor/Persona/Private/PersonaModule.cpp b/Engine/Source/Editor/Persona/Private/PersonaModule.cpp
--- a/Engine/Source/Editor/Persona/Private/PersonaModule.cpp
+++ b/Engine/Source/Editor/Persona/Private/PersonaModule.cpp
@@ -1014,7 +1014,7 @@ TSharedRef< SWidget > FPersonaModule::GenerateCreateAssetMenu(TWeakPtr<IPersonaT
 		Objects.Add(PersonaToolkit->GetSkeleton());
 	}
 
-	AnimationEditorUtils::FillCreateAssetMenu(MenuBuilder, Objects, FAnimAssetCreated::CreateRaw(this, &FPersonaModule::HandleAssetCreated), false);
+	AnimationEditorUtils::FillCreateAssetMenu(MenuBuilder, Objects, FAnimAssetCreated::CreateRaw(const_cast<FPersonaModule*>(this), &FPersonaModule::HandleAssetCreated), false);
 
 	return MenuBuilder.MakeWidget();
 }
@@ -1040,7 +1040,7 @@ void FPersonaModule::FillCreateAnimationMenu(FMenuBuilder& MenuBuilder, TWeakPtr
 			LOCTEXT("CreateAnimation_RefPose_Tooltip", "Create Animation from reference pose."),
 			FSlateIcon(),
 			FUIAction(
-				FExecuteAction::CreateStatic(&AnimationEditorUtils::ExecuteNewAnimAsset<UAnimSequenceFactory, UAnimSequence>, Objects, FString("_Sequence"), FAnimAssetCreated::CreateRaw(this, &FPersonaModule::CreateAnimation, EPoseSourceOption::ReferencePose, InWeakPersonaToolkit), false),
+				FExecuteAction::CreateStatic(&AnimationEditorUtils::ExecuteNewAnimAsset<UAnimSequenceFactory, UAnimSequence>, Objects, FString("_Sequence"), FAnimAssetCreated::CreateRaw(const_cast<FPersonaModule*>(this), &FPersonaModule::CreateAnimation, EPoseSourceOption::ReferencePose, InWeakPersonaToolkit), false),
 				FCanExecuteAction()
 				)
 			);
@@ -1050,7 +1050,7 @@ void FPersonaModule::FillCreateAnimationMenu(FMenuBuilder& MenuBuilder, TWeakPtr
 			LOCTEXT("CreateAnimation_CurrentPose_Tooltip", "Create Animation from current pose."),
 			FSlateIcon(),
 			FUIAction(
-				FExecuteAction::CreateStatic(&AnimationEditorUtils::ExecuteNewAnimAsset<UAnimSequenceFactory, UAnimSequence>, Objects, FString("_Sequence"), FAnimAssetCreated::CreateRaw(this, &FPersonaModule::CreateAnimation, EPoseSourceOption::CurrentPose, InWeakPersonaToolkit), false),
+				FExecuteAction::CreateStatic(&AnimationEditorUtils::ExecuteNewAnimAsset<UAnimSequenceFactory, UAnimSequence>, Objects, FString("_Sequence"), FAnimAssetCreated::CreateRaw(const_cast<FPersonaModule*>(this), &FPersonaModule::CreateAnimation, EPoseSourceOption::CurrentPose, InWeakPersonaToolkit), false),
 				FCanExecuteAction()
 				)
 			);
@@ -1091,7 +1091,7 @@ void FPersonaModule::FillCreateAnimationFromCurrentAnimationMenu(FMenuBuilder& M
 			LOCTEXT("CreateAnimation_CurrentAnimation_AnimData_Tooltip", "Create Animation from Animation Source Data."),
 			FSlateIcon(),
 			FUIAction(
-				FExecuteAction::CreateStatic(&AnimationEditorUtils::ExecuteNewAnimAsset<UAnimSequenceFactory, UAnimSequence>, Objects, FString("_Sequence"), FAnimAssetCreated::CreateRaw(this, &FPersonaModule::CreateAnimation, EPoseSourceOption::CurrentAnimation_AnimData, InWeakPersonaToolkit), false)
+				FExecuteAction::CreateStatic(&AnimationEditorUtils::ExecuteNewAnimAsset<UAnimSequenceFactory, UAnimSequence>, Objects, FString("_Sequence"), FAnimAssetCreated::CreateRaw(const_cast<FPersonaModule*>(this), &FPersonaModule::CreateAnimation, EPoseSourceOption::CurrentAnimation_AnimData, InWeakPersonaToolkit), false)
 			)
 		);
 
@@ -1100,7 +1100,7 @@ void FPersonaModule::FillCreateAnimationFromCurrentAnimationMenu(FMenuBuilder& M
 			LOCTEXT("CreateAnimation_CurrentAnimation_PreviewMesh_Tooltip", "Create Animation by playing on the Current Preview Mesh, including Retargeting, Post Process Graph, or anything you see on the preview mesh."),
 			FSlateIcon(),
 			FUIAction(
-				FExecuteAction::CreateStatic(&AnimationEditorUtils::ExecuteNewAnimAsset<UAnimSequenceFactory, UAnimSequence>, Objects, FString("_Sequence"), FAnimAssetCreated::CreateRaw(this, &FPersonaModule::CreateAnimation, EPoseSourceOption::CurrentAnimation_PreviewMesh, InWeakPersonaToolkit), false)
+				FExecuteAction::CreateStatic(&AnimationEditorUtils::ExecuteNewAnimAsset<UAnimSequenceFactory, UAnimSequence>, Objects, FString("_Sequence"), FAnimAssetCreated::CreateRaw(const_cast<FPersonaModule*>(this), &FPersonaModule::CreateAnimation, EPoseSourceOption::CurrentAnimation_PreviewMesh, InWeakPersonaToolkit), false)
 			)
 		);
 	}
@@ -1129,7 +1129,7 @@ void FPersonaModule::FillCreatePoseAssetMenu(FMenuBuilder& MenuBuilder, TWeakPtr
 			LOCTEXT("CreatePoseAsset_CurrentPose_Tooltip", "Create PoseAsset from current pose."),
 			FSlateIcon(),
 			FUIAction(
-				FExecuteAction::CreateStatic(&AnimationEditorUtils::ExecuteNewAnimAsset<UPoseAssetFactory, UPoseAsset>, Objects, FString("_PoseAsset"), FAnimAssetCreated::CreateRaw(this, &FPersonaModule::CreatePoseAsset, EPoseSourceOption::CurrentPose, InWeakPersonaToolkit), false),
+				FExecuteAction::CreateStatic(&AnimationEditorUtils::ExecuteNewAnimAsset<UPoseAssetFactory, UPoseAsset>, Objects, FString("_PoseAsset"), FAnimAssetCreated::CreateRaw(const_cast<FPersonaModule*>(this), &FPersonaModule::CreatePoseAsset, EPoseSourceOption::CurrentPose, InWeakPersonaToolkit), false),
 				FCanExecuteAction()
 			)
 		);
@@ -1139,7 +1139,7 @@ void FPersonaModule::FillCreatePoseAssetMenu(FMenuBuilder& MenuBuilder, TWeakPtr
 			LOCTEXT("CreatePoseAsset_CurrentAnimation_Tooltip", "Create Animation from current animation."),
 			FSlateIcon(),
 			FUIAction(
-				FExecuteAction::CreateStatic(&AnimationEditorUtils::ExecuteNewAnimAsset<UPoseAssetFactory, UPoseAsset>, Objects, FString("_PoseAsset"), FAnimAssetCreated::CreateRaw(this, &FPersonaModule::CreatePoseAsset, EPoseSourceOption::CurrentAnimation_AnimData, InWeakPersonaToolkit), false),
+				FExecuteAction::CreateStatic(&AnimationEditorUtils::ExecuteNewAnimAsset<UPoseAssetFactory, UPoseAsset>, Objects, FString("_PoseAsset"), FAnimAssetCreated::CreateRaw(const_cast<FPersonaModule*>(this), &FPersonaModule::CreatePoseAsset, EPoseSourceOption::CurrentAnimation_AnimData, InWeakPersonaToolkit), false),
 				FCanExecuteAction()
 			)
 		);
@@ -1174,7 +1174,7 @@ void FPersonaModule::FillInsertPoseMenu(FMenuBuilder& MenuBuilder, TWeakPtr<IPer
 	AssetPickerConfig.Filter.TagsAndValues.Add(TEXT("Skeleton"), FAssetData(Skeleton).GetExportTextName());
 
 	/** The delegate that fires when an asset was selected */
-	AssetPickerConfig.OnAssetSelected = FOnAssetSelected::CreateRaw(this, &FPersonaModule::InsertCurrentPoseToAsset, InWeakPersonaToolkit);
+	AssetPickerConfig.OnAssetSelected = FOnAssetSelected::CreateRaw(const_cast<FPersonaModule*>(this), &FPersonaModule::InsertCurrentPoseToAsset, InWeakPersonaToolkit);
 
 	/** The default view mode should be a list view */
 	AssetPickerConfig.InitialAssetViewType = EAssetViewType::List;
diff --git a/Engine/Source/Editor/Persona/Private/SAnimCurvePanel.cpp b/Engine/Source/Editor/Persona/Private/SAnimCurvePanel.cpp
--- a/Engine/Source/Editor/Persona/Private/SAnimCurvePanel.cpp
+++ b/Engine/Source/Editor/Persona/Private/SAnimCurvePanel.cpp
@@ -1282,14 +1282,14 @@ TSharedRef<SWidget> SAnimCurvePanel::CreateCurveContextMenu(FAnimCurveBaseInterf
 			TypeToggleToolTip = LOCTEXT("TypeToggleToVariableToolTip", "Turns this curve into a variable curve.");
 		}
 
-		NewAction.ExecuteAction.BindSP(this, &SAnimCurvePanel::ToggleCurveTypeMenuCallback, Curve);
+		NewAction.ExecuteAction.BindSP(const_cast<SAnimCurvePanel*>(this), &SAnimCurvePanel::ToggleCurveTypeMenuCallback, Curve);
 		MenuBuilder.AddMenuEntry(
 			TypeToggleLabel,
 			TypeToggleToolTip,
 			FSlateIcon(),
 			NewAction);
 
-		NewAction.ExecuteAction.BindSP(this, &SAnimCurvePanel::DeleteTrack, Curve->CurveUID);
+		NewAction.ExecuteAction.BindSP(const_cast<SAnimCurvePanel*>(this), &SAnimCurvePanel::DeleteTrack, Curve->CurveUID);
 		MenuBuilder.AddMenuEntry(
 			LOCTEXT("RemoveTrack", "Remove Track"),
 			LOCTEXT("RemoveTrackTooltip", "Remove this track"),
diff --git a/Engine/Source/Editor/Persona/Private/SAnimTrackCurvePanel.cpp b/Engine/Source/Editor/Persona/Private/SAnimTrackCurvePanel.cpp
--- a/Engine/Source/Editor/Persona/Private/SAnimTrackCurvePanel.cpp
+++ b/Engine/Source/Editor/Persona/Private/SAnimTrackCurvePanel.cpp
@@ -909,7 +909,7 @@ TSharedRef<SWidget> SAnimTrackCurvePanel::CreateCurveContextMenu(USkeleton::Anim
 		{
 			FUIAction NewAction;
 
-			NewAction.ExecuteAction.BindSP(this, &SAnimTrackCurvePanel::DeleteTrack, CurveUid);
+			NewAction.ExecuteAction.BindSP(const_cast<SAnimTrackCurvePanel*>(this), &SAnimTrackCurvePanel::DeleteTrack, CurveUid);
 			MenuBuilder.AddMenuEntry(
 				LOCTEXT("RemoveTrack", "Remove Track"),
 				LOCTEXT("RemoveTrackTooltip", "Remove this track"),
diff --git a/Engine/Source/Editor/Persona/Private/SAnimViewportToolBar.cpp b/Engine/Source/Editor/Persona/Private/SAnimViewportToolBar.cpp
--- a/Engine/Source/Editor/Persona/Private/SAnimViewportToolBar.cpp
+++ b/Engine/Source/Editor/Persona/Private/SAnimViewportToolBar.cpp
@@ -690,7 +690,7 @@ TSharedRef<SWidget> SAnimViewportToolBar::GenerateCharacterMenu() const
 				InMenuBuilder.AddSubMenu(
 					LOCTEXT("CharacterMenu_ClothingSubMenu", "Clothing"),
 					LOCTEXT("CharacterMenu_ClothingSubMenuToolTip", "Options relating to clothing"),
-					FNewMenuDelegate::CreateRaw(this, &SAnimViewportToolBar::FillCharacterClothingMenu));
+					FNewMenuDelegate::CreateRaw(const_cast<SAnimViewportToolBar*>(this), &SAnimViewportToolBar::FillCharacterClothingMenu));
 			}
 #endif // #if WITH_APEX_CLOTHING
 		}
@@ -711,7 +711,7 @@ TSharedRef<SWidget> SAnimViewportToolBar::GenerateCharacterMenu() const
 		InMenuBuilder.AddSubMenu(
 			LOCTEXT("CharacterMenu_AdvancedSubMenu", "Advanced"),
 			LOCTEXT("CharacterMenu_AdvancedSubMenuToolTip", "Advanced options"),
-			FNewMenuDelegate::CreateRaw(this, &SAnimViewportToolBar::FillCharacterAdvancedMenu));
+			FNewMenuDelegate::CreateRaw(const_cast<SAnimViewportToolBar*>(this), &SAnimViewportToolBar::FillCharacterAdvancedMenu));
 
 		InMenuBuilder.EndSection();
 	}
diff --git a/Engine/Source/Editor/Persona/Private/SAnimationSequenceBrowser.cpp b/Engine/Source/Editor/Persona/Private/SAnimationSequenceBrowser.cpp
--- a/Engine/Source/Editor/Persona/Private/SAnimationSequenceBrowser.cpp
+++ b/Engine/Source/Editor/Persona/Private/SAnimationSequenceBrowser.cpp
@@ -833,7 +833,7 @@ TSharedRef<SWidget> SAnimationSequenceBrowser::CreateHistoryMenu(bool bInBackHis
 
 				MenuBuilder.AddMenuEntry(DisplayName, Tooltip, FSlateIcon(), 
 					FUIAction(
-					FExecuteAction::CreateRaw(this, &SAnimationSequenceBrowser::GoToHistoryIndex, HistoryIdx)
+					FExecuteAction::CreateRaw(const_cast<SAnimationSequenceBrowser*>(this), &SAnimationSequenceBrowser::GoToHistoryIndex, HistoryIdx)
 					), 
 					NAME_None, EUserInterfaceActionType::Button);
 			}
diff --git a/Engine/Source/Editor/UnrealEd/Private/SCommonEditorViewportToolbarBase.cpp b/Engine/Source/Editor/UnrealEd/Private/SCommonEditorViewportToolbarBase.cpp
--- a/Engine/Source/Editor/UnrealEd/Private/SCommonEditorViewportToolbarBase.cpp
+++ b/Engine/Source/Editor/UnrealEd/Private/SCommonEditorViewportToolbarBase.cpp
@@ -429,7 +429,7 @@ TSharedPtr<FExtender> SCommonEditorViewportToolbarBase::GetViewMenuExtender() co
 		TEXT("ViewMode"),
 		EExtensionHook::After,
 		GetInfoProvider().GetViewportWidget()->GetCommandList(),
-		FMenuExtensionDelegate::CreateSP(this, &SCommonEditorViewportToolbarBase::CreateViewMenuExtensions));
+		FMenuExtensionDelegate::CreateSP(const_cast<SCommonEditorViewportToolbarBase*>(this), &SCommonEditorViewportToolbarBase::CreateViewMenuExtensions));
 
 	return GetCombinedExtenderList(ViewModeExtender);
 }
diff --git a/Engine/Source/Runtime/AIModule/Private/BehaviorTree/BehaviorTreeComponent.cpp b/Engine/Source/Runtime/AIModule/Private/BehaviorTree/BehaviorTreeComponent.cpp
--- a/Engine/Source/Runtime/AIModule/Private/BehaviorTree/BehaviorTreeComponent.cpp
+++ b/Engine/Source/Runtime/AIModule/Private/BehaviorTree/BehaviorTreeComponent.cpp
@@ -1784,7 +1784,7 @@ void UBehaviorTreeComponent::RegisterMessageObserver(const UBTTaskNode* TaskNode
 		NodeIdx.InstanceIndex = InstanceStack.Num() - 1;
 
 		TaskMessageObservers.Add(NodeIdx,
-			FAIMessageObserver::Create(this, MessageType, FOnAIMessage::CreateUObject(TaskNode, &UBTTaskNode::ReceivedMessage))
+			FAIMessageObserver::Create(this, MessageType, FOnAIMessage::CreateUObject(const_cast<UBTTaskNode*>(TaskNode), &UBTTaskNode::ReceivedMessage))
 			);
 
 		UE_VLOG(GetOwner(), LogBehaviorTree, Log, TEXT("Message[%s] observer added for %s"),
@@ -1801,7 +1801,7 @@ void UBehaviorTreeComponent::RegisterMessageObserver(const UBTTaskNode* TaskNode
 		NodeIdx.InstanceIndex = InstanceStack.Num() - 1;
 
 		TaskMessageObservers.Add(NodeIdx,
-			FAIMessageObserver::Create(this, MessageType, RequestID, FOnAIMessage::CreateUObject(TaskNode, &UBTTaskNode::ReceivedMessage))
+			FAIMessageObserver::Create(this, MessageType, RequestID, FOnAIMessage::CreateUObject(const_cast<UBTTaskNode*>(TaskNode), &UBTTaskNode::ReceivedMessage))
 			);
 
 		UE_VLOG(GetOwner(), LogBehaviorTree, Log, TEXT("Message[%s:%d] observer added for %s"),
diff --git a/Engine/Source/Runtime/AIModule/Private/EnvironmentQuery/Tests/EnvQueryTest_Overlap.cpp b/Engine/Source/Runtime/AIModule/Private/EnvironmentQuery/Tests/EnvQueryTest_Overlap.cpp
--- a/Engine/Source/Runtime/AIModule/Private/EnvironmentQuery/Tests/EnvQueryTest_Overlap.cpp
+++ b/Engine/Source/Runtime/AIModule/Private/EnvironmentQuery/Tests/EnvQueryTest_Overlap.cpp
@@ -50,11 +50,11 @@ void UEnvQueryTest_Overlap::RunTest(FEnvQueryInstance& QueryInstance) const
 	FRunOverlapSignature OverlapFunc;
 	if (OverlapData.bOnlyBlockingHits)
 	{
-		OverlapFunc.BindUObject(this, &UEnvQueryTest_Overlap::RunOverlapBlocking);
+		OverlapFunc.BindUObject(const_cast<UEnvQueryTest_Overlap*>(this), &UEnvQueryTest_Overlap::RunOverlapBlocking);
 	}
 	else
 	{
-		OverlapFunc.BindUObject(this, &UEnvQueryTest_Overlap::RunOverlap);
+		OverlapFunc.BindUObject(const_cast<UEnvQueryTest_Overlap*>(this), &UEnvQueryTest_Overlap::RunOverlap);
 	}
 
 	for (FEnvQueryInstance::ItemIterator It(this, QueryInstance); It; ++It)
diff --git a/Engine/Source/Runtime/AIModule/Private/EnvironmentQuery/Tests/EnvQueryTest_Trace.cpp b/Engine/Source/Runtime/AIModule/Private/EnvironmentQuery/Tests/EnvQueryTest_Trace.cpp
--- a/Engine/Source/Runtime/AIModule/Private/EnvironmentQuery/Tests/EnvQueryTest_Trace.cpp
+++ b/Engine/Source/Runtime/AIModule/Private/EnvironmentQuery/Tests/EnvQueryTest_Trace.cpp
@@ -53,19 +53,19 @@ void UEnvQueryTest_Trace::RunTest(FEnvQueryInstance& QueryInstance) const
 	switch (TraceData.TraceShape)
 	{
 	case EEnvTraceShape::Line:
-		TraceFunc.BindUObject(this, bTraceToItem ? &UEnvQueryTest_Trace::RunLineTraceTo : &UEnvQueryTest_Trace::RunLineTraceFrom);
+		TraceFunc.BindUObject(const_cast<UEnvQueryTest_Trace*>(this), bTraceToItem ? &UEnvQueryTest_Trace::RunLineTraceTo : &UEnvQueryTest_Trace::RunLineTraceFrom);
 		break;
 
 	case EEnvTraceShape::Box:
-		TraceFunc.BindUObject(this, bTraceToItem ? &UEnvQueryTest_Trace::RunBoxTraceTo : &UEnvQueryTest_Trace::RunBoxTraceFrom);
+		TraceFunc.BindUObject(const_cast<UEnvQueryTest_Trace*>(this), bTraceToItem ? &UEnvQueryTest_Trace::RunBoxTraceTo : &UEnvQueryTest_Trace::RunBoxTraceFrom);
 		break;
 
 	case EEnvTraceShape::Sphere:
-		TraceFunc.BindUObject(this, bTraceToItem ? &UEnvQueryTest_Trace::RunSphereTraceTo : &UEnvQueryTest_Trace::RunSphereTraceFrom);
+		TraceFunc.BindUObject(const_cast<UEnvQueryTest_Trace*>(this), bTraceToItem ? &UEnvQueryTest_Trace::RunSphereTraceTo : &UEnvQueryTest_Trace::RunSphereTraceFrom);
 		break;
 
 	case EEnvTraceShape::Capsule:
-		TraceFunc.BindUObject(this, bTraceToItem ? &UEnvQueryTest_Trace::RunCapsuleTraceTo : &UEnvQueryTest_Trace::RunCapsuleTraceFrom);
+		TraceFunc.BindUObject(const_cast<UEnvQueryTest_Trace*>(this), bTraceToItem ? &UEnvQueryTest_Trace::RunCapsuleTraceTo : &UEnvQueryTest_Trace::RunCapsuleTraceFrom);
 		break;
 
 	default:
diff --git a/Engine/Source/Runtime/AIModule/Private/Perception/AIPerceptionComponent.cpp b/Engine/Source/Runtime/AIModule/Private/Perception/AIPerceptionComponent.cpp
--- a/Engine/Source/Runtime/AIModule/Private/Perception/AIPerceptionComponent.cpp
+++ b/Engine/Source/Runtime/AIModule/Private/Perception/AIPerceptionComponent.cpp
@@ -290,7 +290,7 @@ void UAIPerceptionComponent::GetHostileActors(TArray<AActor*>& OutActors) const
 	if (bDeadDataFound)
 	{
 		FSimpleDelegateGraphTask::CreateAndDispatchWhenReady(
-			FSimpleDelegateGraphTask::FDelegate::CreateUObject(this, &UAIPerceptionComponent::RemoveDeadData),
+			FSimpleDelegateGraphTask::FDelegate::CreateUObject(const_cast<UAIPerceptionComponent*>(this), &UAIPerceptionComponent::RemoveDeadData),
 			GET_STATID(STAT_FSimpleDelegateGraphTask_RequestingRemovalOfDeadPerceptionData), NULL, ENamedThreads::GameThread);
 	}
 }
@@ -329,7 +329,7 @@ const FActorPerceptionInfo* UAIPerceptionComponent::GetFreshestTrace(const FAISe
 	if (bDeadDataFound)
 	{
 		FSimpleDelegateGraphTask::CreateAndDispatchWhenReady(
-			FSimpleDelegateGraphTask::FDelegate::CreateUObject(this, &UAIPerceptionComponent::RemoveDeadData),
+			FSimpleDelegateGraphTask::FDelegate::CreateUObject(const_cast<UAIPerceptionComponent*>(this), &UAIPerceptionComponent::RemoveDeadData),
 			GET_STATID(STAT_FSimpleDelegateGraphTask_RequestingRemovalOfDeadPerceptionData), NULL, ENamedThreads::GameThread);
 	}
 
diff --git a/Engine/Source/Runtime/AppFramework/Private/Widgets/Colors/SColorPicker.cpp b/Engine/Source/Runtime/AppFramework/Private/Widgets/Colors/SColorPicker.cpp
--- a/Engine/Source/Runtime/AppFramework/Private/Widgets/Colors/SColorPicker.cpp
+++ b/Engine/Source/Runtime/AppFramework/Private/Widgets/Colors/SColorPicker.cpp
@@ -892,9 +892,9 @@ TSharedRef<SWidget> SColorPicker::MakeColorSlider( EColorPickerChannels Channel
 				.SliderBarColor(FLinearColor::Transparent)
 				.Style(&FCoreStyle::Get().GetWidgetStyle<FSliderStyle>("ColorPicker.Slider"))
 				.Value(this, &SColorPicker::HandleColorSpinBoxValue, Channel)
-				.OnMouseCaptureBegin(this, &SColorPicker::HandleInteractiveChangeBegin)
-				.OnMouseCaptureEnd(this, &SColorPicker::HandleInteractiveChangeEnd)
-				.OnValueChanged(this, &SColorPicker::HandleColorSpinBoxValueChanged, Channel)
+				.OnMouseCaptureBegin(const_cast<SColorPicker*>(this), &SColorPicker::HandleInteractiveChangeBegin)
+				.OnMouseCaptureEnd(const_cast<SColorPicker*>(this), &SColorPicker::HandleInteractiveChangeEnd)
+				.OnValueChanged(const_cast<SColorPicker*>(this), &SColorPicker::HandleColorSpinBoxValueChanged, Channel)
 		];
 }
 
@@ -1019,9 +1019,9 @@ TSharedRef<SWidget> SColorPicker::MakeColorSpinBox( EColorPickerChannels Channel
 						.Delta(Channel == EColorPickerChannels::Hue ? 1.0f : 0.001f)
 						.Font(SmallLayoutFont)
 						.Value(this, &SColorPicker::HandleColorSpinBoxValue, Channel)
-						.OnBeginSliderMovement(this, &SColorPicker::HandleInteractiveChangeBegin)
-						.OnEndSliderMovement(this, &SColorPicker::HandleInteractiveChangeEnd)
-						.OnValueChanged(this, &SColorPicker::HandleColorSpinBoxValueChanged, Channel)
+						.OnBeginSliderMovement(const_cast<SColorPicker*>(this), &SColorPicker::HandleInteractiveChangeBegin)
+						.OnEndSliderMovement(const_cast<SColorPicker*>(this), &SColorPicker::HandleInteractiveChangeEnd)
+						.OnValueChanged(const_cast<SColorPicker*>(this), &SColorPicker::HandleColorSpinBoxValueChanged, Channel)
 				]
 
 			+ SVerticalBox::Slot()
@@ -1065,7 +1065,7 @@ TSharedRef<SWidget> SColorPicker::MakeColorPreviewBox() const
 								.IgnoreAlpha(true)
 								.ToolTipText(LOCTEXT("OldColorToolTip", "Old color without alpha (drag to theme bar to save)"))
 								.Color(OldColor) 
-								.OnMouseButtonDown(this, &SColorPicker::HandleOldColorBlockMouseButtonDown, false)
+								.OnMouseButtonDown(const_cast<SColorPicker*>(this), &SColorPicker::HandleOldColorBlockMouseButtonDown, false)
 								.UseSRGB(SharedThis(this), &SColorPicker::HandleColorPickerUseSRGB)
 								.Cursor(EMouseCursor::GrabHand)
 						]
@@ -1079,7 +1079,7 @@ TSharedRef<SWidget> SColorPicker::MakeColorPreviewBox() const
 								.ToolTipText(LOCTEXT("OldColorAlphaToolTip", "Old color with alpha (drag to theme bar to save)"))
 								.Color(OldColor) 
 								.Visibility(SharedThis(this), &SColorPicker::HandleAlphaColorBlockVisibility)
-								.OnMouseButtonDown(this, &SColorPicker::HandleOldColorBlockMouseButtonDown, true)
+								.OnMouseButtonDown(const_cast<SColorPicker*>(this), &SColorPicker::HandleOldColorBlockMouseButtonDown, true)
 								.UseSRGB(SharedThis(this), &SColorPicker::HandleColorPickerUseSRGB)
 								.Cursor(EMouseCursor::GrabHand)
 						]
@@ -1097,7 +1097,7 @@ TSharedRef<SWidget> SColorPicker::MakeColorPreviewBox() const
 								.IgnoreAlpha(true)
 								.ToolTipText(LOCTEXT("NewColorToolTip", "New color without alpha (drag to theme bar to save)"))
 								.Color(this, &SColorPicker::GetCurrentColor)
-								.OnMouseButtonDown(this, &SColorPicker::HandleNewColorBlockMouseButtonDown, false)
+								.OnMouseButtonDown(const_cast<SColorPicker*>(this), &SColorPicker::HandleNewColorBlockMouseButtonDown, false)
 								.UseSRGB(SharedThis(this), &SColorPicker::HandleColorPickerUseSRGB)
 								.Cursor(EMouseCursor::GrabHand)
 						]
@@ -1111,7 +1111,7 @@ TSharedRef<SWidget> SColorPicker::MakeColorPreviewBox() const
 								.ToolTipText(LOCTEXT("NewColorAlphaToolTip", "New color with alpha (drag to theme bar to save)"))
 								.Color(this, &SColorPicker::GetCurrentColor)
 								.Visibility(SharedThis(this), &SColorPicker::HandleAlphaColorBlockVisibility)
-								.OnMouseButtonDown(this, &SColorPicker::HandleNewColorBlockMouseButtonDown, true)
+								.OnMouseButtonDown(const_cast<SColorPicker*>(this), &SColorPicker::HandleNewColorBlockMouseButtonDown, true)
 								.UseSRGB(SharedThis(this), &SColorPicker::HandleColorPickerUseSRGB)
 								.Cursor(EMouseCursor::GrabHand)
 						]
diff --git a/Engine/Source/Runtime/AppFramework/Private/Widgets/Testing/STestSuite.cpp b/Engine/Source/Runtime/AppFramework/Private/Widgets/Testing/STestSuite.cpp
--- a/Engine/Source/Runtime/AppFramework/Private/Widgets/Testing/STestSuite.cpp
+++ b/Engine/Source/Runtime/AppFramework/Private/Widgets/Testing/STestSuite.cpp
@@ -2275,7 +2275,7 @@ public:
 
 	FSlateWidgetRun::FWidgetRunInfo OnCreateWidgetDecoratorWidget( const FTextRunInfo& RunInfo, const ISlateStyle* Style ) const
 	{
-		TSharedRef< SWidget > Widget = SNew( SButton ) .OnClicked( this, &SRichTextTest::OnWidgetDecoratorClicked ) 
+		TSharedRef< SWidget > Widget = SNew( SButton ) .OnClicked( const_cast<SRichTextTest*>(this), &SRichTextTest::OnWidgetDecoratorClicked ) 
 			.ToolTip( 
 				SNew( SToolTip )
 				.BorderImage( FTestStyle::Get().GetBrush( "RichText.Tagline.Background" ) )
diff --git a/Engine/Source/Runtime/Core/Public/Delegates/DelegateSignatureImpl.inl b/Engine/Source/Runtime/Core/Public/Delegates/DelegateSignatureImpl.inl
--- a/Engine/Source/Runtime/Core/Public/Delegates/DelegateSignatureImpl.inl
+++ b/Engine/Source/Runtime/Core/Public/Delegates/DelegateSignatureImpl.inl
@@ -8,14 +8,15 @@
 #pragma once
 #include "CoreTypes.h"
 #include "Misc/AssertionMacros.h"
-#include "Templates/TypeWrapper.h"
-#include "Templates/RemoveReference.h"
-#include "Templates/UnrealTemplate.h"
 #include "Misc/Crc.h"
 #include "UObject/NameTypes.h"
+#include "UObject/ScriptDelegates.h"
 #include "UObject/WeakObjectPtrTemplates.h"
+#include "Templates/IsConst.h"
+#include "Templates/RemoveReference.h"
 #include "Templates/SharedPointer.h"
-#include "UObject/ScriptDelegates.h"
+#include "Templates/TypeWrapper.h"
+#include "Templates/UnrealTemplate.h"
 
 class FDelegateBase;
 class FDelegateHandle;
@@ -158,6 +159,8 @@ public:
 	inline static TBaseDelegate<RetValType, ParamTypes...> CreateRaw(UserClass* InUserObject, typename TMemFunPtrType<false, UserClass, RetValType (ParamTypes..., VarTypes...)>::Type InFunc, VarTypes... Vars)
 	FUNCTION_CHECK_RETURN_END
 	{
+		UE_STATIC_DEPRECATE(4.23, TIsConst<UserClass>::Value, "Binding a delegate with a const object pointer and non-const function is deprecated.");
+
 		TBaseDelegate<RetValType, ParamTypes...> Result;
 		TBaseRawMethodDelegateInstance<false, UserClass, TFuncType, VarTypes...>::Create(Result, InUserObject, InFunc, Vars...);
 		return Result;
@@ -183,6 +186,8 @@ public:
 	inline static TBaseDelegate<RetValType, ParamTypes...> CreateSP(const TSharedRef<UserClass, ESPMode::Fast>& InUserObjectRef, typename TMemFunPtrType<false, UserClass, RetValType (ParamTypes..., VarTypes...)>::Type InFunc, VarTypes... Vars)
 	FUNCTION_CHECK_RETURN_END
 	{
+		UE_STATIC_DEPRECATE(4.23, TIsConst<UserClass>::Value, "Binding a delegate with a const object pointer and non-const function is deprecated.");
+
 		TBaseDelegate<RetValType, ParamTypes...> Result;
 		TBaseSPMethodDelegateInstance<false, UserClass, ESPMode::Fast, TFuncType, VarTypes...>::Create(Result, InUserObjectRef, InFunc, Vars...);
 		return Result;
@@ -208,6 +213,8 @@ public:
 	inline static TBaseDelegate<RetValType, ParamTypes...> CreateSP(UserClass* InUserObject, typename TMemFunPtrType<false, UserClass, RetValType (ParamTypes..., VarTypes...)>::Type InFunc, VarTypes... Vars)
 	FUNCTION_CHECK_RETURN_END
 	{
+		UE_STATIC_DEPRECATE(4.23, TIsConst<UserClass>::Value, "Binding a delegate with a const object pointer and non-const function is deprecated.");
+
 		return CreateSP(StaticCastSharedRef<UserClass>(InUserObject->AsShared()), InFunc, Vars...);
 	}
 	template <typename UserClass, typename... VarTypes>
@@ -229,6 +236,8 @@ public:
 	inline static TBaseDelegate<RetValType, ParamTypes...> CreateThreadSafeSP(const TSharedRef<UserClass, ESPMode::ThreadSafe>& InUserObjectRef, typename TMemFunPtrType<false, UserClass, RetValType (ParamTypes..., VarTypes...)>::Type InFunc, VarTypes... Vars)
 	FUNCTION_CHECK_RETURN_END
 	{
+		UE_STATIC_DEPRECATE(4.23, TIsConst<UserClass>::Value, "Binding a delegate with a const object pointer and non-const function is deprecated.");
+
 		TBaseDelegate<RetValType, ParamTypes...> Result;
 		TBaseSPMethodDelegateInstance<false, UserClass, ESPMode::ThreadSafe, TFuncType, VarTypes...>::Create(Result, InUserObjectRef, InFunc, Vars...);
 		return Result;
@@ -254,6 +263,8 @@ public:
 	inline static TBaseDelegate<RetValType, ParamTypes...> CreateThreadSafeSP(UserClass* InUserObject, typename TMemFunPtrType<false, UserClass, RetValType (ParamTypes..., VarTypes...)>::Type InFunc, VarTypes... Vars)
 	FUNCTION_CHECK_RETURN_END
 	{
+		UE_STATIC_DEPRECATE(4.23, TIsConst<UserClass>::Value, "Binding a delegate with a const object pointer and non-const function is deprecated.");
+
 		return CreateThreadSafeSP(StaticCastSharedRef<UserClass>(InUserObject->AsShared()), InFunc, Vars...);
 	}
 	template <typename UserClass, typename... VarTypes>
@@ -291,6 +302,8 @@ public:
 	inline static TBaseDelegate<RetValType, ParamTypes...> CreateUObject(UserClass* InUserObject, typename TMemFunPtrType<false, UserClass, RetValType (ParamTypes..., VarTypes...)>::Type InFunc, VarTypes... Vars)
 	FUNCTION_CHECK_RETURN_END
 	{
+		UE_STATIC_DEPRECATE(4.23, TIsConst<UserClass>::Value, "Binding a delegate with a const object pointer and non-const function is deprecated.");
+
 		TBaseDelegate<RetValType, ParamTypes...> Result;
 		TBaseUObjectMethodDelegateInstance<false, UserClass, TFuncType, VarTypes...>::Create(Result, InUserObject, InFunc, Vars...);
 		return Result;
@@ -439,6 +452,8 @@ public:
 	template <typename UserClass, typename... VarTypes>
 	inline void BindRaw(UserClass* InUserObject, typename TMemFunPtrType<false, UserClass, RetValType (ParamTypes..., VarTypes...)>::Type InFunc, VarTypes... Vars)
 	{
+		UE_STATIC_DEPRECATE(4.23, TIsConst<UserClass>::Value, "Binding a delegate with a const object pointer and non-const function is deprecated.");
+
 		*this = CreateRaw(InUserObject, InFunc, Vars...);
 	}
 	template <typename UserClass, typename... VarTypes>
@@ -453,6 +468,8 @@ public:
 	template <typename UserClass, typename... VarTypes>
 	inline void BindSP(const TSharedRef<UserClass, ESPMode::Fast>& InUserObjectRef, typename TMemFunPtrType<false, UserClass, RetValType (ParamTypes..., VarTypes...)>::Type InFunc, VarTypes... Vars)
 	{
+		UE_STATIC_DEPRECATE(4.23, TIsConst<UserClass>::Value, "Binding a delegate with a const object pointer and non-const function is deprecated.");
+
 		*this = CreateSP(InUserObjectRef, InFunc, Vars...);
 	}
 	template <typename UserClass, typename... VarTypes>
@@ -470,6 +487,8 @@ public:
 	template <typename UserClass, typename... VarTypes>
 	inline void BindSP(UserClass* InUserObject, typename TMemFunPtrType<false, UserClass, RetValType (ParamTypes..., VarTypes...)>::Type InFunc, VarTypes... Vars)
 	{
+		UE_STATIC_DEPRECATE(4.23, TIsConst<UserClass>::Value, "Binding a delegate with a const object pointer and non-const function is deprecated.");
+
 		*this = CreateSP(InUserObject, InFunc, Vars...);
 	}
 	template <typename UserClass, typename... VarTypes>
@@ -487,6 +506,8 @@ public:
 	template <typename UserClass, typename... VarTypes>
 	inline void BindThreadSafeSP(const TSharedRef<UserClass, ESPMode::ThreadSafe>& InUserObjectRef, typename TMemFunPtrType<false, UserClass, RetValType (ParamTypes..., VarTypes...)>::Type InFunc, VarTypes... Vars)
 	{
+		UE_STATIC_DEPRECATE(4.23, TIsConst<UserClass>::Value, "Binding a delegate with a const object pointer and non-const function is deprecated.");
+
 		*this = CreateThreadSafeSP(InUserObjectRef, InFunc, Vars...);
 	}
 	template <typename UserClass, typename... VarTypes>
@@ -504,6 +525,8 @@ public:
 	template <typename UserClass, typename... VarTypes>
 	inline void BindThreadSafeSP(UserClass* InUserObject, typename TMemFunPtrType<false, UserClass, RetValType (ParamTypes..., VarTypes...)>::Type InFunc, VarTypes... Vars)
 	{
+		UE_STATIC_DEPRECATE(4.23, TIsConst<UserClass>::Value, "Binding a delegate with a const object pointer and non-const function is deprecated.");
+
 		*this = CreateThreadSafeSP(InUserObject, InFunc, Vars...);
 	}
 	template <typename UserClass, typename... VarTypes>
@@ -533,6 +556,8 @@ public:
 	template <typename UserClass, typename... VarTypes>
 	inline void BindUObject(UserClass* InUserObject, typename TMemFunPtrType<false, UserClass, RetValType (ParamTypes..., VarTypes...)>::Type InFunc, VarTypes... Vars)
 	{
+		UE_STATIC_DEPRECATE(4.23, TIsConst<UserClass>::Value, "Binding a delegate with a const object pointer and non-const function is deprecated.");
+
 		*this = CreateUObject(InUserObject, InFunc, Vars...);
 	}
 	template <typename UserClass, typename... VarTypes>
@@ -764,6 +789,8 @@ public:
 	template <typename UserClass, typename... VarTypes>
 	inline FDelegateHandle AddRaw(UserClass* InUserObject, typename TMemFunPtrType<false, UserClass, void (ParamTypes..., VarTypes...)>::Type InFunc, VarTypes... Vars)
 	{
+		UE_STATIC_DEPRECATE(4.23, TIsConst<UserClass>::Value, "Binding a delegate with a const object pointer and non-const function is deprecated.");
+
 		return Add(FDelegate::CreateRaw(InUserObject, InFunc, Vars...));
 	}
 	template <typename UserClass, typename... VarTypes>
@@ -783,6 +810,8 @@ public:
 	template <typename UserClass, typename... VarTypes>
 	inline FDelegateHandle AddSP(const TSharedRef<UserClass, ESPMode::Fast>& InUserObjectRef, typename TMemFunPtrType<false, UserClass, void (ParamTypes..., VarTypes...)>::Type InFunc, VarTypes... Vars)
 	{
+		UE_STATIC_DEPRECATE(4.23, TIsConst<UserClass>::Value, "Binding a delegate with a const object pointer and non-const function is deprecated.");
+
 		return Add(FDelegate::CreateSP(InUserObjectRef, InFunc, Vars...));
 	}
 	template <typename UserClass, typename... VarTypes>
@@ -802,6 +831,8 @@ public:
 	template <typename UserClass, typename... VarTypes>
 	inline FDelegateHandle AddSP(UserClass* InUserObject, typename TMemFunPtrType<false, UserClass, void (ParamTypes..., VarTypes...)>::Type InFunc, VarTypes... Vars)
 	{
+		UE_STATIC_DEPRECATE(4.23, TIsConst<UserClass>::Value, "Binding a delegate with a const object pointer and non-const function is deprecated.");
+
 		return Add(FDelegate::CreateSP(InUserObject, InFunc, Vars...));
 	}
 	template <typename UserClass, typename... VarTypes>
@@ -819,6 +850,8 @@ public:
 	template <typename UserClass, typename... VarTypes>
 	inline FDelegateHandle AddThreadSafeSP(const TSharedRef<UserClass, ESPMode::ThreadSafe>& InUserObjectRef, typename TMemFunPtrType<false, UserClass, void (ParamTypes..., VarTypes...)>::Type InFunc, VarTypes... Vars)
 	{
+		UE_STATIC_DEPRECATE(4.23, TIsConst<UserClass>::Value, "Binding a delegate with a const object pointer and non-const function is deprecated.");
+
 		return Add(FDelegate::CreateThreadSafeSP(InUserObjectRef, InFunc, Vars...));
 	}
 	template <typename UserClass, typename... VarTypes>
@@ -838,6 +871,8 @@ public:
 	template <typename UserClass, typename... VarTypes>
 	inline FDelegateHandle AddThreadSafeSP(UserClass* InUserObject, typename TMemFunPtrType<false, UserClass, void (ParamTypes..., VarTypes...)>::Type InFunc, VarTypes... Vars)
 	{
+		UE_STATIC_DEPRECATE(4.23, TIsConst<UserClass>::Value, "Binding a delegate with a const object pointer and non-const function is deprecated.");
+
 		return Add(FDelegate::CreateThreadSafeSP(InUserObject, InFunc, Vars...));
 	}
 	template <typename UserClass, typename... VarTypes>
@@ -871,6 +906,8 @@ public:
 	template <typename UserClass, typename... VarTypes>
 	inline FDelegateHandle AddUObject(UserClass* InUserObject, typename TMemFunPtrType<false, UserClass, void (ParamTypes..., VarTypes...)>::Type InFunc, VarTypes... Vars)
 	{
+		UE_STATIC_DEPRECATE(4.23, TIsConst<UserClass>::Value, "Binding a delegate with a const object pointer and non-const function is deprecated.");
+
 		return Add(FDelegate::CreateUObject(InUserObject, InFunc, Vars...));
 	}
 	template <typename UserClass, typename... VarTypes>
diff --git a/Engine/Source/Runtime/NavigationSystem/Public/NavigationData.h b/Engine/Source/Runtime/NavigationSystem/Public/NavigationData.h
--- a/Engine/Source/Runtime/NavigationSystem/Public/NavigationData.h
+++ b/Engine/Source/Runtime/NavigationSystem/Public/NavigationData.h
@@ -649,7 +649,7 @@ public:
 			STATGROUP_TaskGraphTasks);
 
 		FSimpleDelegateGraphTask::CreateAndDispatchWhenReady(
-			FSimpleDelegateGraphTask::FDelegate::CreateUObject(this, &ANavigationData::RegisterActivePath, SharedPath),
+			FSimpleDelegateGraphTask::FDelegate::CreateUObject(const_cast<ANavigationData*>(this), &ANavigationData::RegisterActivePath, SharedPath),
 			GET_STATID(STAT_FSimpleDelegateGraphTask_AddingPathToActivePaths), NULL, ENamedThreads::GameThread
 		);
 
