Silence VS2019 conversion warnings

diff --git a/Engine/Source/Runtime/Core/Private/Tests/Misc/AsciiSetTest.cpp b/Engine/Source/Runtime/Core/Private/Tests/Misc/AsciiSetTest.cpp
--- a/Engine/Source/Runtime/Core/Private/Tests/Misc/AsciiSetTest.cpp
+++ b/Engine/Source/Runtime/Core/Private/Tests/Misc/AsciiSetTest.cpp
@@ -10,9 +10,9 @@ IMPLEMENT_SIMPLE_AUTOMATION_TEST(TAsciiSetTest, "System.Core.Misc.AsciiSet", EAu
 bool TAsciiSetTest::RunTest(const FString& Parameters)
 {
 	constexpr FAsciiSet Whitespaces(" \v\f\t\r\n");
-	TestTrue(TEXT("Contains"), Whitespaces.Test(' '));
-	TestTrue(TEXT("Contains"), Whitespaces.Test('\n'));
-	TestFalse(TEXT("Contains"), Whitespaces.Test('a'));
+	TestTrue(TEXT("Contains"), Whitespaces.Contains(' '));
+	TestTrue(TEXT("Contains"), Whitespaces.Contains('\n'));
+	TestFalse(TEXT("Contains"), Whitespaces.Contains('a'));
 
 	constexpr FAsciiSet NonWhitespaces = ~Whitespaces;
 	uint32 WhitespaceNum = 0;
@@ -50,11 +50,11 @@ bool TAsciiSetTest::RunTest(const FString& Parameters)
 
 	constexpr FAsciiSet Abc("abc");
 	constexpr FAsciiSet Abcd = Abc + 'd';
-	TestTrue(TEXT("Add"), Abcd.Test('a'));
-	TestTrue(TEXT("Add"), Abcd.Test('b'));
-	TestTrue(TEXT("Add"), Abcd.Test('c'));
-	TestTrue(TEXT("Add"), Abcd.Test('d'));
-	TestFalse(TEXT("Add"), Abcd.Test('e'));
+	TestTrue(TEXT("Add"), Abcd.Contains('a'));
+	TestTrue(TEXT("Add"), Abcd.Contains('b'));
+	TestTrue(TEXT("Add"), Abcd.Contains('c'));
+	TestTrue(TEXT("Add"), Abcd.Contains('d'));
+	TestFalse(TEXT("Add"), Abcd.Contains('e'));
 
 	return true;
 }
diff --git a/Engine/Source/Runtime/Core/Public/Misc/AsciiSet.h b/Engine/Source/Runtime/Core/Public/Misc/AsciiSet.h
--- a/Engine/Source/Runtime/Core/Public/Misc/AsciiSet.h
+++ b/Engine/Source/Runtime/Core/Public/Misc/AsciiSet.h
@@ -34,7 +34,14 @@ public:
 		: FAsciiSet(StringToBitset(Chars))
 	{}
 
-	/** Returns non-zero if a character is part of the set */
+	/** Returns true if a character is part of the set */
+	template<typename CharType>
+	constexpr FORCEINLINE bool Contains(CharType Char) const
+	{
+		return !!TestImpl(TChar<CharType>::ToUnsigned(Char));
+	}
+
+	/** Returns non-zero if a character is part of the set. Prefer Contains() to avoid VS2019 conversion warnings. */
 	template<typename CharType>
 	constexpr FORCEINLINE uint64 Test(CharType Char) const
 	{
