Changed code which used TryEnterCriticalSection followed by EnterCriticalSection into a straight EnterCriticalSection call

Previous code seemed to be based on misunderstanding of where critical section spinning happens. The TryEnter* is redundant if you've already decided that you want to enter the critical section, and the spinning happens inside EnterCriticalSection

diff --git a/Engine/Source/Runtime/Core/Public/Windows/WindowsCriticalSection.h b/Engine/Source/Runtime/Core/Public/Windows/WindowsCriticalSection.h
--- a/Engine/Source/Runtime/Core/Public/Windows/WindowsCriticalSection.h
+++ b/Engine/Source/Runtime/Core/Public/Windows/WindowsCriticalSection.h
@@ -44,11 +44,7 @@ public:
 	 */
 	FORCEINLINE void Lock()
 	{
-		// Spin first before entering critical section, causing ring-0 transition and context switch.
-		if(Windows::TryEnterCriticalSection(&CriticalSection) == 0 )
-		{
-			Windows::EnterCriticalSection(&CriticalSection);
-		}
+		Windows::EnterCriticalSection(&CriticalSection);
 	}
 
 	/**
