Remove deprecation warning macros from FArchive.(cpp|h) since the properties are now private and no longer flagged as deprecated

diff --git a/Engine/Source/Runtime/Core/Private/Serialization/Archive.cpp b/Engine/Source/Runtime/Core/Private/Serialization/Archive.cpp
--- a/Engine/Source/Runtime/Core/Private/Serialization/Archive.cpp
+++ b/Engine/Source/Runtime/Core/Private/Serialization/Archive.cpp
@@ -120,7 +120,6 @@ FArchive::~FArchive()
 }
 
 
-PRAGMA_DISABLE_DEPRECATION_WARNINGS
 void FArchive::Reset()
 {
 #if DEVIRTUALIZE_FLinkerLoad_Serialize
@@ -226,7 +225,6 @@ void FArchive::CopyTrivialFArchiveStatusMembers(const FArchive& ArchiveToCopy)
 	SetBaseLocalizationNamespace(ArchiveToCopy.GetBaseLocalizationNamespace());
 #endif // USE_STABLE_LOCALIZATION_KEYS
 }
-PRAGMA_ENABLE_DEPRECATION_WARNINGS
 
 /**
  * Returns the name of the Archive.  Useful for getting the name of the package a struct or object
@@ -1092,7 +1090,6 @@ void FArchive::LogfImpl(const TCHAR* Fmt, ...)
 	FMemory::SystemFree( Buffer );
 }
 
-PRAGMA_DISABLE_DEPRECATION_WARNINGS
 void FArchive::SetUE4Ver(int32 InVer)
 {
 	ArUE4Ver = InVer;
@@ -1152,4 +1149,3 @@ void FArchive::SetIsPersistent(bool bInIsPersistent)
 {
 	ArIsPersistent = bInIsPersistent;
 }
-PRAGMA_ENABLE_DEPRECATION_WARNINGS
diff --git a/Engine/Source/Runtime/Core/Public/Serialization/Archive.h b/Engine/Source/Runtime/Core/Public/Serialization/Archive.h
--- a/Engine/Source/Runtime/Core/Public/Serialization/Archive.h
+++ b/Engine/Source/Runtime/Core/Public/Serialization/Archive.h
@@ -893,7 +893,6 @@ public:
 		return this;
 	}
 
-PRAGMA_DISABLE_DEPRECATION_WARNINGS
 	FORCEINLINE bool IsLoading() const
 	{
 		return ArIsLoading;
@@ -935,7 +934,6 @@ PRAGMA_DISABLE_DEPRECATION_WARNINGS
 	{
 		return ArIsPersistent;
 	}
-PRAGMA_ENABLE_DEPRECATION_WARNINGS
 
 	FORCEINLINE bool IsError() const
 	{
