Build fixes for non-unity builds with no PCH

diff --git a/Engine/Source/Runtime/Core/Public/Misc/IEngineCrypto.h b/Engine/Source/Runtime/Core/Public/Misc/IEngineCrypto.h
--- a/Engine/Source/Runtime/Core/Public/Misc/IEngineCrypto.h
+++ b/Engine/Source/Runtime/Core/Public/Misc/IEngineCrypto.h
@@ -2,6 +2,7 @@
 
 #include "CoreTypes.h"
 #include "Features/IModularFeatures.h"
+#include "Containers/ArrayView.h"
 
 typedef void* FRSAKeyHandle;
 static const FRSAKeyHandle InvalidRSAKeyHandle = nullptr;
diff --git a/Engine/Source/Runtime/RSA/Private/RSA.cpp b/Engine/Source/Runtime/RSA/Private/RSA.cpp
--- a/Engine/Source/Runtime/RSA/Private/RSA.cpp
+++ b/Engine/Source/Runtime/RSA/Private/RSA.cpp
@@ -2,6 +2,7 @@
 
 #include "RSA.h"
 #include "Features/IModularFeatures.h"
+#include "Modules/ModuleManager.h"
 
 static FCriticalSection GLock;
 
