VS2019 fix

diff --git a/Engine/Source/Runtime/Core/Private/UObject/UnrealNames.cpp b/Engine/Source/Runtime/Core/Private/UObject/UnrealNames.cpp
--- a/Engine/Source/Runtime/Core/Private/UObject/UnrealNames.cpp
+++ b/Engine/Source/Runtime/Core/Private/UObject/UnrealNames.cpp
@@ -243,7 +243,7 @@ struct FNameSlot
 	
 	bool operator==(FNameSlot Rhs) const { return IdAndHash == Rhs.IdAndHash; }
 
-	bool Used() const { return IdAndHash;  }
+	bool Used() const { return !!IdAndHash;  }
 private:
 	uint32 IdAndHash = 0;
 };
