Marked logging trace events as always on.

diff --git a/Engine/Source/Runtime/Core/Private/Logging/LogTrace.cpp b/Engine/Source/Runtime/Core/Private/Logging/LogTrace.cpp
--- a/Engine/Source/Runtime/Core/Private/Logging/LogTrace.cpp
+++ b/Engine/Source/Runtime/Core/Private/Logging/LogTrace.cpp
@@ -8,19 +8,19 @@
 #include "HAL/PlatformTime.h"
 #include "HAL/PlatformTLS.h"
 
-UE_TRACE_EVENT_BEGIN(Logging, LogCategory)
+UE_TRACE_EVENT_BEGIN(Logging, LogCategory, Always)
 	UE_TRACE_EVENT_FIELD(const void*, CategoryPointer)
 	UE_TRACE_EVENT_FIELD(uint8, DefaultVerbosity)
 UE_TRACE_EVENT_END()
 
-UE_TRACE_EVENT_BEGIN(Logging, LogMessageSpec)
+UE_TRACE_EVENT_BEGIN(Logging, LogMessageSpec, Always)
 	UE_TRACE_EVENT_FIELD(const void*, LogPoint)
 	UE_TRACE_EVENT_FIELD(const void*, CategoryPointer)
 	UE_TRACE_EVENT_FIELD(int32, Line)
 	UE_TRACE_EVENT_FIELD(uint8, Verbosity)
 UE_TRACE_EVENT_END()
 
-UE_TRACE_EVENT_BEGIN(Logging, LogMessage)
+UE_TRACE_EVENT_BEGIN(Logging, LogMessage, Always)
 	UE_TRACE_EVENT_FIELD(const void*, LogPoint)
 	UE_TRACE_EVENT_FIELD(uint64, Cycle)
 	UE_TRACE_EVENT_FIELD(uint32, ThreadId)
