Fix for potential hang on startup with recursive GCreateMalloc vs FName construction with LLM.

diff --git a/Engine/Source/Runtime/Core/Public/HAL/LowLevelMemTracker.h b/Engine/Source/Runtime/Core/Public/HAL/LowLevelMemTracker.h
--- a/Engine/Source/Runtime/Core/Public/HAL/LowLevelMemTracker.h
+++ b/Engine/Source/Runtime/Core/Public/HAL/LowLevelMemTracker.h
@@ -287,8 +287,8 @@ extern FName LLMGetTagStat(ELLMTag Tag);
  /**
  * LLM Pause scope macros
  */
-#define LLM_SCOPED_PAUSE_TRACKING(AllocType) FLLMPauseScope SCOPE_NAME(NAME_None, 0, ELLMTracker::Max, AllocType);
-#define LLM_SCOPED_PAUSE_TRACKING_FOR_TRACKER(Tracker, AllocType) FLLMPauseScope SCOPE_NAME(NAME_None, 0, Tracker, AllocType);
+#define LLM_SCOPED_PAUSE_TRACKING(AllocType) FLLMPauseScope SCOPE_NAME(ELLMTag::Untagged, 0, ELLMTracker::Max, AllocType);
+#define LLM_SCOPED_PAUSE_TRACKING_FOR_TRACKER(Tracker, AllocType) FLLMPauseScope SCOPE_NAME(ELLMTag::Untagged, 0, Tracker, AllocType);
 #define LLM_SCOPED_PAUSE_TRACKING_WITH_ENUM_AND_AMOUNT(Tag, Amount, Tracker, AllocType) FLLMPauseScope SCOPE_NAME(Tag, Amount, Tracker, AllocType);
 
 /**
