UnrealInsights:
Cpu scope tracing compiled in by default
More frequent flushing of cpu scope thread buffers
Ensured that cpu scope events can be enabled and disabled dynamically
Added -cpuprofilertrace command line option to enable cpu scopes from the start

diff --git a/Engine/Source/Developer/TraceServices/Private/Analyzers/CpuProfilerTraceAnalysis.cpp b/Engine/Source/Developer/TraceServices/Private/Analyzers/CpuProfilerTraceAnalysis.cpp
--- a/Engine/Source/Developer/TraceServices/Private/Analyzers/CpuProfilerTraceAnalysis.cpp
+++ b/Engine/Source/Developer/TraceServices/Private/Analyzers/CpuProfilerTraceAnalysis.cpp
@@ -17,6 +17,11 @@ void FCpuProfilerAnalyzer::OnAnalysisBegin(const FOnAnalysisContext& Context)
 
 	Builder.RouteEvent(RouteId_EventSpec, "CpuProfiler", "EventSpec");
 	Builder.RouteEvent(RouteId_EventBatch, "CpuProfiler", "EventBatch");
+	Builder.RouteEvent(RouteId_EndCapture, "CpuProfiler", "EndCapture");
+}
+
+void FCpuProfilerAnalyzer::OnAnalysisEnd()
+{
 }
 
 void FCpuProfilerAnalyzer::OnEvent(uint16 RouteId, const FOnEventContext& Context)
@@ -29,11 +34,20 @@ void FCpuProfilerAnalyzer::OnEvent(uint16 RouteId, const FOnEventContext& Contex
 	case RouteId_EventSpec:
 	{
 		uint16 Id = EventData.GetValue("Id").As<uint16>();
-		ScopeIdToEventIdMap.Add(Id, TimingProfilerProvider.AddCpuTimer(reinterpret_cast<const TCHAR*>(EventData.GetAttachment())));
+		if (ScopeIdToEventIdMap.Contains(Id))
+		{
+			TimingProfilerProvider.SetTimerName(ScopeIdToEventIdMap[Id], reinterpret_cast<const TCHAR*>(EventData.GetAttachment()));
+		}
+		else
+		{
+			ScopeIdToEventIdMap.Add(Id, TimingProfilerProvider.AddCpuTimer(reinterpret_cast<const TCHAR*>(EventData.GetAttachment())));
+		}
 		break;
 	}
 	case RouteId_EventBatch:
+	case RouteId_EndCapture:
 	{
+		TotalEventSize += EventData.GetTotalSize();
 		uint32 ThreadId = EventData.GetValue("ThreadId").As<uint32>();
 		TSharedRef<FThreadState> ThreadState = GetThreadState(ThreadId);
 		uint64 LastCycle = ThreadState->LastCycle;
@@ -50,14 +64,22 @@ void FCpuProfilerAnalyzer::OnEvent(uint16 RouteId, const FOnEventContext& Contex
 				EventScopeState& ScopeState = ThreadState->ScopeStack.AddDefaulted_GetRef();
 				ScopeState.StartCycle = ActualCycle;
 				uint16 SpecId = FTraceAnalyzerUtils::Decode7bit(BufferPtr);
-				ScopeState.EventTypeId = ScopeIdToEventIdMap[SpecId];
+				uint32* FindIt = ScopeIdToEventIdMap.Find(SpecId);
+				if (!FindIt)
+				{
+					ScopeState.EventTypeId = ScopeIdToEventIdMap.Add(SpecId, TimingProfilerProvider.AddCpuTimer(TEXT("<unknown>")));
+				}
+				else
+				{
+					ScopeState.EventTypeId = *FindIt;
+				}
 				Trace::FTimingProfilerEvent Event;
 				Event.TimerIndex = ScopeState.EventTypeId;
 				ThreadState->Timeline->AppendBeginEvent(Context.SessionContext.TimestampFromCycle(ScopeState.StartCycle), Event);
+				++TotalScopeCount;
 			}
-			else
+			else if (ThreadState->ScopeStack.Num())
 			{
-				check(ThreadState->ScopeStack.Num());
 				ThreadState->ScopeStack.Pop();
 				ThreadState->Timeline->AppendEndEvent(Context.SessionContext.TimestampFromCycle(ActualCycle));
 			}
@@ -65,9 +87,20 @@ void FCpuProfilerAnalyzer::OnEvent(uint16 RouteId, const FOnEventContext& Contex
 		check(BufferPtr == BufferEnd);
 		if (LastCycle)
 		{
-			Session.UpdateDurationSeconds(Context.SessionContext.TimestampFromCycle(LastCycle));
+			double LastTimestamp = Context.SessionContext.TimestampFromCycle(LastCycle);
+			Session.UpdateDurationSeconds(LastTimestamp);
+			if (RouteId == RouteId_EndCapture)
+			{
+
+				while (ThreadState->ScopeStack.Num())
+				{
+					ThreadState->ScopeStack.Pop();
+					ThreadState->Timeline->AppendEndEvent(LastTimestamp);
+				}
+			}
 		}
 		ThreadState->LastCycle = LastCycle;
+		BytesPerScope = double(TotalEventSize) / double(TotalScopeCount);
 		break;
 	}
 	}
diff --git a/Engine/Source/Developer/TraceServices/Private/Analyzers/CpuProfilerTraceAnalysis.h b/Engine/Source/Developer/TraceServices/Private/Analyzers/CpuProfilerTraceAnalysis.h
--- a/Engine/Source/Developer/TraceServices/Private/Analyzers/CpuProfilerTraceAnalysis.h
+++ b/Engine/Source/Developer/TraceServices/Private/Analyzers/CpuProfilerTraceAnalysis.h
@@ -19,7 +19,7 @@ public:
 	FCpuProfilerAnalyzer(Trace::IAnalysisSession& Session, Trace::FTimingProfilerProvider& TimingProfilerProvider);
 	virtual void OnAnalysisBegin(const FOnAnalysisContext& Context) override;
 	virtual void OnEvent(uint16 RouteId, const FOnEventContext& Context) override;
-	virtual void OnAnalysisEnd() override {};
+	virtual void OnAnalysisEnd() override;
 
 private:
 	struct EventScopeState
@@ -41,10 +41,14 @@ private:
 	{
 		RouteId_EventSpec,
 		RouteId_EventBatch,
+		RouteId_EndCapture,
 	};
 
 	Trace::IAnalysisSession& Session;
 	Trace::FTimingProfilerProvider& TimingProfilerProvider;
 	TMap<uint32, TSharedRef<FThreadState>> ThreadStatesMap;
 	TMap<uint16, uint32> ScopeIdToEventIdMap;
+	uint64 TotalEventSize = 0;
+	uint64 TotalScopeCount = 0;
+	double BytesPerScope = 0.0;
 };
diff --git a/Engine/Source/Developer/TraceServices/Private/Model/TimingProfiler.cpp b/Engine/Source/Developer/TraceServices/Private/Model/TimingProfiler.cpp
--- a/Engine/Source/Developer/TraceServices/Private/Model/TimingProfiler.cpp
+++ b/Engine/Source/Developer/TraceServices/Private/Model/TimingProfiler.cpp
@@ -27,6 +27,20 @@ uint32 FTimingProfilerProvider::AddCpuTimer(const TCHAR* Name)
 	return Timer.Id;
 }
 
+void FTimingProfilerProvider::SetTimerName(uint32 TimerId, const TCHAR* Name)
+{
+	Session.WriteAccessCheck();
+	
+	FTimingProfilerTimer& Timer = Timers[TimerId];
+	Timer.Name = Session.StoreString(Name);
+	uint32 NameHash = 0;
+	for (const TCHAR* c = Name; *c; ++c)
+	{
+		NameHash = (NameHash + *c) * 0x2c2c57ed;
+	}
+	Timer.NameHash = NameHash;
+}
+
 uint32 FTimingProfilerProvider::AddGpuTimer(const TCHAR* Name)
 {
 	Session.WriteAccessCheck();
diff --git a/Engine/Source/Developer/TraceServices/Private/Model/TimingProfilerPrivate.h b/Engine/Source/Developer/TraceServices/Private/Model/TimingProfilerPrivate.h
--- a/Engine/Source/Developer/TraceServices/Private/Model/TimingProfilerPrivate.h
+++ b/Engine/Source/Developer/TraceServices/Private/Model/TimingProfilerPrivate.h
@@ -23,6 +23,7 @@ public:
 	virtual ~FTimingProfilerProvider();
 	uint32 AddCpuTimer(const TCHAR* Name);
 	uint32 AddGpuTimer(const TCHAR* Name);
+	void SetTimerName(uint32 TimerId, const TCHAR* Name);
 	TimelineInternal& EditCpuThreadTimeline(uint32 ThreadId);
 	TimelineInternal& EditGpuTimeline();
 	virtual bool GetCpuThreadTimelineIndex(uint32 ThreadId, uint32& OutTimelineIndex) const override;
diff --git a/Engine/Source/Runtime/Core/Private/ProfilingDebugging/CpuProfilerTrace.cpp b/Engine/Source/Runtime/Core/Private/ProfilingDebugging/CpuProfilerTrace.cpp
--- a/Engine/Source/Runtime/Core/Private/ProfilingDebugging/CpuProfilerTrace.cpp
+++ b/Engine/Source/Runtime/Core/Private/ProfilingDebugging/CpuProfilerTrace.cpp
@@ -1,15 +1,13 @@
 // Copyright 1998-2019 Epic Games, Inc. All Rights Reserved.
 #include "ProfilingDebugging/CpuProfilerTrace.h"
-#include "Templates/UnrealTemplate.h"
 #include "Trace/Trace.h"
 #include "HAL/PlatformTime.h"
 #include "HAL/PlatformTLS.h"
-#include "Containers/Map.h"
 #include "ProfilingDebugging/MiscTrace.h"
 
 #if CPUPROFILERTRACE_ENABLED
 
-UE_TRACE_EVENT_BEGIN(CpuProfiler, EventSpec)
+UE_TRACE_EVENT_BEGIN(CpuProfiler, EventSpec, Always)
 	UE_TRACE_EVENT_FIELD(uint16, Id)
 	UE_TRACE_EVENT_FIELD(uint16, Group)
 UE_TRACE_EVENT_END()
@@ -18,6 +16,10 @@ UE_TRACE_EVENT_BEGIN(CpuProfiler, EventBatch)
 	UE_TRACE_EVENT_FIELD(uint32, ThreadId)
 UE_TRACE_EVENT_END()
 
+UE_TRACE_EVENT_BEGIN(CpuProfiler, EndCapture, Always)
+	UE_TRACE_EVENT_FIELD(uint32, ThreadId)
+UE_TRACE_EVENT_END()
+
 struct FCpuProfilerTraceInternal
 {
 public:
@@ -30,18 +32,18 @@ public:
 
 	struct FThreadState
 	{
-		uint64 CycleBase;
-		uint64 LastCycle;
+		uint64 LastCycle = 0;
 		uint32 ThreadId;
-		uint32 BufferSize;
+		uint16 Depth = 0;
+		uint16 BufferSize = 0;
+		bool bEnabled = false;
 		uint8 Buffer[MaxBufferSize];
 	};
 
-	static FThreadState* GetThreadState() { return ThreadLocalThreadState; }
 	FORCENOINLINE static FThreadState* InitThreadState();
 	FORCENOINLINE static void FlushThreadBuffer(FThreadState* ThreadState);
+	FORCENOINLINE static void EndCapture(FThreadState* ThreadState);
 
-private:
 	static thread_local FThreadState* ThreadLocalThreadState;
 };
 
@@ -50,9 +52,7 @@ thread_local FCpuProfilerTraceInternal::FThreadState* FCpuProfilerTraceInternal:
 FCpuProfilerTraceInternal::FThreadState* FCpuProfilerTraceInternal::InitThreadState()
 {
 	ThreadLocalThreadState = new FThreadState();
-	ThreadLocalThreadState->BufferSize = 0;
 	ThreadLocalThreadState->ThreadId = FPlatformTLS::GetCurrentThreadId();
-	ThreadLocalThreadState->LastCycle = 0;
 	return ThreadLocalThreadState;
 }
 
@@ -64,20 +64,40 @@ void FCpuProfilerTraceInternal::FlushThreadBuffer(FThreadState* ThreadState)
 	ThreadState->BufferSize = 0;
 }
 
+void FCpuProfilerTraceInternal::EndCapture(FThreadState* ThreadState)
+{
+	UE_TRACE_LOG(CpuProfiler, EndCapture, ThreadState->BufferSize)
+		<< EndCapture.ThreadId(ThreadState->ThreadId)
+		<< EndCapture.Attachment(ThreadState->Buffer, ThreadState->BufferSize);
+	ThreadState->BufferSize = 0;
+}
+
 void FCpuProfilerTrace::OutputBeginEvent(uint16 SpecId)
 {
-	uint64 Cycle = FPlatformTime::Cycles64();
-	FCpuProfilerTraceInternal::FThreadState* ThreadState = FCpuProfilerTraceInternal::GetThreadState();
+	FCpuProfilerTraceInternal::FThreadState* ThreadState = FCpuProfilerTraceInternal::ThreadLocalThreadState;
 	if (!ThreadState)
 	{
 		ThreadState = FCpuProfilerTraceInternal::InitThreadState();
 	}
-	uint64 CycleDiff = Cycle - ThreadState->LastCycle;
-	ThreadState->LastCycle = Cycle;
+	++ThreadState->Depth;
+	bool bEventEnabled = UE_TRACE_EVENT_IS_ENABLED(CpuProfiler, EventBatch);
+	if (!bEventEnabled & (bEventEnabled == ThreadState->bEnabled))
+	{
+		return;
+	}
+	ThreadState->bEnabled = bEventEnabled;
+	if (!bEventEnabled)
+	{
+		FCpuProfilerTraceInternal::EndCapture(ThreadState);
+		return;
+	}
 	if (ThreadState->BufferSize >= FCpuProfilerTraceInternal::FullBufferThreshold)
 	{
 		FCpuProfilerTraceInternal::FlushThreadBuffer(ThreadState);
 	}
+	uint64 Cycle = FPlatformTime::Cycles64();
+	uint64 CycleDiff = Cycle - ThreadState->LastCycle;
+	ThreadState->LastCycle = Cycle;
 	uint8* BufferPtr = ThreadState->Buffer + ThreadState->BufferSize;
 	FTraceUtils::Encode7bit((CycleDiff << 1) | 1ull, BufferPtr);
 	FTraceUtils::Encode7bit(SpecId, BufferPtr);
@@ -86,14 +106,26 @@ void FCpuProfilerTrace::OutputBeginEvent(uint16 SpecId)
 
 void FCpuProfilerTrace::OutputEndEvent()
 {
-	uint64 Cycle = FPlatformTime::Cycles64();
-	FCpuProfilerTraceInternal::FThreadState* ThreadState = FCpuProfilerTraceInternal::GetThreadState();
-	uint64 CycleDiff = Cycle - ThreadState->LastCycle;
-	ThreadState->LastCycle = Cycle;
-	if (ThreadState->BufferSize >= FCpuProfilerTraceInternal::FullBufferThreshold)
+	FCpuProfilerTraceInternal::FThreadState* ThreadState = FCpuProfilerTraceInternal::ThreadLocalThreadState;
+	--ThreadState->Depth;
+	bool bEventEnabled = UE_TRACE_EVENT_IS_ENABLED(CpuProfiler, EventBatch);
+	if (!bEventEnabled & (bEventEnabled == ThreadState->bEnabled))
+	{
+		return;
+	}
+	ThreadState->bEnabled = bEventEnabled;
+	if (!bEventEnabled)
+	{
+		FCpuProfilerTraceInternal::EndCapture(ThreadState);
+		return;
+	}
+	if ((ThreadState->Depth == 0) | (ThreadState->BufferSize >= FCpuProfilerTraceInternal::FullBufferThreshold))
 	{
 		FCpuProfilerTraceInternal::FlushThreadBuffer(ThreadState);
 	}
+	uint64 Cycle = FPlatformTime::Cycles64();
+	uint64 CycleDiff = Cycle - ThreadState->LastCycle;
+	ThreadState->LastCycle = Cycle;
 	uint8* BufferPtr = ThreadState->Buffer + ThreadState->BufferSize;
 	FTraceUtils::Encode7bit(CycleDiff << 1, BufferPtr);
 	ThreadState->BufferSize = BufferPtr - ThreadState->Buffer;
@@ -113,5 +145,13 @@ uint16 FCpuProfilerTrace::OutputEventType(const TCHAR* Name, ECpuProfilerGroup G
 	return Id;
 }
 
+void FCpuProfilerTrace::Init(bool bStartEnabled)
+{
+	if (bStartEnabled)
+	{
+		UE_TRACE_EVENT_IS_ENABLED(CpuProfiler, EventBatch);
+		Trace::ToggleEvent(TEXT("CpuProfiler.EventBatch"), true);
+	}
+}
 
 #endif
\ No newline at end of file
diff --git a/Engine/Source/Runtime/Core/Public/ProfilingDebugging/CpuProfilerTrace.h b/Engine/Source/Runtime/Core/Public/ProfilingDebugging/CpuProfilerTrace.h
--- a/Engine/Source/Runtime/Core/Public/ProfilingDebugging/CpuProfilerTrace.h
+++ b/Engine/Source/Runtime/Core/Public/ProfilingDebugging/CpuProfilerTrace.h
@@ -4,8 +4,8 @@
 
 #include "CoreTypes.h"
 
-#if !IS_PROGRAM && !UE_BUILD_SHIPPING
-#define CPUPROFILERTRACE_ENABLED 0
+#if !IS_PROGRAM && !UE_BUILD_SHIPPING && (PLATFORM_WINDOWS || PLATFORM_PS4 || PLATFORM_XBOXONE)
+#define CPUPROFILERTRACE_ENABLED 1
 #else
 #define CPUPROFILERTRACE_ENABLED 0
 #endif
@@ -22,6 +22,7 @@ enum ECpuProfilerGroup
 
 struct FCpuProfilerTrace
 {
+	CORE_API static void Init(bool bStartEnabled);
 	CORE_API static uint16 OutputEventType(const TCHAR* Name, ECpuProfilerGroup Group);
 	CORE_API static void OutputBeginEvent(uint16 SpecId);
 	CORE_API static void OutputEndEvent();
diff --git a/Engine/Source/Runtime/Launch/Private/LaunchEngineLoop.cpp b/Engine/Source/Runtime/Launch/Private/LaunchEngineLoop.cpp
--- a/Engine/Source/Runtime/Launch/Private/LaunchEngineLoop.cpp
+++ b/Engine/Source/Runtime/Launch/Private/LaunchEngineLoop.cpp
@@ -1158,6 +1158,9 @@ DECLARE_CYCLE_STAT( TEXT( "FEngineLoop::PreInit.AfterStats" ), STAT_FEngineLoop_
 int32 FEngineLoop::PreInit(const TCHAR* CmdLine)
 {
 	TRACE_REGISTER_GAME_THREAD(FPlatformTLS::GetCurrentThreadId());
+#if CPUPROFILERTRACE_ENABLED
+	FCpuProfilerTrace::Init(FParse::Param(CmdLine, TEXT("cpuprofilertrace")));
+#endif
 
 	SCOPED_BOOT_TIMING("FEngineLoop::PreInit");
 
