UnrealInsights:
Set thread, bookmark and frame events as always on
Reduced size of frame events

diff --git a/Engine/Source/Developer/TraceServices/Private/Analyzers/MiscTraceAnalysis.cpp b/Engine/Source/Developer/TraceServices/Private/Analyzers/MiscTraceAnalysis.cpp
--- a/Engine/Source/Developer/TraceServices/Private/Analyzers/MiscTraceAnalysis.cpp
+++ b/Engine/Source/Developer/TraceServices/Private/Analyzers/MiscTraceAnalysis.cpp
@@ -5,6 +5,7 @@
 #include "Model/ThreadsPrivate.h"
 #include "Model/BookmarksPrivate.h"
 #include "Model/FramesPrivate.h"
+#include "Common/Utils.h"
 
 FMiscTraceAnalyzer::FMiscTraceAnalyzer(Trace::IAnalysisSession& InSession,
 									   Trace::FThreadProvider& InThreadProvider,
@@ -35,6 +36,10 @@ void FMiscTraceAnalyzer::OnAnalysisBegin(const FOnAnalysisContext& Context)
 	Builder.RouteEvent(RouteId_Bookmark, "Misc", "Bookmark");
 	Builder.RouteEvent(RouteId_BeginFrame, "Misc", "BeginFrame");
 	Builder.RouteEvent(RouteId_EndFrame, "Misc", "EndFrame");
+	Builder.RouteEvent(RouteId_BeginGameFrame, "Misc", "BeginGameFrame");
+	Builder.RouteEvent(RouteId_EndGameFrame, "Misc", "EndGameFrame");
+	Builder.RouteEvent(RouteId_BeginRenderFrame, "Misc", "BeginRenderFrame");
+	Builder.RouteEvent(RouteId_EndRenderFrame, "Misc", "EndRenderFrame");
 }
 
 void FMiscTraceAnalyzer::OnEvent(uint16 RouteId, const FOnEventContext& Context)
@@ -127,6 +132,37 @@ void FMiscTraceAnalyzer::OnEvent(uint16 RouteId, const FOnEventContext& Context)
 		FrameProvider.EndFrame(ETraceFrameType(FrameType), Context.SessionContext.TimestampFromCycle(Cycle));
 		break;
 	}
+	case RouteId_BeginGameFrame:
+	case RouteId_EndGameFrame:
+	case RouteId_BeginRenderFrame:
+	case RouteId_EndRenderFrame:
+	{
+		ETraceFrameType FrameType = TraceFrameType_Count;
+		if (RouteId == RouteId_BeginGameFrame || RouteId == RouteId_EndGameFrame)
+		{
+			FrameType = TraceFrameType_Game;
+		}
+		else if (RouteId == RouteId_BeginRenderFrame || RouteId == RouteId_EndRenderFrame)
+		{
+			FrameType = TraceFrameType_Rendering;
+		}
+		if (FrameType != TraceFrameType_Count)
+		{
+			const uint8* BufferPtr = EventData.GetAttachment();
+			uint64 CycleDiff = FTraceAnalyzerUtils::Decode7bit(BufferPtr);
+			uint64 Cycle = LastFrameCycle[FrameType] + CycleDiff;
+			LastFrameCycle[FrameType] = Cycle;
+			if (RouteId == RouteId_BeginGameFrame || RouteId == RouteId_BeginRenderFrame)
+			{
+				FrameProvider.BeginFrame(FrameType, Context.SessionContext.TimestampFromCycle(Cycle));
+			}
+			else
+			{
+				FrameProvider.EndFrame(FrameType, Context.SessionContext.TimestampFromCycle(Cycle));
+			}
+		}
+		break;
+	}
 	}
 }
 
diff --git a/Engine/Source/Developer/TraceServices/Private/Analyzers/MiscTraceAnalysis.h b/Engine/Source/Developer/TraceServices/Private/Analyzers/MiscTraceAnalysis.h
--- a/Engine/Source/Developer/TraceServices/Private/Analyzers/MiscTraceAnalysis.h
+++ b/Engine/Source/Developer/TraceServices/Private/Analyzers/MiscTraceAnalysis.h
@@ -42,6 +42,10 @@ private:
 		RouteId_Bookmark,
 		RouteId_BeginFrame,
 		RouteId_EndFrame,
+		RouteId_BeginGameFrame,
+		RouteId_EndGameFrame,
+		RouteId_BeginRenderFrame,
+		RouteId_EndRenderFrame,
 	};
 
 	struct FThreadState
@@ -57,5 +61,6 @@ private:
 	Trace::FLogProvider& LogProvider;
 	Trace::FFrameProvider& FrameProvider;
 	TMap<uint32, TSharedRef<FThreadState>> ThreadStateMap;
+	uint64 LastFrameCycle[TraceFrameType_Count] = { 0, 0 };
 };
 
diff --git a/Engine/Source/Runtime/Core/Private/ProfilingDebugging/MiscTrace.cpp b/Engine/Source/Runtime/Core/Private/ProfilingDebugging/MiscTrace.cpp
--- a/Engine/Source/Runtime/Core/Private/ProfilingDebugging/MiscTrace.cpp
+++ b/Engine/Source/Runtime/Core/Private/ProfilingDebugging/MiscTrace.cpp
@@ -8,49 +8,59 @@
 #include "HAL/PlatformTLS.h"
 #include "HAL/PlatformTime.h"
 
-UE_TRACE_EVENT_BEGIN(Misc, RegisterGameThread)
+UE_TRACE_EVENT_BEGIN(Misc, RegisterGameThread, Always)
 	UE_TRACE_EVENT_FIELD(uint32, ThreadId)
 UE_TRACE_EVENT_END()
 
-UE_TRACE_EVENT_BEGIN(Misc, CreateThread)
+UE_TRACE_EVENT_BEGIN(Misc, CreateThread, Always)
 	UE_TRACE_EVENT_FIELD(uint32, CurrentThreadId)
 	UE_TRACE_EVENT_FIELD(uint32, CreatedThreadId)
 	UE_TRACE_EVENT_FIELD(uint32, Priority)
 	UE_TRACE_EVENT_FIELD(uint16, NameSize)
 UE_TRACE_EVENT_END()
 
-UE_TRACE_EVENT_BEGIN(Misc, SetThreadGroup)
+UE_TRACE_EVENT_BEGIN(Misc, SetThreadGroup, Always)
 	UE_TRACE_EVENT_FIELD(uint32, ThreadId)
 UE_TRACE_EVENT_END()
 
-UE_TRACE_EVENT_BEGIN(Misc, BeginThreadGroupScope)
+UE_TRACE_EVENT_BEGIN(Misc, BeginThreadGroupScope, Always)
 	UE_TRACE_EVENT_FIELD(uint32, CurrentThreadId)
 UE_TRACE_EVENT_END()
 
-UE_TRACE_EVENT_BEGIN(Misc, EndThreadGroupScope)
+UE_TRACE_EVENT_BEGIN(Misc, EndThreadGroupScope, Always)
 	UE_TRACE_EVENT_FIELD(uint32, CurrentThreadId)
 UE_TRACE_EVENT_END()
 
-UE_TRACE_EVENT_BEGIN(Misc, BookmarkSpec)
+UE_TRACE_EVENT_BEGIN(Misc, BookmarkSpec, Always)
 	UE_TRACE_EVENT_FIELD(const void*, BookmarkPoint)
 	UE_TRACE_EVENT_FIELD(int32, Line)
 UE_TRACE_EVENT_END()
 
-UE_TRACE_EVENT_BEGIN(Misc, Bookmark)
+UE_TRACE_EVENT_BEGIN(Misc, Bookmark, Always)
 	UE_TRACE_EVENT_FIELD(uint64, Cycle)
 	UE_TRACE_EVENT_FIELD(const void*, BookmarkPoint)
 UE_TRACE_EVENT_END()
 
-UE_TRACE_EVENT_BEGIN(Misc, BeginFrame)
-	UE_TRACE_EVENT_FIELD(uint64, Cycle)
-	UE_TRACE_EVENT_FIELD(uint8, FrameType)
+UE_TRACE_EVENT_BEGIN(Misc, BeginGameFrame, Always)
 UE_TRACE_EVENT_END()
 
-UE_TRACE_EVENT_BEGIN(Misc, EndFrame)
-	UE_TRACE_EVENT_FIELD(uint64, Cycle)
-	UE_TRACE_EVENT_FIELD(uint8, FrameType)
+UE_TRACE_EVENT_BEGIN(Misc, EndGameFrame, Always)
+UE_TRACE_EVENT_END()
+
+UE_TRACE_EVENT_BEGIN(Misc, BeginRenderFrame, Always)
+UE_TRACE_EVENT_END()
+
+UE_TRACE_EVENT_BEGIN(Misc, EndRenderFrame, Always)
 UE_TRACE_EVENT_END()
 
+
+struct FMiscTraceInternal
+{
+	static uint64 LastFrameCycle[TraceFrameType_Count];
+};
+
+uint64 FMiscTraceInternal::LastFrameCycle[TraceFrameType_Count] = { 0, 0 };
+
 void FMiscTrace::OutputRegisterGameThread(uint32 Id)
 {
 	UE_TRACE_LOG(Misc, RegisterGameThread)
@@ -113,16 +123,44 @@ void FMiscTrace::OutputBookmarkInternal(const void* BookmarkPoint, uint16 Encode
 
 void FMiscTrace::OutputBeginFrame(ETraceFrameType FrameType)
 {
-	UE_TRACE_LOG(Misc, BeginFrame)
-		<< BeginFrame.Cycle(FPlatformTime::Cycles64())
-		<< BeginFrame.FrameType(uint8(FrameType));
+	uint64 Cycle = FPlatformTime::Cycles64();
+	uint64 CycleDiff = Cycle - FMiscTraceInternal::LastFrameCycle[FrameType];
+	FMiscTraceInternal::LastFrameCycle[FrameType] = Cycle;
+	uint8 Buffer[9];
+	uint8* BufferPtr = Buffer;
+	FTraceUtils::Encode7bit(CycleDiff, BufferPtr);
+	uint16 BufferSize = BufferPtr - Buffer;
+	if (FrameType == TraceFrameType_Game)
+	{
+		UE_TRACE_LOG(Misc, BeginGameFrame, BufferSize)
+			<< BeginGameFrame.Attachment(&Buffer, BufferSize);
+	}
+	else if (FrameType == TraceFrameType_Rendering)
+	{
+		UE_TRACE_LOG(Misc, BeginRenderFrame, BufferSize)
+			<< BeginRenderFrame.Attachment(&Buffer, BufferSize);
+	}
 }
 
 void FMiscTrace::OutputEndFrame(ETraceFrameType FrameType)
 {
-	UE_TRACE_LOG(Misc, EndFrame)
-		<< EndFrame.Cycle(FPlatformTime::Cycles64())
-		<< EndFrame.FrameType(uint8(FrameType));
+	uint64 Cycle = FPlatformTime::Cycles64();
+	uint64 CycleDiff = Cycle - FMiscTraceInternal::LastFrameCycle[FrameType];
+	FMiscTraceInternal::LastFrameCycle[FrameType] = Cycle;
+	uint8 Buffer[9];
+	uint8* BufferPtr = Buffer;
+	FTraceUtils::Encode7bit(CycleDiff, BufferPtr);
+	uint16 BufferSize = BufferPtr - Buffer;
+	if (FrameType == TraceFrameType_Game)
+	{
+		UE_TRACE_LOG(Misc, EndGameFrame, BufferSize)
+			<< EndGameFrame.Attachment(&Buffer, BufferSize);
+	}
+	else if (FrameType == TraceFrameType_Rendering)
+	{
+		UE_TRACE_LOG(Misc, EndRenderFrame, BufferSize)
+			<< EndRenderFrame.Attachment(&Buffer, BufferSize);
+	}
 }
 
 #endif
