Marked the events that report UE_LOG events' categories and format strings as important.

diff --git a/Engine/Source/Runtime/Core/Private/Logging/LogTrace.cpp b/Engine/Source/Runtime/Core/Private/Logging/LogTrace.cpp
--- a/Engine/Source/Runtime/Core/Private/Logging/LogTrace.cpp
+++ b/Engine/Source/Runtime/Core/Private/Logging/LogTrace.cpp
@@ -8,12 +8,12 @@
 #include "HAL/PlatformTime.h"
 #include "HAL/PlatformTLS.h"
 
-UE_TRACE_EVENT_BEGIN(Logging, LogCategory, Always)
+UE_TRACE_EVENT_BEGIN(Logging, LogCategory, Always|Important)
 	UE_TRACE_EVENT_FIELD(const void*, CategoryPointer)
 	UE_TRACE_EVENT_FIELD(uint8, DefaultVerbosity)
 UE_TRACE_EVENT_END()
 
-UE_TRACE_EVENT_BEGIN(Logging, LogMessageSpec, Always)
+UE_TRACE_EVENT_BEGIN(Logging, LogMessageSpec, Always|Important)
 	UE_TRACE_EVENT_FIELD(const void*, LogPoint)
 	UE_TRACE_EVENT_FIELD(const void*, CategoryPointer)
 	UE_TRACE_EVENT_FIELD(int32, Line)
