UnrealInsights:
Added 'Always' event flags for compiled out trace systems
Hide Stats module unless stats trace is compiled in
Hide loading tracks in timing profiler tab

diff --git a/Engine/Source/Developer/TraceInsights/Private/Insights/Widgets/STimingView.cpp b/Engine/Source/Developer/TraceInsights/Private/Insights/Widgets/STimingView.cpp
--- a/Engine/Source/Developer/TraceInsights/Private/Insights/Widgets/STimingView.cpp
+++ b/Engine/Source/Developer/TraceInsights/Private/Insights/Widgets/STimingView.cpp
@@ -170,7 +170,7 @@ void STimingView::Reset()
 {
 	bShowHideAllGpuTracks = true;
 	bShowHideAllCpuTracks = true;
-	bShowHideAllLoadingTracks = true;
+	bShowHideAllLoadingTracks = false;
 	bShowHideAllIoTracks = false;
 
 	//////////////////////////////////////////////////
diff --git a/Engine/Source/Developer/TraceServices/Private/TraceServicesModule.cpp b/Engine/Source/Developer/TraceServices/Private/TraceServicesModule.cpp
--- a/Engine/Source/Developer/TraceServices/Private/TraceServicesModule.cpp
+++ b/Engine/Source/Developer/TraceServices/Private/TraceServicesModule.cpp
@@ -9,6 +9,7 @@
 #include "Modules/TimingProfilerModule.h"
 #include "Modules/LoadTimeProfilerModule.h"
 #include "Modules/StatsModule.h"
+#include "Stats/StatsTrace.h"
 
 class FTraceServicesModule
 	: public ITraceServicesModule
@@ -64,12 +65,16 @@ void FTraceServicesModule::StartupModule()
 {
 	IModularFeatures::Get().RegisterModularFeature(Trace::ModuleFeatureName, &TimingProfilerModule);
 	IModularFeatures::Get().RegisterModularFeature(Trace::ModuleFeatureName, &LoadTimeProfilerModule);
+#if EXPERIMENTAL_STATSTRACE_ENABLED
 	IModularFeatures::Get().RegisterModularFeature(Trace::ModuleFeatureName, &StatsModule);
+#endif
 }
 
 void FTraceServicesModule::ShutdownModule()
 {
+#if EXPERIMENTAL_STATSTRACE_ENABLED
 	IModularFeatures::Get().UnregisterModularFeature(Trace::ModuleFeatureName, &StatsModule);
+#endif
 	IModularFeatures::Get().UnregisterModularFeature(Trace::ModuleFeatureName, &LoadTimeProfilerModule);
 	IModularFeatures::Get().UnregisterModularFeature(Trace::ModuleFeatureName, &TimingProfilerModule);
 }
diff --git a/Engine/Source/Runtime/Core/Private/ProfilingDebugging/PlatformFileTrace.cpp b/Engine/Source/Runtime/Core/Private/ProfilingDebugging/PlatformFileTrace.cpp
--- a/Engine/Source/Runtime/Core/Private/ProfilingDebugging/PlatformFileTrace.cpp
+++ b/Engine/Source/Runtime/Core/Private/ProfilingDebugging/PlatformFileTrace.cpp
@@ -3,36 +3,35 @@
 
 #if PLATFORMFILETRACE_ENABLED
 
-#include "Trace/Trace.h"
 #include "Misc/CString.h"
 #include "HAL/PlatformTime.h"
 #include "HAL/PlatformTLS.h"
 
-UE_TRACE_EVENT_BEGIN(PlatformFile, BeginOpen)
+UE_TRACE_EVENT_BEGIN(PlatformFile, BeginOpen, Always)
 	UE_TRACE_EVENT_FIELD(uint64, Cycle)
 	UE_TRACE_EVENT_FIELD(uint64, TempHandle)
 	UE_TRACE_EVENT_FIELD(uint32, ThreadId)
 UE_TRACE_EVENT_END()
 
-UE_TRACE_EVENT_BEGIN(PlatformFile, EndOpen)
+UE_TRACE_EVENT_BEGIN(PlatformFile, EndOpen, Always)
 	UE_TRACE_EVENT_FIELD(uint64, Cycle)
 	UE_TRACE_EVENT_FIELD(uint64, TempHandle)
 	UE_TRACE_EVENT_FIELD(uint64, FileHandle)
 UE_TRACE_EVENT_END()
 
-UE_TRACE_EVENT_BEGIN(PlatformFile, BeginClose)
+UE_TRACE_EVENT_BEGIN(PlatformFile, BeginClose, Always)
 	UE_TRACE_EVENT_FIELD(uint64, Cycle)
 	UE_TRACE_EVENT_FIELD(uint64, TempHandle)
 	UE_TRACE_EVENT_FIELD(uint64, FileHandle)
 	UE_TRACE_EVENT_FIELD(uint32, ThreadId)
 UE_TRACE_EVENT_END()
 
-UE_TRACE_EVENT_BEGIN(PlatformFile, EndClose)
+UE_TRACE_EVENT_BEGIN(PlatformFile, EndClose, Always)
 	UE_TRACE_EVENT_FIELD(uint64, Cycle)
 	UE_TRACE_EVENT_FIELD(uint64, TempHandle)
 UE_TRACE_EVENT_END()
 
-UE_TRACE_EVENT_BEGIN(PlatformFile, BeginRead)
+UE_TRACE_EVENT_BEGIN(PlatformFile, BeginRead, Always)
 	UE_TRACE_EVENT_FIELD(uint64, Cycle)
 	UE_TRACE_EVENT_FIELD(uint64, ReadHandle)
 	UE_TRACE_EVENT_FIELD(uint64, FileHandle)
@@ -41,14 +40,14 @@ UE_TRACE_EVENT_BEGIN(PlatformFile, BeginRead)
 	UE_TRACE_EVENT_FIELD(uint32, ThreadId)
 UE_TRACE_EVENT_END()
 
-UE_TRACE_EVENT_BEGIN(PlatformFile, EndRead)
+UE_TRACE_EVENT_BEGIN(PlatformFile, EndRead, Always)
 	UE_TRACE_EVENT_FIELD(uint64, Cycle)
 	UE_TRACE_EVENT_FIELD(uint64, ReadHandle)
 	UE_TRACE_EVENT_FIELD(uint64, SizeRead)
 	UE_TRACE_EVENT_FIELD(uint32, ThreadId)
 UE_TRACE_EVENT_END()
 
-UE_TRACE_EVENT_BEGIN(PlatformFile, BeginWrite)
+UE_TRACE_EVENT_BEGIN(PlatformFile, BeginWrite, Always)
 	UE_TRACE_EVENT_FIELD(uint64, Cycle)
 	UE_TRACE_EVENT_FIELD(uint64, WriteHandle)
 	UE_TRACE_EVENT_FIELD(uint64, FileHandle)
@@ -57,7 +56,7 @@ UE_TRACE_EVENT_BEGIN(PlatformFile, BeginWrite)
 	UE_TRACE_EVENT_FIELD(uint32, ThreadId)
 UE_TRACE_EVENT_END()
 
-UE_TRACE_EVENT_BEGIN(PlatformFile, EndWrite)
+UE_TRACE_EVENT_BEGIN(PlatformFile, EndWrite, Always)
 	UE_TRACE_EVENT_FIELD(uint64, Cycle)
 	UE_TRACE_EVENT_FIELD(uint64, WriteHandle)
 	UE_TRACE_EVENT_FIELD(uint64, SizeWritten)
diff --git a/Engine/Source/Runtime/Core/Private/Stats/StatsTrace.cpp b/Engine/Source/Runtime/Core/Private/Stats/StatsTrace.cpp
--- a/Engine/Source/Runtime/Core/Private/Stats/StatsTrace.cpp
+++ b/Engine/Source/Runtime/Core/Private/Stats/StatsTrace.cpp
@@ -1,7 +1,6 @@
 // Copyright 1998-2019 Epic Games, Inc. All Rights Reserved.
 #include "Stats/StatsTrace.h"
 #include "ProfilingDebugging/MiscTrace.h"
-#include "Trace/Trace.h"
 #include "Templates/Atomic.h"
 #include "Misc/CString.h"
 #include "HAL/PlatformTime.h"
@@ -10,14 +9,14 @@
 
 #if STATSTRACE_ENABLED
 
-UE_TRACE_EVENT_BEGIN(Stats, Spec)
+UE_TRACE_EVENT_BEGIN(Stats, Spec, Always)
 	UE_TRACE_EVENT_FIELD(uint32, Id)
 	UE_TRACE_EVENT_FIELD(bool, IsFloatingPoint)
 	UE_TRACE_EVENT_FIELD(bool, IsMemory)
 	UE_TRACE_EVENT_FIELD(bool, ShouldClearEveryFrame)
 UE_TRACE_EVENT_END()
 
-UE_TRACE_EVENT_BEGIN(Stats, EventBatch)
+UE_TRACE_EVENT_BEGIN(Stats, EventBatch, Always)
 	UE_TRACE_EVENT_FIELD(uint32, ThreadId)
 UE_TRACE_EVENT_END()
 
diff --git a/Engine/Source/Runtime/Core/Public/ProfilingDebugging/MiscTrace.h b/Engine/Source/Runtime/Core/Public/ProfilingDebugging/MiscTrace.h
--- a/Engine/Source/Runtime/Core/Public/ProfilingDebugging/MiscTrace.h
+++ b/Engine/Source/Runtime/Core/Public/ProfilingDebugging/MiscTrace.h
@@ -3,8 +3,9 @@
 #pragma once
 
 #include "CoreTypes.h"
+#include "Trace/Trace.h"
 
-#if !IS_PROGRAM && !UE_BUILD_SHIPPING
+#if UE_TRACE_ENABLED && !UE_BUILD_SHIPPING
 #define MISCTRACE_ENABLED 1
 #else
 #define MISCTRACE_ENABLED 0
diff --git a/Engine/Source/Runtime/Core/Public/ProfilingDebugging/PlatformFileTrace.h b/Engine/Source/Runtime/Core/Public/ProfilingDebugging/PlatformFileTrace.h
--- a/Engine/Source/Runtime/Core/Public/ProfilingDebugging/PlatformFileTrace.h
+++ b/Engine/Source/Runtime/Core/Public/ProfilingDebugging/PlatformFileTrace.h
@@ -3,8 +3,9 @@
 #pragma once
 
 #include "CoreTypes.h"
+#include "Trace/Trace.h"
 
-#if !IS_PROGRAM && !UE_BUILD_SHIPPING
+#if UE_TRACE_ENABLED && !UE_BUILD_SHIPPING
 #define PLATFORMFILETRACE_ENABLED 0
 #else
 #define PLATFORMFILETRACE_ENABLED 0
diff --git a/Engine/Source/Runtime/Core/Public/Stats/StatsTrace.h b/Engine/Source/Runtime/Core/Public/Stats/StatsTrace.h
--- a/Engine/Source/Runtime/Core/Public/Stats/StatsTrace.h
+++ b/Engine/Source/Runtime/Core/Public/Stats/StatsTrace.h
@@ -3,9 +3,12 @@
 #pragma once
 
 #include "CoreTypes.h"
+#include "Trace/Trace.h"
 
-#if STATS && !IS_PROGRAM && !UE_BUILD_SHIPPING
-#define STATSTRACE_ENABLED 0
+#define EXPERIMENTAL_STATSTRACE_ENABLED 0
+
+#if UE_TRACE_ENABLED && STATS && EXPERIMENTAL_STATSTRACE_ENABLED
+#define STATSTRACE_ENABLED 1
 #else
 #define STATSTRACE_ENABLED 0
 #endif
diff --git a/Engine/Source/Runtime/RHI/Private/GpuProfilerTrace.cpp b/Engine/Source/Runtime/RHI/Private/GpuProfilerTrace.cpp
--- a/Engine/Source/Runtime/RHI/Private/GpuProfilerTrace.cpp
+++ b/Engine/Source/Runtime/RHI/Private/GpuProfilerTrace.cpp
@@ -8,7 +8,7 @@
 
 #if GPUPROFILERTRACE_ENABLED
 
-UE_TRACE_EVENT_BEGIN(GpuProfiler, EventSpec)
+UE_TRACE_EVENT_BEGIN(GpuProfiler, EventSpec, Always)
 	UE_TRACE_EVENT_FIELD(const void*, EventType)
 	UE_TRACE_EVENT_FIELD(uint16, NameLength)
 UE_TRACE_EVENT_END()
diff --git a/Engine/Source/Runtime/RHI/Public/GpuProfilerTrace.h b/Engine/Source/Runtime/RHI/Public/GpuProfilerTrace.h
--- a/Engine/Source/Runtime/RHI/Public/GpuProfilerTrace.h
+++ b/Engine/Source/Runtime/RHI/Public/GpuProfilerTrace.h
@@ -3,8 +3,9 @@
 #pragma once
 
 #include "CoreTypes.h"
+#include "Trace/Trace.h"
 
-#if !IS_PROGRAM && !UE_BUILD_SHIPPING
+#if UE_TRACE_ENABLED && !UE_BUILD_SHIPPING
 #define GPUPROFILERTRACE_ENABLED 0
 #else
 #define GPUPROFILERTRACE_ENABLED 0
