Fixed crash due to unaligned write on ARM (Android) -- it's possible using __packed would give better performance here but for now memcpy will have to do

diff --git a/Engine/Source/Runtime/Core/Public/ProfilingDebugging/FormatArgsTrace.h b/Engine/Source/Runtime/Core/Public/ProfilingDebugging/FormatArgsTrace.h
--- a/Engine/Source/Runtime/Core/Public/ProfilingDebugging/FormatArgsTrace.h
+++ b/Engine/Source/Runtime/Core/Public/ProfilingDebugging/FormatArgsTrace.h
@@ -78,7 +78,16 @@ private:
 	static typename TEnableIf<TAnd<TNot<TIsFloatingPoint<T>>, TNot<TIsStringArgument<T>>>::Value>::Type EncodeArgumentInternal(uint8*& TypeCodesPtr, uint8*& PayloadPtr, T Argument)
 	{
 		*TypeCodesPtr++ = FormatArgTypeCode_CategoryInteger | sizeof(T);
+
+#if PLATFORM_SUPPORTS_UNALIGNED_LOADS
 		*reinterpret_cast<T*>(PayloadPtr) = Argument;
+#else
+		// For ARM targets, it's possible that using __packed here would be preferable
+		// but I have not checked the codegen -- it's possible that the compiler generates
+		// the same code for this fixed size memcpy
+		memcpy(PayloadPtr, &Argument, sizeof Argument);
+#endif
+
 		PayloadPtr += sizeof(T);
 	}
 
@@ -86,7 +95,16 @@ private:
 	static typename TEnableIf<TIsFloatingPoint<T>::Value>::Type EncodeArgumentInternal(uint8*& TypeCodesPtr, uint8*& PayloadPtr, T Argument)
 	{
 		*TypeCodesPtr++ = FormatArgTypeCode_CategoryFloatingPoint | sizeof(T);
+
+#if PLATFORM_SUPPORTS_UNALIGNED_LOADS
 		*reinterpret_cast<T*>(PayloadPtr) = Argument;
+#else
+		// For ARM targets, it's possible that using __packed here would be preferable
+		// but I have not checked the codegen -- it's possible that the compiler generates
+		// the same code for this fixed size memcpy
+		memcpy(PayloadPtr, &Argument, sizeof Argument);
+#endif
+
 		PayloadPtr += sizeof(T);
 	}
 
