Fixed logic error in Core.Build.cs which meant stomp was disabled everywhere but on Mac (introduced in CL6631504, identified by CM)

diff --git a/Engine/Source/Runtime/Core/Core.Build.cs b/Engine/Source/Runtime/Core/Core.Build.cs
--- a/Engine/Source/Runtime/Core/Core.Build.cs
+++ b/Engine/Source/Runtime/Core/Core.Build.cs
@@ -199,7 +199,7 @@ public class Core : ModuleRules
         bool bWithMallocStomp = false;
         if (Target.Configuration != UnrealTargetConfiguration.Shipping)
         {
-			if (Target.Platform == UnrealTargetPlatform.Mac || Target.Platform == UnrealTargetPlatform.Mac || Target.Platform == UnrealTargetPlatform.Mac)
+			if (Target.Platform == UnrealTargetPlatform.Mac || Target.Platform == UnrealTargetPlatform.Linux || Target.Platform == UnrealTargetPlatform.Win64)
 			// Target.Platform == UnrealTargetPlatform.Win32: // 32-bit windows can technically be supported, but will likely run out of virtual memory space quickly
 			// Target.Platform == UnrealTargetPlatform.XboxOne: // XboxOne could be supported, as it's similar enough to Win64
 			{
