Fixed incorrect define for enabling tracing of log messages

diff --git a/Engine/Source/Runtime/Core/Public/Logging/LogTrace.h b/Engine/Source/Runtime/Core/Public/Logging/LogTrace.h
--- a/Engine/Source/Runtime/Core/Public/Logging/LogTrace.h
+++ b/Engine/Source/Runtime/Core/Public/Logging/LogTrace.h
@@ -3,8 +3,9 @@
 #pragma once
 
 #include "CoreTypes.h"
+#include "Trace/Trace.h"
 
-#if !IS_PROGRAM && !UE_BUILD_SHIPPING
+#if UE_TRACE_ENABLED && !UE_BUILD_SHIPPING
 #define LOGTRACE_ENABLED 1
 #else
 #define LOGTRACE_ENABLED 0
