Fixes to copy+pasted comments.
Fix to currently-unused CEqualityComparable concept, but which is the replacement for the THasOperatorEquals trait.

diff --git a/Engine/Source/Runtime/Core/Public/Concepts/EqualityComparable.h b/Engine/Source/Runtime/Core/Public/Concepts/EqualityComparable.h
--- a/Engine/Source/Runtime/Core/Public/Concepts/EqualityComparable.h
+++ b/Engine/Source/Runtime/Core/Public/Concepts/EqualityComparable.h
@@ -2,14 +2,12 @@
 
 #pragma once
 
-#include "CoreTypes.h"
-
 /**
  * Describes a type comparable with another type.
  */
 struct CEqualityComparableWith {
-	template <typename T>
-	auto Requires(bool& Result, const T& A, const T& B) -> decltype(
+	template <typename T, typename U>
+	auto Requires(bool& Result, const T& A, const U& B) -> decltype(
 		Result = A == B,
 		Result = B == A,
 		Result = A != B,
diff --git a/Engine/Source/Runtime/Core/Public/Concepts/Insertable.h b/Engine/Source/Runtime/Core/Public/Concepts/Insertable.h
--- a/Engine/Source/Runtime/Core/Public/Concepts/Insertable.h
+++ b/Engine/Source/Runtime/Core/Public/Concepts/Insertable.h
@@ -2,10 +2,8 @@
 
 #pragma once
 
-#include "CoreTypes.h"
-
 /**
- * Describes a type with a GetTypeHash overload.
+ * Describes an insertion operation for a destination type where an instance of another type can be inserted via operator<<.
  */
 template <typename DestType>
 struct CInsertable {
