Second attempt at fixing incorrect define for enabling tracing of log messages

diff --git a/Engine/Source/Runtime/Core/Public/Logging/LogTrace.h b/Engine/Source/Runtime/Core/Public/Logging/LogTrace.h
--- a/Engine/Source/Runtime/Core/Public/Logging/LogTrace.h
+++ b/Engine/Source/Runtime/Core/Public/Logging/LogTrace.h
@@ -3,9 +3,8 @@
 #pragma once
 
 #include "CoreTypes.h"
-#include "Trace/Trace.h"
 
-#if UE_TRACE_ENABLED && !UE_BUILD_SHIPPING
+#if !IS_PROGRAM && !UE_BUILD_SHIPPING && (PLATFORM_WINDOWS || PLATFORM_PS4 || PLATFORM_XBOXONE)
 #define LOGTRACE_ENABLED 1
 #else
 #define LOGTRACE_ENABLED 0
