Common slot code factored out into a base class.

diff --git a/Engine/Source/Runtime/Core/Public/Serialization/StructuredArchive.h b/Engine/Source/Runtime/Core/Public/Serialization/StructuredArchive.h
--- a/Engine/Source/Runtime/Core/Public/Serialization/StructuredArchive.h
+++ b/Engine/Source/Runtime/Core/Public/Serialization/StructuredArchive.h
@@ -96,13 +96,42 @@
 	FStructuredArchive(const FStructuredArchive&) = delete;
 	FStructuredArchive& operator=(const FStructuredArchive&) = delete;
 
+private:
+	class CORE_API FSlotBase
+	{
+	public:
+		FORCEINLINE FArchive& GetUnderlyingArchive()
+		{
+			return Ar.GetUnderlyingArchive();
+		}
+
+	protected:
+		FStructuredArchive& Ar;
+#if WITH_TEXT_ARCHIVE_SUPPORT
+		const int32 Depth;
+		const int32 ElementId;
+
+		FORCEINLINE explicit FSlotBase(FStructuredArchive& InAr, int32 InDepth, int32 InElementId)
+			: Ar(InAr)
+			, Depth(InDepth)
+			, ElementId(InElementId)
+		{
+		}
+#else
+		FORCEINLINE FSlotBase(FStructuredArchive& InAr)
+			: Ar(InAr)
+		{
+		}
+#endif
+	};
+
 public:
 	/**
 	 * Contains a value in the archive; either a field or array/map element. A slot does not know it's name or location,
 	 * and can merely have a value serialized into it. That value may be a literal (eg. int, float) or compound object
 	 * (eg. object, array, map).
 	 */
-	class CORE_API FSlot
+	class CORE_API FSlot final : public FSlotBase
 	{
 	public:
 		FRecord EnterRecord();
@@ -137,11 +166,6 @@
 		void Serialize(TArray<uint8>& Data);
 		void Serialize(void* Data, uint64 DataSize);
 
-		FORCEINLINE FArchive& GetUnderlyingArchive()
-		{
-			return Ar.GetUnderlyingArchive();
-		}
-
 		FORCEINLINE bool IsFilled() const
 		{
 #if WITH_TEXT_ARCHIVE_SUPPORT
@@ -155,30 +179,14 @@
 		friend FStructuredArchive;
 		friend class FStructuredArchiveChildReader;
 
-		FStructuredArchive& Ar;
-#if WITH_TEXT_ARCHIVE_SUPPORT
-		const int32 Depth;
-		const int32 ElementId;
-
-		FORCEINLINE FSlot(FStructuredArchive& InAr, int32 InDepth, int32 InElementId)
-			: Ar(InAr)
-			, Depth(InDepth)
-			, ElementId(InElementId)
-		{
-		}
-#else
-		FORCEINLINE FSlot(FStructuredArchive& InAr)
-			: Ar(InAr)
-		{
-		}
-#endif
+		using FSlotBase::FSlotBase;
 	};
 
 	/**
 	 * Represents a record in the structured archive. An object contains slots that are identified by FArchiveName,
 	 * which may be compiled out with binary-only archives.
 	 */
-	class CORE_API FRecord
+	class CORE_API FRecord final : public FSlotBase
 	{
 	public:
 		FSlot EnterField(FArchiveFieldName Name);
@@ -198,38 +206,17 @@
 			return *this;
 		}
 
-		FORCEINLINE FArchive& GetUnderlyingArchive()
-		{
-			return Ar.GetUnderlyingArchive();
-		}
-
 	private:
 		friend FStructuredArchive;
 
-		FStructuredArchive& Ar;
-#if WITH_TEXT_ARCHIVE_SUPPORT
-		const int32 Depth;
-		const int32 ElementId;
-
-		FORCEINLINE FRecord(FStructuredArchive& InAr, int32 InDepth, int32 InElementId)
-			: Ar(InAr)
-			, Depth(InDepth)
-			, ElementId(InElementId)
-		{
-		}
-#else
-		FORCEINLINE FRecord(FStructuredArchive& InAr)
-			: Ar(InAr)
-		{
-		}
-#endif
+		using FSlotBase::FSlotBase;
 	};
 
 	/**
 	 * Represents an array in the structured archive. An object contains slots that are identified by a FArchiveFieldName,
 	 * which may be compiled out with binary-only archives.
 	 */
-	class CORE_API FArray
+	class CORE_API FArray final : public FSlotBase
 	{
 	public:
 		FSlot EnterElement();
@@ -241,37 +228,16 @@
 			return *this;
 		}
 
-		FORCEINLINE FArchive& GetUnderlyingArchive()
-		{
-			return Ar.GetUnderlyingArchive();
-		}
-
 	private:
 		friend FStructuredArchive;
 
-		FStructuredArchive& Ar;
-#if WITH_TEXT_ARCHIVE_SUPPORT
-		const int32 Depth;
-		const int32 ElementId;
-
-		FORCEINLINE FArray(FStructuredArchive& InAr, int32 InDepth, int32 InElementId)
-			: Ar(InAr)
-			, Depth(InDepth)
-			, ElementId(InElementId)
-		{
-		}
-#else
-		FORCEINLINE FArray(FStructuredArchive& InAr)
-			: Ar(InAr)
-		{
-		}
-#endif
+		using FSlotBase::FSlotBase;
 	};
 
 	/**
 	 * Represents an unsized sequence of slots in the structured archive (similar to an array, but without a known size).
 	 */
-	class CORE_API FStream
+	class CORE_API FStream final : public FSlotBase
 	{
 	public:
 		FSlot EnterElement();
@@ -283,68 +249,26 @@
 			return *this;
 		}
 
-		FORCEINLINE FArchive& GetUnderlyingArchive()
-		{
-			return Ar.GetUnderlyingArchive();
-		}
-
 	private:
 		friend FStructuredArchive;
 
-		FStructuredArchive& Ar;
-#if WITH_TEXT_ARCHIVE_SUPPORT
-		const int32 Depth;
-		const int32 ElementId;
-
-		FORCEINLINE FStream(FStructuredArchive& InAr, int32 InDepth, int32 InElementId)
-			: Ar(InAr)
-			, Depth(InDepth)
-			, ElementId(InElementId)
-		{
-		}
-#else
-		FORCEINLINE FStream(FStructuredArchive& InAr)
-			: Ar(InAr)
-		{
-		}
-#endif
+		using FSlotBase::FSlotBase;
 	};
 
 	/**
 	 * Represents a map in the structured archive. A map is similar to a record, but keys can be read back out from an archive.
 	 * (This is an important distinction for binary archives).
 	 */
-	class CORE_API FMap
+	class CORE_API FMap final : public FSlotBase
 	{
 	public:
 		FSlot EnterElement(FString& Name);
 		FSlot EnterElement_TextOnly(FString& Name, EArchiveValueType& OutType);
 
-		FORCEINLINE FArchive& GetUnderlyingArchive()
-		{
-			return Ar.GetUnderlyingArchive();
-		}
-
 	private:
 		friend FStructuredArchive;
 
-		FStructuredArchive& Ar;
-#if WITH_TEXT_ARCHIVE_SUPPORT
-		const int32 Depth;
-		const int32 ElementId;
-
-		FORCEINLINE FMap(FStructuredArchive& InAr, int32 InDepth, int32 InElementId)
-			: Ar(InAr)
-			, Depth(InDepth)
-			, ElementId(InElementId)
-		{
-		}
-#else
-		FORCEINLINE FMap(FStructuredArchive& InAr)
-			: Ar(InAr)
-		{
-		}
-#endif
+		using FSlotBase::FSlotBase;
 	};
 
 private:
