Lots of Serialize-forwarding code replaced with a single template which forwards onto any overload the underlying formatter supports.

diff --git a/Engine/Source/Runtime/Core/Private/Serialization/StructuredArchive.cpp b/Engine/Source/Runtime/Core/Private/Serialization/StructuredArchive.cpp
--- a/Engine/Source/Runtime/Core/Private/Serialization/StructuredArchive.cpp
+++ b/Engine/Source/Runtime/Core/Private/Serialization/StructuredArchive.cpp
@@ -281,150 +281,6 @@ FStructuredArchive::FMap FStructuredArchive::FSlot::EnterMap(int32& Num)
 	return FMap(Ar, Depth + 1, ElementId);
 }
 
-void FStructuredArchive::FSlot::operator<< (char& Value)
-{
-	Ar.EnterSlot(ElementId);
-
-	int8 AsInt = Value;
-	Ar.Formatter.Serialize(AsInt);
-	Value = AsInt;
-
-	Ar.LeaveSlot();
-}
-
-void FStructuredArchive::FSlot::operator<< (uint8& Value)
-{
-	Ar.EnterSlot(ElementId);
-	Ar.Formatter.Serialize(Value);
-	Ar.LeaveSlot();
-}
-
-void FStructuredArchive::FSlot::operator<< (uint16& Value)
-{
-	Ar.EnterSlot(ElementId);
-	Ar.Formatter.Serialize(Value);
-	Ar.LeaveSlot();
-}
-
-void FStructuredArchive::FSlot::operator<< (uint32& Value)
-{
-	Ar.EnterSlot(ElementId);
-	Ar.Formatter.Serialize(Value);
-	Ar.LeaveSlot();
-}
-
-void FStructuredArchive::FSlot::operator<< (uint64& Value)
-{
-	Ar.EnterSlot(ElementId);
-	Ar.Formatter.Serialize(Value);
-	Ar.LeaveSlot();
-}
-
-void FStructuredArchive::FSlot::operator<< (int8& Value)
-{
-	Ar.EnterSlot(ElementId);
-	Ar.Formatter.Serialize(Value);
-	Ar.LeaveSlot();
-}
-
-void FStructuredArchive::FSlot::operator<< (int16& Value)
-{
-	Ar.EnterSlot(ElementId);
-	Ar.Formatter.Serialize(Value);
-	Ar.LeaveSlot();
-}
-
-void FStructuredArchive::FSlot::operator<< (int32& Value)
-{
-	Ar.EnterSlot(ElementId);
-	Ar.Formatter.Serialize(Value);
-	Ar.LeaveSlot();
-}
-
-void FStructuredArchive::FSlot::operator<< (int64& Value)
-{
-	Ar.EnterSlot(ElementId);
-	Ar.Formatter.Serialize(Value);
-	Ar.LeaveSlot();
-}
-
-void FStructuredArchive::FSlot::operator<< (float& Value)
-{
-	Ar.EnterSlot(ElementId);
-	Ar.Formatter.Serialize(Value);
-	Ar.LeaveSlot();
-}
-
-void FStructuredArchive::FSlot::operator<< (double& Value)
-{
-	Ar.EnterSlot(ElementId);
-	Ar.Formatter.Serialize(Value);
-	Ar.LeaveSlot();
-}
-
-void FStructuredArchive::FSlot::operator<< (bool& Value)
-{
-	Ar.EnterSlot(ElementId);
-	Ar.Formatter.Serialize(Value);
-	Ar.LeaveSlot();
-}
-
-void FStructuredArchive::FSlot::operator<< (FString& Value)
-{
-	Ar.EnterSlot(ElementId);
-	Ar.Formatter.Serialize(Value);
-	Ar.LeaveSlot();
-}
-
-void FStructuredArchive::FSlot::operator<< (FName& Value)
-{
-	Ar.EnterSlot(ElementId);
-	Ar.Formatter.Serialize(Value);
-	Ar.LeaveSlot();
-}
-
-void FStructuredArchive::FSlot::operator<< (UObject*& Value)
-{
-	Ar.EnterSlot(ElementId);
-	Ar.Formatter.Serialize(Value);
-	Ar.LeaveSlot();
-}
-
-void FStructuredArchive::FSlot::operator<< (FText& Value)
-{
-	Ar.EnterSlot(ElementId);
-	Ar.Formatter.Serialize(Value);
-	Ar.LeaveSlot();
-}
-
-void FStructuredArchive::FSlot::operator<< (FWeakObjectPtr& Value)
-{
-	Ar.EnterSlot(ElementId);
-	Ar.Formatter.Serialize(Value);
-	Ar.LeaveSlot();
-}
-
-void FStructuredArchive::FSlot::operator<< (FLazyObjectPtr& Value)
-{
-	Ar.EnterSlot(ElementId);
-	Ar.Formatter.Serialize(Value);
-	Ar.LeaveSlot();
-}
-
-void FStructuredArchive::FSlot::operator<< (FSoftObjectPtr& Value)
-{
-	Ar.EnterSlot(ElementId);
-	Ar.Formatter.Serialize(Value);
-	Ar.LeaveSlot();
-}
-
-void FStructuredArchive::FSlot::operator<< (FSoftObjectPath& Value)
-{
-	Ar.EnterSlot(ElementId);
-	Ar.Formatter.Serialize(Value);
-	Ar.LeaveSlot();
-}
-
 void FStructuredArchive::FSlot::Serialize(TArray<uint8>& Data)
 {
 	Ar.EnterSlot(ElementId);
diff --git a/Engine/Source/Runtime/Core/Public/Concepts/Serializable.h b/Engine/Source/Runtime/Core/Public/Concepts/Serializable.h
new file mode 100644
--- /dev/null
+++ b/Engine/Source/Runtime/Core/Public/Concepts/Serializable.h
@@ -0,0 +1,14 @@
+// Copyright 1998-2019 Epic Games, Inc. All Rights Reserved.
+
+#pragma once
+
+/**
+ * Describes a serializable type (with a Serialize member function) taking the given argument types.
+ */
+template <typename SerializableType>
+struct CSerializable {
+	template <typename... ArgTypes>
+	auto Requires(SerializableType& Obj, ArgTypes&&... Args) -> decltype(
+		Obj.Serialize(Forward<ArgTypes>(Args)...)
+	);
+};
diff --git a/Engine/Source/Runtime/Core/Public/Serialization/StructuredArchive.h b/Engine/Source/Runtime/Core/Public/Serialization/StructuredArchive.h
--- a/Engine/Source/Runtime/Core/Public/Serialization/StructuredArchive.h
+++ b/Engine/Source/Runtime/Core/Public/Serialization/StructuredArchive.h
@@ -6,6 +6,8 @@
 #include "Serialization/StructuredArchiveFormatter.h"
 #include "Formatters/BinaryArchiveFormatter.h"
 #include "Misc/Optional.h"
+#include "Concepts/Serializable.h"
+#include "Templates/Models.h"
 #include "Containers/Array.h"
 
 /**
@@ -127,7 +129,7 @@ private:
 
 public:
 	/**
-	 * Contains a value in the archive; either a field or array/map element. A slot does not know it's name or location,
+	 * Contains a value in the archive; either a field or array/map element. A slot does not know its name or location,
 	 * and can merely have a value serialized into it. That value may be a literal (eg. int, float) or compound object
 	 * (eg. object, array, map).
 	 */
@@ -141,27 +143,20 @@ public:
 		FStream EnterStream_TextOnly(int32& OutNumElements);
 		FMap EnterMap(int32& Num);
 
-		// We don't support chaining writes to a single slot, so these return void.
-		void operator << (char& Value);
-		void operator << (uint8& Value);
-		void operator << (uint16& Value);
-		void operator << (uint32& Value);
-		void operator << (uint64& Value);
-		void operator << (int8& Value);
-		void operator << (int16& Value);
-		void operator << (int32& Value);
-		void operator << (int64& Value);
-		void operator << (float& Value);
-		void operator << (double& Value);
-		void operator << (bool& Value);
-		void operator << (FString& Value);
-		void operator << (FName& Value);
-		void operator << (UObject*& Value);
-		void operator << (FText& Value);
-		void operator << (FWeakObjectPtr& Value);
-		void operator << (FSoftObjectPtr& Value);
-		void operator << (FSoftObjectPath& Value);
-		void operator << (FLazyObjectPtr& Value);
+		// We don't support chaining writes to a single slot, so this returns void.
+		template <typename ArgType>
+		typename TEnableIf<
+			TModels<CSerializable<FArchiveFormatterType>, ArgType>::Value
+		>::Type operator<<(ArgType&& Arg)
+		{
+#if WITH_TEXT_ARCHIVE_SUPPORT
+			Ar.EnterSlot(ElementId);
+			Ar.Formatter.Serialize(Forward<ArgType>(Arg));
+			Ar.LeaveSlot();
+#else
+			Ar.Formatter.Serialize(Forward<ArgType>(Arg));
+#endif
+		}
 
 		void Serialize(TArray<uint8>& Data);
 		void Serialize(void* Data, uint64 DataSize);
@@ -479,118 +474,6 @@ private:
 		return FMap(Ar);
 	}
 
-	FORCEINLINE void FStructuredArchive::FSlot::operator<< (char& Value)
-	{
-		int8 AsInt = Value;
-		Ar.Formatter.Serialize(AsInt);
-		Value = AsInt;
-	}
-
-	FORCEINLINE void FStructuredArchive::FSlot::operator<< (uint8& Value)
-	{
-		Ar.Formatter.Serialize(Value);
-	}
-
-	FORCEINLINE void FStructuredArchive::FSlot::operator<< (uint16& Value)
-	{
-		Ar.Formatter.Serialize(Value);
-	}
-
-	FORCEINLINE void FStructuredArchive::FSlot::operator<< (uint32& Value)
-	{
-		Ar.Formatter.Serialize(Value);
-	}
-
-	FORCEINLINE void FStructuredArchive::FSlot::operator<< (uint64& Value)
-	{
-		Ar.Formatter.Serialize(Value);
-	}
-
-	FORCEINLINE void FStructuredArchive::FSlot::operator<< (int8& Value)
-	{
-		Ar.Formatter.Serialize(Value);
-	}
-
-	FORCEINLINE void FStructuredArchive::FSlot::operator<< (int16& Value)
-	{
-		Ar.Formatter.Serialize(Value);
-	}
-
-	FORCEINLINE void FStructuredArchive::FSlot::operator<< (int32& Value)
-	{
-		Ar.Formatter.Serialize(Value);
-	}
-
-	FORCEINLINE void FStructuredArchive::FSlot::operator<< (int64& Value)
-	{
-		Ar.Formatter.Serialize(Value);
-	}
-
-	FORCEINLINE void FStructuredArchive::FSlot::operator<< (float& Value)
-	{
-		Ar.Formatter.Serialize(Value);
-	}
-
-	FORCEINLINE void FStructuredArchive::FSlot::operator<< (double& Value)
-	{
-		Ar.Formatter.Serialize(Value);
-	}
-
-	FORCEINLINE void FStructuredArchive::FSlot::operator<< (bool& Value)
-	{
-		Ar.Formatter.Serialize(Value);
-	}
-
-	FORCEINLINE void FStructuredArchive::FSlot::operator<< (FString& Value)
-	{
-		Ar.Formatter.Serialize(Value);
-	}
-
-	FORCEINLINE void FStructuredArchive::FSlot::operator<< (FName& Value)
-	{
-		Ar.Formatter.Serialize(Value);
-	}
-
-	FORCEINLINE void FStructuredArchive::FSlot::operator<< (UObject*& Value)
-	{
-		Ar.Formatter.Serialize(Value);
-	}
-
-	FORCEINLINE void FStructuredArchive::FSlot::operator<< (FText& Value)
-	{
-		Ar.Formatter.Serialize(Value);
-	}
-
-	FORCEINLINE void FStructuredArchive::FSlot::operator<< (FWeakObjectPtr& Value)
-	{
-		Ar.Formatter.Serialize(Value);
-	}
-
-	FORCEINLINE void FStructuredArchive::FSlot::operator<< (FSoftObjectPath& Value)
-	{
-		Ar.Formatter.Serialize(Value);
-	}
-
-	FORCEINLINE void FStructuredArchive::FSlot::operator<< (FSoftObjectPtr& Value)
-	{
-		Ar.Formatter.Serialize(Value);
-	}
-
-	FORCEINLINE void FStructuredArchive::FSlot::operator<< (FLazyObjectPtr& Value)
-	{
-		Ar.Formatter.Serialize(Value);
-	}
-
-	FORCEINLINE void FStructuredArchive::FSlot::Serialize(TArray<uint8>& Data)
-	{
-		Ar.Formatter.Serialize(Data);
-	}
-
-	FORCEINLINE void FStructuredArchive::FSlot::Serialize(void* Data, uint64 DataSize)
-	{
-		Ar.Formatter.Serialize(Data, DataSize);
-	}
-
 	//////////// FStructuredArchive::FRecord ////////////
 
 	FORCEINLINE FStructuredArchive::FSlot FStructuredArchive::FRecord::EnterField(FArchiveFieldName Name)
