When an assertion failed on Mac the process return code would be 0 which canonically incorrectly indicates there was no failure. Changed Mac to use a non-zero exit code when forcefully terminated which matches Windows' behaviour.

diff --git a/Engine/Source/Runtime/Core/Private/Mac/MacPlatformMisc.cpp b/Engine/Source/Runtime/Core/Private/Mac/MacPlatformMisc.cpp
--- a/Engine/Source/Runtime/Core/Private/Mac/MacPlatformMisc.cpp
+++ b/Engine/Source/Runtime/Core/Private/Mac/MacPlatformMisc.cpp
@@ -1627,7 +1627,7 @@ static void GracefulTerminationHandler(int32 Signal, siginfo_t* Info, void* Cont
 	}
 	else
 	{
-		_Exit(0);
+		_Exit(1);
 	}
 }
 
@@ -1959,7 +1959,7 @@ void FMacCrashContext::GenerateCrashInfoAndLaunchReporter() const
 		raise(Signal);
 	}
 	
-	_Exit(0);
+	_Exit(1);
 }
 
 void FMacCrashContext::GenerateEnsureInfoAndLaunchReporter() const
