Link error fix - some functions restored that never shouldn't have been deleted.

diff --git a/Engine/Source/Runtime/Core/Public/Serialization/StructuredArchive.h b/Engine/Source/Runtime/Core/Public/Serialization/StructuredArchive.h
--- a/Engine/Source/Runtime/Core/Public/Serialization/StructuredArchive.h
+++ b/Engine/Source/Runtime/Core/Public/Serialization/StructuredArchive.h
@@ -474,6 +474,16 @@ private:
 		return FMap(Ar);
 	}
 
+	FORCEINLINE void FStructuredArchive::FSlot::Serialize(TArray<uint8>& Data)
+	{
+		Ar.Formatter.Serialize(Data);
+	}
+
+	FORCEINLINE void FStructuredArchive::FSlot::Serialize(void* Data, uint64 DataSize)
+	{
+		Ar.Formatter.Serialize(Data, DataSize);
+	}
+
 	//////////// FStructuredArchive::FRecord ////////////
 
 	FORCEINLINE FStructuredArchive::FSlot FStructuredArchive::FRecord::EnterField(FArchiveFieldName Name)
