Build fix

diff --git a/Engine/Source/Runtime/Core/Public/Math/Color.h b/Engine/Source/Runtime/Core/Public/Math/Color.h
--- a/Engine/Source/Runtime/Core/Public/Math/Color.h
+++ b/Engine/Source/Runtime/Core/Public/Math/Color.h
@@ -8,6 +8,7 @@
 #include "Containers/UnrealString.h"
 #include "Misc/Parse.h"
 #include "Serialization/StructuredArchive.h"
+#include "Serialization/ArchiveAdapters.h"
 
 class FFloat16Color;
 
