More build fixes

diff --git a/Engine/Source/Runtime/Core/Public/Misc/DateTime.h b/Engine/Source/Runtime/Core/Public/Misc/DateTime.h
--- a/Engine/Source/Runtime/Core/Public/Misc/DateTime.h
+++ b/Engine/Source/Runtime/Core/Public/Misc/DateTime.h
@@ -6,6 +6,7 @@
 #include "Containers/UnrealString.h"
 #include "Misc/Timespan.h"
 #include "Templates/TypeHash.h"
+#include "Serialization/ArchiveAdapters.h"
 
 class FArchive;
 class FOutputDevice;
diff --git a/Engine/Source/Runtime/Core/Public/UObject/NameTypes.h b/Engine/Source/Runtime/Core/Public/UObject/NameTypes.h
--- a/Engine/Source/Runtime/Core/Public/UObject/NameTypes.h
+++ b/Engine/Source/Runtime/Core/Public/UObject/NameTypes.h
@@ -12,7 +12,7 @@
 #include "Containers/StringConv.h"
 #include "UObject/UnrealNames.h"
 #include "Templates/Atomic.h"
-#include "Serialization/StructuredArchive.h"
+#include "Serialization/ArchiveAdapters.h"
 
 /*----------------------------------------------------------------------------
 	Definitions.
