Even more build fixes

diff --git a/Engine/Source/Runtime/Core/Private/UObject/UnrealNames.cpp b/Engine/Source/Runtime/Core/Private/UObject/UnrealNames.cpp
--- a/Engine/Source/Runtime/Core/Private/UObject/UnrealNames.cpp
+++ b/Engine/Source/Runtime/Core/Private/UObject/UnrealNames.cpp
@@ -22,7 +22,6 @@
 #include "Misc/OutputDeviceRedirector.h"
 #include "HAL/IConsoleManager.h"
 #include "HAL/LowLevelMemTracker.h"
-#include "Serialization/ArchiveFromStructuredArchive.h"
 #include "Hash/CityHash.h"
 
 DEFINE_LOG_CATEGORY_STATIC(LogUnrealNames, Log, All);
@@ -2087,16 +2086,6 @@ void FNameEntry::Write( FArchive& Ar ) const
 	Ar << EntrySerialized;
 }
 
-void FNameEntry::Write(FStructuredArchive::FSlot Slot) const
-{
-	// This path should be unused - since FNameEntry structs are allocated with a dynamic size, we can only save them. Use FNameEntrySerialized to read them back into an intermediate buffer.
-	checkf(!Slot.GetUnderlyingArchive().IsLoading(), TEXT("FNameEntry does not support reading from an archive. Serialize into a FNameEntrySerialized and construct a FNameEntry from that."));
-
-	// Convert to our serialized type
-	FNameEntrySerialized EntrySerialized(*this);
-	Slot << EntrySerialized;
-}
-
 FArchive& operator<<(FArchive& Ar, FNameEntrySerialized& E)
 {
 	if (Ar.IsLoading())
@@ -2180,31 +2169,6 @@ FArchive& operator<<(FArchive& Ar, FNameEntrySerialized& E)
 	return Ar;
 }
 
-void operator<<(FStructuredArchive::FSlot Slot, FNameEntrySerialized& E)
-{
-	if (Slot.GetUnderlyingArchive().IsTextFormat())
-	{
-		FString Str = E.GetPlainNameString();
-		Slot << Str;
-
-		if (Slot.GetUnderlyingArchive().IsLoading())
-		{
-			// mark the name will be wide
-			E.bIsWide = true;
-
-			// get the pointer to the wide array 
-			WIDECHAR* WideName = const_cast<WIDECHAR*>(E.GetWideName());
-			FCString::Strcpy(WideName, 1024, *Str);
-		}
-	}
-	else
-	{
-		FArchiveFromStructuredArchive Ar(Slot);
-		Ar << E;
-	}
-}
-
-
 #if !UE_BUILD_SHIPPING && !UE_BUILD_TEST
 
 #include "Containers/StackTracker.h"
diff --git a/Engine/Source/Runtime/Core/Public/Math/IntPoint.h b/Engine/Source/Runtime/Core/Public/Math/IntPoint.h
--- a/Engine/Source/Runtime/Core/Public/Math/IntPoint.h
+++ b/Engine/Source/Runtime/Core/Public/Math/IntPoint.h
@@ -7,6 +7,7 @@
 #include "Math/UnrealMathUtility.h"
 #include "Containers/UnrealString.h"
 #include "Templates/TypeHash.h"
+#include "Serialization/ArchiveAdapters.h"
 
 /**
  * Structure for integer points in 2-d space.
diff --git a/Engine/Source/Runtime/Core/Public/UObject/NameTypes.h b/Engine/Source/Runtime/Core/Public/UObject/NameTypes.h
--- a/Engine/Source/Runtime/Core/Public/UObject/NameTypes.h
+++ b/Engine/Source/Runtime/Core/Public/UObject/NameTypes.h
@@ -12,7 +12,6 @@
 #include "Containers/StringConv.h"
 #include "UObject/UnrealNames.h"
 #include "Templates/Atomic.h"
-#include "Serialization/ArchiveAdapters.h"
 
 /*----------------------------------------------------------------------------
 	Definitions.
@@ -236,7 +235,6 @@ public:
 	CORE_API int32 GetSizeInBytes() const;
 
 	CORE_API void Write(FArchive& Ar) const;
-	CORE_API void Write(FStructuredArchive::FSlot Slot) const;
 
 	static int32 GetDataOffset();
 	struct FNameStringView MakeView(union FNameBuffer& OptionalDecodeBuffer) const;
@@ -309,12 +307,6 @@ struct FNameEntrySerialized
 	{
 		return Ar << *E;
 	}
-
-	friend CORE_API void operator<<(FStructuredArchive::FSlot Slot, FNameEntrySerialized& E);
-	friend void operator<<(FStructuredArchive::FSlot Slot, FNameEntrySerialized* E)
-	{
-		Slot << *E;
-	}
 };
 
 /**
diff --git a/Engine/Source/Runtime/CoreUObject/Private/UObject/SavePackage.cpp b/Engine/Source/Runtime/CoreUObject/Private/UObject/SavePackage.cpp
--- a/Engine/Source/Runtime/CoreUObject/Private/UObject/SavePackage.cpp
+++ b/Engine/Source/Runtime/CoreUObject/Private/UObject/SavePackage.cpp
@@ -4572,7 +4572,6 @@ FSavePackageResultStruct UPackage::Save(UPackage* InOuter, UObject* Base, EObjec
 				// Save names.
 				if (!bTextFormat)
 				{
-					FStructuredArchive::FStream NameStream = StructuredArchiveRoot.EnterStream(FIELD_NAME_TEXT("Names"));
 #if WITH_EDITOR
 					FArchive::FScopeSetDebugSerializationFlags S(*Linker, DSF_IgnoreDiff, true);
 					FArchiveStackTraceIgnoreScope IgnoreSummaryDiffsScope(DiffSettings.bIgnoreHeaderDiffs);
@@ -4580,7 +4579,7 @@ FSavePackageResultStruct UPackage::Save(UPackage* InOuter, UObject* Base, EObjec
 					Linker->Summary.NameCount = Linker->NameMap.Num();
 					for (int32 i = 0; i < Linker->NameMap.Num(); i++)
 					{
-						FName::GetEntry(Linker->NameMap[i])->Write(NameStream.EnterElement());
+						FName::GetEntry(Linker->NameMap[i])->Write(*Linker);
 						Linker->NameIndices.Add(Linker->NameMap[i], i);
 					}
 				}
