Fix FName crash when constructing from WIDECHAR / TCHAR nullptr

diff --git a/Engine/Source/Runtime/Core/Private/UObject/UnrealNames.cpp b/Engine/Source/Runtime/Core/Private/UObject/UnrealNames.cpp
--- a/Engine/Source/Runtime/Core/Private/UObject/UnrealNames.cpp
+++ b/Engine/Source/Runtime/Core/Private/UObject/UnrealNames.cpp
@@ -1442,10 +1442,13 @@ static int32 GetLengthAndWidth(const WIDECHAR* Str, bool& bOutIsWide)
 {
 	uint32 UserCharBits = 0;
 	const WIDECHAR* It = Str;
-	while (*It)
+	if (Str)
 	{
-		UserCharBits |= TChar<WIDECHAR>::ToUnsigned(*It);
-		++It;
+		while (*It)
+		{
+			UserCharBits |= TChar<WIDECHAR>::ToUnsigned(*It);
+			++It;
+		}
 	}
 
 	bOutIsWide = UserCharBits & 0xffffff80u;
@@ -1935,6 +1938,7 @@ void FName::AutoTest()
 
 	const FName NullName(static_cast<ANSICHAR*>(nullptr));
 	check(NullName.IsNone());
+	check(NullName == FName(static_cast<WIDECHAR*>(nullptr)));
 	check(NullName == FName(NAME_None));
 	check(NullName == FName());
 	check(NullName == FName(""));
