Merged down disabled FNAME_Replace_Not_Safe_For_Threading tests when page protection is on

diff --git a/Engine/Source/Runtime/Core/Private/UObject/UnrealNames.cpp b/Engine/Source/Runtime/Core/Private/UObject/UnrealNames.cpp
--- a/Engine/Source/Runtime/Core/Private/UObject/UnrealNames.cpp
+++ b/Engine/Source/Runtime/Core/Private/UObject/UnrealNames.cpp
@@ -1996,6 +1996,7 @@ void FName::AutoTest()
 		check(FName(TEXT("UNIQUEUNICORN!!"), FNAME_Find) == UniqueName);
 		check(FName("uniqueunicorn!!", FNAME_Find) == UniqueName);
 
+#if !FNAME_WRITE_PROTECT_PAGES
 		// Check FNAME_Replace_Not_Safe_For_Threading updates casing
 		check(0 != UniqueName.GetPlainNameString().Compare("UNIQUEunicorn!!", ESearchCase::CaseSensitive));
 		const FName UniqueNameReplaced("UNIQUEunicorn!!", FNAME_Replace_Not_Safe_For_Threading);
@@ -2010,6 +2011,7 @@ void FName::AutoTest()
 		// Check FNAME_Replace_Not_Safe_For_Threading adds entries that do not exist
 		const FName AddedByReplace("WasAdded!!", FNAME_Replace_Not_Safe_For_Threading);
 		check(FName("WasAdded!!", FNAME_Find) == AddedByReplace);
+#endif
 	
 		Once = false;
 	}
