noMaster compile fix

diff --git a/Engine/Source/Runtime/Core/Private/Serialization/LoadTimeTrace.cpp b/Engine/Source/Runtime/Core/Private/Serialization/LoadTimeTrace.cpp
--- a/Engine/Source/Runtime/Core/Private/Serialization/LoadTimeTrace.cpp
+++ b/Engine/Source/Runtime/Core/Private/Serialization/LoadTimeTrace.cpp
@@ -6,6 +6,7 @@
 
 #include "Trace/Trace.h"
 #include "Misc/CString.h"
+#include "HAL/PlatformTLS.h"
 
 UE_TRACE_EVENT_BEGIN(LoadTime, BeginRequestGroup)
 	UE_TRACE_EVENT_FIELD(uint32, ThreadId)
