Move adapters into StructuredArchive.h so that they are globally visible.

diff --git a/Engine/Source/Programs/UnrealHeaderTool/Private/CodeGenerator.cpp b/Engine/Source/Programs/UnrealHeaderTool/Private/CodeGenerator.cpp
--- a/Engine/Source/Programs/UnrealHeaderTool/Private/CodeGenerator.cpp
+++ b/Engine/Source/Programs/UnrealHeaderTool/Private/CodeGenerator.cpp
@@ -5235,7 +5235,7 @@ FNativeClassHeaderGenerator::FNativeClassHeaderGenerator(
 			{
 				if (!bAddedStructuredArchiveFromArchiveHeader && !bHasSerializeToFArchive)
 				{
-					RelativeIncludes.AddUnique(TEXT("Serialization/StructuredArchiveFromArchive.h"));
+					RelativeIncludes.AddUnique(TEXT("Serialization/StructuredArchive.h"));
 					bAddedStructuredArchiveFromArchiveHeader = true;
 				}
 
diff --git a/Engine/Source/Runtime/Core/Private/Internationalization/GatherableTextData.cpp b/Engine/Source/Runtime/Core/Private/Internationalization/GatherableTextData.cpp
--- a/Engine/Source/Runtime/Core/Private/Internationalization/GatherableTextData.cpp
+++ b/Engine/Source/Runtime/Core/Private/Internationalization/GatherableTextData.cpp
@@ -1,7 +1,7 @@
 // Copyright 1998-2019 Epic Games, Inc. All Rights Reserved.
 
 #include "Internationalization/GatherableTextData.h"
-#include "Serialization/StructuredArchiveFromArchive.h"
+#include "Serialization/StructuredArchive.h"
 
 FArchive& operator<<(FArchive& Archive, FTextSourceSiteContext& This)
 {
diff --git a/Engine/Source/Runtime/Core/Private/Internationalization/InternationalizationMetadata.cpp b/Engine/Source/Runtime/Core/Private/Internationalization/InternationalizationMetadata.cpp
--- a/Engine/Source/Runtime/Core/Private/Internationalization/InternationalizationMetadata.cpp
+++ b/Engine/Source/Runtime/Core/Private/Internationalization/InternationalizationMetadata.cpp
@@ -1,7 +1,7 @@
 // Copyright 1998-2019 Epic Games, Inc. All Rights Reserved.
 
 #include "Internationalization/InternationalizationMetadata.h"
-#include "Serialization/StructuredArchiveFromArchive.h"
+#include "Serialization/StructuredArchive.h"
 
 DEFINE_LOG_CATEGORY_STATIC(LogInternationalizationMetadata, Log, All);
 
diff --git a/Engine/Source/Runtime/Core/Private/Internationalization/Text.cpp b/Engine/Source/Runtime/Core/Private/Internationalization/Text.cpp
--- a/Engine/Source/Runtime/Core/Private/Internationalization/Text.cpp
+++ b/Engine/Source/Runtime/Core/Private/Internationalization/Text.cpp
@@ -19,7 +19,7 @@
 #include "Internationalization/FastDecimalFormat.h"
 #include "Internationalization/TextGeneratorRegistry.h"
 
-#include "Serialization/ArchiveFromStructuredArchive.h"
+#include "Serialization/StructuredArchive.h"
 
 #include "UObject/EditorObjectVersion.h"
 #include "HAL/PlatformProcess.h"
diff --git a/Engine/Source/Runtime/Core/Private/Misc/ObjectThumbnail.cpp b/Engine/Source/Runtime/Core/Private/Misc/ObjectThumbnail.cpp
--- a/Engine/Source/Runtime/Core/Private/Misc/ObjectThumbnail.cpp
+++ b/Engine/Source/Runtime/Core/Private/Misc/ObjectThumbnail.cpp
@@ -1,7 +1,7 @@
 // Copyright 1998-2019 Epic Games, Inc. All Rights Reserved.
 
 #include "Misc/ObjectThumbnail.h"
-#include "Serialization/StructuredArchiveFromArchive.h"
+#include "Serialization/StructuredArchive.h"
 
 /** Static: Thumbnail compressor */
 FThumbnailCompressionInterface* FObjectThumbnail::ThumbnailCompressor = NULL;
diff --git a/Engine/Source/Runtime/Core/Private/Serialization/ArchiveFromStructuredArchive.cpp b/Engine/Source/Runtime/Core/Private/Serialization/ArchiveFromStructuredArchive.cpp
--- a/Engine/Source/Runtime/Core/Private/Serialization/ArchiveFromStructuredArchive.cpp
+++ b/Engine/Source/Runtime/Core/Private/Serialization/ArchiveFromStructuredArchive.cpp
@@ -1,16 +1,41 @@
 // Copyright 1998-2019 Epic Games, Inc. All Rights Reserved.
 
-#include "Serialization/ArchiveFromStructuredArchive.h"
+#include "Serialization/StructuredArchive.h"
 #include "Internationalization/Text.h"
 
 #if WITH_TEXT_ARCHIVE_SUPPORT
 
+struct FArchiveFromStructuredArchive::FImpl
+{
+	explicit FImpl(FStructuredArchive::FSlot Slot)
+		: RootSlot(Slot)
+	{
+	}
+
+	TOptional<FStructuredArchive::FRecord> Root;
+
+	static constexpr int32 MaxBufferSize = 128;
+
+	bool bPendingSerialize = true;
+	bool bWasOpened        = false;
+
+	TArray<uint8> Buffer;
+	int32 Pos = 0;
+
+	TArray<FName> Names;
+	TMap<FName, int32> NameToIndex;
+
+	TArray<FString> ObjectNames;
+	TArray<UObject*> Objects;
+	TBitArray<> ObjectsValid;
+	TMap<UObject*, int32> ObjectToIndex;
+
+	FStructuredArchive::FSlot RootSlot;
+};
+
 FArchiveFromStructuredArchive::FArchiveFromStructuredArchive(FStructuredArchive::FSlot Slot)
 	: FArchiveProxy(Slot.GetUnderlyingArchive())
-	, bPendingSerialize(true)
-	, bWasOpened(false)
-	, Pos(0)
-	, RootSlot(Slot)
+	, Pimpl(Slot)
 {
 	// For some reason, the FArchive copy constructor will copy all the trivial members of the source archive, but then specifically set ArIsFilterEditorOnly to false, with a comment saying
 	// they don't know why it's doing this... make sure we inherit this flag here!
@@ -39,7 +64,7 @@ int64 FArchiveFromStructuredArchive::Tell()
 {
 	if (InnerArchive.IsTextFormat())
 	{
-		return Pos;
+		return Pimpl->Pos;
 	}
 	else
 	{
@@ -57,8 +82,8 @@ void FArchiveFromStructuredArchive::Seek(int64 InPos)
 {
 	if (InnerArchive.IsTextFormat())
 	{
-		check(Pos >= 0 && Pos <= Buffer.Num());
-		Pos = InPos;
+		check(Pimpl->Pos >= 0 && Pimpl->Pos <= Pimpl->Buffer.Num());
+		Pimpl->Pos = InPos;
 	}
 	else
 	{
@@ -70,7 +95,7 @@ bool FArchiveFromStructuredArchive::AtEnd()
 {
 	if (InnerArchive.IsTextFormat())
 	{
-		return Pos == Buffer.Num();
+		return Pimpl->Pos == Pimpl->Buffer.Num();
 	}
 	else
 	{
@@ -88,15 +113,15 @@ FArchive& FArchiveFromStructuredArchive::operator<<(class FName& Value)
 		{
 			int32 NameIdx = 0;
 			Serialize(&NameIdx, sizeof(NameIdx));
-			Value = Names[NameIdx];
+			Value = Pimpl->Names[NameIdx];
 		}
 		else
 		{
-			int32* NameIdxPtr = NameToIndex.Find(Value);
+			int32* NameIdxPtr = Pimpl->NameToIndex.Find(Value);
 			if (NameIdxPtr == nullptr)
 			{
-				NameIdxPtr = &(NameToIndex.Add(Value));
-				*NameIdxPtr = Names.Add(Value);
+				NameIdxPtr = &Pimpl->NameToIndex.Add(Value);
+				*NameIdxPtr = Pimpl->Names.Add(Value);
 			}
 			Serialize(NameIdxPtr, sizeof(*NameIdxPtr));
 		}
@@ -120,13 +145,13 @@ FArchive& FArchiveFromStructuredArchive::operator<<(class UObject*& Value)
 			Serialize(&ObjectIdx, sizeof(ObjectIdx));
 
 			// If this object has already been accessed, return the cached value
-			if (ObjectsValid[ObjectIdx])
+			if (Pimpl->ObjectsValid[ObjectIdx])
 			{
-				Value = Objects[ObjectIdx];
+				Value = Pimpl->Objects[ObjectIdx];
 			}
 			else
 			{
-				FStructuredArchive::FStream Stream = Root->EnterStream(FIELD_NAME_TEXT("Objects"));
+				FStructuredArchive::FStream Stream = Pimpl->Root->EnterStream(FIELD_NAME_TEXT("Objects"));
 
 				// We know exactly which stream index we want to load here, but because of the API we need to read through them
 				// in order, consuming the string name until we reach the entry we want and then load it as a uobject reference.
@@ -138,12 +163,12 @@ FArchive& FArchiveFromStructuredArchive::operator<<(class UObject*& Value)
 				// pointer.
 
 				FString Dummy;
-				for (int32 Index = 0; Index < Objects.Num(); ++Index)
+				for (int32 Index = 0; Index < Pimpl->Objects.Num(); ++Index)
 				{
 					if (Index == ObjectIdx)
 					{
 						Stream.EnterElement() << Value;
-						Objects[ObjectIdx] = Value;
+						Pimpl->Objects[ObjectIdx] = Value;
 					}
 					else
 					{
@@ -151,17 +176,17 @@ FArchive& FArchiveFromStructuredArchive::operator<<(class UObject*& Value)
 					}
 				}
 
-				Objects[ObjectIdx] = Value;
-				ObjectsValid[ObjectIdx] = true;
+				Pimpl->Objects[ObjectIdx] = Value;
+				Pimpl->ObjectsValid[ObjectIdx] = true;
 			}
 		}
 		else
 		{
-			int32* ObjectIdxPtr = ObjectToIndex.Find(Value);
+			int32* ObjectIdxPtr = Pimpl->ObjectToIndex.Find(Value);
 			if (ObjectIdxPtr == nullptr)
 			{
-				ObjectIdxPtr = &(ObjectToIndex.Add(Value));
-				*ObjectIdxPtr = Objects.Add(Value);
+				ObjectIdxPtr = &Pimpl->ObjectToIndex.Add(Value);
+				*ObjectIdxPtr = Pimpl->Objects.Add(Value);
 			}
 			Serialize(ObjectIdxPtr, sizeof(*ObjectIdxPtr));
 		}
@@ -196,21 +221,21 @@ void FArchiveFromStructuredArchive::Serialize(void* V, int64 Length)
 	{
 		if (IsLoading())
 		{
-			if (Pos + Length > Buffer.Num())
+			if (Pimpl->Pos + Length > Pimpl->Buffer.Num())
 			{
 				checkf(false, TEXT("Attempt to read past end of archive"));
 			}
-			FMemory::Memcpy(V, Buffer.GetData() + Pos, Length);
-			Pos += Length;
+			FMemory::Memcpy(V, Pimpl->Buffer.GetData() + Pimpl->Pos, Length);
+			Pimpl->Pos += Length;
 		}
 		else
 		{
-			if (Pos + Length > Buffer.Num())
+			if (Pimpl->Pos + Length > Pimpl->Buffer.Num())
 			{
-				Buffer.AddUninitialized(Pos + Length - Buffer.Num());
+				Pimpl->Buffer.AddUninitialized(Pimpl->Pos + Length - Pimpl->Buffer.Num());
 			}
-			FMemory::Memcpy(Buffer.GetData() + Pos, V, Length);
-			Pos += Length;
+			FMemory::Memcpy(Pimpl->Buffer.GetData() + Pimpl->Pos, V, Length);
+			Pimpl->Pos += Length;
 		}
 	}
 	else
@@ -221,22 +246,22 @@ void FArchiveFromStructuredArchive::Serialize(void* V, int64 Length)
 
 void FArchiveFromStructuredArchive::Commit()
 {
-	if (bWasOpened && InnerArchive.IsTextFormat())
+	if (Pimpl->bWasOpened && InnerArchive.IsTextFormat())
 	{
-		SerializeInternal(Root.GetValue());
+		SerializeInternal(Pimpl->Root.GetValue());
 	}
 }
 
 void FArchiveFromStructuredArchive::SerializeInternal(FStructuredArchive::FRecord Record)
 {
-	check(bWasOpened);
+	check(Pimpl->bWasOpened);
 
-	if (bPendingSerialize)
+	if (Pimpl->bPendingSerialize)
 	{
 		FStructuredArchive::FSlot DataSlot = Record.EnterField(FIELD_NAME_TEXT("Data"));
-		DataSlot.Serialize(Buffer);
+		DataSlot.Serialize(Pimpl->Buffer);
 
-		TOptional<FStructuredArchive::FSlot> ObjectsSlot = Record.TryEnterField(FIELD_NAME_TEXT("Objects"), Objects.Num() > 0);
+		TOptional<FStructuredArchive::FSlot> ObjectsSlot = Record.TryEnterField(FIELD_NAME_TEXT("Objects"), Pimpl->Objects.Num() > 0);
 		if (ObjectsSlot.IsSet())
 		{
 			if (IsLoading())
@@ -249,48 +274,58 @@ void FArchiveFromStructuredArchive::SerializeInternal(FStructuredArchive::FRecor
 				// when we enter the array here. We never read them, so I'm assuming they just sit there
 				// until we destroy this archive wrapper. Perhaps we need something in the API here to just access
 				// the size of the array but not preparing to access it's values?
-				ObjectsSlot.GetValue() << ObjectNames;
+				ObjectsSlot.GetValue() << Pimpl->ObjectNames;
 				//int32 NumEntries = 0;
 				//ObjectsSlot.GetValue().EnterArray(NumEntries);
-				Objects.AddUninitialized(ObjectNames.Num());
-				ObjectsValid.Init(false, ObjectNames.Num());
+				Pimpl->Objects.AddUninitialized(Pimpl->ObjectNames.Num());
+				Pimpl->ObjectsValid.Init(false, Pimpl->ObjectNames.Num());
 			}
 			else
 			{
-				ObjectsSlot.GetValue() << Objects;
+				ObjectsSlot.GetValue() << Pimpl->Objects;
 			}
 		}
 
-		TOptional<FStructuredArchive::FSlot> NamesSlot = Record.TryEnterField(FIELD_NAME_TEXT("Names"), Names.Num() > 0);
+		TOptional<FStructuredArchive::FSlot> NamesSlot = Record.TryEnterField(FIELD_NAME_TEXT("Names"), Pimpl->Names.Num() > 0);
 		if (NamesSlot.IsSet())
 		{
-			NamesSlot.GetValue() << Names;
+			NamesSlot.GetValue() << Pimpl->Names;
 		}
 
-		bPendingSerialize = false;
+		Pimpl->bPendingSerialize = false;
 	}
 }
 
 void FArchiveFromStructuredArchive::OpenArchive()
 {
-	if (!bWasOpened)
+	if (!Pimpl->bWasOpened)
 	{
-		bWasOpened = true;
+		Pimpl->bWasOpened = true;
 
 		if (InnerArchive.IsTextFormat())
 		{
-			Root = RootSlot.EnterRecord();
+			Pimpl->Root = Pimpl->RootSlot.EnterRecord();
 
 			if (IsLoading())
 			{
-				SerializeInternal(Root.GetValue());
+				SerializeInternal(Pimpl->Root.GetValue());
 			}
 		}
 		else
 		{
-			RootSlot.EnterStream();
+			Pimpl->RootSlot.EnterStream();
 		}
 	}
 }
 
-#endif
\ No newline at end of file
+FArchive* FArchiveFromStructuredArchive::GetCacheableArchive()
+{
+	return IsTextFormat() ? nullptr : Pimpl->Root->GetUnderlyingArchive().GetCacheableArchive();
+}
+
+bool FArchiveFromStructuredArchive::ContainsData() const
+{
+	return Pimpl->Buffer.Num() > 0;
+}
+
+#endif
diff --git a/Engine/Source/Runtime/Core/Private/Serialization/CustomVersion.cpp b/Engine/Source/Runtime/Core/Private/Serialization/CustomVersion.cpp
--- a/Engine/Source/Runtime/Core/Private/Serialization/CustomVersion.cpp
+++ b/Engine/Source/Runtime/Core/Private/Serialization/CustomVersion.cpp
@@ -5,7 +5,7 @@
 =============================================================================*/
 
 #include "Serialization/CustomVersion.h"
-#include "Serialization/StructuredArchiveFromArchive.h"
+#include "Serialization/StructuredArchive.h"
 #include "Algo/Sort.h"
 
 namespace
diff --git a/Engine/Source/Runtime/Core/Private/Serialization/StructuredArchiveFromArchive.cpp b/Engine/Source/Runtime/Core/Private/Serialization/StructuredArchiveFromArchive.cpp
new file mode 100644
--- /dev/null
+++ b/Engine/Source/Runtime/Core/Private/Serialization/StructuredArchiveFromArchive.cpp
@@ -0,0 +1,31 @@
+// Copyright 1998-2019 Epic Games, Inc. All Rights Reserved.
+
+#include "Serialization/StructuredArchive.h"
+#include "Containers/Set.h"
+#include "Containers/UnrealString.h"
+
+struct FStructuredArchiveFromArchive::FImpl
+{
+	explicit FImpl(FArchive& Ar)
+		: Formatter(Ar)
+		, StructuredArchive(Formatter)
+		, Slot(StructuredArchive.Open())
+	{
+	}
+
+	FBinaryArchiveFormatter Formatter;
+	FStructuredArchive StructuredArchive;
+	FStructuredArchive::FSlot Slot;
+};
+
+FStructuredArchiveFromArchive::FStructuredArchiveFromArchive(FArchive& Ar)
+	: Pimpl(Ar)
+{
+}
+
+FStructuredArchiveFromArchive::~FStructuredArchiveFromArchive() = default;
+
+FStructuredArchive::FSlot FStructuredArchiveFromArchive::GetSlot()
+{
+	return Pimpl->Slot;
+}
diff --git a/Engine/Source/Runtime/Core/Public/Math/Color.h b/Engine/Source/Runtime/Core/Public/Math/Color.h
--- a/Engine/Source/Runtime/Core/Public/Math/Color.h
+++ b/Engine/Source/Runtime/Core/Public/Math/Color.h
@@ -8,7 +8,6 @@
 #include "Containers/UnrealString.h"
 #include "Misc/Parse.h"
 #include "Serialization/StructuredArchive.h"
-#include "Serialization/ArchiveAdapters.h"
 
 class FFloat16Color;
 
diff --git a/Engine/Source/Runtime/Core/Public/Math/IntPoint.h b/Engine/Source/Runtime/Core/Public/Math/IntPoint.h
--- a/Engine/Source/Runtime/Core/Public/Math/IntPoint.h
+++ b/Engine/Source/Runtime/Core/Public/Math/IntPoint.h
@@ -7,7 +7,6 @@
 #include "Math/UnrealMathUtility.h"
 #include "Containers/UnrealString.h"
 #include "Templates/TypeHash.h"
-#include "Serialization/ArchiveAdapters.h"
 
 /**
  * Structure for integer points in 2-d space.
diff --git a/Engine/Source/Runtime/Core/Public/Math/TwoVectors.h b/Engine/Source/Runtime/Core/Public/Math/TwoVectors.h
--- a/Engine/Source/Runtime/Core/Public/Math/TwoVectors.h
+++ b/Engine/Source/Runtime/Core/Public/Math/TwoVectors.h
@@ -7,7 +7,6 @@
 #include "Math/UnrealMathUtility.h"
 #include "Containers/UnrealString.h"
 #include "Math/Vector.h"
-#include "Serialization/ArchiveAdapters.h"
 
 /**
  * A pair of 3D vectors.
diff --git a/Engine/Source/Runtime/Core/Public/Math/Vector.h b/Engine/Source/Runtime/Core/Public/Math/Vector.h
--- a/Engine/Source/Runtime/Core/Public/Math/Vector.h
+++ b/Engine/Source/Runtime/Core/Public/Math/Vector.h
@@ -18,7 +18,6 @@
 #include "Internationalization/Internationalization.h"
 #include "Math/IntVector.h"
 #include "Math/Axis.h"
-#include "Serialization/ArchiveAdapters.h"
 
 #if PLATFORM_VECTOR_CUBIC_INTERP_SSE
 #include "Math/UnrealMathSSE.h"
diff --git a/Engine/Source/Runtime/Core/Public/Misc/ConfigCacheIni.h b/Engine/Source/Runtime/Core/Public/Misc/ConfigCacheIni.h
--- a/Engine/Source/Runtime/Core/Public/Misc/ConfigCacheIni.h
+++ b/Engine/Source/Runtime/Core/Public/Misc/ConfigCacheIni.h
@@ -18,7 +18,7 @@
 #include "Math/Vector.h"
 #include "Math/Rotator.h"
 #include "Misc/Paths.h"
-#include "Serialization/StructuredArchiveFromArchive.h"
+#include "Serialization/StructuredArchive.h"
 
 CORE_API DECLARE_LOG_CATEGORY_EXTERN(LogConfig, Log, All);
 
diff --git a/Engine/Source/Runtime/Core/Public/Misc/DateTime.h b/Engine/Source/Runtime/Core/Public/Misc/DateTime.h
--- a/Engine/Source/Runtime/Core/Public/Misc/DateTime.h
+++ b/Engine/Source/Runtime/Core/Public/Misc/DateTime.h
@@ -6,7 +6,6 @@
 #include "Containers/UnrealString.h"
 #include "Misc/Timespan.h"
 #include "Templates/TypeHash.h"
-#include "Serialization/ArchiveAdapters.h"
 
 class FArchive;
 class FOutputDevice;
diff --git a/Engine/Source/Runtime/Core/Public/Misc/Guid.h b/Engine/Source/Runtime/Core/Public/Misc/Guid.h
--- a/Engine/Source/Runtime/Core/Public/Misc/Guid.h
+++ b/Engine/Source/Runtime/Core/Public/Misc/Guid.h
@@ -6,7 +6,7 @@
 #include "Misc/AssertionMacros.h"
 #include "Misc/Crc.h"
 #include "Containers/UnrealString.h"
-#include "Serialization/StructuredArchiveFromArchive.h"
+#include "Serialization/StructuredArchive.h"
 
 class FArchive;
 class FOutputDevice;
diff --git a/Engine/Source/Runtime/Core/Public/Serialization/ArchiveAdapters.h b/Engine/Source/Runtime/Core/Public/Serialization/ArchiveAdapters.h
deleted file mode 100644
diff --git a/Engine/Source/Runtime/Core/Public/Serialization/ArchiveFromStructuredArchive.h b/Engine/Source/Runtime/Core/Public/Serialization/ArchiveFromStructuredArchive.h
deleted file mode 100644
diff --git a/Engine/Source/Runtime/Core/Public/Serialization/StructuredArchive.h b/Engine/Source/Runtime/Core/Public/Serialization/StructuredArchive.h
--- a/Engine/Source/Runtime/Core/Public/Serialization/StructuredArchive.h
+++ b/Engine/Source/Runtime/Core/Public/Serialization/StructuredArchive.h
@@ -6,9 +6,12 @@
 #include "Serialization/StructuredArchiveFormatter.h"
 #include "Formatters/BinaryArchiveFormatter.h"
 #include "Misc/Optional.h"
+#include "Concepts/Insertable.h"
 #include "Concepts/Serializable.h"
 #include "Templates/Models.h"
 #include "Containers/Array.h"
+#include "Serialization/ArchiveProxy.h"
+#include "Templates/UniqueObj.h"
 
 /**
  * Class to contain a named value for serialization. Intended to be created as a temporary and passed to object serialization methods.
@@ -405,6 +408,130 @@ private:
 	TOptional<FStructuredArchive::FSlot> Root;
 };
 
+class CORE_API FStructuredArchiveFromArchive
+{
+	struct FImpl;
+
+public:
+	explicit FStructuredArchiveFromArchive(FArchive& Ar);
+	~FStructuredArchiveFromArchive();
+
+	// Non-copyable
+	FStructuredArchiveFromArchive(FStructuredArchiveFromArchive&&) = delete;
+	FStructuredArchiveFromArchive(const FStructuredArchiveFromArchive&) = delete;
+	FStructuredArchiveFromArchive& operator=(FStructuredArchiveFromArchive&&) = delete;
+	FStructuredArchiveFromArchive& operator=(const FStructuredArchiveFromArchive&) = delete;
+
+	FStructuredArchive::FSlot GetSlot();
+
+private:
+	// Implmented as a pimpl in order to reduce dependencies
+	TUniqueObj<FImpl> Pimpl;
+};
+
+#if WITH_TEXT_ARCHIVE_SUPPORT
+
+class CORE_API FArchiveFromStructuredArchive : public FArchiveProxy
+{
+	struct FImpl;
+
+public:
+	explicit FArchiveFromStructuredArchive(FStructuredArchive::FSlot Slot);
+	virtual ~FArchiveFromStructuredArchive();
+
+	// Non-copyable
+	FArchiveFromStructuredArchive(FArchiveFromStructuredArchive&&) = delete;
+	FArchiveFromStructuredArchive(const FArchiveFromStructuredArchive&) = delete;
+	FArchiveFromStructuredArchive& operator=(FArchiveFromStructuredArchive&&) = delete;
+	FArchiveFromStructuredArchive& operator=(const FArchiveFromStructuredArchive&) = delete;
+
+	virtual void Flush() override;
+	virtual bool Close() override;
+
+	virtual int64 Tell() override;
+	virtual int64 TotalSize() override;
+	virtual void Seek(int64 InPos) override;
+	virtual bool AtEnd() override;
+
+	virtual FArchive& operator<<(class FName& Value) override;
+	virtual FArchive& operator<<(class UObject*& Value) override;
+	virtual FArchive& operator<<(class FText& Value) override;
+
+	virtual void Serialize(void* V, int64 Length) override;
+
+	virtual FArchive* GetCacheableArchive() override;
+
+	bool ContainsData() const;
+
+protected:
+	virtual void SerializeInternal(FStructuredArchive::FRecord Record);
+	void OpenArchive();
+
+private:
+	void Commit();
+
+	// Implmented as a pimpl in order to reduce dependencies
+	TUniqueObj<FImpl> Pimpl;
+};
+
+#else
+
+class CORE_API FArchiveFromStructuredArchive
+{
+public:
+	explicit FArchiveFromStructuredArchive(FStructuredArchive::FSlot InSlot)
+		: Ar(InSlot.GetUnderlyingArchive())
+	{
+
+	}
+
+	operator FArchive& () { return Ar; }
+
+private:
+	FArchive& Ar;
+};
+
+#endif
+
+/**
+ * Adapter operator which allows a type to stream to an FArchive when it already supports streaming to an FStructuredArchive::FSlot.
+ *
+ * @param  Ar   The archive to read from or write to.
+ * @param  Obj  The object to read or write.
+ *
+ * @return  A reference to the same archive as Ar.
+ */
+template <typename T>
+typename TEnableIf<
+	!TModels<CInsertable<FArchive&>, T>::Value && TModels<CInsertable<FStructuredArchive::FSlot>, T>::Value,
+	FArchive&
+>::Type operator<<(FArchive& Ar, T& Obj)
+{
+	FStructuredArchiveFromArchive ArAdapt(Ar);
+	ArAdapt.GetSlot() << Obj;
+	return Ar;
+}
+
+/**
+ * Adapter operator which allows a type to stream to an FStructuredArchive::FSlot when it already supports streaming to an FArchive.
+ *
+ * @param  Slot  The slot to read from or write to.
+ * @param  Obj   The object to read or write.
+ */
+template <typename T>
+typename TEnableIf<
+	TModels<CInsertable<FArchive&>, T>::Value &&
+	!TModels<CInsertable<FStructuredArchive::FSlot>, T>::Value
+>::Type operator<<(FStructuredArchive::FSlot Slot, T& Obj)
+{
+#if WITH_TEXT_ARCHIVE_SUPPORT
+	FArchiveFromStructuredArchive Ar(Slot);
+#else
+	FArchive& Ar = Slot.GetUnderlyingArchive();
+#endif
+	Ar << Obj;
+}
+
 #if !WITH_TEXT_ARCHIVE_SUPPORT
 	FORCEINLINE FStructuredArchiveChildReader::FStructuredArchiveChildReader(FStructuredArchive::FSlot InSlot)
 		: OwnedFormatter(nullptr)
diff --git a/Engine/Source/Runtime/Core/Public/Serialization/StructuredArchiveFromArchive.h b/Engine/Source/Runtime/Core/Public/Serialization/StructuredArchiveFromArchive.h
deleted file mode 100644
diff --git a/Engine/Source/Runtime/CoreUObject/Private/Serialization/ArchiveUObjectFromStructuredArchive.cpp b/Engine/Source/Runtime/CoreUObject/Private/Serialization/ArchiveUObjectFromStructuredArchive.cpp
--- a/Engine/Source/Runtime/CoreUObject/Private/Serialization/ArchiveUObjectFromStructuredArchive.cpp
+++ b/Engine/Source/Runtime/CoreUObject/Private/Serialization/ArchiveUObjectFromStructuredArchive.cpp
@@ -10,11 +10,6 @@ FArchiveUObjectFromStructuredArchive::FArchiveUObjectFromStructuredArchive(FStru
 {
 }
 
-FArchiveUObjectFromStructuredArchive::~FArchiveUObjectFromStructuredArchive()
-{
-	Commit();
-}
-
 FArchive& FArchiveUObjectFromStructuredArchive::operator<<(FLazyObjectPtr& Value)
 {
 	OpenArchive();
diff --git a/Engine/Source/Runtime/CoreUObject/Private/UObject/PackageFileSummary.cpp b/Engine/Source/Runtime/CoreUObject/Private/UObject/PackageFileSummary.cpp
--- a/Engine/Source/Runtime/CoreUObject/Private/UObject/PackageFileSummary.cpp
+++ b/Engine/Source/Runtime/CoreUObject/Private/UObject/PackageFileSummary.cpp
@@ -2,7 +2,7 @@
 
 #include "UObject/PackageFileSummary.h"
 #include "UObject/Linker.h"
-#include "Serialization/ArchiveFromStructuredArchive.h"
+#include "Serialization/StructuredArchive.h"
 #include "UObject/UObjectGlobals.h"
 
 FPackageFileSummary::FPackageFileSummary()
diff --git a/Engine/Source/Runtime/CoreUObject/Public/Serialization/ArchiveUObjectFromStructuredArchive.h b/Engine/Source/Runtime/CoreUObject/Public/Serialization/ArchiveUObjectFromStructuredArchive.h
--- a/Engine/Source/Runtime/CoreUObject/Public/Serialization/ArchiveUObjectFromStructuredArchive.h
+++ b/Engine/Source/Runtime/CoreUObject/Public/Serialization/ArchiveUObjectFromStructuredArchive.h
@@ -2,7 +2,7 @@
 
 #pragma once
 
-#include "Serialization/ArchiveFromStructuredArchive.h"
+#include "Serialization/StructuredArchive.h"
 #include "Serialization/ArchiveUObject.h"
 #include "UObject/ObjectResource.h"
 
@@ -18,7 +18,6 @@ class COREUOBJECT_API FArchiveUObjectFromStructuredArchive : public FArchiveFrom
 public:
 
 	FArchiveUObjectFromStructuredArchive(FStructuredArchive::FSlot Slot);
-	virtual ~FArchiveUObjectFromStructuredArchive();
 
 	using FArchive::operator<<; // For visibility of the overloads we don't override
 
