Circular include fix.

diff --git a/Engine/Source/Runtime/Core/Public/Serialization/ArchiveProxy.h b/Engine/Source/Runtime/Core/Public/Serialization/ArchiveProxy.h
--- a/Engine/Source/Runtime/Core/Public/Serialization/ArchiveProxy.h
+++ b/Engine/Source/Runtime/Core/Public/Serialization/ArchiveProxy.h
@@ -5,7 +5,6 @@
 #include "CoreTypes.h"
 #include "Internationalization/TextNamespaceFwd.h"
 #include "Serialization/Archive.h"
-#include "Containers/UnrealString.h"
 
 /**
  * Base class for archive proxies.
