Possible archive adapter fix.

diff --git a/Engine/Source/Runtime/Core/Public/Serialization/StructuredArchive.h b/Engine/Source/Runtime/Core/Public/Serialization/StructuredArchive.h
--- a/Engine/Source/Runtime/Core/Public/Serialization/StructuredArchive.h
+++ b/Engine/Source/Runtime/Core/Public/Serialization/StructuredArchive.h
@@ -525,7 +525,8 @@ typename TEnableIf<
 >::Type operator<<(FStructuredArchive::FSlot Slot, T& Obj)
 {
 #if WITH_TEXT_ARCHIVE_SUPPORT
-	FArchiveFromStructuredArchive Ar(Slot);
+	FArchiveFromStructuredArchive Adapter(Slot);
+	FArchive& Ar = Adapter;
 #else
 	FArchive& Ar = Slot.GetUnderlyingArchive();
 #endif
