Added new simpler RW lock scopes for convenience

Makes it easier to write, read and reason about scope duration.

diff --git a/Engine/Source/Runtime/Core/Public/Misc/ScopeRWLock.h b/Engine/Source/Runtime/Core/Public/Misc/ScopeRWLock.h
--- a/Engine/Source/Runtime/Core/Public/Misc/ScopeRWLock.h
+++ b/Engine/Source/Runtime/Core/Public/Misc/ScopeRWLock.h
@@ -6,6 +6,49 @@
 #include "Misc/AssertionMacros.h"
 #include "HAL/CriticalSection.h"
 
+
+/** Keeps a FRWLock read-locked while this scope lives */
+class FReadScopeLock
+{
+public:
+	explicit FReadScopeLock(FRWLock& InLock)
+		: Lock(InLock)
+	{
+		Lock.ReadLock();
+	}
+
+	~FReadScopeLock()
+	{
+		Lock.ReadUnlock();
+	}
+
+private:
+	FRWLock& Lock;
+
+	UE_NONCOPYABLE(FReadScopeLock);
+};
+
+/** Keeps a FRWLock write-locked while this scope lives */
+class FWriteScopeLock
+{
+public:
+	explicit FWriteScopeLock(FRWLock& InLock)
+		: Lock(InLock)
+	{
+		Lock.WriteLock();
+	}
+
+	~FWriteScopeLock()
+	{
+		Lock.WriteUnlock();
+	}
+	
+private:
+	FRWLock& Lock;
+
+	UE_NONCOPYABLE(FWriteScopeLock);
+};
+
 //
 // A scope lifetime controlled Read or Write lock of referenced mutex object
 //
@@ -16,14 +59,15 @@ enum FRWScopeLockType
 };
 
 /**
- * Implements a scope lock for RW locks.
+ * Keeps a FRWLock read- or write-locked while this scope lives
+ *
  *	Notes:
  *		- PThreads and Win32 API's don't provide a mechanism for upgrading a ownership of a read lock to a write lock - to get round that this system unlocks then acquires a write lock so it can race between.
  */
 class FRWScopeLock
 {
 public:
-	FORCEINLINE FRWScopeLock(FRWLock& InLockObject,FRWScopeLockType InLockType)
+	explicit FRWScopeLock(FRWLock& InLockObject,FRWScopeLockType InLockType)
 	: LockObject(InLockObject)
 	, LockType(InLockType)
 	{
@@ -41,7 +85,7 @@ public:
 	// It releases the read lock _before_ acquiring a new write lock. This is not an atomic operation and the caller should 
 	// not treat it as such. 
 	// E.g. Pointers read from protected data structures prior to this call may be invalid after the function is called. 
-	FORCEINLINE void ReleaseReadOnlyLockAndAcquireWriteLock_USE_WITH_CAUTION()
+	void ReleaseReadOnlyLockAndAcquireWriteLock_USE_WITH_CAUTION()
 	{
 		if(LockType == SLT_ReadOnly)
 		{
@@ -51,7 +95,7 @@ public:
 		}
 	}
 	
-	FORCEINLINE ~FRWScopeLock()
+	~FRWScopeLock()
 	{
 		if(LockType == SLT_ReadOnly)
 		{
@@ -64,11 +108,8 @@ public:
 	}
 	
 private:
-	FRWScopeLock();
-	FRWScopeLock(const FRWScopeLock&);
-	FRWScopeLock& operator=(const FRWScopeLock&);
-	
-private:
+	UE_NONCOPYABLE(FRWScopeLock);
+
 	FRWLock& LockObject;
 	FRWScopeLockType LockType;
 };
