FArchiveFromStructuredArchive and FArchiveUObjectFromStructuredArchive no longer implicitly convert to FArchive.

diff --git a/Engine/Source/Programs/UnrealBuildTool/ProjectFiles/VisualStudio/VCProjectFileGenerator.cs b/Engine/Source/Programs/UnrealBuildTool/ProjectFiles/VisualStudio/VCProjectFileGenerator.cs
--- a/Engine/Source/Programs/UnrealBuildTool/ProjectFiles/VisualStudio/VCProjectFileGenerator.cs
+++ b/Engine/Source/Programs/UnrealBuildTool/ProjectFiles/VisualStudio/VCProjectFileGenerator.cs
@@ -486,6 +486,8 @@ namespace UnrealBuildTool
 						};
 					GatherFoldersFunction(RootFolder.SubFolders);
 
+					AllSolutionFolders.Sort((Lhs, Rhs) => Lhs.FolderName.CompareTo(Rhs.FolderName));
+
 					foreach (VisualStudioSolutionFolder CurFolder in AllSolutionFolders)
 					{
 						string FolderGUIDString = CurFolder.FolderGUID.ToString("B").ToUpperInvariant();
@@ -509,6 +511,7 @@ namespace UnrealBuildTool
 
 
 				// Project files
+				List<MSBuildProjectFile> AllProjectFilesSorted = AllProjectFiles.OrderBy((ProjFile) => ProjFile.ProjectFilePath.GetFileNameWithoutExtension()).Cast<MSBuildProjectFile>().ToList();
 				foreach (MSBuildProjectFile CurProject in AllProjectFiles)
 				{
 					// Visual Studio uses different GUID types depending on the project type
diff --git a/Engine/Source/Runtime/Core/Private/Serialization/ArchiveFromStructuredArchive.cpp b/Engine/Source/Runtime/Core/Private/Serialization/ArchiveFromStructuredArchive.cpp
--- a/Engine/Source/Runtime/Core/Private/Serialization/ArchiveFromStructuredArchive.cpp
+++ b/Engine/Source/Runtime/Core/Private/Serialization/ArchiveFromStructuredArchive.cpp
@@ -5,7 +5,7 @@
 
 #if WITH_TEXT_ARCHIVE_SUPPORT
 
-struct FArchiveFromStructuredArchive::FImpl
+struct FArchiveFromStructuredArchiveImpl::FImpl
 {
 	explicit FImpl(FStructuredArchive::FSlot Slot)
 		: RootSlot(Slot)
@@ -33,7 +33,7 @@ struct FArchiveFromStructuredArchive::FImpl
 	FStructuredArchive::FSlot RootSlot;
 };
 
-FArchiveFromStructuredArchive::FArchiveFromStructuredArchive(FStructuredArchive::FSlot Slot)
+FArchiveFromStructuredArchiveImpl::FArchiveFromStructuredArchiveImpl(FStructuredArchive::FSlot Slot)
 	: FArchiveProxy(Slot.GetUnderlyingArchive())
 	, Pimpl(Slot)
 {
@@ -43,24 +43,24 @@ FArchiveFromStructuredArchive::FArchiveFromStructuredArchive(FStructuredArchive:
 	SetIsTextFormat(false);
 }
 
-FArchiveFromStructuredArchive::~FArchiveFromStructuredArchive()
+FArchiveFromStructuredArchiveImpl::~FArchiveFromStructuredArchiveImpl()
 {
 	Commit();
 }
 
-void FArchiveFromStructuredArchive::Flush()
+void FArchiveFromStructuredArchiveImpl::Flush()
 {
 	Commit();
 	FArchive::Flush();
 }
 
-bool FArchiveFromStructuredArchive::Close()
+bool FArchiveFromStructuredArchiveImpl::Close()
 {
 	Commit();
 	return FArchive::Close();
 }
 
-int64 FArchiveFromStructuredArchive::Tell()
+int64 FArchiveFromStructuredArchiveImpl::Tell()
 {
 	if (InnerArchive.IsTextFormat())
 	{
@@ -72,13 +72,13 @@ int64 FArchiveFromStructuredArchive::Tell()
 	}
 }
 
-int64 FArchiveFromStructuredArchive::TotalSize()
+int64 FArchiveFromStructuredArchiveImpl::TotalSize()
 {
 	checkf(false, TEXT("FArchiveFromStructuredArchive does not support TotalSize()"));
 	return FArchive::TotalSize();
 }
 
-void FArchiveFromStructuredArchive::Seek(int64 InPos)
+void FArchiveFromStructuredArchiveImpl::Seek(int64 InPos)
 {
 	if (InnerArchive.IsTextFormat())
 	{
@@ -91,7 +91,7 @@ void FArchiveFromStructuredArchive::Seek(int64 InPos)
 	}
 }
 
-bool FArchiveFromStructuredArchive::AtEnd()
+bool FArchiveFromStructuredArchiveImpl::AtEnd()
 {
 	if (InnerArchive.IsTextFormat())
 	{
@@ -103,7 +103,7 @@ bool FArchiveFromStructuredArchive::AtEnd()
 	}
 }
 
-FArchive& FArchiveFromStructuredArchive::operator<<(class FName& Value)
+FArchive& FArchiveFromStructuredArchiveImpl::operator<<(class FName& Value)
 {
 	OpenArchive();
 
@@ -133,7 +133,7 @@ FArchive& FArchiveFromStructuredArchive::operator<<(class FName& Value)
 	return *this;
 }
 
-FArchive& FArchiveFromStructuredArchive::operator<<(class UObject*& Value)
+FArchive& FArchiveFromStructuredArchiveImpl::operator<<(class UObject*& Value)
 {
 	OpenArchive();
 
@@ -198,7 +198,7 @@ FArchive& FArchiveFromStructuredArchive::operator<<(class UObject*& Value)
 	return *this;
 }
 
-FArchive& FArchiveFromStructuredArchive::operator<<(class FText& Value)
+FArchive& FArchiveFromStructuredArchiveImpl::operator<<(class FText& Value)
 {
 	OpenArchive();
 
@@ -213,7 +213,7 @@ FArchive& FArchiveFromStructuredArchive::operator<<(class FText& Value)
 	return *this;
 }
 
-void FArchiveFromStructuredArchive::Serialize(void* V, int64 Length)
+void FArchiveFromStructuredArchiveImpl::Serialize(void* V, int64 Length)
 {
 	OpenArchive();
 
@@ -244,7 +244,7 @@ void FArchiveFromStructuredArchive::Serialize(void* V, int64 Length)
 	}
 }
 
-void FArchiveFromStructuredArchive::Commit()
+void FArchiveFromStructuredArchiveImpl::Commit()
 {
 	if (Pimpl->bWasOpened && InnerArchive.IsTextFormat())
 	{
@@ -252,7 +252,7 @@ void FArchiveFromStructuredArchive::Commit()
 	}
 }
 
-void FArchiveFromStructuredArchive::SerializeInternal(FStructuredArchive::FRecord Record)
+void FArchiveFromStructuredArchiveImpl::SerializeInternal(FStructuredArchive::FRecord Record)
 {
 	check(Pimpl->bWasOpened);
 
@@ -296,7 +296,7 @@ void FArchiveFromStructuredArchive::SerializeInternal(FStructuredArchive::FRecor
 	}
 }
 
-void FArchiveFromStructuredArchive::OpenArchive()
+void FArchiveFromStructuredArchiveImpl::OpenArchive()
 {
 	if (!Pimpl->bWasOpened)
 	{
@@ -318,12 +318,12 @@ void FArchiveFromStructuredArchive::OpenArchive()
 	}
 }
 
-FArchive* FArchiveFromStructuredArchive::GetCacheableArchive()
+FArchive* FArchiveFromStructuredArchiveImpl::GetCacheableArchive()
 {
 	return IsTextFormat() ? nullptr : Pimpl->Root->GetUnderlyingArchive().GetCacheableArchive();
 }
 
-bool FArchiveFromStructuredArchive::ContainsData() const
+bool FArchiveFromStructuredArchiveImpl::ContainsData() const
 {
 	return Pimpl->Buffer.Num() > 0;
 }
diff --git a/Engine/Source/Runtime/Core/Public/Internationalization/Text.h b/Engine/Source/Runtime/Core/Public/Internationalization/Text.h
--- a/Engine/Source/Runtime/Core/Public/Internationalization/Text.h
+++ b/Engine/Source/Runtime/Core/Public/Internationalization/Text.h
@@ -735,7 +735,7 @@ public:
 	friend class FTextStringHelper;
 	friend class FStringTableRegistry;
 	friend class FArchive;
-	friend class FArchiveFromStructuredArchive;
+	friend class FArchiveFromStructuredArchiveImpl;
 	friend class FJsonArchiveInputFormatter;
 	friend class FJsonArchiveOutputFormatter;
 	friend class UTextProperty;
diff --git a/Engine/Source/Runtime/Core/Public/Serialization/StructuredArchive.h b/Engine/Source/Runtime/Core/Public/Serialization/StructuredArchive.h
--- a/Engine/Source/Runtime/Core/Public/Serialization/StructuredArchive.h
+++ b/Engine/Source/Runtime/Core/Public/Serialization/StructuredArchive.h
@@ -410,18 +410,14 @@ private:
 
 class CORE_API FStructuredArchiveFromArchive
 {
+	UE_NONCOPYABLE(FStructuredArchiveFromArchive)
+
 	struct FImpl;
 
 public:
 	explicit FStructuredArchiveFromArchive(FArchive& Ar);
 	~FStructuredArchiveFromArchive();
 
-	// Non-copyable
-	FStructuredArchiveFromArchive(FStructuredArchiveFromArchive&&) = delete;
-	FStructuredArchiveFromArchive(const FStructuredArchiveFromArchive&) = delete;
-	FStructuredArchiveFromArchive& operator=(FStructuredArchiveFromArchive&&) = delete;
-	FStructuredArchiveFromArchive& operator=(const FStructuredArchiveFromArchive&) = delete;
-
 	FStructuredArchive::FSlot GetSlot();
 
 private:
@@ -431,19 +427,15 @@ private:
 
 #if WITH_TEXT_ARCHIVE_SUPPORT
 
-class CORE_API FArchiveFromStructuredArchive : public FArchiveProxy
+class CORE_API FArchiveFromStructuredArchiveImpl : public FArchiveProxy
 {
+	UE_NONCOPYABLE(FArchiveFromStructuredArchiveImpl)
+
 	struct FImpl;
 
 public:
-	explicit FArchiveFromStructuredArchive(FStructuredArchive::FSlot Slot);
-	virtual ~FArchiveFromStructuredArchive();
-
-	// Non-copyable
-	FArchiveFromStructuredArchive(FArchiveFromStructuredArchive&&) = delete;
-	FArchiveFromStructuredArchive(const FArchiveFromStructuredArchive&) = delete;
-	FArchiveFromStructuredArchive& operator=(FArchiveFromStructuredArchive&&) = delete;
-	FArchiveFromStructuredArchive& operator=(const FArchiveFromStructuredArchive&) = delete;
+	explicit FArchiveFromStructuredArchiveImpl(FStructuredArchive::FSlot Slot);
+	virtual ~FArchiveFromStructuredArchiveImpl();
 
 	virtual void Flush() override;
 	virtual bool Close() override;
@@ -453,9 +445,13 @@ public:
 	virtual void Seek(int64 InPos) override;
 	virtual bool AtEnd() override;
 
+	using FArchive::operator<<; // For visibility of the overloads we don't override
+
+	//~ Begin FArchive Interface
 	virtual FArchive& operator<<(class FName& Value) override;
 	virtual FArchive& operator<<(class UObject*& Value) override;
 	virtual FArchive& operator<<(class FText& Value) override;
+	//~ End FArchive Interface
 
 	virtual void Serialize(void* V, int64 Length) override;
 
@@ -474,18 +470,33 @@ private:
 	TUniqueObj<FImpl> Pimpl;
 };
 
+class FArchiveFromStructuredArchive
+{
+public:
+	explicit FArchiveFromStructuredArchive(FStructuredArchive::FSlot InSlot)
+		: Impl(InSlot)
+	{
+	}
+
+	      FArchive& GetArchive()       { return Impl; }
+	const FArchive& GetArchive() const { return Impl; }
+
+private:
+	FArchiveFromStructuredArchiveImpl Impl;
+};
+
 #else
 
-class CORE_API FArchiveFromStructuredArchive
+class FArchiveFromStructuredArchive
 {
 public:
 	explicit FArchiveFromStructuredArchive(FStructuredArchive::FSlot InSlot)
 		: Ar(InSlot.GetUnderlyingArchive())
 	{
-
 	}
 
-	operator FArchive& () { return Ar; }
+	      FArchive& GetArchive()       { return Ar; }
+	const FArchive& GetArchive() const { return Ar; }
 
 private:
 	FArchive& Ar;
@@ -526,7 +537,7 @@ typename TEnableIf<
 {
 #if WITH_TEXT_ARCHIVE_SUPPORT
 	FArchiveFromStructuredArchive Adapter(Slot);
-	FArchive& Ar = Adapter;
+	FArchive& Ar = Adapter.GetArchive();
 #else
 	FArchive& Ar = Slot.GetUnderlyingArchive();
 #endif
diff --git a/Engine/Source/Runtime/CoreUObject/Private/Serialization/ArchiveUObjectFromStructuredArchive.cpp b/Engine/Source/Runtime/CoreUObject/Private/Serialization/ArchiveUObjectFromStructuredArchive.cpp
--- a/Engine/Source/Runtime/CoreUObject/Private/Serialization/ArchiveUObjectFromStructuredArchive.cpp
+++ b/Engine/Source/Runtime/CoreUObject/Private/Serialization/ArchiveUObjectFromStructuredArchive.cpp
@@ -4,13 +4,13 @@
 
 #if WITH_TEXT_ARCHIVE_SUPPORT
 
-FArchiveUObjectFromStructuredArchive::FArchiveUObjectFromStructuredArchive(FStructuredArchive::FSlot Slot)
-	: FArchiveFromStructuredArchive(Slot)
+FArchiveUObjectFromStructuredArchiveImpl::FArchiveUObjectFromStructuredArchiveImpl(FStructuredArchive::FSlot Slot)
+	: Super(Slot)
 	, bPendingSerialize(true)
 {
 }
 
-FArchive& FArchiveUObjectFromStructuredArchive::operator<<(FLazyObjectPtr& Value)
+FArchive& FArchiveUObjectFromStructuredArchiveImpl::operator<<(FLazyObjectPtr& Value)
 {
 	OpenArchive();
 
@@ -40,7 +40,7 @@ FArchive& FArchiveUObjectFromStructuredArchive::operator<<(FLazyObjectPtr& Value
 	return *this;
 }
 
-FArchive& FArchiveUObjectFromStructuredArchive::operator<<(FSoftObjectPtr& Value)
+FArchive& FArchiveUObjectFromStructuredArchiveImpl::operator<<(FSoftObjectPtr& Value)
 {
 	OpenArchive();
 
@@ -70,7 +70,7 @@ FArchive& FArchiveUObjectFromStructuredArchive::operator<<(FSoftObjectPtr& Value
 	return *this;
 }
 
-FArchive& FArchiveUObjectFromStructuredArchive::operator<<(FSoftObjectPath& Value)
+FArchive& FArchiveUObjectFromStructuredArchiveImpl::operator<<(FSoftObjectPath& Value)
 {
 	OpenArchive();
 
@@ -117,7 +117,7 @@ FArchive& FArchiveUObjectFromStructuredArchive::operator<<(FSoftObjectPath& Valu
 	return *this;
 }
 
-FArchive& FArchiveUObjectFromStructuredArchive::operator<<(FWeakObjectPtr& Value)
+FArchive& FArchiveUObjectFromStructuredArchiveImpl::operator<<(FWeakObjectPtr& Value)
 {
 	OpenArchive();
 
@@ -147,9 +147,9 @@ FArchive& FArchiveUObjectFromStructuredArchive::operator<<(FWeakObjectPtr& Value
 	return *this;
 }
 
-void FArchiveUObjectFromStructuredArchive::SerializeInternal(FStructuredArchive::FRecord Record)
+void FArchiveUObjectFromStructuredArchiveImpl::SerializeInternal(FStructuredArchive::FRecord Record)
 {
-	FArchiveFromStructuredArchive::SerializeInternal(Record);
+	Super::SerializeInternal(Record);
 
 	if (bPendingSerialize)
 	{
diff --git a/Engine/Source/Runtime/CoreUObject/Private/UObject/Class.cpp b/Engine/Source/Runtime/CoreUObject/Private/UObject/Class.cpp
--- a/Engine/Source/Runtime/CoreUObject/Private/UObject/Class.cpp
+++ b/Engine/Source/Runtime/CoreUObject/Private/UObject/Class.cpp
@@ -2164,7 +2164,8 @@ void UScriptStruct::SerializeItem(FStructuredArchive::FSlot Slot, void* Value, v
 		else
 		{
 #if WITH_TEXT_ARCHIVE_SUPPORT
-			FArchiveUObjectFromStructuredArchive Ar(Slot);
+			FArchiveUObjectFromStructuredArchive Adapter(Slot);
+			FArchive& Ar = Adapter.GetArchive();
 			bItemSerialized = TheCppStructOps->Serialize(Ar, Value);
 			if (bItemSerialized && !Slot.IsFilled())
 			{
diff --git a/Engine/Source/Runtime/CoreUObject/Private/UObject/LinkerLoad.cpp b/Engine/Source/Runtime/CoreUObject/Private/UObject/LinkerLoad.cpp
--- a/Engine/Source/Runtime/CoreUObject/Private/UObject/LinkerLoad.cpp
+++ b/Engine/Source/Runtime/CoreUObject/Private/UObject/LinkerLoad.cpp
@@ -3588,8 +3588,8 @@ void FLinkerLoad::Preload( UObject* Object )
 							{
 								FStructuredArchiveChildReader ChildReader(ExportSlot);
 								FArchiveUObjectFromStructuredArchive Adapter(ChildReader.GetRoot());
-							Object->GetClass()->SerializeDefaultObject(Object, Adapter);
-						}
+								Object->GetClass()->SerializeDefaultObject(Object, Adapter.GetArchive());
+							}
 						}
 						else
 #endif
@@ -3638,7 +3638,7 @@ void FLinkerLoad::Preload( UObject* Object )
 							{
 								FStructuredArchiveChildReader ChildReader(ExportSlot);
 								FArchiveUObjectFromStructuredArchive Adapter(ChildReader.GetRoot());
-								Object->Serialize(Adapter);
+								Object->Serialize(Adapter.GetArchive());
 							}
 						}
 						else
diff --git a/Engine/Source/Runtime/CoreUObject/Private/UObject/PropertyMulticastDelegate.cpp b/Engine/Source/Runtime/CoreUObject/Private/UObject/PropertyMulticastDelegate.cpp
--- a/Engine/Source/Runtime/CoreUObject/Private/UObject/PropertyMulticastDelegate.cpp
+++ b/Engine/Source/Runtime/CoreUObject/Private/UObject/PropertyMulticastDelegate.cpp
@@ -364,7 +364,8 @@ FMulticastScriptDelegate::FInvocationList& UMulticastInlineDelegateProperty::Get
 
 void UMulticastInlineDelegateProperty::SerializeItem(FStructuredArchive::FSlot Slot, void* Value, void const* Defaults) const
 {
-	FArchiveUObjectFromStructuredArchive Ar(Slot);
+	FArchiveUObjectFromStructuredArchive Adapter(Slot);
+	FArchive& Ar = Adapter.GetArchive();
 	Ar << *GetPropertyValuePtr(Value);
 }
 
@@ -469,8 +470,8 @@ FMulticastScriptDelegate::FInvocationList& UMulticastSparseDelegateProperty::Get
 
 void UMulticastSparseDelegateProperty::SerializeItem(FStructuredArchive::FSlot Slot, void* Value, void const* Defaults) const
 {
-	FArchiveUObjectFromStructuredArchive Ar(Slot);
-	SerializeItemInternal(Ar, Value, Defaults);
+	FArchiveUObjectFromStructuredArchive Adapter(Slot);
+	SerializeItemInternal(Adapter.GetArchive(), Value, Defaults);
 }
 
 void UMulticastSparseDelegateProperty::SerializeItemInternal(FArchive& Ar, void* Value, void const* Defaults) const
diff --git a/Engine/Source/Runtime/CoreUObject/Private/UObject/PropertyStruct.cpp b/Engine/Source/Runtime/CoreUObject/Private/UObject/PropertyStruct.cpp
--- a/Engine/Source/Runtime/CoreUObject/Private/UObject/PropertyStruct.cpp
+++ b/Engine/Source/Runtime/CoreUObject/Private/UObject/PropertyStruct.cpp
@@ -353,7 +353,8 @@ EConvertFromTypeResult UStructProperty::ConvertFromType(const FPropertyTag& Tag,
 			}
 			else 
 			{
-				FArchiveUObjectFromStructuredArchive Ar(Slot);
+				FArchiveUObjectFromStructuredArchive Adapter(Slot);
+				FArchive& Ar = Adapter.GetArchive();
 				if (CppStructOps->HasSerializeFromMismatchedTag() && CppStructOps->SerializeFromMismatchedTag(Tag, Ar, DestAddress))
 				{
 
diff --git a/Engine/Source/Runtime/CoreUObject/Private/UObject/SavePackage.cpp b/Engine/Source/Runtime/CoreUObject/Private/UObject/SavePackage.cpp
--- a/Engine/Source/Runtime/CoreUObject/Private/UObject/SavePackage.cpp
+++ b/Engine/Source/Runtime/CoreUObject/Private/UObject/SavePackage.cpp
@@ -5514,7 +5514,7 @@ FSavePackageResultStruct UPackage::Save(UPackage* InOuter, UObject* Base, EObjec
 								else
 								{
 									FArchiveUObjectFromStructuredArchive Adapter(ExportSlot);
-									Export.Object->GetClass()->SerializeDefaultObject(Export.Object, Adapter);
+									Export.Object->GetClass()->SerializeDefaultObject(Export.Object, Adapter.GetArchive());
 								}
 							}
 							else
@@ -5534,7 +5534,7 @@ FSavePackageResultStruct UPackage::Save(UPackage* InOuter, UObject* Base, EObjec
 								else
 								{
 									FArchiveUObjectFromStructuredArchive Adapter(ExportSlot);
-									Export.Object->Serialize(Adapter);
+									Export.Object->Serialize(Adapter.GetArchive());
 								}
 
 #if WITH_EDITOR
diff --git a/Engine/Source/Runtime/CoreUObject/Public/Serialization/ArchiveUObjectFromStructuredArchive.h b/Engine/Source/Runtime/CoreUObject/Public/Serialization/ArchiveUObjectFromStructuredArchive.h
--- a/Engine/Source/Runtime/CoreUObject/Public/Serialization/ArchiveUObjectFromStructuredArchive.h
+++ b/Engine/Source/Runtime/CoreUObject/Public/Serialization/ArchiveUObjectFromStructuredArchive.h
@@ -13,11 +13,13 @@
 
 #if WITH_TEXT_ARCHIVE_SUPPORT
 
-class COREUOBJECT_API FArchiveUObjectFromStructuredArchive : public FArchiveFromStructuredArchive
+class COREUOBJECT_API FArchiveUObjectFromStructuredArchiveImpl : public FArchiveFromStructuredArchiveImpl
 {
+	using Super = FArchiveFromStructuredArchiveImpl;
+
 public:
 
-	FArchiveUObjectFromStructuredArchive(FStructuredArchive::FSlot Slot);
+	FArchiveUObjectFromStructuredArchiveImpl(FStructuredArchive::FSlot Slot);
 
 	using FArchive::operator<<; // For visibility of the overloads we don't override
 
@@ -45,6 +47,21 @@ private:
 	virtual void SerializeInternal(FStructuredArchive::FRecord Record) override;
 };
 
+class FArchiveUObjectFromStructuredArchive
+{
+public:
+	explicit FArchiveUObjectFromStructuredArchive(FStructuredArchive::FSlot InSlot)
+		: Impl(InSlot)
+	{
+	}
+
+	      FArchive& GetArchive()       { return Impl; }
+	const FArchive& GetArchive() const { return Impl; }
+
+private:
+	FArchiveUObjectFromStructuredArchiveImpl Impl;
+};
+
 #else
 
 class COREUOBJECT_API FArchiveUObjectFromStructuredArchive : public FArchiveFromStructuredArchive
diff --git a/Engine/Source/Runtime/CoreUObject/Public/UObject/ObjectMacros.h b/Engine/Source/Runtime/CoreUObject/Public/UObject/ObjectMacros.h
--- a/Engine/Source/Runtime/CoreUObject/Public/UObject/ObjectMacros.h
+++ b/Engine/Source/Runtime/CoreUObject/Public/UObject/ObjectMacros.h
@@ -1415,7 +1415,7 @@ namespace UM
 	}
 
 #define IMPLEMENT_FARCHIVE_SERIALIZER( TClass ) void TClass::Serialize(FArchive& Ar) { TClass::Serialize(FStructuredArchiveFromArchive(Ar).GetSlot().EnterRecord()); }
-#define IMPLEMENT_FSTRUCTUREDARCHIVE_SERIALIZER( TClass ) void TClass::Serialize(FStructuredArchive::FRecord Record) { FArchiveUObjectFromStructuredArchive Ar(Record.EnterField(FIELD_NAME_TEXT("BaseClassAutoGen"))); TClass::Serialize(Ar); }
+#define IMPLEMENT_FSTRUCTUREDARCHIVE_SERIALIZER( TClass ) void TClass::Serialize(FStructuredArchive::FRecord Record) { FArchiveUObjectFromStructuredArchive Ar(Record.EnterField(FIELD_NAME_TEXT("BaseClassAutoGen"))); TClass::Serialize(Ar.GetArchive()); }
 #define DECLARE_FARCHIVE_SERIALIZER( TClass, API ) virtual API void Serialize(FArchive& Ar) override;
 #define DECLARE_FSTRUCTUREDARCHIVE_SERIALIZER( TClass, API ) virtual API void Serialize(FStructuredArchive::FRecord Record) override;
 
