UnrealInsights:
Register cpu profiler TLS buffers for TlsAutoCleanup
Don't create cpu profiler TLS buffers until they're needed

diff --git a/Engine/Source/Runtime/Core/Private/ProfilingDebugging/CpuProfilerTrace.cpp b/Engine/Source/Runtime/Core/Private/ProfilingDebugging/CpuProfilerTrace.cpp
--- a/Engine/Source/Runtime/Core/Private/ProfilingDebugging/CpuProfilerTrace.cpp
+++ b/Engine/Source/Runtime/Core/Private/ProfilingDebugging/CpuProfilerTrace.cpp
@@ -3,7 +3,9 @@
 #include "Trace/Trace.h"
 #include "HAL/PlatformTime.h"
 #include "HAL/PlatformTLS.h"
+#include "HAL/TlsAutoCleanup.h"
 #include "ProfilingDebugging/MiscTrace.h"
+#include "Misc/Parse.h"
 
 #if CPUPROFILERTRACE_ENABLED
 
@@ -30,105 +32,111 @@ public:
 		FullBufferThreshold = MaxBufferSize - MaxEncodedEventSize,
 	};
 
-	struct FThreadState
+	struct FThreadBuffer
+		: public FTlsAutoCleanup
 	{
 		uint64 LastCycle = 0;
-		uint32 ThreadId;
-		uint16 Depth = 0;
 		uint16 BufferSize = 0;
-		bool bEnabled = false;
 		uint8 Buffer[MaxBufferSize];
 	};
 
-	FORCENOINLINE static FThreadState* InitThreadState();
-	FORCENOINLINE static void FlushThreadBuffer(FThreadState* ThreadState);
-	FORCENOINLINE static void EndCapture(FThreadState* ThreadState);
+	FORCENOINLINE static FThreadBuffer* CreateThreadBuffer();
+	FORCENOINLINE static void FlushThreadBuffer(FThreadBuffer* ThreadBuffer);
+	FORCENOINLINE static void EndCapture(FThreadBuffer* ThreadBuffer);
 
-	static thread_local FThreadState* ThreadLocalThreadState;
+	static thread_local uint32 ThreadDepth;
+	static thread_local bool bThreadEnabled;
+	static thread_local FThreadBuffer* ThreadBuffer;
 };
 
-thread_local FCpuProfilerTraceInternal::FThreadState* FCpuProfilerTraceInternal::ThreadLocalThreadState = nullptr;
+thread_local uint32 FCpuProfilerTraceInternal::ThreadDepth = 0;
+thread_local bool FCpuProfilerTraceInternal::bThreadEnabled = false;
+thread_local FCpuProfilerTraceInternal::FThreadBuffer* FCpuProfilerTraceInternal::ThreadBuffer = nullptr;
 
-FCpuProfilerTraceInternal::FThreadState* FCpuProfilerTraceInternal::InitThreadState()
+FCpuProfilerTraceInternal::FThreadBuffer* FCpuProfilerTraceInternal::CreateThreadBuffer()
 {
-	ThreadLocalThreadState = new FThreadState();
-	ThreadLocalThreadState->ThreadId = FPlatformTLS::GetCurrentThreadId();
-	return ThreadLocalThreadState;
+	ThreadBuffer = new FThreadBuffer();
+	ThreadBuffer->Register();
+	return ThreadBuffer;
 }
 
-void FCpuProfilerTraceInternal::FlushThreadBuffer(FThreadState* ThreadState)
+void FCpuProfilerTraceInternal::FlushThreadBuffer(FThreadBuffer* InThreadBuffer)
 {
-	UE_TRACE_LOG(CpuProfiler, EventBatch, ThreadState->BufferSize)
-		<< EventBatch.ThreadId(ThreadState->ThreadId)
-		<< EventBatch.Attachment(ThreadState->Buffer, ThreadState->BufferSize);
-	ThreadState->BufferSize = 0;
+	UE_TRACE_LOG(CpuProfiler, EventBatch, InThreadBuffer->BufferSize)
+		<< EventBatch.ThreadId(FPlatformTLS::GetCurrentThreadId())
+		<< EventBatch.Attachment(InThreadBuffer->Buffer, InThreadBuffer->BufferSize);
+	InThreadBuffer->BufferSize = 0;
 }
 
-void FCpuProfilerTraceInternal::EndCapture(FThreadState* ThreadState)
+void FCpuProfilerTraceInternal::EndCapture(FThreadBuffer* InThreadBuffer)
 {
-	UE_TRACE_LOG(CpuProfiler, EndCapture, ThreadState->BufferSize)
-		<< EndCapture.ThreadId(ThreadState->ThreadId)
-		<< EndCapture.Attachment(ThreadState->Buffer, ThreadState->BufferSize);
-	ThreadState->BufferSize = 0;
+	UE_TRACE_LOG(CpuProfiler, EndCapture, InThreadBuffer->BufferSize)
+		<< EndCapture.ThreadId(FPlatformTLS::GetCurrentThreadId())
+		<< EndCapture.Attachment(InThreadBuffer->Buffer, InThreadBuffer->BufferSize);
+	InThreadBuffer->BufferSize = 0;
 }
 
 void FCpuProfilerTrace::OutputBeginEvent(uint16 SpecId)
 {
-	FCpuProfilerTraceInternal::FThreadState* ThreadState = FCpuProfilerTraceInternal::ThreadLocalThreadState;
-	if (!ThreadState)
-	{
-		ThreadState = FCpuProfilerTraceInternal::InitThreadState();
-	}
-	++ThreadState->Depth;
+	++FCpuProfilerTraceInternal::ThreadDepth;
 	bool bEventEnabled = UE_TRACE_EVENT_IS_ENABLED(CpuProfiler, EventBatch);
-	if ((!bEventEnabled) & (bEventEnabled == ThreadState->bEnabled))
+	if ((!bEventEnabled) & (bEventEnabled == FCpuProfilerTraceInternal::bThreadEnabled))
 	{
 		return;
 	}
-	ThreadState->bEnabled = bEventEnabled;
+	FCpuProfilerTraceInternal::bThreadEnabled = bEventEnabled;
+	FCpuProfilerTraceInternal::FThreadBuffer* ThreadBuffer = FCpuProfilerTraceInternal::ThreadBuffer;
+	if (!ThreadBuffer)
+	{
+		ThreadBuffer = FCpuProfilerTraceInternal::CreateThreadBuffer();
+	}
 	if (!bEventEnabled)
 	{
-		FCpuProfilerTraceInternal::EndCapture(ThreadState);
+		FCpuProfilerTraceInternal::EndCapture(ThreadBuffer);
 		return;
 	}
-	if (ThreadState->BufferSize >= FCpuProfilerTraceInternal::FullBufferThreshold)
+	if (ThreadBuffer->BufferSize >= FCpuProfilerTraceInternal::FullBufferThreshold)
 	{
-		FCpuProfilerTraceInternal::FlushThreadBuffer(ThreadState);
+		FCpuProfilerTraceInternal::FlushThreadBuffer(ThreadBuffer);
 	}
 	uint64 Cycle = FPlatformTime::Cycles64();
-	uint64 CycleDiff = Cycle - ThreadState->LastCycle;
-	ThreadState->LastCycle = Cycle;
-	uint8* BufferPtr = ThreadState->Buffer + ThreadState->BufferSize;
+	uint64 CycleDiff = Cycle - ThreadBuffer->LastCycle;
+	ThreadBuffer->LastCycle = Cycle;
+	uint8* BufferPtr = ThreadBuffer->Buffer + ThreadBuffer->BufferSize;
 	FTraceUtils::Encode7bit((CycleDiff << 1) | 1ull, BufferPtr);
 	FTraceUtils::Encode7bit(SpecId, BufferPtr);
-	ThreadState->BufferSize = BufferPtr - ThreadState->Buffer;
+	ThreadBuffer->BufferSize = BufferPtr - ThreadBuffer->Buffer;
 }
 
 void FCpuProfilerTrace::OutputEndEvent()
 {
-	FCpuProfilerTraceInternal::FThreadState* ThreadState = FCpuProfilerTraceInternal::ThreadLocalThreadState;
-	--ThreadState->Depth;
+	--FCpuProfilerTraceInternal::ThreadDepth;
 	bool bEventEnabled = UE_TRACE_EVENT_IS_ENABLED(CpuProfiler, EventBatch);
-	if ((!bEventEnabled) & (bEventEnabled == ThreadState->bEnabled))
+	if ((!bEventEnabled) & (bEventEnabled == FCpuProfilerTraceInternal::bThreadEnabled))
 	{
 		return;
 	}
-	ThreadState->bEnabled = bEventEnabled;
+	FCpuProfilerTraceInternal::bThreadEnabled = bEventEnabled;
+	FCpuProfilerTraceInternal::FThreadBuffer* ThreadBuffer = FCpuProfilerTraceInternal::ThreadBuffer;
+	if (!ThreadBuffer)
+	{
+		ThreadBuffer = FCpuProfilerTraceInternal::CreateThreadBuffer();
+	}
 	if (!bEventEnabled)
 	{
-		FCpuProfilerTraceInternal::EndCapture(ThreadState);
+		FCpuProfilerTraceInternal::EndCapture(ThreadBuffer);
 		return;
 	}
-	if ((ThreadState->Depth == 0) | (ThreadState->BufferSize >= FCpuProfilerTraceInternal::FullBufferThreshold))
+	if ((FCpuProfilerTraceInternal::ThreadDepth == 0) | (ThreadBuffer->BufferSize >= FCpuProfilerTraceInternal::FullBufferThreshold))
 	{
-		FCpuProfilerTraceInternal::FlushThreadBuffer(ThreadState);
+		FCpuProfilerTraceInternal::FlushThreadBuffer(ThreadBuffer);
 	}
 	uint64 Cycle = FPlatformTime::Cycles64();
-	uint64 CycleDiff = Cycle - ThreadState->LastCycle;
-	ThreadState->LastCycle = Cycle;
-	uint8* BufferPtr = ThreadState->Buffer + ThreadState->BufferSize;
+	uint64 CycleDiff = Cycle - ThreadBuffer->LastCycle;
+	ThreadBuffer->LastCycle = Cycle;
+	uint8* BufferPtr = ThreadBuffer->Buffer + ThreadBuffer->BufferSize;
 	FTraceUtils::Encode7bit(CycleDiff << 1, BufferPtr);
-	ThreadState->BufferSize = BufferPtr - ThreadState->Buffer;
+	ThreadBuffer->BufferSize = BufferPtr - ThreadBuffer->Buffer;
 }
 
 uint16 FCpuProfilerTrace::OutputEventType(const TCHAR* Name, ECpuProfilerGroup Group)
@@ -145,9 +153,9 @@ uint16 FCpuProfilerTrace::OutputEventType(const TCHAR* Name, ECpuProfilerGroup G
 	return Id;
 }
 
-void FCpuProfilerTrace::Init(bool bStartEnabled)
+void FCpuProfilerTrace::Init(const TCHAR* CmdLine)
 {
-	if (bStartEnabled)
+	if (FParse::Param(CmdLine, TEXT("cpuprofilertrace")))
 	{
 		UE_TRACE_EVENT_IS_ENABLED(CpuProfiler, EventBatch);
 		Trace::ToggleEvent(TEXT("CpuProfiler.EventBatch"), true);
diff --git a/Engine/Source/Runtime/Core/Public/ProfilingDebugging/CpuProfilerTrace.h b/Engine/Source/Runtime/Core/Public/ProfilingDebugging/CpuProfilerTrace.h
--- a/Engine/Source/Runtime/Core/Public/ProfilingDebugging/CpuProfilerTrace.h
+++ b/Engine/Source/Runtime/Core/Public/ProfilingDebugging/CpuProfilerTrace.h
@@ -22,7 +22,7 @@ enum ECpuProfilerGroup
 
 struct FCpuProfilerTrace
 {
-	CORE_API static void Init(bool bStartEnabled);
+	CORE_API static void Init(const TCHAR* CmdLine);
 	CORE_API static uint16 OutputEventType(const TCHAR* Name, ECpuProfilerGroup Group);
 	CORE_API static void OutputBeginEvent(uint16 SpecId);
 	CORE_API static void OutputEndEvent();
@@ -41,6 +41,9 @@ struct FCpuProfilerTrace
 	};
 };
 
+#define TRACE_CPUPROFILER_INIT(CmdLine) \
+	FCpuProfilerTrace::Init(CmdLine);
+
 #define TRACE_CPUPROFILER_EVENT_SCOPE_TEXT_GROUP(Name, Group) \
 	static uint16 PREPROCESSOR_JOIN(__CpuProfilerEventSpecId, __LINE__); \
 	if (PREPROCESSOR_JOIN(__CpuProfilerEventSpecId, __LINE__) == 0) { \
@@ -59,6 +62,7 @@ struct FCpuProfilerTrace
 
 #else
 
+#define TRACE_CPUPROFILER_INIT(CmdLine)
 #define TRACE_CPUPROFILER_EVENT_SCOPE_GROUP(Name, Group)
 #define TRACE_CPUPROFILER_EVENT_SCOPE(Name)
 #define TRACE_CPUPROFILER_EVENT_SCOPE_TEXT_GROUP(Name, Group)
diff --git a/Engine/Source/Runtime/Launch/Private/LaunchEngineLoop.cpp b/Engine/Source/Runtime/Launch/Private/LaunchEngineLoop.cpp
--- a/Engine/Source/Runtime/Launch/Private/LaunchEngineLoop.cpp
+++ b/Engine/Source/Runtime/Launch/Private/LaunchEngineLoop.cpp
@@ -1159,9 +1159,7 @@ DECLARE_CYCLE_STAT( TEXT( "FEngineLoop::PreInit.AfterStats" ), STAT_FEngineLoop_
 int32 FEngineLoop::PreInit(const TCHAR* CmdLine)
 {
 	TRACE_REGISTER_GAME_THREAD(FPlatformTLS::GetCurrentThreadId());
-#if CPUPROFILERTRACE_ENABLED
-	FCpuProfilerTrace::Init(FParse::Param(CmdLine, TEXT("cpuprofilertrace")));
-#endif
+	TRACE_CPUPROFILER_INIT(CmdLine);
 	TRACE_PLATFORMFILE_INIT(CmdLine);
 
 	SCOPED_BOOT_TIMING("FEngineLoop::PreInit");
