Allocation-less FString::MatchesWildcard().

diff --git a/Engine/Source/Runtime/Core/Private/Containers/String.cpp b/Engine/Source/Runtime/Core/Private/Containers/String.cpp
--- a/Engine/Source/Runtime/Core/Private/Containers/String.cpp
+++ b/Engine/Source/Runtime/Core/Private/Containers/String.cpp
@@ -39,6 +39,11 @@ namespace UE4String_Private
 		// Skip over common initial non-wildcard-char sequence of Target and Wildcard
 		for (;;)
 		{
+			if (WildcardLength == 0)
+			{
+				return TargetLength == 0;
+			}
+
 			TCHAR WCh = *Wildcard;
 			if (WCh == TEXT('*') || WCh == TEXT('?'))
 			{
@@ -50,11 +55,6 @@ namespace UE4String_Private
 				return false;
 			}
 
-			if (WCh == TEXT('\0'))
-			{
-				return true;
-			}
-
 			++Target;
 			++Wildcard;
 			--TargetLength;
@@ -107,7 +107,7 @@ namespace UE4String_Private
 
 		for (int32 Index = 0; Index <= MaxNum; ++Index)
 		{
-			if (MatchesWildcardRecursive<CompareType>(Target, TargetLength - Index, Wildcard, WildcardLength))
+			if (MatchesWildcardRecursive<CompareType>(Target + Index, TargetLength - Index, Wildcard, WildcardLength))
 			{
 				return true;
 			}
@@ -913,79 +913,19 @@ int32 FString::ParseIntoArray( TArray<FString>& OutArray, const TCHAR* pchDelim,
 	return OutArray.Num();
 }
 
-bool FString::MatchesWildcard(const FString& InWildcard, ESearchCase::Type SearchCase) const
+bool FString::MatchesWildcard(const TCHAR* InWildcard, ESearchCase::Type SearchCase) const
 {
-	FString Wildcard(InWildcard);
-	FString Target(*this);
-	int32 IndexOfStar = Wildcard.Find(TEXT("*"), ESearchCase::CaseSensitive, ESearchDir::FromEnd); // last occurance
-	int32 IndexOfQuestion = Wildcard.Find( TEXT("?"), ESearchCase::CaseSensitive, ESearchDir::FromEnd); // last occurance
-	int32 Suffix = FMath::Max<int32>(IndexOfStar, IndexOfQuestion);
-	if (Suffix == INDEX_NONE)
-	{
-		// no wildcards
-		if (SearchCase == ESearchCase::IgnoreCase)
-		{
-			return FCString::Stricmp( *Target, *Wildcard )==0;
-		}
-		else
-		{
-			return FCString::Strcmp( *Target, *Wildcard )==0;
-		}
-	}
-	else
-	{
-		if (Suffix + 1 < Wildcard.Len())
-		{
-			FString SuffixString = Wildcard.Mid(Suffix + 1);
-			if (!Target.EndsWith(SuffixString, SearchCase))
-			{
-				return false;
-			}
-			Wildcard = Wildcard.Left(Suffix + 1);
-			Target = Target.Left(Target.Len() - SuffixString.Len());
-		}
-		int32 PrefixIndexOfStar = Wildcard.Find(TEXT("*"), ESearchCase::CaseSensitive); 
-		int32 PrefixIndexOfQuestion = Wildcard.Find(TEXT("?"), ESearchCase::CaseSensitive);
-		int32 Prefix = FMath::Min<int32>(PrefixIndexOfStar < 0 ? MAX_int32 : PrefixIndexOfStar, PrefixIndexOfQuestion < 0 ? MAX_int32 : PrefixIndexOfQuestion);
-		check(Prefix >= 0 && Prefix < Wildcard.Len());
-		if (Prefix > 0)
-		{
-			FString PrefixString = Wildcard.Left(Prefix);
-			if (!Target.StartsWith(PrefixString, SearchCase))
-			{
-				return false;
-			}
-			Wildcard = Wildcard.Mid(Prefix);
-			Target = Target.Mid(Prefix);
-		}
-	}
-	// This routine is very slow, though it does ok with one wildcard
-	check(Wildcard.Len());
-	TCHAR FirstWild = Wildcard[0];
-	Wildcard = Wildcard.Right(Wildcard.Len() - 1);
-	if (FirstWild == TEXT('*') || FirstWild == TEXT('?'))
+	const TCHAR* Target = **this;
+	int32        TargetLength = Len();
+	int32        WildcardLength = FCString::Strlen(InWildcard);
+
+	if (SearchCase == ESearchCase::CaseSensitive)
 	{
-		if (!Wildcard.Len())
-		{
-			if (FirstWild == TEXT('*') || Target.Len() < 2)
-			{
-				return true;
-			}
-		}
-		int32 MaxNum = FMath::Min<int32>(Target.Len(), FirstWild == TEXT('?') ? 1 : MAX_int32);
-		for (int32 Index = 0; Index <= MaxNum; Index++)
-		{
-			if (Target.Right(Target.Len() - Index).MatchesWildcard(Wildcard, SearchCase))
-			{
-				return true;
-			}
-		}
-		return false;
+		return UE4String_Private::MatchesWildcardRecursive<UE4String_Private::FCompareCharsCaseSensitive>(Target, TargetLength, InWildcard, WildcardLength);
 	}
 	else
 	{
-		check(0); // we should have dealt with prefix comparison above
-		return false;
+		return UE4String_Private::MatchesWildcardRecursive<UE4String_Private::FCompareCharsCaseInsensitive>(Target, TargetLength, InWildcard, WildcardLength);
 	}
 }
 
diff --git a/Engine/Source/Runtime/Core/Public/Containers/UnrealString.h b/Engine/Source/Runtime/Core/Public/Containers/UnrealString.h
--- a/Engine/Source/Runtime/Core/Public/Containers/UnrealString.h
+++ b/Engine/Source/Runtime/Core/Public/Containers/UnrealString.h
@@ -1508,7 +1508,20 @@ public:
 	 * @return true if this string matches the *?-type wildcard given. 
 	 * @warning This is a simple, SLOW routine. Use with caution
 	 */
-	bool MatchesWildcard(const FString& Wildcard, ESearchCase::Type SearchCase = ESearchCase::IgnoreCase) const;
+	bool MatchesWildcard(const TCHAR* Wildcard, ESearchCase::Type SearchCase = ESearchCase::IgnoreCase) const;
+
+	/**
+	 * Searches this string for a given wild card
+	 *
+	 * @param Wildcard		*?-type wildcard
+	 * @param SearchCase	Indicates whether the search is case sensitive or not ( defaults to ESearchCase::IgnoreCase )
+	 * @return true if this string matches the *?-type wildcard given.
+	 * @warning This is a simple, SLOW routine. Use with caution
+	 */
+	FORCEINLINE bool MatchesWildcard(const FString& Wildcard, ESearchCase::Type SearchCase = ESearchCase::IgnoreCase) const
+	{
+		return MatchesWildcard(*Wildcard, SearchCase);
+	}
 
 	/**
 	 * Removes whitespace characters from the front of this string.
