FScriptSet Add() optimization

diff --git a/Engine/Source/Runtime/Core/Public/Containers/Set.h b/Engine/Source/Runtime/Core/Public/Containers/Set.h
--- a/Engine/Source/Runtime/Core/Public/Containers/Set.h
+++ b/Engine/Source/Runtime/Core/Public/Containers/Set.h
@@ -1620,24 +1620,43 @@ public:
 		}
 	}
 
+private:
+	int32 FindIndexImpl(const void* Element, const FScriptSetLayout& Layout, uint32 KeyHash, TFunctionRef<bool (const void*, const void*)> EqualityFn)
+	{
+		const int32  HashIndex = KeyHash & (HashSize - 1);
+
+		uint8* CurrentElement = nullptr;
+		for (FSetElementId ElementId = GetTypedHash(HashIndex);
+			ElementId.IsValidId();
+			ElementId = GetHashNextIdRef(CurrentElement, Layout))
+		{
+			CurrentElement = (uint8*)Elements.GetData(ElementId, Layout.SparseArrayLayout);
+			if (EqualityFn(Element, CurrentElement))
+			{
+				return ElementId;
+			}
+		}
+	
+		return INDEX_NONE;
+	}
+
+public:
 	int32 FindIndex(const void* Element, const FScriptSetLayout& Layout, TFunctionRef<uint32 (const void*)> GetKeyHash, TFunctionRef<bool (const void*, const void*)> EqualityFn)
 	{
 		if (Elements.Num())
 		{
-			const uint32 KeyHash = GetKeyHash(Element);
-			const int32  HashIndex = KeyHash & (HashSize - 1);
+			return FindIndexImpl(Element, Layout, GetKeyHash(Element), EqualityFn);
+		}
 
-			uint8* CurrentElement = nullptr;
-			for (FSetElementId ElementId = GetTypedHash(HashIndex);
-				ElementId.IsValidId();
-				ElementId = GetHashNextIdRef(CurrentElement, Layout))
-			{
-				CurrentElement = (uint8*)Elements.GetData(ElementId, Layout.SparseArrayLayout);
-				if (EqualityFn(Element, CurrentElement))
-				{
-					return ElementId;
-				}
-			}
+		return INDEX_NONE;
+	}
+
+
+	int32 FindIndexByHash(const void* Element, const FScriptSetLayout& Layout, uint32 KeyHash, TFunctionRef<bool (const void*, const void*)> EqualityFn)
+	{
+		if (Elements.Num())
+		{
+			return FindIndexImpl(Element, Layout, KeyHash, EqualityFn);
 		}
 
 		return INDEX_NONE;
@@ -1645,10 +1664,9 @@ public:
 
 	void Add(const void* Element, const FScriptSetLayout& Layout, TFunctionRef<uint32(const void*)> GetKeyHash, TFunctionRef<bool(const void*, const void*)> EqualityFn, TFunctionRef<void(void*)> ConstructFn, TFunctionRef<void(void*)> DestructFn)
 	{
-		// Minor efficiency concern: we hash the element both here and in the FindIndex() call
 		uint32 KeyHash = GetKeyHash(Element);
 
-		int32 NewElementIndex = FindIndex(Element, Layout, GetKeyHash, EqualityFn);
+		int32 NewElementIndex = FindIndexByHash(Element, Layout, KeyHash, EqualityFn);
 		if (NewElementIndex != INDEX_NONE)
 		{
 			void* ElementPtr = Elements.GetData(NewElementIndex, Layout.SparseArrayLayout);
@@ -1658,7 +1676,7 @@ public:
 
 			// We don't update the hash because we don't need to - the new element
 			// should have the same hash, but let's just check.
-			check(KeyHash == GetKeyHash(ElementPtr));
+			checkSlow(KeyHash == GetKeyHash(ElementPtr));
 		}
 		else
 		{
