TArrayView is now constructible from any contiguous range, to save hardcoding every possible type as a constructor.
Existing range types tagged as contiguous.
Fixes to reference value categories for array overloads of GetData() and GetNum().

diff --git a/Engine/Plugins/Experimental/ChaosEditor/Source/FractureEditor/Private/FractureEditorMode.cpp b/Engine/Plugins/Experimental/ChaosEditor/Source/FractureEditor/Private/FractureEditorMode.cpp
--- a/Engine/Plugins/Experimental/ChaosEditor/Source/FractureEditor/Private/FractureEditorMode.cpp
+++ b/Engine/Plugins/Experimental/ChaosEditor/Source/FractureEditor/Private/FractureEditorMode.cpp
@@ -197,7 +197,7 @@ bool FFractureEditorMode::FrustumSelect(const FConvexVolume& InFrustum, FEditorV
 
 				TMap<int32, FBox> BoundsToBone;
 			
-				GetActorGlobalBounds(MakeArrayView<UGeometryCollectionComponent*>(GeometryComponents), BoundsToBone);
+				GetActorGlobalBounds(MakeArrayView(GeometryComponents), BoundsToBone);
 
 				TArray<int32> SelectedBonesArray;
 
diff --git a/Engine/Source/Runtime/Core/Public/Containers/ArrayView.h b/Engine/Source/Runtime/Core/Public/Containers/ArrayView.h
--- a/Engine/Source/Runtime/Core/Public/Containers/ArrayView.h
+++ b/Engine/Source/Runtime/Core/Public/Containers/ArrayView.h
@@ -3,11 +3,9 @@
 #pragma once
 
 #include "CoreTypes.h"
-#include "Templates/IsInitializerList.h"
 #include "Templates/PointerIsConvertibleFromTo.h"
 #include "Misc/AssertionMacros.h"
 #include "Templates/UnrealTypeTraits.h"
-#include "Math/NumericLimits.h"
 #include "Containers/Array.h"
 
 namespace ArrayViewPrivate
@@ -25,6 +23,23 @@ namespace ArrayViewPrivate
 		 */
 		enum { Value = TPointerIsConvertibleFromTo<T*, ElementType* const>::Value };
 	};
+
+	// Simply forwards to an unqualified GetData(), but can be called from within TArrayView
+	// where GetData() is already a member and so hides any others.
+	template <typename T>
+	FORCEINLINE decltype(auto) GetDataHelper(T&& Arg)
+	{
+		return GetData(Forward<T>(Arg));
+	}
+
+	/**
+	 * Trait testing whether a type is compatible with the view type
+	 */
+	template <typename RangeType, typename ElementType>
+	struct TIsCompatibleRangeType
+	{
+		static constexpr bool Value = TIsCompatibleElementType<typename TRemovePointer<decltype(GetData(DeclVal<RangeType&>()))>::Type, ElementType>::Value;
+	};
 }
 
 /**
@@ -70,79 +85,28 @@ private:
 	template <typename T>
 	using TIsCompatibleElementType = ArrayViewPrivate::TIsCompatibleElementType<T, ElementType>;
 
-public:
-	/**
-	 * Copy constructor from another view
-	 *
-	 * @param Other The source array view to copy
-	 */
-	template <typename OtherElementType,
-		typename = typename TEnableIf<TIsCompatibleElementType<OtherElementType>::Value>::Type>
-	FORCEINLINE TArrayView(const TArrayView<OtherElementType>& Other)
-		: DataPtr(Other.GetData())
-		, ArrayNum(Other.Num())
-	{
-	}
-
-	/**
-	 * Construct a view of a C array with a compatible element type
-	 *
-	 * @param Other The source array to view.
-	 */
-	template <typename OtherElementType, size_t Size,
-		typename = typename TEnableIf<TIsCompatibleElementType<OtherElementType>::Value>::Type>
-	FORCEINLINE TArrayView(OtherElementType (&Other)[Size])
-		: DataPtr(Other)
-		, ArrayNum((int32)Size)
-	{
-		static_assert(Size <= MAX_int32, "Array Size too large! Size is only int32 for compatibility with TArray!");
-	}
-
-	/**
-	 * Construct a view of a C array with a compatible element type
-	 *
-	 * @param Other The source array to view.
-	 */
-	template <typename OtherElementType, size_t Size,
-		typename = typename TEnableIf<TIsCompatibleElementType<const OtherElementType>::Value>::Type>
-	FORCEINLINE TArrayView(const OtherElementType (&Other)[Size])
-		: DataPtr(Other)
-		, ArrayNum((int32)Size)
-	{
-		static_assert(Size <= MAX_int32, "Array Size too large! Size is only int32 for compatibility with TArray!");
-	}
-
-	/**
-	 * Construct a view of a TArray with a compatible element type
-	 *
-	 * @param Other The source array to view.
-	 */
-	template <typename OtherElementType, typename OtherAllocator,
-		typename = typename TEnableIf<TIsCompatibleElementType<OtherElementType>::Value>::Type>
-	FORCEINLINE TArrayView(TArray<OtherElementType, OtherAllocator>& Other)
-		: DataPtr(Other.GetData())
-		, ArrayNum(Other.Num())
-	{
-	}
-
-	template <typename OtherElementType, typename OtherAllocator,
-		typename = typename TEnableIf<TIsCompatibleElementType<const OtherElementType>::Value>::Type>
-	FORCEINLINE TArrayView(const TArray<OtherElementType, OtherAllocator>& Other)
-		: DataPtr(Other.GetData())
-		, ArrayNum(Other.Num())
-	{
-	}
+	template <typename T>
+	using TIsCompatibleRangeType = ArrayViewPrivate::TIsCompatibleRangeType<T, ElementType>;
 
+public:
 	/**
-	 * Construct a view from an initializer list with a compatible element type
+	 * Constructor from another range
 	 *
-	 * @param List The initializer list to view.
+	 * @param Other The source range to copy
 	 */
-	template <typename InitializerListType,
-		typename = typename TEnableIf<TAnd<TIsInitializerList<InitializerListType>, TIsCompatibleElementType<const typename InitializerListType::value_type>>::Value>::Type>
-	FORCEINLINE TArrayView(InitializerListType& List)
-		: DataPtr(&*List.begin())
-		, ArrayNum(List.size())
+	template <
+		typename OtherRangeType,
+		typename CVUnqualifiedOtherRangeType = typename TRemoveCV<typename TRemoveReference<OtherRangeType>::Type>::Type,
+		typename = typename TEnableIf<
+			TAnd<
+				TIsContiguousContainer<CVUnqualifiedOtherRangeType>,
+				TIsCompatibleRangeType<OtherRangeType>
+			>::Value
+		>::Type
+	>
+	FORCEINLINE TArrayView(OtherRangeType&& Other)
+		: DataPtr(ArrayViewPrivate::GetDataHelper(Forward<OtherRangeType>(Other)))
+		, ArrayNum(GetNum(Forward<OtherRangeType>(Other)))
 	{
 	}
 
@@ -161,23 +125,6 @@ public:
 		check(ArrayNum >= 0);
 	}
 
-	/**
-	 * Assignment operator
-	 *
-	 * @param Other The source array view to assign from
-	 */
-	template <typename OtherElementType,
-		typename = typename TEnableIf<TIsCompatibleElementType<OtherElementType>::Value>::Type>
-	FORCEINLINE TArrayView& operator=(const TArrayView<OtherElementType>& Other)
-	{
-		if (this != &Other)
-		{
-			DataPtr = Other.DataPtr;
-			ArrayNum = Other.ArrayNum;
-		}
-		return *this;
-	}
-
 public:
 
 	/**
@@ -579,35 +526,18 @@ struct TIsContiguousContainer<TArrayView<T>>
 
 //////////////////////////////////////////////////////////////////////////
 
-template <typename ElementType, size_t Size>
-TArrayView<ElementType> MakeArrayView(ElementType(&Other)[Size])
-{
-	return TArrayView<ElementType>(Other);
-}
-
-template <typename ElementType, typename Allocator>
-TArrayView<ElementType> MakeArrayView(TArray<ElementType, Allocator>& Other)
-{
-	return TArrayView<ElementType>(Other);
-}
-
-template <typename ElementType, typename Allocator>
-TArrayView<const ElementType> MakeArrayView(const TArray<ElementType, Allocator>& Other)
-{
-	return TArrayView<const ElementType>(Other);
-}
-
-template <typename InitializerListType>
-typename TEnableIf<
-	TIsInitializerList<InitializerListType>::Value,
-	TArrayView<const typename InitializerListType::value_type>
->::Type MakeArrayView(InitializerListType& List)
+template <
+	typename OtherRangeType,
+	typename CVUnqualifiedOtherRangeType = typename TRemoveCV<typename TRemoveReference<OtherRangeType>::Type>::Type,
+	typename = typename TEnableIf<TIsContiguousContainer<CVUnqualifiedOtherRangeType>::Value>::Type
+>
+auto MakeArrayView(OtherRangeType&& Other)
 {
-	return TArrayView<const typename InitializerListType::value_type>(List);
+	return TArrayView<typename TRemovePointer<decltype(GetData(DeclVal<OtherRangeType&>()))>::Type>(Forward<OtherRangeType>(Other));
 }
 
 template<typename ElementType>
-TArrayView<ElementType> MakeArrayView(ElementType* Pointer, int32 Size)
+auto MakeArrayView(ElementType* Pointer, int32 Size)
 {
 	return TArrayView<ElementType>(Pointer, Size);
 }
diff --git a/Engine/Source/Runtime/Core/Public/Containers/StaticArray.h b/Engine/Source/Runtime/Core/Public/Containers/StaticArray.h
--- a/Engine/Source/Runtime/Core/Public/Containers/StaticArray.h
+++ b/Engine/Source/Runtime/Core/Public/Containers/StaticArray.h
@@ -8,15 +8,17 @@
 #include "Delegates/IntegerSequence.h"
 
 /** An array with a static number of elements. */
-template<typename TElement, uint32 NumElements, uint32 Alignment = alignof(TElement)>
+template <typename InElementType, uint32 NumElements, uint32 Alignment = alignof(InElementType)>
 class alignas(Alignment) TStaticArray
 {
 public:
+	using ElementType = InElementType;
+
 	TStaticArray() 
 		: Storage()
 	{}
 
-	explicit TStaticArray(const TElement& DefaultElement) 
+	explicit TStaticArray(const InElementType& DefaultElement)
 		: Storage(TMakeIntegerSequence<uint32, NumElements>(), DefaultElement)
 	{}
 
@@ -36,13 +38,13 @@ public:
 	}
 
 	// Accessors.
-	FORCEINLINE_DEBUGGABLE TElement& operator[](uint32 Index)
+	FORCEINLINE_DEBUGGABLE InElementType& operator[](uint32 Index)
 	{
 		checkSlow(Index < NumElements);
 		return Storage.Elements[Index].Element;
 	}
 
-	FORCEINLINE_DEBUGGABLE const TElement& operator[](uint32 Index) const
+	FORCEINLINE_DEBUGGABLE const InElementType& operator[](uint32 Index) const
 	{
 		checkSlow(Index < NumElements);
 		return Storage.Elements[Index].Element;
@@ -75,7 +77,11 @@ public:
 
 	/** The number of elements in the array. */
 	FORCEINLINE_DEBUGGABLE int32 Num() const { return NumElements; }
-	
+
+	/** A pointer to the first element of the array */
+	FORCEINLINE_DEBUGGABLE       InElementType* GetData()       { return Storage.Elements; }
+	FORCEINLINE_DEBUGGABLE const InElementType* GetData() const { return Storage.Elements; }
+
 	/** Hash function. */
 	friend uint32 GetTypeHash(const TStaticArray& Array)
 	{
@@ -92,12 +98,12 @@ private:
 	struct alignas(Alignment) TArrayStorageElementAligned
 	{
 		TArrayStorageElementAligned() {}
-		TArrayStorageElementAligned(const TElement& InElement)
+		TArrayStorageElementAligned(const InElementType& InElement)
 			: Element(InElement)
 		{
 		}
 
-		TElement Element;
+		InElementType Element;
 	};
 
 	struct TArrayStorage
@@ -108,7 +114,7 @@ private:
 		}
 
 		template<uint32... Indices>
-		TArrayStorage(TIntegerSequence<uint32, Indices...>, const TElement& DefaultElement)
+		TArrayStorage(TIntegerSequence<uint32, Indices...>, const InElementType& DefaultElement)
 			: Elements { ((void)Indices, DefaultElement)... } //Integer Sequence pack expansion duplicates DefaultElement NumElements times and the comma operator throws away the index
 		{
 		}
@@ -120,10 +126,10 @@ private:
 };
 
 /** Creates a static array filled with the specified value. */
-template<typename TElement,uint32 NumElements>
-TStaticArray<TElement,NumElements> MakeUniformStaticArray(typename TCallTraits<TElement>::ParamType InValue)
+template <typename InElementType, uint32 NumElements>
+TStaticArray<InElementType,NumElements> MakeUniformStaticArray(typename TCallTraits<InElementType>::ParamType InValue)
 {
-	TStaticArray<TElement,NumElements> Result;
+	TStaticArray<InElementType,NumElements> Result;
 	for(uint32 ElementIndex = 0;ElementIndex < NumElements;++ElementIndex)
 	{
 		Result[ElementIndex] = InValue;
@@ -131,3 +137,8 @@ TStaticArray<TElement,NumElements> MakeUniformStaticArray(typename TCallTraits<T
 	return Result;
 }
 
+template <typename ElementType, uint32 NumElements, uint32 Alignment>
+struct TIsContiguousContainer<TStaticArray<ElementType, NumElements, Alignment>>
+{
+	enum { Value = true };
+};
diff --git a/Engine/Source/Runtime/Core/Public/Serialization/BufferArchive.h b/Engine/Source/Runtime/Core/Public/Serialization/BufferArchive.h
--- a/Engine/Source/Runtime/Core/Public/Serialization/BufferArchive.h
+++ b/Engine/Source/Runtime/Core/Public/Serialization/BufferArchive.h
@@ -26,3 +26,8 @@ public:
 	virtual FString GetArchiveName() const { return *FString::Printf( TEXT("FBufferArchive %s"), *ArchiveName.ToString()); }
 };
 
+template <>
+struct TIsContiguousContainer<FBufferArchive>
+{
+	static constexpr bool Value = TIsContiguousContainer<TArray<uint8>>::Value;
+};
diff --git a/Engine/Source/Runtime/Core/Public/Templates/UnrealTemplate.h b/Engine/Source/Runtime/Core/Public/Templates/UnrealTemplate.h
--- a/Engine/Source/Runtime/Core/Public/Templates/UnrealTemplate.h
+++ b/Engine/Source/Runtime/Core/Public/Templates/UnrealTemplate.h
@@ -87,14 +87,13 @@ auto GetData(T&& Container) -> decltype(Container.GetData())
 	return Container.GetData();
 }
 
-template <typename T, SIZE_T N>
-CONSTEXPR T* GetData(T (&Container)[N])
-{
-	return Container;
-}
+template <typename T, SIZE_T N> CONSTEXPR       T* GetData(      T (& Container)[N]) { return Container; }
+template <typename T, SIZE_T N> CONSTEXPR       T* GetData(      T (&&Container)[N]) { return Container; }
+template <typename T, SIZE_T N> CONSTEXPR const T* GetData(const T (& Container)[N]) { return Container; }
+template <typename T, SIZE_T N> CONSTEXPR const T* GetData(const T (&&Container)[N]) { return Container; }
 
 template <typename T>
-CONSTEXPR T* GetData(std::initializer_list<T> List)
+CONSTEXPR const T* GetData(std::initializer_list<T> List)
 {
 	return List.begin();
 }
@@ -108,11 +107,10 @@ SIZE_T GetNum(T&& Container)
 	return (SIZE_T)Container.Num();
 }
 
-template <typename T, SIZE_T N>
-CONSTEXPR SIZE_T GetNum(T (&Container)[N])
-{
-	return N;
-}
+template <typename T, SIZE_T N> CONSTEXPR SIZE_T GetNum(      T (& Container)[N]) { return N; }
+template <typename T, SIZE_T N> CONSTEXPR SIZE_T GetNum(      T (&&Container)[N]) { return N; }
+template <typename T, SIZE_T N> CONSTEXPR SIZE_T GetNum(const T (& Container)[N]) { return N; }
+template <typename T, SIZE_T N> CONSTEXPR SIZE_T GetNum(const T (&&Container)[N]) { return N; }
 
 template <typename T>
 CONSTEXPR SIZE_T GetNum(std::initializer_list<T> List)
diff --git a/Engine/Source/Runtime/Engine/Public/Animation/AnimCompressionTypes.h b/Engine/Source/Runtime/Engine/Public/Animation/AnimCompressionTypes.h
--- a/Engine/Source/Runtime/Engine/Public/Animation/AnimCompressionTypes.h
+++ b/Engine/Source/Runtime/Engine/Public/Animation/AnimCompressionTypes.h
@@ -654,6 +654,12 @@ struct TAllocatorTraits<TMaybeMappedAllocator<Alignment>> : TAllocatorTraitsBase
 	enum { SupportsMove = true };
 };
 
+template<typename T, uint32 Alignment>
+struct TIsContiguousContainer<TMaybeMappedArray<T, Alignment>>
+{
+	static constexpr bool Value = TIsContiguousContainer<TArray<T, TMaybeMappedAllocator<Alignment>>>::Value;
+};
+
 struct ENGINE_API FCompressedAnimSequence
 {
 public:
diff --git a/Engine/Source/Runtime/Experimental/ChaosCore/Public/Chaos/ArrayCollectionArray.h b/Engine/Source/Runtime/Experimental/ChaosCore/Public/Chaos/ArrayCollectionArray.h
--- a/Engine/Source/Runtime/Experimental/ChaosCore/Public/Chaos/ArrayCollectionArray.h
+++ b/Engine/Source/Runtime/Experimental/ChaosCore/Public/Chaos/ArrayCollectionArray.h
@@ -82,3 +82,9 @@ class TArrayCollectionArrayView : public TArrayCollectionArrayBase, public TArra
 	}
 };
 }
+
+template<class T>
+struct TIsContiguousContainer<Chaos::TArrayCollectionArrayView<T>>
+{
+	static constexpr bool Value = TIsContiguousContainer<TArrayView<T>>::Value;
+};
diff --git a/Engine/Source/Runtime/Experimental/ChaosCore/Public/Chaos/Vector.h b/Engine/Source/Runtime/Experimental/ChaosCore/Public/Chaos/Vector.h
--- a/Engine/Source/Runtime/Experimental/ChaosCore/Public/Chaos/Vector.h
+++ b/Engine/Source/Runtime/Experimental/ChaosCore/Public/Chaos/Vector.h
@@ -791,6 +791,12 @@ namespace Chaos
 
 } // namespace Chaos
 
+template<class T, int d>
+struct TIsContiguousContainer<Chaos::TVector<T, d>>
+{
+	static constexpr bool Value = TIsContiguousContainer<TArray<T>>::Value;
+};
+
 //template<>
 //uint32 GetTypeHash(const Chaos::TVector<int32, 2>& V)
 //{
