TStaticArray::GetData() now reports if the array elements are non-contiguous, and GetData() static asserts in that case.

diff --git a/Engine/Source/Runtime/Core/Public/Containers/StaticArray.h b/Engine/Source/Runtime/Core/Public/Containers/StaticArray.h
--- a/Engine/Source/Runtime/Core/Public/Containers/StaticArray.h
+++ b/Engine/Source/Runtime/Core/Public/Containers/StaticArray.h
@@ -79,8 +79,8 @@ public:
 	FORCEINLINE_DEBUGGABLE int32 Num() const { return NumElements; }
 
 	/** A pointer to the first element of the array */
-	FORCEINLINE_DEBUGGABLE       InElementType* GetData()       { return Storage.Elements; }
-	FORCEINLINE_DEBUGGABLE const InElementType* GetData() const { return Storage.Elements; }
+	FORCEINLINE_DEBUGGABLE       InElementType* GetData()       { static_assert((alignof(ElementType) % Alignment) == 0, "GetData() cannot be called on a TStaticArray with non-standard alignment"); return &Storage.Elements[0].Element; }
+	FORCEINLINE_DEBUGGABLE const InElementType* GetData() const { static_assert((alignof(ElementType) % Alignment) == 0, "GetData() cannot be called on a TStaticArray with non-standard alignment"); return &Storage.Elements[0].Element; }
 
 	/** Hash function. */
 	friend uint32 GetTypeHash(const TStaticArray& Array)
@@ -140,5 +140,5 @@ TStaticArray<InElementType,NumElements> MakeUniformStaticArray(typename TCallTra
 template <typename ElementType, uint32 NumElements, uint32 Alignment>
 struct TIsContiguousContainer<TStaticArray<ElementType, NumElements, Alignment>>
 {
-	enum { Value = true };
+	enum { Value = (alignof(ElementType) % Alignment) == 0 };
 };
