Changed task graph behavior in single threaded mode: if there is nothing to wait for we just return, otherwise it is a fatal error.

diff --git a/Engine/Source/Runtime/Core/Private/Async/TaskGraph.cpp b/Engine/Source/Runtime/Core/Private/Async/TaskGraph.cpp
--- a/Engine/Source/Runtime/Core/Private/Async/TaskGraph.cpp
+++ b/Engine/Source/Runtime/Core/Private/Async/TaskGraph.cpp
@@ -1464,6 +1464,23 @@ public:
 		}
 		else
 		{
+			if (!FPlatformProcess::SupportsMultithreading())
+			{
+				bool bAnyPending = false;
+				for (int32 Index = 0; Index < Tasks.Num(); Index++)
+				{
+					if (!Tasks[Index]->IsComplete())
+					{
+						bAnyPending = true;
+						break;
+					}
+				}
+				if (!bAnyPending)
+				{
+					return;
+				}
+				UE_LOG(LogTaskGraph, Fatal, TEXT("Recursive waits are not allowed in single threaded mode."));
+			}
 			// We will just stall this thread on an event while we wait
 			FScopedEvent Event;
 			TriggerEventWhenTasksComplete(Event.Get(), Tasks, CurrentThreadIfKnown);
