PR #5988: Fix unitialized string return from FWindowsPlatformProcess::GetApplicationName (Contributed by [REDACTED_PERSON])

diff --git a/Engine/Source/Runtime/Core/Private/Windows/WindowsPlatformProcess.cpp b/Engine/Source/Runtime/Core/Private/Windows/WindowsPlatformProcess.cpp
--- a/Engine/Source/Runtime/Core/Private/Windows/WindowsPlatformProcess.cpp
+++ b/Engine/Source/Runtime/Core/Private/Windows/WindowsPlatformProcess.cpp
@@ -654,13 +654,17 @@ FString FWindowsPlatformProcess::GetApplicationName( uint32 ProcessId )
 		int32 InOutSize = ProcessNameBufferSize;
 		static_assert(sizeof(::DWORD) == sizeof(int32), "DWORD size doesn't match int32. Is it the future or the past?");
 
+		if(
 #if WINVER == 0x0502
-		GetProcessImageFileName(ProcessHandle, ProcessNameBuffer, InOutSize);
+		GetProcessImageFileName(ProcessHandle, ProcessNameBuffer, InOutSize)
 #else
-		QueryFullProcessImageName(ProcessHandle, 0, ProcessNameBuffer, (PDWORD)(&InOutSize));
+		QueryFullProcessImageName(ProcessHandle, 0, ProcessNameBuffer, (PDWORD)(&InOutSize))
 #endif
-
-		Output = ProcessNameBuffer;
+			)
+		{
+			// TODO no null termination guarantee on GetProcessImageFileName?  it returns size as well, whereas QueryFullProcessImageName just returns non-zero on success
+			Output = ProcessNameBuffer;
+		}
 	}
 
 	return Output;
