UnrealInsights:
Added functionality to trace variable values and expose them as counters

diff --git a/Engine/Source/Developer/TraceServices/Private/Analyzers/CountersTraceAnalysis.cpp b/Engine/Source/Developer/TraceServices/Private/Analyzers/CountersTraceAnalysis.cpp
new file mode 100644
--- /dev/null
+++ b/Engine/Source/Developer/TraceServices/Private/Analyzers/CountersTraceAnalysis.cpp
@@ -0,0 +1,75 @@
+// Copyright 1998-2019 Epic Games, Inc. All Rights Reserved.
+#include "CountersTraceAnalysis.h"
+#include "TraceServices/Model/Counters.h"
+#include "ProfilingDebugging/CountersTrace.h"
+
+FCountersAnalyzer::FCountersAnalyzer(Trace::IAnalysisSession& InSession, Trace::ICounterProvider& InCounterProvider)
+	: Session(InSession)
+	, CounterProvider(InCounterProvider)
+{
+
+}
+
+void FCountersAnalyzer::OnAnalysisBegin(const FOnAnalysisContext& Context)
+{
+	auto& Builder = Context.InterfaceBuilder;
+
+	Builder.RouteEvent(RouteId_Spec, "Counters", "Spec");
+	Builder.RouteEvent(RouteId_SetValueInt, "Counters", "SetValueInt");
+	Builder.RouteEvent(RouteId_SetValueFloat, "Counters", "SetValueFloat");
+}
+
+void FCountersAnalyzer::OnAnalysisEnd()
+{
+}
+
+void FCountersAnalyzer::OnEvent(uint16 RouteId, const FOnEventContext& Context)
+{
+	Trace::FAnalysisSessionEditScope _(Session);
+
+	const auto& EventData = Context.EventData;
+	switch (RouteId)
+	{
+	case RouteId_Spec:
+	{
+		uint16 CounterId = EventData.GetValue<uint16>("Id");
+		ETraceCounterType CounterType = static_cast<ETraceCounterType>(EventData.GetValue<uint8>("Type"));
+		Trace::ICounter* Counter = CounterProvider.CreateCounter();
+		if (CounterType == TraceCounterType_Float)
+		{
+			Counter->SetIsFloatingPoint(true);
+		}
+		else if (CounterType == TraceCounterType_Memory)
+		{
+			Counter->SetDisplayHint(Trace::CounterDisplayHint_Memory);
+		}
+		Counter->SetName(Session.StoreString(reinterpret_cast<const TCHAR*>(EventData.GetAttachment())));
+		CountersMap.Add(CounterId, Counter);
+		break;
+	}
+	case RouteId_SetValueInt:
+	{
+		uint16 CounterId = EventData.GetValue<uint16>("CounterId");
+		int64 Value = EventData.GetValue<int64>("Value");
+		double Timestamp = Context.SessionContext.TimestampFromCycle(EventData.GetValue<uint64>("Cycle"));
+		Trace::ICounter** FindCounter = CountersMap.Find(CounterId);
+		if (FindCounter)
+		{
+			(*FindCounter)->SetValue(Timestamp, Value);
+		}
+		break;
+	}
+	case RouteId_SetValueFloat:
+	{
+		uint16 CounterId = EventData.GetValue<uint16>("CounterId");
+		float Value = EventData.GetValue<float>("Value");
+		double Timestamp = Context.SessionContext.TimestampFromCycle(EventData.GetValue<uint64>("Cycle"));
+		Trace::ICounter** FindCounter = CountersMap.Find(CounterId);
+		if (FindCounter)
+		{
+			(*FindCounter)->SetValue(Timestamp, Value);
+		}
+		break;
+	}
+	}
+}
diff --git a/Engine/Source/Developer/TraceServices/Private/Analyzers/CountersTraceAnalysis.h b/Engine/Source/Developer/TraceServices/Private/Analyzers/CountersTraceAnalysis.h
new file mode 100644
--- /dev/null
+++ b/Engine/Source/Developer/TraceServices/Private/Analyzers/CountersTraceAnalysis.h
@@ -0,0 +1,36 @@
+// Copyright 1998-2019 Epic Games, Inc. All Rights Reserved.
+
+#pragma once
+
+#include "Trace/Trace.h"
+#include "Trace/Analyzer.h"
+#include "Containers/Map.h"
+
+namespace Trace
+{
+	class IAnalysisSession;
+	class ICounter;
+	class ICounterProvider;
+}
+
+class FCountersAnalyzer
+	: public Trace::IAnalyzer
+{
+public:
+	FCountersAnalyzer(Trace::IAnalysisSession& Session, Trace::ICounterProvider& CounterProvider);
+	virtual void OnAnalysisBegin(const FOnAnalysisContext& Context) override;
+	virtual void OnEvent(uint16 RouteId, const FOnEventContext& Context) override;
+	virtual void OnAnalysisEnd() override;
+
+private:
+	enum : uint16
+	{
+		RouteId_Spec,
+		RouteId_SetValueInt,
+		RouteId_SetValueFloat,
+	};
+
+	Trace::IAnalysisSession& Session;
+	Trace::ICounterProvider& CounterProvider;
+	TMap<uint16, Trace::ICounter*> CountersMap;
+};
diff --git a/Engine/Source/Developer/TraceServices/Private/Modules/CountersModule.cpp b/Engine/Source/Developer/TraceServices/Private/Modules/CountersModule.cpp
new file mode 100644
--- /dev/null
+++ b/Engine/Source/Developer/TraceServices/Private/Modules/CountersModule.cpp
@@ -0,0 +1,31 @@
+// Copyright 1998-2019 Epic Games, Inc. All Rights Reserved.
+
+#include "CountersModule.h"
+#include "Analyzers/CountersTraceAnalysis.h"
+#include "TraceServices/ModuleService.h"
+#include "TraceServices/Model/AnalysisSession.h"
+#include "TraceServices/Model/Counters.h"
+
+namespace Trace
+{
+
+static const FName CountersModuleName("TraceModule_Counters");
+
+void FCountersModule::GetModuleInfo(FModuleInfo& OutModuleInfo)
+{
+	OutModuleInfo.Name = CountersModuleName;
+	OutModuleInfo.DisplayName = TEXT("Counters");
+}
+	
+void FCountersModule::OnAnalysisBegin(IAnalysisSession& Session)
+{
+	ICounterProvider& CounterProvider = EditCounterProvider(Session);
+	Session.AddAnalyzer(new FCountersAnalyzer(Session, CounterProvider));
+}
+
+void FCountersModule::GetLoggers(TArray<const TCHAR *>& OutLoggers)
+{
+	OutLoggers.Add(TEXT("Counters"));
+}
+
+}
diff --git a/Engine/Source/Developer/TraceServices/Private/Modules/CountersModule.h b/Engine/Source/Developer/TraceServices/Private/Modules/CountersModule.h
new file mode 100644
--- /dev/null
+++ b/Engine/Source/Developer/TraceServices/Private/Modules/CountersModule.h
@@ -0,0 +1,20 @@
+// Copyright 1998-2019 Epic Games, Inc. All Rights Reserved.
+
+#pragma once
+
+#include "TraceServices/ModuleService.h"
+
+namespace Trace
+{
+
+class FCountersModule
+	: public IModule
+{
+public:
+	virtual void GetModuleInfo(FModuleInfo& OutModuleInfo) override;
+	virtual void OnAnalysisBegin(IAnalysisSession& Session) override;
+	virtual void GetLoggers(TArray<const TCHAR *>& OutLoggers) override;
+	virtual void GenerateReports(const IAnalysisSession& Session, const TCHAR* CmdLine, const TCHAR* OutputDirectory) override {}
+};
+
+}
diff --git a/Engine/Source/Developer/TraceServices/Private/TraceServicesModule.cpp b/Engine/Source/Developer/TraceServices/Private/TraceServicesModule.cpp
--- a/Engine/Source/Developer/TraceServices/Private/TraceServicesModule.cpp
+++ b/Engine/Source/Developer/TraceServices/Private/TraceServicesModule.cpp
@@ -10,6 +10,7 @@
 #include "Modules/LoadTimeProfilerModule.h"
 #include "Modules/StatsModule.h"
 #include "Modules/CsvProfilerModule.h"
+#include "Modules/CountersModule.h"
 
 class FTraceServicesModule
 	: public ITraceServicesModule
@@ -31,6 +32,7 @@ private:
 	Trace::FLoadTimeProfilerModule LoadTimeProfilerModule;
 	Trace::FStatsModule StatsModule;
 	Trace::FCsvProfilerModule CsvProfilerModule;
+	Trace::FCountersModule CountersModule;
 };
 
 TSharedPtr<Trace::ISessionService> FTraceServicesModule::GetSessionService()
@@ -70,10 +72,12 @@ void FTraceServicesModule::StartupModule()
 	IModularFeatures::Get().RegisterModularFeature(Trace::ModuleFeatureName, &StatsModule);
 #endif
 	IModularFeatures::Get().RegisterModularFeature(Trace::ModuleFeatureName, &CsvProfilerModule);
+	IModularFeatures::Get().RegisterModularFeature(Trace::ModuleFeatureName, &CountersModule);
 }
 
 void FTraceServicesModule::ShutdownModule()
 {
+	IModularFeatures::Get().UnregisterModularFeature(Trace::ModuleFeatureName, &CountersModule);
 	IModularFeatures::Get().UnregisterModularFeature(Trace::ModuleFeatureName, &CsvProfilerModule);
 #if EXPERIMENTAL_STATSTRACE_ENABLED
 	IModularFeatures::Get().UnregisterModularFeature(Trace::ModuleFeatureName, &StatsModule);
diff --git a/Engine/Source/Runtime/Core/Private/ProfilingDebugging/CountersTrace.cpp b/Engine/Source/Runtime/Core/Private/ProfilingDebugging/CountersTrace.cpp
new file mode 100644
--- /dev/null
+++ b/Engine/Source/Runtime/Core/Private/ProfilingDebugging/CountersTrace.cpp
@@ -0,0 +1,111 @@
+// Copyright 1998-2019 Epic Games, Inc. All Rights Reserved.
+#include "ProfilingDebugging/CountersTrace.h"
+#include "HAL/PlatformTime.h"
+#include "Misc/Parse.h"
+
+#if COUNTERSTRACE_ENABLED
+
+UE_TRACE_EVENT_BEGIN(Counters, Spec, Always)
+	UE_TRACE_EVENT_FIELD(uint16, Id)
+	UE_TRACE_EVENT_FIELD(uint8, Type)
+UE_TRACE_EVENT_END()
+
+UE_TRACE_EVENT_BEGIN(Counters, SetValueInt)
+	UE_TRACE_EVENT_FIELD(uint64, Cycle)
+	UE_TRACE_EVENT_FIELD(int64, Value)
+	UE_TRACE_EVENT_FIELD(uint16, CounterId)
+UE_TRACE_EVENT_END()
+
+UE_TRACE_EVENT_BEGIN(Counters, SetValueFloat)
+	UE_TRACE_EVENT_FIELD(uint64, Cycle)
+	UE_TRACE_EVENT_FIELD(double, Value)
+	UE_TRACE_EVENT_FIELD(uint16, CounterId)
+UE_TRACE_EVENT_END()
+
+
+struct FCountersTracePrivate
+{
+	static uint16 InitCounter(const TCHAR* CounterName, ETraceCounterType Type);
+	
+	static void SetValue(uint16 CounterId, int64 Value)
+	{
+		UE_TRACE_LOG(Counters, SetValueInt)
+			<< SetValueInt.Cycle(FPlatformTime::Cycles64())
+			<< SetValueInt.Value(Value)
+			<< SetValueInt.CounterId(CounterId);
+	}
+
+	static void SetValue(uint16 CounterId, float Value)
+	{
+		UE_TRACE_LOG(Counters, SetValueFloat)
+			<< SetValueFloat.Cycle(FPlatformTime::Cycles64())
+			<< SetValueFloat.Value(Value)
+			<< SetValueFloat.CounterId(CounterId);
+	}
+};
+
+uint16 FCountersTracePrivate::InitCounter(const TCHAR* CounterName, ETraceCounterType Type)
+{
+	static TAtomic<uint16> NextId(1);
+	uint16 CounterId = uint16(NextId++);
+	uint16 NameSize = (FCString::Strlen(CounterName) + 1) * sizeof(TCHAR);
+	UE_TRACE_LOG(Counters, Spec, NameSize)
+		<< Spec.Id(CounterId)
+		<< Spec.Type(uint8(Type))
+		<< Spec.Attachment(CounterName, NameSize);
+	return CounterId;
+}
+
+void FCountersTrace::FCounterInt::Init(const TCHAR* CounterName)
+{
+	CounterId = FCountersTracePrivate::InitCounter(CounterName, TraceCounterType_Int);
+}
+
+void FCountersTrace::FCounterInt::Set(int64 Value)
+{
+	if (Value != CurrentValue)
+	{
+		CurrentValue = Value;
+		FCountersTracePrivate::SetValue(CounterId, Value);
+	}
+}
+
+void FCountersTrace::FCounterFloat::Init(const TCHAR* CounterName)
+{
+	CounterId = FCountersTracePrivate::InitCounter(CounterName, TraceCounterType_Float);
+}
+
+void FCountersTrace::FCounterFloat::Set(float Value)
+{
+	if (Value != CurrentValue)
+	{
+		CurrentValue = Value;
+		FCountersTracePrivate::SetValue(CounterId, Value);
+	}
+}
+
+void FCountersTrace::FCounterMemory::Init(const TCHAR* CounterName)
+{
+	CounterId = FCountersTracePrivate::InitCounter(CounterName, TraceCounterType_Memory);
+}
+
+void FCountersTrace::FCounterMemory::Set(uint64 Value)
+{
+	if (Value != CurrentValue)
+	{
+		CurrentValue = Value;
+		FCountersTracePrivate::SetValue(CounterId, (int64)Value);
+	}
+}
+
+void FCountersTrace::Init(const TCHAR* CmdLine)
+{
+	if (FParse::Param(CmdLine, TEXT("counterstrace")))
+	{
+		UE_TRACE_EVENT_IS_ENABLED(Counters, SetValueInt);
+		UE_TRACE_EVENT_IS_ENABLED(Counters, SetValueFloat);
+		Trace::ToggleEvent(TEXT("Counters"), true);
+	}
+}
+
+#endif
\ No newline at end of file
diff --git a/Engine/Source/Runtime/Core/Public/ProfilingDebugging/CountersTrace.h b/Engine/Source/Runtime/Core/Public/ProfilingDebugging/CountersTrace.h
new file mode 100644
--- /dev/null
+++ b/Engine/Source/Runtime/Core/Public/ProfilingDebugging/CountersTrace.h
@@ -0,0 +1,82 @@
+// Copyright 1998-2019 Epic Games, Inc. All Rights Reserved.
+
+#pragma once
+
+#include "CoreTypes.h"
+#include "Trace/Trace.h"
+
+#if UE_TRACE_ENABLED && !UE_BUILD_SHIPPING
+#define COUNTERSTRACE_ENABLED 1
+#else
+#define COUNTERSTRACE_ENABLED 0
+#endif
+
+enum ETraceCounterType
+{
+	TraceCounterType_Int,
+	TraceCounterType_Float,
+	TraceCounterType_Memory,
+};
+
+#if COUNTERSTRACE_ENABLED
+
+struct FCountersTrace
+{
+	struct FCounterInt
+	{
+		CORE_API void Init(const TCHAR* CounterName);
+		CORE_API void Set(int64 Value);
+
+		int64 CurrentValue;
+		uint16 CounterId;
+	};
+
+	struct FCounterFloat
+	{
+		CORE_API void Init(const TCHAR* CounterName);
+		CORE_API void Set(float Value);
+
+		double CurrentValue;
+		uint16 CounterId;
+	};
+
+	struct FCounterMemory
+	{
+		CORE_API void Init(const TCHAR* CounterName);
+		CORE_API void Set(uint64 Value);
+
+		uint64 CurrentValue;
+		uint16 CounterId;
+	};
+
+	CORE_API static void Init(const TCHAR* CmdLine);
+};
+
+#define TRACE_COUNTERS_INIT(CmdLine) \
+	FCountersTrace::Init(CmdLine);
+
+#define __TRACE_COUNTER_VALUE(CounterType, CounterName, Value) \
+	static FCountersTrace::CounterType PREPROCESSOR_JOIN(__TraceCounter, __LINE__); \
+	if (!PREPROCESSOR_JOIN(__TraceCounter, __LINE__).CounterId) \
+	{ \
+		PREPROCESSOR_JOIN(__TraceCounter, __LINE__).Init(TEXT(#CounterName)); \
+	} \
+	PREPROCESSOR_JOIN(__TraceCounter, __LINE__).Set(Value);
+
+#define TRACE_INTEGER_VALUE(CounterName, Value) \
+	__TRACE_COUNTER_VALUE(FCounterInt, CounterName, Value)
+
+#define TRACE_FLOAT_VALUE(CounterName, Value) \
+	__TRACE_COUNTER_VALUE(FCounterFloat, CounterName, Value)
+
+#define TRACE_MEMORY_VALUE(CounterName, Value) \
+	__TRACE_COUNTER_VALUE(FCounterMemory, CounterName, Value)
+
+#else
+
+#define TRACE_COUNTERS_INIT(CmdLine)
+#define TRACE_INTEGER_VALUE(CounterName, Value)
+#define TRACE_FLOAT_VALUE(CounterName, Value)
+#define TRACE_MEMORY_VALUE(CounterName, Value)
+
+#endif
\ No newline at end of file
diff --git a/Engine/Source/Runtime/Launch/Private/LaunchEngineLoop.cpp b/Engine/Source/Runtime/Launch/Private/LaunchEngineLoop.cpp
--- a/Engine/Source/Runtime/Launch/Private/LaunchEngineLoop.cpp
+++ b/Engine/Source/Runtime/Launch/Private/LaunchEngineLoop.cpp
@@ -31,6 +31,7 @@
 #include "Trace/Trace.h"
 #include "ProfilingDebugging/MiscTrace.h"
 #include "ProfilingDebugging/PlatformFileTrace.h"
+#include "ProfilingDebugging/CountersTrace.h"
 #if WITH_ENGINE
 #include "HAL/PlatformSplash.h"
 #endif
@@ -1161,6 +1162,7 @@ int32 FEngineLoop::PreInit(const TCHAR* CmdLine)
 	TRACE_REGISTER_GAME_THREAD(FPlatformTLS::GetCurrentThreadId());
 	TRACE_CPUPROFILER_INIT(CmdLine);
 	TRACE_PLATFORMFILE_INIT(CmdLine);
+	TRACE_COUNTERS_INIT(CmdLine);
 
 	SCOPED_BOOT_TIMING("FEngineLoop::PreInit");
 
