Don't invalidate the FArchiveFileReaderGeneric internal buffer if the seek will still end up within it's bounds since we can still use the data that was already loaded.

Doesn't have a huge real world impact (atleast on windows) as the OS file cache will most likely be preventing actual additional disk reads but it does cut down on a small amount of overhead.

diff --git a/Engine/Source/Runtime/Core/Private/HAL/FileManagerGeneric.cpp b/Engine/Source/Runtime/Core/Private/HAL/FileManagerGeneric.cpp
--- a/Engine/Source/Runtime/Core/Private/HAL/FileManagerGeneric.cpp
+++ b/Engine/Source/Runtime/Core/Private/HAL/FileManagerGeneric.cpp
@@ -653,15 +653,26 @@ void FArchiveFileReaderGeneric::Seek( int64 InPos )
 {
 	checkf(InPos >= 0, TEXT("Attempted to seek to a negative location (%lld/%lld), file: %s. The file is most likely corrupt."), InPos, Size, *Filename);
 	checkf(InPos <= Size, TEXT("Attempted to seek past the end of file (%lld/%lld), file: %s. The file is most likely corrupt."), InPos, Size, *Filename);
-	if( !SeekLowLevel( InPos ) )
+
+	if (InPos < BufferBase || InPos >= BufferBase + BufferCount)
 	{
-		TCHAR ErrorBuffer[1024];
-		ArIsError = true;
-		UE_LOG(LogFileManager, Error, TEXT("SetFilePointer on %s Failed %lld/%lld: %lld %s"), *Filename, InPos, Size, Pos, FPlatformMisc::GetSystemErrorMessage(ErrorBuffer, 1024, 0));
+		// The seek is outside of the cached buffer so we actually need to seek
+		if (!SeekLowLevel(InPos))
+		{
+			TCHAR ErrorBuffer[1024];
+			ArIsError = true;
+			UE_LOG(LogFileManager, Error, TEXT("SetFilePointer on %s Failed %lld/%lld: %lld %s"), *Filename, InPos, Size, Pos, FPlatformMisc::GetSystemErrorMessage(ErrorBuffer, 1024, 0));
+		}
+
+		Pos = InPos;
+		BufferBase = Pos;
+		BufferCount = 0;
+	}
+	else
+	{
+		// We already have the seek position cached so we can just move the current reading position
+		Pos = InPos;
 	}
-	Pos         = InPos;
-	BufferBase  = Pos;
-	BufferCount = 0;
 }
 
 FArchiveFileReaderGeneric::~FArchiveFileReaderGeneric()
