Fixed possible buffer overruns in WritePipe functions

diff --git a/Engine/Source/Runtime/Core/Private/Mac/MacPlatformProcess.cpp b/Engine/Source/Runtime/Core/Private/Mac/MacPlatformProcess.cpp
--- a/Engine/Source/Runtime/Core/Private/Mac/MacPlatformProcess.cpp
+++ b/Engine/Source/Runtime/Core/Private/Mac/MacPlatformProcess.cpp
@@ -1113,7 +1113,7 @@ bool FMacPlatformProcess::WritePipe(void* WritePipe, const FString& Message, FSt
 
 	// Convert input to UTF8CHAR
 	uint32 BytesAvailable = Message.Len();
-	UTF8CHAR * Buffer = new UTF8CHAR[BytesAvailable + 1];
+	UTF8CHAR * Buffer = new UTF8CHAR[BytesAvailable + 2];
 	for (uint32 i = 0; i < BytesAvailable; i++)
 	{
 		Buffer[i] = Message[i];
diff --git a/Engine/Source/Runtime/Core/Private/Unix/UnixPlatformProcess.cpp b/Engine/Source/Runtime/Core/Private/Unix/UnixPlatformProcess.cpp
--- a/Engine/Source/Runtime/Core/Private/Unix/UnixPlatformProcess.cpp
+++ b/Engine/Source/Runtime/Core/Private/Unix/UnixPlatformProcess.cpp
@@ -522,7 +522,7 @@ bool FUnixPlatformProcess::WritePipe(void* WritePipe, const FString& Message, FS
 
 	// Convert input to UTF8CHAR
 	uint32 BytesAvailable = Message.Len();
-	UTF8CHAR * Buffer = new UTF8CHAR[BytesAvailable + 1];
+	UTF8CHAR * Buffer = new UTF8CHAR[BytesAvailable + 2];
 	for (uint32 i = 0; i < BytesAvailable; i++)
 	{
 		Buffer[i] = Message[i];
diff --git a/Engine/Source/Runtime/Core/Private/Windows/WindowsPlatformProcess.cpp b/Engine/Source/Runtime/Core/Private/Windows/WindowsPlatformProcess.cpp
--- a/Engine/Source/Runtime/Core/Private/Windows/WindowsPlatformProcess.cpp
+++ b/Engine/Source/Runtime/Core/Private/Windows/WindowsPlatformProcess.cpp
@@ -1387,7 +1387,7 @@ bool FWindowsPlatformProcess::WritePipe(void* WritePipe, const FString& Message,
 
 	// Convert input to UTF8CHAR
 	uint32 BytesAvailable = Message.Len();
-	UTF8CHAR * Buffer = new UTF8CHAR[BytesAvailable + 1];
+	UTF8CHAR * Buffer = new UTF8CHAR[BytesAvailable + 2];
 	for (uint32 i = 0; i < BytesAvailable; i++)
 	{
 		Buffer[i] = Message[i];
