GitHub #6050 : Adds support for zero-copy consumes.

diff --git a/Engine/Source/Runtime/Core/Public/Containers/CircularQueue.h b/Engine/Source/Runtime/Core/Public/Containers/CircularQueue.h
--- a/Engine/Source/Runtime/Core/Public/Containers/CircularQueue.h
+++ b/Engine/Source/Runtime/Core/Public/Containers/CircularQueue.h
@@ -25,7 +25,7 @@ template<typename ElementType> class TCircularQueue
 public:
 
 	/**
-	 * Default constructor.
+	 * Constructor.
 	 *
 	 * @param CapacityPlusOne The number of elements that the queue can hold (will be rounded up to the next power of 2).
 	 */
@@ -79,6 +79,26 @@ public:
 		return false;
 	}
 
+	/**
+	 * Removes an item from the front of the queue.
+	 *
+	 * @return true if an element has been removed, false if the queue was empty.
+	 * @note To be called only from consumer thread.
+	 */
+	bool Dequeue()
+	{
+		const uint32 CurrentHead = Head.Load();
+
+		if (CurrentHead != Tail.Load())
+		{
+			Head.Store(Buffer.GetNextIndex(CurrentHead));
+
+			return true;
+		}
+
+		return false;
+	}
+
 	/**
 	 * Empties the queue.
 	 *
@@ -171,7 +191,7 @@ public:
 	 * @return true if an item has been returned, false if the queue was empty.
 	 * @note To be called only from consumer thread.
 	 */
-	bool Peek(ElementType& OutItem)
+	bool Peek(ElementType& OutItem) const
 	{
 		const uint32 CurrentHead = Head.Load();
 
@@ -185,6 +205,25 @@ public:
 		return false;
 	}
 
+	/**
+	 * Returns the oldest item in the queue without removing it.
+	 *
+	 * @return an ElementType pointer if an item has been returned, nullptr if the queue was empty.
+	 * @note To be called only from consumer thread.
+	 * @note The return value is only valid until Dequeue, Empty, or the destructor has been called.
+	 */
+	const ElementType* Peek() const
+	{
+		const uint32 CurrentHead = Head.Load();
+
+		if (CurrentHead != Tail.Load())
+		{
+			return &Buffer[CurrentHead];
+		}
+
+		return nullptr;
+	}
+
 private:
 
 	/** Holds the buffer. */
