Tests for TCircularQueue.

diff --git a/Engine/Source/Runtime/Core/Private/Tests/Misc/CircularQueueTest.cpp b/Engine/Source/Runtime/Core/Private/Tests/Misc/CircularQueueTest.cpp
--- a/Engine/Source/Runtime/Core/Private/Tests/Misc/CircularQueueTest.cpp
+++ b/Engine/Source/Runtime/Core/Private/Tests/Misc/CircularQueueTest.cpp
@@ -19,21 +19,31 @@ bool FCircularQueueTest::RunTest(const FString& Parameters)
 		TestEqual(TEXT("Newly created queues must have zero elements"), Queue.Count(), 0u);
 		TestTrue(TEXT("Newly created queues must be empty"), Queue.IsEmpty());
 		TestFalse(TEXT("Newly created queues must not be full"), Queue.IsFull());
+
+		int32 Value;
+		TestFalse(TEXT("Peek must fail on an empty queue"), Queue.Peek(Value));
+		TestNull(TEXT("Peek must fail on an empty queue"), Queue.Peek());
 	}
 
 	// partially filled
 	{
 		TCircularQueue<int32> Queue(QueueSize);
 		int32 Value = 0;
-	
+
 		TestTrue(TEXT("Adding to an empty queue must succeed"), Queue.Enqueue(666));
 		TestEqual(TEXT("After adding to an empty queue it must have one element"), Queue.Count(), 1u);
 		TestFalse(TEXT("Partially filled queues must not be empty"), Queue.IsEmpty());
 		TestFalse(TEXT("Partially filled queues must not be full"), Queue.IsFull());
 		TestTrue(TEXT("Peeking at a partially filled queue must succeed"), Queue.Peek(Value));
+		TestEqual(TEXT("The peeked at value must be correct"), Value, 666);
+
+		const int32* PeekValue = Queue.Peek();
+		TestNotNull(TEXT("Peeking at a partially filled queue must succeed"), PeekValue);
+		TestEqual(TEXT("The peeked at value must be correct"), *PeekValue, 666);
 	}
 
 	// full queue
+	for (uint32 PeekType = 0; PeekType < 2; PeekType++)
 	{
 		TCircularQueue<int32> Queue(QueueSize);
 
@@ -50,10 +60,22 @@ bool FCircularQueueTest::RunTest(const FString& Parameters)
 
 		for (int32 Index = 0; Index < QueueSize - 1; ++Index)
 		{
-			TestTrue(TEXT("Peeking at a none-empty queue must succeed"), Queue.Peek(Value));
-			TestEqual(TEXT("The peeked at value must be correct"), Value, Index);
-			TestTrue(TEXT("Removing from a non-empty queue must succeed"), Queue.Dequeue(Value));
-			TestEqual(TEXT("The removed value must be correct"), Value, Index);
+			if (PeekType == 0)
+			{
+				TestTrue(TEXT("Peeking at a non-empty queue must succeed"), Queue.Peek(Value));
+				TestEqual(TEXT("The peeked at value must be correct"), Value, Index);
+
+				TestTrue(TEXT("Removing from a non-empty queue must succeed"), Queue.Dequeue(Value));
+				TestEqual(TEXT("The removed value must be correct"), Value, Index);
+			}
+			else
+			{
+				const int32* PeekValue = Queue.Peek();
+				TestNotNull(TEXT("Peeking at a non-empty queue must succeed"), PeekValue);
+				TestEqual(TEXT("The peeked at value must be correct"), *PeekValue, Index);
+
+				TestTrue(TEXT("Removing from a non-empty queue must succeed"), Queue.Dequeue());
+			}
 		}
 
 		TestTrue(TEXT("A queue that had all items removed must be empty"), Queue.IsEmpty());
@@ -69,7 +91,7 @@ bool FCircularQueueTest::RunTest(const FString& Parameters)
 		{
 			TestTrue(TEXT("Adding to non-full queue must succeed"), Queue.Enqueue(Index));
 		}
-		
+
 		int32 Value = 0;
 		const int32 ExpectedSize = QueueSize - 1;
 		for (int32 Index = 0; Index < QueueSize; ++Index)
@@ -80,6 +102,35 @@ bool FCircularQueueTest::RunTest(const FString& Parameters)
 		}
 	}
 
+	// Non-zero initialization - ensures the backing store constructs and destructs non-POD objects properly.
+	{
+		static uint32 Const;
+		static uint32 Dest;
+		static uint32 CopyConst;
+		struct FNonPOD
+		{
+			FNonPOD() { Const++; }
+			~FNonPOD() { Dest++; }
+			FNonPOD(const FNonPOD&) { CopyConst++; }
+		};
+
+		// The current implementation of TCircularQueue doesn't call the held object constructors on but does call the
+		// destructors. Verify that (somewhat surprising, if typically ok for POD) behavior first.
+		Const = 0;
+		Dest = 0;
+		CopyConst = 0;
+		{
+			TCircularQueue<FNonPOD> Queue(QueueSize);
+
+			TestEqual(TEXT("The constructor should not run"), Const, 0);
+			TestEqual(TEXT("The destructor should not run"), Dest, 0);
+			TestEqual(TEXT("The copy constructor should not run"), CopyConst, 0);
+		}
+		TestEqual(TEXT("The constructor should not run"), Const, 0);
+		TestEqual(TEXT("The destructor should run"), Dest, QueueSize);
+		TestEqual(TEXT("The copy constructor should not run"), CopyConst, 0);
+	}
+
 	return true;
 }
 
