Simplified traits to reduce redundant compile-time churn.
Removal of some copy constructors to make them defaulted and match their overridden POD (specifically: trivial copy constructor) trait.

diff --git a/Engine/Source/Runtime/Core/Public/Math/Plane.h b/Engine/Source/Runtime/Core/Public/Math/Plane.h
--- a/Engine/Source/Runtime/Core/Public/Math/Plane.h
+++ b/Engine/Source/Runtime/Core/Public/Math/Plane.h
@@ -27,13 +27,6 @@ public:
 	/** Default constructor (no initialization). */
 	FORCEINLINE FPlane();
 
-	/**
-	 * Copy Constructor.
-	 *
-	 * @param P Plane to copy from.
-	 */
-	FORCEINLINE FPlane(const FPlane& P);
-
 	/**
 	 * Constructor.
 	 *
@@ -385,12 +378,6 @@ FORCEINLINE FPlane::FPlane()
 {}
 
 
-FORCEINLINE FPlane::FPlane(const FPlane& P)
-	:	FVector(P)
-	,	W(P.W)
-{}
-
-
 FORCEINLINE FPlane::FPlane(const FVector4& V)
 	:	FVector(V)
 	,	W(V.W)
diff --git a/Engine/Source/Runtime/Core/Public/Math/Quat.h b/Engine/Source/Runtime/Core/Public/Math/Quat.h
--- a/Engine/Source/Runtime/Core/Public/Math/Quat.h
+++ b/Engine/Source/Runtime/Core/Public/Math/Quat.h
@@ -68,13 +68,6 @@ public:
 	 */
 	FORCEINLINE FQuat(float InX, float InY, float InZ, float InW);
 
-	/**
-	 * Copy constructor.
-	 *
-	 * @param Q A FQuat object to use to create new quaternion from.
-	 */
-	FORCEINLINE FQuat(const FQuat& Q);
-	
 	/**
 	 * Creates and initializes a new quaternion from the given matrix.
 	 *
@@ -707,14 +700,6 @@ FORCEINLINE FQuat::FQuat(float InX, float InY, float InZ, float InW)
 }
 
 
-FORCEINLINE FQuat::FQuat(const FQuat& Q)
-	: X(Q.X)
-	, Y(Q.Y)
-	, Z(Q.Z)
-	, W(Q.W)
-{ }
-
-
 FORCEINLINE FString FQuat::ToString() const
 {
 	return FString::Printf(TEXT("X=%.9f Y=%.9f Z=%.9f W=%.9f"), X, Y, Z, W);
diff --git a/Engine/Source/Runtime/Core/Public/Math/TransformNonVectorized.h b/Engine/Source/Runtime/Core/Public/Math/TransformNonVectorized.h
--- a/Engine/Source/Runtime/Core/Public/Math/TransformNonVectorized.h
+++ b/Engine/Source/Runtime/Core/Public/Math/TransformNonVectorized.h
@@ -188,19 +188,6 @@ public:
 		DiagnosticCheckNaN_All();
 	}
 
-	/**
-	* Copy-constructor
-	*
-	* @param InTransform The source transform from which all components will be copied
-	*/
-	FORCEINLINE FTransform(const FTransform& InTransform) :
-		Rotation(InTransform.Rotation),
-		Translation(InTransform.Translation),
-		Scale3D(InTransform.Scale3D)
-	{
-		DiagnosticCheckNaN_All();
-	}
-
 	/**
 	* Constructor for converting a Matrix (including scale) into a FTransform.
 	*/
diff --git a/Engine/Source/Runtime/Core/Public/Math/TransformVectorized.h b/Engine/Source/Runtime/Core/Public/Math/TransformVectorized.h
--- a/Engine/Source/Runtime/Core/Public/Math/TransformVectorized.h
+++ b/Engine/Source/Runtime/Core/Public/Math/TransformVectorized.h
@@ -237,19 +237,6 @@ public:
 #endif
 	}
 
-	/**
-	 * Copy-constructor
-	 *
-	 * @param InTransform The source transform from which all components will be copied
-	 */
-	FORCEINLINE FTransform(const FTransform& InTransform) : 
-		Rotation(InTransform.Rotation), 
-		Translation(InTransform.Translation), 
-		Scale3D(InTransform.Scale3D)
-	{
-		DiagnosticCheckNaN_All();
-	}
-
 	/**
 	 * Constructor for converting a Matrix (including scale) into a FTransform.
 	 */
diff --git a/Engine/Source/Runtime/Core/Public/Stats/Stats2.h b/Engine/Source/Runtime/Core/Public/Stats/Stats2.h
--- a/Engine/Source/Runtime/Core/Public/Stats/Stats2.h
+++ b/Engine/Source/Runtime/Core/Public/Stats/Stats2.h
@@ -406,16 +406,6 @@ public:
 	{
 	}
 
-	/**
-	 * Copy constructor
-	 */
-	FORCEINLINE_STATS FStatNameAndInfo(FStatNameAndInfo const& Other)
-		: NameAndInfo(Other.NameAndInfo)
-	{
-		static_assert(EStatAllFields::StartShift >= 0, "Too many fields.");
-		CheckInvariants();
-	}
-
 	/**
 	 * Build from a raw FName
 	 */
@@ -699,15 +689,6 @@ struct FStatMessage
 	{
 	}
 
-	/**
-	* Copy constructor
-	*/
-	FORCEINLINE_STATS FStatMessage(FStatMessage const& Other)
-		: StatData(Other.StatData)
-		, NameAndInfo(Other.NameAndInfo)
-	{
-	}
-	
 	/**
 	* Build a meta data message
 	*/
@@ -924,19 +905,6 @@ struct TStatMessage
 		Clear();
 	}
 
-	/**
-	* Copy constructor
-	*/
-	explicit FORCEINLINE_STATS TStatMessage(const TStatMessage& Other)
-		: NameAndInfo(Other.NameAndInfo)
-	{
-		// Copy all fields.
-		for( int32 FieldIndex = 0; FieldIndex < EnumCount; ++FieldIndex )
-		{
-			*((int64*)&StatData+FieldIndex) = *((int64*)&Other.StatData+FieldIndex);
-		}
-	}
-
 	/** Assignment operator for FStatMessage. */
 	TStatMessage& operator=(const FStatMessage& Other)
 	{
diff --git a/Engine/Source/Runtime/Core/Public/Templates/IsPODType.h b/Engine/Source/Runtime/Core/Public/Templates/IsPODType.h
--- a/Engine/Source/Runtime/Core/Public/Templates/IsPODType.h
+++ b/Engine/Source/Runtime/Core/Public/Templates/IsPODType.h
@@ -3,9 +3,6 @@
 #pragma once
 
 #include "CoreTypes.h"
-#include "Templates/IsPointer.h"
-#include "Templates/AndOrNot.h"
-#include "Templates/IsArithmetic.h"
 
 /**
  * Traits class which tests if a type is POD.
@@ -13,5 +10,5 @@
 template <typename T>
 struct TIsPODType 
 { 
-	enum { Value = TOrValue<__is_pod(T) || __is_enum(T), TIsArithmetic<T>, TIsPointer<T>>::Value };
+	enum { Value = __is_pod(T) };
 };
diff --git a/Engine/Source/Runtime/Core/Public/Templates/IsTriviallyCopyAssignable.h b/Engine/Source/Runtime/Core/Public/Templates/IsTriviallyCopyAssignable.h
--- a/Engine/Source/Runtime/Core/Public/Templates/IsTriviallyCopyAssignable.h
+++ b/Engine/Source/Runtime/Core/Public/Templates/IsTriviallyCopyAssignable.h
@@ -3,8 +3,6 @@
 #pragma once
 
 #include "CoreTypes.h"
-#include "Templates/AndOrNot.h"
-#include "Templates/IsPODType.h"
 
 /**
  * Traits class which tests if a type has a trivial copy assignment operator.
@@ -12,5 +10,5 @@
 template <typename T>
 struct TIsTriviallyCopyAssignable
 {
-	enum { Value = TOrValue<__has_trivial_assign(T), TIsPODType<T>>::Value };
+	enum { Value = __has_trivial_assign(T) };
 };
diff --git a/Engine/Source/Runtime/Core/Public/Templates/IsTriviallyCopyConstructible.h b/Engine/Source/Runtime/Core/Public/Templates/IsTriviallyCopyConstructible.h
--- a/Engine/Source/Runtime/Core/Public/Templates/IsTriviallyCopyConstructible.h
+++ b/Engine/Source/Runtime/Core/Public/Templates/IsTriviallyCopyConstructible.h
@@ -3,8 +3,6 @@
 #pragma once
 
 #include "CoreTypes.h"
-#include "Templates/AndOrNot.h"
-#include "Templates/IsPODType.h"
 
 /**
  * Traits class which tests if a type has a trivial copy constructor.
@@ -12,5 +10,5 @@
 template <typename T>
 struct TIsTriviallyCopyConstructible
 {
-	enum { Value = TOrValue<__has_trivial_copy(T), TIsPODType<T>>::Value };
+	enum { Value = __has_trivial_copy(T) };
 };
diff --git a/Engine/Source/Runtime/Core/Public/Templates/UnrealTypeTraits.h b/Engine/Source/Runtime/Core/Public/Templates/UnrealTypeTraits.h
--- a/Engine/Source/Runtime/Core/Public/Templates/UnrealTypeTraits.h
+++ b/Engine/Source/Runtime/Core/Public/Templates/UnrealTypeTraits.h
@@ -288,7 +288,8 @@ template <typename T>
 struct TCallTraitsBase
 {
 private:
-	enum { PassByValue = TOr<TAndValue<(sizeof(T) <= sizeof(void*)), TIsPODType<T>>, TIsArithmetic<T>, TIsPointer<T>>::Value };
+	enum { PassByValue = TOr<TAndValue<(sizeof(T) <= sizeof(void*)), TIsPODType<T>>, TIsArithmetic<T>>::Value };
+
 public:
 	typedef T ValueType;
 	typedef T& Reference;
diff --git a/Engine/Source/Runtime/CoreUObject/Public/UObject/LazyObjectPtr.h b/Engine/Source/Runtime/CoreUObject/Public/UObject/LazyObjectPtr.h
--- a/Engine/Source/Runtime/CoreUObject/Public/UObject/LazyObjectPtr.h
+++ b/Engine/Source/Runtime/CoreUObject/Public/UObject/LazyObjectPtr.h
@@ -132,12 +132,6 @@ public:
 	{
 	}
 
-	/** Construct from another lazy pointer */
-	FORCEINLINE FLazyObjectPtr(const FLazyObjectPtr& Other)
-	{
-		(*this)=Other;
-	}
-
 	/** Construct from object already in memory */
 	explicit FORCEINLINE FLazyObjectPtr(const UObject* Object)
 	{
@@ -150,12 +144,6 @@ public:
 		TPersistentObjectPtr<FUniqueObjectGuid>::operator=(Object);
 	}
 
-	/** Copy from a lazy pointer */
-	FORCEINLINE void operator=(const FLazyObjectPtr& Other)
-	{
-		TPersistentObjectPtr<FUniqueObjectGuid>::operator=(Other);
-	}
-
 	/** Copy from a unique object identifier */
 	FORCEINLINE void operator=(const FUniqueObjectGuid& InObjectID)
 	{
diff --git a/Engine/Source/Runtime/CoreUObject/Public/UObject/PersistentObjectPtr.h b/Engine/Source/Runtime/CoreUObject/Public/UObject/PersistentObjectPtr.h
--- a/Engine/Source/Runtime/CoreUObject/Public/UObject/PersistentObjectPtr.h
+++ b/Engine/Source/Runtime/CoreUObject/Public/UObject/PersistentObjectPtr.h
@@ -38,12 +38,6 @@ public:
 		TagAtLastTest = 0;
 	}
 
-	/** Construct from another pointer of the same type */
-	FORCEINLINE TPersistentObjectPtr(const TPersistentObjectPtr& Other)
-	{
-		(*this)=Other;
-	}
-
 	/** Construct from a unique object identifier */
 	explicit FORCEINLINE TPersistentObjectPtr(const TObjectID& InObjectID)
 		: WeakPtr()
@@ -83,15 +77,6 @@ public:
 		*this = Object;
 	}
 
-	/** Construct from another pointer */
-	FORCEINLINE void operator=(const TPersistentObjectPtr<TObjectID>& Other)
-	{
-		WeakPtr = Other.WeakPtr;
-		TagAtLastTest = Other.TagAtLastTest;
-		ObjectID = Other.ObjectID;
-
-	}
-	
 	/**
 	 * Gets the unique object identifier associated with this lazy pointer. Valid even if pointer is not currently valid
 	 *
diff --git a/Engine/Source/Runtime/CoreUObject/Public/UObject/WeakObjectPtr.h b/Engine/Source/Runtime/CoreUObject/Public/UObject/WeakObjectPtr.h
--- a/Engine/Source/Runtime/CoreUObject/Public/UObject/WeakObjectPtr.h
+++ b/Engine/Source/Runtime/CoreUObject/Public/UObject/WeakObjectPtr.h
@@ -39,10 +39,7 @@ public:
 	 * Construct from another weak pointer
 	 * @param Other weak pointer to copy from
 	 */
-	FORCEINLINE FWeakObjectPtr(const FWeakObjectPtr& Other)
-	{
-		(*this)=Other;
-	}
+	FWeakObjectPtr(const FWeakObjectPtr& Other) = default;
 
 	/**
 	 * Reset the weak pointer back to the null state
@@ -63,11 +60,7 @@ public:
 	 * Construct from another weak pointer
 	 * @param Other weak pointer to copy from
 	 */
-	FORCEINLINE void operator=(const FWeakObjectPtr& Other)
-	{
-		ObjectIndex = Other.ObjectIndex;
-		ObjectSerialNumber = Other.ObjectSerialNumber;
-	}
+	FWeakObjectPtr& operator=(const FWeakObjectPtr& Other) = default;
 
 	/**  
 	 * Compare weak pointers for equality.
