Attribute value support for values in structured archives, and tests.

diff --git a/Engine/Source/Runtime/Core/Private/Serialization/StructuredArchive.cpp b/Engine/Source/Runtime/Core/Private/Serialization/StructuredArchive.cpp
--- a/Engine/Source/Runtime/Core/Private/Serialization/StructuredArchive.cpp
+++ b/Engine/Source/Runtime/Core/Private/Serialization/StructuredArchive.cpp
@@ -14,15 +14,16 @@
 
 struct FStructuredArchive::FContainer
 {
-	int Index;
-	int Count;
+	int  Index                   = 0;
+	int  Count                   = 0;
+	bool bAttributedValueWritten = false;
+
 #if CHECK_UNIQUE_FIELD_NAMES
 	TSet<FString> KeyNames;
 #endif
 
-	FContainer(int InCount)
-		: Index(0)
-		, Count(InCount)
+	explicit FContainer(int InCount)
+		: Count(InCount)
 	{
 	}
 };
@@ -95,7 +96,7 @@ FStructuredArchive::FSlot FStructuredArchive::Open()
 	check(NextElementId == RootElementId + 1);
 	check(CurrentSlotElementId == INDEX_NONE);
 
-	CurrentScope.Add(FElement(RootElementId, EElementType::Root));
+	CurrentScope.Emplace(RootElementId, EElementType::Root);
 
 	CurrentSlotElementId = NextElementId++;
 
@@ -107,20 +108,49 @@ void FStructuredArchive::Close()
 	SetScope(0, RootElementId);
 }
 
-void FStructuredArchive::EnterSlot(int32 ElementId)
+void FStructuredArchive::EnterSlot(int32 ParentDepth, int32 ElementId)
 {
-	checkf(ElementId == CurrentSlotElementId, TEXT("Attempt to serialize data into an invalid slot"));
+	// If the slot being entered has attributes, enter the value slot first.
+	if (ParentDepth + 1 < CurrentScope.Num() && CurrentScope[ParentDepth + 1].Id == ElementId && CurrentScope[ParentDepth + 1].Type == EElementType::AttributedValue)
+	{
+#if DO_GUARD_SLOW
+		checkf(CurrentSlotElementId == INDEX_NONE && !CurrentContainer.Top()->bAttributedValueWritten, TEXT("Attempt to serialize data into an invalid slot"));
+		CurrentContainer.Top()->bAttributedValueWritten = true;
+#else
+		checkf(CurrentSlotElementId == INDEX_NONE, TEXT("Attempt to serialize data into an invalid slot"));
+#endif
+
+		SetScope(ParentDepth + 1, ElementId);
+		Formatter.EnterAttributedValueValue();
+	}
+	else
+	{
+		checkf(ElementId == CurrentSlotElementId, TEXT("Attempt to serialize data into an invalid slot"));
+		CurrentSlotElementId = INDEX_NONE;
+	}
 
-	CurrentSlotElementId = INDEX_NONE;
+	CurrentEnteringAttributeState = EEnteringAttributeState::NotEnteringAttribute;
 }
 
-void FStructuredArchive::EnterSlot(int32 ElementId, EElementType ElementType)
+int32 FStructuredArchive::EnterSlot(int32 ParentDepth, int32 ElementId, EElementType ElementType)
 {
-	checkf(ElementId == CurrentSlotElementId, TEXT("Attempt to serialize data into an invalid slot"));
+	EnterSlot(ParentDepth, ElementId);
 
-	CurrentSlotElementId = INDEX_NONE;
+	int32 Result = ParentDepth + 1;
 
-	CurrentScope.Add(FElement(ElementId, ElementType));
+	// If we're entering the value of an attributed slot, we need to return a depth one higher than usual, because we're
+	// inside an attributed value container.
+	//
+	// We don't need to do adjust for attributes, because entering the attribute slot will bump the depth anyway.
+	if (ParentDepth + 1 < CurrentScope.Num() &&
+		CurrentScope[ParentDepth + 1].Type == EElementType::AttributedValue &&
+		CurrentEnteringAttributeState == EEnteringAttributeState::NotEnteringAttribute)
+	{
+		++Result;
+	}
+
+	CurrentScope.Emplace(ElementId, ElementType);
+	return Result;
 }
 
 void FStructuredArchive::LeaveSlot()
@@ -147,6 +177,9 @@ void FStructuredArchive::LeaveSlot()
 			CurrentContainer.Top()->Index++;
 #endif
 			break;
+		case EElementType::AttributedValue:
+			Formatter.LeaveAttribute();
+			break;
 		}
 	}
 }
@@ -191,6 +224,12 @@ void FStructuredArchive::SetScope(int32 Depth, int32 ElementId)
 				Formatter.LeaveMap();
 #if DO_GUARD_SLOW
 				delete CurrentContainer.Pop(false);
+#endif
+				break;
+			case EElementType::AttributedValue:
+				Formatter.LeaveAttributedValue();
+#if DO_GUARD_SLOW
+				delete CurrentContainer.Pop(false);
 #endif
 				break;
 			}
@@ -213,7 +252,7 @@ void FStructuredArchive::SetScope(int32 Depth, int32 ElementId)
 
 FStructuredArchive::FRecord FStructuredArchive::FSlot::EnterRecord()
 {
-	Ar.EnterSlot(ElementId, EElementType::Record);
+	int32 NewDepth = Ar.EnterSlot(Depth, ElementId, EElementType::Record);
 
 #if DO_GUARD_SLOW
 	Ar.CurrentContainer.Add(new FContainer(0));
@@ -221,12 +260,12 @@ FStructuredArchive::FRecord FStructuredArchive::FSlot::EnterRecord()
 
 	Ar.Formatter.EnterRecord();
 
-	return FRecord(Ar, Depth + 1, ElementId);
+	return FRecord(Ar, NewDepth, ElementId);
 }
 
 FStructuredArchive::FRecord FStructuredArchive::FSlot::EnterRecord_TextOnly(TArray<FString>& OutFieldNames)
 {
-	Ar.EnterSlot(ElementId, EElementType::Record);
+	int32 NewDepth = Ar.EnterSlot(Depth, ElementId, EElementType::Record);
 
 #if DO_GUARD_SLOW
 	Ar.CurrentContainer.Add(new FContainer(0));
@@ -234,12 +273,12 @@ FStructuredArchive::FRecord FStructuredArchive::FSlot::EnterRecord_TextOnly(TArr
 
 	Ar.Formatter.EnterRecord_TextOnly(OutFieldNames);
 
-	return FRecord(Ar, Depth + 1, ElementId);
+	return FRecord(Ar, NewDepth, ElementId);
 }
 
 FStructuredArchive::FArray FStructuredArchive::FSlot::EnterArray(int32& Num)
 {
-	Ar.EnterSlot(ElementId, EElementType::Array);
+	int32 NewDepth = Ar.EnterSlot(Depth, ElementId, EElementType::Array);
 
 	Ar.Formatter.EnterArray(Num);
 
@@ -247,30 +286,30 @@ FStructuredArchive::FArray FStructuredArchive::FSlot::EnterArray(int32& Num)
 	Ar.CurrentContainer.Add(new FContainer(Num));
 #endif
 
-	return FArray(Ar, Depth + 1, ElementId);
+	return FArray(Ar, NewDepth, ElementId);
 }
 
 FStructuredArchive::FStream FStructuredArchive::FSlot::EnterStream()
 {
-	Ar.EnterSlot(ElementId, EElementType::Stream);
+	int32 NewDepth = Ar.EnterSlot(Depth, ElementId, EElementType::Stream);
 
 	Ar.Formatter.EnterStream();
 
-	return FStream(Ar, Depth + 1, ElementId);
+	return FStream(Ar, NewDepth, ElementId);
 }
 
 FStructuredArchive::FStream FStructuredArchive::FSlot::EnterStream_TextOnly(int32& OutNumElements)
 {
-	Ar.EnterSlot(ElementId, EElementType::Stream);
+	int32 NewDepth = Ar.EnterSlot(Depth, ElementId, EElementType::Stream);
 
 	Ar.Formatter.EnterStream_TextOnly(OutNumElements);
 
-	return FStream(Ar, Depth + 1, ElementId);
+	return FStream(Ar, NewDepth, ElementId);
 }
 
 FStructuredArchive::FMap FStructuredArchive::FSlot::EnterMap(int32& Num)
 {
-	Ar.EnterSlot(ElementId, EElementType::Map);
+	int32 NewDepth = Ar.EnterSlot(Depth, ElementId, EElementType::Map);
 
 	Ar.Formatter.EnterMap(Num);
 
@@ -278,19 +317,60 @@ FStructuredArchive::FMap FStructuredArchive::FSlot::EnterMap(int32& Num)
 	Ar.CurrentContainer.Add(new FContainer(Num));
 #endif
 
-	return FMap(Ar, Depth + 1, ElementId);
+	return FMap(Ar, NewDepth, ElementId);
+}
+
+FStructuredArchive::FSlot FStructuredArchive::FSlot::EnterAttribute(FArchiveFieldName AttributeName)
+{
+	check(Ar.CurrentScope.Num() > 0);
+
+	int32 NewDepth = Depth + 1;
+	if (NewDepth >= Ar.CurrentScope.Num() || Ar.CurrentScope[NewDepth].Id != ElementId || Ar.CurrentScope[NewDepth].Type != EElementType::AttributedValue)
+	{
+		int32 NewDepthCheck = Ar.EnterSlot(Depth, ElementId, EElementType::AttributedValue);
+		checkSlow(NewDepth == NewDepthCheck);
+
+		Ar.Formatter.EnterAttributedValue();
+
+#if DO_GUARD_SLOW
+		Ar.CurrentContainer.Add(new FContainer(0));
+#endif
+	}
+
+	Ar.CurrentEnteringAttributeState = EEnteringAttributeState::NotEnteringAttribute;
+
+	int32 AttributedValueId = Ar.CurrentScope[NewDepth].Id;
+
+	Ar.SetScope(NewDepth, AttributedValueId);
+
+	Ar.CurrentSlotElementId = Ar.NextElementId++;
+
+#if DO_GUARD_SLOW
+#if CHECK_UNIQUE_FIELD_NAMES
+	if (!Ar.GetUnderlyingArchive().IsLoading())
+	{
+		FContainer& Container = *Ar.CurrentContainer.Top();
+		checkf(!Container.KeyNames.Contains(Name.Name), TEXT("Multiple attributes called '%s' serialized into attributed value"), AttributeName.Name);
+		Container.KeyNames.Add(Name.Name);
+	}
+#endif
+#endif
+
+	Ar.Formatter.EnterAttribute(AttributeName);
+
+	return FSlot(Ar, NewDepth, Ar.CurrentSlotElementId);
 }
 
 void FStructuredArchive::FSlot::Serialize(TArray<uint8>& Data)
 {
-	Ar.EnterSlot(ElementId);
+	Ar.EnterSlot(Depth, ElementId);
 	Ar.Formatter.Serialize(Data);
 	Ar.LeaveSlot();
 }
 
 void FStructuredArchive::FSlot::Serialize(void* Data, uint64 DataSize)
 {
-	Ar.EnterSlot(ElementId);
+	Ar.EnterSlot(Depth, ElementId);
 	Ar.Formatter.Serialize(Data, DataSize);
 	Ar.LeaveSlot();
 }
diff --git a/Engine/Source/Runtime/Core/Public/Serialization/Formatters/BinaryArchiveFormatter.h b/Engine/Source/Runtime/Core/Public/Serialization/Formatters/BinaryArchiveFormatter.h
--- a/Engine/Source/Runtime/Core/Public/Serialization/Formatters/BinaryArchiveFormatter.h
+++ b/Engine/Source/Runtime/Core/Public/Serialization/Formatters/BinaryArchiveFormatter.h
@@ -42,6 +42,12 @@ public:
 	virtual void EnterMapElement_TextOnly(FString& Name, EArchiveValueType& OutType) override;
 	virtual void LeaveMapElement() override;
 
+	virtual void EnterAttributedValue() override;
+	virtual void EnterAttribute(FArchiveFieldName AttributeName) override;
+	virtual void EnterAttributedValueValue() override;
+	virtual void LeaveAttribute() override;
+	virtual void LeaveAttributedValue() override;
+
 	virtual void Serialize(uint8& Value) override;
 	virtual void Serialize(uint16& Value) override;
 	virtual void Serialize(uint32& Value) override;
@@ -152,6 +158,26 @@ inline void FBinaryArchiveFormatter::LeaveMapElement()
 {
 }
 
+inline void FBinaryArchiveFormatter::EnterAttributedValue()
+{
+}
+
+inline void FBinaryArchiveFormatter::EnterAttribute(FArchiveFieldName AttributeName)
+{
+}
+
+inline void FBinaryArchiveFormatter::EnterAttributedValueValue()
+{
+}
+
+inline void FBinaryArchiveFormatter::LeaveAttribute()
+{
+}
+
+inline void FBinaryArchiveFormatter::LeaveAttributedValue()
+{
+}
+
 inline void FBinaryArchiveFormatter::Serialize(uint8& Value)
 {
 	Inner << Value;
diff --git a/Engine/Source/Runtime/Core/Public/Serialization/StructuredArchive.h b/Engine/Source/Runtime/Core/Public/Serialization/StructuredArchive.h
--- a/Engine/Source/Runtime/Core/Public/Serialization/StructuredArchive.h
+++ b/Engine/Source/Runtime/Core/Public/Serialization/StructuredArchive.h
@@ -28,6 +28,21 @@ template<typename T> struct TNamedValue
 	}
 };
 
+/**
+ * Class to contain a named attribute for serialization. Intended to be created as a temporary and passed to object serialization methods.
+ */
+template<typename T> struct TNamedAttribute
+{
+	FArchiveFieldName Name;
+	T& Value;
+
+	explicit FORCEINLINE TNamedAttribute(FArchiveFieldName InName, T& InValue)
+		: Name(InName)
+		, Value(InValue)
+	{
+	}
+};
+
 /**
  * Helper function to construct a TNamedValue, deducing the value type.
  */
@@ -36,6 +51,14 @@ template<typename T> FORCEINLINE TNamedValue<T> MakeNamedValue(FArchiveFieldName
 	return TNamedValue<T>(Name, Value);
 }
 
+/**
+ * Helper function to construct a TNamedValue, deducing the value type.
+ */
+template<typename T> FORCEINLINE TNamedAttribute<T> MakeNamedAttribute(FArchiveFieldName Name, T& Value)
+{
+	return TNamedAttribute<T>(Name, Value);
+}
+
 /** Construct a TNamedValue given an ANSI string and value reference. */
 #if WITH_TEXT_ARCHIVE_SUPPORT
 	#define NAMED_ITEM(Name, Value) MakeNamedValue(FArchiveFieldName(TEXT(Name)), Value)
@@ -43,9 +66,19 @@ template<typename T> FORCEINLINE TNamedValue<T> MakeNamedValue(FArchiveFieldName
 	#define NAMED_ITEM(Name, Value) MakeNamedValue(FArchiveFieldName(), Value)
 #endif
 
+/** Construct a TNamedValue given an ANSI string and value reference. */
+#if WITH_TEXT_ARCHIVE_SUPPORT
+#define MAKE_NAMED_ATTRIBUTE(Name, Value) MakeNamedAttribute(FArchiveFieldName(TEXT(Name)), Value)
+#else
+#define MAKE_NAMED_ATTRIBUTE(Name, Value) MakeNamedAttribute(FArchiveFieldName(), Value)
+#endif
+
 /** Construct a TNamedValue using the name of a field or variable. */
 #define NAMED_FIELD(Field) NAMED_ITEM(#Field, Field)
 
+/** Construct a TNamedAttribute using the name of a field or variable. */
+#define NAMED_ATTRIBUTE(Field) MAKE_NAMED_ATTRIBUTE(#Field, Field)
+
 /** Typedef for which formatter type to support */
 #if WITH_TEXT_ARCHIVE_SUPPORT
 	typedef FStructuredArchiveFormatter FArchiveFormatterType;
@@ -145,6 +178,7 @@ public:
 		FStream EnterStream();
 		FStream EnterStream_TextOnly(int32& OutNumElements);
 		FMap EnterMap(int32& Num);
+		FSlot EnterAttribute(FArchiveFieldName AttributeName);
 
 		// We don't support chaining writes to a single slot, so this returns void.
 		template <typename ArgType>
@@ -153,7 +187,7 @@ public:
 		>::Type operator<<(ArgType&& Arg)
 		{
 #if WITH_TEXT_ARCHIVE_SUPPORT
-			Ar.EnterSlot(ElementId);
+			Ar.EnterSlot(Depth, ElementId);
 			Ar.Formatter.Serialize(Forward<ArgType>(Arg));
 			Ar.LeaveSlot();
 #else
@@ -161,6 +195,12 @@ public:
 #endif
 		}
 
+		template <typename T>
+		FORCEINLINE void operator<<(TNamedAttribute<T> Item)
+		{
+			EnterAttribute(Item.Name) << Item.Value;
+		}
+
 		void Serialize(TArray<uint8>& Data);
 		void Serialize(void* Data, uint64 DataSize);
 
@@ -291,6 +331,13 @@ private:
 		Array,
 		Stream,
 		Map,
+		AttributedValue,
+	};
+
+	enum class EEnteringAttributeState
+	{
+		NotEnteringAttribute,
+		EnteringAttribute,
 	};
 
 	struct FElement
@@ -337,15 +384,22 @@ private:
 	TArray<FContainer*> CurrentContainer;
 #endif
 
+	/**
+	 * Whether or not we've just entered an attribute
+	 */
+	EEnteringAttributeState CurrentEnteringAttributeState = EEnteringAttributeState::NotEnteringAttribute;
+
 	/**
 	 * Enters the current slot for serializing a value. Asserts if the archive is not in a state about to write to an empty-slot.
 	 */
-	void EnterSlot(int32 ElementId);
+	void EnterSlot(int32 ParentDepth, int32 ElementId);
 
 	/**
 	 * Enters the current slot, adding an element onto the stack. Asserts if the archive is not in a state about to write to an empty-slot.
+	 *
+	 * @return  The depth of the newly-entered slot.
 	 */
-	void EnterSlot(int32 ElementId, EElementType ElementType);
+	int32 EnterSlot(int32 ParentDepth, int32 ElementId, EElementType ElementType);
 
 	/**
 	 * Leaves slot at the top of the current scope
diff --git a/Engine/Source/Runtime/Core/Public/Serialization/StructuredArchiveFormatter.h b/Engine/Source/Runtime/Core/Public/Serialization/StructuredArchiveFormatter.h
--- a/Engine/Source/Runtime/Core/Public/Serialization/StructuredArchiveFormatter.h
+++ b/Engine/Source/Runtime/Core/Public/Serialization/StructuredArchiveFormatter.h
@@ -55,6 +55,8 @@ enum class EArchiveValueType
 	SoftObjectPath,
 	LazyObjectPtr,
 	RawData,
+	AttributedValue,
+	Attribute,
 };
 
 /**
@@ -103,6 +105,12 @@ public:
 	virtual void EnterMapElement_TextOnly(FString& Name, EArchiveValueType& OutType) = 0;
 	virtual void LeaveMapElement() = 0;
 
+	virtual void EnterAttributedValue() = 0;
+	virtual void EnterAttribute(FArchiveFieldName AttributeName) = 0;
+	virtual void EnterAttributedValueValue() = 0;
+	virtual void LeaveAttribute() = 0;
+	virtual void LeaveAttributedValue() = 0;
+
 	virtual void Serialize(uint8& Value) = 0;
 	virtual void Serialize(uint16& Value) = 0;
 	virtual void Serialize(uint32& Value) = 0;
diff --git a/Engine/Source/Runtime/CoreUObject/Private/Serialization/Formatters/JsonArchiveInputFormatter.cpp b/Engine/Source/Runtime/CoreUObject/Private/Serialization/Formatters/JsonArchiveInputFormatter.cpp
--- a/Engine/Source/Runtime/CoreUObject/Private/Serialization/Formatters/JsonArchiveInputFormatter.cpp
+++ b/Engine/Source/Runtime/CoreUObject/Private/Serialization/Formatters/JsonArchiveInputFormatter.cpp
@@ -58,7 +58,7 @@ bool FJsonArchiveInputFormatter::HasDocumentTree() const
 void FJsonArchiveInputFormatter::EnterRecord()
 {
 	TSharedPtr<FJsonValue>& Value = ValueStack.Top();
-	ObjectStack.Add(FObjectRecord(Value->AsObject(), ValueStack.Num()));
+	ObjectStack.Emplace(Value->AsObject(), ValueStack.Num());
 }
 
 void FJsonArchiveInputFormatter::EnterRecord_TextOnly(TArray<FString>& OutFieldNames)
@@ -213,6 +213,58 @@ void FJsonArchiveInputFormatter::LeaveMapElement()
 	++MapIteratorStack.Top();
 }
 
+void FJsonArchiveInputFormatter::EnterAttributedValue()
+{
+	TSharedPtr<FJsonValue>& Value = ValueStack.Top();
+	const TSharedPtr<FJsonObject>* ObjectPtr = nullptr;
+	if (Value->TryGetObject(ObjectPtr))
+	{
+		FJsonObject& ObjectRef = **ObjectPtr;
+
+		TSharedPtr<FJsonValue> Field = ObjectRef.TryGetField(EscapeFieldName(TEXT("_Value")));
+		if (Field.IsValid())
+		{
+			ObjectStack.Emplace(*ObjectPtr, ValueStack.Num());
+			return;
+		}
+	}
+
+	ObjectStack.Emplace(nullptr, ValueStack.Num());
+}
+
+void FJsonArchiveInputFormatter::EnterAttribute(FArchiveFieldName AttributeName)
+{
+	TSharedPtr<FJsonObject>& NewObject = ObjectStack.Top().JsonObject;
+	TSharedPtr<FJsonValue> Field = NewObject->TryGetField(EscapeFieldName(*FString::Printf(TEXT("_%s"), AttributeName.Name)));
+	check(Field.IsValid());
+	ValueStack.Add(Field);
+}
+
+void FJsonArchiveInputFormatter::EnterAttributedValueValue()
+{
+	if (TSharedPtr<FJsonObject>& Object = ObjectStack.Top().JsonObject)
+	{
+		TSharedPtr<FJsonValue> Field = Object->TryGetField(EscapeFieldName(TEXT("_Value")));
+		checkSlow(Field);
+		ValueStack.Add(Field);
+	}
+	else
+	{
+		ValueStack.Add(CopyTemp(ValueStack.Top()));
+	}
+}
+
+void FJsonArchiveInputFormatter::LeaveAttribute()
+{
+	ValueStack.Pop();
+}
+
+void FJsonArchiveInputFormatter::LeaveAttributedValue()
+{
+	check(ValueStack.Num() == ObjectStack.Top().ValueCountOnCreation);
+	ObjectStack.Pop();
+}
+
 void FJsonArchiveInputFormatter::Serialize(uint8& Value)
 {
 	Value = (uint8)ValueStack.Top()->AsNumber();
diff --git a/Engine/Source/Runtime/CoreUObject/Private/Serialization/Formatters/JsonArchiveOutputFormatter.cpp b/Engine/Source/Runtime/CoreUObject/Private/Serialization/Formatters/JsonArchiveOutputFormatter.cpp
--- a/Engine/Source/Runtime/CoreUObject/Private/Serialization/Formatters/JsonArchiveOutputFormatter.cpp
+++ b/Engine/Source/Runtime/CoreUObject/Private/Serialization/Formatters/JsonArchiveOutputFormatter.cpp
@@ -17,9 +17,6 @@ FJsonArchiveOutputFormatter::FJsonArchiveOutputFormatter(FArchive& InInner)
 	, Newline(LINE_TERMINATOR_ANSI, ARRAY_COUNT(LINE_TERMINATOR_ANSI) - 1)
 {
 	Inner.SetIsTextFormat(true);
-
-	bNeedsComma = false;
-	bNeedsNewline = false;
 }
 
 FJsonArchiveOutputFormatter::~FJsonArchiveOutputFormatter()
@@ -182,6 +179,44 @@ void FJsonArchiveOutputFormatter::LeaveMapElement()
 	LeaveField();
 }
 
+void FJsonArchiveOutputFormatter::EnterAttributedValue()
+{
+	NumAttributesStack.Push(0);
+}
+
+void FJsonArchiveOutputFormatter::EnterAttribute(FArchiveFieldName AttributeName)
+{
+	WriteOptionalComma();
+	WriteOptionalNewline();
+	WriteOptionalAttributedBlockOpening();
+	WriteOptionalComma();
+	WriteOptionalNewline();
+	checkf(FCString::Strcmp(AttributeName.Name, TEXT("Value")) != 0, TEXT("Attributes called 'Value' are reserved by the implementation"));
+	WriteFieldName(*FString::Printf(TEXT("_%s"), AttributeName.Name));
+	++NumAttributesStack.Top();
+}
+
+void FJsonArchiveOutputFormatter::LeaveAttribute()
+{
+	bNeedsComma = true;
+	bNeedsNewline = true;
+}
+
+void FJsonArchiveOutputFormatter::LeaveAttributedValue()
+{
+	WriteOptionalAttributedBlockClosing();
+	NumAttributesStack.Pop();
+	bNeedsComma = true;
+	bNeedsNewline = true;
+}
+
+void FJsonArchiveOutputFormatter::EnterAttributedValueValue()
+{
+	WriteOptionalComma();
+	WriteOptionalNewline();
+	WriteOptionalAttributedBlockValue();
+}
+
 void FJsonArchiveOutputFormatter::Serialize(uint8& Value)
 {
 	WriteValue(LexToString(Value));
@@ -462,6 +497,36 @@ void FJsonArchiveOutputFormatter::WriteOptionalNewline()
 	}
 }
 
+void FJsonArchiveOutputFormatter::WriteOptionalAttributedBlockOpening()
+{
+	if (NumAttributesStack.Top() == 0)
+	{
+		Write('{');
+		Newline.Add('\t');
+		bNeedsNewline = true;
+	}
+}
+
+void FJsonArchiveOutputFormatter::WriteOptionalAttributedBlockValue()
+{
+	if (NumAttributesStack.Top() != 0)
+	{
+		WriteFieldName(TEXT("_Value"));
+	}
+}
+
+void FJsonArchiveOutputFormatter::WriteOptionalAttributedBlockClosing()
+{
+	if (NumAttributesStack.Top() != 0)
+	{
+		Newline.Pop(false);
+		WriteOptionalNewline();
+		Write("}");
+		bNeedsComma                  = true;
+		bNeedsNewline                = true;
+	}
+}
+
 void FJsonArchiveOutputFormatter::SerializeStringInternal(const FString& String)
 {
 	FString Result = TEXT("\"");
diff --git a/Engine/Source/Runtime/CoreUObject/Private/Serialization/Formatters/TaggedBinaryArchiveInputFormatter.cpp b/Engine/Source/Runtime/CoreUObject/Private/Serialization/Formatters/TaggedBinaryArchiveInputFormatter.cpp
--- a/Engine/Source/Runtime/CoreUObject/Private/Serialization/Formatters/TaggedBinaryArchiveInputFormatter.cpp
+++ b/Engine/Source/Runtime/CoreUObject/Private/Serialization/Formatters/TaggedBinaryArchiveInputFormatter.cpp
@@ -1,6 +1,7 @@
 // Copyright 1998-2019 Epic Games, Inc. All Rights Reserved.
 
 #include "Serialization/Formatters/TaggedBinaryArchiveInputFormatter.h"
+#include "Algo/Find.h"
 
 #if WITH_TEXT_ARCHIVE_SUPPORT
 
@@ -220,6 +221,77 @@ void FTaggedBinaryArchiveInputFormatter::LeaveMapElement()
 	CurrentType = EArchiveValueType::None;
 }
 
+void FTaggedBinaryArchiveInputFormatter::EnterAttributedValue()
+{
+	ExpectType(EArchiveValueType::AttributedValue);
+
+	int64 AttributeTableOffset = ReadSize();
+	int64 CurrentOffset = Inner.Tell();
+	Inner.Seek(AttributeTableOffset);
+
+	int32 NumAttributes = 0;
+	Inner << NumAttributes;
+
+	FAttributedValue& AttributedValue = AttributedValues.Emplace_GetRef();
+	AttributedValue.Attributes.SetNum(NumAttributes);
+
+	// Read attribute table
+	for (FAttribute& Attribute : AttributedValue.Attributes)
+	{
+		ExpectType(EArchiveValueType::Attribute);
+
+		Inner << Attribute.NameIdx;
+		Attribute.Size = ReadSize();
+	}
+
+	Inner.Seek(CurrentOffset);
+
+	// Calculate attribute offsets from end of table
+	for (FAttribute& Attribute : AttributedValue.Attributes)
+	{
+		Attribute.StartOffset = CurrentOffset;
+		CurrentOffset += Attribute.Size;
+	}
+
+	AttributedValue.ValueOffset = CurrentOffset;
+}
+
+void FTaggedBinaryArchiveInputFormatter::EnterAttribute(FArchiveFieldName AttributeName)
+{
+	ExpectType(EArchiveValueType::Attribute);
+
+	FAttributedValue& AttributedValue = AttributedValues.Top();
+
+	int32 NameIdx = NameToIndex.FindChecked(AttributeName.Name);
+	const FAttribute* Attribute = Algo::FindBy(AttributedValue.Attributes, NameIdx, &FAttribute::NameIdx);
+	check(Attribute);
+	Inner.Seek(Attribute->StartOffset);
+}
+
+void FTaggedBinaryArchiveInputFormatter::EnterAttributedValueValue()
+{
+	ExpectType(EArchiveValueType::Attribute);
+
+	FAttributedValue& AttributedValue = AttributedValues.Top();
+
+	Inner.Seek(AttributedValue.ValueOffset);
+}
+
+void FTaggedBinaryArchiveInputFormatter::LeaveAttribute()
+{
+	FAttributedValue& AttributedValue = AttributedValues.Top();
+
+	CurrentType = EArchiveValueType::None;
+}
+
+void FTaggedBinaryArchiveInputFormatter::LeaveAttributedValue()
+{
+	FAttributedValue& AttributedValue = AttributedValues.Last();
+
+	Inner.Seek(AttributedValue.ValueOffset + AttributedValue.ValueSize);
+	AttributedValues.Pop();
+}
+
 void FTaggedBinaryArchiveInputFormatter::Serialize(uint8& Value)
 {
 	ReadNumericValue(Value);
@@ -367,9 +439,10 @@ EArchiveValueType FTaggedBinaryArchiveInputFormatter::ReadType()
 	return Type;
 }
 
-void FTaggedBinaryArchiveInputFormatter::ExpectType(EArchiveValueType Type)
+void FTaggedBinaryArchiveInputFormatter::ExpectType(EArchiveValueType InType)
 {
-	verify(ReadType() == Type);
+	EArchiveValueType Type = ReadType();
+	ensure(Type == InType);
 }
 
 template<typename NumericType> void FTaggedBinaryArchiveInputFormatter::ReadNumericValue(NumericType& OutValue)
diff --git a/Engine/Source/Runtime/CoreUObject/Private/Serialization/Formatters/TaggedBinaryArchiveOutputFormatter.cpp b/Engine/Source/Runtime/CoreUObject/Private/Serialization/Formatters/TaggedBinaryArchiveOutputFormatter.cpp
--- a/Engine/Source/Runtime/CoreUObject/Private/Serialization/Formatters/TaggedBinaryArchiveOutputFormatter.cpp
+++ b/Engine/Source/Runtime/CoreUObject/Private/Serialization/Formatters/TaggedBinaryArchiveOutputFormatter.cpp
@@ -8,7 +8,6 @@ FTaggedBinaryArchiveOutputFormatter::FTaggedBinaryArchiveOutputFormatter(FArchiv
 	: Inner(InInner)
 	, SerializeObject(InSerializeObject)
 	, StartOffset(Inner.Tell())
-	, NextRecordIdx(0)
 {
 	Inner.SetIsTextFormat(true);
 
@@ -208,6 +207,67 @@ void FTaggedBinaryArchiveOutputFormatter::LeaveMapElement()
 {
 }
 
+void FTaggedBinaryArchiveOutputFormatter::EnterAttributedValue()
+{
+	WriteType(EArchiveValueType::AttributedValue);
+
+	int64 AttributeTableOffsetOffset = Inner.Tell();
+
+	// Position to be written into when the attributed value is left and the attribute table is written
+	int64 AttributeTableOffset = 0;
+	Inner << AttributeTableOffset;
+
+	FAttributedValue& AttributedValue = AttributedValues.Emplace_GetRef();
+	AttributedValue.AttributeTableOffsetOffset = AttributeTableOffsetOffset;
+}
+
+void FTaggedBinaryArchiveOutputFormatter::EnterAttribute(FArchiveFieldName AttributeName)
+{
+	WriteType(EArchiveValueType::Attribute);
+
+	FAttributedValue& AttributedValue = AttributedValues.Top();
+
+	FAttribute& Attribute = AttributedValue.Attributes.Emplace_GetRef();
+	Attribute.NameIdx = FindOrAddName(AttributeName.Name);
+	Attribute.Offset = Inner.Tell();
+}
+
+void FTaggedBinaryArchiveOutputFormatter::EnterAttributedValueValue()
+{
+}
+
+void FTaggedBinaryArchiveOutputFormatter::LeaveAttribute()
+{
+	FAttributedValue& AttributedValue = AttributedValues.Top();
+	FAttribute& Attribute = AttributedValue.Attributes.Emplace_GetRef();
+
+	int64 CurrentOffset = Inner.Tell();
+	Attribute.Size = CurrentOffset - Attribute.Offset;
+}
+
+void FTaggedBinaryArchiveOutputFormatter::LeaveAttributedValue()
+{
+	FAttributedValue& AttributedValue = AttributedValues.Top();
+
+	int64 AttributeTableOffset = Inner.Tell();
+
+	int32 NumAttributes = AttributedValue.Attributes.Num();
+	Inner << NumAttributes;
+	for (FAttribute& Attribute : AttributedValue.Attributes)
+	{
+		Inner << Attribute.NameIdx;
+		WriteSize(Attribute.Size);
+	}
+
+	int64 AttributeTableEnd = Inner.Tell();
+
+	Inner.Seek(AttributedValue.AttributeTableOffsetOffset);
+	Inner << AttributeTableOffset;
+	Inner.Seek(AttributeTableEnd);
+
+	AttributedValues.Pop();
+}
+
 void FTaggedBinaryArchiveOutputFormatter::Serialize(uint8& Value)
 {
 	WriteType(EArchiveValueType::UInt8);
diff --git a/Engine/Source/Runtime/CoreUObject/Public/Serialization/Formatters/JsonArchiveInputFormatter.h b/Engine/Source/Runtime/CoreUObject/Public/Serialization/Formatters/JsonArchiveInputFormatter.h
--- a/Engine/Source/Runtime/CoreUObject/Public/Serialization/Formatters/JsonArchiveInputFormatter.h
+++ b/Engine/Source/Runtime/CoreUObject/Public/Serialization/Formatters/JsonArchiveInputFormatter.h
@@ -48,6 +48,12 @@ public:
 	virtual void EnterMapElement_TextOnly(FString& OutName, EArchiveValueType& OutType) override;
 	virtual void LeaveMapElement() override;
 
+	virtual void EnterAttributedValue() override;
+	virtual void EnterAttribute(FArchiveFieldName AttributeName) override;
+	virtual void EnterAttributedValueValue() override;
+	virtual void LeaveAttribute() override;
+	virtual void LeaveAttributedValue() override;
+
 	virtual void Serialize(uint8& Value) override;
 	virtual void Serialize(uint16& Value) override;
 	virtual void Serialize(uint32& Value) override;
diff --git a/Engine/Source/Runtime/CoreUObject/Public/Serialization/Formatters/JsonArchiveOutputFormatter.h b/Engine/Source/Runtime/CoreUObject/Public/Serialization/Formatters/JsonArchiveOutputFormatter.h
--- a/Engine/Source/Runtime/CoreUObject/Public/Serialization/Formatters/JsonArchiveOutputFormatter.h
+++ b/Engine/Source/Runtime/CoreUObject/Public/Serialization/Formatters/JsonArchiveOutputFormatter.h
@@ -44,6 +44,12 @@ public:
 	virtual void EnterMapElement_TextOnly(FString& Name, EArchiveValueType& OutType) override;
 	virtual void LeaveMapElement() override;
 
+	virtual void EnterAttributedValue() override;
+	virtual void EnterAttribute(FArchiveFieldName AttributeName) override;
+	virtual void EnterAttributedValueValue() override;
+	virtual void LeaveAttribute() override;
+	virtual void LeaveAttributedValue() override;
+
 	virtual void Serialize(uint8& Value) override;
 	virtual void Serialize(uint16& Value) override;
 	virtual void Serialize(uint32& Value) override;
@@ -70,8 +76,10 @@ private:
 	FArchive& Inner;
 
 	TArray<ANSICHAR> Newline;
-	bool bNeedsComma;
-	bool bNeedsNewline;
+	bool bNeedsComma   = false;
+	bool bNeedsNewline = false;
+
+	TArray<int32> NumAttributesStack;
 
 	void Write(ANSICHAR Character);
 
@@ -83,6 +91,9 @@ private:
 
 	void WriteOptionalComma();
 	void WriteOptionalNewline();
+	void WriteOptionalAttributedBlockOpening();
+	void WriteOptionalAttributedBlockValue();
+	void WriteOptionalAttributedBlockClosing();
 
 	void SerializeStringInternal(const FString& String);
 };
diff --git a/Engine/Source/Runtime/CoreUObject/Public/Serialization/Formatters/TaggedBinaryArchiveInputFormatter.h b/Engine/Source/Runtime/CoreUObject/Public/Serialization/Formatters/TaggedBinaryArchiveInputFormatter.h
--- a/Engine/Source/Runtime/CoreUObject/Public/Serialization/Formatters/TaggedBinaryArchiveInputFormatter.h
+++ b/Engine/Source/Runtime/CoreUObject/Public/Serialization/Formatters/TaggedBinaryArchiveInputFormatter.h
@@ -44,6 +44,12 @@ public:
 	virtual void EnterMapElement_TextOnly(FString& OutName, EArchiveValueType& OutType) override;
 	virtual void LeaveMapElement() override;
 
+	virtual void EnterAttributedValue() override;
+	virtual void EnterAttribute(FArchiveFieldName AttributeName) override;
+	virtual void EnterAttributedValueValue() override;
+	virtual void LeaveAttribute() override;
+	virtual void LeaveAttributedValue() override;
+
 	virtual void Serialize(uint8& Value) override;
 	virtual void Serialize(uint16& Value) override;
 	virtual void Serialize(uint32& Value) override;
@@ -87,6 +93,20 @@ private:
 		int32 NumItems;
 	};
 
+	struct FAttribute
+	{
+		int32  NameIdx;
+		int64  StartOffset;
+		uint64 Size;
+	};
+
+	struct FAttributedValue
+	{
+		TArray<FAttribute> Attributes;
+		int64 ValueOffset;
+		int64 ValueSize;
+	};
+
 	FArchive& Inner;
 	TFunction<void(FArchive&, UObject*&)> SerializeObject;
 
@@ -97,6 +117,7 @@ private:
 	TArray<FRecord> Records;
 	TArray<int32> RecordStack;
 	TArray<FStream> Streams;
+	TArray<FAttributedValue> AttributedValues;
 	EArchiveValueType CurrentType;
 
 	uint64 ReadSize();
diff --git a/Engine/Source/Runtime/CoreUObject/Public/Serialization/Formatters/TaggedBinaryArchiveOutputFormatter.h b/Engine/Source/Runtime/CoreUObject/Public/Serialization/Formatters/TaggedBinaryArchiveOutputFormatter.h
--- a/Engine/Source/Runtime/CoreUObject/Public/Serialization/Formatters/TaggedBinaryArchiveOutputFormatter.h
+++ b/Engine/Source/Runtime/CoreUObject/Public/Serialization/Formatters/TaggedBinaryArchiveOutputFormatter.h
@@ -44,6 +44,12 @@ public:
 	virtual void EnterMapElement_TextOnly(FString& Name, EArchiveValueType& OutType) override;
 	virtual void LeaveMapElement() override;
 
+	virtual void EnterAttributedValue() override;
+	virtual void EnterAttribute(FArchiveFieldName AttributeName) override;
+	virtual void EnterAttributedValueValue() override;
+	virtual void LeaveAttribute() override;
+	virtual void LeaveAttributedValue() override;
+
 	virtual void Serialize(uint8& Value) override;
 	virtual void Serialize(uint16& Value) override;
 	virtual void Serialize(uint32& Value) override;
@@ -87,14 +93,30 @@ private:
 		int32 NumItems;
 	};
 
+	struct FAttribute
+	{
+		int32  NameIdx;
+		uint64 Offset;
+		uint64 Size;
+	};
+
+	struct FAttributedValue
+	{
+		int64 AttributeTableOffsetOffset = 0; // The offset to the attribute table offset in the file
+
+		TArray<FAttribute> Attributes;
+	};
+
 	FArchive& Inner;
 	TFunction<void(FArchive&, UObject*&)> SerializeObject;
 
 	TArray<FString> Names;
 	TMap<FString, int32> NameToIndex;
 
-	int64 StartOffset;
-	int32 NextRecordIdx;
+	int64 StartOffset   = 0;
+	int32 NextRecordIdx = 0;
+	TArray<FAttributedValue> AttributedValues;
+	TArray<int32> AttributedValueStack;
 	TArray<FRecord> Records;
 	TArray<int32> RecordStack;
 	TArray<FStream> Streams;
