Fixed compile error when Csv profiling is disabled but Trace is not. This fix should really follow the value of the CSV_PROFILER macro but this cannot be done due to the arrangement of the header files (they'd be circular).

diff --git a/Engine/Source/Runtime/Core/Public/ProfilingDebugging/CsvProfilerTrace.h b/Engine/Source/Runtime/Core/Public/ProfilingDebugging/CsvProfilerTrace.h
--- a/Engine/Source/Runtime/Core/Public/ProfilingDebugging/CsvProfilerTrace.h
+++ b/Engine/Source/Runtime/Core/Public/ProfilingDebugging/CsvProfilerTrace.h
@@ -6,7 +6,7 @@
 #include "Trace/Trace.h"
 
 #if UE_TRACE_ENABLED && !UE_BUILD_SHIPPING
-#define CSVPROFILERTRACE_ENABLED 1
+#define CSVPROFILERTRACE_ENABLED WITH_ENGINE
 #else
 #define CSVPROFILERTRACE_ENABLED 0
 #endif
