Reordering branch logic to be positive-led.

diff --git a/Engine/Source/Runtime/Core/Public/Templates/MemoryOps.h b/Engine/Source/Runtime/Core/Public/Templates/MemoryOps.h
--- a/Engine/Source/Runtime/Core/Public/Templates/MemoryOps.h
+++ b/Engine/Source/Runtime/Core/Public/Templates/MemoryOps.h
@@ -44,7 +44,11 @@ namespace UE4MemoryOps_Private
 	template <typename ElementType, typename SizeType>
 	FORCEINLINE void DefaultConstructItems(void* Address, SizeType Count)
 	{
-		if constexpr (!TIsZeroConstructType<ElementType>::Value)
+		if constexpr (TIsZeroConstructType<ElementType>::Value)
+		{
+			FMemory::Memset(Address, 0, sizeof(ElementType) * Count);
+		}
+		else
 		{
 			ElementType* Element = (ElementType*)Address;
 			while (Count)
@@ -54,10 +58,6 @@ namespace UE4MemoryOps_Private
 				--Count;
 			}
 		}
-		else
-		{
-			FMemory::Memset(Address, 0, sizeof(ElementType) * Count);
-		}
 	}
 
 #else
@@ -206,7 +206,11 @@ namespace UE4MemoryOps_Private
 	template <typename DestinationElementType, typename SourceElementType, typename SizeType>
 	FORCEINLINE void ConstructItems(void* Dest, const SourceElementType* Source, SizeType Count)
 	{
-		if constexpr (!TIsBitwiseConstructible<DestinationElementType, SourceElementType>::Value)
+		if constexpr (TIsBitwiseConstructible<DestinationElementType, SourceElementType>::Value)
+		{
+			FMemory::Memcpy(Dest, Source, sizeof(SourceElementType) * Count);
+		}
+		else
 		{
 			while (Count)
 			{
@@ -216,10 +220,6 @@ namespace UE4MemoryOps_Private
 				--Count;
 			}
 		}
-		else
-		{
-			FMemory::Memcpy(Dest, Source, sizeof(SourceElementType) * Count);
-		}
 	}
 
 #else
@@ -264,7 +264,11 @@ namespace UE4MemoryOps_Private
 	template <typename ElementType, typename SizeType>
 	FORCEINLINE void CopyAssignItems(ElementType* Dest, const ElementType* Source, SizeType Count)
 	{
-		if constexpr (!TIsTriviallyCopyAssignable<ElementType>::Value)
+		if constexpr (TIsTriviallyCopyAssignable<ElementType>::Value)
+		{
+			FMemory::Memcpy(Dest, Source, sizeof(ElementType) * Count);
+		}
+		else
 		{
 			while (Count)
 			{
@@ -274,10 +278,6 @@ namespace UE4MemoryOps_Private
 				--Count;
 			}
 		}
-		else
-		{
-			FMemory::Memcpy(Dest, Source, sizeof(ElementType) * Count);
-		}
 	}
 
 #else
@@ -323,7 +323,19 @@ namespace UE4MemoryOps_Private
 	template <typename DestinationElementType, typename SourceElementType, typename SizeType>
 	FORCEINLINE void RelocateConstructItems(void* Dest, const SourceElementType* Source, SizeType Count)
 	{
-		if constexpr (!UE4MemoryOps_Private::TCanBitwiseRelocate<DestinationElementType, SourceElementType>::Value)
+		if constexpr (UE4MemoryOps_Private::TCanBitwiseRelocate<DestinationElementType, SourceElementType>::Value)
+		{
+			/* All existing UE containers seem to assume trivial relocatability (i.e. memcpy'able) of their members,
+			 * so we're going to assume that this is safe here.  However, it's not generally possible to assume this
+			 * in general as objects which contain pointers/references to themselves are not safe to be trivially
+			 * relocated.
+			 *
+			 * However, it is not yet possible to automatically infer this at compile time, so we can't enable
+			 * different (i.e. safer) implementations anyway. */
+
+			FMemory::Memmove(Dest, Source, sizeof(SourceElementType) * Count);
+		}
+		else
 		{
 			while (Count)
 			{
@@ -336,18 +348,6 @@ namespace UE4MemoryOps_Private
 				--Count;
 			}
 		}
-		else
-		{
-			/* All existing UE containers seem to assume trivial relocatability (i.e. memcpy'able) of their members,
-			 * so we're going to assume that this is safe here.  However, it's not generally possible to assume this
-			 * in general as objects which contain pointers/references to themselves are not safe to be trivially
-			 * relocated.
-			 *
-			 * However, it is not yet possible to automatically infer this at compile time, so we can't enable
-			 * different (i.e. safer) implementations anyway. */
-
-			FMemory::Memmove(Dest, Source, sizeof(SourceElementType) * Count);
-		}
 	}
 
 #else
@@ -403,7 +403,11 @@ namespace UE4MemoryOps_Private
 	template <typename ElementType, typename SizeType>
 	FORCEINLINE void MoveConstructItems(void* Dest, const ElementType* Source, SizeType Count)
 	{
-		if constexpr (!TIsTriviallyCopyConstructible<ElementType>::Value)
+		if constexpr (TIsTriviallyCopyConstructible<ElementType>::Value)
+		{
+			FMemory::Memmove(Dest, Source, sizeof(ElementType) * Count);
+		}
+		else
 		{
 			while (Count)
 			{
@@ -413,10 +417,6 @@ namespace UE4MemoryOps_Private
 				--Count;
 			}
 		}
-		else
-		{
-			FMemory::Memmove(Dest, Source, sizeof(ElementType) * Count);
-		}
 	}
 
 #else
@@ -460,7 +460,11 @@ namespace UE4MemoryOps_Private
 	template <typename ElementType, typename SizeType>
 	FORCEINLINE void MoveAssignItems(ElementType* Dest, const ElementType* Source, SizeType Count)
 	{
-		if constexpr (!TIsTriviallyCopyAssignable<ElementType>::Value)
+		if constexpr (TIsTriviallyCopyAssignable<ElementType>::Value)
+		{
+			FMemory::Memmove(Dest, Source, sizeof(ElementType) * Count);
+		}
+		else
 		{
 			while (Count)
 			{
@@ -470,10 +474,6 @@ namespace UE4MemoryOps_Private
 				--Count;
 			}
 		}
-		else
-		{
-			FMemory::Memmove(Dest, Source, sizeof(ElementType) * Count);
-		}
 	}
 
 #else
