Moved RootElementId constant out of StructuredArchive.h and into an anonymous namespace in the cpp file. Fixes some obscure linker error on Linux, and also hides away an impementation constant which doesn't need to be visible in a header file.

diff --git a/Engine/Source/Runtime/Core/Private/Serialization/StructuredArchive.cpp b/Engine/Source/Runtime/Core/Private/Serialization/StructuredArchive.cpp
--- a/Engine/Source/Runtime/Core/Private/Serialization/StructuredArchive.cpp
+++ b/Engine/Source/Runtime/Core/Private/Serialization/StructuredArchive.cpp
@@ -6,6 +6,14 @@
 
 #if WITH_TEXT_ARCHIVE_SUPPORT
 
+namespace 
+{
+	/**
+	 * First element ID, which is assigned to the root element.
+	 */
+	static const int RootElementId = 1000;
+}
+
 //////////// FStructuredArchive::FContainer ////////////
 
 #if DO_GUARD_SLOW
diff --git a/Engine/Source/Runtime/Core/Public/Serialization/StructuredArchive.h b/Engine/Source/Runtime/Core/Public/Serialization/StructuredArchive.h
--- a/Engine/Source/Runtime/Core/Public/Serialization/StructuredArchive.h
+++ b/Engine/Source/Runtime/Core/Public/Serialization/StructuredArchive.h
@@ -356,11 +356,6 @@ private:
 	struct FContainer;
 #endif
 
-	/**
-	 * First element ID, which is assigned to the root element.
-	 */
-	static const int RootElementId = 1000;
-
 	/**
 	 * The next element id to be assigned
 	 */
