CIS fix.

diff --git a/Engine/Source/Runtime/Core/Public/Templates/Function.h b/Engine/Source/Runtime/Core/Public/Templates/Function.h
--- a/Engine/Source/Runtime/Core/Public/Templates/Function.h
+++ b/Engine/Source/Runtime/Core/Public/Templates/Function.h
@@ -403,7 +403,11 @@ namespace UE4Function_Private
 				return nullptr;
 			}
 
-			constexpr bool bUseInline = TFUNCTION_USES_INLINE_STORAGE && sizeof(TStorageOwnerTypeT<FunctorType, bUnique, false>) <= TFUNCTION_INLINE_SIZE;
+#if TFUNCTION_USES_INLINE_STORAGE
+			constexpr bool bUseInline = sizeof(TStorageOwnerTypeT<FunctorType, bUnique, false>) <= TFUNCTION_INLINE_SIZE;
+#else
+			constexpr bool bUseInline = false;
+#endif
 
 			using OwnedType = TStorageOwnerTypeT<FunctorType, bUnique, !bUseInline>;
 
