PR #6096: Update type test messages (Contributed by [REDACTED_PERSON])

diff --git a/Engine/Source/Runtime/Core/Public/HAL/Platform.h b/Engine/Source/Runtime/Core/Public/HAL/Platform.h
--- a/Engine/Source/Runtime/Core/Public/HAL/Platform.h
+++ b/Engine/Source/Runtime/Core/Public/HAL/Platform.h
@@ -870,30 +870,30 @@ namespace TypeTests
 	static_assert((!TAreTypesEqual<WIDECHAR, UCS2CHAR>::Value), "WIDECHAR and CHAR16 should be different types.");
 	static_assert((TAreTypesEqual<TCHAR, ANSICHAR>::Value == true || TAreTypesEqual<TCHAR, WIDECHAR>::Value == true), "TCHAR should either be ANSICHAR or WIDECHAR.");
 
-	static_assert(sizeof(uint8) == 1, "BYTE type size test failed.");
-	static_assert(int32(uint8(-1)) == 0xFF, "BYTE type sign test failed.");
+	static_assert(sizeof(uint8) == 1, "uint8 type size test failed.");
+	static_assert(int32(uint8(-1)) == 0xFF, "uint8 type sign test failed.");
 
-	static_assert(sizeof(uint16) == 2, "WORD type size test failed.");
-	static_assert(int32(uint16(-1)) == 0xFFFF, "WORD type sign test failed.");
+	static_assert(sizeof(uint16) == 2, "uint16 type size test failed.");
+	static_assert(int32(uint16(-1)) == 0xFFFF, "uint16 type sign test failed.");
 
-	static_assert(sizeof(uint32) == 4, "DWORD type size test failed.");
-	static_assert(int64(uint32(-1)) == int64(0xFFFFFFFF), "DWORD type sign test failed.");
+	static_assert(sizeof(uint32) == 4, "uint32 type size test failed.");
+	static_assert(int64(uint32(-1)) == int64(0xFFFFFFFF), "uint32 type sign test failed.");
 
-	static_assert(sizeof(uint64) == 8, "QWORD type size test failed.");
-	static_assert(uint64(-1) > uint64(0), "QWORD type sign test failed.");
+	static_assert(sizeof(uint64) == 8, "uint64 type size test failed.");
+	static_assert(uint64(-1) > uint64(0), "uint64 type sign test failed.");
 
 
-	static_assert(sizeof(int8) == 1, "SBYTE type size test failed.");
-	static_assert(int32(int8(-1)) == -1, "SBYTE type sign test failed.");
+	static_assert(sizeof(int8) == 1, "int8 type size test failed.");
+	static_assert(int32(int8(-1)) == -1, "int8 type sign test failed.");
 
-	static_assert(sizeof(int16) == 2, "SWORD type size test failed.");
-	static_assert(int32(int16(-1)) == -1, "SWORD type sign test failed.");
+	static_assert(sizeof(int16) == 2, "int16 type size test failed.");
+	static_assert(int32(int16(-1)) == -1, "int16 type sign test failed.");
 
-	static_assert(sizeof(int32) == 4, "INT type size test failed.");
-	static_assert(int64(int32(-1)) == int64(-1), "INT type sign test failed.");
+	static_assert(sizeof(int32) == 4, "int32 type size test failed.");
+	static_assert(int64(int32(-1)) == int64(-1), "int32 type sign test failed.");
 
-	static_assert(sizeof(int64) == 8, "SQWORD type size test failed.");
-	static_assert(int64(-1) < int64(0), "SQWORD type sign test failed.");
+	static_assert(sizeof(int64) == 8, "int64 type size test failed.");
+	static_assert(int64(-1) < int64(0), "int64 type sign test failed.");
 
 	static_assert(sizeof(ANSICHAR) == 1, "ANSICHAR type size test failed.");
 	static_assert(int32(ANSICHAR(-1)) == -1, "ANSICHAR type sign test failed.");
@@ -902,9 +902,6 @@ namespace TypeTests
 
 	static_assert(sizeof(UCS2CHAR) == 2, "UCS2CHAR type size test failed.");
 
-	static_assert(sizeof(uint32) == 4, "BITFIELD type size test failed.");
-	static_assert(int64(uint32(-1)) == int64(0xFFFFFFFF), "BITFIELD type sign test failed.");
-
 	static_assert(sizeof(PTRINT) == sizeof(void *), "PTRINT type size test failed.");
 	static_assert(PTRINT(-1) < PTRINT(0), "PTRINT type sign test failed.");
 
