Unification of TBA macro names and use of TEXT().

diff --git a/Engine/Plugins/FX/Niagara/Source/Niagara/Private/NiagaraSpriteRendererProperties.cpp b/Engine/Plugins/FX/Niagara/Source/Niagara/Private/NiagaraSpriteRendererProperties.cpp
--- a/Engine/Plugins/FX/Niagara/Source/Niagara/Private/NiagaraSpriteRendererProperties.cpp
+++ b/Engine/Plugins/FX/Niagara/Source/Niagara/Private/NiagaraSpriteRendererProperties.cpp
@@ -95,7 +95,7 @@ void UNiagaraSpriteRendererProperties::Serialize(FStructuredArchive::FRecord Rec
 	FArchive& UnderlyingArchive = Record.GetUnderlyingArchive();
 	if (UnderlyingArchive.IsCooking() || (FPlatformProperties::RequiresCookedData() && UnderlyingArchive.IsLoading()))
 	{
-		DerivedData.Serialize(Record.EnterField(FIELD_NAME_TEXT("DerivedData")));
+		DerivedData.Serialize(Record.EnterField(SA_FIELD_NAME(TEXT("DerivedData"))));
 	}
 }
 
diff --git a/Engine/Source/Editor/UMGEditor/Private/WidgetBlueprint.cpp b/Engine/Source/Editor/UMGEditor/Private/WidgetBlueprint.cpp
--- a/Engine/Source/Editor/UMGEditor/Private/WidgetBlueprint.cpp
+++ b/Engine/Source/Editor/UMGEditor/Private/WidgetBlueprint.cpp
@@ -534,8 +534,8 @@ bool FWidgetAnimation_DEPRECATED::SerializeFromMismatchedTag(struct FPropertyTag
 	if(Tag.Type == NAME_StructProperty && Tag.Name == AnimationDataName)
 	{
 		FStructuredArchive::FRecord Record = Slot.EnterRecord();
-		Record << NAMED_FIELD(MovieScene);
-		Record << NAMED_FIELD(AnimationBindings);
+		Record << SA_VALUE(TEXT("MovieScene"), MovieScene);
+		Record << SA_VALUE(TEXT("AnimationBindings"), AnimationBindings);
 		return true;
 	}
 
diff --git a/Engine/Source/Runtime/Core/Private/Internationalization/GatherableTextData.cpp b/Engine/Source/Runtime/Core/Private/Internationalization/GatherableTextData.cpp
--- a/Engine/Source/Runtime/Core/Private/Internationalization/GatherableTextData.cpp
+++ b/Engine/Source/Runtime/Core/Private/Internationalization/GatherableTextData.cpp
@@ -12,12 +12,12 @@ FArchive& operator<<(FArchive& Archive, FTextSourceSiteContext& This)
 void operator<<(FStructuredArchive::FSlot Slot, FTextSourceSiteContext& This)
 {
 	FStructuredArchive::FRecord Record = Slot.EnterRecord();
-	Record << NAMED_ITEM("KeyName", This.KeyName);
-	Record << NAMED_ITEM("SiteDescription", This.SiteDescription);
-	Record << NAMED_ITEM("IsEditorOnly", This.IsEditorOnly);
-	Record << NAMED_ITEM("IsOptional", This.IsOptional);
-	Record << NAMED_ITEM("InfoMetaData", This.InfoMetaData);
-	Record << NAMED_ITEM("KeyMetaData", This.KeyMetaData);
+	Record << SA_VALUE(TEXT("KeyName"), This.KeyName);
+	Record << SA_VALUE(TEXT("SiteDescription"), This.SiteDescription);
+	Record << SA_VALUE(TEXT("IsEditorOnly"), This.IsEditorOnly);
+	Record << SA_VALUE(TEXT("IsOptional"), This.IsOptional);
+	Record << SA_VALUE(TEXT("InfoMetaData"), This.InfoMetaData);
+	Record << SA_VALUE(TEXT("KeyMetaData"), This.KeyMetaData);
 }
 
 FArchive& operator<<(FArchive& Archive, FTextSourceData& This)
@@ -29,8 +29,8 @@ FArchive& operator<<(FArchive& Archive, FTextSourceData& This)
 void operator<<(FStructuredArchive::FSlot Slot, FTextSourceData& This)
 {
 	FStructuredArchive::FRecord Record = Slot.EnterRecord();
-	Record << NAMED_ITEM("SourceString", This.SourceString);
-	Record << NAMED_ITEM("SourceStringMetaData", This.SourceStringMetaData);
+	Record << SA_VALUE(TEXT("SourceString"), This.SourceString);
+	Record << SA_VALUE(TEXT("SourceStringMetaData"), This.SourceStringMetaData);
 }
 
 FArchive& operator<<(FArchive& Archive, FGatherableTextData& This)
@@ -42,7 +42,7 @@ FArchive& operator<<(FArchive& Archive, FGatherableTextData& This)
 void operator<<(FStructuredArchive::FSlot Slot, FGatherableTextData& This)
 {
 	FStructuredArchive::FRecord Record = Slot.EnterRecord();
-	Record << NAMED_ITEM("NamespaceName", This.NamespaceName);
-	Record << NAMED_ITEM("SourceData", This.SourceData);
-	Record << NAMED_ITEM("SourceSiteContexts", This.SourceSiteContexts);
+	Record << SA_VALUE(TEXT("NamespaceName"), This.NamespaceName);
+	Record << SA_VALUE(TEXT("SourceData"), This.SourceData);
+	Record << SA_VALUE(TEXT("SourceSiteContexts"), This.SourceSiteContexts);
 }
\ No newline at end of file
diff --git a/Engine/Source/Runtime/Core/Private/Internationalization/InternationalizationMetadata.cpp b/Engine/Source/Runtime/Core/Private/Internationalization/InternationalizationMetadata.cpp
--- a/Engine/Source/Runtime/Core/Private/Internationalization/InternationalizationMetadata.cpp
+++ b/Engine/Source/Runtime/Core/Private/Internationalization/InternationalizationMetadata.cpp
@@ -293,10 +293,10 @@ namespace
 		}
 
 		int32 MetaDataTypeAsInt = static_cast<int32>(MetaDataType);
-		Record << NAMED_ITEM("Type", MetaDataTypeAsInt);
+		Record << SA_VALUE(TEXT("Type"), MetaDataTypeAsInt);
 		MetaDataType = static_cast<ELocMetadataType>(MetaDataTypeAsInt);
 
-		FStructuredArchive::FSlot ValueSlot = Record.EnterField(FIELD_NAME_TEXT("Value"));
+		FStructuredArchive::FSlot ValueSlot = Record.EnterField(SA_FIELD_NAME(TEXT("Value")));
 
 		switch (MetaDataType)
 		{
@@ -350,7 +350,7 @@ void operator<<(FStructuredArchive::FSlot Slot, FLocMetadataObject& Object)
 {
 	FStructuredArchive::FRecord Record = Slot.EnterRecord();
 	int32 ValueCount = Object.Values.Num();
-	Record << NAMED_FIELD(ValueCount);
+	Record << SA_VALUE(TEXT("ValueCount"), ValueCount);
 
 	bool bIsLoading = Slot.GetUnderlyingArchive().IsLoading();
 
@@ -359,7 +359,7 @@ void operator<<(FStructuredArchive::FSlot Slot, FLocMetadataObject& Object)
 		Object.Values.Reserve(ValueCount);
 	}
 
-	FStructuredArchive::FStream Stream = Record.EnterStream(FIELD_NAME_TEXT("Values"));
+	FStructuredArchive::FStream Stream = Record.EnterStream(SA_FIELD_NAME(TEXT("Values")));
 	TArray<FString> MapKeys;
 	Object.Values.GetKeys(MapKeys);
 	for (int32 i = 0; i < ValueCount; ++i)
@@ -373,9 +373,9 @@ void operator<<(FStructuredArchive::FSlot Slot, FLocMetadataObject& Object)
 			Key = MapKeys[i];
 		}
 
-		ValueRecord << NAMED_FIELD(Key);
+		ValueRecord << SA_VALUE(TEXT("Key"), Key);
 
-		FStructuredArchive::FSlot ValueSlot = ValueRecord.EnterField(FIELD_NAME_TEXT("Value"));
+		FStructuredArchive::FSlot ValueSlot = ValueRecord.EnterField(SA_FIELD_NAME(TEXT("Value")));
 		if (bIsLoading)
 		{
 			TSharedPtr<FLocMetadataValue> Value;
diff --git a/Engine/Source/Runtime/Core/Private/Internationalization/Text.cpp b/Engine/Source/Runtime/Core/Private/Internationalization/Text.cpp
--- a/Engine/Source/Runtime/Core/Private/Internationalization/Text.cpp
+++ b/Engine/Source/Runtime/Core/Private/Internationalization/Text.cpp
@@ -131,19 +131,19 @@ void operator<<(FStructuredArchive::FSlot Slot, FNumberFormattingOptions& Value)
 
 	if (UnderlyingArchive.CustomVer(FEditorObjectVersion::GUID) >= FEditorObjectVersion::AddedAlwaysSignNumberFormattingOption)
 	{
-		Record << NAMED_ITEM("AlwaysSign", Value.AlwaysSign);
+		Record << SA_VALUE(TEXT("AlwaysSign"), Value.AlwaysSign);
 	}
 
-	Record << NAMED_ITEM("UseGrouping", Value.UseGrouping);
+	Record << SA_VALUE(TEXT("UseGrouping"), Value.UseGrouping);
 
 	int8 RoundingModeInt8 = (int8)Value.RoundingMode;
-	Record << NAMED_ITEM("RoundingMode", RoundingModeInt8);
+	Record << SA_VALUE(TEXT("RoundingMode"), RoundingModeInt8);
 	Value.RoundingMode = (ERoundingMode)RoundingModeInt8;
 
-	Record << NAMED_ITEM("MinimumIntegralDigits", Value.MinimumIntegralDigits);
-	Record << NAMED_ITEM("MaximumIntegralDigits", Value.MaximumIntegralDigits);
-	Record << NAMED_ITEM("MinimumFractionalDigits", Value.MinimumFractionalDigits);
-	Record << NAMED_ITEM("MaximumFractionalDigits", Value.MaximumFractionalDigits);
+	Record << SA_VALUE(TEXT("MinimumIntegralDigits"), Value.MinimumIntegralDigits);
+	Record << SA_VALUE(TEXT("MaximumIntegralDigits"), Value.MaximumIntegralDigits);
+	Record << SA_VALUE(TEXT("MinimumFractionalDigits"), Value.MinimumFractionalDigits);
+	Record << SA_VALUE(TEXT("MaximumFractionalDigits"), Value.MaximumFractionalDigits);
 }
 
 uint32 GetTypeHash( const FNumberFormattingOptions& Key )
@@ -806,7 +806,7 @@ void FText::SerializeText(FStructuredArchive::FSlot Slot, FText& Value)
 	if (UnderlyingArchive.IsLoading() && UnderlyingArchive.UE4Ver() < VER_UE4_FTEXT_HISTORY)
 	{
 		FString SourceStringToImplantIntoHistory;
-		Record << NAMED_FIELD(SourceStringToImplantIntoHistory);
+		Record << SA_VALUE(TEXT("SourceStringToImplantIntoHistory"), SourceStringToImplantIntoHistory);
 
 		FTextDisplayStringPtr DisplayString;
 
@@ -816,8 +816,8 @@ void FText::SerializeText(FStructuredArchive::FSlot Slot, FText& Value)
 			FString Namespace;
 			FString Key;
 
-			Record << NAMED_FIELD(Namespace);
-			Record << NAMED_FIELD(Key);
+			Record << SA_VALUE(TEXT("Namespace"), Namespace);
+			Record << SA_VALUE(TEXT("Key"), Key);
 
 			// Get the DisplayString using the namespace, key, and source string.
 			DisplayString = FTextLocalizationManager::Get().GetDisplayString(Namespace, Key, &SourceStringToImplantIntoHistory);
@@ -855,7 +855,7 @@ void FText::SerializeText(FStructuredArchive::FSlot Slot, FText& Value)
 			Value.Flags &= ~(ETextFlag::ConvertedProperty | ETextFlag::InitializedFromString); // Remove conversion flag before saving.
 		}
 	}
-	Record << NAMED_ITEM("Flags", Value.Flags);
+	Record << SA_VALUE(TEXT("Flags"), Value.Flags);
 
 	if (UnderlyingArchive.IsLoading() && UnderlyingArchive.IsPersistent())
 	{
@@ -879,14 +879,14 @@ void FText::SerializeText(FStructuredArchive::FSlot Slot, FText& Value)
 			if (!bSerializeHistory)
 			{
 				int8 HistoryType = (int8)ETextHistoryType::None;
-				Record << NAMED_FIELD(HistoryType);
+				Record << SA_VALUE(TEXT("HistoryType"), HistoryType);
 
 				bool bHasCultureInvariantString = !Value.IsEmpty() && Value.IsCultureInvariant();
-				Record << NAMED_FIELD(bHasCultureInvariantString);
+				Record << SA_VALUE(TEXT("bHasCultureInvariantString"), bHasCultureInvariantString);
 				if (bHasCultureInvariantString)
 				{
 					FString CultureInvariantString = Value.GetSourceString();
-					Record << NAMED_FIELD(CultureInvariantString);
+					Record << SA_VALUE(TEXT("CultureInvariantString"), CultureInvariantString);
 				}
 			}
 		}
@@ -894,7 +894,7 @@ void FText::SerializeText(FStructuredArchive::FSlot Slot, FText& Value)
 		{
 			// The type is serialized during the serialization of the history, during deserialization we need to deserialize it and create the correct history
 			int8 HistoryType = (int8)ETextHistoryType::None;
-			Record << NAMED_FIELD(HistoryType);
+			Record << SA_VALUE(TEXT("HistoryType"), HistoryType);
 
 			// Create the history class based on the serialized type
 			switch((ETextHistoryType)HistoryType)
@@ -972,11 +972,11 @@ void FText::SerializeText(FStructuredArchive::FSlot Slot, FText& Value)
 					if (UnderlyingArchive.CustomVer(FEditorObjectVersion::GUID) >= FEditorObjectVersion::CultureInvariantTextSerializationKeyStability)
 					{
 						bool bHasCultureInvariantString = false;
-						Record << NAMED_FIELD(bHasCultureInvariantString);
+						Record << SA_VALUE(TEXT("bHasCultureInvariantString"), bHasCultureInvariantString);
 						if (bHasCultureInvariantString)
 						{
 							FString CultureInvariantString;
-							Record << NAMED_FIELD(CultureInvariantString);
+							Record << SA_VALUE(TEXT("CultureInvariantString"), CultureInvariantString);
 							Value.TextData = FText(MoveTemp(CultureInvariantString)).TextData;
 						}
 					}
@@ -1185,29 +1185,29 @@ void operator<<(FStructuredArchive::FSlot Slot, FFormatArgumentValue& Value)
 {
 	FStructuredArchive::FRecord Record = Slot.EnterRecord();
 	int8 TypeAsInt8 = Value.GetType();
-	Record << NAMED_ITEM("Type", TypeAsInt8);
+	Record << SA_VALUE(TEXT("Type"), TypeAsInt8);
 	Value.Type = (EFormatArgumentType::Type)TypeAsInt8;
 
 	switch(Value.Type)
 	{
 	case EFormatArgumentType::Double:
 		{
-			Record << NAMED_ITEM("Value", Value.DoubleValue);
+			Record << SA_VALUE(TEXT("Value"), Value.DoubleValue);
 			break;
 		}
 	case EFormatArgumentType::Float:
 		{
-			Record << NAMED_ITEM("Value", Value.FloatValue);
+			Record << SA_VALUE(TEXT("Value"), Value.FloatValue);
 			break;
 		}
 	case EFormatArgumentType::Int:
 		{
-			Record << NAMED_ITEM("Value", Value.IntValue);
+			Record << SA_VALUE(TEXT("Value"), Value.IntValue);
 			break;
 		}
 	case EFormatArgumentType::UInt:
 		{
-			Record << NAMED_ITEM("Value", Value.UIntValue);
+			Record << SA_VALUE(TEXT("Value"), Value.UIntValue);
 			break;
 		}
 	case EFormatArgumentType::Text:
@@ -1216,7 +1216,7 @@ void operator<<(FStructuredArchive::FSlot Slot, FFormatArgumentValue& Value)
 			{
 				Value.TextValue = FText();
 			}
-			Record << NAMED_ITEM("Value", Value.TextValue.GetValue());
+			Record << SA_VALUE(TEXT("Value"), Value.TextValue.GetValue());
 			break;
 		}
 	}
@@ -1391,18 +1391,18 @@ void operator<<(FStructuredArchive::FSlot Slot, FFormatArgumentData& Value)
 		// ArgumentName was changed to be FString rather than FText, so we need to convert older data to ensure serialization stays happy outside of UStruct::SerializeTaggedProperties.
 		if (UnderlyingArchive.UE4Ver() >= VER_UE4_K2NODE_VAR_REFERENCEGUIDS) // There was no version bump for this change, but VER_UE4_K2NODE_VAR_REFERENCEGUIDS was made at almost the same time.
 		{
-			Record << NAMED_ITEM("ArgumentName", Value.ArgumentName);
+			Record << SA_VALUE(TEXT("ArgumentName"), Value.ArgumentName);
 		}
 		else
 		{
 			FText TempValue;
-			Record << NAMED_ITEM("ArgumentName", TempValue);
+			Record << SA_VALUE(TEXT("ArgumentName"), TempValue);
 			Value.ArgumentName = TempValue.ToString();
 		}
 	}
 	if (UnderlyingArchive.IsSaving())
 	{
-		Record << NAMED_ITEM("ArgumentName", Value.ArgumentName);
+		Record << SA_VALUE(TEXT("ArgumentName"), Value.ArgumentName);
 	}
 
 	uint8 TypeAsByte = Value.ArgumentValueType;
@@ -1412,7 +1412,7 @@ void operator<<(FStructuredArchive::FSlot Slot, FFormatArgumentData& Value)
 
 		if (UnderlyingArchive.CustomVer(FEditorObjectVersion::GUID) >= FEditorObjectVersion::TextFormatArgumentDataIsVariant)
 		{
-			Record << NAMED_ITEM("Type", TypeAsByte);
+			Record << SA_VALUE(TEXT("Type"), TypeAsByte);
 		}
 		else
 		{
@@ -1422,25 +1422,25 @@ void operator<<(FStructuredArchive::FSlot Slot, FFormatArgumentData& Value)
 	}
 	else if (UnderlyingArchive.IsSaving())
 	{
-		Record << NAMED_ITEM("Type", TypeAsByte);
+		Record << SA_VALUE(TEXT("Type"), TypeAsByte);
 	}
 
 	Value.ArgumentValueType = (EFormatArgumentType::Type)TypeAsByte;
 	switch (Value.ArgumentValueType)
 	{
 	case EFormatArgumentType::Int:
-		Record << NAMED_ITEM("Value", Value.ArgumentValueInt);
+		Record << SA_VALUE(TEXT("Value"), Value.ArgumentValueInt);
 		break;
 	case EFormatArgumentType::Float:
-		Record << NAMED_ITEM("Value", Value.ArgumentValueFloat);
+		Record << SA_VALUE(TEXT("Value"), Value.ArgumentValueFloat);
 		break;
 	case EFormatArgumentType::Text:
-		Record << NAMED_ITEM("Value", Value.ArgumentValue);
+		Record << SA_VALUE(TEXT("Value"), Value.ArgumentValue);
 		break;
 	case EFormatArgumentType::Gender:
 		{
 			uint8& Gender = (uint8&)Value.ArgumentValueGender;
-			Record << NAMED_ITEM("Value", Gender);
+			Record << SA_VALUE(TEXT("Value"), Gender);
 			break;
 		}
 	default:
diff --git a/Engine/Source/Runtime/Core/Private/Internationalization/TextHistory.cpp b/Engine/Source/Runtime/Core/Private/Internationalization/TextHistory.cpp
--- a/Engine/Source/Runtime/Core/Private/Internationalization/TextHistory.cpp
+++ b/Engine/Source/Runtime/Core/Private/Internationalization/TextHistory.cpp
@@ -772,7 +772,7 @@ void FTextHistory_Base::Serialize(FStructuredArchive::FRecord Record)
 	if(Record.GetUnderlyingArchive().IsSaving())
 	{
 		int8 HistoryType = (int8)ETextHistoryType::Base;
-		Record << NAMED_FIELD(HistoryType);
+		Record << SA_VALUE(TEXT("HistoryType"), HistoryType);
 	}
 }
 
@@ -788,9 +788,9 @@ void FTextHistory_Base::SerializeForDisplayString(FStructuredArchive::FRecord Re
 		FString Namespace;
 		FString Key;
 
-		Record << NAMED_FIELD(Namespace);
-		Record << NAMED_FIELD(Key);
-		Record << NAMED_FIELD(SourceString);
+		Record << SA_VALUE(TEXT("Namespace"), Namespace);
+		Record << SA_VALUE(TEXT("Key"), Key);
+		Record << SA_VALUE(TEXT("SourceString"), SourceString);
 
 #if USE_STABLE_LOCALIZATION_KEYS
 		// Make sure the package namespace for this text property is up-to-date
@@ -870,13 +870,13 @@ void FTextHistory_Base::SerializeForDisplayString(FStructuredArchive::FRecord Re
 		}
 
 		// Serialize the Namespace
-		Record << NAMED_FIELD(Namespace);
+		Record << SA_VALUE(TEXT("Namespace"), Namespace);
 
 		// Serialize the Key
-		Record << NAMED_FIELD(Key);
+		Record << SA_VALUE(TEXT("Key"), Key);
 
 		// Serialize the SourceString
-		Record << NAMED_FIELD(SourceString);
+		Record << SA_VALUE(TEXT("SourceString"), SourceString);
 	}
 }
 
@@ -1080,22 +1080,22 @@ void FTextHistory_NamedFormat::Serialize(FStructuredArchive::FRecord Record)
 	if(BaseArchive.IsSaving())
 	{
 		int8 HistoryType = (int8)ETextHistoryType::NamedFormat;
-		Record << NAMED_FIELD(HistoryType);
+		Record << SA_VALUE(TEXT("HistoryType"), HistoryType);
 	}
 
 	if (BaseArchive.IsSaving())
 	{
 		FText FormatText = SourceFmt.GetSourceText();
-		Record.EnterField(FIELD_NAME_TEXT("FormatText")) << FormatText;
+		Record.EnterField(SA_FIELD_NAME(TEXT("FormatText"))) << FormatText;
 	}
 	else if (BaseArchive.IsLoading())
 	{
 		FText FormatText;
-		Record.EnterField(FIELD_NAME_TEXT("FormatText")) << FormatText;
+		Record.EnterField(SA_FIELD_NAME(TEXT("FormatText"))) << FormatText;
 		SourceFmt = FTextFormat(FormatText);
 	}
 
-	Record << NAMED_FIELD(Arguments);
+	Record << SA_VALUE(TEXT("Arguments"), Arguments);
 }
 
 bool FTextHistory_NamedFormat::StaticShouldReadFromBuffer(const TCHAR* Buffer)
@@ -1233,22 +1233,22 @@ void FTextHistory_OrderedFormat::Serialize(FStructuredArchive::FRecord Record)
 	if(BaseArchive.IsSaving())
 	{
 		int8 HistoryType = (int8)ETextHistoryType::OrderedFormat;
-		Record << NAMED_FIELD(HistoryType);
+		Record << SA_VALUE(TEXT("HistoryType"), HistoryType);
 	}
 
 	if (BaseArchive.IsSaving())
 	{
 		FText FormatText = SourceFmt.GetSourceText();
-		Record.EnterField(FIELD_NAME_TEXT("FormatText")) << FormatText;
+		Record.EnterField(SA_FIELD_NAME(TEXT("FormatText"))) << FormatText;
 	}
 	else if (BaseArchive.IsLoading())
 	{
 		FText FormatText;
-		Record.EnterField(FIELD_NAME_TEXT("FormatText")) << FormatText;
+		Record.EnterField(SA_FIELD_NAME(TEXT("FormatText"))) << FormatText;
 		SourceFmt = FTextFormat(FormatText);
 	}
 
-	Record << NAMED_FIELD(Arguments);
+	Record << SA_VALUE(TEXT("Arguments"), Arguments);
 }
 
 bool FTextHistory_OrderedFormat::StaticShouldReadFromBuffer(const TCHAR* Buffer)
@@ -1385,22 +1385,22 @@ void FTextHistory_ArgumentDataFormat::Serialize(FStructuredArchive::FRecord Reco
 	if(BaseArchive.IsSaving())
 	{
 		int8 HistoryType = (int8)ETextHistoryType::ArgumentFormat;
-		Record << NAMED_FIELD(HistoryType);
+		Record << SA_VALUE(TEXT("HistoryType"), HistoryType);
 	}
 
 	if (BaseArchive.IsSaving())
 	{
 		FText FormatText = SourceFmt.GetSourceText();
-		Record.EnterField(FIELD_NAME_TEXT("FormatText")) << FormatText;
+		Record.EnterField(SA_FIELD_NAME(TEXT("FormatText"))) << FormatText;
 	}
 	else if (BaseArchive.IsLoading())
 	{
 		FText FormatText;
-		Record.EnterField(FIELD_NAME_TEXT("FormatText")) << FormatText;
+		Record.EnterField(SA_FIELD_NAME(TEXT("FormatText"))) << FormatText;
 		SourceFmt = FTextFormat(FormatText);
 	}
 
-	Record << NAMED_FIELD(Arguments);
+	Record << SA_VALUE(TEXT("Arguments"), Arguments);
 }
 
 bool FTextHistory_ArgumentDataFormat::StaticShouldReadFromBuffer(const TCHAR* Buffer)
@@ -1504,10 +1504,10 @@ void FTextHistory_FormatNumber::Serialize(FStructuredArchive::FRecord Record)
 {
 	FArchive& BaseArchive = Record.GetUnderlyingArchive();
 
-	Record << NAMED_FIELD(SourceValue);
+	Record << SA_VALUE(TEXT("SourceValue"), SourceValue);
 
 	bool bHasFormatOptions = FormatOptions.IsSet();
-	Record << NAMED_FIELD(bHasFormatOptions);
+	Record << SA_VALUE(TEXT("bHasFormatOptions"), bHasFormatOptions);
 
 	if(BaseArchive.IsLoading())
 	{
@@ -1524,18 +1524,18 @@ void FTextHistory_FormatNumber::Serialize(FStructuredArchive::FRecord Record)
 	{
 		check(FormatOptions.IsSet());
 		FNumberFormattingOptions& Options = FormatOptions.GetValue();
-		Record << NAMED_FIELD(Options);
+		Record << SA_VALUE(TEXT("Options"), Options);
 	}
 
 	if(BaseArchive.IsSaving())
 	{
 		FString CultureName = TargetCulture.IsValid()? TargetCulture->GetName() : FString();
-		Record << NAMED_FIELD(CultureName);
+		Record << SA_VALUE(TEXT("CultureName"), CultureName);
 	}
 	else if(BaseArchive.IsLoading())
 	{
 		FString CultureName;
-		Record << NAMED_FIELD(CultureName);
+		Record << SA_VALUE(TEXT("CultureName"), CultureName);
 
 		if(!CultureName.IsEmpty())
 		{
@@ -1611,7 +1611,7 @@ void FTextHistory_AsNumber::Serialize(FStructuredArchive::FRecord Record)
 	if(BaseArchive.IsSaving())
 	{
 		int8 HistoryType = (int8)ETextHistoryType::AsNumber;
-		Record << NAMED_FIELD(HistoryType);
+		Record << SA_VALUE(TEXT("HistoryType"), HistoryType);
 	}
 
 	FTextHistory_FormatNumber::Serialize(Record);
@@ -1688,7 +1688,7 @@ void FTextHistory_AsPercent::Serialize(FStructuredArchive::FRecord Record)
 	if(BaseArchive.IsSaving())
 	{
 		int8 HistoryType = (int8)ETextHistoryType::AsPercent;
-		Record << NAMED_FIELD(HistoryType);
+		Record << SA_VALUE(TEXT("HistoryType"), HistoryType);
 	}
 
 	FTextHistory_FormatNumber::Serialize(Record);
@@ -1773,12 +1773,12 @@ void FTextHistory_AsCurrency::Serialize(FStructuredArchive::FRecord Record)
 	if(BaseArchive.IsSaving())
 	{
 		int8 HistoryType = (int8)ETextHistoryType::AsCurrency;
-		Record << NAMED_FIELD(HistoryType);
+		Record << SA_VALUE(TEXT("HistoryType"), HistoryType);
 	}
 
 	if (BaseArchive.UE4Ver() >= VER_UE4_ADDED_CURRENCY_CODE_TO_FTEXT)
 	{
-		Record << NAMED_FIELD(CurrencyCode);
+		Record << SA_VALUE(TEXT("CurrencyCode"), CurrencyCode);
 	}
 
 	FTextHistory_FormatNumber::Serialize(Record);
@@ -1944,29 +1944,29 @@ void FTextHistory_AsDate::Serialize(FStructuredArchive::FRecord Record)
 	if(BaseArchive.IsSaving())
 	{
 		int8 HistoryType = (int8)ETextHistoryType::AsDate;
-		Record << NAMED_FIELD(HistoryType);
+		Record << SA_VALUE(TEXT("HistoryType"), HistoryType);
 	}
 
-	Record << NAMED_FIELD(SourceDateTime);
+	Record << SA_VALUE(TEXT("SourceDateTime"), SourceDateTime);
 
 	int8 DateStyleInt8 = (int8)DateStyle;
-	Record << NAMED_FIELD(DateStyleInt8);
+	Record << SA_VALUE(TEXT("DateStyleInt8"), DateStyleInt8);
 	DateStyle = (EDateTimeStyle::Type)DateStyleInt8;
 
 	if( BaseArchive.UE4Ver() >= VER_UE4_FTEXT_HISTORY_DATE_TIMEZONE )
 	{
-		Record << NAMED_FIELD(TimeZone);
+		Record << SA_VALUE(TEXT("TimeZone"), TimeZone);
 	}
 
 	if(BaseArchive.IsSaving())
 	{
 		FString CultureName = TargetCulture.IsValid()? TargetCulture->GetName() : FString();
-		Record << NAMED_FIELD(CultureName);
+		Record << SA_VALUE(TEXT("CultureName"), CultureName);
 	}
 	else if(BaseArchive.IsLoading())
 	{
 		FString CultureName;
-		Record << NAMED_FIELD(CultureName);
+		Record << SA_VALUE(TEXT("CultureName"), CultureName);
 
 		if(!CultureName.IsEmpty())
 		{
@@ -2052,26 +2052,26 @@ void FTextHistory_AsTime::Serialize(FStructuredArchive::FRecord Record)
 	if(BaseArchive.IsSaving())
 	{
 		int8 HistoryType = (int8)ETextHistoryType::AsTime;
-		Record << NAMED_FIELD(HistoryType);
+		Record << SA_VALUE(TEXT("HistoryType"), HistoryType);
 	}
 
-	Record << NAMED_FIELD(SourceDateTime);
+	Record << SA_VALUE(TEXT("SourceDateTime"), SourceDateTime);
 
 	int8 TimeStyleInt8 = (int8)TimeStyle;
-	Record << NAMED_ITEM("TimeStyle", TimeStyleInt8);
+	Record << SA_VALUE(TEXT("TimeStyle"), TimeStyleInt8);
 	TimeStyle = (EDateTimeStyle::Type)TimeStyleInt8;
 
-	Record << NAMED_FIELD(TimeZone);
+	Record << SA_VALUE(TEXT("TimeZone"), TimeZone);
 
 	if(BaseArchive.IsSaving())
 	{
 		FString CultureName = TargetCulture.IsValid()? TargetCulture->GetName() : FString();
-		Record << NAMED_FIELD(CultureName);
+		Record << SA_VALUE(TEXT("CultureName"), CultureName);
 	}
 	else if(BaseArchive.IsLoading())
 	{
 		FString CultureName;
-		Record << NAMED_FIELD(CultureName);
+		Record << SA_VALUE(TEXT("CultureName"), CultureName);
 
 		if(!CultureName.IsEmpty())
 		{
@@ -2160,30 +2160,30 @@ void FTextHistory_AsDateTime::Serialize(FStructuredArchive::FRecord Record)
 	if(BaseArchive.IsSaving())
 	{
 		int8 HistoryType = (int8)ETextHistoryType::AsDateTime;
-		Record << NAMED_FIELD(HistoryType);
+		Record << SA_VALUE(TEXT("HistoryType"), HistoryType);
 	}
 
-	Record << NAMED_FIELD(SourceDateTime);
+	Record << SA_VALUE(TEXT("SourceDateTime"), SourceDateTime);
 
 	int8 DateStyleInt8 = (int8)DateStyle;
-	Record << NAMED_ITEM("DateStyle", DateStyleInt8);
+	Record << SA_VALUE(TEXT("DateStyle"), DateStyleInt8);
 	DateStyle = (EDateTimeStyle::Type)DateStyleInt8;
 
 	int8 TimeStyleInt8 = (int8)TimeStyle;
-	Record << NAMED_ITEM("TimeStyle", TimeStyleInt8);
+	Record << SA_VALUE(TEXT("TimeStyle"), TimeStyleInt8);
 	TimeStyle = (EDateTimeStyle::Type)TimeStyleInt8;
 
-	Record << NAMED_FIELD(TimeZone);
+	Record << SA_VALUE(TEXT("TimeZone"), TimeZone);
 
 	if(BaseArchive.IsSaving())
 	{
 		FString CultureName = TargetCulture.IsValid()? TargetCulture->GetName() : FString();
-		Record << NAMED_FIELD(CultureName);
+		Record << SA_VALUE(TEXT("CultureName"), CultureName);
 	}
 	else if(BaseArchive.IsLoading())
 	{
 		FString CultureName;
-		Record << NAMED_FIELD(CultureName);
+		Record << SA_VALUE(TEXT("CultureName"), CultureName);
 
 		if(!CultureName.IsEmpty())
 		{
@@ -2263,13 +2263,13 @@ void FTextHistory_Transform::Serialize(FStructuredArchive::FRecord Record)
 	if (BaseArchive.IsSaving())
 	{
 		int8 HistoryType = (int8)ETextHistoryType::Transform;
-		Record << NAMED_FIELD(HistoryType);
+		Record << SA_VALUE(TEXT("HistoryType"), HistoryType);
 	}
 
-	Record << NAMED_FIELD(SourceText);
+	Record << SA_VALUE(TEXT("SourceText"), SourceText);
 
 	uint8& TransformTypeRef = (uint8&)TransformType;
-	Record << NAMED_ITEM("TransformType", TransformTypeRef);
+	Record << SA_VALUE(TEXT("TransformType"), TransformTypeRef);
 }
 
 bool FTextHistory_Transform::StaticShouldReadFromBuffer(const TCHAR* Buffer)
@@ -2450,7 +2450,7 @@ void FTextHistory_StringTableEntry::Serialize(FStructuredArchive::FRecord Record
 	if (BaseArchive.IsSaving())
 	{
 		int8 HistoryType = (int8)ETextHistoryType::StringTableEntry;
-		Record << NAMED_FIELD(HistoryType);
+		Record << SA_VALUE(TEXT("HistoryType"), HistoryType);
 	}
 
 	if (BaseArchive.IsLoading())
@@ -2460,8 +2460,8 @@ void FTextHistory_StringTableEntry::Serialize(FStructuredArchive::FRecord Record
 
 		FName TableId;
 		FString Key;
-		Record << NAMED_FIELD(TableId);
-		Record << NAMED_FIELD(Key);
+		Record << SA_VALUE(TEXT("TableId"), TableId);
+		Record << SA_VALUE(TEXT("Key"), Key);
 
 		// String Table assets should already have been created via dependency loading when using the EDL (although they may not be fully loaded yet)
 		StringTableReferenceData = MakeShared<FStringTableReferenceData, ESPMode::ThreadSafe>();
@@ -2476,8 +2476,8 @@ void FTextHistory_StringTableEntry::Serialize(FStructuredArchive::FRecord Record
 			StringTableReferenceData->GetTableIdAndKey(TableId, Key);
 		}
 
-		Record << NAMED_FIELD(TableId);
-		Record << NAMED_FIELD(Key);
+		Record << SA_VALUE(TEXT("TableId"), TableId);
+		Record << SA_VALUE(TEXT("Key"), Key);
 	}
 
 	// Collect string table asset references
@@ -2618,7 +2618,7 @@ void FTextHistory_StringTableEntry::FStringTableReferenceData::CollectStringTabl
 	if (Record.GetUnderlyingArchive().IsObjectReferenceCollector())
 	{
 		FScopeLock ScopeLock(&DataCS);
-		IStringTableEngineBridge::CollectStringTableAssetReferences(TableId, Record.EnterField(FIELD_NAME_TEXT("AssetReferences")));
+		IStringTableEngineBridge::CollectStringTableAssetReferences(TableId, Record.EnterField(SA_FIELD_NAME(TEXT("AssetReferences"))));
 	}
 }
 
@@ -2746,13 +2746,13 @@ void FTextHistory_TextGenerator::Serialize(FStructuredArchive::FRecord Record)
 	if (BaseArchive.IsSaving())
 	{
 		int8 HistoryType = (int8)ETextHistoryType::TextGenerator;
-		Record << NAMED_FIELD(HistoryType);
+		Record << SA_VALUE(TEXT("HistoryType"), HistoryType);
 	}
 
 	FName GeneratorTypeID = (BaseArchive.IsSaving() && TextGenerator)
 		? TextGenerator->GetTypeID()
 		: FName();
-	Record << NAMED_FIELD(GeneratorTypeID);
+	Record << SA_VALUE(TEXT("GeneratorTypeID"), GeneratorTypeID);
 
 	TArray<uint8> GeneratorContents;
 
@@ -2764,7 +2764,7 @@ void FTextHistory_TextGenerator::Serialize(FStructuredArchive::FRecord Record)
 		if (GeneratorTypeID != NAME_None)
 		{
 			FText::FCreateTextGeneratorDelegate FactoryFunction = FText::FindRegisteredTextGenerator(GeneratorTypeID);
-			Record << NAMED_FIELD(GeneratorContents);
+			Record << SA_VALUE(TEXT("GeneratorContents"), GeneratorContents);
 
 			if (ensureMsgf(FactoryFunction.IsBound(), TEXT("FTextHistory_TextGenerator::Serialize(): Unable to find registered text generator for \"%s\". Use FText::RegisterTextGenerator() to register a handler."),
 				*GeneratorTypeID.ToString()))
@@ -2799,7 +2799,7 @@ void FTextHistory_TextGenerator::Serialize(FStructuredArchive::FRecord Record)
 			FStructuredArchiveFromArchive ArStructuredWriter(ArWriter);
 
 			TextGenerator->Serialize(ArStructuredWriter.GetSlot().EnterRecord());
-			Record << NAMED_FIELD(GeneratorContents);
+			Record << SA_VALUE(TEXT("GeneratorContents"), GeneratorContents);
 
 			if (ArWriter.IsError())
 			{
diff --git a/Engine/Source/Runtime/Core/Private/Misc/EngineVersion.cpp b/Engine/Source/Runtime/Core/Private/Misc/EngineVersion.cpp
--- a/Engine/Source/Runtime/Core/Private/Misc/EngineVersion.cpp
+++ b/Engine/Source/Runtime/Core/Private/Misc/EngineVersion.cpp
@@ -214,11 +214,11 @@ void operator<<(FArchive &Ar, FEngineVersion &Version)
 void operator<<(FStructuredArchive::FSlot Slot, FEngineVersion &Version)
 {
 	FStructuredArchive::FRecord Record = Slot.EnterRecord();
-	Record << NAMED_ITEM("Major", Version.Major);
-	Record << NAMED_ITEM("Minor", Version.Minor);
-	Record << NAMED_ITEM("Patch", Version.Patch);
-	Record << NAMED_ITEM("Changelist", Version.Changelist);
-	Record << NAMED_ITEM("Branch", Version.Branch);
+	Record << SA_VALUE(TEXT("Major"), Version.Major);
+	Record << SA_VALUE(TEXT("Minor"), Version.Minor);
+	Record << SA_VALUE(TEXT("Patch"), Version.Patch);
+	Record << SA_VALUE(TEXT("Changelist"), Version.Changelist);
+	Record << SA_VALUE(TEXT("Branch"), Version.Branch);
 }
 
 
diff --git a/Engine/Source/Runtime/Core/Private/Misc/Guid.cpp b/Engine/Source/Runtime/Core/Private/Misc/Guid.cpp
--- a/Engine/Source/Runtime/Core/Private/Misc/Guid.cpp
+++ b/Engine/Source/Runtime/Core/Private/Misc/Guid.cpp
@@ -285,6 +285,6 @@ void operator<<(FStructuredArchive::FSlot Slot, FGuid& G)
 	else
 	{
 		FStructuredArchive::FRecord Record = Slot.EnterRecord();
-		Record << NAMED_ITEM("A", G.A) << NAMED_ITEM("B", G.B) << NAMED_ITEM("C", G.C) << NAMED_ITEM("D", G.D);
+		Record << SA_VALUE(TEXT("A"), G.A) << SA_VALUE(TEXT("B"), G.B) << SA_VALUE(TEXT("C"), G.C) << SA_VALUE(TEXT("D"), G.D);
 	}
 }
diff --git a/Engine/Source/Runtime/Core/Private/Misc/ObjectThumbnail.cpp b/Engine/Source/Runtime/Core/Private/Misc/ObjectThumbnail.cpp
--- a/Engine/Source/Runtime/Core/Private/Misc/ObjectThumbnail.cpp
+++ b/Engine/Source/Runtime/Core/Private/Misc/ObjectThumbnail.cpp
@@ -36,8 +36,8 @@ void FObjectThumbnail::Serialize( FArchive& Ar )
 void FObjectThumbnail::Serialize(FStructuredArchive::FSlot Slot)
 {
 	FStructuredArchive::FRecord Record = Slot.EnterRecord();
-	Record << NAMED_FIELD(ImageWidth);
-	Record << NAMED_FIELD(ImageHeight);
+	Record << SA_VALUE(TEXT("ImageWidth"), ImageWidth);
+	Record << SA_VALUE(TEXT("ImageHeight"), ImageHeight);
 
 	//if the image thinks it's empty, ensure there is no memory waste
 	if ((ImageWidth == 0) || (ImageHeight == 0))
@@ -52,11 +52,11 @@ void FObjectThumbnail::Serialize(FStructuredArchive::FSlot Slot)
 		CompressImageData();
 	}
 
-	Record << NAMED_FIELD(CompressedImageData);
+	Record << SA_VALUE(TEXT("CompressedImageData"), CompressedImageData);
 
 	if (Slot.GetUnderlyingArchive().IsCountingMemory())
 	{
-		Record << NAMED_FIELD(ImageData) << NAMED_FIELD(bIsDirty);
+		Record << SA_VALUE(TEXT("ImageData"), ImageData) << SA_VALUE(TEXT("bIsDirty"), bIsDirty);
 	}
 
 	if (Slot.GetUnderlyingArchive().IsLoading())
diff --git a/Engine/Source/Runtime/Core/Private/Serialization/ArchiveFromStructuredArchive.cpp b/Engine/Source/Runtime/Core/Private/Serialization/ArchiveFromStructuredArchive.cpp
--- a/Engine/Source/Runtime/Core/Private/Serialization/ArchiveFromStructuredArchive.cpp
+++ b/Engine/Source/Runtime/Core/Private/Serialization/ArchiveFromStructuredArchive.cpp
@@ -151,7 +151,7 @@ FArchive& FArchiveFromStructuredArchiveImpl::operator<<(class UObject*& Value)
 			}
 			else
 			{
-				FStructuredArchive::FStream Stream = Pimpl->Root->EnterStream(FIELD_NAME_TEXT("Objects"));
+				FStructuredArchive::FStream Stream = Pimpl->Root->EnterStream(SA_FIELD_NAME(TEXT("Objects")));
 
 				// We know exactly which stream index we want to load here, but because of the API we need to read through them
 				// in order, consuming the string name until we reach the entry we want and then load it as a uobject reference.
@@ -258,10 +258,10 @@ void FArchiveFromStructuredArchiveImpl::SerializeInternal(FStructuredArchive::FR
 
 	if (Pimpl->bPendingSerialize)
 	{
-		FStructuredArchive::FSlot DataSlot = Record.EnterField(FIELD_NAME_TEXT("Data"));
+		FStructuredArchive::FSlot DataSlot = Record.EnterField(SA_FIELD_NAME(TEXT("Data")));
 		DataSlot.Serialize(Pimpl->Buffer);
 
-		TOptional<FStructuredArchive::FSlot> ObjectsSlot = Record.TryEnterField(FIELD_NAME_TEXT("Objects"), Pimpl->Objects.Num() > 0);
+		TOptional<FStructuredArchive::FSlot> ObjectsSlot = Record.TryEnterField(SA_FIELD_NAME(TEXT("Objects")), Pimpl->Objects.Num() > 0);
 		if (ObjectsSlot.IsSet())
 		{
 			if (IsLoading())
@@ -286,7 +286,7 @@ void FArchiveFromStructuredArchiveImpl::SerializeInternal(FStructuredArchive::FR
 			}
 		}
 
-		TOptional<FStructuredArchive::FSlot> NamesSlot = Record.TryEnterField(FIELD_NAME_TEXT("Names"), Pimpl->Names.Num() > 0);
+		TOptional<FStructuredArchive::FSlot> NamesSlot = Record.TryEnterField(SA_FIELD_NAME(TEXT("Names")), Pimpl->Names.Num() > 0);
 		if (NamesSlot.IsSet())
 		{
 			NamesSlot.GetValue() << Pimpl->Names;
diff --git a/Engine/Source/Runtime/Core/Private/Serialization/CustomVersion.cpp b/Engine/Source/Runtime/Core/Private/Serialization/CustomVersion.cpp
--- a/Engine/Source/Runtime/Core/Private/Serialization/CustomVersion.cpp
+++ b/Engine/Source/Runtime/Core/Private/Serialization/CustomVersion.cpp
@@ -28,8 +28,8 @@ namespace
 	{
 		// Serialize keys
 		FStructuredArchive::FRecord Record = Slot.EnterRecord();
-		Record << NAMED_ITEM("Tag", Version.Tag);
-		Record << NAMED_ITEM("Version", Version.Version);
+		Record << SA_VALUE(TEXT("Tag"), Version.Tag);
+		Record << SA_VALUE(TEXT("Version"), Version.Version);
 	}
 
 	FArchive& operator<<(FArchive& Ar, FEnumCustomVersion_DEPRECATED& Version)
@@ -54,9 +54,9 @@ namespace
 	void operator<<(FStructuredArchive::FSlot Slot, FGuidCustomVersion_DEPRECATED& Version)
 	{
 		FStructuredArchive::FRecord Record = Slot.EnterRecord();
-		Record << NAMED_ITEM("Key", Version.Key);
-		Record << NAMED_ITEM("Version", Version.Version);
-		Record << NAMED_ITEM("FriendlyName", Version.FriendlyName);
+		Record << SA_VALUE(TEXT("Key"), Version.Key);
+		Record << SA_VALUE(TEXT("Version"), Version.Version);
+		Record << SA_VALUE(TEXT("FriendlyName"), Version.FriendlyName);
 	}
 
 	FArchive& operator<<(FArchive& Ar, FGuidCustomVersion_DEPRECATED& Version)
@@ -118,8 +118,8 @@ FArchive& operator<<(FArchive& Ar, FCustomVersion& Version)
 void operator<<(FStructuredArchive::FSlot Slot, FCustomVersion& Version)
 {
 	FStructuredArchive::FRecord Record = Slot.EnterRecord();
-	Record << NAMED_ITEM("Key", Version.Key);
-	Record << NAMED_ITEM("Version", Version.Version);
+	Record << SA_VALUE(TEXT("Key"), Version.Key);
+	Record << SA_VALUE(TEXT("Version"), Version.Version);
 }
 
 void FCustomVersionContainer::Serialize(FArchive& Ar, ECustomVersionSerializationFormat::Type Format)
diff --git a/Engine/Source/Runtime/Core/Public/Math/Box.h b/Engine/Source/Runtime/Core/Public/Math/Box.h
--- a/Engine/Source/Runtime/Core/Public/Math/Box.h
+++ b/Engine/Source/Runtime/Core/Public/Math/Box.h
@@ -466,7 +466,7 @@ public:
 	friend void operator<<(FStructuredArchive::FSlot Slot, FBox& Box)
 	{
 		FStructuredArchive::FRecord Record = Slot.EnterRecord();
-		Record << NAMED_ITEM("Min", Box.Min) << NAMED_ITEM("Max", Box.Max) << NAMED_ITEM("IsValid", Box.IsValid);
+		Record << SA_VALUE(TEXT("Min"), Box.Min) << SA_VALUE(TEXT("Max"), Box.Max) << SA_VALUE(TEXT("IsValid"), Box.IsValid);
 	}
 
 	bool Serialize( FArchive& Ar )
diff --git a/Engine/Source/Runtime/Core/Public/Math/Color.h b/Engine/Source/Runtime/Core/Public/Math/Color.h
--- a/Engine/Source/Runtime/Core/Public/Math/Color.h
+++ b/Engine/Source/Runtime/Core/Public/Math/Color.h
@@ -64,7 +64,7 @@ struct FLinearColor
 	friend void operator<<(FStructuredArchive::FSlot Slot, FLinearColor& Color)
 	{
 		FStructuredArchive::FRecord Record = Slot.EnterRecord();
-		Record << NAMED_ITEM("R", Color.R) << NAMED_ITEM("G", Color.G) << NAMED_ITEM("B", Color.B) << NAMED_ITEM("A", Color.A);
+		Record << SA_VALUE(TEXT("R"), Color.R) << SA_VALUE(TEXT("G"), Color.G) << SA_VALUE(TEXT("B"), Color.B) << SA_VALUE(TEXT("A"), Color.A);
 	}
 
 	bool Serialize(FStructuredArchive::FSlot Slot)
diff --git a/Engine/Source/Runtime/Core/Public/Math/IntPoint.h b/Engine/Source/Runtime/Core/Public/Math/IntPoint.h
--- a/Engine/Source/Runtime/Core/Public/Math/IntPoint.h
+++ b/Engine/Source/Runtime/Core/Public/Math/IntPoint.h
@@ -295,7 +295,7 @@ public:
 	friend void operator<<(FStructuredArchive::FSlot Slot, FIntPoint& Point)
 	{
 		FStructuredArchive::FRecord Record = Slot.EnterRecord();
-		Record << NAMED_ITEM("X", Point.X) << NAMED_ITEM("Y", Point.Y);
+		Record << SA_VALUE(TEXT("X"), Point.X) << SA_VALUE(TEXT("Y"), Point.Y);
 	}
 
 	/**
diff --git a/Engine/Source/Runtime/Core/Public/Math/IntVector.h b/Engine/Source/Runtime/Core/Public/Math/IntVector.h
--- a/Engine/Source/Runtime/Core/Public/Math/IntVector.h
+++ b/Engine/Source/Runtime/Core/Public/Math/IntVector.h
@@ -300,9 +300,9 @@ public:
 	friend void operator<<(FStructuredArchive::FSlot Slot, FIntVector& Vector)
 	{
 		FStructuredArchive::FRecord Record = Slot.EnterRecord();
-		Record << NAMED_ITEM("X", Vector.X);
-		Record << NAMED_ITEM("Y", Vector.Y);
-		Record << NAMED_ITEM("Z", Vector.Z);
+		Record << SA_VALUE(TEXT("X"), Vector.X);
+		Record << SA_VALUE(TEXT("Y"), Vector.Y);
+		Record << SA_VALUE(TEXT("Z"), Vector.Z);
 	}
 
 	bool Serialize( FArchive& Ar )
diff --git a/Engine/Source/Runtime/Core/Public/Math/Vector.h b/Engine/Source/Runtime/Core/Public/Math/Vector.h
--- a/Engine/Source/Runtime/Core/Public/Math/Vector.h
+++ b/Engine/Source/Runtime/Core/Public/Math/Vector.h
@@ -993,9 +993,9 @@ public:
 		// @warning BulkSerialize: FVector is serialized as memory dump
 		// See TArray::BulkSerialize for detailed description of implied limitations.
 		FStructuredArchive::FRecord Record = Slot.EnterRecord();
-		Record << NAMED_ITEM("X", V.X);
-		Record << NAMED_ITEM("Y", V.Y);
-		Record << NAMED_ITEM("Z", V.Z);
+		Record << SA_VALUE(TEXT("X"), V.X);
+		Record << SA_VALUE(TEXT("Y"), V.Y);
+		Record << SA_VALUE(TEXT("Z"), V.Z);
 	}
 	
 	bool Serialize(FStructuredArchive::FSlot Slot)
diff --git a/Engine/Source/Runtime/Core/Public/Serialization/StructuredArchive.h b/Engine/Source/Runtime/Core/Public/Serialization/StructuredArchive.h
--- a/Engine/Source/Runtime/Core/Public/Serialization/StructuredArchive.h
+++ b/Engine/Source/Runtime/Core/Public/Serialization/StructuredArchive.h
@@ -61,24 +61,18 @@ template<typename T> FORCEINLINE TNamedAttribute<T> MakeNamedAttribute(FArchiveF
 
 /** Construct a TNamedValue given an ANSI string and value reference. */
 #if WITH_TEXT_ARCHIVE_SUPPORT
-	#define NAMED_ITEM(Name, Value) MakeNamedValue(FArchiveFieldName(TEXT(Name)), Value)
+	#define SA_VALUE(Name, Value) MakeNamedValue(FArchiveFieldName(Name), Value)
 #else
-	#define NAMED_ITEM(Name, Value) MakeNamedValue(FArchiveFieldName(), Value)
+	#define SA_VALUE(Name, Value) MakeNamedValue(FArchiveFieldName(), Value)
 #endif
 
 /** Construct a TNamedValue given an ANSI string and value reference. */
 #if WITH_TEXT_ARCHIVE_SUPPORT
-#define MAKE_NAMED_ATTRIBUTE(Name, Value) MakeNamedAttribute(FArchiveFieldName(TEXT(Name)), Value)
+	#define SA_ATTRIBUTE(Name, Value) MakeNamedAttribute(FArchiveFieldName(Name), Value)
 #else
-#define MAKE_NAMED_ATTRIBUTE(Name, Value) MakeNamedAttribute(FArchiveFieldName(), Value)
+	#define SA_ATTRIBUTE(Name, Value) MakeNamedAttribute(FArchiveFieldName(), Value)
 #endif
 
-/** Construct a TNamedValue using the name of a field or variable. */
-#define NAMED_FIELD(Field) NAMED_ITEM(#Field, Field)
-
-/** Construct a TNamedAttribute using the name of a field or variable. */
-#define NAMED_ATTRIBUTE(Field) MAKE_NAMED_ATTRIBUTE(#Field, Field)
-
 /** Typedef for which formatter type to support */
 #if WITH_TEXT_ARCHIVE_SUPPORT
 	typedef FStructuredArchiveFormatter FArchiveFormatterType;
diff --git a/Engine/Source/Runtime/Core/Public/Serialization/StructuredArchiveFormatter.h b/Engine/Source/Runtime/Core/Public/Serialization/StructuredArchiveFormatter.h
--- a/Engine/Source/Runtime/Core/Public/Serialization/StructuredArchiveFormatter.h
+++ b/Engine/Source/Runtime/Core/Public/Serialization/StructuredArchiveFormatter.h
@@ -14,15 +14,13 @@
 		explicit FArchiveFieldName(const TCHAR* InName) : Name(InName){ }
 	};
 
-	#define FIELD_NAME_TEXT(x) FArchiveFieldName(TEXT(x))
-	#define FIELD_NAME(x) FArchiveFieldName(x)
+	#define SA_FIELD_NAME(x) FArchiveFieldName(x)
 #else
 	struct FArchiveFieldName
 	{
 	};
 
-	#define FIELD_NAME_TEXT(x) FArchiveFieldName()
-	#define FIELD_NAME(x) FArchiveFieldName()
+	#define SA_FIELD_NAME(x) FArchiveFieldName()
 #endif
 
 /**
diff --git a/Engine/Source/Runtime/Core/Public/UObject/ScriptDelegates.h b/Engine/Source/Runtime/Core/Public/UObject/ScriptDelegates.h
--- a/Engine/Source/Runtime/Core/Public/UObject/ScriptDelegates.h
+++ b/Engine/Source/Runtime/Core/Public/UObject/ScriptDelegates.h
@@ -145,7 +145,7 @@ public:
 	friend void operator<<(FStructuredArchive::FSlot Slot, TScriptDelegate& D)
 	{
 		FStructuredArchive::FRecord Record = Slot.EnterRecord();
-		Record << NAMED_ITEM("Object", D.Object) << NAMED_ITEM("FunctionName",D.FunctionName);
+		Record << SA_VALUE(TEXT("Object"), D.Object) << SA_VALUE(TEXT("FunctionName"),D.FunctionName);
 	}
 
 	/** Comparison operators */
diff --git a/Engine/Source/Runtime/CoreUObject/Private/Misc/WorldCompositionUtility.cpp b/Engine/Source/Runtime/CoreUObject/Private/Misc/WorldCompositionUtility.cpp
--- a/Engine/Source/Runtime/CoreUObject/Private/Misc/WorldCompositionUtility.cpp
+++ b/Engine/Source/Runtime/CoreUObject/Private/Misc/WorldCompositionUtility.cpp
@@ -35,16 +35,16 @@ void operator<<(FStructuredArchive::FSlot Slot, FWorldTileLayer& D)
 	int32 Version = Slot.GetUnderlyingArchive().UE4Ver();
 
 	// Serialized with FPackageFileSummary
-	Record << NAMED_ITEM("Name", D.Name) << NAMED_ITEM("Reserved0", D.Reserved0) << NAMED_ITEM("Reserved1", D.Reserved1);
+	Record << SA_VALUE(TEXT("Name"), D.Name) << SA_VALUE(TEXT("Reserved0"), D.Reserved0) << SA_VALUE(TEXT("Reserved1"), D.Reserved1);
 
 	if (Version >= VER_UE4_WORLD_LEVEL_INFO_UPDATED)
 	{
-		Record << NAMED_ITEM("StreamingDistance", D.StreamingDistance);
+		Record << SA_VALUE(TEXT("StreamingDistance"), D.StreamingDistance);
 	}
 
 	if (Version >= VER_UE4_WORLD_LAYER_ENABLE_DISTANCE_STREAMING)
 	{
-		Record << NAMED_ITEM("DistanceStreamingEnabled", D.DistanceStreamingEnabled);
+		Record << SA_VALUE(TEXT("DistanceStreamingEnabled"), D.DistanceStreamingEnabled);
 	}
 }
 
@@ -64,11 +64,11 @@ void operator<<(FStructuredArchive::FSlot Slot, FWorldTileLODInfo& D)
 	FStructuredArchive::FRecord Record = Slot.EnterRecord();
 
 	// Serialized with FPackageFileSummary
-	Record << NAMED_ITEM("RelativeStreamingDistance", D.RelativeStreamingDistance)
-		<< NAMED_ITEM("Reserved0", D.Reserved0)
-		<< NAMED_ITEM("Reserved1", D.Reserved1)
-		<< NAMED_ITEM("Reserved2", D.Reserved2)
-		<< NAMED_ITEM("Reserved3", D.Reserved3);
+	Record << SA_VALUE(TEXT("RelativeStreamingDistance"), D.RelativeStreamingDistance)
+		<< SA_VALUE(TEXT("Reserved0"), D.Reserved0)
+		<< SA_VALUE(TEXT("Reserved1"), D.Reserved1)
+		<< SA_VALUE(TEXT("Reserved2"), D.Reserved2)
+		<< SA_VALUE(TEXT("Reserved3"), D.Reserved3);
 }
 
 FArchive& operator<<( FArchive& Ar, FWorldTileInfo& D )
@@ -119,26 +119,26 @@ void operator<<(FStructuredArchive::FSlot Slot, FWorldTileInfo& D)
 	int32 ArchiveVersion = Slot.GetUnderlyingArchive().UE4Ver();
 
 	// Serialized with FPackageFileSummary
-	Record << NAMED_ITEM("Position", D.Position) << NAMED_ITEM("Bounds", D.Bounds) << NAMED_ITEM("Layer", D.Layer);
+	Record << SA_VALUE(TEXT("Position"), D.Position) << SA_VALUE(TEXT("Bounds"), D.Bounds) << SA_VALUE(TEXT("Layer"), D.Layer);
 
 	if (ArchiveVersion >= VER_UE4_WORLD_LEVEL_INFO_UPDATED)
 	{
-		Record << NAMED_ITEM("HideInTileView", D.bHideInTileView) << NAMED_ITEM("ParentTilePackageName", D.ParentTilePackageName);
+		Record << SA_VALUE(TEXT("HideInTileView"), D.bHideInTileView) << SA_VALUE(TEXT("ParentTilePackageName"), D.ParentTilePackageName);
 	}
 
 	if (ArchiveVersion >= VER_UE4_WORLD_LEVEL_INFO_LOD_LIST)
 	{
-		Record << NAMED_ITEM("LODList", D.LODList);
+		Record << SA_VALUE(TEXT("LODList"), D.LODList);
 	}
 
 	if (ArchiveVersion >= VER_UE4_WORLD_LEVEL_INFO_ZORDER)
 	{
-		Record << NAMED_ITEM("ZOrder", D.ZOrder);
+		Record << SA_VALUE(TEXT("ZOrder"), D.ZOrder);
 	}
 
 	if (Slot.GetUnderlyingArchive().GetPortFlags() & PPF_DuplicateForPIE)
 	{
-		Record << NAMED_ITEM("AbsolutePosition", D.AbsolutePosition);
+		Record << SA_VALUE(TEXT("AbsolutePosition"), D.AbsolutePosition);
 	}
 }
 
diff --git a/Engine/Source/Runtime/CoreUObject/Private/Serialization/ArchiveUObjectFromStructuredArchive.cpp b/Engine/Source/Runtime/CoreUObject/Private/Serialization/ArchiveUObjectFromStructuredArchive.cpp
--- a/Engine/Source/Runtime/CoreUObject/Private/Serialization/ArchiveUObjectFromStructuredArchive.cpp
+++ b/Engine/Source/Runtime/CoreUObject/Private/Serialization/ArchiveUObjectFromStructuredArchive.cpp
@@ -153,25 +153,25 @@ void FArchiveUObjectFromStructuredArchiveImpl::SerializeInternal(FStructuredArch
 
 	if (bPendingSerialize)
 	{
-		TOptional<FStructuredArchive::FSlot> LazyObjectPtrsSlot = Record.TryEnterField(FIELD_NAME_TEXT("LazyObjectPtrs"), LazyObjectPtrs.Num() > 0);
+		TOptional<FStructuredArchive::FSlot> LazyObjectPtrsSlot = Record.TryEnterField(SA_FIELD_NAME(TEXT("LazyObjectPtrs")), LazyObjectPtrs.Num() > 0);
 		if (LazyObjectPtrsSlot.IsSet())
 		{
 			LazyObjectPtrsSlot.GetValue() << LazyObjectPtrs;
 		}
 
-		TOptional<FStructuredArchive::FSlot> SoftObjectPtrsSlot = Record.TryEnterField(FIELD_NAME_TEXT("SoftObjectPtrs"), SoftObjectPtrs.Num() > 0);
+		TOptional<FStructuredArchive::FSlot> SoftObjectPtrsSlot = Record.TryEnterField(SA_FIELD_NAME(TEXT("SoftObjectPtrs")), SoftObjectPtrs.Num() > 0);
 		if (SoftObjectPtrsSlot.IsSet())
 		{
 			SoftObjectPtrsSlot.GetValue() << SoftObjectPtrs;
 		}
 
-		TOptional<FStructuredArchive::FSlot> SoftObjectPathsSlot = Record.TryEnterField(FIELD_NAME_TEXT("SoftObjectPaths"), SoftObjectPaths.Num() > 0);
+		TOptional<FStructuredArchive::FSlot> SoftObjectPathsSlot = Record.TryEnterField(SA_FIELD_NAME(TEXT("SoftObjectPaths")), SoftObjectPaths.Num() > 0);
 		if (SoftObjectPathsSlot.IsSet())
 		{
 			SoftObjectPathsSlot.GetValue() << SoftObjectPaths;
 		}
 
-		TOptional<FStructuredArchive::FSlot> WeakObjectPtrsSlot = Record.TryEnterField(FIELD_NAME_TEXT("WeakObjectPtrs"), WeakObjectPtrs.Num() > 0);
+		TOptional<FStructuredArchive::FSlot> WeakObjectPtrsSlot = Record.TryEnterField(SA_FIELD_NAME(TEXT("WeakObjectPtrs")), WeakObjectPtrs.Num() > 0);
 		if (WeakObjectPtrsSlot.IsSet())
 		{
 			WeakObjectPtrsSlot.GetValue() << WeakObjectPtrs;
diff --git a/Engine/Source/Runtime/CoreUObject/Private/UObject/Class.cpp b/Engine/Source/Runtime/CoreUObject/Private/UObject/Class.cpp
--- a/Engine/Source/Runtime/CoreUObject/Private/UObject/Class.cpp
+++ b/Engine/Source/Runtime/CoreUObject/Private/UObject/Class.cpp
@@ -940,7 +940,7 @@ void UStruct::SerializeTaggedProperties(FStructuredArchive::FSlot Slot, uint8* D
 			FStructuredArchive::FRecord PropertyRecord = PropertiesStream.EnterElement().EnterRecord();
 
 			FPropertyTag Tag;
-			PropertyRecord << NAMED_FIELD(Tag);
+			PropertyRecord << SA_VALUE(TEXT("Tag"), Tag);
 
 			if( Tag.Name == NAME_None )
 			{
@@ -1091,7 +1091,7 @@ void UStruct::SerializeTaggedProperties(FStructuredArchive::FSlot Slot, uint8* D
 			}
 			else
 			{
-				FStructuredArchive::FSlot ValueSlot = PropertyRecord.EnterField(FIELD_NAME_TEXT("Value"));
+				FStructuredArchive::FSlot ValueSlot = PropertyRecord.EnterField(SA_FIELD_NAME(TEXT("Value")));
 
 				switch (Property->ConvertFromType(Tag, ValueSlot, Data, DefaultsStruct))
 				{
@@ -1181,7 +1181,7 @@ void UStruct::SerializeTaggedProperties(FStructuredArchive::FSlot Slot, uint8* D
 
 						FStructuredArchive::FRecord PropertyRecord = PropertiesStream.EnterElement().EnterRecord();
 
-						PropertyRecord << NAMED_FIELD(Tag);
+						PropertyRecord << SA_VALUE(TEXT("Tag"), Tag);
 
 						// need to know how much data this call to SerializeTaggedProperty consumes, so mark where we are
 						int64 DataOffset = UnderlyingArchive.Tell();
@@ -1194,7 +1194,7 @@ void UStruct::SerializeTaggedProperties(FStructuredArchive::FSlot Slot, uint8* D
 							UnderlyingArchive.ArCustomPropertyList = CustomPropertyNode->SubPropertyList;
 						}
 
-						FStructuredArchive::FSlot PropertyField = PropertyRecord.EnterField(FIELD_NAME_TEXT("Value"));
+						FStructuredArchive::FSlot PropertyField = PropertyRecord.EnterField(SA_FIELD_NAME(TEXT("Value")));
 						Tag.SerializeTaggedProperty(PropertyField, Property, DataPtr, DefaultValue);
 						if (!PropertyField.IsFilled())
 						{
@@ -1228,7 +1228,7 @@ void UStruct::SerializeTaggedProperties(FStructuredArchive::FSlot Slot, uint8* D
 		}
 
 		static FName Temp(NAME_None);
-		PropertiesStream.EnterElement().EnterRecord().EnterField(FIELD_NAME_TEXT("Tag")).EnterRecord() << NAMED_ITEM("Name", Temp);
+		PropertiesStream.EnterElement().EnterRecord().EnterField(SA_FIELD_NAME(TEXT("Tag"))).EnterRecord() << SA_VALUE(TEXT("Name"), Temp);
 	}
 }
 void UStruct::FinishDestroy()
diff --git a/Engine/Source/Runtime/CoreUObject/Private/UObject/LazyObjectPtr.cpp b/Engine/Source/Runtime/CoreUObject/Private/UObject/LazyObjectPtr.cpp
--- a/Engine/Source/Runtime/CoreUObject/Private/UObject/LazyObjectPtr.cpp
+++ b/Engine/Source/Runtime/CoreUObject/Private/UObject/LazyObjectPtr.cpp
@@ -103,7 +103,7 @@ void FLazyObjectPtr::PossiblySerializeObjectGuid(UObject *Object, FStructuredArc
 	{
 		FUniqueObjectGuid Guid = GuidAnnotation.GetAnnotation(Object);
 		bool HasGuid = Guid.IsValid();
-		Record << NAMED_FIELD(HasGuid);
+		Record << SA_VALUE(TEXT("HasGuid"), HasGuid);
 		if (HasGuid)
 		{
 			if (UnderlyingArchive.GetPortFlags() & PPF_DuplicateForPIE)
@@ -120,17 +120,17 @@ void FLazyObjectPtr::PossiblySerializeObjectGuid(UObject *Object, FStructuredArc
 				}
 			}
 
-			Record << NAMED_FIELD(Guid);
+			Record << SA_VALUE(TEXT("Guid"), Guid);
 		}
 	}
 	else if (UnderlyingArchive.IsLoading())
 	{
 		bool HasGuid = false;
-		Record << NAMED_FIELD(HasGuid);
+		Record << SA_VALUE(TEXT("HasGuid"), HasGuid);
 		if (HasGuid)
 		{
 			FUniqueObjectGuid Guid;
-			Record << NAMED_FIELD(Guid);
+			Record << SA_VALUE(TEXT("Guid"), Guid);
 
 			// Don't try and resolve GUIDs when loading a package for diffing
 			const UPackage* Package = Object->GetOutermost();
diff --git a/Engine/Source/Runtime/CoreUObject/Private/UObject/Linker.cpp b/Engine/Source/Runtime/CoreUObject/Private/UObject/Linker.cpp
--- a/Engine/Source/Runtime/CoreUObject/Private/UObject/Linker.cpp
+++ b/Engine/Source/Runtime/CoreUObject/Private/UObject/Linker.cpp
@@ -73,10 +73,10 @@ FArchive& operator<<(FArchive& Ar,FCompressedChunk& Chunk)
 void operator<<(FStructuredArchive::FSlot Slot, FCompressedChunk& Chunk)
 {
 	FStructuredArchive::FRecord Record = Slot.EnterRecord();
-	Record << NAMED_ITEM("UncompressedOffset", Chunk.UncompressedOffset);
-	Record << NAMED_ITEM("UncompressedSize", Chunk.UncompressedSize);
-	Record << NAMED_ITEM("CompressedOffset", Chunk.CompressedOffset);
-	Record << NAMED_ITEM("CompressedSize", Chunk.CompressedSize);
+	Record << SA_VALUE(TEXT("UncompressedOffset"), Chunk.UncompressedOffset);
+	Record << SA_VALUE(TEXT("UncompressedSize"), Chunk.UncompressedSize);
+	Record << SA_VALUE(TEXT("CompressedOffset"), Chunk.CompressedOffset);
+	Record << SA_VALUE(TEXT("CompressedSize"), Chunk.CompressedSize);
 }
 
 /*----------------------------------------------------------------------------
@@ -98,7 +98,7 @@ void FGenerationInfo::Serialize(FArchive& Ar, const struct FPackageFileSummary&
 void FGenerationInfo::Serialize(FStructuredArchive::FSlot Slot, const struct FPackageFileSummary& Summary)
 {
 	FStructuredArchive::FRecord Record = Slot.EnterRecord();
-	Record << NAMED_FIELD(ExportCount) << NAMED_FIELD(NameCount);
+	Record << SA_VALUE(TEXT("ExportCount"), ExportCount) << SA_VALUE(TEXT("NameCount"), NameCount);
 }
 
 #if WITH_EDITORONLY_DATA
diff --git a/Engine/Source/Runtime/CoreUObject/Private/UObject/LinkerLoad.cpp b/Engine/Source/Runtime/CoreUObject/Private/UObject/LinkerLoad.cpp
--- a/Engine/Source/Runtime/CoreUObject/Private/UObject/LinkerLoad.cpp
+++ b/Engine/Source/Runtime/CoreUObject/Private/UObject/LinkerLoad.cpp
@@ -1160,7 +1160,7 @@ FLinkerLoad::ELinkerStatus FLinkerLoad::SerializePackageFileSummary()
 		}
 #endif
 		// Read summary from file.
-		StructuredArchiveRootRecord.GetValue() << NAMED_FIELD(Summary);
+		StructuredArchiveRootRecord.GetValue() << SA_VALUE(TEXT("Summary"), Summary);
 
 		TRACE_LOADTIME_PACKAGE_SUMMARY(this, Summary.TotalHeaderSize, Summary.NameCount, Summary.ImportCount, Summary.ExportCount);
 
@@ -1471,7 +1471,7 @@ FLinkerLoad::ELinkerStatus FLinkerLoad::SerializeGatherableTextDataMap(bool bFor
 		Seek( Summary.GatherableTextDataOffset );
 	}
 
-	FStructuredArchive::FStream Stream = StructuredArchiveRootRecord->EnterStream(FIELD_NAME_TEXT("GatherableTextData"));
+	FStructuredArchive::FStream Stream = StructuredArchiveRootRecord->EnterStream(SA_FIELD_NAME(TEXT("GatherableTextData")));
 	while (GatherableTextDataMapIndex < Summary.GatherableTextDataCount && !IsTimeLimitExceeded(TEXT("serializing gatherable text data map"), 100))
 	{
 		FGatherableTextData* GatherableTextData = new(GatherableTextDataMap)FGatherableTextData;
@@ -1497,7 +1497,7 @@ FLinkerLoad::ELinkerStatus FLinkerLoad::SerializeImportMap()
 		Seek( Summary.ImportOffset );
 	}
 
-	FStructuredArchive::FStream Stream = StructuredArchiveRootRecord->EnterStream(FIELD_NAME_TEXT("ImportTable"));
+	FStructuredArchive::FStream Stream = StructuredArchiveRootRecord->EnterStream(SA_FIELD_NAME(TEXT("ImportTable")));
 
 	while( ImportMapIndex < Summary.ImportCount && !IsTimeLimitExceeded(TEXT("serializing import map"),100) )
 	{
@@ -1704,7 +1704,7 @@ FLinkerLoad::ELinkerStatus FLinkerLoad::SerializeExportMap()
 	}
 
 
-	FStructuredArchive::FStream Stream = StructuredArchiveRootRecord->EnterStream(FIELD_NAME_TEXT("ExportTable"));
+	FStructuredArchive::FStream Stream = StructuredArchiveRootRecord->EnterStream(SA_FIELD_NAME(TEXT("ExportTable")));
 
 	while( ExportMapIndex < Summary.ExportCount && !IsTimeLimitExceeded(TEXT("serializing export map"),100) )
 	{
@@ -1752,7 +1752,7 @@ FLinkerLoad::ELinkerStatus FLinkerLoad::SerializeDependsMap()
 		DependsMap.AddZeroed(Summary.ExportCount);
 	}
 
-	FStructuredArchive::FStream Stream = StructuredArchiveRootRecord->EnterStream(FIELD_NAME_TEXT("DependsMap"));
+	FStructuredArchive::FStream Stream = StructuredArchiveRootRecord->EnterStream(SA_FIELD_NAME(TEXT("DependsMap")));
 	while (DependsMapIndex < Summary.ExportCount && !IsTimeLimitExceeded(TEXT("serializing depends map"), 100))
 	{
 		TArray<FPackageIndex>& Depends = DependsMap[DependsMapIndex];
@@ -1791,7 +1791,7 @@ FLinkerLoad::ELinkerStatus FLinkerLoad::SerializePreloadDependencies()
 		)
 	{
 		//@todoio check endiness and fastpath this as a single serialize
-		FStructuredArchive::FStream Stream = StructuredArchiveRootRecord->EnterStream(FIELD_NAME_TEXT("PreloadDependencies"));
+		FStructuredArchive::FStream Stream = StructuredArchiveRootRecord->EnterStream(SA_FIELD_NAME(TEXT("PreloadDependencies")));
 		for (int32 Index = 0; Index < Summary.PreloadDependencyCount; Index++)
 		{
 			FPackageIndex Idx;
@@ -1824,7 +1824,7 @@ FLinkerLoad::ELinkerStatus FLinkerLoad::SerializeThumbnails( bool bForceEnableIn
 		return LINKER_Loaded;
 	}
 
-	FStructuredArchive::FRecord Record = StructuredArchiveRootRecord->EnterRecord(FIELD_NAME_TEXT("Thumbnails"));
+	FStructuredArchive::FRecord Record = StructuredArchiveRootRecord->EnterRecord(SA_FIELD_NAME(TEXT("Thumbnails")));
 
 	if((Summary.ThumbnailTableOffset > 0) || IsTextFormat())
 	{
@@ -1832,13 +1832,13 @@ FLinkerLoad::ELinkerStatus FLinkerLoad::SerializeThumbnails( bool bForceEnableIn
 
 		if (IsTextFormat())
 		{
-			IndexSlot = Record.TryEnterField(FIELD_NAME_TEXT("Index"), false);
+			IndexSlot = Record.TryEnterField(SA_FIELD_NAME(TEXT("Index")), false);
 		}
 		else
 		{
 			// Seek to the thumbnail table of contents
 			Seek(Summary.ThumbnailTableOffset);
-			IndexSlot.Emplace(Record.EnterField(FIELD_NAME_TEXT("Index")));
+			IndexSlot.Emplace(Record.EnterField(SA_FIELD_NAME(TEXT("Index"))));
 		}
 
 		if (IndexSlot.IsSet())
@@ -1863,11 +1863,11 @@ FLinkerLoad::ELinkerStatus FLinkerLoad::SerializeThumbnails( bool bForceEnableIn
 
 				FString ObjectClassName;
 				// Newer packages always store the class name for each asset
-				IndexRecord << NAMED_FIELD(ObjectClassName);
+				IndexRecord << SA_VALUE(TEXT("ObjectClassName"), ObjectClassName);
 
 				// Object path
 				FString ObjectPathWithoutPackageName;
-				IndexRecord << NAMED_FIELD(ObjectPathWithoutPackageName);
+				IndexRecord << SA_VALUE(TEXT("ObjectPathWithoutPackageName"), ObjectPathWithoutPackageName);
 				const FString ObjectPath(LinkerRoot->GetName() + TEXT(".") + ObjectPathWithoutPackageName);
 
 
@@ -1876,7 +1876,7 @@ FLinkerLoad::ELinkerStatus FLinkerLoad::SerializeThumbnails( bool bForceEnableIn
 				ThumbnailInfo.ObjectFullName = FName(*ObjectFullName);
 
 				// File offset for the thumbnail (already saved out.)
-				IndexRecord << NAMED_ITEM("FileOffset", ThumbnailInfo.FileOffset);
+				IndexRecord << SA_VALUE(TEXT("FileOffset"), ThumbnailInfo.FileOffset);
 
 				// Only bother loading thumbnails that don't already exist in memory yet.  This is because when we
 				// go to load thumbnails that aren't in memory yet when saving packages we don't want to clobber
@@ -1889,7 +1889,7 @@ FLinkerLoad::ELinkerStatus FLinkerLoad::SerializeThumbnails( bool bForceEnableIn
 			}
 
 
-			FStructuredArchive::FStream DataStream = Record.EnterStream(FIELD_NAME_TEXT("Thumbnails"));
+			FStructuredArchive::FStream DataStream = Record.EnterStream(SA_FIELD_NAME(TEXT("Thumbnails")));
 
 			// Now go and load and cache all of the thumbnails
 			for (int32 CurObjectIndex = 0; CurObjectIndex < ThumbnailInfoArray.Num(); ++CurObjectIndex)
@@ -3578,7 +3578,7 @@ void FLinkerLoad::Preload( UObject* Object )
 						if (IsTextFormat())
 						{
 							FString ObjectName = Object->GetPathName(Object->GetOutermost());
-							FStructuredArchive::FSlot ExportSlot = StructuredArchiveRootRecord->EnterField(FIELD_NAME(*ObjectName));
+							FStructuredArchive::FSlot ExportSlot = StructuredArchiveRootRecord->EnterField(SA_FIELD_NAME(*ObjectName));
 
 							if (bClassSupportsTextFormat)
 							{
@@ -3628,7 +3628,7 @@ void FLinkerLoad::Preload( UObject* Object )
 								ObjectName = ObjectPath.Right(ObjectPath.Len() - PackagePath.Len() - 1);
 							}
 
-							FStructuredArchive::FSlot ExportSlot = StructuredArchiveRootRecord->EnterField(FIELD_NAME(*ObjectName));
+							FStructuredArchive::FSlot ExportSlot = StructuredArchiveRootRecord->EnterField(SA_FIELD_NAME(*ObjectName));
 
 							if (bClassSupportsTextFormat)
 							{
diff --git a/Engine/Source/Runtime/CoreUObject/Private/UObject/MetaData.cpp b/Engine/Source/Runtime/CoreUObject/Private/UObject/MetaData.cpp
--- a/Engine/Source/Runtime/CoreUObject/Private/UObject/MetaData.cpp
+++ b/Engine/Source/Runtime/CoreUObject/Private/UObject/MetaData.cpp
@@ -137,14 +137,14 @@ void UMetaData::Serialize(FStructuredArchive::FRecord Record)
 	
 	if (!UnderlyingArchive.IsLoading())
 	{
-		Record << NAMED_FIELD(ObjectMetaDataMap);
-		Record << NAMED_FIELD(RootMetaDataMap);
+		Record << SA_VALUE(TEXT("ObjectMetaDataMap"), ObjectMetaDataMap);
+		Record << SA_VALUE(TEXT("RootMetaDataMap"), RootMetaDataMap);
 	}
 	else
 	{
 		{
 			TMap< FWeakObjectPtr, TMap<FName, FString> > TempMap;
-			Record << NAMED_ITEM("ObjectMetaDataMap", TempMap);
+			Record << SA_VALUE(TEXT("ObjectMetaDataMap"), TempMap);
 
 			const bool bLoadFromLinker = (NULL != UnderlyingArchive.GetLinker());
 			if (bLoadFromLinker && HasAnyFlags(RF_LoadCompleted))
@@ -164,7 +164,7 @@ void UMetaData::Serialize(FStructuredArchive::FRecord Record)
 		if (UnderlyingArchive.CustomVer(FEditorObjectVersion::GUID) >= FEditorObjectVersion::RootMetaDataSupport)
 		{
 			TMap<FName, FString> TempMap;
-			Record << NAMED_ITEM("RootMetaDataMap", TempMap);
+			Record << SA_VALUE(TEXT("RootMetaDataMap"), TempMap);
 
 			const bool bLoadFromLinker = (NULL != UnderlyingArchive.GetLinker());
 			if (bLoadFromLinker && HasAnyFlags(RF_LoadCompleted))
diff --git a/Engine/Source/Runtime/CoreUObject/Private/UObject/Obj.cpp b/Engine/Source/Runtime/CoreUObject/Private/UObject/Obj.cpp
--- a/Engine/Source/Runtime/CoreUObject/Private/UObject/Obj.cpp
+++ b/Engine/Source/Runtime/CoreUObject/Private/UObject/Obj.cpp
@@ -1282,14 +1282,14 @@ void UObject::Serialize(FStructuredArchive::FRecord Record)
 		// Special info.
 		if ((!UnderlyingArchive.IsLoading() && !UnderlyingArchive.IsSaving() && !UnderlyingArchive.IsObjectReferenceCollector()))
 		{
-			Record << NAMED_FIELD(LoadName);
+			Record << SA_VALUE(TEXT("LoadName"), LoadName);
 			if (!UnderlyingArchive.IsIgnoringOuterRef())
 			{
-				Record << NAMED_FIELD(LoadOuter);
+				Record << SA_VALUE(TEXT("LoadOuter"), LoadOuter);
 			}
 			if (!UnderlyingArchive.IsIgnoringClassRef())
 			{
-				Record << NAMED_FIELD(ObjClass);
+				Record << SA_VALUE(TEXT("ObjClass"), ObjClass);
 			}
 		}
 		// Special support for supporting undo/redo of renaming and changing Archetype.
@@ -1299,7 +1299,7 @@ void UObject::Serialize(FStructuredArchive::FRecord Record)
 			{
 				if (UnderlyingArchive.IsLoading())
 				{
-					Record << NAMED_FIELD(LoadName) << NAMED_FIELD(LoadOuter);
+					Record << SA_VALUE(TEXT("LoadName"), LoadName) << SA_VALUE(TEXT("LoadOuter"), LoadOuter);
 
 					// If the name we loaded is different from the current one,
 					// unhash the object, change the name and hash it again.
@@ -1312,7 +1312,7 @@ void UObject::Serialize(FStructuredArchive::FRecord Record)
 				}
 				else
 				{
-					Record << NAMED_FIELD(LoadName) << NAMED_FIELD(LoadOuter);
+					Record << SA_VALUE(TEXT("LoadName"), LoadName) << SA_VALUE(TEXT("LoadOuter"), LoadOuter);
 				}
 			}
 		}
@@ -1321,7 +1321,7 @@ void UObject::Serialize(FStructuredArchive::FRecord Record)
 		// Handle derived UClass objects (exact UClass objects are native only and shouldn't be touched)
 		if (ObjClass != UClass::StaticClass())
 		{
-			SerializeScriptProperties(Record.EnterField(FIELD_NAME_TEXT("Properties")));
+			SerializeScriptProperties(Record.EnterField(SA_FIELD_NAME(TEXT("Properties"))));
 		}
 
 		// Keep track of pending kill
@@ -1330,7 +1330,7 @@ void UObject::Serialize(FStructuredArchive::FRecord Record)
 			bool WasKill = IsPendingKill();
 			if (UnderlyingArchive.IsLoading())
 			{
-				Record << NAMED_FIELD(WasKill);
+				Record << SA_VALUE(TEXT("WasKill"), WasKill);
 				if (WasKill)
 				{
 					MarkPendingKill();
@@ -1342,7 +1342,7 @@ void UObject::Serialize(FStructuredArchive::FRecord Record)
 			}
 			else if (UnderlyingArchive.IsSaving())
 			{
-				Record << NAMED_FIELD(WasKill);
+				Record << SA_VALUE(TEXT("WasKill"), WasKill);
 			}
 		}
 
diff --git a/Engine/Source/Runtime/CoreUObject/Private/UObject/ObjectRedirector.cpp b/Engine/Source/Runtime/CoreUObject/Private/UObject/ObjectRedirector.cpp
--- a/Engine/Source/Runtime/CoreUObject/Private/UObject/ObjectRedirector.cpp
+++ b/Engine/Source/Runtime/CoreUObject/Private/UObject/ObjectRedirector.cpp
@@ -39,7 +39,7 @@ void UObjectRedirector::Serialize( FStructuredArchive::FRecord Record )
 {
 	Super::Serialize(Record);
 
-	Record << NAMED_FIELD(DestinationObject);
+	Record << SA_VALUE(TEXT("DestinationObject"), DestinationObject);
 }
 
 bool UObjectRedirector::NeedsLoadForEditorGame() const
diff --git a/Engine/Source/Runtime/CoreUObject/Private/UObject/ObjectResource.cpp b/Engine/Source/Runtime/CoreUObject/Private/UObject/ObjectResource.cpp
--- a/Engine/Source/Runtime/CoreUObject/Private/UObject/ObjectResource.cpp
+++ b/Engine/Source/Runtime/CoreUObject/Private/UObject/ObjectResource.cpp
@@ -109,19 +109,19 @@ void operator<<(FStructuredArchive::FSlot Slot, FObjectExport& E)
 	FArchive& BaseArchive = Slot.GetUnderlyingArchive();
 	FStructuredArchive::FRecord Record = Slot.EnterRecord();
 
-	Record << NAMED_ITEM("ClassIndex", E.ClassIndex);
-	Record << NAMED_ITEM("SuperIndex", E.SuperIndex);
+	Record << SA_VALUE(TEXT("ClassIndex"), E.ClassIndex);
+	Record << SA_VALUE(TEXT("SuperIndex"), E.SuperIndex);
 
 	if (BaseArchive.UE4Ver() >= VER_UE4_TemplateIndex_IN_COOKED_EXPORTS)
 	{
-		Record << NAMED_ITEM("TemplateIndex", E.TemplateIndex);
+		Record << SA_VALUE(TEXT("TemplateIndex"), E.TemplateIndex);
 	}
 
-	Record << NAMED_ITEM("OuterIndex", E.OuterIndex);
-	Record << NAMED_ITEM("ObjectName", E.ObjectName);
+	Record << SA_VALUE(TEXT("OuterIndex"), E.OuterIndex);
+	Record << SA_VALUE(TEXT("ObjectName"), E.ObjectName);
 
 	uint32 Save = E.ObjectFlags & RF_Load;
-	Record << NAMED_ITEM("ObjectFlags", Save);
+	Record << SA_VALUE(TEXT("ObjectFlags"), Save);
 
 	if (BaseArchive.IsLoading())
 	{
@@ -131,43 +131,43 @@ void operator<<(FStructuredArchive::FSlot Slot, FObjectExport& E)
 	if (BaseArchive.UE4Ver() < VER_UE4_64BIT_EXPORTMAP_SERIALSIZES)
 	{
 		int32 SerialSize = E.SerialSize;
-		Record << NAMED_FIELD(SerialSize);
+		Record << SA_VALUE(TEXT("SerialSize"), SerialSize);
 		E.SerialSize = (int64)SerialSize;
 
 		int32 SerialOffset = E.SerialOffset;
-		Record << NAMED_FIELD(SerialOffset);
+		Record << SA_VALUE(TEXT("SerialOffset"), SerialOffset);
 		E.SerialOffset = SerialOffset;
 	}
 	else
 	{
-		Record << NAMED_ITEM("SerialSize", E.SerialSize);
-		Record << NAMED_ITEM("SerialOffset", E.SerialOffset);
+		Record << SA_VALUE(TEXT("SerialSize"), E.SerialSize);
+		Record << SA_VALUE(TEXT("SerialOffset"), E.SerialOffset);
 	}
 
-	Record << NAMED_ITEM("bForcedExport", E.bForcedExport);
-	Record << NAMED_ITEM("bNotForClient", E.bNotForClient);
-	Record << NAMED_ITEM("bNotForServer", E.bNotForServer);
+	Record << SA_VALUE(TEXT("bForcedExport"), E.bForcedExport);
+	Record << SA_VALUE(TEXT("bNotForClient"), E.bNotForClient);
+	Record << SA_VALUE(TEXT("bNotForServer"), E.bNotForServer);
 
-	Record << NAMED_ITEM("PackageGuid", E.PackageGuid);
-	Record << NAMED_ITEM("PackageFlags", E.PackageFlags);
+	Record << SA_VALUE(TEXT("PackageGuid"), E.PackageGuid);
+	Record << SA_VALUE(TEXT("PackageFlags"), E.PackageFlags);
 
 	if (BaseArchive.UE4Ver() >= VER_UE4_LOAD_FOR_EDITOR_GAME)
 	{
-		Record << NAMED_ITEM("bNotAlwaysLoadedForEditorGame", E.bNotAlwaysLoadedForEditorGame);
+		Record << SA_VALUE(TEXT("bNotAlwaysLoadedForEditorGame"), E.bNotAlwaysLoadedForEditorGame);
 	}
 
 	if (BaseArchive.UE4Ver() >= VER_UE4_COOKED_ASSETS_IN_EDITOR_SUPPORT)
 	{
-		Record << NAMED_ITEM("bIsAsset", E.bIsAsset);
+		Record << SA_VALUE(TEXT("bIsAsset"), E.bIsAsset);
 	}
 
 	if (BaseArchive.UE4Ver() >= VER_UE4_PRELOAD_DEPENDENCIES_IN_COOKED_EXPORTS)
 	{
-		Record << NAMED_ITEM("FirstExportDependency", E.FirstExportDependency);
-		Record << NAMED_ITEM("SerializationBeforeSerializationDependencies", E.SerializationBeforeSerializationDependencies);
-		Record << NAMED_ITEM("CreateBeforeSerializationDependencies", E.CreateBeforeSerializationDependencies);
-		Record << NAMED_ITEM("SerializationBeforeCreateDependencies", E.SerializationBeforeCreateDependencies);
-		Record << NAMED_ITEM("CreateBeforeCreateDependencies", E.CreateBeforeCreateDependencies);
+		Record << SA_VALUE(TEXT("FirstExportDependency"), E.FirstExportDependency);
+		Record << SA_VALUE(TEXT("SerializationBeforeSerializationDependencies"), E.SerializationBeforeSerializationDependencies);
+		Record << SA_VALUE(TEXT("CreateBeforeSerializationDependencies"), E.CreateBeforeSerializationDependencies);
+		Record << SA_VALUE(TEXT("SerializationBeforeCreateDependencies"), E.SerializationBeforeCreateDependencies);
+		Record << SA_VALUE(TEXT("CreateBeforeCreateDependencies"), E.CreateBeforeCreateDependencies);
 	}	
 }
 
@@ -219,10 +219,10 @@ void operator<<(FStructuredArchive::FSlot Slot, FObjectImport& I)
 {
 	FStructuredArchive::FRecord Record = Slot.EnterRecord();
 
-	Record << NAMED_ITEM("ClassPackage", I.ClassPackage);
-	Record << NAMED_ITEM("ClassName", I.ClassName);
-	Record << NAMED_ITEM("OuterIndex", I.OuterIndex);
-	Record << NAMED_ITEM("ObjectName", I.ObjectName);
+	Record << SA_VALUE(TEXT("ClassPackage"), I.ClassPackage);
+	Record << SA_VALUE(TEXT("ClassName"), I.ClassName);
+	Record << SA_VALUE(TEXT("OuterIndex"), I.OuterIndex);
+	Record << SA_VALUE(TEXT("ObjectName"), I.ObjectName);
 
 	if (Slot.GetUnderlyingArchive().IsLoading())
 	{
diff --git a/Engine/Source/Runtime/CoreUObject/Private/UObject/PackageFileSummary.cpp b/Engine/Source/Runtime/CoreUObject/Private/UObject/PackageFileSummary.cpp
--- a/Engine/Source/Runtime/CoreUObject/Private/UObject/PackageFileSummary.cpp
+++ b/Engine/Source/Runtime/CoreUObject/Private/UObject/PackageFileSummary.cpp
@@ -50,7 +50,7 @@ void operator<<(FStructuredArchive::FSlot Slot, FPackageFileSummary& Sum)
 
 	if (bCanStartSerializing)
 	{
-		Record << NAMED_ITEM("Tag", Sum.Tag);
+		Record << SA_VALUE(TEXT("Tag"), Sum.Tag);
 	}
 	// only keep loading if we match the magic
 	if (Sum.Tag == PACKAGE_FILE_TAG || Sum.Tag == PACKAGE_FILE_TAG_SWAPPED)
@@ -86,7 +86,7 @@ void operator<<(FStructuredArchive::FSlot Slot, FPackageFileSummary& Sum)
 		*/
 		const int32 CurrentLegacyFileVersion = -7;
 		int32 LegacyFileVersion = CurrentLegacyFileVersion;
-		Record << NAMED_ITEM("LegacyFileVersion", LegacyFileVersion);
+		Record << SA_VALUE(TEXT("LegacyFileVersion"), LegacyFileVersion);
 
 		if (BaseArchive.IsLoading())
 		{
@@ -104,14 +104,14 @@ void operator<<(FStructuredArchive::FSlot Slot, FPackageFileSummary& Sum)
 				if (LegacyFileVersion != -4)
 				{
 					int32 LegacyUE3Version = 0;
-					Record << NAMED_ITEM("LegacyUE3Version", LegacyUE3Version);
+					Record << SA_VALUE(TEXT("LegacyUE3Version"), LegacyUE3Version);
 				}
-				Record << NAMED_ITEM("FileVersionUE4", Sum.FileVersionUE4);
-				Record << NAMED_ITEM("FileVersionLicenseeUE4", Sum.FileVersionLicenseeUE4);
+				Record << SA_VALUE(TEXT("FileVersionUE4"), Sum.FileVersionUE4);
+				Record << SA_VALUE(TEXT("FileVersionLicenseeUE4"), Sum.FileVersionLicenseeUE4);
 
 				if (LegacyFileVersion <= -2)
 				{
-					Sum.CustomVersionContainer.Serialize(Record.EnterField(FIELD_NAME_TEXT("CustomVersions")), GetCustomVersionFormatForArchive(LegacyFileVersion));
+					Sum.CustomVersionContainer.Serialize(Record.EnterField(SA_FIELD_NAME(TEXT("CustomVersions"))), GetCustomVersionFormatForArchive(LegacyFileVersion));
 				}
 
 				if (!Sum.FileVersionUE4 && !Sum.FileVersionLicenseeUE4)
@@ -143,28 +143,28 @@ void operator<<(FStructuredArchive::FSlot Slot, FPackageFileSummary& Sum)
 			if (Sum.bUnversioned)
 			{
 				int32 Zero = 0;
-				Record << NAMED_ITEM("LegacyUE3version", Zero); // LegacyUE3version
-				Record << NAMED_ITEM("FileVersionUE4", Zero); // VersionUE4
-				Record << NAMED_ITEM("FileVersionLicenseeUE4", Zero); // VersionLicenseeUE4
+				Record << SA_VALUE(TEXT("LegacyUE3version"), Zero); // LegacyUE3version
+				Record << SA_VALUE(TEXT("FileVersionUE4"), Zero); // VersionUE4
+				Record << SA_VALUE(TEXT("FileVersionLicenseeUE4"), Zero); // VersionLicenseeUE4
 
 				FCustomVersionContainer NoCustomVersions;
-				NoCustomVersions.Serialize(Record.EnterField(FIELD_NAME_TEXT("CustomVersions")));
+				NoCustomVersions.Serialize(Record.EnterField(SA_FIELD_NAME(TEXT("CustomVersions"))));
 			}
 			else
 			{
 				// Must write out the last UE3 engine version, so that older versions identify it as new
 				int32 LegacyUE3Version = 864;
-				Record << NAMED_ITEM("LegacyUE3Version", LegacyUE3Version);
-				Record << NAMED_ITEM("FileVersionUE4", Sum.FileVersionUE4);
-				Record << NAMED_ITEM("FileVersionLicenseeUE4", Sum.FileVersionLicenseeUE4);
+				Record << SA_VALUE(TEXT("LegacyUE3Version"), LegacyUE3Version);
+				Record << SA_VALUE(TEXT("FileVersionUE4"), Sum.FileVersionUE4);
+				Record << SA_VALUE(TEXT("FileVersionLicenseeUE4"), Sum.FileVersionLicenseeUE4);
 
 				// Serialise custom version map.
-				Sum.CustomVersionContainer.Serialize(Record.EnterField(FIELD_NAME_TEXT("CustomVersions")));
+				Sum.CustomVersionContainer.Serialize(Record.EnterField(SA_FIELD_NAME(TEXT("CustomVersions"))));
 			}
 		}
-		Record << NAMED_ITEM("TotalHeaderSize", Sum.TotalHeaderSize);
-		Record << NAMED_ITEM("FolderName", Sum.FolderName);
-		Record << NAMED_ITEM("PackageFlags", Sum.PackageFlags);
+		Record << SA_VALUE(TEXT("TotalHeaderSize"), Sum.TotalHeaderSize);
+		Record << SA_VALUE(TEXT("FolderName"), Sum.FolderName);
+		Record << SA_VALUE(TEXT("PackageFlags"), Sum.PackageFlags);
 
 #if WITH_EDITOR
 		if (BaseArchive.IsLoading())
@@ -178,21 +178,21 @@ void operator<<(FStructuredArchive::FSlot Slot, FPackageFileSummary& Sum)
 		{
 			BaseArchive.SetFilterEditorOnly(true);
 		}
-		Record << NAMED_ITEM("NameCount", Sum.NameCount) << NAMED_ITEM("NameOffset", Sum.NameOffset);
+		Record << SA_VALUE(TEXT("NameCount"), Sum.NameCount) << SA_VALUE(TEXT("NameOffset"), Sum.NameOffset);
 		if (!BaseArchive.IsFilterEditorOnly())
 		{
 			if (BaseArchive.IsSaving() || Sum.FileVersionUE4 >= VER_UE4_ADDED_PACKAGE_SUMMARY_LOCALIZATION_ID)
 			{
-				Record << NAMED_ITEM("LocalizationId", Sum.LocalizationId);
+				Record << SA_VALUE(TEXT("LocalizationId"), Sum.LocalizationId);
 			}
 		}
 		if (Sum.FileVersionUE4 >= VER_UE4_SERIALIZE_TEXT_IN_PACKAGES)
 		{
-			Record << NAMED_ITEM("GatherableTextDataCount", Sum.GatherableTextDataCount) << NAMED_ITEM("GatherableTextDataOffset", Sum.GatherableTextDataOffset);
+			Record << SA_VALUE(TEXT("GatherableTextDataCount"), Sum.GatherableTextDataCount) << SA_VALUE(TEXT("GatherableTextDataOffset"), Sum.GatherableTextDataOffset);
 		}
-		Record << NAMED_ITEM("ExportCount", Sum.ExportCount) << NAMED_ITEM("ExportOffset", Sum.ExportOffset);
-		Record << NAMED_ITEM("ImportCount", Sum.ImportCount) << NAMED_ITEM("ImportOffset", Sum.ImportOffset);
-		Record << NAMED_ITEM("DependsOffset", Sum.DependsOffset);
+		Record << SA_VALUE(TEXT("ExportCount"), Sum.ExportCount) << SA_VALUE(TEXT("ExportOffset"), Sum.ExportOffset);
+		Record << SA_VALUE(TEXT("ImportCount"), Sum.ImportCount) << SA_VALUE(TEXT("ImportOffset"), Sum.ImportOffset);
+		Record << SA_VALUE(TEXT("DependsOffset"), Sum.DependsOffset);
 
 		if (BaseArchive.IsLoading() && (Sum.FileVersionUE4 < VER_UE4_OLDEST_LOADABLE_PACKAGE || Sum.FileVersionUE4 > GPackageFileUE4Version))
 		{
@@ -201,25 +201,25 @@ void operator<<(FStructuredArchive::FSlot Slot, FPackageFileSummary& Sum)
 
 		if (BaseArchive.IsSaving() || Sum.FileVersionUE4 >= VER_UE4_ADD_STRING_ASSET_REFERENCES_MAP)
 		{
-			Record << NAMED_ITEM("SoftPackageReferencesCount", Sum.SoftPackageReferencesCount) << NAMED_ITEM("SoftPackageReferencesOffset", Sum.SoftPackageReferencesOffset);
+			Record << SA_VALUE(TEXT("SoftPackageReferencesCount"), Sum.SoftPackageReferencesCount) << SA_VALUE(TEXT("SoftPackageReferencesOffset"), Sum.SoftPackageReferencesOffset);
 		}
 
 		if (BaseArchive.IsSaving() || Sum.FileVersionUE4 >= VER_UE4_ADDED_SEARCHABLE_NAMES)
 		{
-			Record << NAMED_ITEM("SearchableNamesOffset", Sum.SearchableNamesOffset);
+			Record << SA_VALUE(TEXT("SearchableNamesOffset"), Sum.SearchableNamesOffset);
 		}
 
-		Record << NAMED_ITEM("ThumbnailTableOffset", Sum.ThumbnailTableOffset);
+		Record << SA_VALUE(TEXT("ThumbnailTableOffset"), Sum.ThumbnailTableOffset);
 
 		int32 GenerationCount = Sum.Generations.Num();
-		Record << NAMED_ITEM("Guid", Sum.Guid) << NAMED_ITEM("GenerationCount", GenerationCount);
+		Record << SA_VALUE(TEXT("Guid"), Sum.Guid) << SA_VALUE(TEXT("GenerationCount"), GenerationCount);
 		if (BaseArchive.IsLoading() && GenerationCount > 0)
 		{
 			Sum.Generations.Reset(GenerationCount);
 			Sum.Generations.AddZeroed(GenerationCount);
 		}
 
-		FStructuredArchive::FStream GenerationsStream = Record.EnterStream(FIELD_NAME_TEXT("Generations"));
+		FStructuredArchive::FStream GenerationsStream = Record.EnterStream(SA_FIELD_NAME(TEXT("Generations")));
 		for (int32 i = 0; i<GenerationCount; i++)
 		{
 			Sum.Generations[i].Serialize(GenerationsStream.EnterElement(), Sum);
@@ -234,17 +234,17 @@ void operator<<(FStructuredArchive::FSlot Slot, FPackageFileSummary& Sum)
 			if (BaseArchive.IsCooking() || (BaseArchive.IsSaving() && !FEngineVersion::Current().HasChangelist()))
 			{
 				FEngineVersion EmptyEngineVersion;
-				Record << NAMED_ITEM("SavedByEngineVersion", EmptyEngineVersion);
+				Record << SA_VALUE(TEXT("SavedByEngineVersion"), EmptyEngineVersion);
 			}
 			else
 			{
-				Record << NAMED_ITEM("SavedByEngineVersion", Sum.SavedByEngineVersion);
+				Record << SA_VALUE(TEXT("SavedByEngineVersion"), Sum.SavedByEngineVersion);
 			}
 		}
 		else
 		{
 			int32 EngineChangelist = 0;
-			Record << NAMED_ITEM("EngineChangelist", EngineChangelist);
+			Record << SA_VALUE(TEXT("EngineChangelist"), EngineChangelist);
 
 			if (BaseArchive.IsLoading() && EngineChangelist != 0)
 			{
@@ -257,11 +257,11 @@ void operator<<(FStructuredArchive::FSlot Slot, FPackageFileSummary& Sum)
 			if (BaseArchive.IsCooking() || (BaseArchive.IsSaving() && !FEngineVersion::Current().HasChangelist()))
 			{
 				FEngineVersion EmptyEngineVersion;
-				Record << NAMED_ITEM("CompatibleWithEngineVersion", EmptyEngineVersion);
+				Record << SA_VALUE(TEXT("CompatibleWithEngineVersion"), EmptyEngineVersion);
 			}
 			else
 			{
-				Record << NAMED_ITEM("CompatibleWithEngineVersion", Sum.CompatibleWithEngineVersion);
+				Record << SA_VALUE(TEXT("CompatibleWithEngineVersion"), Sum.CompatibleWithEngineVersion);
 			}
 		}
 		else
@@ -272,7 +272,7 @@ void operator<<(FStructuredArchive::FSlot Slot, FPackageFileSummary& Sum)
 			}
 		}
 
-		Record << NAMED_ITEM("CompressionFlags", Sum.CompressionFlags);
+		Record << SA_VALUE(TEXT("CompressionFlags"), Sum.CompressionFlags);
 		if (!FCompression::VerifyCompressionFlagsValid(Sum.CompressionFlags))
 		{
 			UE_LOG(LogLinker, Warning, TEXT("Failed to read package file summary, the file \"%s\" has invalid compression flags (%d)."), *BaseArchive.GetArchiveName(), Sum.CompressionFlags);
@@ -281,7 +281,7 @@ void operator<<(FStructuredArchive::FSlot Slot, FPackageFileSummary& Sum)
 		}
 
 		TArray<FCompressedChunk> CompressedChunks;
-		Record << NAMED_ITEM("CompressedChunks", CompressedChunks);
+		Record << SA_VALUE(TEXT("CompressedChunks"), CompressedChunks);
 
 		if (CompressedChunks.Num())
 		{
@@ -291,32 +291,32 @@ void operator<<(FStructuredArchive::FSlot Slot, FPackageFileSummary& Sum)
 			return; // we can't safely load more than this because we just changed the version to something it is not.
 		}
 
-		Record << NAMED_ITEM("PackageSource", Sum.PackageSource);
+		Record << SA_VALUE(TEXT("PackageSource"), Sum.PackageSource);
 
 		// No longer used: List of additional packages that are needed to be cooked for this package (ie streaming levels)
 		// Keeping the serialization code for backwards compatibility without bumping the package version
 		TArray<FString>	AdditionalPackagesToCook;
-		Record << NAMED_ITEM("AdditionalPackagesToCook", AdditionalPackagesToCook);
+		Record << SA_VALUE(TEXT("AdditionalPackagesToCook"), AdditionalPackagesToCook);
 
 		if (LegacyFileVersion > -7)
 		{
 			int32 NumTextureAllocations = 0;
-			Record << NAMED_ITEM("NumTextureAllocations", NumTextureAllocations);
+			Record << SA_VALUE(TEXT("NumTextureAllocations"), NumTextureAllocations);
 			// We haven't used texture allocation info for ages and it's no longer supported anyway
 			check(NumTextureAllocations == 0);
 		}
 
-		Record << NAMED_ITEM("AssetRegistryDataOffset", Sum.AssetRegistryDataOffset);
-		Record << NAMED_ITEM("BulkDataStartOffset", Sum.BulkDataStartOffset);
+		Record << SA_VALUE(TEXT("AssetRegistryDataOffset"), Sum.AssetRegistryDataOffset);
+		Record << SA_VALUE(TEXT("BulkDataStartOffset"), Sum.BulkDataStartOffset);
 
 		if (Sum.GetFileVersionUE4() >= VER_UE4_WORLD_LEVEL_INFO)
 		{
-			Record << NAMED_ITEM("WorldTileInfoDataOffset", Sum.WorldTileInfoDataOffset);
+			Record << SA_VALUE(TEXT("WorldTileInfoDataOffset"), Sum.WorldTileInfoDataOffset);
 		}
 
 		if (Sum.GetFileVersionUE4() >= VER_UE4_CHANGED_CHUNKID_TO_BE_AN_ARRAY_OF_CHUNKIDS)
 		{
-			Record << NAMED_ITEM("ChunkIDs", Sum.ChunkIDs);
+			Record << SA_VALUE(TEXT("ChunkIDs"), Sum.ChunkIDs);
 		}
 		else if (Sum.GetFileVersionUE4() >= VER_UE4_ADDED_CHUNKID_TO_ASSETDATA_AND_UPACKAGE)
 		{
@@ -324,7 +324,7 @@ void operator<<(FStructuredArchive::FSlot Slot, FPackageFileSummary& Sum)
 			if (BaseArchive.IsLoading())
 			{
 				int ChunkID = -1;
-				Record << NAMED_ITEM("ChunkID", ChunkID);
+				Record << SA_VALUE(TEXT("ChunkID"), ChunkID);
 
 				// don't load <0 entries since an empty array represents the same thing now
 				if (ChunkID >= 0)
@@ -335,7 +335,7 @@ void operator<<(FStructuredArchive::FSlot Slot, FPackageFileSummary& Sum)
 		}
 		if (BaseArchive.IsSaving() || Sum.FileVersionUE4 >= VER_UE4_PRELOAD_DEPENDENCIES_IN_COOKED_EXPORTS)
 		{
-			Record << NAMED_ITEM("PreloadDependencyCount", Sum.PreloadDependencyCount) << NAMED_ITEM("PreloadDependencyOffset", Sum.PreloadDependencyOffset);
+			Record << SA_VALUE(TEXT("PreloadDependencyCount"), Sum.PreloadDependencyCount) << SA_VALUE(TEXT("PreloadDependencyOffset"), Sum.PreloadDependencyOffset);
 		}
 		else
 		{
diff --git a/Engine/Source/Runtime/CoreUObject/Private/UObject/PropertyArray.cpp b/Engine/Source/Runtime/CoreUObject/Private/UObject/PropertyArray.cpp
--- a/Engine/Source/Runtime/CoreUObject/Private/UObject/PropertyArray.cpp
+++ b/Engine/Source/Runtime/CoreUObject/Private/UObject/PropertyArray.cpp
@@ -69,7 +69,7 @@ void UArrayProperty::SerializeItem(FStructuredArchive::FSlot Slot, void* Value,
 
 	FScriptArrayHelper ArrayHelper(this, Value);
 	int32		n		= ArrayHelper.Num();
-	Record << NAMED_ITEM("Count", n);
+	Record << SA_VALUE(TEXT("Count"), n);
 	if( UnderlyingArchive.IsLoading() )
 	{
 		// If using a custom property list, don't empty the array on load. Not all indices may have been serialized, so we need to preserve existing values at those slots.
@@ -98,11 +98,11 @@ void UArrayProperty::SerializeItem(FStructuredArchive::FSlot Slot, void* Value,
 	{
 		if (UnderlyingArchive.IsSaving())
 		{
-			Record << NAMED_FIELD(InnerTag);
+			Record << SA_VALUE(TEXT("InnerTag"), InnerTag);
 		}
 		else if (UnderlyingArchive.IsLoading())
 		{
-			Record << NAMED_FIELD(InnerTag);
+			Record << SA_VALUE(TEXT("InnerTag"), InnerTag);
 
 			auto CanSerializeFromStructWithDifferentName = [](const FPropertyTag& PropertyTag, const UStructProperty* StructProperty)
 			{
@@ -155,7 +155,7 @@ void UArrayProperty::SerializeItem(FStructuredArchive::FSlot Slot, void* Value,
 	// need to know how much data this call to SerializeItem consumes, so mark where we are
 	int64 DataOffset = UnderlyingArchive.Tell();
 
-	FStructuredArchive::FStream ValueStream = Record.EnterField(FIELD_NAME_TEXT("Values")).EnterStream();
+	FStructuredArchive::FStream ValueStream = Record.EnterField(SA_FIELD_NAME(TEXT("Values"))).EnterStream();
 
 	// If we're using a custom property list, first serialize any explicit indices
 	int32 i = 0;
@@ -603,7 +603,7 @@ EConvertFromTypeResult UArrayProperty::ConvertFromType(const FPropertyTag& Tag,
 		FStructuredArchive::FRecord ArrayRecord = Slot.EnterRecord();
 
 		int32 ElementCount = 0;
-		ArrayRecord << NAMED_ITEM("Count", ElementCount);
+		ArrayRecord << SA_VALUE(TEXT("Count"), ElementCount);
 
 		FScriptArrayHelper ScriptArrayHelper(this, ArrayPropertyData);
 		ScriptArrayHelper.EmptyAndAddValues(ElementCount);
@@ -615,7 +615,7 @@ EConvertFromTypeResult UArrayProperty::ConvertFromType(const FPropertyTag& Tag,
 			InnerPropertyTag.Type = Tag.InnerType;
 			InnerPropertyTag.ArrayIndex = 0;
 
-			FStructuredArchive::FStream ValueStream = ArrayRecord.EnterField(FIELD_NAME_TEXT("Value")).EnterStream();
+			FStructuredArchive::FStream ValueStream = ArrayRecord.EnterField(SA_FIELD_NAME(TEXT("Value"))).EnterStream();
 
 			if (Inner->ConvertFromType(InnerPropertyTag, ValueStream.EnterElement(), ScriptArrayHelper.GetRawPtr(0), DefaultsStruct) == EConvertFromTypeResult::Converted)
 			{
diff --git a/Engine/Source/Runtime/CoreUObject/Private/UObject/PropertyMap.cpp b/Engine/Source/Runtime/CoreUObject/Private/UObject/PropertyMap.cpp
--- a/Engine/Source/Runtime/CoreUObject/Private/UObject/PropertyMap.cpp
+++ b/Engine/Source/Runtime/CoreUObject/Private/UObject/PropertyMap.cpp
@@ -284,7 +284,7 @@ void UMapProperty::SerializeItem(FStructuredArchive::FSlot Slot, void* Value, co
 
 		// Delete any explicitly-removed keys
 		int32 NumKeysToRemove = 0;
-		FStructuredArchive::FArray KeysToRemoveArray = Record.EnterArray(FIELD_NAME_TEXT("KeysToRemove"), NumKeysToRemove);
+		FStructuredArchive::FArray KeysToRemoveArray = Record.EnterArray(SA_FIELD_NAME(TEXT("KeysToRemove")), NumKeysToRemove);
 		if (NumKeysToRemove)
 		{
 			TempKeyStorage = (uint8*)FMemory::Malloc(MapLayout.SetLayout.Size);
@@ -305,7 +305,7 @@ void UMapProperty::SerializeItem(FStructuredArchive::FSlot Slot, void* Value, co
 		}
 
 		int32 NumEntries = 0;
-		FStructuredArchive::FArray EntriesArray = Record.EnterArray(FIELD_NAME_TEXT("Entries"), NumEntries);
+		FStructuredArchive::FArray EntriesArray = Record.EnterArray(SA_FIELD_NAME(TEXT("Entries")), NumEntries);
 
 		// Allocate temporary key space if we haven't allocated it already above
 		if (NumEntries != 0 && !TempKeyStorage)
@@ -322,7 +322,7 @@ void UMapProperty::SerializeItem(FStructuredArchive::FSlot Slot, void* Value, co
 			// Read key into temporary storage
 			{
 				FSerializedPropertyScope SerializedProperty(UnderlyingArchive, KeyProp, this);
-				KeyProp->SerializeItem(EntryRecord.EnterField(FIELD_NAME_TEXT("Key")), TempKeyStorage);
+				KeyProp->SerializeItem(EntryRecord.EnterField(SA_FIELD_NAME(TEXT("Key"))), TempKeyStorage);
 			}
 
 			void* ValuePtr = MapHelper.FindOrAdd(TempKeyStorage);
@@ -330,7 +330,7 @@ void UMapProperty::SerializeItem(FStructuredArchive::FSlot Slot, void* Value, co
 			// Deserialize value into hash map-owned memory
 			{
 				FSerializedPropertyScope SerializedProperty(UnderlyingArchive, ValueProp, this);
-				ValueProp->SerializeItem(EntryRecord.EnterField(FIELD_NAME_TEXT("Value")), ValuePtr);
+				ValueProp->SerializeItem(EntryRecord.EnterField(SA_FIELD_NAME(TEXT("Value"))), ValuePtr);
 			}
 		}
 	}
@@ -362,7 +362,7 @@ void UMapProperty::SerializeItem(FStructuredArchive::FSlot Slot, void* Value, co
 
 		// Write out the missing keys
 		int32 MissingKeysNum = Indices.Num();
-		FStructuredArchive::FArray KeysToRemoveArray = Record.EnterArray(FIELD_NAME_TEXT("KeysToRemove"), MissingKeysNum);
+		FStructuredArchive::FArray KeysToRemoveArray = Record.EnterArray(SA_FIELD_NAME(TEXT("KeysToRemove")), MissingKeysNum);
 		{
 			FSerializedPropertyScope SerializedProperty(UnderlyingArchive, KeyProp, this);
 			for (int32 Index : Indices)
@@ -393,7 +393,7 @@ void UMapProperty::SerializeItem(FStructuredArchive::FSlot Slot, void* Value, co
 
 			// Write out differences from defaults
 			int32 Num = Indices.Num();
-			FStructuredArchive::FArray EntriesArray = Record.EnterArray(FIELD_NAME_TEXT("Entries"), Num);
+			FStructuredArchive::FArray EntriesArray = Record.EnterArray(SA_FIELD_NAME(TEXT("Entries")), Num);
 			for (int32 Index : Indices)
 			{
 				uint8* ValuePairPtr = MapHelper.GetPairPtrWithoutCheck(Index);
@@ -401,18 +401,18 @@ void UMapProperty::SerializeItem(FStructuredArchive::FSlot Slot, void* Value, co
 
 				{
 					FSerializedPropertyScope SerializedProperty(UnderlyingArchive, KeyProp, this);
-					KeyProp->SerializeItem(EntryRecord.EnterField(FIELD_NAME_TEXT("Key")), ValuePairPtr);
+					KeyProp->SerializeItem(EntryRecord.EnterField(SA_FIELD_NAME(TEXT("Key"))), ValuePairPtr);
 				}
 				{
 					FSerializedPropertyScope SerializedProperty(UnderlyingArchive, ValueProp, this);
-					ValueProp->SerializeItem(EntryRecord.EnterField(FIELD_NAME_TEXT("Value")), ValuePairPtr + MapLayout.ValueOffset);
+					ValueProp->SerializeItem(EntryRecord.EnterField(SA_FIELD_NAME(TEXT("Value"))), ValuePairPtr + MapLayout.ValueOffset);
 				}
 			}
 		}
 		else
 		{
 			int32 Num = MapHelper.Num();
-			FStructuredArchive::FArray EntriesArray = Record.EnterArray(FIELD_NAME_TEXT("Entries"), Num);
+			FStructuredArchive::FArray EntriesArray = Record.EnterArray(SA_FIELD_NAME(TEXT("Entries")), Num);
 
 			for (int32 Index = 0; Num; ++Index)
 			{
@@ -424,11 +424,11 @@ void UMapProperty::SerializeItem(FStructuredArchive::FSlot Slot, void* Value, co
 
 					{
 						FSerializedPropertyScope SerializedProperty(UnderlyingArchive, KeyProp, this);
-						KeyProp->SerializeItem(EntryRecord.EnterField(FIELD_NAME_TEXT("Key")), ValuePairPtr);
+						KeyProp->SerializeItem(EntryRecord.EnterField(SA_FIELD_NAME(TEXT("Key"))), ValuePairPtr);
 					}
 					{
 						FSerializedPropertyScope SerializedProperty(UnderlyingArchive, ValueProp, this);
-						ValueProp->SerializeItem(EntryRecord.EnterField(FIELD_NAME_TEXT("Value")), ValuePairPtr + MapLayout.ValueOffset);
+						ValueProp->SerializeItem(EntryRecord.EnterField(SA_FIELD_NAME(TEXT("Value"))), ValuePairPtr + MapLayout.ValueOffset);
 					}
 
 					--Num;
@@ -970,7 +970,7 @@ EConvertFromTypeResult UMapProperty::ConvertFromType(const FPropertyTag& Tag, FS
 			// instance that was being written. Presumably we were constructed from our defaults and must now remove 
 			// any of the elements that were not present when we saved this Map:
 			int32 NumKeysToRemove = 0;
-			FStructuredArchive::FArray KeysToRemoveArray = ValueRecord.EnterArray(FIELD_NAME_TEXT("KeysToRemove"), NumKeysToRemove);
+			FStructuredArchive::FArray KeysToRemoveArray = ValueRecord.EnterArray(SA_FIELD_NAME(TEXT("KeysToRemove")), NumKeysToRemove);
 
 			if( NumKeysToRemove != 0 )
 			{
@@ -1004,7 +1004,7 @@ EConvertFromTypeResult UMapProperty::ConvertFromType(const FPropertyTag& Tag, FS
 			}
 
 			int32 NumEntries = 0;
-			FStructuredArchive::FArray EntriesArray = ValueRecord.EnterArray(FIELD_NAME_TEXT("Entries"), NumEntries);
+			FStructuredArchive::FArray EntriesArray = ValueRecord.EnterArray(SA_FIELD_NAME(TEXT("Entries")), NumEntries);
 
 			if( bConversionSucceeded )
 			{
@@ -1018,7 +1018,7 @@ EConvertFromTypeResult UMapProperty::ConvertFromType(const FPropertyTag& Tag, FS
 
 					FStructuredArchive::FRecord FirstPropertyRecord = EntriesArray.EnterElement().EnterRecord();
 
-					if( SerializeOrConvert( KeyProp, KeyPropertyTag, FirstPropertyRecord.EnterField(FIELD_NAME_TEXT("Key")), TempKeyStorage, DefaultsStruct ) )
+					if( SerializeOrConvert( KeyProp, KeyPropertyTag, FirstPropertyRecord.EnterField(SA_FIELD_NAME(TEXT("Key"))), TempKeyStorage, DefaultsStruct ) )
 					{
 						// Add a new default value if the key doesn't currently exist in the map
 						bool bKeyAlreadyPresent = true;
@@ -1034,14 +1034,14 @@ EConvertFromTypeResult UMapProperty::ConvertFromType(const FPropertyTag& Tag, FS
 						KeyProp->CopyCompleteValue_InContainer(NextPairPtr, TempKeyStorage);
 
 						// Deserialize value
-						if( SerializeOrConvert( ValueProp, ValuePropertyTag, FirstPropertyRecord.EnterField(FIELD_NAME_TEXT("Value")), NextPairPtr, DefaultsStruct ) )
+						if( SerializeOrConvert( ValueProp, ValuePropertyTag, FirstPropertyRecord.EnterField(SA_FIELD_NAME(TEXT("Value"))), NextPairPtr, DefaultsStruct ) )
 						{
 							// first entry went fine, convert the rest:
 							for(int32 I = 1; I < NumEntries; ++I)
 							{
 								FStructuredArchive::FRecord PropertyRecord = EntriesArray.EnterElement().EnterRecord();
 
-								verify( SerializeOrConvert( KeyProp, KeyPropertyTag, PropertyRecord.EnterField(FIELD_NAME_TEXT("Key")), TempKeyStorage, DefaultsStruct ) );
+								verify( SerializeOrConvert( KeyProp, KeyPropertyTag, PropertyRecord.EnterField(SA_FIELD_NAME(TEXT("Key"))), TempKeyStorage, DefaultsStruct ) );
 								NextPairIndex = MapHelper.FindMapIndexWithKey(TempKeyStorage);
 								if (NextPairIndex == INDEX_NONE)
 								{
@@ -1051,7 +1051,7 @@ EConvertFromTypeResult UMapProperty::ConvertFromType(const FPropertyTag& Tag, FS
 								NextPairPtr = MapHelper.GetPairPtrWithoutCheck(NextPairIndex);
 								// This copy is unnecessary when the key was already in the map:
 								KeyProp->CopyCompleteValue_InContainer(NextPairPtr, TempKeyStorage);
-								verify( SerializeOrConvert( ValueProp, ValuePropertyTag, PropertyRecord.EnterField(FIELD_NAME_TEXT("Value")), NextPairPtr, DefaultsStruct ) );
+								verify( SerializeOrConvert( ValueProp, ValuePropertyTag, PropertyRecord.EnterField(SA_FIELD_NAME(TEXT("Value"))), NextPairPtr, DefaultsStruct ) );
 							}
 						}
 						else
diff --git a/Engine/Source/Runtime/CoreUObject/Private/UObject/PropertySet.cpp b/Engine/Source/Runtime/CoreUObject/Private/UObject/PropertySet.cpp
--- a/Engine/Source/Runtime/CoreUObject/Private/UObject/PropertySet.cpp
+++ b/Engine/Source/Runtime/CoreUObject/Private/UObject/PropertySet.cpp
@@ -255,7 +255,7 @@ void USetProperty::SerializeItem(FStructuredArchive::FSlot Slot, void* Value, co
 
 		// Delete any explicitly-removed elements
 		int32 NumElementsToRemove = 0;
-		FStructuredArchive::FArray ElementsToRemoveArray = Record.EnterArray(FIELD_NAME_TEXT("ElementsToRemove"), NumElementsToRemove);
+		FStructuredArchive::FArray ElementsToRemoveArray = Record.EnterArray(SA_FIELD_NAME(TEXT("ElementsToRemove")), NumElementsToRemove);
 
 		if (NumElementsToRemove)
 		{
@@ -278,7 +278,7 @@ void USetProperty::SerializeItem(FStructuredArchive::FSlot Slot, void* Value, co
 		}
 
 		int32 Num = 0;
-		FStructuredArchive::FArray ElementsArray = Record.EnterArray(FIELD_NAME_TEXT("Elements"), Num);
+		FStructuredArchive::FArray ElementsArray = Record.EnterArray(SA_FIELD_NAME(TEXT("Elements")), Num);
 
 		// Allocate temporary key space if we haven't allocated it already above
 		if (Num != 0 && !TempElementStorage)
@@ -335,7 +335,7 @@ void USetProperty::SerializeItem(FStructuredArchive::FSlot Slot, void* Value, co
 
 		// Write out the removed elements
 		int32 RemovedElementsNum = Indices.Num();
-		FStructuredArchive::FArray RemovedElementsArray = Record.EnterArray(FIELD_NAME_TEXT("ElementsToRemove"), RemovedElementsNum);
+		FStructuredArchive::FArray RemovedElementsArray = Record.EnterArray(SA_FIELD_NAME(TEXT("ElementsToRemove")), RemovedElementsNum);
 		
 		{
 			FSerializedPropertyScope SerializedProperty(UnderlyingArchive, ElementProp, this);
@@ -367,7 +367,7 @@ void USetProperty::SerializeItem(FStructuredArchive::FSlot Slot, void* Value, co
 
 			// Write out differences from defaults
 			int32 Num = Indices.Num();
-			FStructuredArchive::FArray ElementsArray = Record.EnterArray(FIELD_NAME_TEXT("Elements"), Num);
+			FStructuredArchive::FArray ElementsArray = Record.EnterArray(SA_FIELD_NAME(TEXT("Elements")), Num);
 
 			FSerializedPropertyScope SerializedProperty(UnderlyingArchive, ElementProp, this);
 			for (int32 Index : Indices)
@@ -380,7 +380,7 @@ void USetProperty::SerializeItem(FStructuredArchive::FSlot Slot, void* Value, co
 		else
 		{
 			int32 Num = SetHelper.Num();
-			FStructuredArchive::FArray ElementsArray = Record.EnterArray(FIELD_NAME_TEXT("Elements"), Num);
+			FStructuredArchive::FArray ElementsArray = Record.EnterArray(SA_FIELD_NAME(TEXT("Elements")), Num);
 
 			FSerializedPropertyScope SerializedProperty(UnderlyingArchive, ElementProp, this);
 			for (int32 Index = 0; Num; ++Index)
@@ -820,7 +820,7 @@ EConvertFromTypeResult USetProperty::ConvertFromType(const FPropertyTag& Tag, FS
 			// instance that was being written. Presumably we were constructed from our defaults and must now remove 
 			// any of the elements that were not present when we saved this Set:
 			int32 NumElementsToRemove = 0;
-			FStructuredArchive::FArray ElementsToRemoveArray = ValueRecord.EnterArray(FIELD_NAME_TEXT("ElementsToRemove"), NumElementsToRemove);
+			FStructuredArchive::FArray ElementsToRemoveArray = ValueRecord.EnterArray(SA_FIELD_NAME(TEXT("ElementsToRemove")), NumElementsToRemove);
 
 			if(NumElementsToRemove)
 			{
@@ -853,7 +853,7 @@ EConvertFromTypeResult USetProperty::ConvertFromType(const FPropertyTag& Tag, FS
 			}
 
 			int32 Num = 0;
-			FStructuredArchive::FArray ElementsArray = ValueRecord.EnterArray(FIELD_NAME_TEXT("Elements"), Num);
+			FStructuredArchive::FArray ElementsArray = ValueRecord.EnterArray(SA_FIELD_NAME(TEXT("Elements")), Num);
 
 			if(bConversionSucceeded)
 			{
diff --git a/Engine/Source/Runtime/CoreUObject/Private/UObject/PropertyTag.cpp b/Engine/Source/Runtime/CoreUObject/Private/UObject/PropertyTag.cpp
--- a/Engine/Source/Runtime/CoreUObject/Private/UObject/PropertyTag.cpp
+++ b/Engine/Source/Runtime/CoreUObject/Private/UObject/PropertyTag.cpp
@@ -111,13 +111,13 @@ void operator<<(FStructuredArchive::FSlot Slot, FPropertyTag& Tag)
 	checkf(!UnderlyingArchive.IsSaving() || Tag.Prop, TEXT("FPropertyTag must be constructed with a valid property when used for saving data!"));
 
 	// Name.
-	Record << NAMED_ITEM("Name", Tag.Name);
+	Record << SA_VALUE(TEXT("Name"), Tag.Name);
 	if (Tag.Name.IsNone())
 	{
 		return;
 	}
 
-	Record << NAMED_ITEM("Type", Tag.Type);
+	Record << SA_VALUE(TEXT("Type"), Tag.Type);
 	if (UnderlyingArchive.IsSaving())
 	{
 		// remember the offset of the Size variable - UStruct::SerializeTaggedProperties will update it after the
@@ -126,7 +126,7 @@ void operator<<(FStructuredArchive::FSlot Slot, FPropertyTag& Tag)
 	}
 	{
 		FArchive::FScopeSetDebugSerializationFlags S(UnderlyingArchive, DSF_IgnoreDiff);
-		Record << NAMED_ITEM("Size", Tag.Size) << NAMED_ITEM("ArrayIndex", Tag.ArrayIndex);
+		Record << SA_VALUE(TEXT("Size"), Tag.Size) << SA_VALUE(TEXT("ArrayIndex"), Tag.ArrayIndex);
 	}
 
 	if (Tag.Type.GetNumber() == 0)
@@ -136,10 +136,10 @@ void operator<<(FStructuredArchive::FSlot Slot, FPropertyTag& Tag)
 		// only need to serialize this for structs
 		if (TagType == NAME_StructProperty)
 		{
-			Record << NAMED_ITEM("StructName", Tag.StructName);
+			Record << SA_VALUE(TEXT("StructName"), Tag.StructName);
 			if (Version >= VER_UE4_STRUCT_GUID_IN_PROPERTY_TAG)
 			{
-				Record << NAMED_ITEM("StructGuid", Tag.StructGuid);
+				Record << SA_VALUE(TEXT("StructGuid"), Tag.StructGuid);
 			}
 		}
 		// only need to serialize this for bools
@@ -148,36 +148,36 @@ void operator<<(FStructuredArchive::FSlot Slot, FPropertyTag& Tag)
 			if (UnderlyingArchive.IsSaving())
 			{
 				FSerializedPropertyScope SerializedProperty(UnderlyingArchive, Tag.Prop);
-				Record << NAMED_ITEM("BoolVal", Tag.BoolVal);
+				Record << SA_VALUE(TEXT("BoolVal"), Tag.BoolVal);
 			}
 			else
 			{
-				Record << NAMED_ITEM("BoolVal", Tag.BoolVal);
+				Record << SA_VALUE(TEXT("BoolVal"), Tag.BoolVal);
 			}
 		}
 		// only need to serialize this for bytes/enums
 		else if (TagType == NAME_ByteProperty || TagType == NAME_EnumProperty)
 		{
-			Record << NAMED_ITEM("EnumName", Tag.EnumName);
+			Record << SA_VALUE(TEXT("EnumName"), Tag.EnumName);
 		}
 		// only need to serialize this for arrays
 		else if (TagType == NAME_ArrayProperty)
 		{
 			if (Version >= VAR_UE4_ARRAY_PROPERTY_INNER_TAGS)
 			{
-				Record << NAMED_ITEM("InnerType", Tag.InnerType);
+				Record << SA_VALUE(TEXT("InnerType"), Tag.InnerType);
 			}
 		}
 		else if (Version >= VER_UE4_PROPERTY_TAG_SET_MAP_SUPPORT)
 		{
 			if (TagType == NAME_SetProperty)
 			{
-				Record << NAMED_ITEM("InnerType", Tag.InnerType);
+				Record << SA_VALUE(TEXT("InnerType"), Tag.InnerType);
 			}
 			else if (TagType == NAME_MapProperty)
 			{
-				Record << NAMED_ITEM("InnerType", Tag.InnerType);
-				Record << NAMED_ITEM("ValueType", Tag.ValueType);
+				Record << SA_VALUE(TEXT("InnerType"), Tag.InnerType);
+				Record << SA_VALUE(TEXT("ValueType"), Tag.ValueType);
 			}
 		}
 	}
@@ -185,10 +185,10 @@ void operator<<(FStructuredArchive::FSlot Slot, FPropertyTag& Tag)
 	// Property tags to handle renamed blueprint properties effectively.
 	if (Version >= VER_UE4_PROPERTY_GUID_IN_PROPERTY_TAG)
 	{
-		Record << NAMED_ITEM("HasPropertyGuid", Tag.HasPropertyGuid);
+		Record << SA_VALUE(TEXT("HasPropertyGuid"), Tag.HasPropertyGuid);
 		if (Tag.HasPropertyGuid)
 		{
-			Record << NAMED_ITEM("PropertyGuid", Tag.PropertyGuid);
+			Record << SA_VALUE(TEXT("PropertyGuid"), Tag.PropertyGuid);
 		}
 	}
 }
diff --git a/Engine/Source/Runtime/CoreUObject/Private/UObject/SavePackage.cpp b/Engine/Source/Runtime/CoreUObject/Private/UObject/SavePackage.cpp
--- a/Engine/Source/Runtime/CoreUObject/Private/UObject/SavePackage.cpp
+++ b/Engine/Source/Runtime/CoreUObject/Private/UObject/SavePackage.cpp
@@ -4607,7 +4607,7 @@ FSavePackageResultStruct UPackage::Save(UPackage* InOuter, UObject* Base, EObjec
 				}
 				SlowTask.EnterProgressFrame();
 
-				FStructuredArchive::FStream Stream = StructuredArchiveRoot.EnterStream(FIELD_NAME_TEXT("GatherableTextData"));
+				FStructuredArchive::FStream Stream = StructuredArchiveRoot.EnterStream(SA_FIELD_NAME(TEXT("GatherableTextData")));
 				Linker->Summary.GatherableTextDataOffset = 0;
 				Linker->Summary.GatherableTextDataCount = 0;
 				if (!(Linker->Summary.PackageFlags & PKG_FilterEditorOnly))
@@ -5030,7 +5030,7 @@ FSavePackageResultStruct UPackage::Save(UPackage* InOuter, UObject* Base, EObjec
 				}
 				SlowTask.EnterProgressFrame();
 
-				FStructuredArchive::FStream DependsStream = StructuredArchiveRoot.EnterStream(FIELD_NAME_TEXT("DependsMap"));
+				FStructuredArchive::FStream DependsStream = StructuredArchiveRoot.EnterStream(SA_FIELD_NAME(TEXT("DependsMap")));
 				if (Linker->IsCooking())
 				{
 #if WITH_EDITOR
@@ -5073,7 +5073,7 @@ FSavePackageResultStruct UPackage::Save(UPackage* InOuter, UObject* Base, EObjec
 #if WITH_EDITOR
 						FArchive::FScopeSetDebugSerializationFlags S(*Linker, DSF_IgnoreDiff, true);
 #endif
-						FStructuredArchive::FStream SoftReferenceStream = StructuredArchiveRoot.EnterStream(FIELD_NAME_TEXT("SoftReferences"));
+						FStructuredArchive::FStream SoftReferenceStream = StructuredArchiveRoot.EnterStream(SA_FIELD_NAME(TEXT("SoftReferences")));
 						for (FName& SoftPackageName : Linker->SoftPackageReferenceList)
 						{
 							SoftReferenceStream.EnterElement() << SoftPackageName;
@@ -5082,7 +5082,7 @@ FSavePackageResultStruct UPackage::Save(UPackage* InOuter, UObject* Base, EObjec
 
 					// Save searchable names map
 					Linker->Summary.SearchableNamesOffset = Linker->Tell();
-					Linker->SerializeSearchableNamesMap(StructuredArchiveRoot.EnterField(FIELD_NAME_TEXT("SearchableNames")));
+					Linker->SerializeSearchableNamesMap(StructuredArchiveRoot.EnterField(SA_FIELD_NAME(TEXT("SearchableNames"))));
 				}
 				else
 				{
@@ -5097,13 +5097,13 @@ FSavePackageResultStruct UPackage::Save(UPackage* InOuter, UObject* Base, EObjec
 #endif // WITH_EDITOR
 
 					// Save thumbnails
-					UPackage::SaveThumbnails(InOuter, Linker.Get(), StructuredArchiveRoot.EnterField(FIELD_NAME_TEXT("Thumbnails")));
+					UPackage::SaveThumbnails(InOuter, Linker.Get(), StructuredArchiveRoot.EnterField(SA_FIELD_NAME(TEXT("Thumbnails"))));
 
 					// Save asset registry data so the editor can search for information about assets in this package
-					UPackage::SaveAssetRegistryData(InOuter, Linker.Get(), StructuredArchiveRoot.EnterField(FIELD_NAME_TEXT("AssetRegistry")));
+					UPackage::SaveAssetRegistryData(InOuter, Linker.Get(), StructuredArchiveRoot.EnterField(SA_FIELD_NAME(TEXT("AssetRegistry"))));
 
 					// Save level information used by World browser
-					UPackage::SaveWorldLevelInfo(InOuter, Linker.Get(), StructuredArchiveRoot.EnterField(FIELD_NAME_TEXT("WorldLevelInfo")));
+					UPackage::SaveWorldLevelInfo(InOuter, Linker.Get(), StructuredArchiveRoot.EnterField(SA_FIELD_NAME(TEXT("WorldLevelInfo"))));
 				}
 
 
@@ -5263,7 +5263,7 @@ FSavePackageResultStruct UPackage::Save(UPackage* InOuter, UObject* Base, EObjec
 						}
 					};
 
-					FStructuredArchive::FStream DepedenciesStream = StructuredArchiveRoot.EnterStream(FIELD_NAME_TEXT("PreloadDependencies"));
+					FStructuredArchive::FStream DepedenciesStream = StructuredArchiveRoot.EnterStream(SA_FIELD_NAME(TEXT("PreloadDependencies")));
 					TArray<UObject*> Subobjects;
 					TArray<UObject*> Deps;
 					TSet<FPackageIndex> SerializationBeforeCreateDependencies;
@@ -5514,7 +5514,7 @@ FSavePackageResultStruct UPackage::Save(UPackage* InOuter, UObject* Base, EObjec
 
 							ObjectName.Reset();
 							Export.Object->GetPathName(InOuter, ObjectName);
-							FStructuredArchive::FSlot ExportSlot = StructuredArchiveRoot.EnterField(FIELD_NAME(*ObjectName));
+							FStructuredArchive::FSlot ExportSlot = StructuredArchiveRoot.EnterField(SA_FIELD_NAME(*ObjectName));
 
 #if WITH_EDITOR
 							bool bSupportsText = UClass::IsSafeToSerializeToStructuredArchives(Export.Object->GetClass());
@@ -5855,7 +5855,7 @@ FSavePackageResultStruct UPackage::Save(UPackage* InOuter, UObject* Base, EObjec
 					}
 
 					int32 NumImports = Linker->ImportMap.Num();
-					FStructuredArchive::FStream ImportTableStream = StructuredArchiveRoot.EnterStream(FIELD_NAME_TEXT("ImportTable"));
+					FStructuredArchive::FStream ImportTableStream = StructuredArchiveRoot.EnterStream(SA_FIELD_NAME(TEXT("ImportTable")));
 
 					for (int32 i = 0; i < Linker->ImportMap.Num(); i++)
 					{
@@ -5919,7 +5919,7 @@ FSavePackageResultStruct UPackage::Save(UPackage* InOuter, UObject* Base, EObjec
 				}
 
 				int32 NumExports = Linker->ExportMap.Num();
-				FStructuredArchive::FStream ExportTableStream = StructuredArchiveRoot.EnterStream(FIELD_NAME_TEXT("ExportTable"));
+				FStructuredArchive::FStream ExportTableStream = StructuredArchiveRoot.EnterStream(SA_FIELD_NAME(TEXT("ExportTable")));
 				{
 #if WITH_EDITOR
 					FArchive::FScopeSetDebugSerializationFlags S(*Linker, DSF_IgnoreDiff, true);
@@ -5991,7 +5991,7 @@ FSavePackageResultStruct UPackage::Save(UPackage* InOuter, UObject* Base, EObjec
 #if WITH_EDITOR
 					FArchiveStackTraceIgnoreScope IgnoreSummaryDiffsScope(DiffSettings.bIgnoreHeaderDiffs);
 #endif // WITH_EDITOR
-					StructuredArchiveRoot.EnterField(FIELD_NAME_TEXT("Summary")) << Linker->Summary;
+					StructuredArchiveRoot.EnterField(SA_FIELD_NAME(TEXT("Summary"))) << Linker->Summary;
 				}
 
 				if (!bTextFormat)
@@ -6338,7 +6338,7 @@ void UPackage::SaveThumbnails(UPackage* InOuter, FLinkerSave* Linker, FStructure
 		if( ObjectsWithThumbnails.Num() > 0 )
 		{
 			// Save out the image data for the thumbnails
-			FStructuredArchive::FStream ThumbnailStream = Record.EnterStream(FIELD_NAME_TEXT("Thumbnails"));
+			FStructuredArchive::FStream ThumbnailStream = Record.EnterStream(SA_FIELD_NAME(TEXT("Thumbnails")));
 
 			for( int32 CurObjectIndex = 0; CurObjectIndex < ObjectsWithThumbnails.Num(); ++CurObjectIndex )
 			{
@@ -6359,7 +6359,7 @@ void UPackage::SaveThumbnails(UPackage* InOuter, FLinkerSave* Linker, FStructure
 
 				// Save number of thumbnails
 				int32 ThumbnailCount = ObjectsWithThumbnails.Num();
-				FStructuredArchive::FArray IndexArray = Record.EnterField(FIELD_NAME_TEXT("Index")).EnterArray(ThumbnailCount);
+				FStructuredArchive::FArray IndexArray = Record.EnterField(SA_FIELD_NAME(TEXT("Index"))).EnterArray(ThumbnailCount);
 
 				// Store a list of object names along with the offset in the file where the thumbnail is stored
 				for( int32 CurObjectIndex = 0; CurObjectIndex < ObjectsWithThumbnails.Num(); ++CurObjectIndex )
@@ -6383,9 +6383,9 @@ void UPackage::SaveThumbnails(UPackage* InOuter, FLinkerSave* Linker, FStructure
 					int32 FileOffset = CurObjectThumb.FileOffset;
 
 					IndexArray.EnterElement().EnterRecord()
-						<< NAMED_FIELD(ObjectClassName)
-						<< NAMED_FIELD(ObjectPathWithoutPackageName)
-						<< NAMED_FIELD(FileOffset);
+						<< SA_VALUE(TEXT("ObjectClassName"), ObjectClassName)
+						<< SA_VALUE(TEXT("ObjectPathWithoutPackageName"), ObjectPathWithoutPackageName)
+						<< SA_VALUE(TEXT("FileOffset"), FileOffset);
 				}
 			}
 		}
@@ -6453,9 +6453,9 @@ void UPackage::SaveAssetRegistryData(UPackage* InOuter, FLinkerSave* Linker, FSt
 		int32 TagCount = Tags.Num();
 
 		FStructuredArchive::FRecord AssetRecord = AssetArray.EnterElement().EnterRecord();
-		AssetRecord << NAMED_ITEM("Path", ObjectPath) << NAMED_ITEM("Class", ObjectClassName);
+		AssetRecord << SA_VALUE(TEXT("Path"), ObjectPath) << SA_VALUE(TEXT("Class"), ObjectClassName);
 
-		FStructuredArchive::FMap TagMap = AssetRecord.EnterField(FIELD_NAME_TEXT("Tags")).EnterMap(TagCount);
+		FStructuredArchive::FMap TagMap = AssetRecord.EnterField(SA_FIELD_NAME(TEXT("Tags"))).EnterMap(TagCount);
 				
 		for (TArray<FAssetRegistryTag>::TConstIterator TagIter(Tags); TagIter; ++TagIter)
 		{
diff --git a/Engine/Source/Runtime/CoreUObject/Public/UObject/ObjectMacros.h b/Engine/Source/Runtime/CoreUObject/Public/UObject/ObjectMacros.h
--- a/Engine/Source/Runtime/CoreUObject/Public/UObject/ObjectMacros.h
+++ b/Engine/Source/Runtime/CoreUObject/Public/UObject/ObjectMacros.h
@@ -1415,7 +1415,7 @@ namespace UM
 	}
 
 #define IMPLEMENT_FARCHIVE_SERIALIZER( TClass ) void TClass::Serialize(FArchive& Ar) { TClass::Serialize(FStructuredArchiveFromArchive(Ar).GetSlot().EnterRecord()); }
-#define IMPLEMENT_FSTRUCTUREDARCHIVE_SERIALIZER( TClass ) void TClass::Serialize(FStructuredArchive::FRecord Record) { FArchiveUObjectFromStructuredArchive Ar(Record.EnterField(FIELD_NAME_TEXT("BaseClassAutoGen"))); TClass::Serialize(Ar.GetArchive()); }
+#define IMPLEMENT_FSTRUCTUREDARCHIVE_SERIALIZER( TClass ) void TClass::Serialize(FStructuredArchive::FRecord Record) { FArchiveUObjectFromStructuredArchive Ar(Record.EnterField(SA_FIELD_NAME(TEXT("BaseClassAutoGen")))); TClass::Serialize(Ar.GetArchive()); }
 #define DECLARE_FARCHIVE_SERIALIZER( TClass, API ) virtual API void Serialize(FArchive& Ar) override;
 #define DECLARE_FSTRUCTUREDARCHIVE_SERIALIZER( TClass, API ) virtual API void Serialize(FStructuredArchive::FRecord Record) override;
 
diff --git a/Engine/Source/Runtime/Engine/Private/EngineUtils.cpp b/Engine/Source/Runtime/Engine/Private/EngineUtils.cpp
--- a/Engine/Source/Runtime/Engine/Private/EngineUtils.cpp
+++ b/Engine/Source/Runtime/Engine/Private/EngineUtils.cpp
@@ -551,8 +551,8 @@ FStripDataFlags::FStripDataFlags(FStructuredArchive::FSlot Slot, uint8 InClassFl
 			GlobalStripFlags |= UnderlyingArchive.CookingTarget()->IsServerOnly() ? FStripDataFlags::Server : FStripDataFlags::None;
 			ClassStripFlags = InClassFlags;
 		}
-		Record << NAMED_FIELD(GlobalStripFlags);
-		Record << NAMED_FIELD(ClassStripFlags);
+		Record << SA_VALUE(TEXT("GlobalStripFlags"), GlobalStripFlags);
+		Record << SA_VALUE(TEXT("ClassStripFlags"), ClassStripFlags);
 	}
 }
 
@@ -572,8 +572,8 @@ FStripDataFlags::FStripDataFlags(FStructuredArchive::FSlot Slot, uint8 InGlobalF
 			GlobalStripFlags = InGlobalFlags;
 			ClassStripFlags = InClassFlags;
 		}
-		Record << NAMED_FIELD(GlobalStripFlags);
-		Record << NAMED_FIELD(ClassStripFlags);
+		Record << SA_VALUE(TEXT("GlobalStripFlags"), GlobalStripFlags);
+		Record << SA_VALUE(TEXT("ClassStripFlags"), ClassStripFlags);
 	}
 }
 
diff --git a/Engine/Source/Runtime/Engine/Private/Particles/ParticleModules.cpp b/Engine/Source/Runtime/Engine/Private/Particles/ParticleModules.cpp
--- a/Engine/Source/Runtime/Engine/Private/Particles/ParticleModules.cpp
+++ b/Engine/Source/Runtime/Engine/Private/Particles/ParticleModules.cpp
@@ -1136,7 +1136,7 @@ void UParticleModuleRequired::Serialize(FStructuredArchive::FRecord Record)
 
 		// Save a bool indicating whether this is cooked data
 		// This is needed when loading cooked data, to know to serialize differently
-		Record << NAMED_FIELD(bCooked);
+		Record << SA_VALUE(TEXT("bCooked"), bCooked);
 
 		if (FPlatformProperties::RequiresCookedData() && !bCooked && UnderlyingArchive.IsLoading())
 		{
@@ -1145,7 +1145,7 @@ void UParticleModuleRequired::Serialize(FStructuredArchive::FRecord Record)
 
 		if (bCooked)
 		{
-			DerivedData.Serialize(Record.EnterField(FIELD_NAME_TEXT("DerivedData")));
+			DerivedData.Serialize(Record.EnterField(SA_FIELD_NAME(TEXT("DerivedData"))));
 		}
 	}
 }
diff --git a/Engine/Source/Runtime/Engine/Private/Particles/SubUVAnimation.cpp b/Engine/Source/Runtime/Engine/Private/Particles/SubUVAnimation.cpp
--- a/Engine/Source/Runtime/Engine/Private/Particles/SubUVAnimation.cpp
+++ b/Engine/Source/Runtime/Engine/Private/Particles/SubUVAnimation.cpp
@@ -76,7 +76,7 @@ void USubUVAnimation::Serialize(FStructuredArchive::FRecord Record)
 
 	// Save a bool indicating whether this is cooked data
 	// This is needed when loading cooked data, to know to serialize differently
-	Record << NAMED_FIELD(bCooked);
+	Record << SA_VALUE(TEXT("bCooked"), bCooked);
 
 	if (FPlatformProperties::RequiresCookedData() && !bCooked && UnderlyingArchive.IsLoading())
 	{
@@ -85,7 +85,7 @@ void USubUVAnimation::Serialize(FStructuredArchive::FRecord Record)
 
 	if (bCooked)
 	{
-		DerivedData.Serialize(Record.EnterField(FIELD_NAME_TEXT("DerivedData")));
+		DerivedData.Serialize(Record.EnterField(SA_FIELD_NAME(TEXT("DerivedData"))));
 	}
 }
 
diff --git a/Engine/Source/Runtime/Engine/Private/PerPlatformProperties.cpp b/Engine/Source/Runtime/Engine/Private/PerPlatformProperties.cpp
--- a/Engine/Source/Runtime/Engine/Private/PerPlatformProperties.cpp
+++ b/Engine/Source/Runtime/Engine/Private/PerPlatformProperties.cpp
@@ -51,21 +51,21 @@ ENGINE_API void operator<<(FStructuredArchive::FSlot Slot, TPerPlatformProperty<
 	if (UnderlyingArchive.IsCooking())
 	{
 		bCooked = true;
-		Record << NAMED_FIELD(bCooked);
+		Record << SA_VALUE(TEXT("bCooked"), bCooked);
 		// Save out platform override if it exists and Default otherwise
 		ValueType Value = Property.GetValueForPlatformIdentifiers(UnderlyingArchive.CookingTarget()->GetPlatformInfo().PlatformGroupName);
-		Record << NAMED_FIELD(Value);
+		Record << SA_VALUE(TEXT("Value"), Value);
 	}
 	else
 #endif
 	{
 		StructType* This = StaticCast<StructType*>(&Property);
-		Record << NAMED_FIELD(bCooked);
-		Record << NAMED_ITEM("Value", This->Default);
+		Record << SA_VALUE(TEXT("bCooked"), bCooked);
+		Record << SA_VALUE(TEXT("Value"), This->Default);
 #if WITH_EDITORONLY_DATA
 		if (!bCooked)
 		{
-			Record << NAMED_ITEM("PerPlatform", This->PerPlatform);
+			Record << SA_VALUE(TEXT("PerPlatform"), This->PerPlatform);
 		}
 #endif
 	}
diff --git a/Engine/Source/Runtime/Engine/Private/SoundCue.cpp b/Engine/Source/Runtime/Engine/Private/SoundCue.cpp
--- a/Engine/Source/Runtime/Engine/Private/SoundCue.cpp
+++ b/Engine/Source/Runtime/Engine/Private/SoundCue.cpp
@@ -102,18 +102,18 @@ void USoundCue::Serialize(FStructuredArchive::FRecord Record)
 
 	if (UnderlyingArchive.UE4Ver() >= VER_UE4_COOKED_ASSETS_IN_EDITOR_SUPPORT)
 	{
-		FStripDataFlags StripFlags(Record.EnterField(FIELD_NAME_TEXT("SoundCueStripFlags")));
+		FStripDataFlags StripFlags(Record.EnterField(SA_FIELD_NAME(TEXT("SoundCueStripFlags"))));
 #if WITH_EDITORONLY_DATA
 		if (!StripFlags.IsEditorDataStripped())
 		{
-			Record << NAMED_FIELD(SoundCueGraph);
+			Record << SA_VALUE(TEXT("SoundCueGraph"), SoundCueGraph);
 		}
 #endif // WITH_EDITORONLY_DATA
 	}
 #if WITH_EDITOR
 	else
 	{
-		Record << NAMED_FIELD(SoundCueGraph);
+		Record << SA_VALUE(TEXT("SoundCueGraph"), SoundCueGraph);
 	}
 #endif // WITH_EDITOR
 }
diff --git a/Engine/Source/Runtime/Engine/Private/UserDefinedStruct.cpp b/Engine/Source/Runtime/Engine/Private/UserDefinedStruct.cpp
--- a/Engine/Source/Runtime/Engine/Private/UserDefinedStruct.cpp
+++ b/Engine/Source/Runtime/Engine/Private/UserDefinedStruct.cpp
@@ -77,7 +77,7 @@ void UUserDefinedStruct::Serialize(FStructuredArchive::FRecord Record)
 			uint8* StructData = DefaultStructInstance.GetStructMemory();
 
 			FScopedPlaceholderRawContainerTracker TrackStruct(StructData);
-			SerializeItem(Record.EnterField(FIELD_NAME_TEXT("Data")), StructData, nullptr);
+			SerializeItem(Record.EnterField(SA_FIELD_NAME(TEXT("Data"))), StructData, nullptr);
 
 			// Now that defaults have been loaded we can inspect our properties
 			// and default values and set the StructFlags accordingly:
diff --git a/Engine/Source/Runtime/UMG/Public/Animation/WidgetAnimationBinding.h b/Engine/Source/Runtime/UMG/Public/Animation/WidgetAnimationBinding.h
--- a/Engine/Source/Runtime/UMG/Public/Animation/WidgetAnimationBinding.h
+++ b/Engine/Source/Runtime/UMG/Public/Animation/WidgetAnimationBinding.h
@@ -68,9 +68,9 @@ public:
 	friend void operator<<(FStructuredArchive::FSlot Slot, FWidgetAnimationBinding& Binding)
 	{
 		FStructuredArchive::FRecord Record = Slot.EnterRecord();
-		Record << NAMED_ITEM("WidgetName", Binding.WidgetName);
-		Record << NAMED_ITEM("SlotWidgetName", Binding.SlotWidgetName);
-		Record << NAMED_ITEM("AnimationGuid", Binding.AnimationGuid);
-		Record << NAMED_ITEM("bIsRootWidget", Binding.bIsRootWidget);
+		Record << SA_VALUE(TEXT("WidgetName"), Binding.WidgetName);
+		Record << SA_VALUE(TEXT("SlotWidgetName"), Binding.SlotWidgetName);
+		Record << SA_VALUE(TEXT("AnimationGuid"), Binding.AnimationGuid);
+		Record << SA_VALUE(TEXT("bIsRootWidget"), Binding.bIsRootWidget);
 	}
 };
