PR #6080: VInterpConstantTo Current parameter should be a reference (Contributed by [REDACTED_PERSON])

diff --git a/Engine/Source/Runtime/Core/Private/Math/UnrealMath.cpp b/Engine/Source/Runtime/Core/Private/Math/UnrealMath.cpp
--- a/Engine/Source/Runtime/Core/Private/Math/UnrealMath.cpp
+++ b/Engine/Source/Runtime/Core/Private/Math/UnrealMath.cpp
@@ -2289,7 +2289,7 @@ CORE_API FVector FMath::VInterpNormalRotationTo(const FVector& Current, const FV
 	return Target;
 }
 
-CORE_API FVector FMath::VInterpConstantTo(const FVector Current, const FVector& Target, float DeltaTime, float InterpSpeed)
+CORE_API FVector FMath::VInterpConstantTo(const FVector& Current, const FVector& Target, float DeltaTime, float InterpSpeed)
 {
 	const FVector Delta = Target - Current;
 	const float DeltaM = Delta.Size();
diff --git a/Engine/Source/Runtime/Core/Public/Math/UnrealMathUtility.h b/Engine/Source/Runtime/Core/Public/Math/UnrealMathUtility.h
--- a/Engine/Source/Runtime/Core/Public/Math/UnrealMathUtility.h
+++ b/Engine/Source/Runtime/Core/Public/Math/UnrealMathUtility.h
@@ -944,7 +944,7 @@ struct FMath : public FPlatformMath
 	static CORE_API FVector VInterpNormalRotationTo(const FVector& Current, const FVector& Target, float DeltaTime, float RotationSpeedDegrees);
 
 	/** Interpolate vector from Current to Target with constant step */
-	static CORE_API FVector VInterpConstantTo(const FVector Current, const FVector& Target, float DeltaTime, float InterpSpeed);
+	static CORE_API FVector VInterpConstantTo(const FVector& Current, const FVector& Target, float DeltaTime, float InterpSpeed);
 
 	/** Interpolate vector from Current to Target. Scaled by distance to Target, so it has a strong start speed and ease out. */
 	static CORE_API FVector VInterpTo( const FVector& Current, const FVector& Target, float DeltaTime, float InterpSpeed );
