PR #6095: Construct FVector2D from single float (Contributed by [REDACTED_PERSON])

diff --git a/Engine/Source/Runtime/Core/Public/Math/Vector2D.h b/Engine/Source/Runtime/Core/Public/Math/Vector2D.h
--- a/Engine/Source/Runtime/Core/Public/Math/Vector2D.h
+++ b/Engine/Source/Runtime/Core/Public/Math/Vector2D.h
@@ -53,7 +53,14 @@ public:
 	 * @param InY Y coordinate.
 	 */
 	FORCEINLINE FVector2D(float InX, float InY);
-
+	
+	/**
+	 * Constructor initializing both components to a single float value.
+	 *
+	 * @param InF Value to set both components to.
+	 */
+	explicit FORCEINLINE FVector2D(float InF);
+	
 	/**
 	 * Constructs a vector from an FIntPoint.
 	 *
@@ -613,6 +620,9 @@ FORCEINLINE FVector2D::FVector2D(float InX,float InY)
 	:	X(InX), Y(InY)
 { }
 
+FORCEINLINE FVector2D::FVector2D(float InF)
+	:	X(InF), Y(InF)
+{ }
 
 FORCEINLINE FVector2D::FVector2D(FIntPoint InPos)
 {
@@ -1002,4 +1012,4 @@ FORCEINLINE float FMath::GetRangePct(FVector2D const& Range, float Value)
 FORCEINLINE float FMath::GetRangeValue(FVector2D const& Range, float Pct)
 {
 	return Lerp<float>(Range.X, Range.Y, Pct);
-}
\ No newline at end of file
+}
