Compiler-specific feature macros moved into their own headers - currently incomplete, but they can be built upon.

diff --git a/Engine/Source/Runtime/Core/Public/Android/AndroidPlatform.h b/Engine/Source/Runtime/Core/Public/Android/AndroidPlatform.h
--- a/Engine/Source/Runtime/Core/Public/Android/AndroidPlatform.h
+++ b/Engine/Source/Runtime/Core/Public/Android/AndroidPlatform.h
@@ -7,6 +7,7 @@
 #pragma once
 
 #include "GenericPlatform/GenericPlatform.h"
+#include "Clang/ClangPlatform.h"
 #include "Misc/Build.h"
 
 /** Define the android platform to be the active one **/
@@ -80,12 +81,6 @@ typedef FAndroidTypes FPlatformTypes;
 	#define PLATFORM_COMPILER_HAS_DECLTYPE_AUTO 0
 #endif
 
-#if __has_feature(cxx_if_constexpr)
-	#define PLATFORM_COMPILER_HAS_IF_CONSTEXPR 1
-#else
-	#define PLATFORM_COMPILER_HAS_IF_CONSTEXPR 0
-#endif
-
 // some android platform overrides that sub-platforms can disable
 #ifndef USE_ANDROID_JNI
 	#define USE_ANDROID_JNI							1
diff --git a/Engine/Source/Runtime/Core/Public/Clang/ClangPlatform.h b/Engine/Source/Runtime/Core/Public/Clang/ClangPlatform.h
new file mode 100644
--- /dev/null
+++ b/Engine/Source/Runtime/Core/Public/Clang/ClangPlatform.h
@@ -0,0 +1,13 @@
+// Copyright 1998-2019 Epic Games, Inc. All Rights Reserved.
+
+/*================================================================================
+	ClangPlatform.h: Setup for any Clang-using platform
+==================================================================================*/
+
+#pragma once
+
+#if __has_feature(cxx_if_constexpr)
+	#define PLATFORM_COMPILER_HAS_IF_CONSTEXPR 1
+#else
+	#define PLATFORM_COMPILER_HAS_IF_CONSTEXPR 0
+#endif
diff --git a/Engine/Source/Runtime/Core/Public/HTML5/HTML5Platform.h b/Engine/Source/Runtime/Core/Public/HTML5/HTML5Platform.h
--- a/Engine/Source/Runtime/Core/Public/HTML5/HTML5Platform.h
+++ b/Engine/Source/Runtime/Core/Public/HTML5/HTML5Platform.h
@@ -6,7 +6,7 @@
 
 #pragma once
 
-
+#include "Clang/ClangPlatform.h"
 
 
 /**
@@ -74,12 +74,6 @@ typedef FHTML5Types FPlatformTypes;
 	#define PLATFORM_COMPILER_HAS_DECLTYPE_AUTO 0
 #endif
 
-#if __has_feature(cxx_if_constexpr)
-	#define PLATFORM_COMPILER_HAS_IF_CONSTEXPR 1
-#else
-	#define PLATFORM_COMPILER_HAS_IF_CONSTEXPR 0
-#endif
-
 // Function type macros.
 #define FORCEINLINE		inline __attribute__((__always_inline__))					/* Force code to be inline */
 #define FORCENOINLINE	__attribute__((noinline))									/* Force code to NOT be inline */
diff --git a/Engine/Source/Runtime/Core/Public/HoloLens/HoloLensPlatform.h b/Engine/Source/Runtime/Core/Public/HoloLens/HoloLensPlatform.h
--- a/Engine/Source/Runtime/Core/Public/HoloLens/HoloLensPlatform.h
+++ b/Engine/Source/Runtime/Core/Public/HoloLens/HoloLensPlatform.h
@@ -6,6 +6,8 @@
 
 #pragma once
 
+#include "MSVC/MSVCPlatform.h"
+
 /** Define the HoloLens platform to be the active one **/
 #define PLATFORM_HOLOLENS					1
 
diff --git a/Engine/Source/Runtime/Core/Public/IOS/IOSPlatform.h b/Engine/Source/Runtime/Core/Public/IOS/IOSPlatform.h
--- a/Engine/Source/Runtime/Core/Public/IOS/IOSPlatform.h
+++ b/Engine/Source/Runtime/Core/Public/IOS/IOSPlatform.h
@@ -6,6 +6,8 @@
 
 #pragma once
 
+#include "Clang/ClangPlatform.h"
+
 /**
 * iOS specific types
 **/
@@ -84,12 +86,6 @@ typedef FIOSPlatformTypes FPlatformTypes;
 	#define PLATFORM_COMPILER_HAS_DECLTYPE_AUTO 0
 #endif
 
-#if __has_feature(cxx_if_constexpr)
-	#define PLATFORM_COMPILER_HAS_IF_CONSTEXPR 1
-#else
-	#define PLATFORM_COMPILER_HAS_IF_CONSTEXPR 0
-#endif
-
 //mallocpoison not safe with aligned ansi allocator.  returns the larger unaligned size during Free() which causes writes off the end of the allocation.
 #define UE_USE_MALLOC_FILL_BYTES 0 
 
diff --git a/Engine/Source/Runtime/Core/Public/MSVC/MSVCPlatform.h b/Engine/Source/Runtime/Core/Public/MSVC/MSVCPlatform.h
new file mode 100644
--- /dev/null
+++ b/Engine/Source/Runtime/Core/Public/MSVC/MSVCPlatform.h
@@ -0,0 +1,13 @@
+// Copyright 1998-2019 Epic Games, Inc. All Rights Reserved.
+
+/*================================================================================
+	MSVCPlatform.h: Setup for any MSVC-using platform
+==================================================================================*/
+
+#pragma once
+
+#if _MSC_FULL_VER >= 191125507 && defined(_MSVC_LANG) && _MSVC_LANG >= 201402
+	#define PLATFORM_COMPILER_HAS_IF_CONSTEXPR 1
+#else
+	#define PLATFORM_COMPILER_HAS_IF_CONSTEXPR 0
+#endif
diff --git a/Engine/Source/Runtime/Core/Public/Mac/MacPlatform.h b/Engine/Source/Runtime/Core/Public/Mac/MacPlatform.h
--- a/Engine/Source/Runtime/Core/Public/Mac/MacPlatform.h
+++ b/Engine/Source/Runtime/Core/Public/Mac/MacPlatform.h
@@ -6,6 +6,8 @@
 
 #pragma once
 
+#include "Clang/ClangPlatform.h"
+
 /**
 * Mac specific types
 **/
@@ -60,12 +62,6 @@ typedef FMacPlatformTypes FPlatformTypes;
 	#define PLATFORM_COMPILER_HAS_DECLTYPE_AUTO 0
 #endif
 
-#if __has_feature(cxx_if_constexpr)
-	#define PLATFORM_COMPILER_HAS_IF_CONSTEXPR 1
-#else
-	#define PLATFORM_COMPILER_HAS_IF_CONSTEXPR 0
-#endif
-
 // Function type macros.
 #define VARARGS															/* Functions with variable arguments */
 #define CDECL															/* Standard C function */
diff --git a/Engine/Source/Runtime/Core/Public/Unix/UnixPlatform.h b/Engine/Source/Runtime/Core/Public/Unix/UnixPlatform.h
--- a/Engine/Source/Runtime/Core/Public/Unix/UnixPlatform.h
+++ b/Engine/Source/Runtime/Core/Public/Unix/UnixPlatform.h
@@ -8,6 +8,7 @@
 
 #include <linux/version.h>
 #include <signal.h>
+#include "Clang/ClangPlatform.h"
 
 struct FGenericPlatformTypes;
 
@@ -77,12 +78,6 @@ typedef FUnixPlatformTypes FPlatformTypes;
 	#define PLATFORM_COMPILER_HAS_DECLTYPE_AUTO 0
 #endif
 
-#if __has_feature(cxx_if_constexpr)
-	#define PLATFORM_COMPILER_HAS_IF_CONSTEXPR 1
-#else
-	#define PLATFORM_COMPILER_HAS_IF_CONSTEXPR 0
-#endif
-
 // SOCK_CLOEXEC is available on Unix since 2.6.27
 #if LINUX_VERSION_CODE >= KERNEL_VERSION(2,6,27)
 	#define PLATFORM_HAS_BSD_SOCKET_FEATURE_CLOSE_ON_EXEC	1
diff --git a/Engine/Source/Runtime/Core/Public/Windows/WIndowsPlatform.h b/Engine/Source/Runtime/Core/Public/Windows/WIndowsPlatform.h
--- a/Engine/Source/Runtime/Core/Public/Windows/WIndowsPlatform.h
+++ b/Engine/Source/Runtime/Core/Public/Windows/WIndowsPlatform.h
@@ -4,6 +4,12 @@
 
 #include <sal.h>
 
+#if defined(__clang__)
+	#include "Clang/ClangPlatform.h"
+#else
+	#include "MSVC/MSVCPlatform.h"
+#endif
+
 
 /**
 * Windows specific types
@@ -68,12 +74,6 @@ typedef FWindowsPlatformTypes FPlatformTypes;
 
 #define PLATFORM_GLOBAL_LOG_CATEGORY						LogWindows
 
-#if _MSC_FULL_VER >= 191125507 && defined(_MSVC_LANG) && _MSVC_LANG >= 201402
-	#define PLATFORM_COMPILER_HAS_IF_CONSTEXPR 1
-#else
-	#define PLATFORM_COMPILER_HAS_IF_CONSTEXPR 0
-#endif
-
 
 // Q: Why is there a __nop() before __debugbreak()?
 // A: VS' debug engine has a bug where it will silently swallow explicit
