Added FScriptSet::FindOrAdd() and removed FScriptMap::FindOrAdd() dependency on FScriptSet::Add() implementation details

Also corrected FScriptMap::FindOrAdd() documentation

diff --git a/Engine/Source/Runtime/Core/Public/Containers/Map.h b/Engine/Source/Runtime/Core/Public/Containers/Map.h
--- a/Engine/Source/Runtime/Core/Public/Containers/Map.h
+++ b/Engine/Source/Runtime/Core/Public/Containers/Map.h
@@ -1615,7 +1615,7 @@ public:
 	}
 
 	/**
-	 * Adds a new uninitialized value or destroys the existing value for a given key.
+	 * Constructs a new key-value pair if key didn't exist
 	 *
 	 * No need to rehash after calling. The hash table must be properly hashed before calling.
 	 *
@@ -1628,29 +1628,18 @@ public:
 		TFunctionRef<bool(const void*, const void*)> KeyEqualityFn,
 		TFunctionRef<void(void*, void*)> ConstructPairFn)
 	{
-		void* OutPair = nullptr;
-		bool bCreatePair = true;
 		const int32 ValueOffset = Layout.ValueOffset;
-				
-		Pairs.Add(
+		int32 PairIndex = Pairs.FindOrAdd(
 			Key,
 			Layout.SetLayout,
 			GetKeyHash,
 			KeyEqualityFn,
-			[ConstructPairFn, &bCreatePair, &OutPair, ValueOffset](void* NewPair)
+			[ConstructPairFn, ValueOffset](void* NewPair)
 			{
-				if (bCreatePair)
-				{
-					ConstructPairFn(NewPair, (uint8*)NewPair + ValueOffset);
-				}
-
-				OutPair = NewPair;
-			},
-			[&bCreatePair](void* NewPair) {	bCreatePair = false; }
-		);
+				ConstructPairFn(NewPair, (uint8*)NewPair + ValueOffset);
+			});
 
-		checkSlow(OutPair);
-		return (uint8*)OutPair + ValueOffset;
+		return (uint8*)Pairs.GetData(PairIndex, Layout.SetLayout) + ValueOffset;
 	}
 
 private:
diff --git a/Engine/Source/Runtime/Core/Public/Containers/Set.h b/Engine/Source/Runtime/Core/Public/Containers/Set.h
--- a/Engine/Source/Runtime/Core/Public/Containers/Set.h
+++ b/Engine/Source/Runtime/Core/Public/Containers/Set.h
@@ -1662,14 +1662,25 @@ public:
 		return INDEX_NONE;
 	}
 
-	void Add(const void* Element, const FScriptSetLayout& Layout, TFunctionRef<uint32(const void*)> GetKeyHash, TFunctionRef<bool(const void*, const void*)> EqualityFn, TFunctionRef<void(void*)> ConstructFn, TFunctionRef<void(void*)> DestructFn)
+	int32 FindOrAdd(const void* Element, const FScriptSetLayout& Layout, TFunctionRef<uint32(const void*)> GetKeyHash, TFunctionRef<bool(const void*, const void*)> EqualityFn, TFunctionRef<void(void*)> ConstructFn)
 	{
 		uint32 KeyHash = GetKeyHash(Element);
+		int32 OldElementIndex = FindIndexByHash(Element, Layout, KeyHash, EqualityFn);
+		if (OldElementIndex != INDEX_NONE)
+		{
+			return OldElementIndex;
+		}
 
-		int32 NewElementIndex = FindIndexByHash(Element, Layout, KeyHash, EqualityFn);
-		if (NewElementIndex != INDEX_NONE)
+		return AddNewElement(Layout, GetKeyHash, KeyHash, ConstructFn);
+	}
+
+	void Add(const void* Element, const FScriptSetLayout& Layout, TFunctionRef<uint32(const void*)> GetKeyHash, TFunctionRef<bool(const void*, const void*)> EqualityFn, TFunctionRef<void(void*)> ConstructFn, TFunctionRef<void(void*)> DestructFn)
+	{
+		uint32 KeyHash = GetKeyHash(Element);
+		int32 OldElementIndex = FindIndexByHash(Element, Layout, KeyHash, EqualityFn);
+		if (OldElementIndex != INDEX_NONE)
 		{
-			void* ElementPtr = Elements.GetData(NewElementIndex, Layout.SparseArrayLayout);
+			void* ElementPtr = Elements.GetData(OldElementIndex, Layout.SparseArrayLayout);
 
 			DestructFn(ElementPtr);
 			ConstructFn(ElementPtr);
@@ -1680,30 +1691,37 @@ public:
 		}
 		else
 		{
-			NewElementIndex = Elements.AddUninitialized(Layout.SparseArrayLayout);
+			AddNewElement(Layout, GetKeyHash, KeyHash, ConstructFn);
+		}
+	}
 
-			void* ElementPtr = Elements.GetData(NewElementIndex, Layout.SparseArrayLayout);
-			ConstructFn(ElementPtr);
+private:
+	int32 AddNewElement(const FScriptSetLayout& Layout, TFunctionRef<uint32(const void*)> GetKeyHash, uint32 KeyHash, TFunctionRef<void(void*)> ConstructFn)
+	{
+		int32 NewElementIndex = Elements.AddUninitialized(Layout.SparseArrayLayout);
 
-			const int32 DesiredHashSize = FDefaultSetAllocator::GetNumberOfHashBuckets(Num());
-			if (!HashSize || HashSize < DesiredHashSize)
-			{
-				// rehash, this will link in our new element if needed:
-				Rehash(Layout, GetKeyHash);
-			}
-			else
-			{
-				// link the new element into the set:
-				int32 HashIndex = KeyHash & (HashSize - 1);
-				FSetElementId& TypedHash = GetTypedHash(HashIndex);
-				GetHashIndexRef(ElementPtr, Layout) = HashIndex;
-				GetHashNextIdRef(ElementPtr, Layout) = TypedHash;
-				TypedHash = FSetElementId(NewElementIndex);
-			}
+		void* ElementPtr = Elements.GetData(NewElementIndex, Layout.SparseArrayLayout);
+		ConstructFn(ElementPtr);
+
+		const int32 DesiredHashSize = FDefaultSetAllocator::GetNumberOfHashBuckets(Num());
+		if (!HashSize || HashSize < DesiredHashSize)
+		{
+			// rehash, this will link in our new element if needed:
+			Rehash(Layout, GetKeyHash);
+		}
+		else
+		{
+			// link the new element into the set:
+			int32 HashIndex = KeyHash & (HashSize - 1);
+			FSetElementId& TypedHash = GetTypedHash(HashIndex);
+			GetHashIndexRef(ElementPtr, Layout) = HashIndex;
+			GetHashNextIdRef(ElementPtr, Layout) = TypedHash;
+			TypedHash = FSetElementId(NewElementIndex);
 		}
+
+		return NewElementIndex;
 	}
 
-private:
 	typedef FDefaultSetAllocator Allocator;
 	typedef Allocator::HashAllocator::ForElementType<FSetElementId> HashType;
 
