CIS fix

#rnx

diff --git a/Engine/Source/Runtime/Core/Private/Modules/ModuleManager.cpp b/Engine/Source/Runtime/Core/Private/Modules/ModuleManager.cpp
--- a/Engine/Source/Runtime/Core/Private/Modules/ModuleManager.cpp
+++ b/Engine/Source/Runtime/Core/Private/Modules/ModuleManager.cpp
@@ -308,6 +308,7 @@ void FModuleManager::AddModule(const FName InModuleName)
 	// Add this module to the set of modules that we know about
 	ModuleInfo->OriginalFilename = Prefix + Suffix;
 	ModuleInfo->Filename         = ModuleFilename;
+#endif // !IS_MONOLITHIC
 }
 
 
