Defaulted attribute support in structured archives.

diff --git a/Engine/Source/Runtime/Core/Private/Serialization/StructuredArchive.cpp b/Engine/Source/Runtime/Core/Private/Serialization/StructuredArchive.cpp
--- a/Engine/Source/Runtime/Core/Private/Serialization/StructuredArchive.cpp
+++ b/Engine/Source/Runtime/Core/Private/Serialization/StructuredArchive.cpp
@@ -369,6 +369,52 @@ FStructuredArchive::FSlot FStructuredArchive::FSlot::EnterAttribute(FArchiveFiel
 	return FSlot(Ar, NewDepth, Ar.CurrentSlotElementId);
 }
 
+TOptional<FStructuredArchive::FSlot> FStructuredArchive::FSlot::TryEnterAttribute(FArchiveFieldName AttributeName, bool bEnterWhenWriting)
+{
+	check(Ar.CurrentScope.Num() > 0);
+
+	int32 NewDepth = Depth + 1;
+	if (NewDepth >= Ar.CurrentScope.Num() || Ar.CurrentScope[NewDepth].Id != ElementId || Ar.CurrentScope[NewDepth].Type != EElementType::AttributedValue)
+	{
+		int32 NewDepthCheck = Ar.EnterSlot(Depth, ElementId, EElementType::AttributedValue);
+		checkSlow(NewDepth == NewDepthCheck);
+
+		Ar.Formatter.EnterAttributedValue();
+
+#if DO_GUARD_SLOW
+		Ar.CurrentContainer.Add(new FContainer(0));
+#endif
+	}
+
+#if DO_GUARD_SLOW
+#if CHECK_UNIQUE_FIELD_NAMES
+	if (!Ar.GetUnderlyingArchive().IsLoading())
+	{
+		FContainer& Container = *Ar.CurrentContainer.Top();
+		checkf(!Container.KeyNames.Contains(Name.Name), TEXT("Multiple attributes called '%s' serialized into attributed value"), AttributeName.Name);
+		Container.KeyNames.Add(Name.Name);
+	}
+#endif
+#endif
+
+	int32 AttributedValueId = Ar.CurrentScope[NewDepth].Id;
+
+	Ar.SetScope(NewDepth, AttributedValueId);
+
+	if (Ar.Formatter.TryEnterAttribute(AttributeName, bEnterWhenWriting))
+	{
+		Ar.CurrentEnteringAttributeState = EEnteringAttributeState::NotEnteringAttribute;
+
+		Ar.CurrentSlotElementId = Ar.NextElementId++;
+
+		return FSlot(Ar, NewDepth, Ar.CurrentSlotElementId);
+	}
+	else
+	{
+		return {};
+	}
+}
+
 void FStructuredArchive::FSlot::Serialize(TArray<uint8>& Data)
 {
 	Ar.EnterSlot(Depth, ElementId);
@@ -455,11 +501,11 @@ TOptional<FStructuredArchive::FSlot> FStructuredArchive::FRecord::TryEnterField(
 	if (Ar.Formatter.TryEnterField(Name, bEnterWhenWriting))
 	{
 		Ar.CurrentSlotElementId = Ar.NextElementId++;
-		return TOptional<FStructuredArchive::FSlot>(FSlot(Ar, Depth, Ar.CurrentSlotElementId));
+		return FSlot(Ar, Depth, Ar.CurrentSlotElementId);
 	}
 	else
 	{
-		return TOptional<FStructuredArchive::FSlot>();
+		return {};
 	}
 }
 
diff --git a/Engine/Source/Runtime/Core/Public/Serialization/Formatters/BinaryArchiveFormatter.h b/Engine/Source/Runtime/Core/Public/Serialization/Formatters/BinaryArchiveFormatter.h
--- a/Engine/Source/Runtime/Core/Public/Serialization/Formatters/BinaryArchiveFormatter.h
+++ b/Engine/Source/Runtime/Core/Public/Serialization/Formatters/BinaryArchiveFormatter.h
@@ -47,6 +47,7 @@ public:
 	virtual void EnterAttributedValueValue() override;
 	virtual void LeaveAttribute() override;
 	virtual void LeaveAttributedValue() override;
+	virtual bool TryEnterAttribute(FArchiveFieldName AttributeName, bool bEnterWhenWriting) override;
 
 	virtual void Serialize(uint8& Value) override;
 	virtual void Serialize(uint16& Value) override;
@@ -178,6 +179,17 @@ inline void FBinaryArchiveFormatter::LeaveAttributedValue()
 {
 }
 
+inline bool FBinaryArchiveFormatter::TryEnterAttribute(FArchiveFieldName AttributeName, bool bEnterWhenWriting)
+{
+	bool bValue = bEnterWhenWriting;
+	Inner << bValue;
+	if (bValue)
+	{
+		EnterAttribute(AttributeName);
+	}
+	return bValue;
+}
+
 inline void FBinaryArchiveFormatter::Serialize(uint8& Value)
 {
 	Inner << Value;
diff --git a/Engine/Source/Runtime/Core/Public/Serialization/StructuredArchive.h b/Engine/Source/Runtime/Core/Public/Serialization/StructuredArchive.h
--- a/Engine/Source/Runtime/Core/Public/Serialization/StructuredArchive.h
+++ b/Engine/Source/Runtime/Core/Public/Serialization/StructuredArchive.h
@@ -43,6 +43,24 @@ template<typename T> struct TNamedAttribute
 	}
 };
 
+/**
+ * Class to contain a named attribute for serialization, with a default. Intended to be created as a temporary and passed to object
+ * serialization methods, which can choose not to serialize the attribute if it matches the default.
+ */
+template<typename T> struct TDefaultedNamedAttribute
+{
+	FArchiveFieldName Name;
+	T& Value;
+	const T& Default;
+
+	explicit FORCEINLINE TDefaultedNamedAttribute(FArchiveFieldName InName, T& InValue, const T& InDefault)
+		: Name(InName)
+		, Value(InValue)
+		, Default(InDefault)
+	{
+	}
+};
+
 /**
  * Helper function to construct a TNamedValue, deducing the value type.
  */
@@ -52,13 +70,21 @@ template<typename T> FORCEINLINE TNamedValue<T> MakeNamedValue(FArchiveFieldName
 }
 
 /**
- * Helper function to construct a TNamedValue, deducing the value type.
+ * Helper function to construct a TNamedAttribute, deducing the value type.
  */
 template<typename T> FORCEINLINE TNamedAttribute<T> MakeNamedAttribute(FArchiveFieldName Name, T& Value)
 {
 	return TNamedAttribute<T>(Name, Value);
 }
 
+/**
+ * Helper function to construct a TDefaultedNamedAttribute, deducing the value type.
+ */
+template<typename T> FORCEINLINE TDefaultedNamedAttribute<T> MakeDefaultedNamedAttribute(FArchiveFieldName Name, T& Value, const typename TIdentity<T>::Type& Default)
+{
+	return TDefaultedNamedAttribute<T>(Name, Value, Default);
+}
+
 /** Construct a TNamedValue given an ANSI string and value reference. */
 #if WITH_TEXT_ARCHIVE_SUPPORT
 	#define SA_VALUE(Name, Value) MakeNamedValue(FArchiveFieldName(Name), Value)
@@ -66,13 +92,20 @@ template<typename T> FORCEINLINE TNamedAttribute<T> MakeNamedAttribute(FArchiveF
 	#define SA_VALUE(Name, Value) MakeNamedValue(FArchiveFieldName(), Value)
 #endif
 
-/** Construct a TNamedValue given an ANSI string and value reference. */
+/** Construct a TNamedAttribute given an ANSI string and value reference. */
 #if WITH_TEXT_ARCHIVE_SUPPORT
 	#define SA_ATTRIBUTE(Name, Value) MakeNamedAttribute(FArchiveFieldName(Name), Value)
 #else
 	#define SA_ATTRIBUTE(Name, Value) MakeNamedAttribute(FArchiveFieldName(), Value)
 #endif
 
+/** Construct a TDefaultedNamedAttribute given an ANSI string and value reference. */
+#if WITH_TEXT_ARCHIVE_SUPPORT
+	#define SA_DEFAULTED_ATTRIBUTE(Name, Value, Default) MakeDefaultedNamedAttribute(FArchiveFieldName(Name), Value, Default)
+#else
+	#define SA_DEFAULTED_ATTRIBUTE(Name, Value) MakeNamedAttribute(FArchiveFieldName(), Value)
+#endif
+
 /** Typedef for which formatter type to support */
 #if WITH_TEXT_ARCHIVE_SUPPORT
 	typedef FStructuredArchiveFormatter FArchiveFormatterType;
@@ -173,6 +206,7 @@ public:
 		FStream EnterStream_TextOnly(int32& OutNumElements);
 		FMap EnterMap(int32& Num);
 		FSlot EnterAttribute(FArchiveFieldName AttributeName);
+		TOptional<FSlot> TryEnterAttribute(FArchiveFieldName AttributeName, bool bEnterWhenWriting);
 
 		// We don't support chaining writes to a single slot, so this returns void.
 		template <typename ArgType>
@@ -195,6 +229,19 @@ public:
 			EnterAttribute(Item.Name) << Item.Value;
 		}
 
+		template <typename T>
+		FORCEINLINE void operator<<(TDefaultedNamedAttribute<T> Item)
+		{
+			if (TOptional<FSlot> Attribute = TryEnterAttribute(Item.Name, Item.Value != Item.Default))
+			{
+				Attribute.GetValue() << Item.Value;
+			}
+			else
+			{
+				Item.Value = Item.Default;
+			}
+		}
+
 		void Serialize(TArray<uint8>& Data);
 		void Serialize(void* Data, uint64 DataSize);
 
diff --git a/Engine/Source/Runtime/Core/Public/Serialization/StructuredArchiveFormatter.h b/Engine/Source/Runtime/Core/Public/Serialization/StructuredArchiveFormatter.h
--- a/Engine/Source/Runtime/Core/Public/Serialization/StructuredArchiveFormatter.h
+++ b/Engine/Source/Runtime/Core/Public/Serialization/StructuredArchiveFormatter.h
@@ -108,6 +108,7 @@ public:
 	virtual void EnterAttributedValueValue() = 0;
 	virtual void LeaveAttribute() = 0;
 	virtual void LeaveAttributedValue() = 0;
+	virtual bool TryEnterAttribute(FArchiveFieldName AttributeName, bool bEnterWhenWriting) = 0;
 
 	virtual void Serialize(uint8& Value) = 0;
 	virtual void Serialize(uint16& Value) = 0;
diff --git a/Engine/Source/Runtime/CoreUObject/Private/Serialization/Formatters/JsonArchiveInputFormatter.cpp b/Engine/Source/Runtime/CoreUObject/Private/Serialization/Formatters/JsonArchiveInputFormatter.cpp
--- a/Engine/Source/Runtime/CoreUObject/Private/Serialization/Formatters/JsonArchiveInputFormatter.cpp
+++ b/Engine/Source/Runtime/CoreUObject/Private/Serialization/Formatters/JsonArchiveInputFormatter.cpp
@@ -265,6 +265,23 @@ void FJsonArchiveInputFormatter::LeaveAttributedValue()
 	ObjectStack.Pop();
 }
 
+bool FJsonArchiveInputFormatter::TryEnterAttribute(FArchiveFieldName AttributeName, bool bEnterWhenSaving)
+{
+	if (TSharedPtr<FJsonObject>& Object = ObjectStack.Top().JsonObject)
+	{
+		if (TSharedPtr<FJsonValue> Field = Object->TryGetField(EscapeFieldName(TEXT("_Value"))))
+		{
+			TSharedPtr<FJsonValue> Attribute = Object->TryGetField(EscapeFieldName(*FString::Printf(TEXT("_%s"), AttributeName.Name)));
+			if (Attribute.IsValid())
+			{
+				ValueStack.Add(Attribute);
+				return true;
+			}
+		}
+	}
+	return false;
+}
+
 void FJsonArchiveInputFormatter::Serialize(uint8& Value)
 {
 	Value = (uint8)ValueStack.Top()->AsNumber();
diff --git a/Engine/Source/Runtime/CoreUObject/Private/Serialization/Formatters/JsonArchiveOutputFormatter.cpp b/Engine/Source/Runtime/CoreUObject/Private/Serialization/Formatters/JsonArchiveOutputFormatter.cpp
--- a/Engine/Source/Runtime/CoreUObject/Private/Serialization/Formatters/JsonArchiveOutputFormatter.cpp
+++ b/Engine/Source/Runtime/CoreUObject/Private/Serialization/Formatters/JsonArchiveOutputFormatter.cpp
@@ -217,6 +217,15 @@ void FJsonArchiveOutputFormatter::EnterAttributedValueValue()
 	WriteOptionalAttributedBlockValue();
 }
 
+bool FJsonArchiveOutputFormatter::TryEnterAttribute(FArchiveFieldName AttributeName, bool bEnterWhenSaving)
+{
+	if (bEnterWhenSaving)
+	{
+		EnterAttribute(AttributeName);
+	}
+	return bEnterWhenSaving;
+}
+
 void FJsonArchiveOutputFormatter::Serialize(uint8& Value)
 {
 	WriteValue(LexToString(Value));
diff --git a/Engine/Source/Runtime/CoreUObject/Private/Serialization/Formatters/TaggedBinaryArchiveInputFormatter.cpp b/Engine/Source/Runtime/CoreUObject/Private/Serialization/Formatters/TaggedBinaryArchiveInputFormatter.cpp
--- a/Engine/Source/Runtime/CoreUObject/Private/Serialization/Formatters/TaggedBinaryArchiveInputFormatter.cpp
+++ b/Engine/Source/Runtime/CoreUObject/Private/Serialization/Formatters/TaggedBinaryArchiveInputFormatter.cpp
@@ -112,6 +112,7 @@ void FTaggedBinaryArchiveInputFormatter::EnterField(FArchiveFieldName Name)
 
 	int32 NameIdx = NameToIndex.FindChecked(Name.Name);
 	const FField* Field = Record.Fields.FindByPredicate([NameIdx](const FField& TestField) -> bool { return TestField.NameIdx == NameIdx; });
+	check(Field);
 	Inner.Seek(Field->Offset);
 }
 
@@ -292,6 +293,22 @@ void FTaggedBinaryArchiveInputFormatter::LeaveAttributedValue()
 	AttributedValues.Pop();
 }
 
+bool FTaggedBinaryArchiveInputFormatter::TryEnterAttribute(FArchiveFieldName AttributeName, bool bEnterWhenSaving)
+{
+	ExpectType(EArchiveValueType::Attribute);
+
+	FAttributedValue& AttributedValue = AttributedValues.Top();
+
+	int32 NameIdx = NameToIndex.FindChecked(AttributeName.Name);
+	const FAttribute* Attribute = Algo::FindBy(AttributedValue.Attributes, NameIdx, &FAttribute::NameIdx);
+	if (Attribute)
+	{
+		Inner.Seek(Attribute->StartOffset);
+		return true;
+	}
+	return false;
+}
+
 void FTaggedBinaryArchiveInputFormatter::Serialize(uint8& Value)
 {
 	ReadNumericValue(Value);
diff --git a/Engine/Source/Runtime/CoreUObject/Private/Serialization/Formatters/TaggedBinaryArchiveOutputFormatter.cpp b/Engine/Source/Runtime/CoreUObject/Private/Serialization/Formatters/TaggedBinaryArchiveOutputFormatter.cpp
--- a/Engine/Source/Runtime/CoreUObject/Private/Serialization/Formatters/TaggedBinaryArchiveOutputFormatter.cpp
+++ b/Engine/Source/Runtime/CoreUObject/Private/Serialization/Formatters/TaggedBinaryArchiveOutputFormatter.cpp
@@ -268,6 +268,15 @@ void FTaggedBinaryArchiveOutputFormatter::LeaveAttributedValue()
 	AttributedValues.Pop();
 }
 
+bool FTaggedBinaryArchiveOutputFormatter::TryEnterAttribute(FArchiveFieldName AttributeName, bool bEnterWhenSaving)
+{
+	if (bEnterWhenSaving)
+	{
+		EnterAttribute(AttributeName);
+	}
+	return bEnterWhenSaving;
+}
+
 void FTaggedBinaryArchiveOutputFormatter::Serialize(uint8& Value)
 {
 	WriteType(EArchiveValueType::UInt8);
diff --git a/Engine/Source/Runtime/CoreUObject/Public/Serialization/Formatters/JsonArchiveInputFormatter.h b/Engine/Source/Runtime/CoreUObject/Public/Serialization/Formatters/JsonArchiveInputFormatter.h
--- a/Engine/Source/Runtime/CoreUObject/Public/Serialization/Formatters/JsonArchiveInputFormatter.h
+++ b/Engine/Source/Runtime/CoreUObject/Public/Serialization/Formatters/JsonArchiveInputFormatter.h
@@ -53,6 +53,7 @@ public:
 	virtual void EnterAttributedValueValue() override;
 	virtual void LeaveAttribute() override;
 	virtual void LeaveAttributedValue() override;
+	virtual bool TryEnterAttribute(FArchiveFieldName AttributeName, bool bEnterWhenSavin) override;
 
 	virtual void Serialize(uint8& Value) override;
 	virtual void Serialize(uint16& Value) override;
diff --git a/Engine/Source/Runtime/CoreUObject/Public/Serialization/Formatters/JsonArchiveOutputFormatter.h b/Engine/Source/Runtime/CoreUObject/Public/Serialization/Formatters/JsonArchiveOutputFormatter.h
--- a/Engine/Source/Runtime/CoreUObject/Public/Serialization/Formatters/JsonArchiveOutputFormatter.h
+++ b/Engine/Source/Runtime/CoreUObject/Public/Serialization/Formatters/JsonArchiveOutputFormatter.h
@@ -49,6 +49,7 @@ public:
 	virtual void EnterAttributedValueValue() override;
 	virtual void LeaveAttribute() override;
 	virtual void LeaveAttributedValue() override;
+	virtual bool TryEnterAttribute(FArchiveFieldName AttributeName, bool bEnterWhenSaving) override;
 
 	virtual void Serialize(uint8& Value) override;
 	virtual void Serialize(uint16& Value) override;
diff --git a/Engine/Source/Runtime/CoreUObject/Public/Serialization/Formatters/TaggedBinaryArchiveInputFormatter.h b/Engine/Source/Runtime/CoreUObject/Public/Serialization/Formatters/TaggedBinaryArchiveInputFormatter.h
--- a/Engine/Source/Runtime/CoreUObject/Public/Serialization/Formatters/TaggedBinaryArchiveInputFormatter.h
+++ b/Engine/Source/Runtime/CoreUObject/Public/Serialization/Formatters/TaggedBinaryArchiveInputFormatter.h
@@ -49,6 +49,7 @@ public:
 	virtual void EnterAttributedValueValue() override;
 	virtual void LeaveAttribute() override;
 	virtual void LeaveAttributedValue() override;
+	virtual bool TryEnterAttribute(FArchiveFieldName AttributeName, bool bEnterWhenSaving) override;
 
 	virtual void Serialize(uint8& Value) override;
 	virtual void Serialize(uint16& Value) override;
diff --git a/Engine/Source/Runtime/CoreUObject/Public/Serialization/Formatters/TaggedBinaryArchiveOutputFormatter.h b/Engine/Source/Runtime/CoreUObject/Public/Serialization/Formatters/TaggedBinaryArchiveOutputFormatter.h
--- a/Engine/Source/Runtime/CoreUObject/Public/Serialization/Formatters/TaggedBinaryArchiveOutputFormatter.h
+++ b/Engine/Source/Runtime/CoreUObject/Public/Serialization/Formatters/TaggedBinaryArchiveOutputFormatter.h
@@ -49,6 +49,7 @@ public:
 	virtual void EnterAttributedValueValue() override;
 	virtual void LeaveAttribute() override;
 	virtual void LeaveAttributedValue() override;
+	virtual bool TryEnterAttribute(FArchiveFieldName AttributeName, bool bEnterWhenSaving) override;
 
 	virtual void Serialize(uint8& Value) override;
 	virtual void Serialize(uint16& Value) override;
