Reverting a change SteveR made to try and templatize the FStructuredArchive::FSlot << operator to reduce code verbosity. There is evidence that it is the cause of the very rare linux/mac runtime crash that we have been looking for. If this fixes the nightly CIS then we can look at fixing this code up properly in the future

#rnx

diff --git a/Engine/Source/Runtime/Core/Private/Serialization/StructuredArchive.cpp b/Engine/Source/Runtime/Core/Private/Serialization/StructuredArchive.cpp
--- a/Engine/Source/Runtime/Core/Private/Serialization/StructuredArchive.cpp
+++ b/Engine/Source/Runtime/Core/Private/Serialization/StructuredArchive.cpp
@@ -415,6 +415,150 @@ TOptional<FStructuredArchive::FSlot> FStructuredArchive::FSlot::TryEnterAttribut
 	}
 }
 
+void FStructuredArchive::FSlot::operator<< (char& Value)
+{
+	Ar.EnterSlot(Depth, ElementId);
+
+	int8 AsInt = Value;
+	Ar.Formatter.Serialize(AsInt);
+	Value = AsInt;
+
+	Ar.LeaveSlot();
+}
+
+void FStructuredArchive::FSlot::operator<< (uint8& Value)
+{
+	Ar.EnterSlot(Depth, ElementId);
+	Ar.Formatter.Serialize(Value);
+	Ar.LeaveSlot();
+}
+
+void FStructuredArchive::FSlot::operator<< (uint16& Value)
+{
+	Ar.EnterSlot(Depth, ElementId);
+	Ar.Formatter.Serialize(Value);
+	Ar.LeaveSlot();
+}
+
+void FStructuredArchive::FSlot::operator<< (uint32& Value)
+{
+	Ar.EnterSlot(Depth, ElementId);
+	Ar.Formatter.Serialize(Value);
+	Ar.LeaveSlot();
+}
+
+void FStructuredArchive::FSlot::operator<< (uint64& Value)
+{
+	Ar.EnterSlot(Depth, ElementId);
+	Ar.Formatter.Serialize(Value);
+	Ar.LeaveSlot();
+}
+
+void FStructuredArchive::FSlot::operator<< (int8& Value)
+{
+	Ar.EnterSlot(Depth, ElementId);
+	Ar.Formatter.Serialize(Value);
+	Ar.LeaveSlot();
+}
+
+void FStructuredArchive::FSlot::operator<< (int16& Value)
+{
+	Ar.EnterSlot(Depth, ElementId);
+	Ar.Formatter.Serialize(Value);
+	Ar.LeaveSlot();
+}
+
+void FStructuredArchive::FSlot::operator<< (int32& Value)
+{
+	Ar.EnterSlot(Depth, ElementId);
+	Ar.Formatter.Serialize(Value);
+	Ar.LeaveSlot();
+}
+
+void FStructuredArchive::FSlot::operator<< (int64& Value)
+{
+	Ar.EnterSlot(Depth, ElementId);
+	Ar.Formatter.Serialize(Value);
+	Ar.LeaveSlot();
+}
+
+void FStructuredArchive::FSlot::operator<< (float& Value)
+{
+	Ar.EnterSlot(Depth, ElementId);
+	Ar.Formatter.Serialize(Value);
+	Ar.LeaveSlot();
+}
+
+void FStructuredArchive::FSlot::operator<< (double& Value)
+{
+	Ar.EnterSlot(Depth, ElementId);
+	Ar.Formatter.Serialize(Value);
+	Ar.LeaveSlot();
+}
+
+void FStructuredArchive::FSlot::operator<< (bool& Value)
+{
+	Ar.EnterSlot(Depth, ElementId);
+	Ar.Formatter.Serialize(Value);
+	Ar.LeaveSlot();
+}
+
+void FStructuredArchive::FSlot::operator<< (FString& Value)
+{
+	Ar.EnterSlot(Depth, ElementId);
+	Ar.Formatter.Serialize(Value);
+	Ar.LeaveSlot();
+}
+
+void FStructuredArchive::FSlot::operator<< (FName& Value)
+{
+	Ar.EnterSlot(Depth, ElementId);
+	Ar.Formatter.Serialize(Value);
+	Ar.LeaveSlot();
+}
+
+void FStructuredArchive::FSlot::operator<< (UObject*& Value)
+{
+	Ar.EnterSlot(Depth, ElementId);
+	Ar.Formatter.Serialize(Value);
+	Ar.LeaveSlot();
+}
+
+void FStructuredArchive::FSlot::operator<< (FText& Value)
+{
+	Ar.EnterSlot(Depth, ElementId);
+	Ar.Formatter.Serialize(Value);
+	Ar.LeaveSlot();
+}
+
+void FStructuredArchive::FSlot::operator<< (FWeakObjectPtr& Value)
+{
+	Ar.EnterSlot(Depth, ElementId);
+	Ar.Formatter.Serialize(Value);
+	Ar.LeaveSlot();
+}
+
+void FStructuredArchive::FSlot::operator<< (FLazyObjectPtr& Value)
+{
+	Ar.EnterSlot(Depth, ElementId);
+	Ar.Formatter.Serialize(Value);
+	Ar.LeaveSlot();
+}
+
+void FStructuredArchive::FSlot::operator<< (FSoftObjectPtr& Value)
+{
+	Ar.EnterSlot(Depth, ElementId);
+	Ar.Formatter.Serialize(Value);
+	Ar.LeaveSlot();
+}
+
+void FStructuredArchive::FSlot::operator<< (FSoftObjectPath& Value)
+{
+	Ar.EnterSlot(Depth, ElementId);
+	Ar.Formatter.Serialize(Value);
+	Ar.LeaveSlot();
+}
+
 void FStructuredArchive::FSlot::Serialize(TArray<uint8>& Data)
 {
 	Ar.EnterSlot(Depth, ElementId);
diff --git a/Engine/Source/Runtime/Core/Public/Concepts/Serializable.h b/Engine/Source/Runtime/Core/Public/Concepts/Serializable.h
deleted file mode 100644
diff --git a/Engine/Source/Runtime/Core/Public/Serialization/StructuredArchive.h b/Engine/Source/Runtime/Core/Public/Serialization/StructuredArchive.h
--- a/Engine/Source/Runtime/Core/Public/Serialization/StructuredArchive.h
+++ b/Engine/Source/Runtime/Core/Public/Serialization/StructuredArchive.h
@@ -7,7 +7,6 @@
 #include "Formatters/BinaryArchiveFormatter.h"
 #include "Misc/Optional.h"
 #include "Concepts/Insertable.h"
-#include "Concepts/Serializable.h"
 #include "Templates/Models.h"
 #include "Containers/Array.h"
 #include "Serialization/ArchiveProxy.h"
@@ -192,7 +191,7 @@ private:
 
 public:
 	/**
-	 * Contains a value in the archive; either a field or array/map element. A slot does not know its name or location,
+	 * Contains a value in the archive; either a field or array/map element. A slot does not know it's name or location,
 	 * and can merely have a value serialized into it. That value may be a literal (eg. int, float) or compound object
 	 * (eg. object, array, map).
 	 */
@@ -209,19 +208,27 @@ public:
 		TOptional<FSlot> TryEnterAttribute(FArchiveFieldName AttributeName, bool bEnterWhenWriting);
 
 		// We don't support chaining writes to a single slot, so this returns void.
-		template <typename ArgType>
-		typename TEnableIf<
-			TModels<CSerializable<FArchiveFormatterType>, ArgType>::Value
-		>::Type operator<<(ArgType&& Arg)
-		{
-#if WITH_TEXT_ARCHIVE_SUPPORT
-			Ar.EnterSlot(Depth, ElementId);
-			Ar.Formatter.Serialize(Forward<ArgType>(Arg));
-			Ar.LeaveSlot();
-#else
-			Ar.Formatter.Serialize(Forward<ArgType>(Arg));
-#endif
-		}
+		void operator << (char& Value);
+		void operator << (uint8& Value);
+		void operator << (uint16& Value);
+		void operator << (uint32& Value);
+		void operator << (uint64& Value);
+		void operator << (int8& Value);
+		void operator << (int16& Value);
+		void operator << (int32& Value);
+		void operator << (int64& Value);
+		void operator << (float& Value);
+		void operator << (double& Value);
+		void operator << (bool& Value);
+		void operator << (FString& Value);
+		void operator << (FName& Value);
+		void operator << (UObject*& Value);
+		void operator << (FText& Value);
+		void operator << (FWeakObjectPtr& Value);
+		void operator << (FSoftObjectPtr& Value);
+		void operator << (FSoftObjectPath& Value);
+		void operator << (FLazyObjectPtr& Value);
+
 
 		template <typename T>
 		FORCEINLINE void operator<<(TNamedAttribute<T> Item)
@@ -703,6 +710,108 @@ typename TEnableIf<
 		return FMap(Ar);
 	}
 
+	FORCEINLINE void FStructuredArchive::FSlot::operator<< (char& Value)
+	{
+		int8 AsInt = Value;
+		Ar.Formatter.Serialize(AsInt);
+		Value = AsInt;
+	}
+
+	FORCEINLINE void FStructuredArchive::FSlot::operator<< (uint8& Value)
+	{
+		Ar.Formatter.Serialize(Value);
+	}
+
+	FORCEINLINE void FStructuredArchive::FSlot::operator<< (uint16& Value)
+	{
+		Ar.Formatter.Serialize(Value);
+	}
+
+	FORCEINLINE void FStructuredArchive::FSlot::operator<< (uint32& Value)
+	{
+		Ar.Formatter.Serialize(Value);
+	}
+
+	FORCEINLINE void FStructuredArchive::FSlot::operator<< (uint64& Value)
+	{
+		Ar.Formatter.Serialize(Value);
+	}
+
+	FORCEINLINE void FStructuredArchive::FSlot::operator<< (int8& Value)
+	{
+		Ar.Formatter.Serialize(Value);
+	}
+
+	FORCEINLINE void FStructuredArchive::FSlot::operator<< (int16& Value)
+	{
+		Ar.Formatter.Serialize(Value);
+	}
+
+	FORCEINLINE void FStructuredArchive::FSlot::operator<< (int32& Value)
+	{
+		Ar.Formatter.Serialize(Value);
+	}
+
+	FORCEINLINE void FStructuredArchive::FSlot::operator<< (int64& Value)
+	{
+		Ar.Formatter.Serialize(Value);
+	}
+
+	FORCEINLINE void FStructuredArchive::FSlot::operator<< (float& Value)
+	{
+		Ar.Formatter.Serialize(Value);
+	}
+
+	FORCEINLINE void FStructuredArchive::FSlot::operator<< (double& Value)
+	{
+		Ar.Formatter.Serialize(Value);
+	}
+
+	FORCEINLINE void FStructuredArchive::FSlot::operator<< (bool& Value)
+	{
+		Ar.Formatter.Serialize(Value);
+	}
+
+	FORCEINLINE void FStructuredArchive::FSlot::operator<< (FString& Value)
+	{
+		Ar.Formatter.Serialize(Value);
+	}
+
+	FORCEINLINE void FStructuredArchive::FSlot::operator<< (FName& Value)
+	{
+		Ar.Formatter.Serialize(Value);
+	}
+
+	FORCEINLINE void FStructuredArchive::FSlot::operator<< (UObject*& Value)
+	{
+		Ar.Formatter.Serialize(Value);
+	}
+
+	FORCEINLINE void FStructuredArchive::FSlot::operator<< (FText& Value)
+	{
+		Ar.Formatter.Serialize(Value);
+	}
+
+	FORCEINLINE void FStructuredArchive::FSlot::operator<< (FWeakObjectPtr& Value)
+	{
+		Ar.Formatter.Serialize(Value);
+	}
+
+	FORCEINLINE void FStructuredArchive::FSlot::operator<< (FSoftObjectPath& Value)
+	{
+		Ar.Formatter.Serialize(Value);
+	}
+
+	FORCEINLINE void FStructuredArchive::FSlot::operator<< (FSoftObjectPtr& Value)
+	{
+		Ar.Formatter.Serialize(Value);
+	}
+
+	FORCEINLINE void FStructuredArchive::FSlot::operator<< (FLazyObjectPtr& Value)
+	{
+		Ar.Formatter.Serialize(Value);
+	}
+
 	FORCEINLINE void FStructuredArchive::FSlot::Serialize(TArray<uint8>& Data)
 	{
 		Ar.Formatter.Serialize(Data);
