Optimized FName(NAME_None) and implemented missing but unused FScriptName(EName) and FMinimalName(EName)

Avoid doing function call and EName -> FNameEntryId lookup when constructing an FName with NAME_None. In practice code could just use FName default constructor to achieve the same thing but I'd rather avoid touching call sites.

Makes tagged property loading about 10% faster for an internal project.

diff --git a/Engine/Source/Runtime/Core/Private/UObject/UnrealNames.cpp b/Engine/Source/Runtime/Core/Private/UObject/UnrealNames.cpp
--- a/Engine/Source/Runtime/Core/Private/UObject/UnrealNames.cpp
+++ b/Engine/Source/Runtime/Core/Private/UObject/UnrealNames.cpp
@@ -1696,20 +1696,6 @@ FName::FName(const FNameEntrySerialized& LoadedEntry)
 	: FName(FNameHelper::MakeFromLoaded(LoadedEntry))
 {}
 
-FName::FName(EName Ename, int32 InNumber)
-	: ComparisonIndex(GetNamePool().Find(Ename))
-#if WITH_CASE_PRESERVING_NAME
-	, DisplayIndex(ComparisonIndex)
-#endif
-	, Number(InNumber)
-{
-	check(Ename < NAME_MaxHardcodedNameIndex);
-}
-
-FName::FName(EName Ename)
-	: FName(Ename, NAME_NO_NUMBER_INTERNAL)
-{}
-
 bool FName::operator==(const ANSICHAR* Str) const
 {
 	return FNameHelper::EqualsString(*this, Str);
@@ -2186,6 +2172,11 @@ FArchive& operator<<(FArchive& Ar, FNameEntrySerialized& E)
 	return Ar;
 }
 
+FNameEntryId FNameEntryId::FromValidEName(EName Ename)
+{
+	return GetNamePool().Find(Ename);
+}
+
 #if !UE_BUILD_SHIPPING && !UE_BUILD_TEST
 
 #include "Containers/StackTracker.h"
diff --git a/Engine/Source/Runtime/Core/Public/UObject/NameTypes.h b/Engine/Source/Runtime/Core/Public/UObject/NameTypes.h
--- a/Engine/Source/Runtime/Core/Public/UObject/NameTypes.h
+++ b/Engine/Source/Runtime/Core/Public/UObject/NameTypes.h
@@ -66,8 +66,16 @@ struct FNameEntryId
 
 	/** Create from unstable int produced by this process */
 	CORE_API static FNameEntryId FromUnstableInt(uint32 UnstableInt);
+
+	FORCEINLINE static FNameEntryId FromEName(EName Ename)
+	{
+		return Ename == NAME_None ? FNameEntryId() : FromValidEName(Ename);
+	}
+
 private:
 	uint32 Value;
+
+	CORE_API static FNameEntryId FromValidEName(EName Ename);
 };
 
 CORE_API uint32 GetTypeHash(FNameEntryId Id);
@@ -314,10 +322,13 @@ struct FNameEntrySerialized
  * This is smaller than FName, but you lose the case-preserving behavior
  */
 struct FMinimalName
-	{
+{
 	FMinimalName() {}
-
-	CORE_API FMinimalName(EName N);
+	
+	FMinimalName(EName N)
+		: Index(FNameEntryId::FromEName(N))
+	{
+	}
 
 	FMinimalName(FNameEntryId InIndex, int32 InNumber)
 		: Index(InIndex)
@@ -344,8 +355,12 @@ struct FMinimalName
 struct FScriptName
 {
 	FScriptName() {}
-
-	CORE_API FScriptName(EName N);
+	
+	FScriptName(EName Ename)
+		: ComparisonIndex(FNameEntryId::FromEName(Ename))
+		, DisplayIndex(ComparisonIndex)
+	{
+	}
 
 	FScriptName(FNameEntryId InComparisonIndex, FNameEntryId InDisplayIndex, int32 InNumber)
 		: ComparisonIndex(InComparisonIndex)
@@ -619,7 +634,7 @@ public:
 	 *
 	 * @param N The hardcoded value the string portion of the name will have. The number portion will be NAME_NO_NUMBER
 	 */
-	FName(EName N);
+	FORCEINLINE FName(EName Ename) : FName(Ename, NAME_NO_NUMBER_INTERNAL) {}
 
 	/**
 	 * Create an FName with a hardcoded string index and (instance).
@@ -627,7 +642,15 @@ public:
 	 * @param N The hardcoded value the string portion of the name will have
 	 * @param InNumber The hardcoded value for the number portion of the name
 	 */
-	FName(EName N, int32 InNumber);
+	FORCEINLINE FName(EName Ename, int32 InNumber)
+		: ComparisonIndex(FNameEntryId::FromEName(Ename))
+#if WITH_CASE_PRESERVING_NAME
+		, DisplayIndex(ComparisonIndex)
+#endif
+		, Number(InNumber)
+	{
+	}
+
 
 	/**
 	 * Create an FName from an existing string, but with a different instance.
@@ -918,7 +941,7 @@ typedef FNameFastLess FNameSortIndexes;
 
 /** Slow alphabetical order that is stable / deterministic over process runs */
 struct FNameLexicalLess
-	{
+{
 	FORCEINLINE bool operator()(const FName& A, const FName& B) const
 	{
 		return A.Compare(B) < 0;
@@ -938,7 +961,7 @@ inline void FNameEntry::Decode(WIDECHAR*, uint32) {}
 #endif
 
 struct FNameDebugVisualizer
-	{
+{
 	CORE_API static uint8** GetBlocks();
 private:
 	static constexpr uint32 EntryStride = alignof(FNameEntry);
