We now only make a single call to ::SwitchToThread (give other threads in the current process a chance to run) when attempting to yeild the thread rather than also calling an additional ::Sleep(0) (give threads in any process a chance to run)

This was shown to give small wins during loading of 0.1->0.5(s) depending on the number of cores and the amount of thread contension.

diff --git a/Engine/Source/Runtime/Core/Private/Windows/WindowsPlatformProcess.cpp b/Engine/Source/Runtime/Core/Private/Windows/WindowsPlatformProcess.cpp
--- a/Engine/Source/Runtime/Core/Private/Windows/WindowsPlatformProcess.cpp
+++ b/Engine/Source/Runtime/Core/Private/Windows/WindowsPlatformProcess.cpp
@@ -1230,7 +1230,10 @@ void FWindowsPlatformProcess::SleepNoStats(float Seconds)
 	{
 		::SwitchToThread();
 	}
-	::Sleep(Milliseconds);
+	else
+	{
+		::Sleep(Milliseconds);
+	}
 }
 
 void FWindowsPlatformProcess::SleepInfinite()
