* Buildfix for improper string formatting in call to LogFormattedMessageWithCallstack.

diff --git a/Engine/Source/Runtime/Core/Private/HAL/MallocDoubleFreeFinder.cpp b/Engine/Source/Runtime/Core/Private/HAL/MallocDoubleFreeFinder.cpp
--- a/Engine/Source/Runtime/Core/Private/HAL/MallocDoubleFreeFinder.cpp
+++ b/Engine/Source/Runtime/Core/Private/HAL/MallocDoubleFreeFinder.cpp
@@ -304,7 +304,7 @@ FORCENOINLINE void FMallocDoubleFreeFinder::DumpStackTraceToLog(int32 StackIndex
 	}
 	// Dump the error and flush the log.
 	// ELogVerbosity::Error to make sure it gets printed in log for convenience.
-	FDebug::LogFormattedMessageWithCallstack(LogOutputDevice.GetCategoryName(), __FILE__, __LINE__, L"FMallocDoubleFreeFinder::DumpStackTraceToLog", ANSI_TO_TCHAR(&StackTraceString[0]), ELogVerbosity::Error);
+	FDebug::LogFormattedMessageWithCallstack(LogOutputDevice.GetCategoryName(), __FILE__, __LINE__, TEXT("FMallocDoubleFreeFinder::DumpStackTraceToLog"), ANSI_TO_TCHAR(&StackTraceString[0]), ELogVerbosity::Error);
 	GLog->Flush();
 #endif
 }
