Additional compilation fixes for MallocDoubleFreeFinder.

diff --git a/Engine/Source/Runtime/Core/Private/HAL/MallocDoubleFreeFinder.cpp b/Engine/Source/Runtime/Core/Private/HAL/MallocDoubleFreeFinder.cpp
--- a/Engine/Source/Runtime/Core/Private/HAL/MallocDoubleFreeFinder.cpp
+++ b/Engine/Source/Runtime/Core/Private/HAL/MallocDoubleFreeFinder.cpp
@@ -5,6 +5,7 @@
 =============================================================================*/
 #include "HAL/MallocDoubleFreeFinder.h"
 #include "Logging/LogMacros.h"
+#include "Misc/Parse.h"""
 #include "Misc/ScopeLock.h"
 //#include "Templates/Function.h"
 #include "GenericPlatform/GenericPlatformProcess.h"
diff --git a/Engine/Source/Runtime/Core/Public/HAL/MallocDoubleFreeFinder.h b/Engine/Source/Runtime/Core/Public/HAL/MallocDoubleFreeFinder.h
--- a/Engine/Source/Runtime/Core/Public/HAL/MallocDoubleFreeFinder.h
+++ b/Engine/Source/Runtime/Core/Public/HAL/MallocDoubleFreeFinder.h
@@ -3,6 +3,8 @@
 #pragma once
 
 #include "CoreTypes.h"
+#include "HAL/CriticalSection.h"
+#include "HAL/PlatformTLS.h"
 #include "HAL/MemoryBase.h"
 #include "Containers/Array.h"
 
