Sheesh ... final compilation fixes for MallocDoubleFreeFinder I hope.  Not sure how those pesky extra characters made it in the include.  Also updating copyright to reflect what the new standard appears to be (i.e. not specifying any date range).

diff --git a/Engine/Source/Runtime/Core/Private/HAL/MallocDoubleFreeFinder.cpp b/Engine/Source/Runtime/Core/Private/HAL/MallocDoubleFreeFinder.cpp
--- a/Engine/Source/Runtime/Core/Private/HAL/MallocDoubleFreeFinder.cpp
+++ b/Engine/Source/Runtime/Core/Private/HAL/MallocDoubleFreeFinder.cpp
@@ -1,11 +1,11 @@
-// Copyright 1998-2020 Epic Games, Inc. All Rights Reserved.
+// Copyright Epic Games, Inc. All Rights Reserved.
 
 /*=============================================================================
 	FMallocDoubleFreeFinder.cpp: Memoory tracking allocator
 =============================================================================*/
 #include "HAL/MallocDoubleFreeFinder.h"
 #include "Logging/LogMacros.h"
-#include "Misc/Parse.h"""
+#include "Misc/Parse.h"
 #include "Misc/ScopeLock.h"
 //#include "Templates/Function.h"
 #include "GenericPlatform/GenericPlatformProcess.h"
diff --git a/Engine/Source/Runtime/Core/Public/HAL/MallocDoubleFreeFinder.h b/Engine/Source/Runtime/Core/Public/HAL/MallocDoubleFreeFinder.h
--- a/Engine/Source/Runtime/Core/Public/HAL/MallocDoubleFreeFinder.h
+++ b/Engine/Source/Runtime/Core/Public/HAL/MallocDoubleFreeFinder.h
@@ -1,4 +1,4 @@
-// Copyright 1998-2020 Epic Games, Inc. All Rights Reserved.
+// Copyright Epic Games, Inc. All Rights Reserved.
 
 #pragma once
 
