Another missing include in MallocDoubleFreeFinder.h causing headaches for non-unity builds.

diff --git a/Engine/Source/Runtime/Core/Public/HAL/MallocDoubleFreeFinder.h b/Engine/Source/Runtime/Core/Public/HAL/MallocDoubleFreeFinder.h
--- a/Engine/Source/Runtime/Core/Public/HAL/MallocDoubleFreeFinder.h
+++ b/Engine/Source/Runtime/Core/Public/HAL/MallocDoubleFreeFinder.h
@@ -7,6 +7,7 @@
 #include "HAL/PlatformTLS.h"
 #include "HAL/MemoryBase.h"
 #include "Containers/Array.h"
+#include "Containers/Map.h"
 
 class CORE_API FMallocDoubleFreeFinder final : public FMalloc
 {
