Yet another missing include in MallocDoubleFreeFinder.  Unfortunately UGS is faster than building locally right now

diff --git a/Engine/Source/Runtime/Core/Private/HAL/MallocDoubleFreeFinder.cpp b/Engine/Source/Runtime/Core/Private/HAL/MallocDoubleFreeFinder.cpp
--- a/Engine/Source/Runtime/Core/Private/HAL/MallocDoubleFreeFinder.cpp
+++ b/Engine/Source/Runtime/Core/Private/HAL/MallocDoubleFreeFinder.cpp
@@ -5,6 +5,7 @@
 =============================================================================*/
 #include "HAL/MallocDoubleFreeFinder.h"
 #include "Logging/LogMacros.h"
+#include "Misc/OutputDeviceRedirector.h"
 #include "Misc/Parse.h"
 #include "Misc/ScopeLock.h"
 //#include "Templates/Function.h"
