Added some "const" to FArchiveState member functions

diff --git a/Engine/Source/Runtime/Core/Public/Serialization/Archive.h b/Engine/Source/Runtime/Core/Public/Serialization/Archive.h
--- a/Engine/Source/Runtime/Core/Public/Serialization/Archive.h
+++ b/Engine/Source/Runtime/Core/Public/Serialization/Archive.h
@@ -413,7 +413,7 @@ public:
 	 *
 	 * @return true if the archive filters editor-only content, false otherwise.
 	 */
-	bool IsFilterEditorOnly()
+	bool IsFilterEditorOnly() const
 	{
 		return ArIsFilterEditorOnly;
 	}
@@ -434,7 +434,7 @@ public:
 	 * @note This is intended for game-specific archives and is not true for any of the build in save methods
 	 * @return true if the archive is dealing with save games, false otherwise.
 	 */
-	bool IsSaveGame()
+	bool IsSaveGame() const
 	{
 		return ArIsSaveGame;
 	}
@@ -442,7 +442,7 @@ public:
 	/**
 	 * Whether or not this archive is serializing data being sent/received by the netcode
 	 */
-	FORCEINLINE bool IsNetArchive()
+	FORCEINLINE bool IsNetArchive() const
 	{
 		return ArIsNetArchive;
 	}
