Always align file reader and writer buffer sizes to a power of two to take best advantage of hardware I/O.

diff --git a/Engine/Source/Runtime/Core/Private/HAL/FileManagerGeneric.cpp b/Engine/Source/Runtime/Core/Private/HAL/FileManagerGeneric.cpp
--- a/Engine/Source/Runtime/Core/Private/HAL/FileManagerGeneric.cpp
+++ b/Engine/Source/Runtime/Core/Private/HAL/FileManagerGeneric.cpp
@@ -660,7 +660,8 @@ FArchiveFileReaderGeneric::FArchiveFileReaderGeneric( IFileHandle* InHandle, con
 	, BufferBase( 0 )
 	, Handle( InHandle )
 {
-	BufferArray.Reserve(InBufferSize);
+	BufferSize = FMath::RoundUpToPowerOfTwo64((int64)InBufferSize);
+	BufferArray.Reserve(BufferSize);
 	this->SetIsLoading(true);
 	this->SetIsPersistent(true);
 }
@@ -722,7 +723,7 @@ bool FArchiveFileReaderGeneric::InternalPrecache( int64 PrecacheOffset, int64 Pr
 	if( Pos == PrecacheOffset &&( !BufferBase || !BufferArray.Num() || BufferBase != Pos ) )
 	{
 		BufferBase = Pos;
-		int64 BufferCount = FMath::Min( FMath::Min( PrecacheSize,( int64 )(BufferArray.Max() -( Pos&( BufferArray.Max() -1 ) ) ) ), Size-Pos );
+		int64 BufferCount = FMath::Min( FMath::Min( PrecacheSize,( int64 )(BufferSize -( Pos&(BufferSize -1 ) ) ) ), Size-Pos );
 		BufferCount = FMath::Max(BufferCount, 0LL ); // clamp to 0
 		
 		const bool AllowShrink = false;
@@ -731,7 +732,7 @@ bool FArchiveFileReaderGeneric::InternalPrecache( int64 PrecacheOffset, int64 Pr
 		int64 Count = 0;
 
 		{
-			if (BufferCount > BufferArray.Max() || BufferCount <= 0)
+			if (BufferCount > BufferSize || BufferCount <= 0)
 			{
 				UE_LOG( LogFileManager, Error, TEXT("Invalid BufferCount=%lld while reading %s. File is most likely corrupted. Please verify your installation. Pos=%lld, Size=%lld, PrecacheSize=%lld, PrecacheOffset=%lld"),
 					BufferCount, *Filename, Pos, Size, PrecacheSize, PrecacheOffset );
@@ -766,7 +767,7 @@ void FArchiveFileReaderGeneric::Serialize( void* V, int64 Length )
 		int64 Copy = FMath::Min( Length, BufferBase+BufferArray.Num()-Pos );
 		if( Copy<=0 )
 		{
-			if( Length >= BufferArray.Max() )
+			if( Length >= BufferSize )
 			{
 				int64 Count=0;
 				{
@@ -814,7 +815,8 @@ FArchiveFileWriterGeneric::FArchiveFileWriterGeneric( IFileHandle* InHandle, con
 	, Handle( InHandle )
 	, bLoggingError( false )
 {
-	BufferArray.Reserve(InBufferSize);
+	BufferSize = FMath::RoundUpToPowerOfTwo64((int64)InBufferSize);
+	BufferArray.Reserve(BufferSize);
 	this->SetIsSaving(true);
 	this->SetIsPersistent(true);
 }
@@ -872,7 +874,7 @@ bool FArchiveFileWriterGeneric::Close()
 void FArchiveFileWriterGeneric::Serialize( void* V, int64 Length )
 {
 	Pos += Length;
-	if ( Length >= BufferArray.Max() )
+	if ( Length >= BufferSize )
 	{
 		FlushBuffer();
 		if( !WriteLowLevel( (uint8*)V, Length ) )
@@ -884,7 +886,7 @@ void FArchiveFileWriterGeneric::Serialize( void* V, int64 Length )
 	else
 	{
 		int64 Copy;
-		while( Length >( Copy=BufferArray.Max()-BufferArray.Num() ) )
+		while( Length >( Copy=BufferSize-BufferArray.Num() ) )
 		{
 			BufferArray.Append((uint8*)V, Copy);
 			Length      -= Copy;
diff --git a/Engine/Source/Runtime/Core/Public/HAL/FileManagerGeneric.h b/Engine/Source/Runtime/Core/Public/HAL/FileManagerGeneric.h
--- a/Engine/Source/Runtime/Core/Public/HAL/FileManagerGeneric.h
+++ b/Engine/Source/Runtime/Core/Public/HAL/FileManagerGeneric.h
@@ -268,6 +268,7 @@ protected:
 	int64 BufferBase;
 	TUniquePtr<IFileHandle> Handle;
 	TArray64<uint8> BufferArray;
+	int64 BufferSize;
 };
 
 
@@ -342,5 +343,6 @@ protected:
 	int64 Pos;
 	TUniquePtr<IFileHandle> Handle;
 	TArray64<uint8> BufferArray;
+	int64 BufferSize;
 	bool bLoggingError;
 };
