All programs should dump callstacks

diff --git a/Engine/Source/Runtime/Core/Public/GenericPlatform/GenericPlatformProperties.h b/Engine/Source/Runtime/Core/Public/GenericPlatform/GenericPlatformProperties.h
--- a/Engine/Source/Runtime/Core/Public/GenericPlatform/GenericPlatformProperties.h
+++ b/Engine/Source/Runtime/Core/Public/GenericPlatform/GenericPlatformProperties.h
@@ -288,7 +288,7 @@ struct FGenericPlatformProperties
 	// Whether the platform allows the call stack to be dumped during an assert
 	static FORCEINLINE bool AllowsCallStackDumpDuringAssert()
 	{
-		return false;
+		return IsProgram();
 	}
 
 	// If this platform wants to replace Zlib with a platform-specific version, set the name of the compression format 
