Added FName length bounds-checking on top of existing checks to guard against malicious input in shipping builds

diff --git a/Engine/Source/Runtime/Core/Private/UObject/UnrealNames.cpp b/Engine/Source/Runtime/Core/Private/UObject/UnrealNames.cpp
--- a/Engine/Source/Runtime/Core/Private/UObject/UnrealNames.cpp
+++ b/Engine/Source/Runtime/Core/Private/UObject/UnrealNames.cpp
@@ -1789,7 +1789,7 @@ struct FNameHelper
 		{
 			// Consider _mm_packus_epi16 or similar if this proves too slow
 			ANSICHAR AnsiName[NAME_SIZE];
-			for (int32 I = 0; I < View.Len; ++I)
+			for (int32 I = 0, Len = FMath::Min<int32>(View.Len, NAME_SIZE); I < Len; ++I)
 			{
 				AnsiName[I] = View.Str[I];
 			}
@@ -1803,8 +1803,12 @@ struct FNameHelper
 
 	static FName Make(FNameStringView View, EFindName FindType, int32 InternalNumber)
 	{
-		checkf(View.Len < NAME_SIZE, TEXT("FName's %d max length exceeded. Got %d characters excluding null-terminator."), NAME_SIZE - 1, View.Len);
-
+		if (View.Len >= NAME_SIZE)
+		{
+			checkf(false, TEXT("FName's %d max length exceeded. Got %d characters excluding null-terminator."), NAME_SIZE - 1, View.Len);
+			return FName("ERROR_NAME_SIZE_EXCEEDED");
+		}
+		
 		FNamePool& Pool = GetNamePool();
 
 		FNameEntryId DisplayId, ComparisonId;
