Added framepro.lScopeMinTimeMicroseconds to control the minimum recorded scope duration per platform.

framepro.ScopeMinTimeMicroseconds only applies when using framepro.startrec (not when captuing externally with the framepro ui or starting a capture through code).
Also got rid of framepro.startrecscopeoverride.

Set the default min scope to 25us.
Switch is set to 50us.
PC and Mac are set to 10us.

diff --git a/Engine/Config/Mac/MacEngine.ini b/Engine/Config/Mac/MacEngine.ini
--- a/Engine/Config/Mac/MacEngine.ini
+++ b/Engine/Config/Mac/MacEngine.ini
@@ -11,6 +11,7 @@ UseAudioMixer=true
 [SystemSettings]
 r.setres=1280x720
 g.TimeoutForBlockOnRenderFence=300000
+framepro.ScopeMinTimeMicroseconds=10
 
 [SystemSettingsEditor]
 ; Parallel rendering has not been tested in the editor so just disable it.
diff --git a/Engine/Config/Windows/WindowsEngine.ini b/Engine/Config/Windows/WindowsEngine.ini
--- a/Engine/Config/Windows/WindowsEngine.ini
+++ b/Engine/Config/Windows/WindowsEngine.ini
@@ -16,6 +16,7 @@ ConfigClass=/Engine/DeviceProfiles/SetDeviceProfileWindows.SetDeviceProfileWindo
 
 [SystemSettings]
 r.setres=1280x720
+framepro.ScopeMinTimeMicroseconds=10
 
 [PlatformCrypto]
 PlatformRequiresDataCrypto=True
diff --git a/Engine/Source/Runtime/Core/Private/FramePro/FrameProProfiler.cpp b/Engine/Source/Runtime/Core/Private/FramePro/FrameProProfiler.cpp
--- a/Engine/Source/Runtime/Core/Private/FramePro/FrameProProfiler.cpp
+++ b/Engine/Source/Runtime/Core/Private/FramePro/FrameProProfiler.cpp
@@ -333,6 +333,14 @@ void FFrameProProfiler::PopEvent(const ANSICHAR* Override)
 	}
 }
 
+static int32 ScopeMinTimeMicroseconds = 25;
+static FAutoConsoleVariableRef CVarScopeMinTimeMicroseconds(
+	TEXT("framepro.ScopeMinTimeMicroseconds"),
+	ScopeMinTimeMicroseconds,
+	TEXT("Scopes with time taken below this threshold are not recorded in the FramePro capture.\n")
+	TEXT(" This value is only used when starting framepro captures with framepro.startrec.")
+);
+
 void FFrameProProfiler::StartFrameProRecordingFromCommand(const TArray< FString >& Args)
 {
 	FString FilenameRoot = FString::Printf(TEXT("ProfilePid%d"), FPlatformProcess::GetCurrentProcessId());
@@ -341,18 +349,7 @@ void FFrameProProfiler::StartFrameProRecordingFromCommand(const TArray< FString
 		FilenameRoot = Args[0];
 	}
 
-	StartFrameProRecording(FilenameRoot, 50);
-}
-
-void FFrameProProfiler::StartFrameProRecordingScopeOverrideFromCommand(const TArray< FString >& Args)
-{
-	int32 MinScopeTime = 50;
-	if (Args.Num() > 0 && Args[0].Len() > 0)
-	{
-		MinScopeTime = FCString::Atoi(*Args[0]);
-	}
-
-	StartFrameProRecording(FString::Printf(TEXT("ProfilePid%d"), FPlatformProcess::GetCurrentProcessId()), MinScopeTime);
+	StartFrameProRecording(FilenameRoot, ScopeMinTimeMicroseconds);
 }
 
 FString FFrameProProfiler::StartFrameProRecording(const FString& FilenameRoot, int32 MinScopeTime)
@@ -391,12 +388,6 @@ static FAutoConsoleCommand StartFrameProRecordCommand(
 	FConsoleCommandWithArgsDelegate::CreateStatic(&FFrameProProfiler::StartFrameProRecordingFromCommand)
 );
 
-static FAutoConsoleCommand StartFrameProRecordScopeOverrideCommand(
-	TEXT("framepro.startrecscopeoverride"),
-	TEXT("Start FramePro recording with a minimum event scope override"),
-	FConsoleCommandWithArgsDelegate::CreateStatic(&FFrameProProfiler::StartFrameProRecordingScopeOverrideFromCommand)
-);
-
 void FFrameProProfiler::StopFrameProRecording()
 {
 	FramePro::StopRecording();
diff --git a/Engine/Source/Runtime/Core/Public/FramePro/FrameProProfiler.h b/Engine/Source/Runtime/Core/Public/FramePro/FrameProProfiler.h
--- a/Engine/Source/Runtime/Core/Public/FramePro/FrameProProfiler.h
+++ b/Engine/Source/Runtime/Core/Public/FramePro/FrameProProfiler.h
@@ -24,7 +24,6 @@ public:
 	static void PopEvent(const ANSICHAR* Override);
 
 	static void StartFrameProRecordingFromCommand(const TArray< FString >& Args);
-	static void StartFrameProRecordingScopeOverrideFromCommand(const TArray< FString >& Args);
 	static FString StartFrameProRecording(const FString& FilenameRoot, int32 MinScopeTime);
 	static void StopFrameProRecording();
 
