Fix crash on startup when running an unpackaged project from a different drive than the one containing the engine.

#rnx

diff --git a/Engine/Source/Runtime/Core/Private/GenericPlatform/GenericPlatformMisc.cpp b/Engine/Source/Runtime/Core/Private/GenericPlatform/GenericPlatformMisc.cpp
--- a/Engine/Source/Runtime/Core/Private/GenericPlatform/GenericPlatformMisc.cpp
+++ b/Engine/Source/Runtime/Core/Private/GenericPlatform/GenericPlatformMisc.cpp
@@ -831,7 +831,7 @@ const TCHAR* FGenericPlatformMisc::EngineDir()
 		{
 			EngineDirectory = DefaultEngineDir;
 		}
-		else if (GForeignEngineDir != NULL && FPlatformFileManager::Get().GetPlatformFile().DirectoryExists(*(FPlatformProcess::BaseDir() / FString(GForeignEngineDir) / TEXT("Binaries"))))
+		else if (GForeignEngineDir != NULL && FPlatformFileManager::Get().GetPlatformFile().DirectoryExists(*(FString(GForeignEngineDir) / TEXT("Binaries"))))
 		{
 			EngineDirectory = GForeignEngineDir;
 		}
