FSlowTask improvements
* ShouldCancel now ticks the UI from time to time (only on the game thread, max 5 times per second) so that the Cancel button interaction can be checked without artificially calling EnterProgressFrame
* Added checks to make sure that EnterProgressFrame and ShouldCancel are only called from the game thread

diff --git a/Engine/Source/Runtime/Core/Private/Misc/SlowTask.cpp b/Engine/Source/Runtime/Core/Private/Misc/SlowTask.cpp
--- a/Engine/Source/Runtime/Core/Private/Misc/SlowTask.cpp
+++ b/Engine/Source/Runtime/Core/Private/Misc/SlowTask.cpp
@@ -100,6 +100,8 @@ void FSlowTask::MakeDialogDelayed(float Threshold, bool bShowCancelButton, bool
 
 void FSlowTask::EnterProgressFrame(float ExpectedWorkThisFrame, FText Text)
 {
+	check(!bEnabled || IsInGameThread());
+
 	FrameMessage = Text;
 	CompletedWork += CurrentFrameScope;
 
@@ -148,5 +150,14 @@ void FSlowTask::MakeDialog(bool bShowCancelButton, bool bAllowInPIE)
 
 bool FSlowTask::ShouldCancel() const
 {
-	return Context.ReceivedUserCancel();
+	if (bEnabled)
+	{
+		check(IsInGameThread()); // FSlowTask is only meant to be used on the main thread currently
+
+		// update the UI from time to time (throttling is done in RequestUpdateUI) so that the cancel button interaction can be detected: 
+		Context.RequestUpdateUI();
+
+		return Context.ReceivedUserCancel();
+	}
+	return false;
 }
