Added LLM tags for navigation recast/detour

#rnx

diff --git a/Engine/Source/Runtime/Core/Private/HAL/LowLevelMemTracker.cpp b/Engine/Source/Runtime/Core/Private/HAL/LowLevelMemTracker.cpp
--- a/Engine/Source/Runtime/Core/Private/HAL/LowLevelMemTracker.cpp
+++ b/Engine/Source/Runtime/Core/Private/HAL/LowLevelMemTracker.cpp
@@ -85,6 +85,7 @@ DECLARE_LLM_MEMORY_STAT(TEXT("Particles"), STAT_ParticlesLLM, STATGROUP_LLMFULL)
 DECLARE_LLM_MEMORY_STAT(TEXT("Niagara"), STAT_NiagaraLLM, STATGROUP_LLMFULL);
 DECLARE_LLM_MEMORY_STAT(TEXT("GC"), STAT_GCLLM, STATGROUP_LLMFULL);
 DECLARE_LLM_MEMORY_STAT(TEXT("UI"), STAT_UILLM, STATGROUP_LLMFULL);
+DECLARE_LLM_MEMORY_STAT(TEXT("NavigationRecast"), STAT_NavigationRecastLLM, STATGROUP_LLMFULL);
 DECLARE_LLM_MEMORY_STAT(TEXT("Physics"), STAT_PhysicsLLM, STATGROUP_LLMFULL);
 DECLARE_LLM_MEMORY_STAT(TEXT("PhysX"), STAT_PhysXLLM, STATGROUP_LLMFULL);
 DECLARE_LLM_MEMORY_STAT(TEXT("PhysXGeometry"), STAT_PhysXGeometryLLM, STATGROUP_LLMFULL);
@@ -139,6 +140,7 @@ DECLARE_LLM_MEMORY_STAT(TEXT("Materials"), STAT_MaterialsSummaryLLM, STATGROUP_L
 DECLARE_LLM_MEMORY_STAT(TEXT("Particles"), STAT_ParticlesSummaryLLM, STATGROUP_LLM);
 DECLARE_LLM_MEMORY_STAT(TEXT("Niagara"), STAT_NiagaraSummaryLLM, STATGROUP_LLM);
 DECLARE_LLM_MEMORY_STAT(TEXT("UI"), STAT_UISummaryLLM, STATGROUP_LLM);
+DECLARE_LLM_MEMORY_STAT(TEXT("Navigation"), STAT_NavigationSummaryLLM, STATGROUP_LLM);
 DECLARE_LLM_MEMORY_STAT(TEXT("Textures"), STAT_TexturesSummaryLLM, STATGROUP_LLM);
 DECLARE_LLM_MEMORY_STAT(TEXT("MediaStreaming"), STAT_MediaStreamingSummaryLLM, STATGROUP_LLM);
 
diff --git a/Engine/Source/Runtime/Core/Public/HAL/LowLevelMemTracker.h b/Engine/Source/Runtime/Core/Public/HAL/LowLevelMemTracker.h
--- a/Engine/Source/Runtime/Core/Public/HAL/LowLevelMemTracker.h
+++ b/Engine/Source/Runtime/Core/Public/HAL/LowLevelMemTracker.h
@@ -184,6 +184,7 @@ enum class ELLMTagSet : uint8
 	macro(Niagara,								"Niagara",						GET_STATFNAME(STAT_NiagaraLLM),								GET_STATFNAME(STAT_NiagaraSummaryLLM),			-1)\
 	macro(GC,									"GC",							GET_STATFNAME(STAT_GCLLM),									GET_STATFNAME(STAT_EngineSummaryLLM),			-1)\
 	macro(UI,									"UI",							GET_STATFNAME(STAT_UILLM),									GET_STATFNAME(STAT_UISummaryLLM),				-1)\
+	macro(NavigationRecast,						"NavigationRecast",				GET_STATFNAME(STAT_NavigationRecastLLM),					GET_STATFNAME(STAT_NavigationSummaryLLM),		-1)\
 	macro(Physics,								"Physics",						GET_STATFNAME(STAT_PhysicsLLM),								GET_STATFNAME(STAT_PhysicsSummaryLLM),			-1)\
 	macro(PhysX,								"PhysX",						GET_STATFNAME(STAT_PhysXLLM),								GET_STATFNAME(STAT_PhysXSummaryLLM),			ELLMTag::Physics)\
 	macro(PhysXGeometry,						"PhysXGeometry",				GET_STATFNAME(STAT_PhysXGeometryLLM),						GET_STATFNAME(STAT_PhysXSummaryLLM),			ELLMTag::Physics)\
diff --git a/Engine/Source/Runtime/NavigationSystem/Private/NavMesh/PImplRecastNavMesh.cpp b/Engine/Source/Runtime/NavigationSystem/Private/NavMesh/PImplRecastNavMesh.cpp
--- a/Engine/Source/Runtime/NavigationSystem/Private/NavMesh/PImplRecastNavMesh.cpp
+++ b/Engine/Source/Runtime/NavigationSystem/Private/NavMesh/PImplRecastNavMesh.cpp
@@ -44,6 +44,7 @@ static_assert(RECAST_UNWALKABLE_POLY_COST == DT_UNWALKABLE_POLY_COST, "Unwalkabl
 
 static void* DetourMalloc(int Size, dtAllocHint)
 {
+	LLM_SCOPE(ELLMTag::NavigationRecast);
 	void* Result = FMemory::Malloc(uint32(Size));
 #if STATS
 	const uint32 ActualSize = FMemory::GetAllocSize(Result);
@@ -55,6 +56,7 @@ static void* DetourMalloc(int Size, dtAllocHint)
 
 static void* RecastMalloc(int Size, rcAllocHint)
 {
+	LLM_SCOPE(ELLMTag::NavigationRecast);
 	void* Result = FMemory::Malloc(uint32(Size));
 #if STATS
 	const uint32 ActualSize = FMemory::GetAllocSize(Result);
@@ -66,6 +68,7 @@ static void* RecastMalloc(int Size, rcAllocHint)
 
 static void RecastFree( void* Original )
 {
+	LLM_SCOPE(ELLMTag::NavigationRecast);
 #if STATS
 	const uint32 Size = FMemory::GetAllocSize(Original);
 	DEC_DWORD_STAT_BY(STAT_NavigationMemory, Size);	
