Adding preliminary support for enabling the warning C4244 ('argument': conversion from 'type1' to 'type2', possible loss of data) on a module by module basis to help catch int64->int32 and double->float / float<->int confusion
- UnsafeTypeCastWarningLevel can be set to WarningLevel.Warning or WarningLevel.Error in ModuleName.Build.cs (off by default)
- Currently only supported on MS compilers (Clang ignores the setting for now)

diff --git a/Engine/Source/Programs/AutomationTool/Scripts/CheckBalancedMacros.cs b/Engine/Source/Programs/AutomationTool/Scripts/CheckBalancedMacros.cs
--- a/Engine/Source/Programs/AutomationTool/Scripts/CheckBalancedMacros.cs
+++ b/Engine/Source/Programs/AutomationTool/Scripts/CheckBalancedMacros.cs
@@ -38,6 +38,10 @@ namespace AutomationTool
 				"PRAGMA_DISABLE_SHADOW_VARIABLE_WARNINGS",
 				"PRAGMA_ENABLE_SHADOW_VARIABLE_WARNINGS"
 			},
+			{
+				"PRAGMA_DISABLE_UNSAFE_TYPECAST_WARNINGS",
+				"PRAGMA_ENABLE_UNSAFE_TYPECAST_WARNINGS"
+			},
 			{
 				"PRAGMA_DISABLE_UNDEFINED_IDENTIFIER_WARNINGS",
 				"PRAGMA_ENABLE_UNDEFINED_IDENTIFIER_WARNINGS"
diff --git a/Engine/Source/Programs/UnrealBuildTool/Configuration/ModuleRules.cs b/Engine/Source/Programs/UnrealBuildTool/Configuration/ModuleRules.cs
--- a/Engine/Source/Programs/UnrealBuildTool/Configuration/ModuleRules.cs
+++ b/Engine/Source/Programs/UnrealBuildTool/Configuration/ModuleRules.cs
@@ -606,6 +606,16 @@ namespace UnrealBuildTool
 			set { ShadowVariableWarningLevel = (value ? WarningLevel.Warning : WarningLevel.Off); }
 		}
 
+		/// <summary>
+		/// How to treat unsafe implicit type cast warnings (e.g., double->float or int64->int32)
+		/// </summary>
+		public WarningLevel UnsafeTypeCastWarningLevel
+		{
+			get { return UnsafeTypeCastWarningLevelPrivate ?? Target.UnsafeTypeCastWarningLevel; }
+			set { UnsafeTypeCastWarningLevelPrivate = value; }
+		}
+		private WarningLevel? UnsafeTypeCastWarningLevelPrivate;
+
 		/// <summary>
 		/// Enable warnings for using undefined identifiers in #if expressions
 		/// </summary>
diff --git a/Engine/Source/Programs/UnrealBuildTool/Configuration/TargetRules.cs b/Engine/Source/Programs/UnrealBuildTool/Configuration/TargetRules.cs
--- a/Engine/Source/Programs/UnrealBuildTool/Configuration/TargetRules.cs
+++ b/Engine/Source/Programs/UnrealBuildTool/Configuration/TargetRules.cs
@@ -873,6 +873,12 @@ namespace UnrealBuildTool
 			set { ShadowVariableWarningLevel = (value? WarningLevel.Error : WarningLevel.Warning); }
 		}
 
+		/// <summary>
+		/// Indicates what warning/error level to treat unsafe type casts as on platforms that support it (e.g., double->float or int64->int32)
+		/// </summary>
+		[XmlConfigFile(Category = "BuildConfiguration")]
+		public WarningLevel UnsafeTypeCastWarningLevel = WarningLevel.Off;
+
 		/// <summary>
 		/// Forces the use of undefined identifiers in conditional expressions to be treated as errors.
 		/// </summary>
@@ -2254,6 +2260,11 @@ namespace UnrealBuildTool
 			get { return Inner.ShadowVariableWarningLevel; }
 		}
 
+		public WarningLevel UnsafeTypeCastWarningLevel
+		{
+			get { return Inner.UnsafeTypeCastWarningLevel; }
+		}
+
 		public bool bUndefinedIdentifierErrors
 		{
 			get { return Inner.bUndefinedIdentifierErrors; }
diff --git a/Engine/Source/Programs/UnrealBuildTool/Configuration/UEBuildModuleCPP.cs b/Engine/Source/Programs/UnrealBuildTool/Configuration/UEBuildModuleCPP.cs
--- a/Engine/Source/Programs/UnrealBuildTool/Configuration/UEBuildModuleCPP.cs
+++ b/Engine/Source/Programs/UnrealBuildTool/Configuration/UEBuildModuleCPP.cs
@@ -693,6 +693,10 @@ namespace UnrealBuildTool
 			{
 				return false;
 			}
+			if(ModuleCompileEnvironment.UnsafeTypeCastWarningLevel != CompileEnvironment.UnsafeTypeCastWarningLevel)
+			{
+				return false;
+			}
 			if(ModuleCompileEnvironment.bEnableUndefinedIdentifierWarnings != CompileEnvironment.bEnableUndefinedIdentifierWarnings)
 			{
 				return false;
@@ -757,6 +761,23 @@ namespace UnrealBuildTool
 					Variant += ".NoShadow";
 				}
 			}
+
+			if (CompileEnvironment.UnsafeTypeCastWarningLevel != BaseCompileEnvironment.UnsafeTypeCastWarningLevel)
+			{
+				if (CompileEnvironment.UnsafeTypeCastWarningLevel == WarningLevel.Error)
+				{
+					Variant += ".TypeCastErrors";
+				}
+				else if (CompileEnvironment.UnsafeTypeCastWarningLevel == WarningLevel.Warning)
+				{
+					Variant += ".TypeCastWarnings";
+				}
+				else
+				{
+					Variant += ".NoTypeCast";
+				}
+			}
+			
 			if (CompileEnvironment.bEnableUndefinedIdentifierWarnings != BaseCompileEnvironment.bEnableUndefinedIdentifierWarnings)
 			{
 				if (CompileEnvironment.bEnableUndefinedIdentifierWarnings)
@@ -782,6 +803,7 @@ namespace UnrealBuildTool
 			CompileEnvironment.bUseRTTI = ModuleCompileEnvironment.bUseRTTI;
 			CompileEnvironment.bEnableExceptions = ModuleCompileEnvironment.bEnableExceptions;
 			CompileEnvironment.ShadowVariableWarningLevel = ModuleCompileEnvironment.ShadowVariableWarningLevel;
+			CompileEnvironment.UnsafeTypeCastWarningLevel = ModuleCompileEnvironment.UnsafeTypeCastWarningLevel;
 			CompileEnvironment.bEnableUndefinedIdentifierWarnings = ModuleCompileEnvironment.bEnableUndefinedIdentifierWarnings;
 		}
 
@@ -1157,6 +1179,7 @@ namespace UnrealBuildTool
 			Result.bEnableExceptions |= Rules.bEnableExceptions;
 			Result.bEnableObjCExceptions |= Rules.bEnableObjCExceptions;
 			Result.ShadowVariableWarningLevel = Rules.ShadowVariableWarningLevel;
+			Result.UnsafeTypeCastWarningLevel = Rules.UnsafeTypeCastWarningLevel;
 			Result.bEnableUndefinedIdentifierWarnings = Rules.bEnableUndefinedIdentifierWarnings;
 
 			// If the module overrides the C++ language version, override it on the compile environment
diff --git a/Engine/Source/Programs/UnrealBuildTool/Configuration/UEBuildTarget.cs b/Engine/Source/Programs/UnrealBuildTool/Configuration/UEBuildTarget.cs
--- a/Engine/Source/Programs/UnrealBuildTool/Configuration/UEBuildTarget.cs
+++ b/Engine/Source/Programs/UnrealBuildTool/Configuration/UEBuildTarget.cs
@@ -3340,6 +3340,7 @@ namespace UnrealBuildTool
 			GlobalCompileEnvironment.bEnableExceptions = Rules.bForceEnableExceptions || Rules.bBuildEditor;
 			GlobalCompileEnvironment.bEnableObjCExceptions = Rules.bForceEnableObjCExceptions || Rules.bBuildEditor;
 			GlobalCompileEnvironment.ShadowVariableWarningLevel = Rules.ShadowVariableWarningLevel;
+			GlobalCompileEnvironment.UnsafeTypeCastWarningLevel = Rules.UnsafeTypeCastWarningLevel;
 			GlobalCompileEnvironment.bUndefinedIdentifierWarningsAsErrors = Rules.bUndefinedIdentifierErrors;
 			GlobalCompileEnvironment.bOptimizeForSize = Rules.bCompileForSize;
 			GlobalCompileEnvironment.bUseStaticCRT = Rules.bUseStaticCRT;
diff --git a/Engine/Source/Programs/UnrealBuildTool/Platform/Windows/VCToolChain.cs b/Engine/Source/Programs/UnrealBuildTool/Platform/Windows/VCToolChain.cs
--- a/Engine/Source/Programs/UnrealBuildTool/Platform/Windows/VCToolChain.cs
+++ b/Engine/Source/Programs/UnrealBuildTool/Platform/Windows/VCToolChain.cs
@@ -490,8 +490,6 @@ namespace UnrealBuildTool
 				}
 
 				Arguments.Add("/wd4463"); // 4463 - overflow; assigning 1 to bit-field that can only hold values from -1 to 0
-
-				Arguments.Add("/wd4838"); // 4838: conversion from 'type1' to 'type2' requires a narrowing conversion
 			}
 
 			if(CompileEnvironment.bEnableUndefinedIdentifierWarnings)
@@ -505,6 +503,29 @@ namespace UnrealBuildTool
 					Arguments.Add("/w44668");
 				}
 			}
+
+			// The unsafe type cast warnings setting controls the following warnings currently:
+			//   4244: conversion from 'type1' to 'type2', possible loss of data
+			//   4838: conversion from 'type1' to 'type2' requires a narrowing conversion
+			//@TODO: FLOATPRECISION: Consider doing the following as well:
+			//   4267: 'var' : conversion from 'size_t' to 'type', possible loss of data
+			//   4305: 'identifier' : truncation from 'type1' to 'type2'
+			if (CompileEnvironment.UnsafeTypeCastWarningLevel == WarningLevel.Error)
+ 			{
+ 				Arguments.Add("/we4244");
+				Arguments.Add("/we4838");
+ 			}
+ 			else if (CompileEnvironment.UnsafeTypeCastWarningLevel == WarningLevel.Warning)
+ 			{
+				// Note: The extra 4 is not a typo, /wLXXXX sets warning XXXX to level L
+ 				Arguments.Add("/w44244");
+				Arguments.Add("/w44838");
+			}
+			else
+ 			{
+ 				Arguments.Add("/wd4244");
+				Arguments.Add("/wd4838");
+ 			}
 		}
 
 		void AppendCLArguments_CPP(CppCompileEnvironment CompileEnvironment, List<string> Arguments)
diff --git a/Engine/Source/Programs/UnrealBuildTool/System/CppCompileEnvironment.cs b/Engine/Source/Programs/UnrealBuildTool/System/CppCompileEnvironment.cs
--- a/Engine/Source/Programs/UnrealBuildTool/System/CppCompileEnvironment.cs
+++ b/Engine/Source/Programs/UnrealBuildTool/System/CppCompileEnvironment.cs
@@ -185,6 +185,11 @@ namespace UnrealBuildTool
 		/// </summary>
 		public WarningLevel ShadowVariableWarningLevel = WarningLevel.Warning;
 
+		/// <summary>
+		/// How to treat unsafe implicit type cast warnings (e.g., double->float or int64->int32)
+		/// </summary>
+		public WarningLevel UnsafeTypeCastWarningLevel = WarningLevel.Off;
+
 		/// <summary>
 		/// Whether to warn about the use of undefined identifiers in #if expressions
 		/// </summary>
@@ -400,6 +405,7 @@ namespace UnrealBuildTool
 			bEnableExceptions = Other.bEnableExceptions;
 			bEnableObjCExceptions = Other.bEnableObjCExceptions;
 			ShadowVariableWarningLevel = Other.ShadowVariableWarningLevel;
+			UnsafeTypeCastWarningLevel = Other.UnsafeTypeCastWarningLevel;
 			bUndefinedIdentifierWarningsAsErrors = Other.bUndefinedIdentifierWarningsAsErrors;
 			bEnableUndefinedIdentifierWarnings = Other.bEnableUndefinedIdentifierWarnings;
 			bOptimizeCode = Other.bOptimizeCode;
diff --git a/Engine/Source/Runtime/Core/Public/Clang/ClangPlatformCompilerPreSetup.h b/Engine/Source/Runtime/Core/Public/Clang/ClangPlatformCompilerPreSetup.h
--- a/Engine/Source/Runtime/Core/Public/Clang/ClangPlatformCompilerPreSetup.h
+++ b/Engine/Source/Runtime/Core/Public/Clang/ClangPlatformCompilerPreSetup.h
@@ -73,6 +73,16 @@
 		_Pragma("clang diagnostic pop")
 #endif // PRAGMA_ENABLE_SHADOW_VARIABLE_WARNINGS
 
+#ifndef PRAGMA_DISABLE_UNSAFE_TYPECAST_WARNINGS
+	#define PRAGMA_DISABLE_UNSAFE_TYPECAST_WARNINGS
+		//@TODO: FLOATPRECISION: Add Clang support for this kind of error
+#endif // PRAGMA_DISABLE_UNSAFE_TYPECAST_WARNINGS
+
+#ifndef PRAGMA_ENABLE_UNSAFE_TYPECAST_WARNINGS
+	#define PRAGMA_ENABLE_UNSAFE_TYPECAST_WARNINGS
+		//@TODO: FLOATPRECISION: Add Clang support for this kind of error
+#endif // PRAGMA_ENABLE_UNSAFE_TYPECAST_WARNINGS
+
 #ifndef PRAGMA_DISABLE_UNDEFINED_IDENTIFIER_WARNINGS
 	#define PRAGMA_DISABLE_UNDEFINED_IDENTIFIER_WARNINGS \
 		_Pragma("clang diagnostic push") \
diff --git a/Engine/Source/Runtime/Core/Public/HoloLens/HoloLensPlatformCompilerPreSetup.h b/Engine/Source/Runtime/Core/Public/HoloLens/HoloLensPlatformCompilerPreSetup.h
--- a/Engine/Source/Runtime/Core/Public/HoloLens/HoloLensPlatformCompilerPreSetup.h
+++ b/Engine/Source/Runtime/Core/Public/HoloLens/HoloLensPlatformCompilerPreSetup.h
@@ -29,6 +29,17 @@
 		__pragma(warning(pop))
 #endif // PRAGMA_ENABLE_SHADOW_VARIABLE_WARNINGS
 
+#ifndef PRAGMA_DISABLE_UNSAFE_TYPECAST_WARNINGS
+	#define PRAGMA_DISABLE_UNSAFE_TYPECAST_WARNINGS \
+		__pragma (warning(push)) \
+		__pragma (warning(disable: 4244)) /* 'argument': conversion from 'type1' to 'type2', possible loss of data */
+#endif // PRAGMA_DISABLE_UNSAFE_TYPECAST_WARNINGS
+
+#ifndef PRAGMA_ENABLE_UNSAFE_TYPECAST_WARNINGS
+	#define PRAGMA_ENABLE_UNSAFE_TYPECAST_WARNINGS \
+		__pragma(warning(pop))
+#endif // PRAGMA_ENABLE_UNSAFE_TYPECAST_WARNINGS
+
 #ifndef PRAGMA_DISABLE_UNDEFINED_IDENTIFIER_WARNINGS
 	#define PRAGMA_DISABLE_UNDEFINED_IDENTIFIER_WARNINGS \
 		__pragma(warning(push)) \
diff --git a/Engine/Source/Runtime/Core/Public/HoloLens/HoloLensPlatformCompilerSetup.h b/Engine/Source/Runtime/Core/Public/HoloLens/HoloLensPlatformCompilerSetup.h
--- a/Engine/Source/Runtime/Core/Public/HoloLens/HoloLensPlatformCompilerSetup.h
+++ b/Engine/Source/Runtime/Core/Public/HoloLens/HoloLensPlatformCompilerSetup.h
@@ -18,7 +18,7 @@ static_assert(_MSC_VER >= 1900, "Visual Studio 2015 or newer required.");
 
 #pragma warning (error:      4001 4002 4003 4004      4006 4007 4008 4009 4010 4011 4012 4013 4014 4015 4016 4017 4018 4019 4020 4021 4022 4023 4024 4025 4026 4027 4028 4029 4030 4031 4032 4033 4034 4035 4036 4037 4038 4039 4040 4041 4042 4043 4044 4045 4046 4047 4048 4049 4050 4051 4052 4053 4054 4055 4056 4057 4058 4059 4060           4063 4064 4065 4066 4067 4068 4069 4070 4071 4072 4073 4074 4075 4076 4077 4078 4079 4080 4081 4082 4083 4084 4085 4086 4087 4088 4089 4090 4091 4092 4093 4094 4095 4096 4097 4098 4099)
 #pragma warning (error: 4100 4101 4102 4103 4104 4105 4106 4107 4108 4109 4110 4111 4112 4113 4114 4115 4116 4117 4118 4119 4120 4121 4122 4123 4124 4125 4126 4127 4128 4129 4130 4131 4132 4133 4134 4135 4136 4137 4138 4139 4140 4141 4142 4143 4144 4145 4146 4147 4148 4149 4150 4151 4152 4153 4154 4155 4156 4157 4158 4159 4160 4161 4162 4163 4164      4166 4167 4168 4169 4170 4171 4172 4173 4174 4175 4176 4177 4178 4179 4180 4181 4182 4183 4184 4185 4186 4187 4188 4189 4190 4191 4192 4193 4194 4195 4196 4197 4198 4199)
-#pragma warning (error: 4200 4201 4202 4203 4204 4205 4206 4207 4208 4209 4210 4211 4212 4213 4214 4215 4216 4217 4218 4219 4220 4221 4222 4223 4224 4225 4226 4227 4228 4229 4230 4231 4232 4233 4234 4235 4236 4237 4238 4239 4240 4241      4243 4244 4245 4246 4247 4248 4249 4250 4251 4252 4253      4255 4256 4257 4258 4259 4260 4261 4262 4263 4264           4267 4268 4269 4270 4271 4272 4273 4274 4275 4276 4277 4278 4279 4280 4281 4282 4283 4284 4285 4286 4287 4288 4289 4290 4291 4292 4293 4294 4295      4297 4298 4299)
+#pragma warning (error: 4200 4201 4202 4203 4204 4205 4206 4207 4208 4209 4210 4211 4212 4213 4214 4215 4216 4217 4218 4219 4220 4221 4222 4223 4224 4225 4226 4227 4228 4229 4230 4231 4232 4233 4234 4235 4236 4237 4238 4239 4240 4241      4243      4245 4246 4247 4248 4249 4250 4251 4252 4253      4255 4256 4257 4258 4259 4260 4261 4262 4263 4264           4267 4268 4269 4270 4271 4272 4273 4274 4275 4276 4277 4278 4279 4280 4281 4282 4283 4284 4285 4286 4287 4288 4289 4290 4291 4292 4293 4294 4295      4297 4298 4299)
 #pragma warning (error: 4300 4301 4302 4303 4304      4306 4307 4308 4309 4310           4313 4314 4315 4316 4317 4318 4319 4320 4321 4322 4323 4324 4325 4326 4327 4328 4329 4330 4331 4332 4333 4334 4335 4336 4337 4338      4340 4341      4343 4344      4346 4347 4348 4349      4351 4352 4353 4354 4355 4356 4357 4358 4359 4360 4361 4362 4363 4364      4366 4367 4368 4369           4372      4374 4375 4376 4377 4378 4379 4380 4381 4382 4383 4384 4385 4386 4387      4389 4390 4391 4392 4393 4394 4395 4396 4397 4398 4399)
 #pragma warning (error: 4400 4401 4402 4403 4404 4405 4406 4407 4408 4409 4410 4411      4413 4414 4415 4416 4417 4418 4419 4420 4421 4422 4423 4424 4425 4426 4427 4428 4429 4430 4431 4432 4433 4434      4436      4438 4439 4440 4441 4442 4443      4445 4446 4447 4448 4449 4450 4451 4452 4453 4454 4455                4459 4460 4461 4462           4465 4466 4467 4468 4469 4470           4473 4474 4475 4476 4477 4478 4479 4480      4482 4483 4484 4485 4486 4487 4488 4489 4490 4491 4492 4493 4494 4495 4496 4497 4498 4499)
 #pragma warning (error: 4500 4501 4502 4503 4504 4505 4506 4507 4508 4509      4511 4512 4513 4514 4515 4516 4517 4518 4519 4520 4521 4522 4523 4524 4525 4526 4527 4528 4529 4530 4531 4532 4533 4534 4535 4536 4537 4538 4539 4540 4541 4542 4543 4544 4545 4546                4550 4551 4552 4553 4554      4556 4557 4558 4559 4560 4561 4562 4563 4564 4565 4566 4567 4568 4569 4570      4572 4573      4575 4576 4577 4578 4579 4580 4581 4582 4583 4584 4585 4586 4587 4588 4589 4590 4591      4593 4594 4595 4596 4597 4598     )
@@ -93,7 +93,6 @@ static_assert(_MSC_VER >= 1900, "Visual Studio 2015 or newer required.");
 #pragma warning(disable : 4127) // Conditional expression is constant									
 #pragma warning(disable : 4200) // Zero-length array item at end of structure, a VC-specific extension	
 #pragma warning(disable : 4201) // nonstandard extension used : nameless struct/union	
-#pragma warning(disable : 4244) // 'expression' : conversion from 'type' to 'type', possible loss of data						
 #pragma warning(disable : 4245) // 'initializing': conversion from 'type' to 'type', signed/unsigned mismatch 
 #pragma warning(disable : 4251) // 'type' needs to have dll-interface to be used by clients of 'type'
 #pragma warning(disable : 4291) // typedef-name '' used as synonym for class-name ''                    
diff --git a/Engine/Source/Runtime/Core/Public/Misc/CoreMiscDefines.h b/Engine/Source/Runtime/Core/Public/Misc/CoreMiscDefines.h
--- a/Engine/Source/Runtime/Core/Public/Misc/CoreMiscDefines.h
+++ b/Engine/Source/Runtime/Core/Public/Misc/CoreMiscDefines.h
@@ -246,6 +246,12 @@ struct TStaticDeprecateExpression
 	}; \
 	enum class PREPROCESSOR_JOIN(EDeprecationMsg_, __LINE__) { Value = PREPROCESSOR_JOIN(FDeprecationMsg_, __LINE__)::condition(TStaticDeprecateExpression<!!(bExpression)>()) }
 
+// These defines are used to mark a difference between two pointers as expected to fit into the specified range
+// while still leaving something searchable if the surrounding code is updated to work with a 64 bit count/range
+// in the future
+#define UE_PTRDIFF_TO_INT32(argument) static_cast<int32>(argument)
+#define UE_PTRDIFF_TO_UINT32(argument) static_cast<uint32>(argument)
+
 /**
 * Makes a type non-copyable and non-movable by deleting copy/move constructors and assignment/move operators.
 * The macro should be placed in the public section of the type for better compiler diagnostic messages.
diff --git a/Engine/Source/Runtime/Core/Public/Windows/WindowsPlatformCompilerPreSetup.h b/Engine/Source/Runtime/Core/Public/Windows/WindowsPlatformCompilerPreSetup.h
--- a/Engine/Source/Runtime/Core/Public/Windows/WindowsPlatformCompilerPreSetup.h
+++ b/Engine/Source/Runtime/Core/Public/Windows/WindowsPlatformCompilerPreSetup.h
@@ -74,6 +74,17 @@
 			__pragma(warning(pop))
 	#endif // PRAGMA_ENABLE_SHADOW_VARIABLE_WARNINGS
 
+	#ifndef PRAGMA_DISABLE_UNSAFE_TYPECAST_WARNINGS
+		#define PRAGMA_DISABLE_UNSAFE_TYPECAST_WARNINGS \
+			__pragma (warning(push)) \
+			__pragma (warning(disable: 4244)) /* 'argument': conversion from 'type1' to 'type2', possible loss of data */
+	#endif // PRAGMA_DISABLE_UNSAFE_TYPECAST_WARNINGS
+
+	#ifndef PRAGMA_ENABLE_UNSAFE_TYPECAST_WARNINGS
+		#define PRAGMA_ENABLE_UNSAFE_TYPECAST_WARNINGS \
+			__pragma(warning(pop))
+	#endif // PRAGMA_ENABLE_UNSAFE_TYPECAST_WARNINGS
+
 	#ifndef PRAGMA_DISABLE_UNDEFINED_IDENTIFIER_WARNINGS
 		#define PRAGMA_DISABLE_UNDEFINED_IDENTIFIER_WARNINGS \
 			__pragma(warning(push)) \
diff --git a/Engine/Source/Runtime/Core/Public/Windows/WindowsPlatformCompilerSetup.h b/Engine/Source/Runtime/Core/Public/Windows/WindowsPlatformCompilerSetup.h
--- a/Engine/Source/Runtime/Core/Public/Windows/WindowsPlatformCompilerSetup.h
+++ b/Engine/Source/Runtime/Core/Public/Windows/WindowsPlatformCompilerSetup.h
@@ -28,7 +28,7 @@ static_assert(_MSC_VER != 1914 && _MSC_VER != 1915, "Visual Studio 2017 versions
 
 #pragma warning (error:      4001 4002 4003           4006 4007 4008      4010           4013      4015           4018 4019 4020      4022 4023 4024 4025 4026 4027 4028 4029 4030 4031 4032 4033 4034 4035 4036      4038           4041 4042           4045      4047 4048 4049      4051 4052 4053 4054 4055 4056 4057           4060           4063 4064 4065 4066 4067 4068 4069                4073 4074 4075 4076 4077      4079 4080 4081      4083      4085 4086 4087 4088 4089 4090 4091 4092      4094      4096 4097 4098 4099)
 #pragma warning (error: 4100 4101 4102 4103                          4109           4112 4113 4114 4115 4116 4117      4119 4120 4121 4122 4123 4124 4125      4127      4129 4130 4131 4132 4133                4137 4138           4141 4142 4143 4144 4145 4146                4150      4152 4153 4154 4155 4156 4157 4158 4159 4160 4161 4162 4163 4164      4166 4167 4168                4172      4174 4175 4176 4177 4178 4179 4180 4181 4182 4183      4185 4186 4187 4188 4189 4190 4191 4192 4193 4194 4195 4196 4197      4199)
-#pragma warning (error: 4200 4201 4202 4203 4204 4205 4206 4207 4208      4210 4211 4212 4213 4214 4215 4216      4218      4220 4221      4223 4224      4226 4227 4228 4229 4230      4232 4233 4234 4235      4237 4238 4239 4240           4243 4244 4245                     4250 4251                4255 4256      4258                     4263 4264           4267 4268 4269           4272 4273 4274 4275 4276 4277 4278 4279 4280 4281 4282 4283      4285 4286 4287 4288 4289 4290 4291      4293      4295      4297 4298 4299)
+#pragma warning (error: 4200 4201 4202 4203 4204 4205 4206 4207 4208      4210 4211 4212 4213 4214 4215 4216      4218      4220 4221      4223 4224      4226 4227 4228 4229 4230      4232 4233 4234 4235      4237 4238 4239 4240           4243      4245                     4250 4251                4255 4256      4258                     4263 4264           4267 4268 4269           4272 4273 4274 4275 4276 4277 4278 4279 4280 4281 4282 4283      4285 4286 4287 4288 4289 4290 4291      4293      4295      4297 4298 4299)
 #pragma warning (error:      4301 4302 4303           4306      4308 4309 4310           4313 4314 4315 4316      4318 4319      4321 4322 4323 4324 4325 4326 4327 4328 4329 4330           4333 4334 4335 4336 4337 4338      4340           4343 4344      4346      4348                4352 4353      4355 4356 4357 4358 4359           4362      4364      4366 4367 4368 4369                     4374 4375 4376 4377 4378 4379 4380 4381 4382 4383 4384           4387      4389 4390 4391 4392 4393 4394 4395 4396 4397 4398 4399)
 #pragma warning (error: 4400 4401 4402 4403 4404 4405 4406 4407 4408 4409 4410 4411      4413 4414 4415 4416 4417 4418 4419 4420           4423 4424 4425 4426 4427      4429 4430 4431           4434      4436      4438 4439 4440 4441 4442 4443      4445 4446 4447 4448 4449 4450 4451 4452 4453 4454 4455                     4460 4461 4462      4464 4465 4466 4467 4468 4469 4470           4473 4474 4475 4476 4477 4478      4480      4482 4483 4484 4485 4486 4487 4488 4489 4490 4491 4492 4493 4494 4495 4496 4497 4498 4499)
 #pragma warning (error:           4502 4503      4505 4506      4508 4509      4511 4512 4513 4514 4515 4516 4517 4518 4519      4521 4522 4523           4526                4530 4531 4532 4533 4534 4535 4536 4537 4538      4540 4541 4542 4543 4544 4545 4546                4550 4551 4552 4553 4554      4556 4557 4558 4559      4561 4562      4564 4565 4566      4568 4569 4570      4572 4573      4575 4576      4578      4580 4581 4582 4583 4584 4585 4586 4587 4588 4589      4591      4593 4594 4595 4596 4597 4598     )
@@ -113,7 +113,6 @@ static_assert(_MSC_VER != 1914 && _MSC_VER != 1915, "Visual Studio 2017 versions
 // 4826 - Conversion from 'type1' to 'type2' is sign-extended. This may cause unexpected runtime behavior.							https://docs.microsoft.com/en-us/cpp/error-messages/compiler-warnings/compiler-warnings-c4800-through-c4999
 // 4828 - The file contains a character starting at offset ... that is illegal in the current source character set (codepage ...).	https://docs.microsoft.com/en-us/cpp/error-messages/compiler-warnings/compiler-warnings-c4800-through-c4999
 // 4837 - trigraph detected: '??character' replaced by 'character'																	https://docs.microsoft.com/en-us/cpp/error-messages/compiler-warnings/compiler-warnings-c4800-through-c4999
-// 4838 - conversion from 'type_1' to 'type_2' requires a narrowing conversion														https://docs.microsoft.com/en-us/cpp/error-messages/compiler-warnings/compiler-warning-level-1-c4838
 // 4868 - file(line_number)' compiler may not enforce left-to-right evaluation order in braced initializer list                     https://msdn.microsoft.com/en-us/library/mt656772.aspx
 
 // 4962 - 'function' : Profile-guided optimizations disabled because optimizations caused profile data to become inconsistent		https://docs.microsoft.com/en-us/cpp/error-messages/compiler-warnings/compiler-warning-c4962
@@ -141,7 +140,6 @@ static_assert(_MSC_VER != 1914 && _MSC_VER != 1915, "Visual Studio 2017 versions
 #pragma warning(disable: 4200) // Zero-length array item at end of structure, a VC-specific extension								https://docs.microsoft.com/en-us/cpp/error-messages/compiler-warnings/compiler-warning-levels-2-and-4-c4200
 #pragma warning(disable: 4201) // nonstandard extension used : nameless struct/union												https://docs.microsoft.com/en-us/cpp/error-messages/compiler-warnings/compiler-warning-level-4-c4201
 #pragma warning(disable: 4217) // 'operator' : member template functions cannot be used for copy-assignment or copy-construction	// No docs
-#pragma warning(disable: 4244) // 'argument' : conversion from 'type1' to 'type2', possible loss of data							https://docs.microsoft.com/en-us/cpp/error-messages/compiler-warnings/compiler-warning-level-4-c4201
 #pragma warning(disable: 4245) // 'initializing': conversion from 'type' to 'type', signed/unsigned mismatch						https://docs.microsoft.com/en-us/cpp/error-messages/compiler-warnings/compiler-warning-level-4-c4245
 #pragma warning(disable: 4251) // 'type' needs to have dll-interface to be used by clients of 'type'								https://docs.microsoft.com/en-us/cpp/error-messages/compiler-warnings/compiler-warning-level-1-c4251
 #pragma warning(disable: 4267) // 'var' : conversion from 'size_t' to 'type', possible loss of data									https://docs.microsoft.com/en-us/cpp/error-messages/compiler-warnings/compiler-warning-level-3-c4267
