Core: Fixes for various unsafe implicit type conversion warnings in the Core module (WIP; some files still suppress the warning)
- Ptrdiff -> int32
- Float/int confusion and double/float
- size_t stuff; various changes to the algorithms to use a deduced IndexType template argument and/or decltype to use the appropriate size for indicies and counts
- Fixed GetNum(FString) incorrectly returning SIZE_T instead of int32, and GetNum(container) now returns whatever container.Num() does (so usually int32)

diff --git a/Engine/Source/Runtime/Core/Core.Build.cs b/Engine/Source/Runtime/Core/Core.Build.cs
--- a/Engine/Source/Runtime/Core/Core.Build.cs
+++ b/Engine/Source/Runtime/Core/Core.Build.cs
@@ -223,5 +223,7 @@ public class Core : ModuleRules
 		PrivateDefinitions.Add("PLATFORM_COMPILER_OPTIMIZATION_LTCG=" + (Target.bAllowLTCG ? "1" : "0"));
 		PrivateDefinitions.Add("PLATFORM_COMPILER_OPTIMIZATION_PG=" + (Target.bPGOOptimize ? "1" : "0"));
 		PrivateDefinitions.Add("PLATFORM_COMPILER_OPTIMIZATION_PG_PROFILING=" + (Target.bPGOProfile ? "1" : "0"));
-    }
+
+		UnsafeTypeCastWarningLevel = WarningLevel.Warning;
+	}
 }
diff --git a/Engine/Source/Runtime/Core/Private/Android/AndroidStats.cpp b/Engine/Source/Runtime/Core/Private/Android/AndroidStats.cpp
--- a/Engine/Source/Runtime/Core/Private/Android/AndroidStats.cpp
+++ b/Engine/Source/Runtime/Core/Private/Android/AndroidStats.cpp
@@ -72,7 +72,7 @@ void FAndroidStats::UpdateAndroidStats()
 	{
 		return;
 	}
-	bool bUpdateStats = ((FPlatformTime::ToSeconds(CurrentTime - LastCollectionTime) >= GAndroidCPUStatsUpdateRate));
+	bool bUpdateStats = ((FPlatformTime::ToSeconds64(CurrentTime - LastCollectionTime) >= GAndroidCPUStatsUpdateRate));
 	
 	if( bUpdateStats )
 	{
diff --git a/Engine/Source/Runtime/Core/Private/Async/TaskGraph.cpp b/Engine/Source/Runtime/Core/Private/Async/TaskGraph.cpp
--- a/Engine/Source/Runtime/Core/Private/Async/TaskGraph.cpp
+++ b/Engine/Source/Runtime/Core/Private/Async/TaskGraph.cpp
@@ -1600,7 +1600,7 @@ private:
 		if (TLSPointer)
 		{
 			checkThreadGraph(TLSPointer - WorkerThreads >= 0 && TLSPointer - WorkerThreads < NumThreads);
-			int32 ThreadIndex = TLSPointer - WorkerThreads;
+			int32 ThreadIndex = UE_PTRDIFF_TO_INT32(TLSPointer - WorkerThreads);
 			checkThreadGraph(Thread(ThreadIndex).GetThreadId() == ThreadIndex);
 			if (ThreadIndex < NumNamedThreads)
 			{
@@ -1821,10 +1821,10 @@ public:
 	void FORCEINLINE DoTask(ENamedThreads::Type CurrentThread, const FGraphEventRef& MyCompletionGraphEvent)
 	{
 		{
-			float ThisTime = FPlatformTime::Seconds() - StartTime;
-			if (ThisTime > 0.02f)
+			const double ThisTime = FPlatformTime::Seconds() - StartTime;
+			if (ThisTime > 0.02)
 			{
-				UE_CLOG(GPrintBroadcastWarnings, LogTaskGraph, Warning, TEXT("Task graph took %6.2fms for %s to recieve broadcast."), ThisTime * 1000.0f, Name);
+				UE_CLOG(GPrintBroadcastWarnings, LogTaskGraph, Warning, TEXT("Task graph took %6.2fms for %s to recieve broadcast."), ThisTime * 1000.0, Name);
 			}
 		}
 
@@ -1833,10 +1833,10 @@ public:
 			Function(CurrentThread);
 		}
 		{
-			float ThisTime = FPlatformTime::Seconds() - StartTime;
-			if (ThisTime > 0.02f)
+			const double ThisTime = FPlatformTime::Seconds() - StartTime;
+			if (ThisTime > 0.02)
 			{
-				UE_CLOG(GPrintBroadcastWarnings, LogTaskGraph, Warning, TEXT("Task graph took %6.2fms for %s to recieve broadcast and do processing."), ThisTime * 1000.0f, Name);
+				UE_CLOG(GPrintBroadcastWarnings, LogTaskGraph, Warning, TEXT("Task graph took %6.2fms for %s to recieve broadcast and do processing."), ThisTime * 1000.0, Name);
 			}
 		}
 		if (StallForTaskThread)
@@ -1848,8 +1848,8 @@ public:
 					QUICK_SCOPE_CYCLE_COUNTER(STAT_Broadcast_WaitForOthers);
 					TaskEvent->Wait();
 					{
-						float ThisTime = FPlatformTime::Seconds() - StartTime;
-						if (ThisTime > 0.02f)
+						const double ThisTime = FPlatformTime::Seconds() - StartTime;
+						if (ThisTime > 0.02)
 						{
 							UE_CLOG(GPrintBroadcastWarnings, LogTaskGraph, Warning, TEXT("Task graph took %6.2fms for %s to recieve broadcast do processing and wait for other task threads."), ThisTime * 1000.0f, Name);
 						}
@@ -1860,10 +1860,10 @@ public:
 			{
 				CallerEvent->Trigger();
 				{
-					float ThisTime = FPlatformTime::Seconds() - StartTime;
-					if (ThisTime > 0.02f)
+					const double ThisTime = FPlatformTime::Seconds() - StartTime;
+					if (ThisTime > 0.02)
 					{
-						UE_CLOG(GPrintBroadcastWarnings, LogTaskGraph, Warning, TEXT("Task graph took %6.2fms for %s to recieve broadcast do processing and trigger other task threads."), ThisTime * 1000.0f, Name);
+						UE_CLOG(GPrintBroadcastWarnings, LogTaskGraph, Warning, TEXT("Task graph took %6.2fms for %s to recieve broadcast do processing and trigger other task threads."), ThisTime * 1000.0, Name);
 					}
 				}
 			}
@@ -1973,14 +1973,14 @@ void FTaskGraphInterface::BroadcastSlow_OnlyUseForSpecialPurposes(bool bDoTaskTh
 			TaskEvent->Trigger();
 		}
 		{
-			double StartTimeInner = FPlatformTime::Seconds();
+			const double StartTimeInner = FPlatformTime::Seconds();
 			QUICK_SCOPE_CYCLE_COUNTER(STAT_Broadcast_WaitForTaskThreads);
 			FTaskGraphInterface::Get().WaitUntilTasksComplete(TaskThreadTasks, ENamedThreads::GameThread_Local);
 			{
-				float ThisTime = FPlatformTime::Seconds() - StartTimeInner;
-				if (ThisTime > 0.02f)
+				const double ThisTime = FPlatformTime::Seconds() - StartTimeInner;
+				if (ThisTime > 0.02)
 				{
-					UE_CLOG(GPrintBroadcastWarnings, LogTaskGraph, Warning, TEXT("Task graph took %6.2fms to wait for task thread broadcast."), ThisTime * 1000.0f);
+					UE_CLOG(GPrintBroadcastWarnings, LogTaskGraph, Warning, TEXT("Task graph took %6.2fms to wait for task thread broadcast."), ThisTime * 1000.0);
 				}
 			}
 		}
@@ -2009,10 +2009,10 @@ void FTaskGraphInterface::BroadcastSlow_OnlyUseForSpecialPurposes(bool bDoTaskTh
 			FPlatformMisc::PumpMessagesOutsideMainLoop();
 		}
 		
-		float EndTimeInner = FPlatformTime::Seconds() - StartTimeInner;
-		if (EndTimeInner > 0.02f)
+		const double EndTimeInner = FPlatformTime::Seconds() - StartTimeInner;
+		if (EndTimeInner > 0.02)
 		{
-			UE_CLOG(GPrintBroadcastWarnings, LogTaskGraph, Warning, TEXT("Task graph took %6.2fms to wait for named thread broadcast."), EndTimeInner * 1000.0f);
+			UE_CLOG(GPrintBroadcastWarnings, LogTaskGraph, Warning, TEXT("Task graph took %6.2fms to wait for named thread broadcast."), EndTimeInner * 1000.0);
 		}
 	}
 	for (FEvent* TaskEvent : TaskEvents)
@@ -2024,10 +2024,10 @@ void FTaskGraphInterface::BroadcastSlow_OnlyUseForSpecialPurposes(bool bDoTaskTh
 		FPlatformProcess::ReturnSynchEventToPool(MyEvent);
 	}
 	{
-		float ThisTime = FPlatformTime::Seconds() - StartTime;
-		if (ThisTime > 0.02f)
+		const double ThisTime = FPlatformTime::Seconds() - StartTime;
+		if (ThisTime > 0.02)
 		{
-			UE_CLOG(GPrintBroadcastWarnings, LogTaskGraph, Warning, TEXT("Task graph took %6.2fms to broadcast."), ThisTime * 1000.0f);
+			UE_CLOG(GPrintBroadcastWarnings, LogTaskGraph, Warning, TEXT("Task graph took %6.2fms to broadcast."), ThisTime * 1000.0);
 		}
 	}
 }
diff --git a/Engine/Source/Runtime/Core/Private/Containers/Algo/FindSortedStringCaseInsensitive.cpp b/Engine/Source/Runtime/Core/Private/Containers/Algo/FindSortedStringCaseInsensitive.cpp
--- a/Engine/Source/Runtime/Core/Private/Containers/Algo/FindSortedStringCaseInsensitive.cpp
+++ b/Engine/Source/Runtime/Core/Private/Containers/Algo/FindSortedStringCaseInsensitive.cpp
@@ -35,7 +35,7 @@ namespace Algo
 			// If the subarray is only 1 element, we can use Stricmp to finish the job
 			if (SubArray + 1 == SubArrayEnd)
 			{
-				return (FCString::Stricmp(*SubArray + CharIndex, Str + CharIndex) == 0) ? SubArray - SortedArray : -1;
+				return (FCString::Stricmp(*SubArray + CharIndex, Str + CharIndex) == 0) ? UE_PTRDIFF_TO_INT32(SubArray - SortedArray) : -1;
 			}
 
 			// Get the next char to binary search for
@@ -46,7 +46,7 @@ namespace Algo
 			{
 				// If the subarray's string is also at the end then we've found a match,
 				// otherwise we've only found a prefix, which is not a match
-				return ((*SubArray)[CharIndex] == TEXT('\0')) ? SubArray - SortedArray : -1;
+				return ((*SubArray)[CharIndex] == TEXT('\0')) ? UE_PTRDIFF_TO_INT32(SubArray - SortedArray) : -1;
 			}
 
 			// Now we try to find the equal_range of the character
diff --git a/Engine/Source/Runtime/Core/Private/Containers/ContainersTest.cpp b/Engine/Source/Runtime/Core/Private/Containers/ContainersTest.cpp
--- a/Engine/Source/Runtime/Core/Private/Containers/ContainersTest.cpp
+++ b/Engine/Source/Runtime/Core/Private/Containers/ContainersTest.cpp
@@ -47,7 +47,7 @@ namespace
 
 		int32 Num() const
 		{
-			return End - ConstructedIDs;
+			return UE_PTRDIFF_TO_INT32(End - ConstructedIDs);
 		}
 
 		int32 Add()
diff --git a/Engine/Source/Runtime/Core/Private/Containers/HashTable.cpp b/Engine/Source/Runtime/Core/Private/Containers/HashTable.cpp
--- a/Engine/Source/Runtime/Core/Private/Containers/HashTable.cpp
+++ b/Engine/Source/Runtime/Core/Private/Containers/HashTable.cpp
@@ -19,7 +19,7 @@ CORE_API void FHashTable::Resize( uint32 NewIndexSize )
 
 	if( IndexSize == 0 )
 	{
-		HashMask = HashSize - 1;
+		HashMask = (uint16)(HashSize - 1);
 		Hash = new uint32[ HashSize ];
 		FMemory::Memset( Hash, 0xff, HashSize * 4 );
 	}
@@ -43,7 +43,7 @@ CORE_API float FHashTable::AverageSearch() const
 	for( uint32 Key = 0; Key < HashSize; Key++ )
 	{
 		uint32 NumInBucket = 0;
-		for( uint32 i = First( Key ); IsValid( i ); i = Next( i ) )
+		for( uint32 i = First( (uint16)Key ); IsValid( i ); i = Next( i ) )
 		{
 			NumInBucket++;
 		}
diff --git a/Engine/Source/Runtime/Core/Private/Containers/LockFreeList.cpp b/Engine/Source/Runtime/Core/Private/Containers/LockFreeList.cpp
--- a/Engine/Source/Runtime/Core/Private/Containers/LockFreeList.cpp
+++ b/Engine/Source/Runtime/Core/Private/Containers/LockFreeList.cpp
@@ -48,7 +48,7 @@ void LockFreeLinksExhausted(uint32 TotalNum)
 	UE_LOG(LogTemp, Fatal, TEXT("Consumed %d lock free links; there are no more."), TotalNum);
 }
 
-static void ChangeMem(int32 Delta)
+static void ChangeMem(int64 Delta)
 {
 #if 0   // this is not safe; we cannot change stats in the middle of a lock free operation
 	static FThreadSafeCounter LockFreeListMem;
diff --git a/Engine/Source/Runtime/Core/Private/Containers/StackTracker.cpp b/Engine/Source/Runtime/Core/Private/Containers/StackTracker.cpp
--- a/Engine/Source/Runtime/Core/Private/Containers/StackTracker.cpp
+++ b/Engine/Source/Runtime/Core/Private/Containers/StackTracker.cpp
@@ -143,7 +143,7 @@ void FStackTracker::DumpStackTraces(int32 StackThreshold, FOutputDevice& Ar, flo
 	}
 
 	// Calculate the number of frames we captured.
-	int32 FramesCaptured = 0;
+	uint64 FramesCaptured = 0;
 	if( bIsEnabled )
 	{
 		FramesCaptured = GFrameCounter - StartFrameCounter;
diff --git a/Engine/Source/Runtime/Core/Private/FramePro/FrameProProfiler.cpp b/Engine/Source/Runtime/Core/Private/FramePro/FrameProProfiler.cpp
--- a/Engine/Source/Runtime/Core/Private/FramePro/FrameProProfiler.cpp
+++ b/Engine/Source/Runtime/Core/Private/FramePro/FrameProProfiler.cpp
@@ -275,7 +275,7 @@ void FFrameProProfiler::FrameStart()
 		uint64 CurrentTime = FPlatformTime::Cycles64();
 		if (GFrameProCPUStatsUpdateRate > 0.0f)
 		{
-			bool bUpdateStats = ((FPlatformTime::ToSeconds(CurrentTime - LastCollectionTime) >= GFrameProCPUStatsUpdateRate));
+			bool bUpdateStats = ((FPlatformTime::ToSeconds64(CurrentTime - LastCollectionTime) >= GFrameProCPUStatsUpdateRate));
 			if (bUpdateStats)
 			{
 				LastCollectionTime = CurrentTime;
diff --git a/Engine/Source/Runtime/Core/Private/GenericPlatform/GenericPlatformCrashContext.cpp b/Engine/Source/Runtime/Core/Private/GenericPlatform/GenericPlatformCrashContext.cpp
--- a/Engine/Source/Runtime/Core/Private/GenericPlatform/GenericPlatformCrashContext.cpp
+++ b/Engine/Source/Runtime/Core/Private/GenericPlatform/GenericPlatformCrashContext.cpp
@@ -302,7 +302,7 @@ void FGenericCrashContext::CopySharedCrashContext(FSharedCrashContext& Dst)
 	TCHAR* DynamicDataStart = &Dst.DynamicData[0];
 	TCHAR* DynamicDataPtr = DynamicDataStart;
 
-	Dst.EnabledPluginsOffset = DynamicDataPtr - DynamicDataStart;
+	Dst.EnabledPluginsOffset = (uint32)(DynamicDataPtr - DynamicDataStart);
 	Dst.EnabledPluginsNum = NCached::EnabledPluginsList.Num();
 	for (const FString& Plugin : NCached::EnabledPluginsList)
 	{
@@ -311,7 +311,7 @@ void FGenericCrashContext::CopySharedCrashContext(FSharedCrashContext& Dst)
 	}
 	DynamicDataPtr += FCString::Strlen(DynamicDataPtr) + 1;
 
-	Dst.EngineDataOffset = DynamicDataPtr - DynamicDataStart;
+	Dst.EngineDataOffset = (uint32)(DynamicDataPtr - DynamicDataStart);
 	Dst.EngineDataNum = NCached::EngineData.Num();
 	for (const TPair<FString, FString>& Pair : NCached::EngineData)
 	{
@@ -322,7 +322,7 @@ void FGenericCrashContext::CopySharedCrashContext(FSharedCrashContext& Dst)
 	}
 	DynamicDataPtr += FCString::Strlen(DynamicDataPtr) + 1;
 
-	Dst.GameDataOffset = DynamicDataPtr - DynamicDataStart;
+	Dst.GameDataOffset = (uint32)(DynamicDataPtr - DynamicDataStart);
 	Dst.GameDataNum = NCached::GameData.Num();
 	for (const TPair<FString, FString>& Pair : NCached::GameData)
 	{
diff --git a/Engine/Source/Runtime/Core/Private/GenericPlatform/GenericPlatformMallocCrash.cpp b/Engine/Source/Runtime/Core/Private/GenericPlatform/GenericPlatformMallocCrash.cpp
--- a/Engine/Source/Runtime/Core/Private/GenericPlatform/GenericPlatformMallocCrash.cpp
+++ b/Engine/Source/Runtime/Core/Private/GenericPlatform/GenericPlatformMallocCrash.cpp
@@ -517,7 +517,7 @@ uint8* FGenericPlatformMallocCrash::AllocateFromSmallPool( uint32 AllocationSize
 uint32 FGenericPlatformMallocCrash::GetAllocationSize( void *Original )
 {
 	FPtrInfo* PtrInfo = (FPtrInfo*)((uint8*)Original-PER_ALLOC_OVERHEAD);
-	return PtrInfo->Size;
+	return (uint32)(PtrInfo->Size);
 }
 
 uint32 FGenericPlatformMallocCrash::SafePageSize()
diff --git a/Engine/Source/Runtime/Core/Private/GenericPlatform/GenericPlatformMath.cpp b/Engine/Source/Runtime/Core/Private/GenericPlatform/GenericPlatformMath.cpp
--- a/Engine/Source/Runtime/Core/Private/GenericPlatform/GenericPlatformMath.cpp
+++ b/Engine/Source/Runtime/Core/Private/GenericPlatform/GenericPlatformMath.cpp
@@ -148,11 +148,11 @@ public:
 		check(MathPlatform::IsFinite(One));
 		check(!MathPlatform::IsNaN(MinusOneE37));
 		check(MathPlatform::IsFinite(MinusOneE37));
-		check(MathPlatform::FloorLog2(Zero) == 0);
-		check(MathPlatform::FloorLog2(One) == 0);
-		check(MathPlatform::FloorLog2(Two) == 1);
-		check(MathPlatform::FloorLog2(Twelve) == 3);
-		check(MathPlatform::FloorLog2(Sixteen) == 4);
+		check(MathPlatform::FloorLog2((uint32)Zero) == 0);
+		check(MathPlatform::FloorLog2((uint32)One) == 0);
+		check(MathPlatform::FloorLog2((uint32)Two) == 1);
+		check(MathPlatform::FloorLog2((uint32)Twelve) == 3);
+		check(MathPlatform::FloorLog2((uint32)Sixteen) == 4);
 	}
 };
 
diff --git a/Engine/Source/Runtime/Core/Private/GenericPlatform/GenericPlatformMemory.cpp b/Engine/Source/Runtime/Core/Private/GenericPlatform/GenericPlatformMemory.cpp
--- a/Engine/Source/Runtime/Core/Private/GenericPlatform/GenericPlatformMemory.cpp
+++ b/Engine/Source/Runtime/Core/Private/GenericPlatform/GenericPlatformMemory.cpp
@@ -293,11 +293,11 @@ EPlatformMemorySizeBucket FGenericPlatformMemory::GetMemorySizeBucket()
 		{
 			MemoryBucketRoundingAddition = FCString::Atoi64(**MemoryBucketRoundingAdditionVar);
 		}
-		uint32 TotalPhysicalGB = (Stats.TotalPhysical + MemoryBucketRoundingAddition * 1024 * 1024 - 1) / 1024 / 1024 / 1024;
+		uint32 TotalPhysicalGB = (uint32)((Stats.TotalPhysical + MemoryBucketRoundingAddition * 1024 * 1024 - 1) / 1024 / 1024 / 1024);
 #else
-		uint32 TotalPhysicalGB = (Stats.TotalPhysical + 1024 * 1024 * 1024 - 1) / 1024 / 1024 / 1024;
+		uint32 TotalPhysicalGB = (uint32)((Stats.TotalPhysical + 1024 * 1024 * 1024 - 1) / 1024 / 1024 / 1024);
 #endif
-		uint32 AddressLimitGB = (Stats.AddressLimit + 1024 * 1024 * 1024 - 1) / 1024 / 1024 / 1024;
+		uint32 AddressLimitGB = (uint32)((Stats.AddressLimit + 1024 * 1024 * 1024 - 1) / 1024 / 1024 / 1024);
 		int32 CurMemoryGB = (int32)FMath::Min(TotalPhysicalGB, AddressLimitGB);
 
 		// if at least Smaller is specified, we can set the Bucket
@@ -398,7 +398,7 @@ void FGenericPlatformMemory::MemswapGreaterThan8( void* RESTRICT Ptr1, void* RES
 		Size -= 4;
 	}
 
-	uint32 CommonAlignment = FMath::Min(FMath::CountTrailingZeros(Union1.PtrUint - Union2.PtrUint), 3u);
+	uint32 CommonAlignment = FMath::Min(FMath::CountTrailingZeros((uint32)(Union1.PtrUint - Union2.PtrUint)), 3u);
 	switch (CommonAlignment)
 	{
 		default:
diff --git a/Engine/Source/Runtime/Core/Private/GenericPlatform/GenericPlatformMisc.cpp b/Engine/Source/Runtime/Core/Private/GenericPlatform/GenericPlatformMisc.cpp
--- a/Engine/Source/Runtime/Core/Private/GenericPlatform/GenericPlatformMisc.cpp
+++ b/Engine/Source/Runtime/Core/Private/GenericPlatform/GenericPlatformMisc.cpp
@@ -1207,14 +1207,14 @@ void FGenericPlatformMisc::UpdateHotfixableEnsureSettings()
 	float HandleEnsurePercentInConfig = 100.0f;
 	if (GConfig && GConfig->GetFloat(TEXT("Core.System"), TEXT("HandleEnsurePercent"), HandleEnsurePercentInConfig, GEngineIni))
 	{
-		GenericPlatformMisc::GEnsureChance = HandleEnsurePercentInConfig / 100.0;
+		GenericPlatformMisc::GEnsureChance = HandleEnsurePercentInConfig / 100.0f;
 	}
 	else
 	{
 		float HandleEnsurePercentOnCmdLine = 100.0f;
 		if (!FCommandLine::IsInitialized() && FParse::Value(FCommandLine::Get(), TEXT("handleensurepercent="), HandleEnsurePercentOnCmdLine))
 		{
-			GenericPlatformMisc::GEnsureChance = HandleEnsurePercentOnCmdLine / 100.0;
+			GenericPlatformMisc::GEnsureChance = HandleEnsurePercentOnCmdLine / 100.0f;
 		}
 	}
 
diff --git a/Engine/Source/Runtime/Core/Private/GenericPlatform/GenericPlatformStackWalk.cpp b/Engine/Source/Runtime/Core/Private/GenericPlatform/GenericPlatformStackWalk.cpp
--- a/Engine/Source/Runtime/Core/Private/GenericPlatform/GenericPlatformStackWalk.cpp
+++ b/Engine/Source/Runtime/Core/Private/GenericPlatform/GenericPlatformStackWalk.cpp
@@ -108,7 +108,7 @@ bool FGenericPlatformStackWalk::SymbolInfoToHumanReadableString( const FProgramC
 		}
 
 		// Append the stack line.
-		FCStringAnsi::Strncat(HumanReadableString, StackLine, HumanReadableStringSize);
+		FCStringAnsi::Strncat(HumanReadableString, StackLine, (int32)HumanReadableStringSize);
 
 		// Return true, if we have a valid function name.
 		return bHasValidFunctionName;
@@ -182,7 +182,7 @@ void FGenericPlatformStackWalk::StackWalkAndDump( ANSICHAR* HumanReadableString,
 	while( CurrentDepth < Depth )
 	{
 		FPlatformStackWalk::ProgramCounterToHumanReadableString( CurrentDepth, StackTrace[CurrentDepth], HumanReadableString, HumanReadableStringSize, reinterpret_cast< FGenericCrashContext* >( Context ) );
-		FCStringAnsi::Strncat(HumanReadableString, LINE_TERMINATOR_ANSI, HumanReadableStringSize);
+		FCStringAnsi::Strncat(HumanReadableString, LINE_TERMINATOR_ANSI, (int32)HumanReadableStringSize);
 		CurrentDepth++;
 	}
 }
diff --git a/Engine/Source/Runtime/Core/Private/GenericPlatform/GenericPlatformTime.cpp b/Engine/Source/Runtime/Core/Private/GenericPlatform/GenericPlatformTime.cpp
--- a/Engine/Source/Runtime/Core/Private/GenericPlatform/GenericPlatformTime.cpp
+++ b/Engine/Source/Runtime/Core/Private/GenericPlatform/GenericPlatformTime.cpp
@@ -88,7 +88,7 @@ TCHAR* FGenericPlatformTime::StrDate( TCHAR* Dest, SIZE_T DestSize )
 	int32 MSec;
 
 	FPlatformTime::SystemTime(Year, Month, DayOfWeek, Day, Hour, Min, Sec, MSec );
-	FCString::Snprintf(Dest, DestSize, TEXT("%02d/%02d/%02d"), Month, Day, Year % 100);
+	FCString::Snprintf(Dest, (int32)DestSize, TEXT("%02d/%02d/%02d"), Month, Day, Year % 100);
 	return Dest;
 }
 
@@ -104,7 +104,7 @@ TCHAR* FGenericPlatformTime::StrTime( TCHAR* Dest, SIZE_T DestSize )
 	int32 MSec;
 
 	FPlatformTime::SystemTime(Year, Month, DayOfWeek, Day, Hour, Min, Sec, MSec );
-	FCString::Snprintf( Dest, DestSize, TEXT("%02d:%02d:%02d"), Hour, Min, Sec);
+	FCString::Snprintf( Dest, (int32)DestSize, TEXT("%02d:%02d:%02d"), Hour, Min, Sec);
 	return Dest;
 }
 
@@ -127,31 +127,31 @@ FString FGenericPlatformTime::PrettyTime( double Seconds )
 {
 	if ( Seconds < 1.0 )
 	{
-		return FString::Printf( TEXT("%d ms"), FMath::TruncToInt(Seconds*1000) );
+		return FString::Printf( TEXT("%d ms"), FMath::TruncToInt((float)(Seconds*1000)) );
 	}
 	else if ( Seconds < 10.0 )
 	{
-		int32 Sec = FMath::TruncToInt(Seconds);
-		int32 Ms = FMath::TruncToInt(Seconds*1000) - Sec*1000;
+		int32 Sec = FMath::TruncToInt((float)Seconds);
+		int32 Ms = FMath::TruncToInt((float)(Seconds*1000)) - Sec*1000;
 		return FString::Printf( TEXT("%d.%02d sec"), Sec, Ms/10 );
 	}
 	else if ( Seconds < 60.0 )
 	{
-		int32 Sec = FMath::TruncToInt(Seconds);
-		int32 Ms = FMath::TruncToInt(Seconds*1000) - Sec*1000;
+		int32 Sec = FMath::TruncToInt((float)Seconds);
+		int32 Ms = FMath::TruncToInt((float)(Seconds*1000)) - Sec*1000;
 		return FString::Printf( TEXT("%d.%d sec"), Sec, Ms/100 );
 	}
 	else if ( Seconds < 60.0*60.0 )
 	{
-		int32 Min = FMath::TruncToInt(Seconds/60.0);
-		int32 Sec = FMath::TruncToInt(Seconds) - Min*60;
+		int32 Min = FMath::TruncToInt((float)(Seconds/60.0));
+		int32 Sec = FMath::TruncToInt((float)Seconds) - Min*60;
 		return FString::Printf( TEXT("%d:%02d min"), Min, Sec );
 	}
 	else
 	{
-		int32 Hr = FMath::TruncToInt(Seconds/3600.0);
-		int32 Min = FMath::TruncToInt((Seconds - Hr*3600)/60.0);
-		int32 Sec = FMath::TruncToInt(Seconds - Hr*3600 - Min*60);
+		int32 Hr = FMath::TruncToInt((float)(Seconds/3600.0));
+		int32 Min = FMath::TruncToInt((float)((Seconds - Hr*3600)/60.0));
+		int32 Sec = FMath::TruncToInt((float)(Seconds - Hr*3600 - Min*60));
 		return FString::Printf( TEXT("%d:%02d:%02d hours"), Hr, Min, Sec );
 	}
 }
@@ -200,7 +200,7 @@ struct FCPUTimeDump
 			UE_LOG(LogGenericPlatformTime, Log, TEXT("Delay set to %i second(s), started printing the CPU usage"), Delay);
 
 			GetCPUTimeDelegateHandle   = FTicker::GetCoreTicker().AddTicker( GetCPUTimeDelegate );
-			CPUTimeDumpDelegateHandle  = FTicker::GetCoreTicker().AddTicker( CPUTimeDumpDelegate, Delay );
+			CPUTimeDumpDelegateHandle  = FTicker::GetCoreTicker().AddTicker( CPUTimeDumpDelegate, (float)Delay );
 		}
 	}
 
diff --git a/Engine/Source/Runtime/Core/Private/HAL/Allocators/CachedOSPageAllocator.cpp b/Engine/Source/Runtime/Core/Private/HAL/Allocators/CachedOSPageAllocator.cpp
--- a/Engine/Source/Runtime/Core/Private/HAL/Allocators/CachedOSPageAllocator.cpp
+++ b/Engine/Source/Runtime/Core/Private/HAL/Allocators/CachedOSPageAllocator.cpp
@@ -7,7 +7,7 @@
 #include "HAL/LowLevelMemTracker.h"
 
 
-void* FCachedOSPageAllocator::AllocateImpl(SIZE_T Size, uint32 CachedByteLimit, FFreePageBlock* First, FFreePageBlock* Last, uint32& FreedPageBlocksNum, uint32& CachedTotal)
+void* FCachedOSPageAllocator::AllocateImpl(SIZE_T Size, uint32 CachedByteLimit, FFreePageBlock* First, FFreePageBlock* Last, uint32& FreedPageBlocksNum, SIZE_T& CachedTotal)
 {
 	if (!FPlatformMemory::BinnedPlatformHasMemoryPoolForThisSize(Size) && (Size <= CachedByteLimit / 4))
 	{
@@ -79,7 +79,7 @@ void* FCachedOSPageAllocator::AllocateImpl(SIZE_T Size, uint32 CachedByteLimit,
 	return FPlatformMemory::BinnedAllocFromOS(Size);
 }
 
-void FCachedOSPageAllocator::FreeImpl(void* Ptr, SIZE_T Size, uint32 NumCacheBlocks, uint32 CachedByteLimit, FFreePageBlock* First, uint32& FreedPageBlocksNum, uint32& CachedTotal)
+void FCachedOSPageAllocator::FreeImpl(void* Ptr, SIZE_T Size, uint32 NumCacheBlocks, uint32 CachedByteLimit, FFreePageBlock* First, uint32& FreedPageBlocksNum, SIZE_T& CachedTotal)
 {
 	if (FPlatformMemory::BinnedPlatformHasMemoryPoolForThisSize(Size) || Size > CachedByteLimit / 4)
 	{
@@ -108,7 +108,7 @@ void FCachedOSPageAllocator::FreeImpl(void* Ptr, SIZE_T Size, uint32 NumCacheBlo
 	++FreedPageBlocksNum;
 }
 
-void FCachedOSPageAllocator::FreeAllImpl(FFreePageBlock* First, uint32& FreedPageBlocksNum, uint32& CachedTotal)
+void FCachedOSPageAllocator::FreeAllImpl(FFreePageBlock* First, uint32& FreedPageBlocksNum, SIZE_T& CachedTotal)
 {
 	while (FreedPageBlocksNum)
 	{
diff --git a/Engine/Source/Runtime/Core/Private/HAL/Allocators/PooledVirtualMemoryAllocator.cpp b/Engine/Source/Runtime/Core/Private/HAL/Allocators/PooledVirtualMemoryAllocator.cpp
--- a/Engine/Source/Runtime/Core/Private/HAL/Allocators/PooledVirtualMemoryAllocator.cpp
+++ b/Engine/Source/Runtime/Core/Private/HAL/Allocators/PooledVirtualMemoryAllocator.cpp
@@ -223,7 +223,7 @@ void FPooledVirtualMemoryAllocator::DestroyPool(FPoolDescriptorBase* Pool)
 	// allocated with placement new, do not call delete
 	PoolDesc.Pool->~T64KBAlignedPool();
 
-	FPlatformMemory::FPlatformVirtualMemoryBlock VMBlock(Pool, Pool->VMSizeDivVirtualSizeAlignment);
+	FPlatformMemory::FPlatformVirtualMemoryBlock VMBlock(Pool, (uint32)Pool->VMSizeDivVirtualSizeAlignment);
 	VMBlock.FreeVirtual();
 }
 
diff --git a/Engine/Source/Runtime/Core/Private/HAL/IPlatformFileProfilerWrapper.cpp b/Engine/Source/Runtime/Core/Private/HAL/IPlatformFileProfilerWrapper.cpp
--- a/Engine/Source/Runtime/Core/Private/HAL/IPlatformFileProfilerWrapper.cpp
+++ b/Engine/Source/Runtime/Core/Private/HAL/IPlatformFileProfilerWrapper.cpp
@@ -22,13 +22,13 @@ bool FPlatformFileReadStatsHandle::Read( uint8* Destination, int64 BytesToRead )
 {
 	double Timer = FPlatformTime::Seconds();
 	bool Result = FileHandle->Read(Destination, BytesToRead);
-	float Delta = FPlatformTime::Seconds()-Timer;
+	double Delta = FPlatformTime::Seconds()-Timer;
 	if (Delta > SMALL_NUMBER)
 	{
-		float BytesPerSec = (BytesToRead/1024.0f)/Delta;
-		FPlatformAtomics::InterlockedAdd(BytesPerSecCounter, (int32)(BytesPerSec));
+		const int32 BytesPerSec = (int32)((BytesToRead/1024.0f)/Delta);
+		FPlatformAtomics::InterlockedAdd(BytesPerSecCounter, BytesPerSec);
 	}
-	FPlatformAtomics::InterlockedAdd(BytesReadCounter, BytesToRead);
+	FPlatformAtomics::InterlockedAdd(BytesReadCounter, (int32)BytesToRead);
 	FPlatformAtomics::InterlockedIncrement(ReadsCounter);
 	return Result;
 }
diff --git a/Engine/Source/Runtime/Core/Private/HAL/LowLevelMemTracker.cpp b/Engine/Source/Runtime/Core/Private/HAL/LowLevelMemTracker.cpp
--- a/Engine/Source/Runtime/Core/Private/HAL/LowLevelMemTracker.cpp
+++ b/Engine/Source/Runtime/Core/Private/HAL/LowLevelMemTracker.cpp
@@ -866,10 +866,10 @@ bool FLowLevelMemTracker::Exec(const TCHAR* Cmd, FOutputDevice& Ar)
 
 			TArray<void*> Spam;
 			Spam.Reserve(NumAllocs);
-			uint32 TotalSize = 0;
+			SIZE_T TotalSize = 0;
 			for (int32 Index = 0; Index < NumAllocs; Index++)
 			{
-				int32 Size = (FPlatformMath::Rand() % MaxSize / 2) + MaxSize / 2;
+				SIZE_T Size = (FPlatformMath::Rand() % MaxSize / 2) + MaxSize / 2;
 				TotalSize += Size;
 				Spam.Add(FMemory::Malloc(Size));
 			}
@@ -1301,23 +1301,25 @@ void FLLMTracker::TrackAllocation(const void* Ptr, uint64 Size, ELLMTag DefaultT
 	if (Ptr != nullptr)
 	{
 		// remember the size and tag info
-		int64 tag = State->GetTopTag();
-		if (tag == (int64)ELLMTag::Untagged)
-			tag = (int64)DefaultTag;
+		int64 Tag = State->GetTopTag();
+		if (Tag == (int64)ELLMTag::Untagged)
+		{
+			Tag = (int64)DefaultTag;
+		}
 
 		FLLMTracker::FLowLevelAllocInfo AllocInfo;
 		#if LLM_USE_ALLOC_INFO_STRUCT
-		AllocInfo.Tag = tag;
+		AllocInfo.Tag = Tag;
 			#if LLM_ALLOW_ASSETS_TAGS
 		AllocInfo.AssetTag = State->GetTopAssetTag();
 			#endif
 		#else
-		LLMCheck(tag >= 0 && tag < (int64)LLM_TAG_COUNT);
-		AllocInfo = (ELLMTag)tag;
+		LLMCheck(Tag >= 0 && Tag < (int64)LLM_TAG_COUNT);
+		AllocInfo = (ELLMTag)Tag;
 		#endif
 
-		LLMCheck(Size <= 0xffffffff);
-		GetAllocationMap().Add(Ptr, Size, AllocInfo);
+		LLMCheck(Size <= 0xffffffffu);
+		GetAllocationMap().Add(Ptr, (uint32)Size, AllocInfo);
 	}
 }
 
diff --git a/Engine/Source/Runtime/Core/Private/HAL/MallocBinned.cpp b/Engine/Source/Runtime/Core/Private/HAL/MallocBinned.cpp
--- a/Engine/Source/Runtime/Core/Private/HAL/MallocBinned.cpp
+++ b/Engine/Source/Runtime/Core/Private/HAL/MallocBinned.cpp
@@ -1,15 +1,13 @@
 // Copyright Epic Games, Inc. All Rights Reserved.
 
-/*=============================================================================
-	MallocBinned.cpp: Binned memory allocator
-=============================================================================*/
-
 #include "HAL/MallocBinned.h"
 #include "Misc/ScopeLock.h"
 #include "Misc/BufferedOutputDevice.h"
 
 #include "HAL/MemoryMisc.h"
 
+PRAGMA_DISABLE_UNSAFE_TYPECAST_WARNINGS
+
 /** Malloc binned allocator specific stats. */
 DEFINE_STAT(STAT_Binned_OsCurrent);
 DEFINE_STAT(STAT_Binned_OsPeak);
@@ -413,7 +411,7 @@ struct FMallocBinned::Private
 		FPlatformMemory::OnOutOfMemory(Size, Alignment);
 	}
 
-	static FORCEINLINE void TrackStats(FPoolTable* Table, SIZE_T Size)
+	static FORCEINLINE void TrackStats(FPoolTable* Table, uint32 Size)
 	{
 #if STATS
 		// keep track of memory lost to padding
@@ -530,9 +528,9 @@ struct FMallocBinned::Private
 	{
 		checkSlow(Allocator.HashBuckets);
 
-		uint32 Key       = Ptr >> Allocator.HashKeyShift;
+		uint32 Key       = (uint32)(Ptr >> Allocator.HashKeyShift);
 		uint32 Hash      = Key & (Allocator.MaxHashBuckets - 1);
-		uint32 PoolIndex = ((UPTRINT)Ptr >> Allocator.PoolBitShift) & Allocator.PoolMask;
+		uint32 PoolIndex = (uint32)(((UPTRINT)Ptr >> Allocator.PoolBitShift) & Allocator.PoolMask);
 
 		JumpOffset = 0;
 
@@ -653,7 +651,7 @@ struct FMallocBinned::Private
 				check(TrailingPool);
 
 				//Set trailing pools to point back to first pool
-				TrailingPool->SetAllocationSizes(0, 0, Offset, Allocator.BinnedOSTableIndex);
+				TrailingPool->SetAllocationSizes(0, 0, (uint32)Offset, (uint32)Allocator.BinnedOSTableIndex);
 			}
 
 			
@@ -1367,7 +1365,7 @@ void* FMallocBinned::Malloc(SIZE_T Size, uint32 Alignment)
 #endif
 			checkSlow(Size <= Table->BlockSize);
 
-			Private::TrackStats(Table, Size);
+			Private::TrackStats(Table, (uint32)Size);
 
 			FPoolInfo* Pool = Table->FirstPool;
 			if( !Pool )
@@ -1389,7 +1387,7 @@ void* FMallocBinned::Malloc(SIZE_T Size, uint32 Alignment)
 #endif
 			checkSlow(Size <= Table->BlockSize);
 
-			Private::TrackStats(Table, Size);
+			Private::TrackStats(Table, (uint32)Size);
 
 			FPoolInfo* Pool = Table->FirstPool;
 			if( !Pool )
@@ -1860,3 +1858,5 @@ const TCHAR* FMallocBinned::GetDescriptiveName()
 {
 	return TEXT("binned");
 }
+
+PRAGMA_ENABLE_UNSAFE_TYPECAST_WARNINGS
diff --git a/Engine/Source/Runtime/Core/Private/HAL/MallocBinned2.cpp b/Engine/Source/Runtime/Core/Private/HAL/MallocBinned2.cpp
--- a/Engine/Source/Runtime/Core/Private/HAL/MallocBinned2.cpp
+++ b/Engine/Source/Runtime/Core/Private/HAL/MallocBinned2.cpp
@@ -1,9 +1,5 @@
 // Copyright Epic Games, Inc. All Rights Reserved.
 
-/*=============================================================================
-	MallocBinned.cpp: Binned memory allocator
-=============================================================================*/
-
 #include "HAL/MallocBinned2.h"
 #include "Logging/LogMacros.h"
 #include "Misc/ScopeLock.h"
@@ -15,6 +11,8 @@
 #include "HAL/PlatformMisc.h"
 #include "Misc/App.h"
 
+PRAGMA_DISABLE_UNSAFE_TYPECAST_WARNINGS
+
 #if BINNED2_ALLOW_RUNTIME_TWEAKING
 
 int32 GMallocBinned2PerThreadCaches = DEFAULT_GMallocBinned2PerThreadCaches;
@@ -1053,8 +1051,8 @@ void FMallocBinned2::FlushCurrentThreadCache()
 	QUICK_SCOPE_CYCLE_COUNTER(STAT_FMallocBinned2_FlushCurrentThreadCache);
 	FPerThreadFreeBlockLists* Lists = FPerThreadFreeBlockLists::Get();
 
-	float WaitForMutexTime = 0.0f;
-	float WaitForMutexAndTrimTime = 0.0f;
+	double WaitForMutexTime = 0.0;
+	double WaitForMutexAndTrimTime = 0.0;
 
 	if (Lists)
 	{
@@ -1309,3 +1307,5 @@ void FMallocBinned2::DumpAllocatorStats(class FOutputDevice& Ar)
 		#include "FMemory.inl"
 	#endif
 #endif
+
+PRAGMA_ENABLE_UNSAFE_TYPECAST_WARNINGS
diff --git a/Engine/Source/Runtime/Core/Private/HAL/MallocBinned3.cpp b/Engine/Source/Runtime/Core/Private/HAL/MallocBinned3.cpp
--- a/Engine/Source/Runtime/Core/Private/HAL/MallocBinned3.cpp
+++ b/Engine/Source/Runtime/Core/Private/HAL/MallocBinned3.cpp
@@ -1,11 +1,9 @@
 // Copyright Epic Games, Inc. All Rights Reserved.
 
-/*=============================================================================
-	MallocBinned.cpp: Binned memory allocator
-=============================================================================*/
-
 #include "HAL/MallocBinned3.h"
 
+PRAGMA_DISABLE_UNSAFE_TYPECAST_WARNINGS
+
 #if PLATFORM_64BITS && PLATFORM_HAS_FPlatformVirtualMemoryBlock
 #include "Logging/LogMacros.h"
 #include "Misc/ScopeLock.h"
@@ -16,8 +14,6 @@
 #include "HAL/MemoryMisc.h"
 #include "HAL/PlatformMisc.h"
 
-
-
 #if USE_CACHED_PAGE_ALLOCATOR_FOR_LARGE_ALLOCS
 #include "HAL/Allocators/CachedOSPageAllocator.h"
 #define BINNED3_MAX_CACHED_OS_FREES (64)
@@ -1593,4 +1589,6 @@ void FMallocBinned3::DumpAllocatorStats(class FOutputDevice& Ar)
 		#include "FMemory.inl"
 	#endif
 #endif
-#endif
\ No newline at end of file
+#endif
+
+PRAGMA_ENABLE_UNSAFE_TYPECAST_WARNINGS
diff --git a/Engine/Source/Runtime/Core/Private/HAL/MallocBinnedArena.cpp b/Engine/Source/Runtime/Core/Private/HAL/MallocBinnedArena.cpp
--- a/Engine/Source/Runtime/Core/Private/HAL/MallocBinnedArena.cpp
+++ b/Engine/Source/Runtime/Core/Private/HAL/MallocBinnedArena.cpp
@@ -1,11 +1,9 @@
 // Copyright Epic Games, Inc. All Rights Reserved.
 
-/*=============================================================================
-	MallocBinned.cpp: Binned memory allocator
-=============================================================================*/
-
 #include "HAL/MallocBinnedArena.h"
 
+PRAGMA_DISABLE_UNSAFE_TYPECAST_WARNINGS
+
 #if PLATFORM_64BITS && PLATFORM_HAS_FPlatformVirtualMemoryBlock
 #include "Logging/LogMacros.h"
 #include "Templates/Function.h"
@@ -1353,4 +1351,6 @@ void FMallocBinnedArena::DumpAllocatorStats(class FOutputDevice& Ar)
 		);
 	}
 }
-#endif
\ No newline at end of file
+#endif
+
+PRAGMA_ENABLE_UNSAFE_TYPECAST_WARNINGS
diff --git a/Engine/Source/Runtime/Core/Private/HAL/MallocBinnedCommon.cpp b/Engine/Source/Runtime/Core/Private/HAL/MallocBinnedCommon.cpp
--- a/Engine/Source/Runtime/Core/Private/HAL/MallocBinnedCommon.cpp
+++ b/Engine/Source/Runtime/Core/Private/HAL/MallocBinnedCommon.cpp
@@ -1,9 +1,5 @@
 // Copyright Epic Games, Inc. All Rights Reserved.
 
-/*=============================================================================
-	MallocBinned.cpp: Binned memory allocator
-=============================================================================*/
-
 #include "HAL/MallocBinnedCommon.h"
 #include "Misc/AssertionMacros.h"
 #include "Math/NumericLimits.h"
@@ -11,6 +7,8 @@
 #include "Templates/UnrealTemplate.h"
 #include "Templates/Sorting.h"
 
+PRAGMA_DISABLE_UNSAFE_TYPECAST_WARNINGS
+
 #if PLATFORM_HAS_FPlatformVirtualMemoryBlock
 
 // Block sizes are based around getting the maximum amount of allocations per pool, with as little alignment waste as possible.
@@ -155,7 +153,7 @@ uint8 FSizeTableEntry::FillSizeTable(uint64 PlatformPageSize, FSizeTableEntry* S
 		SizeTable[Index++] = FSizeTableEntry(Size, PlatformPageSize, Size / BasePageSize, BasePageSize, MinimumAlignment);
 	}
 	check(Index < 256);
-	return Index;
+	return (uint8)Index;
 }
 
 uint32 FBitTree::GetMemoryRequirements(uint32 DesiredCapacity)
@@ -473,4 +471,6 @@ uint32 FBitTree::CountOnes(uint32 UpTo) const
 	return Result;
 }
 
-#endif
\ No newline at end of file
+#endif
+
+PRAGMA_ENABLE_UNSAFE_TYPECAST_WARNINGS
diff --git a/Engine/Source/Runtime/Core/Private/HAL/MallocBinnedGPU.cpp b/Engine/Source/Runtime/Core/Private/HAL/MallocBinnedGPU.cpp
--- a/Engine/Source/Runtime/Core/Private/HAL/MallocBinnedGPU.cpp
+++ b/Engine/Source/Runtime/Core/Private/HAL/MallocBinnedGPU.cpp
@@ -1,11 +1,9 @@
 // Copyright Epic Games, Inc. All Rights Reserved.
 
-/*=============================================================================
-	MallocBinned.cpp: Binned memory allocator
-=============================================================================*/
-
 #include "HAL/MallocBinnedGPU.h"
 
+PRAGMA_DISABLE_UNSAFE_TYPECAST_WARNINGS
+
 #if PLATFORM_64BITS && PLATFORM_HAS_FPlatformVirtualMemoryBlock
 #include "Logging/LogMacros.h"
 #include "Templates/Function.h"
@@ -1040,8 +1038,8 @@ void FMallocBinnedGPU::FlushCurrentThreadCache()
 	QUICK_SCOPE_CYCLE_COUNTER(STAT_FMallocBinnedGPU_FlushCurrentThreadCache);
 	FPerThreadFreeBlockLists* Lists = FPerThreadFreeBlockLists::Get(BinnedGPUTlsSlot);
 
-	float WaitForMutexTime = 0.0f;
-	float WaitForMutexAndTrimTime = 0.0f;
+	double WaitForMutexTime = 0.0;
+	double WaitForMutexAndTrimTime = 0.0;
 
 	if (Lists)
 	{
@@ -1295,4 +1293,6 @@ void FMallocBinnedGPU::DumpAllocatorStats(class FOutputDevice& Ar)
 		);
 	}
 }
-#endif
\ No newline at end of file
+#endif
+
+PRAGMA_ENABLE_UNSAFE_TYPECAST_WARNINGS
diff --git a/Engine/Source/Runtime/Core/Private/HAL/UnrealMemory.cpp b/Engine/Source/Runtime/Core/Private/HAL/UnrealMemory.cpp
--- a/Engine/Source/Runtime/Core/Private/HAL/UnrealMemory.cpp
+++ b/Engine/Source/Runtime/Core/Private/HAL/UnrealMemory.cpp
@@ -119,7 +119,7 @@ public:
 			verify(GetAllocationSize(Ptr, Size) && Size);
 			FMemory::Memset(Ptr, uint8(PURGATORY_STOMP_CHECKS_CANARYBYTE), Size);
 			Purgatory[GFrameNumber % PURGATORY_STOMP_CHECKS_FRAMES].Push(Ptr);
-			OutstandingSizeInKB.Add((Size + 1023) / 1024);
+			OutstandingSizeInKB.Add((int32)((Size + 1023) / 1024));
 		}
 		FPlatformMisc::MemoryBarrier();
 		uint32 LocalLastCheckFrame = LastCheckFrame;
@@ -151,7 +151,7 @@ public:
 						}
 					}
 					UsedMalloc->Free(Pop);
-					OutstandingSizeInKB.Subtract((Size + 1023) / 1024);
+					OutstandingSizeInKB.Subtract((int32)((Size + 1023) / 1024));
 				}
 			}
 		}
@@ -633,13 +633,12 @@ void FMemory::RegisterPersistentAuxiliary(void* InMemory, SIZE_T InSize)
 	GPersistentAuxiliary = (uint8 *)InMemory;
 	GPersistentAuxiliaryEnd = GPersistentAuxiliary + InSize;
 }
-void* FMemory::MallocPersistentAuxiliary(SIZE_T InSize, SIZE_T InAlignment)
+void* FMemory::MallocPersistentAuxiliary(SIZE_T InSize, uint32 InAlignment)
 {
 	if (GPersistentAuxiliary != nullptr && GPersistentAuxiliaryEnabled)
 	{
-		SIZE_T Alignment = FMath::Max(InAlignment, (SIZE_T)16);
-		Alignment = FMath::Max(Alignment, (SIZE_T)16);
-		SIZE_T AlignedSize = Align(InSize, Alignment);
+		const uint32 Alignment = FMath::Max<uint32>(InAlignment, 16u);
+		const SIZE_T AlignedSize = Align(InSize, Alignment);
 		// 1st check if there is room, this is atomic but could still fail when actually incrementing the offset.
 		if (GPersistentAuxiliaryCurrentOffset + AlignedSize <= GPersistentAuxiliarySize)
 		{
diff --git a/Engine/Source/Runtime/Core/Private/Internationalization/FastDecimalFormat.cpp b/Engine/Source/Runtime/Core/Private/Internationalization/FastDecimalFormat.cpp
--- a/Engine/Source/Runtime/Core/Private/Internationalization/FastDecimalFormat.cpp
+++ b/Engine/Source/Runtime/Core/Private/Internationalization/FastDecimalFormat.cpp
@@ -399,7 +399,7 @@ void FractionalToString(const double InVal, const FDecimalNumberFormattingRules&
 {
 	SanitizeNumberFormattingOptions(InFormattingOptions);
 
-	if (FMath::IsNaN(InVal))
+	if (FMath::IsNaN((float)InVal)) //@TODO: FLOATPRECISION: Need a double version?
 	{
 		OutString.Append(InFormattingRules.NaNString);
 		return;
@@ -644,7 +644,7 @@ bool StringToIntegral(const TCHAR* InStr, const int32 InStrLen, const FDecimalNu
 	// Only fill in the length if we actually parsed some digits
 	if (IntegralPartDigitCount > 0 && OutParsedLen)
 	{
-		*OutParsedLen = Buffer - InStr;
+		*OutParsedLen = UE_PTRDIFF_TO_INT32(Buffer - InStr);
 	}
 
 	return bResult;
@@ -684,7 +684,7 @@ bool StringToFractional(const TCHAR* InStr, const int32 InStrLen, const FDecimal
 	// Only fill in the length if we actually parsed some digits
 	if (TotalDigitCount > 0 && OutParsedLen)
 	{
-		*OutParsedLen = Buffer - InStr;
+		*OutParsedLen = UE_PTRDIFF_TO_INT32(Buffer - InStr);
 	}
 
 	return bResult;
diff --git a/Engine/Source/Runtime/Core/Private/Internationalization/ICULineBreakIterator.cpp b/Engine/Source/Runtime/Core/Private/Internationalization/ICULineBreakIterator.cpp
--- a/Engine/Source/Runtime/Core/Private/Internationalization/ICULineBreakIterator.cpp
+++ b/Engine/Source/Runtime/Core/Private/Internationalization/ICULineBreakIterator.cpp
@@ -38,7 +38,7 @@ EHangulTextWrappingMethod GetHangulTextWrappingMethod()
 	return EHangulTextWrappingMethod::PerWord;
 }
 
-FORCEINLINE bool IsHangul(const TCHAR InChar)
+FORCEINLINE bool IsHangul(const UChar32 InChar)
 {
 	return InChar >= 0xAC00 && InChar <= 0xD7A3;
 }
@@ -177,7 +177,7 @@ int32 FICULineBreakIterator::MoveToPreviousImpl()
 		if (IsHangul(CharIt.current32()))
 		{
 			// Walk to the start of the Hangul characters
-			while (CharIt.hasPrevious() && IsHangul(static_cast<TCHAR>(CharIt.previous32())))
+			while (CharIt.hasPrevious() && IsHangul(CharIt.previous32()))
 			{
 				InternalPosition = CharIt.getIndex();
 			}
@@ -206,7 +206,7 @@ int32 FICULineBreakIterator::MoveToNextImpl()
 		if (IsHangul(CharIt.current32()))
 		{
 			// Walk to the end of the Hangul characters
-			while (CharIt.hasNext() && IsHangul(static_cast<TCHAR>(CharIt.next32())))
+			while (CharIt.hasNext() && IsHangul(CharIt.next32()))
 			{
 				InternalPosition = CharIt.getIndex();
 			}
diff --git a/Engine/Source/Runtime/Core/Private/Internationalization/Internationalization.cpp b/Engine/Source/Runtime/Core/Private/Internationalization/Internationalization.cpp
--- a/Engine/Source/Runtime/Core/Private/Internationalization/Internationalization.cpp
+++ b/Engine/Source/Runtime/Core/Private/Internationalization/Internationalization.cpp
@@ -351,7 +351,7 @@ FString& FInternationalization::Leetify(FString& SourceString)
 				const TCHAR* ArgumentEndPtr = FCString::Strchr(RawSourceStringPtr + SourceCharIndex + 1, SourceArgumentEndMarker);
 				if (ArgumentEndPtr)
 				{
-					const int32 ArgumentEndIndex = ArgumentEndPtr - RawSourceStringPtr;
+					const int32 ArgumentEndIndex = UE_PTRDIFF_TO_INT32(ArgumentEndPtr - RawSourceStringPtr);
 
 					// Inject a marker before the argument block
 					LeetifiedString.AppendChar(LeetifyArgumentStartMarker);
diff --git a/Engine/Source/Runtime/Core/Private/Internationalization/Text.cpp b/Engine/Source/Runtime/Core/Private/Internationalization/Text.cpp
--- a/Engine/Source/Runtime/Core/Private/Internationalization/Text.cpp
+++ b/Engine/Source/Runtime/Core/Private/Internationalization/Text.cpp
@@ -608,7 +608,7 @@ FText FText::AsCurrencyBase(int64 BaseVal, const FString& CurrencyCode, const FC
 
 	const FDecimalNumberFormattingRules& FormattingRules = Culture.GetCurrencyFormattingRules(CurrencyCode);
 	const FNumberFormattingOptions& FormattingOptions = FormattingRules.CultureDefaultFormattingOptions;
-	double Val = static_cast<double>(BaseVal) / FMath::Pow(10.0f, FormattingOptions.MaximumFractionalDigits);
+	double Val = static_cast<double>(BaseVal) / FMath::Pow(10.0f, (float)FormattingOptions.MaximumFractionalDigits);
 	FString NativeString = FastDecimalFormat::NumberToString(Val, FormattingRules, FormattingOptions);
 
 	FText Result = FText(MakeShared<TGeneratedTextData<FTextHistory_AsCurrency>, ESPMode::ThreadSafe>(MoveTemp(NativeString), FTextHistory_AsCurrency(Val, CurrencyCode, nullptr, TargetCulture)));
@@ -1192,7 +1192,7 @@ bool FText::IdenticalTo( const FText& Other ) const
 void operator<<(FStructuredArchive::FSlot Slot, FFormatArgumentValue& Value)
 {
 	FStructuredArchive::FRecord Record = Slot.EnterRecord();
-	int8 TypeAsInt8 = Value.GetType();
+	int8 TypeAsInt8 = (int8)Value.GetType();
 	Record << SA_VALUE(TEXT("Type"), TypeAsInt8);
 	Value.Type = (EFormatArgumentType::Type)TypeAsInt8;
 
@@ -1413,7 +1413,7 @@ void operator<<(FStructuredArchive::FSlot Slot, FFormatArgumentData& Value)
 		Record << SA_VALUE(TEXT("ArgumentName"), Value.ArgumentName);
 	}
 
-	uint8 TypeAsByte = Value.ArgumentValueType;
+	uint8 TypeAsByte = (uint8)Value.ArgumentValueType;
 	if (UnderlyingArchive.IsLoading())
 	{
 		Value.ResetValue();
@@ -1710,7 +1710,7 @@ bool FTextStringHelper::ReadFromString(const TCHAR* Buffer, FText& OutValue, con
 	{
 		if (OutNumCharsRead)
 		{
-			*OutNumCharsRead = (Buffer - Start);
+			*OutNumCharsRead = UE_PTRDIFF_TO_INT32(Buffer - Start);
 		}
 		return true;
 	}
diff --git a/Engine/Source/Runtime/Core/Private/Internationalization/TextFormatArgumentModifier.cpp b/Engine/Source/Runtime/Core/Private/Internationalization/TextFormatArgumentModifier.cpp
--- a/Engine/Source/Runtime/Core/Private/Internationalization/TextFormatArgumentModifier.cpp
+++ b/Engine/Source/Runtime/Core/Private/Internationalization/TextFormatArgumentModifier.cpp
@@ -26,7 +26,7 @@ bool ITextFormatArgumentModifier::ParseKeyValueArgs(const FTextFormatString& InA
 		{
 			++BufferPtr;
 		}
-		return FTextFormatString::MakeReference(KeyStartPtr, BufferPtr - KeyStartPtr);
+		return FTextFormatString::MakeReference(KeyStartPtr, UE_PTRDIFF_TO_INT32(BufferPtr - KeyStartPtr));
 	};
 
 	auto ParseValue = [&]() -> FTextFormatString
@@ -47,7 +47,7 @@ bool ITextFormatArgumentModifier::ParseKeyValueArgs(const FTextFormatString& InA
 		{
 			++BufferPtr;
 		}
-		return FTextFormatString::MakeReference(ValueStartPtr, BufferPtr - ValueStartPtr);
+		return FTextFormatString::MakeReference(ValueStartPtr, UE_PTRDIFF_TO_INT32(BufferPtr - ValueStartPtr));
 	};
 
 	// Skip leading whitespace in case this string is all whitespace
@@ -130,7 +130,7 @@ bool ITextFormatArgumentModifier::ParseValueArgs(const FTextFormatString& InArgs
 		{
 			++BufferPtr;
 		}
-		return FTextFormatString::MakeReference(ValueStartPtr, BufferPtr - ValueStartPtr);
+		return FTextFormatString::MakeReference(ValueStartPtr, UE_PTRDIFF_TO_INT32(BufferPtr - ValueStartPtr));
 	};
 
 	// Skip leading whitespace in case this string is all whitespace
diff --git a/Engine/Source/Runtime/Core/Private/Internationalization/TextFormatter.cpp b/Engine/Source/Runtime/Core/Private/Internationalization/TextFormatter.cpp
--- a/Engine/Source/Runtime/Core/Private/Internationalization/TextFormatter.cpp
+++ b/Engine/Source/Runtime/Core/Private/Internationalization/TextFormatter.cpp
@@ -24,7 +24,7 @@ struct FStringLiteral
 {
 	explicit FStringLiteral(const FStringToken& InString)
 		: StringStartPos(InString.GetTokenStartPos())
-		, StringLen(InString.GetTokenEndPos() - InString.GetTokenStartPos())
+		, StringLen(UE_PTRDIFF_TO_INT32(InString.GetTokenEndPos() - InString.GetTokenStartPos()))
 	{
 	}
 
@@ -40,7 +40,7 @@ struct FArgumentTokenSpecifier
 {
 	explicit FArgumentTokenSpecifier(const FStringToken& InArgument)
 		: ArgumentNameStartPos(InArgument.GetTokenStartPos())
-		, ArgumentNameLen(InArgument.GetTokenEndPos() - InArgument.GetTokenStartPos())
+		, ArgumentNameLen(UE_PTRDIFF_TO_INT32(InArgument.GetTokenEndPos() - InArgument.GetTokenStartPos()))
 		, ArgumentIndex(INDEX_NONE)
 	{
 		if (ArgumentNameLen > 0)
@@ -79,7 +79,7 @@ struct FArgumentModifierTokenSpecifier
 {
 	FArgumentModifierTokenSpecifier(const FStringToken& InModifierPatternWithPipe, TSharedRef<ITextFormatArgumentModifier> InTextFormatArgumentModifier)
 		: ModifierPatternStartPos(InModifierPatternWithPipe.GetTokenStartPos() + 1) // We don't want to store the pipe
-		, ModifierPatternLen(InModifierPatternWithPipe.GetTokenEndPos() - InModifierPatternWithPipe.GetTokenStartPos() - 1)
+		, ModifierPatternLen(UE_PTRDIFF_TO_INT32(InModifierPatternWithPipe.GetTokenEndPos() - InModifierPatternWithPipe.GetTokenStartPos() - 1))
 		, TextFormatArgumentModifier(MoveTemp(InTextFormatArgumentModifier))
 	{
 	}
@@ -180,7 +180,7 @@ TOptional<FExpressionError> ParseArgumentModifier(const FTextFormatPatternDefini
 
 	// Valid modifier name?
 	FStringToken& IdentifierValue = Identifier.GetValue();
-	FTextFormatter::FCompileTextArgumentModifierFuncPtr CompileTextArgumentModifierFunc = FTextFormatter::Get().FindTextArgumentModifier(FTextFormatString::MakeReference(IdentifierValue.GetTokenStartPos(), IdentifierValue.GetTokenEndPos() - IdentifierValue.GetTokenStartPos()));
+	FTextFormatter::FCompileTextArgumentModifierFuncPtr CompileTextArgumentModifierFunc = FTextFormatter::Get().FindTextArgumentModifier(FTextFormatString::MakeReference(IdentifierValue.GetTokenStartPos(), UE_PTRDIFF_TO_INT32(IdentifierValue.GetTokenEndPos() - IdentifierValue.GetTokenStartPos())));
 	if (!CompileTextArgumentModifierFunc)
 	{
 		return TOptional<FExpressionError>();
@@ -232,7 +232,7 @@ TOptional<FExpressionError> ParseArgumentModifier(const FTextFormatPatternDefini
 
 	// Compile the parameters for this argument modifier
 	FStringToken& ParametersValue = Parameters.GetValue();
-	TSharedPtr<ITextFormatArgumentModifier> CompiledTextArgumentModifier = CompileTextArgumentModifierFunc(FTextFormatString::MakeReference(ParametersValue.GetTokenStartPos(), ParametersValue.GetTokenEndPos() - ParametersValue.GetTokenStartPos()), InPatternDef.AsShared());
+	TSharedPtr<ITextFormatArgumentModifier> CompiledTextArgumentModifier = CompileTextArgumentModifierFunc(FTextFormatString::MakeReference(ParametersValue.GetTokenStartPos(), UE_PTRDIFF_TO_INT32(ParametersValue.GetTokenEndPos() - ParametersValue.GetTokenStartPos())), InPatternDef.AsShared());
 	if (!CompiledTextArgumentModifier.IsValid())
 	{
 		return TOptional<FExpressionError>();
diff --git a/Engine/Source/Runtime/Core/Private/Internationalization/TextHistory.cpp b/Engine/Source/Runtime/Core/Private/Internationalization/TextHistory.cpp
--- a/Engine/Source/Runtime/Core/Private/Internationalization/TextHistory.cpp
+++ b/Engine/Source/Runtime/Core/Private/Internationalization/TextHistory.cpp
@@ -550,10 +550,10 @@ const TCHAR* ReadDateTimeFromBuffer(const TCHAR* Buffer, const FString& TokenMar
 			OutDateTime = FDateTime::FromUnixTimestamp(UnixTimestampValue.GetUIntValue());
 			break;
 		case EFormatArgumentType::Float:
-			OutDateTime = FDateTime::FromUnixTimestamp(UnixTimestampValue.GetFloatValue());
+			OutDateTime = FDateTime::FromUnixTimestamp((int64)UnixTimestampValue.GetFloatValue());
 			break;
 		case EFormatArgumentType::Double:
-			OutDateTime = FDateTime::FromUnixTimestamp(UnixTimestampValue.GetDoubleValue());
+			OutDateTime = FDateTime::FromUnixTimestamp((int64)UnixTimestampValue.GetDoubleValue());
 			break;
 		default:
 			return nullptr;
@@ -1831,10 +1831,10 @@ const TCHAR* FTextHistory_AsCurrency::ReadFromBuffer(const TCHAR* Buffer, const
 		switch (SourceValue.GetType())
 		{
 		case EFormatArgumentType::Int:
-			BaseValue = SourceValue.GetIntValue();
+			BaseValue = (double)SourceValue.GetIntValue();
 			break;
 		case EFormatArgumentType::UInt:
-			BaseValue = SourceValue.GetUIntValue();
+			BaseValue = (double)SourceValue.GetUIntValue();
 			break;
 		case EFormatArgumentType::Float:
 			BaseValue = SourceValue.GetFloatValue();
@@ -1849,7 +1849,7 @@ const TCHAR* FTextHistory_AsCurrency::ReadFromBuffer(const TCHAR* Buffer, const
 		// We need to convert the "base" value back to its pre-divided version
 		const FDecimalNumberFormattingRules& FormattingRules = Culture.GetCurrencyFormattingRules(CurrencyCode);
 		const FNumberFormattingOptions& FormattingOptions = FormattingRules.CultureDefaultFormattingOptions;
-		SourceValue = BaseValue / FMath::Pow(10.0f, FormattingOptions.MaximumFractionalDigits);
+		SourceValue = BaseValue / FMath::Pow(10.0f, (float)FormattingOptions.MaximumFractionalDigits);
 
 		PrepareDisplayStringForRebuild(OutDisplayString);
 		return Buffer;
@@ -1869,10 +1869,10 @@ bool FTextHistory_AsCurrency::WriteToBuffer(FString& Buffer, FTextDisplayStringP
 	switch (SourceValue.GetType())
 	{
 	case EFormatArgumentType::Int:
-		DividedValue = SourceValue.GetIntValue();
+		DividedValue = (double)SourceValue.GetIntValue();
 		break;
 	case EFormatArgumentType::UInt:
-		DividedValue = SourceValue.GetUIntValue();
+		DividedValue = (double)SourceValue.GetUIntValue();
 		break;
 	case EFormatArgumentType::Float:
 		DividedValue = SourceValue.GetFloatValue();
@@ -1887,7 +1887,7 @@ bool FTextHistory_AsCurrency::WriteToBuffer(FString& Buffer, FTextDisplayStringP
 	// We need to convert the value back to its "base" version
 	const FDecimalNumberFormattingRules& FormattingRules = Culture.GetCurrencyFormattingRules(CurrencyCode);
 	const FNumberFormattingOptions& FormattingOptions = FormattingRules.CultureDefaultFormattingOptions;
-	const int64 BaseVal = static_cast<int64>(DividedValue * FMath::Pow(10.0f, FormattingOptions.MaximumFractionalDigits));
+	const int64 BaseVal = static_cast<int64>(DividedValue * FMath::Pow(10.0f, (float)FormattingOptions.MaximumFractionalDigits));
 
 	// Produces LOCGEN_CURRENCY(..., "...", "...")
 	Buffer += TEXT("LOCGEN_CURRENCY(");
diff --git a/Engine/Source/Runtime/Core/Private/Math/BasicMathExpressionEvaluator.cpp b/Engine/Source/Runtime/Core/Private/Math/BasicMathExpressionEvaluator.cpp
--- a/Engine/Source/Runtime/Core/Private/Math/BasicMathExpressionEvaluator.cpp
+++ b/Engine/Source/Runtime/Core/Private/Math/BasicMathExpressionEvaluator.cpp
@@ -40,11 +40,11 @@ namespace ExpressionParser
 		// This call will return false if there is some other data after the number, which is why we check the parsed length instead
 		double PrimaryValue = 0.0;
 		int32 PrimaryParsedLen = 0;
-		FastDecimalFormat::StringToNumber(InStream.GetRead(), InStream.GetEnd() - InStream.GetRead(), InPrimaryFormattingRules, FNumberParsingOptions::DefaultNoGrouping(), PrimaryValue, &PrimaryParsedLen);
+		FastDecimalFormat::StringToNumber(InStream.GetRead(), UE_PTRDIFF_TO_INT32(InStream.GetEnd() - InStream.GetRead()), InPrimaryFormattingRules, FNumberParsingOptions::DefaultNoGrouping(), PrimaryValue, &PrimaryParsedLen);
 
 		double FallbackValue = 0.0;
 		int32 FallbackParsedLen = 0;
-		FastDecimalFormat::StringToNumber(InStream.GetRead(), InStream.GetEnd() - InStream.GetRead(), InFallbackFormattingRules, FNumberParsingOptions::DefaultNoGrouping(), FallbackValue, &FallbackParsedLen);
+		FastDecimalFormat::StringToNumber(InStream.GetRead(), UE_PTRDIFF_TO_INT32(InStream.GetEnd() - InStream.GetRead()), InFallbackFormattingRules, FNumberParsingOptions::DefaultNoGrouping(), FallbackValue, &FallbackParsedLen);
 
 		// We take whichever value parsed the most text from the string
 		if (FallbackParsedLen <= PrimaryParsedLen)
@@ -73,7 +73,7 @@ namespace ExpressionParser
 		// This call will return false if there is some other data after the number, which is why we check the parsed length instead
 		double Value = 0.0;
 		int32 ParsedLen = 0;
-		FastDecimalFormat::StringToNumber(InStream.GetRead(), InStream.GetEnd() - InStream.GetRead(), InFormattingRules, FNumberParsingOptions::DefaultNoGrouping(), Value, &ParsedLen);
+		FastDecimalFormat::StringToNumber(InStream.GetRead(), UE_PTRDIFF_TO_INT32(InStream.GetEnd() - InStream.GetRead()), InFormattingRules, FNumberParsingOptions::DefaultNoGrouping(), Value, &ParsedLen);
 
 		if (OutValue)
 		{
@@ -172,12 +172,12 @@ FBasicMathExpressionEvaluator::FBasicMathExpressionEvaluator()
 
 	JumpTable.MapPreUnary<FPlus>([](double N)			{ return N; });
 	JumpTable.MapPreUnary<FMinus>([](double N)			{ return -N; });
-	JumpTable.MapPreUnary<FSquareRoot>([](double A)		{ return double(FMath::Sqrt(A)); });
+	JumpTable.MapPreUnary<FSquareRoot>([](double A)		{ return double(FMath::Sqrt((float)A)); }); //@TODO: FLOATPRECISION: Needs a double version of Sqrt
 
 	JumpTable.MapBinary<FPlus>([](double A, double B)	{ return A + B; });
 	JumpTable.MapBinary<FMinus>([](double A, double B)	{ return A - B; });
 	JumpTable.MapBinary<FStar>([](double A, double B)	{ return A * B; });
-	JumpTable.MapBinary<FPower>([](double A, double B)	{ return double(FMath::Pow(A, B)); });
+	JumpTable.MapBinary<FPower>([](double A, double B)	{ return double(FMath::Pow((float)A, (float)B)); }); //@TODO: FLOATPRECISION: Needs a double version of Pow
 
 	JumpTable.MapBinary<FForwardSlash>([](double A, double B) -> FExpressionResult {
 		if (B == 0)
@@ -193,7 +193,7 @@ FBasicMathExpressionEvaluator::FBasicMathExpressionEvaluator()
 			return MakeError(LOCTEXT("ModZero", "Modulo zero"));
 		}
 
-		return MakeValue(double(FMath::Fmod(A, B)));
+		return MakeValue(double(FMath::Fmod((float)A, (float)B))); //@TODO: FLOATPRECISION: Needs a double version of FMod
 	});
 }
 
diff --git a/Engine/Source/Runtime/Core/Private/Math/Color.cpp b/Engine/Source/Runtime/Core/Private/Math/Color.cpp
--- a/Engine/Source/Runtime/Core/Private/Math/Color.cpp
+++ b/Engine/Source/Runtime/Core/Private/Math/Color.cpp
@@ -98,19 +98,19 @@ FColor FLinearColor::ToRGBE() const
 	const float	Primary = FMath::Max3( R, G, B );
 	FColor	Color;
 
-	if( Primary < 1E-32 )
+	if( Primary < 1E-32f )
 	{
 		Color = FColor(0,0,0,0);
 	}
 	else
 	{
-		int32	Exponent;
-		const float Scale	= frexp(Primary, &Exponent) / Primary * 255.f;
+		int32 Exponent;
+		const float Scale = (float)frexp(Primary, &Exponent) / Primary * 255.f;
 
-		Color.R		= FMath::Clamp(FMath::TruncToInt(R * Scale), 0, 255);
-		Color.G		= FMath::Clamp(FMath::TruncToInt(G * Scale), 0, 255);
-		Color.B		= FMath::Clamp(FMath::TruncToInt(B * Scale), 0, 255);
-		Color.A		= FMath::Clamp(FMath::TruncToInt(Exponent),-128,127) + 128;
+		Color.R		= (uint8)FMath::Clamp(FMath::TruncToInt(R * Scale), 0, 255);
+		Color.G		= (uint8)FMath::Clamp(FMath::TruncToInt(G * Scale), 0, 255);
+		Color.B		= (uint8)FMath::Clamp(FMath::TruncToInt(B * Scale), 0, 255);
+		Color.A		= (uint8)(FMath::Clamp(Exponent,-128,127) + 128);
 	}
 
 	return Color;
@@ -132,14 +132,14 @@ FColor FLinearColor::ToFColor(const bool bSRGB) const
 		FloatB = FloatB <= 0.0031308f ? FloatB * 12.92f : FMath::Pow( FloatB, 1.0f / 2.4f ) * 1.055f - 0.055f;
 	}
 
-	FColor ret;
+	FColor Result;
 
-	ret.A = FMath::FloorToInt(FloatA * 255.999f);
-	ret.R = FMath::FloorToInt(FloatR * 255.999f);
-	ret.G = FMath::FloorToInt(FloatG * 255.999f);
-	ret.B = FMath::FloorToInt(FloatB * 255.999f);
+	Result.A = (uint8)FMath::FloorToInt(FloatA * 255.999f);
+	Result.R = (uint8)FMath::FloorToInt(FloatR * 255.999f);
+	Result.G = (uint8)FMath::FloorToInt(FloatG * 255.999f);
+	Result.B = (uint8)FMath::FloorToInt(FloatB * 255.999f);
 
-	return ret;
+	return Result;
 }
 
 
@@ -185,16 +185,16 @@ FColor FColor::FromHex( const FString& HexString )
 		const int32 G = FParse::HexDigit(HexString[StartIndex++]);
 		const int32 B = FParse::HexDigit(HexString[StartIndex]);
 
-		return FColor((R << 4) + R, (G << 4) + G, (B << 4) + B, 255);
+		return FColor((uint8)((R << 4) + R), (uint8)((G << 4) + G), (uint8)((B << 4) + B, 255));
 	}
 
 	if (HexString.Len() == 6 + StartIndex)
 	{
 		FColor Result;
 
-		Result.R = (FParse::HexDigit(HexString[StartIndex+0]) << 4) + FParse::HexDigit(HexString[StartIndex+1]);
-		Result.G = (FParse::HexDigit(HexString[StartIndex+2]) << 4) + FParse::HexDigit(HexString[StartIndex+3]);
-		Result.B = (FParse::HexDigit(HexString[StartIndex+4]) << 4) + FParse::HexDigit(HexString[StartIndex+5]);
+		Result.R = (uint8)((FParse::HexDigit(HexString[StartIndex+0]) << 4) + FParse::HexDigit(HexString[StartIndex+1]));
+		Result.G = (uint8)((FParse::HexDigit(HexString[StartIndex+2]) << 4) + FParse::HexDigit(HexString[StartIndex+3]));
+		Result.B = (uint8)((FParse::HexDigit(HexString[StartIndex+4]) << 4) + FParse::HexDigit(HexString[StartIndex+5]));
 		Result.A = 255;
 
 		return Result;
@@ -204,10 +204,10 @@ FColor FColor::FromHex( const FString& HexString )
 	{
 		FColor Result;
 
-		Result.R = (FParse::HexDigit(HexString[StartIndex+0]) << 4) + FParse::HexDigit(HexString[StartIndex+1]);
-		Result.G = (FParse::HexDigit(HexString[StartIndex+2]) << 4) + FParse::HexDigit(HexString[StartIndex+3]);
-		Result.B = (FParse::HexDigit(HexString[StartIndex+4]) << 4) + FParse::HexDigit(HexString[StartIndex+5]);
-		Result.A = (FParse::HexDigit(HexString[StartIndex+6]) << 4) + FParse::HexDigit(HexString[StartIndex+7]);
+		Result.R = (uint8)((FParse::HexDigit(HexString[StartIndex+0]) << 4) + FParse::HexDigit(HexString[StartIndex+1]));
+		Result.G = (uint8)((FParse::HexDigit(HexString[StartIndex+2]) << 4) + FParse::HexDigit(HexString[StartIndex+3]));
+		Result.B = (uint8)((FParse::HexDigit(HexString[StartIndex+4]) << 4) + FParse::HexDigit(HexString[StartIndex+5]));
+		Result.A = (uint8)((FParse::HexDigit(HexString[StartIndex+6]) << 4) + FParse::HexDigit(HexString[StartIndex+7]));
 
 		return Result;
 	}
@@ -218,11 +218,13 @@ FColor FColor::FromHex( const FString& HexString )
 // Convert from RGBE to float as outlined in Gregory Ward's Real Pixels article, Graphics Gems II, page 80.
 FLinearColor FColor::FromRGBE() const
 {
-	if( A == 0 )
+	if (A == 0)
+	{
 		return FLinearColor::Black;
+	}
 	else
 	{
-		const float Scale = ldexp( 1 / 255.0, A - 128 );
+		const float Scale = (float)ldexp( 1 / 255.0f, A - 128 );
 		return FLinearColor( R * Scale, G * Scale, B * Scale, 1.0f );
 	}
 }
@@ -392,11 +394,11 @@ FColor FColor::MakeFromColorTemperature( float Temp )
 
 FColor FColor::MakeRedToGreenColorFromScalar(float Scalar)
 {
-	float RedSclr = FMath::Clamp<float>((1.0f - Scalar)/0.5f,0.f,1.f);
-	float GreenSclr = FMath::Clamp<float>((Scalar/0.5f),0.f,1.f);
-	int32 R = FMath::TruncToInt(255 * RedSclr);
-	int32 G = FMath::TruncToInt(255 * GreenSclr);
-	int32 B = 0;
+	const float RedSclr = FMath::Clamp<float>((1.0f - Scalar)/0.5f,0.f,1.f);
+	const float GreenSclr = FMath::Clamp<float>((Scalar/0.5f),0.f,1.f);
+	const uint8 R = (uint8)FMath::TruncToInt(255 * RedSclr);
+	const uint8 G = (uint8)FMath::TruncToInt(255 * GreenSclr);
+	const uint8 B = 0;
 	return FColor(R, G, B);
 }
 
diff --git a/Engine/Source/Runtime/Core/Private/Math/UnitConversion.cpp b/Engine/Source/Runtime/Core/Private/Math/UnitConversion.cpp
--- a/Engine/Source/Runtime/Core/Private/Math/UnitConversion.cpp
+++ b/Engine/Source/Runtime/Core/Private/Math/UnitConversion.cpp
@@ -6,6 +6,8 @@
 #include "Misc/ExpressionParser.h"
 #include "Math/BasicMathExpressionEvaluator.h"
 
+PRAGMA_DISABLE_UNSAFE_TYPECAST_WARNINGS
+
 #define LOCTEXT_NAMESPACE "UnitConversion"
 
 /** Structure used to match units when parsing */
@@ -775,3 +777,5 @@ namespace UnitConversion
 }	// namespace UnitConversion
 
 #undef LOCTEXT_NAMESPACE
+
+PRAGMA_ENABLE_UNSAFE_TYPECAST_WARNINGS
diff --git a/Engine/Source/Runtime/Core/Private/Math/UnrealMath.cpp b/Engine/Source/Runtime/Core/Private/Math/UnrealMath.cpp
--- a/Engine/Source/Runtime/Core/Private/Math/UnrealMath.cpp
+++ b/Engine/Source/Runtime/Core/Private/Math/UnrealMath.cpp
@@ -1791,8 +1791,8 @@ static bool ComputeProjectedSphereShaft(
 	int32& InOutMaxX
 	)
 {
-	float ViewX = InOutMinX;
-	float ViewSizeX = InOutMaxX - InOutMinX;
+	float ViewX = (float)InOutMinX;
+	float ViewSizeX = (float)(InOutMaxX - InOutMinX);
 
 	// Vertical planes: T = <Nx, 0, Nz, 0>
 	float Discriminant = (FMath::Square(LightX) - FMath::Square(Radius) + FMath::Square(LightZ)) * FMath::Square(LightZ);
@@ -1816,11 +1816,11 @@ static bool ComputeProjectedSphereShaft(
 			float X = (Dot3(P,Axis) / P.W + 1.0f * AxisSign) / 2.0f * AxisSign;
 			if(FMath::IsNegativeFloat(Nxa) ^ FMath::IsNegativeFloat(AxisSign))
 			{
-				InOutMaxX = FMath::Min<int64>(FMath::CeilToInt(ViewSizeX * X + ViewX),InOutMaxX);
+				InOutMaxX = FMath::Min<int32>(FMath::CeilToInt(ViewSizeX * X + ViewX),InOutMaxX);
 			}
 			else
 			{
-				InOutMinX = FMath::Max<int64>(FMath::FloorToInt(ViewSizeX * X + ViewX),InOutMinX);
+				InOutMinX = FMath::Max<int32>(FMath::FloorToInt(ViewSizeX * X + ViewX),InOutMinX);
 			}
 		}
 
@@ -1832,11 +1832,11 @@ static bool ComputeProjectedSphereShaft(
 			float X = (Dot3(P,Axis) / P.W + 1.0f * AxisSign) / 2.0f * AxisSign;
 			if(FMath::IsNegativeFloat(Nxb) ^ FMath::IsNegativeFloat(AxisSign))
 			{
-				InOutMaxX = FMath::Min<int64>(FMath::CeilToInt(ViewSizeX * X + ViewX),InOutMaxX);
+				InOutMaxX = FMath::Min<int32>(FMath::CeilToInt(ViewSizeX * X + ViewX),InOutMaxX);
 			}
 			else
 			{
-				InOutMinX = FMath::Max<int64>(FMath::FloorToInt(ViewSizeX * X + ViewX),InOutMinX);
+				InOutMinX = FMath::Max<int32>(FMath::FloorToInt(ViewSizeX * X + ViewX),InOutMinX);
 			}
 		}
 	}
@@ -2556,9 +2556,9 @@ CORE_API FQuat FMath::QInterpTo(const FQuat& Current, const FQuat& Target, float
 
 CORE_API float ClampFloatTangent( float PrevPointVal, float PrevTime, float CurPointVal, float CurTime, float NextPointVal, float NextTime )
 {
-	const float PrevToNextTimeDiff = FMath::Max< double >( KINDA_SMALL_NUMBER, NextTime - PrevTime );
-	const float PrevToCurTimeDiff = FMath::Max< double >( KINDA_SMALL_NUMBER, CurTime - PrevTime );
-	const float CurToNextTimeDiff = FMath::Max< double >( KINDA_SMALL_NUMBER, NextTime - CurTime );
+	const float PrevToNextTimeDiff = FMath::Max< float >( KINDA_SMALL_NUMBER, NextTime - PrevTime );
+	const float PrevToCurTimeDiff = FMath::Max< float >( KINDA_SMALL_NUMBER, CurTime - PrevTime );
+	const float CurToNextTimeDiff = FMath::Max< float >( KINDA_SMALL_NUMBER, NextTime - CurTime );
 
 	float OutTangentVal = 0.0f;
 
@@ -3238,7 +3238,7 @@ PrecLoop:
 				}
 				else
 				{
-					V = (int32)V % (int32)W;
+					V = (float)((int32)V % (int32)W);
 					c = GrabChar(pStr);
 					goto PrecLoop;
 				}
@@ -3487,7 +3487,7 @@ namespace FMathPerlinHelpers
 	// Curve w/ second derivative vanishing at 0 and 1, from Perlin's improved noise paper
 	FORCEINLINE float SmoothCurve(float X)
 	{
-		return X * X * X * (X * (X * 6.0 - 15.0) + 10.0);
+		return X * X * X * (X * (X * 6.0f - 15.0f) + 10.0f);
 	}
 };
 
diff --git a/Engine/Source/Runtime/Core/Private/Memory/MemoryArena.cpp b/Engine/Source/Runtime/Core/Private/Memory/MemoryArena.cpp
--- a/Engine/Source/Runtime/Core/Private/Memory/MemoryArena.cpp
+++ b/Engine/Source/Runtime/Core/Private/Memory/MemoryArena.cpp
@@ -7,6 +7,8 @@
 #include "Misc/ScopeRWLock.h"
 #include <atomic>
 
+PRAGMA_DISABLE_UNSAFE_TYPECAST_WARNINGS
+
 #if UE_WITH_HEAPARENA
 // dlmalloc for use in FHeapArena
 
@@ -113,7 +115,7 @@ static uint16 AllocArenaId(FMemoryArena* Arena)
 	Tracker.KnownArenas[NewArenaId].Arena	= Arena;
 	Tracker.NextFreeArenaIndex				= NextFree;
 
-	return NewArenaId;
+	return (uint16)NewArenaId;
 }
 
 static void FreeArenaId(uint16 ArenaId, FMemoryArena* Arena)
@@ -474,3 +476,5 @@ FMemoryArena* FArenaMap::MapPtrToArena(const void* VaBase)
 #endif
 
 //////////////////////////////////////////////////////////////////////////
+
+PRAGMA_ENABLE_UNSAFE_TYPECAST_WARNINGS
diff --git a/Engine/Source/Runtime/Core/Private/Misc/AutomationTest.cpp b/Engine/Source/Runtime/Core/Private/Misc/AutomationTest.cpp
--- a/Engine/Source/Runtime/Core/Private/Misc/AutomationTest.cpp
+++ b/Engine/Source/Runtime/Core/Private/Misc/AutomationTest.cpp
@@ -164,7 +164,7 @@ bool FAutomationTestFramework::RunSmokeTests()
 			TMap<FString, FAutomationTestExecutionInfo> OutExecutionInfoMap;
 
 			// Run each valid test
-			FScopedSlowTask SlowTask(TestInfo.Num());
+			FScopedSlowTask SlowTask((float)TestInfo.Num());
 
 			// We disable capturing the stack when running smoke tests, it adds too much overhead to do it at startup.
 			FAutomationTestFramework::Get().SetCaptureStack(false);
diff --git a/Engine/Source/Runtime/Core/Private/Misc/Base64.cpp b/Engine/Source/Runtime/Core/Private/Misc/Base64.cpp
--- a/Engine/Source/Runtime/Core/Private/Misc/Base64.cpp
+++ b/Engine/Source/Runtime/Core/Private/Misc/Base64.cpp
@@ -1,4 +1,4 @@
-﻿// Copyright Epic Games, Inc. All Rights Reserved.
+// Copyright Epic Games, Inc. All Rights Reserved.
 
 #include "Misc/Base64.h"
 #include "Containers/StringConv.h"
@@ -121,7 +121,7 @@ template<typename CharType> uint32 FBase64::Encode(const uint8* Source, uint32 L
 	// Add a null terminator
 	*EncodedBytes = 0;
 
-	return EncodedBytes - Dest;
+	return UE_PTRDIFF_TO_UINT32(EncodedBytes - Dest);
 }
 
 template CORE_API uint32 FBase64::Encode<ANSICHAR>(const uint8* Source, uint32 Length, ANSICHAR* Dest);
diff --git a/Engine/Source/Runtime/Core/Private/Misc/Compression.cpp b/Engine/Source/Runtime/Core/Private/Misc/Compression.cpp
--- a/Engine/Source/Runtime/Core/Private/Misc/Compression.cpp
+++ b/Engine/Source/Runtime/Core/Private/Misc/Compression.cpp
@@ -23,6 +23,8 @@ DEFINE_LOG_CATEGORY(LogCompression);
 
 DECLARE_STATS_GROUP( TEXT( "Compression" ), STATGROUP_Compression, STATCAT_Advanced );
 
+PRAGMA_DISABLE_UNSAFE_TYPECAST_WARNINGS
+
 TMap<FName, struct ICompressionFormat*> FCompression::CompressionFormats;
 FCriticalSection FCompression::CompressionFormatsCriticalSection;
 
@@ -793,3 +795,4 @@ bool FCompression::VerifyCompressionFlagsValid(int32 InCompressionFlags)
 /***********************
   Deprecated functions
 ***********************/
+PRAGMA_ENABLE_UNSAFE_TYPECAST_WARNINGS
diff --git a/Engine/Source/Runtime/Core/Private/Misc/ConfigCacheIni.cpp b/Engine/Source/Runtime/Core/Private/Misc/ConfigCacheIni.cpp
--- a/Engine/Source/Runtime/Core/Private/Misc/ConfigCacheIni.cpp
+++ b/Engine/Source/Runtime/Core/Private/Misc/ConfigCacheIni.cpp
@@ -142,7 +142,7 @@ bool FConfigValue::ExpandValue(const FString& InCollapsedValue, FString& OutExpa
 		const TCHAR* Begin;
 		const TCHAR* End;
 
-		int32 Len() const { return End - Begin; }
+		int32 Len() const { return UE_PTRDIFF_TO_INT32(End - Begin); }
 	};
 
 	// Find substrings of input and expansions to concatenate to final output string
@@ -358,7 +358,7 @@ static void ExtractPropertyValue(const FString& FullStructValue, const FString&
 		}
 
 		// pull out the token
-		Out.AppendChars(*FullStructValue + MatchLoc, Travel - Start);
+		Out.AppendChars(*FullStructValue + MatchLoc, UE_PTRDIFF_TO_INT32(Travel - Start));
 	}
 }
 
@@ -3100,8 +3100,8 @@ struct FConfigMemoryData
 		SIZE_T MaxMem = MemAr.GetMax();
 
 		NameIndent = FMath::Max(NameIndent, ConfigFilename.Len());
-		SizeIndent = FMath::Max(SizeIndent, FString::FromInt(TotalMem).Len());
-		MaxSizeIndent = FMath::Max(MaxSizeIndent, FString::FromInt(MaxMem).Len());
+		SizeIndent = FMath::Max(SizeIndent, FString::FromInt((int32)TotalMem).Len());
+		MaxSizeIndent = FMath::Max(MaxSizeIndent, FString::FromInt((int32)MaxMem).Len());
 		
 		new(MemoryData) FConfigFileMemoryData( ConfigFilename, TotalMem, MaxMem );
 	}
diff --git a/Engine/Source/Runtime/Core/Private/Misc/ConfigManifest.cpp b/Engine/Source/Runtime/Core/Private/Misc/ConfigManifest.cpp
--- a/Engine/Source/Runtime/Core/Private/Misc/ConfigManifest.cpp
+++ b/Engine/Source/Runtime/Core/Private/Misc/ConfigManifest.cpp
@@ -75,7 +75,7 @@ void MigratePreviousEngineInis()
 	IPlatformFile& PlatformFile = FPlatformFileManager::Get().GetPlatformFile();
 	while(MinorVersion >= 0)
 	{
-		const FEngineVersion PreviousVersion(FEngineVersion::Current().GetMajor(), MinorVersion--, 0, 0, FString());
+		const FEngineVersion PreviousVersion((uint16)FEngineVersion::Current().GetMajor(), (uint16)(MinorVersion--), 0, 0, FString());
 	
 		const FString Directory = FString(FPlatformProcess::UserSettingsDir()) / VERSION_TEXT(EPIC_PRODUCT_IDENTIFIER) / PreviousVersion.ToString(EVersionComponent::Minor) / TEXT("Saved") / TEXT("Config") / ANSI_TO_TCHAR(FPlatformProperties::PlatformName());
 		if (FPaths::DirectoryExists(Directory))
diff --git a/Engine/Source/Runtime/Core/Private/Misc/Crc.cpp b/Engine/Source/Runtime/Core/Private/Misc/Crc.cpp
--- a/Engine/Source/Runtime/Core/Private/Misc/Crc.cpp
+++ b/Engine/Source/Runtime/Core/Private/Misc/Crc.cpp
@@ -440,7 +440,7 @@ uint32 FCrc::MemCrc32( const void* InData, int32 Length, uint32 CRC/*=0 */ )
 	const uint8* __restrict Data = (uint8*)InData;
 
 	// First we need to align to 32-bits
-	int32 InitBytes = Align(Data, 4) - Data;
+	int32 InitBytes = UE_PTRDIFF_TO_INT32(Align(Data, 4) - Data);
 
 	if (Length > InitBytes)
 	{
@@ -490,7 +490,7 @@ uint32 FCrc::MemCrc_DEPRECATED(const void* InData, int32 Length, uint32 CRC/*=0
 	const uint8* __restrict Data = (uint8*)InData;
 
 	// First we need to align to 32-bits
-	int32 InitBytes = Align(Data, 4) - Data;
+	int32 InitBytes = UE_PTRDIFF_TO_INT32(Align(Data, 4) - Data);
 
 	if (Length > InitBytes)
 	{
diff --git a/Engine/Source/Runtime/Core/Private/Misc/DateTime.cpp b/Engine/Source/Runtime/Core/Private/Misc/DateTime.cpp
--- a/Engine/Source/Runtime/Core/Private/Misc/DateTime.cpp
+++ b/Engine/Source/Runtime/Core/Private/Misc/DateTime.cpp
@@ -75,7 +75,7 @@ void FDateTime::GetDate(int32& OutYear, int32& OutMonth, int32& OutDay) const
 
 	int32 i, j, k, l, n;
 
-	l = FMath::FloorToInt(GetJulianDay() + 0.5) + 68569;
+	l = FMath::FloorToInt((float)(GetJulianDay() + 0.5)) + 68569;
 	n = 4 * l / 146097;
 	l = l - (146097 * n + 3) / 4;
 	i = 4000 * (l + 1) / 1461001;
@@ -738,7 +738,7 @@ bool FDateTime::ParseIso8601(const TCHAR* DateTimeString, FDateTime& OutDateTime
 				return false;
 			}
 
-			for (int32 Digits = Next - Ptr; Digits < 3; ++Digits)
+			for (int32 Digits = UE_PTRDIFF_TO_INT32(Next - Ptr); Digits < 3; ++Digits)
 			{
 				Millisecond *= 10;
 			}
diff --git a/Engine/Source/Runtime/Core/Private/Misc/DefaultValueHelper.cpp b/Engine/Source/Runtime/Core/Private/Misc/DefaultValueHelper.cpp
--- a/Engine/Source/Runtime/Core/Private/Misc/DefaultValueHelper.cpp
+++ b/Engine/Source/Runtime/Core/Private/Misc/DefaultValueHelper.cpp
@@ -672,11 +672,11 @@ bool FDefaultValueHelper::ParseColor(const FString& Source, FColor& OutVal)
 		return OutVal.InitFromString(Source);
 	}
 
-	const uint8 Alpha = ( NULL != ThirdComma ) ? FCString::Atoi( ThirdComma + 1 ) : 255;
+	const uint8 Alpha = ( NULL != ThirdComma ) ? (uint8)FCString::Atoi( ThirdComma + 1 ) : 255;
 	OutVal = FColor( 
-		FCString::Atoi( Start ), 
-		FCString::Atoi( FirstComma + 1 ), 
-		FCString::Atoi( SecondComma + 1 ),
+		(uint8)FCString::Atoi( Start ), 
+		(uint8)FCString::Atoi( FirstComma + 1 ),
+		(uint8)FCString::Atoi( SecondComma + 1 ),
 		Alpha );
 	
 	return true;
diff --git a/Engine/Source/Runtime/Core/Private/Misc/EngineVersion.cpp b/Engine/Source/Runtime/Core/Private/Misc/EngineVersion.cpp
--- a/Engine/Source/Runtime/Core/Private/Misc/EngineVersion.cpp
+++ b/Engine/Source/Runtime/Core/Private/Misc/EngineVersion.cpp
@@ -168,7 +168,7 @@ bool FEngineVersion::Parse(const FString &Text, FEngineVersion &OutVersion)
 	}
 
 	// Build the output version
-	OutVersion.Set(Major, Minor, Patch, Changelist, Branch);
+	OutVersion.Set((uint16)Major, (uint16)Minor, (uint16)Patch, (uint32)Changelist, Branch);
 	return true;
 }
 
diff --git a/Engine/Source/Runtime/Core/Private/Misc/ExpressionParser.cpp b/Engine/Source/Runtime/Core/Private/Misc/ExpressionParser.cpp
--- a/Engine/Source/Runtime/Core/Private/Misc/ExpressionParser.cpp
+++ b/Engine/Source/Runtime/Core/Private/Misc/ExpressionParser.cpp
@@ -33,7 +33,7 @@ TCHAR FTokenStream::PeekChar(int32 Offset) const
 
 int32 FTokenStream::CharsRemaining() const
 {
-	return End - ReadPos;
+	return UE_PTRDIFF_TO_INT32(End - ReadPos);
 }
 
 bool FTokenStream::IsEmpty() const
@@ -43,7 +43,7 @@ bool FTokenStream::IsEmpty() const
 
 int32 FTokenStream::GetPosition() const
 {
-	return ReadPos - Start;
+	return UE_PTRDIFF_TO_INT32(ReadPos - Start);
 }
 
 FString FTokenStream::GetErrorContext() const
@@ -82,7 +82,7 @@ TOptional<FStringToken> FTokenStream::ParseToken(TFunctionRef<EParseState(TCHAR)
 		return TOptional<FStringToken>();
 	}
 
-	FStringToken Token(OptReadPos, 0, OptReadPos - Start);
+	FStringToken Token(OptReadPos, 0, UE_PTRDIFF_TO_INT32(OptReadPos - Start));
 
 	while (Token.GetTokenEndPos() != End)
 	{
@@ -126,7 +126,7 @@ TOptional<FStringToken> FTokenStream::ParseSymbol(FStringToken* Accumulate) cons
 		return TOptional<FStringToken>();
 	}
 	
-	FStringToken Token(OptReadPos, 0, OptReadPos - Start);
+	FStringToken Token(OptReadPos, 0, UE_PTRDIFF_TO_INT32(OptReadPos - Start));
 	++Token.TokenEnd;
 
 	if (Accumulate)
@@ -146,7 +146,7 @@ TOptional<FStringToken> FTokenStream::ParseSymbol(TCHAR Symbol, FStringToken* Ac
 		return TOptional<FStringToken>();
 	}
 	
-	FStringToken Token(OptReadPos, 0, OptReadPos - Start);
+	FStringToken Token(OptReadPos, 0, UE_PTRDIFF_TO_INT32(OptReadPos - Start));
 
 	if (*Token.TokenEnd == Symbol)
 	{
@@ -178,7 +178,7 @@ TOptional<FStringToken> FTokenStream::ParseToken(const TCHAR* Symbol, FStringTok
 		return TOptional<FStringToken>();		
 	}
 
-	FStringToken Token(OptReadPos, 0, OptReadPos - Start);
+	FStringToken Token(OptReadPos, 0, UE_PTRDIFF_TO_INT32(OptReadPos - Start));
 	
 	if (FCString::Strncmp(Token.GetTokenEndPos(), Symbol, Len) == 0)
 	{
@@ -205,7 +205,7 @@ TOptional<FStringToken> FTokenStream::ParseTokenIgnoreCase(const TCHAR* Symbol,
 		return TOptional<FStringToken>();
 	}
 
-	FStringToken Token(OptReadPos, 0, OptReadPos - Start);
+	FStringToken Token(OptReadPos, 0, UE_PTRDIFF_TO_INT32(OptReadPos - Start));
 	
 	if (FCString::Strnicmp(OptReadPos, Symbol, Len) == 0)
 	{
@@ -240,7 +240,7 @@ TOptional<FStringToken> FTokenStream::GenerateToken(int32 NumChars, FStringToken
 
 	if (IsReadPosValid(OptReadPos, NumChars))
 	{
-		FStringToken Token(OptReadPos, 0, OptReadPos - Start);
+		FStringToken Token(OptReadPos, 0, UE_PTRDIFF_TO_INT32(OptReadPos - Start));
 		Token.TokenEnd += NumChars;
 		if (Accumulate)
 		{
diff --git a/Engine/Source/Runtime/Core/Private/Misc/FeedbackContext.cpp b/Engine/Source/Runtime/Core/Private/Misc/FeedbackContext.cpp
--- a/Engine/Source/Runtime/Core/Private/Misc/FeedbackContext.cpp
+++ b/Engine/Source/Runtime/Core/Private/Misc/FeedbackContext.cpp
@@ -59,9 +59,9 @@ void FFeedbackContext::UpdateProgress( int32 Numerator, int32 Denominator )
 
 	if (LegacyAPIScopes.Num() != 0)
 	{
-		LegacyAPIScopes.Last()->TotalAmountOfWork = Denominator;
-		LegacyAPIScopes.Last()->CompletedWork = Numerator;
-		LegacyAPIScopes.Last()->CurrentFrameScope = Denominator - Numerator;
+		LegacyAPIScopes.Last()->TotalAmountOfWork = (float)Denominator;
+		LegacyAPIScopes.Last()->CompletedWork = (float)Numerator;
+		LegacyAPIScopes.Last()->CurrentFrameScope = (float)(Denominator - Numerator);
 		RequestUpdateUI();
 	}
 }
diff --git a/Engine/Source/Runtime/Core/Private/Misc/FeedbackContextMarkup.cpp b/Engine/Source/Runtime/Core/Private/Misc/FeedbackContextMarkup.cpp
--- a/Engine/Source/Runtime/Core/Private/Misc/FeedbackContextMarkup.cpp
+++ b/Engine/Source/Runtime/Core/Private/Misc/FeedbackContextMarkup.cpp
@@ -146,7 +146,7 @@ bool FFeedbackContextMarkup::ReadInteger(const TCHAR*& Text, uint32& OutInteger)
 	if(FChar::IsDigit(*Text))
 	{
 		TCHAR *End;
-		OutInteger = FCString::Strtoui64(Text, &End, 10);
+		OutInteger = (uint32)FCString::Strtoui64(Text, &End, 10);
 		Text = End;
 		while(FChar::IsWhitespace(*Text)) Text++;
 		return true;
@@ -162,7 +162,7 @@ bool FFeedbackContextMarkup::ReadString(const TCHAR*& Text, FString& OutString)
 		{
 			if(*End == *Text)
 			{
-				OutString = FString(End - (Text + 1), Text + 1);
+				OutString = FString(UE_PTRDIFF_TO_INT32(End - (Text + 1)), Text + 1);
 				do { End++; } while(FChar::IsWhitespace(*End));
 				Text = End;
 				return true;
diff --git a/Engine/Source/Runtime/Core/Private/Misc/FileHelper.cpp b/Engine/Source/Runtime/Core/Private/Misc/FileHelper.cpp
--- a/Engine/Source/Runtime/Core/Private/Misc/FileHelper.cpp
+++ b/Engine/Source/Runtime/Core/Private/Misc/FileHelper.cpp
@@ -42,7 +42,7 @@ bool FFileHelper::LoadFileToArray( TArray<uint8>& Result, const TCHAR* Filename,
 		}
 		return false;
 	}
-	int64 TotalSize = Reader->TotalSize();
+	int32 TotalSize = (int32)Reader->TotalSize();
 	// Allocate slightly larger than file size to avoid re-allocation when caller null terminates file buffer
 	Result.Reset( TotalSize + 2 );
 	Result.AddUninitialized( TotalSize );
@@ -119,7 +119,7 @@ bool FFileHelper::LoadFileToString(FString& Result, FArchive& Reader, EHashOptio
 {
 	FScopedLoadingState ScopedLoadingState(*Reader.GetArchiveName());
 
-	int32 Size = Reader.TotalSize();
+	int64 Size = Reader.TotalSize();
 	if (!Size)
 	{
 		Result.Empty();
@@ -136,7 +136,7 @@ bool FFileHelper::LoadFileToString(FString& Result, FArchive& Reader, EHashOptio
 	Reader.Serialize(Ch, Size);
 	bool Success = Reader.Close();
 
-	BufferToString(Result, Ch, Size);
+	BufferToString(Result, Ch, (int32)Size);
 
 	// handle SHA verify of the file
 	if (EnumHasAnyFlags(VerifyFlags, EHashOptions::EnableVerify) && (EnumHasAnyFlags(VerifyFlags, EHashOptions::ErrorMissingHash) || FSHA1::GetFileSHAHash(*Reader.GetArchiveName(), nullptr)))
@@ -212,7 +212,7 @@ bool FFileHelper::LoadFileToStringArray( TArray<FString>& Result, const TCHAR* F
 			Pos++;
 		}
 
-		Result.Emplace(Pos - LineStart, LineStart);
+		Result.Emplace((int32)(Pos - LineStart), LineStart);
 
 		if(*Pos == '\r')
 		{
@@ -245,7 +245,7 @@ bool FFileHelper::LoadFileToStringArrayWithPredicate(TArray<FString>& Result, co
 			Pos++;
 		}
 
-		FString Line(Pos - LineStart, LineStart);
+		FString Line(UE_PTRDIFF_TO_INT32(Pos - LineStart), LineStart);
 		if (Invoke(Predicate, Line))
 		{
 			Result.Add(MoveTemp(Line));
@@ -658,7 +658,7 @@ bool FFileHelper::LoadANSITextFileToStrings(const TCHAR* InFilename, IFileManage
 	if (TextFile != NULL)
 	{
 		// get the size of the file
-		int32 Size = TextFile->TotalSize();
+		int32 Size = (int32)TextFile->TotalSize();
 		// read the file
 		TArray<uint8> Buffer;
 		Buffer.Empty(Size + 1);
diff --git a/Engine/Source/Runtime/Core/Private/Misc/FrameRate.cpp b/Engine/Source/Runtime/Core/Private/Misc/FrameRate.cpp
--- a/Engine/Source/Runtime/Core/Private/Misc/FrameRate.cpp
+++ b/Engine/Source/Runtime/Core/Private/Misc/FrameRate.cpp
@@ -8,6 +8,8 @@
 
 #define LOCTEXT_NAMESPACE "FFrameRate"
 
+PRAGMA_DISABLE_UNSAFE_TYPECAST_WARNINGS
+
 namespace
 {
 	struct FSeconds         { static const TCHAR* const Moniker; };
@@ -200,7 +202,7 @@ bool FFrameRate::ComputeGridSpacing(float PixelsPerSecond, double& OutMajorInter
 
 		// Showing hours, minutes or seconds
 		static const int32 DesirableBases[]  = { 1, 2, 5, 10, 30, 60 };
-		static const int32 NumDesirableBases = UE_ARRAY_COUNT(DesirableBases);
+		static const SIZE_T NumDesirableBases = UE_ARRAY_COUNT(DesirableBases);
 
 		const int32 Scale     = FMath::CeilToInt(DesiredMajorTickPx / PixelsPerSecond / TimeOrder);
 		const int32 BaseIndex = FMath::Min(Algo::LowerBound(DesirableBases, Scale), NumDesirableBases-1);
@@ -302,4 +304,6 @@ bool TryParseString(FFrameRate& OutFrameRate, const TCHAR* InString)
 	return false;
 }
 
+PRAGMA_ENABLE_UNSAFE_TYPECAST_WARNINGS
+
 #undef LOCTEXT_NAMESPACE
\ No newline at end of file
diff --git a/Engine/Source/Runtime/Core/Private/Misc/MemStack.cpp b/Engine/Source/Runtime/Core/Private/Misc/MemStack.cpp
--- a/Engine/Source/Runtime/Core/Private/Misc/MemStack.cpp
+++ b/Engine/Source/Runtime/Core/Private/Misc/MemStack.cpp
@@ -282,7 +282,7 @@ int32 FMemStackBase::GetByteCount() const
 		}
 		else
 		{
-			Count += Top - Chunk->Data();
+			Count += UE_PTRDIFF_TO_INT32(Top - Chunk->Data());
 		}
 	}
 	return Count;
diff --git a/Engine/Source/Runtime/Core/Private/Misc/OutputDeviceFile.cpp b/Engine/Source/Runtime/Core/Private/Misc/OutputDeviceFile.cpp
--- a/Engine/Source/Runtime/Core/Private/Misc/OutputDeviceFile.cpp
+++ b/Engine/Source/Runtime/Core/Private/Misc/OutputDeviceFile.cpp
@@ -201,7 +201,7 @@ void FAsyncWriter::Serialize(void* InData, int64 Length)
 			if (Length >= Buffer.Num())
 			{
 				// Keep the buffer bigger than we require so that % Buffer.Num() does not return 0 for Lengths = Buffer.Num()
-				Buffer.SetNumUninitialized(Length + 1);
+				Buffer.SetNumUninitialized((int32)(Length + 1));
 			}
 		}
 	}
diff --git a/Engine/Source/Runtime/Core/Private/Misc/Parse.cpp b/Engine/Source/Runtime/Core/Private/Misc/Parse.cpp
--- a/Engine/Source/Runtime/Core/Private/Misc/Parse.cpp
+++ b/Engine/Source/Runtime/Core/Private/Misc/Parse.cpp
@@ -474,7 +474,7 @@ bool FParse::QuotedString( const TCHAR* Buffer, FString& Value, int32* OutNumCha
 
 	if (OutNumCharsRead)
 	{
-		*OutNumCharsRead = (Buffer - Start);
+		*OutNumCharsRead = UE_PTRDIFF_TO_INT32(Buffer - Start);
 	}
 
 	return true;
@@ -590,7 +590,7 @@ bool FParse::Value( const TCHAR* Stream, const TCHAR* Match, int8& Value )
 	if( Temp==NULL )
 		return false;
 	Temp += FCString::Strlen( Match );
-	Value = FCString::Atoi( Temp );
+	Value = (int8)FCString::Atoi( Temp );
 	return Value!=0 || FChar::IsDigit(Temp[0]);
 }
 
diff --git a/Engine/Source/Runtime/Core/Private/Misc/SecureHash.cpp b/Engine/Source/Runtime/Core/Private/Misc/SecureHash.cpp
--- a/Engine/Source/Runtime/Core/Private/Misc/SecureHash.cpp
+++ b/Engine/Source/Runtime/Core/Private/Misc/SecureHash.cpp
@@ -478,15 +478,15 @@ void LexFromString(FMD5Hash& Hash, const TCHAR* Buffer)
 		TCHAR Base = 0;
 		if (InHexChar >= '0' && InHexChar <= '9')
 		{
-			OutDecValue = InHexChar - '0';
+			OutDecValue = (uint8)(InHexChar - '0');
 		}
 		else if (InHexChar >= 'A' && InHexChar <= 'F')
 		{
-			OutDecValue = (InHexChar - 'A') + 10;
+			OutDecValue = (uint8)((InHexChar - 'A') + 10);
 		}
 		else if (InHexChar >= 'a' && InHexChar <= 'f')
 		{
-			OutDecValue = (InHexChar - 'a') + 10;
+			OutDecValue = (uint8)((InHexChar - 'a') + 10);
 		}
 		else
 		{
diff --git a/Engine/Source/Runtime/Core/Private/Misc/StringBuilder.cpp b/Engine/Source/Runtime/Core/Private/Misc/StringBuilder.cpp
--- a/Engine/Source/Runtime/Core/Private/Misc/StringBuilder.cpp
+++ b/Engine/Source/Runtime/Core/Private/Misc/StringBuilder.cpp
@@ -24,12 +24,14 @@ template<typename C>
 TStringBuilderImpl<C>::~TStringBuilderImpl()
 {
 	if (bIsDynamic)
-		FreeBuffer(Base, End - Base);
+	{
+		FreeBuffer(Base, UE_PTRDIFF_TO_INT32(End - Base));
+	}
 }
 
 template<typename C>
 void
-TStringBuilderImpl<C>::Extend(int32 ExtraCapacity)
+TStringBuilderImpl<C>::Extend(SIZE_T ExtraCapacity)
 {
 	check(bIsExtendable);
 
@@ -54,14 +56,14 @@ TStringBuilderImpl<C>::Extend(int32 ExtraCapacity)
 
 template<typename C>
 void*
-TStringBuilderImpl<C>::AllocBuffer(int32 ByteCount)
+TStringBuilderImpl<C>::AllocBuffer(SIZE_T ByteCount)
 {
 	return FMemory::Malloc(ByteCount * sizeof(C));
 }
 
 template<typename C>
 void
-TStringBuilderImpl<C>::FreeBuffer(void* Buffer, int32 ByteCount)
+TStringBuilderImpl<C>::FreeBuffer(void* Buffer, SIZE_T ByteCount)
 {
 	FMemory::Free(Buffer);
 }
diff --git a/Engine/Source/Runtime/Core/Private/Misc/StringFormatter.cpp b/Engine/Source/Runtime/Core/Private/Misc/StringFormatter.cpp
--- a/Engine/Source/Runtime/Core/Private/Misc/StringFormatter.cpp
+++ b/Engine/Source/Runtime/Core/Private/Misc/StringFormatter.cpp
@@ -42,7 +42,7 @@ void AppendToString(const FStringFormatArg& Arg, FString& StringToAppendTo)
 /** Token representing a literal string inside the string */
 struct FStringLiteral
 {
-	FStringLiteral(const FStringToken& InString) : String(InString), Len(InString.GetTokenEndPos() - InString.GetTokenStartPos()) {}
+	FStringLiteral(const FStringToken& InString) : String(InString), Len(UE_PTRDIFF_TO_INT32(InString.GetTokenEndPos() - InString.GetTokenStartPos())) {}
 	/** The string literal token */
 	FStringToken String;
 	/** Cached length of the string */
@@ -52,7 +52,7 @@ struct FStringLiteral
 /** Token representing a user-defined token, such as {Argument} */
 struct FFormatSpecifier
 {
-	FFormatSpecifier(const FStringToken& InIdentifier, const FStringToken& InEntireToken) : Identifier(InIdentifier), EntireToken(InEntireToken), Len(Identifier.GetTokenEndPos() - Identifier.GetTokenStartPos()) {}
+	FFormatSpecifier(const FStringToken& InIdentifier, const FStringToken& InEntireToken) : Identifier(InIdentifier), EntireToken(InEntireToken), Len(UE_PTRDIFF_TO_INT32(Identifier.GetTokenEndPos() - Identifier.GetTokenStartPos())) {}
 
 	/** The identifier part of the token */
 	FStringToken Identifier;
@@ -340,7 +340,7 @@ TValueOrError<FString, FExpressionError> FStringFormatter::FormatInternal(const
 	
 	// This code deliberately tries to reallocate as little as possible
 	FString Formatted;
-	Formatted.Reserve(Tokens.Last().Context.GetTokenEndPos() - InExpression);
+	Formatted.Reserve(UE_PTRDIFF_TO_INT32(Tokens.Last().Context.GetTokenEndPos() - InExpression));
 	for (auto& Token : Tokens)
 	{
 		if (const auto* Literal = Token.Node.Cast<FStringLiteral>())
@@ -374,7 +374,7 @@ TValueOrError<FString, FExpressionError> FStringFormatter::FormatInternal(const
 			else
 			{
 				// No replacement found, so just add the original token string
-				const int32 Length = FormatToken->EntireToken.GetTokenEndPos() - FormatToken->EntireToken.GetTokenStartPos();
+				const int32 Length = UE_PTRDIFF_TO_INT32(FormatToken->EntireToken.GetTokenEndPos() - FormatToken->EntireToken.GetTokenStartPos());
 				Formatted.AppendChars(FormatToken->EntireToken.GetTokenStartPos(), Length);
 			}
 		}
@@ -399,7 +399,7 @@ TValueOrError<FString, FExpressionError> FStringFormatter::FormatInternal(const
 	
 	// This code deliberately tries to reallocate as little as possible
 	FString Formatted;
-	Formatted.Reserve(Tokens.Last().Context.GetTokenEndPos() - InExpression);
+	Formatted.Reserve(UE_PTRDIFF_TO_INT32(Tokens.Last().Context.GetTokenEndPos() - InExpression));
 	for (auto& Token : Tokens)
 	{
 		if (const auto* Literal = Token.Node.Cast<FStringLiteral>())
@@ -423,7 +423,7 @@ TValueOrError<FString, FExpressionError> FStringFormatter::FormatInternal(const
 			else
 			{
 				// No replacement found, so just add the original token string
-				const int32 Length = IndexToken->EntireToken.GetTokenEndPos() - IndexToken->EntireToken.GetTokenStartPos();
+				const int32 Length = UE_PTRDIFF_TO_INT32(IndexToken->EntireToken.GetTokenEndPos() - IndexToken->EntireToken.GetTokenStartPos());
 				Formatted.AppendChars(IndexToken->EntireToken.GetTokenStartPos(), Length);
 			}
 		}
diff --git a/Engine/Source/Runtime/Core/Private/Misc/TextFilterUtils.cpp b/Engine/Source/Runtime/Core/Private/Misc/TextFilterUtils.cpp
--- a/Engine/Source/Runtime/Core/Private/Misc/TextFilterUtils.cpp
+++ b/Engine/Source/Runtime/Core/Private/Misc/TextFilterUtils.cpp
@@ -214,7 +214,7 @@ namespace FastToUpper
 	{
 		if (InChar < ToUpperAdjustmentTableCount)
 		{
-			return InChar + ToUpperAdjustmentTable[(int32)InChar];
+			return (TCHAR)(InChar + ToUpperAdjustmentTable[(int32)InChar]);
 		}
 		return FChar::ToUpper(InChar);
 	}
diff --git a/Engine/Source/Runtime/Core/Private/Modules/ModuleManager.cpp b/Engine/Source/Runtime/Core/Private/Modules/ModuleManager.cpp
--- a/Engine/Source/Runtime/Core/Private/Modules/ModuleManager.cpp
+++ b/Engine/Source/Runtime/Core/Private/Modules/ModuleManager.cpp
@@ -195,7 +195,7 @@ bool FindNewestModuleFile(TArray<FString>& FilesToSearch, const FDateTime& Newer
 
 		// Check the time stamp for this file
 		const FDateTime FoundFileTime = IFileManager::Get().GetTimeStamp(*FoundFilePath);
-		if (ensure(FoundFileTime != -1.0))
+		if (ensure(FoundFileTime != FDateTime::MinValue()))
 		{
 			// Was this file modified more recently than our others?
 			if (FoundFileTime > NewestFoundFileTime)
diff --git a/Engine/Source/Runtime/Core/Private/ProfilingDebugging/ABTesting.cpp b/Engine/Source/Runtime/Core/Private/ProfilingDebugging/ABTesting.cpp
--- a/Engine/Source/Runtime/Core/Private/ProfilingDebugging/ABTesting.cpp
+++ b/Engine/Source/Runtime/Core/Private/ProfilingDebugging/ABTesting.cpp
@@ -109,7 +109,7 @@ const TCHAR* FABTest::TickAndGetCommand()
 	}
 	else if (bABTestActive)
 	{
-		float Delta;
+		double Delta;
 		if (bABScopeTestActive)
 		{
 			Delta = TotalScopeTimeInFrame;
@@ -117,7 +117,7 @@ const TCHAR* FABTest::TickAndGetCommand()
 		else
 		{
 			Delta = FPlatformTime::Seconds() - LastTime;
-			check(Delta > 0);
+			check(Delta > 0.0);
 		}
 
 		FSample* Sample = nullptr;
diff --git a/Engine/Source/Runtime/Core/Private/ProfilingDebugging/CsvProfiler.cpp b/Engine/Source/Runtime/Core/Private/ProfilingDebugging/CsvProfiler.cpp
--- a/Engine/Source/Runtime/Core/Private/ProfilingDebugging/CsvProfiler.cpp
+++ b/Engine/Source/Runtime/Core/Private/ProfilingDebugging/CsvProfiler.cpp
@@ -1031,7 +1031,7 @@ struct FCsvStatBase
 		static const uint8 IsExclusiveInsertedMarker = 0x20;
 	};
 
-	CSV_PROFILER_INLINE void Init(uint64 InStatID, int32 InCategoryIndex, uint32 InFlags, uint64 InTimestamp)
+	CSV_PROFILER_INLINE void Init(uint64 InStatID, int32 InCategoryIndex, uint8 InFlags, uint64 InTimestamp)
 	{
 		Timestamp = InTimestamp;
 		Flags = InFlags;
@@ -1220,7 +1220,7 @@ public:
 		int32 StrLen;
 		ANSICHAR StringBuffer[256];
 
-		if (FMath::Frac(Value) == 0.0)
+		if (FMath::Frac((float)Value) == 0.0f)
 		{
 			StrLen = FCStringAnsi::Snprintf(StringBuffer, 256, "%d", int(Value));
 		}
@@ -1611,7 +1611,7 @@ void FCsvStatSeries::FlushIfDirty()
 		switch (SeriesType)
 		{
 		case EType::TimerData:
-			Value.Value.AsFloat = FPlatformTime::ToMilliseconds64(CurrentValue.AsTimerCycles);
+			Value.Value.AsFloat = (float)FPlatformTime::ToMilliseconds64(CurrentValue.AsTimerCycles);
 			break;
 		case EType::CustomStatInt:
 			Value.Value.AsInt = CurrentValue.AsIntValue;
@@ -2619,7 +2619,7 @@ void FCsvProfiler::BeginFrame()
 					SetMetadata(TEXT("TargetFramerate"), *FString::FromInt(TargetFPS));
 
 #if !UE_BUILD_SHIPPING
-					uint64 ExtraDevelopmentMemoryMB = FPlatformMemory::GetExtraDevelopmentMemorySize()/1024ull/1024ull;
+					int32 ExtraDevelopmentMemoryMB = (int32)(FPlatformMemory::GetExtraDevelopmentMemorySize()/1024ull/1024ull);
 					SetMetadata(TEXT("ExtraDevelopmentMemoryMB"), *FString::FromInt(ExtraDevelopmentMemoryMB)); 
 #endif
 
@@ -2671,7 +2671,7 @@ void FCsvProfiler::EndFrame()
 		// Record the frametime (measured since the last EndFrame)
 		uint64 CurrentTimeStamp = FPlatformTime::Cycles64();
 		uint64 ElapsedCycles = CurrentTimeStamp - LastEndFrameTimestamp;
-		float ElapsedMs = FPlatformTime::ToMilliseconds64(ElapsedCycles);
+		float ElapsedMs = (float)FPlatformTime::ToMilliseconds64(ElapsedCycles);
 		CSV_CUSTOM_STAT_DEFINED(FrameTime, ElapsedMs, ECsvCustomStatOp::Set);
 
 		FPlatformMemoryStats MemoryStats = FPlatformMemory::GetStats();
diff --git a/Engine/Source/Runtime/Core/Private/Serialization/Archive.cpp b/Engine/Source/Runtime/Core/Private/Serialization/Archive.cpp
--- a/Engine/Source/Runtime/Core/Private/Serialization/Archive.cpp
+++ b/Engine/Source/Runtime/Core/Private/Serialization/Archive.cpp
@@ -1,8 +1,5 @@
 // Copyright Epic Games, Inc. All Rights Reserved.
 
-/*=============================================================================
-	UnArchive.cpp: Core archive classes.
-=============================================================================*/
 #include "Serialization/Archive.h"
 #include "Math/UnrealMathUtility.h"
 #include "HAL/UnrealMemory.h"
@@ -26,6 +23,8 @@
 #include "Serialization/CompressedChunkInfo.h"
 #include "Serialization/ArchiveSerializedPropertyChain.h"
 
+PRAGMA_DISABLE_UNSAFE_TYPECAST_WARNINGS
+
 /*-----------------------------------------------------------------------------
 	FArchive implementation.
 -----------------------------------------------------------------------------*/
@@ -1153,3 +1152,5 @@ void FArchive::SetArchiveState(const FArchiveState& InState)
 {
 	ImplicitConv<FArchiveState&>(*this) = InState;
 }
+
+PRAGMA_ENABLE_UNSAFE_TYPECAST_WARNINGS
diff --git a/Engine/Source/Runtime/Core/Private/Serialization/ArchiveFromStructuredArchive.cpp b/Engine/Source/Runtime/Core/Private/Serialization/ArchiveFromStructuredArchive.cpp
--- a/Engine/Source/Runtime/Core/Private/Serialization/ArchiveFromStructuredArchive.cpp
+++ b/Engine/Source/Runtime/Core/Private/Serialization/ArchiveFromStructuredArchive.cpp
@@ -83,7 +83,7 @@ void FArchiveFromStructuredArchiveImpl::Seek(int64 InPos)
 	if (InnerArchive.IsTextFormat())
 	{
 		check(Pimpl->Pos >= 0 && Pimpl->Pos <= Pimpl->Buffer.Num());
-		Pimpl->Pos = InPos;
+		Pimpl->Pos = (int32)InPos;
 	}
 	else
 	{
@@ -211,16 +211,16 @@ void FArchiveFromStructuredArchiveImpl::Serialize(void* V, int64 Length)
 				checkf(false, TEXT("Attempt to read past end of archive"));
 			}
 			FMemory::Memcpy(V, Pimpl->Buffer.GetData() + Pimpl->Pos, Length);
-			Pimpl->Pos += Length;
+			Pimpl->Pos += (int32)Length;
 		}
 		else
 		{
 			if (Pimpl->Pos + Length > Pimpl->Buffer.Num())
 			{
-				Pimpl->Buffer.AddUninitialized(Pimpl->Pos + Length - Pimpl->Buffer.Num());
+				Pimpl->Buffer.AddUninitialized((int32)(Pimpl->Pos + Length - Pimpl->Buffer.Num()));
 			}
 			FMemory::Memcpy(Pimpl->Buffer.GetData() + Pimpl->Pos, V, Length);
-			Pimpl->Pos += Length;
+			Pimpl->Pos += (int32)Length;
 		}
 	}
 	else
diff --git a/Engine/Source/Runtime/Core/Private/Serialization/ArchiveLoadCompressedProxy.cpp b/Engine/Source/Runtime/Core/Private/Serialization/ArchiveLoadCompressedProxy.cpp
--- a/Engine/Source/Runtime/Core/Private/Serialization/ArchiveLoadCompressedProxy.cpp
+++ b/Engine/Source/Runtime/Core/Private/Serialization/ArchiveLoadCompressedProxy.cpp
@@ -5,6 +5,8 @@
 #include "HAL/UnrealMemory.h"
 #include "Containers/Array.h"
 
+PRAGMA_DISABLE_UNSAFE_TYPECAST_WARNINGS
+
 /*----------------------------------------------------------------------------
 	FArchiveLoadCompressedProxy
 ----------------------------------------------------------------------------*/
@@ -119,3 +121,4 @@ int64 FArchiveLoadCompressedProxy::Tell()
 	return RawBytesSerialized;
 }
 
+PRAGMA_ENABLE_UNSAFE_TYPECAST_WARNINGS
diff --git a/Engine/Source/Runtime/Core/Private/Serialization/BitReader.cpp b/Engine/Source/Runtime/Core/Private/Serialization/BitReader.cpp
--- a/Engine/Source/Runtime/Core/Private/Serialization/BitReader.cpp
+++ b/Engine/Source/Runtime/Core/Private/Serialization/BitReader.cpp
@@ -5,6 +5,8 @@
 #include "Logging/LogMacros.h"
 #include "CoreGlobals.h"
 
+PRAGMA_DISABLE_UNSAFE_TYPECAST_WARNINGS
+
 // Table.
 extern const uint8 GShift[8]={0x01,0x02,0x04,0x08,0x10,0x20,0x40,0x80};
 extern const uint8 GMask [8]={0x00,0x01,0x03,0x07,0x0f,0x1f,0x3f,0x7f};
@@ -224,7 +226,7 @@ void FBitReader::CountMemory(FArchive& Ar) const
 	Ar.CountBytes(sizeof(*this), sizeof(*this));
 }
 
-void FBitReader::SetOverflowed(int32 LengthBits)
+void FBitReader::SetOverflowed(int64 LengthBits)
 {
 	UE_LOG(LogNetSerialization, Error, TEXT("FBitReader::SetOverflowed() called! (ReadLen: %i, Remaining: %i, Max: %i)"),
 			LengthBits, (Num - Pos), Num);
@@ -245,7 +247,7 @@ void FBitReader::SerializeBitsWithOffset( void* Dest, int32 DestBit, int64 Lengt
 
 	if (LengthBits != 0)
 	{
-		appBitsCpy((uint8*)Dest, DestBit, Buffer.GetData(), Pos, LengthBits);
+		appBitsCpy((uint8*)Dest, DestBit, Buffer.GetData(), (int32)Pos, (int32)LengthBits);
 		Pos += LengthBits;
 	}
 }
@@ -308,3 +310,5 @@ void FBitReaderMark::Copy( FBitReader& Reader, TArray<uint8> &Buffer )
 		appBitsCpy(Buffer.GetData(), 0, Reader.Buffer.GetData(), Pos, Reader.Pos - Pos);
 	}
 }
+
+PRAGMA_ENABLE_UNSAFE_TYPECAST_WARNINGS
diff --git a/Engine/Source/Runtime/Core/Private/Serialization/BitWriter.cpp b/Engine/Source/Runtime/Core/Private/Serialization/BitWriter.cpp
--- a/Engine/Source/Runtime/Core/Private/Serialization/BitWriter.cpp
+++ b/Engine/Source/Runtime/Core/Private/Serialization/BitWriter.cpp
@@ -1,13 +1,11 @@
 // Copyright Epic Games, Inc. All Rights Reserved.
 
-/*=============================================================================
-	UnBits.h: Unreal bitstream manipulation classes.
-=============================================================================*/
-
 #include "Serialization/BitWriter.h"
 #include "Logging/LogMacros.h"
 #include "CoreGlobals.h"
 
+PRAGMA_DISABLE_UNSAFE_TYPECAST_WARNINGS
+
 extern const uint8 GShift[8];
 extern const uint8 GMask[8];
 
@@ -299,4 +297,6 @@ void FBitWriterMark::Copy( FBitWriter& Writer, TArray<uint8> &Buffer )
 		Buffer[Bytes-1] = 0;	// Make sure the last byte is 0 out, because appBitsCpy wont touch the last bits
 		appBitsCpy(Buffer.GetData(), 0, Writer.Buffer.GetData(), Num, Writer.Num - Num);
 	}
-}
\ No newline at end of file
+}
+
+PRAGMA_ENABLE_UNSAFE_TYPECAST_WARNINGS
