Core: Fixes for various unsafe implicit type conversion warnings in the Core module (WIP; some files still suppress the warning)

diff --git a/Engine/Source/Runtime/Core/Private/Stats/StatsCommand.cpp b/Engine/Source/Runtime/Core/Private/Stats/StatsCommand.cpp
--- a/Engine/Source/Runtime/Core/Private/Stats/StatsCommand.cpp
+++ b/Engine/Source/Runtime/Core/Private/Stats/StatsCommand.cpp
@@ -656,7 +656,7 @@ void DumpCPUSummary(FStatsThreadState const& StatsData, int64 TargetFrame)
 }
 
 static int32 HitchIndex = 0;
-static float TotalHitchTime = 0.0f;
+static double TotalHitchTime = 0.0;
 
 static void DumpHitch(int64 Frame)
 {
@@ -674,14 +674,14 @@ static void DumpHitch(int64 Frame)
 	FStatsThreadState& Stats = FStatsThreadState::GetLocalState();
 	SCOPE_CYCLE_COUNTER(STAT_HitchScan);
 
-	const float GameThreadTime = FPlatformTime::ToSeconds(Stats.GetFastThreadFrameTime(Frame, EThreadType::Game));
-	const float RenderThreadTime = FPlatformTime::ToSeconds(Stats.GetFastThreadFrameTime(Frame, EThreadType::Renderer));
+	const double GameThreadTime = FPlatformTime::ToSeconds64(Stats.GetFastThreadFrameTime(Frame, EThreadType::Game));
+	const double RenderThreadTime = FPlatformTime::ToSeconds64(Stats.GetFastThreadFrameTime(Frame, EThreadType::Renderer));
 	const float HitchThresholdSecs = GHitchThresholdMS * 0.001f;
 
 	if ((GameThreadTime > HitchThresholdSecs) || (RenderThreadTime > HitchThresholdSecs))
 	{
 		HitchIndex++;
-		float ThisHitch = FMath::Max<float>(GameThreadTime, RenderThreadTime) * 1000.0f;
+		double ThisHitch = FMath::Max<double>(GameThreadTime, RenderThreadTime) * 1000.0;
 		TotalHitchTime += ThisHitch;
 		UE_LOG(LogStats, Log, TEXT("------------------Thread Hitch %d, Frame %lld  %6.1fms ---------------"), HitchIndex, Frame, ThisHitch);
 		FRawStatStackNode Stack;
@@ -1988,7 +1988,7 @@ static void StatCmd(FString InCmd, bool bStatCommand, FOutputDevice* Ar /*= null
 			{
 				StatsMasterEnableAdd();
 				HitchIndex = 0;
-				TotalHitchTime = 0.0f;
+				TotalHitchTime = 0.0;
 				DumpHitchDelegateHandle = Stats.NewFrameDelegate.AddStatic(&DumpHitch);
 			}
 			else
diff --git a/Engine/Source/Runtime/Core/Private/Stats/StatsData.cpp b/Engine/Source/Runtime/Core/Private/Stats/StatsData.cpp
--- a/Engine/Source/Runtime/Core/Private/Stats/StatsData.cpp
+++ b/Engine/Source/Runtime/Core/Private/Stats/StatsData.cpp
@@ -313,7 +313,7 @@ void FRawStatStackNode::AddSelf()
 			{
 				Self.GetValue_int64() = ToPackedCallCountDuration(
 					FromPackedCallCountDuration_CallCount(Self.GetValue_int64()),
-					MyTime);
+					(uint32)MyTime);
 				Self.NameAndInfo.SetRawName(NAME_Self);
 				Children.Add(NAME_Self, new FRawStatStackNode(Self));
 			}
@@ -1839,7 +1839,7 @@ FString FStatsUtils::DebugPrint(FStatMessage const& Item)
 		}
 		else if (Item.NameAndInfo.GetFlag(EStatMetaFlags::IsCycle))
 		{
-			Result = FString::Printf(TEXT("%.3fms"), FPlatformTime::ToMilliseconds(Item.GetValue_int64()));
+			Result = FString::Printf(TEXT("%.3fms"), FPlatformTime::ToMilliseconds64(Item.GetValue_int64()));
 		}
 		else
 		{
diff --git a/Engine/Source/Runtime/Core/Private/Stats/StatsFile.cpp b/Engine/Source/Runtime/Core/Private/Stats/StatsFile.cpp
--- a/Engine/Source/Runtime/Core/Private/Stats/StatsFile.cpp
+++ b/Engine/Source/Runtime/Core/Private/Stats/StatsFile.cpp
@@ -278,7 +278,7 @@ void IStatsWriteFile::Stop()
 
 FText IStatsWriteFile::GetFileMetaDesc() const
 {
-	const FTimespan Duration = FTimespan( 0, 0, FPlatformTime::Seconds() - StartTime );
+	const FTimespan Duration = FTimespan( 0, 0, (int32)(FPlatformTime::Seconds() - StartTime) );
 
 	const FText FileMetaDesc = FText::Format( LOCTEXT( "FileMetaDesc_Fmt", "STATS FILE: Duration: {0}, Filesize: {1}" ), FText::AsTimespan( Duration ), FText::AsMemory( (SIZE_T)FileSize ) );
 	return FileMetaDesc;
@@ -758,7 +758,7 @@ void FStatsReadFile::ReadRawStats()
 		else
 		{
 			Frame.Packets.Add( StatPacket );
-			FileInfo.MaximumPacketSize = FMath::Max<int32>( FileInfo.MaximumPacketSize, StatPacket->StatMessages.GetAllocatedSize() );
+			FileInfo.MaximumPacketSize = FMath::Max<int32>( FileInfo.MaximumPacketSize, (int32)StatPacket->StatMessages.GetAllocatedSize() );
 		}
 
 		UpdateReadStageProgress();
@@ -925,7 +925,7 @@ void FStatsReadFile::ProcessStats()
 									// Read OperationSequenceTag.
 									Index++; CurrentStatMessageIndex++;
 									const FStatMessage& SequenceTagMessage = Data[Index];
-									const uint32 SequenceTag = SequenceTagMessage.GetValue_int64();
+									const uint32 SequenceTag = (uint32)SequenceTagMessage.GetValue_int64();
 
 									//ThreadStats->AddMemoryMessage( GET_STATFNAME( STAT_Memory_AllocPtr ), (uint64)(UPTRINT)Ptr | (uint64)EMemoryOperation::Alloc );
 									//ThreadStats->AddMemoryMessage( GET_STATFNAME( STAT_Memory_AllocSize ), Size );
@@ -949,7 +949,7 @@ void FStatsReadFile::ProcessStats()
 									// Read OperationSequenceTag.
 									Index++; CurrentStatMessageIndex++;
 									const FStatMessage& SequenceTagMessage = Data[Index];
-									const uint32 SequenceTag = SequenceTagMessage.GetValue_int64();
+									const uint32 SequenceTag = (uint32)SequenceTagMessage.GetValue_int64();
 
 									//ThreadStats->AddMemoryMessage( GET_STATFNAME( STAT_Memory_FreePtr ), (uint64)(UPTRINT)OldPtr | (uint64)EMemoryOperation::Realloc );
 									//ThreadStats->AddMemoryMessage( GET_STATFNAME( STAT_Memory_AllocPtr ), (uint64)(UPTRINT)NewPtr | (uint64)EMemoryOperation::Realloc );
@@ -962,7 +962,7 @@ void FStatsReadFile::ProcessStats()
 									// Read OperationSequenceTag.
 									Index++; CurrentStatMessageIndex++;
 									const FStatMessage& SequenceTagMessage = Data[Index];
-									const uint32 SequenceTag = SequenceTagMessage.GetValue_int64();
+									const uint32 SequenceTag = (uint32)SequenceTagMessage.GetValue_int64();
 
 									//ThreadStats->AddMemoryMessage( GET_STATFNAME( STAT_Memory_FreePtr ), (uint64)(UPTRINT)Ptr | (uint64)EMemoryOperation::Free );	// 16 bytes total				
 									//ThreadStats->AddMemoryMessage( GET_STATFNAME( STAT_Memory_OperationSequenceTag ), (int64)SequenceTag );
@@ -1086,7 +1086,7 @@ void FStatsReadFile::UpdateCombinedHistoryStats()
 		int32 FramePackets = It.Value.Packets.Num(); // Threads
 		for (const auto& It2 : It.Value.Packets)
 		{
-			FramePacketsSize += It2->StatMessages.GetAllocatedSize();
+			FramePacketsSize += (int32)It2->StatMessages.GetAllocatedSize();
 			FrameStatMessages += It2->StatMessages.Num();
 		}
 
diff --git a/Engine/Source/Runtime/Core/Private/Tests/Misc/CharTest.cpp b/Engine/Source/Runtime/Core/Private/Tests/Misc/CharTest.cpp
--- a/Engine/Source/Runtime/Core/Private/Tests/Misc/CharTest.cpp
+++ b/Engine/Source/Runtime/Core/Private/Tests/Misc/CharTest.cpp
@@ -24,7 +24,7 @@ void RunCharTests(FAutomationTestBase& Test, uint32 MaxChar)
 {
 	for (uint32 I = 0; I < MaxChar; ++I)
 	{
-		CharType C(I);
+		CharType C = (CharType)I;
 		Test.TestEqual("TChar::ToLower()", TChar<CharType>::ToLower(C), crt::tolower(C));
 		Test.TestEqual("TChar::ToUpper()", TChar<CharType>::ToUpper(C), crt::toupper(C));
 	}
diff --git a/Engine/Source/Runtime/Core/Private/Tests/Misc/TimecodeTest.cpp b/Engine/Source/Runtime/Core/Private/Tests/Misc/TimecodeTest.cpp
--- a/Engine/Source/Runtime/Core/Private/Tests/Misc/TimecodeTest.cpp
+++ b/Engine/Source/Runtime/Core/Private/Tests/Misc/TimecodeTest.cpp
@@ -138,7 +138,7 @@ bool FTimecodeTest::RunTest(const FString& Parameters)
 				if (bWrongFrame && bIsDropFrame)
 				{
 					// if new minute but not multiple of 10 mins, 2|4 is expected
-					const int32 NumberOfFramesInSecond = FMath::CeilToInt(FrameRate.AsDecimal());
+					const int32 NumberOfFramesInSecond = FMath::CeilToInt((float)FrameRate.AsDecimal());
 					const int32 NumberOfTimecodesToDrop = NumberOfFramesInSecond <= 30 ? 2 : 4;
 					bWrongFrame = !(TimecodeValue.Frames == NumberOfTimecodesToDrop && PreviousTimecodeValue.Minutes + 1 == TimecodeValue.Minutes && TimecodeValue.Minutes % 10 != 0);
 				}
diff --git a/Engine/Source/Runtime/Core/Private/Windows/WindowsAsyncIO.h b/Engine/Source/Runtime/Core/Private/Windows/WindowsAsyncIO.h
--- a/Engine/Source/Runtime/Core/Private/Windows/WindowsAsyncIO.h
+++ b/Engine/Source/Runtime/Core/Private/Windows/WindowsAsyncIO.h
@@ -4,6 +4,8 @@
 
 #include "ProfilingDebugging/PlatformFileTrace.h"
 
+PRAGMA_DISABLE_UNSAFE_TYPECAST_WARNINGS
+
 class FWindowsReadRequest;
 class FWindowsAsyncReadFileHandle;
 
@@ -488,4 +490,6 @@ bool FWindowsReadRequest::CheckForPrecache()
 const TCHAR* FWindowsReadRequest::GetFileNameForErrorMessagesAndPanicRetry()
 {
 	return *Owner->FileNameForErrorMessagesAndPanicRetry;
-}
\ No newline at end of file
+}
+
+PRAGMA_ENABLE_UNSAFE_TYPECAST_WARNINGS
diff --git a/Engine/Source/Runtime/Core/Private/Windows/WindowsPlatformCrashContext.cpp b/Engine/Source/Runtime/Core/Private/Windows/WindowsPlatformCrashContext.cpp
--- a/Engine/Source/Runtime/Core/Private/Windows/WindowsPlatformCrashContext.cpp
+++ b/Engine/Source/Runtime/Core/Private/Windows/WindowsPlatformCrashContext.cpp
@@ -97,7 +97,7 @@ namespace {
 		// CrashContext.runtime-xml is now a part of the minidump file.
 		MINIDUMP_USER_STREAM CrashContextStream = { 0 };
 		CrashContextStream.Type = UE4_MINIDUMP_CRASHCONTEXT;
-		CrashContextStream.BufferSize = InContext.GetBuffer().GetAllocatedSize();
+		CrashContextStream.BufferSize = (ULONG)InContext.GetBuffer().GetAllocatedSize();
 		CrashContextStream.Buffer = (void*)*InContext.GetBuffer();
 
 		MINIDUMP_USER_STREAM_INFORMATION CrashContextStreamInformation = { 0 };
diff --git a/Engine/Source/Runtime/Core/Private/Windows/WindowsPlatformMemory.cpp b/Engine/Source/Runtime/Core/Private/Windows/WindowsPlatformMemory.cpp
--- a/Engine/Source/Runtime/Core/Private/Windows/WindowsPlatformMemory.cpp
+++ b/Engine/Source/Runtime/Core/Private/Windows/WindowsPlatformMemory.cpp
@@ -257,7 +257,7 @@ const FPlatformMemoryConstants& FWindowsPlatformMemory::GetConstants()
 		MemoryConstants.PageSize = SystemInfo.dwPageSize;
 		MemoryConstants.AddressLimit = FPlatformMath::RoundUpToPowerOfTwo64(MemoryConstants.TotalPhysical);
 
-		MemoryConstants.TotalPhysicalGB = (MemoryConstants.TotalPhysical + 1024 * 1024 * 1024 - 1) / 1024 / 1024 / 1024;
+		MemoryConstants.TotalPhysicalGB = (uint32)((MemoryConstants.TotalPhysical + 1024 * 1024 * 1024 - 1) / 1024 / 1024 / 1024);
 	}
 
 	return MemoryConstants;	
diff --git a/Engine/Source/Runtime/Core/Private/Windows/WindowsPlatformProcess.cpp b/Engine/Source/Runtime/Core/Private/Windows/WindowsPlatformProcess.cpp
--- a/Engine/Source/Runtime/Core/Private/Windows/WindowsPlatformProcess.cpp
+++ b/Engine/Source/Runtime/Core/Private/Windows/WindowsPlatformProcess.cpp
@@ -36,6 +36,8 @@
 
 #pragma comment(lib, "psapi.lib")
 
+PRAGMA_DISABLE_UNSAFE_TYPECAST_WARNINGS
+
 // static variables
 TArray<FString> FWindowsPlatformProcess::DllDirectoryStack;
 TArray<FString> FWindowsPlatformProcess::DllDirectories;
@@ -499,8 +501,8 @@ bool FWindowsPlatformProcess::GetPerFrameProcessorUsage(uint32 ProcessId, float&
 {
 	bool bSuccess = true;
 
-	static float LastProcessTime = 0.f;
-	static float LastIdleTime = 0.f;
+	static double LastProcessTime = 0.f;
+	static double LastIdleTime = 0.f;
 	static uint32 LastFrameNumber = 0;
 
 	if (LastFrameNumber != GFrameNumber)
@@ -1919,4 +1921,6 @@ FString FWindowsPlatformProcess::FProcEnumInfo::GetFullPath() const
 	return GetApplicationName(GetPID());
 }
 
-#include "Windows/HideWindowsPlatformTypes.h"
+PRAGMA_ENABLE_UNSAFE_TYPECAST_WARNINGS
+
+#include "Windows/HideWindowsPlatformTypes.h"
\ No newline at end of file
diff --git a/Engine/Source/Runtime/Core/Private/Windows/WindowsPlatformStackWalk.cpp b/Engine/Source/Runtime/Core/Private/Windows/WindowsPlatformStackWalk.cpp
--- a/Engine/Source/Runtime/Core/Private/Windows/WindowsPlatformStackWalk.cpp
+++ b/Engine/Source/Runtime/Core/Private/Windows/WindowsPlatformStackWalk.cpp
@@ -305,7 +305,7 @@ uint32 FWindowsPlatformStackWalk::CaptureThreadStackBackTrace(uint64 ThreadId, u
 	if (BackTrace == nullptr || MaxDepth == 0)
 		return 0;
 
-	HANDLE ThreadHandle = OpenThread(THREAD_GET_CONTEXT | THREAD_SET_CONTEXT | THREAD_TERMINATE | THREAD_SUSPEND_RESUME, false, ThreadId);
+	HANDLE ThreadHandle = OpenThread(THREAD_GET_CONTEXT | THREAD_SET_CONTEXT | THREAD_TERMINATE | THREAD_SUSPEND_RESUME, false, (DWORD)ThreadId);
 	if (!ThreadHandle)
 		return 0;
 
diff --git a/Engine/Source/Runtime/Core/Private/Windows/WindowsPlatformTime.cpp b/Engine/Source/Runtime/Core/Private/Windows/WindowsPlatformTime.cpp
--- a/Engine/Source/Runtime/Core/Private/Windows/WindowsPlatformTime.cpp
+++ b/Engine/Source/Runtime/Core/Private/Windows/WindowsPlatformTime.cpp
@@ -105,9 +105,9 @@ bool FWindowsPlatformTime::UpdateCPUTime( float /*DeltaTime*/ )
 
 	// IntervalUserAndKernelTime == 0.0f means that the OS hasn't updated the data yet, 
 	// so don't update to avoid oscillating between 0 and calculated value.
-	if( IntervalUserAndKernelTime > 0.0f )
+	if( IntervalUserAndKernelTime > 0.0 )
 	{
-		CPUTimePctRelative = IntervalUserAndKernelTime/IntervalProcessTime * 100.0f;
+		CPUTimePctRelative = (float)(IntervalUserAndKernelTime/IntervalProcessTime * 100.0);
 
 		LastTotalProcessTime = CurrentTotalProcessTime;
 		LastTotalUserAndKernelTime = CurrentTotalUserAndKernelTime;
diff --git a/Engine/Source/Runtime/Core/Public/Algo/BinarySearch.h b/Engine/Source/Runtime/Core/Public/Algo/BinarySearch.h
--- a/Engine/Source/Runtime/Core/Public/Algo/BinarySearch.h
+++ b/Engine/Source/Runtime/Core/Public/Algo/BinarySearch.h
@@ -19,22 +19,22 @@ namespace AlgoImpl
 	 *
 	 * @returns Position of the first element >= Value, may be == Num
 	 */
-	template <typename RangeValueType, typename PredicateValueType, typename ProjectionType, typename SortPredicateType>
-	FORCEINLINE SIZE_T LowerBoundInternal(RangeValueType* First, const SIZE_T Num, const PredicateValueType& Value, ProjectionType Projection, SortPredicateType SortPredicate)
+	template <typename RangeValueType, typename SizeType, typename PredicateValueType, typename ProjectionType, typename SortPredicateType>
+	FORCEINLINE SizeType LowerBoundInternal(RangeValueType* First, const SizeType Num, const PredicateValueType& Value, ProjectionType Projection, SortPredicateType SortPredicate)
 	{
 		// Current start of sequence to check
-		SIZE_T Start = 0;
+		SizeType Start = 0;
 		// Size of sequence to check
-		SIZE_T Size = Num;
+		SizeType Size = Num;
 
 		// With this method, if Size is even it will do one more comparison than necessary, but because Size can be predicted by the CPU it is faster in practice
 		while (Size > 0)
 		{
-			const SIZE_T LeftoverSize = Size % 2;
+			const SizeType LeftoverSize = Size % 2;
 			Size = Size / 2;
 
-			const SIZE_T CheckIndex = Start + Size;
-			const SIZE_T StartIfLess = CheckIndex + LeftoverSize;
+			const SizeType CheckIndex = Start + Size;
+			const SizeType StartIfLess = CheckIndex + LeftoverSize;
 
 			auto&& CheckValue = Invoke(Projection, First[CheckIndex]);
 			Start = SortPredicate(CheckValue, Value) ? StartIfLess : Start;
@@ -52,22 +52,22 @@ namespace AlgoImpl
 	 *
 	 * @returns Position of the first element > Value, may be == Num
 	 */
-	template <typename RangeValueType, typename PredicateValueType, typename ProjectionType, typename SortPredicateType>
-	FORCEINLINE SIZE_T UpperBoundInternal(RangeValueType* First, const SIZE_T Num, const PredicateValueType& Value, ProjectionType Projection, SortPredicateType SortPredicate)
+	template <typename RangeValueType, typename SizeType, typename PredicateValueType, typename ProjectionType, typename SortPredicateType>
+	FORCEINLINE SizeType UpperBoundInternal(RangeValueType* First, const SizeType Num, const PredicateValueType& Value, ProjectionType Projection, SortPredicateType SortPredicate)
 	{
 		// Current start of sequence to check
-		SIZE_T Start = 0;
+		SizeType Start = 0;
 		// Size of sequence to check
-		SIZE_T Size = Num;
+		SizeType Size = Num;
 
 		// With this method, if Size is even it will do one more comparison than necessary, but because Size can be predicted by the CPU it is faster in practice
 		while (Size > 0)
 		{
-			const SIZE_T LeftoverSize = Size % 2;
+			const SizeType LeftoverSize = Size % 2;
 			Size = Size / 2;
 
-			const SIZE_T CheckIndex = Start + Size;
-			const SIZE_T StartIfLess = CheckIndex + LeftoverSize;
+			const SizeType CheckIndex = Start + Size;
+			const SizeType StartIfLess = CheckIndex + LeftoverSize;
 
 			auto&& CheckValue = Invoke(Projection, First[CheckIndex]);
 			Start = !SortPredicate(Value, CheckValue) ? StartIfLess : Start;
@@ -89,12 +89,12 @@ namespace Algo
 	 * @returns Position of the first element >= Value, may be position after last element in range
 	 */
 	template <typename RangeType, typename ValueType, typename SortPredicateType>
-	FORCEINLINE int32 LowerBound(RangeType& Range, const ValueType& Value, SortPredicateType SortPredicate)
+	FORCEINLINE auto LowerBound(RangeType& Range, const ValueType& Value, SortPredicateType SortPredicate) -> decltype(GetNum(Range))
 	{
 		return AlgoImpl::LowerBoundInternal(GetData(Range), GetNum(Range), Value, FIdentityFunctor(), SortPredicate);
 	}
 	template <typename RangeType, typename ValueType>
-	FORCEINLINE int32 LowerBound(RangeType& Range, const ValueType& Value)
+	FORCEINLINE auto LowerBound(RangeType& Range, const ValueType& Value) -> decltype(GetNum(Range))
 	{
 		return AlgoImpl::LowerBoundInternal(GetData(Range), GetNum(Range), Value, FIdentityFunctor(), TLess<>());
 	}
@@ -110,12 +110,12 @@ namespace Algo
 	 * @returns Position of the first element >= Value, may be position after last element in range
 	 */
 	template <typename RangeType, typename ValueType, typename ProjectionType, typename SortPredicateType>
-	FORCEINLINE int32 LowerBoundBy(RangeType& Range, const ValueType& Value, ProjectionType Projection, SortPredicateType SortPredicate)
+	FORCEINLINE auto LowerBoundBy(RangeType& Range, const ValueType& Value, ProjectionType Projection, SortPredicateType SortPredicate) -> decltype(GetNum(Range))
 	{
 		return AlgoImpl::LowerBoundInternal(GetData(Range), GetNum(Range), Value, Projection, SortPredicate);
 	}
 	template <typename RangeType, typename ValueType, typename ProjectionType>
-	FORCEINLINE int32 LowerBoundBy(RangeType& Range, const ValueType& Value, ProjectionType Projection)
+	FORCEINLINE auto LowerBoundBy(RangeType& Range, const ValueType& Value, ProjectionType Projection) -> decltype(GetNum(Range))
 	{
 		return AlgoImpl::LowerBoundInternal(GetData(Range), GetNum(Range), Value, Projection, TLess<>());
 	}
@@ -130,12 +130,12 @@ namespace Algo
 	 * @returns Position of the first element > Value, may be past end of range
 	 */
 	template <typename RangeType, typename ValueType, typename SortPredicateType>
-	FORCEINLINE int32 UpperBound(RangeType& Range, const ValueType& Value, SortPredicateType SortPredicate)
+	FORCEINLINE auto UpperBound(RangeType& Range, const ValueType& Value, SortPredicateType SortPredicate) -> decltype(GetNum(Range))
 	{
 		return AlgoImpl::UpperBoundInternal(GetData(Range), GetNum(Range), Value, FIdentityFunctor(), SortPredicate);
 	}
 	template <typename RangeType, typename ValueType>
-	FORCEINLINE int32 UpperBound(RangeType& Range, const ValueType& Value)
+	FORCEINLINE auto UpperBound(RangeType& Range, const ValueType& Value) -> decltype(GetNum(Range))
 	{
 		return AlgoImpl::UpperBoundInternal(GetData(Range), GetNum(Range), Value, FIdentityFunctor(), TLess<>());
 	}
@@ -151,12 +151,12 @@ namespace Algo
 	 * @returns Position of the first element > Value, may be past end of range
 	 */
 	template <typename RangeType, typename ValueType, typename ProjectionType, typename SortPredicateType>
-	FORCEINLINE int32 UpperBoundBy(RangeType& Range, const ValueType& Value, ProjectionType Projection, SortPredicateType SortPredicate)
+	FORCEINLINE auto UpperBoundBy(RangeType& Range, const ValueType& Value, ProjectionType Projection, SortPredicateType SortPredicate) -> decltype(GetNum(Range))
 	{
 		return AlgoImpl::UpperBoundInternal(GetData(Range), GetNum(Range), Value, Projection, SortPredicate);
 	}
 	template <typename RangeType, typename ValueType, typename ProjectionType>
-	FORCEINLINE int32 UpperBoundBy(RangeType& Range, const ValueType& Value, ProjectionType Projection)
+	FORCEINLINE auto UpperBoundBy(RangeType& Range, const ValueType& Value, ProjectionType Projection) -> decltype(GetNum(Range))
 	{
 		return AlgoImpl::UpperBoundInternal(GetData(Range), GetNum(Range), Value, Projection, TLess<>());
 	}
@@ -170,9 +170,9 @@ namespace Algo
 	 * @return Index of found element, or INDEX_NONE
 	 */
 	template <typename RangeType, typename ValueType, typename SortPredicateType>
-	FORCEINLINE int32 BinarySearch(RangeType& Range, const ValueType& Value, SortPredicateType SortPredicate)
+	FORCEINLINE auto BinarySearch(RangeType& Range, const ValueType& Value, SortPredicateType SortPredicate) -> decltype(GetNum(Range))
 	{
-		SIZE_T CheckIndex = LowerBound(Range, Value, SortPredicate);
+		auto CheckIndex = LowerBound(Range, Value, SortPredicate);
 		if (CheckIndex < GetNum(Range))
 		{
 			auto&& CheckValue = GetData(Range)[CheckIndex];
@@ -185,7 +185,7 @@ namespace Algo
 		return INDEX_NONE;
 	}
 	template <typename RangeType, typename ValueType>
-	FORCEINLINE int32 BinarySearch(RangeType& Range, const ValueType& Value)
+	FORCEINLINE auto BinarySearch(RangeType& Range, const ValueType& Value)
 	{
 		return BinarySearch(Range, Value, TLess<>());
 	}
@@ -200,9 +200,9 @@ namespace Algo
 	 * @return Index of found element, or INDEX_NONE
 	 */
 	template <typename RangeType, typename ValueType, typename ProjectionType, typename SortPredicateType>
-	FORCEINLINE int32 BinarySearchBy(RangeType& Range, const ValueType& Value, ProjectionType Projection, SortPredicateType SortPredicate)
+	FORCEINLINE auto BinarySearchBy(RangeType& Range, const ValueType& Value, ProjectionType Projection, SortPredicateType SortPredicate) -> decltype(GetNum(Range))
 	{
-		SIZE_T CheckIndex = LowerBoundBy(Range, Value, Projection, SortPredicate);
+		auto CheckIndex = LowerBoundBy(Range, Value, Projection, SortPredicate);
 		if (CheckIndex < GetNum(Range))
 		{
 			auto&& CheckValue = Invoke(Projection, GetData(Range)[CheckIndex]);
@@ -215,7 +215,7 @@ namespace Algo
 		return INDEX_NONE;
 	}
 	template <typename RangeType, typename ValueType, typename ProjectionType>
-	FORCEINLINE int32 BinarySearchBy(RangeType& Range, const ValueType& Value, ProjectionType Projection)
+	FORCEINLINE auto BinarySearchBy(RangeType& Range, const ValueType& Value, ProjectionType Projection)
 	{
 		return BinarySearchBy(Range, Value, Projection, TLess<>());
 	}
diff --git a/Engine/Source/Runtime/Core/Public/Algo/Compare.h b/Engine/Source/Runtime/Core/Public/Algo/Compare.h
--- a/Engine/Source/Runtime/Core/Public/Algo/Compare.h
+++ b/Engine/Source/Runtime/Core/Public/Algo/Compare.h
@@ -17,10 +17,11 @@ namespace Algo
 	template <typename InT, typename PredicateT>
 	FORCEINLINE bool CompareByPredicate(const InT& InputA, const InT& InputB, PredicateT Predicate)
 	{
-		if (InputA.Num() == InputB.Num())
-		{
-			uint32 Count = GetNum(InputA);
+		auto Count = GetNum(InputA);
+		auto CountB = GetNum(InputB);
 
+		if (Count == CountB)
+		{
 			auto* A = GetData(InputA);
 			auto* B = GetData(InputB);
 
diff --git a/Engine/Source/Runtime/Core/Public/Algo/Impl/BinaryHeap.h b/Engine/Source/Runtime/Core/Public/Algo/Impl/BinaryHeap.h
--- a/Engine/Source/Runtime/Core/Public/Algo/Impl/BinaryHeap.h
+++ b/Engine/Source/Runtime/Core/Public/Algo/Impl/BinaryHeap.h
@@ -113,7 +113,7 @@ namespace AlgoImpl
 	 * @param	Predicate	A binary predicate object used to specify if one element should precede another.
 	 */
 	template <typename RangeValueType, typename ProjectionType, typename PredicateType>
-	FORCEINLINE void HeapifyInternal(RangeValueType* First, SIZE_T Num, ProjectionType Projection, PredicateType Predicate)
+	FORCEINLINE void HeapifyInternal(RangeValueType* First, int32 Num, ProjectionType Projection, PredicateType Predicate)
 	{
 		for (int32 Index = HeapGetParentIndex(Num - 1); Index >= 0; Index--)
 		{
@@ -130,7 +130,7 @@ namespace AlgoImpl
 	 * @param	Predicate	predicate class
 	 */
 	template <typename RangeValueType, typename ProjectionType, class PredicateType>
-	void HeapSortInternal(RangeValueType* First, SIZE_T Num, ProjectionType Projection, PredicateType Predicate)
+	void HeapSortInternal(RangeValueType* First, int32 Num, ProjectionType Projection, PredicateType Predicate)
 	{
 		TReversePredicate< PredicateType > ReversePredicateWrapper(Predicate); // Reverse the predicate to build a max-heap instead of a min-heap
 		HeapifyInternal(First, Num, Projection, ReversePredicateWrapper);
diff --git a/Engine/Source/Runtime/Core/Public/Algo/IntroSort.h b/Engine/Source/Runtime/Core/Public/Algo/IntroSort.h
--- a/Engine/Source/Runtime/Core/Public/Algo/IntroSort.h
+++ b/Engine/Source/Runtime/Core/Public/Algo/IntroSort.h
@@ -22,8 +22,8 @@ namespace AlgoImpl
 	 * @param Projection	The projection to sort by when applied to the element.
 	 * @param Predicate		predicate class
 	 */
-	template <typename T, typename ProjectionType, typename PredicateType> 
-	void IntroSortInternal(T* First, SIZE_T Num, ProjectionType Projection, PredicateType Predicate)
+	template <typename T, typename IndexType, typename ProjectionType, typename PredicateType> 
+	void IntroSortInternal(T* First, IndexType Num, ProjectionType Projection, PredicateType Predicate)
 	{
 		struct FStack
 		{
@@ -37,13 +37,13 @@ namespace AlgoImpl
 			return;
 		}
 
-		FStack RecursionStack[32]={{First, First+Num-1, (uint32)(FMath::Loge(Num) * 2.f)}}, Current, Inner;
+		FStack RecursionStack[32]={{First, First+Num-1, (uint32)(FMath::Loge((float)Num) * 2.f)}}, Current, Inner;
 		for( FStack* StackTop=RecursionStack; StackTop>=RecursionStack; --StackTop ) //-V625
 		{
 			Current = *StackTop;
 
 		Loop:
-			PTRINT Count = Current.Max - Current.Min + 1;
+			IndexType Count = (IndexType)(Current.Max - Current.Min + 1);
 
 			if ( Current.MaxDepth == 0 )
 			{
diff --git a/Engine/Source/Runtime/Core/Public/Algo/IsHeap.h b/Engine/Source/Runtime/Core/Public/Algo/IsHeap.h
--- a/Engine/Source/Runtime/Core/Public/Algo/IsHeap.h
+++ b/Engine/Source/Runtime/Core/Public/Algo/IsHeap.h
@@ -21,12 +21,12 @@ namespace AlgoImpl
 	 *
 	 * @return	returns		true if the range is a min-heap
 	 */
-	template <typename RangeValueType, typename ProjectionType, typename PredicateType>
-	bool IsHeapInternal(RangeValueType* Heap, SIZE_T Num, ProjectionType Projection, PredicateType Predicate)
+	template <typename RangeValueType, typename IndexType, typename ProjectionType, typename PredicateType>
+	bool IsHeapInternal(RangeValueType* Heap, IndexType Num, ProjectionType Projection, PredicateType Predicate)
 	{
-		for (SIZE_T Index = 1; Index < Num; Index++)
+		for (IndexType Index = 1; Index < Num; Index++)
 		{
-			int32 ParentIndex = HeapGetParentIndex(Index);
+			IndexType ParentIndex = HeapGetParentIndex(Index);
 			if (Predicate( Invoke(Projection, Heap[Index]), Invoke(Projection, Heap[ParentIndex]) ))
 			{
 				return false;
diff --git a/Engine/Source/Runtime/Core/Public/Algo/IsSorted.h b/Engine/Source/Runtime/Core/Public/Algo/IsSorted.h
--- a/Engine/Source/Runtime/Core/Public/Algo/IsSorted.h
+++ b/Engine/Source/Runtime/Core/Public/Algo/IsSorted.h
@@ -11,8 +11,8 @@
 
 namespace AlgoImpl
 {
-	template <typename T, typename ProjectionType, typename PredType>
-	bool IsSortedBy(const T* Range, int32 RangeSize, ProjectionType Proj, PredType Pred)
+	template <typename T, typename IndexType, typename ProjectionType, typename PredType>
+	bool IsSortedBy(const T* Range, IndexType RangeSize, ProjectionType Proj, PredType Pred)
 	{
 		if (RangeSize == 0)
 		{
diff --git a/Engine/Source/Runtime/Core/Public/Algo/Partition.h b/Engine/Source/Runtime/Core/Public/Algo/Partition.h
--- a/Engine/Source/Runtime/Core/Public/Algo/Partition.h
+++ b/Engine/Source/Runtime/Core/Public/Algo/Partition.h
@@ -12,8 +12,8 @@ namespace Algo
 	 * @param	Predicate	unary predicate class
 	 * @return	index of the first element in the second group
 	 */
-	template<class T, class UnaryPredicate>
-	int32 Partition(T* Elements, const int32 Num, const UnaryPredicate& Predicate)
+	template<class T, typename IndexType, class UnaryPredicate>
+	IndexType Partition(T* Elements, const IndexType Num, const UnaryPredicate& Predicate)
 	{
 		T* First = Elements;
 		T* Last = Elements + Num;
@@ -25,7 +25,7 @@ namespace Algo
 				++First;
 				if (First == Last) 
 				{	
-					return First - Elements;
+					return (IndexType)(First - Elements);
 				}
 			}
 		
@@ -34,7 +34,7 @@ namespace Algo
 				--Last;
 				if (First == Last)
 				{
-					return First - Elements;
+					return (IndexType)(First - Elements);
 				}
 			} while (!Predicate(*Last));
 		
@@ -42,6 +42,6 @@ namespace Algo
 			++First;
 		}
 	
-		return First - Elements;
+		return (IndexType)(First - Elements);
 	}
 } //namespace Algo
\ No newline at end of file
diff --git a/Engine/Source/Runtime/Core/Public/Containers/ArrayView.h b/Engine/Source/Runtime/Core/Public/Containers/ArrayView.h
--- a/Engine/Source/Runtime/Core/Public/Containers/ArrayView.h
+++ b/Engine/Source/Runtime/Core/Public/Containers/ArrayView.h
@@ -121,8 +121,10 @@ public:
 	>
 	FORCEINLINE TArrayView(OtherRangeType&& Other)
 		: DataPtr(ArrayViewPrivate::GetDataHelper(Forward<OtherRangeType>(Other)))
-		, ArrayNum(GetNum(Forward<OtherRangeType>(Other)))
 	{
+		const auto InCount = GetNum(Forward<OtherRangeType>(Other));
+		check((InCount >= 0) && (InCount <= TNumericLimits<int32>::Max()));
+		ArrayNum = (int32)InCount;
 	}
 
 	/**
diff --git a/Engine/Source/Runtime/Core/Public/Containers/ContainerAllocationPolicies.h b/Engine/Source/Runtime/Core/Public/Containers/ContainerAllocationPolicies.h
--- a/Engine/Source/Runtime/Core/Public/Containers/ContainerAllocationPolicies.h
+++ b/Engine/Source/Runtime/Core/Public/Containers/ContainerAllocationPolicies.h
@@ -45,7 +45,7 @@ FORCEINLINE SizeType DefaultCalculateSlackShrink(SizeType NumElements, SizeType
 		{
 			if (bAllowQuantize)
 			{
-				Retval = FMemory::QuantizeSize(Retval * BytesPerElement, Alignment) / BytesPerElement;
+				Retval = (SizeType)(FMemory::QuantizeSize(Retval * BytesPerElement, Alignment) / BytesPerElement);
 			}
 		}
 	}
@@ -82,11 +82,11 @@ FORCEINLINE SizeType DefaultCalculateSlackGrow(SizeType NumElements, SizeType Nu
 	}
 	if (bAllowQuantize)
 	{
-		Retval = FMemory::QuantizeSize(Grow * BytesPerElement, Alignment) / BytesPerElement;
+		Retval = (SizeType)(FMemory::QuantizeSize(Grow * BytesPerElement, Alignment) / BytesPerElement);
 	}
 	else
 	{
-		Retval = Grow;
+		Retval = (SizeType)Grow;
 	}
 	// NumElements and MaxElements are stored in 32 bit signed integers so we must be careful not to overflow here.
 	if (NumElements > Retval)
@@ -104,7 +104,7 @@ FORCEINLINE SizeType DefaultCalculateSlackReserve(SizeType NumElements, SIZE_T B
 	checkSlow(NumElements > 0);
 	if (bAllowQuantize)
 	{
-		Retval = FMemory::QuantizeSize(SIZE_T(Retval) * SIZE_T(BytesPerElement), Alignment) / BytesPerElement;
+		Retval = (SizeType)(FMemory::QuantizeSize(SIZE_T(Retval) * SIZE_T(BytesPerElement), Alignment) / BytesPerElement);
 		// NumElements and MaxElements are stored in 32 bit signed integers so we must be careful not to overflow here.
 		if (NumElements > Retval)
 		{
@@ -462,11 +462,11 @@ public:
 		{
 			return DefaultCalculateSlackReserve(NumElements, NumBytesPerElement, true);
 		}
-		FORCEINLINE SizeType CalculateSlackShrink(SizeType NumElements, SizeType NumAllocatedElements, SizeType NumBytesPerElement) const
+		FORCEINLINE SizeType CalculateSlackShrink(SizeType NumElements, SizeType NumAllocatedElements, SIZE_T NumBytesPerElement) const
 		{
 			return DefaultCalculateSlackShrink(NumElements, NumAllocatedElements, NumBytesPerElement, true);
 		}
-		FORCEINLINE SizeType CalculateSlackGrow(SizeType NumElements, SizeType NumAllocatedElements, SizeType NumBytesPerElement) const
+		FORCEINLINE SizeType CalculateSlackGrow(SizeType NumElements, SizeType NumAllocatedElements, SIZE_T NumBytesPerElement) const
 		{
 			return DefaultCalculateSlackGrow(NumElements, NumAllocatedElements, NumBytesPerElement, true);
 		}
diff --git a/Engine/Source/Runtime/Core/Public/Containers/HashTable.h b/Engine/Source/Runtime/Core/Public/Containers/HashTable.h
--- a/Engine/Source/Runtime/Core/Public/Containers/HashTable.h
+++ b/Engine/Source/Runtime/Core/Public/Containers/HashTable.h
@@ -224,7 +224,7 @@ FORCEINLINE FHashTable::FHashTable( uint32 InHashSize, uint32 InIndexSize )
 	
 	if( IndexSize )
 	{
-		HashMask = HashSize - 1;
+		HashMask = (uint16)(HashSize - 1);
 		
 		Hash = new uint32[ HashSize ];
 		NextIndex = new uint32[ IndexSize ];
diff --git a/Engine/Source/Runtime/Core/Public/Containers/SortedMap.h b/Engine/Source/Runtime/Core/Public/Containers/SortedMap.h
--- a/Engine/Source/Runtime/Core/Public/Containers/SortedMap.h
+++ b/Engine/Source/Runtime/Core/Public/Containers/SortedMap.h
@@ -116,7 +116,7 @@ public:
 	 *
 	 * @return number of bytes allocated by this container.
 	 */
-	FORCEINLINE uint32 GetAllocatedSize() const
+	FORCEINLINE SIZE_T GetAllocatedSize() const
 	{
 		return Pairs.GetAllocatedSize();
 	}
diff --git a/Engine/Source/Runtime/Core/Public/Containers/StringConv.h b/Engine/Source/Runtime/Core/Public/Containers/StringConv.h
--- a/Engine/Source/Runtime/Core/Public/Containers/StringConv.h
+++ b/Engine/Source/Runtime/Core/Public/Containers/StringConv.h
@@ -74,7 +74,7 @@ namespace StringConv
 	static FORCEINLINE void DecodeSurrogate(const uint32 Codepoint, uint16& OutHighSurrogate, uint16& OutLowSurrogate)
 	{
 		const uint32 TmpCodepoint = Codepoint - 0x10000;
-		OutHighSurrogate = (TmpCodepoint >> 10) + HIGH_SURROGATE_START_CODEPOINT;
+		OutHighSurrogate = (uint16)((TmpCodepoint >> 10) + HIGH_SURROGATE_START_CODEPOINT);
 		OutLowSurrogate = (TmpCodepoint & 0x3FF) + LOW_SURROGATE_START_CODEPOINT;
 	}
 
@@ -266,7 +266,7 @@ public:
 			}
 		}
 
-		return static_cast<int32>(OutputIterator - OutputIteratorStartPosition);
+		return UE_PTRDIFF_TO_INT32(OutputIterator - OutputIteratorStartPosition);
 	}
 
 	/**
@@ -346,7 +346,7 @@ private:
 				{
 					const uint32 LowSurrogate = Codepoint;
 					// Combine our high and low surrogates together to a single Unicode codepoint
-					Codepoint = StringConv::EncodeSurrogate(HighSurrogate, LowSurrogate);
+					Codepoint = StringConv::EncodeSurrogate((uint16)HighSurrogate, (uint16)LowSurrogate);
 				}
 				else
 				{
@@ -657,7 +657,7 @@ private:
 		while (Source < SourceEnd && DestLen > 0)
 		{
 			// Read our codepoint, advancing the source pointer
-			uint32 Codepoint = CodepointFromUtf8(Source, SourceEnd - Source);
+			uint32 Codepoint = CodepointFromUtf8(Source, UE_PTRDIFF_TO_UINT32(SourceEnd - Source));
 
 #if !PLATFORM_TCHAR_IS_4_BYTES
 			// We want to write out two chars
@@ -670,8 +670,8 @@ private:
 					uint16 LowSurrogate = 0;
 					StringConv::DecodeSurrogate(Codepoint, HighSurrogate, LowSurrogate);
 
-					*(ConvertedBuffer++) = (TCHAR)HighSurrogate;
-					*(ConvertedBuffer++) = (TCHAR)LowSurrogate;
+					*(ConvertedBuffer++) = (ToType)HighSurrogate;
+					*(ConvertedBuffer++) = (ToType)LowSurrogate;
 					DestLen -= 2;
 					continue;
 				}
@@ -686,7 +686,7 @@ private:
 			}
 #endif	// !PLATFORM_TCHAR_IS_4_BYTES
 
-			*(ConvertedBuffer++) = Codepoint;
+			*(ConvertedBuffer++) = (ToType)Codepoint;
 			--DestLen;
 		}
 	}
@@ -757,7 +757,7 @@ public:
 			*(OutputIterator++) = Codepoint;
 		}
 
-		return static_cast<int32>(OutputIterator - OutputIteratorStartPosition);
+		return UE_PTRDIFF_TO_INT32(OutputIterator - OutputIteratorStartPosition);
 	}
 
 	/**
diff --git a/Engine/Source/Runtime/Core/Public/GenericPlatform/MicrosoftPlatformString.h b/Engine/Source/Runtime/Core/Public/GenericPlatform/MicrosoftPlatformString.h
--- a/Engine/Source/Runtime/Core/Public/GenericPlatform/MicrosoftPlatformString.h
+++ b/Engine/Source/Runtime/Core/Public/GenericPlatform/MicrosoftPlatformString.h
@@ -83,7 +83,7 @@ struct FMicrosoftPlatformString : public FGenericPlatformString
 
 	static FORCEINLINE float Atof(const WIDECHAR* String)
 	{
-		return _tstof( String );
+		return (float)_tstof( String );
 	}
 
 	static FORCEINLINE double Atod(const WIDECHAR* String)
diff --git a/Engine/Source/Runtime/Core/Public/HAL/Allocators/CachedOSPageAllocator.h b/Engine/Source/Runtime/Core/Public/HAL/Allocators/CachedOSPageAllocator.h
--- a/Engine/Source/Runtime/Core/Public/HAL/Allocators/CachedOSPageAllocator.h
+++ b/Engine/Source/Runtime/Core/Public/HAL/Allocators/CachedOSPageAllocator.h
@@ -19,17 +19,17 @@ protected:
 		}
 	};
 
-	void* AllocateImpl(SIZE_T Size, uint32 CachedByteLimit, FFreePageBlock* First, FFreePageBlock* Last, uint32& FreedPageBlocksNum, uint32& CachedTotal);
-	void FreeImpl(void* Ptr, SIZE_T Size, uint32 NumCacheBlocks, uint32 CachedByteLimit, FFreePageBlock* First, uint32& FreedPageBlocksNum, uint32& CachedTotal);
-	void FreeAllImpl(FFreePageBlock* First, uint32& FreedPageBlocksNum, uint32& CachedTotal);
+	void* AllocateImpl(SIZE_T Size, uint32 CachedByteLimit, FFreePageBlock* First, FFreePageBlock* Last, uint32& FreedPageBlocksNum, SIZE_T& CachedTotal);
+	void FreeImpl(void* Ptr, SIZE_T Size, uint32 NumCacheBlocks, uint32 CachedByteLimit, FFreePageBlock* First, uint32& FreedPageBlocksNum, SIZE_T& CachedTotal);
+	void FreeAllImpl(FFreePageBlock* First, uint32& FreedPageBlocksNum, SIZE_T& CachedTotal);
 };
 
 template <uint32 NumCacheBlocks, uint32 CachedByteLimit>
 struct TCachedOSPageAllocator : private FCachedOSPageAllocator
 {
 	TCachedOSPageAllocator()
-		: FreedPageBlocksNum(0)
-		, CachedTotal       (0)
+		: CachedTotal(0)
+		, FreedPageBlocksNum(0)
 	{
 	}
 
@@ -49,11 +49,11 @@ struct TCachedOSPageAllocator : private FCachedOSPageAllocator
 
 	uint64 GetCachedFreeTotal()
 	{
-		return static_cast<uint64>(CachedTotal);
+		return CachedTotal;
 	}
 
 private:
 	FFreePageBlock FreedPageBlocks[NumCacheBlocks];
+	SIZE_T         CachedTotal;
 	uint32         FreedPageBlocksNum;
-	uint32         CachedTotal;
 };
diff --git a/Engine/Source/Runtime/Core/Public/HAL/IPlatformFileLogWrapper.h b/Engine/Source/Runtime/Core/Public/HAL/IPlatformFileLogWrapper.h
--- a/Engine/Source/Runtime/Core/Public/HAL/IPlatformFileLogWrapper.h
+++ b/Engine/Source/Runtime/Core/Public/HAL/IPlatformFileLogWrapper.h
@@ -51,7 +51,7 @@ public:
 		FILE_LOG(LogPlatformFile, VeryVerbose, TEXT("Tell %s"), *Filename);
 		double StartTime = FPlatformTime::Seconds();
 		int64 Result = FileHandle->Tell();
-		float ThisTime = (FPlatformTime::Seconds() - StartTime) / 1000.0;
+		double ThisTime = (FPlatformTime::Seconds() - StartTime) / 1000.0;
 		FILE_LOG(LogPlatformFile, VeryVerbose, TEXT("Tell return %lld [%fms]"), Result, ThisTime);
 		return Result;
 	}
@@ -60,7 +60,7 @@ public:
 		FILE_LOG(LogPlatformFile, VeryVerbose, TEXT("Seek %s %lld"), *Filename, NewPosition);
 		double StartTime = FPlatformTime::Seconds();
 		bool Result = FileHandle->Seek(NewPosition);
-		float ThisTime = (FPlatformTime::Seconds() - StartTime) / 1000.0;
+		double ThisTime = (FPlatformTime::Seconds() - StartTime) / 1000.0;
 		FILE_LOG(LogPlatformFile, VeryVerbose, TEXT("Seek return %d [%fms]"), int32(Result), ThisTime);
 		return Result;
 	}
@@ -69,7 +69,7 @@ public:
 		FILE_LOG(LogPlatformFile, VeryVerbose, TEXT("SeekFromEnd %s %lld"), *Filename, NewPositionRelativeToEnd);
 		double StartTime = FPlatformTime::Seconds();
 		bool Result = FileHandle->SeekFromEnd(NewPositionRelativeToEnd);
-		float ThisTime = (FPlatformTime::Seconds() - StartTime) / 1000.0;
+		double ThisTime = (FPlatformTime::Seconds() - StartTime) / 1000.0;
 		FILE_LOG(LogPlatformFile, VeryVerbose, TEXT("SeekFromEnd return %d [%fms]"), int32(Result), ThisTime);
 		return Result;
 	}
@@ -78,7 +78,7 @@ public:
 		FILE_LOG(LogPlatformFile, VeryVerbose, TEXT("Read %s %lld"), *Filename, BytesToRead);
 		double StartTime = FPlatformTime::Seconds();
 		bool Result = FileHandle->Read(Destination, BytesToRead);
-		float ThisTime = (FPlatformTime::Seconds() - StartTime) / 1000.0;
+		double ThisTime = (FPlatformTime::Seconds() - StartTime) / 1000.0;
 		FILE_LOG(LogPlatformFile, VeryVerbose, TEXT("Read return %d [%fms]"), int32(Result), ThisTime);
 		return Result;
 	}
@@ -87,7 +87,7 @@ public:
 		FILE_LOG(LogPlatformFile, VeryVerbose, TEXT("Write %s %lld"), *Filename, BytesToWrite);
 		double StartTime = FPlatformTime::Seconds();
 		bool Result = FileHandle->Write(Source, BytesToWrite);
-		float ThisTime = (FPlatformTime::Seconds() - StartTime) / 1000.0;
+		double ThisTime = (FPlatformTime::Seconds() - StartTime) / 1000.0;
 		FILE_LOG(LogPlatformFile, VeryVerbose, TEXT("Write return %d [%fms]"), int32(Result), ThisTime);
 		return Result;
 	}
@@ -96,7 +96,7 @@ public:
 		FILE_LOG(LogPlatformFile, Verbose, TEXT("Size %s"), *Filename);
 		double StartTime = FPlatformTime::Seconds();
 		int64 Result = FileHandle->Size();
-		float ThisTime = (FPlatformTime::Seconds() - StartTime) / 1000.0;
+		double ThisTime = (FPlatformTime::Seconds() - StartTime) / 1000.0;
 		FILE_LOG(LogPlatformFile, Verbose, TEXT("Size return %lld [%fms]"), Result, ThisTime);
 		return Result;
 	}
@@ -105,7 +105,7 @@ public:
 		FILE_LOG(LogPlatformFile, Verbose, TEXT("Flush %s %s"), *Filename, (bFullFlush ? TEXT("full") : TEXT("partial")));
 		double StartTime = FPlatformTime::Seconds();
 		bool bResult = FileHandle->Flush(bFullFlush);
-		float ThisTime = (FPlatformTime::Seconds() - StartTime) / 1000.0;
+		double ThisTime = (FPlatformTime::Seconds() - StartTime) / 1000.0;
 		FILE_LOG(LogPlatformFile, Verbose, TEXT("Flush return %d [%fms]"), int32(bResult), ThisTime);
 		return bResult;
 	}
@@ -114,7 +114,7 @@ public:
 		FILE_LOG(LogPlatformFile, Verbose, TEXT("Truncate %s %lld"), *Filename, NewSize);
 		double StartTime = FPlatformTime::Seconds();
 		bool bResult = FileHandle->Truncate(NewSize);
-		float ThisTime = (FPlatformTime::Seconds() - StartTime) / 1000.0;
+		double ThisTime = (FPlatformTime::Seconds() - StartTime) / 1000.0;
 		FILE_LOG(LogPlatformFile, Verbose, TEXT("Truncate return %d [%fms]"), int32(bResult), ThisTime);
 		return bResult;
 	}
@@ -170,7 +170,7 @@ public:
 		FILE_LOG(LogPlatformFile, Log, TEXT("%s"), *DataStr);
 		double StartTime = FPlatformTime::Seconds();
 		bool Result = LowerLevel->FileExists(Filename);
-		float ThisTime = (FPlatformTime::Seconds() - StartTime) / 1000.0;
+		double ThisTime = (FPlatformTime::Seconds() - StartTime) / 1000.0;
 		FILE_LOG(LogPlatformFile, Log, TEXT("FileExists return %d [%fms]"), int32(Result), ThisTime);
 		return Result;
 	}
@@ -181,7 +181,7 @@ public:
 		FILE_LOG(LogPlatformFile, Log, TEXT("%s"), *DataStr);
 		double StartTime = FPlatformTime::Seconds();
 		int64 Result = LowerLevel->FileSize(Filename);
-		float ThisTime = (FPlatformTime::Seconds() - StartTime) / 1000.0;
+		double ThisTime = (FPlatformTime::Seconds() - StartTime) / 1000.0;
 		FILE_LOG(LogPlatformFile, Log, TEXT("FileSize return %lld [%fms]"), Result, ThisTime);
 		return Result;
 	}
@@ -192,7 +192,7 @@ public:
 		FILE_LOG(LogPlatformFile, Log, TEXT("%s"), *DataStr);
 		double StartTime = FPlatformTime::Seconds();
 		bool Result = LowerLevel->DeleteFile(Filename);
-		float ThisTime = (FPlatformTime::Seconds() - StartTime) / 1000.0;
+		double ThisTime = (FPlatformTime::Seconds() - StartTime) / 1000.0;
 		FILE_LOG(LogPlatformFile, Log, TEXT("DeleteFile return %d [%fms]"), int32(Result), ThisTime);
 		return Result;
 	}
@@ -203,7 +203,7 @@ public:
 		FILE_LOG(LogPlatformFile, Log, TEXT("%s"), *DataStr);
 		double StartTime = FPlatformTime::Seconds();
 		bool Result = LowerLevel->IsReadOnly(Filename);
-		float ThisTime = (FPlatformTime::Seconds() - StartTime) / 1000.0;
+		double ThisTime = (FPlatformTime::Seconds() - StartTime) / 1000.0;
 		FILE_LOG(LogPlatformFile, Log, TEXT("IsReadOnly return %d [%fms]"), int32(Result), ThisTime);
 		return Result;
 	}
@@ -214,7 +214,7 @@ public:
 		FILE_LOG(LogPlatformFile, Log, TEXT("%s"), *DataStr);
 		double StartTime = FPlatformTime::Seconds();
 		bool Result = LowerLevel->MoveFile(To, From);
-		float ThisTime = (FPlatformTime::Seconds() - StartTime) / 1000.0;
+		double ThisTime = (FPlatformTime::Seconds() - StartTime) / 1000.0;
 		FILE_LOG(LogPlatformFile, Log, TEXT("MoveFile return %d [%fms]"), int32(Result), ThisTime);
 		return Result;
 	}
@@ -225,7 +225,7 @@ public:
 		FILE_LOG(LogPlatformFile, Log, TEXT("%s"), *DataStr);
 		double StartTime = FPlatformTime::Seconds();
 		bool Result = LowerLevel->SetReadOnly(Filename, bNewReadOnlyValue);
-		float ThisTime = (FPlatformTime::Seconds() - StartTime) / 1000.0;
+		double ThisTime = (FPlatformTime::Seconds() - StartTime) / 1000.0;
 		FILE_LOG(LogPlatformFile, Log, TEXT("SetReadOnly return %d [%fms]"), int32(Result), ThisTime);
 		return Result;
 	}
@@ -236,7 +236,7 @@ public:
 		FILE_LOG(LogPlatformFile, Log, TEXT("%s"), *DataStr);
 		double StartTime = FPlatformTime::Seconds();
 		FDateTime Result = LowerLevel->GetTimeStamp(Filename);
-		float ThisTime = (FPlatformTime::Seconds() - StartTime) / 1000.0;
+		double ThisTime = (FPlatformTime::Seconds() - StartTime) / 1000.0;
 		FILE_LOG(LogPlatformFile, Log, TEXT("GetTimeStamp return %llx [%fms]"), Result.GetTicks(), ThisTime);
 		return Result;
 	}
@@ -247,7 +247,7 @@ public:
 		FILE_LOG(LogPlatformFile, Log, TEXT("%s"), *DataStr);
 		double StartTime = FPlatformTime::Seconds();
 		LowerLevel->SetTimeStamp(Filename, DateTime);
-		float ThisTime = (FPlatformTime::Seconds() - StartTime) / 1000.0;
+		double ThisTime = (FPlatformTime::Seconds() - StartTime) / 1000.0;
 		FILE_LOG(LogPlatformFile, Log, TEXT("SetTimeStamp [%fms]"), ThisTime);
 	}
 	virtual FDateTime	GetAccessTimeStamp(const TCHAR* Filename) override
@@ -257,7 +257,7 @@ public:
 		FILE_LOG(LogPlatformFile, Log, TEXT("%s"), *DataStr);
 		double StartTime = FPlatformTime::Seconds();
 		FDateTime Result = LowerLevel->GetAccessTimeStamp(Filename);
-		float ThisTime = (FPlatformTime::Seconds() - StartTime) / 1000.0;
+		double ThisTime = (FPlatformTime::Seconds() - StartTime) / 1000.0;
 		FILE_LOG(LogPlatformFile, Log, TEXT("GetAccessTimeStamp return %llx [%fms]"), Result.GetTicks(), ThisTime);
 		return Result;
 	}
@@ -268,7 +268,7 @@ public:
 		FILE_LOG(LogPlatformFile, Log, TEXT("%s"), *DataStr);
 		double StartTime = FPlatformTime::Seconds();
 		FString Result = LowerLevel->GetFilenameOnDisk(Filename);
-		float ThisTime = (FPlatformTime::Seconds() - StartTime) / 1000.0;
+		double ThisTime = (FPlatformTime::Seconds() - StartTime) / 1000.0;
 		FILE_LOG(LogPlatformFile, Log, TEXT("GetFilenameOnDisk return %s [%fms]"), *Result, ThisTime);
 		return Result;
 	}
@@ -279,7 +279,7 @@ public:
 		FILE_LOG(LogPlatformFile, Log, TEXT("%s"), *DataStr);
 		double StartTime = FPlatformTime::Seconds();
 		IFileHandle* Result = LowerLevel->OpenRead(Filename, bAllowWrite);
-		float ThisTime = (FPlatformTime::Seconds() - StartTime) / 1000.0;
+		double ThisTime = (FPlatformTime::Seconds() - StartTime) / 1000.0;
 		FILE_LOG(LogPlatformFile, Log, TEXT("OpenRead return %llx [%fms]"), uint64(Result), ThisTime);
 		return Result ? (new FLoggedFileHandle(Result, Filename, *this)) : Result;
 	}
@@ -290,7 +290,7 @@ public:
 		FILE_LOG(LogPlatformFile, Log, TEXT("%s"), *DataStr);
 		double StartTime = FPlatformTime::Seconds();
 		IFileHandle* Result = LowerLevel->OpenWrite(Filename, bAppend, bAllowRead);
-		float ThisTime = (FPlatformTime::Seconds() - StartTime) / 1000.0;
+		double ThisTime = (FPlatformTime::Seconds() - StartTime) / 1000.0;
 		FILE_LOG(LogPlatformFile, Log, TEXT("OpenWrite return %llx [%fms]"), uint64(Result), ThisTime);
 		return Result ? (new FLoggedFileHandle(Result, Filename, *this)) : Result;
 	}
@@ -302,7 +302,7 @@ public:
 		FILE_LOG(LogPlatformFile, Log, TEXT("%s"), *DataStr);
 		double StartTime = FPlatformTime::Seconds();
 		bool Result = LowerLevel->DirectoryExists(Directory);
-		float ThisTime = (FPlatformTime::Seconds() - StartTime) / 1000.0;
+		double ThisTime = (FPlatformTime::Seconds() - StartTime) / 1000.0;
 		FILE_LOG(LogPlatformFile, Log, TEXT("DirectoryExists return %d [%fms]"), int32(Result), ThisTime);
 		return Result;
 	}
@@ -313,7 +313,7 @@ public:
 		FILE_LOG(LogPlatformFile, Log, TEXT("%s"), *DataStr);
 		double StartTime = FPlatformTime::Seconds();
 		bool Result = LowerLevel->CreateDirectory(Directory);
-		float ThisTime = (FPlatformTime::Seconds() - StartTime) / 1000.0;
+		double ThisTime = (FPlatformTime::Seconds() - StartTime) / 1000.0;
 		FILE_LOG(LogPlatformFile, Log, TEXT("CreateDirectory return %d [%fms]"), int32(Result), ThisTime);
 		return Result;
 	}
@@ -324,7 +324,7 @@ public:
 		FILE_LOG(LogPlatformFile, Log, TEXT("%s"), *DataStr);
 		double StartTime = FPlatformTime::Seconds();
 		bool Result = LowerLevel->DeleteDirectory(Directory);
-		float ThisTime = (FPlatformTime::Seconds() - StartTime) / 1000.0;
+		double ThisTime = (FPlatformTime::Seconds() - StartTime) / 1000.0;
 		FILE_LOG(LogPlatformFile, Log, TEXT("DeleteDirectory return %d [%fms]"), int32(Result), ThisTime);
 		return Result;
 	}
@@ -335,7 +335,7 @@ public:
 		FILE_LOG(LogPlatformFile, Log, TEXT("%s"), *DataStr);
 		double StartTime = FPlatformTime::Seconds();
 		FFileStatData Result = LowerLevel->GetStatData(FilenameOrDirectory);
-		float ThisTime = (FPlatformTime::Seconds() - StartTime) / 1000.0;
+		double ThisTime = (FPlatformTime::Seconds() - StartTime) / 1000.0;
 		FILE_LOG(LogPlatformFile, Log, TEXT("GetStatData return %d [%fms]"), int32(Result.bIsValid), ThisTime);
 		return Result;
 	}
@@ -353,7 +353,7 @@ public:
 			FILE_LOG(LogPlatformFile, Verbose, TEXT("Visit %s %d"), FilenameOrDirectory, int32(bIsDirectory));
 			double StartTime = FPlatformTime::Seconds();
 			bool Result = Visitor.Visit(FilenameOrDirectory, bIsDirectory);
-			float ThisTime = (FPlatformTime::Seconds() - StartTime) / 1000.0;
+			double ThisTime = (FPlatformTime::Seconds() - StartTime) / 1000.0;
 			FILE_LOG(LogPlatformFile, Verbose, TEXT("Visit return %d [%fms]"), int32(Result), ThisTime);
 			return Result;
 		}
@@ -367,7 +367,7 @@ public:
 		double StartTime = FPlatformTime::Seconds();
 		FLogVisitor LogVisitor(Visitor);
 		bool Result = LowerLevel->IterateDirectory(Directory, LogVisitor);
-		float ThisTime = (FPlatformTime::Seconds() - StartTime) / 1000.0;
+		double ThisTime = (FPlatformTime::Seconds() - StartTime) / 1000.0;
 		FILE_LOG(LogPlatformFile, Log, TEXT("IterateDirectory return %d [%fms]"), int32(Result), ThisTime);
 		return Result;
 	}
@@ -379,7 +379,7 @@ public:
 		double StartTime = FPlatformTime::Seconds();
 		FLogVisitor LogVisitor(Visitor);
 		bool Result = LowerLevel->IterateDirectoryRecursively(Directory, LogVisitor);
-		float ThisTime = (FPlatformTime::Seconds() - StartTime) / 1000.0;
+		double ThisTime = (FPlatformTime::Seconds() - StartTime) / 1000.0;
 		FILE_LOG(LogPlatformFile, Log, TEXT("IterateDirectoryRecursively return %d [%fms]"), int32(Result), ThisTime);
 		return Result;
 	}
@@ -397,7 +397,7 @@ public:
 			FILE_LOG(LogPlatformFile, Verbose, TEXT("Visit %s %d"), FilenameOrDirectory, int32(StatData.bIsDirectory));
 			double StartTime = FPlatformTime::Seconds();
 			bool Result = Visitor.Visit(FilenameOrDirectory, StatData);
-			float ThisTime = (FPlatformTime::Seconds() - StartTime) / 1000.0;
+			double ThisTime = (FPlatformTime::Seconds() - StartTime) / 1000.0;
 			FILE_LOG(LogPlatformFile, Verbose, TEXT("Visit return %d [%fms]"), int32(Result), ThisTime);
 			return Result;
 		}
@@ -411,7 +411,7 @@ public:
 		double StartTime = FPlatformTime::Seconds();
 		FLogStatVisitor LogVisitor(Visitor);
 		bool Result = LowerLevel->IterateDirectoryStat(Directory, LogVisitor);
-		float ThisTime = (FPlatformTime::Seconds() - StartTime) / 1000.0;
+		double ThisTime = (FPlatformTime::Seconds() - StartTime) / 1000.0;
 		FILE_LOG(LogPlatformFile, Log, TEXT("IterateDirectoryStat return %d [%fms]"), int32(Result), ThisTime);
 		return Result;
 	}
@@ -423,7 +423,7 @@ public:
 		double StartTime = FPlatformTime::Seconds();
 		FLogStatVisitor LogVisitor(Visitor);
 		bool Result = LowerLevel->IterateDirectoryStatRecursively(Directory, LogVisitor);
-		float ThisTime = (FPlatformTime::Seconds() - StartTime) / 1000.0;
+		double ThisTime = (FPlatformTime::Seconds() - StartTime) / 1000.0;
 		FILE_LOG(LogPlatformFile, Log, TEXT("IterateDirectoryStatRecursively return %d [%fms]"), int32(Result), ThisTime);
 		return Result;
 	}
@@ -435,7 +435,7 @@ public:
 		FILE_LOG(LogPlatformFile, Log, TEXT("%s"), *DataStr);
 		double StartTime = FPlatformTime::Seconds();
 		bool Result = LowerLevel->DeleteDirectoryRecursively(Directory);
-		float ThisTime = (FPlatformTime::Seconds() - StartTime) / 1000.0;
+		double ThisTime = (FPlatformTime::Seconds() - StartTime) / 1000.0;
 		FILE_LOG(LogPlatformFile, Log, TEXT("DeleteDirectoryRecursively return %d [%fms]"), int32(Result), ThisTime);
 		return Result;
 	}
@@ -446,7 +446,7 @@ public:
 		FILE_LOG(LogPlatformFile, Log, TEXT("%s"), *DataStr);
 		double StartTime = FPlatformTime::Seconds();
 		bool Result = LowerLevel->CopyFile(To, From, ReadFlags, WriteFlags);
-		float ThisTime = (FPlatformTime::Seconds() - StartTime) / 1000.0;
+		double ThisTime = (FPlatformTime::Seconds() - StartTime) / 1000.0;
 		FILE_LOG(LogPlatformFile, Log, TEXT("CopyFile return %d [%fms]"), int32(Result), ThisTime);
 		return Result;
 	}
@@ -476,7 +476,7 @@ public:
 		FILE_LOG(LogPlatformFile, Log, TEXT("%s"), *DataStr);
 		double StartTime = FPlatformTime::Seconds();
 		IAsyncReadFileHandle* Result = LowerLevel->OpenAsyncRead(Filename);
-		float ThisTime = (FPlatformTime::Seconds() - StartTime) / 1000.0;
+		double ThisTime = (FPlatformTime::Seconds() - StartTime) / 1000.0;
 		FILE_LOG(LogPlatformFile, Log, TEXT("OpenAsyncRead return %llx [%fms]"), uint64(Result), ThisTime);
 		//@todo no wrapped logging for async file handles (yet)
 		return Result;
diff --git a/Engine/Source/Runtime/Core/Public/HAL/MallocBinned2.h b/Engine/Source/Runtime/Core/Public/HAL/MallocBinned2.h
--- a/Engine/Source/Runtime/Core/Public/HAL/MallocBinned2.h
+++ b/Engine/Source/Runtime/Core/Public/HAL/MallocBinned2.h
@@ -108,7 +108,7 @@ class CORE_API FMallocBinned2 final : public FMalloc
 			CANARY_VALUE = 0xe3
 		};
 
-		FORCEINLINE FFreeBlock(uint32 InPageSize, uint32 InBlockSize, uint32 InPoolIndex)
+		FORCEINLINE FFreeBlock(uint32 InPageSize, uint16 InBlockSize, uint8 InPoolIndex)
 			: BlockSize(InBlockSize)
 			, PoolIndex(InPoolIndex)
 			, Canary(CANARY_VALUE)
@@ -205,7 +205,7 @@ class CORE_API FMallocBinned2 final : public FMalloc
 
 		void Init(uint32 InPageSize, uint64 InNumPoolsPerPage, uint64 AddressLimit)
 		{
-			uint64 PoolPageToPoolBitShift = FPlatformMath::CeilLogTwo(InNumPoolsPerPage);
+			uint64 PoolPageToPoolBitShift = FPlatformMath::CeilLogTwo64(InNumPoolsPerPage);
 
 			PtrToPoolPageBitShift = FPlatformMath::CeilLogTwo(InPageSize);
 			HashKeyShift          = PtrToPoolPageBitShift + PoolPageToPoolBitShift;
@@ -217,7 +217,7 @@ class CORE_API FMallocBinned2 final : public FMalloc
 		{
 			OutBucketCollision = (UPTRINT)InPtr >> HashKeyShift;
 			OutBucketIndex = uint32(OutBucketCollision & (MaxHashBuckets - 1));
-			OutPoolIndex   = ((UPTRINT)InPtr >> PtrToPoolPageBitShift) & PoolMask;
+			OutPoolIndex   = (uint32)(((UPTRINT)InPtr >> PtrToPoolPageBitShift) & PoolMask);
 		}
 
 		FORCEINLINE uint64 GetMaxHashBuckets() const
diff --git a/Engine/Source/Runtime/Core/Public/HAL/MallocBinned3.h b/Engine/Source/Runtime/Core/Public/HAL/MallocBinned3.h
--- a/Engine/Source/Runtime/Core/Public/HAL/MallocBinned3.h
+++ b/Engine/Source/Runtime/Core/Public/HAL/MallocBinned3.h
@@ -4,6 +4,8 @@
 
 #include "CoreTypes.h"
 
+PRAGMA_DISABLE_UNSAFE_TYPECAST_WARNINGS
+
 #if PLATFORM_64BITS && PLATFORM_HAS_FPlatformVirtualMemoryBlock
 #include "HAL/MallocBinnedCommon.h"
 #include "Misc/AssertionMacros.h"
@@ -112,7 +114,7 @@ class CORE_API FMallocBinned3 final : public FMalloc
 			CANARY_VALUE = 0xe7
 		};
 
-		FORCEINLINE FFreeBlock(uint32 InPageSize, uint32 InBlockSize, uint32 InPoolIndex)
+		FORCEINLINE FFreeBlock(uint32 InPageSize, uint32 InBlockSize, uint8 InPoolIndex)
 			: BlockSizeShifted(InBlockSize >> BINNED3_MINIMUM_ALIGNMENT_SHIFT)
 			, PoolIndex(InPoolIndex)
 			, Canary(CANARY_VALUE)
@@ -470,7 +472,7 @@ class CORE_API FMallocBinned3 final : public FMalloc
 	{
 		uint64 Result = PoolIndexFromPtr(Ptr);
 		check(Result < BINNED3_SMALL_POOL_COUNT);
-		return Result;
+		return (uint32)Result;
 	}
 
 	FORCEINLINE bool IsOSAllocation(const void* Ptr)
@@ -716,4 +718,6 @@ public:
 		#include "FMemory.inl"
 	#endif
 #endif
-#endif
\ No newline at end of file
+#endif
+
+PRAGMA_ENABLE_UNSAFE_TYPECAST_WARNINGS
diff --git a/Engine/Source/Runtime/Core/Public/HAL/MallocBinnedArena.h b/Engine/Source/Runtime/Core/Public/HAL/MallocBinnedArena.h
--- a/Engine/Source/Runtime/Core/Public/HAL/MallocBinnedArena.h
+++ b/Engine/Source/Runtime/Core/Public/HAL/MallocBinnedArena.h
@@ -4,6 +4,8 @@
 
 #include "CoreTypes.h"
 
+PRAGMA_DISABLE_UNSAFE_TYPECAST_WARNINGS
+
 #if PLATFORM_64BITS && PLATFORM_HAS_FPlatformVirtualMemoryBlock
 #include "Misc/AssertionMacros.h"
 #include "Misc/ScopeLock.h"
@@ -350,26 +352,26 @@ class CORE_API FMallocBinnedArena final : public FMalloc
 			{
 				PoolIndex = ArenaParams.PoolCount - 1;
 			}
-			if ((uint8*)Ptr < PoolBaseVMPtr[PoolIndex])
+			if ((uint8*)Ptr < PoolBaseVMPtr[(int32)PoolIndex])
 			{
 				do
 				{
 					PoolIndex--;
 					check(PoolIndex < ArenaParams.PoolCount);
-				} while ((uint8*)Ptr < PoolBaseVMPtr[PoolIndex]);
-				if ((uint8*)Ptr >= PoolBaseVMPtr[PoolIndex] + ArenaParams.MaxMemoryPerBlockSize)
+				} while ((uint8*)Ptr < PoolBaseVMPtr[(int32)PoolIndex]);
+				if ((uint8*)Ptr >= PoolBaseVMPtr[(int32)PoolIndex] + ArenaParams.MaxMemoryPerBlockSize)
 				{
 					PoolIndex = ArenaParams.PoolCount; // was in the gap
 				}
 			}
-			else if ((uint8*)Ptr >= PoolBaseVMPtr[PoolIndex] + ArenaParams.MaxMemoryPerBlockSize)
+			else if ((uint8*)Ptr >= PoolBaseVMPtr[(int32)PoolIndex] + ArenaParams.MaxMemoryPerBlockSize)
 			{
 				do
 				{
 					PoolIndex++;
 					check(PoolIndex < ArenaParams.PoolCount);
-				} while ((uint8*)Ptr >= PoolBaseVMPtr[PoolIndex] + ArenaParams.MaxMemoryPerBlockSize);
-				if ((uint8*)Ptr < PoolBaseVMPtr[PoolIndex])
+				} while ((uint8*)Ptr >= PoolBaseVMPtr[(int32)PoolIndex] + ArenaParams.MaxMemoryPerBlockSize);
+				if ((uint8*)Ptr < PoolBaseVMPtr[(int32)PoolIndex])
 				{
 					PoolIndex = ArenaParams.PoolCount; // was in the gap
 				}
@@ -382,7 +384,7 @@ class CORE_API FMallocBinnedArena final : public FMalloc
 	{
 		return PoolBaseVMPtr[InPoolIndex];
 	}
-	FORCEINLINE uint32 PoolIndexFromPtrChecked(const void* Ptr)
+	FORCEINLINE uint64 PoolIndexFromPtrChecked(const void* Ptr)
 	{
 		uint64 Result = PoolIndexFromPtr(Ptr);
 		check(Result < ArenaParams.PoolCount);
@@ -689,4 +691,6 @@ public:
 	TArray<void*> MallocedPointers;
 };
 
-#endif
\ No newline at end of file
+#endif
+
+PRAGMA_ENABLE_UNSAFE_TYPECAST_WARNINGS
diff --git a/Engine/Source/Runtime/Core/Public/HAL/MallocBinnedGPU.h b/Engine/Source/Runtime/Core/Public/HAL/MallocBinnedGPU.h
--- a/Engine/Source/Runtime/Core/Public/HAL/MallocBinnedGPU.h
+++ b/Engine/Source/Runtime/Core/Public/HAL/MallocBinnedGPU.h
@@ -4,6 +4,8 @@
 
 #include "CoreTypes.h"
 
+PRAGMA_DISABLE_UNSAFE_TYPECAST_WARNINGS
+
 #if PLATFORM_64BITS && PLATFORM_HAS_FPlatformVirtualMemoryBlock
 #include "Misc/AssertionMacros.h"
 #include "Misc/ScopeLock.h"
@@ -372,26 +374,26 @@ class CORE_API FMallocBinnedGPU final : public FMalloc
 			{
 				PoolIndex = ArenaParams.PoolCount - 1;
 			}
-			if ((uint8*)Ptr < PoolBaseVMPtr[PoolIndex])
+			if ((uint8*)Ptr < PoolBaseVMPtr[(int32)PoolIndex])
 			{
 				do
 				{
 					PoolIndex--;
 					check(PoolIndex < ArenaParams.PoolCount);
-				} while ((uint8*)Ptr < PoolBaseVMPtr[PoolIndex]);
-				if ((uint8*)Ptr >= PoolBaseVMPtr[PoolIndex] + ArenaParams.MaxMemoryPerBlockSize)
+				} while ((uint8*)Ptr < PoolBaseVMPtr[(int32)PoolIndex]);
+				if ((uint8*)Ptr >= PoolBaseVMPtr[(int32)PoolIndex] + ArenaParams.MaxMemoryPerBlockSize)
 				{
 					PoolIndex = ArenaParams.PoolCount; // was in the gap
 				}
 			}
-			else if ((uint8*)Ptr >= PoolBaseVMPtr[PoolIndex] + ArenaParams.MaxMemoryPerBlockSize)
+			else if ((uint8*)Ptr >= PoolBaseVMPtr[(int32)PoolIndex] + ArenaParams.MaxMemoryPerBlockSize)
 			{
 				do
 				{
 					PoolIndex++;
 					check(PoolIndex < ArenaParams.PoolCount);
-				} while ((uint8*)Ptr >= PoolBaseVMPtr[PoolIndex] + ArenaParams.MaxMemoryPerBlockSize);
-				if ((uint8*)Ptr < PoolBaseVMPtr[PoolIndex])
+				} while ((uint8*)Ptr >= PoolBaseVMPtr[(int32)PoolIndex] + ArenaParams.MaxMemoryPerBlockSize);
+				if ((uint8*)Ptr < PoolBaseVMPtr[(int32)PoolIndex])
 				{
 					PoolIndex = ArenaParams.PoolCount; // was in the gap
 				}
@@ -404,7 +406,7 @@ class CORE_API FMallocBinnedGPU final : public FMalloc
 	{
 		return PoolBaseVMPtr[InPoolIndex];
 	}
-	FORCEINLINE uint32 PoolIndexFromPtrChecked(const void* Ptr)
+	FORCEINLINE uint64 PoolIndexFromPtrChecked(const void* Ptr)
 	{
 		uint64 Result = PoolIndexFromPtr(Ptr);
 		check(Result < ArenaParams.PoolCount);
@@ -661,4 +663,6 @@ public:
 	TArray<void*> MallocedPointers;
 };
 
-#endif
\ No newline at end of file
+#endif
+
+PRAGMA_ENABLE_UNSAFE_TYPECAST_WARNINGS
diff --git a/Engine/Source/Runtime/Core/Public/HAL/UnrealMemory.h b/Engine/Source/Runtime/Core/Public/HAL/UnrealMemory.h
--- a/Engine/Source/Runtime/Core/Public/HAL/UnrealMemory.h
+++ b/Engine/Source/Runtime/Core/Public/HAL/UnrealMemory.h
@@ -219,7 +219,7 @@ struct CORE_API FMemory
 	* This memory is allocated like a stack, it's never really freed
 	*/
 	static void RegisterPersistentAuxiliary(void* InMemory, SIZE_T InSize);
-	static void* MallocPersistentAuxiliary(SIZE_T InSize, SIZE_T InAlignment = 0);
+	static void* MallocPersistentAuxiliary(SIZE_T InSize, uint32 InAlignment = 0);
 	static void FreePersistentAuxiliary(void* InPtr);
 	static bool IsPersistentAuxiliaryActive();
 	static void DisablePersistentAuxiliary();
diff --git a/Engine/Source/Runtime/Core/Public/Math/BigInt.h b/Engine/Source/Runtime/Core/Public/Math/BigInt.h
--- a/Engine/Source/Runtime/Core/Public/Math/BigInt.h
+++ b/Engine/Source/Runtime/Core/Public/Math/BigInt.h
@@ -348,7 +348,7 @@ public:
 				Carry = Product >> BitsPerWord;
 				ResultBits[WordIndexA + WordIndexB] = (uint32)(Product & (Base - 1));
 			}
-			ResultBits[WordIndexA + NumWordsB] += Carry;
+			ResultBits[WordIndexA + NumWordsB] += (uint32)Carry;
 		}
 
 		for (int32 WordIndex = 0; WordIndex < NumWords; ++WordIndex)
@@ -1246,10 +1246,12 @@ TEncryptionInt FEncryption::ModularPow(TEncryptionInt Base, TEncryptionInt Expon
 
 /// @endcond
 
-template <class TYPE>
+template <class InDataType>
 struct FSignatureBase
 {
-	TYPE Data;
+	typedef InDataType DataType;
+
+	DataType Data;
 
 	FSignatureBase()
 	{
@@ -1298,7 +1300,6 @@ struct FEncryptedSignature : public FSignatureBase<TEncryptionInt>
 
 struct FDecryptedSignature : public FSignatureBase<uint32>
 {
-
 };
 
 namespace FEncryption
@@ -1310,7 +1311,7 @@ namespace FEncryption
 
 	static void DecryptSignature(const FEncryptedSignature& InEncryptedSignature, FDecryptedSignature& OutUnencryptedSignature, const FEncryptionKey& EncryptionKey)
 	{
-		OutUnencryptedSignature.Data = FEncryption::ModularPow(TEncryptionInt(InEncryptedSignature.Data), EncryptionKey.Exponent, EncryptionKey.Modulus).ToInt();
+		OutUnencryptedSignature.Data = (FDecryptedSignature::DataType)FEncryption::ModularPow(TEncryptionInt(InEncryptedSignature.Data), EncryptionKey.Exponent, EncryptionKey.Modulus).ToInt();
 	}
 }
 
diff --git a/Engine/Source/Runtime/Core/Public/Math/Float16.h b/Engine/Source/Runtime/Core/Public/Math/Float16.h
--- a/Engine/Source/Runtime/Core/Public/Math/Float16.h
+++ b/Engine/Source/Runtime/Core/Public/Math/Float16.h
@@ -205,7 +205,7 @@ FORCEINLINE float FFloat16::GetFloat() const
 		else
 		{
 			// Denormal.
-			uint32 MantissaShift = 10 - (uint32)FMath::TruncToInt(FMath::Log2(Mantissa));
+			uint32 MantissaShift = 10 - (uint32)FMath::TruncToInt(FMath::Log2((float)Mantissa));
 			Result.Components.Exponent = 127 - (15 - 1) - MantissaShift;
 			Result.Components.Mantissa = Mantissa << (MantissaShift + 23 - 10);
 		}
diff --git a/Engine/Source/Runtime/Core/Public/Math/IntRect.h b/Engine/Source/Runtime/Core/Public/Math/IntRect.h
--- a/Engine/Source/Runtime/Core/Public/Math/IntRect.h
+++ b/Engine/Source/Runtime/Core/Public/Math/IntRect.h
@@ -312,8 +312,8 @@ public:
 
 FORCEINLINE FIntRect FIntRect::Scale( float Fraction ) const
 {
-	FVector2D Min2D = FVector2D(Min.X, Min.Y) * Fraction;
-	FVector2D Max2D = FVector2D(Max.X, Max.Y) * Fraction;
+	FVector2D Min2D = FVector2D((float)Min.X, (float)Min.Y) * Fraction;
+	FVector2D Max2D = FVector2D((float)Max.X, (float)Max.Y) * Fraction;
 
 	return FIntRect(FMath::FloorToInt(Min2D.X), FMath::FloorToInt(Min2D.Y), FMath::CeilToInt(Max2D.X), FMath::CeilToInt(Max2D.Y));
 }
diff --git a/Engine/Source/Runtime/Core/Public/Math/OrthoMatrix.h b/Engine/Source/Runtime/Core/Public/Math/OrthoMatrix.h
--- a/Engine/Source/Runtime/Core/Public/Math/OrthoMatrix.h
+++ b/Engine/Source/Runtime/Core/Public/Math/OrthoMatrix.h
@@ -45,6 +45,6 @@ FORCEINLINE FReversedZOrthoMatrix::FReversedZOrthoMatrix(float Width,float Heigh
 		FPlane((Width)? (1.0f / Width) : 1.0f,	0.0f,								0.0f,					0.0f),
 		FPlane(0.0f,							(Height)? (1.0f / Height) : 1.f,	0.0f,					0.0f),
 		FPlane(0.0f,							0.0f,								-ZScale,				0.0f),
-		FPlane(0.0f,							0.0f,								1.0 - ZOffset * ZScale,	1.0f)
+		FPlane(0.0f,							0.0f,								1.0f - ZOffset * ZScale,	1.0f)
 	)
 { }
diff --git a/Engine/Source/Runtime/Core/Public/Math/Plane.h b/Engine/Source/Runtime/Core/Public/Math/Plane.h
--- a/Engine/Source/Runtime/Core/Public/Math/Plane.h
+++ b/Engine/Source/Runtime/Core/Public/Math/Plane.h
@@ -271,10 +271,10 @@ public:
 		}
 		else
 		{
-			int16 iX(FMath::RoundToInt(X));
-			int16 iY(FMath::RoundToInt(Y));
-			int16 iZ(FMath::RoundToInt(Z));
-			int16 iW(FMath::RoundToInt(W));
+			int16 iX((int16)FMath::RoundToInt(X));
+			int16 iY((int16)FMath::RoundToInt(Y));
+			int16 iZ((int16)FMath::RoundToInt(Z));
+			int16 iW((int16)FMath::RoundToInt(W));
 			Ar << iX << iY << iZ << iW;
 		}
 		bOutSuccess = true;
diff --git a/Engine/Source/Runtime/Core/Public/Math/RangeBound.h b/Engine/Source/Runtime/Core/Public/Math/RangeBound.h
--- a/Engine/Source/Runtime/Core/Public/Math/RangeBound.h
+++ b/Engine/Source/Runtime/Core/Public/Math/RangeBound.h
@@ -353,7 +353,7 @@ private:
 			: Super(Other) \
 		{ } \
 		 \
-		Name(const int64& InValue) \
+		Name(const ElementType InValue) \
 			: Super(InValue) \
 		{ } \
 		 \
diff --git a/Engine/Source/Runtime/Core/Public/Math/UnitConversion.inl b/Engine/Source/Runtime/Core/Public/Math/UnitConversion.inl
--- a/Engine/Source/Runtime/Core/Public/Math/UnitConversion.inl
+++ b/Engine/Source/Runtime/Core/Public/Math/UnitConversion.inl
@@ -93,15 +93,15 @@ T FUnitConversion::Convert(T InValue, EUnit From, EUnit To)
 			// Put it into kelvin
 			switch (From)
 			{
-				case EUnit::Celsius:			NewValue = NewValue + 273.15;					break;
-				case EUnit::Farenheit:			NewValue = (NewValue + 459.67) * 5.f/9.f;		break;
+				case EUnit::Celsius:			NewValue = NewValue + 273.15f;					break;
+				case EUnit::Farenheit:			NewValue = (NewValue + 459.67f) * 5.f/9.f;		break;
 				default: 																		break;
 			}
 			// And out again
 			switch (To)
 			{
-				case EUnit::Celsius:			return NewValue - 273.15;
-				case EUnit::Farenheit:			return NewValue * 9.f/5.f - 459.67;
+				case EUnit::Celsius:			return NewValue - 273.15f;
+				case EUnit::Farenheit:			return NewValue * 9.f/5.f - 459.67f;
 				default: 						return NewValue;
 			}
 		}
diff --git a/Engine/Source/Runtime/Core/Public/Math/UnrealMathSSE.h b/Engine/Source/Runtime/Core/Public/Math/UnrealMathSSE.h
--- a/Engine/Source/Runtime/Core/Public/Math/UnrealMathSSE.h
+++ b/Engine/Source/Runtime/Core/Public/Math/UnrealMathSSE.h
@@ -1262,8 +1262,8 @@ FORCEINLINE VectorRegister VectorLog2(const VectorRegister& X)
 namespace VectorSinConstantsSSE
 {
 	static const float p = 0.225f;
-	static const float a = (16 * sqrt(p));
-	static const float b = ((1 - p) / sqrt(p));
+	static const float a = (16 * sqrtf(p));
+	static const float b = ((1 - p) / sqrtf(p));
 	static const VectorRegister A = MakeVectorRegister(a, a, a, a);
 	static const VectorRegister B = MakeVectorRegister(b, b, b, b);
 }
diff --git a/Engine/Source/Runtime/Core/Public/Math/UnrealMathUtility.h b/Engine/Source/Runtime/Core/Public/Math/UnrealMathUtility.h
--- a/Engine/Source/Runtime/Core/Public/Math/UnrealMathUtility.h
+++ b/Engine/Source/Runtime/Core/Public/Math/UnrealMathUtility.h
@@ -295,7 +295,7 @@ struct FMath : public FPlatformMath
 		if( Grid==0.f )	return Location;
 		else			
 		{
-			return FloorToFloat((Location + 0.5*Grid)/Grid)*Grid;
+			return FloorToFloat((Location + 0.5f*Grid)/Grid)*Grid;
 		}
 	}
 
@@ -988,7 +988,7 @@ struct FMath : public FPlatformMath
 	 */
 	static float MakePulsatingValue( const double InCurrentTime, const float InPulsesPerSecond, const float InPhase = 0.0f )
 	{
-		return 0.5f + 0.5f * FMath::Sin( ( ( 0.25f + InPhase ) * PI * 2.0 ) + ( InCurrentTime * PI * 2.0 ) * InPulsesPerSecond );
+		return 0.5f + 0.5f * FMath::Sin( ( ( 0.25f + InPhase ) * (float)PI * 2.0f ) + ( (float)InCurrentTime * (float)PI * 2.0f ) * InPulsesPerSecond );
 	}
 
 	// Geometry intersection 
@@ -1547,7 +1547,7 @@ struct FMath : public FPlatformMath
 		check(Ret >= 0);
 		check(Ret <= 255);
 
-		return Ret;
+		return (uint8)Ret;
 	}
 	
 	// @param x assumed to be in this range: -1..1
diff --git a/Engine/Source/Runtime/Core/Public/Math/Vector.h b/Engine/Source/Runtime/Core/Public/Math/Vector.h
--- a/Engine/Source/Runtime/Core/Public/Math/Vector.h
+++ b/Engine/Source/Runtime/Core/Public/Math/Vector.h
@@ -1253,13 +1253,13 @@ FORCEINLINE FVector::FVector(const FLinearColor& InColor)
 }
 
 FORCEINLINE FVector::FVector(FIntVector InVector)
-	: X(InVector.X), Y(InVector.Y), Z(InVector.Z)
+	: X((float)InVector.X), Y((float)InVector.Y), Z((float)InVector.Z)
 {
 	DiagnosticCheckNaN();
 }
 
 FORCEINLINE FVector::FVector(FIntPoint A)
-	: X(A.X), Y(A.Y), Z(0.f)
+	: X((float)A.X), Y((float)A.Y), Z(0.f)
 {
 	DiagnosticCheckNaN();
 }
diff --git a/Engine/Source/Runtime/Core/Public/Math/Vector2D.h b/Engine/Source/Runtime/Core/Public/Math/Vector2D.h
--- a/Engine/Source/Runtime/Core/Public/Math/Vector2D.h
+++ b/Engine/Source/Runtime/Core/Public/Math/Vector2D.h
@@ -961,7 +961,7 @@ FORCEINLINE FIntPoint FVector2D::IntPoint() const
 
 FORCEINLINE FVector2D FVector2D::RoundToVector() const
 {
-	return FVector2D(FMath::RoundToInt(X), FMath::RoundToInt(Y));
+	return FVector2D(FMath::RoundToFloat(X), FMath::RoundToFloat(Y));
 }
 
 FORCEINLINE FVector2D FVector2D::ClampAxes(float MinAxisVal, float MaxAxisVal) const
diff --git a/Engine/Source/Runtime/Core/Public/Math/Vector4.h b/Engine/Source/Runtime/Core/Public/Math/Vector4.h
--- a/Engine/Source/Runtime/Core/Public/Math/Vector4.h
+++ b/Engine/Source/Runtime/Core/Public/Math/Vector4.h
@@ -501,10 +501,10 @@ FORCEINLINE FVector4::FVector4(FVector2D InXY, FVector2D InZW)
 }
 
 FORCEINLINE FVector4::FVector4(const FIntVector4& InVector)
-	: X(InVector.X)
-	, Y(InVector.Y)
-	, Z(InVector.Z)
-	, W(InVector.W)
+	: X((float)InVector.X)
+	, Y((float)InVector.Y)
+	, Z((float)InVector.Z)
+	, W((float)InVector.W)
 {
 }
 
diff --git a/Engine/Source/Runtime/Core/Public/Misc/Attribute.h b/Engine/Source/Runtime/Core/Public/Misc/Attribute.h
--- a/Engine/Source/Runtime/Core/Public/Misc/Attribute.h
+++ b/Engine/Source/Runtime/Core/Public/Misc/Attribute.h
@@ -40,7 +40,7 @@ public:
 	 */
 	template< typename OtherType >
 	TAttribute( const OtherType& InInitialValue )
-		: Value( InInitialValue )		
+		: Value( (ObjectType)InInitialValue )
 		, bIsSet(true)
 		, Getter()
 	{ }
diff --git a/Engine/Source/Runtime/Core/Public/Misc/AutomationTest.h b/Engine/Source/Runtime/Core/Public/Misc/AutomationTest.h
--- a/Engine/Source/Runtime/Core/Public/Misc/AutomationTest.h
+++ b/Engine/Source/Runtime/Core/Public/Misc/AutomationTest.h
@@ -118,7 +118,7 @@ public:
 	/** Constructor */
 	FAutomationTestExecutionInfo() 
 		: bSuccessful( false )
-		, Duration(0.0f)
+		, Duration(0.0)
 		, Errors(0)
 		, Warnings(0)
 	{}
@@ -175,7 +175,7 @@ public:
 	TArray<FString> AnalyticsItems;
 
 	/** Time to complete the task */
-	float Duration;
+	double Duration;
 
 private:
 	/** Any errors that occurred during execution */
@@ -3258,7 +3258,7 @@ public:
 	{
 		if ( !Callback() )
 		{
-			float NewTime = FPlatformTime::Seconds();
+			const double NewTime = FPlatformTime::Seconds();
 			if ( NewTime - StartTime >= Timeout )
 			{
 				TimeoutCallback();
diff --git a/Engine/Source/Runtime/Core/Public/Misc/Char.h b/Engine/Source/Runtime/Core/Public/Misc/Char.h
--- a/Engine/Source/Runtime/Core/Public/Misc/Char.h
+++ b/Engine/Source/Runtime/Core/Public/Misc/Char.h
@@ -78,7 +78,7 @@ struct TChar : TCharBase<CharType, sizeof(CharType)>
 	*/
 	static CharType ToUpper(CharType Char)
 	{
-		return ToUnsigned(Char) - ((uint32(Char) - 'a' < 26u) << 5);
+		return (CharType)(ToUnsigned(Char) - ((uint32(Char) - 'a' < 26u) << 5));
 	}
 
 	/**
@@ -86,7 +86,7 @@ struct TChar : TCharBase<CharType, sizeof(CharType)>
 	*/
 	static CharType ToLower(CharType Char)
 	{
-		return ToUnsigned(Char) + ((uint32(Char) - 'A' < 26u) << 5);
+		return (CharType)(ToUnsigned(Char) + ((uint32(Char) - 'A' < 26u) << 5));
 	}
 
 	static bool IsUpper(CharType Char);
diff --git a/Engine/Source/Runtime/Core/Public/Misc/CompressedGrowableBuffer.h b/Engine/Source/Runtime/Core/Public/Misc/CompressedGrowableBuffer.h
--- a/Engine/Source/Runtime/Core/Public/Misc/CompressedGrowableBuffer.h
+++ b/Engine/Source/Runtime/Core/Public/Misc/CompressedGrowableBuffer.h
@@ -77,7 +77,7 @@ public:
 	 *
 	 * @return number of bytes allocated by this buffer
 	 */
-	uint32 GetAllocatedSize() const
+	SIZE_T GetAllocatedSize() const
 	{
 		return CompressedBuffer.GetAllocatedSize()
 			+ PendingCompressionBuffer.GetAllocatedSize()
diff --git a/Engine/Source/Runtime/Core/Public/Misc/ExpressionParserTypes.h b/Engine/Source/Runtime/Core/Public/Misc/ExpressionParserTypes.h
--- a/Engine/Source/Runtime/Core/Public/Misc/ExpressionParserTypes.h
+++ b/Engine/Source/Runtime/Core/Public/Misc/ExpressionParserTypes.h
@@ -41,7 +41,7 @@ public:
 	FStringToken() : TokenStart(nullptr), TokenEnd(nullptr), LineNumber(0), CharacterIndex(0) {}
 
 	/** Get the string representation of this token */
-	FString GetString() const { return FString(TokenEnd - TokenStart, TokenStart); }
+	FString GetString() const { return FString((int32)(TokenEnd - TokenStart), TokenStart); }
 
 	/** Check if this token is valid */
 	bool IsValid() const { return TokenEnd != TokenStart; }
diff --git a/Engine/Source/Runtime/Core/Public/Misc/FrameRate.h b/Engine/Source/Runtime/Core/Public/Misc/FrameRate.h
--- a/Engine/Source/Runtime/Core/Public/Misc/FrameRate.h
+++ b/Engine/Source/Runtime/Core/Public/Misc/FrameRate.h
@@ -173,7 +173,7 @@ inline double FFrameRate::AsSeconds(FFrameTime FrameTime) const
 
 inline FFrameTime FFrameRate::AsFrameTime(double TimeInSeconds) const
 {
-	// @todo: sequencer-timecode: proper large number integer multiplication/division before coersion to float ?
+	// @todo: sequencer-timecode: proper large number integer multiplication/division before coercion to float ?
 	const double       TimeAsFrame = (TimeInSeconds * Numerator) / Denominator;
 	const FFrameNumber FrameNumber = static_cast<int32>(FMath::FloorToDouble(TimeAsFrame));
 
@@ -183,12 +183,13 @@ inline FFrameTime FFrameRate::AsFrameTime(double TimeInSeconds) const
 		SubFrame = FMath::Min(SubFrame, (double)FFrameTime::MaxSubframe);
 	}
 
-	return FFrameTime(FrameNumber, SubFrame);
+	//@TODO: FLOATPRECISION: FFrameTime needs a general once over for precision (RE: cast to ctor)
+	return FFrameTime(FrameNumber, (float)SubFrame);
 }
 
 inline FFrameNumber FFrameRate::AsFrameNumber(double TimeInSeconds) const
 {
-	// @todo: sequencer-timecode: proper large number integer multiplication/division before coersion to float ?
+	// @todo: sequencer-timecode: proper large number integer multiplication/division before coercion to float ?
 	const double       TimeAsFrame = (double(TimeInSeconds) * Numerator) / Denominator;
 	return static_cast<int32>(FMath::FloorToDouble(TimeAsFrame));
 }
@@ -279,7 +280,8 @@ inline FFrameTime ConvertFrameTime(FFrameTime SourceTime, FFrameRate SourceRate,
 		SubFrame = FMath::Min(SubFrame, 0.999999940);
 	}
 
-	return FFrameTime( (int32)IntegerPart, SubFrame);
+	//@TODO: FLOATPRECISION: FFrameTime needs a general once over for precision (RE: cast to ctor)
+	return FFrameTime( (int32)IntegerPart, (float)SubFrame);
 }
 
 inline FFrameTime FFrameRate::TransformTime(FFrameTime SourceTime, FFrameRate SourceRate, FFrameRate DestinationRate)
diff --git a/Engine/Source/Runtime/Core/Public/Misc/FrameTime.h b/Engine/Source/Runtime/Core/Public/Misc/FrameTime.h
--- a/Engine/Source/Runtime/Core/Public/Misc/FrameTime.h
+++ b/Engine/Source/Runtime/Core/Public/Misc/FrameTime.h
@@ -290,7 +290,7 @@ FORCEINLINE_DEBUGGABLE FFrameTime FFrameTime::FromDecimal(double InDecimalFrame)
 
 	// Ensure fractional parts above the highest sub frame float precision do not round to 0.0
 	double Fraction = InDecimalFrame - FMath::FloorToDouble(InDecimalFrame);
-	return FFrameTime(NewFrame, FMath::Clamp(Fraction, 0.0, (double)MaxSubframe));
+	return FFrameTime(NewFrame, FMath::Clamp((float)Fraction, 0.0f, MaxSubframe));
 }
 
 /** Convert a FFrameTime into a string */
diff --git a/Engine/Source/Runtime/Core/Public/Misc/MemStack.h b/Engine/Source/Runtime/Core/Public/Misc/MemStack.h
--- a/Engine/Source/Runtime/Core/Public/Misc/MemStack.h
+++ b/Engine/Source/Runtime/Core/Public/Misc/MemStack.h
@@ -304,6 +304,8 @@ public:
 		{
 			return Data;
 		}
+
+		//@TODO: FLOATPRECISION: Takes SIZE_T input but doesn't actually support it
 		void ResizeAllocation(SizeType PreviousNumElements, SizeType NumElements,SIZE_T NumBytesPerElement)
 		{
 			void* OldData = Data;
@@ -311,7 +313,7 @@ public:
 			{
 				// Allocate memory from the stack.
 				Data = (ElementType*)FMemStack::Get().PushBytes(
-					NumElements * NumBytesPerElement,
+					(int32)(NumElements * NumBytesPerElement),
 					FMath::Max(Alignment,(uint32)alignof(ElementType))
 					);
 
diff --git a/Engine/Source/Runtime/Core/Public/Misc/StringBuilder.h b/Engine/Source/Runtime/Core/Public/Misc/StringBuilder.h
--- a/Engine/Source/Runtime/Core/Public/Misc/StringBuilder.h
+++ b/Engine/Source/Runtime/Core/Public/Misc/StringBuilder.h
@@ -333,14 +333,16 @@ protected:
 		// for the existing string including NUL terminator
 
 		if ((CurPos + RequiredAdditionalCapacity) < End)
+		{
 			return;
+		}
 
 		Extend(RequiredAdditionalCapacity);
 	}
 
-	CORE_API void	Extend(int32 extraCapacity);
-	CORE_API void*	AllocBuffer(int32 byteCount);
-	CORE_API void	FreeBuffer(void* buffer, int32 byteCount);
+	CORE_API void	Extend(SIZE_T extraCapacity);
+	CORE_API void*	AllocBuffer(SIZE_T byteCount);
+	CORE_API void	FreeBuffer(void* buffer, SIZE_T byteCount);
 
 	C*			Base;
 	C* 			CurPos;
diff --git a/Engine/Source/Runtime/Core/Public/Misc/Timecode.h b/Engine/Source/Runtime/Core/Public/Misc/Timecode.h
--- a/Engine/Source/Runtime/Core/Public/Misc/Timecode.h
+++ b/Engine/Source/Runtime/Core/Public/Misc/Timecode.h
@@ -71,7 +71,7 @@ public:
 	 */
 	FFrameNumber ToFrameNumber(const FFrameRate& InFrameRate) const
 	{
-		const int32 NumberOfFramesInSecond = FMath::CeilToInt(InFrameRate.AsDecimal());
+		const int32 NumberOfFramesInSecond = FMath::CeilToInt((float)InFrameRate.AsDecimal());
 		const int32 NumberOfFramesInMinute = NumberOfFramesInSecond * 60;
 		const int32 NumberOfFramesInHour = NumberOfFramesInMinute * 60;
 
@@ -119,7 +119,7 @@ public:
 	 */
 	static FTimecode FromFrameNumber(const FFrameNumber& InFrameNumber, const FFrameRate& InFrameRate, bool InbDropFrame)
 	{
-		const int32 NumberOfFramesInSecond = FMath::CeilToInt(InFrameRate.AsDecimal());
+		const int32 NumberOfFramesInSecond = FMath::CeilToInt((float)InFrameRate.AsDecimal());
 		const int32 NumberOfFramesInMinute = NumberOfFramesInSecond * 60;
 		const int32 NumberOfFramesInHour = NumberOfFramesInMinute * 60;
 
@@ -139,10 +139,10 @@ public:
 			const int32 NumberOfTimecodesToDrop = NumberOfFramesInSecond <= 30 ? 2 : 4;
 
 			// At an ideal 30fps there would be 18,000 frames every 10 minutes, but at 29.97 there's only 17,982 frames.
-			const int32 NumTrueFramesPerTenMinutes = FMath::FloorToInt((60 * 10) * InFrameRate.AsDecimal());
+			const int32 NumTrueFramesPerTenMinutes = FMath::FloorToInt((float)((60 * 10) * InFrameRate.AsDecimal()));
 
 			// Calculate out how many times we've skipped dropping frames (ie: Minute 15 gives us a value of 1, as we've only didn't drop frames on the 10th minute)
-			const int32 NumTimesSkippedDroppingFrames = FMath::FloorToInt(FMath::Abs(InFrameNumber.Value) / (double)NumTrueFramesPerTenMinutes);
+			const int32 NumTimesSkippedDroppingFrames = FMath::FloorToInt(FMath::Abs(InFrameNumber.Value) / (float)NumTrueFramesPerTenMinutes);
 
 			// Now we can figure out how many frame (displays) have been skipped total; 9 times out of every 10 minutes
 			const int32 NumFramesSkippedTotal = NumTimesSkippedDroppingFrames * 9 * NumberOfTimecodesToDrop;
@@ -159,10 +159,10 @@ public:
 			{
 				// Each minute we slip a little bit more out of sync by a small amount, we just wait until we've accumulated enough error
 				// to skip a whole frame and can catch up.
-				const uint32 NumTrueFramesPerMinute = (uint32)FMath::FloorToInt(60 * InFrameRate.AsDecimal());
+				const uint32 NumTrueFramesPerMinute = (uint32)FMath::FloorToInt(60 * (float)InFrameRate.AsDecimal());
 
 				// Figure out which minute we are (0-9) to see how many to skip
-				int32 CurrentMinuteOfTen = FMath::FloorToInt((FrameInTrueFrames - NumberOfTimecodesToDrop) / (double)NumTrueFramesPerMinute);
+				int32 CurrentMinuteOfTen = FMath::FloorToInt((FrameInTrueFrames - NumberOfTimecodesToDrop) / (float)NumTrueFramesPerMinute);
 				int NumAddedFrames = NumFramesSkippedTotal + (NumberOfTimecodesToDrop * CurrentMinuteOfTen);
 				OffsetFrame += NumAddedFrames;
 			}
diff --git a/Engine/Source/Runtime/Core/Public/Misc/Timespan.h b/Engine/Source/Runtime/Core/Public/Misc/Timespan.h
--- a/Engine/Source/Runtime/Core/Public/Misc/Timespan.h
+++ b/Engine/Source/Runtime/Core/Public/Misc/Timespan.h
@@ -609,7 +609,7 @@ public:
 	 */
 	static FTimespan FromDays(double Days)
 	{
-		return FTimespan(FMath::FloorToDouble(Days * ETimespan::TicksPerDay + 0.5));
+		return FTimespan((int64)FMath::FloorToDouble(Days * ETimespan::TicksPerDay + 0.5));
 	}
 
 	/**
@@ -621,7 +621,7 @@ public:
 	 */
 	static FTimespan FromHours(double Hours)
 	{
-		return FTimespan(FMath::FloorToDouble(Hours * ETimespan::TicksPerHour + 0.5));
+		return FTimespan((int64)FMath::FloorToDouble(Hours * ETimespan::TicksPerHour + 0.5));
 	}
 
 	/**
@@ -633,7 +633,7 @@ public:
 	 */
 	static FTimespan FromMicroseconds(double Microseconds)
 	{
-		return FTimespan(FMath::FloorToDouble(Microseconds * ETimespan::TicksPerMicrosecond + 0.5));
+		return FTimespan((int64)FMath::FloorToDouble(Microseconds * ETimespan::TicksPerMicrosecond + 0.5));
 	}
 
 	/**
@@ -645,7 +645,7 @@ public:
 	 */
 	static FTimespan FromMilliseconds(double Milliseconds)
 	{
-		return FTimespan(FMath::FloorToDouble(Milliseconds * ETimespan::TicksPerMillisecond + 0.5));
+		return FTimespan((int64)FMath::FloorToDouble(Milliseconds * ETimespan::TicksPerMillisecond + 0.5));
 	}
 
 	/**
@@ -657,7 +657,7 @@ public:
 	 */
 	static FTimespan FromMinutes(double Minutes)
 	{
-		return FTimespan(FMath::FloorToDouble(Minutes * ETimespan::TicksPerMinute + 0.5));
+		return FTimespan((int64)FMath::FloorToDouble(Minutes * ETimespan::TicksPerMinute + 0.5));
 	}
 
 	/**
@@ -669,7 +669,7 @@ public:
 	 */
 	static FTimespan FromSeconds(double Seconds)
 	{
-		return FTimespan(FMath::FloorToDouble(Seconds * ETimespan::TicksPerSecond + 0.5));
+		return FTimespan((int64)FMath::FloorToDouble(Seconds * ETimespan::TicksPerSecond + 0.5));
 	}
 
 	/**
diff --git a/Engine/Source/Runtime/Core/Public/ProfilingDebugging/FormatArgsTrace.h b/Engine/Source/Runtime/Core/Public/ProfilingDebugging/FormatArgsTrace.h
--- a/Engine/Source/Runtime/Core/Public/ProfilingDebugging/FormatArgsTrace.h
+++ b/Engine/Source/Runtime/Core/Public/ProfilingDebugging/FormatArgsTrace.h
@@ -41,7 +41,7 @@ struct FFormatArgsTrace
 		uint8* PayloadBufferPtr = TypeCodesBufferPtr + FormatArgsCount;
 		EncodeArgumentsInternal(TypeCodesBufferPtr, PayloadBufferPtr, FormatArgs...);
 		check(PayloadBufferPtr - Buffer == FormatArgsSize);
-		return FormatArgsSize;
+		return (uint16)FormatArgsSize;
 	}
 
 private:
diff --git a/Engine/Source/Runtime/Core/Public/ProfilingDebugging/Histogram.h b/Engine/Source/Runtime/Core/Public/ProfilingDebugging/Histogram.h
--- a/Engine/Source/Runtime/Core/Public/ProfilingDebugging/Histogram.h
+++ b/Engine/Source/Runtime/Core/Public/ProfilingDebugging/Histogram.h
@@ -205,7 +205,7 @@ protected:
 		}
 
 		/** Constructor for a pre-seeded bin */
-		FBin(double MinInclusive, double MaxExclusive, int32 InSum, int32 InCount)
+		FBin(double MinInclusive, double MaxExclusive, double InSum, int32 InCount)  //@TODO: FLOATPRECISION: THIS CHANGES BEHAVIOR
 			: MinValue(MinInclusive)
 			, UpperBound(MaxExclusive)
 			, Sum(InSum)
diff --git a/Engine/Source/Runtime/Core/Public/Serialization/ArrayReader.h b/Engine/Source/Runtime/Core/Public/Serialization/ArrayReader.h
--- a/Engine/Source/Runtime/Core/Public/Serialization/ArrayReader.h
+++ b/Engine/Source/Runtime/Core/Public/Serialization/ArrayReader.h
@@ -37,7 +37,7 @@ public:
 			// Only serialize if we have the requested amount of data
 			if (Offset + Count <= Num())
 			{
-				FMemory::Memcpy(Data, &((*this)[Offset]), Count);
+				FMemory::Memcpy(Data, &((*this)[(int32)Offset]), Count);
 				Offset += Count;
 			}
 			else
diff --git a/Engine/Source/Runtime/Core/Public/Serialization/BitReader.h b/Engine/Source/Runtime/Core/Public/Serialization/BitReader.h
--- a/Engine/Source/Runtime/Core/Public/Serialization/BitReader.h
+++ b/Engine/Source/Runtime/Core/Public/Serialization/BitReader.h
@@ -14,6 +14,8 @@ CORE_API void appBitsCpy( uint8* Dest, int32 DestBit, uint8* Src, int32 SrcBit,
 	FBitReader.
 -----------------------------------------------------------------------------*/
 
+//@TODO: FLOATPRECISION: Public API assumes it can have > 2 GB of bits, but internally uses int32 based TArray for bytes and also has uint32 clamping on bit counts in various places
+
 //
 // Reads bitstreams.
 //
@@ -35,6 +37,7 @@ public:
 
 	FORCEINLINE_DEBUGGABLE void SerializeBits( void* Dest, int64 LengthBits )
 	{
+		//@TODO: FLOATPRECISION: This function/class pretends it is 64-bit aware, e.g., in the type of LengthBits and the Pos member, but it is not as appBitsCpy is only 32 bits, the inner Buffer is a 32 bit TArray, etc...
 		if ( IsError() || Pos+LengthBits > Num)
 		{
 			if (!IsError())
@@ -51,14 +54,14 @@ public:
 		if( LengthBits == 1 )
 		{
 			((uint8*)Dest)[0] = 0;
-			if( Buffer[Pos>>3] & Shift(Pos&7) )
+			if( Buffer[(int32)(Pos>>3)] & Shift(Pos&7) )
 				((uint8*)Dest)[0] |= 0x01;
 			Pos++;
 		}
 		else if (LengthBits != 0)
 		{
 			((uint8*)Dest)[((LengthBits+7)>>3) - 1] = 0;
-			appBitsCpy((uint8*)Dest, 0, Buffer.GetData(), Pos, LengthBits);
+			appBitsCpy((uint8*)Dest, 0, Buffer.GetData(), (int32)Pos, (int32)LengthBits);
 			Pos += LengthBits;
 		}
 	}
@@ -83,7 +86,7 @@ public:
 					break;
 				}
 
-				if (Buffer[LocalPos >> 3] & Shift(LocalPos & 7))
+				if (Buffer[(int32)(LocalPos >> 3)] & Shift(LocalPos & 7))
 				{
 					Value |= Mask;
 				}
@@ -121,7 +124,7 @@ public:
 			}
 			else
 			{
-				Bit = !!(Buffer[LocalPos>>3] & Shift(LocalPos&7));
+				Bit = !!(Buffer[(int32)(LocalPos>>3)] & Shift(LocalPos&7));
 				Pos++;
 			}
 		}
@@ -151,14 +154,14 @@ public:
 	FORCEINLINE_DEBUGGABLE uint8* GetDataPosChecked()
 	{
 		check(Pos % 8 == 0);
-		return &Buffer[Pos >> 3];
+		return &Buffer[(int32)(Pos >> 3)];
 	}
 
-	FORCEINLINE_DEBUGGABLE uint32 GetBytesLeft() const
+	FORCEINLINE_DEBUGGABLE int64 GetBytesLeft() const
 	{
 		return ((Num - Pos) + 7) >> 3;
 	}
-	FORCEINLINE_DEBUGGABLE uint32 GetBitsLeft() const
+	FORCEINLINE_DEBUGGABLE int64 GetBitsLeft() const
 	{
 		return (Num - Pos);
 	}
@@ -180,7 +183,7 @@ public:
 	}
 	FORCEINLINE_DEBUGGABLE void EatByteAlign()
 	{
-		int32 PrePos = Pos;
+		int64 PrePos = Pos;
 
 		// Skip over remaining bits in current byte
 		Pos = (Pos+7) & (~0x07);
@@ -197,7 +200,7 @@ public:
 	 *
 	 * @param LengthBits	The number of bits being read at the time of overflow
 	 */
-	void SetOverflowed(int32 LengthBits);
+	void SetOverflowed(int64 LengthBits);
 
 	void AppendDataFromChecked( FBitReader& Src );
 	void AppendDataFromChecked( uint8* Src, uint32 NumBits );
diff --git a/Engine/Source/Runtime/Core/Public/Serialization/BitWriter.h b/Engine/Source/Runtime/Core/Public/Serialization/BitWriter.h
--- a/Engine/Source/Runtime/Core/Public/Serialization/BitWriter.h
+++ b/Engine/Source/Runtime/Core/Public/Serialization/BitWriter.h
@@ -133,6 +133,8 @@ public:
 
 	FORCEINLINE bool AllowAppend(int64 LengthBits)
 	{
+		//@TODO: FLOATPRECISION: This class pretends it is 64-bit aware, e.g., in the type of LengthBits and the Num/Max members, but it is not as the inner Buffer is a 32 bit TArray, etc...
+
 		if (Num+LengthBits > Max)
 		{
 			if (bAllowResize)
@@ -143,8 +145,8 @@ public:
 				// in chunks).
 
 				Max = FMath::Max<int64>(Max<<1,Num+LengthBits);
-				int32 ByteMax = (Max+7)>>3;
-				Buffer.AddZeroed(ByteMax - Buffer.Num());
+				int64 ByteMax = (Max+7)>>3;
+				Buffer.AddZeroed((int32)(ByteMax - Buffer.Num()));
 				check((Max+7)>>3== Buffer.Num());
 				return true;
 			}
diff --git a/Engine/Source/Runtime/Core/Public/Serialization/MemoryReader.h b/Engine/Source/Runtime/Core/Public/Serialization/MemoryReader.h
--- a/Engine/Source/Runtime/Core/Public/Serialization/MemoryReader.h
+++ b/Engine/Source/Runtime/Core/Public/Serialization/MemoryReader.h
@@ -38,7 +38,7 @@ public:
 			// Only serialize if we have the requested amount of data
 			if (Offset + Num <= TotalSize())
 			{
-				FMemory::Memcpy( Data, &Bytes[Offset], Num );
+				FMemory::Memcpy( Data, &Bytes[(int32)Offset], Num );
 				Offset += Num;
 			}
 			else
@@ -96,7 +96,7 @@ public:
 			// Only serialize if we have the requested amount of data
 			if (Offset + Num <= TotalSize())
 			{
-				FMemory::Memcpy( Data, &Bytes[Offset], Num );
+				FMemory::Memcpy( Data, &Bytes[(int32)Offset], Num );
 				Offset += Num;
 			}
 			else
diff --git a/Engine/Source/Runtime/Core/Public/Serialization/MemoryWriter.h b/Engine/Source/Runtime/Core/Public/Serialization/MemoryWriter.h
--- a/Engine/Source/Runtime/Core/Public/Serialization/MemoryWriter.h
+++ b/Engine/Source/Runtime/Core/Public/Serialization/MemoryWriter.h
@@ -49,7 +49,7 @@ public:
 		
 		if( Num )
 		{
-			FMemory::Memcpy( &Bytes[Offset], Data, Num );
+			FMemory::Memcpy( &Bytes[(int32)Offset], Data, Num );
 			Offset+=Num;
 		}
 	}
diff --git a/Engine/Source/Runtime/Core/Public/Stats/Stats2.h b/Engine/Source/Runtime/Core/Public/Stats/Stats2.h
--- a/Engine/Source/Runtime/Core/Public/Stats/Stats2.h
+++ b/Engine/Source/Runtime/Core/Public/Stats/Stats2.h
@@ -171,7 +171,7 @@ struct TStatId
 #if CPUPROFILERTRACE_ENABLED
 	FORCEINLINE uint16 GetTraceCpuProfilerSpecId() const
 	{
-		return StatIdPtr->TraceCpuProfilerSpecId;
+		return (uint16)(StatIdPtr->TraceCpuProfilerSpecId);
 	}
 #endif
 
@@ -1298,7 +1298,7 @@ public:
 		if( Packet.ThreadType == EThreadType::Other )
 		{
 			FPlatformMisc::MemoryBarrier();
-			uint64 Frame = FStats::GameThreadStatsFrame;
+			int32 Frame = FStats::GameThreadStatsFrame;
 			const bool bFrameHasChanged = Frame > CurrentGameFrame;
 			if( bFrameHasChanged )
 			{
diff --git a/Engine/Source/Runtime/Core/Public/Stats/StatsFile.h b/Engine/Source/Runtime/Core/Public/Stats/StatsFile.h
--- a/Engine/Source/Runtime/Core/Public/Stats/StatsFile.h
+++ b/Engine/Source/Runtime/Core/Public/Stats/StatsFile.h
@@ -157,7 +157,7 @@ protected:
 		else
 		{
 			DestData.Reserve( EStatsFileConstants::MAX_COMPRESSED_SIZE );
-			int32 CompressedSize = DestData.GetAllocatedSize();
+			int32 CompressedSize = (int32)DestData.GetAllocatedSize();
 
 			const bool bResult = FCompression::CompressMemory( NAME_Zlib, DestData.GetData(), CompressedSize, SrcData.GetData(), UncompressedSize );
 			check( bResult );
@@ -800,7 +800,7 @@ public:
 	{
 		// Read FNames.
 		Ar.Seek( Header.FNameTableOffset );
-		out_MetadataMessages.Reserve( Header.NumFNames );
+		out_MetadataMessages.Reserve( (int32)Header.NumFNames );
 		for( int32 Index = 0; Index < Header.NumFNames; Index++ )
 		{
 			ReadFName( Ar, false );
@@ -808,7 +808,7 @@ public:
 
 		// Read metadata messages.
 		Ar.Seek( Header.MetadataMessagesOffset );
-		out_MetadataMessages.Reserve( Header.NumMetadataMessages );
+		out_MetadataMessages.Reserve((int32)Header.NumMetadataMessages );
 		for( int32 Index = 0; Index < Header.NumMetadataMessages; Index++ )
 		{
 			new(out_MetadataMessages)FStatMessage( ReadMessage( Ar, false ) );
diff --git a/Engine/Source/Runtime/Core/Public/Templates/Sorting.h b/Engine/Source/Runtime/Core/Public/Templates/Sorting.h
--- a/Engine/Source/Runtime/Core/Public/Templates/Sorting.h
+++ b/Engine/Source/Runtime/Core/Public/Templates/Sorting.h
@@ -251,14 +251,14 @@ public:
 		while (AStart < BStart && BStart < Num)
 		{
 			// Index after the last value == First[BStart]
-			int32 NewAOffset = AlgoImpl::UpperBoundInternal(First + AStart, BStart - AStart, First[BStart], FIdentityFunctor(), Predicate);
+			int32 NewAOffset = (int32)AlgoImpl::UpperBoundInternal(First + AStart, BStart - AStart, First[BStart], FIdentityFunctor(), Predicate);
 			AStart += NewAOffset;
 
 			if (AStart >= BStart) // done
 				break;
 
 			// Index of the first value == First[AStart]
-			int32 NewBOffset = AlgoImpl::LowerBoundInternal(First + BStart, Num - BStart, First[AStart], FIdentityFunctor(), Predicate);
+			int32 NewBOffset = (int32)AlgoImpl::LowerBoundInternal(First + BStart, Num - BStart, First[AStart], FIdentityFunctor(), Predicate);
 			TRotationPolicy::Rotate(First, AStart, BStart + NewBOffset, NewBOffset);
 			BStart += NewBOffset;
 			AStart += NewBOffset + 1;
diff --git a/Engine/Source/Runtime/Core/Public/Templates/TypeHash.h b/Engine/Source/Runtime/Core/Public/Templates/TypeHash.h
--- a/Engine/Source/Runtime/Core/Public/Templates/TypeHash.h
+++ b/Engine/Source/Runtime/Core/Public/Templates/TypeHash.h
@@ -41,7 +41,7 @@ inline uint32 PointerHash(const void* Key,uint32 C = 0)
 	auto PtrInt = reinterpret_cast<UPTRINT>(Key);
 #endif
 
-	return HashCombine(PtrInt, C);
+	return HashCombine((uint32)PtrInt, C);
 }
 
 
diff --git a/Engine/Source/Runtime/Core/Public/Templates/UnrealTemplate.h b/Engine/Source/Runtime/Core/Public/Templates/UnrealTemplate.h
--- a/Engine/Source/Runtime/Core/Public/Templates/UnrealTemplate.h
+++ b/Engine/Source/Runtime/Core/Public/Templates/UnrealTemplate.h
@@ -90,9 +90,9 @@ CONSTEXPR const T* GetData(std::initializer_list<T> List)
 * Generically gets the number of items in a contiguous container
 */
 template<typename T, typename = typename TEnableIf<TIsContiguousContainer<T>::Value>::Type>
-auto GetNum(T&& Container) -> decltype((SIZE_T)Container.Num())
+auto GetNum(T&& Container) -> decltype(Container.Num())
 {
-	return (SIZE_T)Container.Num();
+	return Container.Num();
 }
 
 template <typename T, SIZE_T N> CONSTEXPR SIZE_T GetNum(      T (& Container)[N]) { return N; }
diff --git a/Engine/Source/Runtime/Core/Public/Windows/WindowsPlatformTime.h b/Engine/Source/Runtime/Core/Public/Windows/WindowsPlatformTime.h
--- a/Engine/Source/Runtime/Core/Public/Windows/WindowsPlatformTime.h
+++ b/Engine/Source/Runtime/Core/Public/Windows/WindowsPlatformTime.h
@@ -33,7 +33,7 @@ struct CORE_API FWindowsPlatformTime
 	{
 		Windows::LARGE_INTEGER Cycles;
 		Windows::QueryPerformanceCounter(&Cycles);
-		return Cycles.QuadPart;
+		return (uint32)Cycles.QuadPart;
 	}
 
 	static FORCEINLINE uint64 Cycles64()
diff --git a/Engine/Source/Runtime/Json/Public/Serialization/JsonReader.h b/Engine/Source/Runtime/Json/Public/Serialization/JsonReader.h
--- a/Engine/Source/Runtime/Json/Public/Serialization/JsonReader.h
+++ b/Engine/Source/Runtime/Json/Public/Serialization/JsonReader.h
@@ -578,7 +578,8 @@ private:
 								return false;
 							}
 
-							HexNum += HexDigit * FMath::Pow(16, Radix);
+							//@TODO: FLOATPRECISION: this is gross
+							HexNum += HexDigit * (int32)FMath::Pow(16, (float)Radix);
 						}
 
 						String += (FString::ElementType)HexNum;
diff --git a/Engine/Source/Runtime/Json/Public/Serialization/JsonSerializerMacros.h b/Engine/Source/Runtime/Json/Public/Serialization/JsonSerializerMacros.h
--- a/Engine/Source/Runtime/Json/Public/Serialization/JsonSerializerMacros.h
+++ b/Engine/Source/Runtime/Json/Public/Serialization/JsonSerializerMacros.h
@@ -634,7 +634,7 @@ public:
 	{
 		if (JsonObject->HasTypedField<EJson::Number>(Name))
 		{
-			Value = JsonObject->GetNumberField(Name);
+			Value = (float)JsonObject->GetNumberField(Name);
 		}
 	}
 	/**
@@ -743,7 +743,7 @@ public:
 			// Iterate all of the keys and their values
 			for (const TPair<FString, TSharedPtr<FJsonValue>>& Pair : JsonMap->Values)
 			{
-				const int32 Value = Pair.Value->AsNumber();
+				const int32 Value = (int32)Pair.Value->AsNumber();
 				Map.Add(Pair.Key, Value);
 			}
 		}
@@ -763,7 +763,7 @@ public:
 			// Iterate all of the keys and their values
 			for (const TPair<FString, TSharedPtr<FJsonValue>>& Pair : JsonMap->Values)
 			{
-				const int64 Value = Pair.Value->AsNumber();
+				const int64 Value = (int64)Pair.Value->AsNumber();
 				Map.Add(Pair.Key, Value);
 			}
 		}
diff --git a/Engine/Source/Runtime/TraceLog/Public/Trace/Detail/Writer.inl b/Engine/Source/Runtime/TraceLog/Public/Trace/Detail/Writer.inl
--- a/Engine/Source/Runtime/TraceLog/Public/Trace/Detail/Writer.inl
+++ b/Engine/Source/Runtime/TraceLog/Public/Trace/Detail/Writer.inl
@@ -48,7 +48,7 @@ private:
 
 ////////////////////////////////////////////////////////////////////////////////
 extern TRACELOG_API uint32 volatile	GLogSerial;
-TRACELOG_API FWriteBuffer*			Writer_NextBuffer(uint16);
+TRACELOG_API FWriteBuffer*			Writer_NextBuffer(int32);
 TRACELOG_API FWriteBuffer*			Writer_GetBuffer();
 
 ////////////////////////////////////////////////////////////////////////////////
@@ -77,7 +77,7 @@ inline FLogInstance Writer_BeginLog(uint16 EventUid, uint16 Size, bool bMaybeHas
 	using namespace Private;
 
 	FWriteBuffer* Buffer = Writer_GetBuffer();
-	uint32 AllocSize = Size + sizeof(FEventHeader) + int(bMaybeHasAux);
+	int32 AllocSize = Size + sizeof(FEventHeader) + int(bMaybeHasAux);
 	Buffer->Cursor += AllocSize;
 	if (UNLIKELY(Buffer->Cursor > (uint8*)Buffer))
 	{
