Misplaced ) and additional cast fixes
#rnx

diff --git a/Engine/Source/Runtime/Core/Private/HAL/MallocDoubleFreeFinder.cpp b/Engine/Source/Runtime/Core/Private/HAL/MallocDoubleFreeFinder.cpp
--- a/Engine/Source/Runtime/Core/Private/HAL/MallocDoubleFreeFinder.cpp
+++ b/Engine/Source/Runtime/Core/Private/HAL/MallocDoubleFreeFinder.cpp
@@ -117,7 +117,7 @@ void FMallocDoubleFreeFinder::Free(void* Ptr)
 	SIZE_T OldSize = 0;
 	UsedMalloc->GetAllocationSize(Ptr, OldSize);
 	UsedMalloc->Free(Ptr);
-	TrackFree(Ptr, OldSize, CallStackIndex);
+	TrackFree(Ptr, (uint32)OldSize, CallStackIndex);
 }
 
 void FMallocDoubleFreeFinder::TrackMalloc(void* Ptr, uint32 Size, int32 CallStackIndex)
diff --git a/Engine/Source/Runtime/Core/Private/Math/Color.cpp b/Engine/Source/Runtime/Core/Private/Math/Color.cpp
--- a/Engine/Source/Runtime/Core/Private/Math/Color.cpp
+++ b/Engine/Source/Runtime/Core/Private/Math/Color.cpp
@@ -185,7 +185,7 @@ FColor FColor::FromHex( const FString& HexString )
 		const int32 G = FParse::HexDigit(HexString[StartIndex++]);
 		const int32 B = FParse::HexDigit(HexString[StartIndex]);
 
-		return FColor((uint8)((R << 4) + R), (uint8)((G << 4) + G), (uint8)((B << 4) + B, 255));
+		return FColor((uint8)((R << 4) + R), (uint8)((G << 4) + G), (uint8)((B << 4) + B), 255);
 	}
 
 	if (HexString.Len() == 6 + StartIndex)
diff --git a/Engine/Source/Runtime/Core/Public/Math/BigInt.h b/Engine/Source/Runtime/Core/Public/Math/BigInt.h
--- a/Engine/Source/Runtime/Core/Public/Math/BigInt.h
+++ b/Engine/Source/Runtime/Core/Public/Math/BigInt.h
@@ -1265,7 +1265,7 @@ struct FSignatureBase
 
 	static int64 Size()
 	{
-		return sizeof(TYPE);
+		return sizeof(InDataType);
 	}
 
 	bool IsValid() const
@@ -1300,6 +1300,7 @@ struct FEncryptedSignature : public FSignatureBase<TEncryptionInt>
 
 struct FDecryptedSignature : public FSignatureBase<uint32>
 {
+
 };
 
 namespace FEncryption
