Core: Catching warnings in another file that is not compiled in by default (MallocLeakDetection.cpp)
#rnx

diff --git a/Engine/Source/Runtime/Core/Private/HAL/MallocLeakDetection.cpp b/Engine/Source/Runtime/Core/Private/HAL/MallocLeakDetection.cpp
--- a/Engine/Source/Runtime/Core/Private/HAL/MallocLeakDetection.cpp
+++ b/Engine/Source/Runtime/Core/Private/HAL/MallocLeakDetection.cpp
@@ -449,8 +449,8 @@ void FMallocLeakDetection::Malloc(void* Ptr, SIZE_T Size)
 				bRecursive = true;
 				FCallstackTrack Callstack;
 				FPlatformStackWalk::CaptureStackBackTrace(Callstack.CallStack, FCallstackTrack::Depth);
-				Callstack.FirstFrame = GFrameCounter;
-				Callstack.LastFrame = GFrameCounter;
+				Callstack.FirstFrame = (uint32)GFrameCounter;
+				Callstack.LastFrame = (uint32)GFrameCounter;
 				Callstack.Size = Size;
 				AddCallstack(Callstack);
 				OpenPointers.Add(Ptr, Callstack);
