Additional C4244 warning fixes for changes in UE main
#rnx

diff --git a/Engine/Source/Runtime/Core/Private/Misc/PathViews.cpp b/Engine/Source/Runtime/Core/Private/Misc/PathViews.cpp
--- a/Engine/Source/Runtime/Core/Private/Misc/PathViews.cpp
+++ b/Engine/Source/Runtime/Core/Private/Misc/PathViews.cpp
@@ -15,7 +15,7 @@ FStringView FPathViews::GetCleanFilename(const FStringView& InPath)
 {
 	if (const TCHAR* StartPos = Algo::FindLastByPredicate(InPath, UE4PathViews_Private::IsSlashOrBackslash))
 	{
-		return InPath.RightChop(StartPos - InPath.GetData() + 1);
+		return InPath.RightChop(UE_PTRDIFF_TO_INT32(StartPos - InPath.GetData() + 1));
 	}
 	return InPath;
 }
@@ -40,7 +40,7 @@ FStringView FPathViews::GetPath(const FStringView& InPath)
 {
 	if (const TCHAR* EndPos = Algo::FindLastByPredicate(InPath, UE4PathViews_Private::IsSlashOrBackslash))
 	{
-		return InPath.Left(EndPos - InPath.GetData());
+		return InPath.Left((FStringView::SizeType)(EndPos - InPath.GetData()));
 	}
 	return FStringView();
 }
@@ -50,7 +50,7 @@ FStringView FPathViews::GetExtension(const FStringView& InPath, bool bIncludeDot
 	if (const TCHAR* Dot = Algo::FindLast(GetCleanFilename(InPath), TEXT('.')))
 	{
 		const TCHAR* Extension = bIncludeDot ? Dot : Dot + 1;
-		return FStringView(Extension, InPath.GetData() + InPath.Len() - Extension);
+		return FStringView(Extension, (FStringView::SizeType)(InPath.GetData() + InPath.Len() - Extension));
 	}
 	return FStringView();
 }
@@ -60,7 +60,7 @@ FStringView FPathViews::GetPathLeaf(const FStringView& InPath)
 	if (const TCHAR* EndPos = Algo::FindLastByPredicate(InPath, UE4PathViews_Private::IsNotSlashOrBackslash))
 	{
 		++EndPos;
-		return GetCleanFilename(InPath.Left(EndPos - InPath.GetData()));
+		return GetCleanFilename(InPath.Left((FStringView::SizeType)(EndPos - InPath.GetData())));
 	}
 	return FStringView();
 }
@@ -69,7 +69,7 @@ void FPathViews::Split(const FStringView& InPath, FStringView& OutPath, FStringV
 {
 	const FStringView CleanName = GetCleanFilename(InPath);
 	const TCHAR* DotPos = Algo::FindLast(CleanName, TEXT('.'));
-	const FStringView::SizeType NameLen = DotPos ? DotPos - CleanName.GetData() : CleanName.Len();
+	const FStringView::SizeType NameLen = DotPos ? (FStringView::SizeType)(DotPos - CleanName.GetData()) : CleanName.Len();
 	OutPath = InPath.LeftChop(CleanName.Len() + 1);
 	OutName = CleanName.Left(NameLen);
 	OutExt = CleanName.RightChop(NameLen + 1);
diff --git a/Engine/Source/Runtime/Core/Private/Tests/Algo/UniqueTest.cpp b/Engine/Source/Runtime/Core/Private/Tests/Algo/UniqueTest.cpp
--- a/Engine/Source/Runtime/Core/Private/Tests/Algo/UniqueTest.cpp
+++ b/Engine/Source/Runtime/Core/Private/Tests/Algo/UniqueTest.cpp
@@ -40,7 +40,7 @@ bool FUniqueTest::RunTest(const FString& Parameters)
 	}
 	{
 		int32 Array[] = {1};
-		int32 NewSize = Unique(Array);
+		int32 NewSize = (int32)Unique(Array);
 		TestEqual(TEXT("`Unique` must support C arrays"), NewSize, 1);
 	}
 	{
diff --git a/Engine/Source/Runtime/Core/Public/Algo/Unique.h b/Engine/Source/Runtime/Core/Public/Algo/Unique.h
--- a/Engine/Source/Runtime/Core/Public/Algo/Unique.h
+++ b/Engine/Source/Runtime/Core/Public/Algo/Unique.h
@@ -29,7 +29,7 @@ namespace AlgoImpl
 			}
 		}
 
-		return Result + 1 - Array;
+		return static_cast<SizeType>(Result + 1 - Array);
 	}
 }
 
diff --git a/Engine/Source/Runtime/Core/Public/Misc/StringBuilder.h b/Engine/Source/Runtime/Core/Public/Misc/StringBuilder.h
--- a/Engine/Source/Runtime/Core/Public/Misc/StringBuilder.h
+++ b/Engine/Source/Runtime/Core/Public/Misc/StringBuilder.h
@@ -96,7 +96,7 @@ public:
 	inline const C* ToString() const	{ EnsureNulTerminated(); return Base; }
 	inline const C* operator*() const	{ EnsureNulTerminated(); return Base; }
 
-	inline operator ViewType() const	{ return ViewType(Base, CurPos - Base); }
+	inline operator ViewType() const	{ return ViewType(Base, (ViewType::SizeType)(CurPos - Base)); }
 
 	inline const C	LastChar() const	{ return *(CurPos - 1); }
 
