Fixed some uint64->uint32 warnings in IoDispatcher
#rnx

diff --git a/Engine/Source/Runtime/Core/Private/IO/IoDispatcherFileBackend.cpp b/Engine/Source/Runtime/Core/Private/IO/IoDispatcherFileBackend.cpp
--- a/Engine/Source/Runtime/Core/Private/IO/IoDispatcherFileBackend.cpp
+++ b/Engine/Source/Runtime/Core/Private/IO/IoDispatcherFileBackend.cpp
@@ -187,9 +187,9 @@ EIoStoreResolveResult FFileIoStore::Resolve(FIoRequestImpl* Request)
 					ResolvedRequest.Request->IoBuffer = FIoBuffer(FIoBuffer::Wrap, TargetVa, ResolvedRequest.ResolvedSize);
 				}
 				PlatformImpl.BeginReadsForRequest(ResolvedRequest);
-				uint64 RequestBeginBlockIndex = ResolvedRequest.ResolvedOffset / CacheBlockSize;
-				uint64 RequestEndBlockIndex = (ResolvedRequest.ResolvedOffset + ResolvedRequest.ResolvedSize - 1) / CacheBlockSize + 1;
-				uint64 BlockCount = RequestEndBlockIndex - RequestBeginBlockIndex;
+				const uint32 RequestBeginBlockIndex = (uint32)(ResolvedRequest.ResolvedOffset / CacheBlockSize);
+				const uint32 RequestEndBlockIndex = (uint32)((ResolvedRequest.ResolvedOffset + ResolvedRequest.ResolvedSize - 1) / CacheBlockSize + 1);
+				const uint32 BlockCount = RequestEndBlockIndex - RequestBeginBlockIndex;
 				check(BlockCount > 0);
 				ReadBlockCached(RequestBeginBlockIndex, ResolvedRequest);
 				if (BlockCount > 1)
@@ -305,7 +305,7 @@ void FFileIoStore::ReadBlockCached(uint32 BlockIndex, const FFileIoStoreResolved
 	FFileIoStoreCacheBlockKey Key;
 	Key.FileHandle = ResolvedRequest.ResolvedFileHandle;
 	Key.BlockIndex = BlockIndex;
-	Key.Hash = HashCombine(ResolvedRequest.ResolvedFileHandle, BlockIndex);
+	Key.Hash = HashCombine(GetTypeHash(ResolvedRequest.ResolvedFileHandle), GetTypeHash(BlockIndex));
 	uint64 BlockOffset = uint64(BlockIndex) * uint64(CacheBlockSize);
 	FFileIoStoreReadBlock* CachedBlock = CachedBlocksMap.FindRef(Key);
 	if (!CachedBlock)
@@ -351,7 +351,8 @@ void FFileIoStore::ReadBlockCached(uint32 BlockIndex, const FFileIoStoreResolved
 		Scatter.Request = ResolvedRequest.Request;
 		Scatter.Dst = ResolvedRequest.Request->IoBuffer.Data() + BlockOffsetInRequest;
 		Scatter.Src = CachedBlock->Buffer + RequestStartOffsetInBlock;
-		Scatter.Size = RequestSizeInBlock;
+		check(RequestSizeInBlock <= TNumericLimits<uint32>::Max());
+		Scatter.Size = (uint32)RequestSizeInBlock;
 	}
 }
 
diff --git a/Engine/Source/Runtime/Core/Public/IO/IoDispatcher.h b/Engine/Source/Runtime/Core/Public/IO/IoDispatcher.h
--- a/Engine/Source/Runtime/Core/Public/IO/IoDispatcher.h
+++ b/Engine/Source/Runtime/Core/Public/IO/IoDispatcher.h
@@ -695,7 +695,7 @@ public:
 		return RequestedOffset;
 	}
 
-	uint32 GetSize() const
+	uint64 GetSize() const
 	{
 		return RequestedSize;
 	}
