Another implicit typecast warning fix
#rnx

diff --git a/Engine/Source/Runtime/Core/Private/Misc/StringBuilder.cpp b/Engine/Source/Runtime/Core/Private/Misc/StringBuilder.cpp
--- a/Engine/Source/Runtime/Core/Private/Misc/StringBuilder.cpp
+++ b/Engine/Source/Runtime/Core/Private/Misc/StringBuilder.cpp
@@ -76,7 +76,7 @@ TStringBuilderImpl<C>::AppendfImpl(BuilderType& Self, const C* Fmt, ...)
 	{
 		va_list ArgPack;
 		va_start(ArgPack, Fmt);
-		const int32 RemainingSize = Self.End - Self.CurPos;
+		const int32 RemainingSize = UE_PTRDIFF_TO_INT32(Self.End - Self.CurPos);
 		const int32 Result = TCString<C>::GetVarArgs(Self.CurPos, RemainingSize, Fmt, ArgPack);
 		va_end(ArgPack);
 
