Prevent signed/unsigned mismatch error in TArrayView ctor safety check() on 32 bit targets
#rnx

diff --git a/Engine/Source/Runtime/Core/Public/Containers/ArrayView.h b/Engine/Source/Runtime/Core/Public/Containers/ArrayView.h
--- a/Engine/Source/Runtime/Core/Public/Containers/ArrayView.h
+++ b/Engine/Source/Runtime/Core/Public/Containers/ArrayView.h
@@ -123,7 +123,7 @@ public:
 		: DataPtr(ArrayViewPrivate::GetDataHelper(Forward<OtherRangeType>(Other)))
 	{
 		const auto InCount = GetNum(Forward<OtherRangeType>(Other));
-		check((InCount >= 0) && (InCount <= TNumericLimits<int32>::Max()));
+		check((InCount >= 0) && ((sizeof(InCount) < sizeof(int32)) || (InCount <= static_cast<decltype(InCount)>(TNumericLimits<int32>::Max()))));
 		ArrayNum = (int32)InCount;
 	}
 
