Core: Reimplementing string builder's operator ViewType() in terms of Len()
#rnx

diff --git a/Engine/Source/Runtime/Core/Public/Misc/StringBuilder.h b/Engine/Source/Runtime/Core/Public/Misc/StringBuilder.h
--- a/Engine/Source/Runtime/Core/Public/Misc/StringBuilder.h
+++ b/Engine/Source/Runtime/Core/Public/Misc/StringBuilder.h
@@ -96,7 +96,7 @@ public:
 	inline const C* ToString() const	{ EnsureNulTerminated(); return Base; }
 	inline const C* operator*() const	{ EnsureNulTerminated(); return Base; }
 
-	inline operator ViewType() const	{ return ViewType(Base, (ViewType::SizeType)(CurPos - Base)); }
+	inline operator ViewType() const	{ return ViewType(Base, Len()); }
 
 	inline const C	LastChar() const	{ return *(CurPos - 1); }
 
