Added public method GetPakchunkLocation(), PrioritizePakchunk() to all platform chunk install and move GetChunkLocation and PrioritizeChunk() to protected.

diff --git a/Engine/Source/Runtime/AssetRegistry/Private/AssetRegistry.cpp b/Engine/Source/Runtime/AssetRegistry/Private/AssetRegistry.cpp
--- a/Engine/Source/Runtime/AssetRegistry/Private/AssetRegistry.cpp
+++ b/Engine/Source/Runtime/AssetRegistry/Private/AssetRegistry.cpp
@@ -1356,9 +1356,9 @@ EAssetAvailability::Type UAssetRegistryImpl::GetAssetAvailability(const FAssetDa
 	EChunkLocation::Type BestLocation = EChunkLocation::DoesNotExist;
 
 	// check all chunks to see which has the best locality
-	for (int32 ChunkId : AssetData.ChunkIDs)
+	for (int32 PakchunkId : AssetData.ChunkIDs)
 	{
-		EChunkLocation::Type ChunkLocation = ChunkInstall->GetChunkLocation(ChunkId);
+		EChunkLocation::Type ChunkLocation = ChunkInstall->GetPakchunkLocation(PakchunkId);
 
 		// if we find one in the best location, early out
 		if (ChunkLocation == EChunkLocation::BestLocation)
@@ -1400,9 +1400,9 @@ float UAssetRegistryImpl::GetAssetAvailabilityProgress(const FAssetData& AssetDa
 	float BestProgress = MAX_FLT;
 
 	// check all chunks to see which has the best time remaining
-	for (int32 ChunkId : AssetData.ChunkIDs)
+	for (int32 PakchunkID : AssetData.ChunkIDs)
 	{
-		float Progress = ChunkInstall->GetChunkProgress(ChunkId, ChunkReportType);
+		float Progress = ChunkInstall->GetChunkProgress(PakchunkID, ChunkReportType);
 
 		// need to flip percentage completes for the comparison
 		if (IsPercentageComplete)
@@ -1445,7 +1445,7 @@ void UAssetRegistryImpl::PrioritizeAssetInstall(const FAssetData& AssetData) con
 		return;
 	}
 
-	ChunkInstall->PrioritizeChunk(AssetData.ChunkIDs[0], EChunkPriority::Immediate);
+	ChunkInstall->PrioritizePakchunk(AssetData.ChunkIDs[0], EChunkPriority::Immediate);
 }
 
 bool UAssetRegistryImpl::AddPath(const FString& PathToAdd)
diff --git a/Engine/Source/Runtime/AssetRegistry/Public/AssetData.h b/Engine/Source/Runtime/AssetRegistry/Public/AssetData.h
--- a/Engine/Source/Runtime/AssetRegistry/Public/AssetData.h
+++ b/Engine/Source/Runtime/AssetRegistry/Public/AssetData.h
@@ -80,7 +80,7 @@ public:
 	FName AssetClass;
 	/** The map of values for properties that were marked AssetRegistrySearchable or added by GetAssetRegistryTags */
 	FAssetDataTagMapSharedView TagsAndValues;
-	/** The IDs of the chunks this asset is located in for streaming install.  Empty if not assigned to a chunk */
+	/** The IDs of the pakchunks this asset is located in for streaming install.  Empty if not assigned to a chunk */
 	TArray<int32> ChunkIDs;
 	/** Asset package flags */
 	uint32 PackageFlags;
diff --git a/Engine/Source/Runtime/Core/Public/GenericPlatform/GenericPlatformChunkInstall.h b/Engine/Source/Runtime/Core/Public/GenericPlatform/GenericPlatformChunkInstall.h
--- a/Engine/Source/Runtime/Core/Public/GenericPlatform/GenericPlatformChunkInstall.h
+++ b/Engine/Source/Runtime/Core/Public/GenericPlatform/GenericPlatformChunkInstall.h
@@ -95,6 +95,16 @@ struct FCustomChunk
 	{}
 };
 
+struct FCustomChunkMapping
+{
+	FString Pattern;
+	uint32	ChunkID;
+
+	FCustomChunkMapping(FString InPattern, uint32 InChunkID) :
+		Pattern(InPattern), ChunkID(InChunkID)
+	{}
+};
+
 /**
 * Interface for platform specific chunk based install
 **/
@@ -106,11 +116,11 @@ public:
 	virtual ~IPlatformChunkInstall() {}
 
 	/**
-	 * Get the current location of a chunk.
-	 * @param ChunkID		The id of the chunk to check.
+	 * Get the current location of a chunk with pakchunk id.
+	 * @param ChunkID		The id of the pak chunk.
 	 * @return				Enum specifying whether the chunk is available to use, waiting to install, or does not exist.
 	 **/
-	virtual EChunkLocation::Type GetChunkLocation( uint32 ChunkID ) = 0;
+	virtual EChunkLocation::Type GetPakchunkLocation( uint32 PakchunkID ) = 0;
 
 	/** 
 	 * Check if a given reporting type is supported.
@@ -121,11 +131,11 @@ public:
 
 	/**
 	 * Get the current install progress of a chunk.  Let the user specify report type for platforms that support more than one.
-	 * @param ChunkID		The id of the chunk to check.
+	 * @param PakchunkID		The id of the chunk to check.
 	 * @param ReportType	The type of progress report you want.
 	 * @return				A value whose meaning is dependent on the ReportType param.
 	 **/
-	virtual float GetChunkProgress( uint32 ChunkID, EChunkProgressReportingType::Type ReportType ) = 0;
+	virtual float GetChunkProgress( uint32 PakchunkID, EChunkProgressReportingType::Type ReportType ) = 0;
 
 	/**
 	 * Inquire about the priority of chunk installation vs. game IO.
@@ -141,11 +151,11 @@ public:
 	
 	/**
 	 * Hint to the installer that we would like to prioritize a specific chunk
-	 * @param ChunkID		The id of the chunk to prioritize.
+	 * @param PakchunkID	The id of the pakchunk to prioritize.
 	 * @param Priority		The priority for the chunk.
 	 * @return				false if the operation is not allowed or the chunk doesn't exist, otherwise true.
 	 **/
-	virtual bool PrioritizeChunk( uint32 ChunkID, EChunkPriority::Type Priority ) = 0;
+	virtual bool PrioritizePakchunk( uint32 PakchunkID, EChunkPriority::Type Priority ) = 0;
 
 	/**
 	 * For platforms that support emulation of the Chunk install.  Starts transfer of the next chunk.
@@ -205,6 +215,24 @@ public:
 	* @return				whether uninstallation task has been kicked
 	*/
 	virtual bool UninstallChunks(const TArray<FCustomChunk>& ChunkTagsID) = 0;
+
+protected:
+		/**
+		 * Get the current location of a chunk.
+		 * Pakchunk index and OS chunk id are not the same.  Call GetPakchunkLocation instead of calling from outside.
+		 * @param ChunkID		The id of the chunk to check.
+		 * @return				Enum specifying whether the chunk is available to use, waiting to install, or does not exist.
+		 **/
+		virtual EChunkLocation::Type GetChunkLocation(uint32 ChunkID) = 0;
+
+		/**
+		 * Hint to the installer that we would like to prioritize a specific chunk
+		 * @param ChunkID		The id of the chunk to prioritize.
+		 * @param Priority		The priority for the chunk.
+		 * @return				false if the operation is not allowed or the chunk doesn't exist, otherwise true.
+		 **/
+		virtual bool PrioritizeChunk(uint32 ChunkID, EChunkPriority::Type Priority) = 0;
+
 };
 
 PRAGMA_DISABLE_DEPRECATION_WARNINGS
@@ -214,9 +242,14 @@ PRAGMA_DISABLE_DEPRECATION_WARNINGS
 class CORE_API FGenericPlatformChunkInstall : public IPlatformChunkInstall
 {
 public:
-	virtual EChunkLocation::Type GetChunkLocation( uint32 ChunkID ) override
+	virtual EChunkLocation::Type GetPakchunkLocation( uint32 PakchunkID ) override
 	{
-		return EChunkLocation::LocalFast;
+		return GetChunkLocation(PakchunkID);
+	}
+
+	virtual bool PrioritizePakchunk(uint32 PakchunkID, EChunkPriority::Type Priority)
+	{
+		return PrioritizeChunk(PakchunkID, Priority);
 	}
 
 	virtual bool GetProgressReportingTypeSupported(EChunkProgressReportingType::Type ReportType) override
@@ -229,7 +262,7 @@ public:
 		return false;
 	}
 
-	virtual float GetChunkProgress( uint32 ChunkID, EChunkProgressReportingType::Type ReportType ) override
+	virtual float GetChunkProgress( uint32 PakchunkID, EChunkProgressReportingType::Type ReportType ) override
 	{
 		if (ReportType == EChunkProgressReportingType::PercentageComplete)
 		{
@@ -248,7 +281,7 @@ public:
 		return false;
 	}
 	
-	virtual bool PrioritizeChunk( uint32 ChunkID, EChunkPriority::Type Priority ) override
+	virtual bool PrioritizeChunk( uint32 PakchunkID, EChunkPriority::Type Priority ) override
 	{
 		return false;
 	}
@@ -307,6 +340,11 @@ protected:
 
 	/** Delegate called when installation succeeds or fails */
 	FPlatformChunkInstallMultiDelegate InstallDelegate;
+
+	virtual EChunkLocation::Type GetChunkLocation(uint32 ChunkID) override
+	{
+		return EChunkLocation::LocalFast;
+	}
 };
 
 PRAGMA_ENABLE_DEPRECATION_WARNINGS
\ No newline at end of file
diff --git a/Engine/Source/Runtime/Core/Public/GenericPlatform/GenericPlatformMisc.h b/Engine/Source/Runtime/Core/Public/GenericPlatform/GenericPlatformMisc.h
--- a/Engine/Source/Runtime/Core/Public/GenericPlatform/GenericPlatformMisc.h
+++ b/Engine/Source/Runtime/Core/Public/GenericPlatform/GenericPlatformMisc.h
@@ -1405,6 +1405,11 @@ public:
 	{
 	}
 
+	FORCEINLINE static int32 GetChunkIDFromPakchunkIndex(int32 PakchunkIndex)
+	{
+		return PakchunkIndex;
+	}
+
 #if !UE_BUILD_SHIPPING
 	/** 
 	 * Returns any platform specific warning messages we want printed on screen
diff --git a/Engine/Source/Runtime/Engine/Private/AssetManager.cpp b/Engine/Source/Runtime/Engine/Private/AssetManager.cpp
--- a/Engine/Source/Runtime/Engine/Private/AssetManager.cpp
+++ b/Engine/Source/Runtime/Engine/Private/AssetManager.cpp
@@ -1661,34 +1661,34 @@ bool UAssetManager::FindMissingChunkList(const TArray<FSoftObjectPath>& AssetLis
 		GetAssetDataForPath(Asset, FoundData);
 		TSet<int32> FoundChunks, MissingChunks, ErrorChunks;
 
-		for (int32 ChunkId : FoundData.ChunkIDs)
+		for (int32 PakchunkId : FoundData.ChunkIDs)
 		{
-			if (!ChunkLocationCache.Contains(ChunkId))
+			if (!ChunkLocationCache.Contains(PakchunkId))
 			{
-				EChunkLocation::Type Location = ChunkInstall->GetChunkLocation(ChunkId);
+				EChunkLocation::Type Location = ChunkInstall->GetPakchunkLocation(PakchunkId);
 
 				// If chunk install thinks the chunk is available, we need to double check with the pak system that it isn't
 				// pending decryption
 				if (Location >= EChunkLocation::LocalSlow && Pak->AnyChunksAvailable())
 				{
-					Location = Pak->GetPakChunkLocation(ChunkId);
+					Location = Pak->GetPakChunkLocation(PakchunkId);
 				}
 
-				ChunkLocationCache.Add(ChunkId, Location);
+				ChunkLocationCache.Add(PakchunkId, Location);
 			}
-			EChunkLocation::Type ChunkLocation = ChunkLocationCache[ChunkId];
+			EChunkLocation::Type ChunkLocation = ChunkLocationCache[PakchunkId];
 
 			switch (ChunkLocation)
 			{			
 			case EChunkLocation::DoesNotExist:
-				ErrorChunks.Add(ChunkId);
+				ErrorChunks.Add(PakchunkId);
 				break;
 			case EChunkLocation::NotAvailable:
-				MissingChunks.Add(ChunkId);
+				MissingChunks.Add(PakchunkId);
 				break;
 			case EChunkLocation::LocalSlow:
 			case EChunkLocation::LocalFast:
-				FoundChunks.Add(ChunkId);
+				FoundChunks.Add(PakchunkId);
 				break;
 			}
 		}
@@ -1751,7 +1751,7 @@ void UAssetManager::AcquireChunkList(const TArray<int32>& ChunkList, FAssetManag
 
 	for (int32 MissingChunk : PendingChunkInstall->PendingChunks)
 	{
-		ChunkInstall->PrioritizeChunk(MissingChunk, Priority);
+		ChunkInstall->PrioritizePakchunk(MissingChunk, Priority);
 	}
 }
 
@@ -1845,9 +1845,9 @@ void UAssetManager::OnChunkDownloaded(uint32 ChunkId, bool bSuccess)
 			TArray<int32> NewPendingList;
 			
 			// Check all chunks if they are done or failed
-			for (int32 PendingChunkId : PendingChunkInstall.PendingChunks)
+			for (int32 PendingPakchunkId : PendingChunkInstall.PendingChunks)
 			{
-				EChunkLocation::Type ChunkLocation = ChunkInstall->GetChunkLocation(PendingChunkId);
+				EChunkLocation::Type ChunkLocation = ChunkInstall->GetPakchunkLocation(PendingPakchunkId);
 
 				switch (ChunkLocation)
 				{
@@ -1855,7 +1855,7 @@ void UAssetManager::OnChunkDownloaded(uint32 ChunkId, bool bSuccess)
 					bFailed = true;
 					break;
 				case EChunkLocation::NotAvailable:
-					NewPendingList.Add(PendingChunkId);
+					NewPendingList.Add(PendingPakchunkId);
 					break;
 				}
 			}
diff --git a/Engine/Source/Runtime/PakFile/Private/IPlatformFilePak.cpp b/Engine/Source/Runtime/PakFile/Private/IPlatformFilePak.cpp
--- a/Engine/Source/Runtime/PakFile/Private/IPlatformFilePak.cpp
+++ b/Engine/Source/Runtime/PakFile/Private/IPlatformFilePak.cpp
@@ -56,7 +56,7 @@ CSV_DECLARE_CATEGORY_MODULE_EXTERN(CORE_API, FileIO);
 
 static FString GMountStartupPaksWildCard = TEXT(MOUNT_STARTUP_PAKS_WILDCARD);
 
-int32 ParseChunkIDFromFilename(const FString& InFilename)
+int32 GetPakchunkIDFromPakFile(const FString& InFilename)
 {
 	FString ChunkIdentifier(TEXT("pakchunk"));
 	FString BaseFilename = FPaths::GetBaseFilename(InFilename);
@@ -4748,7 +4748,7 @@ FPakFile::FPakFile(const TCHAR* Filename, bool bIsSigned)
 	, bSigned(bIsSigned)
 	, bIsValid(false)
 	, bFilenamesRemoved(false)
-	, ChunkID(ParseChunkIDFromFilename(Filename))
+	, ChunkID(GetPakchunkIDFromPakFile(Filename))
 	, bAttemptedPakEntryShrink(false)
 	, bAttemptedPakFilenameUnload(false)
  	, MappedFileHandle(nullptr)
@@ -4778,7 +4778,7 @@ FPakFile::FPakFile(IPlatformFile* LowerLevel, const TCHAR* Filename, bool bIsSig
 	, bSigned(bIsSigned)
 	, bIsValid(false)
 	, bFilenamesRemoved(false)
-	, ChunkID(ParseChunkIDFromFilename(Filename))
+	, ChunkID(GetPakchunkIDFromPakFile(Filename))
 	, bAttemptedPakEntryShrink(false)
 	, bAttemptedPakFilenameUnload(false)
 	, MappedFileHandle(nullptr)
@@ -6067,10 +6067,10 @@ void FPakPlatformFile::FindPakFilesInDirectory(IPlatformFile* LowLevelFile, cons
 					// if a platform supports chunk style installs, make sure that the chunk a pak file resides in is actually fully installed before accepting pak files from it
 					if (ChunkInstall)
 					{
-						int32 ChunkID = ParseChunkIDFromFilename(Filename);
-						if (ChunkID != INDEX_NONE)
+						int32 PakchunkID = GetPakchunkIDFromPakFile(Filename);
+						if (PakchunkID != INDEX_NONE)
 						{
-							if (ChunkInstall->GetChunkLocation(ChunkID) == EChunkLocation::NotAvailable)
+							if (ChunkInstall->GetPakchunkLocation(PakchunkID) == EChunkLocation::NotAvailable)
 							{
 								return true;
 							}
@@ -6574,7 +6574,7 @@ void FPakPlatformFile::RegisterEncryptionKey(const FGuid& InGuid, const FAES::FA
 				UE_LOG(LogPakFile, Log, TEXT("Successfully mounted deferred pak file '%s'"), *Entry.Filename);
 				NumMounted++;
 
-				int32 ChunkID = ParseChunkIDFromFilename(Entry.Filename);
+				int32 ChunkID = GetPakchunkIDFromPakFile(Entry.Filename);
 				if (ChunkID != INDEX_NONE)
 				{
 					ChunksToNotify.Add(ChunkID);
