Fix a compile error in IoDispatcher.cpp

Chunk Install has recently been refactored.  GetChunkLocation() has been moved to protected and GetPakchunkLocation() was added to public.  This was to solve the issue that chunk and pakchunk are no longer 1:1 mapping on some platforms.  (One chunk can contain multiple pakchunk files, some of which has a different pakchunk id than chunk id).

Assume InstallChunkId here actually means pakchunk id.

diff --git a/Engine/Source/Runtime/Core/Private/IO/IoDispatcher.cpp b/Engine/Source/Runtime/Core/Private/IO/IoDispatcher.cpp
--- a/Engine/Source/Runtime/Core/Private/IO/IoDispatcher.cpp
+++ b/Engine/Source/Runtime/Core/Private/IO/IoDispatcher.cpp
@@ -192,7 +192,7 @@ public:
 		for (const FIoStoreInstallManifest::FEntry& InstallManifestEntry : InstallManifest.ReadEntries())
 		{
 			FIoStoreEnvironment PartitionEnvironment(InitialEnvironment, InstallManifestEntry.PartitionName);
-			if (ChunkInstall && ChunkInstall->GetChunkLocation(InstallManifestEntry.InstallChunkId) == EChunkLocation::NotAvailable)
+			if (ChunkInstall && ChunkInstall->GetPakchunkLocation(InstallManifestEntry.InstallChunkId) == EChunkLocation::NotAvailable)
 			{
 				PendingInstallChunks.Add(MakeTuple(InstallManifestEntry.InstallChunkId, PartitionEnvironment));
 			}
