Back out changelist 11065948. Crashes in EdTest

diff --git a/Engine/Source/Runtime/Core/Public/Containers/ContainerAllocationPolicies.h b/Engine/Source/Runtime/Core/Public/Containers/ContainerAllocationPolicies.h
--- a/Engine/Source/Runtime/Core/Public/Containers/ContainerAllocationPolicies.h
+++ b/Engine/Source/Runtime/Core/Public/Containers/ContainerAllocationPolicies.h
@@ -75,15 +75,11 @@ FORCEINLINE SizeType DefaultCalculateSlackGrow(SizeType NumElements, SizeType Nu
 	checkSlow(NumElements > NumAllocatedElements && NumElements > 0);
 
 	SIZE_T Grow = FirstGrow; // this is the amount for the first alloc
-	if (NumAllocatedElements)
+	if (NumAllocatedElements || SIZE_T(NumElements) > Grow)
 	{
 		// Allocate slack for the array proportional to its size.
 		Grow = SIZE_T(NumElements) + 3 * SIZE_T(NumElements) / 8 + ConstantGrow;
 	}
-	else if (SIZE_T(NumElements) > Grow)
-	{
-		Grow = SIZE_T(NumElements);
-	}
 	if (bAllowQuantize)
 	{
 		Retval = (SizeType)(FMemory::QuantizeSize(Grow * BytesPerElement, Alignment) / BytesPerElement);
