Don't inherit some new Metal frame capture environment values (currently DYLD_INSERT_LIBRARIES) from parent process in editor mode or program that macOS Catalina has introduced.

diff --git a/Engine/Source/Runtime/Core/Private/Mac/MacPlatformProcess.cpp b/Engine/Source/Runtime/Core/Private/Mac/MacPlatformProcess.cpp
--- a/Engine/Source/Runtime/Core/Private/Mac/MacPlatformProcess.cpp
+++ b/Engine/Source/Runtime/Core/Private/Mac/MacPlatformProcess.cpp
@@ -261,6 +261,30 @@ FString FMacPlatformProcess::GetGameBundleId()
 
 @end // NSAutoReadPipe
 
+void DisableStallingMetalGPUCaptureKeysForChildProcessWhenDebugging(NSTask* childProcess)
+{
+#if WITH_EDITOR
+	if(FPlatformMisc::IsDebuggerPresent())
+	{
+		// - These are the extra environment keys when turning on GPU Frame Capture via Xcode 11.3 in macOS Catalina (10.15.2):
+		//		DYLD_INSERT_LIBRARIES, DYMTL_TOOLS_DYLIB_PATH, GPUTOOLS_LOAD_GTMTLCAPTURE, GT_HOST_URL_MTL and METAL_LOAD_INTERPOSER
+		// - Both DYLD_INSERT_LIBRARIES and METAL_LOAD_INTERPOSER seem to be new for Catalina (10.15.2) compared to Mojave (10.14.6).
+		// - Using DYLD_INSERT_LIBRARIES seem to be causing a stall at child process startup with Xcode debugger attached to the parent process.
+		// - Removing DYLD_INSERT_LIBRARIES removes the stall for child process startup which is especially useful for ShaderCompileWorker when invoking MetalCompiler and it's other child processes tools.
+
+		// Get the current environment and remove the relevant key(s)
+		NSMutableDictionary* childEnvironment = [[NSProcessInfo processInfo].environment mutableCopy];
+		
+		NSArray* metalKeysToRemove = @[@"DYLD_INSERT_LIBRARIES"];
+		
+		[childEnvironment removeObjectsForKeys:metalKeysToRemove];
+		[childProcess setEnvironment:childEnvironment];
+		
+		[childEnvironment release];
+	}
+#endif
+}
+
 bool FMacPlatformProcess::ExecProcess( const TCHAR* URL, const TCHAR* Params, int32* OutReturnCode, FString* OutStdOut, FString* OutStdErr, const TCHAR* OptionalWorkingDirectory)
 {
 	SCOPED_AUTORELEASE_POOL;
@@ -367,6 +391,8 @@ bool FMacPlatformProcess::ExecProcess( const TCHAR* URL, const TCHAR* Params, in
 			}
 		}
 		
+		DisableStallingMetalGPUCaptureKeysForChildProcessWhenDebugging(ProcessHandle);
+		
 		[ProcessHandle setArguments: Arguments];
 		
 		NSAutoReadPipe* StdOutPipe = [[NSAutoReadPipe new] autorelease];
@@ -536,6 +562,8 @@ FProcHandle FMacPlatformProcess::CreateProc( const TCHAR* URL, const TCHAR* Parm
 				}
 			}
 		}
+		
+		DisableStallingMetalGPUCaptureKeysForChildProcessWhenDebugging(ProcessHandle);
 
 		[ProcessHandle setArguments: Arguments];
 
