ArrayView: Allow slice creation from const views

diff --git a/Engine/Source/Runtime/Core/Public/Containers/ArrayView.h b/Engine/Source/Runtime/Core/Public/Containers/ArrayView.h
--- a/Engine/Source/Runtime/Core/Public/Containers/ArrayView.h
+++ b/Engine/Source/Runtime/Core/Public/Containers/ArrayView.h
@@ -273,7 +273,7 @@ public:
 	 * @param InNum number of elements in the new view
 	 * @returns Sliced view
 	 */
-	FORCEINLINE TArrayView Slice(int32 Index, int32 InNum)
+	FORCEINLINE TArrayView Slice(int32 Index, int32 InNum) const
 	{
 		check(InNum > 0);
 		check(IsValidIndex(Index));
