Add a define to override the build version used for symbols.

#rnx

diff --git a/Engine/Source/Runtime/Core/Private/GenericPlatform/GenericPlatformCrashContext.cpp b/Engine/Source/Runtime/Core/Private/GenericPlatform/GenericPlatformCrashContext.cpp
--- a/Engine/Source/Runtime/Core/Private/GenericPlatform/GenericPlatformCrashContext.cpp
+++ b/Engine/Source/Runtime/Core/Private/GenericPlatform/GenericPlatformCrashContext.cpp
@@ -461,7 +461,11 @@ void FGenericCrashContext::SerializeSessionContext(FString& Buffer)
 	AddCrashPropertyInternal(Buffer, TEXT("GameSessionID"), NCached::Session.GameSessionID);
 
 	// Unique string specifying the symbols to be used by CrashReporter
+#ifdef UE_SYMBOLS_VERSION
+	FString Symbols = FString(UE_SYMBOLS_VERSION);
+#else
 	FString Symbols = FString::Printf(TEXT("%s"), FApp::GetBuildVersion());
+#endif
 #ifdef UE_APP_FLAVOR
 	Symbols = FString::Printf(TEXT("%s-%s"), *Symbols, *FString(UE_APP_FLAVOR));
 #endif
