Add custom chunk mapping type.  Only update pakchunk to chunk for main mappings, aka, pakchunk0-platform.pak but not pakchunk0optional-platform.pak

diff --git a/Engine/Source/Runtime/Core/Public/GenericPlatform/GenericPlatformChunkInstall.h b/Engine/Source/Runtime/Core/Public/GenericPlatform/GenericPlatformChunkInstall.h
--- a/Engine/Source/Runtime/Core/Public/GenericPlatform/GenericPlatformChunkInstall.h
+++ b/Engine/Source/Runtime/Core/Public/GenericPlatform/GenericPlatformChunkInstall.h
@@ -97,11 +97,18 @@ struct FCustomChunk
 
 struct FCustomChunkMapping
 {
+	enum class CustomChunkMappingType : uint8
+	{
+		Main,
+		Optional
+	};
+
 	FString Pattern;
 	uint32	ChunkID;
+	CustomChunkMappingType MappingType;
 
-	FCustomChunkMapping(FString InPattern, uint32 InChunkID) :
-		Pattern(InPattern), ChunkID(InChunkID)
+	FCustomChunkMapping(FString InPattern, uint32 InChunkID, CustomChunkMappingType InMappingType) :
+		Pattern(InPattern), ChunkID(InChunkID), MappingType(InMappingType)
 	{}
 };
 
