Fix crash when automation test fails

diff --git a/Engine/Source/Runtime/Core/Private/Misc/AutomationTest.cpp b/Engine/Source/Runtime/Core/Private/Misc/AutomationTest.cpp
--- a/Engine/Source/Runtime/Core/Private/Misc/AutomationTest.cpp
+++ b/Engine/Source/Runtime/Core/Private/Misc/AutomationTest.cpp
@@ -1304,7 +1304,7 @@ bool FAutomationTestBase::TestNull(const TCHAR* What, const void* Pointer)
 {
 	if ( Pointer != nullptr )
 	{
-		AddError(FString::Printf(TEXT("Expected '%s' to be null."), *What), 1);
+		AddError(FString::Printf(TEXT("Expected '%s' to be null."), What), 1);
 		return false;
 	}
 	return true;
