enable nanomalloc again as the 390MB VM allocation has been fixed
turn on nanomalloc usage for small block allocations

diff --git a/Engine/Source/Runtime/Core/Public/HAL/MallocBinned.h b/Engine/Source/Runtime/Core/Public/HAL/MallocBinned.h
--- a/Engine/Source/Runtime/Core/Public/HAL/MallocBinned.h
+++ b/Engine/Source/Runtime/Core/Public/HAL/MallocBinned.h
@@ -20,7 +20,7 @@
 #endif
 
 #ifndef USE_OS_SMALL_BLOCK_ALLOC
-#define USE_OS_SMALL_BLOCK_ALLOC 0
+#define USE_OS_SMALL_BLOCK_ALLOC PLATFORM_IOS
 #endif //USE_OS_SMALL_BLOCK_ALLOC
 
 #ifndef USE_OS_SMALL_BLOCK_GRAB_MEMORY_FROM_OS
