Fixed static initialization order issue when binding FTextLocalizationManager::OnPakFileMounted

The FTextLocalizationManager constructor would run before the FCoreDelegates::OnPakFileMounted constructor, resulting in the FTextLocalizationManager::OnPakFileMounted binding being cleared. This binding is now made from BeginPreInitTextLocalization which runs before the PAK file loader is created.

#rnx

diff --git a/Engine/Source/Runtime/Core/Private/Internationalization/TextLocalizationManager.cpp b/Engine/Source/Runtime/Core/Private/Internationalization/TextLocalizationManager.cpp
--- a/Engine/Source/Runtime/Core/Private/Internationalization/TextLocalizationManager.cpp
+++ b/Engine/Source/Runtime/Core/Private/Internationalization/TextLocalizationManager.cpp
@@ -292,6 +292,17 @@ void ApplyDefaultCultureSettings(const ELocalizationLoadFlags LocLoadFlags)
 	}
 }
 
+void BeginPreInitTextLocalization()
+{
+	LLM_SCOPE(ELLMTag::Localization);
+
+	SCOPED_BOOT_TIMING("BeginPreInitTextLocalization");
+	DECLARE_SCOPE_CYCLE_COUNTER(TEXT("BeginPreInitTextLocalization"), STAT_BeginPreInitTextLocalization, STATGROUP_LoadTime);
+
+	// Bind this delegate before the PAK file loader is created
+	FCoreDelegates::OnPakFileMounted.AddRaw(&FTextLocalizationManager::Get(), &FTextLocalizationManager::OnPakFileMounted);
+}
+
 void BeginInitTextLocalization()
 {
 	LLM_SCOPE(ELLMTag::Localization);
@@ -400,8 +411,6 @@ FTextLocalizationManager::FTextLocalizationManager()
 	const bool bRefreshResources = false;
 	RegisterTextSource(LocResTextSource.ToSharedRef(), bRefreshResources);
 	RegisterTextSource(PolyglotTextSource.ToSharedRef(), bRefreshResources);
-
-	FCoreDelegates::OnPakFileMounted.AddRaw(this, &FTextLocalizationManager::OnPakFileMounted);
 }
 
 void FTextLocalizationManager::DumpMemoryInfo()
@@ -851,6 +860,8 @@ void FTextLocalizationManager::OnPakFileMounted(const TCHAR* PakFilename, const
 	// Track this so that full resource refreshes (eg, changing culture) work as expected
 	LocResTextSource->RegisterChunkId(ChunkId);
 
+	UE_LOG(LogTextLocalizationManager, Verbose, TEXT("Request to load localization data for chunk %d (from PAK '%s')"), ChunkId, PakFilename);
+
 	if (!bIsInitialized)
 	{
 		// If we're not yet initialized then don't bother patching, as the full initialization path will load the data for this chunk
@@ -870,6 +881,7 @@ void FTextLocalizationManager::OnPakFileMounted(const TCHAR* PakFilename, const
 		{
 			const FString LocalizationTargetForChunk = TextLocalizationResourceUtil::GetLocalizationTargetNameForChunkId(LocalizationTarget, ChunkId);
 			PrioritizedLocalizationPaths.Add(FPaths::ProjectContentDir() / TEXT("Localization") / LocalizationTargetForChunk);
+			UE_LOG(LogTextLocalizationManager, Verbose, TEXT("Loading chunked localization data from '%s'"), *PrioritizedLocalizationPaths.Last());
 		}
 		LocResTextSource->LoadLocalizedResourcesFromPaths(TArrayView<FString>(), PrioritizedLocalizationPaths, TArrayView<FString>(), LocLoadFlags, PrioritizedCultureNames, UnusedNativeResource, LocalizedResource);
 	}
@@ -908,10 +920,12 @@ void FTextLocalizationManager::OnPakFileMounted(const TCHAR* PakFilename, const
 	// Apply the new data
 	if (bNeedsFullRefresh)
 	{
+		UE_LOG(LogTextLocalizationManager, Verbose, TEXT("Patching chunked localization data failed, performing full refresh"));
 		RefreshResources();
 	}
 	else
 	{
+		UE_LOG(LogTextLocalizationManager, Verbose, TEXT("Patching chunked localization data for %d entries"), LocalizedResource.Entries.Num());
 		UpdateFromLocalizations(MoveTemp(LocalizedResource), /*bDirtyTextRevision*/true);
 	}
 }
diff --git a/Engine/Source/Runtime/Core/Public/Internationalization/TextLocalizationManager.h b/Engine/Source/Runtime/Core/Public/Internationalization/TextLocalizationManager.h
--- a/Engine/Source/Runtime/Core/Public/Internationalization/TextLocalizationManager.h
+++ b/Engine/Source/Runtime/Core/Public/Internationalization/TextLocalizationManager.h
@@ -27,6 +27,7 @@ typedef TSharedPtr<FString, ESPMode::ThreadSafe> FTextDisplayStringPtr;
 /** Singleton class that manages display strings for FText. */
 class CORE_API FTextLocalizationManager
 {
+	friend CORE_API void BeginPreInitTextLocalization();
 	friend CORE_API void BeginInitTextLocalization();
 	friend CORE_API void InitEngineTextLocalization();
 	friend CORE_API void InitGameTextLocalization();
diff --git a/Engine/Source/Runtime/Core/Public/Internationalization/TextLocalizationManagerGlobals.h b/Engine/Source/Runtime/Core/Public/Internationalization/TextLocalizationManagerGlobals.h
--- a/Engine/Source/Runtime/Core/Public/Internationalization/TextLocalizationManagerGlobals.h
+++ b/Engine/Source/Runtime/Core/Public/Internationalization/TextLocalizationManagerGlobals.h
@@ -4,6 +4,7 @@
 
 #include "CoreTypes.h"
 
+CORE_API void BeginPreInitTextLocalization();
 CORE_API void BeginInitTextLocalization();
 CORE_API void InitEngineTextLocalization();
 CORE_API void InitGameTextLocalization();
diff --git a/Engine/Source/Runtime/Launch/Private/LaunchEngineLoop.cpp b/Engine/Source/Runtime/Launch/Private/LaunchEngineLoop.cpp
--- a/Engine/Source/Runtime/Launch/Private/LaunchEngineLoop.cpp
+++ b/Engine/Source/Runtime/Launch/Private/LaunchEngineLoop.cpp
@@ -1585,6 +1585,11 @@ int32 FEngineLoop::PreInitPreStartupScreen(const TCHAR* CmdLine)
 	}
 #endif
 
+	{
+		SCOPED_BOOT_TIMING("BeginPreInitTextLocalization");
+		BeginPreInitTextLocalization();
+	}
+
 	// allow the command line to override the platform file singleton
 	bool bFileOverrideFound = false;
 	{
