Core - Adding utility functions for doing GetRangePct, GetRangeValue...etc for TRange<T> types to the math utility.

diff --git a/Engine/Source/Runtime/Core/Public/Math/UnrealMathUtility.h b/Engine/Source/Runtime/Core/Public/Math/UnrealMathUtility.h
--- a/Engine/Source/Runtime/Core/Public/Math/UnrealMathUtility.h
+++ b/Engine/Source/Runtime/Core/Public/Math/UnrealMathUtility.h
@@ -35,6 +35,8 @@ struct  FTransform;
 class  FSphere;
 struct FVector2D;
 struct FLinearColor;
+template<typename ElementType>
+class TRange;
 
 /*-----------------------------------------------------------------------------
 	Floating point constants.
@@ -716,14 +718,15 @@ public:
 	// Interpolation Functions
 
 	/** Calculates the percentage along a line from MinValue to MaxValue that Value is. */
-	static FORCEINLINE float GetRangePct(float MinValue, float MaxValue, float Value)
+	template<class T>
+	static FORCEINLINE T GetRangePct(T MinValue, T MaxValue, T Value)
 	{
 		// Avoid Divide by Zero.
 		// But also if our range is a point, output whether Value is before or after.
-		const float Divisor = MaxValue - MinValue;
+		const T Divisor = MaxValue - MinValue;
 		if (FMath::IsNearlyZero(Divisor))
 		{
-			return (Value >= MaxValue) ? 1.f : 0.f;
+			return (Value >= MaxValue) ? (T)1 : (T)0;
 		}
 
 		return (Value - MinValue) / Divisor;
@@ -748,6 +751,25 @@ public:
 		return GetRangeValue(OutputRange, GetRangePct(InputRange, Value));
 	}
 
+	template<class T>
+	static FORCEINLINE double GetRangePct(TRange<T> const& Range, T Value)
+	{
+		return GetRangePct(Range.GetLowerBoundValue(), Range.GetUpperBoundValue(), Value);
+	}
+
+	template<class T>
+	static FORCEINLINE T GetRangeValue(TRange<T> const& Range, T Pct)
+	{
+		return FMath::Lerp<T>(Range.GetLowerBoundValue(), Range.GetUpperBoundValue(), Pct);
+	}
+
+	template<class T>
+	static FORCEINLINE T GetMappedRangeValueClamped(const TRange<T>& InputRange, const TRange<T>& OutputRange, const T Value)
+	{
+		const T ClampedPct = FMath::Clamp<T>(GetRangePct(InputRange, Value), 0, 1);
+		return GetRangeValue(OutputRange, ClampedPct);
+	}
+
 	/** Performs a linear interpolation between two values, Alpha ranges from 0-1 */
 	template< class T, class U > 
 	static FORCEINLINE_DEBUGGABLE T Lerp( const T& A, const T& B, const U& Alpha )
