Core - Adding EnableIf to the GetRangePct to limit it to floating point templates.

diff --git a/Engine/Source/Runtime/Core/Public/Math/UnrealMathUtility.h b/Engine/Source/Runtime/Core/Public/Math/UnrealMathUtility.h
--- a/Engine/Source/Runtime/Core/Public/Math/UnrealMathUtility.h
+++ b/Engine/Source/Runtime/Core/Public/Math/UnrealMathUtility.h
@@ -5,6 +5,7 @@
 #include "CoreTypes.h"
 #include "Misc/AssertionMacros.h"
 #include "HAL/PlatformMath.h"
+#include "Templates/IsFloatingPoint.h"
 
 
 //#define IMPLEMENT_ASSIGNMENT_OPERATOR_MANUALLY
@@ -718,8 +719,8 @@ public:
 	// Interpolation Functions
 
 	/** Calculates the percentage along a line from MinValue to MaxValue that Value is. */
-	template<class T>
-	static FORCEINLINE T GetRangePct(T MinValue, T MaxValue, T Value)
+	template<typename T>
+	static FORCEINLINE typename TEnableIf<TIsFloatingPoint<T>::Value, T>::Type GetRangePct(T MinValue, T MaxValue, T Value)
 	{
 		// Avoid Divide by Zero.
 		// But also if our range is a point, output whether Value is before or after.
