Core: Additional conversion and narrowing warning fixes for Clang

diff --git a/Engine/Source/Runtime/Core/Private/Logging/LogTrace.cpp b/Engine/Source/Runtime/Core/Private/Logging/LogTrace.cpp
--- a/Engine/Source/Runtime/Core/Private/Logging/LogTrace.cpp
+++ b/Engine/Source/Runtime/Core/Private/Logging/LogTrace.cpp
@@ -28,7 +28,7 @@ UE_TRACE_EVENT_END()
 
 void FLogTrace::OutputLogCategory(const FLogCategoryBase* Category, const TCHAR* Name, ELogVerbosity::Type DefaultVerbosity)
 {
-	uint16 NameSize = (FCString::Strlen(Name) + 1) * sizeof(TCHAR);
+	uint16 NameSize = (uint16)((FCString::Strlen(Name) + 1) * sizeof(TCHAR));
 	UE_TRACE_LOG(Logging, LogCategory, NameSize)
 		<< LogCategory.CategoryPointer(Category)
 		<< LogCategory.DefaultVerbosity(DefaultVerbosity)
@@ -37,8 +37,8 @@ void FLogTrace::OutputLogCategory(const FLogCategoryBase* Category, const TCHAR*
 
 void FLogTrace::OutputLogMessageSpec(const void* LogPoint, const FLogCategoryBase* Category, ELogVerbosity::Type Verbosity, const ANSICHAR* File, int32 Line, const TCHAR* Format)
 {
-	uint16 FileNameSize = strlen(File) + 1;
-	uint16 FormatStringSize = (FCString::Strlen(Format) + 1) * sizeof(TCHAR);
+	uint16 FileNameSize = (uint16)(strlen(File) + 1);
+	uint16 FormatStringSize = (uint16)((FCString::Strlen(Format) + 1) * sizeof(TCHAR));
 	auto StringCopyFunc = [FileNameSize, FormatStringSize, File, Format](uint8* Out) {
 		memcpy(Out, File, FileNameSize);
 		memcpy(Out + FileNameSize, Format, FormatStringSize);
diff --git a/Engine/Source/Runtime/Core/Private/Misc/GuidGenerator.cpp b/Engine/Source/Runtime/Core/Private/Misc/GuidGenerator.cpp
--- a/Engine/Source/Runtime/Core/Private/Misc/GuidGenerator.cpp
+++ b/Engine/Source/Runtime/Core/Private/Misc/GuidGenerator.cpp
@@ -233,7 +233,7 @@ true_random(void)
   if ((rand_ib += 73939) >= 224729)
       rand_ib -= 96233;
   rand_irand = (rand_irand * rand_m) + rand_ia + rand_ib;
-  return (rand_irand >> 16) ^ (rand_irand & RAND_MASK);
+  return unsigned16((rand_irand >> 16) ^ (rand_irand & RAND_MASK));
 }
 
 /*
diff --git a/Engine/Source/Runtime/Core/Private/Misc/SecureHash.cpp b/Engine/Source/Runtime/Core/Private/Misc/SecureHash.cpp
--- a/Engine/Source/Runtime/Core/Private/Misc/SecureHash.cpp
+++ b/Engine/Source/Runtime/Core/Private/Misc/SecureHash.cpp
@@ -517,7 +517,7 @@ void LexFromString(FMD5Hash& Hash, const TCHAR* Buffer)
 			return;
 		}
 
-		Bytes[ByteIndex] = (FirstCharVal << 4) + SecondCharVal;
+		Bytes[ByteIndex] = (uint8)((FirstCharVal << 4) + SecondCharVal);
 	}
 
 	FMemory::Memcpy(Hash.Bytes, Bytes, 16);
diff --git a/Engine/Source/Runtime/Core/Private/ProfilingDebugging/CountersTrace.cpp b/Engine/Source/Runtime/Core/Private/ProfilingDebugging/CountersTrace.cpp
--- a/Engine/Source/Runtime/Core/Private/ProfilingDebugging/CountersTrace.cpp
+++ b/Engine/Source/Runtime/Core/Private/ProfilingDebugging/CountersTrace.cpp
@@ -29,7 +29,7 @@ uint16 FCountersTrace::OutputInitCounter(const TCHAR* CounterName, ETraceCounter
 {
 	static TAtomic<uint16> NextId(1);
 	uint16 CounterId = uint16(NextId++);
-	uint16 NameSize = (FCString::Strlen(CounterName) + 1) * sizeof(TCHAR);
+	uint16 NameSize = (uint16)((FCString::Strlen(CounterName) + 1) * sizeof(TCHAR));
 	UE_TRACE_LOG(Counters, Spec, NameSize)
 		<< Spec.Id(CounterId)
 		<< Spec.Type(uint8(CounterType))
diff --git a/Engine/Source/Runtime/Core/Private/ProfilingDebugging/CpuProfilerTrace.cpp b/Engine/Source/Runtime/Core/Private/ProfilingDebugging/CpuProfilerTrace.cpp
--- a/Engine/Source/Runtime/Core/Private/ProfilingDebugging/CpuProfilerTrace.cpp
+++ b/Engine/Source/Runtime/Core/Private/ProfilingDebugging/CpuProfilerTrace.cpp
@@ -224,7 +224,7 @@ uint32 FCpuProfilerTraceInternal::GetNextSpecId()
 uint32 FCpuProfilerTrace::OutputEventType(const TCHAR* Name, ECpuProfilerGroup Group)
 {
 	uint32 SpecId = FCpuProfilerTraceInternal::GetNextSpecId();
-	uint16 NameSize = (FCString::Strlen(Name) + 1) * sizeof(TCHAR);
+	uint16 NameSize = (uint16)((FCString::Strlen(Name) + 1) * sizeof(TCHAR));
 	UE_TRACE_LOG(CpuProfiler, EventSpec, NameSize)
 		<< EventSpec.Id(SpecId)
 		<< EventSpec.Group(uint16(Group))
@@ -236,7 +236,7 @@ uint32 FCpuProfilerTrace::OutputEventType(const TCHAR* Name, ECpuProfilerGroup G
 uint32 FCpuProfilerTrace::OutputEventType(const ANSICHAR* Name, ECpuProfilerGroup Group)
 {
 	uint32 SpecId = FCpuProfilerTraceInternal::GetNextSpecId();
-	uint16 NameSize = strlen(Name) + 1;
+	uint16 NameSize = (uint16)(strlen(Name) + 1);
 	UE_TRACE_LOG(CpuProfiler, EventSpec, NameSize)
 		<< EventSpec.Id(SpecId)
 		<< EventSpec.Group(uint16(Group))
diff --git a/Engine/Source/Runtime/Core/Private/ProfilingDebugging/CsvProfilerTrace.cpp b/Engine/Source/Runtime/Core/Private/ProfilingDebugging/CsvProfilerTrace.cpp
--- a/Engine/Source/Runtime/Core/Private/ProfilingDebugging/CsvProfilerTrace.cpp
+++ b/Engine/Source/Runtime/Core/Private/ProfilingDebugging/CsvProfilerTrace.cpp
@@ -116,7 +116,7 @@ struct FCsvProfilerTraceInternal
 
 void FCsvProfilerTrace::OutputRegisterCategory(int32 Index, const TCHAR* Name)
 {
-	uint16 NameSize = (FCString::Strlen(Name) + 1) * sizeof(TCHAR);
+	uint16 NameSize = (uint16)((FCString::Strlen(Name) + 1) * sizeof(TCHAR));
 	UE_TRACE_LOG(CsvProfiler, RegisterCategory, NameSize)
 		<< RegisterCategory.Index(Index)
 		<< RegisterCategory.Attachment(Name, NameSize);
@@ -124,7 +124,7 @@ void FCsvProfilerTrace::OutputRegisterCategory(int32 Index, const TCHAR* Name)
 
 void FCsvProfilerTrace::OutputInlineStat(const char* StatName, int32 CategoryIndex)
 {
-	uint16 NameSize = (strlen(StatName) + 1) * sizeof(char);
+	uint16 NameSize = (uint16)((strlen(StatName) + 1) * sizeof(char));
 	UE_TRACE_LOG(CsvProfiler, DefineInlineStat, NameSize)
 		<< DefineInlineStat.StatId(FCsvProfilerTraceInternal::GetStatId(StatName, CategoryIndex))
 		<< DefineInlineStat.CategoryIndex(CategoryIndex)
@@ -142,7 +142,7 @@ void FCsvProfilerTrace::OutputDeclaredStat(const FName& StatName, int32 Category
 {
 	TCHAR NameString[NAME_SIZE];
 	StatName.GetPlainNameString(NameString);
-	uint16 NameSize = (FCString::Strlen(NameString) + 1) * sizeof(TCHAR);
+	uint16 NameSize = (uint16)((FCString::Strlen(NameString) + 1) * sizeof(TCHAR));
 	UE_TRACE_LOG(CsvProfiler, DefineDeclaredStat, NameSize)
 		<< DefineDeclaredStat.StatId(FCsvProfilerTraceInternal::GetStatId(StatName, CategoryIndex))
 		<< DefineDeclaredStat.CategoryIndex(CategoryIndex)
@@ -240,7 +240,7 @@ void FCsvProfilerTrace::OutputCustomStat(const FName& StatName, int32 CategoryIn
 void FCsvProfilerTrace::OutputBeginCapture(const TCHAR* Filename, uint32 RenderThreadId, uint32 RHIThreadId, const char* DefaultWaitStatName, bool bEnableCounts)
 {
 	OutputInlineStat(DefaultWaitStatName, CSV_CATEGORY_INDEX(Exclusive));
-	uint16 NameSize = (FCString::Strlen(Filename) + 1) * sizeof(TCHAR);
+	uint16 NameSize = (uint16)((FCString::Strlen(Filename) + 1) * sizeof(TCHAR));
 	UE_TRACE_LOG(CsvProfiler, BeginCapture, NameSize)
 		<< BeginCapture.Cycle(FPlatformTime::Cycles64())
 		<< BeginCapture.RenderThreadId(RenderThreadId)
@@ -251,7 +251,7 @@ void FCsvProfilerTrace::OutputBeginCapture(const TCHAR* Filename, uint32 RenderT
 
 void FCsvProfilerTrace::OutputEvent(const TCHAR* Text, int32 CategoryIndex, uint64 Cycles)
 {
-	uint16 TextSize = (FCString::Strlen(Text) + 1) * sizeof(TCHAR);
+	uint16 TextSize = (uint16)((FCString::Strlen(Text) + 1) * sizeof(TCHAR));
 	UE_TRACE_LOG(CsvProfiler, Event, TextSize)
 		<< Event.Cycle(Cycles)
 		<< Event.CategoryIndex(CategoryIndex)
@@ -267,8 +267,8 @@ void FCsvProfilerTrace::OutputEndCapture()
 
 void FCsvProfilerTrace::OutputMetadata(const TCHAR* Key, const TCHAR* Value)
 {
-	uint16 KeySize = (FCString::Strlen(Key) + 1) * sizeof(TCHAR);
-	uint16 ValueSize = (FCString::Strlen(Value) + 1) * sizeof(TCHAR);
+	uint16 KeySize = (uint16)((FCString::Strlen(Key) + 1) * sizeof(TCHAR));
+	uint16 ValueSize = (uint16)((FCString::Strlen(Value) + 1) * sizeof(TCHAR));
 	auto Attachment = [Key, KeySize, Value, ValueSize](uint8* Out)
 	{
 		memcpy(Out, Key, KeySize);
diff --git a/Engine/Source/Runtime/Core/Private/ProfilingDebugging/MiscTrace.cpp b/Engine/Source/Runtime/Core/Private/ProfilingDebugging/MiscTrace.cpp
--- a/Engine/Source/Runtime/Core/Private/ProfilingDebugging/MiscTrace.cpp
+++ b/Engine/Source/Runtime/Core/Private/ProfilingDebugging/MiscTrace.cpp
@@ -69,7 +69,7 @@ void FMiscTrace::OutputRegisterGameThread(uint32 Id)
 
 void FMiscTrace::OutputCreateThread(uint32 Id, const TCHAR* Name, uint32 Priority)
 {
-	uint16 NameSize = (FCString::Strlen(Name) + 1) * sizeof(TCHAR);
+	uint16 NameSize = (uint16)((FCString::Strlen(Name) + 1) * sizeof(TCHAR));
 	UE_TRACE_LOG(Misc, CreateThread, NameSize)
 		<< CreateThread.CurrentThreadId(FPlatformTLS::GetCurrentThreadId())
 		<< CreateThread.CreatedThreadId(Id)
@@ -79,7 +79,7 @@ void FMiscTrace::OutputCreateThread(uint32 Id, const TCHAR* Name, uint32 Priorit
 
 void FMiscTrace::OutputSetThreadGroup(uint32 Id, const ANSICHAR* GroupName)
 {
-	uint16 NameSize = strlen(GroupName) + 1;
+	uint16 NameSize = (uint16)(strlen(GroupName) + 1);
 	UE_TRACE_LOG(Misc, SetThreadGroup, NameSize)
 		<< SetThreadGroup.ThreadId(Id)
 		<< SetThreadGroup.Attachment(GroupName, NameSize);
@@ -87,7 +87,7 @@ void FMiscTrace::OutputSetThreadGroup(uint32 Id, const ANSICHAR* GroupName)
 
 void FMiscTrace::OutputBeginThreadGroupScope(const ANSICHAR* GroupName)
 {
-	uint16 NameSize = strlen(GroupName) + 1;
+	uint16 NameSize = (uint16)(strlen(GroupName) + 1);
 	UE_TRACE_LOG(Misc, BeginThreadGroupScope, NameSize)
 		<< BeginThreadGroupScope.CurrentThreadId(FPlatformTLS::GetCurrentThreadId())
 		<< BeginThreadGroupScope.Attachment(GroupName, NameSize);
@@ -101,8 +101,8 @@ void FMiscTrace::OutputEndThreadGroupScope()
 
 void FMiscTrace::OutputBookmarkSpec(const void* BookmarkPoint, const ANSICHAR* File, int32 Line, const TCHAR* Format)
 {
-	uint16 FileNameSize = strlen(File) + 1;
-	uint16 FormatStringSize = (FCString::Strlen(Format) + 1) * sizeof(TCHAR);
+	uint16 FileNameSize = (uint16)(strlen(File) + 1);
+	uint16 FormatStringSize = (uint16)((FCString::Strlen(Format) + 1) * sizeof(TCHAR));
 	auto StringCopyFunc = [FileNameSize, FormatStringSize, File, Format](uint8* Out) {
 		memcpy(Out, File, FileNameSize);
 		memcpy(Out + FileNameSize, Format, FormatStringSize);
diff --git a/Engine/Source/Runtime/Core/Private/Serialization/LoadTimeTrace.cpp b/Engine/Source/Runtime/Core/Private/Serialization/LoadTimeTrace.cpp
--- a/Engine/Source/Runtime/Core/Private/Serialization/LoadTimeTrace.cpp
+++ b/Engine/Source/Runtime/Core/Private/Serialization/LoadTimeTrace.cpp
@@ -24,7 +24,7 @@ FLoadTimeProfilerTrace::FRequestGroupScope::~FRequestGroupScope()
 
 void FLoadTimeProfilerTrace::FRequestGroupScope::OutputBegin()
 {
-	uint16 FormatStringSize = (FCString::Strlen(FormatString) + 1) * sizeof(TCHAR);
+	uint16 FormatStringSize = (uint16)((FCString::Strlen(FormatString) + 1) * sizeof(TCHAR));
 	auto Attachment = [this, FormatStringSize](uint8* Out)
 	{
 		memcpy(Out, FormatString, FormatStringSize);
diff --git a/Engine/Source/Runtime/Core/Private/Tests/Math/UnrealMathTest.cpp b/Engine/Source/Runtime/Core/Private/Tests/Math/UnrealMathTest.cpp
--- a/Engine/Source/Runtime/Core/Private/Tests/Math/UnrealMathTest.cpp
+++ b/Engine/Source/Runtime/Core/Private/Tests/Math/UnrealMathTest.cpp
@@ -1453,34 +1453,34 @@ bool FVectorRegisterAbstractionTest::RunTest(const FString& Parameters)
 		static XYPair XYArray[] =
 		{		
 			// Test normal ranges
-			{ 0.0,	 1.0},
-			{ 1.5,	 1.0},
-			{ 2.8,	 0.3},
-			{-2.8,	 0.3},
-			{ 2.8,	-0.3},
-			{-2.8,	-0.3},
-			{-0.4,	 5.5},
-			{ 0.4,	-5.5},
-			{ 2.8,	 2.0 + KINDA_SMALL_NUMBER},
-			{-2.8,	 2.0 - KINDA_SMALL_NUMBER},
+			{ 0.0f,	 1.0f},
+			{ 1.5f,	 1.0f},
+			{ 2.8f,	 0.3f},
+			{-2.8f,	 0.3f},
+			{ 2.8f,	-0.3f},
+			{-2.8f,	-0.3f},
+			{-0.4f,	 5.5f},
+			{ 0.4f,	-5.5f},
+			{ 2.8f,	 2.0f + KINDA_SMALL_NUMBER},
+			{-2.8f,	 2.0f - KINDA_SMALL_NUMBER},
 
 			// Analytically should be zero but floating point precision can cause results close to Y (or erroneously negative) depending on the method used.
-			{55.8,	 9.3},
-			{1234.1234, 0.1234},
+			{55.8f,	 9.3f},
+			{1234.1234f, 0.1234f},
 
 			// Commonly used for FRotators and angles
-			{725.2,		360.0},
-			{179.9,		 90.0},
-			{ 5.3*PI,	2.*PI},
-			{-5.3*PI,	2.*PI},
+			{725.2f,		360.0f},
+			{179.9f,		 90.0f},
+			{ 5.3f*PI,	2.f*PI},
+			{-5.3f*PI,	2.f*PI},
 
 			// Test extreme ranges
-			{ 1.0,			 KINDA_SMALL_NUMBER},
-			{ 1.0,			-KINDA_SMALL_NUMBER},
+			{ 1.0f,			 KINDA_SMALL_NUMBER},
+			{ 1.0f,			-KINDA_SMALL_NUMBER},
 			{-SMALL_NUMBER,  SMALL_NUMBER},
 			{ SMALL_NUMBER, -SMALL_NUMBER},
-			{ 1.0,			 MIN_flt},
-			{ 1.0,			-MIN_flt},
+			{ 1.0f,			 MIN_flt},
+			{ 1.0f,			-MIN_flt},
 			{ MAX_flt,		 MIN_flt},
 			{ MAX_flt,		-MIN_flt},
 
@@ -1500,7 +1500,7 @@ bool FVectorRegisterAbstractionTest::RunTest(const FString& Parameters)
 			//UE_LOG(LogUnrealMathTest, Warning, TEXT("fmodf(%f, %f) Ours: %f Theirs: %f"), X, Y, Ours, Theirs);
 
 			// A compiler bug causes stock fmodf() to rarely return NaN for valid input, we don't want to report this as a fatal error.
-			if (Y != 0 && FMath::IsNaN(Theirs))
+			if (Y != 0.0f && FMath::IsNaN(Theirs))
 			{
 				UE_LOG(LogUnrealMathTest, Warning, TEXT("fmodf(%f, %f) with valid input resulted in NaN!"), X, Y);
 				continue;
diff --git a/Engine/Source/Runtime/Core/Public/Clang/ClangPlatformCompilerPreSetup.h b/Engine/Source/Runtime/Core/Public/Clang/ClangPlatformCompilerPreSetup.h
--- a/Engine/Source/Runtime/Core/Public/Clang/ClangPlatformCompilerPreSetup.h
+++ b/Engine/Source/Runtime/Core/Public/Clang/ClangPlatformCompilerPreSetup.h
@@ -74,13 +74,17 @@
 #endif // PRAGMA_ENABLE_SHADOW_VARIABLE_WARNINGS
 
 #ifndef PRAGMA_DISABLE_UNSAFE_TYPECAST_WARNINGS
-	#define PRAGMA_DISABLE_UNSAFE_TYPECAST_WARNINGS
-		//@TODO: FLOATPRECISION: Add Clang support for this kind of error
+	#define PRAGMA_DISABLE_UNSAFE_TYPECAST_WARNINGS \
+		_Pragma("clang diagnostic push") \
+		_Pragma("clang diagnostic ignored \"-Wfloat-conversion\"") \
+		_Pragma("clang diagnostic ignored \"-Wimplicit-float-conversion\"") \
+		_Pragma("clang diagnostic ignored \"-Wimplicit-int-conversion\"") \
+		_Pragma("clang diagnostic ignored \"-Wc++11-narrowing\"")
 #endif // PRAGMA_DISABLE_UNSAFE_TYPECAST_WARNINGS
 
 #ifndef PRAGMA_ENABLE_UNSAFE_TYPECAST_WARNINGS
-	#define PRAGMA_ENABLE_UNSAFE_TYPECAST_WARNINGS
-		//@TODO: FLOATPRECISION: Add Clang support for this kind of error
+	#define PRAGMA_ENABLE_UNSAFE_TYPECAST_WARNINGS \
+		_Pragma("clang diagnostic pop")
 #endif // PRAGMA_ENABLE_UNSAFE_TYPECAST_WARNINGS
 
 #ifndef PRAGMA_DISABLE_UNDEFINED_IDENTIFIER_WARNINGS
diff --git a/Engine/Source/Runtime/Core/Public/Clang/ClangPlatformMath.h b/Engine/Source/Runtime/Core/Public/Clang/ClangPlatformMath.h
--- a/Engine/Source/Runtime/Core/Public/Clang/ClangPlatformMath.h
+++ b/Engine/Source/Runtime/Core/Public/Clang/ClangPlatformMath.h
@@ -26,7 +26,7 @@ struct FClangPlatformMath : public FGenericPlatformMath
 			return 32;
 		}
 	
-		return __builtin_clz(Value);
+		return (uint32)__builtin_clz(Value);
 	}
 
 	/**
@@ -60,7 +60,7 @@ struct FClangPlatformMath : public FGenericPlatformMath
 			return 32;
 		}
 	
-		return __builtin_ctz(Value);
+		return (uint32)__builtin_ctz(Value);
 	}
 	
 	/**
@@ -77,6 +77,6 @@ struct FClangPlatformMath : public FGenericPlatformMath
 			return 64;
 		}
 	
-		return __builtin_ctzll(Value);
+		return (uint64)__builtin_ctzll(Value);
 	}
 };
diff --git a/Engine/Source/Runtime/Core/Public/Containers/StringConv.h b/Engine/Source/Runtime/Core/Public/Containers/StringConv.h
--- a/Engine/Source/Runtime/Core/Public/Containers/StringConv.h
+++ b/Engine/Source/Runtime/Core/Public/Containers/StringConv.h
@@ -876,13 +876,13 @@ private:
 				return UNICODE_BOGUS_CHAR_CODEPOINT;
 			}
 
-			const uint32 HighSurrogate = Codepoint;
+			const uint16 HighSurrogate = (uint16)Codepoint;
 			Codepoint = *(++CodeUnitPtr);
 
 			// If our High Surrogate is set, check if this character is the matching low-surrogate
 			if (StringConv::IsLowSurrogate(Codepoint))
 			{
-				const uint32 LowSurrogate = Codepoint;
+				const uint16 LowSurrogate = (uint16)Codepoint;
 
 				// Combine our high and low surrogates together to a single Unicode codepoint
 				Codepoint = StringConv::EncodeSurrogate(HighSurrogate, LowSurrogate);
diff --git a/Engine/Source/Runtime/Core/Public/Containers/UnrealString.h b/Engine/Source/Runtime/Core/Public/Containers/UnrealString.h
--- a/Engine/Source/Runtime/Core/Public/Containers/UnrealString.h
+++ b/Engine/Source/Runtime/Core/Public/Containers/UnrealString.h
@@ -2132,7 +2132,7 @@ inline int32 HexToBytes( const FString& HexString, uint8* OutBytes )
 	}
 	while( *CharPos )
 	{
-		OutBytes[ NumBytes ] = TCharToNibble( *CharPos++ ) << 4;
+		OutBytes[ NumBytes ] = (uint8)(TCharToNibble( *CharPos++ ) << 4);
 		OutBytes[ NumBytes ] += TCharToNibble( *CharPos++ );
 		++NumBytes;
 	}
diff --git a/Engine/Source/Runtime/Core/Public/Delegates/DelegateBase.h b/Engine/Source/Runtime/Core/Public/Delegates/DelegateBase.h
--- a/Engine/Source/Runtime/Core/Public/Delegates/DelegateBase.h
+++ b/Engine/Source/Runtime/Core/Public/Delegates/DelegateBase.h
@@ -256,5 +256,5 @@ private:
 
 inline void* operator new(size_t Size, FDelegateBase& Base)
 {
-	return Base.Allocate(Size);
+	return Base.Allocate((int32)Size);
 }
diff --git a/Engine/Source/Runtime/Core/Public/GenericPlatform/StandardPlatformString.h b/Engine/Source/Runtime/Core/Public/GenericPlatform/StandardPlatformString.h
--- a/Engine/Source/Runtime/Core/Public/GenericPlatform/StandardPlatformString.h
+++ b/Engine/Source/Runtime/Core/Public/GenericPlatform/StandardPlatformString.h
@@ -62,7 +62,7 @@ public:
 
 	static FORCEINLINE int32 Strlen( const WIDECHAR* String )
 	{
-		return wcslen( String );
+		return (int32)wcslen( String );
 	}
 
 	static FORCEINLINE const WIDECHAR* Strstr( const WIDECHAR* String, const WIDECHAR* Find)
@@ -82,7 +82,7 @@ public:
 
 	static FORCEINLINE int32 Atoi(const WIDECHAR* String)
 	{
-		return wcstol( String, NULL, 10 );
+		return (int32)wcstol( String, NULL, 10 );
 	}
 
 	static FORCEINLINE int64 Atoi64(const WIDECHAR* String)
@@ -102,7 +102,7 @@ public:
 
 	static FORCEINLINE int32 Strtoi( const WIDECHAR* Start, WIDECHAR** End, int32 Base )
 	{
-		return wcstol( Start, End, Base );
+		return (int32)wcstol( Start, End, Base );
 	}
 
 	static FORCEINLINE int64 Strtoi64( const WIDECHAR* Start, WIDECHAR** End, int32 Base )
@@ -223,7 +223,7 @@ public:
 
 	static FORCEINLINE int32 Strlen( const ANSICHAR* String )
 	{
-		return strlen( String );
+		return (int32)strlen( String );
 	}
 
 	static FORCEINLINE const ANSICHAR* Strstr( const ANSICHAR* String, const ANSICHAR* Find)
@@ -263,7 +263,7 @@ public:
 
 	static FORCEINLINE int32 Strtoi( const ANSICHAR* Start, ANSICHAR** End, int32 Base )
 	{
-		return strtol( Start, End, Base );
+		return (int32)strtol( Start, End, Base );
 	}
 
 	static FORCEINLINE int64 Strtoi64( const ANSICHAR* Start, ANSICHAR** End, int32 Base )
diff --git a/Engine/Source/Runtime/Core/Public/Math/Float16.h b/Engine/Source/Runtime/Core/Public/Math/Float16.h
--- a/Engine/Source/Runtime/Core/Public/Math/Float16.h
+++ b/Engine/Source/Runtime/Core/Public/Math/Float16.h
@@ -150,7 +150,7 @@ FORCEINLINE void FFloat16::Set(float FP32Value)
          if ( (14 - NewExp) <= 24 ) // Mantissa might be non-zero
          {
              uint32 Mantissa = FP32.Components.Mantissa | 0x800000; // Hidden 1 bit
-             Components.Mantissa = Mantissa >> (14 - NewExp);
+             Components.Mantissa = (uint16)(Mantissa >> (14 - NewExp));
 			 // Check for rounding
              if ( (Mantissa >> (13 - NewExp)) & 1 )
 			 {
@@ -168,7 +168,7 @@ FORCEINLINE void FFloat16::Set(float FP32Value)
 	// Handle normal number.
 	else
 	{
-		Components.Exponent = int32(FP32.Components.Exponent) - 127 + 15;
+		Components.Exponent = uint16(int32(FP32.Components.Exponent) - 127 + 15);
 		Components.Mantissa = uint16(FP32.Components.Mantissa >> 13);
 	}
 }
@@ -183,7 +183,7 @@ FORCEINLINE void FFloat16::SetWithoutBoundsChecks(const float FP32Value)
 	// of that just use Set().
 
 	Components.Sign = FP32.Components.Sign;
-	Components.Exponent = int32(FP32.Components.Exponent) - 127 + 15;
+	Components.Exponent = uint16(int32(FP32.Components.Exponent) - 127 + 15);
 	Components.Mantissa = uint16(FP32.Components.Mantissa >> 13);
 }
 
diff --git a/Engine/Source/Runtime/Core/Public/Math/InterpCurvePoint.h b/Engine/Source/Runtime/Core/Public/Math/InterpCurvePoint.h
--- a/Engine/Source/Runtime/Core/Public/Math/InterpCurvePoint.h
+++ b/Engine/Source/Runtime/Core/Public/Math/InterpCurvePoint.h
@@ -206,7 +206,7 @@ inline void ComputeCurveTangent( float PrevTime, const T& PrevPoint,
 
 	AutoCalcTangent( PrevPoint, CurPoint, NextPoint, Tension, OutTangent );
 
-	const float PrevToNextTimeDiff = FMath::Max< double >( KINDA_SMALL_NUMBER, NextTime - PrevTime );
+	const float PrevToNextTimeDiff = FMath::Max< float >( KINDA_SMALL_NUMBER, NextTime - PrevTime );
 
 	OutTangent /= PrevToNextTimeDiff;
 }
diff --git a/Engine/Source/Runtime/Core/Public/Math/UnitConversion.inl b/Engine/Source/Runtime/Core/Public/Math/UnitConversion.inl
--- a/Engine/Source/Runtime/Core/Public/Math/UnitConversion.inl
+++ b/Engine/Source/Runtime/Core/Public/Math/UnitConversion.inl
@@ -195,7 +195,7 @@ EUnit FUnitConversion::CalculateDisplayUnit(T Value, EUnit InUnits)
 		double This = Convert(Value, InUnits, DisplayUnits[Index]);
 		double Next = Convert(Value, InUnits, DisplayUnits[Index + 1]);
 
-		if (FMath::Abs(FMath::LogX(10.0f, This)) < FMath::Abs(FMath::LogX(10.0f, Next)))
+		if (FMath::Abs(FMath::LogX(10.0f, (float)This)) < FMath::Abs(FMath::LogX(10.0f, (float)Next)))
 		{
 			BestIndex = Index;
 		}
diff --git a/Engine/Source/Runtime/Core/Public/Math/UnrealMathUtility.h b/Engine/Source/Runtime/Core/Public/Math/UnrealMathUtility.h
--- a/Engine/Source/Runtime/Core/Public/Math/UnrealMathUtility.h
+++ b/Engine/Source/Runtime/Core/Public/Math/UnrealMathUtility.h
@@ -1625,7 +1625,7 @@ public:
 	static inline bool ExtractBoolFromBitfield(uint8* Ptr, uint32 Index)
 	{
 		uint8* BytePtr = Ptr + Index / 8;
-		uint8 Mask = 1 << (Index & 0x7);
+		uint8 Mask = (uint8)(1 << (Index & 0x7));
 
 		return (*BytePtr & Mask) != 0;
 	}
@@ -1637,7 +1637,7 @@ public:
 	static inline void SetBoolInBitField(uint8* Ptr, uint32 Index, bool bSet)
 	{
 		uint8* BytePtr = Ptr + Index / 8;
-		uint8 Mask = 1 << (Index & 0x7);
+		uint8 Mask = (uint8)(1 << (Index & 0x7));
 
 		if(bSet)
 		{
diff --git a/Engine/Source/Runtime/Core/Public/Misc/MemStack.h b/Engine/Source/Runtime/Core/Public/Misc/MemStack.h
--- a/Engine/Source/Runtime/Core/Public/Misc/MemStack.h
+++ b/Engine/Source/Runtime/Core/Public/Misc/MemStack.h
@@ -229,39 +229,51 @@ template <class T> inline T* NewOned(FMemStackBase& Mem, int32 Count = 1, int32
 inline void* operator new(size_t Size, FMemStackBase& Mem, int32 Count = 1, int32 Align = DEFAULT_ALIGNMENT)
 {
 	// Get uninitialized memory.
-	return Mem.PushBytes( Size*Count, Align );
+	const size_t SizeInBytes = Size * Count;
+	checkSlow(SizeInBytes <= TNumericLimits<int32>::Max());
+	return Mem.PushBytes( (int32)SizeInBytes, Align );
 }
 inline void* operator new(size_t Size, FMemStackBase& Mem, EMemZeroed Tag, int32 Count = 1, int32 Align = DEFAULT_ALIGNMENT)
 {
 	// Get zero-filled memory.
-	uint8* Result = Mem.PushBytes( Size*Count, Align );
-	FMemory::Memzero( Result, Size*Count );
+	const size_t SizeInBytes = Size * Count;
+	checkSlow(SizeInBytes <= TNumericLimits<int32>::Max());
+	uint8* Result = Mem.PushBytes( (int32)SizeInBytes, Align );
+	FMemory::Memzero( Result, SizeInBytes );
 	return Result;
 }
 inline void* operator new(size_t Size, FMemStackBase& Mem, EMemOned Tag, int32 Count = 1, int32 Align = DEFAULT_ALIGNMENT)
 {
 	// Get one-filled memory.
-	uint8* Result = Mem.PushBytes( Size*Count, Align );
-	FMemory::Memset( Result, 0xff, Size*Count );
+	const size_t SizeInBytes = Size * Count;
+	checkSlow(SizeInBytes <= TNumericLimits<int32>::Max());
+	uint8* Result = Mem.PushBytes( (int32)SizeInBytes, Align );
+	FMemory::Memset( Result, 0xff, SizeInBytes );
 	return Result;
 }
 inline void* operator new[](size_t Size, FMemStackBase& Mem, int32 Count = 1, int32 Align = DEFAULT_ALIGNMENT)
 {
 	// Get uninitialized memory.
-	return Mem.PushBytes( Size*Count, Align );
+	const size_t SizeInBytes = Size * Count;
+	checkSlow(SizeInBytes <= TNumericLimits<int32>::Max());
+	return Mem.PushBytes( (int32)SizeInBytes, Align );
 }
 inline void* operator new[](size_t Size, FMemStackBase& Mem, EMemZeroed Tag, int32 Count = 1, int32 Align = DEFAULT_ALIGNMENT)
 {
 	// Get zero-filled memory.
-	uint8* Result = Mem.PushBytes( Size*Count, Align );
-	FMemory::Memzero( Result, Size*Count );
+	const size_t SizeInBytes = Size * Count;
+	checkSlow(SizeInBytes <= TNumericLimits<int32>::Max());
+	uint8* Result = Mem.PushBytes( (int32)SizeInBytes, Align );
+	FMemory::Memzero( Result, SizeInBytes );
 	return Result;
 }
 inline void* operator new[](size_t Size, FMemStackBase& Mem, EMemOned Tag, int32 Count = 1, int32 Align = DEFAULT_ALIGNMENT)
 {
 	// Get one-filled memory.
-	uint8* Result = Mem.PushBytes( Size*Count, Align );
-	FMemory::Memset( Result, 0xff, Size*Count );
+	const size_t SizeInBytes = Size * Count;
+	checkSlow(SizeInBytes <= TNumericLimits<int32>::Max());
+	uint8* Result = Mem.PushBytes( (int32)SizeInBytes, Align );
+	FMemory::Memset( Result, 0xff, SizeInBytes );
 	return Result;
 }
 
diff --git a/Engine/Source/Runtime/Core/Public/ProfilingDebugging/FormatArgsTrace.h b/Engine/Source/Runtime/Core/Public/ProfilingDebugging/FormatArgsTrace.h
--- a/Engine/Source/Runtime/Core/Public/ProfilingDebugging/FormatArgsTrace.h
+++ b/Engine/Source/Runtime/Core/Public/ProfilingDebugging/FormatArgsTrace.h
@@ -121,7 +121,7 @@ private:
 		*TypeCodesPtr++ = FormatArgTypeCode_CategoryString | sizeof(CharType);
 		if (Argument != nullptr)
 		{
-			uint16 Length = (TCString<CharType>::Strlen(Argument) + 1) * sizeof(CharType);
+			uint16 Length = (uint16)((TCString<CharType>::Strlen(Argument) + 1) * sizeof(CharType));
 			memcpy(PayloadPtr, Argument, Length);
 			PayloadPtr += Length;
 		}
diff --git a/Engine/Source/Runtime/Core/Public/ProfilingDebugging/MiscTrace.h b/Engine/Source/Runtime/Core/Public/ProfilingDebugging/MiscTrace.h
--- a/Engine/Source/Runtime/Core/Public/ProfilingDebugging/MiscTrace.h
+++ b/Engine/Source/Runtime/Core/Public/ProfilingDebugging/MiscTrace.h
@@ -28,7 +28,7 @@ struct FTraceUtils
 	{
 		do
 		{
-			uint8 HasMoreBytes = (Value > uint64(0x7F)) << 7;
+			uint8 HasMoreBytes = (uint8)((Value > uint64(0x7F)) << 7);
 			*(BufferPtr++) = (uint8)(Value & 0x7F) | HasMoreBytes;
 			Value >>= 7;
 		} while (Value > 0);
diff --git a/Engine/Source/Runtime/Core/Public/Serialization/BitReader.h b/Engine/Source/Runtime/Core/Public/Serialization/BitReader.h
--- a/Engine/Source/Runtime/Core/Public/Serialization/BitReader.h
+++ b/Engine/Source/Runtime/Core/Public/Serialization/BitReader.h
@@ -219,7 +219,7 @@ private:
 
 	FORCEINLINE uint8 Shift(uint8 Cnt)
 	{
-		return 1<<Cnt;
+		return (uint8)(1<<Cnt);
 	}
 
 };
