Core: Added __has_warning checking for implicit-float-conversion and implicit-int-conversion Clang warnings
#rnx

diff --git a/Engine/Source/Runtime/Core/Public/Clang/ClangPlatformCompilerPreSetup.h b/Engine/Source/Runtime/Core/Public/Clang/ClangPlatformCompilerPreSetup.h
--- a/Engine/Source/Runtime/Core/Public/Clang/ClangPlatformCompilerPreSetup.h
+++ b/Engine/Source/Runtime/Core/Public/Clang/ClangPlatformCompilerPreSetup.h
@@ -73,12 +73,24 @@
 		_Pragma("clang diagnostic pop")
 #endif // PRAGMA_ENABLE_SHADOW_VARIABLE_WARNINGS
 
+#if __has_warning("-Wimplicit-float-conversion")
+#define DISABLE_IMPLICIT_FLOAT_CONVERSION_FRAGMENT _Pragma("clang diagnostic ignored \"-Wimplicit-float-conversion\"")
+#else
+#define DISABLE_IMPLICIT_FLOAT_CONVERSION_FRAGMENT
+#endif
+
+#if __has_warning("-Wimplicit-int-conversion")
+#define DISABLE_IMPLICIT_INT_CONVERSION_FRAGMENT _Pragma("clang diagnostic ignored \"-Wimplicit-int-conversion\"")
+#else
+#define DISABLE_IMPLICIT_INT_CONVERSION_FRAGMENT
+#endif
+
 #ifndef PRAGMA_DISABLE_UNSAFE_TYPECAST_WARNINGS
 	#define PRAGMA_DISABLE_UNSAFE_TYPECAST_WARNINGS \
 		_Pragma("clang diagnostic push") \
 		_Pragma("clang diagnostic ignored \"-Wfloat-conversion\"") \
-		_Pragma("clang diagnostic ignored \"-Wimplicit-float-conversion\"") \
-		_Pragma("clang diagnostic ignored \"-Wimplicit-int-conversion\"") \
+		DISABLE_IMPLICIT_FLOAT_CONVERSION_FRAGMENT \
+		DISABLE_IMPLICIT_INT_CONVERSION_FRAGMENT \
 		_Pragma("clang diagnostic ignored \"-Wc++11-narrowing\"")
 #endif // PRAGMA_DISABLE_UNSAFE_TYPECAST_WARNINGS
 
