Implement module path name for Unix to avoid hard coded default path

Windows implemented at CL 3292685
Mac implemented at CL 1993104

diff --git a/Engine/Source/Runtime/Core/Private/Unix/UnixPlatformProcess.cpp b/Engine/Source/Runtime/Core/Private/Unix/UnixPlatformProcess.cpp
--- a/Engine/Source/Runtime/Core/Private/Unix/UnixPlatformProcess.cpp
+++ b/Engine/Source/Runtime/Core/Private/Unix/UnixPlatformProcess.cpp
@@ -319,6 +319,16 @@ bool FUnixPlatformProcess::SetProcessLimits(EProcessResource::Type Resource, uin
 	return true;
 }
 
+const FString FUnixPlatformProcess::GetModulesDirectory()
+{
+	static FString CachedModulePath;
+	if (CachedModulePath.IsEmpty())
+	{
+		CachedModulePath = FPaths::GetPath(FString(ExecutablePath()));
+	}
+
+	return CachedModulePath;
+}
 
 const TCHAR* FUnixPlatformProcess::ExecutablePath()
 {
diff --git a/Engine/Source/Runtime/Core/Public/Unix/UnixPlatformProcess.h b/Engine/Source/Runtime/Core/Public/Unix/UnixPlatformProcess.h
--- a/Engine/Source/Runtime/Core/Public/Unix/UnixPlatformProcess.h
+++ b/Engine/Source/Runtime/Core/Public/Unix/UnixPlatformProcess.h
@@ -275,6 +275,7 @@ struct CORE_API FUnixPlatformProcess : public FGenericPlatformProcess
 	static bool WritePipe(void* WritePipe, const FString& Message, FString* OutWritten = nullptr);
 	static bool WritePipe(void* WritePipe, const uint8* Data, const int32 DataLength, int32* OutDataLength = nullptr);
 	static class FRunnableThread* CreateRunnableThread();
+	static const FString GetModulesDirectory();
 	static bool CanLaunchURL(const TCHAR* URL);
 	static void LaunchURL(const TCHAR* URL, const TCHAR* Parms, FString* Error);
 	static FProcHandle CreateProc(const TCHAR* URL, const TCHAR* Parms, bool bLaunchDetached, bool bLaunchHidden, bool bLaunchReallyHidden, uint32* OutProcessID, int32 PriorityModifier, const TCHAR* OptionalWorkingDirectory, void* PipeWriteChild, void* PipeReadChild = nullptr);
