Add Algo::Replace

- Algo::Replace can be used to replace FString::ReplaceCharInline on array views of characters such as on a view of a string builder or a string.

diff --git a/Engine/Source/Runtime/Core/Public/Algo/Replace.h b/Engine/Source/Runtime/Core/Public/Algo/Replace.h
new file mode 100644
--- /dev/null
+++ b/Engine/Source/Runtime/Core/Public/Algo/Replace.h
@@ -0,0 +1,47 @@
+// Copyright Epic Games, Inc. All Rights Reserved.
+
+#pragma once
+
+#include "Templates/Invoke.h"
+#include "Templates/UnrealTemplate.h"
+
+namespace Algo
+{
+	/**
+	 * Replaces all elements that compare equal to one value with a new value.
+	 *
+	 * @param Range The range to search and modify in-place.
+	 * @param InOld The value to search for.
+	 * @param InNew The value to copy in as a replacement.
+	 */
+	template <typename RangeType, typename ValueType>
+	FORCEINLINE void Replace(RangeType&& Range, const ValueType& InOld, const ValueType& InNew)
+	{
+		for (auto& Elem : Forward<RangeType>(Range))
+		{
+			if (Elem == InOld)
+			{
+				Elem = InNew;
+			}
+		}
+	}
+
+	/**
+	 * Replaces all elements that satisfy the predicate with the given value.
+	 *
+	 * @param Range The range to search and modify in-place.
+	 * @param InPred The predicate to apply to each element.
+	 * @param InNew The value to copy in as a replacement for each element satisfying the predicate.
+	 */
+	template <typename RangeType, typename ValueType, typename PredicateType>
+	FORCEINLINE void ReplaceIf(RangeType&& Range, PredicateType InPred, const ValueType& InNew)
+	{
+		for (auto& Elem : Forward<RangeType>(Range))
+		{
+			if (Invoke(InPred, Elem))
+			{
+				Elem = InNew;
+			}
+		}
+	}
+}
