Iteration on StringView

- Add string view literals using ""_SV / ""_ASV / ""_WSV.
- Add constexpr to a handful of functions. Adding greater constexpr support will require edits elsewhere.
- Set OutIndex to INDEX_NONE when Find[Last]Char cannot find the character, for consistency with FString.
- Add the new *Inline versions of Left*, Right*, Mid functions, for consistency with FString.
- Move from Misc to Containers for consistency with FString.
- Add TStringView<CharType> alias to facilitate string view usage in code templated on the character type.
- Make FWideStringView an alias for FStringView to facilitate the creation of TStringView. This seems backward but the explicitly wide version is likely rarely used and this still permits easy forward declaration of FStringView.
- Add TrimStart/TrimEnd/TrimStartAndEnd and inline variants for compatibility with FString.
- Remove the CRTP from TStringViewImpl by using TStringView<T> to determine the required view type without it being provided as a template parameter to TStringViewImpl.
- Document the view API.

diff --git a/Engine/Source/Runtime/Core/Private/Containers/StringTest.cpp b/Engine/Source/Runtime/Core/Private/Containers/StringTest.cpp
--- a/Engine/Source/Runtime/Core/Private/Containers/StringTest.cpp
+++ b/Engine/Source/Runtime/Core/Private/Containers/StringTest.cpp
@@ -3,7 +3,7 @@
 #include "CoreTypes.h"
 #include "Misc/AutomationTest.h"
 #include "Misc/AssertionMacros.h"
-#include "Misc/StringView.h"
+#include "Containers/StringView.h"
 #include "Containers/UnrealString.h"
 #include "Serialization/MemoryReader.h"
 #include "Serialization/MemoryWriter.h"
diff --git a/Engine/Source/Runtime/Core/Public/Containers/StringView.inl b/Engine/Source/Runtime/Core/Public/Containers/StringView.inl
--- a/Engine/Source/Runtime/Core/Public/Containers/StringView.inl
+++ b/Engine/Source/Runtime/Core/Public/Containers/StringView.inl
@@ -16,7 +16,7 @@ template <typename CharType>
 inline typename TStringViewImpl<CharType>::SizeType TStringViewImpl<CharType>::CopyString(CharType* Dest, SizeType CharCount, SizeType Position) const
 {
 	const  SizeType CopyCount = FMath::Min(Size - Position, CharCount);
-	FMemory::Memcpy(Dest, DataPtr + Position, CopyCount);
+	FMemory::Memcpy(Dest, DataPtr + Position, CopyCount * sizeof(CharType));
 	return CopyCount;
 }
 
diff --git a/Engine/Source/Runtime/Core/Public/IO/IoDispatcher.h b/Engine/Source/Runtime/Core/Public/IO/IoDispatcher.h
--- a/Engine/Source/Runtime/Core/Public/IO/IoDispatcher.h
+++ b/Engine/Source/Runtime/Core/Public/IO/IoDispatcher.h
@@ -3,9 +3,9 @@
 #pragma once
 
 #include "CoreTypes.h"
+#include "Containers/StringView.h"
 #include "Containers/UnrealString.h"
 #include "Logging/LogMacros.h"
-#include "Misc/StringView.h"
 #include "Misc/StringBuilder.h"
 #include "Templates/RefCounting.h"
 #include "Templates/UnrealTemplate.h"
diff --git a/Engine/Source/Runtime/Core/Public/Misc/StringView.h b/Engine/Source/Runtime/Core/Public/Misc/StringView.h
--- a/Engine/Source/Runtime/Core/Public/Misc/StringView.h
+++ b/Engine/Source/Runtime/Core/Public/Misc/StringView.h
@@ -2,10 +2,10 @@
 
 #pragma once
 
-//#ifdef _MSC_VER
-//#pragma message(__FILE__"(6): warning : #include Containers/StringView.h instead of Misc/StringView.h")
-//#else
-//#pragma message("#include Containers/StringView.h instead of Misc/StringView.h")
-//#endif
+#ifdef _MSC_VER
+#pragma message(__FILE__"(6): warning : #include Containers/StringView.h instead of Misc/StringView.h")
+#else
+#pragma message("#include Containers/StringView.h instead of Misc/StringView.h")
+#endif
 
 #include "Containers/StringView.h"
