Change Algo::Accumulate to use Invoke

#rnx

diff --git a/Engine/Source/Runtime/Core/Public/Algo/Accumulate.h b/Engine/Source/Runtime/Core/Public/Algo/Accumulate.h
--- a/Engine/Source/Runtime/Core/Public/Algo/Accumulate.h
+++ b/Engine/Source/Runtime/Core/Public/Algo/Accumulate.h
@@ -3,6 +3,7 @@
 #pragma once
 
 #include "CoreTypes.h"
+#include "Templates/Invoke.h"
 
 // TPlus<T> specifically takes const T& and returns T.
 // TPlus<> (empty angle brackets) is late-binding, taking whatever is passed and returning the correct result type for (A+B)
@@ -36,7 +37,7 @@ namespace Algo
 		T Result = MoveTemp(Init);
 		for (const auto& InputElem : Input)
 		{
-			Result = Op(MoveTemp(Result), InputElem);
+			Result = Invoke(Op, MoveTemp(Result), InputElem);
 		}
 		return Result;
 	}
@@ -71,7 +72,7 @@ namespace Algo
 		T Result = MoveTemp(Init);
 		for (const auto& InputElem : Input)
 		{
-			Result = Op(MoveTemp(Result), MapOp(InputElem));
+			Result = Invoke(Op, MoveTemp(Result), Invoke(MapOp, InputElem));
 		}
 		return Result;
 	}
