Add UE::String::ParseLines, UE::String::ParseTokens[Multiple]

- UE::String::ParseLines can be used to split a string view on line endings to find the non-empty lines in the view. Intended to replace functions like FString::ParseIntoArrayLines and FFileHelper::LoadFileToStringArray.
- UE::String::ParseTokens[Multiple] can be used to split a string view on one or more single-character or multi-character delimiters. Intended to replace FString::ParseIntoArray. Separate handling of single-character delimiters reduced the cost by 27% in the case being profiled.
- Extend Algo::CompareByPredicate to allow the input ranges to be different types and to take them by universal reference.
- Extend TArrayView and MakeArrayView to allow array views to be created from initializer lists.

diff --git a/Engine/Source/Runtime/Core/Private/String/ParseLines.cpp b/Engine/Source/Runtime/Core/Private/String/ParseLines.cpp
new file mode 100644
--- /dev/null
+++ b/Engine/Source/Runtime/Core/Private/String/ParseLines.cpp
@@ -0,0 +1,43 @@
+// Copyright Epic Games, Inc. All Rights Reserved.
+
+#include "String/ParseLines.h"
+
+#include "Containers/StringView.h"
+
+namespace UE
+{
+namespace String
+{
+	void ParseLines(const FStringView& View, const TFunctionRef<void(FStringView)>& Visitor)
+	{
+		const TCHAR* ViewIt = View.GetData();
+		const TCHAR* ViewEnd = ViewIt + View.Len();
+		do
+		{
+			const TCHAR* LineStart = ViewIt;
+			const TCHAR* LineEnd = ViewEnd;
+			for (; ViewIt != ViewEnd; ++ViewIt)
+			{
+				const TCHAR CurrentChar = *ViewIt;
+				if (CurrentChar == TEXT('\n'))
+				{
+					LineEnd = ViewIt++;
+					break;
+				}
+				if (CurrentChar == TEXT('\r'))
+				{
+					LineEnd = ViewIt++;
+					if (ViewIt != ViewEnd && *ViewIt == TEXT('\n'))
+					{
+						++ViewIt;
+					}
+					break;
+				}
+			}
+
+			Visitor(FStringView(LineStart, static_cast<FStringView::SizeType>(LineEnd - LineStart)));
+		}
+		while (ViewIt != ViewEnd);
+	}
+}
+}
diff --git a/Engine/Source/Runtime/Core/Private/String/ParseTokens.cpp b/Engine/Source/Runtime/Core/Private/String/ParseTokens.cpp
new file mode 100644
--- /dev/null
+++ b/Engine/Source/Runtime/Core/Private/String/ParseTokens.cpp
@@ -0,0 +1,175 @@
+// Copyright Epic Games, Inc. All Rights Reserved.
+
+#include "String/ParseTokens.h"
+
+#include "Algo/AllOf.h"
+#include "Algo/Find.h"
+#include "Algo/NoneOf.h"
+#include "Containers/BitArray.h"
+#include "Containers/StringView.h"
+
+namespace UE
+{
+namespace String
+{
+	/** Parse tokens with one single-character delimiter. */
+	inline static void ParseTokens1Delim1Char(const FStringView& View, const TCHAR Delimiter, TFunctionRef<void(FStringView)> Visitor)
+	{
+		const TCHAR* ViewIt = View.GetData();
+		const TCHAR* const ViewEnd = ViewIt + View.Len();
+		const TCHAR* NextToken = ViewIt;
+
+		for (;;)
+		{
+			if (ViewIt == ViewEnd)
+			{
+				break;
+			}
+			if (*ViewIt != Delimiter)
+			{
+				++ViewIt;
+				continue;
+			}
+			Visitor(FStringView(NextToken, static_cast<FStringView::SizeType>(ViewIt - NextToken)));
+			NextToken = ++ViewIt;
+		}
+
+		Visitor(FStringView(NextToken, static_cast<FStringView::SizeType>(ViewIt - NextToken)));
+	}
+
+	/** Parse tokens with multiple single-character Basic Latin delimiters. */
+	inline static void ParseTokensNDelim1CharBasicLatin(const FStringView& View, TArrayView<const TCHAR> Delimiters, TFunctionRef<void(FStringView)> Visitor)
+	{
+		TBitArray<> DelimiterMask(false, 128);
+		for (uint32 Delimiter : Delimiters)
+		{
+			DelimiterMask[Delimiter] = true;
+		}
+
+		const TCHAR* ViewIt = View.GetData();
+		const TCHAR* const ViewEnd = ViewIt + View.Len();
+		const TCHAR* NextToken = ViewIt;
+
+		for (;;)
+		{
+			if (ViewIt == ViewEnd)
+			{
+				break;
+			}
+			const uint32 CodePoint = *ViewIt;
+			if (CodePoint >= 128 || !DelimiterMask[CodePoint])
+			{
+				++ViewIt;
+				continue;
+			}
+			Visitor(FStringView(NextToken, static_cast<FStringView::SizeType>(ViewIt - NextToken)));
+			NextToken = ++ViewIt;
+		}
+
+		Visitor(FStringView(NextToken, static_cast<FStringView::SizeType>(ViewIt - NextToken)));
+	}
+
+	/** Parse tokens with multiple single-character delimiters in the Basic Multilingual Plane. */
+	inline static void ParseTokensNDelim1Char(const FStringView& View, TArrayView<const TCHAR> Delimiters, TFunctionRef<void(FStringView)> Visitor)
+	{
+		if (Algo::AllOf(Delimiters, [](TCHAR Delimiter) { return Delimiter < 128; }))
+		{
+			return ParseTokensNDelim1CharBasicLatin(View, Delimiters, Visitor);
+		}
+
+		const TCHAR* ViewIt = View.GetData();
+		const TCHAR* const ViewEnd = ViewIt + View.Len();
+		const TCHAR* NextToken = ViewIt;
+
+		for (;;)
+		{
+			if (ViewIt == ViewEnd)
+			{
+				break;
+			}
+			if (!Algo::Find(Delimiters, *ViewIt))
+			{
+				++ViewIt;
+				continue;
+			}
+			Visitor(FStringView(NextToken, static_cast<FStringView::SizeType>(ViewIt - NextToken)));
+			NextToken = ++ViewIt;
+		}
+
+		Visitor(FStringView(NextToken, static_cast<FStringView::SizeType>(ViewIt - NextToken)));
+	}
+
+	/** Parse tokens with multiple multi-character delimiters. */
+	inline static void ParseTokensNDelimNChar(const FStringView& View, TArrayView<const FStringView> Delimiters, TFunctionRef<void(FStringView)> Visitor)
+	{
+		// This is a naive implementation that takes time proportional to View.Len() * TotalDelimiterLen.
+		// If this function becomes a bottleneck, it can be specialized separately for one and many delimiters.
+		// There are algorithms for each are linear or sub-linear in the length of string to search.
+
+		const FStringView::SizeType ViewLen = View.Len();
+		FStringView::SizeType NextTokenIndex = 0;
+
+		for (FStringView::SizeType ViewIndex = 0; ViewIndex != ViewLen;)
+		{
+			const FStringView RemainingView(View.GetData() + ViewIndex, ViewLen - ViewIndex);
+			auto MatchDelimiter = [&RemainingView](const FStringView& Delimiter) { return RemainingView.StartsWith(Delimiter, ESearchCase::CaseSensitive); };
+			if (const FStringView* Delimiter = Algo::FindByPredicate(Delimiters, MatchDelimiter))
+			{
+				Visitor(FStringView(View.GetData() + NextTokenIndex, ViewIndex - NextTokenIndex));
+				ViewIndex += Delimiter->Len();
+				NextTokenIndex = ViewIndex;
+			}
+			else
+			{
+				++ViewIndex;
+			}
+		}
+
+		Visitor(FStringView(View.GetData() + NextTokenIndex, ViewLen - NextTokenIndex));
+	}
+
+	void ParseTokensMultiple(const FStringView& View, TArrayView<const FStringView> Delimiters, TFunctionRef<void(FStringView)> Visitor)
+	{
+		check(Algo::NoneOf(Delimiters, &FStringView::IsEmpty));
+		switch (Delimiters.Num())
+		{
+		case 0:
+			return Visitor(View);
+		case 1:
+			if (Delimiters[0].Len() == 1)
+			{
+				return ParseTokens1Delim1Char(View, Delimiters[0][0], MoveTemp(Visitor));
+			}
+			return ParseTokensNDelimNChar(View, Delimiters, MoveTemp(Visitor));
+		default:
+			if (Algo::AllOf(Delimiters, [](const FStringView& Delimiter) { return Delimiter.Len() == 1; }))
+			{
+				TArray<TCHAR, TInlineAllocator<32>> DelimiterChars;
+				DelimiterChars.Reserve(Delimiters.Num());
+				for (const FStringView& Delimiter : Delimiters)
+				{
+					DelimiterChars.Add(Delimiter[0]);
+				}
+				return ParseTokensNDelim1Char(View, DelimiterChars, MoveTemp(Visitor));
+			}
+			else
+			{
+				return ParseTokensNDelimNChar(View, Delimiters, MoveTemp(Visitor));
+			}
+		}
+	}
+
+	void ParseTokensMultiple(const FStringView& View, TArrayView<const TCHAR> Delimiters, TFunctionRef<void(FStringView)> Visitor)
+	{
+		switch (Delimiters.Num())
+		{
+		case 0:
+			return Visitor(View);
+		case 1:
+			return ParseTokens1Delim1Char(View, Delimiters[0], MoveTemp(Visitor));
+		default:
+			return ParseTokensNDelim1Char(View, Delimiters, MoveTemp(Visitor));
+		}
+	}
+}
+}
diff --git a/Engine/Source/Runtime/Core/Private/Tests/String/ParseLinesTest.cpp b/Engine/Source/Runtime/Core/Private/Tests/String/ParseLinesTest.cpp
new file mode 100644
--- /dev/null
+++ b/Engine/Source/Runtime/Core/Private/Tests/String/ParseLinesTest.cpp
@@ -0,0 +1,54 @@
+// Copyright Epic Games, Inc. All Rights Reserved.
+
+#include "String/ParseLines.h"
+
+#include "Algo/Compare.h"
+#include "Containers/StringView.h"
+#include "Misc/AutomationTest.h"
+#include "Misc/StringBuilder.h"
+#include "Templates/EqualTo.h"
+
+#if WITH_DEV_AUTOMATION_TESTS
+
+IMPLEMENT_SIMPLE_AUTOMATION_TEST(StringParseLinesTest, "System.Core.String.ParseLines", EAutomationTestFlags::ApplicationContextMask | EAutomationTestFlags::EngineFilter)
+bool StringParseLinesTest::RunTest(const FString& Parameters)
+{
+	auto RunParseLinesTest = [this](const FStringView& View, std::initializer_list<FStringView> ExpectedLines)
+	{
+		TArray<FStringView, TInlineAllocator<8>> ResultLines;
+		UE::String::ParseLines(View, [&ResultLines](FStringView Line) { ResultLines.Add(Line); });
+		if (!Algo::CompareByPredicate(ResultLines, ExpectedLines, TEqualTo<>()))
+		{
+			TStringBuilder<512> Error;
+			Error << TEXT("UE::String::ParseLines failed to parse \"") << FString(View).ReplaceCharWithEscapedChar() << TEXT("\" result {");
+			Error.JoinQuoted(ResultLines, TEXT(", "_SV), TEXT("\""_SV));
+			Error << TEXT("} expected {");
+			Error.JoinQuoted(ExpectedLines, TEXT(", "_SV), TEXT("\""_SV));
+			Error << TEXT("}");
+			AddError(Error.ToString());
+		}
+	};
+
+	RunParseLinesTest(TEXT(""_SV),                        {TEXT(""_SV)});
+	RunParseLinesTest(TEXT("\n"_SV),                      {TEXT(""_SV)});
+	RunParseLinesTest(TEXT("\r"_SV),                      {TEXT(""_SV)});
+	RunParseLinesTest(TEXT("\r\n"_SV),                    {TEXT(""_SV)});
+	RunParseLinesTest(TEXT("\n\n"_SV),                    {TEXT(""_SV), TEXT(""_SV)});
+	RunParseLinesTest(TEXT("\r\r"_SV),                    {TEXT(""_SV), TEXT(""_SV)});
+	RunParseLinesTest(TEXT("\r\n\r\n"_SV),                {TEXT(""_SV), TEXT(""_SV)});
+	RunParseLinesTest(TEXT("\r\nABC"_SV).Left(2),         {TEXT(""_SV)});
+	RunParseLinesTest(TEXT("\r\nABC\r\nDEF"_SV).Left(5),  {TEXT(""_SV), TEXT("ABC"_SV)});
+	RunParseLinesTest(TEXT("ABC DEF"_SV),                 {TEXT("ABC DEF"_SV)});
+	RunParseLinesTest(TEXT("\nABC DEF\n"_SV),             {TEXT(""_SV), TEXT("ABC DEF"_SV)});
+	RunParseLinesTest(TEXT("\rABC DEF\r"_SV),             {TEXT(""_SV), TEXT("ABC DEF"_SV)});
+	RunParseLinesTest(TEXT("\r\nABC DEF\r\n"_SV),         {TEXT(""_SV), TEXT("ABC DEF"_SV)});
+	RunParseLinesTest(TEXT("\r\n\r\nABC DEF\r\n\r\n"_SV), {TEXT(""_SV), TEXT(""_SV), TEXT("ABC DEF"_SV), TEXT(""_SV)});
+	RunParseLinesTest(TEXT("ABC\nDEF"_SV),                {TEXT("ABC"_SV), TEXT("DEF"_SV)});
+	RunParseLinesTest(TEXT("ABC\rDEF"_SV),                {TEXT("ABC"_SV), TEXT("DEF"_SV)});
+	RunParseLinesTest(TEXT("\r\nABC\r\nDEF\r\n"_SV),      {TEXT(""_SV), TEXT("ABC"_SV), TEXT("DEF"_SV)});
+	RunParseLinesTest(TEXT("\r\nABC\r\n\r\nDEF\r\n"_SV),  {TEXT(""_SV), TEXT("ABC"_SV), TEXT(""_SV), TEXT("DEF"_SV)});
+
+	return true;
+}
+
+#endif // WITH_DEV_AUTOMATION_TESTS
diff --git a/Engine/Source/Runtime/Core/Private/Tests/String/ParseTokensTest.cpp b/Engine/Source/Runtime/Core/Private/Tests/String/ParseTokensTest.cpp
new file mode 100644
--- /dev/null
+++ b/Engine/Source/Runtime/Core/Private/Tests/String/ParseTokensTest.cpp
@@ -0,0 +1,98 @@
+// Copyright Epic Games, Inc. All Rights Reserved.
+
+#include "String/ParseTokens.h"
+
+#include "Algo/Compare.h"
+#include "Misc/AutomationTest.h"
+#include "Misc/StringBuilder.h"
+#include "Templates/EqualTo.h"
+
+#if WITH_DEV_AUTOMATION_TESTS
+
+IMPLEMENT_SIMPLE_AUTOMATION_TEST(StringParseTokensByStringTest, "System.Core.String.ParseTokens.ByString", EAutomationTestFlags::ApplicationContextMask | EAutomationTestFlags::EngineFilter)
+bool StringParseTokensByStringTest::RunTest(const FString& Parameters)
+{
+	auto RunParseTokensTest = [this](const FStringView& View, std::initializer_list<FStringView> Delimiters, std::initializer_list<FStringView> ExpectedTokens)
+	{
+		TArray<FStringView, TInlineAllocator<8>> ResultTokens;
+		UE::String::ParseTokensMultiple(View, Delimiters, [&ResultTokens](FStringView Token) { ResultTokens.Add(Token); });
+		if (!Algo::CompareByPredicate(ResultTokens, ExpectedTokens, TEqualTo<>()))
+		{
+			TStringBuilder<512> Error;
+			Error << TEXT("UE::String::ParseTokensMultiple failed to parse \"") << View << TEXT("\" with delimiters {");
+			Error.JoinQuoted(Delimiters, TEXT(", "_SV), TEXT("\""_SV));
+			Error << TEXT("} result {");
+			Error.JoinQuoted(ResultTokens, TEXT(", "_SV), TEXT("\""_SV));
+			Error << TEXT("} expected {");
+			Error.JoinQuoted(ExpectedTokens, TEXT(", "_SV), TEXT("\""_SV));
+			Error << TEXT("}");
+			AddError(Error.ToString());
+		}
+	};
+
+	RunParseTokensTest(TEXT(""),         {},                       {TEXT("")});
+	RunParseTokensTest(TEXT("ABC"),      {},                       {TEXT("ABC")});
+
+	RunParseTokensTest(TEXT(""),         {TEXT(",")},              {TEXT("")});
+	RunParseTokensTest(TEXT(","),        {TEXT(",")},              {TEXT(""), TEXT("")});
+	RunParseTokensTest(TEXT(",,"),       {TEXT(",")},              {TEXT(""), TEXT(""), TEXT("")});
+	RunParseTokensTest(TEXT("ABC"),      {TEXT(",")},              {TEXT("ABC")});
+	RunParseTokensTest(TEXT("A,,C"),     {TEXT(",")},              {TEXT("A"), TEXT(""), TEXT("C")});
+	RunParseTokensTest(TEXT("A,B,C"),    {TEXT(",")},              {TEXT("A"), TEXT("B"), TEXT("C")});
+	RunParseTokensTest(TEXT(",A,B,C,"),  {TEXT(",")},              {TEXT(""), TEXT("A"), TEXT("B"), TEXT("C"), TEXT("")});
+	RunParseTokensTest(TEXT("A\u2022B\u2022C"), {TEXT("\u2022")},  {TEXT("A"), TEXT("B"), TEXT("C")});
+
+	RunParseTokensTest(TEXT("ABCDABCD"), {TEXT("AB")},             {TEXT(""), TEXT("CD"), TEXT("CD")});
+	RunParseTokensTest(TEXT("ABCDABCD"), {TEXT("ABCD")},           {TEXT(""), TEXT(""), TEXT("")});
+	RunParseTokensTest(TEXT("ABCDABCD"), {TEXT("DA")},             {TEXT("ABC"), TEXT("BCD")});
+
+	RunParseTokensTest(TEXT("ABCDABCD"), {TEXT("B"),  TEXT("D")},  {TEXT("A"), TEXT("C"), TEXT("A"), TEXT("C"), TEXT("")});
+	RunParseTokensTest(TEXT("ABCDABCD"), {TEXT("BC"), TEXT("DA")}, {TEXT("A"), TEXT(""), TEXT(""), TEXT("D")});
+
+	RunParseTokensTest(TEXT("A\u2022\u2022B,,C"), {TEXT(",,"), TEXT("\u2022\u2022")},                      {TEXT("A"), TEXT("B"), TEXT("C")});
+	RunParseTokensTest(TEXT("A\u2022\u2022B\u0085\u0085C"), {TEXT("\u0085\u0085"), TEXT("\u2022\u2022")},  {TEXT("A"), TEXT("B"), TEXT("C")});
+
+	return true;
+}
+
+IMPLEMENT_SIMPLE_AUTOMATION_TEST(StringParseTokensByCharTest, "System.Core.String.ParseTokens.ByChar", EAutomationTestFlags::ApplicationContextMask | EAutomationTestFlags::EngineFilter)
+bool StringParseTokensByCharTest::RunTest(const FString& Parameters)
+{
+	auto RunParseTokensTest = [this](const FStringView& View, std::initializer_list<TCHAR> Delimiters, std::initializer_list<FStringView> ExpectedTokens)
+	{
+		TArray<FStringView, TInlineAllocator<8>> ResultTokens;
+		UE::String::ParseTokensMultiple(View, Delimiters, [&ResultTokens](FStringView Token) { ResultTokens.Add(Token); });
+		if (!Algo::CompareByPredicate(ResultTokens, ExpectedTokens, TEqualTo<>()))
+		{
+			TStringBuilder<512> Error;
+			Error << TEXT("UE::String::ParseTokensMultiple failed to parse \"") << View << TEXT("\" with delimiters {");
+			Error.JoinQuoted(Delimiters, TEXT(", "_SV), TEXT("'"_SV));
+			Error << TEXT("} result {");
+			Error.JoinQuoted(ResultTokens, TEXT(", "_SV), TEXT("\""_SV));
+			Error << TEXT("} expected {");
+			Error.JoinQuoted(ExpectedTokens, TEXT(", "_SV), TEXT("\""_SV));
+			Error << TEXT("}");
+			AddError(Error.ToString());
+		}
+	};
+
+	RunParseTokensTest(TEXT(""),         {},                       {TEXT("")});
+	RunParseTokensTest(TEXT("ABC"),      {},                       {TEXT("ABC")});
+
+	RunParseTokensTest(TEXT(""),         {TEXT(',')},              {TEXT("")});
+	RunParseTokensTest(TEXT(","),        {TEXT(',')},              {TEXT(""), TEXT("")});
+	RunParseTokensTest(TEXT(",,"),       {TEXT(',')},              {TEXT(""), TEXT(""), TEXT("")});
+	RunParseTokensTest(TEXT("ABC"),      {TEXT(',')},              {TEXT("ABC")});
+	RunParseTokensTest(TEXT("A,,C"),     {TEXT(',')},              {TEXT("A"), TEXT(""), TEXT("C")});
+	RunParseTokensTest(TEXT("A,B,C"),    {TEXT(',')},              {TEXT("A"), TEXT("B"), TEXT("C")});
+	RunParseTokensTest(TEXT(",A,B,C,"),  {TEXT(',')},              {TEXT(""), TEXT("A"), TEXT("B"), TEXT("C"), TEXT("")});
+	RunParseTokensTest(TEXT("A\u2022B\u2022C"), {TEXT('\u2022')},  {TEXT("A"), TEXT("B"), TEXT("C")});
+
+	RunParseTokensTest(TEXT("ABCDABCD"), {TEXT('B'),  TEXT('D')},                  {TEXT("A"), TEXT("C"), TEXT("A"), TEXT("C"), TEXT("")});
+	RunParseTokensTest(TEXT("A\u2022B,C"), {TEXT(','), TEXT('\u2022')},            {TEXT("A"), TEXT("B"), TEXT("C")});
+	RunParseTokensTest(TEXT("A\u2022B\u0085C"), {TEXT('\u0085'), TEXT('\u2022')},  {TEXT("A"), TEXT("B"), TEXT("C")});
+
+	return true;
+}
+
+#endif // WITH_DEV_AUTOMATION_TESTS
diff --git a/Engine/Source/Runtime/Core/Public/Algo/Compare.h b/Engine/Source/Runtime/Core/Public/Algo/Compare.h
--- a/Engine/Source/Runtime/Core/Public/Algo/Compare.h
+++ b/Engine/Source/Runtime/Core/Public/Algo/Compare.h
@@ -8,38 +8,36 @@
 namespace Algo
 {
 	/**
-	* Compares entries in the container using a user-defined predicate container
+	* Compares two contiguous containers using a predicate to compare pairs of elements.
 	*
-	* @param  InputA     Container used as baseline
-	* @param  InputB     Container to compare against
-	* @param  Predicate  Condition which returns true for elements which are deemed equal
+	* @param  InputA     Container of elements that are used as the first argument to the predicate.
+	* @param  InputB     Container of elements that are used as the second argument to the predicate.
+	* @param  Predicate  Condition which returns true for elements which are deemed equal.
+	*
+	* @return Whether the containers are the same size and the predicate returned true for every pair of elements.
 	*/
-	template <typename InT, typename PredicateT>
-	FORCEINLINE bool CompareByPredicate(const InT& InputA, const InT& InputB, PredicateT Predicate)
+	template <typename InAT, typename InBT, typename PredicateT>
+	constexpr bool CompareByPredicate(InAT&& InputA, InBT&& InputB, PredicateT Predicate)
 	{
-		auto Count = GetNum(InputA);
-		auto CountB = GetNum(InputB);
+		const SIZE_T SizeA = GetNum(InputA);
+		const SIZE_T SizeB = GetNum(InputB);
 
-		if (Count == CountB)
+		if (SizeA != SizeB)
 		{
-			auto* A = GetData(InputA);
-			auto* B = GetData(InputB);
+			return false;
+		}
 
-			while (Count)
-			{
-				if (!(Invoke(Predicate, *A, *B)))
-				{
-					return false;
-				}
+		auto* A = GetData(InputA);
+		auto* B = GetData(InputB);
 
-				++A;
-				++B;
-				--Count;
+		for (SIZE_T Count = SizeA; Count; --Count)
+		{
+			if (!Invoke(Predicate, *A++, *B++))
+			{
+				return false;
 			}
-
-			return true;
 		}
 
-		return false;
+		return true;
 	}
 }
diff --git a/Engine/Source/Runtime/Core/Public/Containers/ArrayView.h b/Engine/Source/Runtime/Core/Public/Containers/ArrayView.h
--- a/Engine/Source/Runtime/Core/Public/Containers/ArrayView.h
+++ b/Engine/Source/Runtime/Core/Public/Containers/ArrayView.h
@@ -142,6 +142,17 @@ public:
 		check(ArrayNum >= 0);
 	}
 
+	/**
+	 * Construct a view of an initializer list.
+	 *
+	 * The caller is responsible for ensuring that the view does not outlive the initializer list.
+	 */
+	FORCEINLINE TArrayView(std::initializer_list<ElementType> List)
+		: DataPtr(ArrayViewPrivate::GetDataHelper(List))
+		, ArrayNum(GetNum(List))
+	{
+	}
+
 public:
 
 	/**
@@ -558,3 +569,9 @@ auto MakeArrayView(ElementType* Pointer, int32 Size)
 {
 	return TArrayView<ElementType>(Pointer, Size);
 }
+
+template <typename T>
+TArrayView<const T> MakeArrayView(std::initializer_list<T> List)
+{
+	return TArrayView<const T>(List);
+}
diff --git a/Engine/Source/Runtime/Core/Public/String/ParseLines.h b/Engine/Source/Runtime/Core/Public/String/ParseLines.h
new file mode 100644
--- /dev/null
+++ b/Engine/Source/Runtime/Core/Public/String/ParseLines.h
@@ -0,0 +1,24 @@
+// Copyright Epic Games, Inc. All Rights Reserved.
+
+#pragma once
+
+#include "CoreTypes.h"
+#include "Templates/Function.h"
+
+class FStringView;
+
+namespace UE
+{
+namespace String
+{
+	/**
+	 * Visit every line in the input view as terminated by any of CRLF, CR, LF.
+	 *
+	 * Empty lines are visited by this function.
+	 *
+	 * @param View The string view to split into lines.
+	 * @param Visitor A function that is called for each line.
+	 */
+	CORE_API void ParseLines(const FStringView& View, const TFunctionRef<void(FStringView)>& Visitor);
+}
+}
diff --git a/Engine/Source/Runtime/Core/Public/String/ParseTokens.h b/Engine/Source/Runtime/Core/Public/String/ParseTokens.h
new file mode 100644
--- /dev/null
+++ b/Engine/Source/Runtime/Core/Public/String/ParseTokens.h
@@ -0,0 +1,67 @@
+// Copyright Epic Games, Inc. All Rights Reserved.
+
+#pragma once
+
+#include "CoreTypes.h"
+#include "Containers/ArrayView.h"
+#include "Templates/Function.h"
+
+class FStringView;
+
+namespace UE
+{
+namespace String
+{
+	/**
+	 * Visit every token in the input view as separated by any of the delimiters.
+	 *
+	 * Comparisons with delimiters are case-sensitive and empty tokens are visited.
+	 * Behavior is undefined when delimiters overlap each other, such as the delimiters
+	 * ("AB, "BC") and the input view "1ABC2".
+	 *
+	 * @param View The string view to split into tokens.
+	 * @param Delimiters An array of non-overlapping non-empty delimiters to split the input string on.
+	 * @param Visitor A function that is called for each token.
+	 */
+	CORE_API void ParseTokensMultiple(const FStringView& View, TArrayView<const FStringView> Delimiters, TFunctionRef<void(FStringView)> Visitor);
+
+	/**
+	 * Visit every token in the input view as separated by any of the delimiters.
+	 *
+	 * Comparisons with delimiters are case-sensitive and empty tokens are visited.
+	 *
+	 * @param View The string view to split into tokens.
+	 * @param Delimiters An array of delimiter characters to split the input string on.
+	 * @param Visitor A function that is called for each token.
+	 */
+	CORE_API void ParseTokensMultiple(const FStringView& View, TArrayView<const TCHAR> Delimiters, TFunctionRef<void(FStringView)> Visitor);
+
+	/**
+	 * Visit every token in the input view as separated by the delimiter.
+	 *
+	 * Comparisons with the delimiter are case-sensitive and empty tokens are visited.
+	 *
+	 * @param View The string view to split into tokens.
+	 * @param Delimiter The non-empty delimiter to split the input string on.
+	 * @param Visitor A function that is called for each token.
+	 */
+	inline void ParseTokens(const FStringView& View, const FStringView& Delimiter, TFunctionRef<void(FStringView)> Visitor)
+	{
+		return ParseTokensMultiple(View, MakeArrayView(&Delimiter, 1), Visitor);
+	}
+
+	/**
+	 * Visit every token in the input view as separated by the delimiter.
+	 *
+	 * Comparisons with the delimiter are case-sensitive and empty tokens are visited.
+	 *
+	 * @param View The string view to split into tokens.
+	 * @param Delimiter The delimiter character to split the input string on.
+	 * @param Visitor A function that is called for each token.
+	 */
+	inline void ParseTokens(const FStringView& View, TCHAR Delimiter, TFunctionRef<void(FStringView)> Visitor)
+	{
+		return ParseTokensMultiple(View, MakeArrayView(&Delimiter, 1), Visitor);
+	}
+}
+}
