Implement LexFromString from FStringView for primitive types

These implementations are not optimized yet. They exist to allow conversion from FString to FStringView.

diff --git a/Engine/Source/Runtime/Core/Private/String/LexFromString.cpp b/Engine/Source/Runtime/Core/Private/String/LexFromString.cpp
new file mode 100644
--- /dev/null
+++ b/Engine/Source/Runtime/Core/Private/String/LexFromString.cpp
@@ -0,0 +1,73 @@
+// Copyright Epic Games, Inc. All Rights Reserved.
+
+#include "String/LexFromString.h"
+
+#include "Containers/StringView.h"
+#include "Containers/UnrealString.h"
+#include "Misc/StringBuilder.h"
+
+void LexFromString(int8& OutValue, const FStringView& InString)
+{
+	TStringBuilder<64> Builder;
+	LexFromString(OutValue, (Builder << InString).ToString());
+}
+
+void LexFromString(int16& OutValue, const FStringView& InString)
+{
+	TStringBuilder<64> Builder;
+	LexFromString(OutValue, (Builder << InString).ToString());
+}
+
+void LexFromString(int32& OutValue, const FStringView& InString)
+{
+	TStringBuilder<64> Builder;
+	LexFromString(OutValue, (Builder << InString).ToString());
+}
+
+void LexFromString(int64& OutValue, const FStringView& InString)
+{
+	TStringBuilder<64> Builder;
+	LexFromString(OutValue, (Builder << InString).ToString());
+}
+
+void LexFromString(uint8& OutValue, const FStringView& InString)
+{
+	TStringBuilder<64> Builder;
+	LexFromString(OutValue, (Builder << InString).ToString());
+}
+
+void LexFromString(uint16& OutValue, const FStringView& InString)
+{
+	TStringBuilder<64> Builder;
+	LexFromString(OutValue, (Builder << InString).ToString());
+}
+
+void LexFromString(uint32& OutValue, const FStringView& InString)
+{
+	TStringBuilder<64> Builder;
+	LexFromString(OutValue, (Builder << InString).ToString());
+}
+
+void LexFromString(uint64& OutValue, const FStringView& InString)
+{
+	TStringBuilder<64> Builder;
+	LexFromString(OutValue, (Builder << InString).ToString());
+}
+
+void LexFromString(float& OutValue, const FStringView& InString)
+{
+	TStringBuilder<64> Builder;
+	LexFromString(OutValue, (Builder << InString).ToString());
+}
+
+void LexFromString(double& OutValue, const FStringView& InString)
+{
+	TStringBuilder<64> Builder;
+	LexFromString(OutValue, (Builder << InString).ToString());
+}
+
+void LexFromString(bool& OutValue, const FStringView& InString)
+{
+	TStringBuilder<64> Builder;
+	LexFromString(OutValue, (Builder << InString).ToString());
+}
diff --git a/Engine/Source/Runtime/Core/Public/String/LexFromString.h b/Engine/Source/Runtime/Core/Public/String/LexFromString.h
new file mode 100644
--- /dev/null
+++ b/Engine/Source/Runtime/Core/Public/String/LexFromString.h
@@ -0,0 +1,19 @@
+// Copyright Epic Games, Inc. All Rights Reserved.
+
+#pragma once
+
+#include "CoreTypes.h"
+
+class FStringView;
+
+CORE_API void LexFromString(int8& OutValue,   const FStringView& InString);
+CORE_API void LexFromString(int16& OutValue,  const FStringView& InString);
+CORE_API void LexFromString(int32& OutValue,  const FStringView& InString);
+CORE_API void LexFromString(int64& OutValue,  const FStringView& InString);
+CORE_API void LexFromString(uint8& OutValue,  const FStringView& InString);
+CORE_API void LexFromString(uint16& OutValue, const FStringView& InString);
+CORE_API void LexFromString(uint32& OutValue, const FStringView& InString);
+CORE_API void LexFromString(uint64& OutValue, const FStringView& InString);
+CORE_API void LexFromString(float& OutValue,  const FStringView& InString);
+CORE_API void LexFromString(double& OutValue, const FStringView& InString);
+CORE_API void LexFromString(bool& OutValue,   const FStringView& InString);
