Add UE::String::BytesToHex, UE::String::HexToBytes, FString += FStringView

- Both hex functions work the same as their FString equivalents but without requiring FString as input or output.
- Optimize TCharToNibble to do a check at the end instead of at the start, as the previous implementation causes the character ranges to be tested twice.
- Change HexToBytes from FString to use UE::String::HexToBytes.
- Add FString += FStringView to facilitate conversion to FStringView.

diff --git a/Engine/Source/Runtime/Core/Private/Containers/String.cpp b/Engine/Source/Runtime/Core/Private/Containers/String.cpp
--- a/Engine/Source/Runtime/Core/Private/Containers/String.cpp
+++ b/Engine/Source/Runtime/Core/Private/Containers/String.cpp
@@ -2,6 +2,7 @@
 
 #include "Containers/UnrealString.h"
 
+#include "Containers/StringView.h"
 #include "CoreGlobals.h"
 #include "CoreTypes.h"
 #include "HAL/UnrealMemory.h"
@@ -10,8 +11,8 @@
 #include "Math/UnrealMathUtility.h"
 #include "Misc/AssertionMacros.h"
 #include "Misc/ByteSwap.h"
-#include "Misc/StringView.h"
 #include "Misc/VarArgs.h"
+#include "String/HexToBytes.h"
 #include "Templates/UnrealTemplate.h"
 
 /* FString implementation
@@ -151,6 +152,13 @@ FString& FString::operator=(const FStringView& Other)
 	return *this;
 }
 
+FString& FString::operator+=(const FStringView& Str)
+{
+	CheckInvariants();
+	AppendChars(Str.GetData(), Str.Len());
+	return *this;
+}
+
 FString::operator FStringView() const
 {
 	return FStringView(Data.GetData(), Len());
@@ -1474,6 +1482,11 @@ FArchive& operator<<( FArchive& Ar, FString& A )
 	return Ar;
 }
 
+int32 HexToBytes(const FString& HexString, uint8* OutBytes)
+{
+	return UE::String::HexToBytes(HexString, OutBytes);
+}
+
 int32 FindMatchingClosingParenthesis(const FString& TargetString, const int32 StartSearch)
 {
 	check(StartSearch >= 0 && StartSearch <= TargetString.Len());// Check for usage, we do not accept INDEX_NONE like other string functions
diff --git a/Engine/Source/Runtime/Core/Private/String/BytesToHex.cpp b/Engine/Source/Runtime/Core/Private/String/BytesToHex.cpp
new file mode 100644
--- /dev/null
+++ b/Engine/Source/Runtime/Core/Private/String/BytesToHex.cpp
@@ -0,0 +1,46 @@
+// Copyright Epic Games, Inc. All Rights Reserved.
+
+#include "String/BytesToHex.h"
+
+#include "Misc/StringBuilder.h"
+
+namespace UE
+{
+namespace String
+{
+	template <typename CharType>
+	void BytesToHexImpl(TArrayView<const uint8> Bytes, CharType* OutHex)
+	{
+		auto NibbleToHex = [](uint8 Value) -> CharType { return Value + (Value > 9 ? 'A' - 10 : '0'); };
+		const uint8* Data = Bytes.GetData();
+		
+		for (const uint8* DataEnd = Data + Bytes.Num(); Data != DataEnd; ++Data)
+		{
+			*OutHex++ = NibbleToHex(*Data >> 4);
+			*OutHex++ = NibbleToHex(*Data & 15);
+		}
+	}
+
+	void BytesToHex(TArrayView<const uint8> Bytes, WIDECHAR* OutHex)
+	{
+		BytesToHexImpl(Bytes, OutHex);
+	}
+
+	void BytesToHex(TArrayView<const uint8> Bytes, ANSICHAR* OutHex)
+	{
+		BytesToHexImpl(Bytes, OutHex);
+	}
+
+	void BytesToHex(TArrayView<const uint8> Bytes, FStringBuilderBase& Builder)
+	{
+		const int32 Offset = Builder.AddUninitialized(Bytes.Num() * 2);
+		BytesToHexImpl(Bytes, GetData(Builder) + Offset);
+	}
+
+	void BytesToHex(TArrayView<const uint8> Bytes, FAnsiStringBuilderBase& Builder)
+	{
+		const int32 Offset = Builder.AddUninitialized(Bytes.Num() * 2);
+		BytesToHexImpl(Bytes, GetData(Builder) + Offset);
+	}
+}
+}
diff --git a/Engine/Source/Runtime/Core/Private/String/HexToBytes.cpp b/Engine/Source/Runtime/Core/Private/String/HexToBytes.cpp
new file mode 100644
--- /dev/null
+++ b/Engine/Source/Runtime/Core/Private/String/HexToBytes.cpp
@@ -0,0 +1,41 @@
+// Copyright Epic Games, Inc. All Rights Reserved.
+
+#include "String/HexToBytes.h"
+
+#include "Containers/StringView.h"
+#include "Containers/UnrealString.h"
+
+namespace UE
+{
+namespace String
+{
+	template <typename CharType>
+	int32 HexToBytesImpl(const TStringView<CharType>& Hex, uint8* const OutBytes)
+	{
+		const int32 HexCount = Hex.Len();
+		const CharType* HexPos = Hex.GetData();
+		const CharType* HexEnd = HexPos + HexCount;
+		uint8* OutPos = OutBytes;
+		if (const bool bPadNibble = (HexCount % 2) == 1)
+		{
+			*OutPos++ = TCharToNibble(*HexPos++);
+		}
+		while (HexPos != HexEnd)
+		{
+			const uint8 HiNibble = uint8(TCharToNibble(*HexPos++) << 4);
+			*OutPos++ = HiNibble | TCharToNibble(*HexPos++);
+		}
+		return static_cast<int32>(OutPos - OutBytes);
+	}
+
+	int32 HexToBytes(const FStringView& Hex, uint8* OutBytes)
+	{
+		return HexToBytesImpl<TCHAR>(Hex, OutBytes);
+	}
+
+	int32 HexToBytes(const FAnsiStringView& Hex, uint8* OutBytes)
+	{
+		return HexToBytesImpl<ANSICHAR>(Hex, OutBytes);
+	}
+}
+}
diff --git a/Engine/Source/Runtime/Core/Public/Containers/UnrealString.h b/Engine/Source/Runtime/Core/Public/Containers/UnrealString.h
--- a/Engine/Source/Runtime/Core/Public/Containers/UnrealString.h
+++ b/Engine/Source/Runtime/Core/Public/Containers/UnrealString.h
@@ -388,6 +388,14 @@ public:
 		*(EndPtr + Count) = 0;
 	}
 
+	/**
+	 * Concatenate this with the given string
+	 *
+	 * @param Str string view to be concatenated onto the end of this
+	 * @return reference to this
+	 */
+	FString& operator+=(const FStringView& Str);
+
 	/**
 	 * Concatenate this with given string
 	 * 
@@ -2098,21 +2106,25 @@ inline const bool CheckTCharIsHex( const TCHAR Char )
 
 /**
  * Convert a TChar to equivalent hex value as a uint8
- * @param Char		The character
+ * @param Hex		The character
  * @return	The uint8 value of a hex character
  */
-inline const uint8 TCharToNibble( const TCHAR Char )
+inline const uint8 TCharToNibble(const TCHAR Hex)
 {
-	check( CheckTCharIsHex( Char ) );
-	if( Char >= TEXT('0') && Char <= TEXT('9') )
+	if (Hex >= '0' && Hex <= '9')
+	{
+		return uint8(Hex - '0');
+	}
+	if (Hex >= 'A' && Hex <= 'F')
 	{
-		return (uint8)(Char - TEXT('0'));
+		return uint8(Hex - 'A' + 10);
 	}
-	else if( Char >= TEXT('A') && Char <= TEXT('F') )
+	if (Hex >= 'a' && Hex <= 'f')
 	{
-		return (uint8)(( Char - TEXT('A') ) + 10);
+		return uint8(Hex - 'a' + 10);
 	}
-	return (uint8)(( Char - TEXT('a') ) + 10);
+	checkf(false, TEXT("'%c' (0x%02X) is not a valid hexadecimal digit"), Hex, Hex);
+	return 0;
 }
 
 /** 
@@ -2121,23 +2133,7 @@ inline const uint8 TCharToNibble( const TCHAR Char )
  * @param OutBytes		Ptr to memory must be preallocated large enough
  * @return	The number of bytes copied
  */
-inline int32 HexToBytes( const FString& HexString, uint8* OutBytes )
-{
-	int32 NumBytes = 0;
-	const bool bPadNibble = ( HexString.Len() % 2 ) == 1;
-	const TCHAR* CharPos = *HexString;
-	if( bPadNibble )
-	{
-		OutBytes[ NumBytes++ ] = TCharToNibble( *CharPos++ );
-	}
-	while( *CharPos )
-	{
-		OutBytes[ NumBytes ] = (uint8)(TCharToNibble( *CharPos++ ) << 4);
-		OutBytes[ NumBytes ] += TCharToNibble( *CharPos++ );
-		++NumBytes;
-	}
-	return NumBytes;
-}
+CORE_API int32 HexToBytes(const FString& HexString, uint8* OutBytes);
 
 /**
  * Generalized API to convert something to a string. Function named after the (deprecated) Lex namespace, which
diff --git a/Engine/Source/Runtime/Core/Public/String/BytesToHex.h b/Engine/Source/Runtime/Core/Public/String/BytesToHex.h
new file mode 100644
--- /dev/null
+++ b/Engine/Source/Runtime/Core/Public/String/BytesToHex.h
@@ -0,0 +1,33 @@
+// Copyright Epic Games, Inc. All Rights Reserved.
+
+#pragma once
+
+#include "CoreTypes.h"
+#include "Containers/ArrayView.h"
+
+class FStringBuilderBase;
+class FAnsiStringBuilderBase;
+
+namespace UE
+{
+namespace String
+{
+	/**
+	 * Convert an array of bytes to a non-null-terminated string of hex digits.
+	 *
+	 * @param Bytes Array of bytes to convert.
+	 * @param OutHex [out] Array of at least 2*Bytes.Len() output characters.
+	 */
+	CORE_API void BytesToHex(TArrayView<const uint8> Bytes, WIDECHAR* OutHex);
+	CORE_API void BytesToHex(TArrayView<const uint8> Bytes, ANSICHAR* OutHex);
+
+	/**
+	 * Append an array of bytes to a string builder as hex digits.
+	 *
+	 * @param Bytes Array of bytes to convert.
+	 * @param Builder Builder to append the converted string to.
+	 */
+	CORE_API void BytesToHex(TArrayView<const uint8> Bytes, FStringBuilderBase& Builder);
+	CORE_API void BytesToHex(TArrayView<const uint8> Bytes, FAnsiStringBuilderBase& Builder);
+}
+}
diff --git a/Engine/Source/Runtime/Core/Public/String/HexToBytes.h b/Engine/Source/Runtime/Core/Public/String/HexToBytes.h
new file mode 100644
--- /dev/null
+++ b/Engine/Source/Runtime/Core/Public/String/HexToBytes.h
@@ -0,0 +1,25 @@
+// Copyright Epic Games, Inc. All Rights Reserved.
+
+#pragma once
+
+#include "CoreTypes.h"
+
+class FStringView;
+class FAnsiStringView;
+
+namespace UE
+{
+namespace String
+{
+	/** 
+	 * Convert an array of hex digits into an array of bytes.
+	 *
+	 * @param Hex Array of hex digits to convert.
+	 * @param OutBytes [out] Array of at least (Hex.Len()+1)/2 output bytes.
+	 *
+	 * @return Number of bytes written to the output.
+	 */
+	CORE_API int32 HexToBytes(const FStringView& Hex, uint8* OutBytes);
+	CORE_API int32 HexToBytes(const FAnsiStringView& Hex, uint8* OutBytes);
+}
+}
