Extend FSHAHash for string view and string builder

- Make FSHAHash::FromString take FStringView to avoid constructing a temporary string.
- Add string builder append operators for FSHAHash.

diff --git a/Engine/Source/Runtime/Core/Public/Misc/SecureHash.h b/Engine/Source/Runtime/Core/Public/Misc/SecureHash.h
--- a/Engine/Source/Runtime/Core/Public/Misc/SecureHash.h
+++ b/Engine/Source/Runtime/Core/Public/Misc/SecureHash.h
@@ -7,10 +7,15 @@
 #include "Containers/UnrealString.h"
 #include "Containers/Map.h"
 #include "Containers/StringConv.h"
+#include "Containers/StringView.h"
 #include "Stats/Stats.h"
 #include "Async/AsyncWork.h"
 #include "Serialization/BufferReader.h"
+#include "String/BytesToHex.h"
+#include "String/HexToBytes.h"
 
+class FAnsiStringBuilderBase;
+class FStringBuilderBase;
 struct FMD5Hash;
 
 /*-----------------------------------------------------------------------------
@@ -204,10 +209,11 @@ public:
 	{
 		return BytesToHex((const uint8*)Hash, sizeof(Hash));
 	}
-	void FromString(const FString& Src)
+
+	inline void FromString(const FStringView& Src)
 	{
 		check(Src.Len() == 40);
-		HexToBytes(Src, Hash);
+		UE::String::HexToBytes(Src, Hash);
 	}
 
 	friend bool operator==(const FSHAHash& X, const FSHAHash& Y)
@@ -231,6 +237,9 @@ public:
 	friend CORE_API void LexFromString(FSHAHash& Hash, const TCHAR*);
 };
 
+inline FStringBuilderBase& operator<<(FStringBuilderBase& Builder, const FSHAHash& Hash) { UE::String::BytesToHex(Hash.Hash, Builder); return Builder; }
+inline FAnsiStringBuilderBase& operator<<(FAnsiStringBuilderBase& Builder, const FSHAHash& Hash) { UE::String::BytesToHex(Hash.Hash, Builder); return Builder; }
+
 class CORE_API FSHA1
 {
 public:
