Add FPathViews::Append

diff --git a/Engine/Source/Runtime/Core/Private/Misc/PathViews.cpp b/Engine/Source/Runtime/Core/Private/Misc/PathViews.cpp
--- a/Engine/Source/Runtime/Core/Private/Misc/PathViews.cpp
+++ b/Engine/Source/Runtime/Core/Private/Misc/PathViews.cpp
@@ -3,7 +3,8 @@
 #include "Misc/PathViews.h"
 
 #include "Algo/FindLast.h"
-#include "Misc/StringView.h"
+#include "Containers/StringView.h"
+#include "Misc/StringBuilder.h"
 
 namespace UE4PathViews_Private
 {
@@ -74,3 +75,12 @@ void FPathViews::Split(const FStringView& InPath, FStringView& OutPath, FStringV
 	OutName = CleanName.Left(NameLen);
 	OutExt = CleanName.RightChop(NameLen + 1);
 }
+
+void FPathViews::Append(FStringBuilderBase& Builder, const FStringView& Suffix)
+{
+	if (Builder.Len() > 0 && !UE4PathViews_Private::IsSlashOrBackslash(Builder.LastChar()))
+	{
+		Builder.Append(TEXT('/'));
+	}
+	Builder.Append(Suffix);
+}
diff --git a/Engine/Source/Runtime/Core/Private/Tests/Misc/PathViewsTest.cpp b/Engine/Source/Runtime/Core/Private/Tests/Misc/PathViewsTest.cpp
--- a/Engine/Source/Runtime/Core/Private/Tests/Misc/PathViewsTest.cpp
+++ b/Engine/Source/Runtime/Core/Private/Tests/Misc/PathViewsTest.cpp
@@ -2,8 +2,8 @@
 
 #include "Misc/PathViews.h"
 
+#include "Containers/StringView.h"
 #include "Misc/AutomationTest.h"
-#include "Misc/StringView.h"
 
 #if WITH_DEV_AUTOMATION_TESTS
 
diff --git a/Engine/Source/Runtime/Core/Public/Misc/PathViews.h b/Engine/Source/Runtime/Core/Public/Misc/PathViews.h
--- a/Engine/Source/Runtime/Core/Public/Misc/PathViews.h
+++ b/Engine/Source/Runtime/Core/Public/Misc/PathViews.h
@@ -4,6 +4,7 @@
 
 #include "CoreTypes.h"
 
+class FStringBuilderBase;
 class FStringView;
 
 class CORE_API FPathViews
@@ -116,4 +117,22 @@ public:
 	 * @param OutExt  [out] Receives the extension portion of the input string, excluding the dot.
 	 */
 	static void Split(const FStringView& InPath, FStringView& OutPath, FStringView& OutName, FStringView& OutExt);
+
+	/**
+	 * Appends the suffix to the path in the builder and ensures that there is a separator between them.
+	 *
+	 * Examples:
+	 * ("",    "")    -> ""
+	 * ("A",   "")    -> "A/"
+	 * ("",    "B")   -> "B"
+	 * ("/",   "B")   -> "/B"
+	 * ("A",   "B")   -> "A/B"
+	 * ("A/",  "B")   -> "A/B"
+	 * ("A\\", "B")   -> "A\\B"
+	 * ("A/B", "C/D") -> "A/B/C/D"
+	 *
+	 * @param Builder A possibly-empty path that may end in a separator.
+	 * @param Suffix A possibly-empty suffix that does not start with a separator.
+	 */
+	static void Append(FStringBuilderBase& Builder, const FStringView& Suffix);
 };
