Add function GetPakchunkIndexFromPakFile() to FGenericPlatformMisc to remove dependency from core to pakfile module.

diff --git a/Engine/Source/Runtime/Core/Private/GenericPlatform/GenericPlatformMisc.cpp b/Engine/Source/Runtime/Core/Private/GenericPlatform/GenericPlatformMisc.cpp
--- a/Engine/Source/Runtime/Core/Private/GenericPlatform/GenericPlatformMisc.cpp
+++ b/Engine/Source/Runtime/Core/Private/GenericPlatform/GenericPlatformMisc.cpp
@@ -1449,3 +1449,32 @@ void FGenericPlatformMisc::ParseChunkIdPakchunkIndexMapping(TArray<FString> Chun
 		}
 	}
 }
+
+int32 FGenericPlatformMisc::GetPakchunkIndexFromPakFile(const FString& InFilename)
+{
+	FString ChunkIdentifier(TEXT("pakchunk"));
+	FString BaseFilename = FPaths::GetBaseFilename(InFilename);
+	int32 ChunkNumber = INDEX_NONE;
+
+	if (BaseFilename.StartsWith(ChunkIdentifier))
+	{
+		int32 StartOfNumber = ChunkIdentifier.Len();
+		int32 DigitCount = 0;
+		if (FChar::IsDigit(BaseFilename[StartOfNumber]))
+		{
+			while ((DigitCount + StartOfNumber) < BaseFilename.Len() && FChar::IsDigit(BaseFilename[StartOfNumber + DigitCount]))
+			{
+				DigitCount++;
+			}
+
+			if ((StartOfNumber + DigitCount) < BaseFilename.Len())
+			{
+				FString ChunkNumberString = BaseFilename.Mid(StartOfNumber, DigitCount);
+				check(ChunkNumberString.IsNumeric());
+				TTypeFromString<int32>::FromString(ChunkNumber, *ChunkNumberString);
+			}
+		}
+	}
+
+	return ChunkNumber;
+}
diff --git a/Engine/Source/Runtime/Core/Public/GenericPlatform/GenericPlatformMisc.h b/Engine/Source/Runtime/Core/Public/GenericPlatform/GenericPlatformMisc.h
--- a/Engine/Source/Runtime/Core/Public/GenericPlatform/GenericPlatformMisc.h
+++ b/Engine/Source/Runtime/Core/Public/GenericPlatform/GenericPlatformMisc.h
@@ -1413,6 +1413,8 @@ public:
 		return PakchunkIndex;
 	}
 
+	static int32 GetPakchunkIndexFromPakFile(const FString& InFilename);
+
 #if !UE_BUILD_SHIPPING
 	/** 
 	 * Returns any platform specific warning messages we want printed on screen
