Fix FName number parsing from unterminated strings

Copy the number to a temporary buffer until we have an overload of Atoi64 that takes the input length. The generic implementation of Atoi64 for TCHAR converts the input to ASCII then converts that to an integer. For unterminated strings, this could cause an access violation or convert an obscenely large input to ASCII.

diff --git a/Engine/Source/Runtime/Core/Private/UObject/UnrealNames.cpp b/Engine/Source/Runtime/Core/Private/UObject/UnrealNames.cpp
--- a/Engine/Source/Runtime/Core/Private/UObject/UnrealNames.cpp
+++ b/Engine/Source/Runtime/Core/Private/UObject/UnrealNames.cpp
@@ -1736,13 +1736,22 @@ struct FNameHelper
 	template<typename CharType>
 	static uint32 ParseNumber(const CharType* Name, int32& InOutLen)
 	{
+		// The input is not guaranteed to be null-terminated so the number must be copied to
+		// a null-terminated buffer until we have an overload of Atoi64 that takes a length.
+		ANSICHAR Buffer[NAME_SIZE];
+		ANSICHAR* TerminatedDigits = Buffer + NAME_SIZE;
+		*--TerminatedDigits = '\0';
+
 		const int32 Len = InOutLen;
 		int32 Digits = 0;
 		for (const CharType* It = Name + Len - 1; It >= Name && *It >= '0' && *It <= '9'; --It)
 		{
+			*--TerminatedDigits = static_cast<ANSICHAR>(*It);
 			++Digits;
 		}
 
+		check(Digits < NAME_SIZE);
+
 		const CharType* FirstDigit = Name + Len - Digits;
 		if (Digits && Digits < Len && *(FirstDigit - 1) == '_')
 		{
@@ -1753,7 +1762,7 @@ struct FNameHelper
 			{
 				// Attempt to convert what's following it to a number
 				// This relies on Name being null-terminated
-				int64 Number = TCString<CharType>::Atoi64(Name + Len - Digits);
+				int64 Number = FCStringAnsi::Atoi64(TerminatedDigits);
 				if (Number < MAX_int32)
 				{
 					InOutLen -= 1 + Digits;
