Optimize FFileHelper::BufferToString to eliminate its use of a temporary allocation

Converting UTF-8 inputs was using a temporary TCHAR buffer the same size as the output. The conversion is now done directly into the allocation for the target string.

diff --git a/Engine/Source/Runtime/Core/Private/Misc/FileHelper.cpp b/Engine/Source/Runtime/Core/Private/Misc/FileHelper.cpp
--- a/Engine/Source/Runtime/Core/Private/Misc/FileHelper.cpp
+++ b/Engine/Source/Runtime/Core/Private/Misc/FileHelper.cpp
@@ -118,10 +118,10 @@ void FFileHelper::BufferToString( FString& Result, const uint8* Buffer, int32 Si
 			Size   -= 3;
 		}
 
-		FUTF8ToTCHAR Conv((const ANSICHAR*)Buffer, Size);
-		int32 Length = Conv.Length();
-		ResultArray.AddUninitialized(Length + 1); // For the null terminator
-		CopyAssignItems(ResultArray.GetData(), Conv.Get(), Length);
+		int32 Length = FUTF8ToTCHAR_Convert::ConvertedLength(reinterpret_cast<const ANSICHAR*>(Buffer), Size);
+		ResultArray.AddUninitialized(Length + 1); // +1 for the null terminator
+		FUTF8ToTCHAR_Convert::Convert(ResultArray.GetData(), ResultArray.Num(), reinterpret_cast<const ANSICHAR*>(Buffer), Size);
+		ResultArray[Length] = TEXT('\0');
 	}
 
 	if (ResultArray.Num() == 1)
