Preport AsyncIODelete feature from //UE4/Dev-Core; saves 4 minutes off of CookByTheBook due to not blocking on the delete of the sandbox from the previous cook.

#rnx

diff --git a/Engine/Source/Editor/UnrealEd/Classes/CookOnTheSide/CookOnTheFlyServer.h b/Engine/Source/Editor/UnrealEd/Classes/CookOnTheSide/CookOnTheFlyServer.h
--- a/Engine/Source/Editor/UnrealEd/Classes/CookOnTheSide/CookOnTheFlyServer.h
+++ b/Engine/Source/Editor/UnrealEd/Classes/CookOnTheSide/CookOnTheFlyServer.h
@@ -9,6 +9,7 @@
 #include "UObject/Object.h"
 #include "UObject/WeakObjectPtr.h"
 #include "UObject/Package.h"
+#include "Templates/UniquePtr.h"
 #include "TickableEditorObject.h"
 #include "IPlatformFileSandboxWrapper.h"
 #include "INetworkFileSystemModule.h"
diff --git a/Engine/Source/Editor/UnrealEd/Private/CookOnTheFlyServer.cpp b/Engine/Source/Editor/UnrealEd/Private/CookOnTheFlyServer.cpp
--- a/Engine/Source/Editor/UnrealEd/Private/CookOnTheFlyServer.cpp
+++ b/Engine/Source/Editor/UnrealEd/Private/CookOnTheFlyServer.cpp
@@ -5143,6 +5143,9 @@ void UCookOnTheFlyServer::DeleteSandboxDirectory(const FString& PlatformName)
 	// The AsyncDelete directory is associated with a sandbox but is necessarily outside of it since it is used to delete the sandbox.
 	// Note that for the Platform we used to create the AsyncIODelete, this Delete will fail because AsyncIODelete refuses to delete its own temproot; this is okay because it will delete the temproot on exit.
 	LocalAsyncIODelete.DeleteDirectory(AsyncDeleteDirectory);
+
+	// UE_DEPRECATED(4.25, "Delete the old location for AsyncDeleteDirectory until all users have cooked at least once")
+	LocalAsyncIODelete.DeleteDirectory(SandboxDirectory + TEXT("AsyncDelete"));
 }
 
 FAsyncIODelete& UCookOnTheFlyServer::GetAsyncIODelete(const FString& PlatformName, const FString* AsyncDeleteDirectory)
@@ -5177,7 +5180,7 @@ FString UCookOnTheFlyServer::GetAsyncDeleteDirectory(const FString& PlatformName
 		FPaths::NormalizeDirectoryName(Buffer);
 		SandboxDirectory = &Buffer;
 	}
-	return (*SandboxDirectory) + TEXT("AsyncDelete");
+	return (*SandboxDirectory) + TEXT("_Del");
 }
 
 void UCookOnTheFlyServer::PopulateCookedPackagesFromDisk(const TArray<ITargetPlatform*>& Platforms)
diff --git a/Engine/Source/Editor/UnrealEd/Private/Cooker/AsyncIODelete.cpp b/Engine/Source/Editor/UnrealEd/Private/Cooker/AsyncIODelete.cpp
--- a/Engine/Source/Editor/UnrealEd/Private/Cooker/AsyncIODelete.cpp
+++ b/Engine/Source/Editor/UnrealEd/Private/Cooker/AsyncIODelete.cpp
@@ -89,14 +89,14 @@ bool FAsyncIODelete::Setup()
 	uint32 ErrorCode;
 	if (!DeleteTempRootDirectory(ErrorCode))
 	{
-		UE_LOG(LogCook, Fatal, TEXT("Could not clear asyncdelete root directory '%s'.  LastError: %i."), *TempRoot, ErrorCode);
+		UE_LOG(LogCook, Error, TEXT("Could not clear asyncdelete root directory '%s'.  LastError: %i."), *TempRoot, ErrorCode);
 		return false;
 	}
 
 	// Create the empty directory to work in
 	if (!IFileManager::Get().MakeDirectory(*TempRoot, true))
 	{
-		UE_LOG(LogCook, Fatal, TEXT("Could not create asyncdelete root directory '%s'.  LastError: %i."), *TempRoot, FPlatformMisc::GetLastError());
+		UE_LOG(LogCook, Error, TEXT("Could not create asyncdelete root directory '%s'.  LastError: %i."), *TempRoot, FPlatformMisc::GetLastError());
 		return false;
 	}
 
@@ -159,7 +159,7 @@ bool FAsyncIODelete::WaitForAllTasks(float TimeLimitSeconds)
 
 	if (TimeLimitSeconds <= 0.f)
 	{
-	TasksComplete->Wait();
+		TasksComplete->Wait();
 	}
 	else
 	{
@@ -216,7 +216,7 @@ bool FAsyncIODelete::Delete(const FStringView& PathToDelete, EPathType ExpectedT
 	}
 
 #if ASYNCIODELETE_ASYNC_ENABLED
-	const FString TempPath = FPaths::Combine(TempRoot, FString::Printf(TEXT("Delete%u"), DeleteCounter));
+	const FString TempPath = FPaths::Combine(TempRoot, FString::Printf(TEXT("%u"), DeleteCounter));
 	DeleteCounter++;
 
 	const bool bReplace = true;
diff --git a/Engine/Source/Runtime/Core/Private/Misc/Paths.cpp b/Engine/Source/Runtime/Core/Private/Misc/Paths.cpp
--- a/Engine/Source/Runtime/Core/Private/Misc/Paths.cpp
+++ b/Engine/Source/Runtime/Core/Private/Misc/Paths.cpp
@@ -1088,10 +1088,40 @@ bool FPaths::CollapseRelativeDirectories(FString& InPath)
 
 void FPaths::RemoveDuplicateSlashes(FString& InPath)
 {
-	while (InPath.Contains(TEXT("//"), ESearchCase::CaseSensitive))
+	TCHAR* Text = InPath.GetCharArray().GetData();
+	if (!Text)
 	{
-		InPath = InPath.Replace(TEXT("//"), TEXT("/"), ESearchCase::CaseSensitive);
+		return;
+	}
+	const TCHAR* const TwoSlashStr = TEXT("//");
+	const TCHAR SlashChr = TEXT('/');
+
+	TCHAR* const EditStart = TCString<TCHAR>::Strstr(Text, TwoSlashStr);
+	if (!EditStart)
+	{
+		return;
 	}
+	const TCHAR* const TextEnd = Text + InPath.Len();
+
+	// Since we know we've found TwoSlashes, point the initial Write head at the spot where the second slash is (which we shall skip), and point the Read head at the next character after the second slash
+	TCHAR* Write = EditStart + 1;	// The position to write the next character of the output
+	const TCHAR* Read = Write + 1;	// The next character of the input to read
+
+	for (; Read < TextEnd; ++Read)
+	{
+		TCHAR NextChar = *Read;
+		if (Write[-1] != SlashChr || NextChar != SlashChr)
+		{
+			*Write++ = NextChar;
+		}
+		else
+		{
+			// Skip the NextChar when adding on a slash after an existing slash, e.g // or more generally before/////after
+			//                                                                       WR                         W  R
+		}
+	}
+	*Write = TEXT('\0');
+	InPath.TrimToNullTerminator();
 }
 
 void FPaths::MakeStandardFilename(FString& InPath)
diff --git a/Engine/Source/Runtime/Core/Private/Tests/Misc/PathsTest.cpp b/Engine/Source/Runtime/Core/Private/Tests/Misc/PathsTest.cpp
--- a/Engine/Source/Runtime/Core/Private/Tests/Misc/PathsTest.cpp
+++ b/Engine/Source/Runtime/Core/Private/Tests/Misc/PathsTest.cpp
@@ -184,6 +184,36 @@ bool FPathTests::RunTest( const FString& Parameters )
 		RunIsUnderDirectoryTest(TEXT("C:/Folder/Subdir/"),	TEXT("C:/Folder/"), true);
 	}
 
+	// RemoveDuplicateSlashes
+	{
+		auto RunRemoveDuplicateSlashesTest = [this](const TCHAR* InPath, const TCHAR* InExpectedResult)
+		{
+			FString ReplacementPath(InPath);
+			const FString ExpectedResult(InExpectedResult);
+			FPaths::RemoveDuplicateSlashes(ReplacementPath);
+			if (!ReplacementPath.Equals(ExpectedResult, ESearchCase::CaseSensitive))
+			{
+				AddError(FString::Printf(TEXT("FPaths::RemoveDuplicateSlashes('%s') != '%s'."), InPath, InExpectedResult));
+			}
+			else if (ReplacementPath.GetCharArray().Num() != 0 && ReplacementPath.GetCharArray().Num() != ExpectedResult.Len() + 1)
+			{
+				AddError(FString::Printf(TEXT("FPaths::RemoveDuplicateSlashes('%s') returned a result with extra space still allocated after the null terminator."), InPath));
+			}
+		};
+
+		RunRemoveDuplicateSlashesTest(TEXT(""), TEXT(""));
+		RunRemoveDuplicateSlashesTest(TEXT("C:/Folder/File.txt"), TEXT("C:/Folder/File.txt"));
+		RunRemoveDuplicateSlashesTest(TEXT("C:/Folder/File/"), TEXT("C:/Folder/File/"));
+		RunRemoveDuplicateSlashesTest(TEXT("/"), TEXT("/"));
+		RunRemoveDuplicateSlashesTest(TEXT("//"), TEXT("/"));
+		RunRemoveDuplicateSlashesTest(TEXT("////"), TEXT("/"));
+		RunRemoveDuplicateSlashesTest(TEXT("/Folder/File"), TEXT("/Folder/File"));
+		RunRemoveDuplicateSlashesTest(TEXT("//Folder/File"), TEXT("/Folder/File")); // Don't use on //UNC paths; it will be stripped!
+		RunRemoveDuplicateSlashesTest(TEXT("/////Folder//////File/////"), TEXT("/Folder/File/"));
+		RunRemoveDuplicateSlashesTest(TEXT("\\\\Folder\\\\File\\\\"), TEXT("\\\\Folder\\\\File\\\\")); // It doesn't strip backslash, and we rely on that in some places
+		RunRemoveDuplicateSlashesTest(TEXT("//\\\\//Folder//\\\\//File//\\\\//"), TEXT("/\\\\/Folder/\\\\/File/\\\\/"));
+	}
+
 	return true;
 }
 
diff --git a/Engine/Source/Runtime/Core/Private/Windows/WindowsPlatformFile.cpp b/Engine/Source/Runtime/Core/Private/Windows/WindowsPlatformFile.cpp
--- a/Engine/Source/Runtime/Core/Private/Windows/WindowsPlatformFile.cpp
+++ b/Engine/Source/Runtime/Core/Private/Windows/WindowsPlatformFile.cpp
@@ -767,27 +767,62 @@ FMappedFileRegionWindows::~FMappedFileRegionWindows()
 **/
 class CORE_API FWindowsPlatformFile : public IPhysicalPlatformFile
 {
-protected:
-	virtual FString NormalizeFilename(const TCHAR* Filename)
+private:
+	FString WindowsNormalizedFilename(const TCHAR* Filename)
 	{
-		FString Result(Filename);
-		FPaths::NormalizeFilename(Result);
-		if (Result.StartsWith(TEXT("//")))
-		{
-			Result = FString(TEXT("\\\\")) + Result.RightChop(2);
-		}
-		return FPaths::ConvertRelativePathToFull(Result);
+		const bool bIsFilename = true; // TODO: Create a platform-independent EPathType enum and use that here
+		return WindowsNormalizedPath(Filename, bIsFilename);
 	}
-	virtual FString NormalizeDirectory(const TCHAR* Directory)
+
+	FString WindowsNormalizedDirname(const TCHAR* Directory)
 	{
-		FString Result(Directory);
-		FPaths::NormalizeDirectoryName(Result);
-		if (Result.StartsWith(TEXT("//")))
+		const bool bIsFilename = false;
+		return WindowsNormalizedPath(Directory, bIsFilename);
+	}
+
+	/**
+	  * Convert from a valid Unreal Path to a canonical and strict-valid Windows Path.
+	  * An Unreal Path may have either \ or / and may have empty directories (two / in a row), and may have .. and may be relative
+	  * A canonical and strict-valid Windows Path has only \, does not have .., does not have empty directories, and is an absolute path, either \\UNC or D:\
+	  * We need to use strict-valid Windows Paths when calling Windows API calls so that we can support the long-path prefix \\?\
+	  */
+	FString WindowsNormalizedPath(const TCHAR* PathString, bool bIsFilename)
+	{
+		FString Result = FPaths::ConvertRelativePathToFull(FString(PathString));
+		// NormalizeFilename was already called by ConvertRelativePathToFull, but we still need to do the extra steps in NormalizeDirectoryName if it is a directory
+		if (!bIsFilename)
+		{
+			FPaths::NormalizeDirectoryName(Result);
+		}
+
+		// Remove duplicate slashes
+		const bool bIsUNCPath = Result.StartsWith(TEXT("//"));
+		if (bIsUNCPath)
+		{
+			// Keep // at the beginning.  If There are more than two / at the beginning, replace them with just //.
+			FPaths::RemoveDuplicateSlashes(Result);
+			Result = TEXT("/") + Result;
+		}
+		else
 		{
-			Result = FString(TEXT("\\\\")) + Result.RightChop(2);
+			FPaths::RemoveDuplicateSlashes(Result);
 		}
-		return FPaths::ConvertRelativePathToFull(Result);
+
+		// We now have a canonical, strict-valid, absolute Unreal Path.  Convert it to a Windows Path.
+		Result.ReplaceCharInline(TEXT('/'), TEXT('\\'), ESearchCase::CaseSensitive);
+
+		// Handle Windows Path length over MAX_PATH
+		if (!bIsUNCPath && Result.Len() > MAX_PATH)
+		{
+			// Pathnames exceeding MAX_PATH are supported by WindowsAPI functions if you add the \\?\ prefix,
+			// This does not work on \\UNC paths in all WindowsAPI functions however, so only do this for D:\ paths.
+			// Our engine does not currently handle UNC paths greater than MAXPATH; we will need to talk to Microsoft if we want to handle that robustly.
+			Result = TEXT("\\\\?\\") + Result;
+		}
+
+		return Result;
 	}
+
 public:
 	//~ For visibility of overloads we don't override
 	using IPhysicalPlatformFile::IterateDirectory;
@@ -795,7 +830,7 @@ public:
 
 	virtual bool FileExists(const TCHAR* Filename) override
 	{
-		uint32 Result = GetFileAttributesW(*NormalizeFilename(Filename));
+		uint32 Result = GetFileAttributesW(*WindowsNormalizedFilename(Filename));
 		if (Result != 0xFFFFFFFF && !(Result & FILE_ATTRIBUTE_DIRECTORY))
 		{
 			return true;
@@ -805,7 +840,7 @@ public:
 	virtual int64 FileSize(const TCHAR* Filename) override
 	{
 		WIN32_FILE_ATTRIBUTE_DATA Info;
-		if (!!GetFileAttributesExW(*NormalizeFilename(Filename), GetFileExInfoStandard, &Info))
+		if (!!GetFileAttributesExW(*WindowsNormalizedFilename(Filename), GetFileExInfoStandard, &Info))
 		{
 			if ((Info.dwFileAttributes & FILE_ATTRIBUTE_DIRECTORY) == 0)
 			{
@@ -819,12 +854,12 @@ public:
 	}
 	virtual bool DeleteFile(const TCHAR* Filename) override
 	{
-		const FString NormalizedFilename = NormalizeFilename(Filename);
+		const FString NormalizedFilename = WindowsNormalizedFilename(Filename);
 		return !!DeleteFileW(*NormalizedFilename);
 	}
 	virtual bool IsReadOnly(const TCHAR* Filename) override
 	{
-		uint32 Result = GetFileAttributesW(*NormalizeFilename(Filename));
+		uint32 Result = GetFileAttributesW(*WindowsNormalizedFilename(Filename));
 		if (Result != 0xFFFFFFFF)
 		{
 			return !!(Result & FILE_ATTRIBUTE_READONLY);
@@ -833,17 +868,17 @@ public:
 	}
 	virtual bool MoveFile(const TCHAR* To, const TCHAR* From) override
 	{
-		return !!MoveFileW(*NormalizeFilename(From), *NormalizeFilename(To));
+		return !!MoveFileW(*WindowsNormalizedFilename(From), *WindowsNormalizedFilename(To));
 	}
 	virtual bool SetReadOnly(const TCHAR* Filename, bool bNewReadOnlyValue) override
 	{
-		return !!SetFileAttributesW(*NormalizeFilename(Filename), bNewReadOnlyValue ? FILE_ATTRIBUTE_READONLY : FILE_ATTRIBUTE_NORMAL);
+		return !!SetFileAttributesW(*WindowsNormalizedFilename(Filename), bNewReadOnlyValue ? FILE_ATTRIBUTE_READONLY : FILE_ATTRIBUTE_NORMAL);
 	}
 
 	virtual FDateTime GetTimeStamp(const TCHAR* Filename) override
 	{
 		WIN32_FILE_ATTRIBUTE_DATA Info;
-		if (GetFileAttributesExW(*NormalizeFilename(Filename), GetFileExInfoStandard, &Info))
+		if (GetFileAttributesExW(*WindowsNormalizedFilename(Filename), GetFileExInfoStandard, &Info))
 		{
 			return WindowsFileTimeToUEDateTime(Info.ftLastWriteTime);
 		}
@@ -854,7 +889,7 @@ public:
 	virtual void SetTimeStamp(const TCHAR* Filename, FDateTime DateTime) override
 	{
 		TRACE_PLATFORMFILE_BEGIN_OPEN(Filename);
-		HANDLE Handle = CreateFileW(*NormalizeFilename(Filename), FILE_WRITE_ATTRIBUTES, FILE_SHARE_READ, nullptr, OPEN_EXISTING, FILE_ATTRIBUTE_NORMAL | FILE_FLAG_OVERLAPPED, nullptr);
+		HANDLE Handle = CreateFileW(*WindowsNormalizedFilename(Filename), FILE_WRITE_ATTRIBUTES, FILE_SHARE_READ, nullptr, OPEN_EXISTING, FILE_ATTRIBUTE_NORMAL | FILE_FLAG_OVERLAPPED, nullptr);
 		TRACE_PLATFORMFILE_END_OPEN(Handle);
 		if (Handle != INVALID_HANDLE_VALUE)
 		{
@@ -876,7 +911,7 @@ public:
 	virtual FDateTime GetAccessTimeStamp(const TCHAR* Filename) override
 	{
 		WIN32_FILE_ATTRIBUTE_DATA Info;
-		if (GetFileAttributesExW(*NormalizeFilename(Filename), GetFileExInfoStandard, &Info))
+		if (GetFileAttributesExW(*WindowsNormalizedFilename(Filename), GetFileExInfoStandard, &Info))
 		{
 			return WindowsFileTimeToUEDateTime(Info.ftLastAccessTime);
 		}
@@ -886,17 +921,14 @@ public:
 
 	virtual FString GetFilenameOnDisk(const TCHAR* Filename) override
 	{
+		FString NormalizedFileName = WindowsNormalizedFilename(Filename);
 		TRACE_PLATFORMFILE_BEGIN_OPEN(Filename);
-		HANDLE hFile = CreateFile(Filename, FILE_READ_ATTRIBUTES, FILE_SHARE_READ | FILE_SHARE_WRITE | FILE_SHARE_DELETE, NULL, OPEN_EXISTING, 0, NULL);
+		HANDLE hFile = CreateFile(*NormalizedFileName, FILE_READ_ATTRIBUTES, FILE_SHARE_READ | FILE_SHARE_WRITE | FILE_SHARE_DELETE, NULL, OPEN_EXISTING, 0, NULL);
 		TRACE_PLATFORMFILE_END_OPEN(hFile);
-		if(hFile == INVALID_HANDLE_VALUE)
+		// If the file exists on disk, read the capitalization from the path on disk, otherwise just return the (normalized) input filename
+		if (hFile != INVALID_HANDLE_VALUE)
 		{
-			return NormalizeFilename(Filename);
-		}
-		else
-		{
-			FString NormalizedFileName;
-			for(uint32 Length = FCString::Strlen(Filename) + 10;;)
+			for (uint32 Length = NormalizedFileName.Len() + 10;;)
 			{
 				TArray<TCHAR>& CharArray = NormalizedFileName.GetCharArray();
 				CharArray.SetNum(Length);
@@ -904,7 +936,7 @@ public:
 				Length = GetFinalPathNameByHandle(hFile, CharArray.GetData(), CharArray.Num(), FILE_NAME_NORMALIZED);
 				if (Length == 0)
 				{
-					NormalizedFileName = Filename;
+					NormalizedFileName = WindowsNormalizedFilename(Filename);
 					break;
 				}
 				if (Length < (uint32)CharArray.Num())
@@ -916,12 +948,12 @@ public:
 			TRACE_PLATFORMFILE_BEGIN_CLOSE(hFile);
 			CloseHandle(hFile);
 			TRACE_PLATFORMFILE_END_CLOSE();
-
-			NormalizedFileName.RemoveFromStart(TEXT("\\\\?\\"), ESearchCase::CaseSensitive);
-			NormalizedFileName.ReplaceInline(TEXT("\\"), TEXT("/"));
-
-			return NormalizedFileName;
 		}
+		// Convert the result back into an UnrealPath(\\ -> / )
+		NormalizedFileName.RemoveFromStart(TEXT("\\\\?\\"), ESearchCase::CaseSensitive);
+		NormalizedFileName.ReplaceCharInline(TEXT('\\'), TEXT('/'), ESearchCase::CaseSensitive);
+
+		return NormalizedFileName;
 	}
 
 #define USE_WINDOWS_ASYNC_IMPL 0
@@ -933,7 +965,7 @@ public:
 		uint32  Create = OPEN_EXISTING;
 
 
-		FString NormalizedFilename = NormalizeFilename(Filename);
+		FString NormalizedFilename = WindowsNormalizedFilename(Filename);
 		TRACE_PLATFORMFILE_BEGIN_OPEN(Filename);
 		HANDLE Handle = CreateFileW(*NormalizedFilename, Access, WinFlags, NULL, Create, FILE_ATTRIBUTE_NORMAL | FILE_FLAG_OVERLAPPED, NULL);
 		TRACE_PLATFORMFILE_END_OPEN(Handle);
@@ -952,7 +984,7 @@ public:
 
 #if USE_OVERLAPPED_IO
 		TRACE_PLATFORMFILE_BEGIN_OPEN(Filename);
-		HANDLE Handle    = CreateFileW(*NormalizeFilename(Filename), Access, WinFlags, NULL, Create, FILE_ATTRIBUTE_NORMAL | FILE_FLAG_OVERLAPPED, NULL);
+		HANDLE Handle    = CreateFileW(*WindowsNormalizedFilename(Filename), Access, WinFlags, NULL, Create, FILE_ATTRIBUTE_NORMAL | FILE_FLAG_OVERLAPPED, NULL);
 		TRACE_PLATFORMFILE_END_OPEN(Handle);
 		if (Handle != INVALID_HANDLE_VALUE)
 		{
@@ -960,7 +992,7 @@ public:
 		}
 #else
 		TRACE_PLATFORMFILE_BEGIN_OPEN(Filename);
-		HANDLE Handle = CreateFileW(*NormalizeFilename(Filename), Access, WinFlags, NULL, Create, FILE_ATTRIBUTE_NORMAL, NULL);
+		HANDLE Handle = CreateFileW(*WindowsNormalizedFilename(Filename), Access, WinFlags, NULL, Create, FILE_ATTRIBUTE_NORMAL, NULL);
 		TRACE_PLATFORMFILE_END_OPEN(Handle);
 		if (Handle != INVALID_HANDLE_VALUE)
 		{
@@ -976,7 +1008,7 @@ public:
 		uint32  WinFlags = FILE_SHARE_READ | (bAllowWrite ? FILE_SHARE_WRITE : 0);
 		uint32  Create = OPEN_EXISTING;
 		TRACE_PLATFORMFILE_BEGIN_OPEN(Filename);
-		HANDLE Handle = CreateFileW(*NormalizeFilename(Filename), Access, WinFlags, NULL, Create, FILE_ATTRIBUTE_NORMAL | FILE_FLAG_OVERLAPPED, NULL);
+		HANDLE Handle = CreateFileW(*WindowsNormalizedFilename(Filename), Access, WinFlags, NULL, Create, FILE_ATTRIBUTE_NORMAL | FILE_FLAG_OVERLAPPED, NULL);
 		TRACE_PLATFORMFILE_END_OPEN(Handle);
 		if (Handle != INVALID_HANDLE_VALUE)
 		{
@@ -991,7 +1023,7 @@ public:
 		uint32  WinFlags  = bAllowRead ? FILE_SHARE_READ : 0;
 		uint32  Create    = bAppend ? OPEN_ALWAYS : CREATE_ALWAYS;
 		TRACE_PLATFORMFILE_BEGIN_OPEN(Filename);
-		HANDLE Handle    = CreateFileW(*NormalizeFilename(Filename), Access, WinFlags, NULL, Create, FILE_ATTRIBUTE_NORMAL, NULL);
+		HANDLE Handle    = CreateFileW(*WindowsNormalizedFilename(Filename), Access, WinFlags, NULL, Create, FILE_ATTRIBUTE_NORMAL, NULL);
 		TRACE_PLATFORMFILE_END_OPEN(Handle);
 		if(Handle != INVALID_HANDLE_VALUE)
 		{
@@ -1017,7 +1049,7 @@ public:
 		uint32  WinFlags = FILE_SHARE_READ;
 		uint32  Create = OPEN_EXISTING;
 		TRACE_PLATFORMFILE_BEGIN_OPEN(Filename);
-		HANDLE Handle = CreateFileW(*NormalizeFilename(Filename), Access, WinFlags, NULL, Create, FILE_ATTRIBUTE_NORMAL, NULL);
+		HANDLE Handle = CreateFileW(*WindowsNormalizedFilename(Filename), Access, WinFlags, NULL, Create, FILE_ATTRIBUTE_NORMAL, NULL);
 		TRACE_PLATFORMFILE_END_OPEN(Handle);
 		if (Handle == INVALID_HANDLE_VALUE)
 		{
@@ -1039,18 +1071,18 @@ public:
 		bool bExists = !FCString::Strlen(Directory);
 		if (!bExists) 
 		{
-			uint32 Result = GetFileAttributesW(*NormalizeDirectory(Directory));
+			uint32 Result = GetFileAttributesW(*WindowsNormalizedDirname(Directory));
 			bExists = (Result != 0xFFFFFFFF && (Result & FILE_ATTRIBUTE_DIRECTORY));
 		}
 		return bExists;
 	}
 	virtual bool CreateDirectory(const TCHAR* Directory) override
 	{
-		return CreateDirectoryW(*NormalizeDirectory(Directory), NULL) || GetLastError() == ERROR_ALREADY_EXISTS;
+		return CreateDirectoryW(*WindowsNormalizedDirname(Directory), NULL) || GetLastError() == ERROR_ALREADY_EXISTS;
 	}
 	virtual bool DeleteDirectory(const TCHAR* Directory) override
 	{
-		RemoveDirectoryW(*NormalizeDirectory(Directory));
+		RemoveDirectoryW(*WindowsNormalizedDirname(Directory));
 		uint32 LastError = GetLastError();
 		const bool bSucceeded = !DirectoryExists(Directory);
 		if (!bSucceeded)
@@ -1062,7 +1094,7 @@ public:
 	virtual FFileStatData GetStatData(const TCHAR* FilenameOrDirectory) override
 	{
 		WIN32_FILE_ATTRIBUTE_DATA Info;
-		if (GetFileAttributesExW(*NormalizeFilename(FilenameOrDirectory), GetFileExInfoStandard, &Info))
+		if (GetFileAttributesExW(*WindowsNormalizedFilename(FilenameOrDirectory), GetFileExInfoStandard, &Info))
 		{
 			const bool bIsDirectory = !!(Info.dwFileAttributes & FILE_ATTRIBUTE_DIRECTORY);
 
@@ -1129,7 +1161,8 @@ public:
 	{
 		bool Result = false;
 		WIN32_FIND_DATAW Data;
-		HANDLE Handle = FindFirstFileW(*(NormalizeDirectory(Directory) / TEXT("*.*")), &Data);
+		FString SearchWildcard = FString(Directory) / TEXT("*.*");
+		HANDLE Handle = FindFirstFileW(*(WindowsNormalizedDirname(*SearchWildcard)), &Data);
 		if (Handle != INVALID_HANDLE_VALUE)
 		{
 			Result = true;
