Add tests for parsing names with numbers from unterminated strings

Requested in review https://p4-swarm.epicgames.net/reviews/11200485 for CL 11201262.

#rnx

diff --git a/Engine/Source/Runtime/Core/Private/UObject/UnrealNames.cpp b/Engine/Source/Runtime/Core/Private/UObject/UnrealNames.cpp
--- a/Engine/Source/Runtime/Core/Private/UObject/UnrealNames.cpp
+++ b/Engine/Source/Runtime/Core/Private/UObject/UnrealNames.cpp
@@ -2282,9 +2282,27 @@ void FName::AutoTest()
 	check(*Cylinder.ToEName() == NAME_Cylinder);
 	check(Cylinder.GetPlainNameString() == TEXT("Cylinder"));
 
+	// Test numbers
+	check(FName("Text_0") == FName("Text", NAME_EXTERNAL_TO_INTERNAL(0)));
+	check(FName("Text_1") == FName("Text", NAME_EXTERNAL_TO_INTERNAL(1)));
+	check(FName("Text_1_0") == FName("Text_1", NAME_EXTERNAL_TO_INTERNAL(0)));
+	check(FName("Text_0_1") == FName("Text_0", NAME_EXTERNAL_TO_INTERNAL(1)));
+	check(FName("Text_00") == FName("Text_00", NAME_NO_NUMBER_INTERNAL));
+	check(FName("Text_01") == FName("Text_01", NAME_NO_NUMBER_INTERNAL));
+
+	// Test unterminated strings
 	check(FName("") == FName(0, "Unused"));
 	check(FName("Used") == FName(4, "UsedUnused"));
 	check(FName("Used") == FName(4, "Used"));
+	check(FName("Used_0") == FName(6, "Used_01"));
+	check(FName("Used_01") == FName(7, "Used_012"));
+	check(FName("Used_123") == FName(8, "Used_123456"));
+	check(FName("Used_123") == FName(8, "Used_123_456"));
+	check(FName("Used_123") == FName(8, TEXT("Used_123456")));
+	check(FName("Used_123") == FName(8, TEXT("Used_123_456")));
+	check(FName("Used_2147483646") == FName(15, TEXT("Used_2147483646123")));
+	check(FName("Used_2147483647") == FName(15, TEXT("Used_2147483647123")));
+	check(FName("Used_2147483648") == FName(15, TEXT("Used_2147483648123")));
 
 	// Test wide strings
 	FString Wide("Wide ");
