Fix unterminated number parsing for new FName constructors with fix length

Alternate fix that does not copy the number

diff --git a/Engine/Source/Runtime/Core/Private/UObject/UnrealNames.cpp b/Engine/Source/Runtime/Core/Private/UObject/UnrealNames.cpp
--- a/Engine/Source/Runtime/Core/Private/UObject/UnrealNames.cpp
+++ b/Engine/Source/Runtime/Core/Private/UObject/UnrealNames.cpp
@@ -1718,6 +1718,19 @@ static FWideStringViewWithWidth MakeUnconvertedView(const WIDECHAR* Str)
 	return View;
 }
 
+// @pre Str contains only digits and the number is smaller than int64 max
+template<typename CharType>
+static constexpr int64 Atoi64(const CharType* Str, int32 Len)
+{
+    int64 N = 0;
+    for (int32 Idx = 0; Idx < Len; ++Idx)
+    {
+        N = 10 * N + Str[Idx] - '0';
+    }
+
+    return N;
+}
+
 /** Templated implementations of non-templated member functions, helps keep header clean */
 struct FNameHelper
 {
@@ -1736,33 +1749,23 @@ struct FNameHelper
 	template<typename CharType>
 	static uint32 ParseNumber(const CharType* Name, int32& InOutLen)
 	{
-		// The input is not guaranteed to be null-terminated so the number must be copied to
-		// a null-terminated buffer until we have an overload of Atoi64 that takes a length.
-		ANSICHAR Buffer[NAME_SIZE];
-		ANSICHAR* TerminatedDigits = Buffer + NAME_SIZE;
-		*--TerminatedDigits = '\0';
-
 		const int32 Len = InOutLen;
 		int32 Digits = 0;
 		for (const CharType* It = Name + Len - 1; It >= Name && *It >= '0' && *It <= '9'; --It)
 		{
-			*--TerminatedDigits = static_cast<ANSICHAR>(*It);
 			++Digits;
 		}
 
-		check(Digits < NAME_SIZE);
-
 		const CharType* FirstDigit = Name + Len - Digits;
-		if (Digits && Digits < Len && *(FirstDigit - 1) == '_')
+		static constexpr int32 MaxDigitsInt32 = 10;
+		if (Digits && Digits < Len && *(FirstDigit - 1) == '_' && Digits <= MaxDigitsInt32)
 		{
 			// check for the case where there are multiple digits after the _ and the first one
 			// is a 0 ("Rocket_04"). Can't split this case. (So, we check if the first char
 			// is not 0 or the length of the number is 1 (since ROcket_0 is valid)
 			if (Digits == 1 || *FirstDigit != '0')
 			{
-				// Attempt to convert what's following it to a number
-				// This relies on Name being null-terminated
-				int64 Number = FCStringAnsi::Atoi64(TerminatedDigits);
+				int64 Number = Atoi64(Name + Len - Digits, Digits);
 				if (Number < MAX_int32)
 				{
 					InOutLen -= 1 + Digits;
