Disable register warnings for C code included into C++17

#rnx

diff --git a/Engine/Plugins/Experimental/PythonScriptPlugin/Source/PythonScriptPlugin/Private/IncludePython.h b/Engine/Plugins/Experimental/PythonScriptPlugin/Source/PythonScriptPlugin/Private/IncludePython.h
--- a/Engine/Plugins/Experimental/PythonScriptPlugin/Source/PythonScriptPlugin/Private/IncludePython.h
+++ b/Engine/Plugins/Experimental/PythonScriptPlugin/Source/PythonScriptPlugin/Private/IncludePython.h
@@ -8,8 +8,10 @@
 #if WITH_PYTHON
 
 THIRD_PARTY_INCLUDES_START
+PRAGMA_DISABLE_REGISTER_WARNINGS
 #include "Python.h"
 #include "structmember.h"
+PRAGMA_ENABLE_REGISTER_WARNINGS
 THIRD_PARTY_INCLUDES_END
 
 DECLARE_STATS_GROUP(TEXT("Python"), STATGROUP_Python, STATCAT_Advanced);
diff --git a/Engine/Source/Runtime/Core/Public/Clang/ClangPlatformCompilerPreSetup.h b/Engine/Source/Runtime/Core/Public/Clang/ClangPlatformCompilerPreSetup.h
--- a/Engine/Source/Runtime/Core/Public/Clang/ClangPlatformCompilerPreSetup.h
+++ b/Engine/Source/Runtime/Core/Public/Clang/ClangPlatformCompilerPreSetup.h
@@ -106,6 +106,17 @@
 		_Pragma("clang diagnostic pop")
 #endif // PRAGMA_ENABLE_REORDER_WARNINGS
 
+#ifndef PRAGMA_DISABLE_REGISTER_WARNINGS
+	#define PRAGMA_DISABLE_REGISTER_WARNINGS \
+		_Pragma("clang diagnostic push") \
+		_Pragma("clang diagnostic ignored \"-Wregister\"")
+#endif // PRAGMA_DISABLE_REGISTER_WARNINGS
+
+#ifndef PRAGMA_ENABLE_REGISTER_WARNINGS
+	#define PRAGMA_ENABLE_REGISTER_WARNINGS \
+		_Pragma("clang diagnostic pop")
+#endif // PRAGMA_ENABLE_REGISTER_WARNINGS
+
 #ifndef PRAGMA_DISABLE_MACRO_REDEFINED_WARNINGS
 	#define PRAGMA_DISABLE_MACRO_REDEFINED_WARNINGS \
 		_Pragma("clang diagnostic push") \
diff --git a/Engine/Source/Runtime/Core/Public/HoloLens/HoloLensPlatformCompilerPreSetup.h b/Engine/Source/Runtime/Core/Public/HoloLens/HoloLensPlatformCompilerPreSetup.h
--- a/Engine/Source/Runtime/Core/Public/HoloLens/HoloLensPlatformCompilerPreSetup.h
+++ b/Engine/Source/Runtime/Core/Public/HoloLens/HoloLensPlatformCompilerPreSetup.h
@@ -62,6 +62,17 @@
 		__pragma(warning(pop))
 #endif // PRAGMA_ENABLE_REORDER_WARNINGS
 
+#ifndef PRAGMA_DISABLE_REGISTER_WARNINGS
+	#define PRAGMA_DISABLE_REGISTER_WARNINGS \
+		__pragma(warning(push)) \
+		__pragma(warning(disable: 5033)) /* 'register' is no longer a supported storage class */
+#endif // PRAGMA_DISABLE_REGISTER_WARNINGS
+
+#ifndef PRAGMA_ENABLE_REGISTER_WARNINGS
+	#define PRAGMA_ENABLE_REGISTER_WARNINGS \
+		__pragma(warning(pop))
+#endif // PRAGMA_ENABLE_REGISTER_WARNINGS
+
 #ifndef PRAGMA_POP
 	#define PRAGMA_POP \
 		__pragma(warning(pop))
diff --git a/Engine/Source/Runtime/Core/Public/Windows/WindowsPlatformCompilerPreSetup.h b/Engine/Source/Runtime/Core/Public/Windows/WindowsPlatformCompilerPreSetup.h
--- a/Engine/Source/Runtime/Core/Public/Windows/WindowsPlatformCompilerPreSetup.h
+++ b/Engine/Source/Runtime/Core/Public/Windows/WindowsPlatformCompilerPreSetup.h
@@ -107,6 +107,17 @@
 			__pragma(warning(pop))
 	#endif // PRAGMA_ENABLE_REORDER_WARNINGS
 
+	#ifndef PRAGMA_DISABLE_REGISTER_WARNINGS
+		#define PRAGMA_DISABLE_REGISTER_WARNINGS \
+			__pragma(warning(push)) \
+			__pragma(warning(disable: 5033)) /* 'register' is no longer a supported storage class */
+	#endif // PRAGMA_DISABLE_REGISTER_WARNINGS
+
+	#ifndef PRAGMA_ENABLE_REGISTER_WARNINGS
+		#define PRAGMA_ENABLE_REGISTER_WARNINGS \
+			__pragma(warning(pop))
+	#endif // PRAGMA_ENABLE_REGISTER_WARNINGS
+
 	#ifndef PRAGMA_POP
 		#define PRAGMA_POP \
 			__pragma(warning(pop))
