Copy up Dev-Enterprise @ CL11133089

diff --git a/Engine/Plugins/Enterprise/LidarPointCloud/Source/LidarPointCloudRuntime/Public/LidarPointCloudOctree.h b/Engine/Plugins/Enterprise/LidarPointCloud/Source/LidarPointCloudRuntime/Public/LidarPointCloudOctree.h
new file mode 100644
--- /dev/null
+++ b/Engine/Plugins/Enterprise/LidarPointCloud/Source/LidarPointCloudRuntime/Public/LidarPointCloudOctree.h
@@ -0,0 +1,526 @@
+// Copyright Epic Games, Inc. All Rights Reserved.
+
+#pragma once
+
+#include "CoreMinimal.h"
+#include "LidarPointCloudShared.h"
+#include "Misc/ScopeLock.h"
+#include "LidarPointCloudSettings.h"
+#include "HAL/ThreadSafeCounter64.h"
+#include "ConvexVolume.h"
+#include "Interfaces/Interface_CollisionDataProvider.h"
+
+class FLidarPointCloudOctree;
+struct FLidarPointCloudTraversalOctree;
+
+/**
+ * WARNING: Exercise caution when modifying the contents of the Octree, as it may be in use by the Rendering Thread via FPointCloudSceneProxy
+ * Use the FLidarPointCloudOctree::DataLock prior to such attempt
+ */
+
+ /**
+ Child ordering
+ 0	X- Y- Z-
+ 1	X- Y- Z+
+ 2	X- Y+ Z-
+ 3	X- Y+ Z+
+ 4	X+ Y- Z-
+ 5	X+ Y- Z+
+ 6	X+ Y+ Z-
+ 7	X+ Y+ Z+
+ */
+
+/**
+ * Represents a single octant in the tree.
+ */
+struct FLidarPointCloudOctreeNode
+{
+public:
+	mutable TArray<FLidarPointCloudPoint> AllocatedPoints;
+	mutable TArray<FLidarPointCloudPoint> PaddingPoints;
+
+private:
+	/** Depth of this node */
+	uint8 Depth;
+
+	/** Location of this node inside the parent node - see the Child Ordering at the top of the file */
+	uint8 LocationInParent;
+
+	/** Center point of this node. */
+	FVector Center;
+
+	/** Stores the children array */
+	// #todo: Change to TIndirectArray<> - investigate increased memory consumption, ~130 bytes / Node
+	TArray<FLidarPointCloudOctreeNode*> Children;
+
+	/** Marks the node for visibility recalculation next time it's necessary */
+	bool bVisibilityDirty;
+
+	/** Stores the number of visible points */
+	int32 NumVisiblePoints;
+
+	FCriticalSection MapLock;
+
+public:
+	FORCEINLINE FLidarPointCloudOctreeNode() : FLidarPointCloudOctreeNode(nullptr, nullptr, 0) {}
+	FORCEINLINE FLidarPointCloudOctreeNode(FLidarPointCloudOctree* Tree, FLidarPointCloudOctreeNode* Parent, const uint8& Depth) : FLidarPointCloudOctreeNode(Tree, Parent, Depth, 0, FVector::ZeroVector) {}
+	FLidarPointCloudOctreeNode(FLidarPointCloudOctree* Tree, FLidarPointCloudOctreeNode* Parent, const uint8& Depth, const uint8& LocationInParent, const FVector& Center);
+	~FLidarPointCloudOctreeNode();
+	FLidarPointCloudOctreeNode(const FLidarPointCloudOctreeNode&) = delete;
+	FLidarPointCloudOctreeNode(FLidarPointCloudOctreeNode&&) = delete;
+	FLidarPointCloudOctreeNode& operator=(const FLidarPointCloudOctreeNode&) = delete;
+	FLidarPointCloudOctreeNode& operator=(FLidarPointCloudOctreeNode&&) = delete;
+
+	/** Returns the sum of grid and padding points allocated to this node. */
+	int32 GetNumPoints() const { return AllocatedPoints.Num() + PaddingPoints.Num(); }
+
+	/** Returns the sum of visible grid and padding points allocated to this node. */
+	int32 GetNumVisiblePoints() const { return NumVisiblePoints; }
+
+	/** Calculates and returns the bounds of this node */
+	FBox GetBounds(const FLidarPointCloudOctree* Tree) const;
+
+	/** Calculates and returns the sphere bounds of this node */
+	FSphere GetSphereBounds(const FLidarPointCloudOctree* Tree) const;
+
+	/** Returns a pointer to the node at the given location, or null if one doesn't exist yet. */
+	FLidarPointCloudOctreeNode* GetChildNodeAtLocation(const uint8& Location) const;
+
+	void UpdateNumVisiblePoints();
+
+	/** Attempts to insert given points to this node or passes it to the children, otherwise. */
+	void InsertPoints(FLidarPointCloudOctree* Tree, const FLidarPointCloudPoint* Points, const int32& Count, ELidarPointCloudDuplicateHandling DuplicateHandling);
+	void InsertPoints(FLidarPointCloudOctree* Tree, FLidarPointCloudPoint** Points, const int32& Count, ELidarPointCloudDuplicateHandling DuplicateHandling);
+
+	/** Removes all points. */
+	void Empty(bool bRecursive = true);
+
+	/** Returns the maximum depth of any children of this node .*/
+	uint32 GetMaxDepth() const;
+
+	/** Returns the amount of memory used by this node */
+	int64 GetAllocatedSize(bool bRecursive) const;
+
+	void Serialize(FArchive& Ar, FLidarPointCloudOctree* Tree);
+
+	/**
+	 * Convenience function, to add point statistics to the Tree table.
+	 * If parameter set to negative value, GetNumPoints will be used
+	 */
+	void AddPointCount(FLidarPointCloudOctree* Tree, int32 PointCount = INT32_MIN);
+
+	friend FLidarPointCloudOctree;
+	friend struct FLidarPointCloudTraversalOctreeNode;
+};
+
+/**
+ * Used for efficient handling of point cloud data.
+ */
+class LIDARPOINTCLOUDRUNTIME_API FLidarPointCloudOctree
+{
+public:
+	/** Stores shared per-LOD node data. */
+	struct FSharedLODData
+	{
+		float Radius;
+		float RadiusSq;
+		FVector Extent;
+		FVector GridSize;
+		FVector NormalizationMultiplier;
+
+		FSharedLODData() {}
+		FSharedLODData(const FVector& InExtent);
+	};
+
+public:
+	/** Maximum allowed depth for any node */
+	static int32 MaxNodeDepth;
+
+	/** Maximum number of unallocated points to keep inside the node before they need to be converted in to a full child node */
+	static int32 MaxBucketSize;
+
+	/** Virtual grid resolution to divide the node into */
+	static int32 NodeGridResolution;
+
+	/** Used for thread safety between rendering and asset operations. */
+	mutable FCriticalSection DataLock;
+
+private:
+	FLidarPointCloudOctreeNode Root;
+	
+	/** Stores shared per-LOD node data. */
+	TArray<FSharedLODData> SharedData;
+	
+	/** Stores number of points per each LOD. */
+	TArray<FThreadSafeCounter64> PointCount;
+
+	/** Stores number of nodes per each LOD. */
+	TArray<FThreadSafeCounter> NodeCount;
+
+	/** Bounds with uniform extent - used for node generation. */
+	FBox UniformBounds;
+
+	/** Bounds covering all nodes. */
+	FBox Bounds;
+
+	/** Bounds covering all points. */
+	FBox PointsBounds;
+
+	/** Used to cache the Allocated Size. */
+	mutable int32 PreviousNodeCount;
+	mutable int64 PreviousPointCount;
+	mutable int64 PreviousAllocatedSize;
+
+	/** Used to notify any linked traversal octrees when they need to re-generate the data. */
+	TArray<FLidarPointCloudTraversalOctree*> LinkedTraversalOctrees;
+
+	/** Stores collision mesh data */
+	FTriMeshCollisionData CollisionMesh;
+
+public:
+	FLidarPointCloudOctree();
+	~FLidarPointCloudOctree();
+	FLidarPointCloudOctree(const FLidarPointCloudOctree&) = delete;
+	FLidarPointCloudOctree(FLidarPointCloudOctree&&) = delete;
+	FLidarPointCloudOctree& operator=(const FLidarPointCloudOctree&) = delete;
+	FLidarPointCloudOctree& operator=(FLidarPointCloudOctree&&) = delete;
+
+	/** Returns true if the Root node exists and has any data assigned. */
+	bool HasData() const { return Root.GetNumPoints() > 0; }
+	
+	/** Returns the number of different LODs. */
+	int32 GetNumLODs() const;
+
+	/** Returns the Octree bounds. */
+	FBox GetUniformBounds() const { return UniformBounds; }
+
+	/** Returns the Cloud bounds. */
+	FBox GetBounds() const { return Bounds; }
+
+	/** Returns the Cloud bounds. */
+	FBox GetPointsBounds() const { return PointsBounds; }
+
+	/** Recalculates and updates points bounds. */
+	void RefreshPointsBounds();
+
+	/** Returns the total number of points. */
+	int64 GetNumPoints() const;
+
+	/** Returns the total number of nodes. */
+	int32 GetNumNodes() const;
+
+	/** Returns the amount of memory used by this Octree */
+	int64 GetAllocatedSize() const;
+
+	/** Returns the grid cell size at root level. */
+	float GetRootCellSize() const { return SharedData[0].GridSize.GetMax(); }
+
+	/** Returns an estimated spacing between points */
+	float GetEstimatedPointSpacing() const;
+
+	/** Returns true, if the Octree has collision built */
+	bool HasCollisionData() const { return CollisionMesh.Vertices.Num() > 0; }
+
+	/** Builds collision using the accuracy provided */
+	void BuildCollision(const float& Accuracy, const bool& bVisibleOnly);
+
+	/** Removes collision mesh data */
+	void RemoveCollision();
+
+	/** Returns pointer to the collision data */
+	const FTriMeshCollisionData* GetCollisionData() const { return &CollisionMesh; }
+
+	/** Populates the given array with points from the tree */
+	void GetPoints(TArray<FLidarPointCloudPoint*>& Points, int64 StartIndex = 0, int64 Count = -1) const;
+
+	/** Populates the array with the list of points within the given sphere. */
+	void GetPointsInSphere(TArray<FLidarPointCloudPoint*>& SelectedPoints, const FSphere& Sphere, const bool& bVisibleOnly) const;
+
+	/** Populates the array with the list of pointers to points within the given box. */
+	void GetPointsInBox(TArray<FLidarPointCloudPoint*>& SelectedPoints, const FBox& Box, const bool& bVisibleOnly) const;
+
+	/** Populates the array with the list of points within the given frustum. */
+	void GetPointsInFrustum(TArray<FLidarPointCloudPoint*>& SelectedPoints, const FConvexVolume& Frustum, const bool& bVisibleOnly) const;
+
+	/** Populates the given array with copies of points from the tree */
+	void GetPointsAsCopies(TArray<FLidarPointCloudPoint>& Points, int64 StartIndex = 0, int64 Count = -1) const;
+
+	/** Populates the array with the list of points within the given sphere. */
+	void GetPointsInSphereAsCopies(TArray<FLidarPointCloudPoint>& SelectedPoints, const FSphere& Sphere, const bool& bVisibleOnly) const;
+
+	/** Populates the array with the list of pointers to points within the given box. */
+	void GetPointsInBoxAsCopies(TArray<FLidarPointCloudPoint>& SelectedPoints, const FBox& Box, const bool& bVisibleOnly) const;
+
+	/** Performs a raycast test against the point cloud. Returns the pointer if hit or nullptr otherwise. */
+	FLidarPointCloudPoint* RaycastSingle(const FLidarPointCloudRay& Ray, const float& Radius, const bool& bVisibleOnly);
+
+	/**
+	 * Performs a raycast test against the point cloud.
+	 * Populates OutHits array with the results.
+	 * Returns true it anything has been hit.
+	 */
+	bool RaycastMulti(const FLidarPointCloudRay& Ray, const float& Radius, const bool& bVisibleOnly, TArray<FLidarPointCloudPoint*>& OutHits);
+	bool RaycastMulti(const FLidarPointCloudRay& Ray, const float& Radius, const bool& bVisibleOnly, TArray<FLidarPointCloudPoint>& OutHits);
+
+	/** Returns true if there are any points within the given sphere. */
+	bool HasPointsInSphere(const FSphere& Sphere, const bool& bVisibleOnly) const;
+
+	/** Returns true if there are any points within the given box. */
+	bool HasPointsInBox(const FBox& Box, const bool& bVisibleOnly) const;
+
+	/** Sets visibility of points within the given sphere. */
+	void SetVisibilityOfPointsInSphere(const bool& bNewVisibility, const FSphere& Sphere);
+
+	/**
+	 * Sets visibility of points within the given sphere.
+	 * Async version - does not wait for completion before returning from the call.
+	 */
+	void SetVisibilityOfPointsInSphereAsync(const bool& bNewVisibility, const FSphere& Sphere, TFunction<void(void)> CompletionCallback = nullptr);
+
+	/** Sets visibility of points within the given box. */
+	void SetVisibilityOfPointsInBox(const bool& bNewVisibility, const FBox& Box);
+
+	/**
+	 * Sets visibility of points within the given box.
+	 * Async version - does not wait for completion before returning from the call.
+	 */
+	void SetVisibilityOfPointsInBoxAsync(const bool& bNewVisibility, const FBox& Box, TFunction<void(void)> CompletionCallback = nullptr);
+
+	/** Sets visibility of points hit by the given ray. */
+	void SetVisibilityOfPointsByRay(const bool& bNewVisibility, const FLidarPointCloudRay& Ray, const float& Radius);
+
+	/**
+	 * Sets visibility of points hit by the given ray.
+	 * Async version - does not wait for completion before returning from the call.
+	 */
+	void SetVisibilityOfPointsByRayAsync(const bool& bNewVisibility, const FLidarPointCloudRay& Ray, const float& Radius, TFunction<void(void)> CompletionCallback = nullptr);
+
+	/** Marks all points hidden */
+	void HideAll();
+
+	/** Marks all points visible */
+	void UnhideAll();
+
+	/**
+	 * Marks all points visible
+	 * Async version - does not wait for completion before returning from the call.
+	 */
+	void ResetVisibilityAsync(TFunction<void(void)> CompletionCallback = nullptr);
+
+	/** Executes the provided action on each of the points. */
+	void ExecuteActionOnAllPoints(TFunction<void(FLidarPointCloudPoint*)> Action, const bool& bVisibleOnly);
+
+	/**
+	 * Executes the provided action on each of the points.
+	 * Async version - does not wait for completion before returning from the call.
+	 */
+	void ExecuteActionOnAllPointsAsync(TFunction<void(FLidarPointCloudPoint*)> Action, const bool& bVisibleOnly, TFunction<void(void)> CompletionCallback = nullptr);
+
+	/** Executes the provided action on each of the points within the given sphere. */
+	void ExecuteActionOnPointsInSphere(TFunction<void(FLidarPointCloudPoint*)> Action, const FSphere& Sphere, const bool& bVisibleOnly);
+
+	/**
+	 * Executes the provided action on each of the points within the given sphere.
+	 * Async version - does not wait for completion before returning from the call.
+	 */
+	void ExecuteActionOnPointsInSphereAsync(TFunction<void(FLidarPointCloudPoint*)> Action, const FSphere& Sphere, const bool& bVisibleOnly, TFunction<void(void)> CompletionCallback = nullptr);
+
+	/** Executes the provided action on each of the points within the given box. */
+	void ExecuteActionOnPointsInBox(TFunction<void(FLidarPointCloudPoint*)> Action, const FBox& Box, const bool& bVisibleOnly);
+
+	/**
+	 * Executes the provided action on each of the points within the given box.
+	 * Async version - does not wait for completion before returning from the call.
+	 */
+	void ExecuteActionOnPointsInBoxAsync(TFunction<void(FLidarPointCloudPoint*)> Action, const FBox& Box, const bool& bVisibleOnly, TFunction<void(void)> CompletionCallback = nullptr);
+
+	/** Executes the provided action on each of the points hit by the given ray. */
+	void ExecuteActionOnPointsByRay(TFunction<void(FLidarPointCloudPoint*)> Action, const FLidarPointCloudRay& Ray, const float& Radius, bool bVisibleOnly);
+
+	/**
+	 * Executes the provided action on each of the points hit by the given ray.
+	 * Async version - does not wait for completion before returning from the call.
+	 */
+	void ExecuteActionOnPointsByRayAsync(TFunction<void(FLidarPointCloudPoint*)> Action, const FLidarPointCloudRay& Ray, const float& Radius, bool bVisibleOnly, TFunction<void(void)> CompletionCallback = nullptr);
+
+	/** Applies the given color to all points */
+	void ApplyColorToAllPoints(const FColor& NewColor, const bool& bVisibleOnly);
+
+	/** Applies the given color to all points within the sphere */
+	void ApplyColorToPointsInSphere(const FColor& NewColor, const FSphere& Sphere, const bool& bVisibleOnly);
+
+	/** Applies the given color to all points within the box */
+	void ApplyColorToPointsInBox(const FColor& NewColor, const FBox& Box, const bool& bVisibleOnly);
+
+	/** Applies the given color to all points hit by the given ray */
+	void ApplyColorToPointsByRay(const FColor& NewColor, const FLidarPointCloudRay& Ray, const float& Radius, bool bVisibleOnly);
+
+	/**
+	 * This should to be called if any manual modification to individual points' visibility has been made.
+	 * If not marked dirty, the rendering may work suboptimally.
+	 */
+	void MarkPointVisibilityDirty();
+
+	/** Initializes the Octree properties. */
+	void Initialize(const FBox& InBounds);
+
+	/** Applies given offset to all contained points. */
+	void ShiftPointsBy(FDoubleVector Offset, bool bRefreshPointsBounds);
+
+	/** Inserts the given point into the Octree structure, internally thread-safe. */
+	FORCEINLINE void InsertPoint(const FLidarPointCloudPoint* Point, ELidarPointCloudDuplicateHandling DuplicateHandling, bool bRefreshPointsBounds) { InsertPoints(Point, 1, DuplicateHandling, bRefreshPointsBounds); }
+
+	/** Inserts group of points into the Octree structure, internally thread-safe. */
+	template <typename T>
+	void InsertPoints(T Points, const int32& Count, ELidarPointCloudDuplicateHandling DuplicateHandling, bool bRefreshPointsBounds)
+	{
+		Root.InsertPoints(this, Points, Count, DuplicateHandling);
+		MarkTraversalOctreesForInvalidation();
+		if (bRefreshPointsBounds)
+		{
+			RefreshPointsBounds();
+		}
+	}
+
+	/** Attempts to remove the given point.  */
+	void RemovePoint(const FLidarPointCloudPoint* Point, bool bRefreshPointsBounds);
+	void RemovePoint(FLidarPointCloudPoint Point, bool bRefreshPointsBounds);
+
+	/** Removes points in bulk */
+	void RemovePoints(TArray<FLidarPointCloudPoint*>& Points, bool bRefreshPointsBounds);
+
+	/** Removes all points within the given sphere */
+	void RemovePointsInSphere(const FSphere& Sphere, const bool& bVisibleOnly, bool bRefreshPointsBounds);
+
+	/** Removes all points within the given box */
+	void RemovePointsInBox(const FBox& Box, const bool& bVisibleOnly, bool bRefreshPointsBounds);
+
+	/** Removes all points hit by the given ray */
+	void RemovePointsByRay(const FLidarPointCloudRay& Ray, const float& Radius, const bool& bVisibleOnly, bool bRefreshPointsBounds);
+
+	/** Removes all hidden points */
+	void RemoveHiddenPoints(bool bRefreshPointsBounds);
+
+	/** Removes all points and, optionally, all nodes except for the root node. Retains the bounds. */
+	void Empty(bool bDestroyNodes);
+
+	/** Adds the given traversal octree to the list of linked octrees. */
+	void RegisterTraversalOctree(FLidarPointCloudTraversalOctree* TraversalOctree)
+	{
+		if (TraversalOctree)
+		{
+			LinkedTraversalOctrees.Add(TraversalOctree);
+		}
+	}
+
+	/** Removes the given traversal octree from the list */
+	void UnregisterTraversalOctree(FLidarPointCloudTraversalOctree* TraversalOctree)
+	{
+		if (TraversalOctree)
+		{
+			LinkedTraversalOctrees.Remove(TraversalOctree);
+		}
+	}
+
+private:
+	/**
+	 * Applies new bounds.
+	 * Returns true if successful
+	 */
+	bool SetNewBounds(const FBox& InBounds);
+	
+	void RemovePoint_Internal(FLidarPointCloudOctreeNode* Node, int32 Index, bool bAllocatedPoint);
+
+	/** Notifies all linked traversal octrees that they should invalidate and regenerate the data. */
+	void MarkTraversalOctreesForInvalidation();
+
+	friend FArchive& operator<<(FArchive& Ar, FLidarPointCloudOctree& O);
+
+	friend FLidarPointCloudOctreeNode;
+	friend struct FLidarPointCloudTraversalOctree;
+};
+
+/**
+ * Represents a single octant in the traversal tree.
+ */
+struct FLidarPointCloudTraversalOctreeNode
+{
+	/** Pointer to the target node. */
+	FLidarPointCloudOctreeNode* DataNode;
+
+	/** Stores the center of the target node in World space. */
+	FVector Center;
+
+	/** Depth of this node */
+	uint8 Depth;
+
+	/** Calculated for use with adaptive sprite scaling */
+	uint8 VirtualDepth;
+
+	// #todo: Can be included as a first item in the Children array to save 8 bytes
+	FLidarPointCloudTraversalOctreeNode* Parent;
+
+	/** Stores the children array */
+	FLidarPointCloudTraversalOctreeNode* Children;
+	uint8 NumChildren : 4;
+
+	/** Holds true if the node has been selected for rendering. */
+	uint8 bSelected : 4;
+
+	FLidarPointCloudTraversalOctreeNode();
+	~FLidarPointCloudTraversalOctreeNode();
+	FLidarPointCloudTraversalOctreeNode(const FLidarPointCloudTraversalOctreeNode&) = delete;
+	FLidarPointCloudTraversalOctreeNode(FLidarPointCloudTraversalOctreeNode&&) = delete;
+	FLidarPointCloudTraversalOctreeNode& operator=(const FLidarPointCloudTraversalOctreeNode&) = delete;
+	FLidarPointCloudTraversalOctreeNode& operator=(FLidarPointCloudTraversalOctreeNode&&) = delete;
+
+	/** Builds the traversal version of the given node. */
+	void Build(FLidarPointCloudOctreeNode* Node, const FTransform& LocalToWorld);
+
+	// #refactor: This is ugly - refactor to avoid the const or call directly in the buffer update with no mutable
+	/** Calculates virtual depth of this node, to be used to estimate the best sprite size */
+	void CalculateVirtualDepth(const TArray<float>& LevelWeights, const float& VDMultiplier, const float& PointSizeBias);
+};
+
+/**
+ * Used as a traversal tree for node selection during rendering
+ */
+struct FLidarPointCloudTraversalOctree
+{
+	FLidarPointCloudTraversalOctreeNode Root;
+
+	/** Stores per-LOD bounds in World space. */
+	TArray<float> RadiiSq;
+	TArray<FVector> Extents;
+
+	/** Stores the number of LODs. */
+	uint8 NumLODs;
+
+	/** Stores number of points per each LOD. */
+	TArray<int64> PointCount;
+
+	/** Total number of points. */
+	int64 NumPoints;
+
+	/** Pointer to the source Octree */
+	FLidarPointCloudOctree* Octree;
+
+	bool bValid;
+
+	/** Build the Traversal tree from the Octree provided */
+	FLidarPointCloudTraversalOctree(FLidarPointCloudOctree* Octree, const FTransform& LocalToWorld);
+
+	~FLidarPointCloudTraversalOctree();
+	FLidarPointCloudTraversalOctree(const FLidarPointCloudTraversalOctree&) = delete;
+	FLidarPointCloudTraversalOctree(FLidarPointCloudTraversalOctree&&) = delete;
+	FLidarPointCloudTraversalOctree& operator=(const FLidarPointCloudTraversalOctree&) = delete;
+	FLidarPointCloudTraversalOctree& operator=(FLidarPointCloudTraversalOctree&&) = delete;
+
+	/** Selects and returns the subset of visible nodes for rendering. */
+	TArray<FLidarPointCloudTraversalOctreeNode*> GetVisibleNodes(const class FSceneView* View, const float& MinScreenSizeSq, const int32& PointBudget, const float& ScreenCenterImportance, const int32& MinDepth, const int32& MaxDepth);
+	
+	void CalculateSpriteSize(TArray<FLidarPointCloudTraversalOctreeNode*>& SelectedNodes, const float& PointSizeBias);
+
+	/** Returns normalized histogram of level weights, one for each LOD. Used for point scaling */
+	TArray<float> CalculateLevelWeights() const;
+};
diff --git a/Engine/Plugins/Enterprise/LidarPointCloud/Source/LidarPointCloudRuntime/Public/LidarPointCloudSettings.h b/Engine/Plugins/Enterprise/LidarPointCloud/Source/LidarPointCloudRuntime/Public/LidarPointCloudSettings.h
new file mode 100644
--- /dev/null
+++ b/Engine/Plugins/Enterprise/LidarPointCloud/Source/LidarPointCloudRuntime/Public/LidarPointCloudSettings.h
@@ -0,0 +1,116 @@
+// Copyright Epic Games, Inc. All Rights Reserved.
+
+#pragma once
+
+#include "CoreMinimal.h"
+#include "UObject/Object.h"
+#include "LidarPointCloudSettings.generated.h"
+
+UENUM(BlueprintType)
+enum class ELidarPointCloudDuplicateHandling : uint8
+{
+	/** Keeps any duplicates found */
+	Ignore,
+	/** Keeps the first point and skips any further duplicates */
+	SelectFirst,
+	/** Selects the brightest of the duplicates */
+	SelectBrighter
+};
+
+UCLASS(config=Engine, defaultconfig)
+class LIDARPOINTCLOUDRUNTIME_API ULidarPointCloudSettings : public UObject
+{
+	GENERATED_BODY()
+
+public:
+	/** Determines how to handle duplicate points (distance < 0.0001). */
+	UPROPERTY(config, EditAnywhere, Category=Octree)
+	ELidarPointCloudDuplicateHandling DuplicateHandling;
+
+	/** Maximum distance between points, within which they are considered to be duplicates */
+	UPROPERTY(config, EditAnywhere, Category=Octree, meta = (ClampMin = "0"))
+	float MaxDistanceForDuplicate;
+
+	/**
+	 * Maximum number of unallocated points to keep inside the node before they need to be converted in to a full child node.
+	 * Lower values will provide finer LOD control at the expense of system RAM and CPU time.
+	 */
+	UPROPERTY(config, EditAnywhere, Category=Octree)
+	int32 MaxBucketSize;
+
+	/**
+	 * Virtual grid resolution to divide the node into.
+	 * Lower values will provide finer LOD control at the expense of system RAM and CPU time.
+	 */
+	UPROPERTY(config, EditAnywhere, Category=Octree)
+	int32 NodeGridResolution;
+
+	/** Enabling this will allow usage of multiple threads during import and processing. */
+	UPROPERTY(config, EditAnywhere, Category=Performance)
+	bool bUseMultithreading;
+
+	/** Determines the maximum amount of points to process in a single batch when using multi-threading. */
+	UPROPERTY(config, EditAnywhere, Category=Performance)
+	int32 MultithreadingInsertionBatchSize;
+
+	/** Enabling this will allow editor to import the point clouds in the background, without blocking the main thread. */
+	UPROPERTY(config, EditAnywhere, Category=Performance)
+	bool bUseAsyncImport;
+
+	/** Determines the maximum size of the buffer to use during importing. */
+	UPROPERTY(config, EditAnywhere, Category=Performance)
+	int32 MaxImportBufferSize;
+
+	/** Determines the maximum size of the buffer to use during exporting. */
+	UPROPERTY(config, EditAnywhere, Category=Performance)
+	int32 ExportBatchSize;
+
+	/** Affects the size of per-thread data for the meshing algorithm. */
+	UPROPERTY(config, EditAnywhere, Category=Collision)
+	int32 MeshingBatchSize;
+
+	/**
+	 * Automatically centers the cloud on import.
+	 * Caution: Preserving original coordinates may cause noticeable precision loss, if the values are too large.
+	 * Should you experience point 'banding' effect, please re-import your cloud with centering enabled.
+	 */
+	UPROPERTY(config, EditAnywhere, Category=Import)
+	bool bAutoCenterOnImport;
+
+	/** Scale to apply during import */
+	UPROPERTY(config, EditAnywhere, Category=Import, meta = (ClampMin = "0.0001"))
+	float ImportScale;
+
+	/**
+	 * Enables 8-bit color detection in LAS files.
+	 * This will slow down the import a little.
+	 */
+	UPROPERTY(config, EditAnywhere, Category=Import, meta = (DisplayName = "Enable 8-Bit LAS Detection"))
+	bool bEnable8BitLASDetection;
+
+	/** Determines the maximum number of points to scan when analyzing the LAS data */
+	UPROPERTY(config, EditAnywhere, Category=Import, meta = (DisplayName = "Max Number Of Points To Scan (LAS)"))
+	int32 MaxNumberOfPointsToScanLAS;
+
+	/** Determines the maximum number of points to scan when analyzing the ASCII data */
+	UPROPERTY(config, EditAnywhere, Category=Import, meta = (DisplayName = "Max Number Of Points To Scan (ASCII)"))
+	int32 MaxNumberOfPointsToScanASCII;
+
+	/** Scale to apply during export. In most cases, this should be equal to an inverted ImportScale */
+	UPROPERTY(config, EditAnywhere, Category=Export, meta = (ClampMin = "0.0001"))
+	float ExportScale;
+
+	/**
+	 * Enabling this will store pre-processed version of the import data as a *.tmp file, and will attempt to use it upon re-import.
+	 * Useful for faster debug iteration on large and slow to import (especially ASCII based) cloud assets.
+	 */
+	UPROPERTY(config, EditAnywhere, Category=Debug)
+	bool bUseIOCaching;
+
+	/** Enabling this will cause Structure and Index buffers to write to the log file anytime they are resized. */
+	UPROPERTY(config, EditAnywhere, Category=Debug)
+	bool bLogBufferExpansion;
+
+public:
+	ULidarPointCloudSettings();
+};
\ No newline at end of file
diff --git a/Engine/Plugins/Enterprise/LidarPointCloud/Source/LidarPointCloudRuntime/Public/LidarPointCloudShared.h b/Engine/Plugins/Enterprise/LidarPointCloud/Source/LidarPointCloudRuntime/Public/LidarPointCloudShared.h
new file mode 100644
--- /dev/null
+++ b/Engine/Plugins/Enterprise/LidarPointCloud/Source/LidarPointCloudRuntime/Public/LidarPointCloudShared.h
@@ -0,0 +1,407 @@
+// Copyright Epic Games, Inc. All Rights Reserved.
+
+#pragma once
+
+#include "CoreMinimal.h"
+#include "LidarPointCloudShared.generated.h"
+
+DECLARE_LOG_CATEGORY_EXTERN(LogLidarPointCloud, Log, All);
+DECLARE_STATS_GROUP(TEXT("Lidar Point Cloud"), STATGROUP_LidarPointCloud, STATCAT_Advanced);
+
+#define PC_LOG(Format, ...) UE_LOG(LogLidarPointCloud, Log, TEXT(Format), ##__VA_ARGS__)
+#define PC_WARNING(Format, ...) UE_LOG(LogLidarPointCloud, Warning, TEXT(Format), ##__VA_ARGS__)
+#define PC_ERROR(Format, ...) UE_LOG(LogLidarPointCloud, Error, TEXT(Format), ##__VA_ARGS__)
+
+#if !CPP
+USTRUCT(noexport)
+struct FDoubleVector
+{
+	UPROPERTY()
+	double X;
+
+	UPROPERTY()
+	double Y;
+
+	UPROPERTY()
+	double Z;
+};
+#endif
+
+struct LIDARPOINTCLOUDRUNTIME_API FDoubleVector
+{
+	double X;
+	double Y;
+	double Z;
+
+	/** A zero vector (0,0,0) */
+	static const FDoubleVector ZeroVector;
+
+	/** One vector (1,1,1) */
+	static const FDoubleVector OneVector;
+
+	/** World up vector (0,0,1) */
+	static const FDoubleVector UpVector;
+
+	/** Unreal forward vector (1,0,0) */
+	static const FDoubleVector ForwardVector;
+
+	/** Unreal right vector (0,1,0) */
+	static const FDoubleVector RightVector;
+
+	FORCEINLINE FDoubleVector() {}
+
+	explicit FORCEINLINE FDoubleVector(double InD) : X(InD), Y(InD), Z(InD) {}
+
+	FORCEINLINE FDoubleVector(double InX, double InY, double InZ) : X(InX), Y(InY), Z(InZ) {}
+
+	FORCEINLINE FDoubleVector(const FVector& V) : X(V.X), Y(V.Y), Z(V.Z) {}
+
+	FORCEINLINE FDoubleVector operator-() const
+	{
+		return FDoubleVector(-X, -Y, -Z);
+	}
+
+	FORCEINLINE FDoubleVector operator+(const FDoubleVector& V) const
+	{
+		return FDoubleVector(X + V.X, Y + V.Y, Z + V.Z);
+	}
+
+	FORCEINLINE FDoubleVector operator-(const FDoubleVector& V) const
+	{
+		return FDoubleVector(X - V.X, Y - V.Y, Z - V.Z);
+	}
+
+	FORCEINLINE FDoubleVector operator+(const FVector& V) const
+	{
+		return FDoubleVector(X + V.X, Y + V.Y, Z + V.Z);
+	}
+
+	FORCEINLINE FDoubleVector operator-(const FVector& V) const
+	{
+		return FDoubleVector(X - V.X, Y - V.Y, Z - V.Z);
+	}
+
+	FORCEINLINE FDoubleVector operator+=(const FDoubleVector& V)
+	{
+		X += V.X; Y += V.Y; Z += V.Z;
+		return *this;
+	}
+
+	FORCEINLINE FDoubleVector operator-=(const FDoubleVector& V)
+	{
+		X -= V.X; Y -= V.Y; Z -= V.Z;
+		return *this;
+	}
+
+	FORCEINLINE FDoubleVector operator+=(const FVector& V)
+	{
+		X += V.X; Y += V.Y; Z += V.Z;
+		return *this;
+	}
+
+	FORCEINLINE FDoubleVector operator-=(const FVector& V)
+	{
+		X -= V.X; Y -= V.Y; Z -= V.Z;
+		return *this;
+	}
+
+	FORCEINLINE FDoubleVector operator*=(double Scale)
+	{
+		X *= Scale; Y *= Scale; Z *= Scale;
+		return *this;
+	}
+
+	FORCEINLINE FDoubleVector operator*(double Scale) const
+	{
+		return FDoubleVector(X * Scale, Y * Scale, Z * Scale);
+	}
+
+	FORCEINLINE FDoubleVector operator*(const FVector& V) const
+	{
+		return FDoubleVector(X * V.X, Y * V.Y, Z * V.Z);
+	}
+
+	FORCEINLINE FDoubleVector operator*(const FDoubleVector& V) const
+	{
+		return FDoubleVector(X * V.X, Y * V.Y, Z * V.Z);
+	}
+
+	FORCEINLINE FDoubleVector operator*(const FIntVector& V) const
+	{
+		return FDoubleVector(X * V.X, Y * V.Y, Z * V.Z);
+	}
+
+	FORCEINLINE FDoubleVector operator/(int32 Scale) const
+	{
+		return FDoubleVector(X / Scale, Y / Scale, Z / Scale);
+	}
+
+	FORCEINLINE FDoubleVector operator/(double Scale) const
+	{
+		return FDoubleVector(X / Scale, Y / Scale, Z / Scale);
+	}
+
+	FORCEINLINE bool Equals(const FDoubleVector& V, float Tolerance = KINDA_SMALL_NUMBER) const
+	{
+		return FMath::Abs(X - V.X) <= Tolerance && FMath::Abs(Y - V.Y) <= Tolerance && FMath::Abs(Z - V.Z) <= Tolerance;
+	}
+
+	FORCEINLINE bool IsZero(float Tolerance = KINDA_SMALL_NUMBER) const
+	{
+		return Equals(FDoubleVector::ZeroVector, Tolerance);
+	}
+
+	FORCEINLINE bool IsNearlyZero(float Tolerance) const
+	{
+		return FMath::Abs(X) <= Tolerance && FMath::Abs(Y) <= Tolerance && FMath::Abs(Z) <= Tolerance;
+	}
+
+	FORCEINLINE FVector ToVector() const { return FVector(X, Y, Z); }
+	FORCEINLINE FIntVector ToIntVector() const { return FIntVector(X, Y, Z); }
+};
+
+#pragma pack(push)
+#pragma pack(1)
+USTRUCT(BlueprintType)
+struct LIDARPOINTCLOUDRUNTIME_API FLidarPointCloudPoint
+{
+	GENERATED_BODY()
+
+public:
+	UPROPERTY(EditAnywhere, BlueprintReadWrite, Category = "Lidar Point Cloud Point")
+	FVector Location;
+
+	UPROPERTY(EditAnywhere, BlueprintReadWrite, Category = "Lidar Point Cloud Point")
+	FColor Color;
+
+	UPROPERTY(EditAnywhere, BlueprintReadWrite, Category = "Lidar Point Cloud Point")
+	uint8 bVisible : 1;
+
+	/** Valid range is 0 - 31. */
+	uint8 ClassificationID : 5;
+
+private:
+	uint8 bSelected : 1;
+	uint8 bMarkedForDeletion : 1;
+
+public:
+	FLidarPointCloudPoint()
+		: Location(FVector::ZeroVector)
+		, Color(FColor::White)
+		, bVisible(true)
+		, ClassificationID(0)
+		, bSelected(false)
+		, bMarkedForDeletion(false)
+	{
+	}
+	FLidarPointCloudPoint(const float& X, const float& Y, const float& Z)
+		: FLidarPointCloudPoint()
+	{
+		Location.X = X;
+		Location.Y = Y;
+		Location.Z = Z;
+	}
+	FLidarPointCloudPoint(const float& X, const float& Y, const float& Z, const float& I)
+		: FLidarPointCloudPoint(X, Y, Z)
+	{
+		Color.A = FMath::FloorToInt(FMath::Clamp(I, 0.0f, 1.0f) * 255.999f);
+	}
+	FLidarPointCloudPoint(const float& X, const float& Y, const float& Z, const float& R, const float& G, const float& B, const float& A = 1.0f)
+		: FLidarPointCloudPoint(X, Y, Z)
+	{
+		Color = FLinearColor(R, G, B, A).ToFColor(false);
+	}
+	FLidarPointCloudPoint(const FVector& Location) : FLidarPointCloudPoint(Location.X, Location.Y, Location.Z) {}
+	FLidarPointCloudPoint(const FVector& Location, const float& R, const float& G, const float& B, const float& A = 1.0f)
+		: FLidarPointCloudPoint(Location)
+	{
+		Color = FLinearColor(R, G, B, A).ToFColor(false);
+	}
+	FLidarPointCloudPoint(const FVector& Location, const float& R, const float& G, const float& B, const float& A, const uint8& ClassificationID)
+		: FLidarPointCloudPoint(Location, R, G, B, A)
+	{
+		this->ClassificationID = ClassificationID;
+	}
+	FLidarPointCloudPoint(const FVector& Location, const FColor& Color, const bool& bVisible, const uint8& ClassificationID)
+		: FLidarPointCloudPoint(Location)
+	{
+		this->Color = Color;
+		this->bVisible = bVisible;
+		this->ClassificationID = ClassificationID;
+	}
+	FLidarPointCloudPoint(const FLidarPointCloudPoint& Other)
+		: FLidarPointCloudPoint(Other.Location, Other.Color, Other.bVisible, Other.ClassificationID)
+	{
+	}
+
+	FORCEINLINE void CopyFrom(const FLidarPointCloudPoint& Other)
+	{
+		Location = Other.Location;
+		Color = Other.Color;
+		bVisible = Other.bVisible;
+		ClassificationID = Other.ClassificationID;
+	}
+
+	bool operator==(const FLidarPointCloudPoint& P) const { return Location == P.Location && Color == P.Color && bVisible == P.bVisible && ClassificationID == P.ClassificationID; }
+
+	friend FArchive& operator<<(FArchive& Ar, FLidarPointCloudPoint& P);
+	friend class FLidarPointCloudOctree;
+	friend class FLidarPointCloudInstanceBuffer;
+#if WITH_EDITOR
+	friend class FLidarPointCloudEditor;
+#endif
+};
+#pragma pack(pop)
+
+/** Used in blueprint latent function execution */
+UENUM(BlueprintType)
+enum class ELidarPointCloudAsyncMode : uint8
+{
+	Success,
+	Failure,
+	Progress
+};
+
+/** Used for Raycasting */
+struct FLidarPointCloudRay
+{
+public:
+	FVector Origin;
+
+private:
+	FVector Direction;
+	FVector InversedDirection;
+
+public:
+	FLidarPointCloudRay() : FLidarPointCloudRay(FVector::ZeroVector, FVector::ForwardVector) {}
+	FLidarPointCloudRay(const FVector& Origin, const FVector& Direction) : Origin(Origin)
+	{
+		SetDirection(Direction);
+	}
+
+	FLidarPointCloudRay& TransformBy(FTransform Transform)
+	{
+		Origin = Transform.TransformPosition(Origin);
+		SetDirection(Transform.TransformVector(Direction));
+		return *this;
+	}
+
+	FORCEINLINE FVector GetDirection() const { return Direction; }
+	FORCEINLINE void SetDirection(const FVector& NewDirection)
+	{
+		Direction = NewDirection;
+		InversedDirection = FVector::OneVector / Direction;
+	}
+
+	/** An Efficient and Robust Ray-Box Intersection Algorithm. Amy Williams et al. 2004. */
+	FORCEINLINE bool Intersects(const FBox& Box) const
+	{
+		float tmin, tmax, tymin, tymax, tzmin, tzmax;
+
+		tmin = ((InversedDirection.X < 0 ? Box.Max.X : Box.Min.X) - Origin.X) * InversedDirection.X;
+		tmax = ((InversedDirection.X < 0 ? Box.Min.X : Box.Max.X) - Origin.X) * InversedDirection.X;
+		tymin = ((InversedDirection.Y < 0 ? Box.Max.Y : Box.Min.Y) - Origin.Y) * InversedDirection.Y;
+		tymax = ((InversedDirection.Y < 0 ? Box.Min.Y : Box.Max.Y) - Origin.Y) * InversedDirection.Y;
+
+		if ((tmin > tymax) || (tymin > tmax))
+		{
+			return false;
+		}
+
+		if (tymin > tmin)
+		{
+			tmin = tymin;
+		}
+
+		if (tymax < tmax)
+		{
+			tmax = tymax;
+		}
+
+		tzmin = ((InversedDirection.Z < 0 ? Box.Max.Z : Box.Min.Z) - Origin.Z) * InversedDirection.Z;
+		tzmax = ((InversedDirection.Z < 0 ? Box.Min.Z : Box.Max.Z) - Origin.Z) * InversedDirection.Z;
+
+		if ((tmin > tzmax) || (tzmin > tmax))
+		{
+			return false;
+		}
+
+		return true;
+	}
+	FORCEINLINE bool Intersects(const FLidarPointCloudPoint* Point, const float& RadiusSq) const
+	{
+		const FVector L = Point->Location - Origin;
+		const float tca = FVector::DotProduct(L, Direction);
+		const float d2 = FVector::DotProduct(L, L) - tca * tca;
+
+		return d2 <= RadiusSq;
+	}
+};
+
+struct FBenchmarkTimer
+{
+	static void Reset()
+	{
+		Time = FPlatformTime::Seconds();
+	}
+	static double Split(uint8 Decimal = 2)
+	{
+		double Now = FPlatformTime::Seconds();
+		double Delta = Now - Time;
+		Time = Now;
+
+		uint32 Multiplier = FMath::Pow(10, Decimal);
+
+		return FMath::RoundToDouble(Delta * Multiplier * 1000) / Multiplier;
+	}
+	static void Log(FString Text, uint8 Decimal = 2)
+	{
+		const double SplitTime = Split(Decimal);
+		PC_LOG("%s: %f ms", *Text, SplitTime);
+	}
+
+private:
+	static double Time;
+};
+
+struct FScopeBenchmarkTimer
+{
+public:
+	bool bActive;
+
+private:
+	double Time;
+	FString Label;
+	float* OutTimer;
+
+public:
+	FScopeBenchmarkTimer(const FString& Label)
+		: bActive(true)
+		, Time(FPlatformTime::Seconds())
+		, Label(Label)
+		, OutTimer(nullptr)
+	{
+	}
+	FScopeBenchmarkTimer(float* OutTimer)
+		: bActive(true)
+		, Time(FPlatformTime::Seconds())
+		, OutTimer(OutTimer)
+	{
+	}
+	~FScopeBenchmarkTimer()
+	{
+		if (bActive)
+		{
+			float Delta = FMath::RoundToDouble((FPlatformTime::Seconds() - Time) * 100000) * 0.01;
+
+			if (OutTimer)
+			{
+				*OutTimer += Delta;
+			}
+			else
+			{
+				PC_LOG("%s: %f ms", *Label, Delta);
+			}
+		}
+	}
+};
\ No newline at end of file
diff --git a/Engine/Plugins/Enterprise/VariantManager/Source/VariantManager/Private/DisplayNodes/VariantManagerFunctionPropertyNode.cpp b/Engine/Plugins/Enterprise/VariantManager/Source/VariantManager/Private/DisplayNodes/VariantManagerFunctionPropertyNode.cpp
--- a/Engine/Plugins/Enterprise/VariantManager/Source/VariantManager/Private/DisplayNodes/VariantManagerFunctionPropertyNode.cpp
+++ b/Engine/Plugins/Enterprise/VariantManager/Source/VariantManager/Private/DisplayNodes/VariantManagerFunctionPropertyNode.cpp
@@ -252,7 +252,7 @@ TSharedRef<SWidget> FVariantManagerFunctionPropertyNode::GetMenuContent()
 				}
 				else
 				{
-					TooltipText = LOCTEXT("InvalidTooltip", "Function must be valid and have either no input pins\nor a single object reference input pin!");
+					TooltipText = LOCTEXT("InvalidTooltip", "Function must be valid and have either no input pins or 4 input pins with the correct types:\nThe bound actor's class, Level Variant Sets, Variant Set and Variant");
 				}
 			}
 			else
diff --git a/Engine/Plugins/Enterprise/VariantManager/Source/VariantManager/Private/VariantManager.cpp b/Engine/Plugins/Enterprise/VariantManager/Source/VariantManager/Private/VariantManager.cpp
--- a/Engine/Plugins/Enterprise/VariantManager/Source/VariantManager/Private/VariantManager.cpp
+++ b/Engine/Plugins/Enterprise/VariantManager/Source/VariantManager/Private/VariantManager.cpp
@@ -121,7 +121,17 @@ UK2Node_FunctionEntry* FVariantManager::CreateDirectorFunction(ULevelVariantSets
 		FEdGraphPinType PinType;
 		PinType.PinCategory = PinClassType && PinClassType->IsChildOf(UInterface::StaticClass()) ? UEdGraphSchema_K2::PC_Interface : UEdGraphSchema_K2::PC_Object;
 		PinType.PinSubCategoryObject = PinClassType ? PinClassType : UObject::StaticClass();
-		EntryNode->CreateUserDefinedPin(FName(TARGET_PIN_NAME), PinType, EGPD_Output, true);
+		EntryNode->CreateUserDefinedPin(TARGET_PIN_NAME, PinType, EGPD_Output, true);
+
+		PinType.PinSubCategoryObject = ULevelVariantSets::StaticClass();
+		EntryNode->CreateUserDefinedPin(LEVEL_VARIANT_SETS_PIN_NAME, PinType, EGPD_Output, true);
+
+		PinType.PinSubCategoryObject = UVariantSet::StaticClass();
+		EntryNode->CreateUserDefinedPin(VARIANT_SET_PIN_NAME, PinType, EGPD_Output, true);
+
+		PinType.PinSubCategoryObject = UVariant::StaticClass();
+		EntryNode->CreateUserDefinedPin(VARIANT_PIN_NAME, PinType, EGPD_Output, true);
+
 		EntryNode->ReconstructNode();
 
 		return EntryNode;
diff --git a/Engine/Plugins/Enterprise/VariantManagerContent/Source/VariantManagerContent/Private/FunctionCaller.cpp b/Engine/Plugins/Enterprise/VariantManagerContent/Source/VariantManagerContent/Private/FunctionCaller.cpp
--- a/Engine/Plugins/Enterprise/VariantManagerContent/Source/VariantManagerContent/Private/FunctionCaller.cpp
+++ b/Engine/Plugins/Enterprise/VariantManagerContent/Source/VariantManagerContent/Private/FunctionCaller.cpp
@@ -4,9 +4,13 @@
 
 #if WITH_EDITORONLY_DATA
 
+#include "LevelVariantSets.h"
+#include "Variant.h"
+#include "VariantSet.h"
+
 #include "EdGraph/EdGraph.h"
-#include "K2Node_FunctionEntry.h"
 #include "EdGraphSchema_K2.h"
+#include "K2Node_FunctionEntry.h"
 #include "UObject/PropertyPortFlags.h"
 
 // This file is based on MovieSceneEvent.cpp
@@ -38,14 +42,41 @@ bool FFunctionCaller::IsValidFunction(UK2Node_FunctionEntry* Function)
 	{
 		return true;
 	}
-	else if (Function->UserDefinedPins.Num() != 1 || Function->UserDefinedPins[0]->PinType.bIsReference)
+	else if (Function->UserDefinedPins.Num() != 4)
+	{
+		return false;
+	}
+
+	TSharedPtr<FUserPinInfo> TargetPin			 = Function->UserDefinedPins[0];
+	TSharedPtr<FUserPinInfo> LevelVariantSetsPin = Function->UserDefinedPins[1];
+	TSharedPtr<FUserPinInfo> VariantSetPin		 = Function->UserDefinedPins[2];
+	TSharedPtr<FUserPinInfo> VariantPin			 = Function->UserDefinedPins[3];
+
+	if (TargetPin->PinType.bIsReference ||
+		TargetPin->PinType.PinCategory != UEdGraphSchema_K2::PC_Object)
+	{
+		return false;
+	}
+	if (LevelVariantSetsPin->PinType.bIsReference ||
+		LevelVariantSetsPin->PinType.PinCategory != UEdGraphSchema_K2::PC_Object ||
+		LevelVariantSetsPin->PinType.PinSubCategoryObject != ULevelVariantSets::StaticClass())
+	{
+		return false;
+	}
+	if (VariantSetPin->PinType.bIsReference ||
+		VariantSetPin->PinType.PinCategory != UEdGraphSchema_K2::PC_Object ||
+		VariantSetPin->PinType.PinSubCategoryObject != UVariantSet::StaticClass())
+	{
+		return false;
+	}
+	if (VariantPin->PinType.bIsReference ||
+		VariantPin->PinType.PinCategory != UEdGraphSchema_K2::PC_Object ||
+		VariantPin->PinType.PinSubCategoryObject != UVariant::StaticClass())
 	{
 		return false;
 	}
 
-	// Pin must be an object or interface property
-	FName PinCategory = Function->UserDefinedPins[0]->PinType.PinCategory;
-	return PinCategory == UEdGraphSchema_K2::PC_Object || PinCategory == UEdGraphSchema_K2::PC_Interface;
+	return true;
 }
 
 uint32 FFunctionCaller::GetDisplayOrder() const
diff --git a/Engine/Plugins/Enterprise/VariantManagerContent/Source/VariantManagerContent/Private/VariantObjectBinding.cpp b/Engine/Plugins/Enterprise/VariantManagerContent/Source/VariantManagerContent/Private/VariantObjectBinding.cpp
--- a/Engine/Plugins/Enterprise/VariantManagerContent/Source/VariantManagerContent/Private/VariantObjectBinding.cpp
+++ b/Engine/Plugins/Enterprise/VariantManagerContent/Source/VariantManagerContent/Private/VariantObjectBinding.cpp
@@ -8,6 +8,7 @@
 #include "Variant.h"
 #include "VariantManagerContentLog.h"
 #include "VariantManagerObjectVersion.h"
+#include "VariantSet.h"
 
 #include "Algo/Sort.h"
 #include "Engine/World.h"
@@ -298,89 +299,64 @@ void UVariantObjectBinding::ExecuteTargetFunction(FName FunctionName)
 	{
 		DirectorInstance->ProcessEvent(Func, nullptr);
 	}
-	else if (Func->NumParms == 1 && Func->PropertyLink && (Func->PropertyLink->GetPropertyFlags() & CPF_ReferenceParm) == 0)
+	else if (Func->NumParms == 4)
 	{
-		if (FObjectProperty* ObjectParameter = CastField<FObjectProperty>(Func->PropertyLink))
+		// Setup parameters to the Func as a struct
+		struct
 		{
-			if (!ObjectParameter->PropertyClass || BoundObject->IsA(ObjectParameter->PropertyClass))
+			UObject* Target = nullptr;
+			ULevelVariantSets* LevelVariantSets = nullptr;
+			UVariantSet* VariantSet = nullptr;
+			UVariant* Variant = nullptr;
+		} FuncParams;
+		FuncParams.Target = BoundObject;
+		FuncParams.Variant = GetParent();
+		FuncParams.VariantSet = FuncParams.Variant ? FuncParams.Variant->GetParent() : nullptr;
+		FuncParams.LevelVariantSets = FuncParams.VariantSet ? FuncParams.VariantSet->GetParent() : nullptr;
+
+		// Check if our bound object is of the right class for this function (that class could be StaticMeshActor/LightActor/etc.)
+		bool bValidClass = false;
+		UClass* ExpectedClass = nullptr;
+		for (TFieldIterator<FObjectProperty> Iter(Func); Iter; ++Iter)
+		{
+			FObjectProperty* Parameter = *Iter;
+			if (!Parameter)
 			{
-				DirectorInstance->ProcessEvent(Func, &BoundObject);
+				continue;
 			}
-			else
+
+			if (Parameter->GetFName() == TARGET_PIN_NAME && (Parameter->GetPropertyFlags() & CPF_ReferenceParm) == 0)
 			{
-				UE_LOG(LogVariantContent, Error, TEXT("Failed to call function '%s' with object '%s' because it is not the correct type. Function expects a '%s' but target object is a '%s'."),
-					*Func->GetName(),
-					*BoundObject->GetName(),
-					*ObjectParameter->PropertyClass->GetName(),
-					*BoundObject->GetClass()->GetName()
-				);
+				ExpectedClass = Parameter->PropertyClass;
+				if (BoundObject->IsA(ExpectedClass))
+				{
+					bValidClass = true;
+				}
+				break;
 			}
 		}
-	}
 
+		if (bValidClass)
+		{
+			DirectorInstance->ProcessEvent(Func, &FuncParams);
+		}
+		else
+		{
+			UE_LOG(LogVariantContent, Error, TEXT("Failed to call function '%s' with object '%s' because it is not the correct type. Function expects a '%s' but target object is a '%s'."),
+				*Func->GetName(),
+				*BoundObject->GetName(),
+				ExpectedClass ? *ExpectedClass->GetName() : TEXT("nullptr"),
+				*BoundObject->GetClass()->GetName()
+			);
+		}
+	}
 }
 
 void UVariantObjectBinding::ExecuteAllTargetFunctions()
 {
-	if (FunctionCallers.Num() == 0)
-	{
-		return;
-	}
-
-	ULevelVariantSets* ParentLVS = GetTypedOuter<ULevelVariantSets>();
-
-	UObject* BoundObject = GetObject();
-	if (!BoundObject)
-	{
-		return;
-	}
-
-	UObject* DirectorInstance = ParentLVS->GetDirectorInstance(BoundObject);
-
 	for (FFunctionCaller& Caller : FunctionCallers)
 	{
-		UFunction* Func = DirectorInstance->FindFunction(Caller.FunctionName);
-
-		if (!Func || !Func->IsValidLowLevel() || Func->IsPendingKillOrUnreachable() || !DirectorInstance->FindFunction(Func->GetFName()))
-		{
-			continue;
-		}
-
-		//need to check if we''re in edit mode and the function is CallInEditor
-#if WITH_EDITOR
-		const static FName NAME_CallInEditor(TEXT("CallInEditor"));
-
-		UWorld* World = DirectorInstance->GetWorld();
-		if (World->WorldType == EWorldType::Editor && !Func->HasMetaData(NAME_CallInEditor))
-		{
-			UE_LOG(LogVariantContent, Warning, TEXT("Cannot call function '%s' as it doesn't have the CallInEditor option checked! Also note that calling this from the editor may have irreversible effects on the level."), *Func->GetName());
-			continue;
-		}
-#endif
-
-		if (Func->NumParms == 0)
-		{
-			DirectorInstance->ProcessEvent(Func, nullptr);
-		}
-		else if (Func->NumParms == 1 && Func->PropertyLink && (Func->PropertyLink->GetPropertyFlags() & CPF_ReferenceParm) == 0)
-		{
-			if (FObjectProperty* ObjectParameter = CastField<FObjectProperty>(Func->PropertyLink))
-			{
-				if (!ObjectParameter->PropertyClass || BoundObject->IsA(ObjectParameter->PropertyClass))
-				{
-					DirectorInstance->ProcessEvent(Func, &BoundObject);
-				}
-				else
-				{
-					UE_LOG(LogVariantContent, Error, TEXT("Failed to call function '%s' with object '%s' because it is not the correct type. Function expects a '%s' but target object is a '%s'."),
-						*Func->GetName(),
-						*BoundObject->GetName(),
-						*ObjectParameter->PropertyClass->GetName(),
-						*BoundObject->GetClass()->GetName()
-					);
-				}
-			}
-		}
+		ExecuteTargetFunction(Caller.FunctionName);
 	}
 }
 
diff --git a/Engine/Plugins/Enterprise/VariantManagerContent/Source/VariantManagerContent/Public/FunctionCaller.h b/Engine/Plugins/Enterprise/VariantManagerContent/Source/VariantManagerContent/Public/FunctionCaller.h
--- a/Engine/Plugins/Enterprise/VariantManagerContent/Source/VariantManagerContent/Public/FunctionCaller.h
+++ b/Engine/Plugins/Enterprise/VariantManagerContent/Source/VariantManagerContent/Public/FunctionCaller.h
@@ -6,6 +6,9 @@
 #include "FunctionCaller.generated.h"
 
 #define TARGET_PIN_NAME TEXT("Target")
+#define LEVEL_VARIANT_SETS_PIN_NAME TEXT("Level Variant Sets")
+#define VARIANT_SET_PIN_NAME TEXT("Variant Set")
+#define VARIANT_PIN_NAME TEXT("Variant")
 
 // This file is based on MovieSceneEvent.h
 
diff --git a/Engine/Plugins/Experimental/Text3D/Source/Text3D/Private/MeshCreator.cpp b/Engine/Plugins/Experimental/Text3D/Source/Text3D/Private/MeshCreator.cpp
--- a/Engine/Plugins/Experimental/Text3D/Source/Text3D/Private/MeshCreator.cpp
+++ b/Engine/Plugins/Experimental/Text3D/Source/Text3D/Private/MeshCreator.cpp
@@ -403,12 +403,10 @@ void FMeshCreator::MirrorGroup(const EText3DGroupType TypeIn, const EText3DGroup
 	{
 		const FMeshTriangle& Triangle = MeshDescription.Triangles()[FTriangleID(GroupIn.FirstTriangle + TriangleIndex)];
 
-		auto InstanceID = [TotalVerticesNum, &Triangle, GroupIn](const int32 Index)
-		{
-			return static_cast<uint32>(TotalVerticesNum + Triangle.GetVertexInstanceID(Index).GetValue() - GroupIn.FirstVertex);
-		};
-
-		Data->AddTriangle(InstanceID(0), InstanceID(2), InstanceID(1));
+		uint32 Instance0 = static_cast<uint32>(TotalVerticesNum + Triangle.GetVertexInstanceID(0).GetValue() - GroupIn.FirstVertex);
+		uint32 Instance2 = static_cast<uint32>(TotalVerticesNum + Triangle.GetVertexInstanceID(2).GetValue() - GroupIn.FirstVertex);
+		uint32 Instance1 = static_cast<uint32>(TotalVerticesNum + Triangle.GetVertexInstanceID(1).GetValue() - GroupIn.FirstVertex);
+		Data->AddTriangle(Instance0, Instance2, Instance1);
 	}
 }
 
diff --git a/Engine/Plugins/Experimental/Text3D/Source/Text3D/Private/Text3DComponent.cpp b/Engine/Plugins/Experimental/Text3D/Source/Text3D/Private/Text3DComponent.cpp
--- a/Engine/Plugins/Experimental/Text3D/Source/Text3D/Private/Text3DComponent.cpp
+++ b/Engine/Plugins/Experimental/Text3D/Source/Text3D/Private/Text3DComponent.cpp
@@ -116,20 +116,23 @@ void UText3DComponent::OnRegister()
 }
 
 #if WITH_EDITOR
-void UText3DComponent::PostEditChangeProperty(struct FPropertyChangedEvent& PropertyChangedEvent)
+void UText3DComponent::PostEditChangeProperty(FPropertyChangedEvent& PropertyChangedEvent)
 {
 	Super::PostEditChangeProperty(PropertyChangedEvent);
 
 	const FName Name = PropertyChangedEvent.GetPropertyName();
-	if (Name == GET_MEMBER_NAME_CHECKED(UText3DComponent, Text) ||
-		Name == GET_MEMBER_NAME_CHECKED(UText3DComponent, Font) ||
-		Name == GET_MEMBER_NAME_CHECKED(UText3DComponent, Extrude) ||
-		Name == GET_MEMBER_NAME_CHECKED(UText3DComponent, Bevel) ||
-		Name == GET_MEMBER_NAME_CHECKED(UText3DComponent, BevelSegments))
+	if (Name == GET_MEMBER_NAME_CHECKED(UText3DComponent, Text))
+	{
+		BuildTextMesh();
+	}
+	else if (Name == GET_MEMBER_NAME_CHECKED(UText3DComponent, Font) ||
+			Name == GET_MEMBER_NAME_CHECKED(UText3DComponent, Extrude) ||
+			Name == GET_MEMBER_NAME_CHECKED(UText3DComponent, Bevel) ||
+			Name == GET_MEMBER_NAME_CHECKED(UText3DComponent, BevelSegments))
 	{
 		if (PropertyChangedEvent.ChangeType != EPropertyChangeType::Interactive)
 		{
-			BuildTextMesh();
+			BuildTextMesh(true);
 		}
 	}
 	else if (Name == GET_MEMBER_NAME_CHECKED(UText3DComponent, BevelType))
@@ -148,7 +151,7 @@ void UText3DComponent::PostEditChangeProperty(struct FPropertyChangedEvent& Prop
 		}
 		}
 
-		BuildTextMesh();
+		BuildTextMesh(true);
 	}
 	else if (Name == GET_MEMBER_NAME_CHECKED(UText3DComponent, FrontMaterial))
 	{
@@ -224,7 +227,7 @@ void UText3DComponent::SetBevelType(const EText3DBevelType Value)
 
 void UText3DComponent::SetBevelSegments(const int32 Value)
 {
-	const int32 NewValue = FMath::Clamp(Value, 1, 10);
+	const int32 NewValue = FMath::Clamp(Value, 1, 15);
 	if (BevelSegments != NewValue)
 	{
 		BevelSegments = NewValue;
@@ -574,7 +577,7 @@ void UText3DComponent::UpdateTransforms()
 	}
 }
 
-void UText3DComponent::BuildTextMesh()
+void UText3DComponent::BuildTextMesh(const bool bCleanCache)
 {
 	bPendingBuild = false;
 	ShapedText->Reset();
@@ -678,7 +681,10 @@ void UText3DComponent::BuildTextMesh()
 
 	TextGeneratedDelegate.Broadcast();
 
-	Subsystem->Cleanup();
+	if (bCleanCache)
+	{
+		Subsystem->Cleanup();
+	}
 }
 
 void UText3DComponent::CheckBevel()
diff --git a/Engine/Plugins/Experimental/Text3D/Source/Text3D/Private/Text3DEngineSubsystem.cpp b/Engine/Plugins/Experimental/Text3D/Source/Text3D/Private/Text3DEngineSubsystem.cpp
--- a/Engine/Plugins/Experimental/Text3D/Source/Text3D/Private/Text3DEngineSubsystem.cpp
+++ b/Engine/Plugins/Experimental/Text3D/Source/Text3D/Private/Text3DEngineSubsystem.cpp
@@ -1,6 +1,7 @@
 // Copyright Epic Games, Inc. All Rights Reserved.
 
 #include "Text3DEngineSubsystem.h"
+#include "Containers/Ticker.h"
 #include "Misc/FileHelper.h"
 #include "Engine/Font.h"
 #include "Engine/Engine.h"
@@ -46,30 +47,66 @@ UText3DEngineSubsystem::UText3DEngineSubsystem()
 
 void UText3DEngineSubsystem::Initialize(FSubsystemCollectionBase& Collection)
 {
+	Super::Initialize(Collection);
+	CleanupTickerHandle = FTicker::GetCoreTicker().AddTicker(FTickerDelegate::CreateUObject(this, &UText3DEngineSubsystem::CleanupTimerCallback), 600.0f);
 }
 
 void UText3DEngineSubsystem::Deinitialize()
 {
+	if (CleanupTickerHandle.IsValid())
+	{
+		FTicker::GetCoreTicker().RemoveTicker(CleanupTickerHandle);
+		CleanupTickerHandle.Reset();
+	}
+
+	Super::Deinitialize();
+}
+
+void UText3DEngineSubsystem::Reset()
+{
+	CachedFonts.Reset();
+}
+
+bool UText3DEngineSubsystem::CleanupTimerCallback(float DeltaTime)
+{
+	Cleanup();
+	return true;
 }
 
 void UText3DEngineSubsystem::Cleanup()
 {
-	TArray<uint32> ToDelete;
-	for (TPair<uint32, FCachedFontData>& CachedFontDataPair : CachedFonts)
-	{
-		if (CachedFontDataPair.Value.Cleanup())
-			ToDelete.Add(CachedFontDataPair.Key);
-	}
+	QUICK_SCOPE_CYCLE_COUNTER(STAT_UText3DEngineSubsystem_Cleanup);
 
-	for (uint32 Key : ToDelete)
+	for (auto It = CachedFonts.CreateIterator(); It; ++It)
 	{
-		CachedFonts.Remove(Key);
+		if (It.Value().Cleanup())
+		{
+			It.RemoveCurrent();
+		}
 	}
 }
 
-FCachedFontData& UText3DEngineSubsystem::GetCachedFontData(class UFont* Font)
+FCachedFontData& UText3DEngineSubsystem::GetCachedFontData(UFont* Font)
 {
 	uint32 FontHash = HashCombine(0, GetTypeHash(Font));
+	if (CachedFonts.Contains(FontHash))
+	{
+		uint32 TypefaceFontDataHash = 0;
+
+		// First we check if the font itself has changed since we last cached it
+		const FCompositeFont* CompositeFont = Font->GetCompositeFont();
+		if (CompositeFont && CompositeFont->DefaultTypeface.Fonts.Num() > 0)
+		{
+			const FTypefaceEntry& Typeface = CompositeFont->DefaultTypeface.Fonts[0];
+			TypefaceFontDataHash = HashCombine(0, GetTypeHash(Typeface.Font));
+		}
+
+		if (CachedFonts[FontHash].GetTypefaceFontDataHash() != TypefaceFontDataHash)
+		{
+			CachedFonts.Remove(FontHash);
+		}
+	}
+
 	if (!CachedFonts.Contains(FontHash))
 	{
 		FCachedFontData& CachedFontData = CachedFonts.Add(FontHash);
@@ -82,6 +119,7 @@ FCachedFontData& UText3DEngineSubsystem::GetCachedFontData(class UFont* Font)
 
 FCachedFontData::FCachedFontData()
 {
+	TypefaceFontDataHash = 0;
 	CacheCounter = MakeShared<int32>(0);
 	Font = nullptr;
 	FreeTypeFace = nullptr;
@@ -94,6 +132,7 @@ FCachedFontData::~FCachedFontData()
 
 void FCachedFontData::ClearFreeTypeFace()
 {
+	TypefaceFontDataHash = 0;
 	if (FreeTypeFace)
 	{
 		FT_Done_Face(FreeTypeFace);
@@ -134,11 +173,17 @@ void FCachedFontData::LoadFreeTypeFace()
 
 	if (FreeTypeFace)
 	{
+		TypefaceFontDataHash = HashCombine(0, GetTypeHash(Typeface.Font));
 		FT_Set_Char_Size(FreeTypeFace, FontSize, FontSize, 96, 96);
 		FT_Set_Pixel_Sizes(FreeTypeFace, FontSize, FontSize);
 	}
 }
 
+uint32 FCachedFontData::GetTypefaceFontDataHash()
+{
+	return TypefaceFontDataHash;
+}
+
 TSharedPtr<int32> FCachedFontData::GetCacheCounter()
 {
 	return CacheCounter;
@@ -238,19 +283,13 @@ bool FCachedFontData::Cleanup()
 		return true;
 	}
 
-	TArray<uint32> CachedFontMeshToDelete;
-	for (TPair<uint32, FCachedFontMeshes>& CachedFontMeshPair : Meshes)
+	for (auto It = Meshes.CreateIterator(); It; ++It)
 	{
-		if (CachedFontMeshPair.Value.GetCacheCount() <= 1)
+		if (It.Value().GetCacheCount() <= 1)
 		{
-			CachedFontMeshToDelete.Add(CachedFontMeshPair.Key);
+			It.RemoveCurrent();
 		}
 	}
 
-	for (uint32 Key : CachedFontMeshToDelete)
-	{
-		Meshes.Remove(Key);
-	}
-
 	return false;
 }
diff --git a/Engine/Plugins/Experimental/Text3D/Source/Text3D/Private/Text3DEngineSubsystem.h b/Engine/Plugins/Experimental/Text3D/Source/Text3D/Private/Text3DEngineSubsystem.h
--- a/Engine/Plugins/Experimental/Text3D/Source/Text3D/Private/Text3DEngineSubsystem.h
+++ b/Engine/Plugins/Experimental/Text3D/Source/Text3D/Private/Text3DEngineSubsystem.h
@@ -43,6 +43,7 @@ public:
 
 	bool Cleanup();
 
+	uint32 GetTypefaceFontDataHash();
 	TSharedPtr<int32> GetCacheCounter();
 	TSharedPtr<int32> GetMeshesCacheCounter(float Extrude, float Bevel, EText3DBevelType BevelType, float BevelSegments);
 
@@ -61,6 +62,7 @@ private:
 	FT_Face FreeTypeFace;
 	FString FontName;
 	TSharedPtr<int32> CacheCounter;
+	uint32 TypefaceFontDataHash;
 };
 
 UCLASS()
@@ -74,6 +76,8 @@ public:
 	virtual void Initialize(FSubsystemCollectionBase& Collection) override;
 	virtual void Deinitialize() override;
 
+
+	void Reset();
 	void Cleanup();
 	FCachedFontData& GetCachedFontData(class UFont* Font);
 
@@ -81,6 +85,10 @@ public:
 	class UMaterial* DefaultMaterial;
 
 private:
+	bool CleanupTimerCallback(float DeltaTime);
+
 	UPROPERTY()
 	TMap<uint32, FCachedFontData> CachedFonts;
+
+	FDelegateHandle CleanupTickerHandle;
 };
diff --git a/Engine/Plugins/Experimental/Text3D/Source/Text3D/Public/Text3DComponent.h b/Engine/Plugins/Experimental/Text3D/Source/Text3D/Public/Text3DComponent.h
--- a/Engine/Plugins/Experimental/Text3D/Source/Text3D/Public/Text3DComponent.h
+++ b/Engine/Plugins/Experimental/Text3D/Source/Text3D/Public/Text3DComponent.h
@@ -59,7 +59,7 @@ public:
 	EText3DBevelType BevelType;
 
 	/** Half Circle Bevel Segments (Defines the amount of tesselation for the bevel part) */
-	UPROPERTY(EditAnywhere, BlueprintReadWrite, BlueprintSetter = SetBevelSegments, Category = "Text3D", Meta = (ClampMin = 1, ClampMax = 10))
+	UPROPERTY(EditAnywhere, BlueprintReadWrite, BlueprintSetter = SetBevelSegments, Category = "Text3D", Meta = (ClampMin = 1, ClampMax = 15))
 	int32 BevelSegments;
 
 	/** Material for the front part */
@@ -238,7 +238,7 @@ private:
 
 	void Rebuild();
 	void UpdateTransforms();
-	void BuildTextMesh();
+	void BuildTextMesh(const bool bCleanCache = false);
 	void CheckBevel();
 	float MaxBevel() const;
 
diff --git a/Engine/Plugins/Experimental/Text3D/Text3D.uplugin b/Engine/Plugins/Experimental/Text3D/Text3D.uplugin
--- a/Engine/Plugins/Experimental/Text3D/Text3D.uplugin
+++ b/Engine/Plugins/Experimental/Text3D/Text3D.uplugin
@@ -18,5 +18,11 @@
       "Type": "Runtime",
       "LoadingPhase": "Default"
     }
+  ],
+  "Plugins": [
+    {
+      "Name": "GeometryProcessing",
+      "Enabled": true
+    }
   ]
 }
diff --git a/Engine/Source/Programs/Enterprise/Datasmith/DatasmithRevitExporter/Private/DatasmithRevitExportContext.cs b/Engine/Source/Programs/Enterprise/Datasmith/DatasmithRevitExporter/Private/DatasmithRevitExportContext.cs
--- a/Engine/Source/Programs/Enterprise/Datasmith/DatasmithRevitExporter/Private/DatasmithRevitExportContext.cs
+++ b/Engine/Source/Programs/Enterprise/Datasmith/DatasmithRevitExporter/Private/DatasmithRevitExportContext.cs
@@ -764,7 +764,7 @@ namespace DatasmithRevitExporter
 
 			private Transform GetPivotTransform(Element InElement)
 			{
-				if (InElement.Location == null)
+				if (InElement.Location == null || (InElement as FamilyInstance) != null)
 				{
 					return null;
 				}
@@ -801,13 +801,7 @@ namespace DatasmithRevitExporter
 
 				// Get pivot basis
 
-				if ((InElement as FamilyInstance) != null)
-				{
-					FamilyInstance CurrentFamilyInstance = InElement as FamilyInstance;
-					BasisX = CurrentFamilyInstance.FacingOrientation.Normalize();
-					ComputeBasis(BasisX, ref BasisY, ref BasisZ);
-				}
-				else if (InElement.GetType() == typeof(Wall))
+				if (InElement.GetType() == typeof(Wall))
 				{
 					BasisY = (InElement as Wall).Orientation.Normalize();
 					BasisX = BasisY.CrossProduct(XYZ.BasisZ).Normalize();
@@ -892,7 +886,7 @@ namespace DatasmithRevitExporter
 				CreateMeshActor(InWorldTransform, out InstanceData.InstanceMesh, out InstanceData.InstanceActor);
 
                 // The Datasmith instance actor is a component in the hierarchy.
-                //InstanceData.InstanceActor.SetIsComponent(true);
+                InstanceData.InstanceActor.SetIsComponent(true);
             }
 
             public FDatasmithFacadeMesh PopInstance()
diff --git a/Engine/Source/Runtime/Core/Public/GenericPlatform/GenericPlatformMisc.h b/Engine/Source/Runtime/Core/Public/GenericPlatform/GenericPlatformMisc.h
--- a/Engine/Source/Runtime/Core/Public/GenericPlatform/GenericPlatformMisc.h
+++ b/Engine/Source/Runtime/Core/Public/GenericPlatform/GenericPlatformMisc.h
@@ -26,44 +26,44 @@ class TFunction;
 #define UE_DEBUG_BREAK() ((void)(FPlatformMisc::IsDebuggerPresent() && ([] () { UE_DEBUG_BREAK_IMPL(); } (), 1)))
 #endif
 
-	/**
+/**
  * Available build configurations. Mirorred from UnrealTargetConfiguration.
-	 */
+ */
 enum class EBuildConfiguration : uint8
-	{
-		/** Unknown build configuration. */
-		Unknown,
+{
+	/** Unknown build configuration. */
+	Unknown,
 
-		/** Debug build. */
-		Debug,
+	/** Debug build. */
+	Debug,
 
-		/** DebugGame build. */
-		DebugGame,
+	/** DebugGame build. */
+	DebugGame,
 
-		/** Development build. */
-		Development,
+	/** Development build. */
+	Development,
 
-		/** Shipping build. */
-		Shipping,
+	/** Shipping build. */
+	Shipping,
 
-		/** Test build. */
-		Test
-	};
+	/** Test build. */
+	Test
+};
 
-	/**
-	 * Returns the string representation of the specified EBuildConfiguration value.
-	 *
+/**
+ * Returns the string representation of the specified EBuildConfiguration value.
+ *
  * @param Configuration The string to get the EBuildConfiguration for.
  * @return An EBuildConfiguration value.
-	 */
+ */
 CORE_API bool LexTryParseString(EBuildConfiguration& OutConfiguration, const TCHAR* Configuration);
 
-	/**
-	 * Returns the string representation of the specified EBuildConfiguration value.
-	 *
-	 * @param Configuration The value to get the string for.
-	 * @return The string representation.
-	 */
+/**
+ * Returns the string representation of the specified EBuildConfiguration value.
+ *
+ * @param Configuration The value to get the string for.
+ * @return The string representation.
+ */
 CORE_API const TCHAR* LexToString(EBuildConfiguration Configuration);
 
 namespace EBuildConfigurations
@@ -116,18 +116,18 @@ namespace EBuildConfigurations
 	CORE_API FText ToText( EBuildConfiguration Configuration );
 }
 
-	/**
+/**
  * Enumerates build target types.
-	 */
+ */
 enum class EBuildTargetType : uint8
-	{
-		/** Unknown build target. */
-		Unknown,
+{
+	/** Unknown build target. */
+	Unknown,
 
-		/** Game target. */
-		Game,
+	/** Game target. */
+	Game,
 
-		/** Server target. */
+	/** Server target. */
 	Server,
 
 	/** Client target. */
@@ -138,23 +138,23 @@ enum class EBuildTargetType : uint8
 
 	/** Program target. */
 	Program,
-	};
+};
 
-	/**
-	 * Returns the string representation of the specified EBuildTarget value.
-	 *
+/**
+ * Returns the string representation of the specified EBuildTarget value.
+ *
  * @param OutType The value to get the string for.
  * @param Text The text to parse.
  * @return The string representation.
-	 */
+ */
 CORE_API bool LexTryParseString(EBuildTargetType& OutType, const TCHAR* Text);
 
-	/**
+/**
  * Returns the string representation of the specified EBuildTargetType value.
-	 *
+ *
  * @param Target The string to get the EBuildTargetType for.
  * @return An EBuildTarget::Type value.
-	 */
+ */
 CORE_API const TCHAR* LexToString(EBuildTargetType Type);
 
 namespace EBuildTargets
@@ -176,7 +176,7 @@ namespace EBuildTargets
 
 	UE_DEPRECATED(4.24, "EBuildTargets::FromString is deprecated. Use LexFromString() instead.")
 	CORE_API EBuildTargetType FromString( const FString& Target );
-	
+
 	UE_DEPRECATED(4.24, "EBuildTargets::FromString is deprecated. Use LexFromString() instead.")
 	CORE_API const TCHAR* ToString(EBuildTargetType Target);
 }
@@ -201,22 +201,22 @@ enum class EDeviceScreenOrientation : uint8
 {
 	/** The orientation is not known */
 	Unknown,
-	
+
 	/** The orientation is portrait with the home button at the bottom */
 	Portrait,
-	
+
 	/** The orientation is portrait with the home button at the top */
 	PortraitUpsideDown,
-	
+
 	/** The orientation is landscape with the home button at the right side */
 	LandscapeLeft,
-	
+
 	/** The orientation is landscape with the home button at the left side */
 	LandscapeRight,
-	
+
 	/** The orientation is as if place on a desk with the screen upward */
 	FaceUp,
-	
+
 	/** The orientation is as if place on a desk with the screen downward */
 	FaceDown
 };
@@ -224,7 +224,7 @@ enum class EDeviceScreenOrientation : uint8
 
 namespace EErrorReportMode
 {
-	/** 
+	/**
 	 * Enumerates supported error reporting modes.
 	 */
 	enum Type
@@ -311,14 +311,14 @@ enum class EMobileHapticsType : uint8
 
 enum class ENetworkConnectionType : uint8
 {
-    /**
-     * Enumerates the network connection types
-     */
-    Unknown,
-    None,
-    AirplaneMode,
-    Cell,
-    WiFi,
+	/**
+	 * Enumerates the network connection types
+	 */
+	Unknown,
+	None,
+	AirplaneMode,
+	Cell,
+	WiFi,
 	WiMAX,
 	Bluetooth,
 	Ethernet,
@@ -333,8 +333,8 @@ enum class ENetworkConnectionType : uint8
 CORE_API const TCHAR* LexToString( ENetworkConnectionType Target );
 
 /**
-* Generic implementation for most platforms
-**/
+ * Generic implementation for most platforms
+ */
 struct CORE_API FGenericPlatformMisc
 {
 	/**
@@ -387,7 +387,7 @@ struct CORE_API FGenericPlatformMisc
 	 * Sets an environment variable to the local process's environment
 	 *
 	 * @param VariableName The name of the variable (ie "Path")
-	 * @param Value The string to set the variable to.	
+	 * @param Value The string to set the variable to.
 	 */
 	static void SetEnvironmentVar(const TCHAR* VariableName, const TCHAR* Value);
 
@@ -406,7 +406,7 @@ struct CORE_API FGenericPlatformMisc
 
 	/**
 	 * Retrieve the Mac address of the current adapter.
-	 * 
+	 *
 	 * @return array of bytes representing the Mac address, or empty array if unable to determine.
 	 */
 	UE_DEPRECATED(4.14, "GetMacAddress is deprecated. It is not reliable on all platforms")
@@ -414,7 +414,7 @@ struct CORE_API FGenericPlatformMisc
 
 	/**
 	 * Retrieve the Mac address of the current adapter as a string.
-	 * 
+	 *
 	 * @return String representing the Mac address, or empty string.
 	 */
 	UE_DEPRECATED(4.14, "GetMacAddressString is deprecated. It is not reliable on all platforms")
@@ -430,7 +430,7 @@ struct CORE_API FGenericPlatformMisc
 
 	/**
 	 * Returns a unique string for device identification
-	 * 
+	 *
 	 * @return the unique string generated by this platform for this device
 	 */
 	UE_DEPRECATED(4.14, "GetUniqueDeviceId is deprecated. Use GetDeviceId instead.")
@@ -440,7 +440,7 @@ struct CORE_API FGenericPlatformMisc
 	 * Returns a unique string for device identification. Differs from the deprecated GetUniqueDeviceId
 	 * in that there is no default implementation (which used unreliable Mac address determiniation).
 	 * This code is expected to use platform-specific methods to identify the device.
-	 * 
+	 *
 	 * WARNING: Use of this method in your app may imply technical certification requirments for your platform!
 	 * For instance, consoles often require cert waivers to be in place before calling APIs that can track a device,
 	 * so be very careful that you are following your platform's protocols for accessing device IDs. See the platform-
@@ -451,10 +451,10 @@ struct CORE_API FGenericPlatformMisc
 	static FString GetDeviceId();
 
 	/**
-	* Returns a unique string for advertising identification
-	*
-	* @return the unique string generated by this platform for this device
-	*/
+	 * Returns a unique string for advertising identification
+	 *
+	 * @return the unique string generated by this platform for this device
+	 */
 	static FString GetUniqueAdvertisingId();
 
 	// #CrashReport: 2015-02-24 Remove
@@ -471,7 +471,7 @@ struct CORE_API FGenericPlatformMisc
 	FORCEINLINE static bool IsDebuggerPresent()
 	{
 #if UE_BUILD_SHIPPING
-		return 0; 
+		return 0;
 #else
 		return 1; // unknown platforms return true so that they can crash into a debugger
 #endif
@@ -488,10 +488,10 @@ struct CORE_API FGenericPlatformMisc
 	}
 
 	/**
-	* Uses cpuid instruction to get the vendor string
-	*
-	* @return	CPU vendor name
-	*/
+	 * Uses cpuid instruction to get the vendor string
+	 *
+	 * @return	CPU vendor name
+	 */
 	static FString GetCPUVendor();
 
 	/**
@@ -515,7 +515,7 @@ struct CORE_API FGenericPlatformMisc
 	/** @return whether to check for specific CPU compatibility or not */
 	static bool NeedsNonoptionalCPUFeaturesCheck();
 
-	/** 
+	/**
 	 * Uses cpuid instruction to get the CPU brand string
 	 *
 	 * @return	CPU brand string
@@ -587,7 +587,7 @@ struct CORE_API FGenericPlatformMisc
 	FORCENOINLINE static void MemoryBarrier();
 
 	/**
-	 * Set a handler to be called when there is a memory warning from the OS 
+	 * Set a handler to be called when there is a memory warning from the OS
 	 *
 	 * @param Handler	The handler to call
 	 */
@@ -602,7 +602,7 @@ struct CORE_API FGenericPlatformMisc
 	{
 		return false;
 	}
-	
+
 	FORCEINLINE static uint32 GetLastError()
 	{
 		return 0;
@@ -630,8 +630,8 @@ public:
 	static void CustomNamedStat(const ANSICHAR* Text, float Value, const ANSICHAR* Graph, const ANSICHAR* Unit) {}
 
 	/**
-	* Profiler color stack - this overrides the color for named events with undefined colors (e.g stat namedevents)
-	*/
+	 * Profiler color stack - this overrides the color for named events with undefined colors (e.g stat namedevents)
+	 */
 	static void BeginProfilerColor(const struct FColor& Color) {}
 	static void EndProfilerColor() {}
 
@@ -641,28 +641,28 @@ public:
 	{
 	}
 
-    /**
-	* Platform specific function for initializing storage of tagged memory buffers
-	*/
+	/**
+	 * Platform specific function for initializing storage of tagged memory buffers
+	 */
 	FORCEINLINE static void InitTaggedStorage(uint32 NumTags)
 	{
 	}
 
-    /**
-	* Platform specific function for freeing storage of tagged memory buffers
-	*/
+	/**
+	 * Platform specific function for freeing storage of tagged memory buffers
+	 */
 	FORCEINLINE static void ShutdownTaggedStorage()
 	{
 	}
 
-    /**
-	* Platform specific function for tagging a memory buffer with a label. Helps see memory access in profilers
-	*/
+	/**
+	 * Platform specific function for tagging a memory buffer with a label. Helps see memory access in profilers
+	 */
 	FORCEINLINE static void TagBuffer(const char* Label, uint32 Category, const void* Buffer, size_t BufferSize)
 	{
 	}
 
-	/** 
+	/**
 	 *	Set the value for the given section and key in the platform specific key->value store
 	 *  Note: The key->value store is user-specific, but may be used to share data between different applications for the same user
 	 *
@@ -674,7 +674,7 @@ public:
 	 */
 	static bool SetStoredValue(const FString& InStoreId, const FString& InSectionName, const FString& InKeyName, const FString& InValue);
 
-	/** 
+	/**
 	 *	Get the value for the given section and key from the platform specific key->value store
 	 *  Note: The key->value store is user-specific, but may be used to share data between different applications for the same user
 	 *
@@ -710,8 +710,8 @@ public:
 	/** Whether LocalPrint can be called from any thread without overlapping */
 	static bool IsLocalPrintThreadSafe() { return false;  }
 
-	/** 
-	 * Whether the platform has a separate debug channel to stdout (eg. OutputDebugString on Windows). Used to suppress messages being output twice 
+	/**
+	 * Whether the platform has a separate debug channel to stdout (eg. OutputDebugString on Windows). Used to suppress messages being output twice
 	 * if both go to the same place.
 	 */
 	static bool HasSeparateChannelForDebugOutput();
@@ -720,7 +720,7 @@ public:
 	 * Requests application exit.
 	 *
 	 * @param	Force	If true, perform immediate exit (dangerous because config code isn't flushed, etc).
-	 *                  If false, request clean main-loop exit from the platform specific code.
+	 *				  If false, request clean main-loop exit from the platform specific code.
 	 */
 	static void RequestExit( bool Force );
 
@@ -728,7 +728,7 @@ public:
 	 * Requests application exit with a specified return code. Name is different from RequestExit() so overloads of just one of functions are possible.
 	 *
 	 * @param	Force 	   If true, perform immediate exit (dangerous because config code isn't flushed, etc).
-	 *                     If false, request clean main-loop exit from the platform specific code.
+	 *					 If false, request clean main-loop exit from the platform specific code.
 	 * @param   ReturnCode This value will be returned from the program (on the platforms where it's possible). Limited to 0-255 to conform with POSIX.
 	 */
 	static void RequestExitWithStatus( bool Force, uint8 ReturnCode );
@@ -759,7 +759,7 @@ public:
 	/** Create a new globally unique identifier. **/
 	static void CreateGuid(struct FGuid& Result);
 
-	/** 
+	/**
 	 * Show a message box if possible, otherwise print a message and return the default
 	 * @param MsgType What sort of options are provided
 	 * @param Text Specific message
@@ -797,7 +797,7 @@ public:
 		return 1;
 	}
 
-	/** 
+	/**
 	 * Platform-specific normalization of path
 	 * E.g. on Linux/Unix platforms, replaces ~ with user home directory, so ~/.config becomes [REDACTED_USER_HOME]/.config (or [REDACTED_USER_HOME]/.config)
 	 */
@@ -818,10 +818,10 @@ public:
 	static bool UseRenderThread();
 
 	/**
-	* Checks if platform wants to allow an audio thread on current device (note: does not imply it will, only if okay given other criteria met)
-	*
-	* @return true if allowed, false if shouldn't use a separate audio thread
-	*/
+	 * Checks if platform wants to allow an audio thread on current device (note: does not imply it will, only if okay given other criteria met)
+	 *
+	 * @return true if allowed, false if shouldn't use a separate audio thread
+	 */
 	static bool AllowAudioThread()
 	{
 		// allow if not overridden
@@ -854,8 +854,8 @@ public:
 	static int32 NumberOfWorkerThreadsToSpawn();
 
 	/**
-	* Return the number of worker threads we should spawn to service IO, NOT based on number of cores
-	*/
+	 * Return the number of worker threads we should spawn to service IO, NOT based on number of cores
+	 */
 	static int32 NumberOfIOWorkerThreadsToSpawn();
 
 	/**
@@ -911,24 +911,24 @@ public:
 	}
 
 	/**
-	*	Return the GamePersistentDownloadDir.  
-	*	On some platforms, returns the writable directory for downloaded data that persists across play sessions.
-	*	This dir is always per-game.
-	*/
+	 *	Return the GamePersistentDownloadDir.
+	 *	On some platforms, returns the writable directory for downloaded data that persists across play sessions.
+	 *	This dir is always per-game.
+	 */
 	static const TCHAR* GamePersistentDownloadDir();
 
 	static const TCHAR* GetUBTPlatform();
 
 	static const TCHAR* GetUBTTarget();
 
-	/** 
+	/**
 	 * Determines the shader format for the plarform
 	 *
 	 * @return	Returns the shader format to be used by that platform
 	 */
 	static const TCHAR* GetNullRHIShaderFormat();
 
-	/** 
+	/**
 	 * Returns the platform specific chunk based install interface
 	 *
 	 * @return	Returns the platform specific chunk based install implementation
@@ -945,7 +945,7 @@ public:
 	/**
 	 * Has the OS execute a command and path pair (such as launch a browser)
 	 *
-	 * @param ComandType OS hint as to the type of command 
+	 * @param ComandType OS hint as to the type of command
 	 * @param Command the command to execute
 	 * @param CommandLine the commands to pass to the executable
 	 * @return whether the command was successful or not
@@ -968,16 +968,16 @@ public:
 	}
 
 	/**
-	* Generates the SHA256 signature of the given data.
-	* 
-	*
-	* @param Data Pointer to the beginning of the data to hash
-	* @param Bytesize Size of the data to has, in bytes.
-	* @param OutSignature Output Structure to hold the computed signature. 
-	*
-	* @return whether the hash was computed successfully
-	*/
-	static bool GetSHA256Signature(const void* Data, uint32 ByteSize, FSHA256Signature& OutSignature);	
+	 * Generates the SHA256 signature of the given data.
+	 *
+	 *
+	 * @param Data Pointer to the beginning of the data to hash
+	 * @param Bytesize Size of the data to has, in bytes.
+	 * @param OutSignature Output Structure to hold the computed signature.
+	 *
+	 * @return whether the hash was computed successfully
+	 */
+	static bool GetSHA256Signature(const void* Data, uint32 ByteSize, FSHA256Signature& OutSignature);
 
 	/**
 	 * Get the default language (for localization) used by this platform.
@@ -1032,7 +1032,7 @@ public:
 	}
 #endif	//#if !UE_BUILD_SHIPPING
 
-	/** 
+	/**
 	 * Allows disabling ensure()s without rebuilding the binary, by either a commandline switch or a hotfix.
 	 *
 	 * @return ensure is allowed
@@ -1067,7 +1067,7 @@ public:
 
 	/** Platform-specific instruction prefetch */
 	FORCEINLINE static void Prefetch(void const* x, int32 offset = 0)
-	{	
+	{
 	}
 
 	/**
@@ -1076,7 +1076,7 @@ public:
 	 * @return the default profile name.
 	 */
 	static const TCHAR* GetDefaultDeviceProfileName();
-	
+
 	/**
 	 * Gets the current battery level.
 	 *
@@ -1087,12 +1087,12 @@ public:
 		return -1;
 	}
 
-	FORCEINLINE static void SetBrightness(float bBright) { } 
+	FORCEINLINE static void SetBrightness(float bBright) { }
 	FORCEINLINE static float GetBrightness() { return 1.0f; }
-    FORCEINLINE static void ResetBrightness() { } // resets brightness to brightness application started with
-    FORCEINLINE static bool SupportsBrightness() { return false; }
+	FORCEINLINE static void ResetBrightness() { } // resets brightness to brightness application started with
+	FORCEINLINE static bool SupportsBrightness() { return false; }
 
-    FORCEINLINE static bool IsInLowPowerMode() { return false;}
+	FORCEINLINE static bool IsInLowPowerMode() { return false;}
 
 	/**
 	 * Returns the current device temperature level.
@@ -1101,7 +1101,7 @@ public:
 	 */
 	static float GetDeviceTemperatureLevel();
 
-	/** 
+	/**
 	 * Allows a game/program/etc to control the game directory in a special place (for instance, monolithic programs that don't have .uprojects)
 	 */
 	static void SetOverrideProjectDir(const FString& InOverrideDir);
@@ -1151,15 +1151,15 @@ public:
 	static bool FullscreenSameAsWindowedFullscreen()
 	{
 		// On some platforms, Fullscreen and WindowedFullscreen behave the same.
-		//     e.g. On Linux, see FLinuxWindow::ReshapeWindow()/SetWindowMode()
-		//          Allowing Fullscreen window mode confuses higher level code (see UE-19996).
+		//	 e.g. On Linux, see FLinuxWindow::ReshapeWindow()/SetWindowMode()
+		//		  Allowing Fullscreen window mode confuses higher level code (see UE-19996).
 		return false;
 	}
 
 	/**
-	 * Returns whether this is a 'stereo only' platform. In general, stereo only platforms will not 
-	 * support on-screen touch input nor require virtual joysticks (though you should use those query 
-	 * functions to verify). The screen is always used for stereo output, and isn't a mode that is 
+	 * Returns whether this is a 'stereo only' platform. In general, stereo only platforms will not
+	 * support on-screen touch input nor require virtual joysticks (though you should use those query
+	 * functions to verify). The screen is always used for stereo output, and isn't a mode that is
 	 * enabled/disabled.
 	 */
 	static bool IsStandaloneStereoOnlyDevice()
@@ -1192,14 +1192,14 @@ public:
 		return false;
 	}
 
-    /**
-     * Returns whether WiFi connection is currently active
-     */
-    static ENetworkConnectionType GetNetworkConnectionType()
-    {
-        return ENetworkConnectionType::Unknown;
-    }
-    
+	/**
+	 * Returns whether WiFi connection is currently active
+	 */
+	static ENetworkConnectionType GetNetworkConnectionType()
+	{
+		return ENetworkConnectionType::Unknown;
+	}
+
 	/**
 	 * Returns whether the platform has variable hardware (configurable/upgradeable system).
 	 */
@@ -1221,7 +1221,7 @@ public:
 	 * Returns whether the platform is running on battery power or not.
 	 */
 	static bool IsRunningOnBattery();
-	
+
 	/**
 	 * Returns the orientation of the device: e.g. Portrait, LandscapeRight.
 	 * @see EScreenOrientation
@@ -1246,11 +1246,11 @@ public:
 
 	/**
 	 * Returns a unique string associated with the login account of the current machine.
-	 * Implemented using persistent storage like the registry on window (using HKCU), so 
+	 * Implemented using persistent storage like the registry on window (using HKCU), so
 	 * is susceptible to anything that could reset or revert that storage if the ID is created,
 	 * which is generally during install or first run of the app.
-	 * 
-	 * Note: This is NOT a user or machine fingerprint, as multiple logins on the same machine will 
+	 *
+	 * Note: This is NOT a user or machine fingerprint, as multiple logins on the same machine will
 	 * not share the same ID, and it is not based on the hardware of the user. It is completely random and
 	 * non-identifiable.
 	 *
@@ -1283,7 +1283,7 @@ public:
 	 */
 	static EConvertibleLaptopMode GetConvertibleLaptopMode();
 
-	/** 
+	/**
 	 * Get a string description of the mode the engine was running in.
 	 */
 	static const TCHAR* GetEngineMode();
@@ -1295,15 +1295,15 @@ public:
 	static TArray<FString> GetPreferredLanguages();
 
 	/**
-	* Returns the currency code associated with the device's locale
-	* @return the currency code associated with the device's locale
-	*/
+	 * Returns the currency code associated with the device's locale
+	 * @return the currency code associated with the device's locale
+ 	*/
 	static FString GetLocalCurrencyCode();
 
 	/**
-	* Returns the currency symbol associated with the device's locale
-	* @return the currency symbol associated with the device's locale
-	*/
+	 * Returns the currency symbol associated with the device's locale
+	 * @return the currency symbol associated with the device's locale
+	 */
 	static FString GetLocalCurrencySymbol();
 
 	/**
@@ -1372,10 +1372,10 @@ public:
 	static TArray<FCustomChunk> GetAllLanguageChunks();
 	static TArray<FCustomChunk> GetCustomChunksByType(ECustomChunkType DesiredChunkType);
 
-	/*
+	/**
 	 * Loads a text file relative to the package root on platforms that distribute apps in package formats.
 	 * For other platforms, the path is relative to the root directory.
-	*/
+	 */
 	static FString LoadTextFileFromPlatformPackage(const FString& RelativePath);
 
 	static bool FileExistsInPlatformPackage(const FString& RelativePath);
@@ -1406,7 +1406,7 @@ public:
 	}
 
 #if !UE_BUILD_SHIPPING
-	/** 
+	/**
 	 * Returns any platform specific warning messages we want printed on screen
 	 */
 	static bool GetPlatformScreenWarnings(TArray<FText>& PlatformScreenWarnings)
diff --git a/Engine/Source/Runtime/Engine/Classes/Materials/Material.h b/Engine/Source/Runtime/Engine/Classes/Materials/Material.h
--- a/Engine/Source/Runtime/Engine/Classes/Materials/Material.h
+++ b/Engine/Source/Runtime/Engine/Classes/Materials/Material.h
@@ -737,6 +737,13 @@ public:
 	UPROPERTY(EditAnywhere, BlueprintReadOnly, Category = Usage)
 	uint32 bUsedWithHairStrands : 1;
 
+	/**
+	 * Indicates that the material and its instances can be use with LiDAR Point Clouds
+	 * This will result in the shaders required to support LiDAR Point Cloud geometries being compiled which will increase shader compile time and memory usage.
+	 */
+	UPROPERTY(EditAnywhere, BlueprintReadOnly, Category = Usage)
+	uint32 bUsedWithLidarPointCloud : 1;
+
 	/** 
 	 * Indicates that the material and its instances can be used with Slate UI and UMG
 	 * This will result in the shaders required to support UI materials being compiled which will increase shader compile time and memory usage.
diff --git a/Engine/Source/Runtime/Engine/Classes/Materials/MaterialInterface.h b/Engine/Source/Runtime/Engine/Classes/Materials/MaterialInterface.h
--- a/Engine/Source/Runtime/Engine/Classes/Materials/MaterialInterface.h
+++ b/Engine/Source/Runtime/Engine/Classes/Materials/MaterialInterface.h
@@ -54,6 +54,7 @@ enum EMaterialUsage
 	MATUSAGE_GeometryCache,
 	MATUSAGE_Water,
 	MATUSAGE_HairStrands,
+	MATUSAGE_LidarPointCloud,
 
 	MATUSAGE_MAX,
 };
diff --git a/Engine/Source/Runtime/Engine/Private/Materials/Material.cpp b/Engine/Source/Runtime/Engine/Private/Materials/Material.cpp
--- a/Engine/Source/Runtime/Engine/Private/Materials/Material.cpp
+++ b/Engine/Source/Runtime/Engine/Private/Materials/Material.cpp
@@ -1323,6 +1323,7 @@ bool UMaterial::GetUsageByFlag(EMaterialUsage Usage) const
 		case MATUSAGE_GeometryCache: UsageValue = bUsedWithGeometryCache; break;
 		case MATUSAGE_Water: UsageValue = bUsedWithWater; break;
 		case MATUSAGE_HairStrands: UsageValue = bUsedWithHairStrands; break;
+		case MATUSAGE_LidarPointCloud: UsageValue = bUsedWithLidarPointCloud; break;
 		default: UE_LOG(LogMaterial, Fatal,TEXT("Unknown material usage: %u"), (int32)Usage);
 	};
 	return UsageValue;
@@ -1466,6 +1467,10 @@ void UMaterial::SetUsageByFlag(EMaterialUsage Usage, bool NewValue)
 		{
 			bUsedWithHairStrands = NewValue; break;
 		}
+		case MATUSAGE_LidarPointCloud:
+		{
+			bUsedWithLidarPointCloud = NewValue; break;
+		}
 		default: UE_LOG(LogMaterial, Fatal,TEXT("Unknown material usage: %u"), (int32)Usage);
 	};
 #if WITH_EDITOR
@@ -1495,6 +1500,7 @@ FString UMaterial::GetUsageName(EMaterialUsage Usage) const
 		case MATUSAGE_GeometryCache: UsageName = TEXT("bUsedWithGeometryCache"); break;
 		case MATUSAGE_Water: UsageName = TEXT("bUsedWithWater"); break;
 		case MATUSAGE_HairStrands: UsageName = TEXT("bUsedWithHairStrands"); break;
+		case MATUSAGE_LidarPointCloud: UsageName = TEXT("bUsedWithLidarPointCloud"); break;
 		default: UE_LOG(LogMaterial, Fatal,TEXT("Unknown material usage: %u"), (int32)Usage);
 	};
 	return UsageName;
diff --git a/Engine/Source/Runtime/Engine/Private/Materials/MaterialShared.cpp b/Engine/Source/Runtime/Engine/Private/Materials/MaterialShared.cpp
--- a/Engine/Source/Runtime/Engine/Private/Materials/MaterialShared.cpp
+++ b/Engine/Source/Runtime/Engine/Private/Materials/MaterialShared.cpp
@@ -1166,6 +1166,11 @@ bool FMaterialResource::IsUsedWithHairStrands() const
 	return Material->bUsedWithHairStrands;
 }
 
+bool FMaterialResource::IsUsedWithLidarPointCloud() const
+{
+	return Material->bUsedWithLidarPointCloud;
+}
+
 bool FMaterialResource::IsUsedWithLandscape() const
 {
 	return false;
diff --git a/Engine/Source/Runtime/Engine/Public/MaterialShared.h b/Engine/Source/Runtime/Engine/Public/MaterialShared.h
--- a/Engine/Source/Runtime/Engine/Public/MaterialShared.h
+++ b/Engine/Source/Runtime/Engine/Public/MaterialShared.h
@@ -1473,6 +1473,7 @@ public:
 	virtual bool IsUsedWithGeometryCache() const { return false; }
 	virtual bool IsUsedWithWater() const { return false; }
 	virtual bool IsUsedWithHairStrands() const { return false; }
+	virtual bool IsUsedWithLidarPointCloud() const { return false; }
 	ENGINE_API virtual enum EMaterialTessellationMode GetTessellationMode() const;
 	virtual bool IsCrackFreeDisplacementEnabled() const { return false; }
 	virtual bool IsAdaptiveTessellationEnabled() const { return false; }
@@ -2247,6 +2248,7 @@ public:
 	ENGINE_API virtual bool IsUsedWithGeometryCache() const override;
 	ENGINE_API virtual bool IsUsedWithWater() const override;
 	ENGINE_API virtual bool IsUsedWithHairStrands() const override;
+	ENGINE_API virtual bool IsUsedWithLidarPointCloud() const override;
 	ENGINE_API virtual enum EMaterialTessellationMode GetTessellationMode() const override;
 	ENGINE_API virtual bool IsCrackFreeDisplacementEnabled() const override;
 	ENGINE_API virtual bool IsAdaptiveTessellationEnabled() const override;
diff --git a/Engine/Source/Runtime/RemoteImportMessaging/Private/RemoteImportMessagingModule.cpp b/Engine/Source/Runtime/RemoteImportMessaging/Private/RemoteImportMessagingModule.cpp
new file mode 100644
--- /dev/null
+++ b/Engine/Source/Runtime/RemoteImportMessaging/Private/RemoteImportMessagingModule.cpp
@@ -0,0 +1,7 @@
+// Copyright Epic Games, Inc. All Rights Reserved.
+#include "RemoteImportMessagingModule.h"
+
+#include "Modules/ModuleInterface.h"
+#include "Modules/ModuleManager.h"
+
+IMPLEMENT_MODULE(FRemoteImportMessagingModule, RemoteImportMessaging);
diff --git a/Engine/Source/Runtime/RemoteImportMessaging/Public/RemoteImportMessages.h b/Engine/Source/Runtime/RemoteImportMessaging/Public/RemoteImportMessages.h
new file mode 100644
--- /dev/null
+++ b/Engine/Source/Runtime/RemoteImportMessaging/Public/RemoteImportMessages.h
@@ -0,0 +1,119 @@
+// Copyright Epic Games, Inc. All Rights Reserved.
+
+#pragma once
+
+#include "CoreMinimal.h"
+
+#include "RemoteImportMessages.generated.h"
+
+/**
+ * Server handle Ping messages by answering a Pong message.
+ * Useful to test client/server connectivity
+ */
+USTRUCT()
+struct FPingMessage
+{
+	GENERATED_BODY()
+
+	UPROPERTY()
+	int32 Version = 1;
+};
+
+
+/**
+ * Server should send this message on client's Ping messages
+ */
+USTRUCT()
+struct FPongMessage
+{
+	GENERATED_BODY()
+
+	UPROPERTY()
+	FString Acknowledgment;
+
+	FPongMessage(const FString& Msg={})
+		: Acknowledgment(FString(TEXT("[ack]:")) + Msg)
+	{}
+};
+
+
+/**
+ * Servers should send this message to notify an internal state change.
+ * When anchors are added/removed, the server should notify clients.
+ */
+USTRUCT()
+struct FServerStateMessage
+{
+	GENERATED_BODY()
+
+	UPROPERTY()
+	FString ServerGUID;
+
+	UPROPERTY()
+	TArray<FString> Anchors;
+};
+
+/**
+ * Client requests inherit from this message
+ */
+USTRUCT()
+struct FRequestMessage
+{
+	GENERATED_BODY()
+
+	UPROPERTY()
+	FGuid RequestID;
+
+	FRequestMessage()
+		: RequestID(FGuid::NewGuid())
+	{
+	}
+};
+
+/**
+ * Server responses inherit from this message
+ */
+USTRUCT()
+struct FResponseMessage
+{
+	GENERATED_BODY()
+
+	UPROPERTY()
+	FGuid SourceRequestID;
+
+	FResponseMessage(const FRequestMessage* SourceRequest=nullptr)
+		: SourceRequestID(SourceRequest ? SourceRequest->RequestID : FGuid{})
+	{
+	}
+};
+
+
+/**
+ * Client request to import a file in the designated anchor destination
+ */
+USTRUCT()
+struct FImportFileRequest : public FRequestMessage
+{
+	GENERATED_BODY()
+
+	UPROPERTY()
+	FString File;
+
+	UPROPERTY()
+	FString Destination;
+};
+
+
+/**
+ * Server response to an import file request
+ */
+USTRUCT()
+struct FImportFileResponse : public FResponseMessage
+{
+	GENERATED_BODY()
+
+	FImportFileResponse(const FImportFileRequest* Request=nullptr)
+		: FResponseMessage(Request)
+	{}
+};
+
diff --git a/Engine/Source/Runtime/RemoteImportMessaging/Public/RemoteImportMessagingModule.h b/Engine/Source/Runtime/RemoteImportMessaging/Public/RemoteImportMessagingModule.h
new file mode 100644
--- /dev/null
+++ b/Engine/Source/Runtime/RemoteImportMessaging/Public/RemoteImportMessagingModule.h
@@ -0,0 +1,26 @@
+// Copyright Epic Games, Inc. All Rights Reserved.
+
+#pragma once
+
+#include "CoreMinimal.h"
+#include "Modules/ModuleInterface.h"
+#include "Modules/ModuleManager.h"
+
+#define REMOTEIMPORTMESSAGING_MODULE_NAME TEXT("RemoteImportMessaging")
+
+/**
+ * Exposes Messages types used between RemoteImport clients and servers.
+ */
+class FRemoteImportMessagingModule : public IModuleInterface
+{
+public:
+	static FRemoteImportMessagingModule& Get()
+	{
+		return FModuleManager::LoadModuleChecked< FRemoteImportMessagingModule >(REMOTEIMPORTMESSAGING_MODULE_NAME);
+	}
+
+	static bool IsAvailable()
+	{
+		return FModuleManager::Get().IsModuleLoaded(REMOTEIMPORTMESSAGING_MODULE_NAME);
+	}
+};
diff --git a/Engine/Source/Runtime/RemoteImportMessaging/RemoteImportMessaging.Build.cs b/Engine/Source/Runtime/RemoteImportMessaging/RemoteImportMessaging.Build.cs
new file mode 100644
--- /dev/null
+++ b/Engine/Source/Runtime/RemoteImportMessaging/RemoteImportMessaging.Build.cs
@@ -0,0 +1,18 @@
+// Copyright Epic Games, Inc. All Rights Reserved.
+
+namespace UnrealBuildTool.Rules
+{
+	public class RemoteImportMessaging : ModuleRules
+	{
+		public RemoteImportMessaging(ReadOnlyTargetRules Target) : base(Target)
+		{
+			PrivateDependencyModuleNames.AddRange(
+				new string[]
+				{
+					"Core",
+					"CoreUObject",
+				}
+			);
+		}
+	}
+}
