Fix CIS. Missing include introduced with CL 10730014.

diff --git a/Engine/Source/Runtime/Core/Private/Misc/Timecode.cpp b/Engine/Source/Runtime/Core/Private/Misc/Timecode.cpp
--- a/Engine/Source/Runtime/Core/Private/Misc/Timecode.cpp
+++ b/Engine/Source/Runtime/Core/Private/Misc/Timecode.cpp
@@ -1,6 +1,7 @@
 // Copyright 1998-2019 Epic Games, Inc. All Rights Reserved.
 
 #include "Misc/Timecode.h"
+#include "HAL/IConsoleManager.h"
 
 
 static TAutoConsoleVariable<bool> CVarUseDropFormatTimecodeByDefaultWhenSupported(
