- Fixed  FMacPlatformProcess::GenerateApplicationPath() to prevent it finding applications anywhere on the machine (anything indexed by Spotlight) like in other branches/versions of Unreal. The function documentation requires to find the application under the Engine directory.

diff --git a/Engine/Source/Runtime/Core/Private/Mac/MacPlatformProcess.cpp b/Engine/Source/Runtime/Core/Private/Mac/MacPlatformProcess.cpp
--- a/Engine/Source/Runtime/Core/Private/Mac/MacPlatformProcess.cpp
+++ b/Engine/Source/Runtime/Core/Private/Mac/MacPlatformProcess.cpp
@@ -78,10 +78,10 @@ FString FMacPlatformProcess::GenerateApplicationPath( const FString& AppName, EB
 	{
 		ExecutableName += FString::Printf(TEXT("-%s-%s"), *PlatformName, LexToString(BuildConfiguration));
 	}
-	
+
 	NSURL* CurrentBundleURL = [[NSBundle mainBundle] bundleURL];
 	NSString* CurrentBundleName = [[CurrentBundleURL lastPathComponent] stringByDeletingPathExtension];
-	if(FString(CurrentBundleName) == ExecutableName)
+	if (FString(CurrentBundleName) == ExecutableName)
 	{
 		CFStringRef FilePath = CFURLCopyFileSystemPath((CFURLRef)CurrentBundleURL, kCFURLPOSIXPathStyle);
 		FString ExecutablePath = FString::Printf(TEXT("%s/Contents/MacOS/%s"), *FString((NSString*)FilePath), *ExecutableName);
@@ -92,7 +92,7 @@ FString FMacPlatformProcess::GenerateApplicationPath( const FString& AppName, EB
 	{
 		// Try expected path of an executable inside an app package in Engine Binaries
 		FString ExecutablePath = FPaths::EngineDir() / FString::Printf(TEXT("Binaries/%s/%s.app/Contents/MacOS/%s"), *PlatformName, *ExecutableName, *ExecutableName);
-			
+
 		NSString* LaunchPath = ExecutablePath.GetNSString();
 		
 		if ([[NSFileManager defaultManager] fileExistsAtPath: LaunchPath])
@@ -101,32 +101,16 @@ FString FMacPlatformProcess::GenerateApplicationPath( const FString& AppName, EB
 		}
 		else
 		{
-			// Next try expected path of a simple executable file in Engine Binaries
+			// Try the path of a simple executable file in Engine Binaries
 			ExecutablePath = FPaths::EngineDir() / FString::Printf(TEXT("Binaries/%s/%s"), *PlatformName, *ExecutableName);
-
 			LaunchPath = ExecutablePath.GetNSString();
 
 			if ([[NSFileManager defaultManager] fileExistsAtPath:LaunchPath])
 			{
 				return ExecutablePath;
 			}
-			else
-			{
-				CFStringRef App = FPlatformString::TCHARToCFString(*ExecutableName);
-				NSWorkspace* Workspace = [NSWorkspace sharedWorkspace];
-				NSString* AppPath = [Workspace fullPathForApplication : (NSString*)App];
-				CFRelease(App);
-				if (AppPath)
-				{
-					ExecutablePath = FString::Printf(TEXT("%s/Contents/MacOS/%s"), *FString(AppPath), *ExecutableName);
-					return ExecutablePath;
-				}
-				else
-				{
-					return FString();
-				}
-			}
 		}
+		return FString(); // Not found.
 	}
 }
 
