Second attempt at fixing >32 bit file writes

Fixed WindowsPlatformFile not handling >32bit read and write sizes correctly

diff --git a/Engine/Source/Runtime/Core/Private/Windows/WindowsPlatformFile.cpp b/Engine/Source/Runtime/Core/Private/Windows/WindowsPlatformFile.cpp
--- a/Engine/Source/Runtime/Core/Private/Windows/WindowsPlatformFile.cpp
+++ b/Engine/Source/Runtime/Core/Private/Windows/WindowsPlatformFile.cpp
@@ -610,7 +610,7 @@ public:
 		int64 TotalNumWritten = 0;
 		do
 		{
-			uint32 BytesToWrite32 = (uint32)FMath::Min<int64>(BytesToWrite - TotalNumWritten, int64(UINT32_MAX));
+			uint32 BytesToWrite32 = (uint32)FMath::Min<int64>(BytesToWrite, int64(UINT32_MAX));
 			uint32 NumWritten = 0;
 			// Now kick off an async write
 			if (!WriteFile(FileHandle, Source, BytesToWrite32, (::DWORD*)&NumWritten, &OverlappedIO))
