UE4 - Added support for specifying per-RHI suggested drivers
#rnx

diff --git a/Engine/Source/Runtime/Core/Public/GenericPlatform/GenericPlatformDriver.h b/Engine/Source/Runtime/Core/Public/GenericPlatform/GenericPlatformDriver.h
--- a/Engine/Source/Runtime/Core/Public/GenericPlatform/GenericPlatformDriver.h
+++ b/Engine/Source/Runtime/Core/Public/GenericPlatform/GenericPlatformDriver.h
@@ -517,7 +517,7 @@ struct FGPUHardware
 	}	
 	
 	// @return a driver version intended to be shown to the user e.g. "15.30.1025.1001 12/17/2015 (Crimson Edition 15.12)"
-	FString GetSuggestedDriverVersion() const
+	FString GetSuggestedDriverVersion(const FString& InRHIName) const
 	{
 		const TCHAR* Section = GetVendorSectionName();
 
@@ -525,7 +525,32 @@ struct FGPUHardware
 
 		if(Section)
 		{
-			GConfig->GetString(Section, TEXT("SuggestedDriverVersion"), Ret, GHardwareIni);
+			TArray<FString> SuggestedDriverVersions;
+			GConfig->GetArray(Section, TEXT("SuggestedDriverVersion"), SuggestedDriverVersions, GHardwareIni);
+
+			// Find specific RHI version first
+			if (InRHIName.Len() > 0)
+			{
+				for (const FString& SuggestedVersion : SuggestedDriverVersions)
+				{
+					int32 Found = SuggestedVersion.Find(InRHIName);
+					if (Found != INDEX_NONE && Found > 0 && SuggestedVersion[Found - 1] == ';')
+					{
+						Ret = SuggestedVersion.Left(Found - 1);
+						return Ret;
+					}
+				}
+			}
+
+			// Return the first generic one
+			for (const FString& SuggestedVersion : SuggestedDriverVersions)
+			{
+				int32 Found = 0;
+				if (!SuggestedVersion.FindChar(';', Found))
+				{
+					return SuggestedVersion;
+				}
+			}
 		}
 
 		return Ret;
diff --git a/Engine/Source/Runtime/RHI/Private/DynamicRHI.cpp b/Engine/Source/Runtime/RHI/Private/DynamicRHI.cpp
--- a/Engine/Source/Runtime/RHI/Private/DynamicRHI.cpp
+++ b/Engine/Source/Runtime/RHI/Private/DynamicRHI.cpp
@@ -148,7 +148,7 @@ static void RHIDetectAndWarnOfBadDrivers(bool bHasEditorToken)
 			FFormatNamedArguments Args;
 			Args.Add(TEXT("AdapterName"), FText::FromString(DriverInfo.DeviceDescription));
 			Args.Add(TEXT("Vendor"), FText::FromString(VendorString));
-			Args.Add(TEXT("RecommendedVer"), FText::FromString(DetectedGPUHardware.GetSuggestedDriverVersion()));
+			Args.Add(TEXT("RecommendedVer"), FText::FromString(DetectedGPUHardware.GetSuggestedDriverVersion(DriverInfo.RHIName)));
 			Args.Add(TEXT("InstalledVer"), FText::FromString(DriverInfo.UserDriverVersion));
 
 			// this message can be suppressed with r.WarnOfBadDrivers=0
