TBA: Write engine version objects out as a concise string when serializing to text

diff --git a/Engine/Source/Runtime/Core/Private/Misc/EngineVersion.cpp b/Engine/Source/Runtime/Core/Private/Misc/EngineVersion.cpp
--- a/Engine/Source/Runtime/Core/Private/Misc/EngineVersion.cpp
+++ b/Engine/Source/Runtime/Core/Private/Misc/EngineVersion.cpp
@@ -213,12 +213,30 @@ void operator<<(FArchive &Ar, FEngineVersion &Version)
 
 void operator<<(FStructuredArchive::FSlot Slot, FEngineVersion &Version)
 {
-	FStructuredArchive::FRecord Record = Slot.EnterRecord();
-	Record << SA_VALUE(TEXT("Major"), Version.Major);
-	Record << SA_VALUE(TEXT("Minor"), Version.Minor);
-	Record << SA_VALUE(TEXT("Patch"), Version.Patch);
-	Record << SA_VALUE(TEXT("Changelist"), Version.Changelist);
-	Record << SA_VALUE(TEXT("Branch"), Version.Branch);
+	FArchive& BaseArchive = Slot.GetUnderlyingArchive();
+	if (BaseArchive.IsTextFormat())
+	{
+		if (BaseArchive.IsLoading())
+		{
+			FString VersionString;
+			Slot << VersionString;
+			FEngineVersion::Parse(VersionString, Version);
+		}
+		else
+		{
+			FString VersionString = Version.ToString();
+			Slot << VersionString;
+		}
+	}
+	else
+	{
+		FStructuredArchive::FRecord Record = Slot.EnterRecord();
+		Record << SA_VALUE(TEXT("Major"), Version.Major);
+		Record << SA_VALUE(TEXT("Minor"), Version.Minor);
+		Record << SA_VALUE(TEXT("Patch"), Version.Patch);
+		Record << SA_VALUE(TEXT("Changelist"), Version.Changelist);
+		Record << SA_VALUE(TEXT("Branch"), Version.Branch);
+	}
 }
 
 
