Tweak FQueuedThread::Run so it doesn't wake up every 10ms even though there's nothing for it to do, unless stats collection is on.

Eliminates lots of unnecessary context switches, especially on wide CPUs

(ushell-p4-cherrypick of 8231213 by [REDACTED_PERSON])

diff --git a/Engine/Source/Runtime/Core/Private/HAL/ThreadingBase.cpp b/Engine/Source/Runtime/Core/Private/HAL/ThreadingBase.cpp
--- a/Engine/Source/Runtime/Core/Private/HAL/ThreadingBase.cpp
+++ b/Engine/Source/Runtime/Core/Private/HAL/ThreadingBase.cpp
@@ -491,11 +491,27 @@ protected:
 			SET_DWORD_STAT( STAT_ThreadPoolDummyCounter, 0 );
 			// We need to wait for shorter amount of time
 			bool bContinueWaiting = true;
-			while( bContinueWaiting )
-			{				
-				DECLARE_SCOPE_CYCLE_COUNTER( TEXT( "FQueuedThread::Run.WaitForWork" ), STAT_FQueuedThread_Run_WaitForWork, STATGROUP_ThreadPoolAsyncTasks );
-				// Wait for some work to do
-				bContinueWaiting = !DoWorkEvent->Wait( 10 );
+
+			// Unless we're collecting stats there doesn't appear to be any reason to wake
+			// up again until there's work to do (or it's time to die)
+
+#if STATS
+			if (FThreadStats::IsCollectingData())
+			{
+				while (bContinueWaiting)
+				{
+					DECLARE_SCOPE_CYCLE_COUNTER(TEXT("FQueuedThread::Run.WaitForWork"), STAT_FQueuedThread_Run_WaitForWork, STATGROUP_ThreadPoolAsyncTasks);
+
+					// Wait for some work to do
+
+					bContinueWaiting = !DoWorkEvent->Wait(10);
+				}
+			}
+#endif
+
+			if (bContinueWaiting)
+			{
+				DoWorkEvent->Wait();
 			}
 
 			IQueuedWork* LocalQueuedWork = QueuedWork;
