- Removed global alloc/free/realloc counters, since these were inaccurate as they are global and could be called multiple times per operation. Additionally they added one atomic per call, which is is bad news, especially since these tracking counters were enabled for TBB
- Also removed inert ad hoc timing profiling code from TBB allocation path
- Flagged Windows ANSI malloc as thread-safe (is *any* malloc implementation likely to be not threadsafe these days?)

(ushell-p4-cherrypick of 8130356 by [REDACTED_PERSON])

diff --git a/Engine/Source/Runtime/Core/Private/HAL/MallocAnsi.cpp b/Engine/Source/Runtime/Core/Private/HAL/MallocAnsi.cpp
--- a/Engine/Source/Runtime/Core/Private/HAL/MallocAnsi.cpp
+++ b/Engine/Source/Runtime/Core/Private/HAL/MallocAnsi.cpp
@@ -33,8 +33,6 @@ FMallocAnsi::FMallocAnsi()
 
 void* FMallocAnsi::Malloc( SIZE_T Size, uint32 Alignment )
 {
-	IncrementTotalMallocCalls();
-
 #if !UE_BUILD_SHIPPING
 	uint64 LocalMaxSingleAlloc = MaxSingleAlloc.Load(EMemoryOrder::Relaxed);
 	if (LocalMaxSingleAlloc != 0 && Size > LocalMaxSingleAlloc)
@@ -73,8 +71,6 @@ void* FMallocAnsi::Malloc( SIZE_T Size, uint32 Alignment )
 
 void* FMallocAnsi::Realloc( void* Ptr, SIZE_T NewSize, uint32 Alignment )
 {
-	IncrementTotalReallocCalls();
-
 #if !UE_BUILD_SHIPPING
 	uint64 LocalMaxSingleAlloc = MaxSingleAlloc.Load(EMemoryOrder::Relaxed);
 	if (LocalMaxSingleAlloc != 0 && NewSize > LocalMaxSingleAlloc)
@@ -159,7 +155,6 @@ void* FMallocAnsi::Realloc( void* Ptr, SIZE_T NewSize, uint32 Alignment )
 
 void FMallocAnsi::Free( void* Ptr )
 {
-	IncrementTotalFreeCalls();
 #if USE_ALIGNED_MALLOC
 	_aligned_free( Ptr );
 #elif PLATFORM_USE_ANSI_POSIX_MALLOC || PLATFORM_USE_ANSI_MEMALIGN
diff --git a/Engine/Source/Runtime/Core/Private/HAL/MallocTBB.cpp b/Engine/Source/Runtime/Core/Private/HAL/MallocTBB.cpp
--- a/Engine/Source/Runtime/Core/Private/HAL/MallocTBB.cpp
+++ b/Engine/Source/Runtime/Core/Private/HAL/MallocTBB.cpp
@@ -30,13 +30,8 @@ THIRD_PARTY_INCLUDES_START
 #include <tbb/scalable_allocator.h>
 THIRD_PARTY_INCLUDES_END
 
-#define MEM_TIME(st)
-
-
 void* FMallocTBB::Malloc( SIZE_T Size, uint32 Alignment )
 {
-	IncrementTotalMallocCalls();
-
 #if !UE_BUILD_SHIPPING
 	uint64 LocalMaxSingleAlloc = MaxSingleAlloc.Load(EMemoryOrder::Relaxed);
 	if (LocalMaxSingleAlloc != 0 && Size > LocalMaxSingleAlloc)
@@ -46,8 +41,6 @@ void* FMallocTBB::Malloc( SIZE_T Size, uint32 Alignment )
 	}
 #endif
 
-	MEM_TIME(MemTime -= FPlatformTime::Seconds());
-
 	void* NewPtr = NULL;
 #if PLATFORM_MAC
 	// macOS expects all allocations to be aligned to 16 bytes, but TBBs default alignment is 8, so on Mac we always have to use scalable_aligned_malloc.
@@ -80,14 +73,11 @@ void* FMallocTBB::Malloc( SIZE_T Size, uint32 Alignment )
 		FMemory::Memset(NewPtr, DEBUG_FILL_NEW, scalable_msize(NewPtr));
 	}
 #endif
-	MEM_TIME(MemTime += FPlatformTime::Seconds());
 	return NewPtr;
 }
 
 void* FMallocTBB::Realloc( void* Ptr, SIZE_T NewSize, uint32 Alignment )
 {
-	IncrementTotalReallocCalls();
-
 #if !UE_BUILD_SHIPPING
 	uint64 LocalMaxSingleAlloc = MaxSingleAlloc.Load(EMemoryOrder::Relaxed);
 	if (LocalMaxSingleAlloc != 0 && NewSize > LocalMaxSingleAlloc)
@@ -97,7 +87,6 @@ void* FMallocTBB::Realloc( void* Ptr, SIZE_T NewSize, uint32 Alignment )
 	}
 #endif
 
-	MEM_TIME(MemTime -= FPlatformTime::Seconds())
 #if UE_BUILD_DEBUG
 	SIZE_T OldSize = 0;
 	if (Ptr)
@@ -135,7 +124,6 @@ void* FMallocTBB::Realloc( void* Ptr, SIZE_T NewSize, uint32 Alignment )
 	{
 		OutOfMemory(NewSize, Alignment);
 	}
-	MEM_TIME(MemTime += FPlatformTime::Seconds())
 	return NewPtr;
 }
 
@@ -145,14 +133,11 @@ void FMallocTBB::Free( void* Ptr )
 	{
 		return;
 	}
-	MEM_TIME(MemTime -= FPlatformTime::Seconds())
 #if UE_BUILD_DEBUG
 	FMemory::Memset(Ptr, DEBUG_FILL_FREED, scalable_msize(Ptr)); 
 #endif
-	IncrementTotalFreeCalls();
 	scalable_free(Ptr);
 
-	MEM_TIME(MemTime += FPlatformTime::Seconds())
 }
 
 bool FMallocTBB::GetAllocationSize(void *Original, SIZE_T &SizeOut)
@@ -161,6 +146,4 @@ bool FMallocTBB::GetAllocationSize(void *Original, SIZE_T &SizeOut)
 	return true;
 }
 
-#undef MEM_TIME
-
 #endif // PLATFORM_SUPPORTS_TBB && TBB_ALLOCATOR_ALLOWED
diff --git a/Engine/Source/Runtime/Core/Private/HAL/MemoryBase.cpp b/Engine/Source/Runtime/Core/Private/HAL/MemoryBase.cpp
--- a/Engine/Source/Runtime/Core/Private/HAL/MemoryBase.cpp
+++ b/Engine/Source/Runtime/Core/Private/HAL/MemoryBase.cpp
@@ -10,10 +10,6 @@ DECLARE_DWORD_COUNTER_STAT(TEXT("Free calls"),				STAT_FreeCalls,STATGROUP_Memor
 DECLARE_DWORD_COUNTER_STAT(TEXT("Realloc calls"),			STAT_ReallocCalls,STATGROUP_MemoryAllocator);
 DECLARE_DWORD_COUNTER_STAT(TEXT("Total Allocator calls"),	STAT_TotalAllocatorCalls,STATGROUP_MemoryAllocator);
 
-TAtomic<uint32> FMalloc::TotalMallocCalls(0);
-TAtomic<uint32> FMalloc::TotalFreeCalls(0);
-TAtomic<uint32> FMalloc::TotalReallocCalls(0);
-
 #if !UE_BUILD_SHIPPING
 TAtomic<uint64> FMalloc::MaxSingleAlloc;
 #endif
@@ -41,9 +37,9 @@ struct FCurrentFrameCalls
 
 	void Update()
 	{
-		uint32 LocalTotalMallocCalls  = FMalloc::TotalMallocCalls.Load(EMemoryOrder::Relaxed);
-		uint32 LocalTotalReallocCalls = FMalloc::TotalReallocCalls.Load(EMemoryOrder::Relaxed);
-		uint32 LocalTotalFreeCalls    = FMalloc::TotalFreeCalls.Load(EMemoryOrder::Relaxed);
+		uint32 LocalTotalMallocCalls  = 0;
+		uint32 LocalTotalReallocCalls = 0;
+		uint32 LocalTotalFreeCalls    = 0;
 
 		MallocCalls      = LocalTotalMallocCalls - LastMallocCalls;
 		ReallocCalls     = LocalTotalReallocCalls - LastReallocCalls;
diff --git a/Engine/Source/Runtime/Core/Public/HAL/MallocPoisonProxy.h b/Engine/Source/Runtime/Core/Public/HAL/MallocPoisonProxy.h
--- a/Engine/Source/Runtime/Core/Public/HAL/MallocPoisonProxy.h
+++ b/Engine/Source/Runtime/Core/Public/HAL/MallocPoisonProxy.h
@@ -42,7 +42,6 @@ public:
 
 	virtual void* Malloc(SIZE_T Size, uint32 Alignment) override
 	{
-		IncrementTotalMallocCalls();
 		void* Result=UsedMalloc->Malloc(Size, Alignment);
 		if (LIKELY(Result != nullptr && Size > 0))
 		{
@@ -56,7 +55,6 @@ public:
 		// NOTE: case when Realloc returns completely new pointer is not handled properly (we would like to have previous memory poisoned completely).
 		// This can be done by avoiding calling Realloc() of the nested allocator and Malloc()/Free()'ing directly, but this is deemed unacceptable
 		// from a performance/fragmentation standpoint.
-		IncrementTotalReallocCalls();
 		SIZE_T OldSize = 0;
 		if (Ptr != nullptr && GetAllocationSize(Ptr, OldSize) && OldSize > 0 && OldSize > NewSize)
 		{
@@ -77,7 +75,6 @@ public:
 	{
 		if (LIKELY(Ptr))
 		{
-			IncrementTotalFreeCalls();
 			SIZE_T AllocSize;
 			if (LIKELY(GetAllocationSize(Ptr, AllocSize) && AllocSize > 0))
 			{
diff --git a/Engine/Source/Runtime/Core/Public/HAL/MallocTBB.h b/Engine/Source/Runtime/Core/Public/HAL/MallocTBB.h
--- a/Engine/Source/Runtime/Core/Public/HAL/MallocTBB.h
+++ b/Engine/Source/Runtime/Core/Public/HAL/MallocTBB.h
@@ -15,16 +15,6 @@ class FMallocTBB
 	: public FMalloc
 {
 public:
-	
-	/**
-	 * Default constructor.
-	 */
-	FMallocTBB() :
-		MemTime(0.0)
-	{ }
-
-public:
-
 	// FMalloc interface.
 	
 	virtual void* Malloc( SIZE_T Size, uint32 Alignment ) override;
@@ -49,10 +39,6 @@ protected:
 		// this is expected not to return
 		FPlatformMemory::OnOutOfMemory(Size, Alignment);
 	}
-
-private:
-
-	double MemTime;
 };
 
 #endif // PLATFORM_SUPPORTS_TBB && TBB_ALLOCATOR_ALLOWED
\ No newline at end of file
diff --git a/Engine/Source/Runtime/Core/Public/HAL/MallocThreadSafeProxy.h b/Engine/Source/Runtime/Core/Public/HAL/MallocThreadSafeProxy.h
--- a/Engine/Source/Runtime/Core/Public/HAL/MallocThreadSafeProxy.h
+++ b/Engine/Source/Runtime/Core/Public/HAL/MallocThreadSafeProxy.h
@@ -39,7 +39,6 @@ public:
 	 */
 	virtual void* Malloc( SIZE_T Size, uint32 Alignment ) override
 	{
-		IncrementTotalMallocCalls();
 		FScopeLock ScopeLock( &SynchronizationObject );
 		return UsedMalloc->Malloc( Size, Alignment );
 	}
@@ -49,7 +48,6 @@ public:
 	 */
 	virtual void* Realloc( void* Ptr, SIZE_T NewSize, uint32 Alignment ) override
 	{
-		IncrementTotalReallocCalls();
 		FScopeLock ScopeLock( &SynchronizationObject );
 		return UsedMalloc->Realloc( Ptr, NewSize, Alignment );
 	}
@@ -61,7 +59,6 @@ public:
 	{
 		if( Ptr )
 		{
-			IncrementTotalFreeCalls();
 			FScopeLock ScopeLock( &SynchronizationObject );
 			UsedMalloc->Free( Ptr );
 		}
diff --git a/Engine/Source/Runtime/Core/Public/HAL/MemoryBase.h b/Engine/Source/Runtime/Core/Public/HAL/MemoryBase.h
--- a/Engine/Source/Runtime/Core/Public/HAL/MemoryBase.h
+++ b/Engine/Source/Runtime/Core/Public/HAL/MemoryBase.h
@@ -191,31 +191,6 @@ public:
 protected:
 	friend struct FCurrentFrameCalls;
 
-	/** Atomically increment total malloc calls. */
-	FORCEINLINE void IncrementTotalMallocCalls()
-	{
-		++TotalMallocCalls;
-	}
-
-	/** Atomically increment total free calls. */
-	FORCEINLINE void IncrementTotalFreeCalls()
-	{
-		++TotalFreeCalls;
-	}
-
-	/** Atomically increment total realloc calls. */
-	FORCEINLINE void IncrementTotalReallocCalls()
-	{
-		++TotalReallocCalls;
-	}
-
-	/** Total number of calls Malloc, if implemented by derived class. */
-	static TAtomic<uint32> TotalMallocCalls;
-	/** Total number of calls Malloc, if implemented by derived class. */
-	static TAtomic<uint32> TotalFreeCalls;
-	/** Total number of calls Malloc, if implemented by derived class. */
-	static TAtomic<uint32> TotalReallocCalls;
-
 #if !UE_BUILD_SHIPPING
 public:
 	/** Limits the maximum single allocation, to this many bytes, for debugging */
diff --git a/Engine/Source/Runtime/Core/Public/Windows/WIndowsPlatform.h b/Engine/Source/Runtime/Core/Public/Windows/WIndowsPlatform.h
--- a/Engine/Source/Runtime/Core/Public/Windows/WIndowsPlatform.h
+++ b/Engine/Source/Runtime/Core/Public/Windows/WIndowsPlatform.h
@@ -61,6 +61,7 @@ typedef FWindowsPlatformTypes FPlatformTypes;
 #define PLATFORM_MAX_FILEPATH_LENGTH_DEPRECATED				WINDOWS_MAX_PATH
 #define PLATFORM_HAS_BSD_SOCKET_FEATURE_WINSOCKETS			1
 #define PLATFORM_USES_MICROSOFT_LIBC_FUNCTIONS				1
+#define PLATFORM_IS_ANSI_MALLOC_THREADSAFE					1
 #define PLATFORM_SUPPORTS_TBB								1
 #define PLATFORM_SUPPORTS_NAMED_PIPES						1
 #define PLATFORM_COMPILER_HAS_TCHAR_WMAIN					1
