Fixed up MallocReplayProxy to account for removed allocation stats

(ushell-p4-cherrypick of 8292842 by [REDACTED_PERSON])

diff --git a/Engine/Source/Runtime/Core/Private/HAL/MallocReplayProxy.cpp b/Engine/Source/Runtime/Core/Private/HAL/MallocReplayProxy.cpp
--- a/Engine/Source/Runtime/Core/Private/HAL/MallocReplayProxy.cpp
+++ b/Engine/Source/Runtime/Core/Private/HAL/MallocReplayProxy.cpp
@@ -86,8 +86,6 @@ void FMallocReplayProxy::InitializeStatsMetadata()
 
 void* FMallocReplayProxy::Malloc(SIZE_T Size, uint32 Alignment)
 {
-	IncrementTotalMallocCalls();
-
 	void* Result = UsedMalloc->Malloc(Size, Alignment);
 	AddToHistory("Malloc", Result, nullptr, Size, Alignment);
 	return Result;
@@ -95,7 +93,6 @@ void* FMallocReplayProxy::Malloc(SIZE_T Size, uint32 Alignment)
 
 void* FMallocReplayProxy::Realloc(void* Ptr, SIZE_T NewSize, uint32 Alignment)
 {
-	IncrementTotalReallocCalls();
 	void* Result = UsedMalloc->Realloc(Ptr, NewSize, Alignment);
 	AddToHistory("Realloc", Result, Ptr, NewSize, Alignment);
 	return Result;
@@ -105,7 +102,6 @@ void FMallocReplayProxy::Free(void* Ptr)
 {
 	if (LIKELY(Ptr))
 	{
-		IncrementTotalFreeCalls();
 		UsedMalloc->Free(Ptr);
 		AddToHistory("Free", nullptr, Ptr, 0, 0);
 	}
