TBA: Fix a few compile errors when text archives are disabled

#rnx

diff --git a/Engine/Source/Runtime/Core/Public/Serialization/StructuredArchive.h b/Engine/Source/Runtime/Core/Public/Serialization/StructuredArchive.h
--- a/Engine/Source/Runtime/Core/Public/Serialization/StructuredArchive.h
+++ b/Engine/Source/Runtime/Core/Public/Serialization/StructuredArchive.h
@@ -102,7 +102,7 @@ template<typename T> FORCEINLINE TDefaultedNamedAttribute<T> MakeDefaultedNamedA
 #if WITH_TEXT_ARCHIVE_SUPPORT
 	#define SA_DEFAULTED_ATTRIBUTE(Name, Value, Default) MakeDefaultedNamedAttribute(FArchiveFieldName(Name), Value, Default)
 #else
-	#define SA_DEFAULTED_ATTRIBUTE(Name, Value) MakeNamedAttribute(FArchiveFieldName(), Value)
+	#define SA_DEFAULTED_ATTRIBUTE(Name, Value, Default) MakeNamedAttribute(FArchiveFieldName(), Value)
 #endif
 
 /** Typedef for which formatter type to support */
@@ -710,6 +710,24 @@ typename TEnableIf<
 		return FMap(Ar);
 	}
 
+	FORCEINLINE FStructuredArchive::FSlot FStructuredArchive::FSlot::EnterAttribute(FArchiveFieldName FieldName)
+	{
+		Ar.Formatter.EnterAttribute(FieldName);
+		return FSlot(Ar);
+	}
+
+	FORCEINLINE TOptional<FStructuredArchive::FSlot> FStructuredArchive::FSlot::TryEnterAttribute(FArchiveFieldName FieldName, bool bEnterWhenWriting)
+	{
+		if (Ar.Formatter.TryEnterAttribute(FieldName, bEnterWhenWriting))
+		{
+			return TOptional<FSlot>(Ar);
+		}
+		else
+		{
+			return TOptional<FSlot>();
+		}
+	}
+
 	FORCEINLINE void FStructuredArchive::FSlot::operator<< (char& Value)
 	{
 		int8 AsInt = Value;
