PR #6151: Use nullptr instead of NULL (Contributed by [REDACTED_PERSON])

(ushell-p4-cherrypick of 8288106 by UnrealBot)

diff --git a/Engine/Source/Runtime/Core/Private/Apple/ApplePlatformCrashContext.cpp b/Engine/Source/Runtime/Core/Private/Apple/ApplePlatformCrashContext.cpp
--- a/Engine/Source/Runtime/Core/Private/Apple/ApplePlatformCrashContext.cpp
+++ b/Engine/Source/Runtime/Core/Private/Apple/ApplePlatformCrashContext.cpp
@@ -12,8 +12,8 @@
 FApplePlatformCrashContext::FApplePlatformCrashContext(ECrashContextType InType, const TCHAR* InErrorMessage)
 :	FGenericCrashContext(InType, InErrorMessage)
 ,	Signal(0)
-,	Info(NULL)
-,	Context(NULL)
+,	Info(nullptr)
+,	Context(nullptr)
 ,	IgnoreDepth(6)
 {
 	SignalDescription[ 0 ] = 0;
@@ -123,7 +123,7 @@ void FApplePlatformCrashContext::CreateExceptionInfoString(int32 Signal, siginfo
 
 void FApplePlatformCrashContext::WriteLine(int ReportFile, const ANSICHAR* Line)
 {
-	if( Line != NULL )
+	if( Line != nullptr )
 	{
 		int64 StringBytes = FCStringAnsi::Strlen(Line);
 		write(ReportFile, Line, StringBytes);
@@ -160,7 +160,7 @@ void FApplePlatformCrashContext::WriteUTF16String(int ReportFile, const TCHAR* U
 
 void FApplePlatformCrashContext::WriteLine(int ReportFile, const TCHAR* Line)
 {
-	if( Line != NULL )
+	if( Line != nullptr )
 	{
 		int64 NumChars = FCString::Strlen(Line);
 		WriteUTF16String(ReportFile, Line, NumChars);
diff --git a/Engine/Source/Runtime/Core/Private/Apple/ApplePlatformFile.cpp b/Engine/Source/Runtime/Core/Private/Apple/ApplePlatformFile.cpp
--- a/Engine/Source/Runtime/Core/Private/Apple/ApplePlatformFile.cpp
+++ b/Engine/Source/Runtime/Core/Private/Apple/ApplePlatformFile.cpp
@@ -88,7 +88,7 @@ public:
 				{
 					UE_LOG(LogInit, Warning, TEXT("Failed to properly close readable file: %s with errno: %d"), *Filename, errno);
 				}
-				ActiveHandles[ HandleSlot ] = NULL;
+				ActiveHandles[ HandleSlot ] = nullptr;
 			}
 		}
 		else
@@ -276,7 +276,7 @@ private:
 		// Look for non-reserved slot
 		for( int32 i = 0; i < ACTIVE_HANDLE_COUNT; ++i )
 		{
-			if( ActiveHandles[ i ] == NULL )
+			if( ActiveHandles[ i ] == nullptr )
 			{
 				HandleSlot = i;
 				break;
@@ -300,7 +300,7 @@ private:
 			HandleSlot = Oldest;
 		}
 
-		ActiveHandles[ HandleSlot ] = NULL;
+		ActiveHandles[ HandleSlot ] = nullptr;
 		AccessTimes[ HandleSlot ] = FPlatformTime::Seconds();
 	}
 #endif
@@ -523,7 +523,7 @@ IFileHandle* FApplePlatformFile::OpenRead(const TCHAR* Filename, bool bAllowWrit
 		if(flock(Handle, LOCK_NB | LOCK_SH) == -1)
 		{
 			close(Handle);
-			return NULL;
+			return nullptr;
 		}
 #endif
 		
@@ -533,7 +533,7 @@ IFileHandle* FApplePlatformFile::OpenRead(const TCHAR* Filename, bool bAllowWrit
 		return new FFileHandleApple(Handle, Filename, true);
 #endif
 	}
-	return NULL;
+	return nullptr;
 }
 
 IFileHandle* FApplePlatformFile::OpenWrite(const TCHAR* Filename, bool bAppend, bool bAllowRead)
@@ -558,7 +558,7 @@ IFileHandle* FApplePlatformFile::OpenWrite(const TCHAR* Filename, bool bAppend,
 		if(flock(Handle, LOCK_NB | LOCK_EX) == -1)
 		{
 			close(Handle);
-			return NULL;
+			return nullptr;
 		}
 #endif
 		
@@ -579,7 +579,7 @@ IFileHandle* FApplePlatformFile::OpenWrite(const TCHAR* Filename, bool bAppend,
 		}
 		return FileHandleApple;
 	}
-	return NULL;
+	return nullptr;
 }
 
 bool FApplePlatformFile::DirectoryExists(const TCHAR* Directory)
diff --git a/Engine/Source/Runtime/Core/Private/Apple/ApplePlatformMisc.cpp b/Engine/Source/Runtime/Core/Private/Apple/ApplePlatformMisc.cpp
--- a/Engine/Source/Runtime/Core/Private/Apple/ApplePlatformMisc.cpp
+++ b/Engine/Source/Runtime/Core/Private/Apple/ApplePlatformMisc.cpp
@@ -91,7 +91,7 @@ int32 FApplePlatformMisc::NumberOfCores()
 	if (NumberOfCores == -1)
 	{
 		SIZE_T Size = sizeof(int32);
-		if (sysctlbyname("hw.ncpu", &NumberOfCores, &Size, NULL, 0) != 0)
+		if (sysctlbyname("hw.ncpu", &NumberOfCores, &Size, nullptr, 0) != 0)
 		{
 			NumberOfCores = 1;
 		}
diff --git a/Engine/Source/Runtime/Core/Private/Apple/ApplePlatformStackWalk.cpp b/Engine/Source/Runtime/Core/Private/Apple/ApplePlatformStackWalk.cpp
--- a/Engine/Source/Runtime/Core/Private/Apple/ApplePlatformStackWalk.cpp
+++ b/Engine/Source/Runtime/Core/Private/Apple/ApplePlatformStackWalk.cpp
@@ -128,7 +128,7 @@ static void AsyncSafeProgramCounterToSymbolInfo( uint64 ProgramCounter, FProgram
 	}
 #else // But on iOS the best we can do is demangle
 	int32 Status = 0;
-	ANSICHAR* DemangledName = NULL;
+	ANSICHAR* DemangledName = nullptr;
 	
 	// Increased the size of the demangle destination to reduce the chances that abi::__cxa_demangle will allocate
 	// this causes the app to hang as malloc isn't signal handler safe. Ideally we wouldn't call this function in a handler.
@@ -180,7 +180,7 @@ uint32 FApplePlatformStackWalk::CaptureStackBackTrace( uint64* BackTrace, uint32
 {
 	// Make sure we have place to store the information before we go through the process of raising
 	// an exception and handling it.
-	if( BackTrace == NULL || MaxDepth == 0 )
+	if( BackTrace == nullptr || MaxDepth == 0 )
 	{
 		return 0;
 	}
