Unused dodgy char serialization overload removed.

diff --git a/Engine/Source/Runtime/Core/Private/Serialization/StructuredArchive.cpp b/Engine/Source/Runtime/Core/Private/Serialization/StructuredArchive.cpp
--- a/Engine/Source/Runtime/Core/Private/Serialization/StructuredArchive.cpp
+++ b/Engine/Source/Runtime/Core/Private/Serialization/StructuredArchive.cpp
@@ -415,17 +415,6 @@ TOptional<FStructuredArchive::FSlot> FStructuredArchive::FSlot::TryEnterAttribut
 	}
 }
 
-void FStructuredArchive::FSlot::operator<< (char& Value)
-{
-	Ar.EnterSlot(Depth, ElementId);
-
-	int8 AsInt = Value;
-	Ar.Formatter.Serialize(AsInt);
-	Value = AsInt;
-
-	Ar.LeaveSlot();
-}
-
 void FStructuredArchive::FSlot::operator<< (uint8& Value)
 {
 	Ar.EnterSlot(Depth, ElementId);
diff --git a/Engine/Source/Runtime/Core/Public/Serialization/StructuredArchive.h b/Engine/Source/Runtime/Core/Public/Serialization/StructuredArchive.h
--- a/Engine/Source/Runtime/Core/Public/Serialization/StructuredArchive.h
+++ b/Engine/Source/Runtime/Core/Public/Serialization/StructuredArchive.h
@@ -208,7 +208,6 @@ public:
 		TOptional<FSlot> TryEnterAttribute(FArchiveFieldName AttributeName, bool bEnterWhenWriting);
 
 		// We don't support chaining writes to a single slot, so this returns void.
-		void operator << (char& Value);
 		void operator << (uint8& Value);
 		void operator << (uint16& Value);
 		void operator << (uint32& Value);
@@ -728,13 +727,6 @@ typename TEnableIf<
 		}
 	}
 
-	FORCEINLINE void FStructuredArchive::FSlot::operator<< (char& Value)
-	{
-		int8 AsInt = Value;
-		Ar.Formatter.Serialize(AsInt);
-		Value = AsInt;
-	}
-
 	FORCEINLINE void FStructuredArchive::FSlot::operator<< (uint8& Value)
 	{
 		Ar.Formatter.Serialize(Value);
