UECORE-482: Crash reporting in external process.

With this change an additional mode is added to crash reporter client, that allows it to run as a monitoring child process. Additionally stack walking and symbol resolution is moved to this new process on Windows. To support this several changes in stack walking had to be made to allow external processes to be queried. All other platforms remain unchanged.

diff --git a/Engine/Source/Developer/CrashDebugHelper/Private/Windows/WindowsPlatformStackWalkExt.cpp b/Engine/Source/Developer/CrashDebugHelper/Private/Windows/WindowsPlatformStackWalkExt.cpp
--- a/Engine/Source/Developer/CrashDebugHelper/Private/Windows/WindowsPlatformStackWalkExt.cpp
+++ b/Engine/Source/Developer/CrashDebugHelper/Private/Windows/WindowsPlatformStackWalkExt.cpp
@@ -20,6 +20,14 @@
 
 #pragma comment( lib, "dbgeng.lib" )
 
+// Use _NT_SYMBOL_PATH for non development builds. We don't want shipping crash reporter to try to 
+// access build servers for example.
+#if !UE_BUILD_SHIPPING
+	#define ALLOW_UNREAL_ACCESS_TO_NT_SYMBOL_PATH 1
+#else
+	#define ALLOW_UNREAL_ACCESS_TO_NT_SYMBOL_PATH 0
+#endif
+
 static IDebugClient5* Client = NULL;
 static IDebugControl4* Control = NULL;
 static IDebugSymbols3* Symbol = NULL;
@@ -34,8 +42,6 @@ FWindowsPlatformStackWalkExt::~FWindowsPlatformStackWalkExt()
 	ShutdownStackWalking();
 }
 
-
-
 bool FWindowsPlatformStackWalkExt::InitStackWalking()
 {
 	if (!FWindowsPlatformMisc::CoInitialize())
@@ -228,6 +234,15 @@ void FWindowsPlatformStackWalkExt::SetSymbolPathsFromModules()
 			CombinedPath += TEXT( ";" );
 		}
 
+#if ALLOW_UNREAL_ACCESS_TO_NT_SYMBOL_PATH
+		FString SymbolPathEnvironmentVariable = FPlatformMisc::GetEnvironmentVariable(L"_NT_SYMBOL_PATH");
+		if (!SymbolPathEnvironmentVariable.IsEmpty())
+		{
+			CombinedPath += SymbolPathEnvironmentVariable;
+			CombinedPath += ";";
+		}
+#endif
+
 		// Set the symbol path
 		Symbol->SetImagePathWide( *CombinedPath );
 		Symbol->SetSymbolPathWide( *CombinedPath );
diff --git a/Engine/Source/Programs/CrashReportClient/Private/CrashReportClientApp.cpp b/Engine/Source/Programs/CrashReportClient/Private/CrashReportClientApp.cpp
--- a/Engine/Source/Programs/CrashReportClient/Private/CrashReportClientApp.cpp
+++ b/Engine/Source/Programs/CrashReportClient/Private/CrashReportClientApp.cpp
@@ -15,6 +15,7 @@
 #include "CrashReportAnalytics.h"
 #include "Modules/ModuleManager.h"
 #include "HAL/PlatformApplicationMisc.h"
+#include "HAL/PlatformCrashContext.h"
 
 #if !CRASH_REPORT_UNATTENDED_ONLY
 	#include "SCrashReportClient.h"
@@ -59,6 +60,15 @@ static bool bImplicitSendFromCmd = false;
 /** If we want to enable analytics */
 static bool AnalyticsEnabledFromCmd = true;
 
+/** If in monitor mode, watch this pid. */
+static uint64 MonitorPid = 0;
+
+/** If in monitor mode, pipe to read data from game. */
+static void* MonitorReadPipe = nullptr;
+
+/** If in monitor mode, pipe to write data to game. */
+static void* MonitorWritePipe = nullptr;
+
 /**
  * Look for the report to upload, either in the command line or in the platform's report queue
  */
@@ -122,6 +132,11 @@ void ParseCommandLine(const TCHAR* CommandLine)
 		{
 			AnalyticsEnabledFromCmd = false;
 		}
+
+		CrashGUIDFromCmd = Params.FindRef(TEXT("CrashGUID"));
+		MonitorPid = FPlatformString::Atoi64(*Params.FindRef(TEXT("MONITOR")));
+		MonitorReadPipe = (void*) FPlatformString::Atoi64(*Params.FindRef(TEXT("READ")));
+		MonitorWritePipe = (void*) FPlatformString::Atoi64(*Params.FindRef(TEXT("WRITE")));
 	}
 
 	if (FoundReportDirectoryAbsolutePaths.Num() == 0)
@@ -353,48 +368,120 @@ void RunUnattended(FPlatformErrorReport ErrorReport)
 	}
 }
 
-
-void RunCrashReportClient(const TCHAR* CommandLine)
+FPlatformErrorReport CollectErrorReport(uint32 Pid, const FSharedCrashContext& SharedCrashContext, void* WritePipe) 
 {
-	// Override the stack size for the thread pool.
-	FQueuedThreadPool::OverrideStackSize = 256 * 1024;
+	// @note: This API is only partially implemented on Mac OS and Linux.
+	FProcHandle ProcessHandle = FPlatformProcess::OpenProcess(Pid);
+
+	// First init the static crash context state
+	FPlatformCrashContext::InitializeFromContext(
+		SharedCrashContext.SessionContext,
+		SharedCrashContext.EnabledPluginsNum > 0 ? &SharedCrashContext.DynamicData[SharedCrashContext.EnabledPluginsOffset] : nullptr,
+		SharedCrashContext.EngineDataNum > 0 ? &SharedCrashContext.DynamicData[SharedCrashContext.EngineDataOffset] : nullptr,
+		SharedCrashContext.GameDataNum > 0 ? &SharedCrashContext.DynamicData[SharedCrashContext.GameDataOffset] : nullptr
+	);
+	// Next create a crash context for the crashed process.
+	FPlatformCrashContext CrashContext(SharedCrashContext.CrashType, SharedCrashContext.ErrorMessage);
+	CrashContext.SetCrashedProcess(ProcessHandle);
+	CrashContext.SetCrashedThreadId(SharedCrashContext.CrashingThreadId);
+	CrashContext.SetNumMinidumpFramesToIgnore(SharedCrashContext.NumStackFramesToIgnore);
+
+	// Initialize the stack walking for the monitored process
+	FPlatformStackWalk::InitStackWalkingForProcess(ProcessHandle);
+
+	for (uint32 ThreadIdx = 0; ThreadIdx < SharedCrashContext.NumThreads; ThreadIdx++)
+	{
+		const uint32 ThreadId = SharedCrashContext.ThreadIds[ThreadIdx];
+		uint64 StackFrames[CR_MAX_STACK_FRAMES] = {0};
+		const uint32 StackFrameCount = FPlatformStackWalk::CaptureThreadStackBackTrace(
+			ThreadId, 
+			StackFrames,
+			CR_MAX_STACK_FRAMES
+		);
+
+		CrashContext.AddPortableThreadCallStack(
+			SharedCrashContext.ThreadIds[ThreadIdx],
+			&SharedCrashContext.ThreadNames[ThreadIdx*CR_MAX_THREAD_NAME_CHARS],
+			StackFrames,
+			StackFrameCount
+		);
+
+		// Add the crashing stack specifically. Is this really needed?
+		if (ThreadId == SharedCrashContext.CrashingThreadId)
+		{
+			CrashContext.SetPortableCallStack(
+				StackFrames + SharedCrashContext.NumStackFramesToIgnore,
+				StackFrameCount - SharedCrashContext.NumStackFramesToIgnore
+			);
+		}
+	}
 
-	// Set up the main loop
-	GEngineLoop.PreInit(CommandLine);
+	// Setup the FPrimaryCrashProperties singleton. If the path is not set it is most likely
+	// that we have crashed during static init, in which case we need to construct a directory
+	// ourself.
+	CrashContext.SerializeContentToBuffer();
 
-	// Make sure all UObject classes are registered and default properties have been initialized
-	ProcessNewlyLoadedUObjects();
+	FString ReportDirectoryAbsolutePath(SharedCrashContext.CrashFilesDirectory);
+	bool DirectoryExists = true;
+	if (ReportDirectoryAbsolutePath.IsEmpty())
+	{
+		DirectoryExists = FGenericCrashContext::CreateCrashReportDirectory(
+			SharedCrashContext.SessionContext.CrashGUIDRoot,
+			SharedCrashContext.SessionContext.GameName,
+			0,
+			ReportDirectoryAbsolutePath);
+	}
 
-	// Tell the module manager is may now process newly-loaded UObjects when new C++ modules are loaded
-	FModuleManager::Get().StartProcessingNewlyLoadedObjects();
+	// Copy platform specific files (e.g. minidump) to output directory if it exists
+	if (DirectoryExists)
+	{
+		CrashContext.CopyPlatformSpecificFiles(*ReportDirectoryAbsolutePath, SharedCrashContext.PlatformCrashContext);
+	}
 
-	// Initialize config.
-	FCrashReportCoreConfig::Get();
+	// At this point the game can continue execution. It is important this happens
+	// as soon as thread state and minidump has been created, so that ensures cause
+	// as little hitch as possible.
+	uint8 ResponseCode[] = { 0xd, 0xe, 0xa, 0xd };
+	FPlatformProcess::WritePipe(WritePipe, ResponseCode, sizeof(ResponseCode));
+
+	// Write out the XML file.
+	const FString CrashContextXMLPath = FPaths::Combine(*ReportDirectoryAbsolutePath, FPlatformCrashContext::CrashContextRuntimeXMLNameW);
+	CrashContext.SerializeAsXML(*CrashContextXMLPath);
+
+	const TCHAR* CrachContextBuffer = *CrashContext.GetBuffer();
+	FPrimaryCrashProperties::Set(new FCrashContext(ReportDirectoryAbsolutePath / TEXT("CrashContext.runtime-xml"), CrachContextBuffer));
+
+	FPlatformErrorReport ErrorReport(ReportDirectoryAbsolutePath);
 
-	const bool bUnattended = 
 #if CRASH_REPORT_UNATTENDED_ONLY
-		true;
+	return ErrorReport;
 #else
-		FApp::IsUnattended();
-#endif // CRASH_REPORT_UNATTENDED_ONLY
 
-	// Find the report to upload in the command line arguments
-	ParseCommandLine(CommandLine);
 
-	// Increase the HttpSendTimeout to 5 minutes
-	GConfig->SetFloat(TEXT("HTTP"), TEXT("HttpSendTimeout"), 5*60.0f, GEngineIni);
+	FString ConfigFilename;
+	if (ErrorReport.FindFirstReportFileWithExtension(ConfigFilename, *FGenericCrashContext::CrashConfigExtension))
+	{
+		FConfigFile CrashConfigFile;
+		CrashConfigFile.Read(ReportDirectoryAbsolutePath / ConfigFilename);
+		FCrashReportCoreConfig::Get().SetProjectConfigOverrides(CrashConfigFile);
+	}
 
-	FPlatformErrorReport::Init();
-	FPlatformErrorReport ErrorReport = LoadErrorReport();
+	return ErrorReport;
+#endif
+}
 
+void SendErrorReport(FPlatformErrorReport& ErrorReport)
+{
 	if (!GIsRequestingExit && ErrorReport.HasFilesToUpload() && FPrimaryCrashProperties::Get() != nullptr)
 	{
-		ErrorReport.SetCrashReportClientVersion(FCrashReportCoreConfig::Get().GetVersion());
+		const bool bUnattended =
+#if CRASH_REPORT_UNATTENDED_ONLY
+			true;
+#else
+			FApp::IsUnattended();
+#endif // CRASH_REPORT_UNATTENDED_ONLY
 
-		if (AnalyticsEnabledFromCmd)
-		{
-			FCrashReportAnalytics::Initialize();
-		}
+		ErrorReport.SetCrashReportClientVersion(FCrashReportCoreConfig::Get().GetVersion());
 
 		if (bUnattended)
 		{
@@ -416,21 +503,87 @@ void RunCrashReportClient(const TCHAR* CommandLine)
 		}
 #endif // !CRASH_REPORT_UNATTENDED_ONLY
 
-		if (AnalyticsEnabledFromCmd)
+	}
+}
+
+bool WaitForCrash(uint32 WatchedProcess, FSharedCrashContext& CrashContext, void* ReadPipe)
+{
+	TArray<uint8> Buffer;
+
+	while (FPlatformProcess::IsApplicationAlive(WatchedProcess))
+	{
+		if (FPlatformProcess::ReadPipeToArray(ReadPipe, Buffer))
 		{
-			// Shutdown analytics.
-			FCrashReportAnalytics::Shutdown();
+			FPlatformMemory::Memcpy(&CrashContext, Buffer.GetData(), Buffer.Num());
+			return true;
+		}
+		
+		FPlatformProcess::Sleep(1.0f);
+	}
+
+	return false;
+}
+
+void RunCrashReportClient(const TCHAR* CommandLine)
+{
+	// Override the stack size for the thread pool.
+	FQueuedThreadPool::OverrideStackSize = 256 * 1024;
+
+	// Increase the HttpSendTimeout to 5 minutes
+	GConfig->SetFloat(TEXT("HTTP"), TEXT("HttpSendTimeout"), 5 * 60.0f, GEngineIni);
+
+	// Set up the main loop
+	GEngineLoop.PreInit(CommandLine);
+
+	// Make sure all UObject classes are registered and default properties have been initialized
+	ProcessNewlyLoadedUObjects();
+
+	// Tell the module manager is may now process newly-loaded UObjects when new C++ modules are loaded
+	FModuleManager::Get().StartProcessingNewlyLoadedObjects();
+
+	// Initialize config.
+	FCrashReportCoreConfig::Get();
+
+	// Find the report to upload in the command line arguments
+	ParseCommandLine(CommandLine);
+
+	FPlatformErrorReport::Init();
+	if (AnalyticsEnabledFromCmd)
+	{
+		FCrashReportAnalytics::Initialize();
+	}
+
+	if (MonitorPid && MonitorWritePipe && MonitorReadPipe) 
+	{
+		while (FPlatformProcess::IsApplicationAlive(MonitorPid))
+		{
+			// Wait for parent process to signal crash
+			FSharedCrashContext CrashContext;
+			if (WaitForCrash(MonitorPid, CrashContext, MonitorReadPipe))
+			{
+				// Build error report in memory
+				FPlatformErrorReport ErrorReport = CollectErrorReport(MonitorPid, CrashContext, MonitorWritePipe);
+				SendErrorReport(ErrorReport);
+			}
 		}
 	}
 	else
 	{
-		// Needed to let systems that are shutting down that we are shutting down by request
-		GIsRequestingExit = true;
+		// Load error report generated by the process from disk
+		FPlatformErrorReport ErrorReport = LoadErrorReport();
+		SendErrorReport(ErrorReport);
+	}
+
+	if (AnalyticsEnabledFromCmd)
+	{
+		FCrashReportAnalytics::Shutdown();
 	}
 
 	FPrimaryCrashProperties::Shutdown();
 	FPlatformErrorReport::ShutDown();
 
+	GIsRequestingExit = true;
+
 	FEngineLoop::AppPreExit();
 	FModuleManager::Get().UnloadModulesAtShutdown();
 	FTaskGraphInterface::Shutdown();
diff --git a/Engine/Source/Runtime/Core/Private/GenericPlatform/GenericPlatformCrashContext.cpp b/Engine/Source/Runtime/Core/Private/GenericPlatform/GenericPlatformCrashContext.cpp
--- a/Engine/Source/Runtime/Core/Private/GenericPlatform/GenericPlatformCrashContext.cpp
+++ b/Engine/Source/Runtime/Core/Private/GenericPlatform/GenericPlatformCrashContext.cpp
@@ -21,6 +21,8 @@
 #include "Stats/Stats.h"
 #include "Internationalization/TextLocalizationManager.h"
 #include "Logging/LogScopedCategoryAndVerbosityOverride.h"
+#include "HAL/PlatformOutputDevices.h"
+#include "Misc/OutputDeviceArchiveWrapper.h"
 
 #ifndef NOINITCRASHREPORTER
 #define NOINITCRASHREPORTER 0
@@ -64,43 +66,13 @@ const FString FGenericCrashContext::EngineModeExDirty = TEXT("Dirty");
 const FString FGenericCrashContext::EngineModeExVanilla = TEXT("Vanilla");
 
 bool FGenericCrashContext::bIsInitialized = false;
-FPlatformMemoryStats FGenericCrashContext::CrashMemoryStats = FPlatformMemoryStats();
 int32 FGenericCrashContext::StaticCrashContextIndex = 0;
 
 const FGuid FGenericCrashContext::ExecutionGuid = FGuid::NewGuid();
 
-namespace NCachedCrashContextProperties
-{
-	static bool bIsInternalBuild;
-	static bool bIsPerforceBuild;
-	static bool bIsSourceDistribution;
-	static bool bIsUE4Release;
-	static TOptional<bool> bIsVanilla;
-	static FString GameName;
-	static FString ExecutableName;
-	static FString DeploymentName;
-	static FString BaseDir;
-	static FString RootDir;
-	static FString EpicAccountId;
-	static FString LoginIdStr;
-	static FString OsVersion;
-	static FString OsSubVersion;
-	static int32 NumberOfCores;
-	static int32 NumberOfCoresIncludingHyperthreads;
-	static FString CPUVendor;
-	static FString CPUBrand;
-	static FString PrimaryGPUBrand;
-	static FString UserName;
-	static FString DefaultLocale;
-	static int32 CrashDumpMode;
-	static int32 SecondsSinceStart;
-	static FString CrashGUIDRoot;
-	static FString UserActivityHint;
-	static FString GameSessionID;
-	static FString CommandLine;
-	static int32 LanguageLCID;
-	static FString CrashReportClientRichText;
-	static FString GameStateName;
+namespace NCached
+{
+	static FSessionContext Session;
 	static TArray<FString> EnabledPluginsList;
 	static TMap<FString, FString> EngineData;
 	static TMap<FString, FString> GameData;
@@ -109,91 +81,111 @@ namespace NCachedCrashContextProperties
 void FGenericCrashContext::Initialize()
 {
 #if !NOINITCRASHREPORTER
-	NCachedCrashContextProperties::bIsInternalBuild = FEngineBuildSettings::IsInternalBuild();
-	NCachedCrashContextProperties::bIsPerforceBuild = FEngineBuildSettings::IsPerforceBuild();
-	NCachedCrashContextProperties::bIsSourceDistribution = FEngineBuildSettings::IsSourceDistribution();
-	NCachedCrashContextProperties::bIsUE4Release = FApp::IsEngineInstalled();
-
-	NCachedCrashContextProperties::GameName = FString::Printf( TEXT("UE4-%s"), FApp::GetProjectName() );
-	NCachedCrashContextProperties::ExecutableName = FPlatformProcess::ExecutableName();
-	NCachedCrashContextProperties::BaseDir = FPlatformProcess::BaseDir();
-	NCachedCrashContextProperties::RootDir = FPlatformMisc::RootDir();
-	NCachedCrashContextProperties::EpicAccountId = FPlatformMisc::GetEpicAccountId();
-	NCachedCrashContextProperties::LoginIdStr = FPlatformMisc::GetLoginId();
-	FPlatformMisc::GetOSVersions(NCachedCrashContextProperties::OsVersion, NCachedCrashContextProperties::OsSubVersion);
-	NCachedCrashContextProperties::NumberOfCores = FPlatformMisc::NumberOfCores();
-	NCachedCrashContextProperties::NumberOfCoresIncludingHyperthreads = FPlatformMisc::NumberOfCoresIncludingHyperthreads();
-
-	NCachedCrashContextProperties::CPUVendor = FPlatformMisc::GetCPUVendor();
-	NCachedCrashContextProperties::CPUBrand = FPlatformMisc::GetCPUBrand();
-	NCachedCrashContextProperties::PrimaryGPUBrand = FPlatformMisc::GetPrimaryGPUBrand();
-	NCachedCrashContextProperties::UserName = FPlatformProcess::UserName();
-	NCachedCrashContextProperties::DefaultLocale = FPlatformMisc::GetDefaultLocale();
-	NCachedCrashContextProperties::CommandLine = FCommandLine::IsInitialized() ? FCommandLine::GetOriginalForLogging() : TEXT(""); 
-
-	// Use -epicapp value from the commandline to start. This will also be set by the game
-	FParse::Value(FCommandLine::Get(), TEXT("EPICAPP="), NCachedCrashContextProperties::DeploymentName);
-
-	if (FInternationalization::IsAvailable())
-	{
-		NCachedCrashContextProperties::LanguageLCID = FInternationalization::Get().GetCurrentCulture()->GetLCID();
-	}
-	else
+	NCached::Session.bIsInternalBuild = FEngineBuildSettings::IsInternalBuild();
+	NCached::Session.bIsPerforceBuild = FEngineBuildSettings::IsPerforceBuild();
+	NCached::Session.bIsSourceDistribution = FEngineBuildSettings::IsSourceDistribution();
+	NCached::Session.ProcessId = FPlatformProcess::GetCurrentProcessId();
+
+	FCString::Strcpy(NCached::Session.GameName, *FString::Printf( TEXT("UE4-%s"), FApp::GetProjectName() ));
+	FCString::Strcpy(NCached::Session.GameSessionID, TEXT(""));
+	FCString::Strcpy(NCached::Session.GameStateName, TEXT(""));
+	FCString::Strcpy(NCached::Session.UserActivityHint, TEXT(""));
+	FCString::Strcpy(NCached::Session.ExecutableName, FPlatformProcess::ExecutableName());
+	FCString::Strcpy(NCached::Session.BaseDir, FPlatformProcess::BaseDir());
+	FCString::Strcpy(NCached::Session.RootDir, FPlatformMisc::RootDir());
+	FCString::Strcpy(NCached::Session.EpicAccountId, *FPlatformMisc::GetEpicAccountId());
+	FCString::Strcpy(NCached::Session.LoginIdStr, *FPlatformMisc::GetLoginId());
+
+	FString OsVersion, OsSubVersion;
+	FPlatformMisc::GetOSVersions(OsVersion, OsSubVersion);
+	FCString::Strcpy(NCached::Session.OsVersion, *OsVersion);
+	FCString::Strcpy(NCached::Session.OsSubVersion, *OsSubVersion);
+
+	NCached::Session.NumberOfCores = FPlatformMisc::NumberOfCores();
+	NCached::Session.NumberOfCoresIncludingHyperthreads = FPlatformMisc::NumberOfCoresIncludingHyperthreads();
+
+	FCString::Strcpy(NCached::Session.CPUVendor, *FPlatformMisc::GetCPUVendor());
+	FCString::Strcpy(NCached::Session.CPUBrand, *FPlatformMisc::GetCPUBrand());
+	FCString::Strcpy(NCached::Session.PrimaryGPUBrand, *FPlatformMisc::GetPrimaryGPUBrand());
+	FCString::Strcpy(NCached::Session.UserName, FPlatformProcess::UserName());
+	FCString::Strcpy(NCached::Session.DefaultLocale, *FPlatformMisc::GetDefaultLocale());
+
+	// Information that cannot be gathered if command line is not initialized (e.g. crash during static init)
+	if (FCommandLine::IsInitialized())
 	{
-		FCulturePtr DefaultCulture = FInternationalization::Get().GetCulture(TEXT("en"));
-		if (DefaultCulture.IsValid())
-		{
-			NCachedCrashContextProperties::LanguageLCID = DefaultCulture->GetLCID();
-		}
-		else
+		NCached::Session.bIsUE4Release = FApp::IsEngineInstalled();
+		FCString::Strcpy(NCached::Session.CommandLine, (FCommandLine::IsInitialized() ? FCommandLine::GetOriginalForLogging() : TEXT("")));
+		FCString::Strcpy(NCached::Session.EngineMode, FGenericPlatformMisc::GetEngineMode());
+		FCString::Strcpy(NCached::Session.EngineModeEx, EngineModeExString());
+
+		// Use -epicapp value from the commandline to start. This will also be set by the game
+		FParse::Value(FCommandLine::Get(), TEXT("EPICAPP="), NCached::Session.DeploymentName, CR_MAX_GENERIC_FIELD_CHARS, true);
+
+		// Using the -fullcrashdump parameter will cause full memory minidumps to be created for crashes
+		NCached::Session.CrashDumpMode = (int32)ECrashDumpMode::Default;
+		if (FPlatformMisc::SupportsFullCrashDumps() && FCommandLine::IsInitialized())
 		{
-			const int DefaultCultureLCID = 1033;
-			NCachedCrashContextProperties::LanguageLCID = DefaultCultureLCID;
+			const TCHAR* CmdLine = FCommandLine::Get();
+			if (FParse::Param(CmdLine, TEXT("fullcrashdumpalways")))
+			{
+				NCached::Session.CrashDumpMode = (int32)ECrashDumpMode::FullDumpAlways;
+			}
+			else if (FParse::Param(CmdLine, TEXT("fullcrashdump")))
+			{
+				NCached::Session.CrashDumpMode = (int32)ECrashDumpMode::FullDump;
+			}
 		}
 	}
 
-	// Using the -fullcrashdump parameter will cause full memory minidumps to be created for crashes
-	NCachedCrashContextProperties::CrashDumpMode = (int32)ECrashDumpMode::Default;
-	if (FPlatformMisc::SupportsFullCrashDumps() && FCommandLine::IsInitialized())
+	// Create a unique base guid for bug report ids
+	const FGuid Guid = FGuid::NewGuid();
+	const FString IniPlatformName(FPlatformProperties::IniPlatformName());
+	FCString::Strcpy(NCached::Session.CrashGUIDRoot, *FString::Printf(TEXT("%s%s-%s"), *CrashGUIDRootPrefix, *IniPlatformName, *Guid.ToString(EGuidFormats::Digits)));
+
+	if (GIsRunning)
 	{
-		const TCHAR* CmdLine = FCommandLine::Get();
-		if (FParse::Param( CmdLine, TEXT("fullcrashdumpalways") ))
+		if (FInternationalization::IsAvailable())
 		{
-			NCachedCrashContextProperties::CrashDumpMode = (int32)ECrashDumpMode::FullDumpAlways;
+			NCached::Session.LanguageLCID = FInternationalization::Get().GetCurrentCulture()->GetLCID();
 		}
-		else if (FParse::Param( CmdLine, TEXT("fullcrashdump") ))
+		else
 		{
-			NCachedCrashContextProperties::CrashDumpMode = (int32)ECrashDumpMode::FullDump;
+			FCulturePtr DefaultCulture = FInternationalization::Get().GetCulture(TEXT("en"));
+			if (DefaultCulture.IsValid())
+			{
+				NCached::Session.LanguageLCID = DefaultCulture->GetLCID();
+			}
+			else
+			{
+				const int DefaultCultureLCID = 1033;
+				NCached::Session.LanguageLCID = DefaultCultureLCID;
+			}
 		}
 	}
 
-	const FGuid Guid = FGuid::NewGuid();
-	const FString IniPlatformName(FPlatformProperties::IniPlatformName());
-	NCachedCrashContextProperties::CrashGUIDRoot = FString::Printf(TEXT("%s%s-%s"), *CrashGUIDRootPrefix, *IniPlatformName, *Guid.ToString(EGuidFormats::Digits));
-
 	// Initialize delegate for updating SecondsSinceStart, because FPlatformTime::Seconds() is not POSIX safe.
 	const float PollingInterval = 1.0f;
 	FTicker::GetCoreTicker().AddTicker( FTickerDelegate::CreateLambda( []( float DeltaTime )
 	{
         QUICK_SCOPE_CYCLE_COUNTER(STAT_NCachedCrashContextProperties_LambdaTicker);
 
-		NCachedCrashContextProperties::SecondsSinceStart = int32(FPlatformTime::Seconds() - GStartTime);
+		NCached::Session.SecondsSinceStart = int32(FPlatformTime::Seconds() - GStartTime);
 		return true;
 	} ), PollingInterval );
 
 	FCoreDelegates::UserActivityStringChanged.AddLambda([](const FString& InUserActivity)
 	{
-		NCachedCrashContextProperties::UserActivityHint = InUserActivity;
+		FCString::Strcpy(NCached::Session.UserActivityHint, *InUserActivity);
 	});
 
 	FCoreDelegates::GameSessionIDChanged.AddLambda([](const FString& InGameSessionID)
 	{
-		NCachedCrashContextProperties::GameSessionID = InGameSessionID;
+		FCString::Strcpy(NCached::Session.GameSessionID, *InGameSessionID);
 	});
 
 	FCoreDelegates::GameStateClassChanged.AddLambda([](const FString& InGameStateName)
 	{
-		NCachedCrashContextProperties::GameStateName = InGameStateName;
+		FCString::Strcpy(NCached::Session.GameStateName, *InGameStateName);
 	});
 
 PRAGMA_DISABLE_DEPRECATION_WARNINGS
@@ -201,18 +193,18 @@ PRAGMA_DISABLE_DEPRECATION_WARNINGS
 	{
 		if (InParams.bSetCrashReportClientMessageText)
 		{
-			NCachedCrashContextProperties::CrashReportClientRichText = InParams.CrashReportClientMessageText;
+			FCString::Strcpy(NCached::Session.CrashReportClientRichText, *InParams.CrashReportClientMessageText);
 		}
 		if (InParams.bSetGameNameSuffix)
 		{
-			NCachedCrashContextProperties::GameName = FString(TEXT("UE4-")) + FApp::GetProjectName() + InParams.GameNameSuffix;
+			FCString::Strcpy(NCached::Session.GameName, *(FString(TEXT("UE4-")) + FApp::GetProjectName() + InParams.GameNameSuffix));
 		}
 	});
 PRAGMA_ENABLE_DEPRECATION_WARNINGS
 
 	FCoreDelegates::IsVanillaProductChanged.AddLambda([](bool bIsVanilla)
 	{
-		NCachedCrashContextProperties::bIsVanilla = bIsVanilla;
+		NCached::Session.bIsVanilla = bIsVanilla;
 	});
 
 	FCoreDelegates::ConfigReadyForUse.AddStatic(FGenericCrashContext::InitializeFromConfig);
@@ -221,6 +213,102 @@ PRAGMA_ENABLE_DEPRECATION_WARNINGS
 #endif	// !NOINITCRASHREPORTER
 }
 
+void FGenericCrashContext::InitializeFromContext(const FSessionContext& Session, const TCHAR* EnabledPluginsStr, const TCHAR* EngineDataStr, const TCHAR* GameDataStr)
+{
+	static const TCHAR* TokenDelim[] = { TEXT(","), TEXT("=") };
+
+	// Copy the session struct which should be all pod types and fixed size buggers
+	FMemory::Memcpy(NCached::Session, Session);
+	
+	// Parse the loaded plugins string, assume comma delimited values.
+	if (EnabledPluginsStr)
+	{
+		TArray<FString> Tokens;
+		FString(EnabledPluginsStr).ParseIntoArray(Tokens, TokenDelim, 2, true);
+		NCached::EnabledPluginsList.Append(Tokens);
+	}
+
+	// Parse engine data, comma delimited key=value pairs.
+	if (EngineDataStr)
+	{
+		TArray<FString> Tokens;
+		FString(EngineDataStr).ParseIntoArray(Tokens, TokenDelim, 2, true);
+		int32 i = 0;
+		while ((i + 1) < Tokens.Num())
+		{
+			const FString& Key = Tokens[i++];
+			const FString& Value = Tokens[i++];
+			NCached::EngineData.Add(Key, Value);
+		}
+	}
+
+	// Parse engine data, comma delimited key=value pairs.
+	if (GameDataStr)
+	{
+		TArray<FString> Tokens;
+		FString(GameDataStr).ParseIntoArray(Tokens, TokenDelim, 2, true);
+		int32 i = 0;
+		while ((i + 1) < Tokens.Num())
+		{
+			const FString& Key = Tokens[i++];
+			const FString& Value = Tokens[i++];
+			NCached::GameData.Add(Key, Value);
+		}
+	}
+
+	bIsInitialized = true;
+}
+
+void FGenericCrashContext::CopySharedCrashContext(FSharedCrashContext& Dst)
+{
+	//Copy the session
+	FMemory::Memcpy(Dst.SessionContext, NCached::Session);
+
+	TCHAR* DynamicDataStart = &Dst.DynamicData[0];
+	TCHAR* DynamicDataPtr = DynamicDataStart;
+
+	Dst.EnabledPluginsOffset = DynamicDataPtr - DynamicDataStart;
+	Dst.EnabledPluginsNum = NCached::EnabledPluginsList.Num();
+	for (const FString& Plugin : NCached::EnabledPluginsList)
+	{
+		FCString::Strcat(DynamicDataPtr, Plugin.Len(), *Plugin);
+		FCString::Strcat(DynamicDataPtr, 1, TEXT(","));
+	}
+	DynamicDataPtr += FCString::Strlen(DynamicDataPtr) + 1;
+
+	Dst.EngineDataOffset = DynamicDataPtr - DynamicDataStart;
+	Dst.EngineDataNum = NCached::EngineData.Num();
+	for (const TPair<FString, FString>& Pair : NCached::EngineData)
+	{
+		FCString::Strcat(DynamicDataPtr, Pair.Key.Len(), *Pair.Key);
+		FCString::Strcat(DynamicDataPtr, 1, TEXT("="));
+		FCString::Strcat(DynamicDataPtr, Pair.Value.Len(), *Pair.Value);
+		FCString::Strcat(DynamicDataPtr, 1, TEXT(","));
+	}
+	DynamicDataPtr += FCString::Strlen(DynamicDataPtr) + 1;
+
+	Dst.GameDataOffset = DynamicDataPtr - DynamicDataStart;
+	Dst.GameDataNum = NCached::GameData.Num();
+	for (const TPair<FString, FString>& Pair : NCached::GameData)
+	{
+		FCString::Strcat(DynamicDataPtr, Pair.Key.Len(), *Pair.Key);
+		FCString::Strcat(DynamicDataPtr, 1, TEXT("="));
+		FCString::Strcat(DynamicDataPtr, Pair.Value.Len(), *Pair.Value);
+		FCString::Strcat(DynamicDataPtr, 1, TEXT(","));
+	}
+	DynamicDataPtr += FCString::Strlen(DynamicDataPtr) + 1;
+}
+
+void FGenericCrashContext::SetMemoryStats(const FPlatformMemoryStats& InMemoryStats)
+{
+	NCached::Session.MemoryStats = InMemoryStats;
+
+	// Update cached OOM stats
+	NCached::Session.bIsOOM = FPlatformMemory::bIsOOM;
+	NCached::Session.OOMAllocationSize = FPlatformMemory::OOMAllocationSize;
+	NCached::Session.OOMAllocationAlignment = FPlatformMemory::OOMAllocationAlignment;
+}
+
 void FGenericCrashContext::InitializeFromConfig()
 {
 #if !NOINITCRASHREPORTER
@@ -240,7 +328,7 @@ void FGenericCrashContext::InitializeFromConfig()
 		CrashConfigFile.Add(ConfigSectionName, CRCConfigSectionCopy);
 
 		CrashConfigFile.Dirty = true;
-		CrashConfigFile.Write(GetCrashConfigFilePath());
+		CrashConfigFile.Write(FString(GetCrashConfigFilePath()));
 	}
 
 	// Read the initial un-localized crash context text
@@ -258,13 +346,14 @@ void FGenericCrashContext::UpdateLocalizedStrings()
 	FText CrashReportClientRichText;
 	if (GConfig->GetText(TEXT("CrashContextProperties"), TEXT("CrashReportClientRichText"), CrashReportClientRichText, GEngineIni))
 	{
-		NCachedCrashContextProperties::CrashReportClientRichText = CrashReportClientRichText.ToString();
+		FCString::Strcpy(NCached::Session.CrashReportClientRichText, *CrashReportClientRichText.ToString());
 	}
 #endif
 }
 
 FGenericCrashContext::FGenericCrashContext(ECrashContextType InType, const TCHAR* InErrorMessage)
 	: Type(InType)
+	, CrashedThreadId(~uint32(0))
 	, ErrorMessage(InErrorMessage)
 	, NumMinidumpFramesToIgnore(0)
 {
@@ -285,20 +374,20 @@ void FGenericCrashContext::SerializeContentToBuffer() const
 	AddCrashProperty( TEXT( "CrashVersion" ), (int32)ECrashDescVersions::VER_3_CrashContext );
 	AddCrashProperty( TEXT( "ExecutionGuid" ), *ExecutionGuid.ToString() );
 	AddCrashProperty( TEXT( "CrashGUID" ), (const TCHAR*)CrashGUID);
-	AddCrashProperty( TEXT( "ProcessId" ), FPlatformProcess::GetCurrentProcessId() );
-	AddCrashProperty( TEXT( "IsInternalBuild" ), NCachedCrashContextProperties::bIsInternalBuild );
-	AddCrashProperty( TEXT( "IsPerforceBuild" ), NCachedCrashContextProperties::bIsPerforceBuild );
-	AddCrashProperty( TEXT( "IsSourceDistribution" ), NCachedCrashContextProperties::bIsSourceDistribution );
+	AddCrashProperty( TEXT( "ProcessId" ), NCached::Session.ProcessId );
+	AddCrashProperty( TEXT( "IsInternalBuild" ), NCached::Session.bIsInternalBuild );
+	AddCrashProperty( TEXT( "IsPerforceBuild" ), NCached::Session.bIsPerforceBuild );
+	AddCrashProperty( TEXT( "IsSourceDistribution" ), NCached::Session.bIsSourceDistribution );
 	AddCrashProperty( TEXT( "IsEnsure" ), (Type == ECrashContextType::Ensure) );
 	AddCrashProperty( TEXT( "IsAssert" ), (Type == ECrashContextType::Assert) );
 	AddCrashProperty( TEXT( "CrashType" ), GetCrashTypeString(Type) );
 
-	AddCrashProperty( TEXT( "SecondsSinceStart" ), NCachedCrashContextProperties::SecondsSinceStart );
+	AddCrashProperty( TEXT( "SecondsSinceStart" ), NCached::Session.SecondsSinceStart );
 
 	// Add common crash properties.
-	if (NCachedCrashContextProperties::GameName.Len() > 0)
+	if (FCString::Strlen(NCached::Session.GameName) > 0)
 	{
-		AddCrashProperty(TEXT("GameName"), *NCachedCrashContextProperties::GameName);
+		AddCrashProperty(TEXT("GameName"), NCached::Session.GameName);
 	}
 	else
 	{
@@ -312,9 +401,9 @@ void FGenericCrashContext::SerializeContentToBuffer() const
 			AddCrashProperty(TEXT("GameName"), TEXT(""));
 		}
 	}
-	AddCrashProperty( TEXT( "ExecutableName" ), *NCachedCrashContextProperties::ExecutableName );
+	AddCrashProperty( TEXT( "ExecutableName" ), NCached::Session.ExecutableName );
 	AddCrashProperty( TEXT( "BuildConfiguration" ), LexToString( FApp::GetBuildConfiguration() ) );
-	AddCrashProperty( TEXT( "GameSessionID" ), *NCachedCrashContextProperties::GameSessionID );
+	AddCrashProperty( TEXT( "GameSessionID" ), NCached::Session.GameSessionID );
 	
 	// Unique string specifying the symbols to be used by CrashReporter
 	FString Symbols = FString::Printf( TEXT( "%s-%s-%s" ), FApp::GetBuildVersion(), FPlatformMisc::GetUBTPlatform(), LexToString(FApp::GetBuildConfiguration())).Replace( TEXT( "+" ), TEXT( "*" ));
@@ -326,28 +415,29 @@ void FGenericCrashContext::SerializeContentToBuffer() const
 
 	AddCrashProperty( TEXT( "PlatformName" ), FPlatformProperties::PlatformName() );
 	AddCrashProperty( TEXT( "PlatformNameIni" ), FPlatformProperties::IniPlatformName());
-	AddCrashProperty( TEXT( "EngineMode" ), FPlatformMisc::GetEngineMode() );
-	AddCrashProperty( TEXT( "EngineModeEx" ), EngineModeExString());
+	AddCrashProperty( TEXT( "EngineMode" ), NCached::Session.EngineMode);
+	AddCrashProperty( TEXT( "EngineModeEx" ), NCached::Session.EngineModeEx);
 
-	AddCrashProperty( TEXT( "DeploymentName"), *NCachedCrashContextProperties::DeploymentName );
+	AddCrashProperty( TEXT( "DeploymentName"), NCached::Session.DeploymentName );
 
 	AddCrashProperty( TEXT( "EngineVersion" ), *FEngineVersion::Current().ToString() );
-	AddCrashProperty( TEXT( "CommandLine" ), *NCachedCrashContextProperties::CommandLine );
-	AddCrashProperty( TEXT( "LanguageLCID" ), NCachedCrashContextProperties::LanguageLCID );
-	AddCrashProperty( TEXT( "AppDefaultLocale" ), *NCachedCrashContextProperties::DefaultLocale );
+	AddCrashProperty( TEXT( "CommandLine" ), NCached::Session.CommandLine );
+	AddCrashProperty( TEXT( "LanguageLCID" ), NCached::Session.LanguageLCID );
+	AddCrashProperty( TEXT( "AppDefaultLocale" ), NCached::Session.DefaultLocale );
 	AddCrashProperty( TEXT( "BuildVersion" ), FApp::GetBuildVersion() );
-	AddCrashProperty( TEXT( "IsUE4Release" ), NCachedCrashContextProperties::bIsUE4Release );
+	AddCrashProperty( TEXT( "IsUE4Release" ), NCached::Session.bIsUE4Release );
 
 	// Remove periods from user names to match AutoReporter user names
 	// The name prefix is read by CrashRepository.AddNewCrash in the website code
-	const bool bSendUserName = NCachedCrashContextProperties::bIsInternalBuild;
-	AddCrashProperty( TEXT( "UserName" ), bSendUserName ? *NCachedCrashContextProperties::UserName.Replace( TEXT( "." ), TEXT( "" ) ) : TEXT( "" ) );
+	const bool bSendUserName = NCached::Session.bIsInternalBuild;
+	FString SanitizedUserName = FString(NCached::Session.UserName).Replace(TEXT("."), TEXT(""));
+	AddCrashProperty( TEXT( "UserName" ), bSendUserName ? *SanitizedUserName : TEXT(""));
 
-	AddCrashProperty( TEXT( "BaseDir" ), *NCachedCrashContextProperties::BaseDir );
-	AddCrashProperty( TEXT( "RootDir" ), *NCachedCrashContextProperties::RootDir );
-	AddCrashProperty( TEXT( "MachineId" ), *NCachedCrashContextProperties::LoginIdStr.ToUpper() );
-	AddCrashProperty( TEXT( "LoginId" ), *NCachedCrashContextProperties::LoginIdStr );
-	AddCrashProperty( TEXT( "EpicAccountId" ), *NCachedCrashContextProperties::EpicAccountId );
+	AddCrashProperty( TEXT( "BaseDir" ), NCached::Session.BaseDir );
+	AddCrashProperty( TEXT( "RootDir" ), NCached::Session.RootDir );
+	AddCrashProperty( TEXT( "MachineId" ), *FString(NCached::Session.LoginIdStr).ToUpper() );
+	AddCrashProperty( TEXT( "LoginId" ), NCached::Session.LoginIdStr );
+	AddCrashProperty( TEXT( "EpicAccountId" ), NCached::Session.EpicAccountId );
 
 	// Legacy callstack element for current crash reporter
 	AddCrashProperty( TEXT( "NumMinidumpFramesToIgnore"), NumMinidumpFramesToIgnore );
@@ -359,32 +449,23 @@ void FGenericCrashContext::SerializeContentToBuffer() const
 
 	AddCrashProperty( TEXT( "SourceContext" ), TEXT( "" ) );
 	AddCrashProperty( TEXT( "UserDescription" ), TEXT( "" ) );
-	AddCrashProperty( TEXT( "UserActivityHint" ), *NCachedCrashContextProperties::UserActivityHint );
+	AddCrashProperty( TEXT( "UserActivityHint" ), NCached::Session.UserActivityHint );
 	AddCrashProperty( TEXT( "ErrorMessage" ), ErrorMessage );
-	AddCrashProperty( TEXT( "CrashDumpMode" ), NCachedCrashContextProperties::CrashDumpMode );
-	AddCrashProperty( TEXT( "CrashReporterMessage" ), *NCachedCrashContextProperties::CrashReportClientRichText );
+	AddCrashProperty( TEXT( "CrashDumpMode" ), NCached::Session.CrashDumpMode );
+	AddCrashProperty( TEXT( "CrashReporterMessage" ), NCached::Session.CrashReportClientRichText );
 
 	// Add misc stats.
-	AddCrashProperty( TEXT( "Misc.NumberOfCores" ), NCachedCrashContextProperties::NumberOfCores );
-	AddCrashProperty( TEXT( "Misc.NumberOfCoresIncludingHyperthreads" ), NCachedCrashContextProperties::NumberOfCoresIncludingHyperthreads );
+	AddCrashProperty( TEXT( "Misc.NumberOfCores" ), NCached::Session.NumberOfCores );
+	AddCrashProperty( TEXT( "Misc.NumberOfCoresIncludingHyperthreads" ), NCached::Session.NumberOfCoresIncludingHyperthreads );
 	AddCrashProperty( TEXT( "Misc.Is64bitOperatingSystem" ), (int32)FPlatformMisc::Is64bitOperatingSystem() );
 
-	AddCrashProperty( TEXT( "Misc.CPUVendor" ), *NCachedCrashContextProperties::CPUVendor );
-	AddCrashProperty( TEXT( "Misc.CPUBrand" ), *NCachedCrashContextProperties::CPUBrand );
-	AddCrashProperty( TEXT( "Misc.PrimaryGPUBrand" ), *NCachedCrashContextProperties::PrimaryGPUBrand );
-	AddCrashProperty( TEXT( "Misc.OSVersionMajor" ), *NCachedCrashContextProperties::OsVersion );
-	AddCrashProperty( TEXT( "Misc.OSVersionMinor" ), *NCachedCrashContextProperties::OsSubVersion );
-
-	AddCrashProperty(TEXT("GameStateName"), *NCachedCrashContextProperties::GameStateName);
+	AddCrashProperty( TEXT( "Misc.CPUVendor" ), NCached::Session.CPUVendor );
+	AddCrashProperty( TEXT( "Misc.CPUBrand" ), NCached::Session.CPUBrand );
+	AddCrashProperty( TEXT( "Misc.PrimaryGPUBrand" ), NCached::Session.PrimaryGPUBrand );
+	AddCrashProperty( TEXT( "Misc.OSVersionMajor" ), NCached::Session.OsVersion );
+	AddCrashProperty( TEXT( "Misc.OSVersionMinor" ), NCached::Session.OsSubVersion );
 
-	// #CrashReport: 2015-07-21 Move to the crash report client.
-	/*{
-		uint64 AppDiskTotalNumberOfBytes = 0;
-		uint64 AppDiskNumberOfFreeBytes = 0;
-		FPlatformMisc::GetDiskTotalAndFreeSpace( FPlatformProcess::BaseDir(), AppDiskTotalNumberOfBytes, AppDiskNumberOfFreeBytes );
-		AddCrashProperty( TEXT( "Misc.AppDiskTotalNumberOfBytes" ), AppDiskTotalNumberOfBytes );
-		AddCrashProperty( TEXT( "Misc.AppDiskNumberOfFreeBytes" ), AppDiskNumberOfFreeBytes );
-	}*/
+	AddCrashProperty(TEXT("GameStateName"), NCached::Session.GameStateName);
 
 	// FPlatformMemory::GetConstants is called in the GCreateMalloc, so we can assume it is always valid.
 	{
@@ -397,15 +478,15 @@ void FGenericCrashContext::SerializeContentToBuffer() const
 		AddCrashProperty( TEXT( "MemoryStats.TotalPhysicalGB" ), MemConstants.TotalPhysicalGB );
 	}
 
-	AddCrashProperty( TEXT( "MemoryStats.AvailablePhysical" ), (uint64)CrashMemoryStats.AvailablePhysical );
-	AddCrashProperty( TEXT( "MemoryStats.AvailableVirtual" ), (uint64)CrashMemoryStats.AvailableVirtual );
-	AddCrashProperty( TEXT( "MemoryStats.UsedPhysical" ), (uint64)CrashMemoryStats.UsedPhysical );
-	AddCrashProperty( TEXT( "MemoryStats.PeakUsedPhysical" ), (uint64)CrashMemoryStats.PeakUsedPhysical );
-	AddCrashProperty( TEXT( "MemoryStats.UsedVirtual" ), (uint64)CrashMemoryStats.UsedVirtual );
-	AddCrashProperty( TEXT( "MemoryStats.PeakUsedVirtual" ), (uint64)CrashMemoryStats.PeakUsedVirtual );
-	AddCrashProperty( TEXT( "MemoryStats.bIsOOM" ), (int32)FPlatformMemory::bIsOOM );
-	AddCrashProperty( TEXT( "MemoryStats.OOMAllocationSize"), (uint64)FPlatformMemory::OOMAllocationSize );
-	AddCrashProperty( TEXT( "MemoryStats.OOMAllocationAlignment"), (int32)FPlatformMemory::OOMAllocationAlignment );
+	AddCrashProperty( TEXT( "MemoryStats.AvailablePhysical" ), (uint64)NCached::Session.MemoryStats.AvailablePhysical );
+	AddCrashProperty( TEXT( "MemoryStats.AvailableVirtual" ), (uint64)NCached::Session.MemoryStats.AvailableVirtual );
+	AddCrashProperty( TEXT( "MemoryStats.UsedPhysical" ), (uint64)NCached::Session.MemoryStats.UsedPhysical );
+	AddCrashProperty( TEXT( "MemoryStats.PeakUsedPhysical" ), (uint64)NCached::Session.MemoryStats.PeakUsedPhysical );
+	AddCrashProperty( TEXT( "MemoryStats.UsedVirtual" ), (uint64)NCached::Session.MemoryStats.UsedVirtual );
+	AddCrashProperty( TEXT( "MemoryStats.PeakUsedVirtual" ), (uint64)NCached::Session.MemoryStats.PeakUsedVirtual );
+	AddCrashProperty( TEXT( "MemoryStats.bIsOOM" ), (int) NCached::Session.bIsOOM );
+	AddCrashProperty( TEXT( "MemoryStats.OOMAllocationSize"), NCached::Session.OOMAllocationSize );
+	AddCrashProperty( TEXT( "MemoryStats.OOMAllocationAlignment"), NCached::Session.OOMAllocationAlignment );
 
 	{
 		FString AllThreadStacks;
@@ -423,11 +504,13 @@ void FGenericCrashContext::SerializeContentToBuffer() const
 	// Add platform specific properties.
 	BeginSection( *PlatformPropertiesTag );
 	AddPlatformSpecificProperties();
+	// The name here is a bit cryptic, but we keep it to avoid breaking backend stuff.
+	AddCrashProperty(TEXT("PlatformCallbackResult"), NCached::Session.CrashType);
 	EndSection( *PlatformPropertiesTag );
 
 	// Add the engine data
 	BeginSection( *EngineDataTag );
-	for (const TPair<FString, FString>& Pair : NCachedCrashContextProperties::EngineData)
+	for (const TPair<FString, FString>& Pair : NCached::EngineData)
 	{
 		AddCrashProperty(*Pair.Key, *Pair.Value);
 	}
@@ -435,7 +518,7 @@ void FGenericCrashContext::SerializeContentToBuffer() const
 
 	// Add the game data
 	BeginSection( *GameDataTag );
-	for (const TPair<FString, FString>& Pair : NCachedCrashContextProperties::GameData)
+	for (const TPair<FString, FString>& Pair : NCached::GameData)
 	{
 		AddCrashProperty(*Pair.Key, *Pair.Value);
 	}
@@ -444,11 +527,11 @@ void FGenericCrashContext::SerializeContentToBuffer() const
 	// Writing out the list of plugin JSON descriptors causes us to run out of memory
 	// in GMallocCrash on console, so enable this only for desktop platforms.
 #if PLATFORM_DESKTOP
-	if(NCachedCrashContextProperties::EnabledPluginsList.Num() > 0)
+	if(NCached::EnabledPluginsList.Num() > 0)
 	{
 		BeginSection(*EnabledPluginsTag);
 
-		for (const FString& Str : NCachedCrashContextProperties::EnabledPluginsList)
+		for (const FString& Str : NCached::EnabledPluginsList)
 		{
 			AddCrashProperty(TEXT("Plugin"), *Str);
 		}
@@ -467,24 +550,29 @@ void FGenericCrashContext::SetNumMinidumpFramesToIgnore(int InNumMinidumpFramesT
 
 void FGenericCrashContext::SetDeploymentName(const FString& EpicApp)
 {
-	NCachedCrashContextProperties::DeploymentName = EpicApp;
+	FCString::Strcpy(NCached::Session.DeploymentName, *EpicApp);
+}
+
+void FGenericCrashContext::SetCrashTrigger(ECrashTrigger Type)
+{
+	NCached::Session.CrashType = (int32)Type;
 }
 
 void FGenericCrashContext::GetUniqueCrashName(TCHAR* GUIDBuffer, int32 BufferSize) const
 {
-	FCString::Snprintf(GUIDBuffer, BufferSize, TEXT("%s_%04i"), *NCachedCrashContextProperties::CrashGUIDRoot, CrashContextIndex);
+	FCString::Snprintf(GUIDBuffer, BufferSize, TEXT("%s_%04i"), NCached::Session.CrashGUIDRoot, CrashContextIndex);
 }
 
 const bool FGenericCrashContext::IsFullCrashDump() const
 {
 	if(Type == ECrashContextType::Ensure)
 	{
-		return (NCachedCrashContextProperties::CrashDumpMode == (int32)ECrashDumpMode::FullDumpAlways);
+		return (NCached::Session.CrashDumpMode == (int32)ECrashDumpMode::FullDumpAlways);
 	}
 	else
 	{
-		return (NCachedCrashContextProperties::CrashDumpMode == (int32)ECrashDumpMode::FullDump) ||
-			(NCachedCrashContextProperties::CrashDumpMode == (int32)ECrashDumpMode::FullDumpAlways);
+		return (NCached::Session.CrashDumpMode == (int32)ECrashDumpMode::FullDump) ||
+			(NCached::Session.CrashDumpMode == (int32)ECrashDumpMode::FullDumpAlways);
 	}
 }
 
@@ -660,7 +748,7 @@ FString FGenericCrashContext::UnescapeXMLString( const FString& Text )
 
 FString FGenericCrashContext::GetCrashGameName()
 {
-	return NCachedCrashContextProperties::GameName;
+	return FString(NCached::Session.GameName);
 }
 
 const TCHAR* FGenericCrashContext::GetCrashTypeString(ECrashContextType Type)
@@ -682,18 +770,18 @@ const TCHAR* FGenericCrashContext::GetCrashTypeString(ECrashContextType Type)
 
 const TCHAR* FGenericCrashContext::EngineModeExString()
 {
-	return !NCachedCrashContextProperties::bIsVanilla.IsSet() ? *FGenericCrashContext::EngineModeExUnknown :
-		(NCachedCrashContextProperties::bIsVanilla.GetValue() ? *FGenericCrashContext::EngineModeExVanilla : *FGenericCrashContext::EngineModeExDirty);
+	return !NCached::Session.bIsVanilla.IsSet() ? *FGenericCrashContext::EngineModeExUnknown :
+		(NCached::Session.bIsVanilla.GetValue() ? *FGenericCrashContext::EngineModeExVanilla : *FGenericCrashContext::EngineModeExDirty);
 }
 
 const TCHAR* FGenericCrashContext::GetCrashConfigFilePath()
 {
-	static FString CrashConfigFilePath;
-	if (CrashConfigFilePath.IsEmpty())
+	if (FCString::Strlen(NCached::Session.CrashConfigFilePath) == 0)
 	{
-		CrashConfigFilePath = FPaths::Combine(GetCrashConfigFolder(), *NCachedCrashContextProperties::CrashGUIDRoot, FGenericCrashContext::CrashConfigFileNameW);
+		FString CrashConfigFilePath = FPaths::Combine(GetCrashConfigFolder(), NCached::Session.CrashGUIDRoot, FGenericCrashContext::CrashConfigFileNameW);
+		FCString::Strcpy(NCached::Session.CrashConfigFilePath, *CrashConfigFilePath);
 	}
-	return *CrashConfigFilePath;
+	return NCached::Session.CrashConfigFilePath;
 }
 
 const TCHAR* FGenericCrashContext::GetCrashConfigFolder()
@@ -733,7 +821,7 @@ void FGenericCrashContext::PurgeOldCrashConfig()
 
 void FGenericCrashContext::ResetEngineData()
 {
-	NCachedCrashContextProperties::EngineData.Reset();
+	NCached::EngineData.Reset();
 }
 
 void FGenericCrashContext::SetEngineData(const FString& Key, const FString& Value)
@@ -743,16 +831,16 @@ void FGenericCrashContext::SetEngineData(const FString& Key, const FString& Valu
 		// for testing purposes, only log values when they change, but don't pay the lookup price normally.
 		UE_SUPPRESS(LogCrashContext, VeryVerbose, 
 		{
-			if (NCachedCrashContextProperties::EngineData.Find(Key))
+			if (NCached::EngineData.Find(Key))
 			{
 				UE_LOG(LogCrashContext, VeryVerbose, TEXT("FGenericCrashContext::SetEngineData(%s, <RemoveKey>)"), *Key);
 			}
 		});
-		NCachedCrashContextProperties::EngineData.Remove(Key);
+		NCached::EngineData.Remove(Key);
 	}
 	else
 	{
-		FString& OldVal = NCachedCrashContextProperties::EngineData.FindOrAdd(Key);
+		FString& OldVal = NCached::EngineData.FindOrAdd(Key);
 		UE_SUPPRESS(LogCrashContext, VeryVerbose, 
 		{
 			if (OldVal != Value)
@@ -766,7 +854,7 @@ void FGenericCrashContext::SetEngineData(const FString& Key, const FString& Valu
 
 void FGenericCrashContext::ResetGameData()
 {
-	NCachedCrashContextProperties::GameData.Reset();
+	NCached::GameData.Reset();
 }
 
 void FGenericCrashContext::SetGameData(const FString& Key, const FString& Value)
@@ -776,16 +864,16 @@ void FGenericCrashContext::SetGameData(const FString& Key, const FString& Value)
 		// for testing purposes, only log values when they change, but don't pay the lookup price normally.
 		UE_SUPPRESS(LogCrashContext, VeryVerbose, 
 		{
-			if (NCachedCrashContextProperties::GameData.Find(Key))
+			if (NCached::GameData.Find(Key))
 			{
 				UE_LOG(LogCrashContext, VeryVerbose, TEXT("FGenericCrashContext::SetGameData(%s, <RemoveKey>)"), *Key);
 			}
 		});
-		NCachedCrashContextProperties::GameData.Remove(Key);
+		NCached::GameData.Remove(Key);
 	}
 	else
 	{
-		FString& OldVal = NCachedCrashContextProperties::GameData.FindOrAdd(Key);
+		FString& OldVal = NCached::GameData.FindOrAdd(Key);
 		UE_SUPPRESS(LogCrashContext, VeryVerbose, 
 		{
 			if (OldVal != Value)
@@ -799,9 +887,52 @@ void FGenericCrashContext::SetGameData(const FString& Key, const FString& Value)
 
 void FGenericCrashContext::AddPlugin(const FString& PluginDesc)
 {
-	NCachedCrashContextProperties::EnabledPluginsList.Add(PluginDesc);
+	NCached::EnabledPluginsList.Add(PluginDesc);
+}
+
+void FGenericCrashContext::DumpLog(const FString& CrashFolderAbsolute)
+{
+	// Copy log
+	const FString LogSrcAbsolute = FPlatformOutputDevices::GetAbsoluteLogFilename();
+	FString LogFilename = FPaths::GetCleanFilename(LogSrcAbsolute);
+	const FString LogDstAbsolute = FPaths::Combine(*CrashFolderAbsolute, *LogFilename);
+
+	// If we have a memory only log, make sure it's dumped to file before we attach it to the report
+#if !NO_LOGGING
+	bool bMemoryOnly = FPlatformOutputDevices::GetLog()->IsMemoryOnly();
+	bool bBacklogEnabled = FOutputDeviceRedirector::Get()->IsBacklogEnabled();
+
+	if (bMemoryOnly || bBacklogEnabled)
+	{
+		TUniquePtr<FArchive> LogFile(IFileManager::Get().CreateFileWriter(*LogDstAbsolute, FILEWRITE_AllowRead));
+		if (LogFile)
+		{
+			if (bMemoryOnly)
+			{
+				FPlatformOutputDevices::GetLog()->Dump(*LogFile);
+			}
+			else
+			{
+				FOutputDeviceArchiveWrapper Wrapper(LogFile.Get());
+				GLog->SerializeBacklog(&Wrapper);
+			}
+
+			LogFile->Flush();
+		}
+	}
+	else
+	{
+		const bool bReplace = true;
+		const bool bEvenIfReadOnly = false;
+		const bool bAttributes = false;
+		FCopyProgress* const CopyProgress = nullptr;
+		static_cast<void>(IFileManager::Get().Copy(*LogDstAbsolute, *LogSrcAbsolute, bReplace, bEvenIfReadOnly, bAttributes, CopyProgress, FILEREAD_AllowWrite, FILEWRITE_AllowRead));	// best effort, so don't care about result: couldn't copy -> tough, no log
+	}
+#endif // !NO_LOGGING
 }
 
+
+
 FORCENOINLINE void FGenericCrashContext::CapturePortableCallStack(int32 NumStackFramesToIgnore, void* Context)
 {
 	// If the callstack is for the executing thread, ignore this function
@@ -868,6 +999,45 @@ void FGenericCrashContext::GetPortableCallStack(const uint64* StackFrames, int32
 	}
 }
 
+void FGenericCrashContext::AddPortableThreadCallStack(uint32 ThreadId, const TCHAR* ThreadName, const uint64* StackFrames, int32 NumStackFrames)
+{
+	// Not implemented for generic class
+}
+
+void FGenericCrashContext::CopyPlatformSpecificFiles(const TCHAR* OutputDirectory, void* Context)
+{
+	// If present, include the crash report config file to pass config values to the CRC
+	const TCHAR* CrashConfigSrcPath = GetCrashConfigFilePath();
+	if (IFileManager::Get().FileExists(CrashConfigSrcPath))
+	{
+		FString CrashConfigFilename = FPaths::GetCleanFilename(CrashConfigSrcPath);
+		const FString CrashConfigDstAbsolute = FPaths::Combine(OutputDirectory, *CrashConfigFilename);
+		IFileManager::Get().Copy(*CrashConfigDstAbsolute, CrashConfigSrcPath);	// best effort, so don't care about result: couldn't copy -> tough, no config
+	}
+}
+
+/**
+ * Attempts to create the output report directory.
+ */
+bool FGenericCrashContext::CreateCrashReportDirectory(const TCHAR* CrashGUIDRoot, const TCHAR* AppName, int32 CrashIndex, FString& OutCrashDirectoryAbsolute)
+{
+	// Generate Crash GUID
+	TCHAR CrashGUID[FGenericCrashContext::CrashGUIDLength];
+	FCString::Snprintf(CrashGUID, FGenericCrashContext::CrashGUIDLength, TEXT("%s_%04i"), CrashGUIDRoot, CrashIndex);
+
+	// The FPaths commands usually checks for command line override, if FCommandLine not yet
+	// initialized we cannot create a directory. Also there is no way of knowing if the file manager
+	// has been created.
+	if (!FCommandLine::IsInitialized())
+	{
+		return false;
+	}
+
+	FString CrashFolder = FPaths::Combine(*FPaths::ProjectSavedDir(), TEXT("Crashes"), CrashGUID);
+	OutCrashDirectoryAbsolute = IFileManager::Get().ConvertToAbsolutePathForExternalAppForWrite(*CrashFolder);
+	return IFileManager::Get().MakeDirectory(*OutCrashDirectoryAbsolute, true);
+}
+
 FProgramCounterSymbolInfoEx::FProgramCounterSymbolInfoEx( FString InModuleName, FString InFunctionName, FString InFilename, uint32 InLineNumber, uint64 InSymbolDisplacement, uint64 InOffsetInModule, uint64 InProgramCounter ) :
 	ModuleName( InModuleName ),
 	FunctionName( InFunctionName ),
diff --git a/Engine/Source/Runtime/Core/Private/GenericPlatform/GenericPlatformMisc.cpp b/Engine/Source/Runtime/Core/Private/GenericPlatform/GenericPlatformMisc.cpp
--- a/Engine/Source/Runtime/Core/Private/GenericPlatform/GenericPlatformMisc.cpp
+++ b/Engine/Source/Runtime/Core/Private/GenericPlatform/GenericPlatformMisc.cpp
@@ -453,7 +453,7 @@ void FGenericPlatformMisc::HandleIOFailure( const TCHAR* Filename )
 void FGenericPlatformMisc::RaiseException(uint32 ExceptionCode)
 {
 	/** This is the last place to gather memory stats before exception. */
-	FGenericCrashContext::CrashMemoryStats = FPlatformMemory::GetStats();
+	FGenericCrashContext::SetMemoryStats(FPlatformMemory::GetStats());
 
 #if HACK_HEADER_GENERATOR && !PLATFORM_EXCEPTIONS_DISABLED
 	// We want Unreal Header Tool to throw an exception but in normal runtime code 
diff --git a/Engine/Source/Runtime/Core/Private/GenericPlatform/GenericPlatformProcess.cpp b/Engine/Source/Runtime/Core/Private/GenericPlatform/GenericPlatformProcess.cpp
--- a/Engine/Source/Runtime/Core/Private/GenericPlatform/GenericPlatformProcess.cpp
+++ b/Engine/Source/Runtime/Core/Private/GenericPlatform/GenericPlatformProcess.cpp
@@ -279,6 +279,12 @@ bool FGenericPlatformProcess::IsApplicationRunning( const TCHAR* ProcName )
 	return false;
 }
 
+bool FGenericPlatformProcess::IsApplicationAlive(uint32 ProcessId)
+{
+	UE_LOG(LogHAL, Fatal, TEXT("FGenericPlatformProcess::IsApplicationAlive not implemented on this platform"));
+	return false;
+}
+
 FString FGenericPlatformProcess::GetApplicationName( uint32 ProcessId )
 {
 	UE_LOG(LogHAL, Fatal, TEXT("FGenericPlatformProcess::GetApplicationName not implemented on this platform"));
@@ -549,11 +555,22 @@ bool FGenericPlatformProcess::SupportsMultithreading()
 }
 
 FGenericPlatformProcess::FSemaphore::FSemaphore(const FString& InName)
+	: FSemaphore(*InName)
 {
-	FCString::Strcpy(Name, sizeof(Name)-1, *InName);
+}
+
+
+FGenericPlatformProcess::FSemaphore::FSemaphore(const TCHAR* InName) 
+{
+	FCString::Strcpy(Name, sizeof(Name) - 1, InName);
 }
 
 FGenericPlatformProcess::FSemaphore* FGenericPlatformProcess::NewInterprocessSynchObject(const FString& Name, bool bCreate, uint32 MaxLocks)
+{
+	return NewInterprocessSynchObject(*Name, bCreate, MaxLocks);
+}
+
+FGenericPlatformProcess::FSemaphore* FGenericPlatformProcess::NewInterprocessSynchObject(const TCHAR* Name, bool bCreate, uint32 MaxLocks)
 {
 	UE_LOG(LogHAL, Fatal, TEXT("FGenericPlatformProcess::NewInterprocessSynchObject not implemented on this platform"));
 	return NULL;
diff --git a/Engine/Source/Runtime/Core/Private/HAL/ExceptionHandling.cpp b/Engine/Source/Runtime/Core/Private/HAL/ExceptionHandling.cpp
--- a/Engine/Source/Runtime/Core/Private/HAL/ExceptionHandling.cpp
+++ b/Engine/Source/Runtime/Core/Private/HAL/ExceptionHandling.cpp
@@ -28,19 +28,8 @@ CORE_API bool GIgnoreDebugger = false;
 CORE_API TCHAR MiniDumpFilenameW[1024] = TEXT("");
 
 
-volatile int32 GCrashType = 0;
 bool GEnsureShowsCRC = false;
 
-void SetCrashType(ECrashType InCrashType)
-{
-	GCrashType = (int32)InCrashType;
-}
-
-int32 GetCrashType()
-{
-	return GCrashType;
-}
-
 void ReportInteractiveEnsure(const TCHAR* InMessage)
 {
 	GEnsureShowsCRC = true;
diff --git a/Engine/Source/Runtime/Core/Private/HAL/ThreadingBase.cpp b/Engine/Source/Runtime/Core/Private/HAL/ThreadingBase.cpp
--- a/Engine/Source/Runtime/Core/Private/HAL/ThreadingBase.cpp
+++ b/Engine/Source/Runtime/Core/Private/HAL/ThreadingBase.cpp
@@ -260,9 +260,12 @@ void FThreadManager::GetAllThreadStackBackTraces(TArray<FThreadStackBackTrace>&
 FThreadManager& FThreadManager::Get()
 {
 	static FThreadManager Singleton;
+	FThreadManager::bIsInitialized = true;
 	return Singleton;
 }
 
+bool FThreadManager::bIsInitialized = false;
+
 
 /*-----------------------------------------------------------------------------
 	FEvent, FScopedEvent
diff --git a/Engine/Source/Runtime/Core/Private/Windows/WindowsPlatformCrashContext.cpp b/Engine/Source/Runtime/Core/Private/Windows/WindowsPlatformCrashContext.cpp
--- a/Engine/Source/Runtime/Core/Private/Windows/WindowsPlatformCrashContext.cpp
+++ b/Engine/Source/Runtime/Core/Private/Windows/WindowsPlatformCrashContext.cpp
@@ -30,20 +30,99 @@
 #include <dbghelp.h>
 #include <Shlwapi.h>
 #include <psapi.h>
+#include <tlhelp32.h>
 
 #ifndef UE_LOG_CRASH_CALLSTACK
 	#define UE_LOG_CRASH_CALLSTACK 1
 #endif
 
+#if !IS_PROGRAM
+	#define USE_CRASH_REPORTER_MONITOR 1
+#else 
+	#define USE_CRASH_REPORTER_MONITOR 0
+#endif
+
+#define CR_CLIENT_MAX_PATH_LEN 265
+#define CR_CLIENT_MAX_ARGS_LEN 256
+// Enable to use a non shipping version of the crash report client
+#define CR_USE_DEVELOPMENT_CLIENT 0
+
 #pragma comment( lib, "version.lib" )
 #pragma comment( lib, "Shlwapi.lib" )
 
+LONG WINAPI UnhandledException(EXCEPTION_POINTERS *ExceptionInfo);
+LONG WINAPI UnhandledStaticInitException(LPEXCEPTION_POINTERS ExceptionInfo);
+
+/** Platform specific constants. */
+enum EConstants
+{
+	UE4_MINIDUMP_CRASHCONTEXT = LastReservedStream + 1,
+};
+
+
 /**
  * Code for an assert exception
  */
 const uint32 AssertExceptionCode = 0x4000;
 const uint32 GPUCrashExceptionCode = 0x8000;
 
+namespace {
+	/**
+	 * Write a Windows minidump to disk
+	 * @param The Crash context with its data already serialized into its buffer
+	 * @param Path Full path of file to write (normally a .dmp file)
+	 * @param ExceptionInfo Pointer to structure containing the exception information
+	 * @return Success or failure
+	 */
+
+	 // #CrashReport: 2014-10-08 Move to FWindowsPlatformCrashContext
+	bool WriteMinidump(HANDLE Process, DWORD ThreadId, FWindowsPlatformCrashContext& InContext, const TCHAR* Path, LPEXCEPTION_POINTERS ExceptionInfo)
+	{
+		// Are we calling this in process or from an external process?
+		const BOOL bIsClientPointers = Process != GetCurrentProcess();
+
+		// Try to create file for minidump.
+		HANDLE FileHandle = CreateFileW(Path, GENERIC_WRITE, 0, NULL, CREATE_ALWAYS, FILE_ATTRIBUTE_NORMAL, NULL);
+
+		if (FileHandle == INVALID_HANDLE_VALUE)
+		{
+			return false;
+		}
+
+		// Initialise structure required by MiniDumpWriteDumps
+		MINIDUMP_EXCEPTION_INFORMATION DumpExceptionInfo = {};
+
+		DumpExceptionInfo.ThreadId = ThreadId;
+		DumpExceptionInfo.ExceptionPointers = ExceptionInfo;
+		DumpExceptionInfo.ClientPointers = bIsClientPointers;
+
+		// CrashContext.runtime-xml is now a part of the minidump file.
+		MINIDUMP_USER_STREAM CrashContextStream = { 0 };
+		CrashContextStream.Type = UE4_MINIDUMP_CRASHCONTEXT;
+		CrashContextStream.BufferSize = InContext.GetBuffer().GetAllocatedSize();
+		CrashContextStream.Buffer = (void*)*InContext.GetBuffer();
+
+		MINIDUMP_USER_STREAM_INFORMATION CrashContextStreamInformation = { 0 };
+		CrashContextStreamInformation.UserStreamCount = 1;
+		CrashContextStreamInformation.UserStreamArray = &CrashContextStream;
+
+		MINIDUMP_TYPE MinidumpType = MiniDumpNormal;//(MINIDUMP_TYPE)(MiniDumpWithPrivateReadWriteMemory|MiniDumpWithDataSegs|MiniDumpWithHandleData|MiniDumpWithFullMemoryInfo|MiniDumpWithThreadInfo|MiniDumpWithUnloadedModules);
+
+		// For ensures by default we use minidump to avoid severe hitches when writing 3GB+ files.
+		// However the crash dump mode will remain the same.
+		bool bShouldBeFullCrashDump = InContext.IsFullCrashDump();
+		if (bShouldBeFullCrashDump)
+		{
+			MinidumpType = (MINIDUMP_TYPE)(MiniDumpWithFullMemory | MiniDumpWithFullMemoryInfo | MiniDumpWithHandleData | MiniDumpWithThreadInfo | MiniDumpWithUnloadedModules);
+		}
+
+		const BOOL Result = MiniDumpWriteDump(Process, GetProcessId(Process), FileHandle, MinidumpType, &DumpExceptionInfo, &CrashContextStreamInformation, NULL);
+		CloseHandle(FileHandle);
+
+		return Result == TRUE;
+	}
+}
+
 /**
  * Stores information about an assert that can be unpacked in the exception handler.
  */
@@ -59,14 +138,14 @@ struct FAssertInfo
 	}
 };
 
-void FWindowsPlatformCrashContext::GetProcModuleHandles(FModuleHandleArray& OutHandles)
+void FWindowsPlatformCrashContext::GetProcModuleHandles(const FProcHandle& ProcessHandle, FModuleHandleArray& OutHandles)
 {
 	// Get all the module handles for the current process. Each module handle is its base address.
 	for (;;)
 	{
 		DWORD BufferSize = OutHandles.Num() * sizeof(HMODULE);
 		DWORD RequiredBufferSize = 0;
-		if (!EnumProcessModules(GetCurrentProcess(), (HMODULE*)OutHandles.GetData(), BufferSize, &RequiredBufferSize))
+		if (!EnumProcessModulesEx(ProcessHandle.IsValid() ? ProcessHandle.Get() : GetCurrentProcess(), (HMODULE*)OutHandles.GetData(), BufferSize, &RequiredBufferSize, LIST_MODULES_ALL))
 		{
 			return;
 		}
@@ -81,6 +160,7 @@ void FWindowsPlatformCrashContext::GetProcModuleHandles(FModuleHandleArray& OutH
 }
 
 void FWindowsPlatformCrashContext::ConvertProgramCountersToStackFrames(
+	const FProcHandle& ProcessHandle,
 	const FModuleHandleArray& SortedModuleHandles,
 	const uint64* ProgramCounters,
 	int32 NumPCs,
@@ -100,7 +180,7 @@ void FWindowsPlatformCrashContext::ConvertProgramCountersToStackFrames(
 		else
 		{
 			TCHAR ModuleName[MAX_PATH];
-			if (GetModuleFileNameW((HMODULE)SortedModuleHandles[ModuleIdx], ModuleName, MAX_PATH) != 0)
+			if (GetModuleFileNameExW(ProcessHandle.IsValid() ? ProcessHandle.Get() : GetCurrentProcess(), (HMODULE)SortedModuleHandles[ModuleIdx], ModuleName, MAX_PATH) != 0)
 			{
 				TCHAR* ModuleNameEnd = FCString::Strrchr(ModuleName, '\\');
 				if (ModuleNameEnd != nullptr)
@@ -116,6 +196,7 @@ void FWindowsPlatformCrashContext::ConvertProgramCountersToStackFrames(
 			}
 			else
 			{
+				const DWORD Err = GetLastError();
 				FCString::Strcpy(ModuleName, TEXT("Unknown"));
 			}
 
@@ -129,67 +210,50 @@ void FWindowsPlatformCrashContext::ConvertProgramCountersToStackFrames(
 void FWindowsPlatformCrashContext::SetPortableCallStack(const uint64* StackTrace, int32 StackTraceDepth)
 {
 	FModuleHandleArray ProcessModuleHandles;
-	GetProcModuleHandles(ProcessModuleHandles);
-	ConvertProgramCountersToStackFrames(ProcessModuleHandles, StackTrace, StackTraceDepth, CallStack);
+	GetProcModuleHandles(ProcessHandle, ProcessModuleHandles);
+	ConvertProgramCountersToStackFrames(ProcessHandle, ProcessModuleHandles, StackTrace, StackTraceDepth, CallStack);
 }
 
 void FWindowsPlatformCrashContext::AddPlatformSpecificProperties() const
 {
 	AddCrashProperty(TEXT("PlatformIsRunningWindows"), 1);
-	// On windows track the crash type
-	AddCrashProperty(TEXT("PlatformCallbackResult"), GetCrashType());
-
 	AddCrashProperty(TEXT("IsRunningOnBattery"), FPlatformMisc::IsRunningOnBattery());
 }
 
 bool FWindowsPlatformCrashContext::GetPlatformAllThreadContextsString(FString& OutStr) const
 {
-	OutStr = AllThreadContexts;
+	for (const FThreadStackFrames& Thread : ThreadCallStacks)
+	{
+		AddThreadContextString(
+			CrashedThreadId, 
+			Thread.ThreadId, 
+			Thread.ThreadName, 
+			Thread.StackFrames,
+			OutStr
+		);
+	}
 	return !OutStr.IsEmpty();
 }
 
-void FWindowsPlatformCrashContext::AddIsCrashed(bool bIsCrashed, FString& OutStr)
-{
-	OutStr += FString::Printf(TEXT("<IsCrashed>%s</IsCrashed>"), bIsCrashed ? TEXT("true") : TEXT("false"));
-	OutStr += LINE_TERMINATOR;
-}
-
-void FWindowsPlatformCrashContext::AddThreadId(uint32 ThreadId, FString& OutStr)
-{
-	OutStr += FString::Printf(TEXT("<ThreadID>%d</ThreadID>"), ThreadId);
-	OutStr += LINE_TERMINATOR;
-}
-
-void FWindowsPlatformCrashContext::AddThreadName(const TCHAR* ThreadName, FString& OutStr)
-{
-	OutStr += FString::Printf(TEXT("<ThreadName>%s</ThreadName>"), ThreadName);
-	OutStr += LINE_TERMINATOR;
-}
-
-void FWindowsPlatformCrashContext::AddThreadContext(
-	const FModuleHandleArray& ProcModuleHandles,
+void FWindowsPlatformCrashContext::AddThreadContextString(
 	uint32 CrashedThreadId,
 	uint32 ThreadId,
 	const FString& ThreadName,
-	const uint64* StackTrace,
-	int32 Depth,
+	const TArray<FCrashStackFrame>& StackFrames,
 	FString& OutStr)
 {
 	OutStr += TEXT("<Thread>");
 	{
 		OutStr += TEXT("<CallStack>");
 
-		TArray<FCrashStackFrame> CrashStackFrames;
-		ConvertProgramCountersToStackFrames(ProcModuleHandles, StackTrace, Depth, CrashStackFrames);
-
 		int32 MaxModuleNameLen = 0;
-		for (const FCrashStackFrame& StFrame : CrashStackFrames)
+		for (const FCrashStackFrame& StFrame : StackFrames)
 		{
 			MaxModuleNameLen = FMath::Max(MaxModuleNameLen, StFrame.ModuleName.Len());
 		}
 
 		FString CallstackStr;
-		for (const FCrashStackFrame& StFrame : CrashStackFrames)
+		for (const FCrashStackFrame& StFrame : StackFrames)
 		{
 			CallstackStr += FString::Printf(TEXT("%-*s 0x%016x + %-8x"), MaxModuleNameLen + 1, *StFrame.ModuleName, StFrame.BaseAddress, StFrame.Offset);
 			CallstackStr += LINE_TERMINATOR;
@@ -198,94 +262,150 @@ void FWindowsPlatformCrashContext::AddThreadContext(
 		OutStr += TEXT("</CallStack>");
 		OutStr += LINE_TERMINATOR;
 	}
-	AddIsCrashed(ThreadId == CrashedThreadId, OutStr);
+	OutStr += FString::Printf(TEXT("<IsCrashed>%s</IsCrashed>"), ThreadId == CrashedThreadId ? TEXT("true") : TEXT("false"));
+	OutStr += LINE_TERMINATOR;
 	// TODO: do we need thread register states?
 	OutStr += TEXT("<Registers></Registers>");
 	OutStr += LINE_TERMINATOR;
-	AddThreadId(ThreadId, OutStr);
-	AddThreadName(*ThreadName, OutStr);
+	OutStr += FString::Printf(TEXT("<ThreadID>%d</ThreadID>"), ThreadId);
+	OutStr += LINE_TERMINATOR;
+	OutStr += FString::Printf(TEXT("<ThreadName>%s</ThreadName>"), *ThreadName);
+	OutStr += LINE_TERMINATOR;
 	OutStr += TEXT("</Thread>");
 	OutStr += LINE_TERMINATOR;
 }
 
-void FWindowsPlatformCrashContext::AddAllThreadContexts(uint32 CrashedThreadId, FString& OutStr)
+void FWindowsPlatformCrashContext::AddPortableThreadCallStack(uint32 ThreadId, const TCHAR* ThreadName, const uint64* StackFrames, int32 NumStackFrames)
 {
-	TArray<typename FThreadManager::FThreadStackBackTrace> StackTraces;
-	FThreadManager::Get().GetAllThreadStackBackTraces(StackTraces);
-
 	FModuleHandleArray ProcModuleHandles;
-	GetProcModuleHandles(ProcModuleHandles);
+	GetProcModuleHandles(ProcessHandle, ProcModuleHandles);
+
+	FThreadStackFrames Thread;
+	Thread.ThreadId = ThreadId;
+	Thread.ThreadName = FString(ThreadName);
+	ConvertProgramCountersToStackFrames(ProcessHandle, ProcModuleHandles, StackFrames, NumStackFrames, Thread.StackFrames);
+	ThreadCallStacks.Push(Thread);
+}
+
+void FWindowsPlatformCrashContext::CopyPlatformSpecificFiles(const TCHAR* OutputDirectory, void* Context)
+{
+	FGenericCrashContext::CopyPlatformSpecificFiles(OutputDirectory, Context);
+
+	// Save minidump
+	LPEXCEPTION_POINTERS ExceptionInfo = (LPEXCEPTION_POINTERS)Context;
+	if (ExceptionInfo != nullptr)
+	{
+		const FString MinidumpFileName = FPaths::Combine(OutputDirectory, *FGenericCrashContext::UE4MinidumpName);
+		WriteMinidump(ProcessHandle.Get(), CrashedThreadId, *this, *MinidumpFileName, ExceptionInfo);
+	}
 
-	for (int32 Idx = 0; Idx < StackTraces.Num(); ++Idx)
+	// If present, include the crash video
+	const FString CrashVideoPath = FPaths::ProjectLogDir() / TEXT("CrashVideo.avi");
+	if (IFileManager::Get().FileExists(*CrashVideoPath))
 	{
-		const auto& ThreadStTrace = StackTraces[Idx];
-		const uint32 ThreadId = ThreadStTrace.ThreadId;
-		const FString& ThreadName = ThreadStTrace.ThreadName;
-		const auto& ThreadPCs = ThreadStTrace.ProgramCounters;
-		AddThreadContext(ProcModuleHandles, CrashedThreadId, ThreadId, ThreadName, ThreadPCs.GetData(), ThreadPCs.Num(), OutStr);
+		FString CrashVideoFilename = FPaths::GetCleanFilename(CrashVideoPath);
+		const FString CrashVideoDstAbsolute = FPaths::Combine(OutputDirectory, *CrashVideoFilename);
+		static_cast<void>(IFileManager::Get().Copy(*CrashVideoDstAbsolute, *CrashVideoPath));	// best effort, so don't care about result: couldn't copy -> tough, no video
 	}
 }
 
-/** Platform specific constants. */
-enum EConstants
+void FWindowsPlatformCrashContext::CaptureAllThreadContexts()
 {
-	UE4_MINIDUMP_CRASHCONTEXT = LastReservedStream + 1,
-};
+	TArray<typename FThreadManager::FThreadStackBackTrace> StackTraces;
+	FThreadManager::Get().GetAllThreadStackBackTraces(StackTraces);
+
+	for (const FThreadManager::FThreadStackBackTrace& Thread : StackTraces)
+	{
+		AddPortableThreadCallStack(Thread.ThreadId, *Thread.ThreadName, Thread.ProgramCounters.GetData(), Thread.ProgramCounters.Num());
+	}	
+}
+
 
 namespace
 {
-static int32 ReportCrashCallCount = 0;
 
-/**
- * Write a Windows minidump to disk
- * @param The Crash context with its data already serialized into its buffer
- * @param Path Full path of file to write (normally a .dmp file)
- * @param ExceptionInfo Pointer to structure containing the exception information
- * @return Success or failure
- */
+static int32 ReportCrashCallCount = 0;
 
-// #CrashReport: 2014-10-08 Move to FWindowsPlatformCrashContext
-bool WriteMinidump(FWindowsPlatformCrashContext& InContext, const TCHAR* Path, LPEXCEPTION_POINTERS ExceptionInfo)
+static FORCEINLINE bool CreatePipeWrite(void*& ReadPipe, void*& WritePipe)
 {
-	// Try to create file for minidump.
-	HANDLE FileHandle = CreateFileW(Path, GENERIC_WRITE, 0, NULL, CREATE_ALWAYS, FILE_ATTRIBUTE_NORMAL, NULL);
-	
-	if (FileHandle == INVALID_HANDLE_VALUE)
+	SECURITY_ATTRIBUTES Attr = { sizeof(SECURITY_ATTRIBUTES), NULL, true };
+
+	if (!::CreatePipe(&ReadPipe, &WritePipe, &Attr, 0))
+	{
+		return false;
+	}
+
+	if (!::SetHandleInformation(WritePipe, HANDLE_FLAG_INHERIT, 0))
 	{
 		return false;
 	}
 
-	// Initialise structure required by MiniDumpWriteDumps
-	MINIDUMP_EXCEPTION_INFORMATION DumpExceptionInfo = {};
+	return true;
+}
+
+/**
+ * Finds the crash reporter binary path.
+ */
+void CreateCrashReportClientPath(TCHAR* OutClientPath, TCHAR* OutWorkingDir, int32 MaxLength)
+{
+#if CR_USE_DEVELOPMENT_CLIENT
+	static const TCHAR CrashReportClientExeName[] = TEXT("CrashReportClient-Win64-Development.exe");
+#else
+	static const TCHAR CrashReportClientExeName[] = TEXT("CrashReportClient.exe");
+#endif 
+
+	// Use the same working directory as this process
+	GetCurrentDirectory(MaxLength, OutWorkingDir);
 
-	DumpExceptionInfo.ThreadId			= GetCurrentThreadId();
-	DumpExceptionInfo.ExceptionPointers	= ExceptionInfo;
-	DumpExceptionInfo.ClientPointers	= FALSE;
+	TCHAR RuntimeFilename[CR_CLIENT_MAX_PATH_LEN];
+	const DWORD RuntimeFilenameLength = GetModuleFileNameEx(GetCurrentProcess(), NULL, RuntimeFilename, CR_CLIENT_MAX_PATH_LEN);
 
-	// CrashContext.runtime-xml is now a part of the minidump file.
-	MINIDUMP_USER_STREAM CrashContextStream ={0};
-	CrashContextStream.Type = UE4_MINIDUMP_CRASHCONTEXT;
-	CrashContextStream.BufferSize = InContext.GetBuffer().GetAllocatedSize();
-	CrashContextStream.Buffer = (void*)*InContext.GetBuffer();
+	// Assume the crash reporter client is in the same location as the game runtime.
+	const TCHAR* LastDelimiter = FCString::Strrchr(RuntimeFilename, TEXT('\\'));
+	FCString::Strncat(OutClientPath, RuntimeFilename, LastDelimiter - RuntimeFilename + 2);
+	FCString::Strncat(OutClientPath, CrashReportClientExeName, MaxLength);
+}
 
-	MINIDUMP_USER_STREAM_INFORMATION CrashContextStreamInformation = {0};
-	CrashContextStreamInformation.UserStreamCount = 1;
-	CrashContextStreamInformation.UserStreamArray = &CrashContextStream;
+/**
+ * Launches crash reporter client and creates the pipes for communication.
+ */
+FProcHandle LaunchCrashReportClient(void** OutWritePipe, void** OutReadPipe)
+{
+	TCHAR CrashReporterClientPath[CR_CLIENT_MAX_PATH_LEN] = { 0 };
+	TCHAR CrashReporterClientArgs[CR_CLIENT_MAX_ARGS_LEN] = { 0 };
+	TCHAR WorkingDirectory[CR_CLIENT_MAX_PATH_LEN] = { 0 };
 
-	MINIDUMP_TYPE MinidumpType = MiniDumpNormal;//(MINIDUMP_TYPE)(MiniDumpWithPrivateReadWriteMemory|MiniDumpWithDataSegs|MiniDumpWithHandleData|MiniDumpWithFullMemoryInfo|MiniDumpWithThreadInfo|MiniDumpWithUnloadedModules);
+	void *PipeChildInRead, *PipeChildInWrite, *PipeChildOutRead, *PipeChildOutWrite;
 
-	// For ensures by default we use minidump to avoid severe hitches when writing 3GB+ files.
-	// However the crash dump mode will remain the same.
-	bool bShouldBeFullCrashDump = InContext.IsFullCrashDump();
-	if (bShouldBeFullCrashDump)
+	if (!CreatePipeWrite(PipeChildInRead, PipeChildInWrite) || !FPlatformProcess::CreatePipe(PipeChildOutRead, PipeChildOutWrite))
 	{
-		MinidumpType = (MINIDUMP_TYPE)(MiniDumpWithFullMemory|MiniDumpWithFullMemoryInfo|MiniDumpWithHandleData|MiniDumpWithThreadInfo|MiniDumpWithUnloadedModules);
+		return FProcHandle();
 	}
 
-	const BOOL Result = MiniDumpWriteDump(GetCurrentProcess(), GetCurrentProcessId(), FileHandle, MinidumpType, &DumpExceptionInfo, &CrashContextStreamInformation, NULL);
-	CloseHandle(FileHandle);
+	// Pass the endpoints to the creator of the client ...
+	*OutWritePipe = PipeChildInWrite;
+	*OutReadPipe = PipeChildOutRead;
+	
+	// ... and the other ends to the child
+	FCString::Sprintf(CrashReporterClientArgs, TEXT(" -READ=%0u -WRITE=%0u"), PipeChildInRead, PipeChildOutWrite);
+
+	{
+		TCHAR PidStr[256] = { 0 };
+		FCString::Sprintf(PidStr, TEXT(" -MONITOR=%u"), FPlatformProcess::GetCurrentProcessId());
+		FCString::Strncat(CrashReporterClientArgs, PidStr, CR_CLIENT_MAX_ARGS_LEN);
+	}
 
-	return Result == TRUE;
+	CreateCrashReportClientPath(CrashReporterClientPath, WorkingDirectory, CR_CLIENT_MAX_PATH_LEN);
+	
+	return FPlatformProcess::CreateProc(
+		CrashReporterClientPath,
+		CrashReporterClientArgs,
+		true, false, false,
+		nullptr, 0, 
+		WorkingDirectory, 
+		nullptr, 
+		nullptr
+	);
 }
 
 /**
@@ -300,6 +420,162 @@ enum class EErrorReportUI
 	ReportInUnattendedMode	
 };
 
+/**
+ * Write required information about the crash to the shared context, and then signal the crash reporter client 
+ * running in monitor mode about the crash.
+ */
+int32 ReportCrashForMonitor(
+	LPEXCEPTION_POINTERS ExceptionInfo,
+	ECrashContextType Type,
+	const TCHAR* ErrorMessage,
+	int NumStackFramesToIgnore,
+	HANDLE CrashingThreadHandle,
+	DWORD CrashingThreadId,
+	FProcHandle& CrashMonitorHandle,
+	FSharedCrashContext* SharedContext,
+	void* WritePipe,
+	void* ReadPipe,
+	EErrorReportUI ReportUI)
+{
+	// Determine UI settings for the crash report. Suppress the user input dialog if we're running in unattended mode
+	// Usage data controls if we want analytics in the crash report client
+	// Finally we cannot call some of these functions if we crash during static init, so check if they are initialized.
+	bool bNoDialog = ReportUI == EErrorReportUI::ReportInUnattendedMode;
+	bool bSendUnattendedBugReports = true;
+	bool bSendUsageData = true;
+	if (FCommandLine::IsInitialized())
+	{
+		bNoDialog |= FApp::IsUnattended() || IsRunningDedicatedServer();
+	}
+	// If we don't have a command line, but know we are dedicated server, also run in no dialog mode.
+	else if (IsRunningDedicatedServer())
+	{
+		bNoDialog = true;
+	}
+
+	if (GConfig)
+	{
+		GConfig->GetBool(TEXT("/Script/UnrealEd.CrashReportsPrivacySettings"), TEXT("bSendUnattendedBugReports"), bSendUnattendedBugReports, GEditorSettingsIni);
+		GConfig->GetBool(TEXT("/Script/UnrealEd.AnalyticsPrivacySettings"), TEXT("bSendUsageData"), bSendUsageData, GEditorSettingsIni);
+	}
+	
+	//Initialize all the static data from crash context
+	FGenericCrashContext::Initialize();
+	FGenericCrashContext::CopySharedCrashContext(*SharedContext);
+
+	// Set the platform specific crash context, so that we can stack walk and minidump from
+	// the crash reporter client.
+	SharedContext->PlatformCrashContext = (void*)ExceptionInfo;
+
+	// Setup up the shared memory area so that the crash report
+	SharedContext->CrashType = Type;
+	SharedContext->CrashingThreadId = CrashingThreadId;
+	SharedContext->NumStackFramesToIgnore = NumStackFramesToIgnore;
+	SharedContext->bNoDialog = bNoDialog;
+	SharedContext->bSendUnattenededBugReports = bSendUnattendedBugReports;
+	SharedContext->bSendUsageData = bSendUsageData;
+	FCString::Strcpy(SharedContext->ErrorMessage, CR_MAX_ERROR_MESSAGE_CHARS-1, ErrorMessage);
+
+	if (GLog)
+	{
+		GLog->PanicFlushThreadedLogs();
+	}
+
+	// Setup all the thread ids and names using snapshot dbghelp. Since it's not possible to 
+	// query thread names from an external process.
+	uint32 ThreadIdx = 0;
+	const bool bThreadManagerAvailable = FThreadManager::IsInitialized();
+	DWORD CurrentProcessId = GetCurrentProcessId();
+	DWORD CurrentThreadId = GetCurrentThreadId();
+	HANDLE ThreadSnapshot = CreateToolhelp32Snapshot(TH32CS_SNAPTHREAD, 0);
+	TArray<HANDLE, TInlineAllocator<CR_MAX_THREADS>> ThreadHandles;
+	if (ThreadSnapshot != INVALID_HANDLE_VALUE)
+	{
+		THREADENTRY32 ThreadEntry;
+		ThreadEntry.dwSize = sizeof(THREADENTRY32);
+		if (Thread32First(ThreadSnapshot, &ThreadEntry))
+		{
+			do
+			{
+				if (ThreadEntry.th32OwnerProcessID == CurrentProcessId)
+				{
+					// Stop the thread (except current!). We will resume once the crash reporter is done.
+					if (ThreadEntry.th32ThreadID != CurrentThreadId)
+					{
+						HANDLE ThreadHandle = OpenThread(THREAD_SUSPEND_RESUME, FALSE, ThreadEntry.th32ThreadID);
+						SuspendThread(ThreadHandle);
+						ThreadHandles.Push(ThreadHandle);
+					}
+
+					SharedContext->ThreadIds[ThreadIdx] = ThreadEntry.th32ThreadID;
+					const TCHAR* ThreadName = nullptr;
+					if (ThreadEntry.th32ThreadID == GGameThreadId)
+					{
+						ThreadName = TEXT("GameThread");
+					}
+					else if (bThreadManagerAvailable)
+					{
+						const FString& TmThreadName = FThreadManager::Get().GetThreadName(ThreadEntry.th32ThreadID);
+						ThreadName = TmThreadName.IsEmpty() ? TEXT("Unknown") : *TmThreadName;
+					}
+					else
+					{
+						ThreadName = TEXT("Unavailable");
+					}
+					FCString::Strcpy(
+						&SharedContext->ThreadNames[ThreadIdx*CR_MAX_THREAD_NAME_CHARS],
+						CR_MAX_THREAD_NAME_CHARS - 1,
+						ThreadName
+					);
+					ThreadIdx++;
+				}
+			} while (Thread32Next(ThreadSnapshot, &ThreadEntry) && (ThreadIdx < CR_MAX_THREADS));
+		}
+	}
+	SharedContext->NumThreads = ThreadIdx;
+	CloseHandle(ThreadSnapshot);
+
+	FString CrashDirectoryAbsolute;
+	if (FGenericCrashContext::CreateCrashReportDirectory(SharedContext->SessionContext.CrashGUIDRoot, SharedContext->SessionContext.GameName, ReportCrashCallCount, CrashDirectoryAbsolute))
+	{
+		FCString::Strcpy(SharedContext->CrashFilesDirectory, *CrashDirectoryAbsolute);
+		// Copy the log file to output
+		FGenericCrashContext::DumpLog(CrashDirectoryAbsolute);
+	}
+
+	// Write the shared context to the pipe
+	int32 OutDataWritten = 0;
+	FPlatformProcess::WritePipe(WritePipe, (UINT8*)SharedContext, sizeof(FSharedCrashContext), &OutDataWritten);
+	check(OutDataWritten == sizeof(FSharedCrashContext));
+
+	// Wait for a response, saying it's ok to continue
+	bool bCanContinueExecution = false;
+	int32 ExitCode = 0;
+	// Would like to use TInlineAllocator here to avoid heap allocation on crashes, but it doesn't work since ReadPipeToArray 
+	// cannot take array with non-default allocator
+	TArray<uint8> ResponseBuffer;
+	ResponseBuffer.AddZeroed(16);
+	while (!FPlatformProcess::GetProcReturnCode(CrashMonitorHandle, &ExitCode) && !bCanContinueExecution)
+	{
+		if (FPlatformProcess::ReadPipeToArray(ReadPipe, ResponseBuffer))
+		{
+			if (ResponseBuffer[0] == 0xd && ResponseBuffer[1] == 0xe &&
+				ResponseBuffer[2] == 0xa && ResponseBuffer[3] == 0xd)
+			{
+				bCanContinueExecution = true;
+			}
+		}
+	}
+
+	for (HANDLE ThreadHandle : ThreadHandles)
+	{
+		ResumeThread(ThreadHandle);
+		CloseHandle(ThreadHandle);
+	}
+
+	return EXCEPTION_CONTINUE_EXECUTION;
+}
+
 /** 
  * Create a crash report, add the user log and video, and save them a unique the crash folder
  * Launch CrashReportClient.exe to read the report and upload to our CR pipeline
@@ -367,72 +643,12 @@ int32 ReportCrashUsingCrashReportClient(FWindowsPlatformCrashContext& InContext,
 			const FString CrashContextXMLPath = FPaths::Combine(*CrashFolderAbsolute, FPlatformCrashContext::CrashContextRuntimeXMLNameW);
 			InContext.SerializeAsXML(*CrashContextXMLPath);
 
-			// Save mindump
-			if (ExceptionInfo != nullptr)
-			{
-				const FString MinidumpFileName = FPaths::Combine(*CrashFolderAbsolute, *FGenericCrashContext::UE4MinidumpName);
-				WriteMinidump(InContext, *MinidumpFileName, ExceptionInfo);
-			}
-
-			// Copy log
-			const FString LogSrcAbsolute = FPlatformOutputDevices::GetAbsoluteLogFilename();
-			FString LogFilename = FPaths::GetCleanFilename(LogSrcAbsolute);
-			const FString LogDstAbsolute = FPaths::Combine(*CrashFolderAbsolute, *LogFilename);
+			// Copy platform specific files (e.g. minidump) to output directory
+			InContext.CopyPlatformSpecificFiles(*CrashFolderAbsolute, (void*) ExceptionInfo);
 
-			// Flush out the log
-			GLog->Flush();
-
-			// If we have a memory only log, make sure it's dumped to file before we attach it to the report
-#if !NO_LOGGING
-			bool bMemoryOnly = FPlatformOutputDevices::GetLog()->IsMemoryOnly();
-			bool bBacklogEnabled = FOutputDeviceRedirector::Get()->IsBacklogEnabled();
-
-			if (bMemoryOnly || bBacklogEnabled)
-			{
-				FArchive* LogFile = IFileManager::Get().CreateFileWriter(*LogDstAbsolute, FILEWRITE_AllowRead);
-				if (LogFile)
-				{
-					if (bMemoryOnly)
-					{
-						FPlatformOutputDevices::GetLog()->Dump(*LogFile);
-					}
-					else
-					{
-						FOutputDeviceArchiveWrapper Wrapper(LogFile);
-						GLog->SerializeBacklog(&Wrapper);
-					}
-
-					LogFile->Flush();
-					delete LogFile;
-				}
-			}
-			else
-			{
-				const bool bReplace = true;
-				const bool bEvenIfReadOnly = false;
-				const bool bAttributes = false;
-				FCopyProgress* const CopyProgress = nullptr;
-				static_cast<void>(IFileManager::Get().Copy(*LogDstAbsolute, *LogSrcAbsolute, bReplace, bEvenIfReadOnly, bAttributes, CopyProgress, FILEREAD_AllowWrite, FILEWRITE_AllowRead));	// best effort, so don't care about result: couldn't copy -> tough, no log
-			}
-#endif // !NO_LOGGING
-
-			// If present, include the crash report config file to pass config values to the CRC
-			const TCHAR* CrashConfigSrcPath = FWindowsPlatformCrashContext::GetCrashConfigFilePath();
-			if (IFileManager::Get().FileExists(CrashConfigSrcPath))
-			{
-				FString CrashConfigFilename = FPaths::GetCleanFilename(CrashConfigSrcPath);
-				const FString CrashConfigDstAbsolute = FPaths::Combine(*CrashFolderAbsolute, *CrashConfigFilename);
-				static_cast<void>(IFileManager::Get().Copy(*CrashConfigDstAbsolute, CrashConfigSrcPath));	// best effort, so don't care about result: couldn't copy -> tough, no config
-			}
-
-			// If present, include the crash video
-			const FString CrashVideoPath = FPaths::ProjectLogDir() / TEXT("CrashVideo.avi");
-			if (IFileManager::Get().FileExists(*CrashVideoPath))
-			{
-				FString CrashVideoFilename = FPaths::GetCleanFilename(CrashVideoPath);
-				const FString CrashVideoDstAbsolute = FPaths::Combine(*CrashFolderAbsolute, *CrashVideoFilename);
-				static_cast<void>(IFileManager::Get().Copy(*CrashVideoDstAbsolute, *CrashVideoPath));	// best effort, so don't care about result: couldn't copy -> tough, no video
-			}
+			// Copy the log file to output
+			GLog->PanicFlushThreadedLogs();
+			FGenericCrashContext::DumpLog(CrashFolderAbsolute);
 
 			// Build machines do not upload these automatically since it is not okay to have lingering processes after the build completes.
 			if (GIsBuildMachine)
@@ -527,123 +743,6 @@ int32 ReportCrashUsingCrashReportClient(FWindowsPlatformCrashContext& InContext,
 
 } // end anonymous namespace
 
-static FCriticalSection EnsureLock;
-static bool bReentranceGuard = false;
-
-#if WINVER > 0x502	// Windows Error Reporting is not supported on Windows XP
-/**
- * A wrapper for ReportCrashUsingCrashReportClient that creates a new ensure crash context
- */
-int32 ReportEnsureUsingCrashReportClient(HANDLE Thread, EXCEPTION_POINTERS* ExceptionInfo, int IgnoreCount, const TCHAR* ErrorMessage, EErrorReportUI ReportUI)
-{
-	FWindowsPlatformCrashContext CrashContext(ECrashContextType::Ensure, ErrorMessage);
-
-	void* ContextWrapper = FWindowsPlatformStackWalk::MakeThreadContextWrapper(ExceptionInfo->ContextRecord, Thread);
-	CrashContext.CapturePortableCallStack(IgnoreCount, ContextWrapper);
-
-	return ReportCrashUsingCrashReportClient(CrashContext, ExceptionInfo, ReportUI);
-}
-#endif
-
-FORCENOINLINE void ReportEnsureInner( const TCHAR* ErrorMessage, int NumStackFramesToIgnore )
-{
-	// Skip this frame and the ::RaiseException call itself
-	NumStackFramesToIgnore += 2;
-
-	/** This is the last place to gather memory stats before exception. */
-	FGenericCrashContext::CrashMemoryStats = FPlatformMemory::GetStats();
-
-#if WINVER > 0x502	// Windows Error Reporting is not supported on Windows XP
-#if !PLATFORM_SEH_EXCEPTIONS_DISABLED
-	__try
-#endif
-	{
-		::RaiseException(1, 0, 0, nullptr);
-	}
-#if !PLATFORM_SEH_EXCEPTIONS_DISABLED
-	__except(ReportEnsureUsingCrashReportClient(GetCurrentThread(), GetExceptionInformation(), NumStackFramesToIgnore, ErrorMessage, IsInteractiveEnsureMode() ? EErrorReportUI::ShowDialog : EErrorReportUI::ReportInUnattendedMode))
-	CA_SUPPRESS(6322)
-	{
-	}
-#endif
-#endif	// WINVER
-}
-
-FORCENOINLINE void ReportAssert(const TCHAR* ErrorMessage, int NumStackFramesToIgnore)
-{
-	/** This is the last place to gather memory stats before exception. */
-	FGenericCrashContext::CrashMemoryStats = FPlatformMemory::GetStats();
-
-	FAssertInfo Info(ErrorMessage, NumStackFramesToIgnore + 2); // +2 for this function and RaiseException()
-
-	ULONG_PTR Arguments[] = { (ULONG_PTR)&Info };
-	::RaiseException( AssertExceptionCode, 0, UE_ARRAY_COUNT(Arguments), Arguments );
-}
-
-FORCENOINLINE void ReportGPUCrash(const TCHAR* ErrorMessage, int NumStackFramesToIgnore)
-{
-	/** This is the last place to gather memory stats before exception. */
-	FGenericCrashContext::CrashMemoryStats = FPlatformMemory::GetStats();
-
-	FAssertInfo Info(ErrorMessage, NumStackFramesToIgnore + 2); // +2 for this function and RaiseException()
-
-	ULONG_PTR Arguments[] = { (ULONG_PTR)&Info };
-	::RaiseException( GPUCrashExceptionCode, 0, UE_ARRAY_COUNT(Arguments), Arguments );
-}
-
-void ReportHang(const TCHAR* ErrorMessage, const uint64* StackFrames, int32 NumStackFrames, uint32 HungThreadId)
-{
-	if (ReportCrashCallCount > 0 || FDebug::HasAsserted())
-	{
-		// Don't report ensures after we've crashed/asserted, they simply may be a result of the crash as
-		// the engine is already in a bad state.
-		return;
-	}
-
-	FWindowsPlatformCrashContext CrashContext(ECrashContextType::Hang, ErrorMessage);
-	CrashContext.SetPortableCallStack(StackFrames, NumStackFrames);
-	CrashContext.SetCrashedThreadId(HungThreadId);
-	CrashContext.CaptureAllThreadContexts();
-
-	EErrorReportUI ReportUI = IsInteractiveEnsureMode() ? EErrorReportUI::ShowDialog : EErrorReportUI::ReportInUnattendedMode;
-	ReportCrashUsingCrashReportClient(CrashContext, nullptr, ReportUI);
-}
-
-// #CrashReport: 2015-05-28 This should be named EngineEnsureHandler
-/** 
- * Report an ensure to the crash reporting system
- */
-FORCENOINLINE void ReportEnsure( const TCHAR* ErrorMessage, int NumStackFramesToIgnore )
-{
-	if (ReportCrashCallCount > 0 || FDebug::HasAsserted())
-	{
-		// Don't report ensures after we've crashed/asserted, they simply may be a result of the crash as
-		// the engine is already in a bad state.
-		return;
-	}
-
-	// Simple re-entrance guard.
-	EnsureLock.Lock();
-
-	if( bReentranceGuard )
-	{
-		EnsureLock.Unlock();
-		return;
-	}
-
-	// Stop checking heartbeat for this thread (and stop the gamethread hitch detector if we're the game thread).
-	// Ensure can take a lot of time (when stackwalking), so we don't want hitches/hangs firing.
-	// These are no-ops on threads that didn't already have a heartbeat etc.
-	FSlowHeartBeatScope SuspendHeartBeat(true);
-	FDisableHitchDetectorScope SuspendGameThreadHitch;
-
-	bReentranceGuard = true;
-
-	ReportEnsureInner(ErrorMessage, NumStackFramesToIgnore + 1);
-
-	bReentranceGuard = false;
-	EnsureLock.Unlock();
-}
 
 #include "Windows/HideWindowsPlatformTypes.h"
 
@@ -700,6 +799,9 @@ void CreateExceptionInfoString(EXCEPTION_RECORD* ExceptionRecord)
 #undef HANDLE_CASE
 }
 
+
+
+
 /** 
  * Crash reporting thread. 
  * We process all the crashes on a separate thread in case the original thread's stack is corrupted (stack overflow etc).
@@ -707,21 +809,33 @@ void CreateExceptionInfoString(EXCEPTION_RECORD* ExceptionRecord)
  **/
 class FCrashReportingThread
 {
-	/** Thread Id */
+	/** Thread Id of reporter thread*/
 	DWORD ThreadId;
-	/** Thread handle */
+	/** Thread handle to reporter thread */
 	HANDLE Thread;
-
 	/** Stops this thread */
 	FThreadSafeCounter StopTaskCounter;
 	/** Signals that the game has crashed */
 	HANDLE CrashEvent;
+	/** Event that signals the crash reporting thread has finished processing the crash */
+	HANDLE CrashHandledEvent;
+
 	/** Exception information */
 	LPEXCEPTION_POINTERS ExceptionInfo;
+	/** ThreadId of the crashed thread */
 	DWORD CrashingThreadId;
+	/** Handle to crashed thread.*/
 	HANDLE CrashingThreadHandle;
-	/** Event that signals the crash reporting thread has finished processing the crash */
-	HANDLE CrashHandledEvent;
+
+	/** Process handle to crash reporter client */
+	FProcHandle CrashClientHandle;
+	/** Pipe for writing to the monitor process. */
+	void* CrashMonitorWritePipe;
+	/** Pipe for reading from the monitor process. */
+	void* CrashMonitorReadPipe;
+	/** Memory allocated for crash context. */
+	FSharedCrashContext SharedContext;
+	
 
 	/** Thread main proc */
 	static DWORD STDCALL CrashReportingThreadProc(LPVOID pThis)
@@ -733,6 +847,14 @@ class FCrashReportingThread
 	/** Main loop that waits for a crash to trigger the report generation */
 	FORCENOINLINE uint32 Run()
 	{
+		// Removed vectored exception handler, we are now guaranteed to
+		// be able to catch unhandled exception in the main try/catch block.
+#if _WIN32_WINNT >= 0x0500
+		if (!FPlatformMisc::IsDebuggerPresent())
+		{
+			RemoveVectoredExceptionHandler(UnhandledStaticInitException);
+		}
+#endif
 		while (StopTaskCounter.GetValue() == 0)
 		{
 			if (WaitForSingleObject(CrashEvent, 500) == WAIT_OBJECT_0)
@@ -757,21 +879,42 @@ class FCrashReportingThread
 public:
 		
 	FCrashReportingThread()
-		: Thread(nullptr)
+		: ThreadId(0)
+		, Thread(nullptr)
 		, CrashEvent(nullptr)
-		, ExceptionInfo(nullptr)
 		, CrashHandledEvent(nullptr)
+		, ExceptionInfo(nullptr)
+		, CrashingThreadId(0)
+		, CrashingThreadHandle(nullptr)
+		, CrashMonitorWritePipe(nullptr)
+		, CrashMonitorReadPipe(nullptr)
 	{
 		// Synchronization objects
 		CrashEvent = CreateEvent(nullptr, true, 0, nullptr);
 		CrashHandledEvent = CreateEvent(nullptr, true, 0, nullptr);
 
+		// Add an exception handler to catch issues during static initialization. This
+		// is removed once the crash reporter thread is started.
+#if _WIN32_WINNT >= 0x0500
+		if (!FPlatformMisc::IsDebuggerPresent())
+		{
+			AddVectoredExceptionHandler(1, UnhandledStaticInitException);
+		}
+#endif
+
+#if USE_CRASH_REPORTER_MONITOR
+		CrashClientHandle = LaunchCrashReportClient(&CrashMonitorWritePipe, &CrashMonitorReadPipe);
+		FMemory::Memzero(SharedContext);
+#endif
+
 		// Create a background thread that will process the crash and generate crash reports
 		Thread = CreateThread(NULL, 0, CrashReportingThreadProc, this, 0, &ThreadId);
 		if (Thread)
 		{
 			SetThreadPriority(Thread, THREAD_PRIORITY_BELOW_NORMAL);
 		}
+
+
 	}
 
 	FORCENOINLINE ~FCrashReportingThread()
@@ -793,6 +936,37 @@ public:
 
 		CloseHandle(CrashHandledEvent);
 		CrashHandledEvent = nullptr;
+
+	}
+
+	/** Ensures are passed trough this. */
+	FORCEINLINE int32 OnEnsure(LPEXCEPTION_POINTERS InExceptionInfo, int NumStackFramesToIgnore, const TCHAR* ErrorMessage, EErrorReportUI ReportUI)
+	{
+		if (CrashClientHandle.IsValid())
+		{
+			return ReportCrashForMonitor(
+				InExceptionInfo, 
+				ECrashContextType::Ensure, 
+				ErrorMessage, 
+				NumStackFramesToIgnore, 
+				GetCurrentThread(), 
+				GetCurrentThreadId(), 
+				CrashClientHandle, 
+				&SharedContext, 
+				CrashMonitorWritePipe, 
+				CrashMonitorReadPipe,
+				ReportUI
+			);
+		}
+		else
+		{
+			FWindowsPlatformCrashContext CrashContext(ECrashContextType::Ensure, ErrorMessage);
+
+			void* ContextWrapper = FWindowsPlatformStackWalk::MakeThreadContextWrapper(InExceptionInfo->ContextRecord, GetCurrentThread());
+			CrashContext.CapturePortableCallStack(NumStackFramesToIgnore, ContextWrapper);
+
+			return ReportCrashUsingCrashReportClient(CrashContext, InExceptionInfo, ReportUI);
+		}
 	}
 
 	/** The thread that crashed calls this function which will trigger the CR thread to report the crash */
@@ -811,6 +985,33 @@ public:
 		return WaitForSingleObject(CrashHandledEvent, 60000) == WAIT_OBJECT_0;
 	}
 
+	/** Crashes during static init should be reported directly to crash monitor. */
+	FORCEINLINE int32 OnCrashDuringStaticInit(LPEXCEPTION_POINTERS InExceptionInfo)
+	{
+		if (CrashClientHandle.IsValid())
+		{
+			const ECrashContextType Type = ECrashContextType::Crash;
+			const int NumStackFramesToIgnore = 1;
+			const TCHAR* ErrorMessage = TEXT("Crash during static initialization");
+
+			return ReportCrashForMonitor(
+				InExceptionInfo,
+				Type,
+				ErrorMessage,
+				NumStackFramesToIgnore,
+				CrashingThreadHandle,
+				CrashingThreadId,
+				CrashClientHandle,
+				&SharedContext,
+				CrashMonitorWritePipe,
+				CrashMonitorReadPipe,
+				EErrorReportUI::ShowDialog
+			);
+		}
+
+		return EXCEPTION_CONTINUE_EXECUTION;
+	}
+
 private:
 
 	/** Handles the crash */
@@ -824,10 +1025,17 @@ private:
 		// Then try run time crash processing and broadcast information about a crash.
 		FCoreDelegates::OnHandleSystemError.Broadcast();
 
+		if (GLog)
+		{
+			GLog->PanicFlushThreadedLogs();
+		}
+		
 		// Get the default settings for the crash context
 		ECrashContextType Type = ECrashContextType::Crash;
 		const TCHAR* ErrorMessage = TEXT("Unhandled exception");
-		int NumStackFramesToIgnore = 0;
+		int NumStackFramesToIgnore = 2;
+
+		void* ContextWrapper = nullptr;
 
 		// If it was an assert or GPU crash, allow overriding the info from the exception parameters
 		if (ExceptionInfo->ExceptionRecord->ExceptionCode == AssertExceptionCode && ExceptionInfo->ExceptionRecord->NumberParameters == 1)
@@ -835,40 +1043,63 @@ private:
 			const FAssertInfo& Info = *(const FAssertInfo*)ExceptionInfo->ExceptionRecord->ExceptionInformation[0];
 			Type = ECrashContextType::Assert;
 			ErrorMessage = Info.ErrorMessage;
-			NumStackFramesToIgnore = Info.NumStackFramesToIgnore;
+			NumStackFramesToIgnore += Info.NumStackFramesToIgnore;
 		}
 		else if (ExceptionInfo->ExceptionRecord->ExceptionCode == GPUCrashExceptionCode && ExceptionInfo->ExceptionRecord->NumberParameters == 1)
 		{
 			const FAssertInfo& Info = *(const FAssertInfo*)ExceptionInfo->ExceptionRecord->ExceptionInformation[0];
 			Type = ECrashContextType::GPUCrash;
 			ErrorMessage = Info.ErrorMessage;
-			NumStackFramesToIgnore = Info.NumStackFramesToIgnore;
+			NumStackFramesToIgnore += Info.NumStackFramesToIgnore;
 		}
 
-		// Not super safe due to dynamic memory allocations, but at least enables new functionality.
-		// Introduces a new runtime crash context. Will replace all Windows related crash reporting.
-		FWindowsPlatformCrashContext CrashContext(Type, ErrorMessage);
+#if USE_CRASH_REPORTER_MONITOR
+		if (CrashClientHandle.IsValid())
+		{
+			// If possible use the crash monitor helper class to report the error. This will do most of the analysis
+			// in the crash reporter client process.
+			ReportCrashForMonitor(
+				ExceptionInfo,
+				Type,
+				ErrorMessage,
+				NumStackFramesToIgnore,
+				CrashingThreadHandle,
+				CrashingThreadId,
+				CrashClientHandle,
+				&SharedContext,
+				CrashMonitorWritePipe,
+				CrashMonitorReadPipe,
+				EErrorReportUI::ShowDialog
+			);
+		}
+		else
+#endif
+		{
+			// Not super safe due to dynamic memory allocations, but at least enables new functionality.
+			// Introduces a new runtime crash context. Will replace all Windows related crash reporting.
+			FWindowsPlatformCrashContext CrashContext(Type, ErrorMessage);
 
-		// Thread context wrapper for stack operations
-		void* ContextWrapper = FWindowsPlatformStackWalk::MakeThreadContextWrapper(ExceptionInfo->ContextRecord, CrashingThreadHandle);
-		CrashContext.CapturePortableCallStack(NumStackFramesToIgnore, ContextWrapper);
-		CrashContext.SetCrashedThreadId(CrashingThreadId);
-		CrashContext.CaptureAllThreadContexts();
+			// Thread context wrapper for stack operations
+			ContextWrapper = FWindowsPlatformStackWalk::MakeThreadContextWrapper(ExceptionInfo->ContextRecord, CrashingThreadHandle);
+			CrashContext.CapturePortableCallStack(NumStackFramesToIgnore, ContextWrapper);
+			CrashContext.SetCrashedThreadId(CrashingThreadId);
+			CrashContext.CaptureAllThreadContexts();
 
-		// Also mark the same number of frames to be ignored if we symbolicate from the minidump
-		CrashContext.SetNumMinidumpFramesToIgnore(NumStackFramesToIgnore);
+			// Also mark the same number of frames to be ignored if we symbolicate from the minidump
+			CrashContext.SetNumMinidumpFramesToIgnore(NumStackFramesToIgnore);
 
-		// First launch the crash reporter client.
+			// First launch the crash reporter client.
 #if WINVER > 0x502	// Windows Error Reporting is not supported on Windows XP
-		if (GUseCrashReportClient)
-		{
-			ReportCrashUsingCrashReportClient(CrashContext, ExceptionInfo, EErrorReportUI::ShowDialog);
-		}
-		else
+			if (GUseCrashReportClient)
+			{
+				ReportCrashUsingCrashReportClient(CrashContext, ExceptionInfo, EErrorReportUI::ShowDialog);
+			}
+			else
 #endif		// WINVER
-		{
-			CrashContext.SerializeContentToBuffer();
-			WriteMinidump(CrashContext, MiniDumpFilenameW, ExceptionInfo);
+			{
+				CrashContext.SerializeContentToBuffer();
+				WriteMinidump(GetCurrentProcess(), GetCurrentThreadId(), CrashContext, MiniDumpFilenameW, ExceptionInfo);
+			}
 		}
 
 		const bool bGenerateRuntimeCallstack =
@@ -919,6 +1150,40 @@ private:
 
 TUniquePtr<FCrashReportingThread> GCrashReportingThread = MakeUnique<FCrashReportingThread>();
 
+
+LONG WINAPI UnhandledStaticInitException(LPEXCEPTION_POINTERS ExceptionInfo)
+{
+	// Top bit in exception code is fatal exceptions. Report those but not other types.
+	if ((ExceptionInfo->ExceptionRecord->ExceptionCode & 0x80000000L) != 0)
+	{
+		// If we get an exception during static init we hope that the crash reporting thread
+		// object has been created, otherwise we cannot handle the exception. This will hopefully 
+		// work even if there is a stack overflow. See 
+		// https://peteronprogramming.wordpress.com/2016/08/10/crashes-you-cant-handle-easily-2-stack-overflows-on-windows/
+		// @note: Even if the object has been created, the actual thread has not been started yet,
+		// (that happens after static init) so we must bypass that and report directly from this thread. 
+		// 
+		if (GCrashReportingThread.IsValid())
+		{
+			return GCrashReportingThread->OnCrashDuringStaticInit(ExceptionInfo);
+		}
+	}
+	
+	return EXCEPTION_CONTINUE_SEARCH;
+}
+
+/**
+ * Fallback for catching exceptions which aren't caught elsewhere. This allows catching exceptions on threads created outside the engine.
+ * Note that Windows does not call this handler if a debugger is attached, separately to our internal logic around crash handling.
+ */
+LONG WINAPI UnhandledException(EXCEPTION_POINTERS *ExceptionInfo)
+{
+	ReportCrash(ExceptionInfo);
+	GIsCriticalError = true;
+	FPlatformMisc::RequestExit(true);
+	return EXCEPTION_CONTINUE_SEARCH;
+}
+
 // #CrashReport: 2015-05-28 This should be named EngineCrashHandler
 int32 ReportCrash( LPEXCEPTION_POINTERS ExceptionInfo )
 {
@@ -937,3 +1202,151 @@ int32 ReportCrash( LPEXCEPTION_POINTERS ExceptionInfo )
 
 	return EXCEPTION_EXECUTE_HANDLER;
 }
+
+static FCriticalSection EnsureLock;
+static bool bReentranceGuard = false;
+
+#if WINVER > 0x502	// Windows Error Reporting is not supported on Windows XP
+/**
+ * A wrapper for ReportCrashUsingCrashReportClient that creates a new ensure crash context
+ */
+int32 ReportEnsureUsingCrashReportClient(EXCEPTION_POINTERS* ExceptionInfo, int NumStackFramesToIgnore, const TCHAR* ErrorMessage, EErrorReportUI ReportUI)
+{
+	return GCrashReportingThread->OnEnsure(ExceptionInfo, NumStackFramesToIgnore, ErrorMessage, ReportUI);
+}
+#endif
+
+FORCENOINLINE void ReportEnsureInner(const TCHAR* ErrorMessage, int NumStackFramesToIgnore)
+{
+	// Skip this frame and the ::RaiseException call itself
+	NumStackFramesToIgnore += 2;
+
+	/** This is the last place to gather memory stats before exception. */
+	FGenericCrashContext::SetMemoryStats(FPlatformMemory::GetStats());
+
+#if WINVER > 0x502	// Windows Error Reporting is not supported on Windows XP
+#if !PLATFORM_SEH_EXCEPTIONS_DISABLED
+	__try
+#endif
+	{
+		::RaiseException(1, 0, 0, nullptr);
+	}
+#if !PLATFORM_SEH_EXCEPTIONS_DISABLED
+	__except (ReportEnsureUsingCrashReportClient( GetExceptionInformation(), NumStackFramesToIgnore, ErrorMessage, IsInteractiveEnsureMode() ? EErrorReportUI::ShowDialog : EErrorReportUI::ReportInUnattendedMode))
+		CA_SUPPRESS(6322)
+	{
+	}
+#endif
+#endif	// WINVER
+}
+
+FORCENOINLINE void ReportAssert(const TCHAR* ErrorMessage, int NumStackFramesToIgnore)
+{
+	/** This is the last place to gather memory stats before exception. */
+	FGenericCrashContext::SetMemoryStats(FPlatformMemory::GetStats());
+
+	FAssertInfo Info(ErrorMessage, NumStackFramesToIgnore + 2); // +2 for this function and RaiseException()
+
+	ULONG_PTR Arguments[] = { (ULONG_PTR)&Info };
+	::RaiseException(AssertExceptionCode, 0, UE_ARRAY_COUNT(Arguments), Arguments);
+}
+
+FORCENOINLINE void ReportGPUCrash(const TCHAR* ErrorMessage, int NumStackFramesToIgnore)
+{
+	/** This is the last place to gather memory stats before exception. */
+	FGenericCrashContext::SetMemoryStats(FPlatformMemory::GetStats());
+
+	FAssertInfo Info(ErrorMessage, NumStackFramesToIgnore + 2); // +2 for this function and RaiseException()
+
+	ULONG_PTR Arguments[] = { (ULONG_PTR)&Info };
+	::RaiseException(GPUCrashExceptionCode, 0, UE_ARRAY_COUNT(Arguments), Arguments);
+}
+
+void ReportHang(const TCHAR* ErrorMessage, const uint64* StackFrames, int32 NumStackFrames, uint32 HungThreadId)
+{
+	if (ReportCrashCallCount > 0 || FDebug::HasAsserted())
+	{
+		// Don't report ensures after we've crashed/asserted, they simply may be a result of the crash as
+		// the engine is already in a bad state.
+		return;
+	}
+
+	FWindowsPlatformCrashContext CrashContext(ECrashContextType::Hang, ErrorMessage);
+	CrashContext.SetPortableCallStack(StackFrames, NumStackFrames);
+	CrashContext.SetCrashedThreadId(HungThreadId);
+	CrashContext.CaptureAllThreadContexts();
+
+	EErrorReportUI ReportUI = IsInteractiveEnsureMode() ? EErrorReportUI::ShowDialog : EErrorReportUI::ReportInUnattendedMode;
+	ReportCrashUsingCrashReportClient(CrashContext, nullptr, ReportUI);
+}
+
+// #CrashReport: 2015-05-28 This should be named EngineEnsureHandler
+/**
+ * Report an ensure to the crash reporting system
+ */
+FORCENOINLINE void ReportEnsure(const TCHAR* ErrorMessage, int NumStackFramesToIgnore)
+{
+	if (ReportCrashCallCount > 0 || FDebug::HasAsserted())
+	{
+		// Don't report ensures after we've crashed/asserted, they simply may be a result of the crash as
+		// the engine is already in a bad state.
+		return;
+	}
+
+	// Simple re-entrance guard.
+	EnsureLock.Lock();
+
+	if (bReentranceGuard)
+	{
+		EnsureLock.Unlock();
+		return;
+	}
+
+	// Stop checking heartbeat for this thread (and stop the gamethread hitch detector if we're the game thread).
+	// Ensure can take a lot of time (when stackwalking), so we don't want hitches/hangs firing.
+	// These are no-ops on threads that didn't already have a heartbeat etc.
+	FSlowHeartBeatScope SuspendHeartBeat(true);
+	FDisableHitchDetectorScope SuspendGameThreadHitch;
+
+	bReentranceGuard = true;
+
+	ReportEnsureInner(ErrorMessage, NumStackFramesToIgnore + 1);
+
+	bReentranceGuard = false;
+	EnsureLock.Unlock();
+}
+
+#if !IS_PROGRAM && 0
+namespace {
+	/** Utility class to test crashes during static initialization. */
+	struct StaticInitCrasher
+	{
+		StaticInitCrasher()
+		{
+			// Test stack overflow
+			//StackOverlowMe(0);
+
+			// Test GPF
+			//int* i = nullptr;
+			//*i = 3;
+
+			// Check assert (shouldn't work during static init)
+			//check(false);
+		}
+
+		void StackOverlowMe(uint32 s) {
+			uint32 buffer[1024];
+			memset(&buffer, 0xdeadbeef, 1024);
+			(void*)buffer;
+			if (s == 0xffffffff)
+				return;
+			StackOverlowMe(s + 1);
+		}
+	};
+
+	static StaticInitCrasher GCrasher;
+}
+#endif
+
+
+
diff --git a/Engine/Source/Runtime/Core/Private/Windows/WindowsPlatformMisc.cpp b/Engine/Source/Runtime/Core/Private/Windows/WindowsPlatformMisc.cpp
--- a/Engine/Source/Runtime/Core/Private/Windows/WindowsPlatformMisc.cpp
+++ b/Engine/Source/Runtime/Core/Private/Windows/WindowsPlatformMisc.cpp
@@ -1624,7 +1624,7 @@ int32 FWindowsPlatformMisc::NumberOfCores()
 	static int32 CoreCount = 0;
 	if (CoreCount == 0)
 	{
-		if (FParse::Param(FCommandLine::Get(), TEXT("usehyperthreading")))
+		if (FCommandLine::IsInitialized() && FParse::Param(FCommandLine::Get(), TEXT("usehyperthreading")))
 		{
 			CoreCount = NumberOfCoresIncludingHyperthreads();
 		}
@@ -1768,7 +1768,7 @@ HWND FWindowsPlatformMisc::GetTopLevelWindowHandle(uint32 ProcessId)
 FORCENOINLINE void FWindowsPlatformMisc::RaiseException( uint32 ExceptionCode )
 {
 	/** This is the last place to gather memory stats before exception. */
-	FGenericCrashContext::CrashMemoryStats = FPlatformMemory::GetStats();
+	FGenericCrashContext::SetMemoryStats(FPlatformMemory::GetStats());
 
 	::RaiseException( ExceptionCode, 0, 0, NULL );
 }
diff --git a/Engine/Source/Runtime/Core/Private/Windows/WindowsPlatformProcess.cpp b/Engine/Source/Runtime/Core/Private/Windows/WindowsPlatformProcess.cpp
--- a/Engine/Source/Runtime/Core/Private/Windows/WindowsPlatformProcess.cpp
+++ b/Engine/Source/Runtime/Core/Private/Windows/WindowsPlatformProcess.cpp
@@ -642,6 +642,15 @@ bool FWindowsPlatformProcess::IsApplicationRunning( const TCHAR* ProcName )
 	return false;
 }
 
+bool FWindowsPlatformProcess::IsApplicationAlive(uint32 ProcessId)
+{
+	DWORD ExitCode;
+	HANDLE ProcessHandle = ::OpenProcess(PROCESS_QUERY_INFORMATION, false, ProcessId);
+	const BOOL result = GetExitCodeProcess(ProcessHandle, &ExitCode);
+	CloseHandle(ProcessHandle);
+	return result && ExitCode == STILL_ACTIVE;
+}
+
 FString FWindowsPlatformProcess::GetApplicationName( uint32 ProcessId )
 {
 	FString Output = TEXT("");
@@ -666,6 +675,7 @@ FString FWindowsPlatformProcess::GetApplicationName( uint32 ProcessId )
 			Output = ProcessNameBuffer;
 		}
 	}
+	CloseHandle(ProcessHandle);
 
 	return Output;
 }
@@ -1444,11 +1454,17 @@ bool FWindowsPlatformProcess::WritePipe(void* WritePipe, const uint8* Data, cons
 #include "Windows/AllowWindowsPlatformTypes.h"
 
 FWindowsPlatformProcess::FWindowsSemaphore::FWindowsSemaphore(const FString & InName, HANDLE InSemaphore)
-	:	FSemaphore(InName)
-	,	Semaphore(InSemaphore)
+	: FWindowsSemaphore(*InName, InSemaphore)
 {
 }
 
+FWindowsPlatformProcess::FWindowsSemaphore::FWindowsSemaphore(const TCHAR* InName, Windows::HANDLE InSemaphore)
+	: FSemaphore(InName)
+	, Semaphore(InSemaphore)
+{
+
+}
+
 FWindowsPlatformProcess::FWindowsSemaphore::~FWindowsSemaphore()
 {
 	// actual cleanup should be done in DeleteInterprocessSynchObject() since it can return errors
@@ -1497,19 +1513,24 @@ void FWindowsPlatformProcess::FWindowsSemaphore::Unlock()
 	}
 }
 
-FWindowsPlatformProcess::FSemaphore * FWindowsPlatformProcess::NewInterprocessSynchObject(const FString & Name, bool bCreate, uint32 MaxLocks)
+FWindowsPlatformProcess::FSemaphore* FWindowsPlatformProcess::NewInterprocessSynchObject(const FString & Name, bool bCreate, uint32 MaxLocks)
+{
+	return NewInterprocessSynchObject(*Name, bCreate, MaxLocks);
+}
+
+FWindowsPlatformProcess::FSemaphore* FWindowsPlatformProcess::NewInterprocessSynchObject(const TCHAR* Name, bool bCreate, uint32 MaxLocks)
 {
 	HANDLE Semaphore = NULL;
 	
 	if (bCreate)
 	{
-		Semaphore = CreateSemaphore(NULL, MaxLocks, MaxLocks, *Name);
+		Semaphore = CreateSemaphore(NULL, MaxLocks, MaxLocks, Name);
 		if (NULL == Semaphore)
 		{
 			DWORD ErrNo = GetLastError();
 			UE_LOG(LogHAL, Warning, TEXT("CreateSemaphore(Attrs=NULL, InitialValue=%d, MaxValue=%d, Name='%s') failed with LastError = %d"),
 				MaxLocks, MaxLocks,
-				*Name,
+				Name,
 				ErrNo);
 			return NULL;
 		}
@@ -1517,13 +1538,13 @@ FWindowsPlatformProcess::FSemaphore * FWindowsPlatformProcess::NewInterprocessSy
 	else
 	{
 		DWORD AccessRights = SYNCHRONIZE | SEMAPHORE_MODIFY_STATE;
-		Semaphore = OpenSemaphore(AccessRights, false, *Name);
+		Semaphore = OpenSemaphore(AccessRights, false, Name);
 		if (NULL == Semaphore)
 		{
 			DWORD ErrNo = GetLastError();
 			UE_LOG(LogHAL, Warning, TEXT("OpenSemaphore(AccessRights=0x%08x, bInherit=false, Name='%s') failed with LastError = %d"),
 				AccessRights,
-				*Name,
+				Name,
 				ErrNo);
 			return NULL;
 		}
diff --git a/Engine/Source/Runtime/Core/Private/Windows/WindowsPlatformStackWalk.cpp b/Engine/Source/Runtime/Core/Private/Windows/WindowsPlatformStackWalk.cpp
--- a/Engine/Source/Runtime/Core/Private/Windows/WindowsPlatformStackWalk.cpp
+++ b/Engine/Source/Runtime/Core/Private/Windows/WindowsPlatformStackWalk.cpp
@@ -36,6 +36,7 @@
 /** Whether appInitStackWalking() has been called successfully or not. */
 static bool GStackWalkingInitialized = false;
 static bool GNeedToRefreshSymbols = false;
+static HANDLE GProcessHandle = INVALID_HANDLE_VALUE;
 
 static const TCHAR* CrashReporterSettings = TEXT("/Script/UnrealEd.CrashReporterSettings");
 
@@ -45,6 +46,21 @@ static const TCHAR* CrashReporterSettings = TEXT("/Script/UnrealEd.CrashReporter
 // Uses StackWalk64 interface which is more reliable, but 500-1000x slower than the fast stacktrace.
 #define USE_SLOW_STACKTRACE 0
 
+// Use _NT_SYMBOL_PATH for non development builds. We don't want shipping crash reporter to try to 
+// access build servers for example.
+#if !UE_BUILD_SHIPPING
+#define ALLOW_UNREAL_ACCESS_TO_NT_SYMBOL_PATH 1
+#else
+#define ALLOW_UNREAL_ACCESS_TO_NT_SYMBOL_PATH 0
+#endif
+
+// Optimization that only loads symbols on demand rather than always loading all symbols for loaded modules.
+#if 1
+#define ON_DEMAND_SYMBOL_LOADING 1
+#else
+#define ON_DEMAND_SYMBOL_LOADING 0
+#endif
+
 typedef bool  (WINAPI *TFEnumProcesses)( uint32* lpidProcess, uint32 cb, uint32* cbNeeded);
 typedef bool  (WINAPI *TFEnumProcessModules)(HANDLE hProcess, HMODULE *lphModule, uint32 cb, LPDWORD lpcbNeeded);
 #if WINVER > 0x502
@@ -62,6 +78,13 @@ static TFGetModuleBaseName		FGetModuleBaseName;
 static TFGetModuleFileNameEx	FGetModuleFileNameEx;
 static TFGetModuleInformation	FGetModuleInformation;
 
+FString GetRemoteStorage(const FString& DownstreamStorage);
+FString GetSymbolSearchPath();
+HMODULE* GetProcessModules(HANDLE ProcessHandle);
+void LoadSymbolsForModule(HMODULE ModuleHandle, const FString& RemoteStorage);
+void LoadSymbolsForProcessModules(const FString &RemoteStorage);
+void LoadSymbolsForModuleByAddress(uint64 Address, const FString& RemoteStorage);
+
 struct FWindowsThreadContextWrapper
 {
 	enum 
@@ -114,7 +137,7 @@ static int32 CaptureStackTraceHelper(uint64 *BackTrace, uint32 MaxDepth, FWindow
 #endif
 	{
 		// Get context, process and thread information.
-		ProcessHandle	= GetCurrentProcess();
+		ProcessHandle	= GProcessHandle;
 
 		// Zero out stack frame.
 		FMemory::Memzero( StackFrame64 );
@@ -185,7 +208,8 @@ static int32 CaptureStackTraceHelper(uint64 *BackTrace, uint32 MaxDepth, FWindow
 
 PRAGMA_DISABLE_OPTIMIZATION // Work around "flow in or out of inline asm code suppresses global optimization" warning C4740.
 
-static int32 CaptureStackTraceHelper(uint64 *BackTrace, uint32 MaxDepth, CONTEXT* Context)
+
+int32 CaptureStackTraceHelper(uint64* BackTrace, uint32 MaxDepth, CONTEXT* Context)
 {
 	FWindowsThreadContextWrapper HelperContext;
 	HelperContext.ThreadHandle = GetCurrentThread();
@@ -404,11 +428,16 @@ void FWindowsPlatformStackWalk::ProgramCounterToSymbolInfo( uint64 ProgramCounte
 	// Initialize stack walking as it loads up symbol information which we require.
 	InitStackWalking();
 
+#if ON_DEMAND_SYMBOL_LOADING
+	// Load symbols for the module
+	LoadSymbolsForModuleByAddress(ProgramCounter, GetSymbolSearchPath());
+#endif
+
 	// Set the program counter.
 	out_SymbolInfo.ProgramCounter = ProgramCounter;
 
 	uint32 LastError = 0;
-	HANDLE ProcessHandle = GetCurrentProcess();
+	HANDLE ProcessHandle = GProcessHandle;
 
 	// Initialize symbol.
 	ANSICHAR SymbolBuffer[sizeof( IMAGEHLP_SYMBOL64 ) + FProgramCounterSymbolInfo::MAX_NAME_LENGTH] = {0};
@@ -469,7 +498,7 @@ void FWindowsPlatformStackWalk::ProgramCounterToSymbolInfo( uint64 ProgramCounte
  *
  * IMPORTANT: Returned value must be deallocated by FMemory::Free().
  */
-static HMODULE* GetProcessModules(HANDLE ProcessHandle)
+HMODULE* GetProcessModules(HANDLE ProcessHandle)
 {
 	const int32 NumModules = FWindowsPlatformStackWalk::GetProcessModuleCount();
 	// Allocate start size (last element reserved for NULL value)
@@ -524,7 +553,7 @@ bool FWindowsPlatformStackWalk::UploadLocalSymbols()
 	SymbolStorage = TEXT( "SRV*" ) + SymbolStorage;
 
 	int32 ErrorCode = 0;
-	HANDLE ProcessHandle = GetCurrentProcess();
+	HANDLE ProcessHandle = GProcessHandle;
 
 	// Enumerate process modules.
 	HMODULE* ModuleHandlePointer = GetProcessModules( ProcessHandle );
@@ -595,114 +624,135 @@ bool FWindowsPlatformStackWalk::UploadLocalSymbols()
 	return true;
 }
 
-/**
- * Loads modules for current process.
- */ 
-static void LoadProcessModules(const FString &RemoteStorage)
+void LoadSymbolsForModule(HMODULE ModuleHandle, const FString& RemoteStorage)
 {
-	int32 ErrorCode = 0;
-	HANDLE ProcessHandle = GetCurrentProcess();
+	HANDLE ProcessHandle = GProcessHandle;
 
-	// Enumerate process modules.
-	HMODULE* ModuleHandlePointer = GetProcessModules(ProcessHandle);
-	if (!ModuleHandlePointer)
-	{
-		ErrorCode = GetLastError();
-		return;
-	}
+	int32 ErrorCode = 0;
 
-	// Load the modules.
-	for( int32 ModuleIndex = 0; ModuleHandlePointer[ModuleIndex]; ModuleIndex++ )
-	{
-		MODULEINFO ModuleInfo = {0};
+	MODULEINFO ModuleInfo = { 0 };
 #if WINVER > 0x502
-		WCHAR ModuleName[FProgramCounterSymbolInfo::MAX_NAME_LENGTH] = {0};
-		WCHAR ImageName[FProgramCounterSymbolInfo::MAX_NAME_LENGTH] = {0};
+	WCHAR ModuleName[FProgramCounterSymbolInfo::MAX_NAME_LENGTH] = { 0 };
+	WCHAR ImageName[FProgramCounterSymbolInfo::MAX_NAME_LENGTH] = { 0 };
 #else
-		ANSICHAR ModuleName[FProgramCounterSymbolInfo::MAX_NAME_LENGTH] = { 0 };
-		ANSICHAR ImageName[FProgramCounterSymbolInfo::MAX_NAME_LENGTH] = { 0 };
+	ANSICHAR ModuleName[FProgramCounterSymbolInfo::MAX_NAME_LENGTH] = { 0 };
+	ANSICHAR ImageName[FProgramCounterSymbolInfo::MAX_NAME_LENGTH] = { 0 };
 #endif
 #if PLATFORM_64BITS
-		static_assert(sizeof( MODULEINFO ) == 24, "Broken alignment for 64bit Windows include.");
+	static_assert(sizeof(MODULEINFO) == 24, "Broken alignment for 64bit Windows include.");
 #else
-		static_assert(sizeof( MODULEINFO ) == 12, "Broken alignment for 32bit Windows include.");
+	static_assert(sizeof(MODULEINFO) == 12, "Broken alignment for 32bit Windows include.");
 #endif
-		FGetModuleInformation( ProcessHandle, ModuleHandlePointer[ModuleIndex], &ModuleInfo, sizeof( ModuleInfo ) );
-		FGetModuleFileNameEx( ProcessHandle, ModuleHandlePointer[ModuleIndex], ImageName, FProgramCounterSymbolInfo::MAX_NAME_LENGTH );
-		FGetModuleBaseName( ProcessHandle, ModuleHandlePointer[ModuleIndex], ModuleName, FProgramCounterSymbolInfo::MAX_NAME_LENGTH );
+	FGetModuleInformation(ProcessHandle, ModuleHandle, &ModuleInfo, sizeof(ModuleInfo));
+	FGetModuleFileNameEx(ProcessHandle, ModuleHandle, ImageName, FProgramCounterSymbolInfo::MAX_NAME_LENGTH);
+	FGetModuleBaseName(ProcessHandle, ModuleHandle, ModuleName, FProgramCounterSymbolInfo::MAX_NAME_LENGTH);
 
-		// Set the search path to find PDBs in the same folder as the DLL.
+	// Set the search path to find PDBs in the same folder as the DLL.
 #if WINVER > 0x502
-		WCHAR SearchPath[MAX_PATH] = {0};
-		WCHAR* FileName = NULL;
-		const auto Result = GetFullPathNameW( ImageName, MAX_PATH, SearchPath, &FileName );
+	WCHAR SearchPath[MAX_PATH] = { 0 };
+	WCHAR* FileName = NULL;
+	const auto Result = GetFullPathNameW(ImageName, MAX_PATH, SearchPath, &FileName);
 #else
-		ANSICHAR SearchPath[MAX_PATH] = { 0 };
-		ANSICHAR* FileName = NULL;
-		const auto Result = GetFullPathNameA( ImageName, MAX_PATH, SearchPath, &FileName );
+	ANSICHAR SearchPath[MAX_PATH] = { 0 };
+	ANSICHAR* FileName = NULL;
+	const auto Result = GetFullPathNameA(ImageName, MAX_PATH, SearchPath, &FileName);
 #endif
 
-		FString SearchPathList;
-		if (Result != 0 && Result < MAX_PATH)
-		{
-			*FileName = 0;
+	FString SearchPathList;
+	if (Result != 0 && Result < MAX_PATH)
+	{
+		*FileName = 0;
 #if WINVER > 0x502
-			SearchPathList = SearchPath;
+		SearchPathList = SearchPath;
 #else
-			SearchPathList = ANSI_TO_TCHAR(SearchPath);
+		SearchPathList = ANSI_TO_TCHAR(SearchPath);
 #endif
-		}
-		if (!RemoteStorage.IsEmpty())
+	}
+	if (!RemoteStorage.IsEmpty())
+	{
+		if (!SearchPathList.IsEmpty())
 		{
-			if (!SearchPathList.IsEmpty())
-			{
-				SearchPathList.AppendChar(';');
-			}
-			SearchPathList.Append(RemoteStorage);
+			SearchPathList.AppendChar(';');
 		}
+		SearchPathList.Append(RemoteStorage);
+	}
 
 #if WINVER > 0x502
-		SymSetSearchPathW(ProcessHandle, *SearchPathList);
+	SymSetSearchPathW(ProcessHandle, *SearchPathList);
 
-		// Load module.
-		const DWORD64 BaseAddress = SymLoadModuleExW( ProcessHandle, ModuleHandlePointer[ModuleIndex], ImageName, ModuleName, (DWORD64) ModuleInfo.lpBaseOfDll, (uint32) ModuleInfo.SizeOfImage, NULL, 0 );
-		if( !BaseAddress )
-		{
-			ErrorCode = GetLastError();
+	// Load module.
+	const DWORD64 BaseAddress = SymLoadModuleExW(ProcessHandle, ModuleHandle, ImageName, ModuleName, (DWORD64)ModuleInfo.lpBaseOfDll, (uint32)ModuleInfo.SizeOfImage, NULL, 0);
+	if (!BaseAddress)
+	{
+		ErrorCode = GetLastError();
 
-			// If the module is already loaded, the return value is zero and GetLastError returns ERROR_SUCCESS.
-			if ( ErrorCode != ERROR_SUCCESS )
-			{
-				UE_LOG(LogWindows, Warning, TEXT("SymLoadModuleExW. Error: %d"), ErrorCode);
-		}
+		// If the module is already loaded, the return value is zero and GetLastError returns ERROR_SUCCESS.
+		if (ErrorCode != ERROR_SUCCESS)
+		{
+			UE_LOG(LogWindows, Warning, TEXT("SymLoadModuleExW. Error: %d"), ErrorCode);
 		}
+	}
 #else
-		SymSetSearchPath(ProcessHandle, TCHAR_TO_ANSI(*SearchPathList));
+	SymSetSearchPath(ProcessHandle, TCHAR_TO_ANSI(*SearchPathList));
 
-		// Load module.
-		const DWORD64 BaseAddress = SymLoadModuleEx( ProcessHandle, ModuleHandlePointer[ModuleIndex], ImageName, ModuleName, (DWORD64)ModuleInfo.lpBaseOfDll, (uint32)ModuleInfo.SizeOfImage, NULL, 0 );
-		if (!BaseAddress)
-		{
-			ErrorCode = GetLastError();
+	// Load module.
+	const DWORD64 BaseAddress = SymLoadModuleEx(ProcessHandle, ModuleHandle, ImageName, ModuleName, (DWORD64)ModuleInfo.lpBaseOfDll, (uint32)ModuleInfo.SizeOfImage, NULL, 0);
+	if (!BaseAddress)
+	{
+		ErrorCode = GetLastError();
 
-			// If the module is already loaded, the return value is zero and GetLastError returns ERROR_SUCCESS.
-			if ( ErrorCode != ERROR_SUCCESS )
-			{
-				UE_LOG(LogWindows, Warning, TEXT("SymLoadModuleEx. Error: %d"), ErrorCode);
-			}
+		// If the module is already loaded, the return value is zero and GetLastError returns ERROR_SUCCESS.
+		if (ErrorCode != ERROR_SUCCESS)
+		{
+			UE_LOG(LogWindows, Warning, TEXT("SymLoadModuleEx. Error: %d"), ErrorCode);
 		}
+	}
 #endif
-	} 
+}
+
+
+/**
+ * Loads modules for current process.
+ */
+void LoadSymbolsForProcessModules(const FString &RemoteStorage)
+{
+	int32 ErrorCode = 0;
+	HANDLE ProcessHandle = GProcessHandle;
+
+	// Enumerate process modules.
+	HMODULE* ModuleHandlePointer = GetProcessModules(ProcessHandle);
+	if (!ModuleHandlePointer)
+	{
+		ErrorCode = GetLastError();
+		return;
+	}
+
+	// Load the modules.
+	for (int32 ModuleIndex = 0; ModuleHandlePointer[ModuleIndex]; ModuleIndex++)
+	{
+		LoadSymbolsForModule(ModuleHandlePointer[ModuleIndex], RemoteStorage);
+	}
 
 	// Free the module handle pointer allocated in case the static array was insufficient.
 	FMemory::Free(ModuleHandlePointer);
 }
 
+void LoadSymbolsForModuleByAddress(uint64 Address, const FString& RemoteStorage)
+{
+	HMODULE ModuleHandle = NULL;
+
+	if (GetModuleHandleEx(GET_MODULE_HANDLE_EX_FLAG_FROM_ADDRESS | GET_MODULE_HANDLE_EX_FLAG_UNCHANGED_REFCOUNT, (LPCTSTR)Address, &ModuleHandle))
+	{
+		LoadSymbolsForModule(ModuleHandle, RemoteStorage);
+	}
+}
+
+
 int32 FWindowsPlatformStackWalk::GetProcessModuleCount()
 {
 	FPlatformStackWalk::InitStackWalking();
 
-	HANDLE ProcessHandle = GetCurrentProcess(); 
+	HANDLE ProcessHandle = GProcessHandle; 
 	uint32 BytesRequired = 0;
 
 	// Enumerate process modules.
@@ -721,7 +771,7 @@ int32 FWindowsPlatformStackWalk::GetProcessModuleSignatures(FStackWalkModuleInfo
 {
 	FPlatformStackWalk::InitStackWalking();
 
-	HANDLE		ProcessHandle = GetCurrentProcess(); 
+	HANDLE		ProcessHandle = GProcessHandle; 
 
 	// Enumerate process modules.
 	HMODULE* ModuleHandlePointer = GetProcessModules(ProcessHandle);
@@ -805,6 +855,9 @@ static void OnModulesChanged( FName ModuleThatChanged, EModuleChangeReason Reaso
 	GNeedToRefreshSymbols = true;
 }
 
+/**
+ * Find the downstream storage path.
+ */
 FString FWindowsPlatformStackWalk::GetDownstreamStorage()
 {
 	FString DownstreamStorage;
@@ -834,10 +887,9 @@ void FWindowsPlatformStackWalk::ReleaseThreadContextWrapper(void* ThreadContext)
 }
 
 /**
- * Create path symbol path.
- * Reference: https://msdn.microsoft.com/en-us/library/ms681416%28v=vs.85%29.aspx?f=255&MSPPError=-2147217396
+ * Create path symbol path for remote storage.
  */
-static FString GetRemoteStorage(const FString& DownstreamStorage)
+FString GetRemoteStorage(const FString& DownstreamStorage)
 {
 	TArray<FString> RemoteStorage;
 	GConfig->GetArray(CrashReporterSettings, TEXT("RemoteStorage"), RemoteStorage, GEditorPerProjectIni);
@@ -863,11 +915,51 @@ static FString GetRemoteStorage(const FString& DownstreamStorage)
 	}
 }
 
+/**
+ * Create path symbol path.
+ * Reference: https://msdn.microsoft.com/en-us/library/ms681416%28v=vs.85%29.aspx?f=255&MSPPError=-2147217396
+ */
+FString GetSymbolSearchPath()
+{
+	FString SymbolSearchPath = GetRemoteStorage(FWindowsPlatformStackWalk::GetDownstreamStorage());
+	if (SymbolSearchPath.IsEmpty())
+	{
+		// by default passing null to SymInitialize will use the current working dir to search for a pdb, 
+		// but to support the basedir argument that allows an exe to run against data in a different location
+		// we'll put the path of the executing module first.
+		TCHAR ModulePath[MAX_PATH] = { 0 };
+		if (::GetModuleFileName(::GetModuleHandle(NULL), ModulePath, MAX_PATH))
+		{
+			SymbolSearchPath = FPaths::GetPath(ModulePath);
+			SymbolSearchPath += ";";
+#if ALLOW_UNREAL_ACCESS_TO_NT_SYMBOL_PATH
+			FString SymbolPathEnvironmentVariable = FPlatformMisc::GetEnvironmentVariable(L"_NT_SYMBOL_PATH");
+			if (!SymbolPathEnvironmentVariable.IsEmpty())
+			{
+				SymbolSearchPath += SymbolPathEnvironmentVariable;
+				SymbolSearchPath += ";";
+			}
+#endif
+			SymbolSearchPath += FPlatformProcess::GetCurrentWorkingDirectory();
+			SymbolSearchPath += ";";
+		}
+	}
+	return SymbolSearchPath;
+}
+
+
+
 /**
  * Initializes the symbol engine if needed.
- */ 
-bool FWindowsPlatformStackWalk::InitStackWalking()
+ */
+bool FWindowsPlatformStackWalk::InitStackWalkingInternal(void* Process)
 {
+	if (GProcessHandle != INVALID_HANDLE_VALUE && !GNeedToRefreshSymbols)
+	{
+		return true;
+	}
+	GProcessHandle = Process;
+
 	// DbgHelp functions are not thread safe, but this function can potentially be called from different
 	// threads in our engine, so we take a critical section
 	static FCriticalSection CriticalSection;
@@ -919,49 +1011,40 @@ bool FWindowsPlatformStackWalk::InitStackWalking()
 		//		SymOpts |= SYMOPT_CASE_INSENSITIVE;
 
 		SymSetOptions( SymOpts );
+
+		FString SymbolSearchPath = GetSymbolSearchPath();
 	
-		// Initialize the symbol engine.		
-		FString RemoteStorage = GetRemoteStorage(GetDownstreamStorage());
+		// Initialize the symbol engine.
 #if WINVER > 0x502
-		if (RemoteStorage.IsEmpty())
-		{
-			// by default passing null to SymInitialize will use the current working dir to search for a pdb, 
-			// but to support the basedir argument that allows an exe to run against data in a different location
-			// we'll put the path of the executing module first.
-			TCHAR ModulePath[MAX_PATH] = { 0 };
-			if (::GetModuleFileName(::GetModuleHandle(NULL), ModulePath, MAX_PATH))
-			{
-				RemoteStorage = FPaths::GetPath(ModulePath);
-				RemoteStorage += ";";
-				RemoteStorage += FPlatformProcess::GetCurrentWorkingDirectory();
-			}
-		}
-		
-		SymInitializeW( GetCurrentProcess(), RemoteStorage.IsEmpty() ? nullptr : *RemoteStorage, true );
+		SymInitializeW( GProcessHandle, SymbolSearchPath.IsEmpty() ? nullptr : *SymbolSearchPath, true );
 #else
-		SymInitialize( GetCurrentProcess(), nullptr, true );
+		SymInitialize( GProcessHandle, nullptr, true );
 #endif
 	
 		GNeedToRefreshSymbols = false;
 		GStackWalkingInitialized = true;
 
+#if !ON_DEMAND_SYMBOL_LOADING
 		if (!FPlatformProperties::IsMonolithicBuild() && FPlatformStackWalk::WantsDetailedCallstacksInNonMonolithicBuilds())
 		{
-			LoadProcessModules( RemoteStorage );
-		}			
+			LoadSymbolsForProcessModules(RemoteStorage);
+		}
+#endif
 	}
 #if WINVER > 0x502
 	else if (GNeedToRefreshSymbols)
 	{
 		// Refresh and reload symbols
-		SymRefreshModuleList( GetCurrentProcess() );
+		SymRefreshModuleList( GProcessHandle );
 
 		GNeedToRefreshSymbols = false;
 
 		if (!FPlatformProperties::IsMonolithicBuild() && FPlatformStackWalk::WantsDetailedCallstacksInNonMonolithicBuilds())
 		{
 			const FString RemoteStorage = GetRemoteStorage( GetDownstreamStorage() );
-			LoadProcessModules( RemoteStorage );
+			// When refresh is needed we cannot track which modules have been loaded and are interesting
+			// so load symbols for all modules the process has loaded.
+			LoadSymbolsForProcessModules( RemoteStorage );
 		}
 	}
 #endif
@@ -969,6 +1052,18 @@ bool FWindowsPlatformStackWalk::InitStackWalking()
 	return GStackWalkingInitialized;
 }
 
+
+bool FWindowsPlatformStackWalk::InitStackWalking()
+{
+	return FWindowsPlatformStackWalk::InitStackWalkingInternal(GetCurrentProcess());
+}
+
+bool FWindowsPlatformStackWalk::InitStackWalkingForProcess(const FProcHandle& Process)
+{
+	return FWindowsPlatformStackWalk::InitStackWalkingInternal(Process.Get());
+}
+
+
 void FWindowsPlatformStackWalk::RegisterOnModulesChanged()
 {
 	// Register for callback so we can reload symbols when new modules are loaded
diff --git a/Engine/Source/Runtime/Core/Public/GenericPlatform/GenericPlatformCrashContext.h b/Engine/Source/Runtime/Core/Public/GenericPlatform/GenericPlatformCrashContext.h
--- a/Engine/Source/Runtime/Core/Public/GenericPlatform/GenericPlatformCrashContext.h
+++ b/Engine/Source/Runtime/Core/Public/GenericPlatform/GenericPlatformCrashContext.h
@@ -82,6 +82,13 @@ struct FCrashStackFrame
 	}
 };
 
+/** Portable thread stack frame */
+struct FThreadStackFrames {
+	FString						ThreadName;
+	uint32						ThreadId;
+	TArray<FCrashStackFrame>	StackFrames;
+};
+
 enum class ECrashContextType
 {
 	Crash,
@@ -94,6 +101,108 @@ enum class ECrashContextType
 	Max
 };
 
+/** In development mode we can cause crashes in order to test reporting systems. */
+enum class ECrashTrigger
+{
+	Debug = -1,
+	Normal = 0
+};
+
+#define CR_MAX_ERROR_MESSAGE_CHARS 2048
+#define CR_MAX_DIRECTORY_CHARS 256
+#define CR_MAX_STACK_FRAMES 256
+#define CR_MAX_THREAD_NAME_CHARS 64
+#define CR_MAX_THREADS 256
+#define CR_MAX_GENERIC_FIELD_CHARS 64
+#define CR_MAX_COMMANDLINE_CHARS 1024
+#define CR_MAX_RICHTEXT_FIELD_CHARS 512
+#define CR_MAX_DYNAMIC_BUFFER_CHARS 1024*16
+
+/**
+ * Fixed size structure that holds session specific state.
+ */
+struct FSessionContext 
+{
+	bool 					bIsInternalBuild;
+	bool 					bIsPerforceBuild;
+	bool 					bIsSourceDistribution;
+	bool 					bIsUE4Release;
+	bool					bIsOOM;
+	uint32					ProcessId;
+	int32 					LanguageLCID;
+	int32 					NumberOfCores;
+	int32 					NumberOfCoresIncludingHyperthreads;
+	int32 					SecondsSinceStart;
+	int32 					CrashDumpMode;
+	int32					CrashType;
+	int32					OOMAllocationAlignment;
+	uint64					OOMAllocationSize;
+	TCHAR 					GameName[CR_MAX_GENERIC_FIELD_CHARS];
+	TCHAR					EngineMode[CR_MAX_GENERIC_FIELD_CHARS];
+	TCHAR					EngineModeEx[CR_MAX_GENERIC_FIELD_CHARS];
+	TCHAR 					ExecutableName[CR_MAX_GENERIC_FIELD_CHARS];
+	TCHAR 					DeploymentName[CR_MAX_GENERIC_FIELD_CHARS];
+	TCHAR 					BaseDir[CR_MAX_DIRECTORY_CHARS];
+	TCHAR 					RootDir[CR_MAX_DIRECTORY_CHARS];
+	TCHAR 					EpicAccountId[CR_MAX_GENERIC_FIELD_CHARS];
+	TCHAR 					LoginIdStr[CR_MAX_GENERIC_FIELD_CHARS];
+	TCHAR 					OsVersion[CR_MAX_GENERIC_FIELD_CHARS];
+	TCHAR 					OsSubVersion[CR_MAX_GENERIC_FIELD_CHARS];
+	TCHAR 					CPUVendor[CR_MAX_GENERIC_FIELD_CHARS];
+	TCHAR 					CPUBrand[CR_MAX_GENERIC_FIELD_CHARS];
+	TCHAR 					PrimaryGPUBrand[CR_MAX_GENERIC_FIELD_CHARS];
+	TCHAR 					UserName[CR_MAX_GENERIC_FIELD_CHARS];
+	TCHAR 					DefaultLocale[CR_MAX_GENERIC_FIELD_CHARS];
+	TCHAR 					CrashGUIDRoot[CR_MAX_GENERIC_FIELD_CHARS];
+	TCHAR 					UserActivityHint[CR_MAX_GENERIC_FIELD_CHARS];
+	TCHAR 					GameSessionID[CR_MAX_GENERIC_FIELD_CHARS];
+	TCHAR 					CommandLine[CR_MAX_COMMANDLINE_CHARS];
+	TCHAR 					CrashReportClientRichText[CR_MAX_RICHTEXT_FIELD_CHARS];
+	TCHAR 					GameStateName[CR_MAX_GENERIC_FIELD_CHARS];
+	TCHAR 					CrashConfigFilePath[CR_MAX_DIRECTORY_CHARS];
+	FPlatformMemoryStats	MemoryStats;
+	TOptional<bool>			bIsVanilla;
+};
+
+/**
+ * Fixed size struct holds crash information and session specific state. It is designed
+ * to shared between processes (e.g. Game and CrashReporterClient).
+ */
+struct FSharedCrashContext
+{
+	// Exception info
+	TCHAR					ErrorMessage[CR_MAX_ERROR_MESSAGE_CHARS];
+	uint32					ThreadIds[CR_MAX_THREADS];
+	TCHAR					ThreadNames[CR_MAX_THREAD_NAME_CHARS * CR_MAX_THREADS];
+	uint32					NumThreads;
+	uint32					CrashingThreadId;
+	uint32					NumStackFramesToIgnore;
+	ECrashContextType		CrashType;
+	// Platform specific crash context (must be portable)
+	void*					PlatformCrashContext;
+	// Platform specific memory stats
+	FPlatformMemoryStats	MemoryStats;
+	// User settings info
+	bool					bNoDialog;
+	bool					bSendUnattenededBugReports;
+	bool					bSendUsageData;
+	// Directory for dumped files
+	TCHAR					CrashFilesDirectory[CR_MAX_DIRECTORY_CHARS];
+	// Game/Engine information not possible to catch out of process
+	FSessionContext			SessionContext;
+	// Count and offset into dynamic buffer to comma separated plugin list
+	uint32					EnabledPluginsNum;
+	uint32					EnabledPluginsOffset;
+	// Count and offset into dynamic buffer to comma separated key=value data for engine information
+	uint32					EngineDataNum;
+	uint32					EngineDataOffset;
+	// Count and offset into dynamic buffer to comma separated key=value data for  game information
+	uint32					GameDataNum;
+	uint32					GameDataOffset;
+	// Fixed size dynamic buffer
+	TCHAR					DynamicData[CR_MAX_DYNAMIC_BUFFER_CHARS];
+};
+
 /**
  *	Contains a runtime crash's properties that are common for all platforms.
  *	This may change in the future.
@@ -102,12 +211,6 @@ struct CORE_API FGenericCrashContext
 {
 public:
 
-	/**
-	* We can't gather memory stats in crash handling function, so we gather them just before raising
-	* exception and use in crash reporting.
-	*/
-	static FPlatformMemoryStats CrashMemoryStats;
-	
 	static const ANSICHAR* CrashContextRuntimeXMLNameA;
 	static const TCHAR* CrashContextRuntimeXMLNameW;
 
@@ -144,6 +247,9 @@ public:
 	/** Initializes crash context related platform specific data that can be impossible to obtain after a crash. */
 	static void Initialize();
 
+	/** Initialized crash context, using a crash context (e.g. shared from another process). */
+	static void InitializeFromContext(const FSessionContext& Context, const TCHAR* EnabledPlugins, const TCHAR* EngineData, const TCHAR* GameData);
+
 	/**
 	 * @return true, if the generic crash context has been initialized.
 	 */
@@ -152,7 +258,7 @@ public:
 		return bIsInitialized;
 	}
 
-	/** Default constructor. */
+	/** Default constructor. Optionally pass a process handle if building a crash context for a process other then current. */
 	FGenericCrashContext(ECrashContextType InType, const TCHAR* ErrorMessage);
 
 	virtual ~FGenericCrashContext() { }
@@ -230,6 +336,26 @@ public:
 	/** Adds a plugin descriptor string to the enabled plugins list in the crash context */
 	static void AddPlugin(const FString& PluginDesc);
 
+	/** Flushes the logs. In the case of in memory logs is used on this configuration, dumps them to file. */
+	static void DumpLog(const FString& CrashFolderAbsolute);
+
+	/** Initializes a shared crash context from current state. Will not set all fields in Dst. */
+	static void CopySharedCrashContext(FSharedCrashContext& Dst);
+
+	/** We can't gather memory stats in crash handling function, so we gather them just before raising
+	  * exception and use in crash reporting. 
+	  */
+	static void SetMemoryStats(const FPlatformMemoryStats& MemoryStats);
+
+	/** Attempts to create the output report directory. */
+	static bool CreateCrashReportDirectory(const TCHAR* CrashGUIDRoot, const TCHAR* AppName, int32 CrashIndex, FString& OutCrashDirectoryAbsolute);
+
+	/** Sets the process id to that has crashed. On supported platforms this will analyze the given process rather than current. Default is current process. */
+	void SetCrashedProcess(const FProcHandle& Process) { ProcessHandle = Process; }
+
+	/** Stores crashing thread id. */
+	void SetCrashedThreadId(uint32 InId) { CrashedThreadId = InId; }
+
 	/** Sets the number of stack frames to ignore when symbolicating from a minidump */
 	void SetNumMinidumpFramesToIgnore(int32 InNumMinidumpFramesToIgnore);
 
@@ -242,6 +368,12 @@ public:
 	/** Gets the portable callstack to a specified stack and puts it into OutCallStack */
 	virtual void GetPortableCallStack(const uint64* StackFrames, int32 NumStackFrames, TArray<FCrashStackFrame>& OutCallStack);
 
+	/** Adds a portable callstack for a thread */
+	virtual void AddPortableThreadCallStack(uint32 ThreadId, const TCHAR* ThreadName, const uint64* StackFrames, int32 NumStackFrames);
+
+	/** Allows platform implementations to copy files to report directory. */
+	virtual void CopyPlatformSpecificFiles(const TCHAR* OutputDirectory, void* Context);
+
 	/**
 	 * @return whether this crash is a non-crash event
 	 */
@@ -252,6 +384,11 @@ public:
 	 */
 	static void SetDeploymentName(const FString& EpicApp);
 
+	/**
+	 * Sets the type of crash triggered. Used to distinguish crashes caused for debugging purposes.
+	 */
+	static void SetCrashTrigger(ECrashTrigger Type);
+
 protected:
 	/**
 	 * @OutStr - a stream of Thread XML elements containing info (e.g. callstack) specific to an active thread
@@ -259,10 +396,13 @@ protected:
 	 */
 	virtual bool GetPlatformAllThreadContextsString(FString& OutStr) const { return false; }
 
+	FProcHandle ProcessHandle;
 	ECrashContextType Type;
+	uint32 CrashedThreadId;
 	const TCHAR* ErrorMessage;
 	int NumMinidumpFramesToIgnore;
 	TArray<FCrashStackFrame> CallStack;
+	TArray<FThreadStackFrames> ThreadCallStacks;
 
 private:
 
diff --git a/Engine/Source/Runtime/Core/Public/GenericPlatform/GenericPlatformProcess.h b/Engine/Source/Runtime/Core/Public/GenericPlatform/GenericPlatformProcess.h
--- a/Engine/Source/Runtime/Core/Public/GenericPlatform/GenericPlatformProcess.h
+++ b/Engine/Source/Runtime/Core/Public/GenericPlatform/GenericPlatformProcess.h
@@ -126,6 +126,13 @@ struct CORE_API FGenericPlatformProcess
 		 */
 		FSemaphore(const FString& InName);
 
+		/**
+		 * Creates and initializes a new instance with the specified name.
+		 *
+		 * @param InName name of the semaphore (all processes should use the same)
+		 */
+		FSemaphore(const TCHAR* InName);
+
 		/** Virtual destructor. */
 		virtual ~FSemaphore() { };
 
@@ -405,6 +412,10 @@ struct CORE_API FGenericPlatformProcess
 	/** Returns true if the specified application is running */
 	static bool IsApplicationRunning( const TCHAR* ProcName );
 
+	/** Returns true if the specified application has alive, false if it has exited. This is slightly different that IsApplicationRunning,
+	as we are able to query applications that has crashed but not yet exited. */
+	static bool IsApplicationAlive(uint32 ProcessId);
+
 	/** Returns the Name of process given by the PID.  Returns Empty string "" if PID not found. */
 	static FString GetApplicationName( uint32 ProcessId );
 
@@ -580,9 +591,20 @@ struct CORE_API FGenericPlatformProcess
 	 * @param Name name (so we can use it across processes).
 	 * @param bCreate If true, the function will try to create, otherwise will try to open existing.
 	 * @param MaxLocks Maximum amount of locks that the semaphore can have (pass 1 to make it act as mutex).
+	 * @return Pointer to heap allocated semaphore object. Caller is responsible for deletion.
 	 */
 	static FSemaphore* NewInterprocessSynchObject(const FString& Name, bool bCreate, uint32 MaxLocks = 1);
 
+	/**
+	 * Creates or opens an interprocess synchronization object.
+	 *
+	 * @param Name name (so we can use it across processes).
+	 * @param bCreate If true, the function will try to create, otherwise will try to open existing.
+	 * @param MaxLocks Maximum amount of locks that the semaphore can have (pass 1 to make it act as mutex).
+	 * @return Pointer to heap allocated semaphore object. Caller is responsible for deletion.
+	 */
+	static FSemaphore* NewInterprocessSynchObject(const TCHAR* Name, bool bCreate, uint32 MaxLocks = 1);
+
 	/**
 	 * Deletes an interprocess synchronization object.
 	 *
diff --git a/Engine/Source/Runtime/Core/Public/GenericPlatform/GenericPlatformStackWalk.h b/Engine/Source/Runtime/Core/Public/GenericPlatform/GenericPlatformStackWalk.h
--- a/Engine/Source/Runtime/Core/Public/GenericPlatform/GenericPlatformStackWalk.h
+++ b/Engine/Source/Runtime/Core/Public/GenericPlatform/GenericPlatformStackWalk.h
@@ -10,6 +10,7 @@
 
 struct FGenericCrashContext;
 struct FProgramCounterSymbolInfoEx;
+struct FProcHandle;
 
 /**
  * This is used to capture all of the module information needed to load pdb's.
@@ -107,6 +108,16 @@ struct CORE_API FGenericPlatformStackWalk
 		return 1;
 	}
 
+	/**
+	 * Like InitStackWalking but initializes stack walking for another process. 
+	 * @param	Process		Process that is going to be queried
+	 * @return				True if successful, false otherwise.
+	 */
+	static bool InitStackWalkingForProcess(const FProcHandle& Process)
+	{
+		return false;
+	}
+
 	/**
 	 * Converts the passed in program counter address to a human readable string and appends it to the passed in one.
 	 * @warning: The code assumes that HumanReadableString is large enough to contain the information.
diff --git a/Engine/Source/Runtime/Core/Public/HAL/ExceptionHandling.h b/Engine/Source/Runtime/Core/Public/HAL/ExceptionHandling.h
--- a/Engine/Source/Runtime/Core/Public/HAL/ExceptionHandling.h
+++ b/Engine/Source/Runtime/Core/Public/HAL/ExceptionHandling.h
@@ -55,15 +55,6 @@ extern CORE_API void ReportHang(const TCHAR*, const uint64* StackFrames, int32 N
 extern CORE_API void NewReportEnsure(const TCHAR* ErrorMessage);
 #endif
 
-enum class ECrashType
-{
-	Debug = -1,
-	Normal = 0
-};
-
-/** In debug builds we can crash on demand so we want to let the crash reporting system know about it */
-extern CORE_API void SetCrashType(ECrashType InCrashType);
-extern CORE_API int32 GetCrashType();
 
 extern CORE_API void ReportInteractiveEnsure(const TCHAR* InMessage);
 extern CORE_API bool IsInteractiveEnsureMode();
\ No newline at end of file
diff --git a/Engine/Source/Runtime/Core/Public/HAL/PlatformCrashContext.h b/Engine/Source/Runtime/Core/Public/HAL/PlatformCrashContext.h
new file mode 100644
--- /dev/null
+++ b/Engine/Source/Runtime/Core/Public/HAL/PlatformCrashContext.h
@@ -0,0 +1,7 @@
+// Copyright 1998-2019 Epic Games, Inc. All Rights Reserved.
+#pragma once
+
+#include "CoreTypes.h"
+
+#include COMPILED_PLATFORM_HEADER(PlatformCrashContext.h)
+
diff --git a/Engine/Source/Runtime/Core/Public/HAL/ThreadManager.h b/Engine/Source/Runtime/Core/Public/HAL/ThreadManager.h
--- a/Engine/Source/Runtime/Core/Public/HAL/ThreadManager.h
+++ b/Engine/Source/Runtime/Core/Public/HAL/ThreadManager.h
@@ -15,6 +15,8 @@ class CORE_API FThreadManager
 	TMap<uint32, class FRunnableThread*> Threads;
 	/** Critical section for ThreadList */
 	FCriticalSection ThreadsCritical;
+	/** Set if thread manager is initialized. */
+	static bool bIsInitialized;
 
 public:
 
@@ -40,6 +42,12 @@ public:
 	/** Returns the name of a thread given its TLS id */
 	const FString& GetThreadName(uint32 ThreadId);
 
+	/** Checks if thread manager has been initialized. Avoids creating the manager trough lazy initialization */
+	FORCEINLINE static bool IsInitialized() 
+	{
+		return bIsInitialized;
+	}
+
 #if PLATFORM_WINDOWS || PLATFORM_MAC
 	struct FThreadStackBackTrace
 	{
diff --git a/Engine/Source/Runtime/Core/Public/Linux/LinuxPlatformCrashContext.h b/Engine/Source/Runtime/Core/Public/Linux/LinuxPlatformCrashContext.h
new file mode 100644
--- /dev/null
+++ b/Engine/Source/Runtime/Core/Public/Linux/LinuxPlatformCrashContext.h
@@ -0,0 +1,5 @@
+// Copyright 1998-2019 Epic Games, Inc. All Rights Reserved.
+
+#pragma once
+
+#include "Unix/UnixPlatformCrashContext.h"
\ No newline at end of file
diff --git a/Engine/Source/Runtime/Core/Public/Windows/WindowsPlatformCrashContext.h b/Engine/Source/Runtime/Core/Public/Windows/WindowsPlatformCrashContext.h
--- a/Engine/Source/Runtime/Core/Public/Windows/WindowsPlatformCrashContext.h
+++ b/Engine/Source/Runtime/Core/Public/Windows/WindowsPlatformCrashContext.h
@@ -9,7 +9,6 @@ struct CORE_API FWindowsPlatformCrashContext : public FGenericCrashContext
 {
 	FWindowsPlatformCrashContext(ECrashContextType InType, const TCHAR* InErrorMessage)
 		: FGenericCrashContext(InType, InErrorMessage)
-		, CrashedThreadId(-1)
 	{
 	}
 
@@ -17,9 +16,11 @@ struct CORE_API FWindowsPlatformCrashContext : public FGenericCrashContext
 
 	virtual void AddPlatformSpecificProperties() const override;
 
-	void CaptureAllThreadContexts() { AddAllThreadContexts(CrashedThreadId, AllThreadContexts); }
+	virtual void AddPortableThreadCallStack(uint32 ThreadId, const TCHAR* ThreadName, const uint64* StackFrames, int32 NumStackFrames) override;
 
-	void SetCrashedThreadId(uint32 InId) { CrashedThreadId = InId; }
+	virtual void CopyPlatformSpecificFiles(const TCHAR* OutputDirectory, void* Context) override;
+
+	void CaptureAllThreadContexts();
 
 protected:
 	virtual bool GetPlatformAllThreadContextsString(FString& OutStr) const override;
@@ -28,42 +29,22 @@ private:
 	// Helpers
 	typedef TArray<void*, TInlineAllocator<128>> FModuleHandleArray;
 	
-	static void GetProcModuleHandles(FModuleHandleArray& OutHandles);
+	static void GetProcModuleHandles(const FProcHandle& Process, FModuleHandleArray& OutHandles);
 
 	static void ConvertProgramCountersToStackFrames(
+		const FProcHandle& Process,
 		const FModuleHandleArray& SortedModuleHandles,
 		const uint64* ProgramCounters,
 		int32 NumPCs,
 		TArray<FCrashStackFrame>& OutStackFrames);
 
-	static void AddIsCrashed(bool bIsCrashed, FString& OutStr);
-	static void AddThreadId(uint32 ThreadId, FString& OutStr);
-	static void AddThreadName(const TCHAR* ThreadName, FString& OutStr);
-	static void AddThreadContext(
-		const FModuleHandleArray& ProcModuleHandles,
+	static void AddThreadContextString(
 		uint32 CrashedThreadId,
 		uint32 ThreadId,
 		const FString& ThreadName,
-		const uint64* StackTrace,
-		int32 Depth,
-		FString& OutStr);
-	static void AddAllThreadContexts(uint32 CrashedThreadId, FString& OutStr);
-
-	/**
-	* <Thread>
-	*   <CallStack>...</CallStack>
-	*   <IsCrashed>...</IsCrashed>
-	*   <Registers>...</Registers>
-	*   <ThreadID>...</ThreadID>
-	*   <ThreadName>...</ThreadName>
-	* </Thead>
-	* <Thread>...</Thread>
-	* ...
-	*/
-	FString AllThreadContexts;
+		const TArray<FCrashStackFrame>& StackFrames,
+		FString& OutStr);	
 
-	// ID of the crashed thread
-	uint32 CrashedThreadId;
 };
 
 typedef FWindowsPlatformCrashContext FPlatformCrashContext;
diff --git a/Engine/Source/Runtime/Core/Public/Windows/WindowsPlatformProcess.h b/Engine/Source/Runtime/Core/Public/Windows/WindowsPlatformProcess.h
--- a/Engine/Source/Runtime/Core/Public/Windows/WindowsPlatformProcess.h
+++ b/Engine/Source/Runtime/Core/Public/Windows/WindowsPlatformProcess.h
@@ -46,6 +46,9 @@ struct CORE_API FWindowsPlatformProcess
 		/** Constructor */
 		FWindowsSemaphore(const FString& InName, Windows::HANDLE InSemaphore);
 
+		/** Allocation free constructor */
+		FWindowsSemaphore(const TCHAR* InName, Windows::HANDLE InSemaphore);
+
 		/** Destructor */
 		virtual ~FWindowsSemaphore();
 
@@ -160,6 +163,7 @@ public:
 	static bool GetPerFrameProcessorUsage(uint32 ProcessId, float& ProcessUsageFraction, float& IdleUsageFraction);
 	static bool IsApplicationRunning( uint32 ProcessId );
 	static bool IsApplicationRunning( const TCHAR* ProcName );
+	static bool IsApplicationAlive(uint32 ProcessId);
 	static FString GetApplicationName( uint32 ProcessId );	
 	static bool ExecProcess(const TCHAR* URL, const TCHAR* Params, int32* OutReturnCode, FString* OutStdOut, FString* OutStdErr, const TCHAR* OptionalWorkingDirectory = NULL);
 	static bool ExecElevatedProcess(const TCHAR* URL, const TCHAR* Params, int32* OutReturnCode);
@@ -178,6 +182,7 @@ public:
 	static bool WritePipe(void* WritePipe, const FString& Message, FString* OutWritten = nullptr);
 	static bool WritePipe(void* WritePipe, const uint8* Data, const int32 DataLength, int32* OutDataLength = nullptr);
 	static FSemaphore* NewInterprocessSynchObject(const FString& Name, bool bCreate, uint32 MaxLocks = 1);
+	static FSemaphore* NewInterprocessSynchObject(const TCHAR* Name, bool bCreate, uint32 MaxLocks = 1);
 	static bool DeleteInterprocessSynchObject(FSemaphore * Object);
 	static bool Daemonize();
 protected:
diff --git a/Engine/Source/Runtime/Core/Public/Windows/WindowsPlatformStackWalk.h b/Engine/Source/Runtime/Core/Public/Windows/WindowsPlatformStackWalk.h
--- a/Engine/Source/Runtime/Core/Public/Windows/WindowsPlatformStackWalk.h
+++ b/Engine/Source/Runtime/Core/Public/Windows/WindowsPlatformStackWalk.h
@@ -12,6 +12,7 @@ struct CORE_API FWindowsPlatformStackWalk
 	: public FGenericPlatformStackWalk
 {
 	static bool InitStackWalking();
+	static bool InitStackWalkingForProcess(const FProcHandle& Process);
 	
 	static TArray<FProgramCounterSymbolInfo> GetStack(int32 IgnoreCount, int32 MaxDepth = 100, void* Context = nullptr);
 
@@ -43,6 +44,9 @@ struct CORE_API FWindowsPlatformStackWalk
 
 	static void* MakeThreadContextWrapper(void* Context, void* ThreadHandle);
 	static void ReleaseThreadContextWrapper(void* ThreadContext);
+
+private:
+	static bool InitStackWalkingInternal(void* Process);
 };
 
 
diff --git a/Engine/Source/Runtime/CrashReportCore/Private/CrashDescription.cpp b/Engine/Source/Runtime/CrashReportCore/Private/CrashDescription.cpp
--- a/Engine/Source/Runtime/CrashReportCore/Private/CrashDescription.cpp
+++ b/Engine/Source/Runtime/CrashReportCore/Private/CrashDescription.cpp
@@ -169,10 +169,17 @@ void FPrimaryCrashProperties::UpdateIDs()
 	LoginId = FPlatformMisc::GetLoginId();
 }
 
-void FPrimaryCrashProperties::ReadXML( const FString& CrashContextFilepath  )
+void FPrimaryCrashProperties::ReadXML( const FString& CrashContextFilepath, const TCHAR* Buffer  )
 {
 	XmlFilepath = CrashContextFilepath;
-	XmlFile = new FXmlFile( XmlFilepath );
+	if (Buffer)
+	{
+		XmlFile = new FXmlFile(Buffer, EConstructMethod::ConstructFromBuffer);
+	}
+	else
+	{
+		XmlFile = new FXmlFile(XmlFilepath);
+	}
 	if (XmlFile->IsValid())
 	{
 		TimeOfCrash = FDateTime::UtcNow().GetTicks();
@@ -446,10 +453,20 @@ void FPrimaryCrashProperties::Save()
 	FCrashContextReader
 -----------------------------------------------------------------------------*/
 
-FCrashContext::FCrashContext( const FString& CrashContextFilepath )
+FCrashContext::FCrashContext(const FString& CrashContextFilePath, const TCHAR* Buffer)
 {
-	ReadXML( CrashContextFilepath );
+	ReadXML(CrashContextFilePath, Buffer);
+	SetupPrimaryCrashProperties();
+}
 
+FCrashContext::FCrashContext(const FString& CrashContextFilepath)
+{
+	ReadXML(CrashContextFilepath);
+	SetupPrimaryCrashProperties();
+}
+
+void FCrashContext::SetupPrimaryCrashProperties() 
+{
 	const bool bIsValid = XmlFile->IsValid();
 	if (bIsValid)
 	{
diff --git a/Engine/Source/Runtime/CrashReportCore/Private/GenericErrorReport.cpp b/Engine/Source/Runtime/CrashReportCore/Private/GenericErrorReport.cpp
--- a/Engine/Source/Runtime/CrashReportCore/Private/GenericErrorReport.cpp
+++ b/Engine/Source/Runtime/CrashReportCore/Private/GenericErrorReport.cpp
@@ -189,7 +189,10 @@ void FGenericErrorReport::SetPrimaryCrashProperties( FPrimaryCrashProperties& ou
 			CallStack.RemoveAt(0, FMath::Min(CallStack.Num(), (int32)NumMinidumpFramesToIgnore));
 		}
 
-		out_PrimaryCrashProperties.CallStack = CallStack;
+		if (CallStack.Num() > 0)
+		{
+			out_PrimaryCrashProperties.CallStack = CallStack;
+		}
 		out_PrimaryCrashProperties.Modules = Helper->CrashInfo.ModuleNames;
 		out_PrimaryCrashProperties.SourceContext = Helper->CrashInfo.SourceContext;
 
diff --git a/Engine/Source/Runtime/CrashReportCore/Public/CrashDescription.h b/Engine/Source/Runtime/CrashReportCore/Public/CrashDescription.h
--- a/Engine/Source/Runtime/CrashReportCore/Public/CrashDescription.h
+++ b/Engine/Source/Runtime/CrashReportCore/Public/CrashDescription.h
@@ -453,8 +453,8 @@ public:
 	void Save();
 
 protected:
-	/** Reads previously set XML file. */
-	void ReadXML( const FString& CrashContextFilepath );
+	/** Reads previously set XML file. Optionally give a buffer to parse from. Any changes will be stored to CrashContextFilePath. */
+	void ReadXML( const FString& CrashContextFilepath, const TCHAR* Buffer = nullptr );
 
 	/** Sets the CrasgGUID based on the report's path. */
 	void SetCrashGUID( const FString& Filepath );
@@ -534,6 +534,14 @@ struct FCrashContext : public FPrimaryCrashProperties
 {
 	/** Initializes instance based on specified Crash Context filepath. */
 	explicit FCrashContext( const FString& CrashContextFilepath );
+
+	/** Initializes an instance based on a in memory buffer. */
+	FCrashContext(const FString& CrashContextFilepath, const TCHAR* Buffer);
+
+private:
+
+	/** Initializes the primary crash property fields from the xml file */
+	void SetupPrimaryCrashProperties();
 };
 
 /** Crash context based on the Window Error Reporting WER files, obsolete, only for backward compatibility. */
diff --git a/Engine/Source/Runtime/Engine/Private/UnrealEngine.cpp b/Engine/Source/Runtime/Engine/Private/UnrealEngine.cpp
--- a/Engine/Source/Runtime/Engine/Private/UnrealEngine.cpp
+++ b/Engine/Source/Runtime/Engine/Private/UnrealEngine.cpp
@@ -8030,7 +8030,7 @@ bool UEngine::PerformError(const TCHAR* Cmd, FOutputDevice& Ar)
 			[](FRHICommandList& RHICmdList)
 			{
 				UE_LOG(LogEngine, Warning, TEXT("Printed warning to log."));
-				SetCrashType(ECrashType::Debug);
+				FGenericCrashContext::SetCrashTrigger(ECrashTrigger::Debug);
 				UE_LOG(LogEngine, Fatal, TEXT("Crashing the renderthread at your request"));
 			});
 		return true;
@@ -8042,7 +8042,7 @@ bool UEngine::PerformError(const TCHAR* Cmd, FOutputDevice& Ar)
 			static void Check()
 			{				
 				UE_LOG(LogEngine, Warning, TEXT("Printed warning to log."));
-				SetCrashType(ECrashType::Debug);
+				FGenericCrashContext::SetCrashTrigger(ECrashTrigger::Debug);
 				check(!"Crashing the renderthread via check(0) at your request");
 			}
 		};
@@ -8059,7 +8059,7 @@ bool UEngine::PerformError(const TCHAR* Cmd, FOutputDevice& Ar)
 			[](FRHICommandList& RHICmdList)
 			{
 				UE_LOG(LogEngine, Warning, TEXT("Printed warning to log."));
-				SetCrashType(ECrashType::Debug);
+				FGenericCrashContext::SetCrashTrigger(ECrashTrigger::Debug);
 				*(int32 *)3 = 123;
 			});
 		return true;
@@ -8070,8 +8070,8 @@ bool UEngine::PerformError(const TCHAR* Cmd, FOutputDevice& Ar)
 			[](FRHICommandList& RHICmdList)
 			{
 				UE_LOG(LogEngine, Warning, TEXT("Printed warning to log."));
-			SetCrashType(ECrashType::Debug);
-			LowLevelFatalError(TEXT("FError::LowLevelFatal test"));
+				FGenericCrashContext::SetCrashTrigger(ECrashTrigger::Debug);
+				LowLevelFatalError(TEXT("FError::LowLevelFatal test"));
 			});
 		return true;
 	}
@@ -8095,7 +8095,7 @@ bool UEngine::PerformError(const TCHAR* Cmd, FOutputDevice& Ar)
 			static void Crash(ENamedThreads::Type, const  FGraphEventRef&)
 			{
 				UE_LOG(LogEngine, Warning, TEXT("Printed warning to log."));
-				SetCrashType(ECrashType::Debug);
+				FGenericCrashContext::SetCrashTrigger(ECrashTrigger::Debug);
 				UE_LOG(LogEngine, Fatal, TEXT("Crashing the worker thread at your request"));
 			}
 		};
@@ -8120,7 +8120,7 @@ bool UEngine::PerformError(const TCHAR* Cmd, FOutputDevice& Ar)
 			static void Check(ENamedThreads::Type, const FGraphEventRef&)
 			{
 				UE_LOG(LogEngine, Warning, TEXT("Printed warning to log."));
-				SetCrashType(ECrashType::Debug);
+				FGenericCrashContext::SetCrashTrigger(ECrashTrigger::Debug);
 				check(!"Crashing a worker thread via check(0) at your request");
 			}
 		};
@@ -8145,7 +8145,7 @@ bool UEngine::PerformError(const TCHAR* Cmd, FOutputDevice& Ar)
 			static void GPF(ENamedThreads::Type, const FGraphEventRef&)
 			{
 				UE_LOG(LogEngine, Warning, TEXT("Printed warning to log."));
-				SetCrashType(ECrashType::Debug);
+				FGenericCrashContext::SetCrashTrigger(ECrashTrigger::Debug);
 				*(int32 *)3 = 123;
 			}
 		};
@@ -8185,7 +8185,7 @@ bool UEngine::PerformError(const TCHAR* Cmd, FOutputDevice& Ar)
 					if (FPlatformTime::Seconds() >= CrashTime)
 					{
 						UE_LOG(LogEngine, Warning, TEXT("Printed warning to log."));
-						SetCrashType(ECrashType::Debug);
+						FGenericCrashContext::SetCrashTrigger(ECrashTrigger::Debug);
 						UE_LOG(LogEngine, Fatal, TEXT("Crashing the worker thread at your request"));
 						break;
 					}
@@ -8225,7 +8225,7 @@ bool UEngine::PerformError(const TCHAR* Cmd, FOutputDevice& Ar)
 					if (FPlatformTime::Seconds() >= CrashTime)
 					{
 						UE_LOG(LogEngine, Warning, TEXT("Printed warning to log."));
-						SetCrashType(ECrashType::Debug);
+						FGenericCrashContext::SetCrashTrigger(ECrashTrigger::Debug);
 						*(int32 *)3 = 123;
 						break;
 					}
@@ -8273,7 +8273,7 @@ bool UEngine::PerformError(const TCHAR* Cmd, FOutputDevice& Ar)
 			static void Fatal(ENamedThreads::Type, const FGraphEventRef&)
 			{
 				UE_LOG(LogEngine, Warning, TEXT("Printed warning to log."));
-				SetCrashType(ECrashType::Debug);
+				FGenericCrashContext::SetCrashTrigger(ECrashTrigger::Debug);
 				LowLevelFatalError(TEXT("FError::LowLevelFatal test"));
 			}
 		};
@@ -8293,14 +8293,14 @@ bool UEngine::PerformError(const TCHAR* Cmd, FOutputDevice& Ar)
 	else if (FParse::Command(&Cmd, TEXT("CRASH")))
 	{
 		UE_LOG(LogEngine, Warning, TEXT("Printed warning to log."));
-		SetCrashType(ECrashType::Debug);
+		FGenericCrashContext::SetCrashTrigger(ECrashTrigger::Debug);
 		UE_LOG(LogEngine, Fatal, TEXT("%s"), TEXT("Crashing the gamethread at your request"));
 		return true;
 	}
 	else if (FParse::Command(&Cmd, TEXT("CHECK")))
 	{
 		UE_LOG(LogEngine, Warning, TEXT("Printed warning to log."));
-		SetCrashType(ECrashType::Debug);
+		FGenericCrashContext::SetCrashTrigger(ECrashTrigger::Debug);
 		check(!"Crashing the game thread via check(0) at your request");
 		return true;
 	}
@@ -8308,7 +8308,7 @@ bool UEngine::PerformError(const TCHAR* Cmd, FOutputDevice& Ar)
 	{
 		UE_LOG(LogEngine, Warning, TEXT("Printed warning to log."));
 		Ar.Log(TEXT("Crashing with voluntary GPF"));
-		SetCrashType(ECrashType::Debug);
+		FGenericCrashContext::SetCrashTrigger(ECrashTrigger::Debug);
 		// changed to 3 from NULL because clang noticed writing to NULL and warned about it
 		*(int32 *)3 = 123;
 		return true;
@@ -8332,7 +8332,7 @@ bool UEngine::PerformError(const TCHAR* Cmd, FOutputDevice& Ar)
 	else if (FParse::Command(&Cmd, TEXT("FATAL")))
 	{
 		UE_LOG(LogEngine, Warning, TEXT("Printed warning to log."));
-		SetCrashType(ECrashType::Debug);
+		FGenericCrashContext::SetCrashTrigger(ECrashTrigger::Debug);
 		LowLevelFatalError(TEXT("FError::LowLevelFatal test"));
 		return true;
 	}
@@ -8340,13 +8340,13 @@ bool UEngine::PerformError(const TCHAR* Cmd, FOutputDevice& Ar)
 	{
 		// stack overflow test - this case should be caught by /GS (Buffer Overflow Check) compile option
 		ANSICHAR SrcBuffer[] = "12345678901234567890123456789012345678901234567890";
-		SetCrashType(ECrashType::Debug);
+		FGenericCrashContext::SetCrashTrigger(ECrashTrigger::Debug);
 		BufferOverflowFunction(UE_ARRAY_COUNT(SrcBuffer), SrcBuffer);
 		return true;
 	}
 	else if (FParse::Command(&Cmd, TEXT("CRTINVALID")))
 	{
-		SetCrashType(ECrashType::Debug);
+	FGenericCrashContext::SetCrashTrigger(ECrashTrigger::Debug);
 		FString::Printf(TEXT("%s"), (const char*)nullptr);
 		return true;
 	}
@@ -8415,7 +8415,7 @@ bool UEngine::PerformError(const TCHAR* Cmd, FOutputDevice& Ar)
 	{
 		Ar.Logf(TEXT("Recursing to create a very deep callstack."));
 		GLog->Flush();
-		SetCrashType(ECrashType::Debug);
+		FGenericCrashContext::SetCrashTrigger(ECrashTrigger::Debug);
 		InfiniteRecursionFunction(1);
 		Ar.Logf(TEXT("You will never see this log line."));
 		return true;
@@ -8427,7 +8427,7 @@ bool UEngine::PerformError(const TCHAR* Cmd, FOutputDevice& Ar)
 		{
 			static void InfiniteRecursion(ENamedThreads::Type, const FGraphEventRef&)
 			{
-				SetCrashType(ECrashType::Debug);
+				FGenericCrashContext::SetCrashTrigger(ECrashTrigger::Debug);
 				InfiniteRecursionFunction(1);
 			}
 		};
@@ -8447,7 +8447,7 @@ bool UEngine::PerformError(const TCHAR* Cmd, FOutputDevice& Ar)
 	else if (FParse::Command(&Cmd, TEXT("EATMEM")))
 	{
 		Ar.Log(TEXT("Eating up all available memory"));
-		SetCrashType(ECrashType::Debug);
+		FGenericCrashContext::SetCrashTrigger(ECrashTrigger::Debug);
 		while (1)
 		{
 			void* Eat = FMemory::Malloc(65536);
@@ -8464,7 +8464,7 @@ bool UEngine::PerformError(const TCHAR* Cmd, FOutputDevice& Ar)
 	else if (FParse::Command(&Cmd, TEXT("STACKOVERFLOW")))
 	{
 		Ar.Log(TEXT("Infinite recursion to cause stack overflow"));
-		SetCrashType(ECrashType::Debug);
+		FGenericCrashContext::SetCrashTrigger(ECrashTrigger::Debug);
 		StackOverflowFunction(nullptr);
 		return true;
 	}
@@ -8475,7 +8475,7 @@ bool UEngine::PerformError(const TCHAR* Cmd, FOutputDevice& Ar)
 		{
 			static void StackOverflow(ENamedThreads::Type, const FGraphEventRef&)
 			{
-				SetCrashType(ECrashType::Debug);
+				FGenericCrashContext::SetCrashTrigger(ECrashTrigger::Debug);
 				StackOverflowFunction(nullptr);
 			}
 		};
@@ -8495,7 +8495,7 @@ bool UEngine::PerformError(const TCHAR* Cmd, FOutputDevice& Ar)
 	else if (FParse::Command(&Cmd, TEXT("SOFTLOCK")))
 	{
 		Ar.Log(TEXT("Hanging the current thread"));
-		SetCrashType(ECrashType::Debug);
+		FGenericCrashContext::SetCrashTrigger(ECrashTrigger::Debug);
 		while (1)
 		{
 			FPlatformProcess::Sleep(1.0f);
@@ -8505,7 +8505,7 @@ bool UEngine::PerformError(const TCHAR* Cmd, FOutputDevice& Ar)
 	else if (FParse::Command(&Cmd, TEXT("INFINITELOOP")))
 	{
 		Ar.Log(TEXT("Hanging the current thread (CPU-intensive)"));
-		SetCrashType(ECrashType::Debug);
+		FGenericCrashContext::SetCrashTrigger(ECrashTrigger::Debug);
 		for(;;)
 		{
 		}
diff --git a/Engine/Source/Runtime/Launch/Private/Windows/LaunchWindows.cpp b/Engine/Source/Runtime/Launch/Private/Windows/LaunchWindows.cpp
--- a/Engine/Source/Runtime/Launch/Private/Windows/LaunchWindows.cpp
+++ b/Engine/Source/Runtime/Launch/Private/Windows/LaunchWindows.cpp
@@ -107,19 +107,6 @@ void SetupWindowsEnvironment( void )
 	_CrtSetDebugFillThreshold( 0 );
 #endif
 }
-
-/**
- * Fallback for catching exceptions which aren't caught elsewhere. This allows catching exceptions on threads created outside the engine.
- * Note that Windows does not call this handler if a debugger is attached, separately to our internal logic around crash handling.
- */
-LONG WINAPI UnhandledException(EXCEPTION_POINTERS *ExceptionInfo)
-{
-	ReportCrash(ExceptionInfo);
-	GIsCriticalError = true;
-	FPlatformMisc::RequestExit(true);
-	return EXCEPTION_CONTINUE_SEARCH;
-}
-
 /**
  * The inner exception handler catches crashes/asserts in native C++ code and is the only way to get the correct callstack
  * when running a 64-bit executable. However, XAudio2 doesn't always like this and it may result in no sound.
@@ -266,9 +253,6 @@ int32 WINAPI WinMain( _In_ HINSTANCE hInInstance, _In_opt_ HINSTANCE hPrevInstan
 	}
 	else
 	{
-		// Install an unhandled exception filter, to catch any exceptions on threads that are not created by the engine.
-		SetUnhandledExceptionFilter(UnhandledException);
-
 		// Use structured exception handling to trap any crashes, walk the the stack and display a crash dialog box.
 #if !PLATFORM_SEH_EXCEPTIONS_DISABLED
 		__try
diff --git a/Engine/Source/Runtime/Unix/UnixCommonStartup/Private/UnixCommonStartup.cpp b/Engine/Source/Runtime/Unix/UnixCommonStartup/Private/UnixCommonStartup.cpp
--- a/Engine/Source/Runtime/Unix/UnixCommonStartup/Private/UnixCommonStartup.cpp
+++ b/Engine/Source/Runtime/Unix/UnixCommonStartup/Private/UnixCommonStartup.cpp
@@ -28,6 +28,9 @@ void CommonUnixCrashHandler(const FGenericCrashContext& GenericContext)
 	const FUnixCrashContext& Context = static_cast< const FUnixCrashContext& >( GenericContext );
 	printf("CommonUnixCrashHandler: Signal=%d\n", Context.Signal);
 
+	/** This is the last place to gather memory stats */
+	FGenericCrashContext::SetMemoryStats(FPlatformMemory::GetStats());
+
 	// better than having mutable fields?
 	const_cast< FUnixCrashContext& >(Context).CaptureStackTrace();
 	if (GLog)
