Fix to the payload when binding a UFunction to a delegate with reference parameters.

diff --git a/Engine/Source/Runtime/Core/Public/Delegates/DelegateInstancesImpl.h b/Engine/Source/Runtime/Core/Public/Delegates/DelegateInstancesImpl.h
--- a/Engine/Source/Runtime/Core/Public/Delegates/DelegateInstancesImpl.h
+++ b/Engine/Source/Runtime/Core/Public/Delegates/DelegateInstancesImpl.h
@@ -113,7 +113,7 @@ public:
 
 	virtual RetValType Execute(ParamTypes... Params) const override final
 	{
-		typedef TPayload<RetValType (ParamTypes..., VarTypes...)> FParmsWithPayload;
+		typedef TPayload<RetValType (typename TDecay<ParamTypes>::Type..., typename TDecay<VarTypes> ::Type...)> FParmsWithPayload;
 
 		checkSlow(IsSafeToExecute());
 
