Fix some compile erros in WindowsPlatformCrashContext.cpp. TOptional changing signature and ambigous TCHAR arguments.

diff --git a/Engine/Source/Runtime/Core/Private/Windows/WindowsPlatformCrashContext.cpp b/Engine/Source/Runtime/Core/Private/Windows/WindowsPlatformCrashContext.cpp
--- a/Engine/Source/Runtime/Core/Private/Windows/WindowsPlatformCrashContext.cpp
+++ b/Engine/Source/Runtime/Core/Private/Windows/WindowsPlatformCrashContext.cpp
@@ -295,7 +295,7 @@ void FWindowsPlatformCrashContext::CopyPlatformSpecificFiles(const TCHAR* Output
 	LPEXCEPTION_POINTERS ExceptionInfo = (LPEXCEPTION_POINTERS)Context;
 	if (ExceptionInfo != nullptr)
 	{
-		const FString MinidumpFileName = FPaths::Combine(OutputDirectory, *FGenericCrashContext::UE4MinidumpName);
+		const FString MinidumpFileName = FPaths::Combine(OutputDirectory, FGenericCrashContext::UE4MinidumpName);
 		WriteMinidump(ProcessHandle.Get(), CrashedThreadId, *this, *MinidumpFileName, ExceptionInfo);
 	}
 
@@ -1169,7 +1169,7 @@ LONG WINAPI UnhandledStaticInitException(LPEXCEPTION_POINTERS ExceptionInfo)
 		// @note: Even if the object has been created, the actual thread has not been started yet,
 		// (that happens after static init) so we must bypass that and report directly from this thread. 
 		// 
-		if (GCrashReportingThread.IsValid())
+		if (GCrashReportingThread.IsSet())
 		{
 			return GCrashReportingThread->OnCrashDuringStaticInit(ExceptionInfo);
 		}
