Fixing UHT crash on startup (merge fallout)

diff --git a/Engine/Source/Runtime/Core/Private/GenericPlatform/GenericPlatformCrashContext.cpp b/Engine/Source/Runtime/Core/Private/GenericPlatform/GenericPlatformCrashContext.cpp
--- a/Engine/Source/Runtime/Core/Private/GenericPlatform/GenericPlatformCrashContext.cpp
+++ b/Engine/Source/Runtime/Core/Private/GenericPlatform/GenericPlatformCrashContext.cpp
@@ -140,7 +140,7 @@ void FGenericCrashContext::Initialize()
 	// Create a unique base guid for bug report ids
 	const FGuid Guid = FGuid::NewGuid();
 	const FString IniPlatformName(FPlatformProperties::IniPlatformName());
-	FCString::Strcpy(NCached::Session.CrashGUIDRoot, *FString::Printf(TEXT("%s%s-%s"), *CrashGUIDRootPrefix, *IniPlatformName, *Guid.ToString(EGuidFormats::Digits)));
+	FCString::Strcpy(NCached::Session.CrashGUIDRoot, *FString::Printf(TEXT("%s%s-%s"), CrashGUIDRootPrefix, *IniPlatformName, *Guid.ToString(EGuidFormats::Digits)));
 
 	if (GIsRunning)
 	{
