Fixed deprecation warnings.

diff --git a/Engine/Source/Runtime/Core/Private/Logging/LogSuppressionInterface.cpp b/Engine/Source/Runtime/Core/Private/Logging/LogSuppressionInterface.cpp
--- a/Engine/Source/Runtime/Core/Private/Logging/LogSuppressionInterface.cpp
+++ b/Engine/Source/Runtime/Core/Private/Logging/LogSuppressionInterface.cpp
@@ -136,7 +136,7 @@ class FLogSuppressionImplementation: public FLogSuppressionInterface, private FS
 			FString Command = SubCmds[Index].TrimStart();
 			if (Command.StartsWith(LogString))
 			{
-				Command = Command.Right(Command.Len() - (ARRAY_COUNT(LogString) - 1));
+				Command = Command.Right(Command.Len() - (UE_ARRAY_COUNT(LogString) - 1));
 			}
 			TArray<FString> CommandParts;
 			Command.ParseIntoArrayWS(CommandParts);
diff --git a/Engine/Source/Runtime/Core/Private/Misc/ConfigCacheIni.cpp b/Engine/Source/Runtime/Core/Private/Misc/ConfigCacheIni.cpp
--- a/Engine/Source/Runtime/Core/Private/Misc/ConfigCacheIni.cpp
+++ b/Engine/Source/Runtime/Core/Private/Misc/ConfigCacheIni.cpp
@@ -988,7 +988,7 @@ void FConfigFile::OverrideFromCommandline(FConfigFile* File, const FString& File
 				CommandlineOption.Section = CommandlineOption.Section.Replace(CommandlineOverrideSpecifiers::PropertyStartIdentifier, TEXT(""));
 				CommandlineOption.Section = CommandlineOption.Section.Replace(CommandlineOverrideSpecifiers::SectionStartIdentifier, TEXT(""));
 
-				CommandlineOption.PropertyKey = SectionAndKey.Mid(SectionNameEndIndex + ARRAY_COUNT(CommandlineOverrideSpecifiers::PropertyStartIdentifier) - 1);
+				CommandlineOption.PropertyKey = SectionAndKey.Mid(SectionNameEndIndex + UE_ARRAY_COUNT(CommandlineOverrideSpecifiers::PropertyStartIdentifier) - 1);
 				CommandlineOption.PropertyValue = Value;
 
 				// now put it into this into the cache
diff --git a/Engine/Source/Runtime/Core/Private/Windows/WindowsPlatformMisc.cpp b/Engine/Source/Runtime/Core/Private/Windows/WindowsPlatformMisc.cpp
--- a/Engine/Source/Runtime/Core/Private/Windows/WindowsPlatformMisc.cpp
+++ b/Engine/Source/Runtime/Core/Private/Windows/WindowsPlatformMisc.cpp
@@ -381,7 +381,7 @@ int32 FWindowsOSVersionHelper::GetOSVersions( FString& OutOSVersionLabel, FStrin
 	OSVersionLabel[0] = 0;
 	OSSubVersionLabel[0] = 0;
 
-	int32 Result = GetOSVersionsHelper( OSVersionLabel, ARRAY_COUNT(OSVersionLabel), OSSubVersionLabel, ARRAY_COUNT(OSSubVersionLabel) );
+	int32 Result = GetOSVersionsHelper( OSVersionLabel, UE_ARRAY_COUNT(OSVersionLabel), OSSubVersionLabel, UE_ARRAY_COUNT(OSSubVersionLabel) );
 
 	OutOSVersionLabel = OSVersionLabel;
 	OutOSSubVersionLabel = OSSubVersionLabel;
@@ -2527,7 +2527,7 @@ void FWindowsPlatformMisc::GetOSVersions( FString& OutOSVersionLabel, FString& O
 		{
 			OSVersionLabel[0] = 0;
 			OSSubVersionLabel[0] = 0;
-			GetOSVersionsHelper( OSVersionLabel, ARRAY_COUNT(OSVersionLabel), OSSubVersionLabel, ARRAY_COUNT(OSSubVersionLabel) );
+			GetOSVersionsHelper( OSVersionLabel, UE_ARRAY_COUNT(OSVersionLabel), OSSubVersionLabel, UE_ARRAY_COUNT(OSSubVersionLabel) );
 		}
 
 		TCHAR OSVersionLabel[128];
@@ -2546,7 +2546,7 @@ FString FWindowsPlatformMisc::GetOSVersion()
 		FOSVersionInitializer()
 		{
 			CachedOSVersion[0] = 0;
-			if (!GetOSVersionHelper(CachedOSVersion, ARRAY_COUNT(CachedOSVersion)))
+			if (!GetOSVersionHelper(CachedOSVersion, UE_ARRAY_COUNT(CachedOSVersion)))
 			{
 				CachedOSVersion[0] = 0;
 			}
diff --git a/Engine/Source/Runtime/Core/Private/Windows/WindowsPlatformProcess.cpp b/Engine/Source/Runtime/Core/Private/Windows/WindowsPlatformProcess.cpp
--- a/Engine/Source/Runtime/Core/Private/Windows/WindowsPlatformProcess.cpp
+++ b/Engine/Source/Runtime/Core/Private/Windows/WindowsPlatformProcess.cpp
@@ -1144,7 +1144,7 @@ const FString FWindowsPlatformProcess::GetModulesDirectory()
 		}
 
 		// Get the directory for it
-		GetModuleFileName(hCurrentModule, Result, ARRAY_COUNT(Result));
+		GetModuleFileName(hCurrentModule, Result, UE_ARRAY_COUNT(Result));
 		*FCString::Strrchr(Result, '\\') = 0;
 
 		// Normalize the resulting path
