EnterSlot() overload renamed to stop it being an overload.

diff --git a/Engine/Source/Runtime/Core/Private/Serialization/StructuredArchive.cpp b/Engine/Source/Runtime/Core/Private/Serialization/StructuredArchive.cpp
--- a/Engine/Source/Runtime/Core/Private/Serialization/StructuredArchive.cpp
+++ b/Engine/Source/Runtime/Core/Private/Serialization/StructuredArchive.cpp
@@ -140,7 +140,7 @@ void FStructuredArchive::EnterSlot(int32 ParentDepth, int32 ElementId)
 	CurrentEnteringAttributeState = EEnteringAttributeState::NotEnteringAttribute;
 }
 
-int32 FStructuredArchive::EnterSlot(int32 ParentDepth, int32 ElementId, EElementType ElementType)
+int32 FStructuredArchive::EnterSlotAsType(int32 ParentDepth, int32 ElementId, EElementType ElementType)
 {
 	EnterSlot(ParentDepth, ElementId);
 
@@ -260,7 +260,7 @@ void FStructuredArchive::SetScope(int32 Depth, int32 ElementId)
 
 FStructuredArchive::FRecord FStructuredArchive::FSlot::EnterRecord()
 {
-	int32 NewDepth = Ar.EnterSlot(Depth, ElementId, EElementType::Record);
+	int32 NewDepth = Ar.EnterSlotAsType(Depth, ElementId, EElementType::Record);
 
 #if DO_GUARD_SLOW
 	Ar.CurrentContainer.Add(new FContainer(0));
@@ -273,7 +273,7 @@ FStructuredArchive::FRecord FStructuredArchive::FSlot::EnterRecord()
 
 FStructuredArchive::FRecord FStructuredArchive::FSlot::EnterRecord_TextOnly(TArray<FString>& OutFieldNames)
 {
-	int32 NewDepth = Ar.EnterSlot(Depth, ElementId, EElementType::Record);
+	int32 NewDepth = Ar.EnterSlotAsType(Depth, ElementId, EElementType::Record);
 
 #if DO_GUARD_SLOW
 	Ar.CurrentContainer.Add(new FContainer(0));
@@ -286,7 +286,7 @@ FStructuredArchive::FRecord FStructuredArchive::FSlot::EnterRecord_TextOnly(TArr
 
 FStructuredArchive::FArray FStructuredArchive::FSlot::EnterArray(int32& Num)
 {
-	int32 NewDepth = Ar.EnterSlot(Depth, ElementId, EElementType::Array);
+	int32 NewDepth = Ar.EnterSlotAsType(Depth, ElementId, EElementType::Array);
 
 	Ar.Formatter.EnterArray(Num);
 
@@ -299,7 +299,7 @@ FStructuredArchive::FArray FStructuredArchive::FSlot::EnterArray(int32& Num)
 
 FStructuredArchive::FStream FStructuredArchive::FSlot::EnterStream()
 {
-	int32 NewDepth = Ar.EnterSlot(Depth, ElementId, EElementType::Stream);
+	int32 NewDepth = Ar.EnterSlotAsType(Depth, ElementId, EElementType::Stream);
 
 	Ar.Formatter.EnterStream();
 
@@ -308,7 +308,7 @@ FStructuredArchive::FStream FStructuredArchive::FSlot::EnterStream()
 
 FStructuredArchive::FStream FStructuredArchive::FSlot::EnterStream_TextOnly(int32& OutNumElements)
 {
-	int32 NewDepth = Ar.EnterSlot(Depth, ElementId, EElementType::Stream);
+	int32 NewDepth = Ar.EnterSlotAsType(Depth, ElementId, EElementType::Stream);
 
 	Ar.Formatter.EnterStream_TextOnly(OutNumElements);
 
@@ -317,7 +317,7 @@ FStructuredArchive::FStream FStructuredArchive::FSlot::EnterStream_TextOnly(int3
 
 FStructuredArchive::FMap FStructuredArchive::FSlot::EnterMap(int32& Num)
 {
-	int32 NewDepth = Ar.EnterSlot(Depth, ElementId, EElementType::Map);
+	int32 NewDepth = Ar.EnterSlotAsType(Depth, ElementId, EElementType::Map);
 
 	Ar.Formatter.EnterMap(Num);
 
@@ -335,7 +335,7 @@ FStructuredArchive::FSlot FStructuredArchive::FSlot::EnterAttribute(FArchiveFiel
 	int32 NewDepth = Depth + 1;
 	if (NewDepth >= Ar.CurrentScope.Num() || Ar.CurrentScope[NewDepth].Id != ElementId || Ar.CurrentScope[NewDepth].Type != EElementType::AttributedValue)
 	{
-		int32 NewDepthCheck = Ar.EnterSlot(Depth, ElementId, EElementType::AttributedValue);
+		int32 NewDepthCheck = Ar.EnterSlotAsType(Depth, ElementId, EElementType::AttributedValue);
 		checkSlow(NewDepth == NewDepthCheck);
 
 		Ar.Formatter.EnterAttributedValue();
@@ -376,7 +376,7 @@ TOptional<FStructuredArchive::FSlot> FStructuredArchive::FSlot::TryEnterAttribut
 	int32 NewDepth = Depth + 1;
 	if (NewDepth >= Ar.CurrentScope.Num() || Ar.CurrentScope[NewDepth].Id != ElementId || Ar.CurrentScope[NewDepth].Type != EElementType::AttributedValue)
 	{
-		int32 NewDepthCheck = Ar.EnterSlot(Depth, ElementId, EElementType::AttributedValue);
+		int32 NewDepthCheck = Ar.EnterSlotAsType(Depth, ElementId, EElementType::AttributedValue);
 		checkSlow(NewDepth == NewDepthCheck);
 
 		Ar.Formatter.EnterAttributedValue();
diff --git a/Engine/Source/Runtime/Core/Public/Serialization/StructuredArchive.h b/Engine/Source/Runtime/Core/Public/Serialization/StructuredArchive.h
--- a/Engine/Source/Runtime/Core/Public/Serialization/StructuredArchive.h
+++ b/Engine/Source/Runtime/Core/Public/Serialization/StructuredArchive.h
@@ -441,7 +441,7 @@ private:
 	 *
 	 * @return  The depth of the newly-entered slot.
 	 */
-	int32 EnterSlot(int32 ParentDepth, int32 ElementId, EElementType ElementType);
+	int32 EnterSlotAsType(int32 ParentDepth, int32 ElementId, EElementType ElementType);
 
 	/**
 	 * Leaves slot at the top of the current scope
