Element IDs have their own type.

diff --git a/Engine/Source/Runtime/Core/Private/Serialization/StructuredArchive.cpp b/Engine/Source/Runtime/Core/Private/Serialization/StructuredArchive.cpp
--- a/Engine/Source/Runtime/Core/Private/Serialization/StructuredArchive.cpp
+++ b/Engine/Source/Runtime/Core/Private/Serialization/StructuredArchive.cpp
@@ -6,14 +6,6 @@
 
 #if WITH_TEXT_ARCHIVE_SUPPORT
 
-namespace 
-{
-	/**
-	 * First element ID, which is assigned to the root element.
-	 */
-	static const int RootElementId = 1000;
-}
-
 //////////// FStructuredArchive::FContainer ////////////
 
 #if DO_GUARD_SLOW
@@ -77,8 +69,6 @@ FStructuredArchive::FStructuredArchive(FArchiveFormatterType& InFormatter)
 #else
 	, bRequiresStructuralMetadata(InFormatter.HasDocumentTree())
 #endif
-	, NextElementId(RootElementId + 1)
-	, CurrentSlotElementId(INDEX_NONE)
 {
 	CurrentScope.Reserve(32);
 #if DO_GUARD_SLOW
@@ -101,12 +91,13 @@ FStructuredArchive::~FStructuredArchive()
 FStructuredArchive::FSlot FStructuredArchive::Open()
 {
 	check(CurrentScope.Num() == 0);
-	check(NextElementId == RootElementId + 1);
-	check(CurrentSlotElementId == INDEX_NONE);
+	check(!RootElementId.IsValid());
+	check(!CurrentSlotElementId.IsValid());
 
+	RootElementId = ElementIdGenerator.Generate();
 	CurrentScope.Emplace(RootElementId, EElementType::Root);
 
-	CurrentSlotElementId = NextElementId++;
+	CurrentSlotElementId = ElementIdGenerator.Generate();
 
 	return FSlot(*this, 0, CurrentSlotElementId);
 }
@@ -116,16 +107,16 @@ void FStructuredArchive::Close()
 	SetScope(0, RootElementId);
 }
 
-void FStructuredArchive::EnterSlot(int32 ParentDepth, int32 ElementId)
+void FStructuredArchive::EnterSlot(int32 ParentDepth, FElementId ElementId)
 {
 	// If the slot being entered has attributes, enter the value slot first.
 	if (ParentDepth + 1 < CurrentScope.Num() && CurrentScope[ParentDepth + 1].Id == ElementId && CurrentScope[ParentDepth + 1].Type == EElementType::AttributedValue)
 	{
 #if DO_GUARD_SLOW
-		checkf(CurrentSlotElementId == INDEX_NONE && !CurrentContainer.Top()->bAttributedValueWritten, TEXT("Attempt to serialize data into an invalid slot"));
+		checkf(!CurrentSlotElementId.IsValid() && !CurrentContainer.Top()->bAttributedValueWritten, TEXT("Attempt to serialize data into an invalid slot"));
 		CurrentContainer.Top()->bAttributedValueWritten = true;
 #else
-		checkf(CurrentSlotElementId == INDEX_NONE, TEXT("Attempt to serialize data into an invalid slot"));
+		checkf(!CurrentSlotElementId.IsValid(), TEXT("Attempt to serialize data into an invalid slot"));
 #endif
 
 		SetScope(ParentDepth + 1, ElementId);
@@ -134,13 +125,13 @@ void FStructuredArchive::EnterSlot(int32 ParentDepth, int32 ElementId)
 	else
 	{
 		checkf(ElementId == CurrentSlotElementId, TEXT("Attempt to serialize data into an invalid slot"));
-		CurrentSlotElementId = INDEX_NONE;
+		CurrentSlotElementId.Reset();
 	}
 
 	CurrentEnteringAttributeState = EEnteringAttributeState::NotEnteringAttribute;
 }
 
-int32 FStructuredArchive::EnterSlotAsType(int32 ParentDepth, int32 ElementId, EElementType ElementType)
+int32 FStructuredArchive::EnterSlotAsType(int32 ParentDepth, FElementId ElementId, EElementType ElementType)
 {
 	EnterSlot(ParentDepth, ElementId);
 
@@ -192,11 +183,11 @@ void FStructuredArchive::LeaveSlot()
 	}
 }
 
-void FStructuredArchive::SetScope(int32 Depth, int32 ElementId)
+void FStructuredArchive::SetScope(int32 Depth, FElementId ElementId)
 {
 	// Make sure the scope is valid
 	checkf(Depth < CurrentScope.Num() && CurrentScope[Depth].Id == ElementId, TEXT("Invalid scope for writing to archive"));
-	checkf(CurrentSlotElementId == INDEX_NONE || GetUnderlyingArchive().IsLoading(), TEXT("Cannot change scope until having written a value to the current slot"));
+	checkf(!CurrentSlotElementId.IsValid() || GetUnderlyingArchive().IsLoading(), TEXT("Cannot change scope until having written a value to the current slot"));
 
 	// Roll back to the correct scope
 	if (bRequiresStructuralMetadata)
@@ -347,11 +338,11 @@ FStructuredArchive::FSlot FStructuredArchive::FSlot::EnterAttribute(FArchiveFiel
 
 	Ar.CurrentEnteringAttributeState = EEnteringAttributeState::NotEnteringAttribute;
 
-	int32 AttributedValueId = Ar.CurrentScope[NewDepth].Id;
+	FElementId AttributedValueId = Ar.CurrentScope[NewDepth].Id;
 
 	Ar.SetScope(NewDepth, AttributedValueId);
 
-	Ar.CurrentSlotElementId = Ar.NextElementId++;
+	Ar.CurrentSlotElementId = Ar.ElementIdGenerator.Generate();
 
 #if DO_GUARD_SLOW
 #if CHECK_UNIQUE_FIELD_NAMES
@@ -397,7 +388,7 @@ TOptional<FStructuredArchive::FSlot> FStructuredArchive::FSlot::TryEnterAttribut
 #endif
 #endif
 
-	int32 AttributedValueId = Ar.CurrentScope[NewDepth].Id;
+	FElementId AttributedValueId = Ar.CurrentScope[NewDepth].Id;
 
 	Ar.SetScope(NewDepth, AttributedValueId);
 
@@ -405,7 +396,7 @@ TOptional<FStructuredArchive::FSlot> FStructuredArchive::FSlot::TryEnterAttribut
 	{
 		Ar.CurrentEnteringAttributeState = EEnteringAttributeState::NotEnteringAttribute;
 
-		Ar.CurrentSlotElementId = Ar.NextElementId++;
+		Ar.CurrentSlotElementId = Ar.ElementIdGenerator.Generate();
 
 		return FSlot(Ar, NewDepth, Ar.CurrentSlotElementId);
 	}
@@ -568,7 +559,7 @@ FStructuredArchive::FSlot FStructuredArchive::FRecord::EnterField(FArchiveFieldN
 {
 	Ar.SetScope(Depth, ElementId);
 
-	Ar.CurrentSlotElementId = Ar.NextElementId++;
+	Ar.CurrentSlotElementId = Ar.ElementIdGenerator.Generate();
 
 #if DO_GUARD_SLOW
 #if CHECK_UNIQUE_FIELD_NAMES
@@ -633,7 +624,7 @@ TOptional<FStructuredArchive::FSlot> FStructuredArchive::FRecord::TryEnterField(
 
 	if (Ar.Formatter.TryEnterField(Name, bEnterWhenWriting))
 	{
-		Ar.CurrentSlotElementId = Ar.NextElementId++;
+		Ar.CurrentSlotElementId = Ar.ElementIdGenerator.Generate();
 		return FSlot(Ar, Depth, Ar.CurrentSlotElementId);
 	}
 	else
@@ -652,7 +643,7 @@ FStructuredArchive::FSlot FStructuredArchive::FArray::EnterElement()
 	checkf(Ar.CurrentContainer.Top()->Index < Ar.CurrentContainer.Top()->Count, TEXT("Serialized too many array elements"));
 #endif
 
-	Ar.CurrentSlotElementId = Ar.NextElementId++;
+	Ar.CurrentSlotElementId = Ar.ElementIdGenerator.Generate();
 
 	Ar.Formatter.EnterArrayElement();
 
@@ -667,7 +658,7 @@ FStructuredArchive::FSlot FStructuredArchive::FArray::EnterElement_TextOnly(EArc
 	checkf(Ar.CurrentContainer.Top()->Index < Ar.CurrentContainer.Top()->Count, TEXT("Serialized too many array elements"));
 #endif
 
-	Ar.CurrentSlotElementId = Ar.NextElementId++;
+	Ar.CurrentSlotElementId = Ar.ElementIdGenerator.Generate();
 
 	Ar.Formatter.EnterArrayElement_TextOnly(OutType);
 
@@ -680,7 +671,7 @@ FStructuredArchive::FSlot FStructuredArchive::FStream::EnterElement()
 {
 	Ar.SetScope(Depth, ElementId);
 
-	Ar.CurrentSlotElementId = Ar.NextElementId++;
+	Ar.CurrentSlotElementId = Ar.ElementIdGenerator.Generate();
 
 	Ar.Formatter.EnterStreamElement();
 
@@ -691,7 +682,7 @@ FStructuredArchive::FSlot FStructuredArchive::FStream::EnterElement_TextOnly(EAr
 {
 	Ar.SetScope(Depth, ElementId);
 
-	Ar.CurrentSlotElementId = Ar.NextElementId++;
+	Ar.CurrentSlotElementId = Ar.ElementIdGenerator.Generate();
 
 	Ar.Formatter.EnterStreamElement_TextOnly(OutType);
 
@@ -708,7 +699,7 @@ FStructuredArchive::FSlot FStructuredArchive::FMap::EnterElement(FString& Name)
 	checkf(Ar.CurrentContainer.Top()->Index < Ar.CurrentContainer.Top()->Count, TEXT("Serialized too many map elements"));
 #endif
 
-	Ar.CurrentSlotElementId = Ar.NextElementId++;
+	Ar.CurrentSlotElementId = Ar.ElementIdGenerator.Generate();
 
 #if DO_GUARD_SLOW
 #if CHECK_UNIQUE_FIELD_NAMES
@@ -745,7 +736,7 @@ FStructuredArchive::FSlot FStructuredArchive::FMap::EnterElement_TextOnly(FStrin
 	checkf(Ar.CurrentContainer.Top()->Index < Ar.CurrentContainer.Top()->Count, TEXT("Serialized too many map elements"));
 #endif
 
-	Ar.CurrentSlotElementId = Ar.NextElementId++;
+	Ar.CurrentSlotElementId = Ar.ElementIdGenerator.Generate();
 
 #if DO_GUARD_SLOW
 #if CHECK_UNIQUE_FIELD_NAMES
diff --git a/Engine/Source/Runtime/Core/Public/Serialization/StructuredArchive.h b/Engine/Source/Runtime/Core/Public/Serialization/StructuredArchive.h
--- a/Engine/Source/Runtime/Core/Public/Serialization/StructuredArchive.h
+++ b/Engine/Source/Runtime/Core/Public/Serialization/StructuredArchive.h
@@ -161,6 +161,39 @@ public:
 	FStructuredArchive& operator=(const FStructuredArchive&) = delete;
 
 private:
+	struct FElementId
+	{
+		FElementId() = default;
+
+		explicit FElementId(uint32 InId)
+			: Id(InId)
+		{
+		}
+
+		bool IsValid() const
+		{
+			return Id != 0;
+		}
+
+		void Reset()
+		{
+			Id = 0;
+		}
+
+		bool operator==(const FElementId& Rhs) const
+		{
+			return Id == Rhs.Id;
+		}
+
+		bool operator!=(const FElementId& Rhs) const
+		{
+			return Id != Rhs.Id;
+		}
+
+	private:
+		uint32 Id = 0;
+	};
+
 	class CORE_API FSlotBase
 	{
 	public:
@@ -173,9 +206,9 @@ private:
 		FStructuredArchive& Ar;
 #if WITH_TEXT_ARCHIVE_SUPPORT
 		const int32 Depth;
-		const int32 ElementId;
+		const FElementId ElementId;
 
-		FORCEINLINE explicit FSlotBase(FStructuredArchive& InAr, int32 InDepth, int32 InElementId)
+		FORCEINLINE explicit FSlotBase(FStructuredArchive& InAr, int32 InDepth, FElementId InElementId)
 			: Ar(InAr)
 			, Depth(InDepth)
 			, ElementId(InElementId)
@@ -389,10 +422,10 @@ private:
 
 	struct FElement
 	{
-		int Id;
+		FElementId Id;
 		EElementType Type;
 
-		FElement(int InId, EElementType InType)
+		FElement(FElementId InId, EElementType InType)
 			: Id(InId)
 			, Type(InType)
 		{
@@ -403,15 +436,31 @@ private:
 	struct FContainer;
 #endif
 
+	struct FIdGenerator
+	{
+		FElementId Generate()
+		{
+			return FElementId(NextId++);
+		}
+
+	private:
+		uint32 NextId = 1;
+	};
+
 	/**
 	 * The next element id to be assigned
 	 */
-	int NextElementId;
+	FIdGenerator ElementIdGenerator;
+
+	/**
+	 * The ID of the root element.
+	 */
+	FElementId RootElementId;
 
 	/**
-	 * The element ID assigned for the current slot. Slots are transient, and only exist as placeholders until something is written into them. This is reset to INDEX_NONE when something is created in a slot, and the created item can assume the element id.
+	 * The element ID assigned for the current slot. Slots are transient, and only exist as placeholders until something is written into them. This is reset to 0 when something is created in a slot, and the created item can assume the element id.
 	 */
-	int CurrentSlotElementId;
+	FElementId CurrentSlotElementId;
 
 	/**
 	 * Tracks the current stack of objects being written. Used by SetScope() to ensure that scopes are always closed correctly in the underlying formatter,
@@ -434,14 +483,14 @@ private:
 	/**
 	 * Enters the current slot for serializing a value. Asserts if the archive is not in a state about to write to an empty-slot.
 	 */
-	void EnterSlot(int32 ParentDepth, int32 ElementId);
+	void EnterSlot(int32 ParentDepth, FElementId ElementId);
 
 	/**
 	 * Enters the current slot, adding an element onto the stack. Asserts if the archive is not in a state about to write to an empty-slot.
 	 *
 	 * @return  The depth of the newly-entered slot.
 	 */
-	int32 EnterSlotAsType(int32 ParentDepth, int32 ElementId, EElementType ElementType);
+	int32 EnterSlotAsType(int32 ParentDepth, FElementId ElementId, EElementType ElementType);
 
 	/**
 	 * Leaves slot at the top of the current scope
@@ -451,7 +500,7 @@ private:
 	/**
 	 * Switches to the scope at the given element id, updating the formatter state as necessary
 	 */
-	void SetScope(int32 Depth, int32 ElementId);
+	void SetScope(int32 Depth, FElementId ElementId);
 #endif
 };
 
