Fixed a couple of CI no-pch build failures

#rnx

diff --git a/Engine/Source/Runtime/Core/Private/GenericPlatform/GenericPlatformCrashContext.cpp b/Engine/Source/Runtime/Core/Private/GenericPlatform/GenericPlatformCrashContext.cpp
--- a/Engine/Source/Runtime/Core/Private/GenericPlatform/GenericPlatformCrashContext.cpp
+++ b/Engine/Source/Runtime/Core/Private/GenericPlatform/GenericPlatformCrashContext.cpp
@@ -18,6 +18,7 @@
 #include "Misc/App.h"
 #include "Misc/EngineVersion.h"
 #include "Misc/EngineBuildSettings.h"
+#include "Misc/OutputDeviceRedirector.h"
 #include "Stats/Stats.h"
 #include "Internationalization/TextLocalizationManager.h"
 #include "Logging/LogScopedCategoryAndVerbosityOverride.h"
diff --git a/Engine/Source/Runtime/Core/Public/GenericPlatform/GenericPlatformCrashContext.h b/Engine/Source/Runtime/Core/Public/GenericPlatform/GenericPlatformCrashContext.h
--- a/Engine/Source/Runtime/Core/Public/GenericPlatform/GenericPlatformCrashContext.h
+++ b/Engine/Source/Runtime/Core/Public/GenericPlatform/GenericPlatformCrashContext.h
@@ -4,6 +4,7 @@
 
 #include "CoreTypes.h"
 #include "HAL/PlatformMemory.h"
+#include "HAL/PlatformProcess.h"
 #include "Containers/UnrealString.h"
 
 struct FProgramCounterSymbolInfo;
diff --git a/Engine/Source/Runtime/CoreUObject/Private/UObject/ObjectResource.cpp b/Engine/Source/Runtime/CoreUObject/Private/UObject/ObjectResource.cpp
--- a/Engine/Source/Runtime/CoreUObject/Private/UObject/ObjectResource.cpp
+++ b/Engine/Source/Runtime/CoreUObject/Private/UObject/ObjectResource.cpp
@@ -2,6 +2,7 @@
 
 #include "UObject/ObjectResource.h"
 #include "UObject/Class.h"
+#include "Templates/Casts.h"
 
 /*-----------------------------------------------------------------------------
 	Helper functions.
