Guarded access to GCrashReportingThread.

diff --git a/Engine/Source/Runtime/Core/Private/Windows/WindowsPlatformCrashContext.cpp b/Engine/Source/Runtime/Core/Private/Windows/WindowsPlatformCrashContext.cpp
--- a/Engine/Source/Runtime/Core/Private/Windows/WindowsPlatformCrashContext.cpp
+++ b/Engine/Source/Runtime/Core/Private/Windows/WindowsPlatformCrashContext.cpp
@@ -1159,6 +1159,7 @@ TOptional<FCrashReportingThread> GCrashReportingThread(InPlace);
 
 LONG WINAPI UnhandledStaticInitException(LPEXCEPTION_POINTERS ExceptionInfo)
 {
+#if !NOINITCRASHREPORTER
 	// Top bit in exception code is fatal exceptions. Report those but not other types.
 	if ((ExceptionInfo->ExceptionRecord->ExceptionCode & 0x80000000L) != 0)
 	{
@@ -1174,6 +1175,7 @@ LONG WINAPI UnhandledStaticInitException(LPEXCEPTION_POINTERS ExceptionInfo)
 			return GCrashReportingThread->OnCrashDuringStaticInit(ExceptionInfo);
 		}
 	}
+#endif
 	
 	return EXCEPTION_CONTINUE_SEARCH;
 }
@@ -1220,7 +1222,11 @@ static bool bReentranceGuard = false;
  */
 int32 ReportEnsureUsingCrashReportClient(EXCEPTION_POINTERS* ExceptionInfo, int NumStackFramesToIgnore, const TCHAR* ErrorMessage, EErrorReportUI ReportUI)
 {
+#if !NOINITCRASHREPORTER
 	return GCrashReportingThread->OnEnsure(ExceptionInfo, NumStackFramesToIgnore, ErrorMessage, ReportUI);
+#else 
+	return EXCEPTION_EXECUTE_HANDLER;
+#endif
 }
 #endif
 
