Slot position type for representing a depth/ID pair of a slot within the hierarchy.

diff --git a/Engine/Source/Runtime/Core/Private/Serialization/StructuredArchive.cpp b/Engine/Source/Runtime/Core/Private/Serialization/StructuredArchive.cpp
--- a/Engine/Source/Runtime/Core/Private/Serialization/StructuredArchive.cpp
+++ b/Engine/Source/Runtime/Core/Private/Serialization/StructuredArchive.cpp
@@ -104,11 +104,14 @@ FStructuredArchive::FSlot FStructuredArchive::Open()
 
 void FStructuredArchive::Close()
 {
-	SetScope(0, RootElementId);
+	SetScope(FSlotPosition(0, RootElementId));
 }
 
-void FStructuredArchive::EnterSlot(int32 ParentDepth, FElementId ElementId)
+void FStructuredArchive::EnterSlot(FSlotPosition Slot)
 {
+	int32      ParentDepth = Slot.Depth;
+	FElementId ElementId   = Slot.ElementId;
+
 	// If the slot being entered has attributes, enter the value slot first.
 	if (ParentDepth + 1 < CurrentScope.Num() && CurrentScope[ParentDepth + 1].Id == ElementId && CurrentScope[ParentDepth + 1].Type == EElementType::AttributedValue)
 	{
@@ -119,7 +122,7 @@ void FStructuredArchive::EnterSlot(int32 ParentDepth, FElementId ElementId)
 		checkf(!CurrentSlotElementId.IsValid(), TEXT("Attempt to serialize data into an invalid slot"));
 #endif
 
-		SetScope(ParentDepth + 1, ElementId);
+		SetScope(FSlotPosition(ParentDepth + 1, ElementId));
 		Formatter.EnterAttributedValueValue();
 	}
 	else
@@ -131,25 +134,25 @@ void FStructuredArchive::EnterSlot(int32 ParentDepth, FElementId ElementId)
 	CurrentEnteringAttributeState = EEnteringAttributeState::NotEnteringAttribute;
 }
 
-int32 FStructuredArchive::EnterSlotAsType(int32 ParentDepth, FElementId ElementId, EElementType ElementType)
+int32 FStructuredArchive::EnterSlotAsType(FSlotPosition Slot, EElementType ElementType)
 {
-	EnterSlot(ParentDepth, ElementId);
+	EnterSlot(Slot);
 
-	int32 Result = ParentDepth + 1;
+	int32 NewSlotDepth = Slot.Depth + 1;
 
 	// If we're entering the value of an attributed slot, we need to return a depth one higher than usual, because we're
 	// inside an attributed value container.
 	//
 	// We don't need to do adjust for attributes, because entering the attribute slot will bump the depth anyway.
-	if (ParentDepth + 1 < CurrentScope.Num() &&
-		CurrentScope[ParentDepth + 1].Type == EElementType::AttributedValue &&
+	if (NewSlotDepth < CurrentScope.Num() &&
+		CurrentScope[NewSlotDepth].Type == EElementType::AttributedValue &&
 		CurrentEnteringAttributeState == EEnteringAttributeState::NotEnteringAttribute)
 	{
-		++Result;
+		++NewSlotDepth;
 	}
 
-	CurrentScope.Emplace(ElementId, ElementType);
-	return Result;
+	CurrentScope.Emplace(Slot.ElementId, ElementType);
+	return NewSlotDepth;
 }
 
 void FStructuredArchive::LeaveSlot()
@@ -183,16 +186,16 @@ void FStructuredArchive::LeaveSlot()
 	}
 }
 
-void FStructuredArchive::SetScope(int32 Depth, FElementId ElementId)
+void FStructuredArchive::SetScope(FSlotPosition Slot)
 {
 	// Make sure the scope is valid
-	checkf(Depth < CurrentScope.Num() && CurrentScope[Depth].Id == ElementId, TEXT("Invalid scope for writing to archive"));
+	checkf(Slot.Depth < CurrentScope.Num() && CurrentScope[Slot.Depth].Id == Slot.ElementId, TEXT("Invalid scope for writing to archive"));
 	checkf(!CurrentSlotElementId.IsValid() || GetUnderlyingArchive().IsLoading(), TEXT("Cannot change scope until having written a value to the current slot"));
 
 	// Roll back to the correct scope
 	if (bRequiresStructuralMetadata)
 	{
-		for (int32 CurrentDepth = CurrentScope.Num() - 1; CurrentDepth > Depth; CurrentDepth--)
+		for (int32 CurrentDepth = CurrentScope.Num() - 1; CurrentDepth > Slot.Depth; CurrentDepth--)
 		{
 			// Leave the current element
 			const FElement& Element = CurrentScope[CurrentDepth];
@@ -243,7 +246,7 @@ void FStructuredArchive::SetScope(int32 Depth, FElementId ElementId)
 	else
 	{
 		// Remove all the top elements from the stack
-		CurrentScope.RemoveAt(Depth + 1, CurrentScope.Num() - (Depth + 1));
+		CurrentScope.RemoveAt(Slot.Depth + 1, CurrentScope.Num() - (Slot.Depth + 1));
 	}
 }
 
@@ -251,7 +254,7 @@ void FStructuredArchive::SetScope(int32 Depth, FElementId ElementId)
 
 FStructuredArchive::FRecord FStructuredArchive::FSlot::EnterRecord()
 {
-	int32 NewDepth = Ar.EnterSlotAsType(Depth, ElementId, EElementType::Record);
+	int32 NewDepth = Ar.EnterSlotAsType(*this, EElementType::Record);
 
 #if DO_GUARD_SLOW
 	Ar.CurrentContainer.Add(new FContainer(0));
@@ -264,7 +267,7 @@ FStructuredArchive::FRecord FStructuredArchive::FSlot::EnterRecord()
 
 FStructuredArchive::FRecord FStructuredArchive::FSlot::EnterRecord_TextOnly(TArray<FString>& OutFieldNames)
 {
-	int32 NewDepth = Ar.EnterSlotAsType(Depth, ElementId, EElementType::Record);
+	int32 NewDepth = Ar.EnterSlotAsType(*this, EElementType::Record);
 
 #if DO_GUARD_SLOW
 	Ar.CurrentContainer.Add(new FContainer(0));
@@ -277,7 +280,7 @@ FStructuredArchive::FRecord FStructuredArchive::FSlot::EnterRecord_TextOnly(TArr
 
 FStructuredArchive::FArray FStructuredArchive::FSlot::EnterArray(int32& Num)
 {
-	int32 NewDepth = Ar.EnterSlotAsType(Depth, ElementId, EElementType::Array);
+	int32 NewDepth = Ar.EnterSlotAsType(*this, EElementType::Array);
 
 	Ar.Formatter.EnterArray(Num);
 
@@ -290,7 +293,7 @@ FStructuredArchive::FArray FStructuredArchive::FSlot::EnterArray(int32& Num)
 
 FStructuredArchive::FStream FStructuredArchive::FSlot::EnterStream()
 {
-	int32 NewDepth = Ar.EnterSlotAsType(Depth, ElementId, EElementType::Stream);
+	int32 NewDepth = Ar.EnterSlotAsType(*this, EElementType::Stream);
 
 	Ar.Formatter.EnterStream();
 
@@ -299,7 +302,7 @@ FStructuredArchive::FStream FStructuredArchive::FSlot::EnterStream()
 
 FStructuredArchive::FStream FStructuredArchive::FSlot::EnterStream_TextOnly(int32& OutNumElements)
 {
-	int32 NewDepth = Ar.EnterSlotAsType(Depth, ElementId, EElementType::Stream);
+	int32 NewDepth = Ar.EnterSlotAsType(*this, EElementType::Stream);
 
 	Ar.Formatter.EnterStream_TextOnly(OutNumElements);
 
@@ -308,7 +311,7 @@ FStructuredArchive::FStream FStructuredArchive::FSlot::EnterStream_TextOnly(int3
 
 FStructuredArchive::FMap FStructuredArchive::FSlot::EnterMap(int32& Num)
 {
-	int32 NewDepth = Ar.EnterSlotAsType(Depth, ElementId, EElementType::Map);
+	int32 NewDepth = Ar.EnterSlotAsType(*this, EElementType::Map);
 
 	Ar.Formatter.EnterMap(Num);
 
@@ -326,7 +329,7 @@ FStructuredArchive::FSlot FStructuredArchive::FSlot::EnterAttribute(FArchiveFiel
 	int32 NewDepth = Depth + 1;
 	if (NewDepth >= Ar.CurrentScope.Num() || Ar.CurrentScope[NewDepth].Id != ElementId || Ar.CurrentScope[NewDepth].Type != EElementType::AttributedValue)
 	{
-		int32 NewDepthCheck = Ar.EnterSlotAsType(Depth, ElementId, EElementType::AttributedValue);
+		int32 NewDepthCheck = Ar.EnterSlotAsType(*this, EElementType::AttributedValue);
 		checkSlow(NewDepth == NewDepthCheck);
 
 		Ar.Formatter.EnterAttributedValue();
@@ -340,7 +343,7 @@ FStructuredArchive::FSlot FStructuredArchive::FSlot::EnterAttribute(FArchiveFiel
 
 	FElementId AttributedValueId = Ar.CurrentScope[NewDepth].Id;
 
-	Ar.SetScope(NewDepth, AttributedValueId);
+	Ar.SetScope(FSlotPosition(NewDepth, AttributedValueId));
 
 	Ar.CurrentSlotElementId = Ar.ElementIdGenerator.Generate();
 
@@ -367,7 +370,7 @@ TOptional<FStructuredArchive::FSlot> FStructuredArchive::FSlot::TryEnterAttribut
 	int32 NewDepth = Depth + 1;
 	if (NewDepth >= Ar.CurrentScope.Num() || Ar.CurrentScope[NewDepth].Id != ElementId || Ar.CurrentScope[NewDepth].Type != EElementType::AttributedValue)
 	{
-		int32 NewDepthCheck = Ar.EnterSlotAsType(Depth, ElementId, EElementType::AttributedValue);
+		int32 NewDepthCheck = Ar.EnterSlotAsType(*this, EElementType::AttributedValue);
 		checkSlow(NewDepth == NewDepthCheck);
 
 		Ar.Formatter.EnterAttributedValue();
@@ -390,7 +393,7 @@ TOptional<FStructuredArchive::FSlot> FStructuredArchive::FSlot::TryEnterAttribut
 
 	FElementId AttributedValueId = Ar.CurrentScope[NewDepth].Id;
 
-	Ar.SetScope(NewDepth, AttributedValueId);
+	Ar.SetScope(FSlotPosition(NewDepth, AttributedValueId));
 
 	if (Ar.Formatter.TryEnterAttribute(AttributeName, bEnterWhenWriting))
 	{
@@ -408,147 +411,147 @@ TOptional<FStructuredArchive::FSlot> FStructuredArchive::FSlot::TryEnterAttribut
 
 void FStructuredArchive::FSlot::operator<< (uint8& Value)
 {
-	Ar.EnterSlot(Depth, ElementId);
+	Ar.EnterSlot(*this);
 	Ar.Formatter.Serialize(Value);
 	Ar.LeaveSlot();
 }
 
 void FStructuredArchive::FSlot::operator<< (uint16& Value)
 {
-	Ar.EnterSlot(Depth, ElementId);
+	Ar.EnterSlot(*this);
 	Ar.Formatter.Serialize(Value);
 	Ar.LeaveSlot();
 }
 
 void FStructuredArchive::FSlot::operator<< (uint32& Value)
 {
-	Ar.EnterSlot(Depth, ElementId);
+	Ar.EnterSlot(*this);
 	Ar.Formatter.Serialize(Value);
 	Ar.LeaveSlot();
 }
 
 void FStructuredArchive::FSlot::operator<< (uint64& Value)
 {
-	Ar.EnterSlot(Depth, ElementId);
+	Ar.EnterSlot(*this);
 	Ar.Formatter.Serialize(Value);
 	Ar.LeaveSlot();
 }
 
 void FStructuredArchive::FSlot::operator<< (int8& Value)
 {
-	Ar.EnterSlot(Depth, ElementId);
+	Ar.EnterSlot(*this);
 	Ar.Formatter.Serialize(Value);
 	Ar.LeaveSlot();
 }
 
 void FStructuredArchive::FSlot::operator<< (int16& Value)
 {
-	Ar.EnterSlot(Depth, ElementId);
+	Ar.EnterSlot(*this);
 	Ar.Formatter.Serialize(Value);
 	Ar.LeaveSlot();
 }
 
 void FStructuredArchive::FSlot::operator<< (int32& Value)
 {
-	Ar.EnterSlot(Depth, ElementId);
+	Ar.EnterSlot(*this);
 	Ar.Formatter.Serialize(Value);
 	Ar.LeaveSlot();
 }
 
 void FStructuredArchive::FSlot::operator<< (int64& Value)
 {
-	Ar.EnterSlot(Depth, ElementId);
+	Ar.EnterSlot(*this);
 	Ar.Formatter.Serialize(Value);
 	Ar.LeaveSlot();
 }
 
 void FStructuredArchive::FSlot::operator<< (float& Value)
 {
-	Ar.EnterSlot(Depth, ElementId);
+	Ar.EnterSlot(*this);
 	Ar.Formatter.Serialize(Value);
 	Ar.LeaveSlot();
 }
 
 void FStructuredArchive::FSlot::operator<< (double& Value)
 {
-	Ar.EnterSlot(Depth, ElementId);
+	Ar.EnterSlot(*this);
 	Ar.Formatter.Serialize(Value);
 	Ar.LeaveSlot();
 }
 
 void FStructuredArchive::FSlot::operator<< (bool& Value)
 {
-	Ar.EnterSlot(Depth, ElementId);
+	Ar.EnterSlot(*this);
 	Ar.Formatter.Serialize(Value);
 	Ar.LeaveSlot();
 }
 
 void FStructuredArchive::FSlot::operator<< (FString& Value)
 {
-	Ar.EnterSlot(Depth, ElementId);
+	Ar.EnterSlot(*this);
 	Ar.Formatter.Serialize(Value);
 	Ar.LeaveSlot();
 }
 
 void FStructuredArchive::FSlot::operator<< (FName& Value)
 {
-	Ar.EnterSlot(Depth, ElementId);
+	Ar.EnterSlot(*this);
 	Ar.Formatter.Serialize(Value);
 	Ar.LeaveSlot();
 }
 
 void FStructuredArchive::FSlot::operator<< (UObject*& Value)
 {
-	Ar.EnterSlot(Depth, ElementId);
+	Ar.EnterSlot(*this);
 	Ar.Formatter.Serialize(Value);
 	Ar.LeaveSlot();
 }
 
 void FStructuredArchive::FSlot::operator<< (FText& Value)
 {
-	Ar.EnterSlot(Depth, ElementId);
+	Ar.EnterSlot(*this);
 	Ar.Formatter.Serialize(Value);
 	Ar.LeaveSlot();
 }
 
 void FStructuredArchive::FSlot::operator<< (FWeakObjectPtr& Value)
 {
-	Ar.EnterSlot(Depth, ElementId);
+	Ar.EnterSlot(*this);
 	Ar.Formatter.Serialize(Value);
 	Ar.LeaveSlot();
 }
 
 void FStructuredArchive::FSlot::operator<< (FLazyObjectPtr& Value)
 {
-	Ar.EnterSlot(Depth, ElementId);
+	Ar.EnterSlot(*this);
 	Ar.Formatter.Serialize(Value);
 	Ar.LeaveSlot();
 }
 
 void FStructuredArchive::FSlot::operator<< (FSoftObjectPtr& Value)
 {
-	Ar.EnterSlot(Depth, ElementId);
+	Ar.EnterSlot(*this);
 	Ar.Formatter.Serialize(Value);
 	Ar.LeaveSlot();
 }
 
 void FStructuredArchive::FSlot::operator<< (FSoftObjectPath& Value)
 {
-	Ar.EnterSlot(Depth, ElementId);
+	Ar.EnterSlot(*this);
 	Ar.Formatter.Serialize(Value);
 	Ar.LeaveSlot();
 }
 
 void FStructuredArchive::FSlot::Serialize(TArray<uint8>& Data)
 {
-	Ar.EnterSlot(Depth, ElementId);
+	Ar.EnterSlot(*this);
 	Ar.Formatter.Serialize(Data);
 	Ar.LeaveSlot();
 }
 
 void FStructuredArchive::FSlot::Serialize(void* Data, uint64 DataSize)
 {
-	Ar.EnterSlot(Depth, ElementId);
+	Ar.EnterSlot(*this);
 	Ar.Formatter.Serialize(Data, DataSize);
 	Ar.LeaveSlot();
 }
@@ -557,7 +560,7 @@ void FStructuredArchive::FSlot::Serialize(void* Data, uint64 DataSize)
 
 FStructuredArchive::FSlot FStructuredArchive::FRecord::EnterField(FArchiveFieldName Name)
 {
-	Ar.SetScope(Depth, ElementId);
+	Ar.SetScope(*this);
 
 	Ar.CurrentSlotElementId = Ar.ElementIdGenerator.Generate();
 
@@ -609,7 +612,7 @@ FStructuredArchive::FMap FStructuredArchive::FRecord::EnterMap(FArchiveFieldName
 
 TOptional<FStructuredArchive::FSlot> FStructuredArchive::FRecord::TryEnterField(FArchiveFieldName Name, bool bEnterWhenWriting)
 {
-	Ar.SetScope(Depth, ElementId);
+	Ar.SetScope(*this);
 
 #if DO_GUARD_SLOW
 #if CHECK_UNIQUE_FIELD_NAMES
@@ -637,7 +640,7 @@ TOptional<FStructuredArchive::FSlot> FStructuredArchive::FRecord::TryEnterField(
 
 FStructuredArchive::FSlot FStructuredArchive::FArray::EnterElement()
 {
-	Ar.SetScope(Depth, ElementId);
+	Ar.SetScope(*this);
 
 #if DO_GUARD_SLOW
 	checkf(Ar.CurrentContainer.Top()->Index < Ar.CurrentContainer.Top()->Count, TEXT("Serialized too many array elements"));
@@ -652,7 +655,7 @@ FStructuredArchive::FSlot FStructuredArchive::FArray::EnterElement()
 
 FStructuredArchive::FSlot FStructuredArchive::FArray::EnterElement_TextOnly(EArchiveValueType& OutType)
 {
-	Ar.SetScope(Depth, ElementId);
+	Ar.SetScope(*this);
 
 #if DO_GUARD_SLOW
 	checkf(Ar.CurrentContainer.Top()->Index < Ar.CurrentContainer.Top()->Count, TEXT("Serialized too many array elements"));
@@ -669,7 +672,7 @@ FStructuredArchive::FSlot FStructuredArchive::FArray::EnterElement_TextOnly(EArc
 
 FStructuredArchive::FSlot FStructuredArchive::FStream::EnterElement()
 {
-	Ar.SetScope(Depth, ElementId);
+	Ar.SetScope(*this);
 
 	Ar.CurrentSlotElementId = Ar.ElementIdGenerator.Generate();
 
@@ -680,7 +683,7 @@ FStructuredArchive::FSlot FStructuredArchive::FStream::EnterElement()
 
 FStructuredArchive::FSlot FStructuredArchive::FStream::EnterElement_TextOnly(EArchiveValueType& OutType)
 {
-	Ar.SetScope(Depth, ElementId);
+	Ar.SetScope(*this);
 
 	Ar.CurrentSlotElementId = Ar.ElementIdGenerator.Generate();
 
@@ -693,7 +696,7 @@ FStructuredArchive::FSlot FStructuredArchive::FStream::EnterElement_TextOnly(EAr
 
 FStructuredArchive::FSlot FStructuredArchive::FMap::EnterElement(FString& Name)
 {
-	Ar.SetScope(Depth, ElementId);
+	Ar.SetScope(*this);
 
 #if DO_GUARD_SLOW
 	checkf(Ar.CurrentContainer.Top()->Index < Ar.CurrentContainer.Top()->Count, TEXT("Serialized too many map elements"));
@@ -730,7 +733,7 @@ FStructuredArchive::FSlot FStructuredArchive::FMap::EnterElement(FString& Name)
 
 FStructuredArchive::FSlot FStructuredArchive::FMap::EnterElement_TextOnly(FString& Name, EArchiveValueType& OutType)
 {
-	Ar.SetScope(Depth, ElementId);
+	Ar.SetScope(*this);
 
 #if DO_GUARD_SLOW
 	checkf(Ar.CurrentContainer.Top()->Index < Ar.CurrentContainer.Top()->Count, TEXT("Serialized too many map elements"));
diff --git a/Engine/Source/Runtime/Core/Public/Serialization/StructuredArchive.h b/Engine/Source/Runtime/Core/Public/Serialization/StructuredArchive.h
--- a/Engine/Source/Runtime/Core/Public/Serialization/StructuredArchive.h
+++ b/Engine/Source/Runtime/Core/Public/Serialization/StructuredArchive.h
@@ -194,8 +194,30 @@ private:
 		uint32 Id = 0;
 	};
 
+	// Represents a position of a slot within the hierarchy.
+	class CORE_API FSlotPosition
+	{
+		friend class FStructuredArchive;
+
+	protected:
+		int32 Depth;
+		FElementId ElementId;
+
+		FORCEINLINE explicit FSlotPosition(int32 InDepth, FElementId InElementId)
+			: Depth(InDepth)
+			, ElementId(InElementId)
+		{
+		}
+	};
+
+	// The base class of all slot types
 	class CORE_API FSlotBase
+#if WITH_TEXT_ARCHIVE_SUPPORT
+		: protected FSlotPosition
+#endif
 	{
+		friend class FStructuredArchive;
+
 	public:
 		FORCEINLINE FArchive& GetUnderlyingArchive()
 		{
@@ -204,14 +226,11 @@ private:
 
 	protected:
 		FStructuredArchive& Ar;
-#if WITH_TEXT_ARCHIVE_SUPPORT
-		const int32 Depth;
-		const FElementId ElementId;
 
+#if WITH_TEXT_ARCHIVE_SUPPORT
 		FORCEINLINE explicit FSlotBase(FStructuredArchive& InAr, int32 InDepth, FElementId InElementId)
-			: Ar(InAr)
-			, Depth(InDepth)
-			, ElementId(InElementId)
+			: FSlotPosition(InDepth, InElementId)
+			, Ar(InAr)
 		{
 		}
 #else
@@ -483,14 +502,14 @@ private:
 	/**
 	 * Enters the current slot for serializing a value. Asserts if the archive is not in a state about to write to an empty-slot.
 	 */
-	void EnterSlot(int32 ParentDepth, FElementId ElementId);
+	void EnterSlot(FSlotPosition Slot);
 
 	/**
 	 * Enters the current slot, adding an element onto the stack. Asserts if the archive is not in a state about to write to an empty-slot.
 	 *
 	 * @return  The depth of the newly-entered slot.
 	 */
-	int32 EnterSlotAsType(int32 ParentDepth, FElementId ElementId, EElementType ElementType);
+	int32 EnterSlotAsType(FSlotPosition Slot, EElementType ElementType);
 
 	/**
 	 * Leaves slot at the top of the current scope
@@ -498,9 +517,9 @@ private:
 	void LeaveSlot();
 
 	/**
-	 * Switches to the scope at the given element id, updating the formatter state as necessary
+	 * Switches to the scope for the given slot.
 	 */
-	void SetScope(int32 Depth, FElementId ElementId);
+	void SetScope(FSlotPosition Slot);
 #endif
 };
 
