RAII for the container array.

diff --git a/Engine/Source/Runtime/Core/Private/Serialization/StructuredArchive.cpp b/Engine/Source/Runtime/Core/Private/Serialization/StructuredArchive.cpp
--- a/Engine/Source/Runtime/Core/Private/Serialization/StructuredArchive.cpp
+++ b/Engine/Source/Runtime/Core/Private/Serialization/StructuredArchive.cpp
@@ -79,13 +79,6 @@ FStructuredArchive::FStructuredArchive(FArchiveFormatterType& InFormatter)
 FStructuredArchive::~FStructuredArchive()
 {
 	Close();
-
-#if DO_GUARD_SLOW
-	while(CurrentContainer.Num() > 0)
-	{
-		delete CurrentContainer.Pop();
-	}
-#endif
 }
 
 FStructuredArchive::FSlot FStructuredArchive::Open()
@@ -204,7 +197,7 @@ void FStructuredArchive::SetScope(FSlotPosition Slot)
 			case EElementType::Record:
 				Formatter.LeaveRecord();
 #if DO_GUARD_SLOW
-				delete CurrentContainer.Pop(false);
+				CurrentContainer.Pop(false);
 #endif
 				break;
 			case EElementType::Array:
@@ -213,7 +206,7 @@ void FStructuredArchive::SetScope(FSlotPosition Slot)
 #endif
 				Formatter.LeaveArray();
 #if DO_GUARD_SLOW
-				delete CurrentContainer.Pop(false);
+				CurrentContainer.Pop(false);
 #endif
 				break;
 			case EElementType::Stream:
@@ -225,13 +218,13 @@ void FStructuredArchive::SetScope(FSlotPosition Slot)
 #endif
 				Formatter.LeaveMap();
 #if DO_GUARD_SLOW
-				delete CurrentContainer.Pop(false);
+				CurrentContainer.Pop(false);
 #endif
 				break;
 			case EElementType::AttributedValue:
 				Formatter.LeaveAttributedValue();
 #if DO_GUARD_SLOW
-				delete CurrentContainer.Pop(false);
+				CurrentContainer.Pop(false);
 #endif
 				break;
 			}
@@ -257,7 +250,7 @@ FStructuredArchive::FRecord FStructuredArchive::FSlot::EnterRecord()
 	int32 NewDepth = Ar.EnterSlotAsType(*this, EElementType::Record);
 
 #if DO_GUARD_SLOW
-	Ar.CurrentContainer.Add(new FContainer(0));
+	Ar.CurrentContainer.Emplace(0);
 #endif
 
 	Ar.Formatter.EnterRecord();
@@ -270,7 +263,7 @@ FStructuredArchive::FRecord FStructuredArchive::FSlot::EnterRecord_TextOnly(TArr
 	int32 NewDepth = Ar.EnterSlotAsType(*this, EElementType::Record);
 
 #if DO_GUARD_SLOW
-	Ar.CurrentContainer.Add(new FContainer(0));
+	Ar.CurrentContainer.Emplace(0);
 #endif
 
 	Ar.Formatter.EnterRecord_TextOnly(OutFieldNames);
@@ -285,7 +278,7 @@ FStructuredArchive::FArray FStructuredArchive::FSlot::EnterArray(int32& Num)
 	Ar.Formatter.EnterArray(Num);
 
 #if DO_GUARD_SLOW
-	Ar.CurrentContainer.Add(new FContainer(Num));
+	Ar.CurrentContainer.Emplace(Num);
 #endif
 
 	return FArray(Ar, NewDepth, ElementId);
@@ -316,7 +309,7 @@ FStructuredArchive::FMap FStructuredArchive::FSlot::EnterMap(int32& Num)
 	Ar.Formatter.EnterMap(Num);
 
 #if DO_GUARD_SLOW
-	Ar.CurrentContainer.Add(new FContainer(Num));
+	Ar.CurrentContainer.Emplace(Num);
 #endif
 
 	return FMap(Ar, NewDepth, ElementId);
@@ -335,7 +328,7 @@ FStructuredArchive::FSlot FStructuredArchive::FSlot::EnterAttribute(FArchiveFiel
 		Ar.Formatter.EnterAttributedValue();
 
 #if DO_GUARD_SLOW
-		Ar.CurrentContainer.Add(new FContainer(0));
+		Ar.CurrentContainer.Emplace(0);
 #endif
 	}
 
@@ -376,7 +369,7 @@ TOptional<FStructuredArchive::FSlot> FStructuredArchive::FSlot::TryEnterAttribut
 		Ar.Formatter.EnterAttributedValue();
 
 #if DO_GUARD_SLOW
-		Ar.CurrentContainer.Add(new FContainer(0));
+		Ar.CurrentContainer.Emplace(0);
 #endif
 	}
 
diff --git a/Engine/Source/Runtime/Core/Public/Serialization/StructuredArchive.h b/Engine/Source/Runtime/Core/Public/Serialization/StructuredArchive.h
--- a/Engine/Source/Runtime/Core/Public/Serialization/StructuredArchive.h
+++ b/Engine/Source/Runtime/Core/Public/Serialization/StructuredArchive.h
@@ -491,7 +491,7 @@ private:
 	/**
 	 * For arrays and maps, stores the loop counter and size of the container. Also stores key names for records and maps in builds with DO_GUARD_SLOW enabled.
 	 */
-	TArray<FContainer*> CurrentContainer;
+	TArray<TUniqueObj<FContainer>> CurrentContainer;
 #endif
 
 	/**
