Changed CsvProfiler singleton creation to avoid "magic statics" initialization

diff --git a/Engine/Source/Runtime/Core/Private/ProfilingDebugging/CsvProfiler.cpp b/Engine/Source/Runtime/Core/Private/ProfilingDebugging/CsvProfiler.cpp
--- a/Engine/Source/Runtime/Core/Private/ProfilingDebugging/CsvProfiler.cpp
+++ b/Engine/Source/Runtime/Core/Private/ProfilingDebugging/CsvProfiler.cpp
@@ -2423,13 +2423,20 @@ void FCsvProfilerThreadDataProcessor::Process(FCsvProcessThreadDataStats& OutSta
 
 FCsvProfiler* FCsvProfiler::Get()
 {
-	static TUniquePtr<FCsvProfiler> Instance;
-	if (!Instance.IsValid())
+	static FCsvProfiler* InstancePtr;
+
+	if (!InstancePtr)
 	{
-		LLM_SCOPE(ELLMTag::CsvProfiler);
-		Instance = MakeUnique<FCsvProfiler>();
+		// It's important that the initializer goes here to avoid the overhead of
+		// "magic static" initialization on every call (mostly an issue with MSVC
+		// because of their epoch-based initialization scheme which doesn't seem
+		// to make any real sense on x86)
+
+		static FCsvProfiler Instance;
+		InstancePtr = &Instance;
 	}
-	return Instance.Get();
+
+	return InstancePtr;
 }
 
 FCsvProfiler::FCsvProfiler()
