New macros which represent the container checks better.

diff --git a/Engine/Source/Runtime/Core/Private/Serialization/StructuredArchive.cpp b/Engine/Source/Runtime/Core/Private/Serialization/StructuredArchive.cpp
--- a/Engine/Source/Runtime/Core/Private/Serialization/StructuredArchive.cpp
+++ b/Engine/Source/Runtime/Core/Private/Serialization/StructuredArchive.cpp
@@ -8,9 +8,7 @@
 
 //////////// FStructuredArchive::FContainer ////////////
 
-#if DO_GUARD_SLOW
-
-#define CHECK_UNIQUE_FIELD_NAMES 0
+#if DO_STRUCTURED_ARCHIVE_CONTAINER_CHECKS
 
 struct FStructuredArchive::FContainer
 {
@@ -18,7 +16,7 @@ struct FStructuredArchive::FContainer
 	int  Count                   = 0;
 	bool bAttributedValueWritten = false;
 
-#if CHECK_UNIQUE_FIELD_NAMES
+#if DO_STRUCTURED_ARCHIVE_UNIQUE_FIELD_NAME_CHECKS
 	TSet<FString> KeyNames;
 #endif
 
@@ -64,14 +62,14 @@ FStructuredArchiveChildReader::~FStructuredArchiveChildReader()
 
 FStructuredArchive::FStructuredArchive(FArchiveFormatterType& InFormatter)
 	: Formatter(InFormatter)
-#if DO_GUARD_SLOW
+#if DO_STRUCTURED_ARCHIVE_CONTAINER_CHECKS
 	, bRequiresStructuralMetadata(true)
 #else
 	, bRequiresStructuralMetadata(InFormatter.HasDocumentTree())
 #endif
 {
 	CurrentScope.Reserve(32);
-#if DO_GUARD_SLOW
+#if DO_STRUCTURED_ARCHIVE_CONTAINER_CHECKS
 	CurrentContainer.Reserve(32);
 #endif
 }
@@ -108,7 +106,7 @@ void FStructuredArchive::EnterSlot(FSlotPosition Slot)
 	// If the slot being entered has attributes, enter the value slot first.
 	if (ParentDepth + 1 < CurrentScope.Num() && CurrentScope[ParentDepth + 1].Id == ElementId && CurrentScope[ParentDepth + 1].Type == EElementType::AttributedValue)
 	{
-#if DO_GUARD_SLOW
+#if DO_STRUCTURED_ARCHIVE_CONTAINER_CHECKS
 		checkf(!CurrentSlotElementId.IsValid() && !CurrentContainer.Top()->bAttributedValueWritten, TEXT("Attempt to serialize data into an invalid slot"));
 		CurrentContainer.Top()->bAttributedValueWritten = true;
 #else
@@ -159,7 +157,7 @@ void FStructuredArchive::LeaveSlot()
 			break;
 		case EElementType::Array:
 			Formatter.LeaveArrayElement();
-#if DO_GUARD_SLOW
+#if DO_STRUCTURED_ARCHIVE_CONTAINER_CHECKS
 			CurrentContainer.Top()->Index++;
 #endif
 			break;
@@ -168,7 +166,7 @@ void FStructuredArchive::LeaveSlot()
 			break;
 		case EElementType::Map:
 			Formatter.LeaveMapElement();
-#if DO_GUARD_SLOW
+#if DO_STRUCTURED_ARCHIVE_CONTAINER_CHECKS
 			CurrentContainer.Top()->Index++;
 #endif
 			break;
@@ -196,16 +194,16 @@ void FStructuredArchive::SetScope(FSlotPosition Slot)
 			{
 			case EElementType::Record:
 				Formatter.LeaveRecord();
-#if DO_GUARD_SLOW
+#if DO_STRUCTURED_ARCHIVE_CONTAINER_CHECKS
 				CurrentContainer.Pop(false);
 #endif
 				break;
 			case EElementType::Array:
-#if DO_GUARD_SLOW
+#if DO_STRUCTURED_ARCHIVE_CONTAINER_CHECKS
 				checkf(GetUnderlyingArchive().IsLoading() || CurrentContainer.Top()->Index == CurrentContainer.Top()->Count, TEXT("Incorrect number of elements serialized in array"));
 #endif
 				Formatter.LeaveArray();
-#if DO_GUARD_SLOW
+#if DO_STRUCTURED_ARCHIVE_CONTAINER_CHECKS
 				CurrentContainer.Pop(false);
 #endif
 				break;
@@ -213,17 +211,17 @@ void FStructuredArchive::SetScope(FSlotPosition Slot)
 				Formatter.LeaveStream();
 				break;
 			case EElementType::Map:
-#if DO_GUARD_SLOW
+#if DO_STRUCTURED_ARCHIVE_CONTAINER_CHECKS
 				checkf(CurrentContainer.Top()->Index == CurrentContainer.Top()->Count, TEXT("Incorrect number of elements serialized in map"));
 #endif
 				Formatter.LeaveMap();
-#if DO_GUARD_SLOW
+#if DO_STRUCTURED_ARCHIVE_CONTAINER_CHECKS
 				CurrentContainer.Pop(false);
 #endif
 				break;
 			case EElementType::AttributedValue:
 				Formatter.LeaveAttributedValue();
-#if DO_GUARD_SLOW
+#if DO_STRUCTURED_ARCHIVE_CONTAINER_CHECKS
 				CurrentContainer.Pop(false);
 #endif
 				break;
@@ -249,7 +247,7 @@ FStructuredArchive::FRecord FStructuredArchive::FSlot::EnterRecord()
 {
 	int32 NewDepth = Ar.EnterSlotAsType(*this, EElementType::Record);
 
-#if DO_GUARD_SLOW
+#if DO_STRUCTURED_ARCHIVE_CONTAINER_CHECKS
 	Ar.CurrentContainer.Emplace(0);
 #endif
 
@@ -262,7 +260,7 @@ FStructuredArchive::FRecord FStructuredArchive::FSlot::EnterRecord_TextOnly(TArr
 {
 	int32 NewDepth = Ar.EnterSlotAsType(*this, EElementType::Record);
 
-#if DO_GUARD_SLOW
+#if DO_STRUCTURED_ARCHIVE_CONTAINER_CHECKS
 	Ar.CurrentContainer.Emplace(0);
 #endif
 
@@ -277,7 +275,7 @@ FStructuredArchive::FArray FStructuredArchive::FSlot::EnterArray(int32& Num)
 
 	Ar.Formatter.EnterArray(Num);
 
-#if DO_GUARD_SLOW
+#if DO_STRUCTURED_ARCHIVE_CONTAINER_CHECKS
 	Ar.CurrentContainer.Emplace(Num);
 #endif
 
@@ -308,7 +306,7 @@ FStructuredArchive::FMap FStructuredArchive::FSlot::EnterMap(int32& Num)
 
 	Ar.Formatter.EnterMap(Num);
 
-#if DO_GUARD_SLOW
+#if DO_STRUCTURED_ARCHIVE_CONTAINER_CHECKS
 	Ar.CurrentContainer.Emplace(Num);
 #endif
 
@@ -327,7 +325,7 @@ FStructuredArchive::FSlot FStructuredArchive::FSlot::EnterAttribute(FArchiveFiel
 
 		Ar.Formatter.EnterAttributedValue();
 
-#if DO_GUARD_SLOW
+#if DO_STRUCTURED_ARCHIVE_CONTAINER_CHECKS
 		Ar.CurrentContainer.Emplace(0);
 #endif
 	}
@@ -340,8 +338,8 @@ FStructuredArchive::FSlot FStructuredArchive::FSlot::EnterAttribute(FArchiveFiel
 
 	Ar.CurrentSlotElementId = Ar.ElementIdGenerator.Generate();
 
-#if DO_GUARD_SLOW
-#if CHECK_UNIQUE_FIELD_NAMES
+#if DO_STRUCTURED_ARCHIVE_CONTAINER_CHECKS
+#if DO_STRUCTURED_ARCHIVE_UNIQUE_FIELD_NAME_CHECKS
 	if (!Ar.GetUnderlyingArchive().IsLoading())
 	{
 		FContainer& Container = *Ar.CurrentContainer.Top();
@@ -368,13 +366,13 @@ TOptional<FStructuredArchive::FSlot> FStructuredArchive::FSlot::TryEnterAttribut
 
 		Ar.Formatter.EnterAttributedValue();
 
-#if DO_GUARD_SLOW
+#if DO_STRUCTURED_ARCHIVE_CONTAINER_CHECKS
 		Ar.CurrentContainer.Emplace(0);
 #endif
 	}
 
-#if DO_GUARD_SLOW
-#if CHECK_UNIQUE_FIELD_NAMES
+#if DO_STRUCTURED_ARCHIVE_CONTAINER_CHECKS
+#if DO_STRUCTURED_ARCHIVE_UNIQUE_FIELD_NAME_CHECKS
 	if (!Ar.GetUnderlyingArchive().IsLoading())
 	{
 		FContainer& Container = *Ar.CurrentContainer.Top();
@@ -557,8 +555,8 @@ FStructuredArchive::FSlot FStructuredArchive::FRecord::EnterField(FArchiveFieldN
 
 	Ar.CurrentSlotElementId = Ar.ElementIdGenerator.Generate();
 
-#if DO_GUARD_SLOW
-#if CHECK_UNIQUE_FIELD_NAMES
+#if DO_STRUCTURED_ARCHIVE_CONTAINER_CHECKS
+#if DO_STRUCTURED_ARCHIVE_UNIQUE_FIELD_NAME_CHECKS
 	if (!Ar.GetUnderlyingArchive().IsLoading())
 	{
 		FContainer& Container = *Ar.CurrentContainer.Top();
@@ -607,8 +605,8 @@ TOptional<FStructuredArchive::FSlot> FStructuredArchive::FRecord::TryEnterField(
 {
 	Ar.SetScope(*this);
 
-#if DO_GUARD_SLOW
-#if CHECK_UNIQUE_FIELD_NAMES
+#if DO_STRUCTURED_ARCHIVE_CONTAINER_CHECKS
+#if DO_STRUCTURED_ARCHIVE_UNIQUE_FIELD_NAME_CHECKS
 	if (!GetUnderlyingArchive().IsLoading())
 	{
 		FContainer& Container = *Ar.CurrentContainer.Top();
@@ -635,7 +633,7 @@ FStructuredArchive::FSlot FStructuredArchive::FArray::EnterElement()
 {
 	Ar.SetScope(*this);
 
-#if DO_GUARD_SLOW
+#if DO_STRUCTURED_ARCHIVE_CONTAINER_CHECKS
 	checkf(Ar.CurrentContainer.Top()->Index < Ar.CurrentContainer.Top()->Count, TEXT("Serialized too many array elements"));
 #endif
 
@@ -650,7 +648,7 @@ FStructuredArchive::FSlot FStructuredArchive::FArray::EnterElement_TextOnly(EArc
 {
 	Ar.SetScope(*this);
 
-#if DO_GUARD_SLOW
+#if DO_STRUCTURED_ARCHIVE_CONTAINER_CHECKS
 	checkf(Ar.CurrentContainer.Top()->Index < Ar.CurrentContainer.Top()->Count, TEXT("Serialized too many array elements"));
 #endif
 
@@ -691,14 +689,14 @@ FStructuredArchive::FSlot FStructuredArchive::FMap::EnterElement(FString& Name)
 {
 	Ar.SetScope(*this);
 
-#if DO_GUARD_SLOW
+#if DO_STRUCTURED_ARCHIVE_CONTAINER_CHECKS
 	checkf(Ar.CurrentContainer.Top()->Index < Ar.CurrentContainer.Top()->Count, TEXT("Serialized too many map elements"));
 #endif
 
 	Ar.CurrentSlotElementId = Ar.ElementIdGenerator.Generate();
 
-#if DO_GUARD_SLOW
-#if CHECK_UNIQUE_FIELD_NAMES
+#if DO_STRUCTURED_ARCHIVE_CONTAINER_CHECKS
+#if DO_STRUCTURED_ARCHIVE_UNIQUE_FIELD_NAME_CHECKS
 	if(Ar.GetUnderlyingArchive().IsSaving())
 	{
 		FContainer& Container = *Ar.CurrentContainer.Top();
@@ -710,8 +708,8 @@ FStructuredArchive::FSlot FStructuredArchive::FMap::EnterElement(FString& Name)
 
 	Ar.Formatter.EnterMapElement(Name);
 
-#if DO_GUARD_SLOW
-#if CHECK_UNIQUE_FIELD_NAMES
+#if DO_STRUCTURED_ARCHIVE_CONTAINER_CHECKS
+#if DO_STRUCTURED_ARCHIVE_UNIQUE_FIELD_NAME_CHECKS
 	if(Ar.GetUnderlyingArchive().IsLoading())
 	{
 		FContainer& Container = *Ar.CurrentContainer.Top();
@@ -728,14 +726,14 @@ FStructuredArchive::FSlot FStructuredArchive::FMap::EnterElement_TextOnly(FStrin
 {
 	Ar.SetScope(*this);
 
-#if DO_GUARD_SLOW
+#if DO_STRUCTURED_ARCHIVE_CONTAINER_CHECKS
 	checkf(Ar.CurrentContainer.Top()->Index < Ar.CurrentContainer.Top()->Count, TEXT("Serialized too many map elements"));
 #endif
 
 	Ar.CurrentSlotElementId = Ar.ElementIdGenerator.Generate();
 
-#if DO_GUARD_SLOW
-#if CHECK_UNIQUE_FIELD_NAMES
+#if DO_STRUCTURED_ARCHIVE_CONTAINER_CHECKS
+#if DO_STRUCTURED_ARCHIVE_UNIQUE_FIELD_NAME_CHECKS
 	if(Ar.GetUnderlyingArchive().IsSaving())
 	{
 		FContainer& Container = *Ar.CurrentContainer.Top();
@@ -747,8 +745,8 @@ FStructuredArchive::FSlot FStructuredArchive::FMap::EnterElement_TextOnly(FStrin
 
 	Ar.Formatter.EnterMapElement_TextOnly(Name, OutType);
 
-#if DO_GUARD_SLOW
-#if CHECK_UNIQUE_FIELD_NAMES
+#if DO_STRUCTURED_ARCHIVE_CONTAINER_CHECKS
+#if DO_STRUCTURED_ARCHIVE_UNIQUE_FIELD_NAME_CHECKS
 	if(Ar.GetUnderlyingArchive().IsLoading())
 	{
 		FContainer& Container = *Ar.CurrentContainer.Top();
diff --git a/Engine/Source/Runtime/Core/Public/Serialization/StructuredArchive.h b/Engine/Source/Runtime/Core/Public/Serialization/StructuredArchive.h
--- a/Engine/Source/Runtime/Core/Public/Serialization/StructuredArchive.h
+++ b/Engine/Source/Runtime/Core/Public/Serialization/StructuredArchive.h
@@ -12,6 +12,24 @@
 #include "Serialization/ArchiveProxy.h"
 #include "Templates/UniqueObj.h"
 
+/**
+ * DO_STRUCTURED_ARCHIVE_CONTAINER_CHECKS - if set, checks that nested container types are serialized correctly.
+ */
+#ifndef DO_STRUCTURED_ARCHIVE_CONTAINER_CHECKS
+	#if DO_GUARD_SLOW
+		#define DO_STRUCTURED_ARCHIVE_CONTAINER_CHECKS 1
+	#else
+		#define DO_STRUCTURED_ARCHIVE_CONTAINER_CHECKS 0
+	#endif
+#endif
+
+/**
+ * DO_STRUCTURED_ARCHIVE_UNIQUE_FIELD_NAME_CHECKS - if set, checks that field names are unique within a container.  Requires DO_STRUCTURED_ARCHIVE_CONTAINER_CHECKS.
+ */
+#ifndef DO_STRUCTURED_ARCHIVE_UNIQUE_FIELD_NAME_CHECKS
+	#define DO_STRUCTURED_ARCHIVE_UNIQUE_FIELD_NAME_CHECKS 0
+#endif
+
 /**
  * Class to contain a named value for serialization. Intended to be created as a temporary and passed to object serialization methods.
  */
@@ -451,10 +469,6 @@ private:
 		}
 	};
 
-#if DO_GUARD_SLOW
-	struct FContainer;
-#endif
-
 	struct FIdGenerator
 	{
 		FElementId Generate()
@@ -487,9 +501,11 @@ private:
 	 */
 	TArray<FElement, TNonRelocatableInlineAllocator<32>> CurrentScope;
 
-#if DO_GUARD_SLOW
+#if DO_STRUCTURED_ARCHIVE_CONTAINER_CHECKS
+	struct FContainer;
+
 	/**
-	 * For arrays and maps, stores the loop counter and size of the container. Also stores key names for records and maps in builds with DO_GUARD_SLOW enabled.
+	 * For arrays and maps, stores the loop counter and size of the container. Also stores key names for records and maps in builds with DO_STRUCTURED_ARCHIVE_CONTAINER_CHECKS enabled.
 	 */
 	TArray<TUniqueObj<FContainer>> CurrentContainer;
 #endif
