Fix for deserializing an attributed value without having read an attribute first.

diff --git a/Engine/Source/Runtime/Core/Private/Serialization/StructuredArchive.cpp b/Engine/Source/Runtime/Core/Private/Serialization/StructuredArchive.cpp
--- a/Engine/Source/Runtime/Core/Private/Serialization/StructuredArchive.cpp
+++ b/Engine/Source/Runtime/Core/Private/Serialization/StructuredArchive.cpp
@@ -98,7 +98,7 @@ void FStructuredArchive::Close()
 	SetScope(FSlotPosition(0, RootElementId));
 }
 
-void FStructuredArchive::EnterSlot(FSlotPosition Slot)
+void FStructuredArchive::EnterSlot(FSlotPosition Slot, bool bEnteringAttributedValue)
 {
 	int32      ParentDepth = Slot.Depth;
 	FElementId ElementId   = Slot.ElementId;
@@ -116,6 +116,16 @@ void FStructuredArchive::EnterSlot(FSlotPosition Slot)
 		SetScope(FSlotPosition(ParentDepth + 1, ElementId));
 		Formatter.EnterAttributedValueValue();
 	}
+	else if (!bEnteringAttributedValue && Formatter.TryEnterAttributedValueValue())
+	{
+		int32 NewDepth = EnterSlotAsType(FSlotPosition(ParentDepth, ElementId), EElementType::AttributedValue);
+		check(NewDepth == ParentDepth + 1);
+		FElementId AttributedValueId = CurrentScope[NewDepth].Id;
+		SetScope(FSlotPosition(NewDepth, AttributedValueId));
+#if DO_STRUCTURED_ARCHIVE_CONTAINER_CHECKS
+		CurrentContainer.Emplace(0);
+#endif
+	}
 	else
 	{
 		checkf(ElementId == CurrentSlotElementId, TEXT("Attempt to serialize data into an invalid slot"));
@@ -127,7 +137,7 @@ void FStructuredArchive::EnterSlot(FSlotPosition Slot)
 
 int32 FStructuredArchive::EnterSlotAsType(FSlotPosition Slot, EElementType ElementType)
 {
-	EnterSlot(Slot);
+	EnterSlot(Slot, ElementType == EElementType::AttributedValue);
 
 	int32 NewSlotDepth = Slot.Depth + 1;
 
diff --git a/Engine/Source/Runtime/Core/Public/Serialization/Formatters/BinaryArchiveFormatter.h b/Engine/Source/Runtime/Core/Public/Serialization/Formatters/BinaryArchiveFormatter.h
--- a/Engine/Source/Runtime/Core/Public/Serialization/Formatters/BinaryArchiveFormatter.h
+++ b/Engine/Source/Runtime/Core/Public/Serialization/Formatters/BinaryArchiveFormatter.h
@@ -48,6 +48,7 @@ public:
 	virtual void LeaveAttribute() override;
 	virtual void LeaveAttributedValue() override;
 	virtual bool TryEnterAttribute(FArchiveFieldName AttributeName, bool bEnterWhenWriting) override;
+	virtual bool TryEnterAttributedValueValue() override;
 
 	virtual void Serialize(uint8& Value) override;
 	virtual void Serialize(uint16& Value) override;
@@ -171,6 +172,11 @@ inline void FBinaryArchiveFormatter::EnterAttributedValueValue()
 {
 }
 
+inline bool FBinaryArchiveFormatter::TryEnterAttributedValueValue()
+{
+	return false;
+}
+
 inline void FBinaryArchiveFormatter::LeaveAttribute()
 {
 }
diff --git a/Engine/Source/Runtime/Core/Public/Serialization/StructuredArchive.h b/Engine/Source/Runtime/Core/Public/Serialization/StructuredArchive.h
--- a/Engine/Source/Runtime/Core/Public/Serialization/StructuredArchive.h
+++ b/Engine/Source/Runtime/Core/Public/Serialization/StructuredArchive.h
@@ -518,7 +518,7 @@ private:
 	/**
 	 * Enters the current slot for serializing a value. Asserts if the archive is not in a state about to write to an empty-slot.
 	 */
-	void EnterSlot(FSlotPosition Slot);
+	void EnterSlot(FSlotPosition Slot, bool bEnteringAttributedValue = false);
 
 	/**
 	 * Enters the current slot, adding an element onto the stack. Asserts if the archive is not in a state about to write to an empty-slot.
diff --git a/Engine/Source/Runtime/Core/Public/Serialization/StructuredArchiveFormatter.h b/Engine/Source/Runtime/Core/Public/Serialization/StructuredArchiveFormatter.h
--- a/Engine/Source/Runtime/Core/Public/Serialization/StructuredArchiveFormatter.h
+++ b/Engine/Source/Runtime/Core/Public/Serialization/StructuredArchiveFormatter.h
@@ -109,6 +109,7 @@ public:
 	virtual void LeaveAttribute() = 0;
 	virtual void LeaveAttributedValue() = 0;
 	virtual bool TryEnterAttribute(FArchiveFieldName AttributeName, bool bEnterWhenWriting) = 0;
+	virtual bool TryEnterAttributedValueValue() = 0;
 
 	virtual void Serialize(uint8& Value) = 0;
 	virtual void Serialize(uint16& Value) = 0;
diff --git a/Engine/Source/Runtime/CoreUObject/Private/Serialization/Formatters/JsonArchiveInputFormatter.cpp b/Engine/Source/Runtime/CoreUObject/Private/Serialization/Formatters/JsonArchiveInputFormatter.cpp
--- a/Engine/Source/Runtime/CoreUObject/Private/Serialization/Formatters/JsonArchiveInputFormatter.cpp
+++ b/Engine/Source/Runtime/CoreUObject/Private/Serialization/Formatters/JsonArchiveInputFormatter.cpp
@@ -254,6 +254,26 @@ void FJsonArchiveInputFormatter::EnterAttributedValueValue()
 	}
 }
 
+bool FJsonArchiveInputFormatter::TryEnterAttributedValueValue()
+{
+	TSharedPtr<FJsonValue>& Value = ValueStack.Top();
+	const TSharedPtr<FJsonObject>* ObjectPtr = nullptr;
+	if (Value->TryGetObject(ObjectPtr))
+	{
+		FJsonObject& ObjectRef = **ObjectPtr;
+
+		TSharedPtr<FJsonValue> Field = ObjectRef.TryGetField(EscapeFieldName(TEXT("_Value")));
+		if (Field.IsValid())
+		{
+			ObjectStack.Emplace(*ObjectPtr, ValueStack.Num());
+			ValueStack.Add(Field);
+			return true;
+		}
+	}
+
+	return false;
+}
+
 void FJsonArchiveInputFormatter::LeaveAttribute()
 {
 	ValueStack.Pop();
diff --git a/Engine/Source/Runtime/CoreUObject/Private/Serialization/Formatters/JsonArchiveOutputFormatter.cpp b/Engine/Source/Runtime/CoreUObject/Private/Serialization/Formatters/JsonArchiveOutputFormatter.cpp
--- a/Engine/Source/Runtime/CoreUObject/Private/Serialization/Formatters/JsonArchiveOutputFormatter.cpp
+++ b/Engine/Source/Runtime/CoreUObject/Private/Serialization/Formatters/JsonArchiveOutputFormatter.cpp
@@ -217,6 +217,11 @@ void FJsonArchiveOutputFormatter::EnterAttributedValueValue()
 	WriteOptionalAttributedBlockValue();
 }
 
+bool FJsonArchiveOutputFormatter::TryEnterAttributedValueValue()
+{
+	return false;
+}
+
 bool FJsonArchiveOutputFormatter::TryEnterAttribute(FArchiveFieldName AttributeName, bool bEnterWhenSaving)
 {
 	if (bEnterWhenSaving)
diff --git a/Engine/Source/Runtime/CoreUObject/Public/Serialization/Formatters/JsonArchiveInputFormatter.h b/Engine/Source/Runtime/CoreUObject/Public/Serialization/Formatters/JsonArchiveInputFormatter.h
--- a/Engine/Source/Runtime/CoreUObject/Public/Serialization/Formatters/JsonArchiveInputFormatter.h
+++ b/Engine/Source/Runtime/CoreUObject/Public/Serialization/Formatters/JsonArchiveInputFormatter.h
@@ -54,6 +54,7 @@ public:
 	virtual void LeaveAttribute() override;
 	virtual void LeaveAttributedValue() override;
 	virtual bool TryEnterAttribute(FArchiveFieldName AttributeName, bool bEnterWhenSavin) override;
+	virtual bool TryEnterAttributedValueValue() override;
 
 	virtual void Serialize(uint8& Value) override;
 	virtual void Serialize(uint16& Value) override;
diff --git a/Engine/Source/Runtime/CoreUObject/Public/Serialization/Formatters/JsonArchiveOutputFormatter.h b/Engine/Source/Runtime/CoreUObject/Public/Serialization/Formatters/JsonArchiveOutputFormatter.h
--- a/Engine/Source/Runtime/CoreUObject/Public/Serialization/Formatters/JsonArchiveOutputFormatter.h
+++ b/Engine/Source/Runtime/CoreUObject/Public/Serialization/Formatters/JsonArchiveOutputFormatter.h
@@ -51,6 +51,8 @@ public:
 	virtual void LeaveAttributedValue() override;
 	virtual bool TryEnterAttribute(FArchiveFieldName AttributeName, bool bEnterWhenSaving) override;
 
+	virtual bool TryEnterAttributedValueValue() override;
+
 	virtual void Serialize(uint8& Value) override;
 	virtual void Serialize(uint16& Value) override;
 	virtual void Serialize(uint32& Value) override;
