GitHub #6163 : Remove unused lambda capture

diff --git a/Engine/Source/Runtime/Core/Public/Containers/Map.h b/Engine/Source/Runtime/Core/Public/Containers/Map.h
--- a/Engine/Source/Runtime/Core/Public/Containers/Map.h
+++ b/Engine/Source/Runtime/Core/Public/Containers/Map.h
@@ -1560,7 +1560,7 @@ public:
 				Key,
 				MapLayout.SetLayout,
 				GetKeyHash, // We 'know' that the implementation of Find doesn't call GetKeyHash on anything except Key
-				[KeyEqualityFn, MapLayout](const void* InKey, const void* InPair )
+				[KeyEqualityFn](const void* InKey, const void* InPair )
 				{
 					return KeyEqualityFn(InKey, (uint8*)InPair);
 				}
