UnrealInsights:

Improved counters functionality

diff --git a/Engine/Source/Developer/TraceServices/Private/Analyzers/CountersTraceAnalysis.cpp b/Engine/Source/Developer/TraceServices/Private/Analyzers/CountersTraceAnalysis.cpp
--- a/Engine/Source/Developer/TraceServices/Private/Analyzers/CountersTraceAnalysis.cpp
+++ b/Engine/Source/Developer/TraceServices/Private/Analyzers/CountersTraceAnalysis.cpp
@@ -30,12 +30,13 @@ void FCountersAnalyzer::OnEvent(uint16 RouteId, const FOnEventContext& Context)
 	{
 		uint16 CounterId = EventData.GetValue<uint16>("Id");
 		ETraceCounterType CounterType = static_cast<ETraceCounterType>(EventData.GetValue<uint8>("Type"));
+		ETraceCounterDisplayHint CounterDisplayHint = static_cast<ETraceCounterDisplayHint>(EventData.GetValue<uint8>("DisplayHint"));
 		Trace::ICounter* Counter = CounterProvider.CreateCounter();
 		if (CounterType == TraceCounterType_Float)
 		{
 			Counter->SetIsFloatingPoint(true);
 		}
-		else if (CounterType == TraceCounterType_Memory)
+		if (CounterDisplayHint == TraceCounterDisplayHint_Memory)
 		{
 			Counter->SetDisplayHint(Trace::CounterDisplayHint_Memory);
 		}
diff --git a/Engine/Source/Runtime/Core/Private/ProfilingDebugging/CountersTrace.cpp b/Engine/Source/Runtime/Core/Private/ProfilingDebugging/CountersTrace.cpp
--- a/Engine/Source/Runtime/Core/Private/ProfilingDebugging/CountersTrace.cpp
+++ b/Engine/Source/Runtime/Core/Private/ProfilingDebugging/CountersTrace.cpp
@@ -8,6 +8,7 @@
 UE_TRACE_EVENT_BEGIN(Counters, Spec, Always)
 	UE_TRACE_EVENT_FIELD(uint16, Id)
 	UE_TRACE_EVENT_FIELD(uint8, Type)
+	UE_TRACE_EVENT_FIELD(uint8, DisplayHint)
 UE_TRACE_EVENT_END()
 
 UE_TRACE_EVENT_BEGIN(Counters, SetValueInt)
@@ -23,79 +24,33 @@ UE_TRACE_EVENT_BEGIN(Counters, SetValueFloat)
 UE_TRACE_EVENT_END()
 
 
-struct FCountersTracePrivate
-{
-	static uint16 InitCounter(const TCHAR* CounterName, ETraceCounterType Type);
-	
-	static void SetValue(uint16 CounterId, int64 Value)
-	{
-		UE_TRACE_LOG(Counters, SetValueInt)
-			<< SetValueInt.Cycle(FPlatformTime::Cycles64())
-			<< SetValueInt.Value(Value)
-			<< SetValueInt.CounterId(CounterId);
-	}
-
-	static void SetValue(uint16 CounterId, float Value)
-	{
-		UE_TRACE_LOG(Counters, SetValueFloat)
-			<< SetValueFloat.Cycle(FPlatformTime::Cycles64())
-			<< SetValueFloat.Value(Value)
-			<< SetValueFloat.CounterId(CounterId);
-	}
-};
-
-uint16 FCountersTracePrivate::InitCounter(const TCHAR* CounterName, ETraceCounterType Type)
+uint16 FCountersTrace::OutputInitCounter(const TCHAR* CounterName, ETraceCounterType CounterType, ETraceCounterDisplayHint CounterDisplayHint)
 {
 	static TAtomic<uint16> NextId(1);
 	uint16 CounterId = uint16(NextId++);
 	uint16 NameSize = (FCString::Strlen(CounterName) + 1) * sizeof(TCHAR);
 	UE_TRACE_LOG(Counters, Spec, NameSize)
 		<< Spec.Id(CounterId)
-		<< Spec.Type(uint8(Type))
+		<< Spec.Type(uint8(CounterType))
+		<< Spec.DisplayHint(uint8(CounterDisplayHint))
 		<< Spec.Attachment(CounterName, NameSize);
 	return CounterId;
 }
 
-void FCountersTrace::FCounterInt::Init(const TCHAR* CounterName)
+void FCountersTrace::OutputSetValue(uint16 CounterId, int64 Value)
 {
-	CounterId = FCountersTracePrivate::InitCounter(CounterName, TraceCounterType_Int);
+	UE_TRACE_LOG(Counters, SetValueInt)
+		<< SetValueInt.Cycle(FPlatformTime::Cycles64())
+		<< SetValueInt.Value(Value)
+		<< SetValueInt.CounterId(CounterId);
 }
 
-void FCountersTrace::FCounterInt::Set(int64 Value)
+void FCountersTrace::OutputSetValue(uint16 CounterId, double Value)
 {
-	if (Value != CurrentValue)
-	{
-		CurrentValue = Value;
-		FCountersTracePrivate::SetValue(CounterId, Value);
-	}
-}
-
-void FCountersTrace::FCounterFloat::Init(const TCHAR* CounterName)
-{
-	CounterId = FCountersTracePrivate::InitCounter(CounterName, TraceCounterType_Float);
-}
-
-void FCountersTrace::FCounterFloat::Set(float Value)
-{
-	if (Value != CurrentValue)
-	{
-		CurrentValue = Value;
-		FCountersTracePrivate::SetValue(CounterId, Value);
-	}
-}
-
-void FCountersTrace::FCounterMemory::Init(const TCHAR* CounterName)
-{
-	CounterId = FCountersTracePrivate::InitCounter(CounterName, TraceCounterType_Memory);
-}
-
-void FCountersTrace::FCounterMemory::Set(uint64 Value)
-{
-	if (Value != CurrentValue)
-	{
-		CurrentValue = Value;
-		FCountersTracePrivate::SetValue(CounterId, (int64)Value);
-	}
+	UE_TRACE_LOG(Counters, SetValueFloat)
+		<< SetValueFloat.Cycle(FPlatformTime::Cycles64())
+		<< SetValueFloat.Value(Value)
+		<< SetValueFloat.CounterId(CounterId);
 }
 
 void FCountersTrace::Init(const TCHAR* CmdLine)
diff --git a/Engine/Source/Runtime/Core/Public/ProfilingDebugging/CountersTrace.h b/Engine/Source/Runtime/Core/Public/ProfilingDebugging/CountersTrace.h
--- a/Engine/Source/Runtime/Core/Public/ProfilingDebugging/CountersTrace.h
+++ b/Engine/Source/Runtime/Core/Public/ProfilingDebugging/CountersTrace.h
@@ -15,38 +15,127 @@ enum ETraceCounterType
 {
 	TraceCounterType_Int,
 	TraceCounterType_Float,
-	TraceCounterType_Memory,
+};
+
+enum ETraceCounterDisplayHint
+{
+	TraceCounterDisplayHint_None,
+	TraceCounterDisplayHint_Memory,
 };
 
 #if COUNTERSTRACE_ENABLED
 
 struct FCountersTrace
 {
-	struct FCounterInt
-	{
-		CORE_API void Init(const TCHAR* CounterName);
-		CORE_API void Set(int64 Value);
-
-		int64 CurrentValue;
-		uint16 CounterId;
-	};
+	CORE_API static uint16 OutputInitCounter(const TCHAR* CounterName, ETraceCounterType CounterType, ETraceCounterDisplayHint CounterDisplayHint);
+	CORE_API static void OutputSetValue(uint16 CounterId, int64 Value);
+	CORE_API static void OutputSetValue(uint16 CounterId, double Value);
 
-	struct FCounterFloat
+	class FCounter
 	{
-		CORE_API void Init(const TCHAR* CounterName);
-		CORE_API void Set(float Value);
-
-		double CurrentValue;
-		uint16 CounterId;
-	};
-
-	struct FCounterMemory
-	{
-		CORE_API void Init(const TCHAR* CounterName);
-		CORE_API void Set(uint64 Value);
-
-		uint64 CurrentValue;
+	public:
+		FCounter()
+			: IntValue(0)
+		{
+		}
+
+		FCounter(const TCHAR* InCounterName, ETraceCounterType InCounterType, ETraceCounterDisplayHint InCounterDisplayHint)
+			: IntValue(0)
+			, CounterType(InCounterType)
+		{
+			CounterId = OutputInitCounter(InCounterName, InCounterType, InCounterDisplayHint);
+		}
+
+		void Init(const TCHAR* InCounterName, ETraceCounterType InCounterType, ETraceCounterDisplayHint InCounterDisplayHint)
+		{
+			CounterType = InCounterType;
+			CounterId = OutputInitCounter(InCounterName, InCounterType, InCounterDisplayHint);
+		}
+
+		void Set(int64 Value)
+		{
+			if (CounterType == TraceCounterType_Int)
+			{
+				if (IntValue != Value)
+				{
+					IntValue = Value;
+					OutputSetValue(CounterId, IntValue);
+				}
+			}
+			else
+			{
+				double ValueAsDouble = static_cast<double>(Value);
+				if (FloatValue != ValueAsDouble)
+				{
+					FloatValue = ValueAsDouble;
+					OutputSetValue(CounterId, FloatValue);
+				}
+			}
+		}
+
+		void Set(double Value)
+		{
+			if (CounterType == TraceCounterType_Int)
+			{
+				int64 ValueAsInt = static_cast<int64>(Value);
+				if (IntValue != ValueAsInt)
+				{
+					IntValue = ValueAsInt;
+					OutputSetValue(CounterId, IntValue);
+				}
+			}
+			else
+			{
+				if (FloatValue != Value)
+				{
+					FloatValue = Value;
+					OutputSetValue(CounterId, FloatValue);
+				}
+			}
+		}
+
+		void Add(int64 Value)
+		{
+			if (Value != 0)
+			{
+				if (CounterType == TraceCounterType_Int)
+				{
+					IntValue += Value;
+					OutputSetValue(CounterId, IntValue);
+				}
+				else
+				{
+					FloatValue += static_cast<double>(Value);
+					OutputSetValue(CounterId, FloatValue);
+				}
+			}
+		}
+
+		void Add(double Value)
+		{
+			if (Value != 0)
+			{
+				if (CounterType == TraceCounterType_Int)
+				{
+					IntValue += static_cast<int64>(Value);
+					OutputSetValue(CounterId, IntValue);
+				}
+				else
+				{
+					FloatValue += Value;
+					OutputSetValue(CounterId, FloatValue);
+				}
+			}
+		}
+
+	private:
+		union
+		{
+			int64 IntValue;
+			double FloatValue;
+		};
 		uint16 CounterId;
+		ETraceCounterType CounterType;
 	};
 
 	CORE_API static void Init(const TCHAR* CmdLine);
@@ -55,28 +144,65 @@ struct FCountersTrace
 #define TRACE_COUNTERS_INIT(CmdLine) \
 	FCountersTrace::Init(CmdLine);
 
-#define __TRACE_COUNTER_VALUE(CounterType, CounterName, Value) \
-	static FCountersTrace::CounterType PREPROCESSOR_JOIN(__TraceCounter, __LINE__); \
+#define __TRACE_DECLARE_COUNTER(CounterName, CounterType, CounterDisplayHint) \
+	FCountersTrace::FCounter PREPROCESSOR_JOIN(__GTraceCounter, CounterName)(TEXT(#CounterName), CounterType, CounterDisplayHint);
+
+#define __TRACE_DECLARE_INLINE_COUNTER(CounterName, CounterType, CounterDisplayHint) \
+	static FCountersTrace::FCounter PREPROCESSOR_JOIN(__TraceCounter, __LINE__); \
 	if (!PREPROCESSOR_JOIN(__TraceCounter, __LINE__).CounterId) \
 	{ \
-		PREPROCESSOR_JOIN(__TraceCounter, __LINE__).Init(TEXT(#CounterName)); \
-	} \
-	PREPROCESSOR_JOIN(__TraceCounter, __LINE__).Set(Value);
+		PREPROCESSOR_JOIN(__TraceCounter, __LINE__).Init(TEXT(#CounterName), CounterType, CounterDisplayHint); \
+	}
 
-#define TRACE_INTEGER_VALUE(CounterName, Value) \
-	__TRACE_COUNTER_VALUE(FCounterInt, CounterName, Value)
+#define TRACE_INT_VALUE(CounterName, Value) \
+	__TRACE_DECLARE_INLINE_COUNTER(CounterName, TraceCounterType_Int, TraceCounterDisplayHint_None) \
+	PREPROCESSOR_JOIN(__TraceCounter, __LINE__).Set(Value);
 
 #define TRACE_FLOAT_VALUE(CounterName, Value) \
-	__TRACE_COUNTER_VALUE(FCounterFloat, CounterName, Value)
+	__TRACE_DECLARE_INLINE_COUNTER(CounterName, TraceCounterType_Float, TraceCounterDisplayHint_None) \
+	PREPROCESSOR_JOIN(__TraceCounter, __LINE__).Set(Value);
 
 #define TRACE_MEMORY_VALUE(CounterName, Value) \
-	__TRACE_COUNTER_VALUE(FCounterMemory, CounterName, Value)
+	__TRACE_DECLARE_INLINE_COUNTER(CounterName, TraceCounterType_Int, TraceCounterDisplayHint_Memory) \
+	PREPROCESSOR_JOIN(__TraceCounter, __LINE__).Set(Value);
+
+#define TRACE_DECLARE_INT_COUNTER(CounterName) \
+	__TRACE_DECLARE_COUNTER(CounterName, TraceCounterType_Int, TraceCounterDisplayHint_None)
+
+#define TRACE_DECLARE_FLOAT_COUNTER(CounterName) \
+	__TRACE_DECLARE_COUNTER(CounterName, TraceCounterType_Float, TraceCounterDisplayHint_None)
+
+#define TRACE_DECLARE_MEMORY_COUNTER(CounterName) \
+	__TRACE_DECLARE_COUNTER(CounterName, TraceCounterType_Int, TraceCounterDisplayHint_Memory)
+
+#define TRACE_DECLARE_COUNTER_EXTERN(CounterName) \
+	extern FCountersTrace::FCounter PREPROCESSOR_JOIN(__GTraceCounter, CounterName);
+
+#define TRACE_COUNTER_SET(CounterName, Value) \
+	PREPROCESSOR_JOIN(__GTraceCounter, CounterName).Set(Value);
+
+#define TRACE_COUNTER_ADD(CounterName, Value) \
+	PREPROCESSOR_JOIN(__GTraceCounter, CounterName).Add(Value);
+
+#define TRACE_COUNTER_INCREMENT(CounterName) \
+	PREPROCESSOR_JOIN(__GTraceCounter, CounterName).Add(1);
+
+#define TRACE_COUNTER_DECREMENT(CounterName) \
+	PREPROCESSOR_JOIN(__GTraceCounter, CounterName).Add(-1);
 
 #else
 
 #define TRACE_COUNTERS_INIT(CmdLine)
-#define TRACE_INTEGER_VALUE(CounterName, Value)
+#define TRACE_INT_VALUE(CounterName, Value)
 #define TRACE_FLOAT_VALUE(CounterName, Value)
 #define TRACE_MEMORY_VALUE(CounterName, Value)
+#define TRACE_DECLARE_INT_COUNTER(CounterName)
+#define TRACE_DECLARE_FLOAT_COUNTER(CounterName)
+#define TRACE_DECLARE_MEMORY_COUNTER(CounterName)
+#define TRACE_DECLARE_COUNTER_EXTERN(CounterName)
+#define TRACE_COUNTER_SET(CounterName, Value)
+#define TRACE_COUNTER_ADD(CounterName, Value)
+#define TRACE_COUNTER_INCREMENT(CounterName)
+#define TRACE_COUNTER_DECREMENT(CounterName)
 
 #endif
\ No newline at end of file
