Fix for TStringConversion's broken move constructor.

diff --git a/Engine/Source/Runtime/Core/Public/Containers/StringConv.h b/Engine/Source/Runtime/Core/Public/Containers/StringConv.h
--- a/Engine/Source/Runtime/Core/Public/Containers/StringConv.h
+++ b/Engine/Source/Runtime/Core/Public/Containers/StringConv.h
@@ -654,9 +654,14 @@ public:
 	 * Move constructor
 	 */
 	TStringConversion(TStringConversion&& Other)
-		: Converter(MoveTemp((Converter&)Other))
+		: Converter(MoveTemp(ImplicitConv<Converter&>(Other)))
+		, Ptr(Other.Ptr)
+		, StringLength(Other.StringLength)
 	{
 		AllocatorType::MoveToEmpty(Other);
+
+		Other.Ptr          = nullptr;
+		Other.StringLength = 0;
 	}
 
 	/**
