Enabled Cpu profiler trace events on one of the consoles.

diff --git a/Engine/Source/Runtime/Core/Public/ProfilingDebugging/CpuProfilerTrace.h b/Engine/Source/Runtime/Core/Public/ProfilingDebugging/CpuProfilerTrace.h
--- a/Engine/Source/Runtime/Core/Public/ProfilingDebugging/CpuProfilerTrace.h
+++ b/Engine/Source/Runtime/Core/Public/ProfilingDebugging/CpuProfilerTrace.h
@@ -4,7 +4,7 @@
 
 #include "CoreTypes.h"
 
-#if !IS_PROGRAM && !UE_BUILD_SHIPPING && (PLATFORM_WINDOWS || PLATFORM_PS4 || PLATFORM_XBOXONE)
+#if !IS_PROGRAM && !UE_BUILD_SHIPPING && (PLATFORM_WINDOWS || PLATFORM_PS4 || PLATFORM_XBOXONE || PLATFORM_SWITCH)
 #define CPUPROFILERTRACE_ENABLED 1
 #else
 #define CPUPROFILERTRACE_ENABLED 0
