Remove outdated part of comment

#github 6166

diff --git a/Engine/Source/Runtime/Core/Public/Containers/Map.h b/Engine/Source/Runtime/Core/Public/Containers/Map.h
--- a/Engine/Source/Runtime/Core/Public/Containers/Map.h
+++ b/Engine/Source/Runtime/Core/Public/Containers/Map.h
@@ -1583,7 +1583,7 @@ public:
 		return nullptr;
 	}
 
-	/** Adds the (key, value) pair to the map, returning true if the element was added, or false if the element was already present and has been overwritten */
+	/** Adds the (key, value) pair to the map */
 	void Add(
 		const void* Key,
 		const void* Value,
