Fixed the editor hanging after a SIGSEGV crash because the vectored exception handler was unregistered using the wrong handle.
  - Stored and used the correct handle th unergister the handler.

Fixed the editor hanging after a SIGSEGV crash because the CrashReporterClient monitoring the Editor process was killed by the user and the pipe could not be written.
  - Checked if the CrashReporterClient process was still running before trying reporting the crash.

diff --git a/Engine/Source/Runtime/Core/Private/Windows/WindowsPlatformCrashContext.cpp b/Engine/Source/Runtime/Core/Private/Windows/WindowsPlatformCrashContext.cpp
--- a/Engine/Source/Runtime/Core/Private/Windows/WindowsPlatformCrashContext.cpp
+++ b/Engine/Source/Runtime/Core/Private/Windows/WindowsPlatformCrashContext.cpp
@@ -826,6 +826,8 @@ class FCrashReportingThread
 	DWORD CrashingThreadId;
 	/** Handle to crashed thread.*/
 	HANDLE CrashingThreadHandle;
+	/** Handle used to remove vectored exception handler. */
+	HANDLE VectoredExceptionHandle;
 
 	/** Process handle to crash reporter client */
 	FProcHandle CrashClientHandle;
@@ -852,7 +854,7 @@ class FCrashReportingThread
 #if _WIN32_WINNT >= 0x0500
 		if (!FPlatformMisc::IsDebuggerPresent())
 		{
-			RemoveVectoredExceptionHandler(UnhandledStaticInitException);
+			RemoveVectoredExceptionHandler(VectoredExceptionHandle);
 		}
 #endif
 		while (StopTaskCounter.GetValue() == 0)
@@ -886,6 +888,7 @@ public:
 		, ExceptionInfo(nullptr)
 		, CrashingThreadId(0)
 		, CrashingThreadHandle(nullptr)
+		, VectoredExceptionHandle(INVALID_HANDLE_VALUE)
 		, CrashMonitorWritePipe(nullptr)
 		, CrashMonitorReadPipe(nullptr)
 	{
@@ -898,7 +901,7 @@ public:
 #if _WIN32_WINNT >= 0x0500
 		if (!FPlatformMisc::IsDebuggerPresent())
 		{
-			AddVectoredExceptionHandler(1, UnhandledStaticInitException);
+			VectoredExceptionHandle = AddVectoredExceptionHandler(1, UnhandledStaticInitException);
 		}
 #endif
 
@@ -942,7 +945,7 @@ public:
 	/** Ensures are passed trough this. */
 	FORCEINLINE int32 OnEnsure(LPEXCEPTION_POINTERS InExceptionInfo, int NumStackFramesToIgnore, const TCHAR* ErrorMessage, EErrorReportUI ReportUI)
 	{
-		if (CrashClientHandle.IsValid())
+		if (CrashClientHandle.IsValid() && FPlatformProcess::IsProcRunning(CrashClientHandle))
 		{
 			return ReportCrashForMonitor(
 				InExceptionInfo, 
@@ -988,7 +991,7 @@ public:
 	/** Crashes during static init should be reported directly to crash monitor. */
 	FORCEINLINE int32 OnCrashDuringStaticInit(LPEXCEPTION_POINTERS InExceptionInfo)
 	{
-		if (CrashClientHandle.IsValid())
+		if (CrashClientHandle.IsValid() && FPlatformProcess::IsProcRunning(CrashClientHandle))
 		{
 			const ECrashContextType Type = ECrashContextType::Crash;
 			const int NumStackFramesToIgnore = 1;
@@ -1054,7 +1057,7 @@ private:
 		}
 
 #if USE_CRASH_REPORTER_MONITOR
-		if (CrashClientHandle.IsValid())
+		if (CrashClientHandle.IsValid() && FPlatformProcess::IsProcRunning(CrashClientHandle))
 		{
 			// If possible use the crash monitor helper class to report the error. This will do most of the analysis
 			// in the crash reporter client process.
