Structured archive slot types moved to global namespace so that they can be forward-declared.

diff --git a/Engine/Source/Runtime/Core/Private/Serialization/StructuredArchive.cpp b/Engine/Source/Runtime/Core/Private/Serialization/StructuredArchive.cpp
--- a/Engine/Source/Runtime/Core/Private/Serialization/StructuredArchive.cpp
+++ b/Engine/Source/Runtime/Core/Private/Serialization/StructuredArchive.cpp
@@ -27,7 +27,7 @@ struct FStructuredArchive::FContainer
 };
 #endif
 
-FStructuredArchiveChildReader::FStructuredArchiveChildReader(FStructuredArchive::FSlot InSlot)
+FStructuredArchiveChildReader::FStructuredArchiveChildReader(FStructuredArchiveSlot InSlot)
 	: OwnedFormatter(nullptr)
 	, Archive(nullptr)
 {
@@ -79,32 +79,32 @@ FStructuredArchive::~FStructuredArchive()
 	Close();
 }
 
-FStructuredArchive::FSlot FStructuredArchive::Open()
+FStructuredArchiveSlot FStructuredArchive::Open()
 {
 	check(CurrentScope.Num() == 0);
 	check(!RootElementId.IsValid());
 	check(!CurrentSlotElementId.IsValid());
 
 	RootElementId = ElementIdGenerator.Generate();
-	CurrentScope.Emplace(RootElementId, EElementType::Root);
+	CurrentScope.Emplace(RootElementId, UE4StructuredArchive_Private::EElementType::Root);
 
 	CurrentSlotElementId = ElementIdGenerator.Generate();
 
-	return FSlot(*this, 0, CurrentSlotElementId);
+	return FStructuredArchiveSlot(*this, 0, CurrentSlotElementId);
 }
 
 void FStructuredArchive::Close()
 {
-	SetScope(FSlotPosition(0, RootElementId));
+	SetScope(UE4StructuredArchive_Private::FSlotPosition(0, RootElementId));
 }
 
-void FStructuredArchive::EnterSlot(FSlotPosition Slot, bool bEnteringAttributedValue)
+void FStructuredArchive::EnterSlot(UE4StructuredArchive_Private::FSlotPosition Slot, bool bEnteringAttributedValue)
 {
-	int32      ParentDepth = Slot.Depth;
-	FElementId ElementId   = Slot.ElementId;
+	int32                                    ParentDepth = Slot.Depth;
+	UE4StructuredArchive_Private::FElementId ElementId   = Slot.ElementId;
 
 	// If the slot being entered has attributes, enter the value slot first.
-	if (ParentDepth + 1 < CurrentScope.Num() && CurrentScope[ParentDepth + 1].Id == ElementId && CurrentScope[ParentDepth + 1].Type == EElementType::AttributedValue)
+	if (ParentDepth + 1 < CurrentScope.Num() && CurrentScope[ParentDepth + 1].Id == ElementId && CurrentScope[ParentDepth + 1].Type == UE4StructuredArchive_Private::EElementType::AttributedValue)
 	{
 #if DO_STRUCTURED_ARCHIVE_CONTAINER_CHECKS
 		checkf(!CurrentSlotElementId.IsValid() && !CurrentContainer.Top()->bAttributedValueWritten, TEXT("Attempt to serialize data into an invalid slot"));
@@ -113,15 +113,15 @@ void FStructuredArchive::EnterSlot(FSlotPosition Slot, bool bEnteringAttributedV
 		checkf(!CurrentSlotElementId.IsValid(), TEXT("Attempt to serialize data into an invalid slot"));
 #endif
 
-		SetScope(FSlotPosition(ParentDepth + 1, ElementId));
+		SetScope(UE4StructuredArchive_Private::FSlotPosition(ParentDepth + 1, ElementId));
 		Formatter.EnterAttributedValueValue();
 	}
 	else if (!bEnteringAttributedValue && Formatter.TryEnterAttributedValueValue())
 	{
-		int32 NewDepth = EnterSlotAsType(FSlotPosition(ParentDepth, ElementId), EElementType::AttributedValue);
+		int32 NewDepth = EnterSlotAsType(UE4StructuredArchive_Private::FSlotPosition(ParentDepth, ElementId), UE4StructuredArchive_Private::EElementType::AttributedValue);
 		check(NewDepth == ParentDepth + 1);
-		FElementId AttributedValueId = CurrentScope[NewDepth].Id;
-		SetScope(FSlotPosition(NewDepth, AttributedValueId));
+		UE4StructuredArchive_Private::FElementId AttributedValueId = CurrentScope[NewDepth].Id;
+		SetScope(UE4StructuredArchive_Private::FSlotPosition(NewDepth, AttributedValueId));
 #if DO_STRUCTURED_ARCHIVE_CONTAINER_CHECKS
 		CurrentContainer.Emplace(0);
 #endif
@@ -132,12 +132,12 @@ void FStructuredArchive::EnterSlot(FSlotPosition Slot, bool bEnteringAttributedV
 		CurrentSlotElementId.Reset();
 	}
 
-	CurrentEnteringAttributeState = EEnteringAttributeState::NotEnteringAttribute;
+	CurrentEnteringAttributeState = UE4StructuredArchive_Private::EEnteringAttributeState::NotEnteringAttribute;
 }
 
-int32 FStructuredArchive::EnterSlotAsType(FSlotPosition Slot, EElementType ElementType)
+int32 FStructuredArchive::EnterSlotAsType(UE4StructuredArchive_Private::FSlotPosition Slot, UE4StructuredArchive_Private::EElementType ElementType)
 {
-	EnterSlot(Slot, ElementType == EElementType::AttributedValue);
+	EnterSlot(Slot, ElementType == UE4StructuredArchive_Private::EElementType::AttributedValue);
 
 	int32 NewSlotDepth = Slot.Depth + 1;
 
@@ -146,8 +146,8 @@ int32 FStructuredArchive::EnterSlotAsType(FSlotPosition Slot, EElementType Eleme
 	//
 	// We don't need to do adjust for attributes, because entering the attribute slot will bump the depth anyway.
 	if (NewSlotDepth < CurrentScope.Num() &&
-		CurrentScope[NewSlotDepth].Type == EElementType::AttributedValue &&
-		CurrentEnteringAttributeState == EEnteringAttributeState::NotEnteringAttribute)
+		CurrentScope[NewSlotDepth].Type == UE4StructuredArchive_Private::EElementType::AttributedValue &&
+		CurrentEnteringAttributeState == UE4StructuredArchive_Private::EEnteringAttributeState::NotEnteringAttribute)
 	{
 		++NewSlotDepth;
 	}
@@ -162,32 +162,32 @@ void FStructuredArchive::LeaveSlot()
 	{
 		switch (CurrentScope.Top().Type)
 		{
-		case EElementType::Record:
+		case UE4StructuredArchive_Private::EElementType::Record:
 			Formatter.LeaveField();
 			break;
-		case EElementType::Array:
+		case UE4StructuredArchive_Private::EElementType::Array:
 			Formatter.LeaveArrayElement();
 #if DO_STRUCTURED_ARCHIVE_CONTAINER_CHECKS
 			CurrentContainer.Top()->Index++;
 #endif
 			break;
-		case EElementType::Stream:
+		case UE4StructuredArchive_Private::EElementType::Stream:
 			Formatter.LeaveStreamElement();
 			break;
-		case EElementType::Map:
+		case UE4StructuredArchive_Private::EElementType::Map:
 			Formatter.LeaveMapElement();
 #if DO_STRUCTURED_ARCHIVE_CONTAINER_CHECKS
 			CurrentContainer.Top()->Index++;
 #endif
 			break;
-		case EElementType::AttributedValue:
+		case UE4StructuredArchive_Private::EElementType::AttributedValue:
 			Formatter.LeaveAttribute();
 			break;
 		}
 	}
 }
 
-void FStructuredArchive::SetScope(FSlotPosition Slot)
+void FStructuredArchive::SetScope(UE4StructuredArchive_Private::FSlotPosition Slot)
 {
 	// Make sure the scope is valid
 	checkf(Slot.Depth < CurrentScope.Num() && CurrentScope[Slot.Depth].Id == Slot.ElementId, TEXT("Invalid scope for writing to archive"));
@@ -202,13 +202,13 @@ void FStructuredArchive::SetScope(FSlotPosition Slot)
 			const FElement& Element = CurrentScope[CurrentDepth];
 			switch (Element.Type)
 			{
-			case EElementType::Record:
+			case UE4StructuredArchive_Private::EElementType::Record:
 				Formatter.LeaveRecord();
 #if DO_STRUCTURED_ARCHIVE_CONTAINER_CHECKS
 				CurrentContainer.Pop(false);
 #endif
 				break;
-			case EElementType::Array:
+			case UE4StructuredArchive_Private::EElementType::Array:
 #if DO_STRUCTURED_ARCHIVE_CONTAINER_CHECKS
 				checkf(GetUnderlyingArchive().IsLoading() || CurrentContainer.Top()->Index == CurrentContainer.Top()->Count, TEXT("Incorrect number of elements serialized in array"));
 #endif
@@ -217,10 +217,10 @@ void FStructuredArchive::SetScope(FSlotPosition Slot)
 				CurrentContainer.Pop(false);
 #endif
 				break;
-			case EElementType::Stream:
+			case UE4StructuredArchive_Private::EElementType::Stream:
 				Formatter.LeaveStream();
 				break;
-			case EElementType::Map:
+			case UE4StructuredArchive_Private::EElementType::Map:
 #if DO_STRUCTURED_ARCHIVE_CONTAINER_CHECKS
 				checkf(CurrentContainer.Top()->Index == CurrentContainer.Top()->Count, TEXT("Incorrect number of elements serialized in map"));
 #endif
@@ -229,7 +229,7 @@ void FStructuredArchive::SetScope(FSlotPosition Slot)
 				CurrentContainer.Pop(false);
 #endif
 				break;
-			case EElementType::AttributedValue:
+			case UE4StructuredArchive_Private::EElementType::AttributedValue:
 				Formatter.LeaveAttributedValue();
 #if DO_STRUCTURED_ARCHIVE_CONTAINER_CHECKS
 				CurrentContainer.Pop(false);
@@ -251,11 +251,11 @@ void FStructuredArchive::SetScope(FSlotPosition Slot)
 	}
 }
 
-//////////// FStructuredArchive::FSlot ////////////
+//////////// FStructuredArchiveSlot ////////////
 
-FStructuredArchive::FRecord FStructuredArchive::FSlot::EnterRecord()
+FStructuredArchiveRecord FStructuredArchiveSlot::EnterRecord()
 {
-	int32 NewDepth = Ar.EnterSlotAsType(*this, EElementType::Record);
+	int32 NewDepth = Ar.EnterSlotAsType(*this, UE4StructuredArchive_Private::EElementType::Record);
 
 #if DO_STRUCTURED_ARCHIVE_CONTAINER_CHECKS
 	Ar.CurrentContainer.Emplace(0);
@@ -263,12 +263,12 @@ FStructuredArchive::FRecord FStructuredArchive::FSlot::EnterRecord()
 
 	Ar.Formatter.EnterRecord();
 
-	return FRecord(Ar, NewDepth, ElementId);
+	return FStructuredArchiveRecord(Ar, NewDepth, ElementId);
 }
 
-FStructuredArchive::FRecord FStructuredArchive::FSlot::EnterRecord_TextOnly(TArray<FString>& OutFieldNames)
+FStructuredArchiveRecord FStructuredArchiveSlot::EnterRecord_TextOnly(TArray<FString>& OutFieldNames)
 {
-	int32 NewDepth = Ar.EnterSlotAsType(*this, EElementType::Record);
+	int32 NewDepth = Ar.EnterSlotAsType(*this, UE4StructuredArchive_Private::EElementType::Record);
 
 #if DO_STRUCTURED_ARCHIVE_CONTAINER_CHECKS
 	Ar.CurrentContainer.Emplace(0);
@@ -276,12 +276,12 @@ FStructuredArchive::FRecord FStructuredArchive::FSlot::EnterRecord_TextOnly(TArr
 
 	Ar.Formatter.EnterRecord_TextOnly(OutFieldNames);
 
-	return FRecord(Ar, NewDepth, ElementId);
+	return FStructuredArchiveRecord(Ar, NewDepth, ElementId);
 }
 
-FStructuredArchive::FArray FStructuredArchive::FSlot::EnterArray(int32& Num)
+FStructuredArchiveArray FStructuredArchiveSlot::EnterArray(int32& Num)
 {
-	int32 NewDepth = Ar.EnterSlotAsType(*this, EElementType::Array);
+	int32 NewDepth = Ar.EnterSlotAsType(*this, UE4StructuredArchive_Private::EElementType::Array);
 
 	Ar.Formatter.EnterArray(Num);
 
@@ -289,30 +289,30 @@ FStructuredArchive::FArray FStructuredArchive::FSlot::EnterArray(int32& Num)
 	Ar.CurrentContainer.Emplace(Num);
 #endif
 
-	return FArray(Ar, NewDepth, ElementId);
+	return FStructuredArchiveArray(Ar, NewDepth, ElementId);
 }
 
-FStructuredArchive::FStream FStructuredArchive::FSlot::EnterStream()
+FStructuredArchiveStream FStructuredArchiveSlot::EnterStream()
 {
-	int32 NewDepth = Ar.EnterSlotAsType(*this, EElementType::Stream);
+	int32 NewDepth = Ar.EnterSlotAsType(*this, UE4StructuredArchive_Private::EElementType::Stream);
 
 	Ar.Formatter.EnterStream();
 
-	return FStream(Ar, NewDepth, ElementId);
+	return FStructuredArchiveStream(Ar, NewDepth, ElementId);
 }
 
-FStructuredArchive::FStream FStructuredArchive::FSlot::EnterStream_TextOnly(int32& OutNumElements)
+FStructuredArchiveStream FStructuredArchiveSlot::EnterStream_TextOnly(int32& OutNumElements)
 {
-	int32 NewDepth = Ar.EnterSlotAsType(*this, EElementType::Stream);
+	int32 NewDepth = Ar.EnterSlotAsType(*this, UE4StructuredArchive_Private::EElementType::Stream);
 
 	Ar.Formatter.EnterStream_TextOnly(OutNumElements);
 
-	return FStream(Ar, NewDepth, ElementId);
+	return FStructuredArchiveStream(Ar, NewDepth, ElementId);
 }
 
-FStructuredArchive::FMap FStructuredArchive::FSlot::EnterMap(int32& Num)
+FStructuredArchiveMap FStructuredArchiveSlot::EnterMap(int32& Num)
 {
-	int32 NewDepth = Ar.EnterSlotAsType(*this, EElementType::Map);
+	int32 NewDepth = Ar.EnterSlotAsType(*this, UE4StructuredArchive_Private::EElementType::Map);
 
 	Ar.Formatter.EnterMap(Num);
 
@@ -320,17 +320,17 @@ FStructuredArchive::FMap FStructuredArchive::FSlot::EnterMap(int32& Num)
 	Ar.CurrentContainer.Emplace(Num);
 #endif
 
-	return FMap(Ar, NewDepth, ElementId);
+	return FStructuredArchiveMap(Ar, NewDepth, ElementId);
 }
 
-FStructuredArchive::FSlot FStructuredArchive::FSlot::EnterAttribute(FArchiveFieldName AttributeName)
+FStructuredArchiveSlot FStructuredArchiveSlot::EnterAttribute(FArchiveFieldName AttributeName)
 {
 	check(Ar.CurrentScope.Num() > 0);
 
 	int32 NewDepth = Depth + 1;
-	if (NewDepth >= Ar.CurrentScope.Num() || Ar.CurrentScope[NewDepth].Id != ElementId || Ar.CurrentScope[NewDepth].Type != EElementType::AttributedValue)
+	if (NewDepth >= Ar.CurrentScope.Num() || Ar.CurrentScope[NewDepth].Id != ElementId || Ar.CurrentScope[NewDepth].Type != UE4StructuredArchive_Private::EElementType::AttributedValue)
 	{
-		int32 NewDepthCheck = Ar.EnterSlotAsType(*this, EElementType::AttributedValue);
+		int32 NewDepthCheck = Ar.EnterSlotAsType(*this, UE4StructuredArchive_Private::EElementType::AttributedValue);
 		checkSlow(NewDepth == NewDepthCheck);
 
 		Ar.Formatter.EnterAttributedValue();
@@ -340,11 +340,11 @@ FStructuredArchive::FSlot FStructuredArchive::FSlot::EnterAttribute(FArchiveFiel
 #endif
 	}
 
-	Ar.CurrentEnteringAttributeState = EEnteringAttributeState::NotEnteringAttribute;
+	Ar.CurrentEnteringAttributeState = UE4StructuredArchive_Private::EEnteringAttributeState::NotEnteringAttribute;
 
-	FElementId AttributedValueId = Ar.CurrentScope[NewDepth].Id;
+	UE4StructuredArchive_Private::FElementId AttributedValueId = Ar.CurrentScope[NewDepth].Id;
 
-	Ar.SetScope(FSlotPosition(NewDepth, AttributedValueId));
+	Ar.SetScope(UE4StructuredArchive_Private::FSlotPosition(NewDepth, AttributedValueId));
 
 	Ar.CurrentSlotElementId = Ar.ElementIdGenerator.Generate();
 
@@ -361,17 +361,17 @@ FStructuredArchive::FSlot FStructuredArchive::FSlot::EnterAttribute(FArchiveFiel
 
 	Ar.Formatter.EnterAttribute(AttributeName);
 
-	return FSlot(Ar, NewDepth, Ar.CurrentSlotElementId);
+	return FStructuredArchiveSlot(Ar, NewDepth, Ar.CurrentSlotElementId);
 }
 
-TOptional<FStructuredArchive::FSlot> FStructuredArchive::FSlot::TryEnterAttribute(FArchiveFieldName AttributeName, bool bEnterWhenWriting)
+TOptional<FStructuredArchiveSlot> FStructuredArchiveSlot::TryEnterAttribute(FArchiveFieldName AttributeName, bool bEnterWhenWriting)
 {
 	check(Ar.CurrentScope.Num() > 0);
 
 	int32 NewDepth = Depth + 1;
-	if (NewDepth >= Ar.CurrentScope.Num() || Ar.CurrentScope[NewDepth].Id != ElementId || Ar.CurrentScope[NewDepth].Type != EElementType::AttributedValue)
+	if (NewDepth >= Ar.CurrentScope.Num() || Ar.CurrentScope[NewDepth].Id != ElementId || Ar.CurrentScope[NewDepth].Type != UE4StructuredArchive_Private::EElementType::AttributedValue)
 	{
-		int32 NewDepthCheck = Ar.EnterSlotAsType(*this, EElementType::AttributedValue);
+		int32 NewDepthCheck = Ar.EnterSlotAsType(*this, UE4StructuredArchive_Private::EElementType::AttributedValue);
 		checkSlow(NewDepth == NewDepthCheck);
 
 		Ar.Formatter.EnterAttributedValue();
@@ -392,17 +392,17 @@ TOptional<FStructuredArchive::FSlot> FStructuredArchive::FSlot::TryEnterAttribut
 #endif
 #endif
 
-	FElementId AttributedValueId = Ar.CurrentScope[NewDepth].Id;
+	UE4StructuredArchive_Private::FElementId AttributedValueId = Ar.CurrentScope[NewDepth].Id;
 
-	Ar.SetScope(FSlotPosition(NewDepth, AttributedValueId));
+	Ar.SetScope(UE4StructuredArchive_Private::FSlotPosition(NewDepth, AttributedValueId));
 
 	if (Ar.Formatter.TryEnterAttribute(AttributeName, bEnterWhenWriting))
 	{
-		Ar.CurrentEnteringAttributeState = EEnteringAttributeState::NotEnteringAttribute;
+		Ar.CurrentEnteringAttributeState = UE4StructuredArchive_Private::EEnteringAttributeState::NotEnteringAttribute;
 
 		Ar.CurrentSlotElementId = Ar.ElementIdGenerator.Generate();
 
-		return FSlot(Ar, NewDepth, Ar.CurrentSlotElementId);
+		return FStructuredArchiveSlot(Ar, NewDepth, Ar.CurrentSlotElementId);
 	}
 	else
 	{
@@ -410,156 +410,156 @@ TOptional<FStructuredArchive::FSlot> FStructuredArchive::FSlot::TryEnterAttribut
 	}
 }
 
-void FStructuredArchive::FSlot::operator<< (uint8& Value)
+void FStructuredArchiveSlot::operator<< (uint8& Value)
 {
 	Ar.EnterSlot(*this);
 	Ar.Formatter.Serialize(Value);
 	Ar.LeaveSlot();
 }
 
-void FStructuredArchive::FSlot::operator<< (uint16& Value)
+void FStructuredArchiveSlot::operator<< (uint16& Value)
 {
 	Ar.EnterSlot(*this);
 	Ar.Formatter.Serialize(Value);
 	Ar.LeaveSlot();
 }
 
-void FStructuredArchive::FSlot::operator<< (uint32& Value)
+void FStructuredArchiveSlot::operator<< (uint32& Value)
 {
 	Ar.EnterSlot(*this);
 	Ar.Formatter.Serialize(Value);
 	Ar.LeaveSlot();
 }
 
-void FStructuredArchive::FSlot::operator<< (uint64& Value)
+void FStructuredArchiveSlot::operator<< (uint64& Value)
 {
 	Ar.EnterSlot(*this);
 	Ar.Formatter.Serialize(Value);
 	Ar.LeaveSlot();
 }
 
-void FStructuredArchive::FSlot::operator<< (int8& Value)
+void FStructuredArchiveSlot::operator<< (int8& Value)
 {
 	Ar.EnterSlot(*this);
 	Ar.Formatter.Serialize(Value);
 	Ar.LeaveSlot();
 }
 
-void FStructuredArchive::FSlot::operator<< (int16& Value)
+void FStructuredArchiveSlot::operator<< (int16& Value)
 {
 	Ar.EnterSlot(*this);
 	Ar.Formatter.Serialize(Value);
 	Ar.LeaveSlot();
 }
 
-void FStructuredArchive::FSlot::operator<< (int32& Value)
+void FStructuredArchiveSlot::operator<< (int32& Value)
 {
 	Ar.EnterSlot(*this);
 	Ar.Formatter.Serialize(Value);
 	Ar.LeaveSlot();
 }
 
-void FStructuredArchive::FSlot::operator<< (int64& Value)
+void FStructuredArchiveSlot::operator<< (int64& Value)
 {
 	Ar.EnterSlot(*this);
 	Ar.Formatter.Serialize(Value);
 	Ar.LeaveSlot();
 }
 
-void FStructuredArchive::FSlot::operator<< (float& Value)
+void FStructuredArchiveSlot::operator<< (float& Value)
 {
 	Ar.EnterSlot(*this);
 	Ar.Formatter.Serialize(Value);
 	Ar.LeaveSlot();
 }
 
-void FStructuredArchive::FSlot::operator<< (double& Value)
+void FStructuredArchiveSlot::operator<< (double& Value)
 {
 	Ar.EnterSlot(*this);
 	Ar.Formatter.Serialize(Value);
 	Ar.LeaveSlot();
 }
 
-void FStructuredArchive::FSlot::operator<< (bool& Value)
+void FStructuredArchiveSlot::operator<< (bool& Value)
 {
 	Ar.EnterSlot(*this);
 	Ar.Formatter.Serialize(Value);
 	Ar.LeaveSlot();
 }
 
-void FStructuredArchive::FSlot::operator<< (FString& Value)
+void FStructuredArchiveSlot::operator<< (FString& Value)
 {
 	Ar.EnterSlot(*this);
 	Ar.Formatter.Serialize(Value);
 	Ar.LeaveSlot();
 }
 
-void FStructuredArchive::FSlot::operator<< (FName& Value)
+void FStructuredArchiveSlot::operator<< (FName& Value)
 {
 	Ar.EnterSlot(*this);
 	Ar.Formatter.Serialize(Value);
 	Ar.LeaveSlot();
 }
 
-void FStructuredArchive::FSlot::operator<< (UObject*& Value)
+void FStructuredArchiveSlot::operator<< (UObject*& Value)
 {
 	Ar.EnterSlot(*this);
 	Ar.Formatter.Serialize(Value);
 	Ar.LeaveSlot();
 }
 
-void FStructuredArchive::FSlot::operator<< (FText& Value)
+void FStructuredArchiveSlot::operator<< (FText& Value)
 {
 	Ar.EnterSlot(*this);
 	Ar.Formatter.Serialize(Value);
 	Ar.LeaveSlot();
 }
 
-void FStructuredArchive::FSlot::operator<< (FWeakObjectPtr& Value)
+void FStructuredArchiveSlot::operator<< (FWeakObjectPtr& Value)
 {
 	Ar.EnterSlot(*this);
 	Ar.Formatter.Serialize(Value);
 	Ar.LeaveSlot();
 }
 
-void FStructuredArchive::FSlot::operator<< (FLazyObjectPtr& Value)
+void FStructuredArchiveSlot::operator<< (FLazyObjectPtr& Value)
 {
 	Ar.EnterSlot(*this);
 	Ar.Formatter.Serialize(Value);
 	Ar.LeaveSlot();
 }
 
-void FStructuredArchive::FSlot::operator<< (FSoftObjectPtr& Value)
+void FStructuredArchiveSlot::operator<< (FSoftObjectPtr& Value)
 {
 	Ar.EnterSlot(*this);
 	Ar.Formatter.Serialize(Value);
 	Ar.LeaveSlot();
 }
 
-void FStructuredArchive::FSlot::operator<< (FSoftObjectPath& Value)
+void FStructuredArchiveSlot::operator<< (FSoftObjectPath& Value)
 {
 	Ar.EnterSlot(*this);
 	Ar.Formatter.Serialize(Value);
 	Ar.LeaveSlot();
 }
 
-void FStructuredArchive::FSlot::Serialize(TArray<uint8>& Data)
+void FStructuredArchiveSlot::Serialize(TArray<uint8>& Data)
 {
 	Ar.EnterSlot(*this);
 	Ar.Formatter.Serialize(Data);
 	Ar.LeaveSlot();
 }
 
-void FStructuredArchive::FSlot::Serialize(void* Data, uint64 DataSize)
+void FStructuredArchiveSlot::Serialize(void* Data, uint64 DataSize)
 {
 	Ar.EnterSlot(*this);
 	Ar.Formatter.Serialize(Data, DataSize);
 	Ar.LeaveSlot();
 }
 
-//////////// FStructuredArchive::FObject ////////////
+//////////// FStructuredArchiveRecord ////////////
 
-FStructuredArchive::FSlot FStructuredArchive::FRecord::EnterField(FArchiveFieldName Name)
+FStructuredArchiveSlot FStructuredArchiveRecord::EnterField(FArchiveFieldName Name)
 {
 	Ar.SetScope(*this);
 
@@ -578,40 +578,40 @@ FStructuredArchive::FSlot FStructuredArchive::FRecord::EnterField(FArchiveFieldN
 
 	Ar.Formatter.EnterField(Name);
 
-	return FSlot(Ar, Depth, Ar.CurrentSlotElementId);
+	return FStructuredArchiveSlot(Ar, Depth, Ar.CurrentSlotElementId);
 }
 
-FStructuredArchive::FRecord FStructuredArchive::FRecord::EnterRecord(FArchiveFieldName Name)
+FStructuredArchiveRecord FStructuredArchiveRecord::EnterRecord(FArchiveFieldName Name)
 {
 	return EnterField(Name).EnterRecord();
 }
 
-FStructuredArchive::FRecord FStructuredArchive::FRecord::EnterRecord_TextOnly(FArchiveFieldName Name, TArray<FString>& OutFieldNames)
+FStructuredArchiveRecord FStructuredArchiveRecord::EnterRecord_TextOnly(FArchiveFieldName Name, TArray<FString>& OutFieldNames)
 {
 	return EnterField(Name).EnterRecord_TextOnly(OutFieldNames);
 }
 
-FStructuredArchive::FArray FStructuredArchive::FRecord::EnterArray(FArchiveFieldName Name, int32& Num)
+FStructuredArchiveArray FStructuredArchiveRecord::EnterArray(FArchiveFieldName Name, int32& Num)
 {
 	return EnterField(Name).EnterArray(Num);
 }
 
-FStructuredArchive::FStream FStructuredArchive::FRecord::EnterStream(FArchiveFieldName Name)
+FStructuredArchiveStream FStructuredArchiveRecord::EnterStream(FArchiveFieldName Name)
 {
 	return EnterField(Name).EnterStream();
 }
 
-FStructuredArchive::FStream FStructuredArchive::FRecord::EnterStream_TextOnly(FArchiveFieldName Name, int32& OutNumElements)
+FStructuredArchiveStream FStructuredArchiveRecord::EnterStream_TextOnly(FArchiveFieldName Name, int32& OutNumElements)
 {
 	return EnterField(Name).EnterStream_TextOnly(OutNumElements);
 }
 
-FStructuredArchive::FMap FStructuredArchive::FRecord::EnterMap(FArchiveFieldName Name, int32& Num)
+FStructuredArchiveMap FStructuredArchiveRecord::EnterMap(FArchiveFieldName Name, int32& Num)
 {
 	return EnterField(Name).EnterMap(Num);
 }
 
-TOptional<FStructuredArchive::FSlot> FStructuredArchive::FRecord::TryEnterField(FArchiveFieldName Name, bool bEnterWhenWriting)
+TOptional<FStructuredArchiveSlot> FStructuredArchiveRecord::TryEnterField(FArchiveFieldName Name, bool bEnterWhenWriting)
 {
 	Ar.SetScope(*this);
 
@@ -629,7 +629,7 @@ TOptional<FStructuredArchive::FSlot> FStructuredArchive::FRecord::TryEnterField(
 	if (Ar.Formatter.TryEnterField(Name, bEnterWhenWriting))
 	{
 		Ar.CurrentSlotElementId = Ar.ElementIdGenerator.Generate();
-		return FSlot(Ar, Depth, Ar.CurrentSlotElementId);
+		return FStructuredArchiveSlot(Ar, Depth, Ar.CurrentSlotElementId);
 	}
 	else
 	{
@@ -637,9 +637,9 @@ TOptional<FStructuredArchive::FSlot> FStructuredArchive::FRecord::TryEnterField(
 	}
 }
 
-//////////// FStructuredArchive::FArray ////////////
+//////////// FStructuredArchiveArray ////////////
 
-FStructuredArchive::FSlot FStructuredArchive::FArray::EnterElement()
+FStructuredArchiveSlot FStructuredArchiveArray::EnterElement()
 {
 	Ar.SetScope(*this);
 
@@ -651,10 +651,10 @@ FStructuredArchive::FSlot FStructuredArchive::FArray::EnterElement()
 
 	Ar.Formatter.EnterArrayElement();
 
-	return FSlot(Ar, Depth, Ar.CurrentSlotElementId);
+	return FStructuredArchiveSlot(Ar, Depth, Ar.CurrentSlotElementId);
 }
 
-FStructuredArchive::FSlot FStructuredArchive::FArray::EnterElement_TextOnly(EArchiveValueType& OutType)
+FStructuredArchiveSlot FStructuredArchiveArray::EnterElement_TextOnly(EArchiveValueType& OutType)
 {
 	Ar.SetScope(*this);
 
@@ -666,12 +666,12 @@ FStructuredArchive::FSlot FStructuredArchive::FArray::EnterElement_TextOnly(EArc
 
 	Ar.Formatter.EnterArrayElement_TextOnly(OutType);
 
-	return FSlot(Ar, Depth, Ar.CurrentSlotElementId);
+	return FStructuredArchiveSlot(Ar, Depth, Ar.CurrentSlotElementId);
 }
 
-//////////// FStructuredArchive::FStream ////////////
+//////////// FStructuredArchiveStream ////////////
 
-FStructuredArchive::FSlot FStructuredArchive::FStream::EnterElement()
+FStructuredArchiveSlot FStructuredArchiveStream::EnterElement()
 {
 	Ar.SetScope(*this);
 
@@ -679,10 +679,10 @@ FStructuredArchive::FSlot FStructuredArchive::FStream::EnterElement()
 
 	Ar.Formatter.EnterStreamElement();
 
-	return FSlot(Ar, Depth, Ar.CurrentSlotElementId);
+	return FStructuredArchiveSlot(Ar, Depth, Ar.CurrentSlotElementId);
 }
 
-FStructuredArchive::FSlot FStructuredArchive::FStream::EnterElement_TextOnly(EArchiveValueType& OutType)
+FStructuredArchiveSlot FStructuredArchiveStream::EnterElement_TextOnly(EArchiveValueType& OutType)
 {
 	Ar.SetScope(*this);
 
@@ -690,12 +690,12 @@ FStructuredArchive::FSlot FStructuredArchive::FStream::EnterElement_TextOnly(EAr
 
 	Ar.Formatter.EnterStreamElement_TextOnly(OutType);
 
-	return FSlot(Ar, Depth, Ar.CurrentSlotElementId);
+	return FStructuredArchiveSlot(Ar, Depth, Ar.CurrentSlotElementId);
 }
 
-//////////// FStructuredArchive::FMap ////////////
+//////////// FStructuredArchiveMap ////////////
 
-FStructuredArchive::FSlot FStructuredArchive::FMap::EnterElement(FString& Name)
+FStructuredArchiveSlot FStructuredArchiveMap::EnterElement(FString& Name)
 {
 	Ar.SetScope(*this);
 
@@ -729,10 +729,10 @@ FStructuredArchive::FSlot FStructuredArchive::FMap::EnterElement(FString& Name)
 #endif
 #endif
 
-	return FSlot(Ar, Depth, Ar.CurrentSlotElementId);
+	return FStructuredArchiveSlot(Ar, Depth, Ar.CurrentSlotElementId);
 }
 
-FStructuredArchive::FSlot FStructuredArchive::FMap::EnterElement_TextOnly(FString& Name, EArchiveValueType& OutType)
+FStructuredArchiveSlot FStructuredArchiveMap::EnterElement_TextOnly(FString& Name, EArchiveValueType& OutType)
 {
 	Ar.SetScope(*this);
 
@@ -766,7 +766,7 @@ FStructuredArchive::FSlot FStructuredArchive::FMap::EnterElement_TextOnly(FStrin
 #endif
 #endif
 
-	return FSlot(Ar, Depth, Ar.CurrentSlotElementId);
+	return FStructuredArchiveSlot(Ar, Depth, Ar.CurrentSlotElementId);
 }
 
 #endif
diff --git a/Engine/Source/Runtime/Core/Public/Serialization/StructuredArchive.h b/Engine/Source/Runtime/Core/Public/Serialization/StructuredArchive.h
--- a/Engine/Source/Runtime/Core/Public/Serialization/StructuredArchive.h
+++ b/Engine/Source/Runtime/Core/Public/Serialization/StructuredArchive.h
@@ -130,55 +130,16 @@ template<typename T> FORCEINLINE TDefaultedNamedAttribute<T> MakeDefaultedNamedA
 	typedef FBinaryArchiveFormatter FArchiveFormatterType;
 #endif
 
-/**
- * Manages the state of an underlying FStructuredArchiveFormatter, and provides a consistent API for reading and writing to a structured archive.
- * 
- * Both reading and writing to the archive are *forward only* from an interface point of view. There is no point at which it is possible to 
- * require seeking.
- */
-class CORE_API FStructuredArchive
+class FStructuredArchive;
+class FStructuredArchiveChildReader;
+class FStructuredArchiveSlot;
+class FStructuredArchiveRecord;
+class FStructuredArchiveArray;
+class FStructuredArchiveStream;
+class FStructuredArchiveMap;
+
+namespace UE4StructuredArchive_Private
 {
-public:
-	class FSlot;
-	class FRecord;
-	class FArray;
-	class FStream;
-	class FMap;
-
-	/**
-	 * Constructor.
-	 *
-	 * @param InFormatter Formatter for the archive data
-	 */
-	FStructuredArchive(FArchiveFormatterType& InFormatter);
-	
-	/**
-	 * Default destructor. Closes the archive.
-	 */
-	~FStructuredArchive();
-
-	/**
-	 * Start writing to the archive, and gets an interface to the root slot.
-	 */
-	FSlot Open();
-
-	/**
-	 * Flushes any remaining scope to the underlying formatter and closes the archive.
-	 */
-	void Close();
-
-	/**
-	 * Gets the serialization context from the underlying archive.
-	 */
-	FORCEINLINE FArchive& GetUnderlyingArchive()
-	{
-		return Formatter.GetUnderlyingArchive();
-	}
-
-	FStructuredArchive(const FStructuredArchive&) = delete;
-	FStructuredArchive& operator=(const FStructuredArchive&) = delete;
-
-private:
 	struct FElementId
 	{
 		FElementId() = default;
@@ -217,7 +178,7 @@ private:
 	{
 		friend class FStructuredArchive;
 
-	protected:
+	public:
 		int32 Depth;
 		FElementId ElementId;
 
@@ -237,10 +198,7 @@ private:
 		friend class FStructuredArchive;
 
 	public:
-		FORCEINLINE FArchive& GetUnderlyingArchive()
-		{
-			return Ar.GetUnderlyingArchive();
-		}
+		FArchive& GetUnderlyingArchive() const;
 
 	protected:
 		FStructuredArchive& Ar;
@@ -259,172 +217,244 @@ private:
 #endif
 	};
 
+	enum class EElementType : unsigned char
+	{
+		Root,
+		Record,
+		Array,
+		Stream,
+		Map,
+		AttributedValue,
+	};
+
+	enum class EEnteringAttributeState
+	{
+		NotEnteringAttribute,
+		EnteringAttribute,
+	};
+}
+
+/**
+ * Contains a value in the archive; either a field or array/map element. A slot does not know it's name or location,
+ * and can merely have a value serialized into it. That value may be a literal (eg. int, float) or compound object
+ * (eg. object, array, map).
+ */
+class CORE_API FStructuredArchiveSlot final : public UE4StructuredArchive_Private::FSlotBase
+{
 public:
-	/**
-	 * Contains a value in the archive; either a field or array/map element. A slot does not know it's name or location,
-	 * and can merely have a value serialized into it. That value may be a literal (eg. int, float) or compound object
-	 * (eg. object, array, map).
-	 */
-	class CORE_API FSlot final : public FSlotBase
+	FStructuredArchiveRecord EnterRecord();
+	FStructuredArchiveRecord EnterRecord_TextOnly(TArray<FString>& OutFieldNames);
+	FStructuredArchiveArray EnterArray(int32& Num);
+	FStructuredArchiveStream EnterStream();
+	FStructuredArchiveStream EnterStream_TextOnly(int32& OutNumElements);
+	FStructuredArchiveMap EnterMap(int32& Num);
+	FStructuredArchiveSlot EnterAttribute(FArchiveFieldName AttributeName);
+	TOptional<FStructuredArchiveSlot> TryEnterAttribute(FArchiveFieldName AttributeName, bool bEnterWhenWriting);
+
+	// We don't support chaining writes to a single slot, so this returns void.
+	void operator << (uint8& Value);
+	void operator << (uint16& Value);
+	void operator << (uint32& Value);
+	void operator << (uint64& Value);
+	void operator << (int8& Value);
+	void operator << (int16& Value);
+	void operator << (int32& Value);
+	void operator << (int64& Value);
+	void operator << (float& Value);
+	void operator << (double& Value);
+	void operator << (bool& Value);
+	void operator << (FString& Value);
+	void operator << (FName& Value);
+	void operator << (UObject*& Value);
+	void operator << (FText& Value);
+	void operator << (FWeakObjectPtr& Value);
+	void operator << (FSoftObjectPtr& Value);
+	void operator << (FSoftObjectPath& Value);
+	void operator << (FLazyObjectPtr& Value);
+
+
+	template <typename T>
+	FORCEINLINE void operator<<(TNamedAttribute<T> Item)
+	{
+		EnterAttribute(Item.Name) << Item.Value;
+	}
+
+	template <typename T>
+	FORCEINLINE void operator<<(TDefaultedNamedAttribute<T> Item)
 	{
-	public:
-		FRecord EnterRecord();
-		FRecord EnterRecord_TextOnly(TArray<FString>& OutFieldNames);
-		FArray EnterArray(int32& Num);
-		FStream EnterStream();
-		FStream EnterStream_TextOnly(int32& OutNumElements);
-		FMap EnterMap(int32& Num);
-		FSlot EnterAttribute(FArchiveFieldName AttributeName);
-		TOptional<FSlot> TryEnterAttribute(FArchiveFieldName AttributeName, bool bEnterWhenWriting);
-
-		// We don't support chaining writes to a single slot, so this returns void.
-		void operator << (uint8& Value);
-		void operator << (uint16& Value);
-		void operator << (uint32& Value);
-		void operator << (uint64& Value);
-		void operator << (int8& Value);
-		void operator << (int16& Value);
-		void operator << (int32& Value);
-		void operator << (int64& Value);
-		void operator << (float& Value);
-		void operator << (double& Value);
-		void operator << (bool& Value);
-		void operator << (FString& Value);
-		void operator << (FName& Value);
-		void operator << (UObject*& Value);
-		void operator << (FText& Value);
-		void operator << (FWeakObjectPtr& Value);
-		void operator << (FSoftObjectPtr& Value);
-		void operator << (FSoftObjectPath& Value);
-		void operator << (FLazyObjectPtr& Value);
-
-
-		template <typename T>
-		FORCEINLINE void operator<<(TNamedAttribute<T> Item)
+		if (TOptional<FStructuredArchiveSlot> Attribute = TryEnterAttribute(Item.Name, Item.Value != Item.Default))
 		{
-			EnterAttribute(Item.Name) << Item.Value;
+			Attribute.GetValue() << Item.Value;
 		}
-
-		template <typename T>
-		FORCEINLINE void operator<<(TDefaultedNamedAttribute<T> Item)
+		else
 		{
-			if (TOptional<FSlot> Attribute = TryEnterAttribute(Item.Name, Item.Value != Item.Default))
-			{
-				Attribute.GetValue() << Item.Value;
-			}
-			else
-			{
-				Item.Value = Item.Default;
-			}
+			Item.Value = Item.Default;
 		}
+	}
 
-		void Serialize(TArray<uint8>& Data);
-		void Serialize(void* Data, uint64 DataSize);
+	void Serialize(TArray<uint8>& Data);
+	void Serialize(void* Data, uint64 DataSize);
 
-		FORCEINLINE bool IsFilled() const
-		{
-#if WITH_TEXT_ARCHIVE_SUPPORT
-			return Ar.CurrentSlotElementId != ElementId;
-#else
-			return true;
-#endif
-		}
+	bool IsFilled() const;
 
-	private:
-		friend FStructuredArchive;
-		friend class FStructuredArchiveChildReader;
+private:
+	friend FStructuredArchive;
+	friend FStructuredArchiveChildReader;
+	friend FStructuredArchiveSlot;
+	friend FStructuredArchiveRecord;
+	friend FStructuredArchiveArray;
+	friend FStructuredArchiveStream;
+	friend FStructuredArchiveMap;
+
+	using UE4StructuredArchive_Private::FSlotBase::FSlotBase;
+};
 
-		using FSlotBase::FSlotBase;
-	};
+/**
+ * Represents a record in the structured archive. An object contains slots that are identified by FArchiveName,
+ * which may be compiled out with binary-only archives.
+ */
+class CORE_API FStructuredArchiveRecord final : public UE4StructuredArchive_Private::FSlotBase
+{
+public:
+	FStructuredArchiveSlot EnterField(FArchiveFieldName Name);
+	FStructuredArchiveSlot EnterField_TextOnly(FArchiveFieldName Name, EArchiveValueType& OutType);
+	FStructuredArchiveRecord EnterRecord(FArchiveFieldName Name);
+	FStructuredArchiveRecord EnterRecord_TextOnly(FArchiveFieldName Name, TArray<FString>& OutFieldNames);
+	FStructuredArchiveArray EnterArray(FArchiveFieldName Name, int32& Num);
+	FStructuredArchiveStream EnterStream(FArchiveFieldName Name);
+	FStructuredArchiveStream EnterStream_TextOnly(FArchiveFieldName Name, int32& OutNumElements);
+	FStructuredArchiveMap EnterMap(FArchiveFieldName Name, int32& Num);
 
-	/**
-	 * Represents a record in the structured archive. An object contains slots that are identified by FArchiveName,
-	 * which may be compiled out with binary-only archives.
-	 */
-	class CORE_API FRecord final : public FSlotBase
+	TOptional<FStructuredArchiveSlot> TryEnterField(FArchiveFieldName Name, bool bEnterForSaving);
+
+	template<typename T> FORCEINLINE FStructuredArchiveRecord& operator<<(TNamedValue<T> Item)
 	{
-	public:
-		FSlot EnterField(FArchiveFieldName Name);
-		FSlot EnterField_TextOnly(FArchiveFieldName Name, EArchiveValueType& OutType);
-		FRecord EnterRecord(FArchiveFieldName Name);
-		FRecord EnterRecord_TextOnly(FArchiveFieldName Name, TArray<FString>& OutFieldNames);
-		FArray EnterArray(FArchiveFieldName Name, int32& Num);
-		FStream EnterStream(FArchiveFieldName Name);
-		FStream EnterStream_TextOnly(FArchiveFieldName Name, int32& OutNumElements);
-		FMap EnterMap(FArchiveFieldName Name, int32& Num);
-
-		TOptional<FSlot> TryEnterField(FArchiveFieldName Name, bool bEnterForSaving);
-
-		template<typename T> FORCEINLINE FRecord& operator<<(TNamedValue<T> Item)
-		{
-			EnterField(Item.Name) << Item.Value;
-			return *this;
-		}
+		EnterField(Item.Name) << Item.Value;
+		return *this;
+	}
 
-	private:
-		friend FStructuredArchive;
+private:
+	friend FStructuredArchive;
+	friend FStructuredArchiveSlot;
 
-		using FSlotBase::FSlotBase;
-	};
+	using UE4StructuredArchive_Private::FSlotBase::FSlotBase;
+};
 
-	/**
-	 * Represents an array in the structured archive. An object contains slots that are identified by a FArchiveFieldName,
-	 * which may be compiled out with binary-only archives.
-	 */
-	class CORE_API FArray final : public FSlotBase
+/**
+ * Represents an array in the structured archive. An object contains slots that are identified by a FArchiveFieldName,
+ * which may be compiled out with binary-only archives.
+ */
+class CORE_API FStructuredArchiveArray final : public UE4StructuredArchive_Private::FSlotBase
+{
+public:
+	FStructuredArchiveSlot EnterElement();
+	FStructuredArchiveSlot EnterElement_TextOnly(EArchiveValueType& OutType);
+
+	template<typename T> FORCEINLINE FStructuredArchiveArray& operator<<(T& Item)
 	{
-	public:
-		FSlot EnterElement();
-		FSlot EnterElement_TextOnly(EArchiveValueType& OutType);
+		EnterElement() << Item;
+		return *this;
+	}
 
-		template<typename T> FORCEINLINE FArray& operator<<(T& Item)
-		{
-			EnterElement() << Item;
-			return *this;
-		}
+private:
+	friend FStructuredArchive;
+	friend FStructuredArchiveSlot;
 
-	private:
-		friend FStructuredArchive;
+	using UE4StructuredArchive_Private::FSlotBase::FSlotBase;
+};
 
-		using FSlotBase::FSlotBase;
-	};
+/**
+ * Represents an unsized sequence of slots in the structured archive (similar to an array, but without a known size).
+ */
+class CORE_API FStructuredArchiveStream final : public UE4StructuredArchive_Private::FSlotBase
+{
+public:
+	FStructuredArchiveSlot EnterElement();
+	FStructuredArchiveSlot EnterElement_TextOnly(EArchiveValueType& OutType);
 
-	/**
-	 * Represents an unsized sequence of slots in the structured archive (similar to an array, but without a known size).
-	 */
-	class CORE_API FStream final : public FSlotBase
+	template<typename T> FORCEINLINE FStructuredArchiveStream& operator<<(T& Item)
 	{
-	public:
-		FSlot EnterElement();
-		FSlot EnterElement_TextOnly(EArchiveValueType& OutType);
+		EnterElement() << Item;
+		return *this;
+	}
 
-		template<typename T> FORCEINLINE FStream& operator<<(T& Item)
-		{
-			EnterElement() << Item;
-			return *this;
-		}
+private:
+	friend FStructuredArchive;
+	friend FStructuredArchiveSlot;
 
-	private:
-		friend FStructuredArchive;
+	using UE4StructuredArchive_Private::FSlotBase::FSlotBase;
+};
 
-		using FSlotBase::FSlotBase;
-	};
+/**
+ * Represents a map in the structured archive. A map is similar to a record, but keys can be read back out from an archive.
+ * (This is an important distinction for binary archives).
+ */
+class CORE_API FStructuredArchiveMap final : public UE4StructuredArchive_Private::FSlotBase
+{
+public:
+	FStructuredArchiveSlot EnterElement(FString& Name);
+	FStructuredArchiveSlot EnterElement_TextOnly(FString& Name, EArchiveValueType& OutType);
+
+private:
+	friend FStructuredArchive;
+	friend FStructuredArchiveSlot;
+
+	using UE4StructuredArchive_Private::FSlotBase::FSlotBase;
+};
+
+/**
+ * Manages the state of an underlying FStructuredArchiveFormatter, and provides a consistent API for reading and writing to a structured archive.
+ * 
+ * Both reading and writing to the archive are *forward only* from an interface point of view. There is no point at which it is possible to 
+ * require seeking.
+ */
+class CORE_API FStructuredArchive
+{
+	friend FStructuredArchiveSlot;
+	friend FStructuredArchiveRecord;
+	friend FStructuredArchiveArray;
+	friend FStructuredArchiveStream;
+	friend FStructuredArchiveMap;
+
+public:
+	using FSlot   = FStructuredArchiveSlot;
+	using FRecord = FStructuredArchiveRecord;
+	using FArray  = FStructuredArchiveArray;
+	using FStream = FStructuredArchiveStream;
+	using FMap    = FStructuredArchiveMap;
 
 	/**
-	 * Represents a map in the structured archive. A map is similar to a record, but keys can be read back out from an archive.
-	 * (This is an important distinction for binary archives).
+	 * Constructor.
+	 *
+	 * @param InFormatter Formatter for the archive data
 	 */
-	class CORE_API FMap final : public FSlotBase
-	{
-	public:
-		FSlot EnterElement(FString& Name);
-		FSlot EnterElement_TextOnly(FString& Name, EArchiveValueType& OutType);
+	FStructuredArchive(FArchiveFormatterType& InFormatter);
+	
+	/**
+	 * Default destructor. Closes the archive.
+	 */
+	~FStructuredArchive();
 
-	private:
-		friend FStructuredArchive;
+	/**
+	 * Start writing to the archive, and gets an interface to the root slot.
+	 */
+	FStructuredArchiveSlot Open();
 
-		using FSlotBase::FSlotBase;
-	};
+	/**
+	 * Flushes any remaining scope to the underlying formatter and closes the archive.
+	 */
+	void Close();
+
+	/**
+	 * Gets the serialization context from the underlying archive.
+	 */
+	FORCEINLINE FArchive& GetUnderlyingArchive() const
+	{
+		return Formatter.GetUnderlyingArchive();
+	}
+
+	FStructuredArchive(const FStructuredArchive&) = delete;
+	FStructuredArchive& operator=(const FStructuredArchive&) = delete;
 
 private:
 
@@ -441,28 +471,12 @@ private:
 	 */
 	const bool bRequiresStructuralMetadata;
 
-	enum class EElementType : unsigned char
-	{
-		Root,
-		Record,
-		Array,
-		Stream,
-		Map,
-		AttributedValue,
-	};
-
-	enum class EEnteringAttributeState
-	{
-		NotEnteringAttribute,
-		EnteringAttribute,
-	};
-
 	struct FElement
 	{
-		FElementId Id;
-		EElementType Type;
+		UE4StructuredArchive_Private::FElementId Id;
+		UE4StructuredArchive_Private::EElementType Type;
 
-		FElement(FElementId InId, EElementType InType)
+		FElement(UE4StructuredArchive_Private::FElementId InId, UE4StructuredArchive_Private::EElementType InType)
 			: Id(InId)
 			, Type(InType)
 		{
@@ -471,9 +485,9 @@ private:
 
 	struct FIdGenerator
 	{
-		FElementId Generate()
+		UE4StructuredArchive_Private::FElementId Generate()
 		{
-			return FElementId(NextId++);
+			return UE4StructuredArchive_Private::FElementId(NextId++);
 		}
 
 	private:
@@ -488,12 +502,12 @@ private:
 	/**
 	 * The ID of the root element.
 	 */
-	FElementId RootElementId;
+	UE4StructuredArchive_Private::FElementId RootElementId;
 
 	/**
 	 * The element ID assigned for the current slot. Slots are transient, and only exist as placeholders until something is written into them. This is reset to 0 when something is created in a slot, and the created item can assume the element id.
 	 */
-	FElementId CurrentSlotElementId;
+	UE4StructuredArchive_Private::FElementId CurrentSlotElementId;
 
 	/**
 	 * Tracks the current stack of objects being written. Used by SetScope() to ensure that scopes are always closed correctly in the underlying formatter,
@@ -513,19 +527,19 @@ private:
 	/**
 	 * Whether or not we've just entered an attribute
 	 */
-	EEnteringAttributeState CurrentEnteringAttributeState = EEnteringAttributeState::NotEnteringAttribute;
+	UE4StructuredArchive_Private::EEnteringAttributeState CurrentEnteringAttributeState = UE4StructuredArchive_Private::EEnteringAttributeState::NotEnteringAttribute;
 
 	/**
 	 * Enters the current slot for serializing a value. Asserts if the archive is not in a state about to write to an empty-slot.
 	 */
-	void EnterSlot(FSlotPosition Slot, bool bEnteringAttributedValue = false);
+	void EnterSlot(UE4StructuredArchive_Private::FSlotPosition Slot, bool bEnteringAttributedValue = false);
 
 	/**
 	 * Enters the current slot, adding an element onto the stack. Asserts if the archive is not in a state about to write to an empty-slot.
 	 *
 	 * @return  The depth of the newly-entered slot.
 	 */
-	int32 EnterSlotAsType(FSlotPosition Slot, EElementType ElementType);
+	int32 EnterSlotAsType(UE4StructuredArchive_Private::FSlotPosition Slot, UE4StructuredArchive_Private::EElementType ElementType);
 
 	/**
 	 * Leaves slot at the top of the current scope
@@ -535,15 +549,29 @@ private:
 	/**
 	 * Switches to the scope for the given slot.
 	 */
-	void SetScope(FSlotPosition Slot);
+	void SetScope(UE4StructuredArchive_Private::FSlotPosition Slot);
 #endif
 };
 
+FORCEINLINE FArchive& UE4StructuredArchive_Private::FSlotBase::GetUnderlyingArchive() const
+{
+	return Ar.GetUnderlyingArchive();
+}
+
+FORCEINLINE bool FStructuredArchiveSlot::IsFilled() const
+{
+#if WITH_TEXT_ARCHIVE_SUPPORT
+	return Ar.CurrentSlotElementId != ElementId;
+#else
+	return true;
+#endif
+}
+
 template <typename T>
-FORCEINLINE_DEBUGGABLE void operator<<(FStructuredArchive::FSlot Slot, TArray<T>& InArray)
+FORCEINLINE_DEBUGGABLE void operator<<(FStructuredArchiveSlot Slot, TArray<T>& InArray)
 {
 	int32 NumElements = InArray.Num();
-	FStructuredArchive::FArray Array = Slot.EnterArray(NumElements);
+	FStructuredArchiveArray Array = Slot.EnterArray(NumElements);
 
 	if (Slot.GetUnderlyingArchive().IsLoading())
 	{
@@ -552,13 +580,13 @@ FORCEINLINE_DEBUGGABLE void operator<<(FStructuredArchive::FSlot Slot, TArray<T>
 
 	for (int32 ElementIndex = 0; ElementIndex < NumElements; ++ElementIndex)
 	{
-		FStructuredArchive::FSlot ElementSlot = Array.EnterElement();
+		FStructuredArchiveSlot ElementSlot = Array.EnterElement();
 		ElementSlot << InArray[ElementIndex];
 	}
 }
 
 template <>
-FORCEINLINE_DEBUGGABLE void operator<<(FStructuredArchive::FSlot Slot, TArray<uint8>& InArray)
+FORCEINLINE_DEBUGGABLE void operator<<(FStructuredArchiveSlot Slot, TArray<uint8>& InArray)
 {
 	Slot.Serialize(InArray);
 }
@@ -576,16 +604,16 @@ class CORE_API FStructuredArchiveChildReader
 {
 public:
 
-	FStructuredArchiveChildReader(FStructuredArchive::FSlot InSlot);
+	FStructuredArchiveChildReader(FStructuredArchiveSlot InSlot);
 	~FStructuredArchiveChildReader();
 
-	FORCEINLINE FStructuredArchive::FSlot GetRoot() { return Root.GetValue(); }
+	FORCEINLINE FStructuredArchiveSlot GetRoot() { return Root.GetValue(); }
 
 private:
 
 	FStructuredArchiveFormatter* OwnedFormatter;
 	FStructuredArchive* Archive;
-	TOptional<FStructuredArchive::FSlot> Root;
+	TOptional<FStructuredArchiveSlot> Root;
 };
 
 class CORE_API FStructuredArchiveFromArchive
@@ -598,7 +626,7 @@ public:
 	explicit FStructuredArchiveFromArchive(FArchive& Ar);
 	~FStructuredArchiveFromArchive();
 
-	FStructuredArchive::FSlot GetSlot();
+	FStructuredArchiveSlot GetSlot();
 
 private:
 	// Implmented as a pimpl in order to reduce dependencies
@@ -614,7 +642,7 @@ class CORE_API FArchiveFromStructuredArchiveImpl : public FArchiveProxy
 	struct FImpl;
 
 public:
-	explicit FArchiveFromStructuredArchiveImpl(FStructuredArchive::FSlot Slot);
+	explicit FArchiveFromStructuredArchiveImpl(FStructuredArchiveSlot Slot);
 	virtual ~FArchiveFromStructuredArchiveImpl();
 
 	virtual void Flush() override;
@@ -640,7 +668,7 @@ public:
 	bool ContainsData() const;
 
 protected:
-	virtual void SerializeInternal(FStructuredArchive::FRecord Record);
+	virtual void SerializeInternal(FStructuredArchiveRecord Record);
 	void OpenArchive();
 
 private:
@@ -653,7 +681,7 @@ private:
 class FArchiveFromStructuredArchive
 {
 public:
-	explicit FArchiveFromStructuredArchive(FStructuredArchive::FSlot InSlot)
+	explicit FArchiveFromStructuredArchive(FStructuredArchiveSlot InSlot)
 		: Impl(InSlot)
 	{
 	}
@@ -670,7 +698,7 @@ private:
 class FArchiveFromStructuredArchive
 {
 public:
-	explicit FArchiveFromStructuredArchive(FStructuredArchive::FSlot InSlot)
+	explicit FArchiveFromStructuredArchive(FStructuredArchiveSlot InSlot)
 		: Ar(InSlot.GetUnderlyingArchive())
 	{
 	}
@@ -685,7 +713,7 @@ private:
 #endif
 
 /**
- * Adapter operator which allows a type to stream to an FArchive when it already supports streaming to an FStructuredArchive::FSlot.
+ * Adapter operator which allows a type to stream to an FArchive when it already supports streaming to an FStructuredArchiveSlot.
  *
  * @param  Ar   The archive to read from or write to.
  * @param  Obj  The object to read or write.
@@ -694,7 +722,7 @@ private:
  */
 template <typename T>
 typename TEnableIf<
-	!TModels<CInsertable<FArchive&>, T>::Value && TModels<CInsertable<FStructuredArchive::FSlot>, T>::Value,
+	!TModels<CInsertable<FArchive&>, T>::Value && TModels<CInsertable<FStructuredArchiveSlot>, T>::Value,
 	FArchive&
 >::Type operator<<(FArchive& Ar, T& Obj)
 {
@@ -704,7 +732,7 @@ typename TEnableIf<
 }
 
 /**
- * Adapter operator which allows a type to stream to an FStructuredArchive::FSlot when it already supports streaming to an FArchive.
+ * Adapter operator which allows a type to stream to an FStructuredArchiveSlot when it already supports streaming to an FArchive.
  *
  * @param  Slot  The slot to read from or write to.
  * @param  Obj   The object to read or write.
@@ -712,8 +740,8 @@ typename TEnableIf<
 template <typename T>
 typename TEnableIf<
 	TModels<CInsertable<FArchive&>, T>::Value &&
-	!TModels<CInsertable<FStructuredArchive::FSlot>, T>::Value
->::Type operator<<(FStructuredArchive::FSlot Slot, T& Obj)
+	!TModels<CInsertable<FStructuredArchiveSlot>, T>::Value
+>::Type operator<<(FStructuredArchiveSlot Slot, T& Obj)
 {
 #if WITH_TEXT_ARCHIVE_SUPPORT
 	FArchiveFromStructuredArchive Adapter(Slot);
@@ -725,7 +753,7 @@ typename TEnableIf<
 }
 
 #if !WITH_TEXT_ARCHIVE_SUPPORT
-	FORCEINLINE FStructuredArchiveChildReader::FStructuredArchiveChildReader(FStructuredArchive::FSlot InSlot)
+	FORCEINLINE FStructuredArchiveChildReader::FStructuredArchiveChildReader(FStructuredArchiveSlot InSlot)
 		: OwnedFormatter(nullptr)
 		, Archive(nullptr)
 	{
@@ -748,262 +776,262 @@ typename TEnableIf<
 	{
 	}
 
-	FORCEINLINE FStructuredArchive::FSlot FStructuredArchive::Open()
+	FORCEINLINE FStructuredArchiveSlot FStructuredArchive::Open()
 	{
-		return FSlot(*this);
+		return FStructuredArchiveSlot(*this);
 	}
 
 	FORCEINLINE void FStructuredArchive::Close()
 	{
 	}
 
-	//////////// FStructuredArchive::FSlot ////////////
+	//////////// FStructuredArchiveSlot ////////////
 
-	FORCEINLINE FStructuredArchive::FRecord FStructuredArchive::FSlot::EnterRecord()
+	FORCEINLINE FStructuredArchiveRecord FStructuredArchiveSlot::EnterRecord()
 	{
-		return FRecord(Ar);
+		return FStructuredArchiveRecord(Ar);
 	}
 
-	FORCEINLINE FStructuredArchive::FRecord FStructuredArchive::FSlot::EnterRecord_TextOnly(TArray<FString>& OutFieldNames)
+	FORCEINLINE FStructuredArchiveRecord FStructuredArchiveSlot::EnterRecord_TextOnly(TArray<FString>& OutFieldNames)
 	{
 		Ar.Formatter.EnterRecord_TextOnly(OutFieldNames);
-		return FRecord(Ar);
+		return FStructuredArchiveRecord(Ar);
 	}
 
-	FORCEINLINE FStructuredArchive::FArray FStructuredArchive::FSlot::EnterArray(int32& Num)
+	FORCEINLINE FStructuredArchiveArray FStructuredArchiveSlot::EnterArray(int32& Num)
 	{
 		Ar.Formatter.EnterArray(Num);
-		return FArray(Ar);
+		return FStructuredArchiveArray(Ar);
 	}
 
-	FORCEINLINE FStructuredArchive::FStream FStructuredArchive::FSlot::EnterStream()
+	FORCEINLINE FStructuredArchiveStream FStructuredArchiveSlot::EnterStream()
 	{
-		return FStream(Ar);
+		return FStructuredArchiveStream(Ar);
 	}
 
-	FORCEINLINE FStructuredArchive::FStream FStructuredArchive::FSlot::EnterStream_TextOnly(int32& OutNumElements)
+	FORCEINLINE FStructuredArchiveStream FStructuredArchiveSlot::EnterStream_TextOnly(int32& OutNumElements)
 	{
 		Ar.Formatter.EnterStream_TextOnly(OutNumElements);
-		return FStream(Ar);
+		return FStructuredArchiveStream(Ar);
 	}
 
-	FORCEINLINE FStructuredArchive::FMap FStructuredArchive::FSlot::EnterMap(int32& Num)
+	FORCEINLINE FStructuredArchiveMap FStructuredArchiveSlot::EnterMap(int32& Num)
 	{
 		Ar.Formatter.EnterMap(Num);
-		return FMap(Ar);
+		return FStructuredArchiveMap(Ar);
 	}
 
-	FORCEINLINE FStructuredArchive::FSlot FStructuredArchive::FSlot::EnterAttribute(FArchiveFieldName FieldName)
+	FORCEINLINE FStructuredArchiveSlot FStructuredArchiveSlot::EnterAttribute(FArchiveFieldName FieldName)
 	{
 		Ar.Formatter.EnterAttribute(FieldName);
-		return FSlot(Ar);
+		return FStructuredArchiveSlot(Ar);
 	}
 
-	FORCEINLINE TOptional<FStructuredArchive::FSlot> FStructuredArchive::FSlot::TryEnterAttribute(FArchiveFieldName FieldName, bool bEnterWhenWriting)
+	FORCEINLINE TOptional<FStructuredArchiveSlot> FStructuredArchiveSlot::TryEnterAttribute(FArchiveFieldName FieldName, bool bEnterWhenWriting)
 	{
 		if (Ar.Formatter.TryEnterAttribute(FieldName, bEnterWhenWriting))
 		{
-			return TOptional<FSlot>(Ar);
+			return TOptional<FStructuredArchiveSlot>(Ar);
 		}
 		else
 		{
-			return TOptional<FSlot>();
+			return TOptional<FStructuredArchiveSlot>();
 		}
 	}
 
-	FORCEINLINE void FStructuredArchive::FSlot::operator<< (uint8& Value)
+	FORCEINLINE void FStructuredArchiveSlot::operator<< (uint8& Value)
 	{
 		Ar.Formatter.Serialize(Value);
 	}
 
-	FORCEINLINE void FStructuredArchive::FSlot::operator<< (uint16& Value)
+	FORCEINLINE void FStructuredArchiveSlot::operator<< (uint16& Value)
 	{
 		Ar.Formatter.Serialize(Value);
 	}
 
-	FORCEINLINE void FStructuredArchive::FSlot::operator<< (uint32& Value)
+	FORCEINLINE void FStructuredArchiveSlot::operator<< (uint32& Value)
 	{
 		Ar.Formatter.Serialize(Value);
 	}
 
-	FORCEINLINE void FStructuredArchive::FSlot::operator<< (uint64& Value)
+	FORCEINLINE void FStructuredArchiveSlot::operator<< (uint64& Value)
 	{
 		Ar.Formatter.Serialize(Value);
 	}
 
-	FORCEINLINE void FStructuredArchive::FSlot::operator<< (int8& Value)
+	FORCEINLINE void FStructuredArchiveSlot::operator<< (int8& Value)
 	{
 		Ar.Formatter.Serialize(Value);
 	}
 
-	FORCEINLINE void FStructuredArchive::FSlot::operator<< (int16& Value)
+	FORCEINLINE void FStructuredArchiveSlot::operator<< (int16& Value)
 	{
 		Ar.Formatter.Serialize(Value);
 	}
 
-	FORCEINLINE void FStructuredArchive::FSlot::operator<< (int32& Value)
+	FORCEINLINE void FStructuredArchiveSlot::operator<< (int32& Value)
 	{
 		Ar.Formatter.Serialize(Value);
 	}
 
-	FORCEINLINE void FStructuredArchive::FSlot::operator<< (int64& Value)
+	FORCEINLINE void FStructuredArchiveSlot::operator<< (int64& Value)
 	{
 		Ar.Formatter.Serialize(Value);
 	}
 
-	FORCEINLINE void FStructuredArchive::FSlot::operator<< (float& Value)
+	FORCEINLINE void FStructuredArchiveSlot::operator<< (float& Value)
 	{
 		Ar.Formatter.Serialize(Value);
 	}
 
-	FORCEINLINE void FStructuredArchive::FSlot::operator<< (double& Value)
+	FORCEINLINE void FStructuredArchiveSlot::operator<< (double& Value)
 	{
 		Ar.Formatter.Serialize(Value);
 	}
 
-	FORCEINLINE void FStructuredArchive::FSlot::operator<< (bool& Value)
+	FORCEINLINE void FStructuredArchiveSlot::operator<< (bool& Value)
 	{
 		Ar.Formatter.Serialize(Value);
 	}
 
-	FORCEINLINE void FStructuredArchive::FSlot::operator<< (FString& Value)
+	FORCEINLINE void FStructuredArchiveSlot::operator<< (FString& Value)
 	{
 		Ar.Formatter.Serialize(Value);
 	}
 
-	FORCEINLINE void FStructuredArchive::FSlot::operator<< (FName& Value)
+	FORCEINLINE void FStructuredArchiveSlot::operator<< (FName& Value)
 	{
 		Ar.Formatter.Serialize(Value);
 	}
 
-	FORCEINLINE void FStructuredArchive::FSlot::operator<< (UObject*& Value)
+	FORCEINLINE void FStructuredArchiveSlot::operator<< (UObject*& Value)
 	{
 		Ar.Formatter.Serialize(Value);
 	}
 
-	FORCEINLINE void FStructuredArchive::FSlot::operator<< (FText& Value)
+	FORCEINLINE void FStructuredArchiveSlot::operator<< (FText& Value)
 	{
 		Ar.Formatter.Serialize(Value);
 	}
 
-	FORCEINLINE void FStructuredArchive::FSlot::operator<< (FWeakObjectPtr& Value)
+	FORCEINLINE void FStructuredArchiveSlot::operator<< (FWeakObjectPtr& Value)
 	{
 		Ar.Formatter.Serialize(Value);
 	}
 
-	FORCEINLINE void FStructuredArchive::FSlot::operator<< (FSoftObjectPath& Value)
+	FORCEINLINE void FStructuredArchiveSlot::operator<< (FSoftObjectPath& Value)
 	{
 		Ar.Formatter.Serialize(Value);
 	}
 
-	FORCEINLINE void FStructuredArchive::FSlot::operator<< (FSoftObjectPtr& Value)
+	FORCEINLINE void FStructuredArchiveSlot::operator<< (FSoftObjectPtr& Value)
 	{
 		Ar.Formatter.Serialize(Value);
 	}
 
-	FORCEINLINE void FStructuredArchive::FSlot::operator<< (FLazyObjectPtr& Value)
+	FORCEINLINE void FStructuredArchiveSlot::operator<< (FLazyObjectPtr& Value)
 	{
 		Ar.Formatter.Serialize(Value);
 	}
 
-	FORCEINLINE void FStructuredArchive::FSlot::Serialize(TArray<uint8>& Data)
+	FORCEINLINE void FStructuredArchiveSlot::Serialize(TArray<uint8>& Data)
 	{
 		Ar.Formatter.Serialize(Data);
 	}
 
-	FORCEINLINE void FStructuredArchive::FSlot::Serialize(void* Data, uint64 DataSize)
+	FORCEINLINE void FStructuredArchiveSlot::Serialize(void* Data, uint64 DataSize)
 	{
 		Ar.Formatter.Serialize(Data, DataSize);
 	}
 
-	//////////// FStructuredArchive::FRecord ////////////
+	//////////// FStructuredArchiveRecord ////////////
 
-	FORCEINLINE FStructuredArchive::FSlot FStructuredArchive::FRecord::EnterField(FArchiveFieldName Name)
+	FORCEINLINE FStructuredArchiveSlot FStructuredArchiveRecord::EnterField(FArchiveFieldName Name)
 	{
-		return FSlot(Ar);
+		return FStructuredArchiveSlot(Ar);
 	}
 
-	FORCEINLINE FStructuredArchive::FSlot FStructuredArchive::FRecord::EnterField_TextOnly(FArchiveFieldName Name, EArchiveValueType& OutType)
+	FORCEINLINE FStructuredArchiveSlot FStructuredArchiveRecord::EnterField_TextOnly(FArchiveFieldName Name, EArchiveValueType& OutType)
 	{
 		Ar.Formatter.EnterField_TextOnly(Name, OutType);
-		return FSlot(Ar);
+		return FStructuredArchiveSlot(Ar);
 	}
 
-	FORCEINLINE FStructuredArchive::FRecord FStructuredArchive::FRecord::EnterRecord(FArchiveFieldName Name)
+	FORCEINLINE FStructuredArchiveRecord FStructuredArchiveRecord::EnterRecord(FArchiveFieldName Name)
 	{
 		return EnterField(Name).EnterRecord();
 	}
 
-	FORCEINLINE FStructuredArchive::FRecord FStructuredArchive::FRecord::EnterRecord_TextOnly(FArchiveFieldName Name, TArray<FString>& OutFieldNames)
+	FORCEINLINE FStructuredArchiveRecord FStructuredArchiveRecord::EnterRecord_TextOnly(FArchiveFieldName Name, TArray<FString>& OutFieldNames)
 	{
 		return EnterField(Name).EnterRecord_TextOnly(OutFieldNames);
 	}
 
-	FORCEINLINE FStructuredArchive::FArray FStructuredArchive::FRecord::EnterArray(FArchiveFieldName Name, int32& Num)
+	FORCEINLINE FStructuredArchiveArray FStructuredArchiveRecord::EnterArray(FArchiveFieldName Name, int32& Num)
 	{
 		return EnterField(Name).EnterArray(Num);
 	}
 
-	FORCEINLINE FStructuredArchive::FStream FStructuredArchive::FRecord::EnterStream(FArchiveFieldName Name)
+	FORCEINLINE FStructuredArchiveStream FStructuredArchiveRecord::EnterStream(FArchiveFieldName Name)
 	{
 		return EnterField(Name).EnterStream();
 	}
 
-	FORCEINLINE FStructuredArchive::FMap FStructuredArchive::FRecord::EnterMap(FArchiveFieldName Name, int32& Num)
+	FORCEINLINE FStructuredArchiveMap FStructuredArchiveRecord::EnterMap(FArchiveFieldName Name, int32& Num)
 	{
 		return EnterField(Name).EnterMap(Num);
 	}
 
-	FORCEINLINE TOptional<FStructuredArchive::FSlot> FStructuredArchive::FRecord::TryEnterField(FArchiveFieldName Name, bool bEnterWhenWriting)
+	FORCEINLINE TOptional<FStructuredArchiveSlot> FStructuredArchiveRecord::TryEnterField(FArchiveFieldName Name, bool bEnterWhenWriting)
 	{
 		if (Ar.Formatter.TryEnterField(Name, bEnterWhenWriting))
 		{
-			return TOptional<FStructuredArchive::FSlot>(FSlot(Ar));
+			return TOptional<FStructuredArchiveSlot>(FStructuredArchiveSlot(Ar));
 		}
 		else
 		{
-			return TOptional<FStructuredArchive::FSlot>();
+			return TOptional<FStructuredArchiveSlot>();
 		}
 	}
 
-	//////////// FStructuredArchive::FArray ////////////
+	//////////// FStructuredArchiveArray ////////////
 
-	FORCEINLINE FStructuredArchive::FSlot FStructuredArchive::FArray::EnterElement()
+	FORCEINLINE FStructuredArchiveSlot FStructuredArchiveArray::EnterElement()
 	{
-		return FSlot(Ar);
+		return FStructuredArchiveSlot(Ar);
 	}
 
-	FORCEINLINE FStructuredArchive::FSlot FStructuredArchive::FArray::EnterElement_TextOnly(EArchiveValueType& OutType)
+	FORCEINLINE FStructuredArchiveSlot FStructuredArchiveArray::EnterElement_TextOnly(EArchiveValueType& OutType)
 	{
 		Ar.Formatter.EnterArrayElement_TextOnly(OutType);
-		return FSlot(Ar);
+		return FStructuredArchiveSlot(Ar);
 	}
 
-	//////////// FStructuredArchive::FStream ////////////
+	//////////// FStructuredArchiveStream ////////////
 
-	FORCEINLINE FStructuredArchive::FSlot FStructuredArchive::FStream::EnterElement()
+	FORCEINLINE FStructuredArchiveSlot FStructuredArchiveStream::EnterElement()
 	{
-		return FSlot(Ar);
+		return FStructuredArchiveSlot(Ar);
 	}
 
-	FORCEINLINE FStructuredArchive::FSlot FStructuredArchive::FStream::EnterElement_TextOnly(EArchiveValueType& OutType)
+	FORCEINLINE FStructuredArchiveSlot FStructuredArchiveStream::EnterElement_TextOnly(EArchiveValueType& OutType)
 	{
 		Ar.Formatter.EnterStreamElement_TextOnly(OutType);
-		return FSlot(Ar);
+		return FStructuredArchiveSlot(Ar);
 	}
 
-	//////////// FStructuredArchive::FMap ////////////
+	//////////// FStructuredArchiveMap ////////////
 
-	FORCEINLINE FStructuredArchive::FSlot FStructuredArchive::FMap::EnterElement(FString& Name)
+	FORCEINLINE FStructuredArchiveSlot FStructuredArchiveMap::EnterElement(FString& Name)
 	{
 		Ar.Formatter.EnterMapElement(Name);
-		return FSlot(Ar);
+		return FStructuredArchiveSlot(Ar);
 	}
 
-	FORCEINLINE FStructuredArchive::FSlot FStructuredArchive::FMap::EnterElement_TextOnly(FString& Name, EArchiveValueType& OutType)
+	FORCEINLINE FStructuredArchiveSlot FStructuredArchiveMap::EnterElement_TextOnly(FString& Name, EArchiveValueType& OutType)
 	{
 		Ar.Formatter.EnterMapElement_TextOnly(Name, OutType);
-		return FSlot(Ar);
+		return FStructuredArchiveSlot(Ar);
 	}
 
 #endif
