Changed so FLargeMemoryWriter::ReleaseData() returns a pointer to allow a single call to take ownership, instead of one to GetData() followed by ReleaseData()

(ushell-p4-cherrypick of 8640022 by [REDACTED_PERSON])

diff --git a/Engine/Source/Runtime/Core/Private/Serialization/LargeMemoryData.cpp b/Engine/Source/Runtime/Core/Private/Serialization/LargeMemoryData.cpp
--- a/Engine/Source/Runtime/Core/Private/Serialization/LargeMemoryData.cpp
+++ b/Engine/Source/Runtime/Core/Private/Serialization/LargeMemoryData.cpp
@@ -67,11 +67,15 @@ bool FLargeMemoryData::Read(void* OutData, int64 InOffset, int64 InNum) const
 	}
 }
 
-void FLargeMemoryData::ReleaseOwnership()
+uint8* FLargeMemoryData::ReleaseOwnership()
 {
+	uint8* ReturnData = Data;
+
 	Data = nullptr;
 	NumBytes = 0;
 	MaxBytes = 0;
+
+	return ReturnData;
 }
 
 void FLargeMemoryData::Reserve(int64 NewMax)
diff --git a/Engine/Source/Runtime/Core/Public/Serialization/LargeMemoryData.h b/Engine/Source/Runtime/Core/Public/Serialization/LargeMemoryData.h
--- a/Engine/Source/Runtime/Core/Public/Serialization/LargeMemoryData.h
+++ b/Engine/Source/Runtime/Core/Public/Serialization/LargeMemoryData.h
@@ -46,7 +46,7 @@ public:
 	}
 
 	/** Releases ownership of the written data. */
-	void ReleaseOwnership();
+	uint8* ReleaseOwnership();
 
 	/** Check whether data is allocated or if the ownership was released. */
 	bool HasData() const 
diff --git a/Engine/Source/Runtime/Core/Public/Serialization/LargeMemoryWriter.h b/Engine/Source/Runtime/Core/Public/Serialization/LargeMemoryWriter.h
--- a/Engine/Source/Runtime/Core/Public/Serialization/LargeMemoryWriter.h
+++ b/Engine/Source/Runtime/Core/Public/Serialization/LargeMemoryWriter.h
@@ -42,10 +42,13 @@ public:
 
 	/** 
 	 * Releases ownership of the written data
+	 *
+	 * Also returns the pointer, so that the caller only needs to call this function to take control
+	 * of the memory.
 	 */
-	FORCEINLINE void ReleaseOwnership()
+	FORCEINLINE uint8* ReleaseOwnership()
 	{
-		Data.ReleaseOwnership();
+		return Data.ReleaseOwnership();
 	}
 
 private:
