Crash reporting needs to support 'CrashReportClient.exe' missing.

CreateProc emits a warning message if it fails to fiind the executable, so before running it make sure the file exists on disk.

diff --git a/Engine/Source/Runtime/Core/Private/Windows/WindowsPlatformCrashContext.cpp b/Engine/Source/Runtime/Core/Private/Windows/WindowsPlatformCrashContext.cpp
--- a/Engine/Source/Runtime/Core/Private/Windows/WindowsPlatformCrashContext.cpp
+++ b/Engine/Source/Runtime/Core/Private/Windows/WindowsPlatformCrashContext.cpp
@@ -344,9 +344,9 @@ static FORCEINLINE bool CreatePipeWrite(void*& ReadPipe, void*& WritePipe)
 }
 
 /**
- * Finds the crash reporter binary path.
+ * Finds the crash reporter binary path. Returns true if the file exists.
  */
-void CreateCrashReportClientPath(TCHAR* OutClientPath, TCHAR* OutWorkingDir, int32 MaxLength)
+bool CreateCrashReportClientPath(TCHAR* OutClientPath, TCHAR* OutWorkingDir, int32 MaxLength)
 {
 #if CR_USE_DEVELOPMENT_CLIENT
 	static const TCHAR CrashReportClientExeName[] = TEXT("CrashReportClient-Win64-Development.exe");
@@ -364,6 +364,8 @@ void CreateCrashReportClientPath(TCHAR* OutClientPath, TCHAR* OutWorkingDir, int
 	const TCHAR* LastDelimiter = FCString::Strrchr(RuntimeFilename, TEXT('\\'));
 	FCString::Strncat(OutClientPath, RuntimeFilename, LastDelimiter - RuntimeFilename + 2);
 	FCString::Strncat(OutClientPath, CrashReportClientExeName, MaxLength);
+
+	return GetFileAttributesW(OutClientPath) != INVALID_FILE_ATTRIBUTES;
 }
 
 /**
@@ -395,17 +397,20 @@ FProcHandle LaunchCrashReportClient(void** OutWritePipe, void** OutReadPipe)
 		FCString::Strncat(CrashReporterClientArgs, PidStr, CR_CLIENT_MAX_ARGS_LEN);
 	}
 
-	CreateCrashReportClientPath(CrashReporterClientPath, WorkingDirectory, CR_CLIENT_MAX_PATH_LEN);
-	
-	return FPlatformProcess::CreateProc(
-		CrashReporterClientPath,
-		CrashReporterClientArgs,
-		true, false, false,
-		nullptr, 0, 
-		WorkingDirectory, 
-		nullptr, 
-		nullptr
-	);
+	// Launch the crash reporter if the client exists
+	if (CreateCrashReportClientPath(CrashReporterClientPath, WorkingDirectory, CR_CLIENT_MAX_PATH_LEN))
+	{
+		return FPlatformProcess::CreateProc(
+			CrashReporterClientPath,
+			CrashReporterClientArgs,
+			true, false, false,
+			nullptr, 0,
+			WorkingDirectory,
+			nullptr,
+			nullptr
+		);
+	}
+	return FProcHandle();
 }
 
 /**
@@ -713,8 +718,14 @@ int32 ReportCrashUsingCrashReportClient(FWindowsPlatformCrashContext& InContext,
 				CrashReportClientArguments += FString::Format(TEXT(" -abslog=\"{0}\""), { AbsCrashReportClientLog });
 			}
 
-			FString CrashClientPath = FPaths::Combine(*FPaths::EngineDir(), TEXT("Binaries"), FPlatformProcess::GetBinariesSubdirectory(), CrashReportClientExeName);
-			bCrashReporterRan = FPlatformProcess::CreateProc(*CrashClientPath, *CrashReportClientArguments, true, false, false, NULL, 0, NULL, NULL).IsValid();
+			TCHAR CrashReporterClientPath[CR_CLIENT_MAX_PATH_LEN] = { 0 };
+			TCHAR WorkingDirectory[CR_CLIENT_MAX_PATH_LEN] = { 0 };
+
+			if (CreateCrashReportClientPath(CrashReporterClientPath, WorkingDirectory, CR_CLIENT_MAX_PATH_LEN))
+			{
+				FString CrashClientPath = FPaths::Combine(*FPaths::EngineDir(), TEXT("Binaries"), FPlatformProcess::GetBinariesSubdirectory(), CrashReportClientExeName);
+				bCrashReporterRan = FPlatformProcess::CreateProc(*CrashClientPath, *CrashReportClientArguments, true, false, false, NULL, 0, NULL, NULL).IsValid();
+			}
 
 			// Restore the dll directory
 			if (CurrentDllDirectory)
