Changed the way path to the crash reporter is calculated, and fixing issues with running crash reporter using the old (non-monitor) path.

* Previous assumption that the crash report client executable would reside in the same directory as the running process was wrong. Added a more correct calculation of the path.
* Using a path to the executable there is no need to use a working directory.
* Changes to crash context had not been reflected to the non-monitor path of the code, resulting in the crash dialog displaying "you have no symbols". This was due to not correctly setting current process handle, which in turn resulted in a corrupt minidump file.

diff --git a/Engine/Source/Runtime/Core/Private/Windows/WindowsPlatformCrashContext.cpp b/Engine/Source/Runtime/Core/Private/Windows/WindowsPlatformCrashContext.cpp
--- a/Engine/Source/Runtime/Core/Private/Windows/WindowsPlatformCrashContext.cpp
+++ b/Engine/Source/Runtime/Core/Private/Windows/WindowsPlatformCrashContext.cpp
@@ -346,7 +346,7 @@ static FORCEINLINE bool CreatePipeWrite(void*& ReadPipe, void*& WritePipe)
 /**
  * Finds the crash reporter binary path. Returns true if the file exists.
  */
-bool CreateCrashReportClientPath(TCHAR* OutClientPath, TCHAR* OutWorkingDir, int32 MaxLength)
+bool CreateCrashReportClientPath(TCHAR* OutClientPath, int32 MaxLength)
 {
 #if CR_USE_DEVELOPMENT_CLIENT
 	static const TCHAR CrashReportClientExeName[] = TEXT("CrashReportClient-Win64-Development.exe");
@@ -354,18 +354,18 @@ bool CreateCrashReportClientPath(TCHAR* OutClientPath, TCHAR* OutWorkingDir, int
 	static const TCHAR CrashReportClientExeName[] = TEXT("CrashReportClient.exe");
 #endif 
 
-	// Use the same working directory as this process
-	GetCurrentDirectory(MaxLength, OutWorkingDir);
+	const TCHAR* EngineDir = FPlatformMisc::EngineDir();
+	const TCHAR* BinariesDir = FPlatformProcess::GetBinariesSubdirectory();
 
-	TCHAR RuntimeFilename[CR_CLIENT_MAX_PATH_LEN];
-	const DWORD RuntimeFilenameLength = GetModuleFileNameEx(GetCurrentProcess(), NULL, RuntimeFilename, CR_CLIENT_MAX_PATH_LEN);
-
-	// Assume the crash reporter client is in the same location as the game runtime.
-	const TCHAR* LastDelimiter = FCString::Strrchr(RuntimeFilename, TEXT('\\'));
-	FCString::Strncat(OutClientPath, RuntimeFilename, LastDelimiter - RuntimeFilename + 2);
+	// Find the path to crash reporter binary. Avoid creating FStrings.
+	FCString::Strncat(OutClientPath, EngineDir, MaxLength);
+	FCString::Strncat(OutClientPath, TEXT("Binaries/"), MaxLength);
+	FCString::Strncat(OutClientPath, BinariesDir, MaxLength);
+	FCString::Strncat(OutClientPath, TEXT("/"), MaxLength);
 	FCString::Strncat(OutClientPath, CrashReportClientExeName, MaxLength);
 
-	return GetFileAttributesW(OutClientPath) != INVALID_FILE_ATTRIBUTES;
+	const DWORD Results = GetFileAttributesW(OutClientPath);
+	return Results != INVALID_FILE_ATTRIBUTES;
 }
 
 /**
@@ -375,7 +375,6 @@ FProcHandle LaunchCrashReportClient(void** OutWritePipe, void** OutReadPipe)
 {
 	TCHAR CrashReporterClientPath[CR_CLIENT_MAX_PATH_LEN] = { 0 };
 	TCHAR CrashReporterClientArgs[CR_CLIENT_MAX_ARGS_LEN] = { 0 };
-	TCHAR WorkingDirectory[CR_CLIENT_MAX_PATH_LEN] = { 0 };
 
 	void *PipeChildInRead, *PipeChildInWrite, *PipeChildOutRead, *PipeChildOutWrite;
 
@@ -398,14 +397,14 @@ FProcHandle LaunchCrashReportClient(void** OutWritePipe, void** OutReadPipe)
 	}
 
 	// Launch the crash reporter if the client exists
-	if (CreateCrashReportClientPath(CrashReporterClientPath, WorkingDirectory, CR_CLIENT_MAX_PATH_LEN))
+	if (CreateCrashReportClientPath(CrashReporterClientPath, CR_CLIENT_MAX_PATH_LEN))
 	{
 		return FPlatformProcess::CreateProc(
 			CrashReporterClientPath,
 			CrashReporterClientArgs,
 			true, false, false,
 			nullptr, 0,
-			WorkingDirectory,
+			nullptr,
 			nullptr,
 			nullptr
 		);
@@ -719,12 +718,10 @@ int32 ReportCrashUsingCrashReportClient(FWindowsPlatformCrashContext& InContext,
 			}
 
 			TCHAR CrashReporterClientPath[CR_CLIENT_MAX_PATH_LEN] = { 0 };
-			TCHAR WorkingDirectory[CR_CLIENT_MAX_PATH_LEN] = { 0 };
 
-			if (CreateCrashReportClientPath(CrashReporterClientPath, WorkingDirectory, CR_CLIENT_MAX_PATH_LEN))
+			if (CreateCrashReportClientPath(CrashReporterClientPath, CR_CLIENT_MAX_PATH_LEN))
 			{
-				FString CrashClientPath = FPaths::Combine(*FPaths::EngineDir(), TEXT("Binaries"), FPlatformProcess::GetBinariesSubdirectory(), CrashReportClientExeName);
-				bCrashReporterRan = FPlatformProcess::CreateProc(*CrashClientPath, *CrashReportClientArguments, true, false, false, NULL, 0, NULL, NULL).IsValid();
+				bCrashReporterRan = FPlatformProcess::CreateProc(CrashReporterClientPath, *CrashReportClientArguments, true, false, false, NULL, 0, NULL, NULL).IsValid();
 			}
 
 			// Restore the dll directory
@@ -975,7 +972,7 @@ public:
 		else
 		{
 			FWindowsPlatformCrashContext CrashContext(ECrashContextType::Ensure, ErrorMessage);
-
+			CrashContext.SetCrashedProcess(FProcHandle(::GetCurrentProcess()));
 			void* ContextWrapper = FWindowsPlatformStackWalk::MakeThreadContextWrapper(InExceptionInfo->ContextRecord, GetCurrentThread());
 			CrashContext.CapturePortableCallStack(NumStackFramesToIgnore, ContextWrapper);
 
@@ -1095,6 +1092,7 @@ private:
 
 			// Thread context wrapper for stack operations
 			ContextWrapper = FWindowsPlatformStackWalk::MakeThreadContextWrapper(ExceptionInfo->ContextRecord, CrashingThreadHandle);
+			CrashContext.SetCrashedProcess(FProcHandle(::GetCurrentProcess()));
 			CrashContext.CapturePortableCallStack(NumStackFramesToIgnore, ContextWrapper);
 			CrashContext.SetCrashedThreadId(CrashingThreadId);
 			CrashContext.CaptureAllThreadContexts();
@@ -1300,8 +1298,9 @@ void ReportHang(const TCHAR* ErrorMessage, const uint64* StackFrames, int32 NumS
 	}
 
 	FWindowsPlatformCrashContext CrashContext(ECrashContextType::Hang, ErrorMessage);
-	CrashContext.SetPortableCallStack(StackFrames, NumStackFrames);
+	CrashContext.SetCrashedProcess(FProcHandle(::GetCurrentProcess()));
 	CrashContext.SetCrashedThreadId(HungThreadId);
+	CrashContext.SetPortableCallStack(StackFrames, NumStackFrames);
 	CrashContext.CaptureAllThreadContexts();
 
 	EErrorReportUI ReportUI = IsInteractiveEnsureMode() ? EErrorReportUI::ShowDialog : EErrorReportUI::ReportInUnattendedMode;
