Fixed deprecation warning.

diff --git a/Engine/Source/Runtime/Core/Private/HAL/LowLevelMemTracker.cpp b/Engine/Source/Runtime/Core/Private/HAL/LowLevelMemTracker.cpp
--- a/Engine/Source/Runtime/Core/Private/HAL/LowLevelMemTracker.cpp
+++ b/Engine/Source/Runtime/Core/Private/HAL/LowLevelMemTracker.cpp
@@ -215,7 +215,7 @@ bool LLMPrivate::HandleAssert(bool bLog, const TCHAR* Format, ...)
 	if (bLog)
 	{
 		TCHAR DescriptionString[4096];
-		GET_VARARGS(DescriptionString, ARRAY_COUNT(DescriptionString), ARRAY_COUNT(DescriptionString) - 1, Format, Format);
+		GET_VARARGS(DescriptionString, UE_ARRAY_COUNT(DescriptionString), UE_ARRAY_COUNT(DescriptionString) - 1, Format, Format);
 
 		FPlatformMisc::LowLevelOutputDebugString(DescriptionString);
 
