Merge 8669459 from //UE4/Dev-Rendering to Dev-Core (//UE4/Dev-Core)

---

[macOS][Platform] Restrict exclusive file lock access to the Editor

This change restricts the exclusive file locking on used on macOS (to mimic the Windows file handle behavior) to just the Editor, to enable localhost network file sharing over sockets (for UnrealInsights).

---

#rnx

diff --git a/Engine/Source/Runtime/Core/Private/Apple/ApplePlatformFile.cpp b/Engine/Source/Runtime/Core/Private/Apple/ApplePlatformFile.cpp
--- a/Engine/Source/Runtime/Core/Private/Apple/ApplePlatformFile.cpp
+++ b/Engine/Source/Runtime/Core/Private/Apple/ApplePlatformFile.cpp
@@ -553,7 +553,7 @@ IFileHandle* FApplePlatformFile::OpenWrite(const TCHAR* Filename, bool bAppend,
 	
 	if (Handle != -1)
 	{
-#if PLATFORM_MAC && !UE_BUILD_SHIPPING
+#if PLATFORM_MAC && UE_EDITOR && !UE_BUILD_SHIPPING
 		// No blocking attempt exclusive lock, failure means we should not have opened the file for writing, protect against multiple instances and client/server versions
 		if(flock(Handle, LOCK_NB | LOCK_EX) == -1)
 		{
