Fixed static analysis warning about double-slashes in an include path.

diff --git a/Engine/Source/Runtime/Core/Private/FramePro/FramePro.cpp b/Engine/Source/Runtime/Core/Private/FramePro/FramePro.cpp
--- a/Engine/Source/Runtime/Core/Private/FramePro/FramePro.cpp
+++ b/Engine/Source/Runtime/Core/Private/FramePro/FramePro.cpp
@@ -8430,7 +8430,7 @@ namespace FramePro
 		#include <winsock2.h>
 		#include <ws2tcpip.h>
 		#if FRAMEPRO_PLATFORM_UE4
-			#include "Windows//HideWindowsPlatformTypes.h"
+			#include "Windows/HideWindowsPlatformTypes.h"
 		#endif
 	#endif
 
