Don't call FLLMAllocator::Free if the map had not yet been allocated.

Prevents attempts to call ::VirtualFree on a nulllptr on windows which will invoke an Application Verifier error.

diff --git a/Engine/Source/Runtime/Core/Public/HAL/LowLevelMemTracker.h b/Engine/Source/Runtime/Core/Public/HAL/LowLevelMemTracker.h
--- a/Engine/Source/Runtime/Core/Public/HAL/LowLevelMemTracker.h
+++ b/Engine/Source/Runtime/Core/Public/HAL/LowLevelMemTracker.h
@@ -359,10 +359,13 @@ public:
 
 	void Free(void* Ptr, size_t Size)
 	{
-		Size = Align(Size, Alignment);
-		FScopeLock Lock(&CriticalSection);
-		PlatformFree(Ptr, Size);
-		Total -= Size;
+		if (Ptr != nullptr)
+		{
+			Size = Align(Size, Alignment);
+			FScopeLock Lock(&CriticalSection);
+			PlatformFree(Ptr, Size);
+			Total -= Size;
+		}
 	}
 
 	int64 GetTotal() const
