Copying //UE4/Portal-Staging to //UE4/Main (Source: //Portal/Main/Engine @ 8661229)

Change 8553543 by [REDACTED_PERSON]

	We now set a fixed value of 24 to CEF screenInfo colorDepth when off screen rendering is enabled.

Change 8235770 by [REDACTED_PERSON]

	Fix for CEFJsScripting memory stomp and memory alignment errors called out by the stomp memory allocator when running with -stompmalloc on the commandline.

Change 8065597 by [REDACTED_PERSON]

	BuildPatchServices: Improving ChunkBuildDirectory behaviour when dealing with empty build or builds only containing empty files.
	This is now fully supported as part of the generation flow rather than an early detected edge case.

Change 7641628 by [REDACTED_PERSON]

	BuildPatchServices: Tweaks to serialisation safety.
	BuildPatchServices: DiskChunkStore fix for crash when chunkdump serialisation fails due to disk space.

Change 7436869 by [REDACTED_PERSON]

	BuildPatchServices: Fix regression in BPT package chunks for cloud save improvements now require message pumping.

Change 7326553 by [REDACTED_PERSON]

	BuildPatchServices: Attempt to restore functionality of the -SkipBuildPatchPrereq commandline.  This stopped functioning in CL# 6655502 with the BPS DLC related refactors.
	The issue was reported by a number of users that were relying on this commandline as a last resort workaround for prereq install issues.

Change 7323945 by [REDACTED_PERSON]

	BuildPatchServices: BuildPatchTool: Install time coefficient values exposed by DiffManifests.
	This gives an indicative install duration.
	The time is not necessarily accurate, but the simulation is a constant and so the value is highly comparable between different builds.

Change 7310352 by [REDACTED_PERSON]

	BuildPatchServices: Adding support for overriding http path for chunk requests. This allows the ability to support signed urls when downloading patch data.

Change 7095282 by [REDACTED_PERSON]

	BuildPatchServices: Fix regression with manifests that have no core files.

Change 7092198 by leigh.swift

	BuildPatchServices: Reuse existing code from FBuildPatchAppManifest::GetChunkShaHash in FBuildPatchManifestSet::GetChunkShaHash.
	This fixes an issue where older manifest files that did not ship with chunk sha values in them, can skip sha validation of chunks like pre-DLC launcher did.

Change 6959115 by [REDACTED_PERSON]

	Added functionality that will optionally expose embedded browser console logs to the client.

Change 6835841 by [REDACTED_PERSON]

	BuildPatchServices: Rearranging manifest save logic to avoid unnecessary seeking forwards, which avoids an assert when undetected write failures occur.

Change 6684994 by [REDACTED_PERSON]

	BuildPatchServices: Don't clean empty directories if staging only.

Change 6655502 by [REDACTED_PERSON], Leigh.Swift, Wes.Fudala

	BuildPatchServices: Restructuring how installers are configured and make use of manifest files in order to combine multiple actions on an installation directory into one installer.
	This resembles a feature set for a better DLC installation experience.

Change 6404031 by [REDACTED_PERSON]

	BuildPatchTool: Only append ".manifest" to output filename if output filename has been specified on the command line.
	This was causing a manifest file called literally ".manifest" to be output to the clouddir if -OutputFilename was not specified.

Change 6077240 by [REDACTED_PERSON]

	Execution of browser resource load complete delegate now happens on the main thread.

Change 6076171 by [REDACTED_PERSON]

	BuildPatchTool: PatchGeneration: ChunkDeltaOptimise: PackageChunks: Improved corrupt output protection against ill timed taskkill, by serialising to temp filename, and then rename on success.
	BuildPatchTool: PatchGeneration: Manifest file extension added if not provided, fixing an oversight and inconsistency with other mode behaviours.
	BuildPatchTool: Compactify: Only warn when failing to get a file size, if the file still exists. Otherwise log instead.

Change 6049003 by [REDACTED_PERSON]

	BuildPatchServices: Adding ProcessRequiredDiskSpace to Launcher.Install.Stats which represents how much disk space the install/update process needed to complete.
	BuildPatchServices: Adding ProcessAvailableDiskSpace to Launcher.Install.Stats which represents how much disk space was available at the time of checking required disk space.

Change 5915157 by [REDACTED_PERSON]

	BuildPatchTool: Adding a statistic to diffmanifests for temporary disk space requirement to apply the patch.

Change 5934838 by [REDACTED_PERSON]

	BuildPatchTool: PackageChunks: Adding support to provide a tagset for the previous build manifest when producing chunkdbs. This allows expanding the chunks saved out to cover tagsets not installed in the previous build.

Change 5838666 by wes.fudala

	Browser can now bubble up the state of completed web resource loads.

Change 5689493 by [REDACTED_PERSON]

	Adding new x86 and x64 MS VC141 CRT redist, version 14.16.27012

Change 5689462 by [REDACTED_PERSON]

	Fixing process handle leaks on windows.
	Core was leaking for getting an application name.

Change 5500917 by [REDACTED_PERSON]

	BuildPatchTool: Adding new arg DiffAbortThreshold to ChunkDeltaOptimise mode which allows skipping of the operation if the original delta is so large that it would take too long to process, and likely have little benefit.
	BuildPatchTool: Switching some Log output to use Display so that it will appear in EC and CMD windows.

Change 5337482 by [REDACTED_PERSON]

	BuildPatchTool: Fix for DiffManifests mode not accurately representing delta size for tagged install sets.

Change 5261246 by [REDACTED_PERSON]

	BuildPatchServices: Fix for file download needing to mock response codes for higher layer statistics code which tracks data sizes and speeds.
	This is a regression from previous change to correct download failure vs corruption statistics.

Change 5224725 by [REDACTED_PERSON]

	BuildPatchServices: Fix for delta download of more than 0 bytes when no update is necessary.
	BuildPatchServices: Skip requesting delta metafile if no file changes are actually required for a patch.
	BuildPatchTool: Reduce unnecessary data produced by BPT ChunkDeltaOptimise mode.

Change 5010941 by [REDACTED_PERSON]

	BuildPatchServices: Add download scaling based on average speed per request, maximum count, and download health.

Change 5010845 by [REDACTED_PERSON]

	BuildPatchServices: IDownload refactored to have specific request and response success functions, to make it clearer that a successful request does not mean the response was also good.
	BuildPatchServices: Fixed issues with download failures reporting as corruptions.

Change 5000643 by [REDACTED_PERSON]

	Remove HttpServiceTracker from UE4.

Change 4884381 by [REDACTED_PERSON]

	BuildPatchTool: Fix for Package Chunks mode hanging when no chunks were required.

Change 4848675 by [REDACTED_PERSON], Leigh.Swift

	Speculative fixes for graphics device lost related crash, by adding additional d3d api result checks.
	Improved logging for graphics device lost handling.
	Improved logging for tracking down common font loading failure resulting in an ensure.

Change 4831134 by [REDACTED_PERSON]

	BuildPatchTool: Fix for crash in patchgeneration when fast-forward path replays no match.

Change 4801714 by [REDACTED_PERSON]

	Fix for CEF issue encountered when building using Mac Mojave + XCode10.

Change 4719149 by [REDACTED_PERSON]

	BuildPatchTool: PatchGeneration mode cyclic data optimisation, reduces SHA calculation requirement counts for cyclic data.
	BuildPatchTool: PatchGeneration mode fix for a bug causing non-optimal match insertion idx searching when there are 10k+ matches per scanner.

Change 4680963 by [REDACTED_PERSON]

	BuildPatchTool: ChunkDeltaOptimise mode is now FeatureLevel upgrade / downgrade aware.

Change 4680947 by [REDACTED_PERSON]

	BuildPatchTool: Compactify speed improvements for massive network cloud directories.

Change 4656991 by [REDACTED_PERSON]

	BuildPatchServices: Make sure chunk writer robustly discovers if a chunk fails to save out.

Change 4647815 by [REDACTED_PERSON]

	Upping the minimum wait time for UdpMessageBeacon thread so that it will not always wait 0ms when network sends are failing, reducing disconnect CPU usage.
	Adding configurable tick rate logic to XmppConnectionJingle thread. It will now default to 100Hz max.

Change 4627355 by [REDACTED_PERSON]

	Fixed a problem with CEF being unable to find locale pak files on Mac for certain language/region combinations

Change 4620800 by [REDACTED_PERSON]

	Fix for CEF crash when disabling a web window that has not yet got a parent window. There's no need to worry about focus in this case.

Change 4590207 by [REDACTED_PERSON]

	BuildPatchTool: PackageChunks mode now supports FeatureLevel arg

Change 4590103 by leigh.swift

	BuildPatchTool: Adding new mode ChunkDeltaOptimise which reducing the download size when patching between two specific builds in a specific direction.
	BuildPatchTool: Updated Enumeration, DiffManifests, Compactify, PackageChunks, and VerifyChunks modes to take account of new delta data.
	BuildPatchServices: Installers now have a single shared memory chunk store, which reduces the requirement for booting

Change 4590089 by [REDACTED_PERSON]

	BuildPatchTool: Adding new mode ChunkDeltaOptimise which reducing the download size when patching between two specific builds in a specific direction.
	BuildPatchTool: Updated Enumeration, DiffManifests, Compactify, PackageChunks, and VerifyChunks modes to take account of new delta data.
	BuildPatchServices: Installers now have a single shared memory chunk store, which reduces the requirement for booting

Change 4341076 by [REDACTED_PERSON]

	BuildPatchServices: Making FBuildPatchAppManifest::GetRemovableFiles  more robust to handle directories with or without trailing slash.

Change 4331754 by [REDACTED_PERSON]

	BuildPatchTool: Added support for selecting ChunkWindowSize when generating patches.
	BuildPatchTool: Added support for providing the FeatureLevel command-line argument to indicate the data version that should be saved out by patch generation. This warns about defaulting to LatestJson if not provided.
	BuildPatchTool: Added support for generating patches with recognition for any chunks with any ChunkWindowSize found in the provided CloudDir.
	BuildPatchTool: Added command-line -IgnoreOtherWindowSizes param which if provided, the generation code will only accept chunk matches that are the same as ChunkWindowSize.

	BuildPatchServices: Fixes for supporting installations that use any ChunkWindowSize.
	BuildPatchServices: New manifest file format to reduce file size, this is now raw compressed binary data.

(ushell-p4-cherrypick of 8675597 by [REDACTED_PERSON])

diff --git a/Engine/Source/Runtime/Core/Private/Windows/WindowsPlatformProcess.cpp b/Engine/Source/Runtime/Core/Private/Windows/WindowsPlatformProcess.cpp
--- a/Engine/Source/Runtime/Core/Private/Windows/WindowsPlatformProcess.cpp
+++ b/Engine/Source/Runtime/Core/Private/Windows/WindowsPlatformProcess.cpp
@@ -675,7 +675,7 @@ FString FWindowsPlatformProcess::GetApplicationName( uint32 ProcessId )
 			Output = ProcessNameBuffer;
 		}
 
-		CloseHandle(ProcessHandle);
+		::CloseHandle(ProcessHandle);
 	}
 
 	return Output;
@@ -750,8 +750,17 @@ bool FWindowsPlatformProcess::ExecProcess(const TCHAR* URL, const TCHAR* Params,
 		CommandLine = FString::Printf(TEXT("%s %s"), URL, Params);
 	}
 
+	// We only want to add the EXTENDED_STARTUPINFO_PRESENT flag if StartupInfoEx.lpAttributeList is actually setup.
+	// If StartupInfoEx.lpAttributeList is NULL when the EXTENDED_STARTUPINFO_PRESENT flag is used, then CreateProcess causes an Access Violation crash on some Win32 configurations.
+	// This is specific to when the call is redirected to APIHook_CreateProcessW in AcLayers.dll rather than the standard CreateProcessW implementation in kernel32.dll.
+	uint32 CreateFlags = NORMAL_PRIORITY_CLASS | DETACHED_PROCESS;
+	if (StartupInfoEx.lpAttributeList != NULL)
+	{
+		CreateFlags |= EXTENDED_STARTUPINFO_PRESENT;
+	}
+
 	PROCESS_INFORMATION ProcInfo;
-	if (CreateProcess(NULL, CommandLine.GetCharArray().GetData(), NULL, NULL, TRUE, NORMAL_PRIORITY_CLASS | DETACHED_PROCESS | EXTENDED_STARTUPINFO_PRESENT, NULL, OptionalWorkingDirectory, &StartupInfoEx.StartupInfo, &ProcInfo))
+	if (CreateProcess(NULL, CommandLine.GetCharArray().GetData(), NULL, NULL, TRUE, CreateFlags, NULL, OptionalWorkingDirectory, &StartupInfoEx.StartupInfo, &ProcInfo))
 	{
 		if (hStdOutRead != NULL)
 		{
