Reverted CR_USE_DEVELOPMENT_CLIENT to zero (forget in the previous submit)

diff --git a/Engine/Source/Runtime/Core/Private/Windows/WindowsPlatformCrashContext.cpp b/Engine/Source/Runtime/Core/Private/Windows/WindowsPlatformCrashContext.cpp
--- a/Engine/Source/Runtime/Core/Private/Windows/WindowsPlatformCrashContext.cpp
+++ b/Engine/Source/Runtime/Core/Private/Windows/WindowsPlatformCrashContext.cpp
@@ -45,7 +45,7 @@
 #define CR_CLIENT_MAX_PATH_LEN 265
 #define CR_CLIENT_MAX_ARGS_LEN 256
 // Enable to use a non shipping version of the crash report client
-#define CR_USE_DEVELOPMENT_CLIENT 1
+#define CR_USE_DEVELOPMENT_CLIENT 0
 
 #pragma comment( lib, "version.lib" )
 #pragma comment( lib, "Shlwapi.lib" )
