ARRAY_COUNT -> UE_ARRAY_COUNT

diff --git a/Engine/Source/Runtime/Core/Private/Misc/ConfigCacheIni.cpp b/Engine/Source/Runtime/Core/Private/Misc/ConfigCacheIni.cpp
--- a/Engine/Source/Runtime/Core/Private/Misc/ConfigCacheIni.cpp
+++ b/Engine/Source/Runtime/Core/Private/Misc/ConfigCacheIni.cpp
@@ -572,7 +572,7 @@ constexpr int32 GetStaticKey(int32 LayerIndex, int32 LayerExpansionIndex, int32
 	return LayerIndex * 10000 + LayerExpansionIndex * 100 + PlatformIndex;
 }
 
-constexpr int32 MaxStaticHierarchyKey = GetStaticKey(ARRAY_COUNT(GConfigLayers) - 1, ARRAY_COUNT(GConfigLayerExpansions) - 1, MaxPlatformIndex);
+constexpr int32 MaxStaticHierarchyKey = GetStaticKey(UE_ARRAY_COUNT(GConfigLayers) - 1, UE_ARRAY_COUNT(GConfigLayerExpansions) - 1, MaxPlatformIndex);
 
 /*-----------------------------------------------------------------------------
 	FConfigFileHierarchy
