Fixed CIS error due to PLATFORM_WCHAR_IS_4_BYTES being merged out in error.

#rnx

diff --git a/Engine/Source/Runtime/Core/Public/Mac/MacPlatform.h b/Engine/Source/Runtime/Core/Public/Mac/MacPlatform.h
--- a/Engine/Source/Runtime/Core/Public/Mac/MacPlatform.h
+++ b/Engine/Source/Runtime/Core/Public/Mac/MacPlatform.h
@@ -44,6 +44,7 @@ typedef FMacPlatformTypes FPlatformTypes;
 //#define PLATFORM_USE_LS_SPEC_FOR_WIDECHAR			1
 #define PLATFORM_USE_SYSTEM_VSWPRINTF				0
 #define PLATFORM_COMPILER_DISTINGUISHES_INT_AND_LONG			1
+#define PLATFORM_WCHAR_IS_4_BYTES					1
 #if PLATFORM_MAC_USE_CHAR16
 #define PLATFORM_TCHAR_IS_CHAR16					1
 #else
