Update LexTryParseString() to not use FCString::IsNumeric(), which fails for a number of number formats that can be converted by LexToString perfectly well. The function now runs LexToString() and compares the output with 0. When the result is 0, we perform a rudimentary check that the input string was 0 and return true if so.

diff --git a/Engine/Source/Runtime/Core/Private/Tests/Misc/StringTest.cpp b/Engine/Source/Runtime/Core/Private/Tests/Misc/StringTest.cpp
new file mode 100644
--- /dev/null
+++ b/Engine/Source/Runtime/Core/Private/Tests/Misc/StringTest.cpp
@@ -0,0 +1,117 @@
+// Copyright 1998-2019 Epic Games, Inc. All Rights Reserved.
+
+#include "CoreTypes.h"
+#include "Containers/UnrealString.h"
+#include "Misc/AutomationTest.h"
+
+#if WITH_DEV_AUTOMATION_TESTS
+
+IMPLEMENT_SIMPLE_AUTOMATION_TEST(FStringTest, "System.Core.Misc.String", EAutomationTestFlags::ApplicationContextMask | EAutomationTestFlags::SmokeFilter)
+
+bool FStringTest::RunTest( const FString& Parameters )
+{
+	// Test that LexFromString can intepret all the numerical formats we expect it to
+	{
+		// Test float values
+
+		float Value;
+
+		// Basic numbers
+		TestTrue(TEXT("(float conversion from string) 1"), LexTryParseString(Value, (TEXT("1"))) && Value == 1);
+		TestTrue(TEXT("(float conversion from string) 1.0"), LexTryParseString(Value, (TEXT("1.0"))) && Value == 1);
+		TestTrue(TEXT("(float conversion from string) .1"), LexTryParseString(Value, (TEXT(".5"))) && Value == 0.5);
+		TestTrue(TEXT("(float conversion from string) 1."), LexTryParseString(Value, (TEXT("1."))) && Value == 1);
+
+		// Variations of 0
+		TestTrue(TEXT("(float conversion from string) 0"), LexTryParseString(Value, (TEXT("0"))) && Value == 0);
+		TestTrue(TEXT("(float conversion from string) -0"), LexTryParseString(Value, (TEXT("-0"))) && Value == 0);
+		TestTrue(TEXT("(float conversion from string) 0.0"), LexTryParseString(Value, (TEXT("0.0"))) && Value == 0);
+		TestTrue(TEXT("(float conversion from string) .0"), LexTryParseString(Value, (TEXT(".0"))) && Value == 0);
+		TestTrue(TEXT("(float conversion from string) 0."), LexTryParseString(Value, (TEXT("0."))) && Value == 0);
+		TestTrue(TEXT("(float conversion from string) 0. 111"), LexTryParseString(Value, (TEXT("0. 111"))) && Value == 0);
+
+		// Scientific notation
+		TestTrue(TEXT("(float conversion from string) 1.0e+10"), LexTryParseString(Value, (TEXT("1.0e+10"))) && Value == 1.0e+10f);
+		TestTrue(TEXT("(float conversion from string) 1.0e-10"), LexTryParseString(Value, (TEXT("1.99999999e-11"))) && Value == 1.99999999e-11f);
+		TestTrue(TEXT("(float conversion from string) 1e+10"), LexTryParseString(Value, (TEXT("1e+10"))) && Value == 1e+10f);
+
+		// Float as hex
+
+		// Non-finite special numbers
+		TestTrue(TEXT("(float conversion from string) inf"), LexTryParseString(Value, (TEXT("inf"))));
+		TestTrue(TEXT("(float conversion from string) nan"), LexTryParseString(Value, (TEXT("nan"))));
+		TestTrue(TEXT("(float conversion from string) nan(ind)"), LexTryParseString(Value, (TEXT("nan(ind)"))));
+
+		// nan/inf etc. are detected from the start of the string, regardless of any other characters that come afterwards
+		TestTrue(TEXT("(float conversion from string) nananananananana"), LexTryParseString(Value, (TEXT("nananananananana"))));
+		TestTrue(TEXT("(float conversion from string) nan(ind)!"), LexTryParseString(Value, (TEXT("nan(ind)!"))));
+		TestTrue(TEXT("(float conversion from string) infinity"), LexTryParseString(Value, (TEXT("infinity"))));
+
+		// Some numbers with whitespace
+		TestTrue(TEXT("(float conversion from string) whitespace"), LexTryParseString(Value, (TEXT("   2.5   "))) && Value == 2.5);
+		TestTrue(TEXT("(float conversion from string) whitespace"), LexTryParseString(Value, (TEXT("\t3.0\t"))) && Value == 3.0);
+		TestTrue(TEXT("(float conversion from string) whitespace"), LexTryParseString(Value, (TEXT("4.0   \t"))) && Value == 4.0);
+		TestTrue(TEXT("(float conversion from string) whitespace"), LexTryParseString(Value, (TEXT("\r\n5.25"))) && Value == 5.25);
+		TestTrue(TEXT("(float conversion from string) whitespace"), LexTryParseString(Value, (TEXT(" 6 . 2 "))) && Value == 6.0);
+		TestTrue(TEXT("(float conversion from string) whitespace"), LexTryParseString(Value, (TEXT(" 56 . 2 "))) && Value == 56.0);
+		TestTrue(TEXT("(float conversion from string) whitespace"), LexTryParseString(Value, (TEXT(" 5 6 . 2 "))) && Value == 5.0);
+
+		// Failure cases
+		TestFalse(TEXT("(float no conversion from string) not a number"), LexTryParseString(Value, (TEXT("not a number"))));
+		TestFalse(TEXT("(float no conversion from string) <empty string>"), LexTryParseString(Value, (TEXT(""))));
+		TestFalse(TEXT("(float conversion from string) ."), LexTryParseString(Value, (TEXT("."))));
+	}
+
+	{
+		// Test integer values
+
+		int32 Value;
+
+		// Basic numbers
+		TestTrue(TEXT("(int32 conversion from string) 1"), LexTryParseString(Value, (TEXT("1"))) && Value == 1);
+		TestTrue(TEXT("(int32 conversion from string) 1.0"), LexTryParseString(Value, (TEXT("1.0"))) && Value == 1);
+		TestTrue(TEXT("(int32 conversion from string) 3.1"), LexTryParseString(Value, (TEXT("3.1"))) && Value == 3);
+		TestTrue(TEXT("(int32 conversion from string) 0.5"), LexTryParseString(Value, (TEXT("0.5"))) && Value == 0);
+		TestTrue(TEXT("(int32 conversion from string) 1."), LexTryParseString(Value, (TEXT("1."))) && Value == 1);
+
+		// Variations of 0
+		TestTrue(TEXT("(int32 conversion from string) 0"), LexTryParseString(Value, (TEXT("0"))) && Value == 0);
+		TestTrue(TEXT("(int32 conversion from string) 0.0"), LexTryParseString(Value, (TEXT("0.0"))) && Value == 0);
+		TestFalse(TEXT("(int32 conversion from string) .0"), LexTryParseString(Value, (TEXT(".0"))) && Value == 0);
+		TestTrue(TEXT("(int32 conversion from string) 0."), LexTryParseString(Value, (TEXT("0."))) && Value == 0);
+
+		// Scientific notation
+		TestTrue(TEXT("(int32 conversion from string) 1.0e+10"), LexTryParseString(Value, (TEXT("1.0e+10"))));
+		TestTrue(TEXT("(int32 conversion from string) 1.0e-10"), LexTryParseString(Value, (TEXT("1.0e-10"))));
+		TestTrue(TEXT("(int32 conversion from string) 1.0e+10"), LexTryParseString(Value, (TEXT("0.0e+10"))));
+		TestTrue(TEXT("(int32 conversion from string) 1.0e-10"), LexTryParseString(Value, (TEXT("0.0e-10"))));
+		TestTrue(TEXT("(int32 conversion from string) 1e+10"), LexTryParseString(Value, (TEXT("1e+10"))));
+		TestTrue(TEXT("(int32 conversion from string) 1e-10"), LexTryParseString(Value, (TEXT("1e-10"))));
+
+		// Float as hex
+		TestTrue(TEXT("(int32 conversion from string) 1.0e+10"), LexTryParseString(Value, (TEXT("1.0e+10"))));
+
+		// Some numbers with whitespace
+		TestTrue(TEXT("(int32 conversion from string) whitespace"), LexTryParseString(Value, (TEXT("   2.5   "))) && Value == 2.5);
+		TestTrue(TEXT("(int32 conversion from string) whitespace"), LexTryParseString(Value, (TEXT("\t3.0\t"))) && Value == 3.0);
+		TestTrue(TEXT("(int32 conversion from string) whitespace"), LexTryParseString(Value, (TEXT("4.0   \t"))) && Value == 4.0);
+		TestTrue(TEXT("(int32 conversion from string) whitespace"), LexTryParseString(Value, (TEXT("\r\n5.25"))) && Value == 5.25);
+		TestTrue(TEXT("(int32 conversion from string) whitespace"), LexTryParseString(Value, (TEXT(" 6 . 2 "))) && Value == 6.0);
+		TestTrue(TEXT("(int32 conversion from string) whitespace"), LexTryParseString(Value, (TEXT(" 56 . 2 "))) && Value == 56.0);
+		TestTrue(TEXT("(int32 conversion from string) whitespace"), LexTryParseString(Value, (TEXT(" 5 6 . 2 "))) && Value == 5.0);
+
+		// Non-finite special numbers. All shouldn't parse into an int
+		TestFalse(TEXT("(int32 no conversion from string) inf"), LexTryParseString(Value, (TEXT("inf"))));
+		TestFalse(TEXT("(int32 no conversion from string) nan"), LexTryParseString(Value, (TEXT("nan"))));
+		TestFalse(TEXT("(int32 no conversion from string) nan(ind)"), LexTryParseString(Value, (TEXT("nan(ind)"))));
+		TestFalse(TEXT("(int32 no conversion from string) nananananananana"), LexTryParseString(Value, (TEXT("nananananananana"))));
+		TestFalse(TEXT("(int32 no conversion from string) nan(ind)!"), LexTryParseString(Value, (TEXT("nan(ind)!"))));
+		TestFalse(TEXT("(int32 no conversion from string) infinity"), LexTryParseString(Value, (TEXT("infinity"))));
+		TestFalse(TEXT("(float conversion from string) ."), LexTryParseString(Value, (TEXT("."))));
+		TestFalse(TEXT("(float conversion from string) <empyty string>"), LexTryParseString(Value, (TEXT(""))));
+	}
+
+	return true;
+}
+
+#endif //WITH_DEV_AUTOMATION_TESTS
diff --git a/Engine/Source/Runtime/Core/Public/Containers/UnrealString.h b/Engine/Source/Runtime/Core/Public/Containers/UnrealString.h
--- a/Engine/Source/Runtime/Core/Public/Containers/UnrealString.h
+++ b/Engine/Source/Runtime/Core/Public/Containers/UnrealString.h
@@ -2150,19 +2150,42 @@ inline FString LexToSanitizedString(double Value)
 	return FString::SanitizeFloat(Value);
 }
 
-
 /** Parse a string into this type, returning whether it was successful */
 /** Specialization for arithmetic types */
 template<typename T>
 static typename TEnableIf<TIsArithmetic<T>::Value, bool>::Type
 LexTryParseString(T& OutValue, const TCHAR* Buffer)
 {
-	if (FCString::IsNumeric(Buffer))
+	if (Buffer[0] == '\0')
 	{
-		LexFromString(OutValue, Buffer);
-		return true;
+		OutValue = 0;
+		return false;
+	}
+
+	LexFromString(OutValue, Buffer);
+	if (OutValue == 0 && FMath::IsFinite(OutValue))
+	{
+		bool bSawZero = false;
+		TCHAR C = *Buffer;
+		while (C != '\0' && (C == '+' || C == '-' || FChar::IsWhitespace(C)))
+		{
+			C = *(++Buffer);
+		}
+
+		while (C != '\0' && !FChar::IsWhitespace(C) && (TIsFloatingPoint<T>::Value || C != '.'))
+		{
+			bSawZero = bSawZero || (C == '0');
+			if (!bSawZero && C != '.')
+			{
+				return false;
+			}
+
+			C = *(++Buffer);
+		}
+		return bSawZero;
 	}
-	return false;
+	
+	return true;
 }
 
 /** Try and parse a bool - always returns true */
