Build fix

#rnx

diff --git a/Engine/Source/Runtime/Core/Public/Containers/UnrealString.h b/Engine/Source/Runtime/Core/Public/Containers/UnrealString.h
--- a/Engine/Source/Runtime/Core/Public/Containers/UnrealString.h
+++ b/Engine/Source/Runtime/Core/Public/Containers/UnrealString.h
@@ -23,6 +23,7 @@
 #include "Templates/AndOrNot.h"
 #include "Templates/IsArrayOrRefOfType.h"
 #include "Templates/TypeHash.h"
+#include "Templates/IsFloatingPoint.h"
 
 struct FStringFormatArg;
 template<typename KeyType,typename ValueType,typename SetAllocator ,typename KeyFuncs > class TMap;
