Added flag to FGenericCrashContext to tell if the crash reporter is being handled out-of-process or not.

diff --git a/Engine/Source/Runtime/Core/Private/GenericPlatform/GenericPlatformCrashContext.cpp b/Engine/Source/Runtime/Core/Private/GenericPlatform/GenericPlatformCrashContext.cpp
--- a/Engine/Source/Runtime/Core/Private/GenericPlatform/GenericPlatformCrashContext.cpp
+++ b/Engine/Source/Runtime/Core/Private/GenericPlatform/GenericPlatformCrashContext.cpp
@@ -67,6 +67,7 @@ const TCHAR* const FGenericCrashContext::EngineModeExDirty = TEXT("Dirty");
 const TCHAR* const FGenericCrashContext::EngineModeExVanilla = TEXT("Vanilla");
 
 bool FGenericCrashContext::bIsInitialized = false;
+bool FGenericCrashContext::bIsOutOfProcess = false;
 int32 FGenericCrashContext::StaticCrashContextIndex = 0;
 
 const FGuid FGenericCrashContext::ExecutionGuid = FGuid::NewGuid();
diff --git a/Engine/Source/Runtime/Core/Private/Windows/WindowsPlatformCrashContext.cpp b/Engine/Source/Runtime/Core/Private/Windows/WindowsPlatformCrashContext.cpp
--- a/Engine/Source/Runtime/Core/Private/Windows/WindowsPlatformCrashContext.cpp
+++ b/Engine/Source/Runtime/Core/Private/Windows/WindowsPlatformCrashContext.cpp
@@ -45,7 +45,7 @@
 #define CR_CLIENT_MAX_PATH_LEN 265
 #define CR_CLIENT_MAX_ARGS_LEN 256
 // Enable to use a non shipping version of the crash report client
-#define CR_USE_DEVELOPMENT_CLIENT 0
+#define CR_USE_DEVELOPMENT_CLIENT 1
 
 #pragma comment( lib, "version.lib" )
 #pragma comment( lib, "Shlwapi.lib" )
@@ -827,6 +827,7 @@ void CreateExceptionInfoString(EXCEPTION_RECORD* ExceptionRecord)
  **/
 class FCrashReportingThread
 {
+private:
 	/** Thread Id of reporter thread*/
 	DWORD ThreadId;
 	/** Thread handle to reporter thread */
@@ -935,7 +936,7 @@ public:
 			SetThreadPriority(Thread, THREAD_PRIORITY_BELOW_NORMAL);
 		}
 
-
+		FGenericCrashContext::SetIsOutOfProcessCrashReporter(CrashClientHandle.IsValid());
 	}
 
 	FORCENOINLINE ~FCrashReportingThread()
diff --git a/Engine/Source/Runtime/Core/Public/GenericPlatform/GenericPlatformCrashContext.h b/Engine/Source/Runtime/Core/Public/GenericPlatform/GenericPlatformCrashContext.h
--- a/Engine/Source/Runtime/Core/Public/GenericPlatform/GenericPlatformCrashContext.h
+++ b/Engine/Source/Runtime/Core/Public/GenericPlatform/GenericPlatformCrashContext.h
@@ -259,6 +259,20 @@ public:
 		return bIsInitialized;
 	}
 
+	/**
+	 * @return true if crash reporting is being handled out-of-process.
+	 */
+	static bool IsOutOfProcessCrashReporter()
+	{
+		return bIsOutOfProcess;
+	}
+
+	/** Set whether or not the out-of-process crash reporter is running. */
+	static void SetIsOutOfProcessCrashReporter(bool bInValue)
+	{
+		bIsOutOfProcess = bInValue;
+	}
+
 	/** Default constructor. Optionally pass a process handle if building a crash context for a process other then current. */
 	FGenericCrashContext(ECrashContextType InType, const TCHAR* ErrorMessage);
 
@@ -434,6 +448,9 @@ private:
 	/**	Whether the Initialize() has been called */
 	static bool bIsInitialized;
 
+	/** Whether or not crash reporting is being handled out-of-process. */
+	static bool bIsOutOfProcess;
+
 	/**	Static counter records how many crash contexts have been constructed */
 	static int32 StaticCrashContextIndex;
 
