Allow fallback to a Development CrashReportClient in non-shipping builds

#rnx

diff --git a/Engine/Source/Runtime/Core/Private/Windows/WindowsPlatformCrashContext.cpp b/Engine/Source/Runtime/Core/Private/Windows/WindowsPlatformCrashContext.cpp
--- a/Engine/Source/Runtime/Core/Private/Windows/WindowsPlatformCrashContext.cpp
+++ b/Engine/Source/Runtime/Core/Private/Windows/WindowsPlatformCrashContext.cpp
@@ -44,8 +44,6 @@
 
 #define CR_CLIENT_MAX_PATH_LEN 265
 #define CR_CLIENT_MAX_ARGS_LEN 256
-// Enable to use a non shipping version of the crash report client
-#define CR_USE_DEVELOPMENT_CLIENT 1
 
 #pragma comment( lib, "version.lib" )
 #pragma comment( lib, "Shlwapi.lib" )
@@ -348,24 +346,36 @@ static FORCEINLINE bool CreatePipeWrite(void*& ReadPipe, void*& WritePipe)
  */
 bool CreateCrashReportClientPath(TCHAR* OutClientPath, int32 MaxLength)
 {
-#if CR_USE_DEVELOPMENT_CLIENT
-	static const TCHAR CrashReportClientExeName[] = TEXT("CrashReportClient-Win64-Development.exe");
-#else
-	static const TCHAR CrashReportClientExeName[] = TEXT("CrashReportClient.exe");
-#endif 
+	auto CreateCrashReportClientPathImpl = [&OutClientPath, MaxLength](const TCHAR* CrashReportClientExeName) -> bool
+	{
+		const TCHAR* EngineDir = FPlatformMisc::EngineDir();
+		const TCHAR* BinariesDir = FPlatformProcess::GetBinariesSubdirectory();
+
+		// Find the path to crash reporter binary. Avoid creating FStrings.
+		*OutClientPath = 0;
+		FCString::Strncat(OutClientPath, EngineDir, MaxLength);
+		FCString::Strncat(OutClientPath, TEXT("Binaries/"), MaxLength);
+		FCString::Strncat(OutClientPath, BinariesDir, MaxLength);
+		FCString::Strncat(OutClientPath, TEXT("/"), MaxLength);
+		FCString::Strncat(OutClientPath, CrashReportClientExeName, MaxLength);
+
+		const DWORD Results = GetFileAttributesW(OutClientPath);
+		return Results != INVALID_FILE_ATTRIBUTES;
+	};
 
-	const TCHAR* EngineDir = FPlatformMisc::EngineDir();
-	const TCHAR* BinariesDir = FPlatformProcess::GetBinariesSubdirectory();
+	if (CreateCrashReportClientPathImpl(TEXT("CrashReportClient.exe")))
+	{
+		return true;
+	}
 
-	// Find the path to crash reporter binary. Avoid creating FStrings.
-	FCString::Strncat(OutClientPath, EngineDir, MaxLength);
-	FCString::Strncat(OutClientPath, TEXT("Binaries/"), MaxLength);
-	FCString::Strncat(OutClientPath, BinariesDir, MaxLength);
-	FCString::Strncat(OutClientPath, TEXT("/"), MaxLength);
-	FCString::Strncat(OutClientPath, CrashReportClientExeName, MaxLength);
+#if !(UE_BUILD_TEST || UE_BUILD_SHIPPING)
+	if (CreateCrashReportClientPathImpl(TEXT("CrashReportClient-Win64-Development.exe")))
+	{
+		return true;
+	}
+#endif
 
-	const DWORD Results = GetFileAttributesW(OutClientPath);
-	return Results != INVALID_FILE_ATTRIBUTES;
+	return false;
 }
 
 /**
