Fixed iOS compile errors if LLM is disabled.

#rnx

diff --git a/Engine/Source/Runtime/Core/Private/HAL/MallocBinned.cpp b/Engine/Source/Runtime/Core/Private/HAL/MallocBinned.cpp
--- a/Engine/Source/Runtime/Core/Private/HAL/MallocBinned.cpp
+++ b/Engine/Source/Runtime/Core/Private/HAL/MallocBinned.cpp
@@ -23,7 +23,7 @@ DEFINE_STAT(STAT_Binned_CurrentAllocs);
 DEFINE_STAT(STAT_Binned_TotalAllocs);
 DEFINE_STAT(STAT_Binned_SlackCurrent);
 
-#if USE_OS_SMALL_BLOCK_GRAB_MEMORY_FROM_OS
+#if USE_OS_SMALL_BLOCK_GRAB_MEMORY_FROM_OS && ENABLE_LOW_LEVEL_MEM_TRACKER
 DEFINE_STAT(STAT_Binned_NanoMallocPages_Current);
 DEFINE_STAT(STAT_Binned_NanoMallocPages_Peak);
 DEFINE_STAT(STAT_Binned_NanoMallocPages_WasteCurrent);
@@ -1087,7 +1087,7 @@ void FMallocBinned::GetAllocatorStats( FGenericMemoryStats& out_Stats )
 	out_Stats.Add( GET_STATDESCRIPTION( STAT_Binned_TotalAllocs ), LocalTotalAllocs );
 	out_Stats.Add( GET_STATDESCRIPTION( STAT_Binned_SlackCurrent ), LocalSlackCurrent );
 	
-#if USE_OS_SMALL_BLOCK_GRAB_MEMORY_FROM_OS
+#if USE_OS_SMALL_BLOCK_GRAB_MEMORY_FROM_OS && ENABLE_LOW_LEVEL_MEM_TRACKER
 	out_Stats.Add( GET_STATDESCRIPTION( STAT_Binned_NanoMallocPages_Current ), LocalNanoMallocPagesCurrent );
 	out_Stats.Add( GET_STATDESCRIPTION( STAT_Binned_NanoMallocPages_Peak ), LocalNanoMallocPagesPeak );
 	out_Stats.Add( GET_STATDESCRIPTION( STAT_Binned_NanoMallocPages_WasteCurrent ), LocalNanoMallocPagesWaste );
@@ -1675,7 +1675,7 @@ void FMallocBinned::UpdateStats()
 	SIZE_T	LocalTotalAllocs = 0;
 	SIZE_T	LocalSlackCurrent = 0;
 
-#if USE_OS_SMALL_BLOCK_GRAB_MEMORY_FROM_OS
+#if USE_OS_SMALL_BLOCK_GRAB_MEMORY_FROM_OS && ENABLE_LOW_LEVEL_MEM_TRACKER
 	SIZE_T LocalNanoMallocPagesCurrent = 0;
 	SIZE_T LocalNanoMallocPagesPeak = 0;
 	SIZE_T LocalNanoMallocPagesWaste = 0;
@@ -1699,7 +1699,7 @@ void FMallocBinned::UpdateStats()
 		LocalTotalAllocs = TotalAllocs;
 		LocalSlackCurrent = SlackCurrent;
 		
-#if USE_OS_SMALL_BLOCK_GRAB_MEMORY_FROM_OS
+#if USE_OS_SMALL_BLOCK_GRAB_MEMORY_FROM_OS && ENABLE_LOW_LEVEL_MEM_TRACKER
 		LocalNanoMallocPagesCurrent = NanoMallocPagesCurrent;
 		LocalNanoMallocPagesPeak = NanoMallocPagesPeak;
 		LocalNanoMallocPagesWaste = NanoMallocPagesWaste;
@@ -1717,7 +1717,7 @@ void FMallocBinned::UpdateStats()
 	SET_DWORD_STAT( STAT_Binned_TotalAllocs, LocalTotalAllocs );
 	SET_MEMORY_STAT( STAT_Binned_SlackCurrent, LocalSlackCurrent );
 	
-#if USE_OS_SMALL_BLOCK_GRAB_MEMORY_FROM_OS
+#if USE_OS_SMALL_BLOCK_GRAB_MEMORY_FROM_OS && ENABLE_LOW_LEVEL_MEM_TRACKER
 	SET_MEMORY_STAT( STAT_Binned_NanoMallocPages_Current, LocalNanoMallocPagesCurrent );
 	SET_MEMORY_STAT( STAT_Binned_NanoMallocPages_Peak, LocalNanoMallocPagesPeak );
 	SET_MEMORY_STAT( STAT_Binned_NanoMallocPages_WasteCurrent, LocalNanoMallocPagesWaste );
