UnrealInsights:
TRACE_CPUPROFILER_EVENT_SCOPE_TEXT now supports dynamic strings
Fixed instrumentation points that were incorrectly using TRACE_CPUPROFILER_EVENT_SCOPE_TEXT for static scope names

diff --git a/Engine/Source/Developer/DerivedDataCache/Private/DerivedDataCache.cpp b/Engine/Source/Developer/DerivedDataCache/Private/DerivedDataCache.cpp
--- a/Engine/Source/Developer/DerivedDataCache/Private/DerivedDataCache.cpp
+++ b/Engine/Source/Developer/DerivedDataCache/Private/DerivedDataCache.cpp
@@ -30,7 +30,7 @@ DEFINE_STAT(STAT_DDC_SyncBuildTime);
 DEFINE_STAT(STAT_DDC_ExistTime);
 
 //#define DDC_SCOPE_CYCLE_COUNTER(x) QUICK_SCOPE_CYCLE_COUNTER(STAT_ ## x)
-#define DDC_SCOPE_CYCLE_COUNTER(x) TRACE_CPUPROFILER_EVENT_SCOPE_TEXT(TEXT(#x));
+#define DDC_SCOPE_CYCLE_COUNTER(x) TRACE_CPUPROFILER_EVENT_SCOPE(x);
 
 
 #if ENABLE_COOK_STATS
@@ -158,12 +158,12 @@ class FDerivedDataCache : public FDerivedDataCacheInterface
 		/** Async worker that checks the cache backend and if that fails, calls the deriver to build the data and then puts the results to the cache **/
 		void DoWork()
 		{
-			TRACE_CPUPROFILER_EVENT_SCOPE_TEXT(TEXT("DDC_DoWork"));
+			TRACE_CPUPROFILER_EVENT_SCOPE(DDC_DoWork);
 
 			const int32 NumBeforeDDC = Data.Num();
 			bool bGetResult;
 			{
-				TRACE_CPUPROFILER_EVENT_SCOPE_TEXT(TEXT("DDC_Get"));
+				TRACE_CPUPROFILER_EVENT_SCOPE(DDC_Get);
 
 				INC_DWORD_STAT(STAT_DDC_NumGets);
 				STAT(double ThisTime = 0);
@@ -217,7 +217,7 @@ class FDerivedDataCache : public FDerivedDataCacheInterface
 			else if (DataDeriver)
 			{
 				{
-					TRACE_CPUPROFILER_EVENT_SCOPE_TEXT(TEXT("DDC_Build"));
+					TRACE_CPUPROFILER_EVENT_SCOPE(DDC_Build);
 
 					INC_DWORD_STAT(STAT_DDC_NumBuilds);
 					STAT(double ThisTime = 0);
@@ -234,7 +234,7 @@ class FDerivedDataCache : public FDerivedDataCacheInterface
 				{
 					check(Data.Num());
 
-					TRACE_CPUPROFILER_EVENT_SCOPE_TEXT(TEXT("DDC_Put"));
+					TRACE_CPUPROFILER_EVENT_SCOPE(DDC_Put);
 
 					INC_DWORD_STAT(STAT_DDC_NumPuts);
 					STAT(double ThisTime = 0);
diff --git a/Engine/Source/Developer/TraceServices/Private/Analyzers/CpuProfilerTraceAnalysis.cpp b/Engine/Source/Developer/TraceServices/Private/Analyzers/CpuProfilerTraceAnalysis.cpp
--- a/Engine/Source/Developer/TraceServices/Private/Analyzers/CpuProfilerTraceAnalysis.cpp
+++ b/Engine/Source/Developer/TraceServices/Private/Analyzers/CpuProfilerTraceAnalysis.cpp
@@ -16,6 +16,7 @@ void FCpuProfilerAnalyzer::OnAnalysisBegin(const FOnAnalysisContext& Context)
 	auto& Builder = Context.InterfaceBuilder;
 
 	Builder.RouteEvent(RouteId_EventSpec, "CpuProfiler", "EventSpec");
+	Builder.RouteEvent(RouteId_EventSpecEx, "CpuProfiler", "EventSpecEx");
 	Builder.RouteEvent(RouteId_EventBatch, "CpuProfiler", "EventBatch");
 	Builder.RouteEvent(RouteId_EndCapture, "CpuProfiler", "EndCapture");
 }
@@ -29,14 +30,22 @@ bool FCpuProfilerAnalyzer::OnEvent(uint16 RouteId, const FOnEventContext& Contex
 	{
 	case RouteId_EventSpec:
 	{
-		uint16 Id = EventData.GetValue<uint16>("Id");
-		if (ScopeIdToEventIdMap.Contains(Id))
+		uint32 Id = EventData.GetValue<uint16>("Id");
+		const TCHAR* ScopeName = Session.StoreString(reinterpret_cast<const TCHAR*>(EventData.GetAttachment()));
+		DefineScope(Id, ScopeName);
+		break;
+	}
+	case RouteId_EventSpecEx:
+	{
+		uint32 Id = EventData.GetValue<uint32>("Id");
+		uint8 CharSize = EventData.GetValue<uint8>("CharSize");
+		if (CharSize == 1)
 		{
-			TimingProfilerProvider.SetTimerName(ScopeIdToEventIdMap[Id], reinterpret_cast<const TCHAR*>(EventData.GetAttachment()));
+			DefineScope(Id, Session.StoreString(StringCast<TCHAR>(reinterpret_cast<const ANSICHAR*>(EventData.GetAttachment())).Get()));
 		}
-		else
+		else if (CharSize == 2)
 		{
-			ScopeIdToEventIdMap.Add(Id, TimingProfilerProvider.AddCpuTimer(reinterpret_cast<const TCHAR*>(EventData.GetAttachment())));
+			DefineScope(Id, Session.StoreString(reinterpret_cast<const TCHAR*>(EventData.GetAttachment())));
 		}
 		break;
 	}
@@ -59,7 +68,7 @@ bool FCpuProfilerAnalyzer::OnEvent(uint16 RouteId, const FOnEventContext& Contex
 			{
 				EventScopeState& ScopeState = ThreadState->ScopeStack.AddDefaulted_GetRef();
 				ScopeState.StartCycle = ActualCycle;
-				uint16 SpecId = FTraceAnalyzerUtils::Decode7bit(BufferPtr);
+				uint32 SpecId = FTraceAnalyzerUtils::Decode7bit(BufferPtr);
 				uint32* FindIt = ScopeIdToEventIdMap.Find(SpecId);
 				if (!FindIt)
 				{
@@ -104,6 +113,28 @@ bool FCpuProfilerAnalyzer::OnEvent(uint16 RouteId, const FOnEventContext& Contex
 	return true;
 }
 
+void FCpuProfilerAnalyzer::DefineScope(uint32 Id, const TCHAR* Name)
+{
+	if (ScopeIdToEventIdMap.Contains(Id))
+	{
+		TimingProfilerProvider.SetTimerName(ScopeIdToEventIdMap[Id], Name);
+	}
+	else
+	{
+		uint32* FindEventIdByName = ScopeNameToEventIdMap.Find(Name);
+		if (FindEventIdByName)
+		{
+			ScopeIdToEventIdMap.Add(Id, *FindEventIdByName);
+		}
+		else
+		{
+			uint32 NewTimerId = TimingProfilerProvider.AddCpuTimer(Name);
+			ScopeIdToEventIdMap.Add(Id, NewTimerId);
+			ScopeNameToEventIdMap.Add(Name, NewTimerId);
+		}
+	}
+}
+
 TSharedRef<FCpuProfilerAnalyzer::FThreadState> FCpuProfilerAnalyzer::GetThreadState(uint32 ThreadId)
 {
 	if (!ThreadStatesMap.Contains(ThreadId))
diff --git a/Engine/Source/Developer/TraceServices/Private/Analyzers/CpuProfilerTraceAnalysis.h b/Engine/Source/Developer/TraceServices/Private/Analyzers/CpuProfilerTraceAnalysis.h
--- a/Engine/Source/Developer/TraceServices/Private/Analyzers/CpuProfilerTraceAnalysis.h
+++ b/Engine/Source/Developer/TraceServices/Private/Analyzers/CpuProfilerTraceAnalysis.h
@@ -34,11 +34,13 @@ private:
 		double LastCycle = 0.0;
 	};
 
+	void DefineScope(uint32 Id, const TCHAR* ScopeName);
 	TSharedRef<FThreadState> GetThreadState(uint32 ThreadId);
 
 	enum : uint16
 	{
 		RouteId_EventSpec,
+		RouteId_EventSpecEx,
 		RouteId_EventBatch,
 		RouteId_EndCapture,
 	};
@@ -46,7 +48,8 @@ private:
 	Trace::IAnalysisSession& Session;
 	Trace::FTimingProfilerProvider& TimingProfilerProvider;
 	TMap<uint32, TSharedRef<FThreadState>> ThreadStatesMap;
-	TMap<uint16, uint32> ScopeIdToEventIdMap;
+	TMap<uint32, uint32> ScopeIdToEventIdMap;
+	TMap<const TCHAR*, uint32> ScopeNameToEventIdMap;
 	uint64 TotalEventSize = 0;
 	uint64 TotalScopeCount = 0;
 	double BytesPerScope = 0.0;
diff --git a/Engine/Source/Editor/UnrealEd/Private/Commandlets/CookCommandlet.cpp b/Engine/Source/Editor/UnrealEd/Private/Commandlets/CookCommandlet.cpp
--- a/Engine/Source/Editor/UnrealEd/Private/Commandlets/CookCommandlet.cpp
+++ b/Engine/Source/Editor/UnrealEd/Private/Commandlets/CookCommandlet.cpp
@@ -597,7 +597,7 @@ int32 UCookCommandlet::Main(const FString& CmdLineParams)
 
 bool UCookCommandlet::CookByTheBook( const TArray<ITargetPlatform*>& Platforms, TArray<FString>& FilesInPath )
 {
-	TRACE_CPUPROFILER_EVENT_SCOPE_TEXT(TEXT("CookByTheBook"));
+	TRACE_CPUPROFILER_EVENT_SCOPE(CookByTheBook);
 
 	COOK_STAT(FScopedDurationTimer CookByTheBookTimer(DetailedCookStats::CookByTheBookTimeSec));
 	UCookOnTheFlyServer *CookOnTheFlyServer = NewObject<UCookOnTheFlyServer>();
diff --git a/Engine/Source/Editor/UnrealEd/Private/Commandlets/TextAssetCommandlet.cpp b/Engine/Source/Editor/UnrealEd/Private/Commandlets/TextAssetCommandlet.cpp
--- a/Engine/Source/Editor/UnrealEd/Private/Commandlets/TextAssetCommandlet.cpp
+++ b/Engine/Source/Editor/UnrealEd/Private/Commandlets/TextAssetCommandlet.cpp
@@ -146,7 +146,7 @@ bool UTextAssetCommandlet::DoTextAssetProcessing(const FString& InCommandLine)
 
 bool UTextAssetCommandlet::DoTextAssetProcessing(const FProcessingArgs& InArgs)
 {
-	TRACE_CPUPROFILER_EVENT_SCOPE_TEXT(TEXT("UTextAssetCommandlet::Main"));
+	TRACE_CPUPROFILER_EVENT_SCOPE(UTextAssetCommandlet::Main);
 
 	RepairDamagedFiles();
 
diff --git a/Engine/Source/Editor/UnrealEd/Private/CookOnTheFlyServer.cpp b/Engine/Source/Editor/UnrealEd/Private/CookOnTheFlyServer.cpp
--- a/Engine/Source/Editor/UnrealEd/Private/CookOnTheFlyServer.cpp
+++ b/Engine/Source/Editor/UnrealEd/Private/CookOnTheFlyServer.cpp
@@ -311,7 +311,7 @@ void ClearHierarchyTimers()
 
 #define CREATE_TIMER(name, incrementScope) FScopeTimer ScopeTimer##name(__COUNTER__, #name, incrementScope); 
 
-#define SCOPE_TIMER(name)				TRACE_CPUPROFILER_EVENT_SCOPE_TEXT(TEXT(#name)); CREATE_TIMER(name, true); ScopeTimer##name.Start();
+#define SCOPE_TIMER(name)				TRACE_CPUPROFILER_EVENT_SCOPE(name); CREATE_TIMER(name, true); ScopeTimer##name.Start();
 
 #else
 #define SCOPE_TIMER(name)
@@ -1464,7 +1464,7 @@ UCookOnTheFlyServer::~UCookOnTheFlyServer()
 // this tick only happens in the editor cook commandlet directly calls tick on the side
 void UCookOnTheFlyServer::Tick(float DeltaTime)
 {
-	TRACE_CPUPROFILER_EVENT_SCOPE_TEXT(TEXT("UCookOnTheFlyServer::Tick"));
+	TRACE_CPUPROFILER_EVENT_SCOPE(UCookOnTheFlyServer::Tick);
 
 	check(IsCookingInEditor());
 
@@ -4070,7 +4070,7 @@ void UCookOnTheFlyServer::SaveCookedPackage(UPackage* Package, uint32 SaveFlags,
 
 void UCookOnTheFlyServer::Initialize( ECookMode::Type DesiredCookMode, ECookInitializationFlags InCookFlags, const FString &InOutputDirectoryOverride )
 {
-	TRACE_CPUPROFILER_EVENT_SCOPE_TEXT(TEXT("UCookOnTheFlyServer::Initialize"));
+	TRACE_CPUPROFILER_EVENT_SCOPE(UCookOnTheFlyServer::Initialize);
 
 	OutputDirectoryOverride = InOutputDirectoryOverride;
 	CurrentCookMode = DesiredCookMode;
@@ -5064,7 +5064,7 @@ FName UCookOnTheFlyServer::ConvertCookedPathToUncookedPath(
 
 void UCookOnTheFlyServer::GetAllCookedFiles(TMap<FName, FName>& UncookedPathToCookedPath, const FString& SandboxRootDir)
 {
-	TRACE_CPUPROFILER_EVENT_SCOPE_TEXT(TEXT("UCookOnTheFlyServer::GetAllCookedFiles"));
+	TRACE_CPUPROFILER_EVENT_SCOPE(UCookOnTheFlyServer::GetAllCookedFiles);
 
 	TArray<FString> CookedFiles;
 	{
@@ -5093,7 +5093,7 @@ void UCookOnTheFlyServer::GetAllCookedFiles(TMap<FName, FName>& UncookedPathToCo
 
 void UCookOnTheFlyServer::PopulateCookedPackagesFromDisk(const TArray<ITargetPlatform*>& Platforms)
 {
-	TRACE_CPUPROFILER_EVENT_SCOPE_TEXT(TEXT("UCookOnTheFlyServer::PopulateCookedPackagesFromDisk"));
+	TRACE_CPUPROFILER_EVENT_SCOPE(UCookOnTheFlyServer::PopulateCookedPackagesFromDisk);
 
 	// See what files are out of date in the sandbox folder
 	for (int32 Index = 0; Index < Platforms.Num(); Index++)
@@ -5357,7 +5357,7 @@ const FString ExtractPackageNameFromObjectPath( const FString ObjectPath )
 
 void UCookOnTheFlyServer::CleanSandbox(const bool bIterative)
 {
-	TRACE_CPUPROFILER_EVENT_SCOPE_TEXT(TEXT("UCookOnTheFlyServer::CleanSandbox"));
+	TRACE_CPUPROFILER_EVENT_SCOPE(UCookOnTheFlyServer::CleanSandbox);
 
 	const TArray<ITargetPlatform*>& Platforms = GetCookingTargetPlatforms();
 
@@ -5432,7 +5432,7 @@ void UCookOnTheFlyServer::CleanSandbox(const bool bIterative)
 
 void UCookOnTheFlyServer::GenerateAssetRegistry()
 {
-	TRACE_CPUPROFILER_EVENT_SCOPE_TEXT(TEXT("UCookOnTheFlyServer::GenerateAssetRegistry"));
+	TRACE_CPUPROFILER_EVENT_SCOPE(UCookOnTheFlyServer::GenerateAssetRegistry);
 
 	// Cache asset registry for later
 	FAssetRegistryModule& AssetRegistryModule = FModuleManager::LoadModuleChecked<FAssetRegistryModule>(TEXT("AssetRegistry"));
@@ -5571,7 +5571,7 @@ void UCookOnTheFlyServer::AddFileToCook( TArray<FName>& InOutFilesToCook, const
 
 void UCookOnTheFlyServer::CollectFilesToCook(TArray<FName>& FilesInPath, const TArray<FString>& CookMaps, const TArray<FString>& InCookDirectories, const TArray<FString> &IniMapSections, ECookByTheBookOptions FilesToCookFlags)
 {
-	TRACE_CPUPROFILER_EVENT_SCOPE_TEXT(TEXT("UCookOnTheFlyServer::GenerateAssetRegistry"));
+	TRACE_CPUPROFILER_EVENT_SCOPE(UCookOnTheFlyServer::GenerateAssetRegistry);
 
 #if OUTPUT_TIMING
 	SCOPE_TIMER(CollectFilesToCook);
diff --git a/Engine/Source/Editor/UnrealEd/Private/UnrealEdGlobals.cpp b/Engine/Source/Editor/UnrealEd/Private/UnrealEdGlobals.cpp
--- a/Engine/Source/Editor/UnrealEd/Private/UnrealEdGlobals.cpp
+++ b/Engine/Source/Editor/UnrealEd/Private/UnrealEdGlobals.cpp
@@ -202,7 +202,7 @@ int32 EditorReinit()
 
 void EditorExit()
 {
-	TRACE_CPUPROFILER_EVENT_SCOPE_TEXT(TEXT("EditorExit"));
+	TRACE_CPUPROFILER_EVENT_SCOPE(EditorExit);
 
 	GLevelEditorModeTools().SetDefaultMode(FBuiltinEditorModes::EM_Default);
 	GLevelEditorModeTools().DeactivateAllModes(); // this also activates the default mode
diff --git a/Engine/Source/Runtime/Core/Private/Misc/CoreGlobals.cpp b/Engine/Source/Runtime/Core/Private/Misc/CoreGlobals.cpp
--- a/Engine/Source/Runtime/Core/Private/Misc/CoreGlobals.cpp
+++ b/Engine/Source/Runtime/Core/Private/Misc/CoreGlobals.cpp
@@ -9,7 +9,7 @@
 #include "Misc/CoreStats.h"
 #include "Misc/Compression.h"
 #include "Misc/LazySingleton.h"
-
+#include "ProfilingDebugging/MiscTrace.h"
 
 #ifndef FAST_PATH_UNIQUE_NAME_GENERATION
 #define FAST_PATH_UNIQUE_NAME_GENERATION (!WITH_EDITORONLY_DATA)
diff --git a/Engine/Source/Runtime/Core/Private/Modules/ModuleManager.cpp b/Engine/Source/Runtime/Core/Private/Modules/ModuleManager.cpp
--- a/Engine/Source/Runtime/Core/Private/Modules/ModuleManager.cpp
+++ b/Engine/Source/Runtime/Core/Private/Modules/ModuleManager.cpp
@@ -635,7 +635,7 @@ void FModuleManager::UnloadModulesAtShutdown()
 {
 	ensure(IsInGameThread());
 
-	TRACE_CPUPROFILER_EVENT_SCOPE_TEXT(TEXT("UnloadModulesAtShutdown"));
+	TRACE_CPUPROFILER_EVENT_SCOPE(UnloadModulesAtShutdown);
 
 	struct FModulePair
 	{
diff --git a/Engine/Source/Runtime/Core/Private/ProfilingDebugging/CpuProfilerTrace.cpp b/Engine/Source/Runtime/Core/Private/ProfilingDebugging/CpuProfilerTrace.cpp
--- a/Engine/Source/Runtime/Core/Private/ProfilingDebugging/CpuProfilerTrace.cpp
+++ b/Engine/Source/Runtime/Core/Private/ProfilingDebugging/CpuProfilerTrace.cpp
@@ -9,9 +9,10 @@
 
 #if CPUPROFILERTRACE_ENABLED
 
-UE_TRACE_EVENT_BEGIN(CpuProfiler, EventSpec, Always)
-	UE_TRACE_EVENT_FIELD(uint16, Id)
+UE_TRACE_EVENT_BEGIN(CpuProfiler, EventSpecEx, Always)
+	UE_TRACE_EVENT_FIELD(uint32, Id)
 	UE_TRACE_EVENT_FIELD(uint16, Group)
+	UE_TRACE_EVENT_FIELD(uint8, CharSize)
 UE_TRACE_EVENT_END()
 
 UE_TRACE_EVENT_BEGIN(CpuProfiler, EventBatch)
@@ -24,22 +25,50 @@ UE_TRACE_EVENT_END()
 
 struct FCpuProfilerTraceInternal
 {
-public:
 	enum
 	{
 		MaxBufferSize = 255,
-		MaxEncodedEventSize = 12, // 9 + 3
+		MaxEncodedEventSize = 14, // 9 + 5
 		FullBufferThreshold = MaxBufferSize - MaxEncodedEventSize,
 	};
 
+	template<typename CharType>
+	struct FDynamicScopeNameMapKeyFuncs
+	{
+		typedef const CharType* KeyType;
+		typedef const CharType* KeyInitType;
+		typedef const TPairInitializer<const CharType*, uint32>& ElementInitType;
+
+		enum { bAllowDuplicateKeys = false };
+
+		static FORCEINLINE bool Matches(const CharType* A, const CharType* B)
+		{
+			return TCString<CharType>::Stricmp(A, B) == 0;
+		}
+
+		static FORCEINLINE uint32 GetKeyHash(const CharType* Key)
+		{
+			return GetTypeHash(Key);
+		}
+
+		static FORCEINLINE KeyInitType GetSetKey(ElementInitType Element)
+		{
+			return Element.Key;
+		}
+	};
+
 	struct FThreadBuffer
 		: public FTlsAutoCleanup
 	{
 		uint64 LastCycle = 0;
 		uint16 BufferSize = 0;
 		uint8 Buffer[MaxBufferSize];
+		FMemStackBase DynamicScopeNamesMemory;
+		TMap<const ANSICHAR*, uint32, FDefaultSetAllocator, FDynamicScopeNameMapKeyFuncs<ANSICHAR>> DynamicAnsiScopeNamesMap;
+		TMap<const TCHAR*, uint32, FDefaultSetAllocator, FDynamicScopeNameMapKeyFuncs<TCHAR>> DynamicTCharScopeNamesMap;
 	};
 
+	uint32 static GetNextSpecId();
 	FORCENOINLINE static FThreadBuffer* CreateThreadBuffer();
 	FORCENOINLINE static void FlushThreadBuffer(FThreadBuffer* ThreadBuffer);
 	FORCENOINLINE static void EndCapture(FThreadBuffer* ThreadBuffer);
@@ -76,7 +105,7 @@ void FCpuProfilerTraceInternal::EndCapture(FThreadBuffer* InThreadBuffer)
 	InThreadBuffer->BufferSize = 0;
 }
 
-void FCpuProfilerTrace::OutputBeginEvent(uint16 SpecId)
+void FCpuProfilerTrace::OutputBeginEvent(uint32 SpecId)
 {
 	++FCpuProfilerTraceInternal::ThreadDepth;
 	bool bEventEnabled = UE_TRACE_EVENT_IS_ENABLED(CpuProfiler, EventBatch);
@@ -108,6 +137,58 @@ void FCpuProfilerTrace::OutputBeginEvent(uint16 SpecId)
 	ThreadBuffer->BufferSize = BufferPtr - ThreadBuffer->Buffer;
 }
 
+void FCpuProfilerTrace::OutputBeginDynamicEvent(const ANSICHAR* Name, ECpuProfilerGroup Group)
+{
+	if (UE_TRACE_EVENT_IS_ENABLED(CpuProfiler, EventBatch))
+	{
+		FCpuProfilerTraceInternal::FThreadBuffer* ThreadBuffer = FCpuProfilerTraceInternal::ThreadBuffer;
+		if (!ThreadBuffer)
+		{
+			ThreadBuffer = FCpuProfilerTraceInternal::CreateThreadBuffer();
+		}
+		uint32& SpecId = ThreadBuffer->DynamicAnsiScopeNamesMap.FindOrAdd(Name);
+		if (SpecId == 0)
+		{
+			int32 NameSize = strlen(Name) + 1;
+			ANSICHAR* NameCopy = reinterpret_cast<ANSICHAR*>(ThreadBuffer->DynamicScopeNamesMemory.Alloc(NameSize, alignof(ANSICHAR)));
+			FMemory::Memmove(NameCopy, Name, NameSize);
+			SpecId = OutputEventType(NameCopy, Group);
+			ThreadBuffer->DynamicAnsiScopeNamesMap.Add(NameCopy, SpecId);
+		}
+		OutputBeginEvent(SpecId);
+	}
+	else
+	{
+		OutputBeginEvent(0);
+	}
+}
+
+void FCpuProfilerTrace::OutputBeginDynamicEvent(const TCHAR* Name, ECpuProfilerGroup Group)
+{
+	if (UE_TRACE_EVENT_IS_ENABLED(CpuProfiler, EventBatch))
+	{
+		FCpuProfilerTraceInternal::FThreadBuffer* ThreadBuffer = FCpuProfilerTraceInternal::ThreadBuffer;
+		if (!ThreadBuffer)
+		{
+			ThreadBuffer = FCpuProfilerTraceInternal::CreateThreadBuffer();
+		}
+		uint32& SpecId = ThreadBuffer->DynamicTCharScopeNamesMap.FindOrAdd(Name);
+		if (SpecId == 0)
+		{
+			int32 NameSize = (FCString::Strlen(Name) + 1) * sizeof(TCHAR);
+			TCHAR* NameCopy = reinterpret_cast<TCHAR*>(ThreadBuffer->DynamicScopeNamesMemory.Alloc(NameSize, alignof(TCHAR)));
+			FMemory::Memmove(NameCopy, Name, NameSize);
+			SpecId = OutputEventType(NameCopy, Group);
+			ThreadBuffer->DynamicTCharScopeNamesMap.Add(NameCopy, SpecId);
+		}
+		OutputBeginEvent(SpecId);
+	}
+	else
+	{
+		OutputBeginEvent(0);
+	}
+}
+
 void FCpuProfilerTrace::OutputEndEvent()
 {
 	--FCpuProfilerTraceInternal::ThreadDepth;
@@ -139,18 +220,34 @@ void FCpuProfilerTrace::OutputEndEvent()
 	ThreadBuffer->BufferSize = BufferPtr - ThreadBuffer->Buffer;
 }
 
-uint16 FCpuProfilerTrace::OutputEventType(const TCHAR* Name, ECpuProfilerGroup Group)
+uint32 FCpuProfilerTraceInternal::GetNextSpecId()
 {
 	static TAtomic<uint32> NextSpecId(1);
-	uint32 WideId = NextSpecId++;
-	check(WideId <= 0xFFFF);
-	uint16 Id = uint16(WideId);
+	return NextSpecId++;
+}
+
+uint32 FCpuProfilerTrace::OutputEventType(const TCHAR* Name, ECpuProfilerGroup Group)
+{
+	uint32 SpecId = FCpuProfilerTraceInternal::GetNextSpecId();
 	uint16 NameSize = (FCString::Strlen(Name) + 1) * sizeof(TCHAR);
-	UE_TRACE_LOG(CpuProfiler, EventSpec, NameSize)
-		<< EventSpec.Id(Id)
-		<< EventSpec.Group(uint16(Group))
-		<< EventSpec.Attachment(Name, NameSize);
-	return Id;
+	UE_TRACE_LOG(CpuProfiler, EventSpecEx, NameSize)
+		<< EventSpecEx.Id(SpecId)
+		<< EventSpecEx.Group(uint16(Group))
+		<< EventSpecEx.CharSize(sizeof(TCHAR))
+		<< EventSpecEx.Attachment(Name, NameSize);
+	return SpecId;
+}
+
+uint32 FCpuProfilerTrace::OutputEventType(const ANSICHAR* Name, ECpuProfilerGroup Group)
+{
+	uint32 SpecId = FCpuProfilerTraceInternal::GetNextSpecId();
+	uint16 NameSize = strlen(Name) + 1;
+	UE_TRACE_LOG(CpuProfiler, EventSpecEx, NameSize)
+		<< EventSpecEx.Id(SpecId)
+		<< EventSpecEx.Group(uint16(Group))
+		<< EventSpecEx.CharSize(1)
+		<< EventSpecEx.Attachment(Name, NameSize);
+	return SpecId;
 }
 
 void FCpuProfilerTrace::Init(const TCHAR* CmdLine)
diff --git a/Engine/Source/Runtime/Core/Public/CoreGlobals.h b/Engine/Source/Runtime/Core/Public/CoreGlobals.h
--- a/Engine/Source/Runtime/Core/Public/CoreGlobals.h
+++ b/Engine/Source/Runtime/Core/Public/CoreGlobals.h
@@ -71,7 +71,7 @@ struct CORE_API FScopedBootTiming
 };
 
 
-#define SCOPED_BOOT_TIMING(x) TRACE_CPUPROFILER_EVENT_SCOPE_TEXT(TEXT(x)); FScopedBootTiming ANONYMOUS_VARIABLE(BootTiming_)(x);
+#define SCOPED_BOOT_TIMING(x) TRACE_CPUPROFILER_EVENT_SCOPE_TEXT(x); FScopedBootTiming ANONYMOUS_VARIABLE(BootTiming_)(x);
 
 #define GLog GetGlobalLogSingleton()
 extern CORE_API FConfigCacheIni* GConfig;
diff --git a/Engine/Source/Runtime/Core/Public/HAL/PlatformMisc.h b/Engine/Source/Runtime/Core/Public/HAL/PlatformMisc.h
--- a/Engine/Source/Runtime/Core/Public/HAL/PlatformMisc.h
+++ b/Engine/Source/Runtime/Core/Public/HAL/PlatformMisc.h
@@ -119,7 +119,6 @@ public:
 #define SCOPED_NAMED_EVENT_FSTRING(Text, Color)  FScopedNamedEvent       ANONYMOUS_VARIABLE(NamedEvent_)         (Color, *Text);
 
 #define SCOPED_NAMED_EVENT_TEXT(Text, Color) \
-	TRACE_CPUPROFILER_EVENT_SCOPE_TEXT(TEXT(Text)) \
 	FScopedNamedEventStatic ANONYMOUS_VARIABLE(NamedEvent_)         (Color, NAMED_EVENT_STR(Text));
 
 #define SCOPED_NAMED_EVENT_F(Format, Color, ...) FScopedNamedEvent       ANONYMOUS_VARIABLE(NamedEvent_)         (Color, *FString::Printf(Format, __VA_ARGS__));
diff --git a/Engine/Source/Runtime/Core/Public/ProfilingDebugging/CpuProfilerTrace.h b/Engine/Source/Runtime/Core/Public/ProfilingDebugging/CpuProfilerTrace.h
--- a/Engine/Source/Runtime/Core/Public/ProfilingDebugging/CpuProfilerTrace.h
+++ b/Engine/Source/Runtime/Core/Public/ProfilingDebugging/CpuProfilerTrace.h
@@ -23,13 +23,16 @@ enum ECpuProfilerGroup
 struct FCpuProfilerTrace
 {
 	CORE_API static void Init(const TCHAR* CmdLine);
-	CORE_API static uint16 OutputEventType(const TCHAR* Name, ECpuProfilerGroup Group);
-	CORE_API static void OutputBeginEvent(uint16 SpecId);
+	FORCENOINLINE CORE_API static uint32 OutputEventType(const ANSICHAR* Name, ECpuProfilerGroup Group);
+	FORCENOINLINE CORE_API static uint32 OutputEventType(const TCHAR* Name, ECpuProfilerGroup Group);
+	CORE_API static void OutputBeginEvent(uint32 SpecId);
+	CORE_API static void OutputBeginDynamicEvent(const ANSICHAR* Name, ECpuProfilerGroup Group);
+	CORE_API static void OutputBeginDynamicEvent(const TCHAR* Name, ECpuProfilerGroup Group);
 	CORE_API static void OutputEndEvent();
 
 	struct FEventScope
 	{
-		FEventScope(uint16 InSpecId)
+		FEventScope(uint32 InSpecId)
 		{
 			OutputBeginEvent(InSpecId);
 		}
@@ -39,27 +42,45 @@ struct FCpuProfilerTrace
 			OutputEndEvent();
 		}
 	};
+
+	struct FDynamicEventScope
+	{
+		FDynamicEventScope(const ANSICHAR* InEventName, ECpuProfilerGroup Group)
+		{
+			OutputBeginDynamicEvent(InEventName, Group);
+		}
+
+		FDynamicEventScope(const TCHAR* InEventName, ECpuProfilerGroup Group)
+		{
+			OutputBeginDynamicEvent(InEventName, Group);
+		}
+
+		~FDynamicEventScope()
+		{
+			OutputEndEvent();
+		}
+	};
 };
 
 #define TRACE_CPUPROFILER_INIT(CmdLine) \
 	FCpuProfilerTrace::Init(CmdLine);
 
-#define TRACE_CPUPROFILER_EVENT_SCOPE_TEXT_GROUP(Name, Group) \
+#define TRACE_CPUPROFILER_EVENT_SCOPE_GROUP(Name, Group) \
 	static uint16 PREPROCESSOR_JOIN(__CpuProfilerEventSpecId, __LINE__); \
 	if (PREPROCESSOR_JOIN(__CpuProfilerEventSpecId, __LINE__) == 0) { \
-		PREPROCESSOR_JOIN(__CpuProfilerEventSpecId, __LINE__) = FCpuProfilerTrace::OutputEventType(Name, Group); \
+		PREPROCESSOR_JOIN(__CpuProfilerEventSpecId, __LINE__) = FCpuProfilerTrace::OutputEventType(#Name, Group); \
 	} \
 	FCpuProfilerTrace::FEventScope PREPROCESSOR_JOIN(__CpuProfilerEventScope, __LINE__)(PREPROCESSOR_JOIN(__CpuProfilerEventSpecId, __LINE__));
 
+#define TRACE_CPUPROFILER_EVENT_SCOPE(Name) \
+	TRACE_CPUPROFILER_EVENT_SCOPE_GROUP(Name, CpuProfilerGroup_Default)
+
+#define TRACE_CPUPROFILER_EVENT_SCOPE_TEXT_GROUP(Name, Group) \
+	FCpuProfilerTrace::FDynamicEventScope PREPROCESSOR_JOIN(__CpuProfilerEventScope, __LINE__)(Name, Group);
+	
 #define TRACE_CPUPROFILER_EVENT_SCOPE_TEXT(Name) \
 	TRACE_CPUPROFILER_EVENT_SCOPE_TEXT_GROUP(Name, CpuProfilerGroup_Default)
 
-#define TRACE_CPUPROFILER_EVENT_SCOPE_GROUP(Name, Group) \
-	TRACE_CPUPROFILER_EVENT_SCOPE_TEXT_GROUP(TEXT(#Name), Group);
-
-#define TRACE_CPUPROFILER_EVENT_SCOPE(Name) \
-	TRACE_CPUPROFILER_EVENT_SCOPE_TEXT_GROUP(TEXT(#Name), CpuProfilerGroup_Default);
-
 #else
 
 #define TRACE_CPUPROFILER_INIT(CmdLine)
diff --git a/Engine/Source/Runtime/CoreUObject/Private/UObject/LinkerLoad.cpp b/Engine/Source/Runtime/CoreUObject/Private/UObject/LinkerLoad.cpp
--- a/Engine/Source/Runtime/CoreUObject/Private/UObject/LinkerLoad.cpp
+++ b/Engine/Source/Runtime/CoreUObject/Private/UObject/LinkerLoad.cpp
@@ -1000,7 +1000,7 @@ FLinkerLoad::ELinkerStatus FLinkerLoad::CreateLoader(
 			{
 				INC_DWORD_STAT(STAT_TextAssetLinkerCount);
 				DECLARE_SCOPE_CYCLE_COUNTER(TEXT("FLinkerLoad::CreateTextArchiveFormatter"), STAT_LinkerLoad_CreateTextArchiveFormatter, STATGROUP_LinkerLoad);
-				TRACE_CPUPROFILER_EVENT_SCOPE_TEXT(TEXT("FLinkerLoad::CreateTextArchiveFormatter"));
+				TRACE_CPUPROFILER_EVENT_SCOPE(FLinkerLoad::CreateTextArchiveFormatter);
 				Loader = IFileManager::Get().CreateFileReader(*Filename);
 				StructuredArchiveFormatter = new FJsonArchiveInputFormatter(*this, [this](const FString& InFullPath)
 				{
diff --git a/Engine/Source/Runtime/Engine/Private/StaticMesh.cpp b/Engine/Source/Runtime/Engine/Private/StaticMesh.cpp
--- a/Engine/Source/Runtime/Engine/Private/StaticMesh.cpp
+++ b/Engine/Source/Runtime/Engine/Private/StaticMesh.cpp
@@ -2293,7 +2293,7 @@ void FStaticMeshRenderData::Cache(UStaticMesh* Owner, const FStaticMeshLODSettin
 
 
 	{
-		TRACE_CPUPROFILER_EVENT_SCOPE_TEXT(TEXT("StaticMesh_Cache"));
+		TRACE_CPUPROFILER_EVENT_SCOPE(StaticMesh_Cache);
 
 		COOK_STAT(auto Timer = StaticMeshCookStats::UsageStats.TimeSyncWork());
 		int32 T0 = FPlatformTime::Cycles();
