Fix FName("none") != FName("None") in dev and debug configs

#udn 527400

diff --git a/Engine/Source/Runtime/Core/Private/UObject/UnrealNames.cpp b/Engine/Source/Runtime/Core/Private/UObject/UnrealNames.cpp
--- a/Engine/Source/Runtime/Core/Private/UObject/UnrealNames.cpp
+++ b/Engine/Source/Runtime/Core/Private/UObject/UnrealNames.cpp
@@ -304,13 +304,13 @@ public:
 		return FNameEntryHandle(CurrentBlock, ByteOffset / Stride);
 	}
 
-	FNameEntryHandle Create(FNameStringView Name, FNameEntryId ComparisonId, FNameEntryHeader Header)
+	FNameEntryHandle Create(FNameStringView Name, TOptional<FNameEntryId> ComparisonId, FNameEntryHeader Header)
 	{
 		FNameEntryHandle Handle = Allocate(FNameEntry::GetDataOffset() + Name.BytesWithoutTerminator());
 		FNameEntry& Entry = Resolve(Handle);
 
 #if WITH_CASE_PRESERVING_NAME
-		Entry.ComparisonId = ComparisonId ? ComparisonId : FNameEntryId(Handle);
+		Entry.ComparisonId = ComparisonId.IsSet() ? ComparisonId.GetValue() : FNameEntryId(Handle);
 #endif
 
 		Entry.Header = Header;
@@ -533,7 +533,7 @@ struct FNameValue
 
 	FNameStringView Name;
 	FNameHash Hash;
-	FNameEntryId ComparisonId;
+	TOptional<FNameEntryId> ComparisonId;
 };
 
 using FNameComparisonValue = FNameValue<ENameCase::IgnoreCase>;
@@ -1970,6 +1970,8 @@ void FName::AutoTest()
 	check(NullName == FName(""));
 	check(NullName == FName(TEXT("")));
 	check(NullName == FName("None"));
+	check(NullName == FName("none"));
+	check(NullName == FName("NONE"));
 	check(NullName == FName(TEXT("None")));
 	check(FName().ToEName());
 	check(*FName().ToEName() == NAME_None);
