Added some unit tests for LexTryParseFromString

diff --git a/Engine/Source/Runtime/Core/Private/Containers/StringTest.cpp b/Engine/Source/Runtime/Core/Private/Containers/StringTest.cpp
--- a/Engine/Source/Runtime/Core/Private/Containers/StringTest.cpp
+++ b/Engine/Source/Runtime/Core/Private/Containers/StringTest.cpp
@@ -135,4 +135,107 @@ bool FStringUnicodeTest::RunTest(const FString& Parameters)
 	return true;
 }
 
+IMPLEMENT_SIMPLE_AUTOMATION_TEST(FLexTryParseStringTest, "System.Core.Misc.LexTryParseString", EAutomationTestFlags::ApplicationContextMask | EAutomationTestFlags::SmokeFilter)
+
+bool FLexTryParseStringTest::RunTest(const FString& Parameters)
+{
+	// Test that LexFromString can intepret all the numerical formats we expect it to
+	{
+		// Test float values
+
+		float Value;
+
+		// Basic numbers
+		TestTrue(TEXT("(float conversion from string) basic numbers"), LexTryParseString(Value, (TEXT("1"))) && Value == 1);
+		TestTrue(TEXT("(float conversion from string) basic numbers"), LexTryParseString(Value, (TEXT("1.0"))) && Value == 1);
+		TestTrue(TEXT("(float conversion from string) basic numbers"), LexTryParseString(Value, (TEXT(".5"))) && Value == 0.5);
+		TestTrue(TEXT("(float conversion from string) basic numbers"), LexTryParseString(Value, (TEXT("1."))) && Value == 1);
+
+		// Variations of 0
+		TestTrue(TEXT("(float conversion from string) variations of 0"), LexTryParseString(Value, (TEXT("0"))) && Value == 0);
+		TestTrue(TEXT("(float conversion from string) variations of 0"), LexTryParseString(Value, (TEXT("-0"))) && Value == 0);
+		TestTrue(TEXT("(float conversion from string) variations of 0"), LexTryParseString(Value, (TEXT("0.0"))) && Value == 0);
+		TestTrue(TEXT("(float conversion from string) variations of 0"), LexTryParseString(Value, (TEXT(".0"))) && Value == 0);
+		TestTrue(TEXT("(float conversion from string) variations of 0"), LexTryParseString(Value, (TEXT("0."))) && Value == 0);
+		TestTrue(TEXT("(float conversion from string) variations of 0"), LexTryParseString(Value, (TEXT("0. 111"))) && Value == 0);
+
+		// Scientific notation
+		TestTrue(TEXT("(float conversion from string) scientific notation"), LexTryParseString(Value, (TEXT("1.0e+10"))) && Value == 1.0e+10f);
+		TestTrue(TEXT("(float conversion from string) scientific notation"), LexTryParseString(Value, (TEXT("1.99999999e-11"))) && Value == 1.99999999e-11f);
+		TestTrue(TEXT("(float conversion from string) scientific notation"), LexTryParseString(Value, (TEXT("1e+10"))) && Value == 1e+10f);
+
+		// Non-finite special numbers
+		TestTrue(TEXT("(float conversion from string) inf"), LexTryParseString(Value, (TEXT("inf"))));
+		TestTrue(TEXT("(float conversion from string) nan"), LexTryParseString(Value, (TEXT("nan"))));
+		TestTrue(TEXT("(float conversion from string) nan(ind)"), LexTryParseString(Value, (TEXT("nan(ind)"))));
+
+		// nan/inf etc. are detected from the start of the string, regardless of any other characters that come afterwards
+		TestTrue(TEXT("(float conversion from string) nananananananana"), LexTryParseString(Value, (TEXT("nananananananana"))));
+		TestTrue(TEXT("(float conversion from string) nan(ind)!"), LexTryParseString(Value, (TEXT("nan(ind)!"))));
+		TestTrue(TEXT("(float conversion from string) infinity"), LexTryParseString(Value, (TEXT("infinity"))));
+
+		// Some numbers with whitespace
+		TestTrue(TEXT("(float conversion from string) whitespace"), LexTryParseString(Value, (TEXT("   2.5   "))) && Value == 2.5);
+		TestTrue(TEXT("(float conversion from string) whitespace"), LexTryParseString(Value, (TEXT("\t3.0\t"))) && Value == 3.0);
+		TestTrue(TEXT("(float conversion from string) whitespace"), LexTryParseString(Value, (TEXT("4.0   \t"))) && Value == 4.0);
+		TestTrue(TEXT("(float conversion from string) whitespace"), LexTryParseString(Value, (TEXT("\r\n5.25"))) && Value == 5.25);
+		TestTrue(TEXT("(float conversion from string) whitespace"), LexTryParseString(Value, (TEXT(" 6 . 2 "))) && Value == 6.0);
+		TestTrue(TEXT("(float conversion from string) whitespace"), LexTryParseString(Value, (TEXT(" 56 . 2 "))) && Value == 56.0);
+		TestTrue(TEXT("(float conversion from string) whitespace"), LexTryParseString(Value, (TEXT(" 5 6 . 2 "))) && Value == 5.0);
+
+		// Failure cases
+		TestFalse(TEXT("(float no conversion from string) not a number"), LexTryParseString(Value, (TEXT("not a number"))));
+		TestFalse(TEXT("(float no conversion from string) <empty string>"), LexTryParseString(Value, (TEXT(""))));
+		TestFalse(TEXT("(float conversion from string) ."), LexTryParseString(Value, (TEXT("."))));
+	}
+
+	{
+		// Test integer values
+
+		int32 Value;
+
+		// Basic numbers
+		TestTrue(TEXT("(int32 conversion from string) basic numbers"), LexTryParseString(Value, (TEXT("1"))) && Value == 1);
+		TestTrue(TEXT("(int32 conversion from string) basic numbers"), LexTryParseString(Value, (TEXT("1.0"))) && Value == 1);
+		TestTrue(TEXT("(int32 conversion from string) basic numbers"), LexTryParseString(Value, (TEXT("3.1"))) && Value == 3);
+		TestTrue(TEXT("(int32 conversion from string) basic numbers"), LexTryParseString(Value, (TEXT("0.5"))) && Value == 0);
+		TestTrue(TEXT("(int32 conversion from string) basic numbers"), LexTryParseString(Value, (TEXT("1."))) && Value == 1);
+
+		// Variations of 0
+		TestTrue(TEXT("(int32 conversion from string) variations of 0"), LexTryParseString(Value, (TEXT("0"))) && Value == 0);
+		TestTrue(TEXT("(int32 conversion from string) variations of 0"), LexTryParseString(Value, (TEXT("0.0"))) && Value == 0);
+		TestFalse(TEXT("(int32 conversion from string) variations of 0"), LexTryParseString(Value, (TEXT(".0"))) && Value == 0);
+		TestTrue(TEXT("(int32 conversion from string) variations of 0"), LexTryParseString(Value, (TEXT("0."))) && Value == 0);
+
+		// Scientific notation
+		TestTrue(TEXT("(int32 conversion from string) scientific notation"), LexTryParseString(Value, (TEXT("1.0e+10"))) && Value == 1);
+		TestTrue(TEXT("(int32 conversion from string) scientific notation"), LexTryParseString(Value, (TEXT("6.0e-10"))) && Value == 6);
+		TestTrue(TEXT("(int32 conversion from string) scientific notation"), LexTryParseString(Value, (TEXT("0.0e+10"))) && Value == 0);
+		TestTrue(TEXT("(int32 conversion from string) scientific notation"), LexTryParseString(Value, (TEXT("0.0e-10"))) && Value == 0);
+		TestTrue(TEXT("(int32 conversion from string) scientific notation"), LexTryParseString(Value, (TEXT("3e+10"))) && Value == 3);
+		TestTrue(TEXT("(int32 conversion from string) scientific notation"), LexTryParseString(Value, (TEXT("4e-10"))) && Value == 4);
+
+		// Some numbers with whitespace
+		TestTrue(TEXT("(int32 conversion from string) whitespace"), LexTryParseString(Value, (TEXT("   2.5   "))) && Value == 2);
+		TestTrue(TEXT("(int32 conversion from string) whitespace"), LexTryParseString(Value, (TEXT("\t3.0\t"))) && Value == 3);
+		TestTrue(TEXT("(int32 conversion from string) whitespace"), LexTryParseString(Value, (TEXT("4.0   \t"))) && Value == 4);
+		TestTrue(TEXT("(int32 conversion from string) whitespace"), LexTryParseString(Value, (TEXT("\r\n5.25"))) && Value == 5);
+		TestTrue(TEXT("(int32 conversion from string) whitespace"), LexTryParseString(Value, (TEXT(" 6 . 2 "))) && Value == 6);
+		TestTrue(TEXT("(int32 conversion from string) whitespace"), LexTryParseString(Value, (TEXT(" 56 . 2 "))) && Value == 56);
+		TestTrue(TEXT("(int32 conversion from string) whitespace"), LexTryParseString(Value, (TEXT(" 5 6 . 2 "))) && Value == 5);
+
+		// Non-finite special numbers. All shouldn't parse into an int
+		TestFalse(TEXT("(int32 no conversion from string) inf"), LexTryParseString(Value, (TEXT("inf"))));
+		TestFalse(TEXT("(int32 no conversion from string) nan"), LexTryParseString(Value, (TEXT("nan"))));
+		TestFalse(TEXT("(int32 no conversion from string) nan(ind)"), LexTryParseString(Value, (TEXT("nan(ind)"))));
+		TestFalse(TEXT("(int32 no conversion from string) nananananananana"), LexTryParseString(Value, (TEXT("nananananananana"))));
+		TestFalse(TEXT("(int32 no conversion from string) nan(ind)!"), LexTryParseString(Value, (TEXT("nan(ind)!"))));
+		TestFalse(TEXT("(int32 no conversion from string) infinity"), LexTryParseString(Value, (TEXT("infinity"))));
+		TestFalse(TEXT("(float no conversion from string) ."), LexTryParseString(Value, (TEXT("."))));
+		TestFalse(TEXT("(float no conversion from string) <empyty string>"), LexTryParseString(Value, (TEXT(""))));
+	}
+
+	return true;
+}
+
 #endif // WITH_DEV_AUTOMATION_TESTS
