UnrealInsights:
Removed cpu profiler EventSpecEx event again as it was not actually needed for backwards compatibility

diff --git a/Engine/Source/Developer/TraceServices/Private/Analyzers/CpuProfilerTraceAnalysis.cpp b/Engine/Source/Developer/TraceServices/Private/Analyzers/CpuProfilerTraceAnalysis.cpp
--- a/Engine/Source/Developer/TraceServices/Private/Analyzers/CpuProfilerTraceAnalysis.cpp
+++ b/Engine/Source/Developer/TraceServices/Private/Analyzers/CpuProfilerTraceAnalysis.cpp
@@ -16,7 +16,6 @@ void FCpuProfilerAnalyzer::OnAnalysisBegin(const FOnAnalysisContext& Context)
 	auto& Builder = Context.InterfaceBuilder;
 
 	Builder.RouteEvent(RouteId_EventSpec, "CpuProfiler", "EventSpec");
-	Builder.RouteEvent(RouteId_EventSpecEx, "CpuProfiler", "EventSpecEx");
 	Builder.RouteEvent(RouteId_EventBatch, "CpuProfiler", "EventBatch");
 	Builder.RouteEvent(RouteId_EndCapture, "CpuProfiler", "EndCapture");
 }
@@ -29,21 +28,14 @@ bool FCpuProfilerAnalyzer::OnEvent(uint16 RouteId, const FOnEventContext& Contex
 	switch (RouteId)
 	{
 	case RouteId_EventSpec:
-	{
-		uint32 Id = EventData.GetValue<uint16>("Id");
-		const TCHAR* ScopeName = Session.StoreString(reinterpret_cast<const TCHAR*>(EventData.GetAttachment()));
-		DefineScope(Id, ScopeName);
-		break;
-	}
-	case RouteId_EventSpecEx:
 	{
 		uint32 Id = EventData.GetValue<uint32>("Id");
 		uint8 CharSize = EventData.GetValue<uint8>("CharSize");
-		if (CharSize == 1)
+		if (CharSize == sizeof(ANSICHAR))
 		{
 			DefineScope(Id, Session.StoreString(StringCast<TCHAR>(reinterpret_cast<const ANSICHAR*>(EventData.GetAttachment())).Get()));
 		}
-		else if (CharSize == 2)
+		else if (CharSize == 0 || CharSize == sizeof(TCHAR)) // 0 for backwards compatibility
 		{
 			DefineScope(Id, Session.StoreString(reinterpret_cast<const TCHAR*>(EventData.GetAttachment())));
 		}
diff --git a/Engine/Source/Developer/TraceServices/Private/Analyzers/CpuProfilerTraceAnalysis.h b/Engine/Source/Developer/TraceServices/Private/Analyzers/CpuProfilerTraceAnalysis.h
--- a/Engine/Source/Developer/TraceServices/Private/Analyzers/CpuProfilerTraceAnalysis.h
+++ b/Engine/Source/Developer/TraceServices/Private/Analyzers/CpuProfilerTraceAnalysis.h
@@ -40,7 +40,6 @@ private:
 	enum : uint16
 	{
 		RouteId_EventSpec,
-		RouteId_EventSpecEx,
 		RouteId_EventBatch,
 		RouteId_EndCapture,
 	};
diff --git a/Engine/Source/Runtime/Core/Private/ProfilingDebugging/CpuProfilerTrace.cpp b/Engine/Source/Runtime/Core/Private/ProfilingDebugging/CpuProfilerTrace.cpp
--- a/Engine/Source/Runtime/Core/Private/ProfilingDebugging/CpuProfilerTrace.cpp
+++ b/Engine/Source/Runtime/Core/Private/ProfilingDebugging/CpuProfilerTrace.cpp
@@ -6,10 +6,12 @@
 #include "HAL/TlsAutoCleanup.h"
 #include "ProfilingDebugging/MiscTrace.h"
 #include "Misc/Parse.h"
+#include "Containers/Map.h"
+#include "Misc/MemStack.h"
 
 #if CPUPROFILERTRACE_ENABLED
 
-UE_TRACE_EVENT_BEGIN(CpuProfiler, EventSpecEx, Always)
+UE_TRACE_EVENT_BEGIN(CpuProfiler, EventSpec, Always)
 	UE_TRACE_EVENT_FIELD(uint32, Id)
 	UE_TRACE_EVENT_FIELD(uint16, Group)
 	UE_TRACE_EVENT_FIELD(uint8, CharSize)
@@ -230,11 +232,11 @@ uint32 FCpuProfilerTrace::OutputEventType(const TCHAR* Name, ECpuProfilerGroup G
 {
 	uint32 SpecId = FCpuProfilerTraceInternal::GetNextSpecId();
 	uint16 NameSize = (FCString::Strlen(Name) + 1) * sizeof(TCHAR);
-	UE_TRACE_LOG(CpuProfiler, EventSpecEx, NameSize)
-		<< EventSpecEx.Id(SpecId)
-		<< EventSpecEx.Group(uint16(Group))
-		<< EventSpecEx.CharSize(sizeof(TCHAR))
-		<< EventSpecEx.Attachment(Name, NameSize);
+	UE_TRACE_LOG(CpuProfiler, EventSpec, NameSize)
+		<< EventSpec.Id(SpecId)
+		<< EventSpec.Group(uint16(Group))
+		<< EventSpec.CharSize(sizeof(TCHAR))
+		<< EventSpec.Attachment(Name, NameSize);
 	return SpecId;
 }
 
@@ -242,11 +244,11 @@ uint32 FCpuProfilerTrace::OutputEventType(const ANSICHAR* Name, ECpuProfilerGrou
 {
 	uint32 SpecId = FCpuProfilerTraceInternal::GetNextSpecId();
 	uint16 NameSize = strlen(Name) + 1;
-	UE_TRACE_LOG(CpuProfiler, EventSpecEx, NameSize)
-		<< EventSpecEx.Id(SpecId)
-		<< EventSpecEx.Group(uint16(Group))
-		<< EventSpecEx.CharSize(1)
-		<< EventSpecEx.Attachment(Name, NameSize);
+	UE_TRACE_LOG(CpuProfiler, EventSpec, NameSize)
+		<< EventSpec.Id(SpecId)
+		<< EventSpec.Group(uint16(Group))
+		<< EventSpec.CharSize(1)
+		<< EventSpec.Attachment(Name, NameSize);
 	return SpecId;
 }
 
