UnrealInsights:
Fixed broken TRACE_COUNTER macros

diff --git a/Engine/Source/Runtime/Core/Public/ProfilingDebugging/CountersTrace.h b/Engine/Source/Runtime/Core/Public/ProfilingDebugging/CountersTrace.h
--- a/Engine/Source/Runtime/Core/Public/ProfilingDebugging/CountersTrace.h
+++ b/Engine/Source/Runtime/Core/Public/ProfilingDebugging/CountersTrace.h
@@ -48,8 +48,11 @@ struct FCountersTrace
 
 		void Init(const TCHAR* InCounterName, ETraceCounterType InCounterType, ETraceCounterDisplayHint InCounterDisplayHint)
 		{
-			CounterType = InCounterType;
-			CounterId = OutputInitCounter(InCounterName, InCounterType, InCounterDisplayHint);
+			if (CounterId == 0)
+			{
+				CounterType = InCounterType;
+				CounterId = OutputInitCounter(InCounterName, InCounterType, InCounterDisplayHint);
+			}
 		}
 
 		void Set(int64 Value)
@@ -149,22 +152,19 @@ struct FCountersTrace
 
 #define __TRACE_DECLARE_INLINE_COUNTER(CounterName, CounterType, CounterDisplayHint) \
 	static FCountersTrace::FCounter PREPROCESSOR_JOIN(__TraceCounter, __LINE__); \
-	if (!PREPROCESSOR_JOIN(__TraceCounter, __LINE__).CounterId) \
-	{ \
-		PREPROCESSOR_JOIN(__TraceCounter, __LINE__).Init(TEXT(#CounterName), CounterType, CounterDisplayHint); \
-	}
+	PREPROCESSOR_JOIN(__TraceCounter, __LINE__).Init(TEXT(#CounterName), CounterType, CounterDisplayHint); \
 
 #define TRACE_INT_VALUE(CounterName, Value) \
 	__TRACE_DECLARE_INLINE_COUNTER(CounterName, TraceCounterType_Int, TraceCounterDisplayHint_None) \
-	PREPROCESSOR_JOIN(__TraceCounter, __LINE__).Set(Value);
+	PREPROCESSOR_JOIN(__TraceCounter, __LINE__).Set(int64(Value));
 
 #define TRACE_FLOAT_VALUE(CounterName, Value) \
 	__TRACE_DECLARE_INLINE_COUNTER(CounterName, TraceCounterType_Float, TraceCounterDisplayHint_None) \
-	PREPROCESSOR_JOIN(__TraceCounter, __LINE__).Set(Value);
+	PREPROCESSOR_JOIN(__TraceCounter, __LINE__).Set(double(Value));
 
 #define TRACE_MEMORY_VALUE(CounterName, Value) \
 	__TRACE_DECLARE_INLINE_COUNTER(CounterName, TraceCounterType_Int, TraceCounterDisplayHint_Memory) \
-	PREPROCESSOR_JOIN(__TraceCounter, __LINE__).Set(Value);
+	PREPROCESSOR_JOIN(__TraceCounter, __LINE__).Set(int64(Value));
 
 #define TRACE_DECLARE_INT_COUNTER(CounterName) \
 	__TRACE_DECLARE_COUNTER(CounterName, TraceCounterType_Int, TraceCounterDisplayHint_None)
