Added explicit shutdown of CpuProfilerTrace to ensure that game thread state memory is freed before exiting

diff --git a/Engine/Source/Runtime/Core/Private/ProfilingDebugging/CpuProfilerTrace.cpp b/Engine/Source/Runtime/Core/Private/ProfilingDebugging/CpuProfilerTrace.cpp
--- a/Engine/Source/Runtime/Core/Private/ProfilingDebugging/CpuProfilerTrace.cpp
+++ b/Engine/Source/Runtime/Core/Private/ProfilingDebugging/CpuProfilerTrace.cpp
@@ -261,4 +261,14 @@ void FCpuProfilerTrace::Init(const TCHAR* CmdLine)
 	}
 }
 
+void FCpuProfilerTrace::Shutdown()
+{
+	Trace::ToggleEvent(TEXT("CpuProfiler.EventBatch"), false);
+	if (FCpuProfilerTraceInternal::ThreadBuffer)
+	{
+		delete FCpuProfilerTraceInternal::ThreadBuffer;
+		FCpuProfilerTraceInternal::ThreadBuffer = nullptr;
+	}
+}
+
 #endif
\ No newline at end of file
diff --git a/Engine/Source/Runtime/Core/Public/ProfilingDebugging/CpuProfilerTrace.h b/Engine/Source/Runtime/Core/Public/ProfilingDebugging/CpuProfilerTrace.h
--- a/Engine/Source/Runtime/Core/Public/ProfilingDebugging/CpuProfilerTrace.h
+++ b/Engine/Source/Runtime/Core/Public/ProfilingDebugging/CpuProfilerTrace.h
@@ -23,6 +23,7 @@ enum ECpuProfilerGroup
 struct FCpuProfilerTrace
 {
 	CORE_API static void Init(const TCHAR* CmdLine);
+	CORE_API static void Shutdown();
 	FORCENOINLINE CORE_API static uint32 OutputEventType(const ANSICHAR* Name, ECpuProfilerGroup Group);
 	FORCENOINLINE CORE_API static uint32 OutputEventType(const TCHAR* Name, ECpuProfilerGroup Group);
 	CORE_API static void OutputBeginEvent(uint32 SpecId);
@@ -65,6 +66,9 @@ struct FCpuProfilerTrace
 #define TRACE_CPUPROFILER_INIT(CmdLine) \
 	FCpuProfilerTrace::Init(CmdLine);
 
+#define TRACE_CPUPROFILER_SHUTDOWN() \
+	FCpuProfilerTrace::Shutdown();
+
 #define TRACE_CPUPROFILER_EVENT_SCOPE_GROUP(Name, Group) \
 	static uint16 PREPROCESSOR_JOIN(__CpuProfilerEventSpecId, __LINE__); \
 	if (PREPROCESSOR_JOIN(__CpuProfilerEventSpecId, __LINE__) == 0) { \
@@ -84,6 +88,7 @@ struct FCpuProfilerTrace
 #else
 
 #define TRACE_CPUPROFILER_INIT(CmdLine)
+#define TRACE_CPUPROFILER_SHUTDOWN()
 #define TRACE_CPUPROFILER_EVENT_SCOPE_GROUP(Name, Group)
 #define TRACE_CPUPROFILER_EVENT_SCOPE(Name)
 #define TRACE_CPUPROFILER_EVENT_SCOPE_TEXT_GROUP(Name, Group)
diff --git a/Engine/Source/Runtime/Launch/Private/LaunchEngineLoop.cpp b/Engine/Source/Runtime/Launch/Private/LaunchEngineLoop.cpp
--- a/Engine/Source/Runtime/Launch/Private/LaunchEngineLoop.cpp
+++ b/Engine/Source/Runtime/Launch/Private/LaunchEngineLoop.cpp
@@ -3821,6 +3821,8 @@ void FEngineLoop::Exit()
 	IStreamingManager::Shutdown();
 
 	FPlatformMisc::ShutdownTaggedStorage();
+
+	TRACE_CPUPROFILER_SHUTDOWN();
 }
 
 
