Fix for a comment which mentions code which no longer compiles.

diff --git a/Engine/Source/Runtime/Core/Public/Containers/ArrayView.h b/Engine/Source/Runtime/Core/Public/Containers/ArrayView.h
--- a/Engine/Source/Runtime/Core/Public/Containers/ArrayView.h
+++ b/Engine/Source/Runtime/Core/Public/Containers/ArrayView.h
@@ -55,16 +55,31 @@ namespace ArrayViewPrivate
  * }
  * 
  * could be called as:
- *     SumAll(MyTArray);
+ *     SumAll(MyTArray);\
  *     SumAll(MyCArray);
- *     SumAll({1, 2, 3});
  *     SumAll(MakeArrayView(Ptr, Num));
- * 
+ *
+ *     auto Values = { 1, 2, 3 };
+ *     SumAll(Values);
+ *
  * Note:
  *   View classes are not const-propagating! If you want a view where the elements are const, you need "TArrayView<const T>" not "const TArrayView<T>"!
+ *   TArrayViews cannot be implicitly constructed from an rvalue initializer lists:
+ *
+ * , as a compile-time guard against :
+ *
+ *   SumAll({ 1, 2, 3 }); // error
  *
  * Caution:
  *   Treat a view like a *reference* to the elements in the array. DO NOT free or reallocate the array while the view exists!
+ *   For this reason, TArrayViews are not constructible from rvalue initializer lists:
+ *
+ *   SumAll({ 1, 2, 3 }); // error
+ *
+ *   This is to avoid this issue:
+ *
+ *   TArrayView<int> View = { 1, 2, 3 }; // construction of array view from rvalue initializer list
+ *   int n = View[0]; // undefined behavior, as the initializer list was destroyed at the end of the previous line
  */
 template<typename InElementType>
 class TArrayView
