Enabled log trace events on Switch.

#rnx

diff --git a/Engine/Source/Runtime/Core/Public/Logging/LogTrace.h b/Engine/Source/Runtime/Core/Public/Logging/LogTrace.h
--- a/Engine/Source/Runtime/Core/Public/Logging/LogTrace.h
+++ b/Engine/Source/Runtime/Core/Public/Logging/LogTrace.h
@@ -4,7 +4,7 @@
 
 #include "CoreTypes.h"
 
-#if !IS_PROGRAM && !UE_BUILD_SHIPPING && (PLATFORM_WINDOWS || PLATFORM_UNIX || PLATFORM_MAC || PLATFORM_PS4 || PLATFORM_XBOXONE || PLATFORM_ANDROID)
+#if !IS_PROGRAM && !UE_BUILD_SHIPPING && (PLATFORM_WINDOWS || PLATFORM_UNIX || PLATFORM_MAC || PLATFORM_PS4 || PLATFORM_XBOXONE || PLATFORM_ANDROID || PLATFORM_SWITCH)
 #define LOGTRACE_ENABLED 1
 #else
 #define LOGTRACE_ENABLED 0
