Compile out experimental memory arena code

diff --git a/Engine/Source/Runtime/Core/Private/Memory/MemoryArena.cpp b/Engine/Source/Runtime/Core/Private/Memory/MemoryArena.cpp
--- a/Engine/Source/Runtime/Core/Private/Memory/MemoryArena.cpp
+++ b/Engine/Source/Runtime/Core/Private/Memory/MemoryArena.cpp
@@ -7,6 +7,7 @@
 #include "Misc/ScopeRWLock.h"
 #include <atomic>
 
+#if UE_WITH_HEAPARENA
 // dlmalloc for use in FHeapArena
 
 extern "C" {
@@ -36,6 +37,7 @@ extern "C" {
 #	pragma warning(pop)
 #endif
 }
+#endif
 
 //////////////////////////////////////////////////////////////////////////
 //
@@ -217,6 +219,8 @@ FArenaPointer ArenaRealloc(FArenaPointer InPtr, SIZE_T OldSize, SIZE_T NewSize,
 
 //////////////////////////////////////////////////////////////////////////
 
+#if UE_WITH_HEAPARENA
+
 FHeapArena::FHeapArena()
 {
 	HeapHandle = create_mspace(1024 * 1024, /* locked */ 1);
@@ -247,6 +251,8 @@ const TCHAR* FHeapArena::InternalGetDebugName() const
 	return TEXT("HeapArena");
 }
 
+#endif
+
 //////////////////////////////////////////////////////////////////////////
 
 FMallocArena::FMallocArena()
@@ -307,6 +313,8 @@ const TCHAR* FAnsiArena::InternalGetDebugName() const
 }
 
 
+#if UE_WITH_ARENAMAP
+
 //////////////////////////////////////////////////////////////////////////
 //
 // Arena map - arenas need to register with the arena map in order to
@@ -463,4 +471,6 @@ FMemoryArena* FArenaMap::MapPtrToArena(const void* VaBase)
 	return nullptr;
 }
 
+#endif
+
 //////////////////////////////////////////////////////////////////////////
diff --git a/Engine/Source/Runtime/Core/Public/Memory/MemoryArena.h b/Engine/Source/Runtime/Core/Public/Memory/MemoryArena.h
--- a/Engine/Source/Runtime/Core/Public/Memory/MemoryArena.h
+++ b/Engine/Source/Runtime/Core/Public/Memory/MemoryArena.h
@@ -20,6 +20,9 @@
 #	endif
 #endif
 
+#define UE_WITH_HEAPARENA	0
+#define UE_WITH_ARENAMAP	0
+
 class FMemoryArena;
 class FHeapArena;
 
@@ -143,6 +146,7 @@ CORE_API FArenaPointer ArenaRealloc(FMemoryArena* Arena, void* InPtr, SIZE_T Old
 	it
   */
 
+#if UE_WITH_HEAPARENA
 class FHeapArena : public FMemoryArena
 {
 public:
@@ -157,6 +161,7 @@ private:
 
 	void* HeapHandle = nullptr;
 };
+#endif
 
 /** Default heap allocator
 
@@ -198,6 +203,8 @@ private:
 
 	Provides an efficient mechanism for mapping pointers to arenas
   */
+
+#if UE_WITH_ARENAMAP
 class FArenaMap
 {
 public:
@@ -216,3 +223,4 @@ public:
 	// Map a virtual address to a memory arena
 	CORE_API static FMemoryArena* MapPtrToArena(const void* VaBase);
 };
+#endif
