Made cpp-internal string view classes not violate single definition rule wrt FAnsiStringView

diff --git a/Engine/Source/Runtime/Core/Private/UObject/UnrealNames.cpp b/Engine/Source/Runtime/Core/Private/UObject/UnrealNames.cpp
--- a/Engine/Source/Runtime/Core/Private/UObject/UnrealNames.cpp
+++ b/Engine/Source/Runtime/Core/Private/UObject/UnrealNames.cpp
@@ -1438,7 +1438,7 @@ static bool StringAndNumberEqualsString(const CharType1* Name, uint32 NameLen, i
 	return Str[NameLen] == '_' && NumberEqualsString(Number, Str + NameLen + 1);
 }
 
-struct FAnsiStringView
+struct FNameAnsiStringView
 {
 	using CharType = ANSICHAR;
 
@@ -1455,12 +1455,12 @@ struct FWideStringViewWithWidth
 	bool bIsWide;
 };
 
-static FAnsiStringView MakeUnconvertedView(const ANSICHAR* Str, int32 Len)
+static FNameAnsiStringView MakeUnconvertedView(const ANSICHAR* Str, int32 Len)
 {
 	return { Str, Len };
 }
 
-static FAnsiStringView MakeUnconvertedView(const ANSICHAR* Str)
+static FNameAnsiStringView MakeUnconvertedView(const ANSICHAR* Str)
 {
 	return { Str, Str ? FCStringAnsi::Strlen(Str) : 0 };
 }
@@ -1549,7 +1549,7 @@ struct FNameHelper
 		return MakeWithNumber(View, FindType, NAME_NO_NUMBER_INTERNAL);
 	}
 
-	static FName MakeWithNumber(FAnsiStringView View, EFindName FindType, int32 InternalNumber)
+	static FName MakeWithNumber(FNameAnsiStringView	 View, EFindName FindType, int32 InternalNumber)
 	{
 		// Ignore the supplied number if the name string is empty
 		// to keep the semantics of the old FName implementation
