Refactoring of FStandardPlatformString::GetVarArgs to make it pass its result to a lambda instead.

diff --git a/Engine/Source/Runtime/Core/Private/GenericPlatform/StandardPlatformString.cpp b/Engine/Source/Runtime/Core/Private/GenericPlatform/StandardPlatformString.cpp
--- a/Engine/Source/Runtime/Core/Private/GenericPlatform/StandardPlatformString.cpp
+++ b/Engine/Source/Runtime/Core/Private/GenericPlatform/StandardPlatformString.cpp
@@ -163,8 +163,11 @@ static int32 FormatString(const FFormatInfo& Info, VA_LIST_REF ArgPtr, WIDECHAR*
 	}
 }
 
-static const WIDECHAR* GetFormattedArgument(const FFormatInfo& Info, VA_LIST_REF ArgPtr, WIDECHAR* Formatted, int32 &InOutLength)
+template <typename CallableType>
+static void FormatArgument(const FFormatInfo& Info, VA_LIST_REF ArgPtr, CallableType Callable)
 {
+	WIDECHAR Formatted[1024];
+
 	if (FChar::ToLower(Info.Type) == LITERAL(WIDECHAR, 's'))
 	{
 		if (Info.HasDynamicWidth)
@@ -173,20 +176,20 @@ static const WIDECHAR* GetFormattedArgument(const FFormatInfo& Info, VA_LIST_REF
 			const WIDECHAR* String = va_arg(ArgPtr, WIDECHAR*);
 			if (String)
 			{
-				InOutLength = swprintf(Formatted, InOutLength, Info.Format, Width, String);
-				return Formatted;
+				int32 Length = swprintf(Formatted, UE_ARRAY_COUNT(Formatted), Info.Format, Width, String);
+				Callable(Formatted, Length);
 			}
 			else
 			{
-				return TEXT("(null)");
+				Callable(TEXT("(null)"), UE_ARRAY_COUNT(Formatted));
 			}
 		}
 		// Is it a plain string?
 		else if (FChar::ToLower(Info.Format[1]) == LITERAL(WIDECHAR, 's'))
 		{
 			const WIDECHAR* String = va_arg(ArgPtr, WIDECHAR*);
-			InOutLength = FCString::Strlen(String);
-			return String ? String : TEXT("(null)");
+			int32 Length = FCString::Strlen(String);
+			Callable(String ? String : TEXT("(null)"), Length);
 		}
 		// Some form of string requiring formatting, such as a left- or right-justified string
 		else
@@ -194,9 +197,10 @@ static const WIDECHAR* GetFormattedArgument(const FFormatInfo& Info, VA_LIST_REF
 			// We call swprintf directly which may expect %S for a widechar string. This will be fixed up
 			// by the time we get here (See above in GetFormattingInfo).
 			const WIDECHAR* String = va_arg(ArgPtr, WIDECHAR*);
-			InOutLength = swprintf(Formatted, InOutLength, Info.Format, String);
-			return Formatted;
+			int32 Length = swprintf(Formatted, UE_ARRAY_COUNT(Formatted), Info.Format, String);
+			Callable(Formatted, Length);
 		}
+		return;
 	}
 
 	if (FChar::ToLower(Info.Type) == LITERAL(WIDECHAR, 'c'))
@@ -204,32 +208,33 @@ static const WIDECHAR* GetFormattedArgument(const FFormatInfo& Info, VA_LIST_REF
 		const WIDECHAR Char = (WIDECHAR)va_arg(ArgPtr, int32);
 		Formatted[0] = Char;
 		Formatted[1] = 0;
-		InOutLength = 1;
-		return Formatted;
+		Callable(Formatted, 1);
+		return;
 	}
 
+	int32 Length = 0;
 	if (FChar::ToLower(Info.Type) == LITERAL(WIDECHAR, 'a') || FChar::ToLower(Info.Type) == LITERAL(WIDECHAR, 'e')
 		|| FChar::ToLower(Info.Type) == LITERAL(WIDECHAR, 'f') || FChar::ToLower(Info.Type) == LITERAL(WIDECHAR, 'g'))
 	{
-		InOutLength = FormatString<long double, double>(Info, ArgPtr, Formatted, InOutLength);
+		Length = FormatString<long double, double>(Info, ArgPtr, Formatted, UE_ARRAY_COUNT(Formatted));
 	}
 	else if (Info.Type == LITERAL(WIDECHAR, 'p'))
 	{
 		void* Value = va_arg(ArgPtr, void*);
-		InOutLength = swprintf(Formatted, InOutLength, Info.Format, Value);
+		Length = swprintf(Formatted, UE_ARRAY_COUNT(Formatted), Info.Format, Value);
 	}
 	else if (FChar::ToLower(Info.Type) == LITERAL(WIDECHAR, 'd') || FChar::ToLower(Info.Type) == LITERAL(WIDECHAR, 'i'))
 	{
-		InOutLength = FormatString<int64, int32>(Info, ArgPtr, Formatted, InOutLength);
+		Length = FormatString<int64, int32>(Info, ArgPtr, Formatted, UE_ARRAY_COUNT(Formatted));
 	}
 	else if (FChar::ToLower(Info.Type) == LITERAL(WIDECHAR, 'o') || FChar::ToLower(Info.Type) == LITERAL(WIDECHAR, 'u') || FChar::ToLower(Info.Type) == LITERAL(WIDECHAR, 'x'))
 	{
-		InOutLength = FormatString<uint64, uint32>(Info, ArgPtr, Formatted, InOutLength);
+		Length = FormatString<uint64, uint32>(Info, ArgPtr, Formatted, UE_ARRAY_COUNT(Formatted));
 	}
 
-	check(InOutLength != -1);
+	check(Length != -1);
 
-	return Formatted;
+	Callable(Formatted, Length);
 }
 
 int32 FStandardPlatformString::GetVarArgs( WIDECHAR* Dest, SIZE_T DestSize, const WIDECHAR*& Fmt, va_list ArgPtr )
@@ -263,24 +268,26 @@ int32 FStandardPlatformString::GetVarArgs( WIDECHAR* Dest, SIZE_T DestSize, cons
 			FFormatInfo Info;
 			Format += GetFormattingInfo(Format, Info);
 
-			WIDECHAR Formatted[1024];
-			int32 Length = UE_ARRAY_COUNT(Formatted);
-			const WIDECHAR* FormattedArg = GetFormattedArgument(Info, ArgPtr, Formatted, Length);
-			if (FormattedArg && Length > 0)
+			bool bLengthExceeded = false;
+			FormatArgument(Info, ArgPtr, [&Dest, &DestSize, &bLengthExceeded](const TCHAR* FormattedArg, int32 Length)
 			{
-				if (Length < DestSize)
+				if (FormattedArg && Length > 0)
 				{
+					if (Length >= DestSize)
+					{
+						Length = DestSize;
+						bLengthExceeded = true;
+					}
+
 					FMemory::Memcpy(Dest, FormattedArg, Length * sizeof(WIDECHAR));
 					Dest += Length;
 					DestSize -= Length;
 				}
-				else
-				{
-					FMemory::Memcpy(Dest, FormattedArg, DestSize * sizeof(WIDECHAR));
-					Dest += DestSize;
-					*Dest = 0;
-					return -1;
-				}
+			});
+			if (bLengthExceeded)
+			{
+				*Dest = 0;
+				return -1;
 			}
 		}
 		else
