Fix for -1 being wrongly returned when the formatted string matches the buffer length exactly.

diff --git a/Engine/Source/Runtime/Core/Private/GenericPlatform/StandardPlatformString.cpp b/Engine/Source/Runtime/Core/Private/GenericPlatform/StandardPlatformString.cpp
--- a/Engine/Source/Runtime/Core/Private/GenericPlatform/StandardPlatformString.cpp
+++ b/Engine/Source/Runtime/Core/Private/GenericPlatform/StandardPlatformString.cpp
@@ -273,7 +273,7 @@ int32 FStandardPlatformString::GetVarArgs( WIDECHAR* Dest, SIZE_T DestSize, cons
 			{
 				if (FormattedArg && Length > 0)
 				{
-					if (Length >= DestSize)
+					if (Length > DestSize)
 					{
 						Length = DestSize;
 						bLengthExceeded = true;
