Fixed up some TCHAR <-> WIDECHAR and SIZE_T confusion

diff --git a/Engine/Source/Runtime/Core/Private/Misc/StringBuilder.cpp b/Engine/Source/Runtime/Core/Private/Misc/StringBuilder.cpp
--- a/Engine/Source/Runtime/Core/Private/Misc/StringBuilder.cpp
+++ b/Engine/Source/Runtime/Core/Private/Misc/StringBuilder.cpp
@@ -28,12 +28,12 @@ TStringBuilderImpl<C>::~TStringBuilderImpl()
 
 template<typename C>
 void
-TStringBuilderImpl<C>::Extend(SIZE_T ExtraCapacity)
+TStringBuilderImpl<C>::Extend(int32 ExtraCapacity)
 {
 	check(bIsExtendable);
 
-	const size_t OldCapacity = End - Base;
-	const size_t NewCapacity = NextPowerOfTwo(OldCapacity + ExtraCapacity);
+	const SIZE_T OldCapacity = End - Base;
+	const SIZE_T NewCapacity = NextPowerOfTwo(OldCapacity + ExtraCapacity);
 
 	C* NewBase = (C*)AllocBuffer(NewCapacity);
 
@@ -53,14 +53,14 @@ TStringBuilderImpl<C>::Extend(SIZE_T ExtraCapacity)
 
 template<typename C>
 void*
-TStringBuilderImpl<C>::AllocBuffer(SIZE_T ByteCount)
+TStringBuilderImpl<C>::AllocBuffer(int32 ByteCount)
 {
 	return FMemory::Malloc(ByteCount * sizeof(C));
 }
 
 template<typename C>
 void
-TStringBuilderImpl<C>::FreeBuffer(void* Buffer, SIZE_T ByteCount)
+TStringBuilderImpl<C>::FreeBuffer(void* Buffer, int32 ByteCount)
 {
 	FMemory::Free(Buffer);
 }
@@ -68,7 +68,7 @@ TStringBuilderImpl<C>::FreeBuffer(void* Buffer, SIZE_T ByteCount)
 // Instantiate templates once
 
 template class TStringBuilderImpl<ANSICHAR>;
-template class TStringBuilderImpl<WIDECHAR>;
+template class TStringBuilderImpl<TCHAR>;
 
 #if 0
 #if !(UE_BUILD_SHIPPING || UE_BUILD_TEST)
diff --git a/Engine/Source/Runtime/Core/Public/Misc/StringBuilder.h b/Engine/Source/Runtime/Core/Public/Misc/StringBuilder.h
--- a/Engine/Source/Runtime/Core/Public/Misc/StringBuilder.h
+++ b/Engine/Source/Runtime/Core/Public/Misc/StringBuilder.h
@@ -17,18 +17,24 @@
 	the stack as a function local variable to avoid heap allocations.
 
 	The buffer is always contiguous and the class is not intended to be
-	used to construct extremely large strings.
+	used to construct extremely large strings. 
+	
+	This is not intended to be used as a mechanism for holding on to 
+	strings for a long time - the use case is explicitly to aid in 
+	*constructing* strings on the stack	and subsequently passing the 
+	string into a function call or a more permanent string storage 
+	mechanism like FString et al
 
 	The amount of buffer space to allocate is specified via a template
 	parameter and if the constructed string should overflow this initial
 	buffer, a new buffer will be allocated using regular dynamic memory
 	allocations. For instances where you absolutely must not allocate
 	any memory, you should use the fixed variants named
-	TFixed...StringBuilder
+	TFixed*StringBuilder
 
 	Overflow allocation should be the exceptional case however -- always
-	try to size the buffer so that it can hold the vast majority strings
-	you expect to construct.
+	try to size the buffer so that it can hold the vast majority of 
+	strings you expect to construct.
 
 	Be mindful that stack is a limited resource, so if you are writing a
 	highly recursive function you may want to use some other mechanism
@@ -40,19 +46,20 @@ template<typename C>
 class TStringBuilderImpl
 {
 public:
-	TStringBuilderImpl() = default;
-	CORE_API ~TStringBuilderImpl();
+				TStringBuilderImpl() = default;
+	CORE_API	~TStringBuilderImpl();
+
+				TStringBuilderImpl(const TStringBuilderImpl&) = delete;
+				TStringBuilderImpl(TStringBuilderImpl&&) = delete;
 
-	TStringBuilderImpl(const TStringBuilderImpl&) = delete;
-	TStringBuilderImpl(TStringBuilderImpl&&) = delete;
 	const TStringBuilderImpl& operator=(const TStringBuilderImpl&) = delete;
 	const TStringBuilderImpl& operator=(TStringBuilderImpl&&) = delete;
 
-	inline SIZE_T	Len() const			{ return CurPos - Base; }
+	inline int32	Len() const			{ return int32(CurPos - Base); }
 	inline const C* ToString() const	{ EnsureNulTerminated(); return Base; }
 	inline const C* operator*() const	{ EnsureNulTerminated(); return Base; }
 
-	inline const C	LastChar() const			{ return *(CurPos - 1); }
+	inline const C	LastChar() const	{ return *(CurPos - 1); }
 
 	inline TStringBuilderImpl& Append(C Char)
 	{
@@ -69,14 +76,14 @@ public:
 	inline TStringBuilderImpl& AppendAnsi(const ANSICHAR* NulTerminatedString)
 #endif
 	{
-		const SIZE_T Length = TCString<ANSICHAR>::Strlen(NulTerminatedString);
+		const int32 Length = TCString<ANSICHAR>::Strlen(NulTerminatedString);
 
 		EnsureCapacity(Length);
 
 		C* RESTRICT Dest = CurPos;
 		CurPos += Length;
 
-		for (SIZE_T i = 0; i < Length; ++i)
+		for (int32 i = 0; i < Length; ++i)
 		{
 			Dest[i] = NulTerminatedString[i];
 		}
@@ -89,14 +96,14 @@ public:
 		return AppendAnsi(AnsiString.Data(), AnsiString.Len());
 	}
 
-	inline TStringBuilderImpl& AppendAnsi(const ANSICHAR* NulTerminatedString, const SIZE_T Length)
+	inline TStringBuilderImpl& AppendAnsi(const ANSICHAR* NulTerminatedString, const int32 Length)
 	{
 		EnsureCapacity(Length);
 
 		C* RESTRICT Dest = CurPos;
 		CurPos += Length;
 
-		for (SIZE_T i = 0; i < Length; ++i)
+		for (int32 i = 0; i < Length; ++i)
 		{
 			Dest[i] = NulTerminatedString[i];
 		}
@@ -110,7 +117,7 @@ public:
 	inline TStringBuilderImpl& Append(const C* NulTerminatedString)
 #endif
 	{
-		const SIZE_T Length = TCString<C>::Strlen(NulTerminatedString);
+		const int32 Length = TCString<C>::Strlen(NulTerminatedString);
 
 		EnsureCapacity(Length);
 		C* RESTRICT Dest = CurPos;
@@ -121,10 +128,10 @@ public:
 		return *this;
 	}
 
-	inline TStringBuilderImpl& Append(const C* NulTerminatedString, SIZE_T MaxChars)
+	inline TStringBuilderImpl& Append(const C* NulTerminatedString, int32 MaxChars)
 	{
-		const SIZE_T StringLength = TCString<C>::Strlen(NulTerminatedString);
-		const SIZE_T Length = FPlatformMath::Min<SIZE_T>(MaxChars, StringLength);
+		const int32 StringLength = TCString<C>::Strlen(NulTerminatedString);
+		const int32 Length = FPlatformMath::Min<int32>(MaxChars, StringLength);
 
 		EnsureCapacity(Length);
 		C* RESTRICT Dest = CurPos;
@@ -136,16 +143,16 @@ public:
 	}
 
 #if USE_STRING_LITERAL_PATH
-	template<SIZE_T ArrayLength>
+	template<int32 ArrayLength>
 	inline TStringBuilderImpl& Append(C(&CharArray)[ArrayLength])
 	{
 		return Append(&CharArray[0]);
 	}
 
-	template<SIZE_T LiteralLength>
+	template<int32 LiteralLength>
 	inline TStringBuilderImpl& Append(const C (&StringLiteral)[LiteralLength])
 	{
-		constexpr SIZE_T Length = LiteralLength - 1;
+		constexpr int32 Length = LiteralLength - 1;
 		EnsureCapacity(Length);
 		C* RESTRICT Dest = CurPos;
 		CurPos += Length;
@@ -155,16 +162,16 @@ public:
 		return *this;
 	}
 
-	template<SIZE_T LiteralLength>
+	template<int32 LiteralLength>
 	inline TStringBuilderImpl& AppendAnsi(const ANSICHAR(&StringLiteral)[LiteralLength])
 	{
-		constexpr SIZE_T Length = LiteralLength - 1;
+		constexpr int32 Length = LiteralLength - 1;
 		EnsureCapacity(Length);
 
 		C* RESTRICT Dest = CurPos;
 		CurPos += Length;
 
-		for (SIZE_T i = 0; i < Length; ++i)
+		for (int32 i = 0; i < Length; ++i)
 		{
 			Dest[i] = StringLiteral[i];
 		}
@@ -174,7 +181,7 @@ public:
 #endif
 
 protected:
-	inline void Initialize(C* InBase, SIZE_T InCapacity)
+	inline void Initialize(C* InBase, int32 InCapacity)
 	{
 		Base	= InBase;
 		CurPos	= InBase;
@@ -189,7 +196,7 @@ protected:
 		}
 	}
 
-	inline void EnsureCapacity(SIZE_T RequiredAdditionalCapacity)
+	inline void EnsureCapacity(int32 RequiredAdditionalCapacity)
 	{
 		// precondition: we know the current buffer has enough capacity
 		// for the existing string including NUL terminator
@@ -200,9 +207,9 @@ protected:
 		Extend(RequiredAdditionalCapacity);
 	}
 
-	CORE_API void	Extend(SIZE_T extraCapacity);
-	CORE_API void*	AllocBuffer(SIZE_T byteCount);
-	CORE_API void	FreeBuffer(void* buffer, SIZE_T byteCount);
+	CORE_API void	Extend(int32 extraCapacity);
+	CORE_API void*	AllocBuffer(int32 byteCount);
+	CORE_API void	FreeBuffer(void* buffer, int32 byteCount);
 
 	C*			Base;
 	C* 			CurPos;
@@ -221,7 +228,7 @@ public:
 	operator FAnsiStringView() const { return FAnsiStringView(Base, CurPos - Base); }
 
 protected:
-	inline FAnsiStringBuilderBase(ANSICHAR* BufferPointer, SIZE_T BufferCapacity)
+	inline FAnsiStringBuilderBase(ANSICHAR* BufferPointer, int32 BufferCapacity)
 	{
 		Initialize(BufferPointer, BufferCapacity);
 	}
@@ -229,7 +236,7 @@ protected:
 	~FAnsiStringBuilderBase() = default;
 };
 
-template<SIZE_T N>
+template<int32 N>
 class TFixedAnsiStringBuilder : public FAnsiStringBuilderBase
 {
 public:
@@ -244,7 +251,7 @@ private:
 	ANSICHAR	StringBuffer[N];
 };
 
-template<SIZE_T N>
+template<int32 N>
 class TAnsiStringBuilder : public FAnsiStringBuilderBase
 {
 public:
@@ -262,23 +269,23 @@ private:
 
 //////////////////////////////////////////////////////////////////////////
 
-extern template class TStringBuilderImpl<WIDECHAR>;
+extern template class TStringBuilderImpl<TCHAR>;
 
-class FWideStringBuilderBase : public TStringBuilderImpl<WIDECHAR>
+class FStringBuilderBase : public TStringBuilderImpl<TCHAR>
 {
 public:
-	operator FWideStringView() const { return FWideStringView(Base, CurPos - Base); }
+	operator FStringView() const { return FStringView(Base, CurPos - Base); }
 
-	using TStringBuilderImpl<WIDECHAR>::Append;
+	using TStringBuilderImpl<TCHAR>::Append;
 
-	inline TStringBuilderImpl& Append(const FWideStringView& StringView)
+	inline TStringBuilderImpl& Append(const FStringView& StringView)
 	{
-		const SIZE_T StringLength = StringView.Len();
+		const int32 StringLength = StringView.Len();
 
 		EnsureCapacity(StringLength);
-		WIDECHAR* RESTRICT Dest = CurPos;
+		TCHAR* RESTRICT Dest = CurPos;
 
-		FMemory::Memcpy(CurPos, StringView.Data(), StringLength * sizeof(WIDECHAR));
+		FMemory::Memcpy(CurPos, StringView.Data(), StringLength * sizeof(TCHAR));
 
 		CurPos += StringLength;
 
@@ -286,73 +293,65 @@ public:
 	}
 
 protected:
-	FWideStringBuilderBase(WIDECHAR* BufferPointer, size_t BufferCapacity)
+	FStringBuilderBase(TCHAR* BufferPointer, int32 BufferCapacity)
 	{
 		Initialize(BufferPointer, BufferCapacity);
 	}
 
-	~FWideStringBuilderBase() = default;
+	~FStringBuilderBase() = default;
 };
 
-template<SIZE_T N>
-class TFixedWideStringBuilder : public FWideStringBuilderBase
+template<int32 N>
+class TFixedStringBuilder : public FStringBuilderBase
 {
 public:
-	inline TFixedWideStringBuilder()
-	: FWideStringBuilderBase(StringBuffer, N)
+	inline TFixedStringBuilder()
+	: FStringBuilderBase(StringBuffer, N)
 	{
 	}
 
-	~TFixedWideStringBuilder() = default;
+	~TFixedStringBuilder() = default;
 
 private:
-	WIDECHAR	StringBuffer[N];
+	TCHAR	StringBuffer[N];
 };
 
-template<SIZE_T N>
-class TWideStringBuilder : public FWideStringBuilderBase
+template<int32 N>
+class TStringBuilder : public FStringBuilderBase
 {
 public:
-	inline TWideStringBuilder()
-	: FWideStringBuilderBase(StringBuffer, N)
+	inline TStringBuilder()
+	: FStringBuilderBase(StringBuffer, N)
 	{
 		bIsExtendable = true;
 	}
 
-	~TWideStringBuilder() = default;
+	~TStringBuilder() = default;
 
 private:
-	WIDECHAR	StringBuffer[N];
+	TCHAR	StringBuffer[N];
 };
 
-// In theory, this should probably be configurable but it seems like the
-// engine currently does not support narrow characters.
-
-static_assert(sizeof(TCHAR) == sizeof(WIDECHAR), "Currently expecting wide TCHAR only");
-
-template<SIZE_T N> using TFixedStringBuilder	= TFixedWideStringBuilder<N>;
-template<SIZE_T N> using TStringBuilder			= TWideStringBuilder<N>;
-
 // Append operator implementations
 
-inline FWideStringBuilderBase&				operator<<(FWideStringBuilderBase& Builder, const WIDECHAR Char)					{ Builder.Append(Char); return Builder; }
+inline FStringBuilderBase&					operator<<(FStringBuilderBase& Builder, const TCHAR Char)						{ Builder.Append(Char); return Builder; }
 
 #if USE_STRING_LITERAL_PATH
 
-inline FWideStringBuilderBase&				operator<<(FWideStringBuilderBase& Builder, const WIDECHAR*& NulTerminatedString)	{ Builder.Append(NulTerminatedString); return Builder; }
-inline FWideStringBuilderBase&				operator<<(FWideStringBuilderBase& Builder, const ANSICHAR*& NulTerminatedString)	{ Builder.AppendAnsi(NulTerminatedString); return Builder; }
+inline FStringBuilderBase&					operator<<(FStringBuilderBase& Builder, const TCHAR*& NulTerminatedString)		{ Builder.Append(NulTerminatedString); return Builder; }
+inline FStringBuilderBase&					operator<<(FStringBuilderBase& Builder, const ANSICHAR*& NulTerminatedString)	{ Builder.AppendAnsi(NulTerminatedString); return Builder; }
 
-template<SIZE_T N> FWideStringBuilderBase&	operator<<(FWideStringBuilderBase& Builder, const WIDECHAR (&StringLiteral)[N])		{ Builder.Append(StringLiteral); return Builder; }
-template<SIZE_T N> FWideStringBuilderBase&	operator<<(FWideStringBuilderBase& Builder, const ANSICHAR (&StringLiteral)[N])		{ Builder.AppendAnsi(StringLiteral); return Builder; }
+template<int32 N> FStringBuilderBase&		operator<<(FStringBuilderBase& Builder, const TCHAR (&StringLiteral)[N])		{ Builder.Append(StringLiteral); return Builder; }
+template<int32 N> FStringBuilderBase&		operator<<(FStringBuilderBase& Builder, const ANSICHAR (&StringLiteral)[N])		{ Builder.AppendAnsi(StringLiteral); return Builder; }
 
-template<SIZE_T N> FWideStringBuilderBase&	operator<<(FWideStringBuilderBase& Builder, WIDECHAR (&NulTerminatedString)[N])		{ Builder.Append(&NulTerminatedString[0], TCString<WIDECHAR>::Strlen(NulTerminatedString)); return Builder; }
-template<SIZE_T N> FWideStringBuilderBase&	operator<<(FWideStringBuilderBase& Builder, ANSICHAR (&NulTerminatedString)[N])		{ Builder.AppendAnsi(&NulTerminatedString[0], TCString<ANSICHAR>::Strlen(NulTerminatedString)); return Builder; }
+template<int32 N> FStringBuilderBase&		operator<<(FStringBuilderBase& Builder, TCHAR (&NulTerminatedString)[N])		{ Builder.Append(&NulTerminatedString[0], TCString<TCHAR>::Strlen(NulTerminatedString)); return Builder; }
+template<int32 N> FStringBuilderBase&		operator<<(FStringBuilderBase& Builder, ANSICHAR (&NulTerminatedString)[N])		{ Builder.AppendAnsi(&NulTerminatedString[0], TCString<ANSICHAR>::Strlen(NulTerminatedString)); return Builder; }
 
 #else
 
-inline FWideStringBuilderBase&				operator<<(FWideStringBuilderBase& Builder, const WIDECHAR* NulTerminatedString)	{ Builder.Append(NulTerminatedString); return Builder; }
-inline FWideStringBuilderBase&				operator<<(FWideStringBuilderBase& Builder, const ANSICHAR* NulTerminatedString)	{ Builder.AppendAnsi(NulTerminatedString); return Builder; }
+inline FStringBuilderBase&					operator<<(FStringBuilderBase& Builder, const TCHAR* NulTerminatedString)		{ Builder.Append(NulTerminatedString); return Builder; }
+inline FStringBuilderBase&					operator<<(FStringBuilderBase& Builder, const ANSICHAR* NulTerminatedString)	{ Builder.AppendAnsi(NulTerminatedString); return Builder; }
 
 #endif
 
-inline FWideStringBuilderBase&				operator<<(FWideStringBuilderBase& Builder, const FWideStringView& Str)				{ Builder.Append(Str.Data(), Str.Len()); return Builder; }
+inline FStringBuilderBase&					operator<<(FStringBuilderBase& Builder, const FStringView& Str)					{ Builder.Append(Str.Data(), Str.Len()); return Builder; }
diff --git a/Engine/Source/Runtime/Core/Public/Misc/StringView.h b/Engine/Source/Runtime/Core/Public/Misc/StringView.h
--- a/Engine/Source/Runtime/Core/Public/Misc/StringView.h
+++ b/Engine/Source/Runtime/Core/Public/Misc/StringView.h
@@ -19,19 +19,21 @@ template<typename C>
 class TStringViewImpl
 {
 public:
+	typedef int32 SizeType;
+
 	inline TStringViewImpl(const C* InData)
 	:	DataPtr(InData)
 	,	Size(TCString<C>::Strlen(InData))
 	{
 	}
 
-	inline TStringViewImpl(const C* InData, SIZE_T InSize)
+	inline TStringViewImpl(const C* InData, SizeType InSize)
 	:	DataPtr(InData)
 	,	Size(InSize)
 	{
 	}
 
-	inline const C& operator[](SIZE_T Pos) const	{ return DataPtr[Pos]; }
+	inline const C& operator[](SizeType Pos) const	{ return DataPtr[Pos]; }
 
 	// Q: should this even be here? Semantics differ from FString as it won't return 
 	//    a null terminated string
@@ -41,18 +43,18 @@ public:
 
 	// Capacity
 
-	inline SIZE_T	Len() const						{ return Size; }
+	inline SizeType	Len() const						{ return Size; }
 	inline bool		IsEmpty() const					{ return Size == 0; }
 
 	// Modifiers
 
-	inline void		RemovePrefix(SIZE_T CharCount)	{ DataPtr += CharCount; Size -= CharCount; }
-	inline void		RemoveSuffix(SIZE_T CharCount)	{ Size -= CharCount; }
+	inline void		RemovePrefix(SizeType CharCount)	{ DataPtr += CharCount; Size -= CharCount; }
+	inline void		RemoveSuffix(SizeType CharCount)	{ Size -= CharCount; }
 
 	// Operations
 
-	inline SIZE_T			CopyString(C* Dest, SIZE_T CharCount, SIZE_T Position = 0)	{ memcpy(Dest, DataPtr + Position, FMath::Min(Size - Position, CharCount)); }
-	inline TStringViewImpl& SubStr(SIZE_T Position, SIZE_T CharCount)					{ return TStringViewImpl(DataPtr + Position, FMath::Min(Size - Position, CharCount)); }
+	inline SIZE_T			CopyString(C* Dest, SizeType CharCount, SizeType Position = 0)	{ memcpy(Dest, DataPtr + Position, FMath::Min(Size - Position, CharCount)); }
+	inline TStringViewImpl& SubStr(SizeType Position, SizeType CharCount)					{ return TStringViewImpl(DataPtr + Position, FMath::Min(Size - Position, CharCount)); }
 
 	// Comparison
 
@@ -74,17 +76,9 @@ public:
 		return 0 == TCString<C>::Strnicmp(Suffix.Data(), this->Data() + this->Len() - SuffixLen, SuffixLen);
 	}
 
-	// Find
-
-	int32 Find(const FString& SubStr, ESearchCase::Type SearchCase = ESearchCase::IgnoreCase, ESearchDir::Type SearchDir = ESearchDir::FromStart) const
-	{
-		// :TODO:
-		return 0;
-	}
-
 	inline bool FindChar(C InChar, int32& OutIndex) const
 	{
-		for (SIZE_T i = 0; i < Size; ++i)
+		for (SizeType i = 0; i < Size; ++i)
 		{
 			if (DataPtr[i] == InChar)
 			{
@@ -103,7 +97,7 @@ public:
 			return false;
 		}
 
-		for (SIZE_T i = Size - 1; i >= 0; --i)
+		for (SizeType i = Size - 1; i >= 0; --i)
 		{
 			if (DataPtr[i] == InChar)
 			{
@@ -124,7 +118,7 @@ public:
 
 protected:
 	const C*	DataPtr;
-	SIZE_T		Size;
+	SizeType	Size;
 };
 
 template<typename C>
@@ -149,7 +143,7 @@ inline bool operator!=(const TStringViewImpl<C>& lhs, const TStringViewImpl<C>&
 class FAnsiStringView : public TStringViewImpl<ANSICHAR>
 {
 public:
-	FAnsiStringView(const ANSICHAR* InString, SIZE_T InStrlen)
+	FAnsiStringView(const ANSICHAR* InString, SizeType InStrlen)
 	:	TStringViewImpl<ANSICHAR>(InString, InStrlen)
 	{
 	}
@@ -160,23 +154,21 @@ public:
 	}
 };
 
-class FWideStringView : public TStringViewImpl<WIDECHAR>
+class FStringView : public TStringViewImpl<TCHAR>
 {
 public:
-	FWideStringView(const FString& InString)
-	:	TStringViewImpl<WIDECHAR>(*InString, InString.Len())
+	FStringView(const FString& InString)
+	:	TStringViewImpl<TCHAR>(*InString, InString.Len())
 	{
 	}
 
-	FWideStringView(const WIDECHAR* InString)
-	:	TStringViewImpl<WIDECHAR>(InString)
+	FStringView(const TCHAR* InString)
+	:	TStringViewImpl<TCHAR>(InString)
 	{
 	}
 
-	FWideStringView(const WIDECHAR* InString, SIZE_T InStrlen)
-	:	TStringViewImpl<WIDECHAR>(InString, InStrlen)
+	FStringView(const TCHAR* InString, SizeType InStrlen)
+	:	TStringViewImpl<TCHAR>(InString, InStrlen)
 	{
 	}
 };
-
-using FStringView = FWideStringView;
