Removed operator*() since it's quite dangerous. Client code should use string builders to get a properly terminated string, or call .ToString() to convert into FString

diff --git a/Engine/Source/Runtime/Core/Private/IO/IoDispatcher.cpp b/Engine/Source/Runtime/Core/Private/IO/IoDispatcher.cpp
--- a/Engine/Source/Runtime/Core/Private/IO/IoDispatcher.cpp
+++ b/Engine/Source/Runtime/Core/Private/IO/IoDispatcher.cpp
@@ -376,7 +376,7 @@ FIoStoreEnvironment::~FIoStoreEnvironment()
 
 void FIoStoreEnvironment::InitializeFileEnvironment(FStringView InRootPath)
 {
-	RootPath = *InRootPath;
+	RootPath = InRootPath.ToString();
 }
 
 //////////////////////////////////////////////////////////////////////////
diff --git a/Engine/Source/Runtime/Core/Private/Misc/StringView.cpp b/Engine/Source/Runtime/Core/Private/Misc/StringView.cpp
new file mode 100644
--- /dev/null
+++ b/Engine/Source/Runtime/Core/Private/Misc/StringView.cpp
@@ -0,0 +1,9 @@
+// Copyright 1998-2019 Epic Games, Inc. All Rights Reserved.
+
+#include "Misc/StringView.h"
+
+FString 
+FStringView::ToString() const
+{
+	return FString(Len(), Data());
+}
diff --git a/Engine/Source/Runtime/Core/Public/Misc/StringView.h b/Engine/Source/Runtime/Core/Public/Misc/StringView.h
--- a/Engine/Source/Runtime/Core/Public/Misc/StringView.h
+++ b/Engine/Source/Runtime/Core/Public/Misc/StringView.h
@@ -9,10 +9,43 @@
 /** String View
 
 	* A string view is implicitly constructible from FString and const char* style strings
+
 	* A string view does not own any data nor does it attempt to control any lifetimes, it 
 	  merely points at a subrange of characters in some other string. It's up to the user
 	  to ensure the underlying string stays valid for the lifetime of the string view.
 
+	* A string view does not represent a NUL terminated string and therefore you should
+	  never pass in the pointer returned by Data() into a C-string API accepting only a
+	  pointer. You must either use a string builder to make a properly terminated string,
+	  or the ToString() function below if you absolutely must, and can live with the
+	  knowledge that you just added yet one more memory allocation even though memory
+	  allocations are not cheap.
+
+	String views are a good fit for arguments to functions which don't wish to care
+	which style of string construction is used by the caller. If you accept strings via
+	string views then the caller is free to use FString, FStringBuilder or raw C strings
+	or any other type which can be converted into a string builder.
+
+	I.e a function such as
+
+	void DoFoo(const FStringView& InString);
+
+	May be called as:
+
+	void MultiFoo()
+	{
+	     FString MyFoo(TEXT("Zoo"));
+		 const TCHAR* MyFooStr = *MyFoo;
+
+		 TStringBuilder<64> BuiltFoo;
+		 BuiltFoo.Append(TEXT("ABC"));
+
+		 DoFoo(MyFoo);
+		 DoFoo(MyFooStr);
+		 DoFoo(TEXT("ABC"));
+		 DoFoo(BuiltFoo);
+	}
+
   */
 
 template<typename C>
@@ -35,10 +68,6 @@ public:
 
 	inline const C& operator[](SizeType Pos) const	{ return DataPtr[Pos]; }
 
-	// Q: should this even be here? Semantics differ from FString as it won't return 
-	//    a null terminated string
-	inline const C* operator*() const				{ return DataPtr; }			
-
 	inline const C* Data() const					{ return DataPtr; }
 
 	// Capacity
@@ -121,6 +150,8 @@ protected:
 	SizeType	Size;
 };
 
+// Case-insensitive comparison operators
+
 template<typename C>
 inline bool operator==(const TStringViewImpl<C>& lhs, const TStringViewImpl<C>& rhs)
 {
@@ -171,4 +202,22 @@ public:
 	:	TStringViewImpl<TCHAR>(InString, InStrlen)
 	{
 	}
+
+	/** Convert to a dynamic string
+
+		Remember that a string view is not necessarily
+		NUL terminated so whenever you end up passing a
+		string into an API which accepts a plain C string
+		then you will need to ensure it's NUL terminated
+		first. This is one way of doing that, involving
+		a memory allocation as well as a string copy,
+		along with a memory deallocation when the string
+		goes away.
+
+		I would encourage you to use a TStringBuilder<>
+		instead since this avoids hitting the heap, and
+		the heap is always more expensive even if it
+		seems seductively convenient.
+	  */
+	CORE_API FString ToString() const;
 };
