Fix for formatting %S, which is ANSICHAR when TCHARs are wide.

diff --git a/Engine/Source/Runtime/Core/Private/GenericPlatform/GenericWidePlatformString.cpp b/Engine/Source/Runtime/Core/Private/GenericPlatform/GenericWidePlatformString.cpp
--- a/Engine/Source/Runtime/Core/Private/GenericPlatform/GenericWidePlatformString.cpp
+++ b/Engine/Source/Runtime/Core/Private/GenericPlatform/GenericWidePlatformString.cpp
@@ -757,7 +757,39 @@ int32 FGenericWidePlatformString::GetVarArgs( WIDECHAR* Dest, SIZE_T DestSize, c
 			{
 				Src++;
 				static const TCHAR* Null = TEXT("(null)");
-				const TCHAR *Val = va_arg(ArgPtr, TCHAR *);
+				const TCHAR* Val = va_arg(ArgPtr, TCHAR*);
+				if (Val == nullptr)
+				{
+					Val = Null;
+				}
+
+				int RetCnt = Strlen(Val);
+				int Spaces = FPlatformMath::Max(FPlatformMath::Abs(FieldLen) - RetCnt, 0);
+				if (Spaces > 0 && FieldLen > 0)
+				{
+					DestIter.Write(TEXT(' '), Spaces);
+				}
+				DestIter.Write(Val, RetCnt);
+				if (Spaces > 0 && FieldLen < 0)
+				{
+					DestIter.Write(TEXT(' '), Spaces);
+				}
+				if (!DestIter)
+				{
+					return -1;
+				}
+				break;
+			}
+
+			case 'S':
+			{
+				// The %S format represents a string which is the opposite of %s - wide if TCHAR is narrow, or narrow if TCHAR is wide
+
+				using OtherCharType = TChooseClass<TIsSame<TCHAR, ANSICHAR>::Value, WIDECHAR, ANSICHAR>::Result;
+
+				Src++;
+				static const OtherCharType* Null = LITERAL(OtherCharType, "(null)");
+				const OtherCharType* Val = va_arg(ArgPtr, OtherCharType*);
 				if (Val == nullptr)
 				{
 					Val = Null;
diff --git a/Engine/Source/Runtime/Core/Private/GenericPlatform/StandardPlatformString.cpp b/Engine/Source/Runtime/Core/Private/GenericPlatform/StandardPlatformString.cpp
--- a/Engine/Source/Runtime/Core/Private/GenericPlatform/StandardPlatformString.cpp
+++ b/Engine/Source/Runtime/Core/Private/GenericPlatform/StandardPlatformString.cpp
@@ -35,7 +35,7 @@ static int32 GetFormattingInfo(const WIDECHAR* Format, FFormatInfo& OutInfo)
 
 	// Skip flags
 	while (*Format == LITERAL(WIDECHAR, '#') || *Format == LITERAL(WIDECHAR, '0') || *Format == LITERAL(WIDECHAR, '-')
-		   || *Format == LITERAL(WIDECHAR, ' ') || *Format == LITERAL(WIDECHAR, '+') || *Format == LITERAL(WIDECHAR, '\''))
+		|| *Format == LITERAL(WIDECHAR, ' ') || *Format == LITERAL(WIDECHAR, '+') || *Format == LITERAL(WIDECHAR, '\''))
 	{
 		Format++;
 	}
@@ -96,11 +96,11 @@ static int32 GetFormattingInfo(const WIDECHAR* Format, FFormatInfo& OutInfo)
 
 	// The only valid length modifier for floating point types is L, all other modifiers should be ignored. Length modifier for void pointers should also be ignored.
 	if (OutInfo.LengthModifier != LITERAL(WIDECHAR, 'L') &&
-		  (OutInfo.Type == LITERAL(WIDECHAR, 'f') || OutInfo.Type == LITERAL(WIDECHAR, 'F')
-		|| OutInfo.Type == LITERAL(WIDECHAR, 'e') || OutInfo.Type == LITERAL(WIDECHAR, 'E')
-		|| OutInfo.Type == LITERAL(WIDECHAR, 'g') || OutInfo.Type == LITERAL(WIDECHAR, 'G')
-		|| OutInfo.Type == LITERAL(WIDECHAR, 'a') || OutInfo.Type == LITERAL(WIDECHAR, 'A')
-		|| OutInfo.Type == LITERAL(WIDECHAR, 'p')))
+		(OutInfo.Type == LITERAL(WIDECHAR, 'f') || OutInfo.Type == LITERAL(WIDECHAR, 'F')
+			|| OutInfo.Type == LITERAL(WIDECHAR, 'e') || OutInfo.Type == LITERAL(WIDECHAR, 'E')
+			|| OutInfo.Type == LITERAL(WIDECHAR, 'g') || OutInfo.Type == LITERAL(WIDECHAR, 'G')
+			|| OutInfo.Type == LITERAL(WIDECHAR, 'a') || OutInfo.Type == LITERAL(WIDECHAR, 'A')
+			|| OutInfo.Type == LITERAL(WIDECHAR, 'p')))
 	{
 		OutInfo.LengthModifier = 0;
 	}
@@ -109,23 +109,37 @@ static int32 GetFormattingInfo(const WIDECHAR* Format, FFormatInfo& OutInfo)
 
 	FMemory::Memcpy(OutInfo.Format, FormatStart, FormatLength * sizeof(WIDECHAR));
 	int32 OutInfoFormatLength = FormatLength;
-	if (OutInfo.HasDynamicWidth && FChar::ToLower(OutInfo.Type) == LITERAL(WIDECHAR, 's'))
+	if (OutInfo.HasDynamicWidth)
 	{
-		OutInfo.Format[OutInfoFormatLength - 1] = 'l';
-		OutInfo.Format[OutInfoFormatLength++] = 's';
+		if (OutInfo.Type == LITERAL(WIDECHAR, 's'))
+		{
+			OutInfo.Format[OutInfoFormatLength - 1] = 'l';
+			OutInfo.Format[OutInfoFormatLength++] = 's';
+		}
+		else if (OutInfo.Type == LITERAL(WIDECHAR, 'S'))
+		{
+			OutInfo.Format[OutInfoFormatLength - 1] = 'h';
+			OutInfo.Format[OutInfoFormatLength++] = 's';
+		}
 	}
 	OutInfo.Format[OutInfoFormatLength] = 0;
-	
+
 	// HACKHACKHACK
 	// This formatting function expects to understand %s as a string no matter which char width.
 	// On mac (and possibly others) this must be fixed up to %S for widechars.
 	// So we will do the fixup ONLY if this is a widechar system and the format is given as %s.
 	// BUG: This function still doesn't handle char16_t correctly.
-	if (sizeof(WIDECHAR) == sizeof(wchar_t) &&
-		OutInfo.Type == LITERAL(WIDECHAR, 's'))
+	if (sizeof(WIDECHAR) == sizeof(wchar_t))
 	{
-		checkSlow(OutInfo.Format[OutInfoFormatLength-1] == LITERAL(WIDECHAR, 's'));
-		OutInfo.Format[OutInfoFormatLength-1] = LITERAL(WIDECHAR, 'S');
+		if (OutInfo.Type == LITERAL(WIDECHAR, 's'))
+		{
+			checkSlow(OutInfo.Format[OutInfoFormatLength - 1] == LITERAL(WIDECHAR, 's'));
+			OutInfo.Format[OutInfoFormatLength - 1] = LITERAL(WIDECHAR, 'S');
+		}
+		else if (OutInfo.Type == LITERAL(WIDECHAR, 'S'))
+		{
+			OutInfo.Format[OutInfoFormatLength - 1] = LITERAL(WIDECHAR, 's');
+		}
 	}
 
 	return FormatLength;
@@ -183,14 +197,22 @@ static void FormatArgument(const FFormatInfo& Info, VA_LIST_REF ArgPtr, Callable
 			{
 				Callable(TEXT("(null)"), UE_ARRAY_COUNT(Formatted));
 			}
+			return;
 		}
 		// Is it a plain string?
-		else if (FChar::ToLower(Info.Format[1]) == LITERAL(WIDECHAR, 's'))
+		else if (Info.Type == LITERAL(WIDECHAR, 's') && Info.Format[1] == LITERAL(WIDECHAR, 'S'))
 		{
 			const WIDECHAR* String = va_arg(ArgPtr, WIDECHAR*);
 			int32 Length = FCString::Strlen(String);
 			Callable(String ? String : TEXT("(null)"), Length);
 		}
+		// Is it a wide string?
+		else if (Info.Type == LITERAL(WIDECHAR, 'S') && Info.Format[1] == LITERAL(WIDECHAR, 's'))
+		{
+			const ANSICHAR* String = va_arg(ArgPtr, ANSICHAR*);
+			int32 FormattedLength = FCStringAnsi::Strlen(String);
+			Callable(String ? String : "(null)", FormattedLength);
+		}
 		// Some form of string requiring formatting, such as a left- or right-justified string
 		else
 		{
@@ -269,7 +291,7 @@ int32 FStandardPlatformString::GetVarArgs( WIDECHAR* Dest, SIZE_T DestSize, cons
 			Format += GetFormattingInfo(Format, Info);
 
 			bool bLengthExceeded = false;
-			FormatArgument(Info, ArgPtr, [&Dest, &DestSize, &bLengthExceeded](const TCHAR* FormattedArg, int32 Length)
+			FormatArgument(Info, ArgPtr, [&Dest, &DestSize, &bLengthExceeded](const auto* FormattedArg, int32 Length)
 			{
 				if (FormattedArg && Length > 0)
 				{
@@ -279,7 +301,18 @@ int32 FStandardPlatformString::GetVarArgs( WIDECHAR* Dest, SIZE_T DestSize, cons
 						bLengthExceeded = true;
 					}
 
-					FMemory::Memcpy(Dest, FormattedArg, Length * sizeof(WIDECHAR));
+					if (sizeof(*FormattedArg) == sizeof(WIDECHAR))
+					{
+						FMemory::Memcpy(Dest, FormattedArg, Length * sizeof(WIDECHAR));
+					}
+					else
+					{
+						for (int32 Index = 0; Index != Length; ++Index)
+						{
+							Dest[Index] = FormattedArg[Index];
+						}
+					}
+
 					Dest += Length;
 					DestSize -= Length;
 				}
diff --git a/Engine/Source/Runtime/Core/Public/Android/AndroidPlatformString.h b/Engine/Source/Runtime/Core/Public/Android/AndroidPlatformString.h
--- a/Engine/Source/Runtime/Core/Public/Android/AndroidPlatformString.h
+++ b/Engine/Source/Runtime/Core/Public/Android/AndroidPlatformString.h
@@ -440,6 +440,10 @@ struct FAndroidPlatformString : public FGenericPlatformString
 						NewFormat[NewIndex++] = LITERAL(WIDECHAR, 'l');
 						NewFormat[NewIndex] = *NextChar;
 					}
+					else if (*NextChar == LITERAL(WIDECHAR, 'S'))
+					{
+						NewFormat[NewIndex] = LITERAL(WIDECHAR, 's');
+					}
 					else
 					{
 						NewFormat[NewIndex] = *NextChar;
diff --git a/Engine/Source/Runtime/Core/Public/GenericPlatform/StandardPlatformString.h b/Engine/Source/Runtime/Core/Public/GenericPlatform/StandardPlatformString.h
--- a/Engine/Source/Runtime/Core/Public/GenericPlatform/StandardPlatformString.h
+++ b/Engine/Source/Runtime/Core/Public/GenericPlatform/StandardPlatformString.h
@@ -162,6 +162,10 @@ public:
 						NewFormat[NewIndex++] = LITERAL(WIDECHAR, 'l');
 						NewFormat[NewIndex] = *NextChar;
 					}
+					else if (*NextChar == LITERAL(WIDECHAR, 'S'))
+					{
+						NewFormat[NewIndex] = LITERAL(WIDECHAR, 's');
+					}
 					else
 					{
 						NewFormat[NewIndex] = *NextChar;
