Fixed incorrect logic in trace cpu profiler dynamic scope names map

diff --git a/Engine/Source/Runtime/Core/Private/ProfilingDebugging/CpuProfilerTrace.cpp b/Engine/Source/Runtime/Core/Private/ProfilingDebugging/CpuProfilerTrace.cpp
--- a/Engine/Source/Runtime/Core/Private/ProfilingDebugging/CpuProfilerTrace.cpp
+++ b/Engine/Source/Runtime/Core/Private/ProfilingDebugging/CpuProfilerTrace.cpp
@@ -8,6 +8,7 @@
 #include "Misc/Parse.h"
 #include "Containers/Map.h"
 #include "Misc/MemStack.h"
+#include "Misc/Crc.h"
 
 #if CPUPROFILERTRACE_ENABLED
 
@@ -50,7 +51,7 @@ struct FCpuProfilerTraceInternal
 
 		static FORCEINLINE uint32 GetKeyHash(const CharType* Key)
 		{
-			return GetTypeHash(Key);
+			return FCrc::Strihash_DEPRECATED(Key);
 		}
 
 		static FORCEINLINE KeyInitType GetSetKey(ElementInitType Element)
@@ -148,16 +149,21 @@ void FCpuProfilerTrace::OutputBeginDynamicEvent(const ANSICHAR* Name, ECpuProfil
 		{
 			ThreadBuffer = FCpuProfilerTraceInternal::CreateThreadBuffer();
 		}
-		uint32& SpecId = ThreadBuffer->DynamicAnsiScopeNamesMap.FindOrAdd(Name);
-		if (SpecId == 0)
+		uint32* FindSpecId = ThreadBuffer->DynamicAnsiScopeNamesMap.Find(Name);
+		if (!FindSpecId)
 		{
 			int32 NameSize = strlen(Name) + 1;
 			ANSICHAR* NameCopy = reinterpret_cast<ANSICHAR*>(ThreadBuffer->DynamicScopeNamesMemory.Alloc(NameSize, alignof(ANSICHAR)));
 			FMemory::Memmove(NameCopy, Name, NameSize);
-			SpecId = OutputEventType(NameCopy, Group);
+			uint32 SpecId = OutputEventType(NameCopy, Group);
 			ThreadBuffer->DynamicAnsiScopeNamesMap.Add(NameCopy, SpecId);
+			OutputBeginEvent(SpecId);
 		}
-		OutputBeginEvent(SpecId);
+		else
+		{
+			OutputBeginEvent(*FindSpecId);
+		}
+		
 	}
 	else
 	{
@@ -174,16 +180,20 @@ void FCpuProfilerTrace::OutputBeginDynamicEvent(const TCHAR* Name, ECpuProfilerG
 		{
 			ThreadBuffer = FCpuProfilerTraceInternal::CreateThreadBuffer();
 		}
-		uint32& SpecId = ThreadBuffer->DynamicTCharScopeNamesMap.FindOrAdd(Name);
-		if (SpecId == 0)
+		uint32* FindSpecId = ThreadBuffer->DynamicTCharScopeNamesMap.Find(Name);
+		if (!FindSpecId)
 		{
 			int32 NameSize = (FCString::Strlen(Name) + 1) * sizeof(TCHAR);
 			TCHAR* NameCopy = reinterpret_cast<TCHAR*>(ThreadBuffer->DynamicScopeNamesMemory.Alloc(NameSize, alignof(TCHAR)));
 			FMemory::Memmove(NameCopy, Name, NameSize);
-			SpecId = OutputEventType(NameCopy, Group);
+			uint32 SpecId = OutputEventType(NameCopy, Group);
 			ThreadBuffer->DynamicTCharScopeNamesMap.Add(NameCopy, SpecId);
+			OutputBeginEvent(SpecId);
+		}
+		else
+		{
+			OutputBeginEvent(*FindSpecId);
 		}
-		OutputBeginEvent(SpecId);
 	}
 	else
 	{
