More descriptive log output when not able to launch crash report client.

diff --git a/Engine/Source/Runtime/Core/Private/Windows/WindowsPlatformCrashContext.cpp b/Engine/Source/Runtime/Core/Private/Windows/WindowsPlatformCrashContext.cpp
--- a/Engine/Source/Runtime/Core/Private/Windows/WindowsPlatformCrashContext.cpp
+++ b/Engine/Source/Runtime/Core/Private/Windows/WindowsPlatformCrashContext.cpp
@@ -659,7 +659,7 @@ int32 ReportCrashUsingCrashReportClient(FWindowsPlatformCrashContext& InContext,
 
 	if( bCanRunCrashReportClient )
 	{
-		static const TCHAR CrashReportClientExeName[] = TEXT("CrashReportClient.exe");
+		TCHAR CrashReporterClientPath[CR_CLIENT_MAX_PATH_LEN] = { 0 };
 		bool bCrashReporterRan = false;
 
 		// Generate Crash GUID
@@ -745,8 +745,6 @@ int32 ReportCrashUsingCrashReportClient(FWindowsPlatformCrashContext& InContext,
 				CrashReportClientArguments += FString::Format(TEXT(" -abslog=\"{0}\""), { AbsCrashReportClientLog });
 			}
 
-			TCHAR CrashReporterClientPath[CR_CLIENT_MAX_PATH_LEN] = { 0 };
-
 			if (CreateCrashReportClientPath(CrashReporterClientPath, CR_CLIENT_MAX_PATH_LEN))
 			{
 				bCrashReporterRan = FPlatformProcess::CreateProc(CrashReporterClientPath, *CrashReportClientArguments, true, false, false, NULL, 0, NULL, NULL).IsValid();
@@ -762,7 +760,7 @@ int32 ReportCrashUsingCrashReportClient(FWindowsPlatformCrashContext& InContext,
 
 		if (!bCrashReporterRan && !bNoDialog)
 		{
-			UE_LOG(LogWindows, Log, TEXT("Could not start %s"), CrashReportClientExeName);
+			UE_LOG(LogWindows, Log, TEXT("Could not start crash report client using %s"), CrashReporterClientPath);
 			FPlatformMemory::DumpStats(*GWarn);
 			FText MessageTitle(FText::Format(
 				NSLOCTEXT("MessageDialog", "AppHasCrashed", "The {0} {1} has crashed and will close"),
