Revert FArchive custom version changes after copy-up.

diff --git a/Engine/Source/Runtime/Core/Private/Serialization/Archive.cpp b/Engine/Source/Runtime/Core/Private/Serialization/Archive.cpp
--- a/Engine/Source/Runtime/Core/Private/Serialization/Archive.cpp
+++ b/Engine/Source/Runtime/Core/Private/Serialization/Archive.cpp
@@ -61,11 +61,9 @@ FArchive::FArchive(const FArchive& ArchiveToCopy)
 	// Don't know why this is set to false, but this is what the original copying code did
 	ArIsFilterEditorOnly  = false;
 
-	SharedCustomVersionContainerForOptimizedLoading = ArchiveToCopy.SharedCustomVersionContainerForOptimizedLoading;
 	bCustomVersionsAreReset = ArchiveToCopy.bCustomVersionsAreReset;
 	if (ArchiveToCopy.CustomVersionContainer)
 	{
-		check(!SharedCustomVersionContainerForOptimizedLoading);
 		CustomVersionContainer = new FCustomVersionContainer(*ArchiveToCopy.CustomVersionContainer);
 	}
 	else
@@ -87,11 +85,9 @@ FArchive& FArchive::operator=(const FArchive& ArchiveToCopy)
 	// Don't know why this is set to false, but this is what the original copying code did
 	ArIsFilterEditorOnly  = false;
 
-	SharedCustomVersionContainerForOptimizedLoading = ArchiveToCopy.SharedCustomVersionContainerForOptimizedLoading;
 	bCustomVersionsAreReset = ArchiveToCopy.bCustomVersionsAreReset;
 	if (ArchiveToCopy.CustomVersionContainer)
 	{
-		check(!SharedCustomVersionContainerForOptimizedLoading);
 		if (!CustomVersionContainer)
 		{
 			CustomVersionContainer = new FCustomVersionContainer(*ArchiveToCopy.CustomVersionContainer);
@@ -432,11 +428,6 @@ void FArchive::SerializeBool( bool& D )
 
 const FCustomVersionContainer& FArchive::GetCustomVersions() const
 {
-	if (SharedCustomVersionContainerForOptimizedLoading)
-	{
-		return *SharedCustomVersionContainerForOptimizedLoading;
-	}
-
 	if (!CustomVersionContainer)
 	{
 		CustomVersionContainer = new FCustomVersionContainer;
@@ -463,7 +454,6 @@ const FCustomVersionContainer& FArchive::GetCustomVersions() const
 
 void FArchive::SetCustomVersions(const FCustomVersionContainer& NewVersions)
 {
-	check(!SharedCustomVersionContainerForOptimizedLoading);
 	if (!CustomVersionContainer)
 	{
 		CustomVersionContainer = new FCustomVersionContainer(NewVersions);
@@ -475,15 +465,8 @@ void FArchive::SetCustomVersions(const FCustomVersionContainer& NewVersions)
 	bCustomVersionsAreReset = false;
 }
 
-void FArchive::SetSharedCustomVersionContainerForOptimizedLoading(const FCustomVersionContainer* SharedCustomVersionContainer)
-{
-	check(!CustomVersionContainer);
-	SharedCustomVersionContainerForOptimizedLoading = SharedCustomVersionContainer;
-}
-
 void FArchive::ResetCustomVersions()
 {
-	check(!SharedCustomVersionContainerForOptimizedLoading);
 	bCustomVersionsAreReset = true;
 }
 
@@ -494,7 +477,6 @@ void FArchive::UsingCustomVersion(const FGuid& Key)
 	{
 		return;
 	}
-	check(!SharedCustomVersionContainerForOptimizedLoading);
 
 	auto* RegisteredVersion = FCustomVersionContainer::GetRegistered().GetVersion(Key);
 
@@ -518,14 +500,6 @@ int32 FArchive::CustomVer(const FGuid& Key) const
 
 void FArchive::SetCustomVersion(const FGuid& Key, int32 Version, FName FriendlyName)
 {
-	if (SharedCustomVersionContainerForOptimizedLoading)
-	{
-		// Handle licensee proxy archives that may want to set their own custom versions instead of relying on the original
-		check(!CustomVersionContainer);
-		CustomVersionContainer = new FCustomVersionContainer(*SharedCustomVersionContainerForOptimizedLoading);
-		SharedCustomVersionContainerForOptimizedLoading = nullptr;
-	}
-
 	const_cast<FCustomVersionContainer&>(GetCustomVersions()).SetVersion(Key, Version, FriendlyName);
 }
 
diff --git a/Engine/Source/Runtime/Core/Public/Serialization/Archive.h b/Engine/Source/Runtime/Core/Public/Serialization/Archive.h
--- a/Engine/Source/Runtime/Core/Public/Serialization/Archive.h
+++ b/Engine/Source/Runtime/Core/Public/Serialization/Archive.h
@@ -1068,13 +1068,6 @@ public:
 	 */
 	virtual void SetCustomVersions(const FCustomVersionContainer& CustomVersionContainer);
 
-	/**
-	 * Sets the shared custom version container for this archive when loading (to avoid allocations).
-	 *
-	 * @param SharedCustomVersionContainer - The shared container of custom versions to use in this archive.
-	 */
-	void SetSharedCustomVersionContainerForOptimizedLoading(const FCustomVersionContainer* SharedCustomVersionContainer);
-
 	/** Resets the custom version numbers for this archive. */
 	virtual void ResetCustomVersions();
 
@@ -1574,7 +1567,6 @@ private:
 	* Keeping it as a heap-allocated object also helps with performance in some cases as we don't need to construct it for archives that don't care about custom versions.
 	*/
 	mutable FCustomVersionContainer* CustomVersionContainer = nullptr;
-	const FCustomVersionContainer* SharedCustomVersionContainerForOptimizedLoading = nullptr;
 
 public:
 	/** Custom property list attribute. If the flag below is set, only these properties will be iterated during serialization. If NULL, then no properties will be iterated. */
diff --git a/Engine/Source/Runtime/CoreUObject/Private/Serialization/AsyncLoading2.cpp b/Engine/Source/Runtime/CoreUObject/Private/Serialization/AsyncLoading2.cpp
--- a/Engine/Source/Runtime/CoreUObject/Private/Serialization/AsyncLoading2.cpp
+++ b/Engine/Source/Runtime/CoreUObject/Private/Serialization/AsyncLoading2.cpp
@@ -3035,7 +3035,7 @@ void FAsyncPackage2::EventDrivenSerializeExport(int32 LocalExportIndex)
 		Ar.SetUE4Ver(Linker->Summary.GetFileVersionUE4());
 		Ar.SetLicenseeUE4Ver(Linker->Summary.GetFileVersionLicenseeUE4());
 		Ar.SetEngineVer(Linker->Summary.SavedByEngineVersion);
-		Ar.SetSharedCustomVersionContainerForOptimizedLoading(&SummaryVersions);
+		Ar.SetCustomVersions(SummaryVersions);
 
 		FArchive* OldLoader = Linker->Loader;
 		Linker->Loader = &Ar;
diff --git a/Engine/Source/Runtime/CoreUObject/Private/UObject/LinkerLoad.cpp b/Engine/Source/Runtime/CoreUObject/Private/UObject/LinkerLoad.cpp
--- a/Engine/Source/Runtime/CoreUObject/Private/UObject/LinkerLoad.cpp
+++ b/Engine/Source/Runtime/CoreUObject/Private/UObject/LinkerLoad.cpp
@@ -1332,7 +1332,7 @@ FLinkerLoad::ELinkerStatus FLinkerLoad::UpdateFromPackageFileSummary()
 	SetUE4Ver(Summary.GetFileVersionUE4());
 	SetLicenseeUE4Ver(Summary.GetFileVersionLicenseeUE4());
 	SetEngineVer(Summary.SavedByEngineVersion);
-	SetSharedCustomVersionContainerForOptimizedLoading(&SummaryVersions);
+	SetCustomVersions(SummaryVersions);
 
 	if (Summary.PackageFlags & PKG_FilterEditorOnly)
 	{
@@ -1355,7 +1355,7 @@ FLinkerLoad::ELinkerStatus FLinkerLoad::UpdateFromPackageFileSummary()
 		Loader->SetUE4Ver(Summary.GetFileVersionUE4());
 		Loader->SetLicenseeUE4Ver(Summary.GetFileVersionLicenseeUE4());
 		Loader->SetEngineVer(Summary.SavedByEngineVersion);
-		Loader->SetSharedCustomVersionContainerForOptimizedLoading(&SummaryVersions);
+		Loader->SetCustomVersions(SummaryVersions);
 	}
 
 	if (UPackage* LinkerRootPackage = LinkerRoot)
