Workaround for GregorianCalendar memory stomp

#rnx

diff --git a/Engine/Source/Runtime/Core/Private/Internationalization/ICUCulture.h b/Engine/Source/Runtime/Core/Private/Internationalization/ICUCulture.h
--- a/Engine/Source/Runtime/Core/Private/Internationalization/ICUCulture.h
+++ b/Engine/Source/Runtime/Core/Private/Internationalization/ICUCulture.h
@@ -193,6 +193,11 @@ public:
 
 	static FString GetCanonicalName(const FString& Name);
 
+	const icu::Locale& GetLocale() const
+	{
+		return ICULocale;
+	}
+
 	TSharedRef<const icu::BreakIterator> GetBreakIterator(const EBreakIteratorType Type);
 	TSharedRef<const icu::Collator, ESPMode::ThreadSafe> GetCollator(const ETextComparisonLevel::Type ComparisonLevel);
 	TSharedRef<const icu::DateFormat, ESPMode::ThreadSafe> GetDateFormatter(const EDateTimeStyle::Type DateStyle, const FString& TimeZone);
diff --git a/Engine/Source/Runtime/Core/Private/Internationalization/ICUInternationalization.cpp b/Engine/Source/Runtime/Core/Private/Internationalization/ICUInternationalization.cpp
--- a/Engine/Source/Runtime/Core/Private/Internationalization/ICUInternationalization.cpp
+++ b/Engine/Source/Runtime/Core/Private/Internationalization/ICUInternationalization.cpp
@@ -744,29 +744,54 @@ void FICUInternationalization::InitializeTimeZone()
 void FICUInternationalization::InitializeInvariantGregorianCalendar()
 {
 	UErrorCode ICUStatus = U_ZERO_ERROR;
-	InvariantGregorianCalendar = MakeUnique<icu::GregorianCalendar>(ICUStatus);
-	InvariantGregorianCalendar->setTimeZone(icu::TimeZone::getUnknown());
+
+	// UE-81211 workaround: We use the createInstance API here as there is a size discrepancy between GregorianCalendar in UE4 and ICU which causes memory stomps
+	//InvariantGregorianCalendar = MakeUnique<icu::GregorianCalendar>(ICUStatus);
+	//InvariantGregorianCalendar->setTimeZone(icu::TimeZone::getUnknown());
+	
+	icu::Calendar* ICUCalendar = icu::Calendar::createInstance(icu::TimeZone::getUnknown(), I18N->InvariantCulture->Implementation->GetLocale(), ICUStatus);
+	if (ensureMsgf(ICUCalendar && ICUCalendar->getDynamicClassID() == icu::GregorianCalendar::getStaticClassID(), TEXT("The invariant culture was expected to return a Gregorian calendar!")))
+	{
+		InvariantGregorianCalendar.Reset(static_cast<icu::GregorianCalendar*>(ICUCalendar));
+	}
+	else
+	{
+		delete ICUCalendar;
+	}
 }
 
 UDate FICUInternationalization::UEDateTimeToICUDate(const FDateTime& DateTime)
 {
-	// UE4 and ICU have a different time scale for pre-Gregorian dates, so we can't just use the UNIX timestamp from the UE4 DateTime
-	// Instead we have to explode the UE4 DateTime into its component parts, and then use an ICU GregorianCalendar (set to the "unknown" 
-	// timezone so it doesn't apply any adjustment to the time) to reconstruct the DateTime as an ICU UDate in the correct scale
-	int32 Year, Month, Day;
-	DateTime.GetDate(Year, Month, Day);
-	const int32 Hour = DateTime.GetHour();
-	const int32 Minute = DateTime.GetMinute();
-	const int32 Second = DateTime.GetSecond();
+	UDate ICUDate = 0;
 
+	if (InvariantGregorianCalendar)
 	{
-		FScopeLock Lock(&InvariantGregorianCalendarCS);
+		// UE4 and ICU have a different time scale for pre-Gregorian dates, so we can't just use the UNIX timestamp from the UE4 DateTime
+		// Instead we have to explode the UE4 DateTime into its component parts, and then use an ICU GregorianCalendar (set to the "unknown" 
+		// timezone so it doesn't apply any adjustment to the time) to reconstruct the DateTime as an ICU UDate in the correct scale
+		int32 Year, Month, Day;
+		DateTime.GetDate(Year, Month, Day);
+		const int32 Hour = DateTime.GetHour();
+		const int32 Minute = DateTime.GetMinute();
+		const int32 Second = DateTime.GetSecond();
 
-		InvariantGregorianCalendar->set(Year, Month - 1, Day, Hour, Minute, Second);
-		
-		UErrorCode ICUStatus = U_ZERO_ERROR;
-		return InvariantGregorianCalendar->getTime(ICUStatus);
+		{
+			FScopeLock Lock(&InvariantGregorianCalendarCS);
+
+			InvariantGregorianCalendar->set(Year, Month - 1, Day, Hour, Minute, Second);
+
+			UErrorCode ICUStatus = U_ZERO_ERROR;
+			ICUDate = InvariantGregorianCalendar->getTime(ICUStatus);
+		}
 	}
+	else
+	{
+		// This method is less accurate (see the comment above), but will work well enough if an ICU GregorianCalendar isn't available
+		const int64 UnixTimestamp = DateTime.ToUnixTimestamp();
+		ICUDate = static_cast<UDate>(static_cast<double>(UnixTimestamp) * U_MILLIS_PER_SECOND);
+	}
+
+	return ICUDate;
 }
 
 UBool FICUInternationalization::OpenDataFile(const void* InContext, void** OutFileContext, void** OutContents, const char* InPath)
diff --git a/Engine/Source/Runtime/Core/Public/Internationalization/Culture.h b/Engine/Source/Runtime/Core/Public/Internationalization/Culture.h
--- a/Engine/Source/Runtime/Core/Public/Internationalization/Culture.h
+++ b/Engine/Source/Runtime/Core/Public/Internationalization/Culture.h
@@ -25,6 +25,7 @@ class CORE_API FCulture
 	friend class FText;
 	friend class FTextChronoFormatter;
 	friend class FICUBreakIteratorManager;
+	friend class FICUInternationalization;
 #endif
 
 public:
