Reverting workaround in 9406568 as 9422192 fixes the underlying size issue

#rnx

diff --git a/Engine/Source/Runtime/Core/Private/Internationalization/ICUCulture.h b/Engine/Source/Runtime/Core/Private/Internationalization/ICUCulture.h
--- a/Engine/Source/Runtime/Core/Private/Internationalization/ICUCulture.h
+++ b/Engine/Source/Runtime/Core/Private/Internationalization/ICUCulture.h
@@ -193,11 +193,6 @@ public:
 
 	static FString GetCanonicalName(const FString& Name);
 
-	const icu::Locale& GetLocale() const
-	{
-		return ICULocale;
-	}
-
 	TSharedRef<const icu::BreakIterator> GetBreakIterator(const EBreakIteratorType Type);
 	TSharedRef<const icu::Collator, ESPMode::ThreadSafe> GetCollator(const ETextComparisonLevel::Type ComparisonLevel);
 	TSharedRef<const icu::DateFormat, ESPMode::ThreadSafe> GetDateFormatter(const EDateTimeStyle::Type DateStyle, const FString& TimeZone);
diff --git a/Engine/Source/Runtime/Core/Private/Internationalization/ICUInternationalization.cpp b/Engine/Source/Runtime/Core/Private/Internationalization/ICUInternationalization.cpp
--- a/Engine/Source/Runtime/Core/Private/Internationalization/ICUInternationalization.cpp
+++ b/Engine/Source/Runtime/Core/Private/Internationalization/ICUInternationalization.cpp
@@ -744,19 +744,10 @@ void FICUInternationalization::InitializeTimeZone()
 void FICUInternationalization::InitializeInvariantGregorianCalendar()
 {
 	UErrorCode ICUStatus = U_ZERO_ERROR;
-
-	// UE-81211 workaround: We use the createInstance API here as there is a size discrepancy between GregorianCalendar in UE4 and ICU which causes memory stomps
-	//InvariantGregorianCalendar = MakeUnique<icu::GregorianCalendar>(ICUStatus);
-	//InvariantGregorianCalendar->setTimeZone(icu::TimeZone::getUnknown());
-	
-	icu::Calendar* ICUCalendar = icu::Calendar::createInstance(icu::TimeZone::getUnknown(), I18N->InvariantCulture->Implementation->GetLocale(), ICUStatus);
-	if (ensureMsgf(ICUCalendar && ICUCalendar->getDynamicClassID() == icu::GregorianCalendar::getStaticClassID(), TEXT("The invariant culture was expected to return a Gregorian calendar!")))
-	{
-		InvariantGregorianCalendar.Reset(static_cast<icu::GregorianCalendar*>(ICUCalendar));
-	}
-	else
+	InvariantGregorianCalendar = MakeUnique<icu::GregorianCalendar>(ICUStatus);
+	if (InvariantGregorianCalendar)
 	{
-		delete ICUCalendar;
+		InvariantGregorianCalendar->setTimeZone(icu::TimeZone::getUnknown());
 	}
 }
 
diff --git a/Engine/Source/Runtime/Core/Public/Internationalization/Culture.h b/Engine/Source/Runtime/Core/Public/Internationalization/Culture.h
--- a/Engine/Source/Runtime/Core/Public/Internationalization/Culture.h
+++ b/Engine/Source/Runtime/Core/Public/Internationalization/Culture.h
@@ -25,7 +25,6 @@ class CORE_API FCulture
 	friend class FText;
 	friend class FTextChronoFormatter;
 	friend class FICUBreakIteratorManager;
-	friend class FICUInternationalization;
 #endif
 
 public:
