Exposed working directory for FInteractiveProcess.

#okforgituhub
#pr 5930

diff --git a/Engine/Source/Runtime/Core/Private/Misc/InteractiveProcess.cpp b/Engine/Source/Runtime/Core/Private/Misc/InteractiveProcess.cpp
--- a/Engine/Source/Runtime/Core/Private/Misc/InteractiveProcess.cpp
+++ b/Engine/Source/Runtime/Core/Private/Misc/InteractiveProcess.cpp
@@ -32,11 +32,18 @@ static FORCEINLINE bool CreatePipeWrite(void*& ReadPipe, void*& WritePipe)
 DEFINE_LOG_CATEGORY(LogInteractiveProcess);
 
 FInteractiveProcess::FInteractiveProcess(const FString& InURL, const FString& InParams, bool InHidden, bool LongTime)
+	: FInteractiveProcess(InURL, InParams, FString(), InHidden, LongTime)
+{
+
+}
+
+FInteractiveProcess::FInteractiveProcess(const FString& InURL, const FString& InParams, const FString& InWorkingDir, bool InHidden, bool LongTime)
 	: bCanceling(false)
 	, bHidden(InHidden)
 	, bKillTree(false)
 	, URL(InURL)
 	, Params(InParams)
+	, WorkingDir(InWorkingDir)
 	, ReadPipeParent(nullptr)
 	, WritePipeParent(nullptr)
 	, ReadPipeChild(nullptr)
@@ -98,7 +105,7 @@ bool FInteractiveProcess::Launch()
 		return false;
 	}
 
-	ProcessHandle = FPlatformProcess::CreateProc(*URL, *Params, false, bHidden, bHidden, nullptr, 0, nullptr, WritePipeChild, ReadPipeChild);
+	ProcessHandle = FPlatformProcess::CreateProc(*URL, *Params, false, bHidden, bHidden, nullptr, 0, WorkingDir.IsEmpty() ? nullptr : *WorkingDir, WritePipeChild, ReadPipeChild);
 
 	if (ProcessHandle.IsValid() == false)
 	{
diff --git a/Engine/Source/Runtime/Core/Public/Misc/InteractiveProcess.h b/Engine/Source/Runtime/Core/Public/Misc/InteractiveProcess.h
--- a/Engine/Source/Runtime/Core/Public/Misc/InteractiveProcess.h
+++ b/Engine/Source/Runtime/Core/Public/Misc/InteractiveProcess.h
@@ -45,6 +45,16 @@ public:
 	*/
 	FInteractiveProcess(const FString& InURL, const FString& InParams, bool InHidden, bool LongTime = false);
 
+	/**
+	* Creates a new interactive process.
+	*
+	* @param InURL The URL of the executable to launch.
+	* @param InParams The command line parameters.
+	* @param InWorkingDir The URL of the working dir where the executable should launch.
+	* @param InHidden Whether the window of the process should be hidden.
+	*/
+	FInteractiveProcess(const FString& InURL, const FString& InParams, const FString& InWorkingDir, bool InHidden, bool LongTime = false);
+
 	/** Destructor. */
 	~FInteractiveProcess();
 
@@ -187,6 +197,9 @@ private:
 	// Holds the command line parameters. */
 	FString Params;
 
+	// Holds the URL of the working dir for the process. */
+	FString WorkingDir;
+
 	// Holds the handle to the process. */
 	FProcHandle ProcessHandle;
 
