FArchiveState separated out from FArchive.

diff --git a/Engine/Source/Runtime/Core/Private/Internationalization/InternationalizationMetadata.cpp b/Engine/Source/Runtime/Core/Private/Internationalization/InternationalizationMetadata.cpp
--- a/Engine/Source/Runtime/Core/Private/Internationalization/InternationalizationMetadata.cpp
+++ b/Engine/Source/Runtime/Core/Private/Internationalization/InternationalizationMetadata.cpp
@@ -352,7 +352,7 @@ void operator<<(FStructuredArchive::FSlot Slot, FLocMetadataObject& Object)
 	int32 ValueCount = Object.Values.Num();
 	Record << SA_VALUE(TEXT("ValueCount"), ValueCount);
 
-	bool bIsLoading = Slot.GetUnderlyingArchive().IsLoading();
+	bool bIsLoading = Slot.GetArchiveState().IsLoading();
 
 	if (bIsLoading)
 	{
@@ -414,13 +414,13 @@ TSharedRef<FLocMetadataValue> FLocMetadataValueString::Clone() const
 
 FLocMetadataValueString::FLocMetadataValueString(FStructuredArchive::FSlot Slot)
 {
-	check(Slot.GetUnderlyingArchive().IsLoading());
+	check(Slot.GetArchiveState().IsLoading());
 	Slot << Value;
 }
 
 void FLocMetadataValueString::Serialize( FLocMetadataValueString& Value, FStructuredArchive::FSlot Slot )
 {
-	check(!Slot.GetUnderlyingArchive().IsLoading());
+	check(!Slot.GetArchiveState().IsLoading());
 
 	FString StringValue = Value.Value;
 	Slot << StringValue;
@@ -445,14 +445,14 @@ TSharedRef<FLocMetadataValue> FLocMetadataValueBoolean::Clone() const
 
 FLocMetadataValueBoolean::FLocMetadataValueBoolean(FStructuredArchive::FSlot Slot)
 {
-	check(Slot.GetUnderlyingArchive().IsLoading());
+	check(Slot.GetArchiveState().IsLoading());
 
 	Slot << Value;
 }
 
 void FLocMetadataValueBoolean::Serialize( FLocMetadataValueBoolean& Value, FStructuredArchive::FSlot Slot )
 {
-	check(!Slot.GetUnderlyingArchive().IsLoading());
+	check(!Slot.GetArchiveState().IsLoading());
 
 	bool BoolValue = Value.Value;
 	Slot << BoolValue;
@@ -552,7 +552,7 @@ FString FLocMetadataValueArray::ToString() const
 
 FLocMetadataValueArray::FLocMetadataValueArray(FStructuredArchive::FSlot Slot)
 {
-	check(Slot.GetUnderlyingArchive().IsLoading());
+	check(Slot.GetArchiveState().IsLoading());
 
 	int32 ElementCount;
 	FStructuredArchive::FArray Array = Slot.EnterArray(ElementCount);
@@ -566,7 +566,7 @@ FLocMetadataValueArray::FLocMetadataValueArray(FStructuredArchive::FSlot Slot)
 
 void FLocMetadataValueArray::Serialize(FLocMetadataValueArray& Value, FStructuredArchive::FSlot Slot)
 {
-	check(!Slot.GetUnderlyingArchive().IsLoading());
+	check(!Slot.GetArchiveState().IsLoading());
 
 	int32 ElementCount = Value.Value.Num();
 	FStructuredArchive::FArray Array = Slot.EnterArray(ElementCount);
@@ -619,14 +619,14 @@ FString FLocMetadataValueObject::ToString() const
 FLocMetadataValueObject::FLocMetadataValueObject(FStructuredArchive::FSlot Slot)
 	: Value(new FLocMetadataObject)
 {
-	check(Slot.GetUnderlyingArchive().IsLoading());
+	check(Slot.GetArchiveState().IsLoading());
 
 	Slot << *Value;
 }
 
 void FLocMetadataValueObject::Serialize( FLocMetadataValueObject& Value, FStructuredArchive::FSlot Slot )
 {
-	check(!Slot.GetUnderlyingArchive().IsLoading());
+	check(!Slot.GetArchiveState().IsLoading());
 
 	Slot << *(Value.Value);
 }
diff --git a/Engine/Source/Runtime/Core/Private/Internationalization/Text.cpp b/Engine/Source/Runtime/Core/Private/Internationalization/Text.cpp
--- a/Engine/Source/Runtime/Core/Private/Internationalization/Text.cpp
+++ b/Engine/Source/Runtime/Core/Private/Internationalization/Text.cpp
@@ -1220,7 +1220,7 @@ void operator<<(FStructuredArchive::FSlot Slot, FFormatArgumentValue& Value)
 		}
 	case EFormatArgumentType::Text:
 		{
-			if(Slot.GetUnderlyingArchive().IsLoading())
+			if(Slot.GetArchiveState().IsLoading())
 			{
 				Value.TextValue = FText();
 			}
diff --git a/Engine/Source/Runtime/Core/Private/Internationalization/TextHistory.cpp b/Engine/Source/Runtime/Core/Private/Internationalization/TextHistory.cpp
--- a/Engine/Source/Runtime/Core/Private/Internationalization/TextHistory.cpp
+++ b/Engine/Source/Runtime/Core/Private/Internationalization/TextHistory.cpp
@@ -676,7 +676,7 @@ bool FTextHistory::GetHistoricNumericData(const FText& InText, FHistoricTextNume
 
 void FTextHistory::SerializeForDisplayString(FStructuredArchive::FRecord Record, FTextDisplayStringPtr& InOutDisplayString)
 {
-	if(Record.GetUnderlyingArchive().IsLoading())
+	if(Record.GetArchiveState().IsLoading())
 	{
 		PrepareDisplayStringForRebuild(InOutDisplayString);
 	}
@@ -2489,7 +2489,7 @@ void FTextHistory_StringTableEntry::Serialize(FStructuredArchive::FRecord Record
 
 void FTextHistory_StringTableEntry::SerializeForDisplayString(FStructuredArchive::FRecord Record, FTextDisplayStringPtr& InOutDisplayString)
 {
-	if (Record.GetUnderlyingArchive().IsLoading())
+	if (Record.GetArchiveState().IsLoading())
 	{
 		// We will definitely need to do a rebuild later
 		Revision = 0;
diff --git a/Engine/Source/Runtime/Core/Private/Misc/Guid.cpp b/Engine/Source/Runtime/Core/Private/Misc/Guid.cpp
--- a/Engine/Source/Runtime/Core/Private/Misc/Guid.cpp
+++ b/Engine/Source/Runtime/Core/Private/Misc/Guid.cpp
@@ -320,9 +320,10 @@ FArchive& operator<<(FArchive& Ar, FGuid& G)
 
 void operator<<(FStructuredArchive::FSlot Slot, FGuid& G)
 {
-	if (Slot.GetUnderlyingArchive().IsTextFormat())
+	const FArchiveState& State = Slot.GetArchiveState();
+	if (State.IsTextFormat())
 	{
-		if (Slot.GetUnderlyingArchive().IsLoading())
+		if (State.IsLoading())
 		{
 			FString AsString;
 			Slot << AsString;
diff --git a/Engine/Source/Runtime/Core/Private/Misc/ObjectThumbnail.cpp b/Engine/Source/Runtime/Core/Private/Misc/ObjectThumbnail.cpp
--- a/Engine/Source/Runtime/Core/Private/Misc/ObjectThumbnail.cpp
+++ b/Engine/Source/Runtime/Core/Private/Misc/ObjectThumbnail.cpp
@@ -59,7 +59,7 @@ void FObjectThumbnail::Serialize(FStructuredArchive::FSlot Slot)
 		Record << SA_VALUE(TEXT("ImageData"), ImageData) << SA_VALUE(TEXT("bIsDirty"), bIsDirty);
 	}
 
-	if (Slot.GetUnderlyingArchive().IsLoading())
+	if (Slot.GetArchiveState().IsLoading())
 	{
 		bLoadedFromDisk = true;
 		if ((ImageWidth>0) && (ImageHeight>0))
diff --git a/Engine/Source/Runtime/Core/Private/Serialization/Archive.cpp b/Engine/Source/Runtime/Core/Private/Serialization/Archive.cpp
--- a/Engine/Source/Runtime/Core/Private/Serialization/Archive.cpp
+++ b/Engine/Source/Runtime/Core/Private/Serialization/Archive.cpp
@@ -30,7 +30,7 @@
 	FArchive implementation.
 -----------------------------------------------------------------------------*/
 
-FArchive::FArchive()
+FArchiveState::FArchiveState()
 {
 #if DEVIRTUALIZE_FLinkerLoad_Serialize
 	ActiveFPLB = &InlineFPLB;
@@ -44,7 +44,7 @@ FArchive::FArchive()
 	Reset();
 }
 
-FArchive::FArchive(const FArchive& ArchiveToCopy)
+FArchiveState::FArchiveState(const FArchiveState& ArchiveToCopy)
 {
 #if DEVIRTUALIZE_FLinkerLoad_Serialize
 	ActiveFPLB = &InlineFPLB;
@@ -72,7 +72,7 @@ FArchive::FArchive(const FArchive& ArchiveToCopy)
 	}
 }
 
-FArchive& FArchive::operator=(const FArchive& ArchiveToCopy)
+FArchiveState& FArchiveState::operator=(const FArchiveState& ArchiveToCopy)
 {
 #if DEVIRTUALIZE_FLinkerLoad_Serialize
 	ActiveFPLB = &InlineFPLB;
@@ -106,7 +106,7 @@ FArchive& FArchive::operator=(const FArchive& ArchiveToCopy)
 	return *this;
 }
 
-FArchive::~FArchive()
+FArchiveState::~FArchiveState()
 {
 	delete CustomVersionContainer;
 
@@ -117,8 +117,7 @@ FArchive::~FArchive()
 #endif // USE_STABLE_LOCALIZATION_KEYS
 }
 
-
-void FArchive::Reset()
+void FArchiveState::Reset()
 {
 #if DEVIRTUALIZE_FLinkerLoad_Serialize
 	ActiveFPLB->Reset();
@@ -178,7 +177,7 @@ void FArchive::Reset()
 	ResetCustomVersions();
 }
 
-void FArchive::CopyTrivialFArchiveStatusMembers(const FArchive& ArchiveToCopy)
+void FArchiveState::CopyTrivialFArchiveStatusMembers(const FArchiveState& ArchiveToCopy)
 {
 	ArUE4Ver                             = ArchiveToCopy.ArUE4Ver;
 	ArLicenseeUE4Ver                     = ArchiveToCopy.ArLicenseeUE4Ver;
@@ -218,7 +217,7 @@ void FArchive::CopyTrivialFArchiveStatusMembers(const FArchive& ArchiveToCopy)
 	ArCustomPropertyList                 = ArchiveToCopy.ArCustomPropertyList;
 	ArUseCustomPropertyList              = ArchiveToCopy.ArUseCustomPropertyList;
 	CookingTargetPlatform                = ArchiveToCopy.CookingTargetPlatform;
-	SerializedProperty					 = ArchiveToCopy.SerializedProperty;
+	SerializedProperty                   = ArchiveToCopy.SerializedProperty;
 #if USE_STABLE_LOCALIZATION_KEYS
 	SetBaseLocalizationNamespace(ArchiveToCopy.GetBaseLocalizationNamespace());
 #endif // USE_STABLE_LOCALIZATION_KEYS
@@ -230,12 +229,12 @@ void FArchive::CopyTrivialFArchiveStatusMembers(const FArchive& ArchiveToCopy)
  *
  * This is overridden for the specific Archive Types
  **/
-FString FArchive::GetArchiveName() const
+FString FArchiveState::GetArchiveName() const
 {
 	return TEXT("FArchive");
 }
 
-void FArchive::GetSerializedPropertyChain(TArray<class UProperty*>& OutProperties) const
+void FArchiveState::GetSerializedPropertyChain(TArray<class UProperty*>& OutProperties) const
 {
 	if (SerializedPropertyChain)
 	{
@@ -249,7 +248,7 @@ void FArchive::GetSerializedPropertyChain(TArray<class UProperty*>& OutPropertie
 	}
 }
 
-void FArchive::SetSerializedPropertyChain(const FArchiveSerializedPropertyChain* InSerializedPropertyChain, class UProperty* InSerializedPropertyOverride)
+void FArchiveState::SetSerializedPropertyChain(const FArchiveSerializedPropertyChain* InSerializedPropertyChain, class UProperty* InSerializedPropertyOverride)
 {
 	if (InSerializedPropertyChain && InSerializedPropertyChain->GetNumProperties() > 0)
 	{
@@ -316,14 +315,14 @@ void FArchive::PopSerializedProperty(class UProperty* InProperty, const bool bIs
 }
 
 #if WITH_EDITORONLY_DATA
-bool FArchive::IsEditorOnlyPropertyOnTheStack() const
+bool FArchiveState::IsEditorOnlyPropertyOnTheStack() const
 {
 	return SerializedPropertyChain && SerializedPropertyChain->HasEditorOnlyProperty();
 }
 #endif
 
 #if USE_STABLE_LOCALIZATION_KEYS
-void FArchive::SetBaseLocalizationNamespace(const FString& InLocalizationNamespace)
+void FArchiveState::SetBaseLocalizationNamespace(const FString& InLocalizationNamespace)
 {
 	if (InLocalizationNamespace.IsEmpty())
 	{
@@ -339,15 +338,15 @@ void FArchive::SetBaseLocalizationNamespace(const FString& InLocalizationNamespa
 		*LocalizationNamespacePtr = InLocalizationNamespace;
 	}
 }
-FString FArchive::GetBaseLocalizationNamespace() const
+FString FArchiveState::GetBaseLocalizationNamespace() const
 {
 	return LocalizationNamespacePtr ? *LocalizationNamespacePtr : FString();
 }
-void FArchive::SetLocalizationNamespace(const FString& InLocalizationNamespace)
+void FArchiveState::SetLocalizationNamespace(const FString& InLocalizationNamespace)
 {
 	SetBaseLocalizationNamespace(InLocalizationNamespace);
 }
-FString FArchive::GetLocalizationNamespace() const
+FString FArchiveState::GetLocalizationNamespace() const
 {
 	return GetBaseLocalizationNamespace();
 }
@@ -426,7 +425,7 @@ void FArchive::SerializeBool( bool& D )
 }
 #endif
 
-const FCustomVersionContainer& FArchive::GetCustomVersions() const
+const FCustomVersionContainer& FArchiveState::GetCustomVersions() const
 {
 	if (!CustomVersionContainer)
 	{
@@ -452,7 +451,7 @@ const FCustomVersionContainer& FArchive::GetCustomVersions() const
 	return *CustomVersionContainer;
 }
 
-void FArchive::SetCustomVersions(const FCustomVersionContainer& NewVersions)
+void FArchiveState::SetCustomVersions(const FCustomVersionContainer& NewVersions)
 {
 	if (!CustomVersionContainer)
 	{
@@ -465,7 +464,7 @@ void FArchive::SetCustomVersions(const FCustomVersionContainer& NewVersions)
 	bCustomVersionsAreReset = false;
 }
 
-void FArchive::ResetCustomVersions()
+void FArchiveState::ResetCustomVersions()
 {
 	bCustomVersionsAreReset = true;
 }
@@ -487,7 +486,7 @@ void FArchive::UsingCustomVersion(const FGuid& Key)
 	const_cast<FCustomVersionContainer&>(GetCustomVersions()).SetVersion(Key, RegisteredVersion->Version, RegisteredVersion->GetFriendlyName());
 }
 
-int32 FArchive::CustomVer(const FGuid& Key) const
+int32 FArchiveState::CustomVer(const FGuid& Key) const
 {
 	auto* CustomVersion = GetCustomVersions().GetVersion(Key);
 
@@ -498,7 +497,7 @@ int32 FArchive::CustomVer(const FGuid& Key) const
 	return CustomVersion ? CustomVersion->Version : -1;
 }
 
-void FArchive::SetCustomVersion(const FGuid& Key, int32 Version, FName FriendlyName)
+void FArchiveState::SetCustomVersion(const FGuid& Key, int32 Version, FName FriendlyName)
 {
 	const_cast<FCustomVersionContainer&>(GetCustomVersions()).SetVersion(Key, Version, FriendlyName);
 }
@@ -1088,62 +1087,67 @@ void FArchive::LogfImpl(const TCHAR* Fmt, ...)
 	FMemory::SystemFree( Buffer );
 }
 
-void FArchive::SetUE4Ver(int32 InVer)
+void FArchiveState::SetUE4Ver(int32 InVer)
 {
 	ArUE4Ver = InVer;
 }
 
-void FArchive::SetLicenseeUE4Ver(int32 InVer)
+void FArchiveState::SetLicenseeUE4Ver(int32 InVer)
 {
 	ArLicenseeUE4Ver = InVer;
 }
 
-void FArchive::SetEngineVer(const FEngineVersionBase& InVer)
+void FArchiveState::SetEngineVer(const FEngineVersionBase& InVer)
 {
 	ArEngineVer = InVer;
 }
 
-void FArchive::SetEngineNetVer(const uint32 InEngineNetVer)
+void FArchiveState::SetEngineNetVer(const uint32 InEngineNetVer)
 {
 	ArEngineNetVer = InEngineNetVer;
 }
 
-void FArchive::SetGameNetVer(const uint32 InGameNetVer)
+void FArchiveState::SetGameNetVer(const uint32 InGameNetVer)
 {
 	ArGameNetVer = InGameNetVer;
 }
 
-void FArchive::SetIsLoading(bool bInIsLoading)
+void FArchiveState::SetIsLoading(bool bInIsLoading)
 {
 	ArIsLoading = bInIsLoading;
 }
 
-void FArchive::SetIsSaving(bool bInIsSaving)
+void FArchiveState::SetIsSaving(bool bInIsSaving)
 {
 	ArIsSaving = bInIsSaving;
 }
 
-void FArchive::SetIsTransacting(bool bInIsTransacting)
+void FArchiveState::SetIsTransacting(bool bInIsTransacting)
 {
 	ArIsTransacting = bInIsTransacting;
 }
 
-void FArchive::SetIsTextFormat(bool bInIsTextFormat)
+void FArchiveState::SetIsTextFormat(bool bInIsTextFormat)
 {
 	ArIsTextFormat = bInIsTextFormat;
 }
 
-void FArchive::SetWantBinaryPropertySerialization(bool bInWantBinaryPropertySerialization)
+void FArchiveState::SetWantBinaryPropertySerialization(bool bInWantBinaryPropertySerialization)
 {
 	ArWantBinaryPropertySerialization = bInWantBinaryPropertySerialization;
 }
 
-void FArchive::SetForceUnicode(bool bInForceUnicode)
+void FArchiveState::SetForceUnicode(bool bInForceUnicode)
 {
 	ArForceUnicode = bInForceUnicode;
 }
 
-void FArchive::SetIsPersistent(bool bInIsPersistent)
+void FArchiveState::SetIsPersistent(bool bInIsPersistent)
 {
 	ArIsPersistent = bInIsPersistent;
 }
+
+void FArchive::SetArchiveState(const FArchiveState& InState)
+{
+	ImplicitConv<FArchiveState&>(*this) = InState;
+}
diff --git a/Engine/Source/Runtime/Core/Public/Serialization/Archive.h b/Engine/Source/Runtime/Core/Public/Serialization/Archive.h
--- a/Engine/Source/Runtime/Core/Public/Serialization/Archive.h
+++ b/Engine/Source/Runtime/Core/Public/Serialization/Archive.h
@@ -16,6 +16,7 @@
 #include "Templates/AndOrNot.h"
 #include "Templates/IsArrayOrRefOfType.h"
 
+class FArchive;
 class FCustomVersionContainer;
 class FLinker;
 class FName;
@@ -44,6 +45,796 @@
 // Helper macro to make serializing a bitpacked boolean in an archive easier
 #define FArchive_Serialize_BitfieldBool(ARCHIVE, BITFIELD_BOOL) { bool TEMP_BITFIELD_BOOL = BITFIELD_BOOL; ARCHIVE << TEMP_BITFIELD_BOOL; BITFIELD_BOOL = TEMP_BITFIELD_BOOL; }
 
+struct CORE_API FArchiveState
+{
+private:
+	// Only FArchive is allowed to instantiate this, by inheritance
+	friend class FArchive;
+
+	FArchiveState();
+
+	/** Copy constructor. */
+	FArchiveState(const FArchiveState&);
+
+	/**
+	 * Copy assignment operator.
+	 *
+	 * @param ArchiveToCopy The archive to copy from.
+	 */
+	FArchiveState& operator=(const FArchiveState& ArchiveToCopy);
+
+	virtual ~FArchiveState() = 0;
+
+public:
+	virtual void CountBytes(SIZE_T InNum, SIZE_T InMax) { }
+
+	/**
+	 * Returns the name of the Archive.  Useful for getting the name of the package a struct or object
+	 * is in when a loading error occurs.
+	 *
+	 * This is overridden for the specific Archive Types
+	 */
+	virtual FString GetArchiveName() const;
+
+	/**
+	 * If this archive is a FLinkerLoad or FLinkerSave, returns a pointer to the ULinker portion.
+	 *
+	 * @return The linker, or nullptr if the archive is not a linker.
+	 */
+	virtual FLinker* GetLinker()
+	{
+		return nullptr;
+	}
+
+	virtual int64 Tell()
+	{
+		return INDEX_NONE;
+	}
+
+	virtual int64 TotalSize()
+	{
+		return INDEX_NONE;
+	}
+
+	virtual bool AtEnd()
+	{
+		int64 Pos = Tell();
+
+		return ((Pos != INDEX_NONE) && (Pos >= TotalSize()));
+	}
+
+	virtual bool GetError()
+	{
+		return ArIsError;
+	}
+
+	void SetError()
+	{
+		ArIsError = true;
+	}
+
+	FORCEINLINE bool IsByteSwapping()
+	{
+	#if PLATFORM_LITTLE_ENDIAN
+		bool SwapBytes = ArForceByteSwapping;
+	#else
+		bool SwapBytes = this->IsPersistent();
+	#endif
+		return SwapBytes;
+	}
+
+	/** Sets a flag indicating that this archive contains code. */
+	void ThisContainsCode()
+	{
+		ArContainsCode = true;
+	}
+
+	/** Sets a flag indicating that this archive contains a ULevel or UWorld object. */
+	void ThisContainsMap()
+	{
+		ArContainsMap = true;
+	}
+
+	/** Sets a flag indicating that this archive contains data required to be gathered for localization. */
+	void ThisRequiresLocalizationGather()
+	{
+		ArRequiresLocalizationGather = true;
+	}
+
+	/**
+	* Called to retrieve the archetype from the event driven loader. If this returns null, then call GetArchetype yourself.
+	*/
+	virtual UObject* GetArchetypeFromLoader(const UObject* Obj)
+	{
+		return nullptr;
+	}
+
+	FORCEINLINE int32 UE4Ver() const
+	{
+		return ArUE4Ver;
+	}
+
+	FORCEINLINE int32 LicenseeUE4Ver() const
+	{
+		return ArLicenseeUE4Ver;
+	}
+
+	FORCEINLINE FEngineVersionBase EngineVer() const
+	{
+		return ArEngineVer;
+	}
+
+	FORCEINLINE uint32 EngineNetVer() const
+	{
+		return ArEngineNetVer;
+	}
+
+	FORCEINLINE uint32 GameNetVer() const
+	{
+		return ArGameNetVer;
+	}
+
+	/**
+	 * Queries a custom version from the archive.  If the archive is being used to write, the custom version must have already been registered.
+	 *
+	 * @param Key The guid of the custom version to query.
+	 * @return The version number, or 0 if the custom tag isn't stored in the archive.
+	 */
+	int32 CustomVer(const struct FGuid& Key) const;
+
+	FORCEINLINE bool IsLoading() const
+	{
+		return ArIsLoading;
+	}
+
+	FORCEINLINE bool IsSaving() const
+	{
+		return ArIsSaving;
+	}
+
+	FORCEINLINE bool IsTransacting() const
+	{
+		if (FPlatformProperties::HasEditorOnlyData())
+		{
+			return ArIsTransacting;
+		}
+		else
+		{
+			return false;
+		}
+	}
+
+	FORCEINLINE bool IsTextFormat() const
+	{
+		return (ArIsTextFormat && WITH_TEXT_ARCHIVE_SUPPORT);
+	}
+
+	FORCEINLINE bool WantBinaryPropertySerialization() const
+	{
+		return ArWantBinaryPropertySerialization;
+	}
+
+	FORCEINLINE bool IsForcingUnicode() const
+	{
+		return ArForceUnicode;
+	}
+
+	FORCEINLINE bool IsPersistent() const
+	{
+		return ArIsPersistent;
+	}
+
+	FORCEINLINE bool IsError() const
+	{
+		return ArIsError;
+	}
+
+	FORCEINLINE bool IsCriticalError() const
+	{
+		return ArIsCriticalError;
+	}
+
+	FORCEINLINE bool ContainsCode() const
+	{
+		return ArContainsCode;
+	}
+
+	FORCEINLINE bool ContainsMap() const
+	{
+		return ArContainsMap;
+	}
+
+	FORCEINLINE bool RequiresLocalizationGather() const
+	{
+		return ArRequiresLocalizationGather;
+	}
+
+	FORCEINLINE bool ForceByteSwapping() const
+	{
+		return ArForceByteSwapping;
+	}
+
+	FORCEINLINE bool IsSerializingDefaults() const
+	{
+		return (ArSerializingDefaults > 0) ? true : false;
+	}
+
+	FORCEINLINE bool IsIgnoringArchetypeRef() const
+	{
+		return ArIgnoreArchetypeRef;
+	}
+
+	FORCEINLINE bool DoDelta() const
+	{
+		return !ArNoDelta;
+	}
+
+	FORCEINLINE bool DoIntraPropertyDelta() const
+	{
+		return !ArNoIntraPropertyDelta;
+	}
+
+	FORCEINLINE bool IsIgnoringOuterRef() const
+	{
+		return ArIgnoreOuterRef;
+	}
+
+	FORCEINLINE bool IsIgnoringClassGeneratedByRef() const
+	{
+		return ArIgnoreClassGeneratedByRef;
+	}
+
+	FORCEINLINE bool IsIgnoringClassRef() const
+	{
+		return ArIgnoreClassRef;
+	}
+
+	FORCEINLINE bool IsAllowingLazyLoading() const
+	{
+		return ArAllowLazyLoading;
+	}
+
+	FORCEINLINE bool IsObjectReferenceCollector() const
+	{
+		return ArIsObjectReferenceCollector;
+	}
+
+	FORCEINLINE bool IsModifyingWeakAndStrongReferences() const
+	{
+		return ArIsModifyingWeakAndStrongReferences;
+	}
+
+	FORCEINLINE bool IsCountingMemory() const
+	{
+		return ArIsCountingMemory;
+	}
+
+	FORCEINLINE uint32 GetPortFlags() const
+	{
+		return ArPortFlags;
+	}
+
+	FORCEINLINE bool HasAnyPortFlags(uint32 Flags) const
+	{
+		return ((ArPortFlags & Flags) != 0);
+	}
+
+	FORCEINLINE bool HasAllPortFlags(uint32 Flags) const
+	{
+		return ((ArPortFlags & Flags) == Flags);
+	}
+
+	FORCEINLINE uint32 GetDebugSerializationFlags() const
+	{
+#if WITH_EDITOR
+		return ArDebugSerializationFlags;
+#else
+		return 0;
+#endif
+	}
+
+	FORCEINLINE bool ShouldSkipBulkData() const
+	{
+		return ArShouldSkipBulkData;
+	}
+
+	FORCEINLINE int64 GetMaxSerializeSize() const
+	{
+		return ArMaxSerializeSize;
+	}
+
+	/**
+	 * Gets the custom version numbers for this archive.
+	 *
+	 * @return The container of custom versions in the archive.
+	 */
+	virtual const FCustomVersionContainer& GetCustomVersions() const;
+
+	/**
+	 * Sets the custom version numbers for this archive.
+	 *
+	 * @param CustomVersionContainer - The container of custom versions to copy into the archive.
+	 */
+	virtual void SetCustomVersions(const FCustomVersionContainer& CustomVersionContainer);
+
+	/** Resets the custom version numbers for this archive. */
+	virtual void ResetCustomVersions();
+
+	/**
+	 * Sets a specific custom version
+	 *
+	 * @param Key - The guid of the custom version to query.
+	 * @param Version - The version number to set key to
+	 * @param FriendlyName - Friendly name corresponding to the key
+	 */
+	void SetCustomVersion(const struct FGuid& Key, int32 Version, FName FriendlyName);
+
+	/**
+	 * Toggle byte order swapping. This is needed in rare cases when we already know that the data
+	 * swapping has already occurred or if we know that it will be handled later.
+	 *
+	 * @param Enabled	set to true to enable byte order swapping
+	 */
+	void SetByteSwapping(bool Enabled)
+	{
+		ArForceByteSwapping = Enabled;
+	}
+
+	/**
+	 * Sets the archive's property serialization modifier flags
+	 *
+	 * @param	InPortFlags		the new flags to use for property serialization
+	 */
+	void SetPortFlags(uint32 InPortFlags)
+	{
+		ArPortFlags = InPortFlags;
+	}
+
+	/**
+	 * Sets the archives custom serialization modifier flags (nothing to do with PortFlags or Custom versions)
+	 *
+	 * @param InCustomFlags the new flags to use for custom serialization
+	 */
+	void SetDebugSerializationFlags(uint32 InCustomFlags)
+	{
+#if WITH_EDITOR
+		ArDebugSerializationFlags = InCustomFlags;
+#endif
+	}
+
+	/**
+	 * Indicates whether this archive is filtering editor-only on save or contains data that had editor-only content stripped.
+	 *
+	 * @return true if the archive filters editor-only content, false otherwise.
+	 */
+	bool IsFilterEditorOnly()
+	{
+		return ArIsFilterEditorOnly;
+	}
+
+	/**
+	 * Sets a flag indicating that this archive needs to filter editor-only content.
+	 *
+	 * @param InFilterEditorOnly Whether to filter editor-only content.
+	 */
+	virtual void SetFilterEditorOnly(bool InFilterEditorOnly)
+	{
+		ArIsFilterEditorOnly = InFilterEditorOnly;
+	}
+
+	/**
+	 * Indicates whether this archive is saving or loading game state
+	 *
+	 * @note This is intended for game-specific archives and is not true for any of the build in save methods
+	 * @return true if the archive is dealing with save games, false otherwise.
+	 */
+	bool IsSaveGame()
+	{
+		return ArIsSaveGame;
+	}
+
+	/**
+	 * Whether or not this archive is serializing data being sent/received by the netcode
+	 */
+	FORCEINLINE bool IsNetArchive()
+	{
+		return ArIsNetArchive;
+	}
+
+	/**
+	 * Checks whether the archive is used for cooking.
+	 *
+	 * @return true if the archive is used for cooking, false otherwise.
+	 */
+	FORCEINLINE bool IsCooking() const
+	{
+		check(!CookingTargetPlatform || (!IsLoading() && !IsTransacting() && IsSaving()));
+
+		return !!CookingTargetPlatform;
+	}
+
+	/**
+	 * Returns the cooking target platform.
+	 *
+	 * @return Target platform.
+	 */
+	FORCEINLINE const ITargetPlatform* CookingTarget() const
+	{
+		return CookingTargetPlatform;
+	}
+
+	/**
+	 * Sets the cooking target platform.
+	 *
+	 * @param InCookingTarget The target platform to set.
+	 */
+	FORCEINLINE void SetCookingTarget(const ITargetPlatform* InCookingTarget)
+	{
+		CookingTargetPlatform = InCookingTarget;
+	}
+
+	/**
+	 * Checks whether the archive is used to resolve out-of-date enum indexes
+	 * If function returns true, the archive should be called only for objects containing user defined enum
+	 *
+	 * @return true if the archive is used to resolve out-of-date enum indexes
+	 */
+	virtual bool UseToResolveEnumerators() const
+	{
+		return false;
+	}
+
+	/**
+	 * Checks whether the archive wants to skip the property independent of the other flags
+	 */
+	virtual bool ShouldSkipProperty(const UProperty* InProperty) const
+	{
+		return false;
+	}
+
+	/**
+	 * Overrides the property that is currently being serialized
+	 * @note: You likely want to call PushSerializedProperty/PopSerializedProperty instead
+	 *
+	 * @param InProperty Pointer to the property that is currently being serialized
+	 */
+	FORCEINLINE void SetSerializedProperty(UProperty* InProperty)
+	{
+		SerializedProperty = InProperty;
+	}
+
+	/**
+	 * Gets the property that is currently being serialized
+	 *
+	 * @return Pointer to the property that is currently being serialized
+	 */
+	FORCEINLINE class UProperty* GetSerializedProperty() const
+	{
+		return SerializedProperty;
+	}
+
+	/**
+	 * Gets the chain of properties that are currently being serialized
+	 * @note This populates the array in stack order, so the 0th entry in the array is the top of the stack of properties
+	 */
+	void GetSerializedPropertyChain(TArray<class UProperty*>& OutProperties) const;
+
+	/**
+	 * Get the raw serialized property chain for this archive
+	 * @note Accessing this directly can avoid an array allocation depending on your use-case
+	 */
+	FORCEINLINE const FArchiveSerializedPropertyChain* GetSerializedPropertyChain() const
+	{
+		return SerializedPropertyChain;
+	}
+
+	/**
+	 * Set the raw serialized property chain for this archive, optionally overriding the serialized property too (or null to use the head of the property chain)
+	 */
+	void SetSerializedPropertyChain(const FArchiveSerializedPropertyChain* InSerializedPropertyChain, class UProperty* InSerializedPropertyOverride = nullptr);
+
+#if WITH_EDITORONLY_DATA
+	/** Returns true if the stack of currently serialized properties contains an editor-only property */
+	virtual bool IsEditorOnlyPropertyOnTheStack() const;
+#endif
+
+	/* Sets the current UObject serialization context for this archive */
+	virtual void SetSerializeContext(FUObjectSerializeContext* InLoadContext) {}
+
+	/* Gets the current UObject serialization context for this archive */
+	virtual FUObjectSerializeContext* GetSerializeContext() { return nullptr; }
+
+#if USE_STABLE_LOCALIZATION_KEYS
+	/**
+	 * Set the localization namespace that this archive should use when serializing text properties.
+	 * This is typically the namespace used by the package being serialized (if serializing a package, or an object within a package).
+	 */
+	virtual void SetLocalizationNamespace(const FString& InLocalizationNamespace);
+
+	/**
+	 * Get the localization namespace that this archive should use when serializing text properties.
+	 * This is typically the namespace used by the package being serialized (if serializing a package, or an object within a package).
+	 */
+	virtual FString GetLocalizationNamespace() const;
+#endif // USE_STABLE_LOCALIZATION_KEYS
+
+	/** Resets all of the base archive members. */
+	virtual void Reset();
+
+public:
+#if DEVIRTUALIZE_FLinkerLoad_Serialize
+	/* These are used for fastpath inline serializers  */
+	struct FFastPathLoadBuffer
+	{
+		const uint8* StartFastPathLoadBuffer;
+		const uint8* EndFastPathLoadBuffer;
+		const uint8* OriginalFastPathLoadBuffer;
+		FORCEINLINE FFastPathLoadBuffer()
+		{
+			Reset();
+		}
+		FORCEINLINE void Reset()
+		{
+			StartFastPathLoadBuffer = nullptr;
+			EndFastPathLoadBuffer = nullptr;
+			OriginalFastPathLoadBuffer = nullptr;
+		}
+	};
+	//@todoio FArchive is really a horrible class and the way it is proxied by FLinkerLoad is double terrible. It makes the fast path really hacky and slower than it would need to be.
+	FFastPathLoadBuffer* ActiveFPLB;
+	FFastPathLoadBuffer InlineFPLB;
+#endif
+
+// These will be private in FArchive
+protected:
+	/** Copies all of the members except CustomVersionContainer */
+	void CopyTrivialFArchiveStatusMembers(const FArchiveState& ArchiveStatusToCopy);
+
+	/** Whether this archive is for loading data. */
+	uint8 ArIsLoading : 1;
+
+	/** Whether this archive is for saving data. */
+	uint8 ArIsSaving : 1;
+
+	/** Whether archive is transacting. */
+	uint8 ArIsTransacting : 1;
+
+	/** Whether this archive serializes to a text format. Text format archives should use high level constructs from FStructuredArchive for delimiting data rather than manually seeking through the file. */
+	uint8 ArIsTextFormat : 1;
+
+	/** Whether this archive wants properties to be serialized in binary form instead of tagged. */
+	uint8 ArWantBinaryPropertySerialization : 1;
+
+	/** Whether this archive wants to always save strings in unicode format */
+	uint8 ArForceUnicode : 1;
+
+	/** Whether this archive saves to persistent storage. */
+	uint8 ArIsPersistent : 1;
+
+public:
+	/** Whether this archive contains errors. */
+	uint8 ArIsError : 1;
+
+	/** Whether this archive contains critical errors. */
+	uint8 ArIsCriticalError : 1;
+
+	/** Quickly tell if an archive contains script code. */
+	uint8 ArContainsCode : 1;
+
+	/** Used to determine whether FArchive contains a level or world. */
+	uint8 ArContainsMap : 1;
+
+	/** Used to determine whether FArchive contains data required to be gathered for localization. */
+	uint8 ArRequiresLocalizationGather : 1;
+
+	/** Whether we should forcefully swap bytes. */
+	uint8 ArForceByteSwapping : 1;
+
+	/** If true, we will not serialize the ObjectArchetype reference in UObject. */
+	uint8 ArIgnoreArchetypeRef : 1;
+
+	/** If true, do not perform delta serialization of properties. */
+	uint8 ArNoDelta : 1;
+
+	/** If true, do not perform delta serialization within properties (e.g. TMaps and TSets). */
+	uint8 ArNoIntraPropertyDelta : 1;
+
+	/** If true, we will not serialize the Outer reference in UObject. */
+	uint8 ArIgnoreOuterRef : 1;
+
+	/** If true, we will not serialize ClassGeneratedBy reference in UClass. */
+	uint8 ArIgnoreClassGeneratedByRef : 1;
+
+	/** If true, UObject::Serialize will skip serialization of the Class property. */
+	uint8 ArIgnoreClassRef : 1;
+
+	/** Whether to allow lazy loading. */
+	uint8 ArAllowLazyLoading : 1;
+
+	/** Whether this archive only cares about serializing object references. */
+	uint8 ArIsObjectReferenceCollector : 1;
+
+	/** Whether a reference collector is modifying the references and wants both weak and strong ones */
+	uint8 ArIsModifyingWeakAndStrongReferences : 1;
+
+	/** Whether this archive is counting memory and therefore wants e.g. TMaps to be serialized. */
+	uint8 ArIsCountingMemory : 1;
+
+	/** Whether bulk data serialization should be skipped or not. */
+	uint8 ArShouldSkipBulkData : 1;
+
+	/** Whether editor only properties are being filtered from the archive (or has been filtered). */
+	uint8 ArIsFilterEditorOnly : 1;
+
+	/** Whether this archive is saving/loading game state */
+	uint8 ArIsSaveGame : 1;
+
+	/** Whether or not this archive is sending/receiving network data */
+	uint8 ArIsNetArchive : 1;
+
+	/** Set TRUE to use the custom property list attribute for serialization. */
+	uint8 ArUseCustomPropertyList : 1;
+
+	/** Whether we are currently serializing defaults. > 0 means yes, <= 0 means no. */
+	int32 ArSerializingDefaults;
+
+	/** Modifier flags that be used when serializing UProperties */
+	uint32 ArPortFlags;
+
+	/** Max size of data that this archive is allowed to serialize. */
+	int64 ArMaxSerializeSize;
+
+	/**
+	 * Sets whether this archive is for loading data.
+	 *
+	 * @param bInIsLoading  true if this archive is for loading, false otherwise.
+	 */
+	virtual void SetIsLoading(bool bInIsLoading);
+
+	/**
+	 * Sets whether this archive is for saving data.
+	 *
+	 * @param bInIsSaving  true if this archive is for saving, false otherwise.
+	 */
+	virtual void SetIsSaving(bool bInIsSaving);
+
+	/**
+	 * Sets whether this archive is for transacting.
+	 *
+	 * @param bInIsTransacting  true if this archive is for transacting, false otherwise.
+	 */
+	virtual void SetIsTransacting(bool bInIsTransacting);
+
+	/**
+	 * Sets whether this archive is in text format.
+	 *
+	 * @param bInIsTextFormat  true if this archive is in text format, false otherwise.
+	 */
+	virtual void SetIsTextFormat(bool bInIsTextFormat);
+
+	/**
+	 * Sets whether this archive wants binary property serialization.
+	 *
+	 * @param bInWantBinaryPropertySerialization  true if this archive wants binary serialization, false otherwise.
+	 */
+	virtual void SetWantBinaryPropertySerialization(bool bInWantBinaryPropertySerialization);
+
+	/**
+	 * Sets whether this archive wants to force saving as Unicode.
+	 * This is needed when we need to make sure ANSI strings are saved as Unicode.
+	 *
+	 * @param bInForceUnicode  true if this archive wants to force saving as Unicode, false otherwise.
+	 */
+	virtual void SetForceUnicode(bool bInForceUnicode);
+
+	/**
+	 * Sets whether this archive is to persistent storage.
+	 *
+	 * @param bInIsPersistent  true if this archive is to persistent storage, false otherwise.
+	 */
+	virtual void SetIsPersistent(bool bInIsPersistent);
+
+	/**
+	 * Sets the archive version number. Used by the code that makes sure that FLinkerLoad's 
+	 * internal archive versions match the file reader it creates.
+	 *
+	 * @param UE4Ver	new version number
+	 */
+	virtual void SetUE4Ver(int32 InVer);
+
+	/**
+	 * Sets the archive licensee version number. Used by the code that makes sure that FLinkerLoad's 
+	 * internal archive versions match the file reader it creates.
+	 *
+	 * @param Ver	new version number
+	 */
+	virtual void SetLicenseeUE4Ver(int32 InVer);
+
+	/**
+	 * Sets the archive engine version. Used by the code that makes sure that FLinkerLoad's
+	 * internal archive versions match the file reader it creates.
+	 *
+	 * @param InVer	new version number
+	 */
+	virtual void SetEngineVer(const FEngineVersionBase& InVer);
+
+	/**
+	 * Sets the archive engine network version.
+	 */
+	virtual void SetEngineNetVer(const uint32 InEngineNetVer);
+
+	/**
+	 * Sets the archive game network version.
+	 */
+	virtual void SetGameNetVer(const uint32 InGameNetVer);
+
+// These will be private in FArchive
+protected:
+	/** Holds the archive version. */
+	int32 ArUE4Ver;
+
+	/** Holds the archive version for licensees. */
+	int32 ArLicenseeUE4Ver;
+
+	/** Holds the engine version. */
+	FEngineVersionBase ArEngineVer;
+
+	/** Holds the engine network protocol version. */
+	uint32 ArEngineNetVer;
+
+	/** Holds the game network protocol version. */
+	uint32 ArGameNetVer;
+
+	/**
+	* All the custom versions stored in the archive.
+	* Stored as a pointer to a heap-allocated object because of a 3-way dependency between TArray, FCustomVersionContainer and FArchive, which is too much work to change right now.
+	* Keeping it as a heap-allocated object also helps with performance in some cases as we don't need to construct it for archives that don't care about custom versions.
+	*/
+	mutable FCustomVersionContainer* CustomVersionContainer = nullptr;
+
+public:
+	/** Custom property list attribute. If the flag below is set, only these properties will be iterated during serialization. If NULL, then no properties will be iterated. */
+	const struct FCustomPropertyListNode* ArCustomPropertyList;
+
+#if WITH_EDITOR
+	/** Custom serialization modifier flags can be used for anything */
+	uint32 ArDebugSerializationFlags;
+#endif
+
+// These will be private in FArchive
+protected:
+	/** Holds the cooking target platform. */
+	const ITargetPlatform* CookingTargetPlatform;
+
+	/** Holds the pointer to the property that is currently being serialized */
+	UProperty* SerializedProperty;
+
+	/** Holds the chain of properties that are currently being serialized */
+	FArchiveSerializedPropertyChain* SerializedPropertyChain;
+
+#if USE_STABLE_LOCALIZATION_KEYS
+	/**
+	 * The localization namespace that this archive should use when serializing text properties.
+	 * This is typically the namespace used by the package being serialized (if serializing a package, or an object within a package).
+	 * Stored as a pointer to a heap-allocated string because of a dependency between TArray (thus FString) and FArchive; null should be treated as an empty string.
+	 */
+	FString* LocalizationNamespacePtr;
+
+	/** See GetLocalizationNamespace */
+	FString GetBaseLocalizationNamespace() const;
+
+	/** See SetLocalizationNamespace */
+	void SetBaseLocalizationNamespace(const FString& InLocalizationNamespace);
+#endif // USE_STABLE_LOCALIZATION_KEYS
+
+	/**
+	 * Indicates if the custom versions container is in a 'reset' state.  This will be used to defer the choice about how to
+	 * populate the container until it is needed, where the read/write state will be known.
+	 */
+	mutable bool bCustomVersionsAreReset;
+
+};
+
 /**
  * TCheckedObjPtr
  *
@@ -133,25 +924,13 @@
  * Base class for archives that can be used for loading, saving, and garbage
  * collecting in a byte order neutral way.
  */
-class CORE_API FArchive
+class CORE_API FArchive : private FArchiveState
 {
 public:
-
-	/** Default constructor. */
-	FArchive();
-
-	/** Copy constructor. */
-	FArchive(const FArchive&);
-
-	/**
-	 * Copy assignment operator.
-	 *
-	 * @param ArchiveToCopy The archive to copy from.
-	 */
-	FArchive& operator=(const FArchive& ArchiveToCopy);
-
-	/** Destructor. */
-	virtual ~FArchive();
+	FArchive() = default;
+	FArchive(const FArchive&) = default;
+	FArchive& operator=(const FArchive& ArchiveToCopy) = default;
+	~FArchive() = default;
 
 public:
 
@@ -610,43 +1389,13 @@
 
 	virtual void Preload(UObject* Object) { }
 
-	virtual void CountBytes(SIZE_T InNum, SIZE_T InMax) { }
+	using FArchiveState::CountBytes;
+	using FArchiveState::GetArchiveName;
+	using FArchiveState::GetLinker;
+	using FArchiveState::Tell;
+	using FArchiveState::TotalSize;
+	using FArchiveState::AtEnd;
 
-	/**
-  	 * Returns the name of the Archive.  Useful for getting the name of the package a struct or object
-	 * is in when a loading error occurs.
-	 *
-	 * This is overridden for the specific Archive Types
-	 */
-	virtual FString GetArchiveName() const;
-
-	/**
-	 * If this archive is a FLinkerLoad or FLinkerSave, returns a pointer to the ULinker portion.
-	 *
-	 * @return The linker, or nullptr if the archive is not a linker.
-	 */
-	virtual FLinker* GetLinker()
-	{
-		return nullptr;
-	}
-
-	virtual int64 Tell()
-	{
-		return INDEX_NONE;
-	}
-
-	virtual int64 TotalSize()
-	{
-		return INDEX_NONE;
-	}
-
-	virtual bool AtEnd()
-	{
-		int64 Pos = Tell();
-
-		return ((Pos != INDEX_NONE) && (Pos >= TotalSize()));
-	}
-
 	virtual void Seek(int64 InPos) { }
 
 	/**
@@ -726,15 +1475,8 @@
 		return !ArIsError;
 	}
 
-	virtual bool GetError()
-	{
-		return ArIsError;
-	}
-
-	void SetError() 
-	{ 
-		ArIsError = true; 
-	}
+	using FArchiveState::GetError;
+	using FArchiveState::SetError;
 
 	/**
 	 * Serializes and compresses/ uncompresses data. This is a shared helper function for compression
@@ -750,15 +1492,7 @@
 
 
 
-	FORCEINLINE bool IsByteSwapping()
-	{
-#if PLATFORM_LITTLE_ENDIAN
-		bool SwapBytes = ArForceByteSwapping;
-#else
-		bool SwapBytes = this->IsPersistent();
-#endif
-		return SwapBytes;
-	}
+	using FArchiveState::IsByteSwapping;
 
 	// Used to do byte swapping on small items. This does not happen usually, so we don't want it inline
 	void ByteSwap(void* V, int32 Length);
@@ -774,24 +1508,10 @@
 		return *this;
 	}
 
-	/** Sets a flag indicating that this archive contains code. */
-	void ThisContainsCode()
-	{
-		ArContainsCode = true;
-	}
+	using FArchiveState::ThisContainsCode;
+	using FArchiveState::ThisContainsMap;
+	using FArchiveState::ThisRequiresLocalizationGather;
 
-	/** Sets a flag indicating that this archive contains a ULevel or UWorld object. */
-	void ThisContainsMap() 
-	{
-		ArContainsMap = true;
-	}
-
-	/** Sets a flag indicating that this archive contains data required to be gathered for localization. */
-	void ThisRequiresLocalizationGather()
-	{
-		ArRequiresLocalizationGather = true;
-	}
-
 	/** Sets a flag indicating that this archive is currently serializing class/struct defaults. */
 	void StartSerializingDefaults() 
 	{
@@ -819,13 +1539,7 @@
 	 */
 	virtual void MarkSearchableName(const UObject* TypeObject, const FName& ValueName) const { }
 
-	/**
-	* Called to retrieve the archetype from the event driven loader. If this returns null, then call GetArchetype yourself.
-	*/
-	virtual UObject* GetArchetypeFromLoader(const UObject* Obj)
-	{
-		return nullptr;
-	}
+	using FArchiveState::GetArchetypeFromLoader;
 
 private:
 	void VARARGS LogfImpl(const TCHAR* Fmt, ...);
@@ -841,30 +1555,11 @@
 		LogfImpl(Fmt, Args...);
 	}
 
-	FORCEINLINE int32 UE4Ver() const
-	{
-		return ArUE4Ver;
-	}
-	
-	FORCEINLINE int32 LicenseeUE4Ver() const
-	{
-		return ArLicenseeUE4Ver;
-	}
-
-	FORCEINLINE FEngineVersionBase EngineVer() const
-	{
-		return ArEngineVer;
-	}
-
-	FORCEINLINE uint32 EngineNetVer() const
-	{
-		return ArEngineNetVer;
-	}
-
-	FORCEINLINE uint32 GameNetVer() const
-	{
-		return ArGameNetVer;
-	}
+	using FArchiveState::UE4Ver;
+	using FArchiveState::LicenseeUE4Ver;
+	using FArchiveState::EngineVer;
+	using FArchiveState::EngineNetVer;
+	using FArchiveState::GameNetVer;
 
 	/**
 	 * Registers the custom version to the archive.  This is used to inform the archive that custom version information is about to be stored.
@@ -874,13 +1569,7 @@
 	 */
 	virtual void UsingCustomVersion(const struct FGuid& Guid);
 
-	/**
-	 * Queries a custom version from the archive.  If the archive is being used to write, the custom version must have already been registered.
-	 *
-	 * @param Key The guid of the custom version to query.
-	 * @return The version number, or 0 if the custom tag isn't stored in the archive.
-	 */
-	int32 CustomVer(const struct FGuid& Key) const;
+	using FArchiveState::CustomVer;
 
 	/**
 	 * Returns a pointer to an archive that represents the same data that the current archive covers, but that can be cached and reused later
@@ -893,357 +1582,57 @@
 		return this;
 	}
 
-	FORCEINLINE bool IsLoading() const
-	{
-		return ArIsLoading;
-	}
-
-	FORCEINLINE bool IsSaving() const
-	{
-		return ArIsSaving;
-	}
-
-	FORCEINLINE bool IsTransacting() const
-	{
-		if (FPlatformProperties::HasEditorOnlyData())
-		{
-			return ArIsTransacting;
-		}
-		else
-		{
-			return false;
-		}
-	}
-
-	FORCEINLINE bool IsTextFormat() const
-	{
-		return (ArIsTextFormat && WITH_TEXT_ARCHIVE_SUPPORT);
-	}
-
-	FORCEINLINE bool WantBinaryPropertySerialization() const
-	{
-		return ArWantBinaryPropertySerialization;
-	}
-
-	FORCEINLINE bool IsForcingUnicode() const
-	{
-		return ArForceUnicode;
-	}
-
-	FORCEINLINE bool IsPersistent() const
-	{
-		return ArIsPersistent;
-	}
-
-	FORCEINLINE bool IsError() const
-	{
-		return ArIsError;
-	}
-
-	FORCEINLINE bool IsCriticalError() const
-	{
-		return ArIsCriticalError;
-	}
-
-	FORCEINLINE bool ContainsCode() const
-	{
-		return ArContainsCode;
-	}
-
-	FORCEINLINE bool ContainsMap() const
-	{
-		return ArContainsMap;
-	}
-
-	FORCEINLINE bool RequiresLocalizationGather() const
-	{
-		return ArRequiresLocalizationGather;
-	}
-
-	FORCEINLINE bool ForceByteSwapping() const
-	{
-		return ArForceByteSwapping;
-	}
-
-	FORCEINLINE bool IsSerializingDefaults() const
-	{
-		return (ArSerializingDefaults > 0) ? true : false;
-	}
-
-	FORCEINLINE bool IsIgnoringArchetypeRef() const
-	{
-		return ArIgnoreArchetypeRef;
-	}
-
-	FORCEINLINE bool DoDelta() const
-	{
-		return !ArNoDelta;
-	}
-
-	FORCEINLINE bool DoIntraPropertyDelta() const
-	{
-		return !ArNoIntraPropertyDelta;
-	}
-
-	FORCEINLINE bool IsIgnoringOuterRef() const
-	{
-		return ArIgnoreOuterRef;
-	}
-
-	FORCEINLINE bool IsIgnoringClassGeneratedByRef() const
-	{
-		return ArIgnoreClassGeneratedByRef;
-	}
-
-	FORCEINLINE bool IsIgnoringClassRef() const
-	{
-		return ArIgnoreClassRef;
-	}
-
-	FORCEINLINE bool IsAllowingLazyLoading() const
-	{
-		return ArAllowLazyLoading;
-	}
-
-	FORCEINLINE bool IsObjectReferenceCollector() const
-	{
-		return ArIsObjectReferenceCollector;
-	}
-
-	FORCEINLINE bool IsModifyingWeakAndStrongReferences() const
-	{
-		return ArIsModifyingWeakAndStrongReferences;
-	}
-
-	FORCEINLINE bool IsCountingMemory() const
-	{
-		return ArIsCountingMemory;
-	}
-
-	FORCEINLINE uint32 GetPortFlags() const
-	{
-		return ArPortFlags;
-	}
-
-	FORCEINLINE bool HasAnyPortFlags(uint32 Flags) const
-	{
-		return ((ArPortFlags & Flags) != 0);
-	}
-
-	FORCEINLINE bool HasAllPortFlags(uint32 Flags) const
-	{
-		return ((ArPortFlags & Flags) == Flags);
-	}
-
-	FORCEINLINE uint32 GetDebugSerializationFlags() const
-	{
-#if WITH_EDITOR
-		return ArDebugSerializationFlags;
-#else
-		return 0;
-#endif
-	}
-
-	FORCEINLINE bool ShouldSkipBulkData() const
-	{
-		return ArShouldSkipBulkData;
-	}
-
-	FORCEINLINE int64 GetMaxSerializeSize() const
-	{
-		return ArMaxSerializeSize;
-	}
+	using FArchiveState::IsLoading;
+	using FArchiveState::IsSaving;
+	using FArchiveState::IsTransacting;
+	using FArchiveState::IsTextFormat;
+	using FArchiveState::WantBinaryPropertySerialization;
+	using FArchiveState::IsForcingUnicode;
+	using FArchiveState::IsPersistent;
+	using FArchiveState::IsError;
+	using FArchiveState::IsCriticalError;
+	using FArchiveState::ContainsCode;
+	using FArchiveState::ContainsMap;
+	using FArchiveState::RequiresLocalizationGather;
+	using FArchiveState::ForceByteSwapping;
+	using FArchiveState::IsSerializingDefaults;
+	using FArchiveState::IsIgnoringArchetypeRef;
+	using FArchiveState::DoDelta;
+	using FArchiveState::DoIntraPropertyDelta;
+	using FArchiveState::IsIgnoringOuterRef;
+	using FArchiveState::IsIgnoringClassGeneratedByRef;
+	using FArchiveState::IsIgnoringClassRef;
+	using FArchiveState::IsAllowingLazyLoading;
+	using FArchiveState::IsObjectReferenceCollector;
+	using FArchiveState::IsModifyingWeakAndStrongReferences;
+	using FArchiveState::IsCountingMemory;
+	using FArchiveState::GetPortFlags;
+	using FArchiveState::HasAnyPortFlags;
+	using FArchiveState::HasAllPortFlags;
+	using FArchiveState::GetDebugSerializationFlags;
+	using FArchiveState::ShouldSkipBulkData;
+	using FArchiveState::GetMaxSerializeSize;
+	using FArchiveState::GetCustomVersions;
+	using FArchiveState::SetCustomVersions;
+	using FArchiveState::ResetCustomVersions;
+	using FArchiveState::SetCustomVersion;
+	using FArchiveState::SetByteSwapping;
+	using FArchiveState::SetPortFlags;
+	using FArchiveState::SetDebugSerializationFlags;
+	using FArchiveState::IsFilterEditorOnly;
+	using FArchiveState::SetFilterEditorOnly;
+	using FArchiveState::IsSaveGame;
+	using FArchiveState::IsNetArchive;
+	using FArchiveState::IsCooking;
+	using FArchiveState::CookingTarget;
+	using FArchiveState::SetCookingTarget;
+	using FArchiveState::UseToResolveEnumerators;
+	using FArchiveState::ShouldSkipProperty;
+	using FArchiveState::SetSerializedProperty;
+	using FArchiveState::GetSerializedProperty;
+	using FArchiveState::GetSerializedPropertyChain;
+	using FArchiveState::SetSerializedPropertyChain;
 
 	/**
-	 * Gets the custom version numbers for this archive.
-	 *
-	 * @return The container of custom versions in the archive.
-	 */
-	virtual const FCustomVersionContainer& GetCustomVersions() const;
-
-	/**
-	 * Sets the custom version numbers for this archive.
-	 *
-	 * @param CustomVersionContainer - The container of custom versions to copy into the archive.
-	 */
-	virtual void SetCustomVersions(const FCustomVersionContainer& CustomVersionContainer);
-
-	/** Resets the custom version numbers for this archive. */
-	virtual void ResetCustomVersions();
-
-	/**
-	 * Sets a specific custom version
-	 *
-	 * @param Key - The guid of the custom version to query.
-	 * @param Version - The version number to set key to
-	 * @param FriendlyName - Friendly name corresponding to the key
-	 */
-	void SetCustomVersion(const struct FGuid& Key, int32 Version, FName FriendlyName);
-
-	/**
-	 * Toggle byte order swapping. This is needed in rare cases when we already know that the data
-	 * swapping has already occurred or if we know that it will be handled later.
-	 *
-	 * @param Enabled	set to true to enable byte order swapping
-	 */
-	void SetByteSwapping(bool Enabled)
-	{
-		ArForceByteSwapping = Enabled;
-	}
-
-	/**
-	 * Sets the archive's property serialization modifier flags
-	 *
-	 * @param	InPortFlags		the new flags to use for property serialization
-	 */
-	void SetPortFlags(uint32 InPortFlags)
-	{
-		ArPortFlags = InPortFlags;
-	}
-
-	/**
-	 * Sets the archives custom serialization modifier flags (nothing to do with PortFlags or Custom versions)
-	 * 
-	 * @param InCustomFlags the new flags to use for custom serialization
-	 */
-	void SetDebugSerializationFlags(uint32 InCustomFlags)
-	{
-#if WITH_EDITOR
-		ArDebugSerializationFlags = InCustomFlags;
-#endif
-	}
-
-	/**
-	 * Indicates whether this archive is filtering editor-only on save or contains data that had editor-only content stripped.
-	 *
-	 * @return true if the archive filters editor-only content, false otherwise.
-	 */
-	bool IsFilterEditorOnly()
-	{
-		return ArIsFilterEditorOnly;
-	}
-
-	/**
-	 * Sets a flag indicating that this archive needs to filter editor-only content.
-	 *
-	 * @param InFilterEditorOnly Whether to filter editor-only content.
-	 */
-	virtual void SetFilterEditorOnly(bool InFilterEditorOnly)
-	{
-		ArIsFilterEditorOnly = InFilterEditorOnly;
-	}
-
-	/**
-	 * Indicates whether this archive is saving or loading game state
-	 *
-	 * @note This is intended for game-specific archives and is not true for any of the build in save methods
-	 * @return true if the archive is dealing with save games, false otherwise.
-	 */
-	bool IsSaveGame()
-	{
-		return ArIsSaveGame;
-	}
-
-	/**
-	 * Whether or not this archive is serializing data being sent/received by the netcode
-	 */
-	FORCEINLINE bool IsNetArchive()
-	{
-		return ArIsNetArchive;
-	}
-
-	/**
-	 * Checks whether the archive is used for cooking.
-	 *
-	 * @return true if the archive is used for cooking, false otherwise.
-	 */
-	FORCEINLINE bool IsCooking() const	
-	{
-		check(!CookingTargetPlatform || (!IsLoading() && !IsTransacting() && IsSaving()));
-
-		return !!CookingTargetPlatform;
-	}
-
-	/**
-	 * Returns the cooking target platform.
-	 *
-	 * @return Target platform.
-	 */
-	FORCEINLINE const ITargetPlatform* CookingTarget() const 
-	{
-		return CookingTargetPlatform;
-	}
-
-	/**
-	 * Sets the cooking target platform.
-	 *
-	 * @param InCookingTarget The target platform to set.
-	 */
-	FORCEINLINE void SetCookingTarget(const ITargetPlatform* InCookingTarget) 
-	{
-		CookingTargetPlatform = InCookingTarget;
-	}
-
-	/**
-	 * Checks whether the archive is used to resolve out-of-date enum indexes
-	 * If function returns true, the archive should be called only for objects containing user defined enum
-	 *
-	 * @return true if the archive is used to resolve out-of-date enum indexes
-	 */
-	virtual bool UseToResolveEnumerators() const
-	{
-		return false;
-	}
-
-	/**
-	 * Checks whether the archive wants to skip the property independent of the other flags
-	 */
-	virtual bool ShouldSkipProperty(const UProperty* InProperty) const
-	{
-		return false;
-	}
-
-	/**
-	 * Overrides the property that is currently being serialized
-	 * @note: You likely want to call PushSerializedProperty/PopSerializedProperty instead
-	 * 
-	 * @param InProperty Pointer to the property that is currently being serialized
-	 */
-	FORCEINLINE void SetSerializedProperty(UProperty* InProperty)
-	{
-		SerializedProperty = InProperty;
-	}
-
-	/**
-	 * Gets the property that is currently being serialized
-	 *
-	 * @return Pointer to the property that is currently being serialized
-	 */
-	FORCEINLINE class UProperty* GetSerializedProperty() const
-	{
-		return SerializedProperty;
-	}
-
-	/**
-	 * Gets the chain of properties that are currently being serialized
-	 * @note This populates the array in stack order, so the 0th entry in the array is the top of the stack of properties
-	 */
-	void GetSerializedPropertyChain(TArray<class UProperty*>& OutProperties) const;
-
-	/**
-	 * Get the raw serialized property chain for this archive
-	 * @note Accessing this directly can avoid an array allocation depending on your use-case
-	 */
-	FORCEINLINE const FArchiveSerializedPropertyChain* GetSerializedPropertyChain() const
-	{
-		return SerializedPropertyChain;
-	}
-
-	/**
-	 * Set the raw serialized property chain for this archive, optionally overriding the serialized property too (or null to use the head of the property chain)
-	 */
-	void SetSerializedPropertyChain(const FArchiveSerializedPropertyChain* InSerializedPropertyChain, class UProperty* InSerializedPropertyOverride = nullptr);
-
-	/**
 	 * Push a property that is currently being serialized onto the property stack
 	 * 
@@ -1261,15 +1650,11 @@
 	virtual void PopSerializedProperty(class UProperty* InProperty, const bool bIsEditorOnlyProperty);
 
 #if WITH_EDITORONLY_DATA
-	/** Returns true if the stack of currently serialized properties contains an editor-only property */
-	virtual bool IsEditorOnlyPropertyOnTheStack() const;
+	using FArchiveState::IsEditorOnlyPropertyOnTheStack;
 #endif
 
-	/* Sets the current UObject serialization context for this archive */
-	virtual void SetSerializeContext(FUObjectSerializeContext* InLoadContext) {}
+	using FArchiveState::SetSerializeContext;
+	using FArchiveState::GetSerializeContext;
 
-	/* Gets the current UObject serialization context for this archive */
-	virtual FUObjectSerializeContext* GetSerializeContext() { return nullptr; }
-
 	/** 
 	 * Adds external read dependency 
@@ -1279,21 +1664,12 @@
 	virtual bool AttachExternalReadDependency(FExternalReadCallback& ReadCallback) { return false; };
 
 #if USE_STABLE_LOCALIZATION_KEYS
-	/**
-	 * Set the localization namespace that this archive should use when serializing text properties.
-	 * This is typically the namespace used by the package being serialized (if serializing a package, or an object within a package).
-	 */
-	virtual void SetLocalizationNamespace(const FString& InLocalizationNamespace);
-
-	/**
-	 * Get the localization namespace that this archive should use when serializing text properties.
-	 * This is typically the namespace used by the package being serialized (if serializing a package, or an object within a package).
-	 */
-	virtual FString GetLocalizationNamespace() const;
+	using FArchiveState::SetLocalizationNamespace;
+	using FArchiveState::GetLocalizationNamespace;
 #endif // USE_STABLE_LOCALIZATION_KEYS
 
 	/** Resets all of the base archive members. */
-	virtual void Reset();
+	using FArchiveState::Reset;
 
 #if DEVIRTUALIZE_FLinkerLoad_Serialize
 private:
@@ -1333,26 +1709,9 @@
 	}
 
 public:
-	/* These are used for fastpath inline serializers  */
-	struct FFastPathLoadBuffer
-	{
-		const uint8* StartFastPathLoadBuffer;
-		const uint8* EndFastPathLoadBuffer;
-		const uint8* OriginalFastPathLoadBuffer;
-		FORCEINLINE FFastPathLoadBuffer()
-		{
-			Reset();
-		}
-		FORCEINLINE void Reset()
-		{
-			StartFastPathLoadBuffer = nullptr;
-			EndFastPathLoadBuffer = nullptr;
-			OriginalFastPathLoadBuffer = nullptr;
-		}
-	};
 	//@todoio FArchive is really a horrible class and the way it is proxied by FLinkerLoad is double terrible. It makes the fast path really hacky and slower than it would need to be.
-	FFastPathLoadBuffer* ActiveFPLB;
-	FFastPathLoadBuffer InlineFPLB;
+	using FArchiveState::ActiveFPLB;
+	using FArchiveState::InlineFPLB;
 
 #else
 	template<SIZE_T Size>
@@ -1362,215 +1721,73 @@
 	}
 #endif
 
+	const FArchiveState& GetArchiveState() const
+	{
+		return ImplicitConv<const FArchiveState&>(*this);
+	}
+
+	virtual void SetArchiveState(const FArchiveState& InState);
+
 private:
-	/** Copies all of the members except CustomVersionContainer */
-	void CopyTrivialFArchiveStatusMembers(const FArchive& ArchiveStatusToCopy);
-
-	/** Whether this archive is for loading data. */
-	uint8 ArIsLoading : 1;
-
-	/** Whether this archive is for saving data. */
-	uint8 ArIsSaving : 1;
-
-	/** Whether archive is transacting. */
-	uint8 ArIsTransacting : 1;
-
-	/** Whether this archive serializes to a text format. Text format archives should use high level constructs from FStructuredArchive for delimiting data rather than manually seeking through the file. */
-	uint8 ArIsTextFormat : 1;
-
-	/** Whether this archive wants properties to be serialized in binary form instead of tagged. */
-	uint8 ArWantBinaryPropertySerialization : 1;
-
-	/** Whether this archive wants to always save strings in unicode format */
-	uint8 ArForceUnicode : 1;
-
-	/** Whether this archive saves to persistent storage. */
-	uint8 ArIsPersistent : 1;
+	using FArchiveState::ArIsLoading;
+	using FArchiveState::ArIsSaving;
+	using FArchiveState::ArIsTransacting;
+	using FArchiveState::ArIsTextFormat;
+	using FArchiveState::ArWantBinaryPropertySerialization;
+	using FArchiveState::ArForceUnicode;
+	using FArchiveState::ArIsPersistent;
 
 public:
+	using FArchiveState::ArIsError;
+	using FArchiveState::ArIsCriticalError;
+	using FArchiveState::ArContainsCode;
+	using FArchiveState::ArContainsMap;
+	using FArchiveState::ArRequiresLocalizationGather;
+	using FArchiveState::ArForceByteSwapping;
+	using FArchiveState::ArIgnoreArchetypeRef;
+	using FArchiveState::ArNoDelta;
+	using FArchiveState::ArNoIntraPropertyDelta;
+	using FArchiveState::ArIgnoreOuterRef;
+	using FArchiveState::ArIgnoreClassGeneratedByRef;
+	using FArchiveState::ArIgnoreClassRef;
+	using FArchiveState::ArAllowLazyLoading;
+	using FArchiveState::ArIsObjectReferenceCollector;
+	using FArchiveState::ArIsModifyingWeakAndStrongReferences;
+	using FArchiveState::ArIsCountingMemory;
+	using FArchiveState::ArShouldSkipBulkData;
+	using FArchiveState::ArIsFilterEditorOnly;
+	using FArchiveState::ArIsSaveGame;
+	using FArchiveState::ArIsNetArchive;
+	using FArchiveState::ArUseCustomPropertyList;
+	using FArchiveState::ArSerializingDefaults;
+	using FArchiveState::ArPortFlags;
+	using FArchiveState::ArMaxSerializeSize;
 
-	/** Whether this archive contains errors. */
-	uint8 ArIsError : 1;
-
-	/** Whether this archive contains critical errors. */
-	uint8 ArIsCriticalError : 1;
-
-	/** Quickly tell if an archive contains script code. */
-	uint8 ArContainsCode : 1;
-
-	/** Used to determine whether FArchive contains a level or world. */
-	uint8 ArContainsMap : 1;
-
-	/** Used to determine whether FArchive contains data required to be gathered for localization. */
-	uint8 ArRequiresLocalizationGather : 1;
-
-	/** Whether we should forcefully swap bytes. */
-	uint8 ArForceByteSwapping : 1;
-
-	/** If true, we will not serialize the ObjectArchetype reference in UObject. */
-	uint8 ArIgnoreArchetypeRef : 1;
-
-	/** If true, do not perform delta serialization of properties. */
-	uint8 ArNoDelta : 1;
-
-	/** If true, do not perform delta serialization within properties (e.g. TMaps and TSets). */
-	uint8 ArNoIntraPropertyDelta : 1;
-
-	/** If true, we will not serialize the Outer reference in UObject. */
-	uint8 ArIgnoreOuterRef : 1;
-
-	/** If true, we will not serialize ClassGeneratedBy reference in UClass. */
-	uint8 ArIgnoreClassGeneratedByRef : 1;
-
-	/** If true, UObject::Serialize will skip serialization of the Class property. */
-	uint8 ArIgnoreClassRef : 1;
-
-	/** Whether to allow lazy loading. */
-	uint8 ArAllowLazyLoading : 1;
-
-	/** Whether this archive only cares about serializing object references. */
-	uint8 ArIsObjectReferenceCollector : 1;
-
-	/** Whether a reference collector is modifying the references and wants both weak and strong ones */
-	uint8 ArIsModifyingWeakAndStrongReferences : 1;
-
-	/** Whether this archive is counting memory and therefore wants e.g. TMaps to be serialized. */
-	uint8 ArIsCountingMemory : 1;
-
-	/** Whether bulk data serialization should be skipped or not. */
-	uint8 ArShouldSkipBulkData : 1;
-
-	/** Whether editor only properties are being filtered from the archive (or has been filtered). */
-	uint8 ArIsFilterEditorOnly : 1;
-
-	/** Whether this archive is saving/loading game state */
-	uint8 ArIsSaveGame : 1;
-
-	/** Whether or not this archive is sending/receiving network data */
-	uint8 ArIsNetArchive : 1;
-
-	/** Set TRUE to use the custom property list attribute for serialization. */
-	uint8 ArUseCustomPropertyList : 1;
-
-	/** Whether we are currently serializing defaults. > 0 means yes, <= 0 means no. */
-	int32 ArSerializingDefaults;
-
-	/** Modifier flags that be used when serializing UProperties */
-	uint32 ArPortFlags;
-			
-	/** Max size of data that this archive is allowed to serialize. */
-	int64 ArMaxSerializeSize;
-
-	/**
-	 * Sets whether this archive is for loading data.
-	 *
-	 * @param bInIsLoading  true if this archive is for loading, false otherwise.
-	 */
-	virtual void SetIsLoading(bool bInIsLoading);
-
-	/**
-	 * Sets whether this archive is for saving data.
-	 *
-	 * @param bInIsSaving  true if this archive is for saving, false otherwise.
-	 */
-	virtual void SetIsSaving(bool bInIsSaving);
-
-	/**
-	 * Sets whether this archive is for transacting.
-	 *
-	 * @param bInIsTransacting  true if this archive is for transacting, false otherwise.
-	 */
-	virtual void SetIsTransacting(bool bInIsTransacting);
-
-	/**
-	 * Sets whether this archive is in text format.
-	 *
-	 * @param bInIsTextFormat  true if this archive is in text format, false otherwise.
-	 */
-	virtual void SetIsTextFormat(bool bInIsTextFormat);
-
-	/**
-	 * Sets whether this archive wants binary property serialization.
-	 *
-	 * @param bInWantBinaryPropertySerialization  true if this archive wants binary serialization, false otherwise.
-	 */
-	virtual void SetWantBinaryPropertySerialization(bool bInWantBinaryPropertySerialization);
-
-	/**
-	 * Sets whether this archive wants to force saving as Unicode.
-	 * This is needed when we need to make sure ANSI strings are saved as Unicode.
-	 *
-	 * @param bInForceUnicode  true if this archive wants to force saving as Unicode, false otherwise.
-	 */
-	virtual void SetForceUnicode(bool bInForceUnicode);
-
-	/**
-	 * Sets whether this archive is to persistent storage.
-	 *
-	 * @param bInIsPersistent  true if this archive is to persistent storage, false otherwise.
-	 */
-	virtual void SetIsPersistent(bool bInIsPersistent);
-
-	/**
-	 * Sets the archive version number. Used by the code that makes sure that FLinkerLoad's 
-	 * internal archive versions match the file reader it creates.
-	 *
-	 * @param UE4Ver	new version number
-	 */
-	virtual void SetUE4Ver(int32 InVer);
-
-	/**
-	 * Sets the archive licensee version number. Used by the code that makes sure that FLinkerLoad's 
-	 * internal archive versions match the file reader it creates.
-	 *
-	 * @param Ver	new version number
-	 */
-	virtual void SetLicenseeUE4Ver(int32 InVer);
-
-	/**
-	 * Sets the archive engine version. Used by the code that makes sure that FLinkerLoad's
-	 * internal archive versions match the file reader it creates.
-	 *
-	 * @param InVer	new version number
-	 */
-	virtual void SetEngineVer(const FEngineVersionBase& InVer);
-
-	/**
-	 * Sets the archive engine network version.
-	 */
-	virtual void SetEngineNetVer(const uint32 InEngineNetVer);
-
-	/**
-	 * Sets the archive game network version.
-	 */
-	virtual void SetGameNetVer(const uint32 InGameNetVer);
+public:
+	using FArchiveState::SetIsLoading;
+	using FArchiveState::SetIsSaving;
+	using FArchiveState::SetIsTransacting;
+	using FArchiveState::SetIsTextFormat;
+	using FArchiveState::SetWantBinaryPropertySerialization;
+	using FArchiveState::SetForceUnicode;
+	using FArchiveState::SetIsPersistent;
+	using FArchiveState::SetUE4Ver;
+	using FArchiveState::SetLicenseeUE4Ver;
+	using FArchiveState::SetEngineVer;
+	using FArchiveState::SetEngineNetVer;
+	using FArchiveState::SetGameNetVer;
 
 private:
-	/** Holds the archive version. */
-	int32 ArUE4Ver;
-
-	/** Holds the archive version for licensees. */
-	int32 ArLicenseeUE4Ver;
-
-	/** Holds the engine version. */
-	FEngineVersionBase ArEngineVer;
-
-	/** Holds the engine network protocol version. */
-	uint32 ArEngineNetVer;
-
-	/** Holds the game network protocol version. */
-	uint32 ArGameNetVer;
-
-	/**
-	* All the custom versions stored in the archive.
-	* Stored as a pointer to a heap-allocated object because of a 3-way dependency between TArray, FCustomVersionContainer and FArchive, which is too much work to change right now.
-	* Keeping it as a heap-allocated object also helps with performance in some cases as we don't need to construct it for archives that don't care about custom versions.
-	*/
-	mutable FCustomVersionContainer* CustomVersionContainer = nullptr;
+	using FArchiveState::ArUE4Ver;
+	using FArchiveState::ArLicenseeUE4Ver;
+	using FArchiveState::ArEngineVer;
+	using FArchiveState::ArEngineNetVer;
+	using FArchiveState::ArGameNetVer;
+	using FArchiveState::CustomVersionContainer;
 
 public:
 	/** Custom property list attribute. If the flag below is set, only these properties will be iterated during serialization. If NULL, then no properties will be iterated. */
-	const struct FCustomPropertyListNode* ArCustomPropertyList;
+	using FArchiveState::ArCustomPropertyList;
 
 	class FScopeSetDebugSerializationFlags
 	{
@@ -1617,9 +1834,9 @@
 
 #if WITH_EDITOR
 	/** Custom serialization modifier flags can be used for anything */
-	uint32 ArDebugSerializationFlags;
+	using FArchiveState::ArDebugSerializationFlags;
 	/** Debug stack storage if you want to add data to the archive for usage further down the serialization stack this should be used in conjunction with the FScopeAddDebugData struct */
-	
+
 	virtual void PushDebugDataString(const FName& DebugData);
 	virtual void PopDebugDataString() { }
 
@@ -1636,15 +1853,16 @@
 		}
 	};
 #endif	
+
 private:
 	/** Holds the cooking target platform. */
-	const ITargetPlatform* CookingTargetPlatform;
+	using FArchiveState::CookingTargetPlatform;
 
 	/** Holds the pointer to the property that is currently being serialized */
-	UProperty* SerializedProperty;
+	using FArchiveState::SerializedProperty;
 
 	/** Holds the chain of properties that are currently being serialized */
-	FArchiveSerializedPropertyChain* SerializedPropertyChain;
+	using FArchiveState::SerializedPropertyChain;
 
 #if USE_STABLE_LOCALIZATION_KEYS
 	/**
@@ -1652,22 +1870,18 @@
 	 * This is typically the namespace used by the package being serialized (if serializing a package, or an object within a package).
 	 * Stored as a pointer to a heap-allocated string because of a dependency between TArray (thus FString) and FArchive; null should be treated as an empty string.
 	 */
-	FString* LocalizationNamespacePtr;
-
-	/** See SetLocalizationNamespace */
-	void SetBaseLocalizationNamespace(const FString& InLocalizationNamespace);
-
-	/** See GetLocalizationNamespace */
-	FString GetBaseLocalizationNamespace() const;
+	using FArchiveState::LocalizationNamespacePtr;
 #endif // USE_STABLE_LOCALIZATION_KEYS
 
 	/**
 	 * Indicates if the custom versions container is in a 'reset' state.  This will be used to defer the choice about how to
 	 * populate the container until it is needed, where the read/write state will be known.
 	 */
-	mutable bool bCustomVersionsAreReset;
+	using FArchiveState::bCustomVersionsAreReset;
 };
 
+static_assert(sizeof(FArchive) == sizeof(FArchiveState), "New FArchive members should be added to FArchiveState instead");
 
+
 /**
  * Template for archive constructors.
diff --git a/Engine/Source/Runtime/Core/Public/Serialization/ArchiveProxy.h b/Engine/Source/Runtime/Core/Public/Serialization/ArchiveProxy.h
--- a/Engine/Source/Runtime/Core/Public/Serialization/ArchiveProxy.h
+++ b/Engine/Source/Runtime/Core/Public/Serialization/ArchiveProxy.h
@@ -264,6 +264,11 @@ public:
 		return InnerArchive.GetCacheableArchive();
 	}
 
+	virtual void SetArchiveState(const ::FArchiveState& InState) override
+	{
+		return InnerArchive.SetArchiveState(InState);
+	}
+
 protected:
 
 	/** Holds the archive that this archive is a proxy to. */
diff --git a/Engine/Source/Runtime/Core/Public/Serialization/StructuredArchive.h b/Engine/Source/Runtime/Core/Public/Serialization/StructuredArchive.h
--- a/Engine/Source/Runtime/Core/Public/Serialization/StructuredArchive.h
+++ b/Engine/Source/Runtime/Core/Public/Serialization/StructuredArchive.h
@@ -200,6 +200,8 @@ namespace UE4StructuredArchive_Private
 	public:
 		FArchive& GetUnderlyingArchive() const;
 
+		const FArchiveState& GetArchiveState() const;
+
 	protected:
 		FStructuredArchive& Ar;
 
@@ -453,6 +455,14 @@ public:
 		return Formatter.GetUnderlyingArchive();
 	}
 
+	/**
+	 * Gets the archiving state.
+	 */
+	FORCEINLINE const FArchiveState& GetArchiveState() const
+	{
+		return GetUnderlyingArchive().GetArchiveState();
+	}
+
 	FStructuredArchive(const FStructuredArchive&) = delete;
 	FStructuredArchive& operator=(const FStructuredArchive&) = delete;
 
@@ -558,6 +568,11 @@ FORCEINLINE FArchive& UE4StructuredArchive_Private::FSlotBase::GetUnderlyingArch
 	return Ar.GetUnderlyingArchive();
 }
 
+FORCEINLINE const FArchiveState& UE4StructuredArchive_Private::FSlotBase::GetArchiveState() const
+{
+	return Ar.GetArchiveState();
+}
+
 FORCEINLINE bool FStructuredArchiveSlot::IsFilled() const
 {
 #if WITH_TEXT_ARCHIVE_SUPPORT
@@ -573,7 +588,7 @@ FORCEINLINE_DEBUGGABLE void operator<<(FStructuredArchiveSlot Slot, TArray<T>& I
 	int32 NumElements = InArray.Num();
 	FStructuredArchiveArray Array = Slot.EnterArray(NumElements);
 
-	if (Slot.GetUnderlyingArchive().IsLoading())
+	if (Slot.GetArchiveState().IsLoading())
 	{
 		InArray.SetNum(NumElements);
 	}
