TArray64 used instead of manual buffer management.

diff --git a/Engine/Source/Editor/SerializedRecorderInterface/Public/Serializers/MovieSceneSectionSerialization.h b/Engine/Source/Editor/SerializedRecorderInterface/Public/Serializers/MovieSceneSectionSerialization.h
--- a/Engine/Source/Editor/SerializedRecorderInterface/Public/Serializers/MovieSceneSectionSerialization.h
+++ b/Engine/Source/Editor/SerializedRecorderInterface/Public/Serializers/MovieSceneSectionSerialization.h
@@ -93,7 +93,7 @@ public:
 		Size = Handle->Size();
 		Pos = Handle->Tell();
 		BufferBase = 0;
-		BufferCount = 0;
+		BufferArray.Reset();
 	}
 
 	/** FArchiveFileReaderGeneric interface*/
@@ -114,7 +114,7 @@ public:
 	void BeginArchiving()
 	{
 		Pos = Handle->Tell();
-		BufferCount = 0;
+		BufferArray.Reset();
 	}
 
 	/** FArchiveFileWriterGeneric interface*/
diff --git a/Engine/Source/Runtime/Core/Private/HAL/FileManagerGeneric.cpp b/Engine/Source/Runtime/Core/Private/HAL/FileManagerGeneric.cpp
--- a/Engine/Source/Runtime/Core/Private/HAL/FileManagerGeneric.cpp
+++ b/Engine/Source/Runtime/Core/Private/HAL/FileManagerGeneric.cpp
@@ -638,12 +638,9 @@ FArchiveFileReaderGeneric::FArchiveFileReaderGeneric( IFileHandle* InHandle, con
 	, Size( InSize )
 	, Pos( 0 )
 	, BufferBase( 0 )
-	, BufferCount( 0 )
 	, Handle( InHandle )
-	, Buffer(nullptr)
-	, BufferSize(InBufferSize)
 {
-	Buffer = new uint8[BufferSize];
+	BufferArray.Reserve(InBufferSize);
 	this->SetIsLoading(true);
 	this->SetIsPersistent(true);
 }
@@ -654,7 +651,7 @@ void FArchiveFileReaderGeneric::Seek( int64 InPos )
 	checkf(InPos >= 0, TEXT("Attempted to seek to a negative location (%lld/%lld), file: %s. The file is most likely corrupt."), InPos, Size, *Filename);
 	checkf(InPos <= Size, TEXT("Attempted to seek past the end of file (%lld/%lld), file: %s. The file is most likely corrupt."), InPos, Size, *Filename);
 
-	if (InPos < BufferBase || InPos >= BufferBase + BufferCount)
+	if (InPos < BufferBase || InPos >= BufferBase + BufferArray.Num())
 	{
 		// The seek is outside of the cached buffer so we actually need to seek
 		if (!SeekLowLevel(InPos))
@@ -666,7 +663,7 @@ void FArchiveFileReaderGeneric::Seek( int64 InPos )
 
 		Pos = InPos;
 		BufferBase = Pos;
-		BufferCount = 0;
+		BufferArray.Reset();
 	}
 	else
 	{
@@ -678,8 +675,6 @@ void FArchiveFileReaderGeneric::Seek( int64 InPos )
 FArchiveFileReaderGeneric::~FArchiveFileReaderGeneric()
 {
 	Close();
-	check(Buffer != nullptr);
-	delete[] Buffer;
 }
 
 void FArchiveFileReaderGeneric::ReadLowLevel( uint8* Dest, int64 CountToRead, int64& OutBytesRead )
@@ -713,15 +708,16 @@ bool FArchiveFileReaderGeneric::Close()
 bool FArchiveFileReaderGeneric::InternalPrecache( int64 PrecacheOffset, int64 PrecacheSize )
 {
 	// Only precache at current position and avoid work if precaching same offset twice.
-	if( Pos == PrecacheOffset &&( !BufferBase || !BufferCount || BufferBase != Pos ) )
+	if( Pos == PrecacheOffset &&( !BufferBase || !BufferArray.Num() || BufferBase != Pos ) )
 	{
 		BufferBase = Pos;
-		BufferCount = FMath::Min( FMath::Min( PrecacheSize,( int64 )( BufferSize -( Pos&( BufferSize -1 ) ) ) ), Size-Pos );
-		BufferCount = FMath::Max( BufferCount, 0LL ); // clamp to 0
+		int64 BufferCount = FMath::Min( FMath::Min( PrecacheSize,( int64 )(BufferArray.Max() -( Pos&( BufferArray.Max() -1 ) ) ) ), Size-Pos );
+		BufferCount = FMath::Max(BufferCount, 0LL ); // clamp to 0
+		BufferArray.SetNumUninitialized(BufferCount);
 		int64 Count = 0;
 
 		{
-			if (BufferCount > BufferSize || BufferCount <= 0)
+			if (BufferCount > BufferArray.Max() || BufferCount <= 0)
 			{
 				UE_LOG( LogFileManager, Error, TEXT("Invalid BufferCount=%lld while reading %s. File is most likely corrupted. Please verify your installation. Pos=%lld, Size=%lld, PrecacheSize=%lld, PrecacheOffset=%lld"),
 					BufferCount, *Filename, Pos, Size, PrecacheSize, PrecacheOffset );
@@ -729,7 +725,7 @@ bool FArchiveFileReaderGeneric::InternalPrecache( int64 PrecacheOffset, int64 Pr
 				return false;
 			}
 
-			ReadLowLevel( Buffer, BufferCount, Count );
+			ReadLowLevel( BufferArray.GetData(), BufferArray.Num(), Count );
 		}
 
 		if( Count!=BufferCount )
@@ -753,10 +749,10 @@ void FArchiveFileReaderGeneric::Serialize( void* V, int64 Length )
 
 	while( Length>0 )
 	{
-		int64 Copy = FMath::Min( Length, BufferBase+BufferCount-Pos );
+		int64 Copy = FMath::Min( Length, BufferBase+BufferArray.Num()-Pos );
 		if( Copy<=0 )
 		{
-			if( Length >= BufferSize )
+			if( Length >= BufferArray.Max() )
 			{
 				int64 Count=0;
 				{
@@ -778,7 +774,7 @@ void FArchiveFileReaderGeneric::Serialize( void* V, int64 Length )
 				UE_LOG( LogFileManager, Warning, TEXT( "ReadFile failed during precaching for file %s" ),*Filename );
 				return;
 			}
-			Copy = FMath::Min( Length, BufferBase+BufferCount-Pos );
+			Copy = FMath::Min( Length, BufferBase+BufferArray.Num()-Pos );
 			if( Copy<=0 )
 			{
 				ArIsError = true;
@@ -790,7 +786,7 @@ void FArchiveFileReaderGeneric::Serialize( void* V, int64 Length )
 				return;
 			}
 		}
-		FMemory::Memcpy( V, Buffer+Pos-BufferBase, Copy );
+		FMemory::Memcpy( V, BufferArray.GetData()+Pos-BufferBase, Copy );
 		Pos       += Copy;
 		Length    -= Copy;
 		V          =( uint8* )V + Copy;
@@ -801,13 +797,10 @@ FArchiveFileWriterGeneric::FArchiveFileWriterGeneric( IFileHandle* InHandle, con
 	: Filename( InFilename )
 	, Flags( InFlags )
 	, Pos( InPos )
-	, BufferCount( 0 )
 	, Handle( InHandle )
-	, Buffer(nullptr)
-	, BufferSize(InBufferSize)
 	, bLoggingError( false )
 {
-	Buffer = new uint8[BufferSize];
+	BufferArray.Reserve(InBufferSize);
 	this->SetIsSaving(true);
 	this->SetIsPersistent(true);
 }
@@ -815,8 +808,6 @@ FArchiveFileWriterGeneric::FArchiveFileWriterGeneric( IFileHandle* InHandle, con
 FArchiveFileWriterGeneric::~FArchiveFileWriterGeneric()
 {
 	Close();
-	check(Buffer != nullptr);
-	delete[] Buffer;
 }
 
 bool FArchiveFileWriterGeneric::CloseLowLevel()
@@ -867,7 +858,7 @@ bool FArchiveFileWriterGeneric::Close()
 void FArchiveFileWriterGeneric::Serialize( void* V, int64 Length )
 {
 	Pos += Length;
-	if ( Length >= BufferSize )
+	if ( Length >= BufferArray.Max() )
 	{
 		FlushBuffer();
 		if( !WriteLowLevel( (uint8*)V, Length ) )
@@ -879,20 +870,16 @@ void FArchiveFileWriterGeneric::Serialize( void* V, int64 Length )
 	else
 	{
 		int64 Copy;
-		while( Length >( Copy=BufferSize-BufferCount ) )
+		while( Length >( Copy=BufferArray.Max()-BufferArray.Num() ) )
 		{
-			FMemory::Memcpy( Buffer+BufferCount, V, Copy );
-			BufferCount += Copy;
-			check( BufferCount <= BufferSize && BufferCount >= 0 );
+			BufferArray.Append((uint8*)V, Copy);
 			Length      -= Copy;
 			V            =( uint8* )V + Copy;
 			FlushBuffer();
 		}
 		if( Length )
 		{
-			FMemory::Memcpy( Buffer+BufferCount, V, Length );
-			BufferCount += Length;
-			check( BufferCount <= BufferSize && BufferCount >= 0 );
+			BufferArray.Append((uint8*)V, Length);
 		}
 	}
 }
@@ -908,16 +895,15 @@ void FArchiveFileWriterGeneric::Flush()
 bool FArchiveFileWriterGeneric::FlushBuffer()
 {
 	bool bDidWriteData = false;
-	if (BufferCount)
+	if (int64 BufferNum = BufferArray.Num())
 	{
-		check(BufferCount <= BufferSize && BufferCount > 0);
-		bDidWriteData = WriteLowLevel(Buffer, BufferCount);
+		bDidWriteData = WriteLowLevel(BufferArray.GetData(), BufferNum);
 		if (!bDidWriteData)
 		{
 			ArIsError = true;
 			LogWriteError(TEXT("Error flushing file"));
 		}
-		BufferCount = 0;
+		BufferArray.Reset();
 	}
 	return bDidWriteData;
 }
diff --git a/Engine/Source/Runtime/Core/Public/HAL/FileManagerGeneric.h b/Engine/Source/Runtime/Core/Public/HAL/FileManagerGeneric.h
--- a/Engine/Source/Runtime/Core/Public/HAL/FileManagerGeneric.h
+++ b/Engine/Source/Runtime/Core/Public/HAL/FileManagerGeneric.h
@@ -266,10 +266,8 @@ protected:
 	int64 Size;
 	int64 Pos;
 	int64 BufferBase;
-	int64 BufferCount;
 	TUniquePtr<IFileHandle> Handle;
-	uint8* Buffer;
-	uint32 BufferSize;
+	TArray64<uint8> BufferArray;
 };
 
 
@@ -342,9 +340,7 @@ protected:
 	FString Filename;
 	uint32 Flags;
 	int64 Pos;
-	int64 BufferCount;
 	TUniquePtr<IFileHandle> Handle;
-	uint8* Buffer;
-	uint32 BufferSize;
+	TArray64<uint8> BufferArray;
 	bool bLoggingError;
 };
