GitHub #6283 : Use Emplace instead

diff --git a/Engine/Source/Runtime/Core/Private/Misc/FileHelper.cpp b/Engine/Source/Runtime/Core/Private/Misc/FileHelper.cpp
--- a/Engine/Source/Runtime/Core/Private/Misc/FileHelper.cpp
+++ b/Engine/Source/Runtime/Core/Private/Misc/FileHelper.cpp
@@ -177,7 +177,7 @@ bool FFileHelper::LoadFileToStringArray( TArray<FString>& Result, const TCHAR* F
 			Pos++;
 		}
 
-		Result.Add(FString(Pos - LineStart, LineStart));
+		Result.Emplace(Pos - LineStart, LineStart);
 
 		if(*Pos == '\r')
 		{
