GitHub #6271 : Deduce types for void specialisations of TLess and TGreater

diff --git a/Engine/Source/Runtime/Core/Public/Templates/Greater.h b/Engine/Source/Runtime/Core/Public/Templates/Greater.h
--- a/Engine/Source/Runtime/Core/Public/Templates/Greater.h
+++ b/Engine/Source/Runtime/Core/Public/Templates/Greater.h
@@ -3,6 +3,7 @@
 #pragma once
 
 #include "CoreTypes.h"
+#include "UnrealTemplate.h"
 
 /**
  * Binary predicate class for sorting elements in reverse order.  Assumes < operator is defined for the template type.
@@ -21,9 +22,9 @@ struct TGreater
 template <>
 struct TGreater<void>
 {
-	template <typename T>
-	FORCEINLINE bool operator()(const T& A, const T& B) const
+	template <typename T, typename U>
+	FORCEINLINE bool operator()(T&& A, U&& B) const
 	{
-		return B < A;
+		return Forward<U>(B) < Forward<T>(A);
 	}
 };
diff --git a/Engine/Source/Runtime/Core/Public/Templates/Less.h b/Engine/Source/Runtime/Core/Public/Templates/Less.h
--- a/Engine/Source/Runtime/Core/Public/Templates/Less.h
+++ b/Engine/Source/Runtime/Core/Public/Templates/Less.h
@@ -3,6 +3,7 @@
 #pragma once
 
 #include "CoreTypes.h"
+#include "UnrealTemplate.h"
 
 /**
  * Binary predicate class for sorting elements in order.  Assumes < operator is defined for the template type.
@@ -22,9 +23,9 @@ struct TLess
 template <>
 struct TLess<void>
 {
-	template <typename T>
-	FORCEINLINE bool operator()(const T& A, const T& B) const
+	template <typename T, typename U>
+	FORCEINLINE bool operator()(T&& A, U&& B) const
 	{
-		return A < B;
+		return Forward<T>(A) < Forward<U>(B);
 	}
 };
