[IoDispatcher] Dynamic allocation and I/O Queue
- changed fixed sized batch/request buffers to be dynamically allocated
- changed TIoStatusOr to not create result value until valid result
- changed FString to fixed size error message in FIoStatus
- refactored IoQueue to better utilize I/O dispatcher and remove overhead

diff --git a/Engine/Source/Runtime/Core/Private/IO/IoDispatcher.cpp b/Engine/Source/Runtime/Core/Private/IO/IoDispatcher.cpp
--- a/Engine/Source/Runtime/Core/Private/IO/IoDispatcher.cpp
+++ b/Engine/Source/Runtime/Core/Private/IO/IoDispatcher.cpp
@@ -1,21 +1,132 @@
 // Copyright 1998-2019 Epic Games, Inc. All Rights Reserved.
 
 #include "IO/IoDispatcher.h"
-#include "HAL/FileManager.h"
 #include "Async/MappedFileHandle.h"
 #include "Containers/ChunkedArray.h"
 #include "HAL/CriticalSection.h"
+#include "HAL/FileManager.h"
+#include "HAL/Runnable.h"
+#include "HAL/RunnableThread.h"
 #include "Logging/LogMacros.h"
+#include "Memory/MemoryArena.h"
 #include "Misc/CString.h"
 #include "Misc/EventPool.h"
 #include "Misc/ScopeRWLock.h"
+#include "Misc/ScopeLock.h"
 #include "Misc/SecureHash.h"
 #include "Misc/StringBuilder.h"
-#include "Memory/MemoryArena.h"
 #include "Misc/LazySingleton.h"
+#include "Misc/CoreDelegates.h"
 
 DEFINE_LOG_CATEGORY_STATIC(LogIoDispatch, Log, All);
 
+//////////////////////////////////////////////////////////////////////////
+
+template <typename T, uint32 BlockSize = 128>
+class TBlockAllocator
+{
+public:
+	~TBlockAllocator()
+	{
+		FreeBlocks();
+	}
+
+	FORCEINLINE T* Alloc()
+	{
+		FScopeLock _(&CriticalSection);
+
+		if (!NextFree)
+		{
+			//TODO: Virtual alloc
+			FBlock* Block = new FBlock;
+			
+			for (int32 ElementIndex = 0; ElementIndex < BlockSize; ++ElementIndex)
+			{
+				FElement* Element	= &Block->Elements[ElementIndex];
+				Element->Next		= NextFree;
+				NextFree			= Element;
+			}
+
+			Block->Next	= Blocks;
+			Blocks		= Block;
+		}
+
+		FElement* Element	= NextFree;
+		NextFree			= Element->Next;
+		
+		++NumElements;
+
+		return Element->Buffer.GetTypedPtr();
+	}
+
+	FORCEINLINE void Free(T* Ptr)
+	{
+		FScopeLock _(&CriticalSection);
+
+		FElement* Element	= reinterpret_cast<FElement*>(Ptr);
+		Element->Next		= NextFree;
+		NextFree			= Element;
+
+		--NumElements;
+	}
+
+	template <typename... ArgsType>
+	T* Construct(ArgsType&&... Args)
+	{
+		return new(Alloc()) T(Forward<ArgsType>(Args)...);
+	}
+
+	void Destroy(T* Ptr)
+	{
+		Ptr->~T();
+		Free(Ptr);
+	}
+
+	void Trim()
+	{
+		FScopeLock _(&CriticalSection);
+		if (!NumElements)
+		{
+			FreeBlocks();
+		}
+	}
+
+private:
+	void FreeBlocks()
+	{
+		FBlock* Block = Blocks;
+		while (Block)
+		{
+			FBlock* Tmp = Block;
+			Block = Block->Next;
+			delete Tmp;
+		}
+
+		Blocks		= nullptr;
+		NextFree	= nullptr;
+		NumElements = 0;
+	}
+
+	struct FElement
+	{
+		TTypeCompatibleBytes<T> Buffer;
+		FElement* Next;
+	};
+
+	struct FBlock
+	{
+		FElement Elements[BlockSize];
+		FBlock* Next = nullptr;
+	};
+
+	FBlock*				Blocks		= nullptr;
+	FElement*			NextFree	= nullptr;
+	int32				NumElements = 0;
+	FCriticalSection	CriticalSection;
+};
+
+//////////////////////////////////////////////////////////////////////////
+
 const TCHAR* GetIoErrorText(EIoErrorCode ErrorCode)
 {
 	static constexpr const TCHAR* ErrorCodeText[]
@@ -38,23 +149,33 @@ const FIoStatus FIoStatus::Ok		{ EIoErrorCode::Ok,				TEXT("OK")				};
 const FIoStatus FIoStatus::Unknown	{ EIoErrorCode::Unknown,		TEXT("Unknown Status")	};
 const FIoStatus FIoStatus::Invalid	{ EIoErrorCode::InvalidCode,	TEXT("Invalid Code")	};
 
+FIoStatus::FIoStatus()
+{
+}
+
+FIoStatus::~FIoStatus()
+{
+}
+
 FIoStatus::FIoStatus(EIoErrorCode Code)
 :	ErrorCode(Code)
-,	ErrorMessage(GetIoErrorText(Code))
 {
+	ErrorMessage[0] = 0;
 }
 
 FIoStatus::FIoStatus(EIoErrorCode Code, const FStringView& InErrorMessage)
 : ErrorCode(Code)
-, ErrorMessage(InErrorMessage.Data(), InErrorMessage.Len())
 {
+	const int32 ErrorMessageLength = FMath::Min(MaxErrorMessageLength - 1, InErrorMessage.Len());
+	FPlatformString::Convert(ErrorMessage, ErrorMessageLength, InErrorMessage.Data(), ErrorMessageLength);
+	ErrorMessage[ErrorMessageLength] = 0;
 }
 
 FIoStatus& 
 FIoStatus::operator=(const FIoStatus& Other)
 {
 	ErrorCode = Other.ErrorCode;
-	ErrorMessage = Other.ErrorMessage;
+	FMemory::Memcpy(ErrorMessage, Other.ErrorMessage, MaxErrorMessageLength * sizeof(TCHAR));
 
 	return *this;
 }
@@ -62,14 +183,14 @@ FIoStatus::operator=(const FIoStatus& Other)
 bool		
 FIoStatus::operator==(const FIoStatus& Other) const
 {
-	return ErrorCode == Other.ErrorCode
-		&& ErrorMessage == Other.ErrorMessage;
+	return ErrorCode == Other.ErrorCode &&
+		FPlatformString::Stricmp(ErrorMessage, Other.ErrorMessage) == 0;
 }
 
 FString	
 FIoStatus::ToString() const
 {
-	return ErrorMessage;
+	return FString::Format(TEXT("{0} ({1})"), { ErrorMessage, GetIoErrorText(ErrorCode) });
 }
 
 //////////////////////////////////////////////////////////////////////////
@@ -141,7 +262,6 @@ public:
 	}
 };
 
-
 //////////////////////////////////////////////////////////////////////////
 
 class FIoBufferPool
@@ -760,25 +880,19 @@ void FIoStoreWriter::FlushMetadata()
 class FIoBatchImpl
 {
 public:
-	uint32			FirstRequest = ~uint32(0);
+	FIoRequestImpl* FirstRequest	= nullptr;
+	FIoBatchImpl*	NextBatch		= nullptr;
 	TAtomic<uint32>	OutstandingRequests { 0 };
-
-	void Reset() 
-	{
-		// TODO: need to release IO requests back to pool
-
-		FirstRequest		= ~uint32(0);
-		OutstandingRequests = 0;
-	}
 };
 
 class FIoRequestImpl
 {
 public:
+	FIoChunkId				ChunkId;
 	FIoReadOptions			Options;
 	TIoStatusOr<FIoBuffer>	Result;
-	uint32					NextRequestId = ~uint32(0);
-	FIoChunkId				ChunkId;
+	uint64					UserData	= 0;
+	FIoRequestImpl*			NextRequest = nullptr;
 };
 
 //////////////////////////////////////////////////////////////////////////
@@ -861,183 +975,116 @@ public:
 
 //////////////////////////////////////////////////////////////////////////
 
-class FIoDispatcherImpl
-{
-	FIoDispatcher* Owner;
-
-	// Batch allocation - could use something more lightweight but I'm not sure this
-	// is a sufficiently high frequency operation to warrant something lock-free
+using FRequestAllocator	= TBlockAllocator<FIoRequestImpl, 4096>;
+using FBatchAllocator	= TBlockAllocator<FIoBatchImpl, 4096>;
 
-	FRWLock	RwLockBatchId;
-	uint32	FirstFreeBatchId = 0;
-
-	// Request allocation - could use something more lightweight but I'm not sure this
-	// is a sufficiently high frequency operation to warrant something lock-free
-
-	FRWLock	RwLockRequestId;
-	uint32	FirstFreeRequestId = ~0u;
+//////////////////////////////////////////////////////////////////////////
 
+class FIoDispatcherImpl
+{
 public:
-	FIoDispatcherImpl(FIoDispatcher* InOwner)
-	:	Owner(InOwner)
+	FIoDispatcherImpl()
 	{
-		for (int i = 0; i < (MaxBatches - 1); ++i)
-		{
-			Batches[i].FirstRequest = i + 1;
-		}
-
-		Batches[MaxBatches - 1].FirstRequest = ~uint32(0);
-
-		for (int i = 0; i < InitialIoRequests; ++i)
-		{
-			ReserveNewRequest();
-		}
-
-		Batches[MaxBatches - 1].FirstRequest = ~uint32(0);
-
 		IoStore = new FIoStoreImpl;
-	}
-
-	static const int	MaxBatches = 65536;
-	FIoBatchImpl		Batches[MaxBatches];
-	int32				NumBatches = 0;
-
-	static const int	MaxIoRequests = 65536;
-	static const int	InitialIoRequests = 1;
 
-	TChunkedArray<FIoRequestImpl, MaxIoRequests> Requests;
-
-	FIoBatchImpl& MapBatch(uint32 BatchId)
-	{
-		return Batches[BatchId];
-	}
-
-	void ReserveNewRequest()
-	{
-		// NOTE: RwLockRequestId must be held while this is called
-
-		FIoRequestImpl& NewRequest = *new(Requests) FIoRequestImpl;
-		NewRequest.NextRequestId = FirstFreeRequestId;
-		FirstFreeRequestId = Requests.Num() - 1;
+		FCoreDelegates::GetMemoryTrimDelegate().AddLambda([this]()
+		{
+			RequestAllocator.Trim();
+			BatchAllocator.Trim();
+		});
 	}
 
-	FIoRequestImpl& AllocRequest(const FIoChunkId& Chunk, FIoReadOptions Options, uint32 BatchId)
+	FIoRequestImpl* AllocRequest(const FIoChunkId& ChunkId, FIoReadOptions Options, uint64 UserData = 0)
 	{
-		const uint32 RequestId = AllocRequestId();
+		FIoRequestImpl* Request = RequestAllocator.Construct();
 
-		FIoRequestImpl& Request = Requests[RequestId];
-		FIoBatchImpl& Batch		= MapBatch(BatchId);
-
-		Request.ChunkId			= Chunk;
-		Request.Options			= Options;
-		Request.NextRequestId	= Batch.FirstRequest;
-		Request.Result			= FIoStatus::Unknown;
-
-		Batch.FirstRequest		= RequestId;
+		Request->ChunkId		= ChunkId;
+		Request->Options		= Options;
+		Request->Result			= FIoStatus::Unknown;
+		Request->UserData		= UserData;
+		Request->NextRequest	= nullptr;
 
 		return Request;
 	}
 
-	uint32 AllocRequestId()
+	FIoRequestImpl* AllocRequest(FIoBatchImpl* Batch, const FIoChunkId& ChunkId, FIoReadOptions Options, uint64 UserData = 0)
 	{
-		FWriteScopeLock _(RwLockRequestId);
-
-		if (FirstFreeRequestId == ~uint32(0))
-		{
-			ReserveNewRequest();
-		}
-
-		const uint32 RequestId	= FirstFreeRequestId;
+		FIoRequestImpl* Request	= AllocRequest(ChunkId, Options, UserData);
 
-		FIoRequestImpl& Request	= Requests[RequestId];
-		FirstFreeRequestId		= Request.NextRequestId;
-		Request.NextRequestId	= ~uint32(0);
+		Request->NextRequest	= Batch->FirstRequest;
+		Batch->FirstRequest		= Request;
 
-		return RequestId;
+		return Request;
 	}
 
-	uint32 FreeRequestId(const uint32 RequestId)
+	void FreeRequest(FIoRequestImpl* Request)
 	{
-		FWriteScopeLock _(RwLockRequestId);
-
-		uint32 NextRequestId				= Requests[RequestId].NextRequestId;
-		Requests[RequestId].NextRequestId	= FirstFreeRequestId;
-		FirstFreeRequestId					= RequestId;
-
-		return NextRequestId;
+		RequestAllocator.Destroy(Request);
 	}
 
-	uint32 AllocBatchId()
+	FIoBatchImpl* AllocBatch(FIoRequestImpl* FirstRequest = nullptr)
 	{
-		check(NumBatches < MaxBatches);
-
-		FWriteScopeLock _(RwLockBatchId);
-
-		const uint32 BatchId	= FirstFreeBatchId;
-		FIoBatchImpl& Batch		= Batches[BatchId];
-		FirstFreeBatchId		= Batch.FirstRequest;
-		Batch.FirstRequest		= ~uint32(0);
-
-		NumBatches++;
+		FIoBatchImpl* Batch			= BatchAllocator.Construct();
 
-		return BatchId;
+		Batch->FirstRequest			= FirstRequest;
+		Batch->OutstandingRequests	= 0;
+	
+		return Batch;
 	}
 
-	void FreeBatchId(const uint32 BatchId)
+	void FreeBatch(FIoBatchImpl* Batch)
 	{
-		FIoBatchImpl& Batch = Batches[BatchId];
+		FIoRequestImpl* Request = Batch->FirstRequest;
 
-		uint32 RequestId = Batch.FirstRequest; 
-		while (RequestId != ~uint32(0))
+		while (Request)
 		{
-			RequestId = FreeRequestId(RequestId);
+			FIoRequestImpl* Tmp = Request;	
+			Request = Request->NextRequest;
+			FreeRequest(Tmp);
 		}
-	
-		FWriteScopeLock _(RwLockBatchId);
 
-		Batches[BatchId].FirstRequest	= FirstFreeBatchId;
-		FirstFreeBatchId				= BatchId;
-
-		NumBatches--;
-		check(NumBatches >= 0);
+		BatchAllocator.Destroy(Batch);
 	}
 
 	template<typename Func>
-	void IterateBatch(const uint32 BatchId, Func&& InCallbackFunction)
+	void IterateBatch(const FIoBatchImpl* Batch, Func&& InCallbackFunction)
 	{
-		FIoBatchImpl& Batch = Batches[BatchId];
-
-		uint32 RequestId = Batch.FirstRequest;
+		FIoRequestImpl* Request = Batch->FirstRequest;
 
-		while (RequestId != ~uint32(0))
+		while (Request)
 		{
-			FIoRequestImpl& Request = Requests[RequestId];
+			const bool bDoContinue = InCallbackFunction(Request);
 
-			const bool DoContinue = InCallbackFunction(Request);
-
-			if (!DoContinue)
-			{
-				return;
-			}
-
-			RequestId = Request.NextRequestId;
+			Request = bDoContinue ? Request->NextRequest : nullptr;
 		}
 	}
 
-	void IssueBatch(const uint32 BatchId)
+	void IssueBatch(const FIoBatchImpl* Batch)
 	{
 		// At this point the batch is immutable and we should start
 		// doing the work.
 
-		IterateBatch(BatchId, [this](FIoRequestImpl& Request) {
-			TIoStatusOr<FIoBuffer> Resolved = IoStore->Resolve(Request.ChunkId);
-
-			Request.Result = Resolved;
+		IterateBatch(Batch, [this](FIoRequestImpl* Request)
+		{
+			Request->Result = IoStore->Resolve(Request->ChunkId);
 
 			return true;
 		});
 	}
 
+	bool IsBatchReady(const FIoBatchImpl* Batch)
+	{
+		bool bIsReady = true;
+
+		IterateBatch(Batch, [&bIsReady](FIoRequestImpl* Request)
+		{
+			bIsReady &= Request->Result.Status().IsCompleted();
+			return bIsReady;
+		});
+
+		return bIsReady;
+	}
+
 	void Mount(FIoStoreReader* IoStoreReader)
 	{
 		IoStore->Mount(IoStoreReader);
@@ -1050,17 +1097,20 @@ public:
 
 private:
 	TRefCountPtr<FIoStoreImpl>	IoStore;
+	FRequestAllocator			RequestAllocator;
+	FBatchAllocator				BatchAllocator;
 };
 
 //////////////////////////////////////////////////////////////////////////
 
 FIoDispatcher::FIoDispatcher()
-:	Impl(new FIoDispatcherImpl(this))
+:	Impl(new FIoDispatcherImpl())
 {
 }
 
 FIoDispatcher::~FIoDispatcher()
 {
+	delete Impl;
 }
 
 void		
@@ -1078,34 +1128,34 @@ FIoDispatcher::Unmount(FIoStoreReader* IoStore)
 FIoBatch
 FIoDispatcher::NewBatch()
 {
-	return FIoBatch(Impl, Impl->AllocBatchId());
+	return FIoBatch(Impl, Impl->AllocBatch());
 }
 
 void
 FIoDispatcher::FreeBatch(FIoBatch Batch)
 {
-	Impl->FreeBatchId(Batch.BatchId);
+	Impl->FreeBatch(Batch.Impl);
 }
 
 //////////////////////////////////////////////////////////////////////////
 
-FIoBatch::FIoBatch(FIoDispatcherImpl* OwningIoDispatcher, uint32 InBatchId)
-:	Dispatcher(OwningIoDispatcher)
-,	BatchId(InBatchId)
-,	CompletionEvent(TLazySingleton<FEventPool<EEventPoolTypes::ManualReset>>::Get().GetEventFromPool())
+FIoBatch::FIoBatch(FIoDispatcherImpl* InDispatcher, FIoBatchImpl* InImpl)
+:	Dispatcher(InDispatcher)
+,	Impl(InImpl)
+,	CompletionEvent()
 {
 }
 
 FIoRequest
-FIoBatch::Read(const FIoChunkId& Chunk, FIoReadOptions Options)
+FIoBatch::Read(const FIoChunkId& ChunkId, FIoReadOptions Options)
 {
-	return FIoRequest(Dispatcher->AllocRequest(Chunk, Options, BatchId));
+	return FIoRequest(Dispatcher->AllocRequest(Impl, ChunkId, Options));
 }
 
 void 
 FIoBatch::ForEachRequest(TFunction<bool(FIoRequest&)> Callback)
 {
-	Dispatcher->IterateBatch(BatchId, [&](FIoRequestImpl& InRequest) {
+	Dispatcher->IterateBatch(Impl, [&](FIoRequestImpl* InRequest) {
 		FIoRequest Request(InRequest);
 		return Callback(Request);
 	});
@@ -1114,12 +1164,13 @@ FIoBatch::ForEachRequest(TFunction<bool(FIoRequest&)> Callback)
 void 
 FIoBatch::Issue()
 {
-	Dispatcher->IssueBatch(BatchId);
+	Dispatcher->IssueBatch(Impl);
 }
 
 void 
 FIoBatch::Wait()
 {
+	//TODO: Create synchronization event here when it's actually needed
 	unimplemented();
 }
 
@@ -1154,3 +1205,290 @@ FIoRequest::GetChunkId() const
 {
 	return Impl->ChunkId;
 }
+
+const TIoStatusOr<FIoBuffer>&
+FIoRequest::GetResult() const
+{
+	return Impl->Result;
+}
+
+//////////////////////////////////////////////////////////////////////////
+
+class FIoQueueImpl
+	: private FRunnable
+{
+public:
+	FIoQueueImpl(FIoDispatcherImpl& InDispatcher, FIoQueue::FBatchReadyCallback&& InBatchReadyCallback)
+		: Dispatcher(InDispatcher)
+		, BatchReadyCallback(Forward<FIoQueue::FBatchReadyCallback>(InBatchReadyCallback))
+	{
+		WakeUpEvent = FGenericPlatformProcess::GetSynchEventFromPool(true);
+		Thread = FRunnableThread::Create(this, TEXT("IoQueueThread"), 0, TPri_Normal);
+	}
+
+	~FIoQueueImpl()
+	{
+		Stop();
+		Thread->Kill(true);
+		Thread = nullptr;
+		FPlatformProcess::ReturnSynchEventToPool(WakeUpEvent);
+	}
+
+	void Enqueue(const FIoChunkId& ChunkId, FIoReadOptions ReadOptions, uint64 UserData, bool bDeferBatch)
+	{
+		FIoRequestImpl* Request = Dispatcher.AllocRequest(ChunkId, ReadOptions, UserData);
+		check(Request->NextRequest == nullptr);
+
+		{
+			// TODO: CAS
+			FScopeLock _(&QueuedLock);
+
+			Request->NextRequest	= FirstQueuedRequest;
+			FirstQueuedRequest		= Request;
+		}
+
+		if (!bDeferBatch)
+		{
+			IssueBatch();
+		}
+	
+		{
+			FScopeLock _(&NumPendingLock);
+
+			if (0 == NumPending++)
+			{
+				WakeUpEvent->Trigger();
+			}
+		}
+	}
+
+	bool Dequeue(FIoChunkId& ChunkId, TIoStatusOr<FIoBuffer>& Result, uint64& UserData)
+	{
+		FIoRequestImpl* CompletedRequest = nullptr;
+		{
+			// TODO: CAS
+			FScopeLock _(&CompletedLock);
+			
+			CompletedRequest		= FirstCompletedRequest;
+			FirstCompletedRequest	= CompletedRequest ? CompletedRequest->NextRequest : FirstCompletedRequest;
+		}
+
+		if (CompletedRequest)
+		{
+			ChunkId		= CompletedRequest->ChunkId;
+			Result		= CompletedRequest->Result;
+			UserData	= CompletedRequest->UserData;
+
+			Dispatcher.FreeRequest(CompletedRequest);
+
+			{
+				FScopeLock _(&NumPendingLock);
+
+				if (0 == --NumPending)
+				{
+					check(NumPending >= 0);
+					WakeUpEvent->Reset();
+				}
+			}
+
+			return true;
+		}
+
+		return false;
+	}
+
+	void IssueBatch()
+	{
+		FIoRequestImpl* QueuedRequests = nullptr;
+		{
+			// TODO: CAS
+			FScopeLock _(&QueuedLock);
+
+			QueuedRequests 		= FirstQueuedRequest;
+			FirstQueuedRequest	= nullptr;
+		}
+
+		if (QueuedRequests)
+		{
+			FIoBatchImpl* NewBatch = Dispatcher.AllocBatch(QueuedRequests);
+			{
+				// TODO: CAS
+				FScopeLock _(&PendingLock);
+
+				NewBatch->NextBatch	= FirstPendingBatch;
+				FirstPendingBatch	= NewBatch;
+			}
+		}
+	}
+
+	bool IsEmpty() const
+	{
+		FScopeLock _(&NumPendingLock);
+		return NumPending == 0;
+	}
+
+private:
+	
+	uint32 Run() override
+	{
+		bIsRunning = true;
+
+		FBatchQueue IssuedBatches;
+
+		while (bIsRunning)
+		{
+			if (!bIsRunning)
+			{
+				break;
+			}
+
+			// Dispatch pending batch
+			{
+				FIoBatchImpl* PendingBatch = nullptr;
+				{
+					// TODO: CAS
+					FScopeLock _(&PendingLock);
+
+					PendingBatch		= FirstPendingBatch;
+					FirstPendingBatch	= PendingBatch ? PendingBatch->NextBatch : nullptr;
+				}
+
+				if (PendingBatch)
+				{
+					PendingBatch->NextBatch = nullptr;
+					
+					Dispatcher.IssueBatch(PendingBatch);
+					Enqueue(PendingBatch, IssuedBatches);
+				}
+			}
+
+			// Process issued batches
+			if (FIoBatchImpl* IssuedBatch = Peek(IssuedBatches))
+			{
+				if (Dispatcher.IsBatchReady(IssuedBatch))
+				{
+					Dequeue(IssuedBatches);
+
+					FIoRequestImpl* Request = IssuedBatch->FirstRequest;
+					while (Request)
+					{
+						FIoRequestImpl* CompletedRequest	= Request;
+						Request								= Request->NextRequest;
+						
+						{
+							// TODO: CAS
+							FScopeLock _(&CompletedLock);
+
+							CompletedRequest->NextRequest	= FirstCompletedRequest;
+							FirstCompletedRequest			= CompletedRequest;
+						}
+					}
+
+					BatchReadyCallback();
+
+					IssuedBatch->FirstRequest = nullptr;
+					Dispatcher.FreeBatch(IssuedBatch);
+				}
+			}
+			else
+			{
+				WakeUpEvent->Wait();
+			}
+		}
+
+		return 0;
+	}
+
+	void Stop() override
+	{
+		if (bIsRunning)
+		{
+			bIsRunning = false;
+			WakeUpEvent->Trigger();
+		}
+	}
+
+	struct FBatchQueue
+	{
+		FIoBatchImpl* Head = nullptr;
+		FIoBatchImpl* Tail = nullptr;
+	};
+
+	void Enqueue(FIoBatchImpl* Batch, FBatchQueue& Batches)
+	{
+		if (Batches.Tail == nullptr)
+		{
+			Batches.Head = Batches.Tail	= Batch;
+		}
+		else
+		{
+			Batches.Tail->NextBatch	= Batch;
+			Batches.Tail			= Batch;
+		}
+	}
+
+	FIoBatchImpl* Dequeue(FBatchQueue& Batches)
+	{
+		if (FIoBatchImpl* Batch = Batches.Head)
+		{
+			Batches.Head		= Batch->NextBatch;
+			Batches.Tail		= Batches.Head == nullptr ? nullptr : Batches.Tail;
+			Batch->NextBatch	= nullptr;
+
+			return Batch;
+		}
+		
+		return nullptr;
+	}
+
+	FIoBatchImpl* Peek(FBatchQueue& Queue)
+	{
+		return Queue.Head;
+	}
+
+	FIoDispatcherImpl&				Dispatcher;
+	FIoQueue::FBatchReadyCallback	BatchReadyCallback;
+	FRunnableThread*				Thread					= nullptr;
+	FEvent*							WakeUpEvent				= nullptr;
+	bool							bIsRunning				= false;
+	FIoRequestImpl*					FirstQueuedRequest		= nullptr;
+	FIoBatchImpl*					FirstPendingBatch		= nullptr;
+	FIoRequestImpl*					FirstCompletedRequest	= nullptr;
+	int32							NumPending				= 0;
+	FCriticalSection				QueuedLock;
+	FCriticalSection				PendingLock;
+	FCriticalSection				CompletedLock;
+	mutable FCriticalSection		NumPendingLock;
+};
+
+//////////////////////////////////////////////////////////////////////////
+FIoQueue::FIoQueue(FIoDispatcher& IoDispatcher, FBatchReadyCallback BatchReadyCallback)
+: Impl(new FIoQueueImpl(*IoDispatcher.Impl, MoveTemp(BatchReadyCallback)))
+{ 
+}
+
+FIoQueue::~FIoQueue() = default;
+
+void
+FIoQueue::Enqueue(const FIoChunkId& ChunkId, FIoReadOptions ReadOptions, uint64 UserData, bool bDeferBatch)
+{
+	Impl->Enqueue(ChunkId, ReadOptions, UserData, bDeferBatch);
+}
+
+bool
+FIoQueue::Dequeue(FIoChunkId& ChunkId, TIoStatusOr<FIoBuffer>& Result, uint64& UserData)
+{
+	return Impl->Dequeue(ChunkId, Result, UserData);
+}
+
+void
+FIoQueue::IssueBatch()
+{
+	return Impl->IssueBatch();
+}
+
+bool
+FIoQueue::IsEmpty() const
+{
+	return Impl->IsEmpty();
+}
diff --git a/Engine/Source/Runtime/Core/Public/IO/IoDispatcher.h b/Engine/Source/Runtime/Core/Public/IO/IoDispatcher.h
--- a/Engine/Source/Runtime/Core/Public/IO/IoDispatcher.h
+++ b/Engine/Source/Runtime/Core/Public/IO/IoDispatcher.h
@@ -8,6 +8,7 @@
 #include "Misc/StringBuilder.h"
 #include "Templates/RefCounting.h"
 #include "Templates/UnrealTemplate.h"
+#include "Templates/TypeCompatibleBytes.h"
 #include "HAL/PlatformAtomics.h"
 
 #if __cplusplus >= 201703L
@@ -23,7 +24,9 @@ class FIoDispatcher;
 class FIoStoreReader;
 class FIoStoreWriter;
 
+class FIoQueueImpl;
 class FIoRequestImpl;
+class FIoBatchImpl;
 class FIoDispatcherImpl;
 class FIoStoreReaderImpl;
 class FIoStoreWriterImpl;
@@ -51,8 +54,8 @@ enum class EIoErrorCode
 class FIoStatus
 {
 public:
-						FIoStatus() = default;
-						~FIoStatus() = default;
+	CORE_API			FIoStatus();
+	CORE_API			~FIoStatus();
 
 	CORE_API			FIoStatus(EIoErrorCode Code, const FStringView& InErrorMessage);
 	CORE_API			FIoStatus(EIoErrorCode Code);
@@ -66,13 +69,16 @@ public:
 	inline EIoErrorCode	GetErrorCode() const { return ErrorCode; }
 	CORE_API FString	ToString() const;
 
-	static const FIoStatus Ok;
-	static const FIoStatus Unknown;
-	static const FIoStatus Invalid;
+	CORE_API static const FIoStatus Ok;
+	CORE_API static const FIoStatus Unknown;
+	CORE_API static const FIoStatus Invalid;
 
 private:
+	static constexpr int32 MaxErrorMessageLength = 256;
+	using FErrorMessage = TCHAR[MaxErrorMessageLength];
+
 	EIoErrorCode	ErrorCode = EIoErrorCode::Ok;
-	FString			ErrorMessage;
+	FErrorMessage	ErrorMessage;
 
 	friend class FIoStatusBuilder;
 };
@@ -102,17 +108,28 @@ class TIoStatusOr
 	template<typename U> friend class TIoStatusOr;
 
 public:
-	TIoStatusOr() : StatusValue(FIoStatus::Unknown), Value{} {}
-	TIoStatusOr(const FIoStatus& Status);
+	TIoStatusOr() : StatusValue(FIoStatus::Unknown) { }
+	TIoStatusOr(const TIoStatusOr& Other);
+	explicit TIoStatusOr(TIoStatusOr&& Other);
+
+	TIoStatusOr(FIoStatus Status);
 	TIoStatusOr(const T& Value);
+	explicit TIoStatusOr(T&& Value);
+
+	~TIoStatusOr();
+
+	template <typename... ArgTypes>
+	explicit TIoStatusOr(ArgTypes&&... Args);
 
 	template<typename U>
 	TIoStatusOr(const TIoStatusOr<U>& Other);
 
-	TIoStatusOr<T>& operator=(const TIoStatusOr<T>& other);
+	TIoStatusOr<T>& operator=(const TIoStatusOr<T>& Other);
+	TIoStatusOr<T>& operator=(TIoStatusOr<T>&& Other);
+	TIoStatusOr<T>& operator=(T&& OtherValue);
 
 	template<typename U>
-	TIoStatusOr<T>& operator=(const TIoStatusOr<U>& other);
+	TIoStatusOr<T>& operator=(const TIoStatusOr<U>& Other);
 
 	const FIoStatus&	Status() const;
 	bool				IsOk() const;
@@ -123,8 +140,8 @@ public:
 	void				Reset();
 
 private:
-	FIoStatus	StatusValue;
-	T			Value;
+	FIoStatus				StatusValue;
+	TTypeCompatibleBytes<T>	Value;
 };
 
 CORE_API void StatusOrCrash(const FIoStatus& Status);
@@ -132,8 +149,12 @@ CORE_API void StatusOrCrash(const FIoStatus& Status);
 template<typename T>
 void TIoStatusOr<T>::Reset()
 {
+	if (StatusValue.IsOk())
+	{
+		((T*)&Value)->~T();
+	}
+
 	StatusValue = FIoStatus::Unknown;
-	Value = T();
 }
 
 template<typename T>
@@ -144,7 +165,7 @@ const T& TIoStatusOr<T>::ValueOrDie()
 		StatusOrCrash(StatusValue);
 	}
 
-	return Value;
+	return *Value.GetTypedPtr();
 }
 
 template<typename T>
@@ -155,28 +176,64 @@ T TIoStatusOr<T>::ConsumeValueOrDie()
 		StatusOrCrash(StatusValue);
 	}
 
-	return MoveTemp(Value);
+	StatusValue = FIoStatus::Unknown;
+
+	return MoveTemp(*Value.GetTypedPtr());
 }
 
 template<typename T>
-TIoStatusOr<T>::TIoStatusOr(const FIoStatus& Status)
+TIoStatusOr<T>::TIoStatusOr(const TIoStatusOr& Other)
 {
-	if (Status.IsOk())
+	StatusValue = Other.StatusValue;
+	if (StatusValue.IsOk())
 	{
-		// This doesn't make a whole lot of sense. If everything's ok then
-		// we should be returning a T
-		StatusValue = FIoStatus::Invalid;
+		new(&Value) T(*(const T*)&Other.Value);
 	}
-	else
+}
+
+template<typename T>
+TIoStatusOr<T>::TIoStatusOr(TIoStatusOr&& Other)
+{
+	StatusValue = Other.StatusValue;
+	if (StatusValue.IsOk())
 	{
-		StatusValue = Status;
+		new(&Value) T(MoveTempIfPossible(*(T*)&Other.Value));
 	}
 }
 
+template<typename T>
+TIoStatusOr<T>::TIoStatusOr(FIoStatus Status)
+{
+	check(!Status.IsOk());
+	StatusValue = Status;
+}
+
 template<typename T>
 TIoStatusOr<T>::TIoStatusOr(const T& InValue)
 {
-	Value = InValue;
+	StatusValue = FIoStatus::Ok;
+	new(&Value) T(InValue);
+}
+
+template<typename T>
+TIoStatusOr<T>::TIoStatusOr(T&& Value)
+{
+	StatusValue = FIoStatus::Ok;
+	new(&Value) T(MoveTempIfPossible(Value));
+}
+
+template <typename T>
+template <typename... ArgTypes>
+TIoStatusOr<T>::TIoStatusOr(ArgTypes&&... Args)
+{
+	StatusValue = FIoStatus::Ok;
+	new(&Value) T(Forward<ArgTypes>(Args)...);
+}
+
+template<typename T>
+TIoStatusOr<T>::~TIoStatusOr()
+{
+	Reset();
 }
 
 template<typename T>
@@ -195,28 +252,69 @@ template<typename T>
 TIoStatusOr<T>&
 TIoStatusOr<T>::operator=(const TIoStatusOr<T>& Other)
 {
+	Reset();
 	StatusValue = Other.StatusValue;
-	Value = Other.Value;
+
+	if (StatusValue.IsOk())
+	{
+		new(&Value) T(*(const T*)&Other.Value);
+	}
+
+	return *this;
+}
+
+template<typename T>
+TIoStatusOr<T>&
+TIoStatusOr<T>::operator=(TIoStatusOr<T>&& Other)
+{
+	if (&Other != this)
+	{
+		Reset();
+		StatusValue = Other.StatusValue;
+ 
+		if (StatusValue.IsOk())
+		{
+			new(&Value) T(MoveTempIfPossible(*(T*)&Other.Value));
+		}
+	}
+
 	return *this;
 }
 
+template<typename T>
+TIoStatusOr<T>&
+TIoStatusOr<T>::operator=(T&& OtherValue)
+{
+	Reset();
+	
+	StatusValue = FIoStatus::Ok;
+	new(&Value) T(MoveTempIfPossible(Value));
+}
+
 template<typename T>
 template<typename U>
 TIoStatusOr<T>::TIoStatusOr(const TIoStatusOr<U>& Other)
 :	StatusValue(Other.StatusValue)
-,	Value(Other.StatusValue.IsOk() ? Other.Value : T()) 
 {
+	if (StatusValue.IsOk())
+	{
+		new(&Value) T(*(const U*)&Other.Value);
+	}
 }
 
 template<typename T>
 template<typename U>
 TIoStatusOr<T>& TIoStatusOr<T>::operator=(const TIoStatusOr<U>& Other)
 {
+	Reset();
+
 	StatusValue = Other.StatusValue;
+
 	if (StatusValue.IsOk())
 	{
-		Value = Other.Value;
+		new(&Value) T(*(const U*)&Other.Value);
 	}
+
 	return *this;
 }
 
@@ -435,19 +533,21 @@ public:
 	FIoRequest(const FIoRequest&) = default;
 	FIoRequest& operator=(const FIoRequest&) = default;
 
-	CORE_API bool				IsOk() const;
-	CORE_API FIoStatus			Status() const;
-	CORE_API FIoBuffer			GetChunk();
-	CORE_API const FIoChunkId&	GetChunkId() const;
+	CORE_API bool							IsOk() const;
+	CORE_API FIoStatus						Status() const;
+	CORE_API FIoBuffer						GetChunk();
+	CORE_API const FIoChunkId&				GetChunkId() const;
+	CORE_API const TIoStatusOr<FIoBuffer>&	GetResult() const;
 
 private:
 	FIoRequestImpl* Impl = nullptr;
 
-	explicit FIoRequest(FIoRequestImpl& InImpl)
-	: Impl(&InImpl)
+	explicit FIoRequest(FIoRequestImpl* InImpl)
+	: Impl(InImpl)
 	{
 	}
 
+	friend class FIoDispatcher;
 	friend class FIoDispatcherImpl;
 	friend class FIoBatch;
 };
@@ -461,7 +561,7 @@ class FIoBatch
 {
 	friend class FIoDispatcher;
 
-	FIoBatch(FIoDispatcherImpl* OwningIoDispatcher, uint32 InBatchId);
+	FIoBatch(FIoDispatcherImpl* InDispatcher, FIoBatchImpl* InImpl);
 
 public:
 	CORE_API FIoRequest Read(const FIoChunkId& Chunk, FIoReadOptions Options);
@@ -474,7 +574,7 @@ public:
 
 private:
 	FIoDispatcherImpl*	Dispatcher		= nullptr;
-	uint32				BatchId			= ~uint32(0);
+	FIoBatchImpl*		Impl			= nullptr;
 	FEvent*				CompletionEvent = nullptr;
 };
 
@@ -496,10 +596,11 @@ public:
 	FIoDispatcher& operator=(const FIoDispatcher&) = delete;
 
 private:
-	FIoDispatcherImpl* Impl;
+	FIoDispatcherImpl* Impl = nullptr;
 
 	friend class FIoRequest;
 	friend class FIoBatch;
+	friend class FIoQueue;
 };
 
 //////////////////////////////////////////////////////////////////////////
@@ -560,3 +661,24 @@ private:
 };
 
 //////////////////////////////////////////////////////////////////////////
+
+class FIoQueue
+{
+public:
+	using FBatchReadyCallback = TFunction<void()>;
+
+	CORE_API		FIoQueue(FIoDispatcher& IoDispatcher, FBatchReadyCallback BatchReadyCallback);
+	CORE_API		~FIoQueue();
+
+					FIoQueue(const FIoQueue&) = delete;
+					FIoQueue(FIoQueue&&) = delete;
+					FIoQueue& operator=(const FIoQueue&) = delete;
+					FIoQueue& operator=(FIoQueue&&) = delete;
+
+	CORE_API void	Enqueue(const FIoChunkId& ChunkId, FIoReadOptions ReadOptions, uint64 UserData, bool bDeferBatch = false);
+	CORE_API bool	Dequeue(FIoChunkId& ChunkId, TIoStatusOr<FIoBuffer>& Result, uint64& UserData);
+	CORE_API void	IssueBatch();
+	CORE_API bool	IsEmpty() const;
+private:
+	TUniquePtr<FIoQueueImpl> Impl;
+};
diff --git a/Engine/Source/Runtime/CoreUObject/Private/Serialization/AsyncLoading2.cpp b/Engine/Source/Runtime/CoreUObject/Private/Serialization/AsyncLoading2.cpp
--- a/Engine/Source/Runtime/CoreUObject/Private/Serialization/AsyncLoading2.cpp
+++ b/Engine/Source/Runtime/CoreUObject/Private/Serialization/AsyncLoading2.cpp
@@ -43,7 +43,6 @@
 #include "Serialization/AsyncPackage.h"
 #include "Serialization/Zenaphore.h"
 #include "IO/IoDispatcher.h"
-#include "Containers/CircularQueue.h"
 #include "UObject/GCObject.h"
 #include "UObject/ObjectRedirector.h"
 
@@ -411,164 +410,6 @@ uint16 GetChunkIndex(const FIoChunkId& ChunkId)
 
 }
 
-class FIoRequestQueue final : private FRunnable
-{
-public:
-	struct FCompletionEvent { FAsyncPackage2* Package; FIoBatch IoBatch; FIoRequest IoRequest; };
-
-	FIoRequestQueue(FIoDispatcher& InIoDispatcher, FZenaphore& InZenaphore, const uint32 InCapacity = 1024);
-	~FIoRequestQueue();
-	void EnqueueRequest(FAsyncPackage2* Package, const FIoChunkId& ChunkId);
-	TOptional<FCompletionEvent> DequeueCompletionEvent();
-	bool HasPendingRequests() const;
-	bool WaitForRequest(float SecondsToWait);
-
-private:
-	struct FSubmissionRequest { FAsyncPackage2* Package; FIoChunkId ChunkId; };
-	using FPendingRequest = FCompletionEvent;
-
-	uint32 Run() override;
-	void Stop() override;
-
-	FIoDispatcher& IoDispatcher;
-	FZenaphore& Zenaphore;
-	FRunnableThread* Thread = nullptr;
-	FEvent* WakeUpEvent = nullptr;
-	bool bIsRunning = false;
-	TCircularQueue<FSubmissionRequest> RequestQueue;
-	FCriticalSection RequestQueueLock;
-	TCircularQueue<FCompletionEvent> CompletionQueue;
-	FCriticalSection CompletionQueueLock;
-	TArray<FPendingRequest> PendingRequests;
-	TAtomic<int32> NumPendingRequests;
-};
-
-FIoRequestQueue::FIoRequestQueue(FIoDispatcher& InIoDispatcher, FZenaphore& InZenaphore, const uint32 InCapacity)
-: IoDispatcher(InIoDispatcher)
-, Zenaphore(InZenaphore)
-, RequestQueue(InCapacity)
-, CompletionQueue(InCapacity)
-, NumPendingRequests(0)
-{
-	WakeUpEvent = FGenericPlatformProcess::GetSynchEventFromPool(false);
-	Thread = FRunnableThread::Create(this, TEXT("IoRequestQueue"), 0, TPri_Normal);
-}
-
-FIoRequestQueue::~FIoRequestQueue()
-{
-	Stop();
-	Thread->Kill(true);
-	Thread = nullptr;
-	FPlatformProcess::ReturnSynchEventToPool(WakeUpEvent);
-}
-
-void FIoRequestQueue::EnqueueRequest(FAsyncPackage2* Package, const FIoChunkId& ChunkId)
-{
-	FScopeLock _(&RequestQueueLock);
-	bool bEnqueued = RequestQueue.Enqueue(FSubmissionRequest { Package, ChunkId });
-	NumPendingRequests++;
-	check(bEnqueued);
-	WakeUpEvent->Trigger();
-}
-
-TOptional<FIoRequestQueue::FCompletionEvent> FIoRequestQueue::DequeueCompletionEvent()
-{
-	TOptional<FCompletionEvent> Out;
-	FScopeLock _(&CompletionQueueLock);
-	if (const FCompletionEvent* CompletionEvent = CompletionQueue.Peek())
-	{
-		Out = *CompletionEvent;
-		CompletionQueue.Dequeue();
-	}
-
-	return Out;
-}
-
-bool FIoRequestQueue::HasPendingRequests() const
-{
-	return NumPendingRequests > 0 || !CompletionQueue.IsEmpty() || !RequestQueue.IsEmpty();
-}
-
-bool FIoRequestQueue::WaitForRequest(float SecondsToWait)
-{
-	return true;
-}
-
-void FIoRequestQueue::Stop()
-{
-	if (bIsRunning)
-	{
-		bIsRunning = false;
-		WakeUpEvent->Trigger();
-	}
-}
-
-uint32 FIoRequestQueue::Run()
-{
-	bIsRunning = true;
-
-	while (bIsRunning)
-	{
-		if (!bIsRunning)
-		{
-			break;
-		}
-
-		// Process incoming
-		{
-			FSubmissionRequest Request;
-			bool bDequeuedRequest = false;
-			{
-				FScopeLock _(&RequestQueueLock);
-				bDequeuedRequest = RequestQueue.Dequeue(Request);
-			}
-			
-			if (bDequeuedRequest)
-			{
-				FIoBatch IoBatch = IoDispatcher.NewBatch();
-				FIoRequest IoRequest = IoBatch.Read(Request.ChunkId, FIoReadOptions());
-				IoBatch.Issue();
-				PendingRequests.Add(FPendingRequest { Request.Package, IoBatch, IoRequest });
-			}
-		}
-
-		// Process pending
-		{
-			int32 NumPending = PendingRequests.Num();
-			for (int32 I = 0; I < PendingRequests.Num();)
-			{
-				FPendingRequest& PendingRequest = PendingRequests[I];
-				if (PendingRequest.IoRequest.Status().IsCompleted())
-				{
-					{
-						FScopeLock _(&CompletionQueueLock);
-						bool bEnqueued = CompletionQueue.Enqueue(PendingRequest);
-						check(bEnqueued);
-					}
-					PendingRequests.RemoveAtSwap(I, 1, false);
-					NumPendingRequests--;
-				}
-				else
-				{
-					++I;
-				}
-			}
-
-			if (NumPending != PendingRequests.Num())
-			{
-				Zenaphore.NotifyOne();
-			}
-		}
-
-		if (PendingRequests.Num() == 0 && NumPendingRequests == 0)
-		{
-			WakeUpEvent->Wait();
-		}
-	}
-
-	return 0;
-}
-
 enum class EAsyncPackageLoadingState2 : uint8
 {
 	NewPackage,
@@ -1169,7 +1010,7 @@ public:
 	static EAsyncPackageState::Type Event_Tick(FAsyncPackage2* Package, int32);
 	static EAsyncPackageState::Type Event_Delete(FAsyncPackage2* Package, int32);
 
-	void ProcessIoRequest(FIoRequest& IoRequest);
+	void ProcessIoRequest(const FIoChunkId& ChunkId, TIoStatusOr<FIoBuffer>& Result);
 
 	void MarkNewObjectForLoadIfItIsAnExport(UObject *Object);
 
@@ -1398,9 +1239,8 @@ private:
 	/** I/O Dispatcher */
 	FGlobalNameMap GlobalNameMap;
 	FIoStoreEnvironment IoStoreEnvironment;
-	TUniquePtr<FIoStoreReader> IoStoreReader;
 	FIoDispatcher IoDispatcher;
-	TUniquePtr<FIoRequestQueue> IoRequestQueue;
+	TUniquePtr<FIoQueue> IoQueue;
 
 	/** Package store */
 	FPackageStoreEntryRuntime* StoreEntriesRuntime = nullptr;
@@ -1804,13 +1644,17 @@ void FAsyncLoadingThread2Impl::InitializeLoading()
 		TRACE_CPUPROFILER_EVENT_SCOPE(InitIoDispatcher);
 
 		IoStoreEnvironment.InitializeFileEnvironment(RootDir);
-		IoStoreReader = MakeUnique<FIoStoreReader>(IoStoreEnvironment);
+		TUniquePtr<FIoStoreReader> IoStoreReader = MakeUnique<FIoStoreReader>(IoStoreEnvironment);
 		FIoStatus ReaderStatus = IoStoreReader->Initialize(TEXT("PackageLoader"));
 
 		UE_CLOG(!ReaderStatus.IsOk(), LogStreaming, Error, TEXT("Failed to initialize I/O dispatcher: '%s'"), *ReaderStatus.ToString());
 
-		IoDispatcher.Mount(IoStoreReader.Get());
-		IoRequestQueue = MakeUnique<FIoRequestQueue>(IoDispatcher, AltZenaphore, 131072);
+		IoDispatcher.Mount(IoStoreReader.Release());
+
+		IoQueue = MakeUnique<FIoQueue>(IoDispatcher, TFunction<void()>([this]()
+		{
+			AltZenaphore.NotifyOne();
+		}));
 	}
 
 	{
@@ -3294,26 +3138,24 @@ EAsyncPackageState::Type FAsyncPackage2::Event_Delete(FAsyncPackage2* Package, i
 	return EAsyncPackageState::Complete;
 }
 
-void FAsyncPackage2::ProcessIoRequest(FIoRequest& IoRequest)
+void FAsyncPackage2::ProcessIoRequest(const FIoChunkId& ChunkId, TIoStatusOr<FIoBuffer>& Result)
 {
-	UE_CLOG(!IoRequest.IsOk(), LogStreaming, Warning, TEXT("I/O Error: '%s', package: '%s'"), *IoRequest.Status().ToString(), *Desc.NameToLoad.ToString()); 
-
-	FIoChunkId ChunkId = IoRequest.GetChunkId();
+	UE_CLOG(!Result.IsOk(), LogStreaming, Warning, TEXT("I/O Error: '%s', package: '%s'"), *Result.Status().ToString(), *Desc.NameToLoad.ToString()); 
+	
+	const FIoBuffer& IoBuffer = Result.ValueOrDie();
 
 	switch (GetChunkType(ChunkId))
 	{
 	case EChunkType::PackageSummary:
 	{
-		FIoBuffer SummaryIoBuffer = IoRequest.GetChunk();
-		PackageSummaryBuffer = MakeUnique<uint8[]>(SummaryIoBuffer.DataSize());
-		FMemory::Memcpy(PackageSummaryBuffer.Get(), SummaryIoBuffer.Data(), SummaryIoBuffer.DataSize());
+		PackageSummaryBuffer = MakeUnique<uint8[]>(IoBuffer.DataSize());
+		FMemory::Memcpy(PackageSummaryBuffer.Get(), IoBuffer.Data(), IoBuffer.DataSize());
 		GetNode(EEventLoadNode2::Package_LoadSummary)->ReleaseBarrier();
 		break;
 	}
 	case EChunkType::ExportData:
 	{
 		const int32 LocalExportIndex = GetChunkIndex(ChunkId);
-		FIoBuffer IoBuffer = IoRequest.GetChunk();
 		check(ExportIoBuffers.Num() > LocalExportIndex);
 		ExportIoBuffers[LocalExportIndex] = IoBuffer;
 		GetNode(EEventLoadNode2::ImportOrExport_Serialize, FPackageIndex::FromExport(LocalExportIndex))->ReleaseBarrier();
@@ -3398,13 +3240,15 @@ EAsyncPackageState::Type FAsyncLoadingThread2Impl::ProcessAsyncLoadingFromGameTh
 			}
 
 			bool bProcessedResourceCompletionEvents = false;
-			TOptional<FIoRequestQueue::FCompletionEvent> IoCompletionEvent = IoRequestQueue->DequeueCompletionEvent();
-			while (IoCompletionEvent)
+			FIoChunkId ChunkId;
+			TIoStatusOr<FIoBuffer> Result;
+			uint64 UserData;
+
+			while(IoQueue->Dequeue(ChunkId, Result, UserData))
 			{
-				FAsyncPackage2* Package = IoCompletionEvent.GetValue().Package;
-				Package->ProcessIoRequest(IoCompletionEvent.GetValue().IoRequest);
-				IoDispatcher.FreeBatch(IoCompletionEvent.GetValue().IoBatch);
-				IoCompletionEvent = IoRequestQueue->DequeueCompletionEvent();
+				FAsyncPackage2* Package = reinterpret_cast<FAsyncPackage2*>(UserData);
+				Package->ProcessIoRequest(ChunkId, Result);
+				bProcessedResourceCompletionEvents = true;
 			}
 			if (bProcessedResourceCompletionEvents)
 			{
@@ -3427,7 +3271,7 @@ EAsyncPackageState::Type FAsyncLoadingThread2Impl::ProcessAsyncLoadingFromGameTh
 				break;
 			}
 
-			if (IoRequestQueue->HasPendingRequests())
+			if (!IoQueue->IsEmpty())
 			{
 				FPlatformProcess::Sleep(0.001f);
 			}
@@ -3957,15 +3801,17 @@ uint32 FAsyncLoadingThread2Impl::Run()
 				do
 				{
 					bDidSomething = false;
+					FIoChunkId ChunkId;
+					TIoStatusOr<FIoBuffer> Result;
+					uint64 UserData;
+
 					for (;;)
 					{
-						TOptional<FIoRequestQueue::FCompletionEvent> IoCompletionEvent = IoRequestQueue->DequeueCompletionEvent();
-						if (IoCompletionEvent)
+						if (IoQueue->Dequeue(ChunkId, Result, UserData))
 						{
 							TRACE_CPUPROFILER_EVENT_SCOPE(ProcessResourceCompletionEvent);
-							FAsyncPackage2* Package = IoCompletionEvent.GetValue().Package;
-							Package->ProcessIoRequest(IoCompletionEvent.GetValue().IoRequest);
-							IoDispatcher.FreeBatch(IoCompletionEvent.GetValue().IoBatch);
+							FAsyncPackage2* Package = reinterpret_cast<FAsyncPackage2*>(UserData);
+							Package->ProcessIoRequest(ChunkId, Result);
 							bDidSomething = true;
 						}
 						else
@@ -4269,7 +4115,7 @@ void FAsyncLoadingThread2Impl::FireCompletedCompiledInImport(FGCObject* AsyncPac
 
 void FAsyncLoadingThread2Impl::EnqueueIoRequest(FAsyncPackage2* Package, const FIoChunkId& ChunkId)
 {
-	IoRequestQueue->EnqueueRequest(Package, ChunkId);
+	IoQueue->Enqueue(ChunkId, FIoReadOptions(), reinterpret_cast<uint64>(Package));
 }
 
 /*-----------------------------------------------------------------------------
