Removed FAsyncUncompress, added override specifiers to virtual functions

diff --git a/Engine/Source/Runtime/Core/Public/Async/AsyncWork.h b/Engine/Source/Runtime/Core/Public/Async/AsyncWork.h
--- a/Engine/Source/Runtime/Core/Public/Async/AsyncWork.h
+++ b/Engine/Source/Runtime/Core/Public/Async/AsyncWork.h
@@ -290,7 +290,7 @@ class FAsyncTask
 	/** 
 	* Performs the work, this is only called from a pool thread.
 	**/
-	virtual void DoThreadedWork()
+	virtual void DoThreadedWork() override
 	{
 		DoWork();
 		FinishThreadedWork();
@@ -300,7 +300,7 @@ class FAsyncTask
 	 * Always called from the thread pool. Called if the task is removed from queue before it has started which might happen at exit.
 	 * If the user job can abandon, we do that, otherwise we force the work to be done now (doing nothing would not be safe).
 	 */
-	virtual void Abandon(void)
+	virtual void Abandon(void) override
 	{
 		if (Task.CanAbandon())
 		{
@@ -378,7 +378,7 @@ public:
 	}
 
 	/* Retrieve embedded user job, not legal to call while a job is in process
-		* @return reference to embedded user job 
+	* @return reference to embedded user job 
 	**/
 	TTask &GetTask()
 	{
@@ -387,7 +387,7 @@ public:
 	}
 
 	/* Retrieve embedded user job, not legal to call while a job is in process
-		* @return reference to embedded user job 
+	* @return reference to embedded user job 
 	**/
 	const TTask &GetTask() const
 	{
@@ -497,7 +497,7 @@ public:
 
 	/** Returns true if the work and TASK has completed, false while it's still in progress. 
 	 * prior to returning true, it synchronizes so the task can be destroyed or reused
-	*/
+	 */
 	bool IsDone()
 	{
 		if (!IsWorkDone())
@@ -521,12 +521,13 @@ public:
 		return true;
 	}
 
-	/** Returns true if the work has not been started or has been completed. NOT to be used for synchronization, but great for check()'s */
+	/** Returns true if the work has not been started or has been completed. 
+	 * NOT to be used for synchronization, but great for check()'s 
+	 */
 	bool IsIdle() const
 	{
 		return WorkNotFinishedCounter.GetValue() == 0 && QueuedPool == 0;
 	}
-
 };
 
 /**
@@ -543,80 +544,3 @@ public:
 	{
 	}
 };
-
-/**
- * Asynchronous decompression, used for decompressing chunks of memory in the background
- */
-class FAsyncUncompress : public FNonAbandonableTask
-{
-	/** Format to use										*/
-	FName CompressionFormat;
-	/** Buffer containing uncompressed data					*/
-	void*	UncompressedBuffer;
-	/** Size of uncompressed data in bytes					*/
-	int32		UncompressedSize;
-	/** Buffer compressed data is going to be written to	*/
-	void*	CompressedBuffer;
-	/** Size of CompressedBuffer data in bytes				*/
-	int32		CompressedSize;
-	/** Flags to control decompression						*/
-	ECompressionFlags Flags;
-	/** Whether the source memory is padded with a full cache line at the end */
-	bool	bIsSourceMemoryPadded;
-
-public:
-	/**
-	 * Initializes the data and creates the event.
-	 *
-	 * @param	CompressionFormat	Format to use for decompression
-	 * @param	UncompressedBuffer	Buffer containing uncompressed data
-	 * @param	UncompressedSize	Size of uncompressed data in bytes
-	 * @param	CompressedBuffer	Buffer compressed data is going to be written to
-	 * @param	CompressedSize		Size of CompressedBuffer data in bytes
-	 * @param	Flags				Flags to control what method to use for decompression
-	 */
-	FAsyncUncompress(
-		FName InCompressionFormat,
-		void* InUncompressedBuffer,
-		int32 InUncompressedSize,
-		void* InCompressedBuffer,
-		int32 InCompressedSize,
-		ECompressionFlags InFlags = COMPRESS_NoFlags)
-		: CompressionFormat(InCompressionFormat)
-		, UncompressedBuffer(InUncompressedBuffer)
-		, UncompressedSize(InUncompressedSize)
-		, CompressedBuffer(InCompressedBuffer)
-		, CompressedSize(InCompressedSize)
-		, Flags(InFlags)
-	{
-	}
-
-	UE_DEPRECATED(4.20, "Use FAsyncCompress with FName for Format")
-		FAsyncUncompress(
-			ECompressionFlags InFlags,
-			void* InUncompressedBuffer,
-			int32 InUncompressedSize,
-			void* InCompressedBuffer,
-			int32 InCompressedSize,
-			bool bIsSourcePadded = false)
-		: FAsyncUncompress(FCompression::GetCompressionFormatFromDeprecatedFlags(InFlags), InUncompressedBuffer, InUncompressedSize,
-			InCompressedBuffer, InCompressedSize, (ECompressionFlags)(InFlags | (bIsSourcePadded ? COMPRESS_SourceIsPadded : COMPRESS_NoFlags)))
-	{
-	}
-
-	/**
-	 * Performs the async decompression
-	 */
-	void DoWork()
-	{
-		// Uncompress from memory to memory.
-		verify( FCompression::UncompressMemory( CompressionFormat, UncompressedBuffer, UncompressedSize, CompressedBuffer, CompressedSize, Flags ) );
-	}
-
-	FORCEINLINE TStatId GetStatId() const
-	{
-		RETURN_QUICK_DECLARE_CYCLE_STAT(FAsyncUncompress, STATGROUP_ThreadPoolAsyncTasks);
-	}
-};
-
-
