Removed FGenericPlatformMisc::HandleIOFailure() since it's no longer in use

diff --git a/Engine/Source/Runtime/Core/Private/GenericPlatform/GenericPlatformMisc.cpp b/Engine/Source/Runtime/Core/Private/GenericPlatform/GenericPlatformMisc.cpp
--- a/Engine/Source/Runtime/Core/Private/GenericPlatform/GenericPlatformMisc.cpp
+++ b/Engine/Source/Runtime/Core/Private/GenericPlatform/GenericPlatformMisc.cpp
@@ -459,11 +459,6 @@ void FGenericPlatformMisc::MemoryBarrier()
 {
 }
 
-void FGenericPlatformMisc::HandleIOFailure( const TCHAR* Filename )
-{
-	UE_LOG(LogGenericPlatformMisc, Fatal,TEXT("I/O failure operating on '%s'"), Filename ? Filename : TEXT("Unknown file"));
-}
-
 void FGenericPlatformMisc::RaiseException(uint32 ExceptionCode)
 {
 	/** This is the last place to gather memory stats before exception. */
diff --git a/Engine/Source/Runtime/Core/Public/GenericPlatform/GenericPlatformMisc.h b/Engine/Source/Runtime/Core/Public/GenericPlatform/GenericPlatformMisc.h
--- a/Engine/Source/Runtime/Core/Public/GenericPlatform/GenericPlatformMisc.h
+++ b/Engine/Source/Runtime/Core/Public/GenericPlatform/GenericPlatformMisc.h
@@ -578,13 +578,6 @@ struct CORE_API FGenericPlatformMisc
 	 */
 	FORCENOINLINE static void MemoryBarrier();
 
-	/**
-	 * Handles IO failure by ending gameplay.
-	 *
-	 * @param Filename	If not nullptr, name of the file the I/O error occured with
-	 */
-	void static HandleIOFailure( const TCHAR* Filename );
-
 	/**
 	 * Set a handler to be called when there is a memory warning from the OS 
 	 *
