Add TryMalloc(), which may return a nullptr if the memory allocation request cannot be fulfilled

diff --git a/Engine/Source/Runtime/Core/Private/HAL/MallocAnsi.cpp b/Engine/Source/Runtime/Core/Private/HAL/MallocAnsi.cpp
--- a/Engine/Source/Runtime/Core/Private/HAL/MallocAnsi.cpp
+++ b/Engine/Source/Runtime/Core/Private/HAL/MallocAnsi.cpp
@@ -151,35 +151,41 @@ FMallocAnsi::FMallocAnsi()
 #endif
 }
 
-void* FMallocAnsi::Malloc( SIZE_T Size, uint32 Alignment )
+void* FMallocAnsi::TryMalloc(SIZE_T Size, uint32 Alignment)
 {
 #if !UE_BUILD_SHIPPING
 	uint64 LocalMaxSingleAlloc = MaxSingleAlloc.Load(EMemoryOrder::Relaxed);
 	if (LocalMaxSingleAlloc != 0 && Size > LocalMaxSingleAlloc)
 	{
-		FPlatformMemory::OnOutOfMemory(Size, Alignment);
 		return nullptr;
 	}
 #endif
 
 	Alignment = FMath::Max(Size >= 16 ? (uint32)16 : (uint32)8, Alignment);
 
-	void* Result = AnsiMalloc(Size, Alignment); 
+	void* Result = AnsiMalloc(Size, Alignment);
+
+	return Result;
+}
+
+void* FMallocAnsi::Malloc(SIZE_T Size, uint32 Alignment)
+{
+	void* Result = TryMalloc(Size, Alignment); 
 
-	if (Result == nullptr)
+	if (Result == nullptr && Size)
 	{
 		FPlatformMemory::OnOutOfMemory(Size, Alignment);
 	}
+
 	return Result;
 }
 
-void* FMallocAnsi::Realloc( void* Ptr, SIZE_T NewSize, uint32 Alignment )
+void* FMallocAnsi::TryRealloc(void* Ptr, SIZE_T NewSize, uint32 Alignment)
 {
 #if !UE_BUILD_SHIPPING
 	uint64 LocalMaxSingleAlloc = MaxSingleAlloc.Load(EMemoryOrder::Relaxed);
 	if (LocalMaxSingleAlloc != 0 && NewSize > LocalMaxSingleAlloc)
 	{
-		FPlatformMemory::OnOutOfMemory(NewSize, Alignment);
 		return nullptr;
 	}
 #endif
@@ -188,6 +194,13 @@ void* FMallocAnsi::Realloc( void* Ptr, SIZE_T NewSize, uint32 Alignment )
 
 	void* Result = AnsiRealloc(Ptr, NewSize, Alignment);
 
+	return Result;
+}
+
+void* FMallocAnsi::Realloc( void* Ptr, SIZE_T NewSize, uint32 Alignment )
+{
+	void* Result = TryRealloc(Ptr, NewSize, Alignment);
+
 	if (Result == nullptr && NewSize != 0)
 	{
 		FPlatformMemory::OnOutOfMemory(NewSize, Alignment);
diff --git a/Engine/Source/Runtime/Core/Private/HAL/MallocStomp.cpp b/Engine/Source/Runtime/Core/Private/HAL/MallocStomp.cpp
--- a/Engine/Source/Runtime/Core/Private/HAL/MallocStomp.cpp
+++ b/Engine/Source/Runtime/Core/Private/HAL/MallocStomp.cpp
@@ -47,6 +47,18 @@ FMallocStomp::FMallocStomp(const bool InUseUnderrunMode)
 }
 
 void* FMallocStomp::Malloc(SIZE_T Size, uint32 Alignment)
+{
+	void* Result = TryMalloc(Size, Alignment);
+
+	if (Result == nullptr)
+	{
+		FPlatformMemory::OnOutOfMemory(Size, Alignment);
+	}
+
+	return Result;
+}
+
+void* FMallocStomp::TryMalloc(SIZE_T Size, uint32 Alignment)
 {
 	if (Size == 0U)
 	{
@@ -93,11 +105,10 @@ void* FMallocStomp::Malloc(SIZE_T Size, uint32 Alignment)
 
 	if (!FullAllocationPointer)
 	{
-		// this is expected not to return
-		FPlatformMemory::OnOutOfMemory(Size, Alignment);
+		return nullptr;
 	}
 
-	void *ReturnedPointer = nullptr;
+	void* ReturnedPointer = nullptr;
 	static const SIZE_T AllocationDataSize = sizeof(FAllocationData);
 
 	const FAllocationData AllocData = { FullAllocationPointer, TotalAllocationSize, AlignedSize, SentinelExpectedValue };
@@ -113,8 +124,8 @@ void* FMallocStomp::Malloc(SIZE_T Size, uint32 Alignment)
 		void* CommittedMemory = VirtualAlloc(AllocDataPointerStart, AllocFullPageSize, MEM_COMMIT, PAGE_READWRITE);
 		if (!CommittedMemory)
 		{
-			// Failed to allocate and commit physical memory pages. Report OOM.
-			FPlatformMemory::OnOutOfMemory(Size, Alignment);
+			// Failed to allocate and commit physical memory pages. 
+			return nullptr;
 		}
 		check(CommittedMemory == AllocDataPointerStart);
 #else
@@ -131,8 +142,8 @@ void* FMallocStomp::Malloc(SIZE_T Size, uint32 Alignment)
 		void* CommittedMemory = VirtualAlloc(FullAllocationPointer, AllocFullPageSize, MEM_COMMIT, PAGE_READWRITE);
 		if (!CommittedMemory)
 		{
-			// Failed to allocate and commit physical memory pages. Report OOM.
-			FPlatformMemory::OnOutOfMemory(Size, Alignment);
+			// Failed to allocate and commit physical memory pages
+			return nullptr;
 		}
 		check(CommittedMemory == FullAllocationPointer);
 #else
@@ -149,24 +160,40 @@ void* FMallocStomp::Malloc(SIZE_T Size, uint32 Alignment)
 
 void* FMallocStomp::Realloc(void* InPtr, SIZE_T NewSize, uint32 Alignment)
 {
-	if(NewSize == 0U)
+	void* Result = TryRealloc(InPtr, NewSize, Alignment);
+
+	if (Result == nullptr && NewSize)
+	{
+		FPlatformMemory::OnOutOfMemory(NewSize, Alignment);
+	}
+
+	return Result;
+}
+
+void* FMallocStomp::TryRealloc(void* InPtr, SIZE_T NewSize, uint32 Alignment)
+{
+	if (NewSize == 0U)
 	{
 		Free(InPtr);
 		return nullptr;
 	}
 
-	void *ReturnPtr = nullptr;
-	if(InPtr != nullptr)
+	void* ReturnPtr = nullptr;
+
+	if (InPtr != nullptr)
 	{
-		ReturnPtr = Malloc(NewSize, Alignment);
+		ReturnPtr = TryMalloc(NewSize, Alignment);
 
-		FAllocationData *AllocDataPtr = reinterpret_cast<FAllocationData*>(reinterpret_cast<uint8*>(InPtr) - sizeof(FAllocationData));
-		FMemory::Memcpy(ReturnPtr, InPtr, FMath::Min(AllocDataPtr->Size, NewSize));
-		Free(InPtr);
+		if (ReturnPtr != nullptr)
+		{
+			FAllocationData* AllocDataPtr = reinterpret_cast<FAllocationData*>(reinterpret_cast<uint8*>(InPtr) - sizeof(FAllocationData));
+			FMemory::Memcpy(ReturnPtr, InPtr, FMath::Min(AllocDataPtr->Size, NewSize));
+			Free(InPtr);
+		}
 	}
 	else
 	{
-		ReturnPtr = Malloc(NewSize, Alignment);
+		ReturnPtr = TryMalloc(NewSize, Alignment);
 	}
 
 	return ReturnPtr;
diff --git a/Engine/Source/Runtime/Core/Private/HAL/MallocTBB.cpp b/Engine/Source/Runtime/Core/Private/HAL/MallocTBB.cpp
--- a/Engine/Source/Runtime/Core/Private/HAL/MallocTBB.cpp
+++ b/Engine/Source/Runtime/Core/Private/HAL/MallocTBB.cpp
@@ -30,13 +30,12 @@ THIRD_PARTY_INCLUDES_START
 #include <tbb/scalable_allocator.h>
 THIRD_PARTY_INCLUDES_END
 
-void* FMallocTBB::Malloc( SIZE_T Size, uint32 Alignment )
+void* FMallocTBB::TryMalloc( SIZE_T Size, uint32 Alignment )
 {
 #if !UE_BUILD_SHIPPING
 	uint64 LocalMaxSingleAlloc = MaxSingleAlloc.Load(EMemoryOrder::Relaxed);
 	if (LocalMaxSingleAlloc != 0 && Size > LocalMaxSingleAlloc)
 	{
-		FPlatformMemory::OnOutOfMemory(Size, Alignment);
 		return nullptr;
 	}
 #endif
@@ -63,26 +62,34 @@ void* FMallocTBB::Malloc( SIZE_T Size, uint32 Alignment )
 	}
 #endif
 
-	if( !NewPtr && Size )
-	{
-		OutOfMemory(Size, Alignment);
-	}
 #if UE_BUILD_DEBUG
-	else if (Size)
+	if (Size && NewPtr != nullptr)
 	{
 		FMemory::Memset(NewPtr, DEBUG_FILL_NEW, scalable_msize(NewPtr));
 	}
 #endif
+
 	return NewPtr;
 }
 
-void* FMallocTBB::Realloc( void* Ptr, SIZE_T NewSize, uint32 Alignment )
+void* FMallocTBB::Malloc(SIZE_T Size, uint32 Alignment)
+{
+	void* Result = TryMalloc(Size, Alignment);
+
+	if (Result == nullptr && Size)
+	{
+		OutOfMemory(Size, Alignment);
+	}
+
+	return Result;
+}
+
+void* FMallocTBB::TryRealloc( void* Ptr, SIZE_T NewSize, uint32 Alignment )
 {
 #if !UE_BUILD_SHIPPING
 	uint64 LocalMaxSingleAlloc = MaxSingleAlloc.Load(EMemoryOrder::Relaxed);
 	if (LocalMaxSingleAlloc != 0 && NewSize > LocalMaxSingleAlloc)
 	{
-		FPlatformMemory::OnOutOfMemory(NewSize, Alignment);
 		return nullptr;
 	}
 #endif
@@ -120,11 +127,20 @@ void* FMallocTBB::Realloc( void* Ptr, SIZE_T NewSize, uint32 Alignment )
 		FMemory::Memset((uint8*)NewPtr + OldSize, DEBUG_FILL_NEW, scalable_msize(NewPtr) -OldSize);
 	}
 #endif
-	if( !NewPtr && NewSize )
+
+	return NewPtr;
+}
+
+void* FMallocTBB::Realloc(void* Ptr, SIZE_T NewSize, uint32 Alignment)
+{
+	void* Result = TryRealloc(Ptr, NewSize, Alignment);
+
+	if (Result == nullptr && NewSize)
 	{
 		OutOfMemory(NewSize, Alignment);
 	}
-	return NewPtr;
+
+	return Result;
 }
 
 void FMallocTBB::Free( void* Ptr )
diff --git a/Engine/Source/Runtime/Core/Private/HAL/MemoryBase.cpp b/Engine/Source/Runtime/Core/Private/HAL/MemoryBase.cpp
--- a/Engine/Source/Runtime/Core/Private/HAL/MemoryBase.cpp
+++ b/Engine/Source/Runtime/Core/Private/HAL/MemoryBase.cpp
@@ -91,3 +91,13 @@ void FMalloc::GetAllocatorStats( FGenericMemoryStats& out_Stats )
 	out_Stats.Add( TEXT("Total Allocator calls"), GetCurrentFrameCalls().AllocatorCalls );
 #endif // STATS
 }
+
+void* FMalloc::TryMalloc(SIZE_T Count, uint32 Alignment)
+{
+	return Malloc(Count, Alignment);
+}
+
+void* FMalloc::TryRealloc(void* Original, SIZE_T Count, uint32 Alignment)
+{
+	return Realloc(Original, Count, Alignment);
+}
diff --git a/Engine/Source/Runtime/Core/Public/HAL/MallocAnsi.h b/Engine/Source/Runtime/Core/Public/HAL/MallocAnsi.h
--- a/Engine/Source/Runtime/Core/Public/HAL/MallocAnsi.h
+++ b/Engine/Source/Runtime/Core/Public/HAL/MallocAnsi.h
@@ -18,7 +18,7 @@ CORE_API void AnsiFree(void* Ptr);
 //
 // ANSI C memory allocator.
 //
-class FMallocAnsi
+class FMallocAnsi final
 	: public FMalloc
 {
 	
@@ -31,8 +31,12 @@ public:
 	// FMalloc interface.
 	virtual void* Malloc( SIZE_T Size, uint32 Alignment ) override;
 
+	virtual void* TryMalloc(SIZE_T Size, uint32 Alignment) override;
+
 	virtual void* Realloc( void* Ptr, SIZE_T NewSize, uint32 Alignment ) override;
 
+	virtual void* TryRealloc(void* Ptr, SIZE_T NewSize, uint32 Alignment) override;
+
 	virtual void Free( void* Ptr ) override;
 
 	virtual bool GetAllocationSize(void *Original, SIZE_T &SizeOut) override;
diff --git a/Engine/Source/Runtime/Core/Public/HAL/MallocStomp.h b/Engine/Source/Runtime/Core/Public/HAL/MallocStomp.h
--- a/Engine/Source/Runtime/Core/Public/HAL/MallocStomp.h
+++ b/Engine/Source/Runtime/Core/Public/HAL/MallocStomp.h
@@ -22,7 +22,7 @@
  * - Read or writes off the beginning of an allocation.
  * - Read or writes after freeing an allocation. 
  */
-class FMallocStomp : public FMalloc
+class FMallocStomp final : public FMalloc
 {
 private:
 #if PLATFORM_64BITS
@@ -71,6 +71,8 @@ public:
 	 */
 	virtual void* Malloc(SIZE_T Size, uint32 Alignment) override;
 
+	virtual void* TryMalloc(SIZE_T Size, uint32 Alignment) override;
+
 	/**
 	 * Changes the size of the memory block pointed to by OldPtr.
 	 * The function may move the memory block to a new location.
@@ -82,6 +84,8 @@ public:
 	 */
 	virtual void* Realloc(void* InPtr, SIZE_T NewSize, uint32 Alignment) override;
 
+	virtual void* TryRealloc(void* InPtr, SIZE_T NewSize, uint32 Alignment) override;
+
 	/**
 	 * Frees a memory allocation and verifies the sentinel in the process.
 	 *
diff --git a/Engine/Source/Runtime/Core/Public/HAL/MallocTBB.h b/Engine/Source/Runtime/Core/Public/HAL/MallocTBB.h
--- a/Engine/Source/Runtime/Core/Public/HAL/MallocTBB.h
+++ b/Engine/Source/Runtime/Core/Public/HAL/MallocTBB.h
@@ -11,14 +11,16 @@
 /**
  * TBB 64-bit scalable memory allocator.
  */
-class FMallocTBB
+class FMallocTBB final
 	: public FMalloc
 {
 public:
 	// FMalloc interface.
 	
-	virtual void* Malloc( SIZE_T Size, uint32 Alignment ) override;
-	virtual void* Realloc( void* Ptr, SIZE_T NewSize, uint32 Alignment ) override;
+	virtual void* Malloc(SIZE_T Size, uint32 Alignment) override;
+	virtual void* TryMalloc(SIZE_T Size, uint32 Alignment) override;
+	virtual void* Realloc(void* Ptr, SIZE_T NewSize, uint32 Alignment) override;
+	virtual void* TryRealloc(void* Ptr, SIZE_T NewSize, uint32 Alignment) override;
 	virtual void Free( void* Ptr ) override;
 	virtual bool GetAllocationSize( void *Original, SIZE_T &SizeOut ) override;
 	virtual void Trim(bool bTrimThreadCaches) override;
diff --git a/Engine/Source/Runtime/Core/Public/HAL/MemoryBase.h b/Engine/Source/Runtime/Core/Public/HAL/MemoryBase.h
--- a/Engine/Source/Runtime/Core/Public/HAL/MemoryBase.h
+++ b/Engine/Source/Runtime/Core/Public/HAL/MemoryBase.h
@@ -81,11 +81,24 @@ public:
 	 */
 	virtual void* Malloc( SIZE_T Count, uint32 Alignment=DEFAULT_ALIGNMENT ) = 0;
 
+	/**
+	 * TryMalloc - like Malloc(), but may return a nullptr result if the allocation
+	 *             request cannot be satisfied.
+	 */
+	virtual void* TryMalloc( SIZE_T Count, uint32 Alignment=DEFAULT_ALIGNMENT );
+
 	/** 
 	 * Realloc
 	 */
 	virtual void* Realloc( void* Original, SIZE_T Count, uint32 Alignment=DEFAULT_ALIGNMENT ) = 0;
 
+	/** 
+	 * TryRealloc - like Realloc(), but may return a nullptr if the allocation
+	 *              request cannot be satisfied. Note that in this case the memory
+	 *              pointed to by Original will still be valid
+	 */
+	virtual void* TryRealloc(void* Original, SIZE_T Count, uint32 Alignment=DEFAULT_ALIGNMENT);
+
 	/** 
 	 * Free
 	 */
