Changed logic in FArchiveFileReaderGeneric::InternalPrecache to avoid ever shrinking buffer. Previously the buffer could be shrunk to a potentially *very* small size if you read data at the end of a file before seeking to the beginning and continuing to read other data.

diff --git a/Engine/Source/Runtime/Core/Private/HAL/FileManagerGeneric.cpp b/Engine/Source/Runtime/Core/Private/HAL/FileManagerGeneric.cpp
--- a/Engine/Source/Runtime/Core/Private/HAL/FileManagerGeneric.cpp
+++ b/Engine/Source/Runtime/Core/Private/HAL/FileManagerGeneric.cpp
@@ -713,7 +713,10 @@ bool FArchiveFileReaderGeneric::InternalPrecache( int64 PrecacheOffset, int64 Pr
 		BufferBase = Pos;
 		int64 BufferCount = FMath::Min( FMath::Min( PrecacheSize,( int64 )(BufferArray.Max() -( Pos&( BufferArray.Max() -1 ) ) ) ), Size-Pos );
 		BufferCount = FMath::Max(BufferCount, 0LL ); // clamp to 0
-		BufferArray.SetNumUninitialized(BufferCount);
+		
+		const bool AllowShrink = false;
+		BufferArray.SetNumUninitialized(BufferCount, AllowShrink);
+
 		int64 Count = 0;
 
 		{
