Cleaned up FQueuedThreadPool interface to make it easier to plug in alternative implementations

diff --git a/Engine/Source/Runtime/Core/Private/HAL/ThreadingBase.cpp b/Engine/Source/Runtime/Core/Private/HAL/ThreadingBase.cpp
--- a/Engine/Source/Runtime/Core/Private/HAL/ThreadingBase.cpp
+++ b/Engine/Source/Runtime/Core/Private/HAL/ThreadingBase.cpp
@@ -469,80 +469,30 @@ class FQueuedThread
 protected:
 
 	/** The event that tells the thread there is work to do. */
-	FEvent* DoWorkEvent;
+	FEvent* DoWorkEvent = nullptr;
 
 	/** If true, the thread should exit. */
-	TAtomic<bool> TimeToDie;
+	TAtomic<bool> TimeToDie { false };
 
 	/** The work this thread is doing. */
-	IQueuedWork* volatile QueuedWork;
+	IQueuedWork* volatile QueuedWork = nullptr;
 
 	/** The pool this thread belongs to. */
-	class FQueuedThreadPool* OwningThreadPool;
+	class FQueuedThreadPoolBase* OwningThreadPool = nullptr;
 
 	/** My Thread  */
-	FRunnableThread* Thread;
+	FRunnableThread* Thread = nullptr;
 
 	/**
 	 * The real thread entry point. It waits for work events to be queued. Once
 	 * an event is queued, it executes it and goes back to waiting.
 	 */
-	virtual uint32 Run() override
-	{
-		while (!TimeToDie.Load(EMemoryOrder::Relaxed))
-		{
-			// This will force sending the stats packet from the previous frame.
-			SET_DWORD_STAT( STAT_ThreadPoolDummyCounter, 0 );
-			// We need to wait for shorter amount of time
-			bool bContinueWaiting = true;
-
-			// Unless we're collecting stats there doesn't appear to be any reason to wake
-			// up again until there's work to do (or it's time to die)
-
-#if STATS
-			if (FThreadStats::IsCollectingData())
-			{
-			while( bContinueWaiting )
-			{				
-				DECLARE_SCOPE_CYCLE_COUNTER( TEXT( "FQueuedThread::Run.WaitForWork" ), STAT_FQueuedThread_Run_WaitForWork, STATGROUP_ThreadPoolAsyncTasks );
-
-				// Wait for some work to do
-
-				bContinueWaiting = !DoWorkEvent->Wait( 10 );
-			}
-			}
-#endif
-
-			if (bContinueWaiting)
-			{
-				DoWorkEvent->Wait();
-			}
-
-			IQueuedWork* LocalQueuedWork = QueuedWork;
-			QueuedWork = nullptr;
-			FPlatformMisc::MemoryBarrier();
-			check(LocalQueuedWork || TimeToDie.Load(EMemoryOrder::Relaxed)); // well you woke me up, where is the job or termination request?
-			while (LocalQueuedWork)
-			{
-				// Tell the object to do the work
-				LocalQueuedWork->DoThreadedWork();
-				// Let the object cleanup before we remove our ref to it
-				LocalQueuedWork = OwningThreadPool->ReturnToPoolOrGetNextJob(this);
-			} 
-		}
-		return 0;
-	}
+	virtual uint32 Run() override;
 
 public:
 
 	/** Default constructor **/
-	FQueuedThread()
-		: DoWorkEvent(nullptr)
-		, TimeToDie(false)
-		, QueuedWork(nullptr)
-		, OwningThreadPool(nullptr)
-		, Thread(nullptr)
-	{ }
+	FQueuedThread() = default;
 
 	/**
 	 * Creates the thread with the specified stack size and creates the various
@@ -553,7 +503,7 @@ public:
 	 * @param ThreadPriority priority of new thread
 	 * @return True if the thread and all of its initialization was successful, false otherwise
 	 */
-	virtual bool Create(class FQueuedThreadPool* InPool,uint32 InStackSize = 0,EThreadPriority ThreadPriority=TPri_Normal)
+	virtual bool Create(class FQueuedThreadPoolBase* InPool,uint32 InStackSize = 0,EThreadPriority ThreadPriority=TPri_Normal)
 	{
 		static int32 PoolThreadIndex = 0;
 		const FString PoolThreadName = FString::Printf( TEXT( "PoolThread %d" ), PoolThreadIndex );
@@ -574,7 +524,7 @@ public:
 	 *
 	 * @return True if the thread exited graceful, false otherwise
 	 */
-	virtual bool KillThread()
+	bool KillThread()
 	{
 		bool bDidExitOK = true;
 		// Tell the thread it needs to die
@@ -609,7 +559,6 @@ public:
 		// Tell the thread to wake up and do its job
 		DoWorkEvent->Trigger();
 	}
-
 };
 
 
@@ -805,7 +754,7 @@ public:
 		return !!QueuedWork.RemoveSingle(InQueuedWork);
 	}
 
-	virtual IQueuedWork* ReturnToPoolOrGetNextJob(FQueuedThread* InQueuedThread) override
+	IQueuedWork* ReturnToPoolOrGetNextJob(FQueuedThread* InQueuedThread)
 	{
 		check(InQueuedThread != nullptr);
 		IQueuedWork* Work = nullptr;
@@ -840,6 +789,54 @@ FQueuedThreadPool* FQueuedThreadPool::Allocate()
 	return new FQueuedThreadPoolBase;
 }
 
+//////////////////////////////////////////////////////////////////////////
+
+uint32
+FQueuedThread::Run()
+{
+	while (!TimeToDie.Load(EMemoryOrder::Relaxed))
+	{
+		// This will force sending the stats packet from the previous frame.
+		SET_DWORD_STAT(STAT_ThreadPoolDummyCounter, 0);
+		// We need to wait for shorter amount of time
+		bool bContinueWaiting = true;
+
+		// Unless we're collecting stats there doesn't appear to be any reason to wake
+		// up again until there's work to do (or it's time to die)
+
+#if STATS
+		if (FThreadStats::IsCollectingData())
+		{
+			while (bContinueWaiting)
+			{
+				DECLARE_SCOPE_CYCLE_COUNTER(TEXT("FQueuedThread::Run.WaitForWork"), STAT_FQueuedThread_Run_WaitForWork, STATGROUP_ThreadPoolAsyncTasks);
+
+				// Wait for some work to do
+
+				bContinueWaiting = !DoWorkEvent->Wait(10);
+			}
+		}
+#endif
+
+		if (bContinueWaiting)
+		{
+			DoWorkEvent->Wait();
+		}
+
+		IQueuedWork* LocalQueuedWork = QueuedWork;
+		QueuedWork = nullptr;
+		FPlatformMisc::MemoryBarrier();
+		check(LocalQueuedWork || TimeToDie.Load(EMemoryOrder::Relaxed)); // well you woke me up, where is the job or termination request?
+		while (LocalQueuedWork)
+		{
+			// Tell the object to do the work
+			LocalQueuedWork->DoThreadedWork();
+			// Let the object cleanup before we remove our ref to it
+			LocalQueuedWork = OwningThreadPool->ReturnToPoolOrGetNextJob(this);
+		}
+	}
+	return 0;
+}
 
 /*-----------------------------------------------------------------------------
 	FThreadSingletonInitializer
diff --git a/Engine/Source/Runtime/Core/Public/Misc/QueuedThreadPool.h b/Engine/Source/Runtime/Core/Public/Misc/QueuedThreadPool.h
--- a/Engine/Source/Runtime/Core/Public/Misc/QueuedThreadPool.h
+++ b/Engine/Source/Runtime/Core/Public/Misc/QueuedThreadPool.h
@@ -48,14 +48,6 @@ public:
 	 */
 	virtual bool RetractQueuedWork( IQueuedWork* InQueuedWork ) = 0;
 
-	/**
-	 * Places a thread back into the available pool
-	 *
-	 * @param InQueuedThread The thread that is ready to be pooled
-	 * @return next job or null if there is no job available now
-	 */
-	virtual IQueuedWork* ReturnToPoolOrGetNextJob( class FQueuedThread* InQueuedThread ) = 0;
-
 	/**
 	 * Get the number of queued threads
 	 */
