Pass in the FString by const ref rather than by copy.

diff --git a/Engine/Source/Runtime/Core/Public/Misc/AutomationEvent.h b/Engine/Source/Runtime/Core/Public/Misc/AutomationEvent.h
--- a/Engine/Source/Runtime/Core/Public/Misc/AutomationEvent.h
+++ b/Engine/Source/Runtime/Core/Public/Misc/AutomationEvent.h
@@ -15,14 +15,14 @@ enum class EAutomationEventType : uint8
 struct CORE_API FAutomationEvent
 {
 public:
-	FAutomationEvent(EAutomationEventType InType, FString InMessage)
+	FAutomationEvent(EAutomationEventType InType, const FString& InMessage)
 		: Type(InType)
 		, Message(InMessage)
 		, Context()
 	{
 	}
 
-	FAutomationEvent(EAutomationEventType InType, FString InMessage, FString InContext)
+	FAutomationEvent(EAutomationEventType InType, const FString& InMessage, const FString& InContext)
 		: Type(InType)
 		, Message(InMessage)
 		, Context(InContext)
