Add a way to map additional FIoChunkId to a chunk that has already been added to FIoStoreWriter which in turn will make it easier to address partial ranges of a file at runtime.

Also added some new error codes to cover some of the new error conditions.

diff --git a/Engine/Source/Runtime/Core/Private/IO/IoDispatcher.cpp b/Engine/Source/Runtime/Core/Private/IO/IoDispatcher.cpp
--- a/Engine/Source/Runtime/Core/Private/IO/IoDispatcher.cpp
+++ b/Engine/Source/Runtime/Core/Private/IO/IoDispatcher.cpp
@@ -154,6 +154,8 @@ const TCHAR* GetIoErrorText(EIoErrorCode ErrorCode)
 		TEXT("Write Error"),
 		TEXT("Not Found"),
 		TEXT("Corrupt Toc"),
+		TEXT("Unknown ChunkID")
+		TEXT("Invalid Parameter")
 	};
 
 	return ErrorCodeText[static_cast<uint32>(ErrorCode)];
@@ -825,6 +827,39 @@ public:
 		}
 	}
 
+	UE_NODISCARD FIoStatus MapPartialRange(FIoChunkId OriginalChunkId, uint64 Offset, uint64 Length, FIoChunkId ChunkIdPartialRange)
+	{
+		//TODO: Does RelativeOffset + Length overflow?
+
+		const FIoStoreTocEntry* Entry = Toc.Find(OriginalChunkId);
+		if (Entry == nullptr)
+		{
+			return FIoStatus(EIoErrorCode::UnknownChunkID, TEXT("OriginalChunkId does not exist in the container"));
+		}
+
+		if (Toc.Find(ChunkIdPartialRange) != nullptr)
+		{
+			return FIoStatus(EIoErrorCode::InvalidParameter, TEXT("ChunkIdPartialRange is already mapped"));
+		}
+
+		if (Offset + Length > Entry->GetLength())
+		{
+			return FIoStatus(EIoErrorCode::InvalidParameter, TEXT("The given range (Offset/Length) is not within the bounds of OriginalChunkId's data"));
+		}
+
+		FIoStoreTocEntry TocEntry;
+
+		TocEntry.SetOffset(Entry->GetOffset() + Offset);
+		TocEntry.SetLength(Length);
+		TocEntry.ChunkId = ChunkIdPartialRange;
+
+		Toc.Add(ChunkIdPartialRange, TocEntry);
+
+		IsMetadataDirty = true;
+
+		return FIoStatus::Ok;
+	}
+
 	UE_NODISCARD FIoStatus FlushMetadata()
 	{
 		TocFileHandle->Seek(0);
@@ -884,6 +919,11 @@ void FIoStoreWriter::Append(FIoChunkId ChunkId, FIoBuffer Chunk)
 	Impl->Append(ChunkId, Chunk);
 }
 
+void FIoStoreWriter::MapPartialRange(FIoChunkId OriginalChunkId, uint64 Offset, uint64 Length, FIoChunkId ChunkIdPartialRange)
+{
+	Impl->MapPartialRange(OriginalChunkId, Offset, Length, ChunkIdPartialRange);
+}
+
 void FIoStoreWriter::FlushMetadata()
 {
 	Impl->FlushMetadata();
diff --git a/Engine/Source/Runtime/Core/Public/IO/IoDispatcher.h b/Engine/Source/Runtime/Core/Public/IO/IoDispatcher.h
--- a/Engine/Source/Runtime/Core/Public/IO/IoDispatcher.h
+++ b/Engine/Source/Runtime/Core/Public/IO/IoDispatcher.h
@@ -48,7 +48,9 @@ enum class EIoErrorCode
 	FileNotOpen,
 	WriteError,
 	NotFound,
-	CorruptToc
+	CorruptToc,
+	UnknownChunkID,
+	InvalidParameter
 };
 
 class FIoStatus
@@ -654,6 +656,16 @@ public:
 
 	CORE_API FIoStatus	Initialize();
 	CORE_API void		Append(FIoChunkId ChunkId, FIoBuffer Chunk);
+
+	/**
+	 * Creates an addressable range in an already mapped Chunk.
+	 *
+	 * @param OriginalChunkId The FIoChunkId of the original chunk from which you want to create the range
+	 * @param Offset The number of bytes into the original chunk that the range should start
+	 * @param Length The length of the range in bytes
+	 * @param ChunkIdPartialRange The FIoChunkId that will map to the range
+	 */
+	CORE_API void		MapPartialRange(FIoChunkId OriginalChunkId, uint64 Offset, uint64 Length, FIoChunkId ChunkIdPartialRange);
 	CORE_API void		FlushMetadata();
 
 private:
