TBA: Renamed SA_DEFAULTED_ATTRIBUTE to SA_OPTIONAL_ATTRIBUTE, and forced it to behave the same way when text archive support is enabled and disabled. The existing code wouldn't process the stored existence bool when text archives were disabled.

diff --git a/Engine/Source/Runtime/Core/Public/Serialization/StructuredArchive.h b/Engine/Source/Runtime/Core/Public/Serialization/StructuredArchive.h
--- a/Engine/Source/Runtime/Core/Public/Serialization/StructuredArchive.h
+++ b/Engine/Source/Runtime/Core/Public/Serialization/StructuredArchive.h
@@ -64,13 +64,13 @@ template<typename T> struct TNamedAttribute
  * Class to contain a named attribute for serialization, with a default. Intended to be created as a temporary and passed to object
  * serialization methods, which can choose not to serialize the attribute if it matches the default.
  */
-template<typename T> struct TDefaultedNamedAttribute
+template<typename T> struct TOptionalNamedAttribute
 {
 	FArchiveFieldName Name;
 	T& Value;
 	const T& Default;
 
-	explicit FORCEINLINE TDefaultedNamedAttribute(FArchiveFieldName InName, T& InValue, const T& InDefault)
+	explicit FORCEINLINE TOptionalNamedAttribute(FArchiveFieldName InName, T& InValue, const T& InDefault)
 		: Name(InName)
 		, Value(InValue)
 		, Default(InDefault)
@@ -95,33 +95,21 @@ template<typename T> FORCEINLINE TNamedAttribute<T> MakeNamedAttribute(FArchiveF
 }
 
 /**
- * Helper function to construct a TDefaultedNamedAttribute, deducing the value type.
+ * Helper function to construct a TOptionalNamedAttribute, deducing the value type.
  */
-template<typename T> FORCEINLINE TDefaultedNamedAttribute<T> MakeDefaultedNamedAttribute(FArchiveFieldName Name, T& Value, const typename TIdentity<T>::Type& Default)
+template<typename T> FORCEINLINE TOptionalNamedAttribute<T> MakeOptionalNamedAttribute(FArchiveFieldName Name, T& Value, const typename TIdentity<T>::Type& Default)
 {
-	return TDefaultedNamedAttribute<T>(Name, Value, Default);
+	return TOptionalNamedAttribute<T>(Name, Value, Default);
 }
 
 /** Construct a TNamedValue given an ANSI string and value reference. */
-#if WITH_TEXT_ARCHIVE_SUPPORT
-	#define SA_VALUE(Name, Value) MakeNamedValue(FArchiveFieldName(Name), Value)
-#else
-	#define SA_VALUE(Name, Value) MakeNamedValue(FArchiveFieldName(), Value)
-#endif
+#define SA_VALUE(Name, Value) MakeNamedValue(FArchiveFieldName(Name), Value)
 
 /** Construct a TNamedAttribute given an ANSI string and value reference. */
-#if WITH_TEXT_ARCHIVE_SUPPORT
-	#define SA_ATTRIBUTE(Name, Value) MakeNamedAttribute(FArchiveFieldName(Name), Value)
-#else
-	#define SA_ATTRIBUTE(Name, Value) MakeNamedAttribute(FArchiveFieldName(), Value)
-#endif
+#define SA_ATTRIBUTE(Name, Value) MakeNamedAttribute(FArchiveFieldName(Name), Value)
 
-/** Construct a TDefaultedNamedAttribute given an ANSI string and value reference. */
-#if WITH_TEXT_ARCHIVE_SUPPORT
-	#define SA_DEFAULTED_ATTRIBUTE(Name, Value, Default) MakeDefaultedNamedAttribute(FArchiveFieldName(Name), Value, Default)
-#else
-	#define SA_DEFAULTED_ATTRIBUTE(Name, Value, Default) MakeNamedAttribute(FArchiveFieldName(), Value)
-#endif
+/** Construct a TOptionalNamedAttribute given an ANSI string and value reference. */
+#define SA_OPTIONAL_ATTRIBUTE(Name, Value, Default) MakeOptionalNamedAttribute(FArchiveFieldName(Name), Value, Default)
 
 /** Typedef for which formatter type to support */
 #if WITH_TEXT_ARCHIVE_SUPPORT
@@ -282,7 +270,7 @@ public:
 	}
 
 	template <typename T>
-	FORCEINLINE void operator<<(TDefaultedNamedAttribute<T> Item)
+	FORCEINLINE void operator<<(TOptionalNamedAttribute<T> Item)
 	{
 		if (TOptional<FStructuredArchiveSlot> Attribute = TryEnterAttribute(Item.Name, Item.Value != Item.Default))
 		{
diff --git a/Engine/Source/Runtime/Core/Public/Serialization/StructuredArchiveFormatter.h b/Engine/Source/Runtime/Core/Public/Serialization/StructuredArchiveFormatter.h
--- a/Engine/Source/Runtime/Core/Public/Serialization/StructuredArchiveFormatter.h
+++ b/Engine/Source/Runtime/Core/Public/Serialization/StructuredArchiveFormatter.h
@@ -6,22 +6,21 @@
 #include "CoreFwd.h"
 
 // Define a structure to encapsulate a field name, which compiles to an empty object if WITH_TEXT_ARCHIVE_SUPPORT = 0
+struct FArchiveFieldName
+{
 #if WITH_TEXT_ARCHIVE_SUPPORT
-	struct FArchiveFieldName
-	{
-		const TCHAR* Name;
-
-		explicit FArchiveFieldName(const TCHAR* InName) : Name(InName){ }
-	};
+	const TCHAR* Name;
+#endif
 
-	#define SA_FIELD_NAME(x) FArchiveFieldName(x)
-#else
-	struct FArchiveFieldName
+	explicit FArchiveFieldName(const TCHAR* InName) 
+#if WITH_TEXT_ARCHIVE_SUPPORT
+		: Name(InName)
+#endif
 	{
-	};
+	}
+};
 
-	#define SA_FIELD_NAME(x) FArchiveFieldName()
-#endif
+#define SA_FIELD_NAME(x) FArchiveFieldName(x)
 
 /**
  * Specifies the type of a value in a slot. Used by FContextFreeArchiveFormatter for introspection.
diff --git a/Engine/Source/Runtime/CoreUObject/Private/UObject/ObjectResource.cpp b/Engine/Source/Runtime/CoreUObject/Private/UObject/ObjectResource.cpp
--- a/Engine/Source/Runtime/CoreUObject/Private/UObject/ObjectResource.cpp
+++ b/Engine/Source/Runtime/CoreUObject/Private/UObject/ObjectResource.cpp
@@ -207,8 +207,8 @@ void operator<<(FStructuredArchive::FSlot Slot, FObjectTextExport& E)
 	}
 
 	Slot << SA_ATTRIBUTE(TEXT("Class"), ClassName);
-	Slot << SA_DEFAULTED_ATTRIBUTE(TEXT("Outer"), OuterName, FString());
-	Slot << SA_DEFAULTED_ATTRIBUTE(TEXT("SuperStruct"), SuperStructName, FString());
+	Slot << SA_OPTIONAL_ATTRIBUTE(TEXT("Outer"), OuterName, FString());
+	Slot << SA_OPTIONAL_ATTRIBUTE(TEXT("SuperStruct"), SuperStructName, FString());
 
 	if (BaseArchive.IsLoading())
 	{
@@ -218,21 +218,21 @@ void operator<<(FStructuredArchive::FSlot Slot, FObjectTextExport& E)
 	}
 
 	uint32 Save = E.Export.ObjectFlags & RF_Load;
-	Slot << SA_DEFAULTED_ATTRIBUTE(TEXT("ObjectFlags"), Save, 0);
+	Slot << SA_OPTIONAL_ATTRIBUTE(TEXT("ObjectFlags"), Save, 0);
 	if (BaseArchive.IsLoading())
 	{
 		E.Export.ObjectFlags = EObjectFlags(Save & RF_Load);
 	}
 
-	Slot << SA_DEFAULTED_ATTRIBUTE(TEXT("bForcedExport"), E.Export.bForcedExport, false);
-	Slot << SA_DEFAULTED_ATTRIBUTE(TEXT("bNotForClient"), E.Export.bNotForClient, false);
-	Slot << SA_DEFAULTED_ATTRIBUTE(TEXT("bNotForServer"), E.Export.bNotForServer, false);
+	Slot << SA_OPTIONAL_ATTRIBUTE(TEXT("bForcedExport"), E.Export.bForcedExport, false);
+	Slot << SA_OPTIONAL_ATTRIBUTE(TEXT("bNotForClient"), E.Export.bNotForClient, false);
+	Slot << SA_OPTIONAL_ATTRIBUTE(TEXT("bNotForServer"), E.Export.bNotForServer, false);
 
-	Slot << SA_DEFAULTED_ATTRIBUTE(TEXT("PackageGuid"), E.Export.PackageGuid, FGuid());
-	Slot << SA_DEFAULTED_ATTRIBUTE(TEXT("PackageFlags"), E.Export.PackageFlags, 0);
+	Slot << SA_OPTIONAL_ATTRIBUTE(TEXT("PackageGuid"), E.Export.PackageGuid, FGuid());
+	Slot << SA_OPTIONAL_ATTRIBUTE(TEXT("PackageFlags"), E.Export.PackageFlags, 0);
 
-	Slot << SA_DEFAULTED_ATTRIBUTE(TEXT("bNotAlwaysLoadedForEditorGame"), E.Export.bNotAlwaysLoadedForEditorGame, false);
-	Slot << SA_DEFAULTED_ATTRIBUTE(TEXT("bIsAsset"), E.Export.bIsAsset, false);
+	Slot << SA_OPTIONAL_ATTRIBUTE(TEXT("bNotAlwaysLoadedForEditorGame"), E.Export.bNotAlwaysLoadedForEditorGame, false);
+	Slot << SA_OPTIONAL_ATTRIBUTE(TEXT("bIsAsset"), E.Export.bIsAsset, false);
 }
 
 /*-----------------------------------------------------------------------------
diff --git a/Engine/Source/Runtime/CoreUObject/Private/UObject/PropertyArray.cpp b/Engine/Source/Runtime/CoreUObject/Private/UObject/PropertyArray.cpp
--- a/Engine/Source/Runtime/CoreUObject/Private/UObject/PropertyArray.cpp
+++ b/Engine/Source/Runtime/CoreUObject/Private/UObject/PropertyArray.cpp
@@ -67,7 +67,7 @@ void UArrayProperty::SerializeItem(FStructuredArchive::FSlot Slot, void* Value,
 	if (UnderlyingArchive.IsTextFormat() && InnerTag.Type == NAME_StructProperty)
 	{
 		Slot << SA_ATTRIBUTE(TEXT("InnerStructName"), InnerTag.StructName);
-		Slot << SA_DEFAULTED_ATTRIBUTE(TEXT("InnerStructGuid"), InnerTag.StructGuid, FGuid());
+		Slot << SA_OPTIONAL_ATTRIBUTE(TEXT("InnerStructGuid"), InnerTag.StructGuid, FGuid());
 	}
 	
 	// Ensure that the Inner itself has been loaded before calling SerializeItem() on it
diff --git a/Engine/Source/Runtime/CoreUObject/Private/UObject/PropertyTag.cpp b/Engine/Source/Runtime/CoreUObject/Private/UObject/PropertyTag.cpp
--- a/Engine/Source/Runtime/CoreUObject/Private/UObject/PropertyTag.cpp
+++ b/Engine/Source/Runtime/CoreUObject/Private/UObject/PropertyTag.cpp
@@ -124,7 +124,7 @@ void operator<<(FStructuredArchive::FSlot Slot, FPropertyTag& Tag)
 			{
 				if (bIsTextFormat)
 				{
-					Slot << SA_DEFAULTED_ATTRIBUTE(TEXT("StructGuid"), Tag.StructGuid, FGuid());
+					Slot << SA_OPTIONAL_ATTRIBUTE(TEXT("StructGuid"), Tag.StructGuid, FGuid());
 				}
 				else
 				{
@@ -150,7 +150,7 @@ void operator<<(FStructuredArchive::FSlot Slot, FPropertyTag& Tag)
 		{
 			if (UnderlyingArchive.IsTextFormat())
 			{
-				Slot << SA_DEFAULTED_ATTRIBUTE(TEXT("EnumName"), Tag.EnumName, NAME_None);
+				Slot << SA_OPTIONAL_ATTRIBUTE(TEXT("EnumName"), Tag.EnumName, NAME_None);
 			}
 			else
 			{
@@ -188,7 +188,7 @@ void operator<<(FStructuredArchive::FSlot Slot, FPropertyTag& Tag)
 	{
 		if (bIsTextFormat)
 		{
-			Slot << SA_DEFAULTED_ATTRIBUTE(TEXT("PropertyGuid"), Tag.PropertyGuid, FGuid());
+			Slot << SA_OPTIONAL_ATTRIBUTE(TEXT("PropertyGuid"), Tag.PropertyGuid, FGuid());
 			Tag.HasPropertyGuid = Tag.PropertyGuid.IsValid();
 		}
 		else
