Include check for shipping build too, just in case

#rnx

diff --git a/Engine/Source/Runtime/Core/Public/Stats/StatsTrace.h b/Engine/Source/Runtime/Core/Public/Stats/StatsTrace.h
--- a/Engine/Source/Runtime/Core/Public/Stats/StatsTrace.h
+++ b/Engine/Source/Runtime/Core/Public/Stats/StatsTrace.h
@@ -7,7 +7,7 @@
 
 #define EXPERIMENTAL_STATSTRACE_ENABLED 0
 
-#if UE_TRACE_ENABLED && STATS && EXPERIMENTAL_STATSTRACE_ENABLED
+#if UE_TRACE_ENABLED && STATS && EXPERIMENTAL_STATSTRACE_ENABLED && !UE_BUILD_SHIPPING
 #define STATSTRACE_ENABLED 1
 #else
 #define STATSTRACE_ENABLED 0
