Fixed merge error when Dev-RenPlatStage copied up to Main.

#rnx

diff --git a/Engine/Source/Runtime/Core/Public/Logging/LogTrace.h b/Engine/Source/Runtime/Core/Public/Logging/LogTrace.h
--- a/Engine/Source/Runtime/Core/Public/Logging/LogTrace.h
+++ b/Engine/Source/Runtime/Core/Public/Logging/LogTrace.h
@@ -3,9 +3,10 @@
 #pragma once
 
 #include "CoreTypes.h"
+#include "Trace/Config.h"
 
 #if !defined(LOGTRACE_ENABLED)
-#if !IS_PROGRAM && !UE_BUILD_SHIPPING && (PLATFORM_WINDOWS || PLATFORM_PS4 || PLATFORM_UNIX || PLATFORM_XBOXONE)
+#if UE_TRACE_ENABLED && !UE_BUILD_SHIPPING
 #define LOGTRACE_ENABLED 1
 #else
 #define LOGTRACE_ENABLED 0
