Allow the user to control if Misc and Stats tracing is compiled in or not.

#rnx

diff --git a/Engine/Source/Runtime/Core/Public/ProfilingDebugging/MiscTrace.h b/Engine/Source/Runtime/Core/Public/ProfilingDebugging/MiscTrace.h
--- a/Engine/Source/Runtime/Core/Public/ProfilingDebugging/MiscTrace.h
+++ b/Engine/Source/Runtime/Core/Public/ProfilingDebugging/MiscTrace.h
@@ -6,11 +6,13 @@
 #include "Trace/Config.h"
 #include "ProfilingDebugging/FormatArgsTrace.h"
 
+#if !defined(MISCTRACE_ENABLED)
 #if UE_TRACE_ENABLED && !UE_BUILD_SHIPPING
 #define MISCTRACE_ENABLED 1
 #else
 #define MISCTRACE_ENABLED 0
 #endif
+#endif
 
 enum ETraceFrameType
 {
diff --git a/Engine/Source/Runtime/Core/Public/Stats/StatsTrace.h b/Engine/Source/Runtime/Core/Public/Stats/StatsTrace.h
--- a/Engine/Source/Runtime/Core/Public/Stats/StatsTrace.h
+++ b/Engine/Source/Runtime/Core/Public/Stats/StatsTrace.h
@@ -7,11 +7,13 @@
 
 #define EXPERIMENTAL_STATSTRACE_ENABLED 0
 
+#if !defined(STATSTRACE_ENABLED)
 #if UE_TRACE_ENABLED && STATS && EXPERIMENTAL_STATSTRACE_ENABLED && !UE_BUILD_SHIPPING
 #define STATSTRACE_ENABLED 1
 #else
 #define STATSTRACE_ENABLED 0
 #endif
+#endif
 
 #if STATSTRACE_ENABLED
 
