Move support for .ToSharedRef() and .Pin().

diff --git a/Engine/Source/Runtime/Core/Public/Templates/SharedPointer.h b/Engine/Source/Runtime/Core/Public/Templates/SharedPointer.h
--- a/Engine/Source/Runtime/Core/Public/Templates/SharedPointer.h
+++ b/Engine/Source/Runtime/Core/Public/Templates/SharedPointer.h
@@ -792,14 +792,28 @@ public:
 	 * @return  Reference to the object
 	 */
 	// NOTE: The following is an Unreal extension to standard shared_ptr behavior
-	FORCEINLINE TSharedRef< ObjectType, Mode > ToSharedRef() const
+	FORCEINLINE TSharedRef< ObjectType, Mode > ToSharedRef() const&
 	{
- 		// If this assert goes off, it means a shared reference was created from a shared pointer that was nullptr.
- 		// Shared references are never allowed to be null.  Consider using TSharedPtr instead.
+		// If this assert goes off, it means a shared reference was created from a shared pointer that was nullptr.
+		// Shared references are never allowed to be null.  Consider using TSharedPtr instead.
 		check( IsValid() );
 		return TSharedRef< ObjectType, Mode >( *this );
 	}
 
+	/**
+	 * Converts a shared pointer to a shared reference.  The pointer *must* be valid or an assertion will trigger.
+	 *
+	 * @return  Reference to the object
+	 */
+	 // NOTE: The following is an Unreal extension to standard shared_ptr behavior
+	FORCEINLINE TSharedRef< ObjectType, Mode > ToSharedRef() &&
+	{
+		// If this assert goes off, it means a shared reference was created from a shared pointer that was nullptr.
+		// Shared references are never allowed to be null.  Consider using TSharedPtr instead.
+		check( IsValid() );
+		return TSharedRef< ObjectType, Mode >( MoveTemp( *this ) );
+	}
+
 	/**
 	 * Returns the object referenced by this pointer, or nullptr if no object is reference
 	 *
@@ -911,6 +925,32 @@ private:
 		}
 	}
 
+	/**
+	 * Constructs a shared pointer from a weak pointer, allowing you to access the object (if it
+	 * hasn't expired yet.)  Remember, if there are no more shared references to the object, the
+	 * shared pointer will not be valid.  You should always check to make sure this shared
+	 * pointer is valid before trying to dereference the shared pointer!
+	 *
+	 * NOTE: This constructor is private to force users to be explicit when converting a weak
+	 *       pointer to a shared pointer.  Use the weak pointer's Pin() method instead!
+	 */
+	template <
+		typename OtherType,
+		typename = typename TEnableIf<TPointerIsConvertibleFromTo<OtherType, ObjectType>::Value>::Type
+	>
+	FORCEINLINE explicit TSharedPtr(TWeakPtr< OtherType, Mode >&& InWeakPtr)
+		: Object(nullptr)
+		, SharedReferenceCount( MoveTemp( InWeakPtr.WeakReferenceCount ) )
+	{
+		// Check that the shared reference was created from the weak reference successfully.  We'll only
+		// cache a pointer to the object if we have a valid shared reference.
+		if (SharedReferenceCount.IsValid())
+		{
+			Object = InWeakPtr.Object;
+			InWeakPtr.Object = nullptr;
+		}
+	}
+
 	/**
 	 * Computes a hash code for this object
 	 *
@@ -1132,11 +1172,24 @@ public:
 	 *
 	 * @return  Shared pointer for this object (will only be valid if still referenced!)
 	 */
-	FORCEINLINE TSharedPtr< ObjectType, Mode > Pin() const
+	FORCEINLINE TSharedPtr< ObjectType, Mode > Pin() const&
 	{
 		return TSharedPtr< ObjectType, Mode >( *this );
 	}
 
+	/**
+	 * Converts this weak pointer to a shared pointer that you can use to access the object (if it
+	 * hasn't expired yet.)  Remember, if there are no more shared references to the object, the
+	 * returned shared pointer will not be valid.  You should always check to make sure the returned
+	 * pointer is valid before trying to dereference the shared pointer!
+	 *
+	 * @return  Shared pointer for this object (will only be valid if still referenced!)
+	 */
+	FORCEINLINE TSharedPtr< ObjectType, Mode > Pin() &&
+	{
+		return TSharedPtr< ObjectType, Mode >( MoveTemp( *this ) );
+	}
+
 	/**
 	 * Checks to see if this weak pointer actually has a valid reference to an object
 	 *
@@ -1233,7 +1286,7 @@ public:
 
 		// Now that we've verified the shared pointer is valid, we'll convert it to a shared reference
 		// and return it!
-		return SharedThis.ToSharedRef();
+		return MoveTemp( SharedThis ).ToSharedRef();
 	}
 
 	/**
@@ -1260,7 +1313,7 @@ public:
 
 		// Now that we've verified the shared pointer is valid, we'll convert it to a shared reference
 		// and return it!
-		return SharedThis.ToSharedRef();
+		return MoveTemp( SharedThis ).ToSharedRef();
 	}
 
 protected:
diff --git a/Engine/Source/Runtime/Core/Public/Templates/SharedPointerInternals.h b/Engine/Source/Runtime/Core/Public/Templates/SharedPointerInternals.h
--- a/Engine/Source/Runtime/Core/Public/Templates/SharedPointerInternals.h
+++ b/Engine/Source/Runtime/Core/Public/Templates/SharedPointerInternals.h
@@ -435,6 +435,30 @@ namespace SharedPointerInternals
 			}
 		}
 
+		/** Creates a shared referencer object from a weak referencer object.  This will only result
+		    in a valid object reference if the object already has at least one other shared referencer. */
+		FSharedReferencer( FWeakReferencer< Mode >&& InWeakReference )
+			: ReferenceController( InWeakReference.ReferenceController )
+		{
+			// If the incoming reference had an object associated with it, then go ahead and increment the
+			// shared reference count
+			if( ReferenceController != nullptr )
+			{
+				// Attempt to elevate a weak reference to a shared one.  For this to work, the object this
+				// weak counter is associated with must already have at least one shared reference.  We'll
+				// never revive a pointer that has already expired!
+				if( !TOps::ConditionallyAddSharedReference(ReferenceController) )
+				{
+					ReferenceController = nullptr;
+				}
+
+				// Tell the reference counter object that we're no longer referencing the object with
+				// this weak pointer
+				TOps::ReleaseWeakReference(InWeakReference.ReferenceController);
+				InWeakReference.ReferenceController = nullptr;
+			}
+		}
+
 		/** Destructor. */
 		FORCEINLINE ~FSharedReferencer()
 		{
