Fixed WindowsPlatformFile not handling >32bit read and write sizes correctly

diff --git a/Engine/Source/Runtime/Core/Private/Windows/WindowsPlatformFile.cpp b/Engine/Source/Runtime/Core/Private/Windows/WindowsPlatformFile.cpp
--- a/Engine/Source/Runtime/Core/Private/Windows/WindowsPlatformFile.cpp
+++ b/Engine/Source/Runtime/Core/Private/Windows/WindowsPlatformFile.cpp
@@ -556,62 +556,84 @@ public:
 	virtual bool Read(uint8* Destination, int64 BytesToRead) override
 	{
 		check(IsValid());
-		uint32 NumRead = 0;
 		// Now kick off an async read
 		TRACE_PLATFORMFILE_BEGIN_READ(&OverlappedIO, FileHandle, FilePos, BytesToRead);
-		if (!ReadFile(FileHandle, Destination, BytesToRead, (::DWORD*)&NumRead, &OverlappedIO))
+
+		int64 TotalNumRead = 0;
+		do
 		{
-			uint32 ErrorCode = GetLastError();
-			if (ErrorCode != ERROR_IO_PENDING)
-			{
-				// Read failed
-				TRACE_PLATFORMFILE_END_READ(&OverlappedIO, 0);
-				return false;
-			}
-			// Wait for the read to complete
-			NumRead = 0;
-			if (!GetOverlappedResult(FileHandle, &OverlappedIO, (::DWORD*)&NumRead, true))
+			uint32 BytesToRead32 = FMath::Min(BytesToRead, int64(UINT32_MAX));
+			uint32 NumRead = 0;
+
+			if (!ReadFile(FileHandle, Destination, BytesToRead32, (::DWORD*)&NumRead, &OverlappedIO))
 			{
-				// Read failed
-				TRACE_PLATFORMFILE_END_READ(&OverlappedIO, 0);
-				return false;
+				uint32 ErrorCode = GetLastError();
+				if (ErrorCode != ERROR_IO_PENDING)
+				{
+					// Read failed
+					TRACE_PLATFORMFILE_END_READ(&OverlappedIO, 0);
+					return false;
+				}
+				// Wait for the read to complete
+				NumRead = 0;
+				if (!GetOverlappedResult(FileHandle, &OverlappedIO, (::DWORD*)&NumRead, true))
+				{
+					// Read failed
+					TRACE_PLATFORMFILE_END_READ(&OverlappedIO, 0);
+					return false;
+				}
 			}
-		}
-		TRACE_PLATFORMFILE_END_READ(&OverlappedIO, NumRead);
-		// Update where we are in the file
-		FilePos += NumRead;
-		UpdateOverlappedPos();
+
+			BytesToRead -= BytesToRead32;
+			Destination += BytesToRead32;
+			TotalNumRead += NumRead;
+			// Update where we are in the file
+			FilePos += NumRead;
+			UpdateOverlappedPos();
+		} while (BytesToRead > 0);
+		TRACE_PLATFORMFILE_END_READ(&OverlappedIO, TotalNumRead);
 		return true;
 	}
 	virtual bool Write(const uint8* Source, int64 BytesToWrite) override
 	{
 		check(IsValid());
-		uint32 NumWritten = 0;
-		// Now kick off an async write
+
 		TRACE_PLATFORMFILE_BEGIN_WRITE(this, FileHandle, FilePos, BytesToWrite);
-		if (!WriteFile(FileHandle, Source, BytesToWrite, (::DWORD*)&NumWritten, &OverlappedIO))
+
+		int64 TotalNumWritten = 0;
+		do
 		{
-			uint32 ErrorCode = GetLastError();
-			if (ErrorCode != ERROR_IO_PENDING)
+			uint32 BytesToWrite32 = FMath::Min(BytesToWrite - TotalNumWritten, int64(UINT32_MAX));
+			uint32 NumWritten = 0;
+			// Now kick off an async write
+			if (!WriteFile(FileHandle, Source, BytesToWrite32, (::DWORD*)&NumWritten, &OverlappedIO))
 			{
-				// Write failed
-				TRACE_PLATFORMFILE_END_WRITE(this, 0);
-				return false;
-			}
-			// Wait for the write to complete
-			NumWritten = 0;
-			if (!GetOverlappedResult(FileHandle, &OverlappedIO, (::DWORD*)&NumWritten, true))
-			{
-				// Write failed
-				TRACE_PLATFORMFILE_END_WRITE(this, 0);
-				return false;
+				uint32 ErrorCode = GetLastError();
+				if (ErrorCode != ERROR_IO_PENDING)
+				{
+					// Write failed
+					TRACE_PLATFORMFILE_END_WRITE(this, 0);
+					return false;
+				}
+				// Wait for the write to complete
+				NumWritten = 0;
+				if (!GetOverlappedResult(FileHandle, &OverlappedIO, (::DWORD*)&NumWritten, true))
+				{
+					// Write failed
+					TRACE_PLATFORMFILE_END_WRITE(this, 0);
+					return false;
+				}
 			}
-		}
-		TRACE_PLATFORMFILE_END_WRITE(this, NumWritten);
-		// Update where we are in the file
-		FilePos += NumWritten;
-		UpdateOverlappedPos();
-		FileSize = FMath::Max(FilePos, FileSize);
+			BytesToWrite -= BytesToWrite32;
+			Source += BytesToWrite32;
+			TotalNumWritten += NumWritten;
+			// Update where we are in the file
+			FilePos += NumWritten;
+			UpdateOverlappedPos();
+			FileSize = FMath::Max(FilePos, FileSize);
+		} while (BytesToWrite > 0);
+
+		TRACE_PLATFORMFILE_END_WRITE(this, TotalNumWritten);
 		return true;
 	}
 	virtual bool Flush(const bool bFullFlush = false) override
