Missing include fix.

#rnx

diff --git a/Engine/Source/Runtime/Core/Private/ProfilingDebugging/PlatformFileTrace.cpp b/Engine/Source/Runtime/Core/Private/ProfilingDebugging/PlatformFileTrace.cpp
--- a/Engine/Source/Runtime/Core/Private/ProfilingDebugging/PlatformFileTrace.cpp
+++ b/Engine/Source/Runtime/Core/Private/ProfilingDebugging/PlatformFileTrace.cpp
@@ -7,6 +7,7 @@
 #include "HAL/PlatformTime.h"
 #include "HAL/PlatformTLS.h"
 #include "Misc/Parse.h"
+#include "Trace/Trace.h"
 
 UE_TRACE_EVENT_BEGIN(PlatformFile, BeginOpen)
 	UE_TRACE_EVENT_FIELD(uint64, Cycle)
@@ -150,4 +151,4 @@ void FPlatformFileTrace::Init(const TCHAR* CmdLine)
 	}
 }
 
-#endif
\ No newline at end of file
+#endif
