ProcessHandle getter in FMonitoredProcess and FInteractiveProcess

Merge request by ccarvalheiro.

diff --git a/Engine/Source/Runtime/Core/Public/Misc/InteractiveProcess.h b/Engine/Source/Runtime/Core/Public/Misc/InteractiveProcess.h
--- a/Engine/Source/Runtime/Core/Public/Misc/InteractiveProcess.h
+++ b/Engine/Source/Runtime/Core/Public/Misc/InteractiveProcess.h
@@ -65,6 +65,16 @@ public:
 	*/
 	FTimespan GetDuration() const;
 
+	/**
+	* Gets the Process Handle. The instance can be invalid if the process was not created.
+	*
+	* @return The Process Handle
+	*/
+	FProcHandle GetProcessHandle() const
+	{
+		return ProcessHandle;
+	}
+
 	/**
 	* Checks whether the process is still running.
 	*
diff --git a/Engine/Source/Runtime/Core/Public/Misc/MonitoredProcess.h b/Engine/Source/Runtime/Core/Public/Misc/MonitoredProcess.h
--- a/Engine/Source/Runtime/Core/Public/Misc/MonitoredProcess.h
+++ b/Engine/Source/Runtime/Core/Public/Misc/MonitoredProcess.h
@@ -78,6 +78,16 @@ public:
 	 */
 	FTimespan GetDuration() const;
 
+	/**
+	 * Gets the Process Handle. The instance can be invalid if the process was not created.
+	 *
+	 * @return The Process Handle
+	 */
+	FProcHandle GetProcessHandle() const
+	{
+		return ProcessHandle;
+	}
+
 	/**
 	 * Checks whether the process is still running.
 	 *
