More non-unity fixes.

#rnx

diff --git a/Engine/Source/Runtime/Core/Private/ProfilingDebugging/CsvProfilerTrace.cpp b/Engine/Source/Runtime/Core/Private/ProfilingDebugging/CsvProfilerTrace.cpp
--- a/Engine/Source/Runtime/Core/Private/ProfilingDebugging/CsvProfilerTrace.cpp
+++ b/Engine/Source/Runtime/Core/Private/ProfilingDebugging/CsvProfilerTrace.cpp
@@ -3,6 +3,7 @@
 #include "ProfilingDebugging/CsvProfiler.h"
 #include "HAL/PlatformTime.h"
 #include "UObject/NameTypes.h"
+#include "Trace/Trace.h"
 
 #if CSVPROFILERTRACE_ENABLED
 
@@ -278,4 +279,4 @@ void FCsvProfilerTrace::OutputMetadata(const TCHAR* Key, const TCHAR* Value)
 		<< Metadata.Attachment(Attachment);
 }
 
-#endif
\ No newline at end of file
+#endif
diff --git a/Engine/Source/Runtime/Core/Public/ProfilingDebugging/CsvProfilerTrace.h b/Engine/Source/Runtime/Core/Public/ProfilingDebugging/CsvProfilerTrace.h
--- a/Engine/Source/Runtime/Core/Public/ProfilingDebugging/CsvProfilerTrace.h
+++ b/Engine/Source/Runtime/Core/Public/ProfilingDebugging/CsvProfilerTrace.h
@@ -11,6 +11,8 @@
 #define CSVPROFILERTRACE_ENABLED 0
 #endif
 
+class FName;
+
 #if CSVPROFILERTRACE_ENABLED
 
 struct FCsvProfilerTrace
@@ -98,4 +100,4 @@ struct FCsvProfilerTrace
 #define TRACE_CSV_PROFILER_END_CAPTURE(...)
 #define TRACE_CSV_PROFILER_METADATA(...)
 
-#endif
\ No newline at end of file
+#endif
