GitHub #6280 : Add Algo::ForEach.

diff --git a/Engine/Source/Runtime/Core/Public/Algo/ForEach.h b/Engine/Source/Runtime/Core/Public/Algo/ForEach.h
new file mode 100644
--- /dev/null
+++ b/Engine/Source/Runtime/Core/Public/Algo/ForEach.h
@@ -0,0 +1,43 @@
+// Copyright 1998-2019 Epic Games, Inc. All Rights Reserved.
+
+#pragma once
+
+#include "CoreTypes.h"
+#include "Templates/Invoke.h"
+
+namespace Algo
+{
+	/**
+	 * Conditionally invokes a callable to each element in a range
+	 *
+	 * @param  Input      Any iterable type
+	 * @param  Predicate  Condition which returns true for elements that should be called with and false for elements that should be skipped
+	 * @param  Callable   Callable object
+	 */
+	template <typename InT, typename PredicateT, typename CallableT>
+	FORCEINLINE void ForEachIf(InT& Input, PredicateT Predicate, CallableT Callable)
+	{
+		for (auto& Value : Input)
+		{
+			if (Invoke(Predicate, Value))
+			{
+				Invoke(Callable, Value);
+			}
+		}
+	}
+
+	/**
+	 * Invokes a callable to each element in a range
+	 *
+	 * @param  Input      Any iterable type
+	 * @param  Callable   Callable object
+	 */
+	template <typename InT, typename CallableT>
+	FORCEINLINE void ForEach(InT& Input, CallableT Callable)
+	{
+		for (auto& Value : Input)
+		{
+			Invoke(Callable, Value);
+		}
+	}
+}
