GitHub #6281 : Add predicate version of LoadFileToStringArray

diff --git a/Engine/Source/Runtime/Core/Private/Misc/FileHelper.cpp b/Engine/Source/Runtime/Core/Private/Misc/FileHelper.cpp
--- a/Engine/Source/Runtime/Core/Private/Misc/FileHelper.cpp
+++ b/Engine/Source/Runtime/Core/Private/Misc/FileHelper.cpp
@@ -192,6 +192,43 @@ bool FFileHelper::LoadFileToStringArray( TArray<FString>& Result, const TCHAR* F
 	return true;
 }
 
+bool FFileHelper::LoadFileToStringArrayWithPredicate(TArray<FString>& Result, const TCHAR* Filename, TFunctionRef<bool(const FString&)> Predicate, EHashOptions VerifyFlags)
+{
+	Result.Empty();
+
+	FString Buffer;
+	if (!LoadFileToString(Buffer, Filename, VerifyFlags))
+	{
+		return false;
+	}
+
+	for (const TCHAR* Pos = *Buffer; *Pos != 0; )
+	{
+		const TCHAR* LineStart = Pos;
+		while (*Pos != 0 && *Pos != '\r' && *Pos != '\n')
+		{
+			Pos++;
+		}
+
+		FString Line(Pos - LineStart, LineStart);
+		if (Invoke(Predicate, Line))
+		{
+			Result.Add(MoveTemp(Line));
+		}
+
+		if (*Pos == '\r')
+		{
+			Pos++;
+		}
+		if (*Pos == '\n')
+		{
+			Pos++;
+		}
+	}
+
+	return true;
+}
+
 /**
  * Save a binary array to a file.
  */
diff --git a/Engine/Source/Runtime/Core/Public/Misc/FileHelper.h b/Engine/Source/Runtime/Core/Public/Misc/FileHelper.h
--- a/Engine/Source/Runtime/Core/Public/Misc/FileHelper.h
+++ b/Engine/Source/Runtime/Core/Public/Misc/FileHelper.h
@@ -64,6 +64,16 @@ struct CORE_API FFileHelper
 	 */
 	static bool LoadFileToStringArray( TArray<FString>& Result, const TCHAR* Filename, EHashOptions VerifyFlags = EHashOptions::None );
 
+	/**
+	 * Load a text file to an array of strings, filtered by a user-defined predicate. Supports all combination of ANSI/Unicode files and platforms.
+	 *
+	 * @param Result       String representation of the loaded file
+	 * @param Filename     Name of the file to load
+	 * @param Predicate    Condition for whether or not to add the line to the array
+	 * @param VerifyFlags  Flags controlling the hash verification behavior ( see EHashOptions )
+	 */
+	static bool LoadFileToStringArrayWithPredicate(TArray<FString>& Result, const TCHAR* Filename, TFunctionRef<bool(const FString&)> Predicate, EHashOptions VerifyFlags = EHashOptions::None);
+
 	/**
 	 * Save a binary array to a file.
 	 */
