Fixed up doc comment which was referencing parameters which have gone away a long time ago

diff --git a/Engine/Source/Runtime/Core/Private/Async/TaskGraph.cpp b/Engine/Source/Runtime/Core/Private/Async/TaskGraph.cpp
--- a/Engine/Source/Runtime/Core/Private/Async/TaskGraph.cpp
+++ b/Engine/Source/Runtime/Core/Private/Async/TaskGraph.cpp
@@ -411,8 +411,6 @@ public:
 	 *	Sets up some basic information for a thread. Meant to be called from a "main" thread. Also creates the stall event.
 	 *	@param InThreadId; Thread index for this thread.
 	 *	@param InPerThreadIDTLSSlot; TLS slot to store the pointer to me into (later)
-	 *	@param bInAllowsStealsFromMe; If true, this is a worker thread and any other thread can steal tasks from my incoming queue.
-	 *	@param bInStealsFromOthers If true, this is a worker thread and I will attempt to steal tasks when I run out of work.
 	**/
 	void Setup(ENamedThreads::Type InThreadId, uint32 InPerThreadIDTLSSlot, FWorkerThread* InOwnerWorker)
 	{
