[I/O Dispatcher] Extracted platform independent parts from I/O dispatcher

diff --git a/Engine/Source/Runtime/Core/Private/IO/IoBuffer.cpp b/Engine/Source/Runtime/Core/Private/IO/IoBuffer.cpp
new file mode 100644
--- /dev/null
+++ b/Engine/Source/Runtime/Core/Private/IO/IoBuffer.cpp
@@ -0,0 +1,123 @@
+// Copyright 1998-2019 Epic Games, Inc. All Rights Reserved.
+
+#include "IO/IoDispatcher.h"
+
+//////////////////////////////////////////////////////////////////////////
+
+FIoBuffer::BufCore::BufCore()
+{
+}
+
+FIoBuffer::BufCore::~BufCore()
+{
+	if (IsMemoryOwned())
+	{
+		FMemory::Free(Data());
+	}
+}
+
+FIoBuffer::BufCore::BufCore(const uint8* InData, uint64 InSize, bool InOwnsMemory)
+{
+	SetDataAndSize(InData, InSize);
+
+	SetIsOwned(InOwnsMemory);
+}
+
+FIoBuffer::BufCore::BufCore(const uint8* InData, uint64 InSize, const BufCore* InOuter)
+:	OuterCore(InOuter)
+{
+	SetDataAndSize(InData, InSize);
+}
+
+FIoBuffer::BufCore::BufCore(uint64 InSize)
+{
+	uint8* NewBuffer = reinterpret_cast<uint8*>(FMemory::Malloc(InSize));
+
+	SetDataAndSize(NewBuffer, InSize);
+
+	SetIsOwned(true);
+}
+
+FIoBuffer::BufCore::BufCore(ECloneTag, uint8* InData, uint64 InSize)
+:	FIoBuffer::BufCore(InSize)
+{
+	FMemory::Memcpy(Data(), InData, InSize);
+}
+
+void
+FIoBuffer::BufCore::CheckRefCount() const
+{
+	// Verify that Release() is not being called on an object which is already at a zero refcount
+	check(NumRefs != 0);
+}
+
+void
+FIoBuffer::BufCore::SetDataAndSize(const uint8* InData, uint64 InSize)
+{
+	// This is intentionally not split into SetData and SetSize to enable different storage
+	// strategies for flags in the future (in unused pointer bits)
+
+	DataPtr			= const_cast<uint8*>(InData);
+	DataSizeLow		= uint32(InSize & 0xffffffffu);
+	DataSizeHigh	= (InSize >> 32) & 0xffu;
+}
+
+void
+FIoBuffer::BufCore::SetSize(uint64 InSize)
+{
+	SetDataAndSize(Data(), InSize);
+}
+
+void
+FIoBuffer::BufCore::MakeOwned()
+{
+	if (IsMemoryOwned())
+		return;
+
+	const uint64 BufferSize = DataSize();
+	uint8* NewBuffer		= reinterpret_cast<uint8*>(FMemory::Malloc(BufferSize));
+
+	FMemory::Memcpy(NewBuffer, Data(), BufferSize);
+
+	SetDataAndSize(NewBuffer, BufferSize);
+
+	SetIsOwned(true);
+}
+
+//////////////////////////////////////////////////////////////////////////
+
+FIoBuffer::FIoBuffer()
+:	CorePtr(new BufCore)
+{
+}
+
+FIoBuffer::FIoBuffer(uint64 InSize)
+:	CorePtr(new BufCore(InSize))
+{
+}
+
+FIoBuffer::FIoBuffer(const void* Data, uint64 InSize, const FIoBuffer& OuterBuffer)
+:	CorePtr(new BufCore((uint8*) Data, InSize, OuterBuffer.CorePtr))
+{
+}
+
+FIoBuffer::FIoBuffer(FIoBuffer::EWrapTag, const void* Data, uint64 InSize)
+:	CorePtr(new BufCore((uint8*)Data, InSize, /* ownership */ false))
+{
+}
+
+FIoBuffer::FIoBuffer(FIoBuffer::EAssumeOwnershipTag, const void* Data, uint64 InSize)
+:	CorePtr(new BufCore((uint8*)Data, InSize, /* ownership */ true))
+{
+}
+
+FIoBuffer::FIoBuffer(FIoBuffer::ECloneTag, const void* Data, uint64 InSize)
+:	CorePtr(new BufCore(Clone, (uint8*)Data, InSize))
+{
+}
+
+void		
+FIoBuffer::MakeOwned() const
+{
+	CorePtr->MakeOwned();
+}
diff --git a/Engine/Source/Runtime/Core/Private/IO/IoDispatcher.cpp b/Engine/Source/Runtime/Core/Private/IO/IoDispatcher.cpp
--- a/Engine/Source/Runtime/Core/Private/IO/IoDispatcher.cpp
+++ b/Engine/Source/Runtime/Core/Private/IO/IoDispatcher.cpp
@@ -1,6 +1,7 @@
 // Copyright 1998-2019 Epic Games, Inc. All Rights Reserved.
 
 #include "IO/IoDispatcher.h"
+#include "IO/IoStore.h"
 #include "Async/MappedFileHandle.h"
 #include "Containers/ChunkedArray.h"
 #include "HAL/CriticalSection.h"
@@ -23,6 +24,8 @@
 
 DEFINE_LOG_CATEGORY_STATIC(LogIoDispatch, Log, All);
 
+#if !defined(PLATFORM_IMPLEMENTS_IO)
+
 //////////////////////////////////////////////////////////////////////////
 
 UE_TRACE_EVENT_BEGIN(IoDispatcher, BatchIssued, Always)
@@ -141,477 +144,6 @@ private:
 
 //////////////////////////////////////////////////////////////////////////
 
-const TCHAR* GetIoErrorText(EIoErrorCode ErrorCode)
-{
-	static constexpr const TCHAR* ErrorCodeText[]
-	{
-		TEXT("OK"),
-		TEXT("Unknown Status"),
-		TEXT("Invalid Code"),
-		TEXT("Cancelled"),
-		TEXT("FileOpen Failed"),
-		TEXT("File Not Open"),
-		TEXT("Write Error"),
-		TEXT("Not Found"),
-		TEXT("Corrupt Toc"),
-		TEXT("Unknown ChunkID")
-		TEXT("Invalid Parameter")
-	};
-
-	return ErrorCodeText[static_cast<uint32>(ErrorCode)];
-}
-
-const FIoStatus FIoStatus::Ok		{ EIoErrorCode::Ok,				TEXT("OK")				};
-const FIoStatus FIoStatus::Unknown	{ EIoErrorCode::Unknown,		TEXT("Unknown Status")	};
-const FIoStatus FIoStatus::Invalid	{ EIoErrorCode::InvalidCode,	TEXT("Invalid Code")	};
-
-FIoStatus::FIoStatus()
-{
-}
-
-FIoStatus::~FIoStatus()
-{
-}
-
-FIoStatus::FIoStatus(EIoErrorCode Code)
-:	ErrorCode(Code)
-{
-	ErrorMessage[0] = 0;
-}
-
-FIoStatus::FIoStatus(EIoErrorCode Code, const FStringView& InErrorMessage)
-: ErrorCode(Code)
-{
-	const int32 ErrorMessageLength = FMath::Min(MaxErrorMessageLength - 1, InErrorMessage.Len());
-	FPlatformString::Convert(ErrorMessage, ErrorMessageLength, InErrorMessage.Data(), ErrorMessageLength);
-	ErrorMessage[ErrorMessageLength] = 0;
-}
-
-FIoStatus& 
-FIoStatus::operator=(const FIoStatus& Other)
-{
-	ErrorCode = Other.ErrorCode;
-	FMemory::Memcpy(ErrorMessage, Other.ErrorMessage, MaxErrorMessageLength * sizeof(TCHAR));
-
-	return *this;
-}
-
-bool		
-FIoStatus::operator==(const FIoStatus& Other) const
-{
-	return ErrorCode == Other.ErrorCode &&
-		FPlatformString::Stricmp(ErrorMessage, Other.ErrorMessage) == 0;
-}
-
-FString	
-FIoStatus::ToString() const
-{
-	return FString::Format(TEXT("{0} ({1})"), { ErrorMessage, GetIoErrorText(ErrorCode) });
-}
-
-//////////////////////////////////////////////////////////////////////////
-
-FIoStatusBuilder::FIoStatusBuilder(EIoErrorCode InStatusCode)
-:	StatusCode(InStatusCode)
-{
-}
-
-FIoStatusBuilder::FIoStatusBuilder(const FIoStatus& InStatus, FStringView String)
-:	StatusCode(InStatus.ErrorCode)
-{
-	Message.Append(String.Data(), String.Len());
-}
-
-FIoStatusBuilder::~FIoStatusBuilder()
-{
-}
-
-FIoStatusBuilder::operator FIoStatus()
-{
-	return FIoStatus(StatusCode, Message);
-}
-
-FIoStatusBuilder& 
-FIoStatusBuilder::operator<<(FStringView String)
-{
-	Message.Append(String.Data(), String.Len());
-
-	return *this;
-}
-
-FIoStatusBuilder operator<<(const FIoStatus& Status, FStringView String)
-{ 
-	return FIoStatusBuilder(Status, String);
-}
-
-//////////////////////////////////////////////////////////////////////////
-
-void StatusOrCrash(const FIoStatus& Status)
-{
-	// TODO
-}
-
-//////////////////////////////////////////////////////////////////////////
-
-void FIoChunkId::GenerateFromData(const void* InData, SIZE_T InDataSize)
-{
-	FSHA1 Sha1;
-	Sha1.Update(reinterpret_cast<const uint8*>(InData), InDataSize);
-	Sha1.Final();
-	
-	uint8 Sha1Hash[20];
-	Sha1.GetHash(Sha1Hash);
-
-	memcpy(Id, Sha1Hash, sizeof Id);
-}
-
-//////////////////////////////////////////////////////////////////////////
-
-class FIoBufferManager
-{
-public:
-	TChunkedArray<FIoBuffer::BufCore, 65536> BufCores;
-
-	void AllocBufCore()
-	{
-		new(BufCores) FIoBuffer::BufCore;
-	}
-};
-
-//////////////////////////////////////////////////////////////////////////
-
-class FIoBufferPool
-{
-public:
-	FIoBufferPool();
-	~FIoBufferPool();
-
-	TIoStatusOr<FIoBuffer>	Alloc(uint64 RequiredCapacity);
-	void					Free(FIoBuffer InBuffer);
-};
-
-FIoBufferPool::FIoBufferPool()
-{
-}
-
-FIoBufferPool::~FIoBufferPool()
-{
-}
-
-TIoStatusOr<FIoBuffer> FIoBufferPool::Alloc(uint64 RequiredCapacity)
-{
-	return FIoBuffer(RequiredCapacity);
-}
-
-void FIoBufferPool::Free(FIoBuffer InBuffer)
-{
-	auto Destroyer = MoveTemp(InBuffer);
-}
-
-//////////////////////////////////////////////////////////////////////////
-
-FIoBuffer::BufCore::BufCore()
-{
-}
-
-FIoBuffer::BufCore::~BufCore()
-{
-	if (IsMemoryOwned())
-	{
-		FMemory::Free(Data());
-	}
-}
-
-FIoBuffer::BufCore::BufCore(const uint8* InData, uint64 InSize, bool InOwnsMemory)
-{
-	SetDataAndSize(InData, InSize);
-
-	SetIsOwned(InOwnsMemory);
-}
-
-FIoBuffer::BufCore::BufCore(const uint8* InData, uint64 InSize, const BufCore* InOuter)
-:	OuterCore(InOuter)
-{
-	SetDataAndSize(InData, InSize);
-}
-
-FIoBuffer::BufCore::BufCore(uint64 InSize)
-{
-	uint8* NewBuffer = reinterpret_cast<uint8*>(FMemory::Malloc(InSize));
-
-	SetDataAndSize(NewBuffer, InSize);
-
-	SetIsOwned(true);
-}
-
-FIoBuffer::BufCore::BufCore(ECloneTag, uint8* InData, uint64 InSize)
-:	FIoBuffer::BufCore(InSize)
-{
-	FMemory::Memcpy(Data(), InData, InSize);
-}
-
-void
-FIoBuffer::BufCore::CheckRefCount() const
-{
-	// Verify that Release() is not being called on an object which is already at a zero refcount
-	check(NumRefs != 0);
-}
-
-void
-FIoBuffer::BufCore::SetDataAndSize(const uint8* InData, uint64 InSize)
-{
-	// This is intentionally not split into SetData and SetSize to enable different storage
-	// strategies for flags in the future (in unused pointer bits)
-
-	DataPtr			= const_cast<uint8*>(InData);
-	DataSizeLow		= uint32(InSize & 0xffffffffu);
-	DataSizeHigh	= (InSize >> 32) & 0xffu;
-}
-
-void
-FIoBuffer::BufCore::SetSize(uint64 InSize)
-{
-	SetDataAndSize(Data(), InSize);
-}
-
-void
-FIoBuffer::BufCore::MakeOwned()
-{
-	if (IsMemoryOwned())
-		return;
-
-	const uint64 BufferSize = DataSize();
-	uint8* NewBuffer		= reinterpret_cast<uint8*>(FMemory::Malloc(BufferSize));
-
-	FMemory::Memcpy(NewBuffer, Data(), BufferSize);
-
-	SetDataAndSize(NewBuffer, BufferSize);
-
-	SetIsOwned(true);
-}
-
-//////////////////////////////////////////////////////////////////////////
-
-FIoBuffer::FIoBuffer()
-:	CorePtr(new BufCore)
-{
-}
-
-FIoBuffer::FIoBuffer(uint64 InSize)
-:	CorePtr(new BufCore(InSize))
-{
-}
-
-FIoBuffer::FIoBuffer(const void* Data, uint64 InSize, const FIoBuffer& OuterBuffer)
-:	CorePtr(new BufCore((uint8*) Data, InSize, OuterBuffer.CorePtr))
-{
-}
-
-FIoBuffer::FIoBuffer(FIoBuffer::EWrapTag, const void* Data, uint64 InSize)
-:	CorePtr(new BufCore((uint8*)Data, InSize, /* ownership */ false))
-{
-}
-
-FIoBuffer::FIoBuffer(FIoBuffer::EAssumeOwnershipTag, const void* Data, uint64 InSize)
-:	CorePtr(new BufCore((uint8*)Data, InSize, /* ownership */ true))
-{
-}
-
-FIoBuffer::FIoBuffer(FIoBuffer::ECloneTag, const void* Data, uint64 InSize)
-:	CorePtr(new BufCore(Clone, (uint8*)Data, InSize))
-{
-}
-
-void		
-FIoBuffer::MakeOwned() const
-{
-	CorePtr->MakeOwned();
-}
-
-//////////////////////////////////////////////////////////////////////////
-
-static const char TocMagicImg[] = "-==--==--==--==-";
-
-struct FIoStoreTocHeader
-{
-	uint8	TocMagic[16];
-	uint32	TocHeaderSize;
-	uint32	TocEntryCount;
-	uint32	TocEntrySize;	// For sanity checking
-	uint32	TocPad[25];
-
-	void MakeMagic()
-	{
-		FMemory::Memcpy(TocMagic, TocMagicImg, sizeof TocMagic);
-	}
-
-	bool CheckMagic()
-	{
-		return FMemory::Memcmp(TocMagic, TocMagicImg, sizeof TocMagic) == 0;
-	}
-};
-
-struct FIoStoreTocEntry
-{
-	// We use 5 bytes for offset and size, this is enough to represent 
-	// an offset and size of 1PB
-	uint8		OffsetAndLength[5 + 5];
-
-	// TBD: should the chunk ID use content addressing, or names, or a
-	//      mix of both?
-	FIoChunkId	ChunkId;
-
-	inline uint64 GetOffset() const
-	{
-		return OffsetAndLength[4]
-			| (uint64(OffsetAndLength[3]) << 8)
-			| (uint64(OffsetAndLength[2]) << 16)
-			| (uint64(OffsetAndLength[1]) << 24)
-			| (uint64(OffsetAndLength[0]) << 32)
-			;
-	}
-
-	inline uint64 GetLength() const
-	{
-		return OffsetAndLength[9]
-			| (uint64(OffsetAndLength[8]) << 8)
-			| (uint64(OffsetAndLength[7]) << 16)
-			| (uint64(OffsetAndLength[6]) << 24)
-			| (uint64(OffsetAndLength[5]) << 32)
-			;
-	}
-
-	inline void SetOffset(uint64 Offset)
-	{
-		OffsetAndLength[0] = uint8(Offset >> 32);
-		OffsetAndLength[1] = uint8(Offset >> 24);
-		OffsetAndLength[2] = uint8(Offset >> 16);
-		OffsetAndLength[3] = uint8(Offset >>  8);
-		OffsetAndLength[4] = uint8(Offset >>  0);
-	}
-
-	inline void SetLength(uint64 Length)
-	{
-		OffsetAndLength[5] = uint8(Length >> 32);
-		OffsetAndLength[6] = uint8(Length >> 24);
-		OffsetAndLength[7] = uint8(Length >> 16);
-		OffsetAndLength[8] = uint8(Length >> 8);
-		OffsetAndLength[9] = uint8(Length >> 0);
-	}
-};
-
-//////////////////////////////////////////////////////////////////////////
-
-FIoStoreEnvironment::FIoStoreEnvironment()
-{
-}
-
-FIoStoreEnvironment::~FIoStoreEnvironment()
-{
-}
-
-void FIoStoreEnvironment::InitializeFileEnvironment(FStringView InRootPath)
-{
-	RootPath = InRootPath.ToString();
-}
-
-//////////////////////////////////////////////////////////////////////////
-
-/** Very simple string class
-
-	The backing memory is 
-  */
-class FCompactString
-{
-	typedef TCHAR CharT;
-
-public:
-						FCompactString();
-	explicit			FCompactString(FStringView InString);
-						FCompactString(const FCompactString& InString);
-						FCompactString(FCompactString&& InString);
-						~FCompactString();
-
-	inline const CharT* operator*() const						{ return Chars; }
-
-	inline FCompactString& operator=(FStringView Rhs)			{ Assign(Rhs.Data(), Rhs.Len()); return *this; }
-	inline FCompactString& operator=(const FCompactString& Rhs) { Assign(*Rhs); return *this; }
-	inline FCompactString& operator=(FCompactString&& Rhs)		{ /* swap guts */ auto Temp = Chars; Chars = Rhs.Chars; Rhs.Chars = Temp; return *this; }
-
-private:
-	void Initialize(const CharT* InChars, SIZE_T InCharCount);
-	void Initialize(const CharT* InChars);
-	void Assign(const CharT* InChars, SIZE_T InCharCount);
-	void Assign(const CharT* InChars);
-
-	TArenaPointer<CharT>	Chars;
-};
-
-FCompactString::FCompactString()
-{
-}
-
-FCompactString::~FCompactString()
-{
-}
-
-FCompactString::FCompactString(FCompactString&& InString)
-{
-	Chars = InString.Chars;
-	InString.Chars = nullptr;
-}
-
-FCompactString::FCompactString(const FCompactString& InString)
-{
-	Initialize(*InString);
-}
-
-FCompactString::FCompactString(FStringView InString)
-{
-	Initialize(InString.Data(), InString.Len());
-}
-
-void 
-FCompactString::Initialize(const CharT* InChars, SIZE_T InCharCount)
-{
-	// The count does not include the NUL terminator
-	++InCharCount;
-
-	Chars = reinterpret_cast<CharT*>(FMemory::Malloc(InCharCount * sizeof(CharT)));
-
-	FMemory::Memcpy(Chars, InChars, InCharCount * sizeof(CharT));
-}
-
-void
-FCompactString::Initialize(const CharT* InChars)
-{
-	Initialize(InChars, TCString<CharT>::Strlen(InChars));
-}
-
-void
-FCompactString::Assign(const CharT* InChars, SIZE_T InCharCount)
-{
-	// The count does not include the NUL terminator
-	++InCharCount;
-
-	if (Chars)
-	{
-		FMemory::Free(Chars);
-	}
-
-	Chars = reinterpret_cast<CharT*>(FMemory::Malloc(InCharCount * sizeof(CharT)));
-
-	FMemory::Memcpy(Chars, InChars, InCharCount * sizeof(CharT));
-}
-
-void
-FCompactString::Assign(const CharT* InChars)
-{
-	Assign(InChars, TCString<CharT>::Strlen(InChars));
-}
-
-//////////////////////////////////////////////////////////////////////////
-
 class FIoStoreReaderImpl
 {
 public:
@@ -636,7 +168,7 @@ public:
 
 private:
 	FIoStoreEnvironment&				Environment;
-	FCompactString						UniqueId;
+	FString								UniqueId;
 	TMap<FIoChunkId, FIoStoreTocEntry>	Toc;
 	TUniquePtr<IFileHandle>				ContainerFileHandle;
 	TUniquePtr<IMappedFileHandle>		ContainerMappedFileHandle;
@@ -647,7 +179,7 @@ FIoStatus FIoStoreReaderImpl::Open(FStringView InUniqueId)
 {
 	IPlatformFile& Ipf = IPlatformFile::GetPlatformPhysical();
 
-	UniqueId = InUniqueId;
+	UniqueId = InUniqueId.ToString();
 
 	const FString& RootPath = Environment.GetRootPath();
 
@@ -755,182 +287,6 @@ FIoStoreReader::Initialize(FStringView UniqueId)
 
 //////////////////////////////////////////////////////////////////////////
 
-class FIoStoreWriterImpl
-{
-public:
-	FIoStoreWriterImpl(FIoStoreEnvironment& InEnvironment)
-	:	Environment(InEnvironment)
-	{
-	}
-
-	FIoStatus Open()
-	{
-		IPlatformFile& Ipf = IPlatformFile::GetPlatformPhysical();
-
-		const FString& RootPath = Environment.GetRootPath();
-		Ipf.CreateDirectoryTree(*RootPath);
-
-		TStringBuilder<256> ContainerFilePath;
-		ContainerFilePath.Append(RootPath);
-		if (ContainerFilePath.LastChar() != '/')
-			ContainerFilePath.Append(TEXT('/'));
-
-		TStringBuilder<256> TocFilePath;
-		TocFilePath.Append(ContainerFilePath);
-
-		ContainerFilePath.Append(TEXT("Container.ucas"));
-		TocFilePath.Append(TEXT("Container.utoc"));
-
-		ContainerFileHandle.Reset(Ipf.OpenWrite(*ContainerFilePath, /* append */ false, /* allowread */ true));
-
-		if (!ContainerFileHandle)
-		{ 
-			return FIoStatusBuilder(EIoErrorCode::FileOpenFailed) << TEXT("Failed to open IoStore container file '") << *ContainerFilePath << TEXT("'");
-		}
-
-		TocFileHandle.Reset(Ipf.OpenWrite(*TocFilePath, /* append */ false, /* allowread */ true));
-
-		if (!TocFileHandle)
-		{
-			return FIoStatusBuilder(EIoErrorCode::FileOpenFailed) << TEXT("Failed to open IoStore TOC file '") << *TocFilePath << TEXT("'");
-		}
-
-		return FIoStatus::Ok;
-	}
-
-	UE_NODISCARD FIoStatus Append(FIoChunkId ChunkId, FIoBuffer Chunk)
-	{
-		if (!ContainerFileHandle)
-		{
-			return FIoStatus(EIoErrorCode::FileNotOpen, TEXT("No container file to append to"));
-		}
-
-		FIoStoreTocEntry TocEntry;
-
-		TocEntry.SetOffset(ContainerFileHandle->Tell());
-		TocEntry.SetLength(Chunk.DataSize());
-		TocEntry.ChunkId = ChunkId;
-
-		IsMetadataDirty = true;
-
-		const bool Success = ContainerFileHandle->Write(Chunk.Data(), Chunk.DataSize());
-
-		if (Success)
-		{
-			Toc.Add(ChunkId, TocEntry);
-
-			return FIoStatus::Ok;
-		}
-		else
-		{
-			return FIoStatus(EIoErrorCode::WriteError, TEXT("Append failed"));
-		}
-	}
-
-	UE_NODISCARD FIoStatus MapPartialRange(FIoChunkId OriginalChunkId, uint64 Offset, uint64 Length, FIoChunkId ChunkIdPartialRange)
-	{
-		//TODO: Does RelativeOffset + Length overflow?
-
-		const FIoStoreTocEntry* Entry = Toc.Find(OriginalChunkId);
-		if (Entry == nullptr)
-		{
-			return FIoStatus(EIoErrorCode::UnknownChunkID, TEXT("OriginalChunkId does not exist in the container"));
-		}
-
-		if (Toc.Find(ChunkIdPartialRange) != nullptr)
-		{
-			return FIoStatus(EIoErrorCode::InvalidParameter, TEXT("ChunkIdPartialRange is already mapped"));
-		}
-
-		if (Offset + Length > Entry->GetLength())
-		{
-			return FIoStatus(EIoErrorCode::InvalidParameter, TEXT("The given range (Offset/Length) is not within the bounds of OriginalChunkId's data"));
-		}
-
-		FIoStoreTocEntry TocEntry;
-
-		TocEntry.SetOffset(Entry->GetOffset() + Offset);
-		TocEntry.SetLength(Length);
-		TocEntry.ChunkId = ChunkIdPartialRange;
-
-		Toc.Add(ChunkIdPartialRange, TocEntry);
-
-		IsMetadataDirty = true;
-
-		return FIoStatus::Ok;
-	}
-
-	UE_NODISCARD FIoStatus FlushMetadata()
-	{
-		TocFileHandle->Seek(0);
-
-		FIoStoreTocHeader TocHeader;
-		FMemory::Memset(TocHeader, 0);
-
-		TocHeader.MakeMagic();
-		TocHeader.TocHeaderSize = sizeof TocHeader;
-		TocHeader.TocEntryCount = Toc.Num();
-		TocHeader.TocEntrySize = sizeof(FIoStoreTocEntry);
-
-		const bool Success = TocFileHandle->Write(reinterpret_cast<const uint8*>(&TocHeader), sizeof TocHeader);
-
-		if (!Success)
-		{
-			return FIoStatus(EIoErrorCode::WriteError, TEXT("TOC write failed"));
-		}
-
-		for (auto& _: Toc)
-		{
-			FIoStoreTocEntry& TocEntry = _.Value;
-			
-			TocFileHandle->Write(reinterpret_cast<const uint8*>(&TocEntry), sizeof TocEntry);
-		}
-
-		return FIoStatus::Ok;
-	}
-
-private:
-	FIoStoreEnvironment&				Environment;
-	TMap<FIoChunkId, FIoStoreTocEntry>	Toc;
-	TUniquePtr<IFileHandle>				ContainerFileHandle;
-	TUniquePtr<IFileHandle>				TocFileHandle;
-	bool								IsMetadataDirty = true;
-};
-
-FIoStoreWriter::FIoStoreWriter(FIoStoreEnvironment& InEnvironment)
-:	Impl(new FIoStoreWriterImpl(InEnvironment))
-{
-}
-
-FIoStoreWriter::~FIoStoreWriter()
-{
-	Impl->FlushMetadata();
-
-	delete Impl;
-}
-
-FIoStatus FIoStoreWriter::Initialize()
-{
-	return Impl->Open();
-}
-
-void FIoStoreWriter::Append(FIoChunkId ChunkId, FIoBuffer Chunk)
-{
-	Impl->Append(ChunkId, Chunk);
-}
-
-void FIoStoreWriter::MapPartialRange(FIoChunkId OriginalChunkId, uint64 Offset, uint64 Length, FIoChunkId ChunkIdPartialRange)
-{
-	Impl->MapPartialRange(OriginalChunkId, Offset, Length, ChunkIdPartialRange);
-}
-
-void FIoStoreWriter::FlushMetadata()
-{
-	Impl->FlushMetadata();
-}
-
-//////////////////////////////////////////////////////////////////////////
-
 class FIoBatchImpl
 {
 public:
@@ -989,51 +345,6 @@ private:
 
 //////////////////////////////////////////////////////////////////////////
 
-class FIoWorker : FRefCountBase
-{
-public:
-					FIoWorker(FIoStoreImpl* InIoStore);
-	virtual			~FIoWorker();
-
-	virtual void	IssueBatch(FIoBatchImpl& Batch) = 0;
-
-private:
-	TRefCountPtr<FIoStoreImpl> IoStore;
-};
-
-FIoWorker::FIoWorker(FIoStoreImpl* InIoStore)
-:	IoStore(InIoStore)
-{
-}
-
-FIoWorker::~FIoWorker()
-{
-}
-
-class FIoWorkerFile : FIoWorker
-{
-public:
-	FIoWorkerFile(FIoStoreImpl* InIoStore)
-	:	FIoWorker(InIoStore)
-	{
-	}
-
-	~FIoWorkerFile()
-	{
-	}
-
-	virtual void IssueBatch(FIoBatchImpl& Batch) override
-	{
-	}
-};
-
-//////////////////////////////////////////////////////////////////////////
-
-using FRequestAllocator	= TBlockAllocator<FIoRequestImpl, 4096>;
-using FBatchAllocator	= TBlockAllocator<FIoBatchImpl, 4096>;
-
-//////////////////////////////////////////////////////////////////////////
-
 class FIoDispatcherImpl
 {
 public:
@@ -1166,6 +477,9 @@ public:
 	}
 
 private:
+	using FRequestAllocator		= TBlockAllocator<FIoRequestImpl, 4096>;
+	using FBatchAllocator		= TBlockAllocator<FIoBatchImpl, 4096>;
+
 	TRefCountPtr<FIoStoreImpl>	IoStore;
 	FRequestAllocator			RequestAllocator;
 	FBatchAllocator				BatchAllocator;
@@ -1562,3 +876,5 @@ FIoQueue::IsEmpty() const
 {
 	return Impl->IsEmpty();
 }
+
+#endif // PLATFORM_IMPLEMENTS_IO
diff --git a/Engine/Source/Runtime/Core/Private/IO/IoStatus.cpp b/Engine/Source/Runtime/Core/Private/IO/IoStatus.cpp
new file mode 100644
--- /dev/null
+++ b/Engine/Source/Runtime/Core/Private/IO/IoStatus.cpp
@@ -0,0 +1,117 @@
+// Copyright 1998-2019 Epic Games, Inc. All Rights Reserved.
+
+#include "IO/IoDispatcher.h"
+
+//////////////////////////////////////////////////////////////////////////
+
+static const TCHAR* GetIoErrorText(EIoErrorCode ErrorCode)
+{
+	static constexpr const TCHAR* ErrorCodeText[]
+	{
+		TEXT("OK"),
+		TEXT("Unknown Status"),
+		TEXT("Invalid Code"),
+		TEXT("Cancelled"),
+		TEXT("FileOpen Failed"),
+		TEXT("File Not Open"),
+		TEXT("Write Error"),
+		TEXT("Not Found"),
+		TEXT("Corrupt Toc"),
+		TEXT("Unknown ChunkID"),
+		TEXT("Invalid Parameter")
+	};
+
+	return ErrorCodeText[static_cast<uint32>(ErrorCode)];
+}
+
+//////////////////////////////////////////////////////////////////////////
+
+const FIoStatus FIoStatus::Ok		{ EIoErrorCode::Ok,				TEXT("OK")				};
+const FIoStatus FIoStatus::Unknown	{ EIoErrorCode::Unknown,		TEXT("Unknown Status")	};
+const FIoStatus FIoStatus::Invalid	{ EIoErrorCode::InvalidCode,	TEXT("Invalid Code")	};
+
+//////////////////////////////////////////////////////////////////////////
+
+FIoStatus::FIoStatus()
+{
+}
+
+FIoStatus::~FIoStatus()
+{
+}
+
+FIoStatus::FIoStatus(EIoErrorCode Code)
+:	ErrorCode(Code)
+{
+	ErrorMessage[0] = 0;
+}
+
+FIoStatus::FIoStatus(EIoErrorCode Code, const FStringView& InErrorMessage)
+: ErrorCode(Code)
+{
+	const int32 ErrorMessageLength = FMath::Min(MaxErrorMessageLength - 1, InErrorMessage.Len());
+	FPlatformString::Convert(ErrorMessage, ErrorMessageLength, InErrorMessage.Data(), ErrorMessageLength);
+	ErrorMessage[ErrorMessageLength] = 0;
+}
+
+FIoStatus& 
+FIoStatus::operator=(const FIoStatus& Other)
+{
+	ErrorCode = Other.ErrorCode;
+	FMemory::Memcpy(ErrorMessage, Other.ErrorMessage, MaxErrorMessageLength * sizeof(TCHAR));
+
+	return *this;
+}
+
+bool		
+FIoStatus::operator==(const FIoStatus& Other) const
+{
+	return ErrorCode == Other.ErrorCode &&
+		FPlatformString::Stricmp(ErrorMessage, Other.ErrorMessage) == 0;
+}
+
+FString	
+FIoStatus::ToString() const
+{
+	return FString::Format(TEXT("{0} ({1})"), { ErrorMessage, GetIoErrorText(ErrorCode) });
+}
+
+void StatusOrCrash(const FIoStatus& Status)
+{
+	// TODO
+}
+
+//////////////////////////////////////////////////////////////////////////
+
+FIoStatusBuilder::FIoStatusBuilder(EIoErrorCode InStatusCode)
+:	StatusCode(InStatusCode)
+{
+}
+
+FIoStatusBuilder::FIoStatusBuilder(const FIoStatus& InStatus, FStringView String)
+:	StatusCode(InStatus.ErrorCode)
+{
+	Message.Append(String.Data(), String.Len());
+}
+
+FIoStatusBuilder::~FIoStatusBuilder()
+{
+}
+
+FIoStatusBuilder::operator FIoStatus()
+{
+	return FIoStatus(StatusCode, Message);
+}
+
+FIoStatusBuilder& 
+FIoStatusBuilder::operator<<(FStringView String)
+{
+	Message.Append(String.Data(), String.Len());
+
+	return *this;
+}
+
+FIoStatusBuilder operator<<(const FIoStatus& Status, FStringView String)
+{ 
+	return FIoStatusBuilder(Status, String);
+}
diff --git a/Engine/Source/Runtime/Core/Private/IO/IoStore.cpp b/Engine/Source/Runtime/Core/Private/IO/IoStore.cpp
new file mode 100644
--- /dev/null
+++ b/Engine/Source/Runtime/Core/Private/IO/IoStore.cpp
@@ -0,0 +1,185 @@
+#include "IO/IoStore.h"
+
+//////////////////////////////////////////////////////////////////////////
+
+FIoStoreEnvironment::FIoStoreEnvironment()
+{
+}
+
+FIoStoreEnvironment::~FIoStoreEnvironment()
+{
+}
+
+void FIoStoreEnvironment::InitializeFileEnvironment(FStringView InRootPath)
+{
+	RootPath = InRootPath.ToString();
+}
+
+//////////////////////////////////////////////////////////////////////////
+
+class FIoStoreWriterImpl
+{
+public:
+	FIoStoreWriterImpl(FIoStoreEnvironment& InEnvironment)
+	:	Environment(InEnvironment)
+	{
+	}
+
+	FIoStatus Open()
+	{
+		IPlatformFile& Ipf = IPlatformFile::GetPlatformPhysical();
+
+		const FString& RootPath = Environment.GetRootPath();
+		Ipf.CreateDirectoryTree(*RootPath);
+
+		TStringBuilder<256> ContainerFilePath;
+		ContainerFilePath.Append(RootPath);
+		if (ContainerFilePath.LastChar() != '/')
+			ContainerFilePath.Append(TEXT('/'));
+
+		TStringBuilder<256> TocFilePath;
+		TocFilePath.Append(ContainerFilePath);
+
+		ContainerFilePath.Append(TEXT("Container.ucas"));
+		TocFilePath.Append(TEXT("Container.utoc"));
+
+		ContainerFileHandle.Reset(Ipf.OpenWrite(*ContainerFilePath, /* append */ false, /* allowread */ true));
+
+		if (!ContainerFileHandle)
+		{ 
+			return FIoStatusBuilder(EIoErrorCode::FileOpenFailed) << TEXT("Failed to open IoStore container file '") << *ContainerFilePath << TEXT("'");
+		}
+
+		TocFileHandle.Reset(Ipf.OpenWrite(*TocFilePath, /* append */ false, /* allowread */ true));
+
+		if (!TocFileHandle)
+		{
+			return FIoStatusBuilder(EIoErrorCode::FileOpenFailed) << TEXT("Failed to open IoStore TOC file '") << *TocFilePath << TEXT("'");
+		}
+
+		return FIoStatus::Ok;
+	}
+
+	UE_NODISCARD FIoStatus Append(FIoChunkId ChunkId, FIoBuffer Chunk)
+	{
+		if (!ContainerFileHandle)
+		{
+			return FIoStatus(EIoErrorCode::FileNotOpen, TEXT("No container file to append to"));
+		}
+
+		FIoStoreTocEntry TocEntry;
+
+		TocEntry.SetOffset(ContainerFileHandle->Tell());
+		TocEntry.SetLength(Chunk.DataSize());
+		TocEntry.ChunkId = ChunkId;
+
+		IsMetadataDirty = true;
+
+		const bool Success = ContainerFileHandle->Write(Chunk.Data(), Chunk.DataSize());
+
+		if (Success)
+		{
+			Toc.Add(ChunkId, TocEntry);
+
+			return FIoStatus::Ok;
+		}
+		else
+		{
+			return FIoStatus(EIoErrorCode::WriteError, TEXT("Append failed"));
+		}
+	}
+
+	UE_NODISCARD FIoStatus MapPartialRange(FIoChunkId OriginalChunkId, uint64 Offset, uint64 Length, FIoChunkId ChunkIdPartialRange)
+	{
+		//TODO: Does RelativeOffset + Length overflow?
+
+		const FIoStoreTocEntry* Entry = Toc.Find(OriginalChunkId);
+		if (Entry == nullptr)
+		{
+			return FIoStatus(EIoErrorCode::UnknownChunkID, TEXT("OriginalChunkId does not exist in the container"));
+		}
+
+		if (Toc.Find(ChunkIdPartialRange) != nullptr)
+		{
+			return FIoStatus(EIoErrorCode::InvalidParameter, TEXT("ChunkIdPartialRange is already mapped"));
+		}
+
+		if (Offset + Length > Entry->GetLength())
+		{
+			return FIoStatus(EIoErrorCode::InvalidParameter, TEXT("The given range (Offset/Length) is not within the bounds of OriginalChunkId's data"));
+		}
+
+		FIoStoreTocEntry TocEntry;
+
+		TocEntry.SetOffset(Entry->GetOffset() + Offset);
+		TocEntry.SetLength(Length);
+		TocEntry.ChunkId = ChunkIdPartialRange;
+
+		Toc.Add(ChunkIdPartialRange, TocEntry);
+
+		IsMetadataDirty = true;
+
+		return FIoStatus::Ok;
+	}
+
+	UE_NODISCARD FIoStatus FlushMetadata()
+	{
+		TocFileHandle->Seek(0);
+
+		FIoStoreTocHeader TocHeader;
+		FMemory::Memset(TocHeader, 0);
+
+		TocHeader.MakeMagic();
+		TocHeader.TocHeaderSize = sizeof TocHeader;
+		TocHeader.TocEntryCount = Toc.Num();
+		TocHeader.TocEntrySize = sizeof(FIoStoreTocEntry);
+
+		const bool Success = TocFileHandle->Write(reinterpret_cast<const uint8*>(&TocHeader), sizeof TocHeader);
+
+		if (!Success)
+		{
+			return FIoStatus(EIoErrorCode::WriteError, TEXT("TOC write failed"));
+		}
+
+		for (auto& _: Toc)
+		{
+			FIoStoreTocEntry& TocEntry = _.Value;
+			
+			TocFileHandle->Write(reinterpret_cast<const uint8*>(&TocEntry), sizeof TocEntry);
+		}
+
+		return FIoStatus::Ok;
+	}
+
+private:
+	FIoStoreEnvironment&				Environment;
+	TMap<FIoChunkId, FIoStoreTocEntry>	Toc;
+	TUniquePtr<IFileHandle>				ContainerFileHandle;
+	TUniquePtr<IFileHandle>				TocFileHandle;
+	bool								IsMetadataDirty = true;
+};
+
+FIoStoreWriter::FIoStoreWriter(FIoStoreEnvironment& InEnvironment)
+:	Impl(new FIoStoreWriterImpl(InEnvironment))
+{
+}
+
+FIoStoreWriter::~FIoStoreWriter()
+{
+	Impl->FlushMetadata();
+}
+
+FIoStatus FIoStoreWriter::Initialize()
+{
+	return Impl->Open();
+}
+
+void FIoStoreWriter::Append(FIoChunkId ChunkId, FIoBuffer Chunk)
+{
+	Impl->Append(ChunkId, Chunk);
+}
+
+void FIoStoreWriter::FlushMetadata()
+{
+	Impl->FlushMetadata();
+}
diff --git a/Engine/Source/Runtime/Core/Private/IO/IoStore.h b/Engine/Source/Runtime/Core/Private/IO/IoStore.h
new file mode 100644
--- /dev/null
+++ b/Engine/Source/Runtime/Core/Private/IO/IoStore.h
@@ -0,0 +1,81 @@
+// Copyright 1998-2019 Epic Games, Inc. All Rights Reserved.
+
+#pragma once
+
+#include "IO/IoDispatcher.h"
+
+/**
+ * I/O Store TOC header.
+ */
+struct FIoStoreTocHeader
+{
+	static constexpr char TocMagicImg[] = "-==--==--==--==-";
+
+	uint8	TocMagic[16];
+	uint32	TocHeaderSize;
+	uint32	TocEntryCount;
+	uint32	TocEntrySize;	// For sanity checking
+	uint32	TocPad[25];
+
+	void MakeMagic()
+	{
+		FMemory::Memcpy(TocMagic, TocMagicImg, sizeof TocMagic);
+	}
+
+	bool CheckMagic()
+	{
+		return FMemory::Memcmp(TocMagic, TocMagicImg, sizeof TocMagic) == 0;
+	}
+};
+
+/**
+ * I/O Store TOC entry.
+ */
+struct FIoStoreTocEntry
+{
+	// We use 5 bytes for offset and size, this is enough to represent 
+	// an offset and size of 1PB
+	uint8		OffsetAndLength[5 + 5];
+
+	// TBD: should the chunk ID use content addressing, or names, or a
+	//      mix of both?
+	FIoChunkId	ChunkId;
+
+	inline uint64 GetOffset() const
+	{
+		return OffsetAndLength[4]
+			| (uint64(OffsetAndLength[3]) << 8)
+			| (uint64(OffsetAndLength[2]) << 16)
+			| (uint64(OffsetAndLength[1]) << 24)
+			| (uint64(OffsetAndLength[0]) << 32)
+			;
+	}
+
+	inline uint64 GetLength() const
+	{
+		return OffsetAndLength[9]
+			| (uint64(OffsetAndLength[8]) << 8)
+			| (uint64(OffsetAndLength[7]) << 16)
+			| (uint64(OffsetAndLength[6]) << 24)
+			| (uint64(OffsetAndLength[5]) << 32)
+			;
+	}
+
+	inline void SetOffset(uint64 Offset)
+	{
+		OffsetAndLength[0] = uint8(Offset >> 32);
+		OffsetAndLength[1] = uint8(Offset >> 24);
+		OffsetAndLength[2] = uint8(Offset >> 16);
+		OffsetAndLength[3] = uint8(Offset >>  8);
+		OffsetAndLength[4] = uint8(Offset >>  0);
+	}
+
+	inline void SetLength(uint64 Length)
+	{
+		OffsetAndLength[5] = uint8(Length >> 32);
+		OffsetAndLength[6] = uint8(Length >> 24);
+		OffsetAndLength[7] = uint8(Length >> 16);
+		OffsetAndLength[8] = uint8(Length >> 8);
+		OffsetAndLength[9] = uint8(Length >> 0);
+	}
+};
diff --git a/Engine/Source/Runtime/Core/Public/IO/IoDispatcher.h b/Engine/Source/Runtime/Core/Public/IO/IoDispatcher.h
--- a/Engine/Source/Runtime/Core/Public/IO/IoDispatcher.h
+++ b/Engine/Source/Runtime/Core/Public/IO/IoDispatcher.h
@@ -97,7 +97,7 @@ public:
 	CORE_API			FIoStatusBuilder(const FIoStatus& InStatus, FStringView String);
 	CORE_API			~FIoStatusBuilder();
 
-	inline operator FIoStatus();
+	CORE_API			operator FIoStatus();
 
 	CORE_API FIoStatusBuilder& operator<<(FStringView String);
 };
