Build fix Linux

diff --git a/Engine/Source/Runtime/Core/Private/IO/IoStore.cpp b/Engine/Source/Runtime/Core/Private/IO/IoStore.cpp
--- a/Engine/Source/Runtime/Core/Private/IO/IoStore.cpp
+++ b/Engine/Source/Runtime/Core/Private/IO/IoStore.cpp
@@ -2,6 +2,10 @@
 
 //////////////////////////////////////////////////////////////////////////
 
+constexpr char FIoStoreTocHeader::TocMagicImg[];
+
+//////////////////////////////////////////////////////////////////////////
+
 FIoStoreEnvironment::FIoStoreEnvironment()
 {
 }
