Build fix

diff --git a/Engine/Source/Runtime/Core/Private/IO/IoStore.cpp b/Engine/Source/Runtime/Core/Private/IO/IoStore.cpp
--- a/Engine/Source/Runtime/Core/Private/IO/IoStore.cpp
+++ b/Engine/Source/Runtime/Core/Private/IO/IoStore.cpp
@@ -1,4 +1,7 @@
 #include "IO/IoStore.h"
+#include "Containers/Map.h"
+#include "HAL/FileManager.h"
+#include "Templates/UniquePtr.h"
 
 //////////////////////////////////////////////////////////////////////////
 
