[I/O Dispatcher] Added missing definition and change return type to reflect implementation

diff --git a/Engine/Source/Runtime/Core/Private/IO/IoStore.cpp b/Engine/Source/Runtime/Core/Private/IO/IoStore.cpp
--- a/Engine/Source/Runtime/Core/Private/IO/IoStore.cpp
+++ b/Engine/Source/Runtime/Core/Private/IO/IoStore.cpp
@@ -1,3 +1,5 @@
+// Copyright 1998-2019 Epic Games, Inc. All Rights Reserved.
+
 #include "IO/IoStore.h"
 #include "Containers/Map.h"
 #include "HAL/FileManager.h"
@@ -32,7 +34,7 @@ public:
 	{
 	}
 
-	FIoStatus Open()
+	FIoStatus Initialize()
 	{
 		IPlatformFile& Ipf = IPlatformFile::GetPlatformPhysical();
 
@@ -178,15 +180,20 @@ FIoStoreWriter::~FIoStoreWriter()
 
 FIoStatus FIoStoreWriter::Initialize()
 {
-	return Impl->Open();
+	return Impl->Initialize();
 }
 
-void FIoStoreWriter::Append(FIoChunkId ChunkId, FIoBuffer Chunk)
+FIoStatus FIoStoreWriter::Append(FIoChunkId ChunkId, FIoBuffer Chunk)
 {
-	Impl->Append(ChunkId, Chunk);
+	return Impl->Append(ChunkId, Chunk);
 }
 
-void FIoStoreWriter::FlushMetadata()
+FIoStatus FIoStoreWriter::MapPartialRange(FIoChunkId OriginalChunkId, uint64 Offset, uint64 Length, FIoChunkId ChunkIdPartialRange)
 {
-	Impl->FlushMetadata();
+	return Impl->MapPartialRange(OriginalChunkId, Offset, Length, ChunkIdPartialRange);
+}
+
+FIoStatus FIoStoreWriter::FlushMetadata()
+{
+	return Impl->FlushMetadata();
 }
diff --git a/Engine/Source/Runtime/Core/Public/IO/IoDispatcher.h b/Engine/Source/Runtime/Core/Public/IO/IoDispatcher.h
--- a/Engine/Source/Runtime/Core/Public/IO/IoDispatcher.h
+++ b/Engine/Source/Runtime/Core/Public/IO/IoDispatcher.h
@@ -655,7 +655,7 @@ public:
 	FIoStoreWriter& operator=(const FIoStoreWriter&) = delete;
 
 	CORE_API FIoStatus	Initialize();
-	CORE_API void		Append(FIoChunkId ChunkId, FIoBuffer Chunk);
+	CORE_API FIoStatus	Append(FIoChunkId ChunkId, FIoBuffer Chunk);
 
 	/**
 	 * Creates an addressable range in an already mapped Chunk.
@@ -665,8 +665,8 @@ public:
 	 * @param Length The length of the range in bytes
 	 * @param ChunkIdPartialRange The FIoChunkId that will map to the range
 	 */
-	CORE_API void		MapPartialRange(FIoChunkId OriginalChunkId, uint64 Offset, uint64 Length, FIoChunkId ChunkIdPartialRange);
-	CORE_API void		FlushMetadata();
+	CORE_API FIoStatus	MapPartialRange(FIoChunkId OriginalChunkId, uint64 Offset, uint64 Length, FIoChunkId ChunkIdPartialRange);
+	CORE_API FIoStatus	FlushMetadata();
 
 private:
 	FIoStoreWriterImpl*		Impl;
