A FIoChunkId with no values in it's Id array is now considered to be 'invalid' and would indicate that the FIoChunkId is not mapped to any data. An invalid FIoChunkId can be created by calling FIoChunkId::CreateEmptyId.
Add checks when mapping FIoChunkId to make sure that they are valid.
FIoChunkId will now serialize with FileArchives.

diff --git a/Engine/Source/Runtime/Core/Private/IO/IoStore.cpp b/Engine/Source/Runtime/Core/Private/IO/IoStore.cpp
--- a/Engine/Source/Runtime/Core/Private/IO/IoStore.cpp
+++ b/Engine/Source/Runtime/Core/Private/IO/IoStore.cpp
@@ -76,6 +76,16 @@ public:
 			return FIoStatus(EIoErrorCode::FileNotOpen, TEXT("No container file to append to"));
 		}
 
+		if (!ChunkId.IsValid())
+		{
+			return FIoStatus(EIoErrorCode::InvalidParameter, TEXT("ChunkId is not valid!"));
+		}
+
+		if (Toc.Find(ChunkId) != nullptr)
+		{
+			return FIoStatus(EIoErrorCode::InvalidParameter, TEXT("ChunkId is already mapped"));
+		}
+
 		FIoStoreTocEntry TocEntry;
 
 		TocEntry.SetOffset(ContainerFileHandle->Tell());
@@ -108,6 +118,11 @@ public:
 			return FIoStatus(EIoErrorCode::UnknownChunkID, TEXT("OriginalChunkId does not exist in the container"));
 		}
 
+		if (!ChunkIdPartialRange.IsValid())
+		{
+			return FIoStatus(EIoErrorCode::InvalidParameter, TEXT("ChunkIdPartialRange is not valid!"));
+		}
+
 		if (Toc.Find(ChunkIdPartialRange) != nullptr)
 		{
 			return FIoStatus(EIoErrorCode::InvalidParameter, TEXT("ChunkIdPartialRange is already mapped"));
diff --git a/Engine/Source/Runtime/Core/Public/IO/IoDispatcher.h b/Engine/Source/Runtime/Core/Public/IO/IoDispatcher.h
--- a/Engine/Source/Runtime/Core/Public/IO/IoDispatcher.h
+++ b/Engine/Source/Runtime/Core/Public/IO/IoDispatcher.h
@@ -460,6 +460,15 @@ private:
 class FIoChunkId
 {
 public:
+	static inline FIoChunkId CreateEmptyId()
+	{
+		FIoChunkId ChunkId;
+		uint8 Data[12] = { 0 };
+		ChunkId.Set(Data, sizeof Data);
+
+		return ChunkId;
+	}
+
 	friend uint32 GetTypeHash(FIoChunkId InId)
 	{
 		uint32 Hash = 5381;
@@ -470,6 +479,12 @@ public:
 		return Hash;
 	}
 
+	friend FArchive& operator<<(FArchive& Ar, FIoChunkId& ChunkId)
+	{
+		Ar.Serialize(&ChunkId.Id, sizeof Id);
+		return Ar;
+	}
+
 	inline bool operator==(const FIoChunkId& Rhs) const
 	{
 		return 0 == FMemory::Memcmp(Id, Rhs.Id, sizeof Id);
@@ -483,6 +498,12 @@ public:
 		FMemory::Memcpy(Id, InIdPtr, sizeof Id);
 	}
 
+	bool IsValid() const
+	{
+		uint8 Data[12] = { 0 };
+		return FMemory::Memcmp(Id, Data, sizeof Id) != 0;
+	}
+
 private:
 	uint8	Id[12];
 };
