The FIoDispatcher system can now be polled to find out the size of a chunk via its FIoChunkId.

diff --git a/Engine/Source/Runtime/Core/Private/IO/IoDispatcher.cpp b/Engine/Source/Runtime/Core/Private/IO/IoDispatcher.cpp
--- a/Engine/Source/Runtime/Core/Private/IO/IoDispatcher.cpp
+++ b/Engine/Source/Runtime/Core/Private/IO/IoDispatcher.cpp
@@ -166,6 +166,20 @@ public:
 		return FIoBuffer(FIoBuffer::Wrap, MappedRegion->GetMappedPtr() + Entry->GetOffset(), Entry->GetLength());
 	}
 
+	TIoStatusOr<uint64> GetSizeForChunk(const FIoChunkId& ChunkId) const
+	{
+		const FIoStoreTocEntry* Entry = Toc.Find(ChunkId);
+
+		if (Entry != nullptr)
+		{
+			return Entry->GetLength();
+		}
+		else
+		{
+			return FIoStatus(EIoErrorCode::NotFound);
+		}
+	}
+
 private:
 	FIoStoreEnvironment&				Environment;
 	FString								UniqueId;
@@ -338,6 +352,21 @@ public:
 		return FIoStatus(EIoErrorCode::NotFound);
 	}
 
+	TIoStatusOr<uint64> GetSizeForChunk(const FIoChunkId& ChunkId) const
+	{
+		for (const auto& IoStore : IoStores)
+		{
+			TIoStatusOr<uint64> Result = IoStore->Impl->GetSizeForChunk(ChunkId);
+
+			if (Result.IsOk())
+			{
+				return Result;
+			}
+		}
+
+		return FIoStatus(EIoErrorCode::NotFound);
+	}
+
 private:
 	FRWLock									RwLockIoStores;
 	TArray<TRefCountPtr<FIoStoreReader>>	IoStores;
@@ -411,6 +440,11 @@ public:
 		BatchAllocator.Destroy(Batch);
 	}
 
+	TIoStatusOr<uint64> GetSizeForChunk(const FIoChunkId& ChunkId) const
+	{
+		return IoStore->GetSizeForChunk(ChunkId);
+	}
+
 	template<typename Func>
 	void IterateBatch(const FIoBatchImpl* Batch, Func&& InCallbackFunction)
 	{
@@ -521,6 +555,13 @@ FIoDispatcher::FreeBatch(FIoBatch Batch)
 	Impl->FreeBatch(Batch.Impl);
 }
 
+// Polling methods
+TIoStatusOr<uint64>	
+FIoDispatcher::GetSizeForChunk(const FIoChunkId& ChunkId) const
+{
+	return Impl->GetSizeForChunk(ChunkId);
+}
+
 //////////////////////////////////////////////////////////////////////////
 
 FIoBatch::FIoBatch(FIoDispatcherImpl* InDispatcher, FIoBatchImpl* InImpl)
diff --git a/Engine/Source/Runtime/Core/Public/IO/IoDispatcher.h b/Engine/Source/Runtime/Core/Public/IO/IoDispatcher.h
--- a/Engine/Source/Runtime/Core/Public/IO/IoDispatcher.h
+++ b/Engine/Source/Runtime/Core/Public/IO/IoDispatcher.h
@@ -606,14 +606,17 @@ private:
 class FIoDispatcher
 {
 public:
-	CORE_API			FIoDispatcher();
-	CORE_API virtual	~FIoDispatcher();
+	CORE_API						FIoDispatcher();
+	CORE_API virtual				~FIoDispatcher();
 
-	CORE_API void		Mount(FIoStoreReader* IoStore);
-	CORE_API void		Unmount(FIoStoreReader* IoStore);
+	CORE_API void					Mount(FIoStoreReader* IoStore);
+	CORE_API void					Unmount(FIoStoreReader* IoStore);
 
-	CORE_API FIoBatch	NewBatch();
-	CORE_API void		FreeBatch(FIoBatch Batch);
+	CORE_API FIoBatch				NewBatch();
+	CORE_API void					FreeBatch(FIoBatch Batch);
+
+	// Polling methods
+	CORE_API TIoStatusOr<uint64>	GetSizeForChunk(const FIoChunkId& ChunkId) const;
 
 	FIoDispatcher(const FIoDispatcher&) = default;
 	FIoDispatcher& operator=(const FIoDispatcher&) = delete;
