[I/O Dispatcher] Single read when parsing TOC file

diff --git a/Engine/Source/Runtime/Core/Private/IO/IoDispatcher.cpp b/Engine/Source/Runtime/Core/Private/IO/IoDispatcher.cpp
--- a/Engine/Source/Runtime/Core/Private/IO/IoDispatcher.cpp
+++ b/Engine/Source/Runtime/Core/Private/IO/IoDispatcher.cpp
@@ -206,14 +206,6 @@ FIoStatus FIoStoreReaderImpl::Open(FStringView InUniqueId)
 	TocFilePath.Append(ContainerFilePath);
 	TocFilePath.Append(TEXT("Container.utoc"));
 
-	TUniquePtr<IFileHandle>	TocFileHandle;
-	TocFileHandle.Reset(Ipf.OpenRead(*TocFilePath, /* allowwrite */ false));
-
-	if (!TocFileHandle)
-	{
-		return FIoStatusBuilder(EIoErrorCode::FileOpenFailed) << TEXT("Failed to open IoStore TOC file '") << *TocFilePath << TEXT("'");
-	}
-
 	ContainerFilePath.Append(TEXT("Container.ucas"));
 	ContainerFileHandle.Reset(Ipf.OpenRead(*ContainerFilePath, /* allowwrite */ false));
 	
@@ -232,52 +224,60 @@ FIoStatus FIoStoreReaderImpl::Open(FStringView InUniqueId)
 		return FIoStatusBuilder(EIoErrorCode::FileOpenFailed) << TEXT("Failed to memory map IoStore container file '") << *ContainerFilePath << TEXT("'");
 	}
 
-	// Parse TOC
-	//
-	// This should ultimately be a read-in-place operation but it looks like this for now
+	TUniquePtr<uint8[]> TocBuffer;
+	bool bTocReadOk = false;
 
-	FIoStoreTocHeader Header;
-	TocFileHandle->Read(reinterpret_cast<uint8*>(&Header), sizeof Header);
+	{
+		TUniquePtr<IFileHandle>	TocFileHandle(Ipf.OpenRead(*TocFilePath, /* allowwrite */ false));
 
-	if (!Header.CheckMagic())
+		if (!TocFileHandle)
+		{
+			return FIoStatusBuilder(EIoErrorCode::FileOpenFailed) << TEXT("Failed to open IoStore TOC file '") << *TocFilePath << TEXT("'");
+		}
+
+		const int64 TocSize = TocFileHandle->Size();
+		TocBuffer = MakeUnique<uint8[]>(TocSize);
+		bTocReadOk = TocFileHandle->Read(TocBuffer.Get(), TocSize);
+	}
+
+	if (!bTocReadOk)
+	{
+		return FIoStatusBuilder(EIoErrorCode::CorruptToc) << TEXT("Failed to read IoStore TOC file '") << *TocFilePath << TEXT("'");
+	}
+
+	const FIoStoreTocHeader* Header = reinterpret_cast<const FIoStoreTocHeader*>(TocBuffer.Get());
+
+	if (!Header->CheckMagic())
 	{
 		return FIoStatusBuilder(EIoErrorCode::CorruptToc) << TEXT("TOC header magic mismatch while reading '") << *TocFilePath << TEXT("'");
 	}
 
-	if (Header.TocHeaderSize != sizeof Header)
+	if (Header->TocHeaderSize != sizeof FIoStoreTocHeader)
 	{
 		return FIoStatusBuilder(EIoErrorCode::CorruptToc) << TEXT("TOC header size mismatch while reading '") << *TocFilePath << TEXT("'");
 	}
 
-	if (Header.TocEntrySize != sizeof(FIoStoreTocEntry))
+	if (Header->TocEntrySize != sizeof(FIoStoreTocEntry))
 	{
 		return FIoStatusBuilder(EIoErrorCode::CorruptToc) << TEXT("TOC entry size mismatch while reading '") << *TocFilePath << TEXT("'");
 	}
 
-	uint32 EntryCount = Header.TocEntryCount;
+	const FIoStoreTocEntry* Entry = reinterpret_cast<const FIoStoreTocEntry*>(TocBuffer.Get() + sizeof FIoStoreTocHeader);
+	uint32 EntryCount = Header->TocEntryCount;
+
+	Toc.Reserve(EntryCount);
 
 	while(EntryCount--)
 	{
-		FIoStoreTocEntry Entry;
-		bool Success = TocFileHandle->Read(reinterpret_cast<uint8*>(&Entry), sizeof Entry);
-
-		if (!Success)
-		{
-			return FIoStatusBuilder(EIoErrorCode::CorruptToc) << TEXT("failed to read TOC entry while reading '") << *TocFilePath << TEXT("'");
-		}
-
-		if ((Entry.GetOffset() + Entry.GetLength()) > ContainerSize)
+		if ((Entry->GetOffset() + Entry->GetLength()) > ContainerSize)
 		{
 			// TODO: add details
 			return FIoStatusBuilder(EIoErrorCode::CorruptToc) << TEXT("TOC entry out of container bounds while reading '") << *TocFilePath << TEXT("'");
 		}
 
-		Toc.Add(Entry.ChunkId, Entry);
-	}
+		Toc.Add(Entry->ChunkId, *Entry);
 
-	if (TocFileHandle->Tell() != TocFileHandle->Size())
-	{
-		return FIoStatusBuilder(EIoErrorCode::CorruptToc) << TEXT("TOC file contains trailing garbage '") << *TocFilePath << TEXT("'");
+		++Entry;
 	}
 
 	return FIoStatus::Ok;
diff --git a/Engine/Source/Runtime/Core/Private/IO/IoStore.h b/Engine/Source/Runtime/Core/Private/IO/IoStore.h
--- a/Engine/Source/Runtime/Core/Private/IO/IoStore.h
+++ b/Engine/Source/Runtime/Core/Private/IO/IoStore.h
@@ -22,7 +22,7 @@ struct FIoStoreTocHeader
 		FMemory::Memcpy(TocMagic, TocMagicImg, sizeof TocMagic);
 	}
 
-	bool CheckMagic()
+	bool CheckMagic() const
 	{
 		return FMemory::Memcmp(TocMagic, TocMagicImg, sizeof TocMagic) == 0;
 	}
