Fix missing parentheses

diff --git a/Engine/Source/Runtime/Core/Private/IO/IoDispatcher.cpp b/Engine/Source/Runtime/Core/Private/IO/IoDispatcher.cpp
--- a/Engine/Source/Runtime/Core/Private/IO/IoDispatcher.cpp
+++ b/Engine/Source/Runtime/Core/Private/IO/IoDispatcher.cpp
@@ -252,7 +252,7 @@ FIoStatus FIoStoreReaderImpl::Open(FStringView InUniqueId)
 		return FIoStatusBuilder(EIoErrorCode::CorruptToc) << TEXT("TOC header magic mismatch while reading '") << *TocFilePath << TEXT("'");
 	}
 
-	if (Header->TocHeaderSize != sizeof FIoStoreTocHeader)
+	if (Header->TocHeaderSize != sizeof(FIoStoreTocHeader))
 	{
 		return FIoStatusBuilder(EIoErrorCode::CorruptToc) << TEXT("TOC header size mismatch while reading '") << *TocFilePath << TEXT("'");
 	}
@@ -262,7 +262,7 @@ FIoStatus FIoStoreReaderImpl::Open(FStringView InUniqueId)
 		return FIoStatusBuilder(EIoErrorCode::CorruptToc) << TEXT("TOC entry size mismatch while reading '") << *TocFilePath << TEXT("'");
 	}
 
-	const FIoStoreTocEntry* Entry = reinterpret_cast<const FIoStoreTocEntry*>(TocBuffer.Get() + sizeof FIoStoreTocHeader);
+	const FIoStoreTocEntry* Entry = reinterpret_cast<const FIoStoreTocEntry*>(TocBuffer.Get() + sizeof(FIoStoreTocHeader));
 	uint32 EntryCount = Header->TocEntryCount;
 
 	Toc.Reserve(EntryCount);
