IoStatusOr unit tests

diff --git a/Engine/Source/Runtime/Core/Private/Tests/IO/IoStatusOrTest.cpp b/Engine/Source/Runtime/Core/Private/Tests/IO/IoStatusOrTest.cpp
new file mode 100644
--- /dev/null
+++ b/Engine/Source/Runtime/Core/Private/Tests/IO/IoStatusOrTest.cpp
@@ -0,0 +1,153 @@
+// Copyright 1998-2019 Epic Games, Inc. All Rights Reserved.
+
+#include "CoreTypes.h"
+#include "Misc/AssertionMacros.h"
+#include "Misc/AutomationTest.h"
+#include "IO/IoDispatcher.h"
+
+#if WITH_DEV_AUTOMATION_TESTS
+
+struct FIoStatusTestType
+{
+	FIoStatusTestType() { }
+	FIoStatusTestType(const FIoStatusTestType& Other)
+		: Text(Other.Text) { }
+	FIoStatusTestType(FIoStatusTestType&&) = default;
+
+	FIoStatusTestType(const FString& InText)
+		: Text(InText) { }
+	FIoStatusTestType(FString&& InText)
+		: Text(MoveTemp(InText)) { }
+
+	FIoStatusTestType& operator=(const FIoStatusTestType& Other) = default;
+	FIoStatusTestType& operator=(FIoStatusTestType&& Other) = default;
+	FIoStatusTestType& operator=(const FString& OtherText)
+	{
+		Text = OtherText;
+		return *this;
+	}
+
+	FString Text;
+};
+
+IMPLEMENT_SIMPLE_AUTOMATION_TEST(FIoStatusOrTest, "System.Core.IO.IoStatusOr", EAutomationTestFlags::ApplicationContextMask | EAutomationTestFlags::SmokeFilter)
+
+void TestConstruct(FAutomationTestBase& Test)
+{
+	{
+		TIoStatusOr<FIoStatusTestType> Result;
+		Test.TestEqual("Default IoStatus is Unknown", Result.Status(), FIoStatus::Unknown);
+	}
+
+	{
+		const TIoStatusOr<FIoStatusTestType> Other;
+		TIoStatusOr<FIoStatusTestType> Result(Other);
+		Test.TestEqual("Copy construct", Result.Status(), FIoStatus::Unknown);
+	}
+
+	{
+		const FIoStatus IoStatus(EIoErrorCode::InvalidCode);
+		TIoStatusOr<FIoStatusTestType> Result(IoStatus);
+		Test.TestEqual("Construct with status", Result.Status().GetErrorCode(), EIoErrorCode::InvalidCode);
+	}
+
+	{
+		const FString ExpectedText("Unreal");
+		const FIoStatusTestType Type(ExpectedText);
+		TIoStatusOr<FIoStatusTestType> Result(Type);
+		Test.TestEqual("Construct with value", Result.ValueOrDie().Text, ExpectedText);
+	}
+
+	{
+		const FString ExpectedText("Unreal");
+		TIoStatusOr<FIoStatusTestType> Result(FIoStatusTestType("Unreal"));
+		Test.TestEqual("Construct with temporary value", Result.ValueOrDie().Text, ExpectedText);
+	}
+
+	{
+		TIoStatusOr<FIoStatusTestType> Result(FString("Unreal"));
+		Test.TestEqual("Construct with value arguments", Result.ValueOrDie().Text, FString("Unreal"));
+	}
+}
+
+void TestAssignment(FAutomationTestBase& Test)
+{
+	{
+		const EIoErrorCode ExpectedErrorCode = EIoErrorCode::InvalidCode;
+		TIoStatusOr<FIoStatusTestType> Other = FIoStatus(ExpectedErrorCode);
+		TIoStatusOr<FIoStatusTestType> Result;
+		Result = Other;
+		Test.TestEqual("Assign IoStatusOr with status", Result.Status().GetErrorCode(), ExpectedErrorCode);
+	}
+
+	{
+		const EIoErrorCode ExpectedErrorCode = EIoErrorCode::InvalidCode;
+		TIoStatusOr<FIoStatusTestType> Result;
+		Result = TIoStatusOr<FIoStatusTestType>(FIoStatus(ExpectedErrorCode));
+		Test.TestEqual("Assign temporary IoStatusOr with status", Result.Status().GetErrorCode(), ExpectedErrorCode);
+	}
+
+	{
+		const FString ExpectedText("Unreal");
+		TIoStatusOr<FIoStatusTestType> Other = FIoStatusTestType(ExpectedText);
+		TIoStatusOr<FIoStatusTestType> Result;
+		Result = Other;
+		Test.TestEqual("Assign IoStatusOr with value", Result.ValueOrDie().Text, ExpectedText);
+	}
+
+	{
+		const FString ExpectedText("Unreal");
+		TIoStatusOr<FIoStatusTestType> Result;
+		Result = TIoStatusOr<FIoStatusTestType>(ExpectedText);
+		Test.TestEqual("Assign temporary IoStatusOr with value", Result.ValueOrDie().Text, ExpectedText);
+	}
+
+	{
+		const EIoErrorCode ExpectedErrorCode = EIoErrorCode::InvalidCode;
+		const FIoStatus IoStatus(ExpectedErrorCode);
+		TIoStatusOr<FIoStatusTestType> Result;
+		Result = IoStatus; 
+		Test.TestEqual("Assign status", Result.Status().GetErrorCode(), ExpectedErrorCode);
+	}
+
+	{
+		const EIoErrorCode ExpectedErrorCode = EIoErrorCode::InvalidCode;
+		TIoStatusOr<FIoStatusTestType> Result;
+		Result = FIoStatus(ExpectedErrorCode);
+		Test.TestEqual("Assign temporary status", Result.Status().GetErrorCode(), ExpectedErrorCode);
+	}
+
+	{
+		const FString ExpectedText("Unreal");
+		const FIoStatusTestType Value(ExpectedText);
+		TIoStatusOr<FIoStatusTestType> Result;
+		Result = Value;
+		Test.TestEqual("Assign value", Result.ValueOrDie().Text, ExpectedText);
+	}
+
+	{
+		const FString ExpectedText("Unreal");
+		TIoStatusOr<FIoStatusTestType> Result;
+		Result = FIoStatusTestType(ExpectedText);
+		Test.TestEqual("Assign temporary value", Result.ValueOrDie().Text, ExpectedText);
+	}
+}
+
+void TestConsumeValue(FAutomationTestBase& Test)
+{
+	const FString ExpectedText("Unreal");
+	TIoStatusOr<FIoStatusTestType> Result = FIoStatusTestType(ExpectedText);
+	FIoStatusTestType Value = Result.ConsumeValueOrDie();
+	Test.TestEqual("Consume value or die with valid value", Value.Text, ExpectedText);
+}
+
+bool FIoStatusOrTest::RunTest(const FString& Parameters)
+{
+	TestConstruct(*this);
+	TestAssignment(*this);
+	TestConsumeValue(*this);
+
+	return true;
+}
+
+#endif //WITH_DEV_AUTOMATION_TESTS
diff --git a/Engine/Source/Runtime/Core/Public/IO/IoDispatcher.h b/Engine/Source/Runtime/Core/Public/IO/IoDispatcher.h
--- a/Engine/Source/Runtime/Core/Public/IO/IoDispatcher.h
+++ b/Engine/Source/Runtime/Core/Public/IO/IoDispatcher.h
@@ -31,13 +31,9 @@ class FIoDispatcherImpl;
 class FIoStoreReaderImpl;
 class FIoStoreWriterImpl;
 
-//////////////////////////////////////////////////////////////////////////
-//
-// IO Status classes modeled after Google Status / StatusOr
-//
-// TODO: prevent nullptr value in StatusOr
-//
-
+/*
+ * I/O error code.
+ */
 enum class EIoErrorCode
 {
 	Ok,
@@ -53,6 +49,9 @@ enum class EIoErrorCode
 	InvalidParameter
 };
 
+/**
+ * I/O status with error code and message.
+ */
 class FIoStatus
 {
 public:
@@ -85,9 +84,9 @@ private:
 	friend class FIoStatusBuilder;
 };
 
-/** Helper to make it easier to generate meaningful error messages
-    for FIoStatusBuilder.
-  */
+/**
+ * Helper to make it easier to generate meaningful error messages.
+ */
 class FIoStatusBuilder
 {
 	EIoErrorCode		StatusCode;
@@ -104,6 +103,9 @@ public:
 
 CORE_API FIoStatusBuilder operator<<(const FIoStatus& Status, FStringView String);
 
+/**
+ * Optional I/O result or error status.
+ */
 template<typename T>
 class TIoStatusOr
 {
@@ -112,11 +114,11 @@ class TIoStatusOr
 public:
 	TIoStatusOr() : StatusValue(FIoStatus::Unknown) { }
 	TIoStatusOr(const TIoStatusOr& Other);
-	explicit TIoStatusOr(TIoStatusOr&& Other);
+	TIoStatusOr(TIoStatusOr&& Other);
 
-	TIoStatusOr(FIoStatus Status);
-	TIoStatusOr(const T& Value);
-	explicit TIoStatusOr(T&& Value);
+	TIoStatusOr(FIoStatus InStatus);
+	TIoStatusOr(const T& InValue);
+	TIoStatusOr(T&& InValue);
 
 	~TIoStatusOr();
 
@@ -128,6 +130,8 @@ public:
 
 	TIoStatusOr<T>& operator=(const TIoStatusOr<T>& Other);
 	TIoStatusOr<T>& operator=(TIoStatusOr<T>&& Other);
+	TIoStatusOr<T>& operator=(const FIoStatus& OtherStatus);
+	TIoStatusOr<T>& operator=(const T& OtherValue);
 	TIoStatusOr<T>& operator=(T&& OtherValue);
 
 	template<typename U>
@@ -204,10 +208,10 @@ TIoStatusOr<T>::TIoStatusOr(TIoStatusOr&& Other)
 }
 
 template<typename T>
-TIoStatusOr<T>::TIoStatusOr(FIoStatus Status)
+TIoStatusOr<T>::TIoStatusOr(FIoStatus InStatus)
 {
-	check(!Status.IsOk());
-	StatusValue = Status;
+	check(!InStatus.IsOk());
+	StatusValue = InStatus;
 }
 
 template<typename T>
@@ -218,10 +222,10 @@ TIoStatusOr<T>::TIoStatusOr(const T& InValue)
 }
 
 template<typename T>
-TIoStatusOr<T>::TIoStatusOr(T&& Value)
+TIoStatusOr<T>::TIoStatusOr(T&& InValue)
 {
 	StatusValue = FIoStatus::Ok;
-	new(&Value) T(MoveTempIfPossible(Value));
+	new(&Value) T(MoveTempIfPossible(InValue));
 }
 
 template <typename T>
@@ -254,12 +258,15 @@ template<typename T>
 TIoStatusOr<T>&
 TIoStatusOr<T>::operator=(const TIoStatusOr<T>& Other)
 {
-	Reset();
-	StatusValue = Other.StatusValue;
-
-	if (StatusValue.IsOk())
+	if (&Other != this)
 	{
-		new(&Value) T(*(const T*)&Other.Value);
+		Reset();
+		StatusValue = Other.StatusValue;
+
+		if (StatusValue.IsOk())
+		{
+			new(&Value) T(*(const T*)&Other.Value);
+		}
 	}
 
 	return *this;
@@ -285,12 +292,44 @@ TIoStatusOr<T>::operator=(TIoStatusOr<T>&& Other)
 
 template<typename T>
 TIoStatusOr<T>&
-TIoStatusOr<T>::operator=(T&& OtherValue)
+TIoStatusOr<T>::operator=(const FIoStatus& OtherStatus)
 {
+	check(!OtherStatus.IsOk());
+
 	Reset();
-	
-	StatusValue = FIoStatus::Ok;
-	new(&Value) T(MoveTempIfPossible(Value));
+	StatusValue = OtherStatus;
+
+	return *this;
+}
+
+template<typename T>
+TIoStatusOr<T>&
+TIoStatusOr<T>::operator=(const T& OtherValue)
+{
+	if (&OtherValue != (T*)&Value)
+	{
+		Reset();
+		
+		StatusValue = FIoStatus::Ok;
+		new(&Value) T(OtherValue);
+	}
+
+	return *this;
+}
+
+template<typename T>
+TIoStatusOr<T>&
+TIoStatusOr<T>::operator=(T&& OtherValue)
+{
+	if (&OtherValue != (T*)&Value)
+	{
+		Reset();
+		
+		StatusValue = FIoStatus::Ok;
+		new(&Value) T(MoveTempIfPossible(OtherValue));
+	}
+
+	return *this;
 }
 
 template<typename T>
