Update copyright to 2019 to silence the builder email warning.

diff --git a/Engine/Source/Runtime/Core/Public/Misc/GeneratedTypeName.h b/Engine/Source/Runtime/Core/Public/Misc/GeneratedTypeName.h
--- a/Engine/Source/Runtime/Core/Public/Misc/GeneratedTypeName.h
+++ b/Engine/Source/Runtime/Core/Public/Misc/GeneratedTypeName.h
@@ -1,4 +1,4 @@
-// Copyright 1998-2018 Epic Games, Inc. All Rights Reserved.
+// Copyright 1998-2019 Epic Games, Inc. All Rights Reserved.
 
 #pragma once
 
