[I/O Dispatcher] Set value before status in IoStatusOr

diff --git a/Engine/Source/Runtime/Core/Private/IO/IoDispatcher.cpp b/Engine/Source/Runtime/Core/Private/IO/IoDispatcher.cpp
--- a/Engine/Source/Runtime/Core/Private/IO/IoDispatcher.cpp
+++ b/Engine/Source/Runtime/Core/Private/IO/IoDispatcher.cpp
@@ -3,13 +3,10 @@
 #include "IO/IoDispatcher.h"
 #include "IO/IoStore.h"
 #include "Async/MappedFileHandle.h"
-#include "Containers/ChunkedArray.h"
 #include "HAL/CriticalSection.h"
 #include "HAL/FileManager.h"
 #include "HAL/Runnable.h"
 #include "HAL/RunnableThread.h"
-#include "Logging/LogMacros.h"
-#include "Memory/MemoryArena.h"
 #include "Misc/CString.h"
 #include "Misc/EventPool.h"
 #include "Misc/ScopeRWLock.h"
@@ -22,7 +19,7 @@
 
 #define IODISPATCHER_TRACE_ENABLED !UE_BUILD_SHIPPING
 
-DEFINE_LOG_CATEGORY_STATIC(LogIoDispatch, Log, All);
+DEFINE_LOG_CATEGORY(LogIoDispatcher);
 
 const FIoChunkId FIoChunkId::InvalidChunkId = FIoChunkId::CreateEmptyId();
 
diff --git a/Engine/Source/Runtime/Core/Private/IO/IoStatus.cpp b/Engine/Source/Runtime/Core/Private/IO/IoStatus.cpp
--- a/Engine/Source/Runtime/Core/Private/IO/IoStatus.cpp
+++ b/Engine/Source/Runtime/Core/Private/IO/IoStatus.cpp
@@ -54,8 +54,7 @@ FIoStatus::FIoStatus(EIoErrorCode Code, const FStringView& InErrorMessage)
 	ErrorMessage[ErrorMessageLength] = 0;
 }
 
-FIoStatus& 
-FIoStatus::operator=(const FIoStatus& Other)
+FIoStatus& FIoStatus::operator=(const FIoStatus& Other)
 {
 	ErrorCode = Other.ErrorCode;
 	FMemory::Memcpy(ErrorMessage, Other.ErrorMessage, MaxErrorMessageLength * sizeof(TCHAR));
@@ -63,22 +62,28 @@ FIoStatus::operator=(const FIoStatus& Other)
 	return *this;
 }
 
-bool		
-FIoStatus::operator==(const FIoStatus& Other) const
+FIoStatus& FIoStatus::operator=(const EIoErrorCode InErrorCode)
+{
+	ErrorCode = InErrorCode;
+	ErrorMessage[0] = 0;
+
+	return *this;
+}
+
+bool FIoStatus::operator==(const FIoStatus& Other) const
 {
 	return ErrorCode == Other.ErrorCode &&
 		FPlatformString::Stricmp(ErrorMessage, Other.ErrorMessage) == 0;
 }
 
-FString	
-FIoStatus::ToString() const
+FString FIoStatus::ToString() const
 {
 	return FString::Format(TEXT("{0} ({1})"), { ErrorMessage, GetIoErrorText(ErrorCode) });
 }
 
 void StatusOrCrash(const FIoStatus& Status)
 {
-	// TODO
+	UE_LOG(LogIoDispatcher, Fatal, TEXT("I/O Error '%s'"), *Status.ToString());
 }
 
 //////////////////////////////////////////////////////////////////////////
@@ -103,8 +108,7 @@ FIoStatusBuilder::operator FIoStatus()
 	return FIoStatus(StatusCode, Message);
 }
 
-FIoStatusBuilder& 
-FIoStatusBuilder::operator<<(FStringView String)
+FIoStatusBuilder& FIoStatusBuilder::operator<<(FStringView String)
 {
 	Message.Append(String.Data(), String.Len());
 
diff --git a/Engine/Source/Runtime/Core/Public/IO/IoDispatcher.h b/Engine/Source/Runtime/Core/Public/IO/IoDispatcher.h
--- a/Engine/Source/Runtime/Core/Public/IO/IoDispatcher.h
+++ b/Engine/Source/Runtime/Core/Public/IO/IoDispatcher.h
@@ -4,6 +4,7 @@
 
 #include "CoreTypes.h"
 #include "Containers/UnrealString.h"
+#include "Logging/LogMacros.h"
 #include "Misc/StringView.h"
 #include "Misc/StringBuilder.h"
 #include "Templates/RefCounting.h"
@@ -31,6 +32,8 @@ class FIoDispatcherImpl;
 class FIoStoreReaderImpl;
 class FIoStoreWriterImpl;
 
+CORE_API DECLARE_LOG_CATEGORY_EXTERN(LogIoDispatcher, Log, All);
+
 /*
  * I/O error code.
  */
@@ -61,6 +64,7 @@ public:
 	CORE_API			FIoStatus(EIoErrorCode Code, const FStringView& InErrorMessage);
 	CORE_API			FIoStatus(EIoErrorCode Code);
 	CORE_API FIoStatus&	operator=(const FIoStatus& Other);
+	CORE_API FIoStatus&	operator=(const EIoErrorCode InErrorCode);
 
 	CORE_API bool		operator==(const FIoStatus& Other) const;
 			 bool		operator!=(const FIoStatus& Other) const { return !operator==(Other); }
@@ -75,7 +79,7 @@ public:
 	CORE_API static const FIoStatus Invalid;
 
 private:
-	static constexpr int32 MaxErrorMessageLength = 256;
+	static constexpr int32 MaxErrorMessageLength = 128;
 	using FErrorMessage = TCHAR[MaxErrorMessageLength];
 
 	EIoErrorCode	ErrorCode = EIoErrorCode::Ok;
@@ -155,12 +159,13 @@ CORE_API void StatusOrCrash(const FIoStatus& Status);
 template<typename T>
 void TIoStatusOr<T>::Reset()
 {
-	if (StatusValue.IsOk())
+	EIoErrorCode ErrorCode = StatusValue.GetErrorCode();
+	StatusValue = EIoErrorCode::Unknown;
+
+	if (ErrorCode == EIoErrorCode::Ok)
 	{
 		((T*)&Value)->~T();
 	}
-
-	StatusValue = FIoStatus::Unknown;
 }
 
 template<typename T>
@@ -204,6 +209,7 @@ TIoStatusOr<T>::TIoStatusOr(TIoStatusOr&& Other)
 	if (StatusValue.IsOk())
 	{
 		new(&Value) T(MoveTempIfPossible(*(T*)&Other.Value));
+		Other.StatusValue = EIoErrorCode::Unknown;
 	}
 }
 
@@ -261,11 +267,15 @@ TIoStatusOr<T>::operator=(const TIoStatusOr<T>& Other)
 	if (&Other != this)
 	{
 		Reset();
-		StatusValue = Other.StatusValue;
 
-		if (StatusValue.IsOk())
+		if (Other.StatusValue.IsOk())
 		{
 			new(&Value) T(*(const T*)&Other.Value);
+			StatusValue = EIoErrorCode::Ok;
+		}
+		else
+		{
+			StatusValue = Other.StatusValue;
 		}
 	}
 
@@ -279,11 +289,16 @@ TIoStatusOr<T>::operator=(TIoStatusOr<T>&& Other)
 	if (&Other != this)
 	{
 		Reset();
-		StatusValue = Other.StatusValue;
  
-		if (StatusValue.IsOk())
+		if (Other.StatusValue.IsOk())
 		{
 			new(&Value) T(MoveTempIfPossible(*(T*)&Other.Value));
+			Other.StatusValue = EIoErrorCode::Unknown;
+			StatusValue = EIoErrorCode::Ok;
+		}
+		else
+		{
+			StatusValue = Other.StatusValue;
 		}
 	}
 
@@ -310,8 +325,8 @@ TIoStatusOr<T>::operator=(const T& OtherValue)
 	{
 		Reset();
 		
-		StatusValue = FIoStatus::Ok;
 		new(&Value) T(OtherValue);
+		StatusValue = EIoErrorCode::Ok;
 	}
 
 	return *this;
@@ -325,8 +340,8 @@ TIoStatusOr<T>::operator=(T&& OtherValue)
 	{
 		Reset();
 		
-		StatusValue = FIoStatus::Ok;
 		new(&Value) T(MoveTempIfPossible(OtherValue));
+		StatusValue = EIoErrorCode::Ok;
 	}
 
 	return *this;
@@ -349,11 +364,14 @@ TIoStatusOr<T>& TIoStatusOr<T>::operator=(const TIoStatusOr<U>& Other)
 {
 	Reset();
 
-	StatusValue = Other.StatusValue;
-
-	if (StatusValue.IsOk())
+	if (Other.StatusValue.IsOk())
 	{
 		new(&Value) T(*(const U*)&Other.Value);
+		StatusValue = EIoErrorCode::Ok;
+	}
+	else
+	{
+		StatusValue = Other.StatusValue;
 	}
 
 	return *this;
