Fix error message.

diff --git a/Engine/Source/Runtime/Core/Private/Math/Matrix.cpp b/Engine/Source/Runtime/Core/Private/Math/Matrix.cpp
--- a/Engine/Source/Runtime/Core/Private/Math/Matrix.cpp
+++ b/Engine/Source/Runtime/Core/Private/Math/Matrix.cpp
@@ -4,6 +4,6 @@
 
 void FMatrix::ErrorEnsure(const TCHAR* Message)
 {
-	UE_LOG(LogUnrealMath, Error, TEXT("FMatrix::InverseFast(), trying to invert a NIL matrix, this results in NaNs! Use Inverse() instead."));
-	ensureMsgf(false, TEXT("FMatrix::InverseFast(), trying to invert a NIL matrix, this results in NaNs! Use Inverse() instead."));
+	UE_LOG(LogUnrealMath, Error, TEXT("%s"), Message);
+	ensureMsgf(false, TEXT("%s"), Message);
 }
