Allow moving between different heap allocators with different size types, e.g. TArray->TArray64.

diff --git a/Engine/Source/Runtime/Core/Public/Containers/Array.h b/Engine/Source/Runtime/Core/Public/Containers/Array.h
--- a/Engine/Source/Runtime/Core/Public/Containers/Array.h
+++ b/Engine/Source/Runtime/Core/Public/Containers/Array.h
@@ -16,6 +16,7 @@
 #include "Algo/HeapSort.h"
 #include "Algo/IsHeap.h"
 #include "Algo/Impl/BinaryHeap.h"
+#include "Templates/AndOrNot.h"
 #include "Templates/IdentityFunctor.h"
 #include "Templates/Less.h"
 #include "Templates/ChooseClass.h"
@@ -252,7 +253,7 @@ namespace UE4Array_Private
 		enum
 		{
 			Value =
-				TAreTypesEqual<FromAllocatorType, ToAllocatorType>::Value && // Allocators must be equal
+				TOr<TAreTypesEqual<FromAllocatorType, ToAllocatorType>, TCanMoveBetweenAllocators<FromAllocatorType, ToAllocatorType>>::Value && // Allocators must be equal or move-compatible
 				TContainerTraits<FromArrayType>::MoveWillEmptyContainer &&   // A move must be allowed to leave the source array empty
 				(
 					TAreTypesEqual         <ToElementType, FromElementType>::Value || // The element type of the container must be the same, or...
@@ -402,6 +403,27 @@ public:
 	}
 
 private:
+#if !PLATFORM_COMPILER_HAS_IF_CONSTEXPR
+	template <
+		typename FromArrayType,
+		typename ToArrayType,
+		typename TEnableIf<TCanMoveBetweenAllocators<typename FromArrayType::Allocator, typename ToArrayType::Allocator>::Value>::Type* = nullptr
+	>
+	static FORCEINLINE void MoveAllocatorToEmpty(FromArrayType& FromArray, ToArrayType& ToArray)
+	{
+		ToArray.AllocatorInstance.MoveToEmptyFromOtherAllocator<typename FromArrayType::Allocator>(FromArray.AllocatorInstance);
+	}
+
+	template <
+		typename FromArrayType,
+		typename ToArrayType,
+		typename TEnableIf<!TCanMoveBetweenAllocators<typename FromArrayType::Allocator, typename ToArrayType::Allocator>::Value>::Type* = nullptr
+	>
+	static FORCEINLINE void MoveAllocatorToEmpty(FromArrayType& FromArray, ToArrayType& ToArray)
+	{
+		ToArray.AllocatorInstance.MoveToEmpty(FromArray.AllocatorInstance);
+	}
+#endif
 
 	/**
 	 * Moves or copies array. Depends on the array type traits.
@@ -414,10 +436,28 @@ private:
 	template <typename FromArrayType, typename ToArrayType>
 	static FORCEINLINE typename TEnableIf<UE4Array_Private::TCanMoveTArrayPointersBetweenArrayTypes<FromArrayType, ToArrayType>::Value>::Type MoveOrCopy(ToArrayType& ToArray, FromArrayType& FromArray, SizeType PrevMax)
 	{
-		ToArray.AllocatorInstance.MoveToEmpty(FromArray.AllocatorInstance);
+		using FromAllocatorType = typename FromArrayType::Allocator;
+		using ToAllocatorType   = typename ToArrayType::Allocator;
+
+#if PLATFORM_COMPILER_HAS_IF_CONSTEXPR
+		if constexpr (TCanMoveBetweenAllocators<FromAllocatorType, ToAllocatorType>::Value)
+		{
+			ToArray.AllocatorInstance.MoveToEmptyFromOtherAllocator<FromAllocatorType>(FromArray.AllocatorInstance);
+		}
+		else
+		{
+			ToArray.AllocatorInstance.MoveToEmpty(FromArray.AllocatorInstance);
+		}
+#else
+		MoveAllocatorToEmpty(FromArray, ToArray);
+#endif
 
 		ToArray  .ArrayNum = FromArray.ArrayNum;
 		ToArray  .ArrayMax = FromArray.ArrayMax;
+
+		// Ensure the destination container could hold the source range (when the allocator size types shrink)
+		checkf(ToArray.ArrayNum == FromArray.ArrayNum && ToArray.ArrayMax == FromArray.ArrayMax, TEXT("Data lost when moving to a container with a more constrained size type"));
+
 		FromArray.ArrayNum = 0;
 		FromArray.ArrayMax = 0;
 	}
diff --git a/Engine/Source/Runtime/Core/Public/Containers/ContainerAllocationPolicies.h b/Engine/Source/Runtime/Core/Public/Containers/ContainerAllocationPolicies.h
--- a/Engine/Source/Runtime/Core/Public/Containers/ContainerAllocationPolicies.h
+++ b/Engine/Source/Runtime/Core/Public/Containers/ContainerAllocationPolicies.h
@@ -132,6 +132,12 @@ struct TAllocatorTraits : TAllocatorTraitsBase<AllocatorType>
 {
 };
 
+template <typename FromAllocatorType, typename ToAllocatorType>
+struct TCanMoveBetweenAllocators
+{
+	enum { Value = false };
+};
+
 /** This is the allocation policy interface; it exists purely to document the policy's interface, and should not be used. */
 class FContainerAllocatorInterface
 {
@@ -141,6 +147,8 @@ public:
 
 	/** Determines whether the user of the allocator may use the ForAnyElementType inner class. */
 	enum { NeedsElementType = true };
+
+	/** Determines whether the user of the allocator should do range checks */
 	enum { RequireRangeCheck = true };
 
 	/**
@@ -152,11 +160,22 @@ public:
 	{
 		/**
 		 * Moves the state of another allocator into this one.
+		 *
 		 * Assumes that the allocator is currently empty, i.e. memory may be allocated but any existing elements have already been destructed (if necessary).
 		 * @param Other - The allocator to move the state from.  This allocator should be left in a valid empty state.
 		 */
 		void MoveToEmpty(ForElementType& Other);
 
+		/**
+		 * Moves the state of another allocator into this one.  The allocator can be different, and the type must be specified.
+		 * This function should only be called if TAllocatorTraits<AllocatorType>::SupportsMoveFromOtherAllocator is true.
+		 *
+		 * Assumes that the allocator is currently empty, i.e. memory may be allocated but any existing elements have already been destructed (if necessary).
+		 * @param Other - The allocator to move the state from.  This allocator should be left in a valid empty state.
+		 */
+		template <typename OtherAllocatorType>
+		void MoveToEmptyFromOtherAllocator(typename OtherAllocatorType::template ForElementType<ElementType>& Other);
+
 		/** Accesses the container's current data. */
 		ElementType* GetAllocation() const;
 
@@ -367,6 +386,9 @@ public:
 
 	class ForAnyElementType
 	{
+		template <int>
+		friend class TSizedHeapAllocator;
+
 	public:
 		/** Default constructor. */
 		ForAnyElementType()
@@ -374,23 +396,37 @@ public:
 		{}
 
 		/**
-		 * Moves the state of another allocator into this one.
+		 * Moves the state of another allocator into this one.  The allocator can be different.
+		 *
 		 * Assumes that the allocator is currently empty, i.e. memory may be allocated but any existing elements have already been destructed (if necessary).
 		 * @param Other - The allocator to move the state from.  This allocator should be left in a valid empty state.
 		 */
-		FORCEINLINE void MoveToEmpty(ForAnyElementType& Other)
+		template <typename OtherAllocator>
+		FORCEINLINE void MoveToEmptyFromOtherAllocator(typename OtherAllocator::ForAnyElementType& Other)
 		{
-			checkSlow(this != &Other);
+			checkSlow((void*)this != (void*)&Other);
 
 			if (Data)
 			{
 				FMemory::Free(Data);
 			}
 
-			Data       = Other.Data;
+			Data = Other.Data;
 			Other.Data = nullptr;
 		}
 
+		/**
+		 * Moves the state of another allocator into this one.
+		 * Moves the state of another allocator into this one.  The allocator can be different.
+		 *
+		 * Assumes that the allocator is currently empty, i.e. memory may be allocated but any existing elements have already been destructed (if necessary).
+		 * @param Other - The allocator to move the state from.  This allocator should be left in a valid empty state.
+		 */
+		FORCEINLINE void MoveToEmpty(ForAnyElementType& Other)
+		{
+			this->MoveToEmptyFromOtherAllocator<TSizedHeapAllocator>(Other);
+		}
+
 		/** Destructor. */
 		FORCEINLINE ~ForAnyElementType()
 		{
@@ -470,6 +506,13 @@ struct TAllocatorTraits<TSizedHeapAllocator<IndexSize>> : TAllocatorTraitsBase<T
 
 using FHeapAllocator = TSizedHeapAllocator<32>;
 
+template <uint8 FromIndexSize, uint8 ToIndexSize>
+struct TCanMoveBetweenAllocators<TSizedHeapAllocator<FromIndexSize>, TSizedHeapAllocator<ToIndexSize>>
+{
+	// Allow conversions between different int width versions of the allocator
+	enum { Value = true };
+};
+
 /**
  * The inline allocation policy allocates up to a specified number of elements in the same allocation as the container.
  * Any allocation needed beyond that causes all data to be moved into an indirect allocation.
@@ -1038,3 +1081,5 @@ template <> struct TAllocatorTraits<FDefaultAllocator>            : TAllocatorTr
 template <> struct TAllocatorTraits<FDefaultSetAllocator>         : TAllocatorTraits<typename FDefaultSetAllocator        ::Typedef> {};
 template <> struct TAllocatorTraits<FDefaultBitArrayAllocator>    : TAllocatorTraits<typename FDefaultBitArrayAllocator   ::Typedef> {};
 template <> struct TAllocatorTraits<FDefaultSparseArrayAllocator> : TAllocatorTraits<typename FDefaultSparseArrayAllocator::Typedef> {};
+
+template <uint8 FromIndexSize, uint8 ToIndexSize> struct TCanMoveBetweenAllocators<TSizedDefaultAllocator<FromIndexSize>, TSizedDefaultAllocator<ToIndexSize>> : TCanMoveBetweenAllocators<typename TSizedDefaultAllocator<FromIndexSize>::Typedef, typename TSizedDefaultAllocator<ToIndexSize>::Typedef> {};
