PR #6253: Allow exception support when building with clang for windows.  This is mainly to support UHT error reporting. (Contributed by [REDACTED_PERSON])

#rn minor

diff --git a/Engine/Source/Runtime/Core/Public/Windows/WIndowsPlatform.h b/Engine/Source/Runtime/Core/Public/Windows/WIndowsPlatform.h
--- a/Engine/Source/Runtime/Core/Public/Windows/WIndowsPlatform.h
+++ b/Engine/Source/Runtime/Core/Public/Windows/WIndowsPlatform.h
@@ -46,9 +46,6 @@ typedef FWindowsPlatformTypes FPlatformTypes;
 #if defined(__clang__)
 	// @todo clang: Clang compiler on Windows doesn't support SEH exception handling yet (__try/__except)
 	#define PLATFORM_SEH_EXCEPTIONS_DISABLED				1
-
-	// @todo clang: Clang compiler on Windows doesn't support C++ exception handling yet (try/throw/catch)
-	#define PLATFORM_EXCEPTIONS_DISABLED					1
 #endif
 
 #define PLATFORM_SUPPORTS_PRAGMA_PACK						1
