Cleaned up our StringView and implemented a number of methods from FString, making it easier to directly replace existing FString use with FStringView.

- Implement ::EndsWith (single character version)
- Implement ::Left/Mid/Right methods based on FString
- Implement ::FindLastCharByPredicate based on FString
- Implement ::Equals based on FString
- Implement ::Compare based on FString

- Moved any large inline definition to an inl file to make the interface easier to read
- EndsWith now compiles in win32 (replace SIZE_T with SizeType)
- SubStr now asserts that the position is within the StringView, similar to std::string_view (which throws)
- ToString can now be used by Ansi and Wide char versions of the view.
- Remove inheritence to make returning a TStringViewImpl from methods (ie substr) easier for the caller as before we could only return the base type, not FAnsiStringView/FStringView
- Adding unit tests
- Add natvis for TStringViewImpl

diff --git a/Engine/Source/Runtime/Core/Private/Misc/StringView.cpp b/Engine/Source/Runtime/Core/Private/Misc/StringView.cpp
--- a/Engine/Source/Runtime/Core/Private/Misc/StringView.cpp
+++ b/Engine/Source/Runtime/Core/Private/Misc/StringView.cpp
@@ -1,9 +1,3 @@
 // Copyright 1998-2019 Epic Games, Inc. All Rights Reserved.
 
 #include "Misc/StringView.h"
-
-FString 
-FStringView::ToString() const
-{
-	return FString(Len(), Data());
-}
diff --git a/Engine/Source/Runtime/Core/Private/Tests/Misc/StringViewTest.cpp b/Engine/Source/Runtime/Core/Private/Tests/Misc/StringViewTest.cpp
new file mode 100644
--- /dev/null
+++ b/Engine/Source/Runtime/Core/Private/Tests/Misc/StringViewTest.cpp
@@ -0,0 +1,681 @@
+// Copyright 1998-2019 Epic Games, Inc. All Rights Reserved.
+
+#include "Misc/StringView.h"
+#include "Misc/AutomationTest.h"
+
+#if WITH_DEV_AUTOMATION_TESTS 
+
+#define TEST_NAME_ROOT "System.Core.Misc.StringView"
+constexpr const uint32 TestFlags = EAutomationTestFlags::ApplicationContextMask | EAutomationTestFlags::EngineFilter;
+
+IMPLEMENT_SIMPLE_AUTOMATION_TEST(FStringViewTestCtor, TEST_NAME_ROOT ".Ctor", TestFlags)
+bool FStringViewTestCtor::RunTest(const FString& Parameters)
+{
+	// Empty View
+	{
+		FStringView View( TEXT("") );
+		TestEqual("", View.Len(), 0);
+		TestTrue("View.IsEmpty", View.IsEmpty());
+	}
+
+	// Create from a wchar literal
+	{
+		FStringView View(TEXT("Test Ctor"));
+		TestEqual(TEXT("View length"), View.Len(), FCStringWide::Strlen(TEXT("Test Ctor")));
+		TestEqual(TEXT("The result of Strncmp"), FCStringWide::Strncmp(View.Data(), TEXT("Test Ctor"), View.Len()), 0);
+		TestFalse("View.IsEmpty", View.IsEmpty());
+	}
+
+	// Create from a sub section of a wchar literal
+	{
+		FStringView View(TEXT("Test SubSection Ctor"), 4);
+		TestEqual(TEXT("View length"), View.Len(), 4);
+		TestEqual(TEXT("The result of Strncmp"), FCStringWide::Strncmp(View.Data(), TEXT("Test"), View.Len()), 0);
+		TestFalse("View.IsEmpty", View.IsEmpty());
+	}
+
+	// Create from a FString
+	{
+		FString String(TEXT("String Object"));
+		FStringView View(String);
+
+		TestEqual(TEXT("View length"), View.Len(), String.Len());
+		TestEqual(TEXT("The result of Strncmp"), FCStringWide::Strncmp(View.Data(), *String, View.Len()), 0);
+		TestFalse("View.IsEmpty", View.IsEmpty());
+	}
+
+	// Create from a ansi literal
+	{
+		FAnsiStringView View("Test Ctor");
+		TestEqual(TEXT("View length"), View.Len(), FCStringAnsi::Strlen("Test Ctor"));
+		TestEqual(TEXT("The result of Strncmp"), FCStringAnsi::Strncmp(View.Data(), "Test Ctor", View.Len()), 0);
+		TestFalse("View.IsEmpty", View.IsEmpty());
+	}
+
+	// Create from a sub section of an ansi literal
+	{
+		FAnsiStringView View("Test SubSection Ctor", 4);
+		TestEqual(TEXT("View length"), View.Len(), 4);
+		TestEqual(TEXT("The result of Strncmp"), FCStringAnsi::Strncmp(View.Data(), "Test", View.Len()), 0);
+		TestFalse("View.IsEmpty", View.IsEmpty());
+	}
+
+	return true;
+}
+
+IMPLEMENT_SIMPLE_AUTOMATION_TEST(FStringViewTestIterators, TEST_NAME_ROOT ".Iterators", TestFlags)
+bool FStringViewTestIterators::RunTest(const FString& Parameters)
+{
+	// Iterate over a string view
+	{
+		const TCHAR* StringLiteralSrc = TEXT("Iterator!");
+		FStringView View(StringLiteralSrc);
+
+		for (TCHAR C : View)
+		{
+			TestTrue(TEXT("Iterators(0)-Iteration"), C == *StringLiteralSrc++);
+		}
+
+		// Make sure we iterated over the entire string
+		TestTrue(TEXT("Iterators(0-EndCheck"), *StringLiteralSrc == '\0');
+	}
+
+	// Iterate over a partial string view
+	{
+		const TCHAR* StringLiteralSrc = TEXT("Iterator|with extras!");
+		FStringView View(StringLiteralSrc, 8);
+
+		for (TCHAR C : View)
+		{
+			TestTrue(TEXT("Iterators(1)-Iteration"), C == *StringLiteralSrc++);
+		}
+
+		// Make sure we only iterated over the part of the string that the view represents
+		TestTrue(TEXT("Iterators(1)-EndCheck"), *StringLiteralSrc == '|');
+	}
+
+	return true;
+}
+
+IMPLEMENT_SIMPLE_AUTOMATION_TEST(FStringViewTestEquality, TEST_NAME_ROOT ".Equality", TestFlags)
+bool FStringViewTestEquality::RunTest(const FString& Parameters)
+{
+	const TCHAR* WideStringLiteralSrc = TEXT("String To Test!");
+	const TCHAR* WideStringLiteralLower = TEXT("string to test!");
+	const TCHAR* WideStringLiteralUpper = TEXT("STRING TO TEST!");
+	const TCHAR* WideStringLiteralShort = TEXT("String To");
+	const TCHAR* WideStringLiteralLonger = TEXT("String To Test! Extended");
+
+	FStringView WideView(WideStringLiteralSrc);
+
+	TestTrue(TEXT("Equality(0)"), WideView == WideStringLiteralSrc);
+	TestFalse(TEXT("Equality(1)"), WideView != WideStringLiteralSrc);	
+	TestTrue(TEXT("Equality(2)"), WideView == WideStringLiteralLower);
+	TestFalse(TEXT("Equality(3)"), WideView != WideStringLiteralLower);
+	TestTrue(TEXT("Equality(4)"), WideView == WideStringLiteralUpper);
+	TestFalse(TEXT("Equality(5)"), WideView != WideStringLiteralUpper);
+	TestFalse(TEXT("Equality(6)"), WideView == WideStringLiteralShort);
+	TestTrue(TEXT("Equality(7)"), WideView != WideStringLiteralShort);
+	TestFalse(TEXT("Equality(8)"), WideView == WideStringLiteralLonger);
+	TestTrue(TEXT("Equality(9)"), WideView != WideStringLiteralLonger);
+
+	FString WideStringSrc = WideStringLiteralSrc;
+	FString WideStringLower = WideStringLiteralLower;
+	FString WideStringUpper = WideStringLiteralUpper;
+	FString WideStringShort = WideStringLiteralShort;
+	FString WideStringLonger = WideStringLiteralLonger;
+
+	TestTrue(TEXT("Equality(10)"), WideView == WideStringSrc);
+	TestFalse(TEXT("Equality(11)"), WideView != WideStringSrc);
+	TestTrue(TEXT("Equality(12)"), WideView == WideStringLower);
+	TestFalse(TEXT("Equality(13)"), WideView != WideStringLower);
+	TestTrue(TEXT("Equality(14)"), WideView == WideStringUpper);
+	TestFalse(TEXT("Equality(15)"), WideView != WideStringUpper);
+	TestFalse(TEXT("Equality(16)"), WideView == WideStringShort);
+	TestTrue(TEXT("Equality(17)"), WideView != WideStringShort);
+	TestFalse(TEXT("Equality(18)"), WideView == WideStringLonger);
+	TestTrue(TEXT("Equality(19)"), WideView != WideStringLonger);
+
+	FStringView IdenticalView(WideStringLiteralSrc);
+
+	TestTrue(TEXT("Equality(20)"), WideView == IdenticalView);
+	TestTrue(TEXT("Equality(21)"), IdenticalView == WideView);
+
+	return true;
+}
+
+IMPLEMENT_SIMPLE_AUTOMATION_TEST(FStringViewTestComparisonCaseSensitive, TEST_NAME_ROOT ".ComparisonCaseSensitive", TestFlags)
+bool FStringViewTestComparisonCaseSensitive::RunTest(const FString& Parameters)
+{
+	// Basic comparisons involving case
+	{
+		const TCHAR* WideStringLiteralSrc = TEXT("String To Test!");
+		const TCHAR* WideStringLiteralLower = TEXT("string to test!");
+		const TCHAR* WideStringLiteralUpper = TEXT("STRING TO TEST!");
+
+		FStringView WideView(WideStringLiteralSrc);
+
+		TestTrue(TEXT("ComparisonCaseSensitive(0)"), WideView.Compare(WideStringLiteralSrc, ESearchCase::CaseSensitive) == 0);
+		TestFalse(TEXT("ComparisonCaseSensitive(1)"), WideView.Compare(WideStringLiteralLower, ESearchCase::CaseSensitive) > 0);
+		TestFalse(TEXT("ComparisonCaseSensitive(2)"), WideView.Compare(WideStringLiteralUpper, ESearchCase::CaseSensitive) < 0);
+
+		FStringView EmptyView(TEXT(""));
+		TestTrue(TEXT("ComparisonCaseSensitive(3)"), WideView.Compare(EmptyView, ESearchCase::CaseSensitive) > 0);
+
+		FStringView IdenticalView(WideStringLiteralSrc);
+		TestTrue(TEXT("ComparisonCaseSensitive(4)"), WideView.Compare(IdenticalView, ESearchCase::CaseSensitive) == 0);
+
+	}
+
+	// Test comparisons of different lengths
+	{
+		const TCHAR* WideStringLiteralUpper = TEXT("ABCDEF");
+		const TCHAR* WideStringLiteralLower = TEXT("abcdef");
+
+		const TCHAR* WideStringLiteralUpperFirst = TEXT("ABCdef");
+		const TCHAR* WideStringLiteralLowerFirst = TEXT("abcDEF");
+
+		FStringView ViewLongUpper(WideStringLiteralUpper);
+		FStringView ViewLongLower(WideStringLiteralLower);
+
+		// Note that the characters after these views are in a different case, this will help catch over read issues
+		FStringView ViewShortUpper(WideStringLiteralUpperFirst, 3);
+		FStringView ViewShortLower(WideStringLiteralLowerFirst, 3);
+
+		// Same length, different cases
+		TestTrue(TEXT("ComparisonCaseSensitive(5)"), ViewLongUpper.Compare(ViewLongLower, ESearchCase::CaseSensitive) < 0);
+		TestTrue(TEXT("ComparisonCaseSensitive(6)"), ViewLongLower.Compare(ViewLongUpper, ESearchCase::CaseSensitive) > 0);
+
+		// Same case, different lengths
+		TestTrue(TEXT("ComparisonCaseSensitive(7)"), ViewLongUpper.Compare(ViewShortUpper, ESearchCase::CaseSensitive) > 0);
+		TestTrue(TEXT("ComparisonCaseSensitive(8)"), ViewShortUpper.Compare(ViewLongUpper, ESearchCase::CaseSensitive) < 0);
+
+		// Different length, different cases
+		TestTrue(TEXT("ComparisonCaseSensitive(9)"), ViewLongUpper.Compare(ViewShortLower, ESearchCase::CaseSensitive) < 0);
+		TestTrue(TEXT("ComparisonCaseSensitive(10)"), ViewShortLower.Compare(ViewLongUpper, ESearchCase::CaseSensitive) > 0);
+	}
+
+	return true;
+}
+
+IMPLEMENT_SIMPLE_AUTOMATION_TEST(FStringViewTestComparisonCaseInsensitive, TEST_NAME_ROOT ".ComparisonCaseInsensitive", TestFlags)
+bool FStringViewTestComparisonCaseInsensitive::RunTest(const FString& Parameters)
+{
+	// Basic comparisons involving case
+	{
+		const TCHAR* WideStringLiteralSrc = TEXT("String To Test!");
+		const TCHAR* WideStringLiteralLower = TEXT("string to test!");
+		const TCHAR* WideStringLiteralUpper = TEXT("STRING TO TEST!");
+
+		FStringView WideView(WideStringLiteralSrc);
+
+		TestTrue(TEXT("ComparisonCaseInsensitive(0)"), WideView.Compare(WideStringLiteralSrc, ESearchCase::IgnoreCase) == 0);
+		TestTrue(TEXT("ComparisonCaseInsensitive(1)"), WideView.Compare(WideStringLiteralLower, ESearchCase::IgnoreCase) == 0);
+		TestTrue(TEXT("ComparisonCaseInsensitive(2)"), WideView.Compare(WideStringLiteralUpper, ESearchCase::IgnoreCase) == 0);
+
+		FStringView EmptyView(TEXT(""));
+		TestTrue(TEXT("ComparisonCaseInsensitive(3)"), WideView.Compare(EmptyView, ESearchCase::IgnoreCase) > 0);
+
+		FStringView IdenticalView(WideStringLiteralSrc);
+		TestTrue(TEXT("ComparisonCaseInsensitive(4)"), WideView.Compare(IdenticalView, ESearchCase::IgnoreCase) == 0);
+
+	}
+
+	// Test comparisons of different lengths
+	{
+		const TCHAR* WideStringLiteralUpper = TEXT("ABCDEF");
+		const TCHAR* WideStringLiteralLower = TEXT("abcdef");
+
+		const TCHAR* WideStringLiteralUpperFirst = TEXT("ABCdef");
+		const TCHAR* WideStringLiteralLowerFirst = TEXT("abcDEF");
+
+		FStringView ViewLongUpper(WideStringLiteralUpper);
+		FStringView ViewLongLower(WideStringLiteralLower);
+
+		// Note that the characters after these views are in a different case, this will help catch over read issues
+		FStringView ViewShortUpper(WideStringLiteralUpperFirst, 3);
+		FStringView ViewShortLower(WideStringLiteralLowerFirst, 3);
+
+		// Same length, different cases
+		TestTrue(TEXT("ComparisonCaseInsensitive(5)"), ViewLongUpper.Compare(ViewLongLower, ESearchCase::IgnoreCase) == 0);
+		TestTrue(TEXT("ComparisonCaseInsensitive(6)"), ViewLongLower.Compare(ViewLongUpper, ESearchCase::IgnoreCase) == 0);
+
+		// Same case, different lengths
+		TestTrue(TEXT("ComparisonCaseInsensitive(7)"), ViewLongUpper.Compare(ViewShortUpper, ESearchCase::IgnoreCase) > 0);
+		TestTrue(TEXT("ComparisonCaseInsensitive(8)"), ViewShortUpper.Compare(ViewLongUpper, ESearchCase::IgnoreCase) < 0);
+
+		// Different length, different cases
+		TestTrue(TEXT("ComparisonCaseInsensitive(9)"), ViewLongUpper.Compare(ViewShortLower, ESearchCase::IgnoreCase) > 0);
+		TestTrue(TEXT("ComparisonCaseInsensitive(10)"), ViewShortLower.Compare(ViewLongUpper, ESearchCase::IgnoreCase) < 0);
+	}
+
+	return true;
+}
+
+IMPLEMENT_SIMPLE_AUTOMATION_TEST(FStringViewTestArrayAccessor, TEST_NAME_ROOT ".ArrayAccessor", TestFlags)
+bool FStringViewTestArrayAccessor::RunTest(const FString& Parameters)
+{
+	const TCHAR* SrcString = TEXT("String To Test");
+	FStringView View(SrcString);
+
+	for( int32 i = 0; i < View.Len(); ++i )
+	{
+		TestEqual(TEXT("the character accessed"), View[i], SrcString[i]);
+	}
+
+	return true;
+}
+
+IMPLEMENT_SIMPLE_AUTOMATION_TEST(FStringViewTestArrayModifiers, TEST_NAME_ROOT ".Modifiers", TestFlags)
+bool FStringViewTestArrayModifiers::RunTest(const FString& Parameters)
+{
+	const TCHAR* FullText = TEXT("PrefixSuffix");
+	const TCHAR* Prefix = TEXT("Prefix");
+	const TCHAR* Suffix = TEXT("Suffix");
+
+	// Remove prefix
+	{
+		FStringView View(FullText);
+		View.RemovePrefix(FCStringWide::Strlen(Prefix));
+
+		TestEqual(TEXT("View length"), View.Len(), FCStringWide::Strlen(Suffix));
+		TestEqual(TEXT("The result of Strncmp"), FCStringWide::Strncmp(View.Data(), Suffix, View.Len()), 0);
+	}
+
+	// Remove suffix
+	{
+		FStringView View(FullText);
+		View.RemoveSuffix(FCStringWide::Strlen(Suffix));
+
+		TestEqual(TEXT("View length"), View.Len(), FCStringWide::Strlen(Prefix));
+		TestEqual(TEXT("The result of Strncmp"), FCStringWide::Strncmp(View.Data(), Prefix, View.Len()), 0);
+	}
+
+	return true;
+}
+
+IMPLEMENT_SIMPLE_AUTOMATION_TEST(FStringViewTestStartsWith, TEST_NAME_ROOT ".StartsWith", TestFlags)
+bool FStringViewTestStartsWith::RunTest(const FString& Parameters)
+{
+	// Test an empty view
+	{
+		FStringView View( TEXT(""));
+		TestTrue(" View.StartsWith", View.StartsWith(TEXT("")));
+		TestFalse(" View.StartsWith", View.StartsWith(TEXT("Text")));
+		TestFalse(" View.StartsWith", View.StartsWith('A'));
+	}
+
+	// Test a valid view with the correct text
+	{
+		FStringView View(TEXT("String to test"));
+		TestTrue(" View.StartsWith", View.StartsWith(TEXT("String")));
+		TestTrue(" View.StartsWith", View.StartsWith('S'));
+	}
+
+	// Test a valid view with incorrect text
+	{
+		FStringView View(TEXT("String to test"));
+		TestFalse(" View.StartsWith", View.StartsWith(TEXT("test")));
+		TestFalse(" View.StartsWith", View.StartsWith('t'));
+	}
+
+	// Test a valid view with the correct text but with different case
+	{
+		FStringView View(TEXT("String to test"));
+		TestTrue(" View.StartsWith", View.StartsWith(TEXT("sTrInG")));
+
+		// Searching by char is case sensitive to keep compatibility with FString
+		TestFalse(" View.StartsWith", View.StartsWith('s'));
+	}
+
+	return true;
+}
+
+IMPLEMENT_SIMPLE_AUTOMATION_TEST(FStringViewTestEndsWith, TEST_NAME_ROOT ".EndsWith", TestFlags)
+bool FStringViewTestEndsWith::RunTest(const FString& Parameters)
+{
+	// Test an empty view
+	{
+		FStringView View(TEXT(""));
+		TestTrue(" View.EndsWith", View.EndsWith(TEXT("")));
+		TestFalse(" View.EndsWith", View.EndsWith(TEXT("Text")));
+		TestFalse(" View.EndsWith", View.EndsWith('A'));
+	}
+
+	// Test a valid view with the correct text
+	{
+		FStringView View(TEXT("String to test"));
+		TestTrue(" View.EndsWith", View.EndsWith(TEXT("test")));
+		TestTrue(" View.EndsWith", View.EndsWith('t'));
+	}
+
+	// Test a valid view with incorrect text
+	{
+		FStringView View(TEXT("String to test"));
+		TestFalse(" View.EndsWith", View.EndsWith(TEXT("String")));
+		TestFalse(" View.EndsWith", View.EndsWith('S'));
+	}
+
+	// Test a valid view with the correct text but with different case
+	{
+		FStringView View(TEXT("String to test"));
+		TestTrue(" View.EndsWith", View.EndsWith(TEXT("TeST")));
+
+		// Searching by char is case sensitive to keep compatibility with FString
+		TestFalse(" View.EndsWith", View.EndsWith('T')); 
+	}
+
+	return true;
+}
+
+IMPLEMENT_SIMPLE_AUTOMATION_TEST(FStringViewTestSubStr, TEST_NAME_ROOT ".SubStr", TestFlags)
+bool FStringViewTestSubStr::RunTest(const FString& Parameters)
+{
+	{
+		FStringView EmptyView(TEXT(""));
+		FStringView EmptyResult = EmptyView.SubStr(0, 10);
+		TestTrue(TEXT("FStringView::SubStr(0)"), EmptyResult.IsEmpty());
+
+		// The following line is commented out as it would fail an assert and currently we cannot test for this in unit tests 
+		// FStringView OutofBoundsResult = EmptyView.SubStr(1000, 10000); 
+		FStringView OutofBoundsResult = EmptyView.SubStr(0, 10000);
+		TestTrue(TEXT("FStringView::SubStr(1)"), OutofBoundsResult.IsEmpty());
+	}
+
+	{
+		FStringView View(TEXT("A test string|"));
+		FStringView Word0 = View.SubStr(0, 1);
+		FStringView Word1 = View.SubStr(2, 4);
+		FStringView Word2 = View.SubStr(7, 6);
+		FStringView NullTerminatorResult = View.SubStr(14, 1024);	// We can create a substr that starts at the end of the 
+																	// string since the null terminator is still valid
+		FStringView OutofBoundsResult = View.SubStr(0, 1024);
+
+		TestTrue(TEXT("FStringView::SubStr(2)"), FCString::Strncmp(Word0.Data(), TEXT("A"), Word0.Len()) == 0);
+		TestTrue(TEXT("FStringView::SubStr(3)"), FCString::Strncmp(Word1.Data(), TEXT("test"), Word1.Len()) == 0);
+		TestTrue(TEXT("FStringView::SubStr(4)"), FCString::Strncmp(Word2.Data(), TEXT("string"), Word2.Len()) == 0);
+		TestTrue(TEXT("FStringView::SubStr(5)"), NullTerminatorResult.IsEmpty());
+		TestTrue(TEXT("FStringView::SubStr(6)"), View == OutofBoundsResult);
+	}
+
+	return true;
+}
+
+IMPLEMENT_SIMPLE_AUTOMATION_TEST(FStringViewTestLeft, TEST_NAME_ROOT ".Left", TestFlags)
+bool FStringViewTestLeft::RunTest(const FString& Parameters)
+{
+	{
+		FStringView EmptyView(TEXT(""));
+		FStringView EmptyResult = EmptyView.Left(0);
+		TestTrue("FStringView::Left", EmptyResult.IsEmpty());
+
+		FStringView OutofBoundsResult = EmptyView.Left(1024);
+		TestTrue("FStringView::Left", OutofBoundsResult.IsEmpty());
+	}
+	
+	{
+		FStringView View(TEXT("A test string"));
+		FStringView Result = View.Left(8);
+
+		TestTrue(TEXT("FStringView::Left"), FCString::Strncmp(Result.Data(), TEXT("A test s"), Result.Len()) == 0);
+
+		FStringView OutofBoundsResult = View.Left(1024);
+		TestTrue(TEXT("FStringView::Left"), FCString::Strncmp(OutofBoundsResult.Data(), TEXT("A test string"), OutofBoundsResult.Len()) == 0);
+	}
+
+	return true;
+}
+
+IMPLEMENT_SIMPLE_AUTOMATION_TEST(FStringViewTestLeftChop, TEST_NAME_ROOT ".LeftChop", TestFlags)
+bool FStringViewTestLeftChop::RunTest(const FString& Parameters)
+{
+	{
+		FStringView EmptyView(TEXT(""));
+		FStringView EmptyResult = EmptyView.LeftChop(0);
+		TestTrue("FStringView::LeftChop", EmptyResult.IsEmpty());
+
+		FStringView OutofBoundsResult = EmptyView.LeftChop(1024);
+		TestTrue("FStringView::LeftChop", OutofBoundsResult.IsEmpty());
+	}
+
+	{
+		FStringView View(TEXT("A test string"));
+		FStringView Result = View.LeftChop(5);
+
+		TestTrue(TEXT("FStringView::LeftChop"), FCString::Strncmp(Result.Data(), TEXT("A test s"), Result.Len()) == 0);
+
+		FStringView OutofBoundsResult = View.LeftChop(1024);
+		TestTrue(TEXT("FStringView::LeftChop"), FCString::Strncmp(OutofBoundsResult.Data(), TEXT("A test string"), OutofBoundsResult.Len()) == 0);
+	}
+
+	return true;
+}
+
+IMPLEMENT_SIMPLE_AUTOMATION_TEST(FStringViewTestRight, TEST_NAME_ROOT ".Right", TestFlags)
+bool FStringViewTestRight::RunTest(const FString& Parameters)
+{
+	{
+		FStringView EmptyView(TEXT(""));
+		FStringView EmptyResult = EmptyView.Right(0);
+		TestTrue("FStringView::Right", EmptyResult.IsEmpty());
+
+		FStringView OutofBoundsResult = EmptyView.Right(1024);
+		TestTrue("FStringView::Right", OutofBoundsResult.IsEmpty());
+	}
+
+	{
+		FStringView View(TEXT("A test string"));
+		FStringView Result = View.Right(8);
+
+		TestTrue(TEXT("FStringView::Right"), FCString::Strncmp(Result.Data(), TEXT("t string"), Result.Len()) == 0);
+
+		FStringView OutofBoundsResult = View.Right(1024);
+		TestTrue(TEXT("FStringView::Right"), FCString::Strncmp(OutofBoundsResult.Data(), TEXT("A test string"), OutofBoundsResult.Len()) == 0);
+	}
+
+	return true;
+}
+
+IMPLEMENT_SIMPLE_AUTOMATION_TEST(FStringViewTestRightChop, TEST_NAME_ROOT ".RightChop", TestFlags)
+bool FStringViewTestRightChop::RunTest(const FString& Parameters)
+{
+	{
+		FStringView EmptyView(TEXT(""));
+		FStringView EmptyResult = EmptyView.RightChop(0);
+		TestTrue("FStringView::RightChop", EmptyResult.IsEmpty());
+
+		FStringView OutofBoundsResult = EmptyView.RightChop(1024);
+		TestTrue("FStringView::RightChop", OutofBoundsResult.IsEmpty());
+	}
+
+	{
+		FStringView View(TEXT("A test string"));
+		FStringView Result = View.RightChop(3);
+
+		TestTrue(TEXT("FStringView::RightChop"), FCString::Strncmp(Result.Data(), TEXT("est string"), Result.Len()) == 0);
+
+		FStringView OutofBoundsResult = View.RightChop(1024);
+		TestTrue(TEXT("FStringView::RightChop"), FCString::Strncmp(OutofBoundsResult.Data(), TEXT("A test string"), OutofBoundsResult.Len()) == 0);
+	}
+
+	return true;
+}
+
+IMPLEMENT_SIMPLE_AUTOMATION_TEST(FStringViewTestMid, TEST_NAME_ROOT ".Mid", TestFlags)
+bool FStringViewTestMid::RunTest(const FString& Parameters)
+{
+	{
+		FStringView EmptyView(TEXT(""));
+		FStringView EmptyResult = EmptyView.Mid(0, 10);
+		TestTrue(TEXT("FStringView::Mid(0)"), EmptyResult.IsEmpty());
+
+		// The following line is commented out as it would fail an assert and currently we cannot test for this in unit tests 
+		// FStringView OutofBoundsResult = EmptyView.Mid(1000, 10000); 
+		FStringView OutofBoundsResult = EmptyView.Mid(0, 10000);
+		TestTrue(TEXT("FStringView::Mid(1)"), OutofBoundsResult.IsEmpty());
+	}
+
+	{
+		FStringView View(TEXT("A test string|"));
+		FStringView Word0 = View.Mid(0, 1);
+		FStringView Word1 = View.Mid(2, 4);
+		FStringView Word2 = View.Mid(7, 6);
+		FStringView NullTerminatorResult = View.Mid(14, 1024);	// We can call Mid with a position that starts at the end of the 
+																// string since the null terminator is still valid
+		FStringView OutofBoundsResult = View.Mid(0, 1024);
+
+		TestTrue(TEXT("FStringView::Mid(2)"), FCString::Strncmp(Word0.Data(), TEXT("A"), Word0.Len()) == 0);
+		TestTrue(TEXT("FStringView::Mid(3)"), FCString::Strncmp(Word1.Data(), TEXT("test"), Word1.Len()) == 0);
+		TestTrue(TEXT("FStringView::Mid(4)"), FCString::Strncmp(Word2.Data(), TEXT("string"), Word2.Len()) == 0);
+		TestTrue(TEXT("FStringView::Mid(5)"), NullTerminatorResult.IsEmpty());
+		TestTrue(TEXT("FStringView::Mid(6)"), View == OutofBoundsResult);
+	}
+
+	return true;
+}
+
+IMPLEMENT_SIMPLE_AUTOMATION_TEST(FStringViewTestFindChar, TEST_NAME_ROOT ".FindChar", TestFlags)
+bool FStringViewTestFindChar::RunTest(const FString& Parameters)
+{
+	FStringView EmptyView = TEXT("");
+	FStringView View = TEXT("aBce Fga");
+
+	{
+		int32 Index = INDEX_NONE;
+		TestFalse(TEXT("FStringView::FindChar-Return(0)"), EmptyView.FindChar('a', Index));
+		TestEqual(TEXT("FStringView::FindChar-Index(0)"), Index, INDEX_NONE);
+	}
+
+	{
+		int32 Index = INDEX_NONE;
+		TestTrue(TEXT("FStringView::FindChar-Return(1)"), View.FindChar('a', Index));
+		TestEqual(TEXT("FStringView::FindChar-Index(1)"), Index, 0);
+	}
+
+	{
+		int32 Index = INDEX_NONE;
+		TestTrue(TEXT("FStringView::FindChar-Return(2)"), View.FindChar('F', Index));
+		TestEqual(TEXT("FStringView::FindChar-Index(2)"), Index, 5);
+	}
+
+	{
+		int32 Index = INDEX_NONE;
+		TestFalse(TEXT("FStringView::FindChar-Return(3)"), View.FindChar('A', Index));
+		TestEqual(TEXT("FStringView::FindChar-Index(3)"), Index, INDEX_NONE);
+	}
+
+	{
+		int32 Index = INDEX_NONE;
+		TestFalse(TEXT("FStringView::FindChar-Return(4)"), View.FindChar('d', Index));
+		TestEqual(TEXT("FStringView::FindChar-Index(4)"), Index, INDEX_NONE);
+	}
+
+	{
+		int32 Index = INDEX_NONE;
+		TestTrue(TEXT("FStringView::FindChar-Return(5)"), View.FindChar(' ', Index));
+		TestEqual(TEXT("FStringView::FindChar-Index(5)"), Index, 4);
+	}
+
+	return true; 
+}
+
+IMPLEMENT_SIMPLE_AUTOMATION_TEST(FStringViewTestFindLastChar, TEST_NAME_ROOT ".FindLastChar", TestFlags)
+bool FStringViewTestFindLastChar::RunTest(const FString& Parameters)
+{
+	FStringView EmptyView = TEXT("");
+	FStringView View = TEXT("aBce Fga");
+
+	{
+		int32 Index = INDEX_NONE;
+		TestFalse(TEXT("FStringView::FindChar-Return(0)"), EmptyView.FindLastChar('a', Index));
+		TestEqual(TEXT("FStringView::FindChar-Index(0)"), Index, INDEX_NONE);
+	}
+
+	{
+		int32 Index = INDEX_NONE;
+		TestTrue(TEXT("FStringView::FindChar-Return(1)"), View.FindLastChar('a', Index));
+		TestEqual(TEXT("FStringView::FindChar-Index(1)"), Index, 7);
+	}
+
+	{
+		int32 Index = INDEX_NONE;
+		TestTrue(TEXT("FStringView::FindChar-Return(2)"), View.FindLastChar('B', Index));
+		TestEqual(TEXT("FStringView::FindChar-Index(2)"), Index, 1);
+	}
+
+	{
+		int32 Index = INDEX_NONE;
+		TestFalse(TEXT("FStringView::FindChar-Return(3)"), View.FindLastChar('A', Index));
+		TestEqual(TEXT("FStringView::FindChar-Index(3)"), Index, INDEX_NONE);
+	}
+
+	{
+		int32 Index = INDEX_NONE;
+		TestFalse(TEXT("FStringView::FindChar-Return(4)"), View.FindLastChar('d', Index));
+		TestEqual(TEXT("FStringView::FindChar-Index(4)"), Index, INDEX_NONE);
+	}
+
+	{
+		int32 Index = INDEX_NONE;
+		TestTrue(TEXT("FStringView::FindChar-Return(5)"), View.FindLastChar(' ', Index));
+		TestEqual(TEXT("FStringView::FindChar-Index(5)"), Index, 4);
+	}
+
+	return true; 
+}
+
+IMPLEMENT_SIMPLE_AUTOMATION_TEST(FStringViewTestFindLastCharByPredicate, TEST_NAME_ROOT ".FindLastCharByPredicate", TestFlags)
+bool FStringViewTestFindLastCharByPredicate::RunTest(const FString& Parameters)
+{
+	auto AcceptMoneyCharacters = [](TCHAR InputChar) { return InputChar == TEXT('£') || InputChar == TEXT('$') || InputChar == TEXT('€'); };
+
+	FStringView ViewDollars			= TEXT("The currency we have is in $");
+	FStringView ViewPounds			= TEXT("Another currency £ could be found?");
+	FStringView ViewMixedCurrency	= TEXT("We have both £,$ and € to spend");
+	FStringView ViewNoCurreny		= TEXT("We are a day late and a dollar short");
+	FStringView ViewEmpty			= TEXT("");
+
+	{
+		FStringView::SizeType Result = ViewDollars.FindLastCharByPredicate(AcceptMoneyCharacters);
+		TestTrue(TEXT("FStringView::FindLastCharByPredicate(0)"), Result == 27);
+
+		FStringView::SizeType ResultWithLimits = ViewDollars.FindLastCharByPredicate(AcceptMoneyCharacters, 27);
+		TestTrue(TEXT("FStringView::FindLastCharByPredicate(1)"), ResultWithLimits == INDEX_NONE);
+	}
+
+	{
+		FStringView::SizeType Result = ViewPounds.FindLastCharByPredicate(AcceptMoneyCharacters);
+		TestTrue(TEXT("FStringView::FindLastCharByPredicate(2)"), Result == 17);
+
+		FStringView::SizeType ResultWithLimits = ViewPounds.FindLastCharByPredicate(AcceptMoneyCharacters, 17);
+		TestTrue(TEXT("FStringView::FindLastCharByPredicate(3)"), ResultWithLimits == INDEX_NONE);
+	}
+
+	{
+		FStringView::SizeType EuroPosition = ViewMixedCurrency.FindLastCharByPredicate(AcceptMoneyCharacters);
+		TestTrue(TEXT("FStringView::FindLastCharByPredicate(4)"), EuroPosition == 21);
+
+		FStringView::SizeType DollarPosition = ViewMixedCurrency.FindLastCharByPredicate(AcceptMoneyCharacters, 21);
+		TestTrue(TEXT("FStringView::FindLastCharByPredicate(5)"), DollarPosition == 15);
+
+		FStringView::SizeType PoundPosition = ViewMixedCurrency.FindLastCharByPredicate(AcceptMoneyCharacters, 15);
+		TestTrue(TEXT("FStringView::FindLastCharByPredicate(6)"), PoundPosition == 13);
+
+		FStringView::SizeType Result = ViewMixedCurrency.FindLastCharByPredicate(AcceptMoneyCharacters, 13);
+		TestTrue(TEXT("FStringView::FindLastCharByPredicate(7)"), Result == INDEX_NONE);
+	}
+
+	{
+		FStringView::SizeType Result = ViewNoCurreny.FindLastCharByPredicate(AcceptMoneyCharacters);
+		TestTrue(TEXT("FStringView::FindLastCharByPredicate(8)"), Result == INDEX_NONE);
+	}
+	
+	{
+		FStringView::SizeType Result = ViewEmpty.FindLastCharByPredicate(AcceptMoneyCharacters);
+		TestTrue(TEXT("FStringView::FindLastCharByPredicate(9)"), Result == INDEX_NONE);
+	}
+
+	return true;
+}
+
+#endif //WITH_DEV_AUTOMATION_TESTS
diff --git a/Engine/Source/Runtime/Core/Public/Misc/StringView.h b/Engine/Source/Runtime/Core/Public/Misc/StringView.h
--- a/Engine/Source/Runtime/Core/Public/Misc/StringView.h
+++ b/Engine/Source/Runtime/Core/Public/Misc/StringView.h
@@ -60,12 +60,46 @@ public:
 	{
 	}
 
+	template<	typename CharType,
+				typename = typename TEnableIf<TIsCharType<CharType>::Value>::Type>
+	TStringViewImpl(const CharType * InData)
+	{
+		static_assert(TIsCharType<CharType>::Value, "Trying to assign an incompatible string type to the TStringViewImpl!");
+	}
+
 	inline TStringViewImpl(const C* InData, SizeType InSize)
 	:	DataPtr(InData)
 	,	Size(InSize)
 	{
 	}
 
+	template<	typename CharType,
+				typename = typename TEnableIf<TIsCharType<CharType>::Value>::Type>
+	TStringViewImpl(const CharType * InData, SizeType InSize)
+	{
+		static_assert(TIsCharType<CharType>::Value, "Trying to assign an incompatible string type to the TStringViewImpl!");
+	}
+
+	// Constructor for FString where the StringView is of the correct type
+	template<	typename U = C,
+				typename = typename TEnableIf< TIsSame<U, FString::ElementType>::Value >::Type >
+	TStringViewImpl(const FString& InString)
+	:	DataPtr(*InString)
+	,	Size(InString.Len())
+	{
+
+	}
+
+	// Constructor for FString where the StringView is of the incorrect type
+	// In this case we give a clear static assert as to why the code will not compile.
+	template<	typename U = C,
+				typename = typename TEnableIf< !TIsSame<U, FString::ElementType>::Value >::Type,
+				typename = U >
+	TStringViewImpl(const FString & InString)
+	{
+		static_assert(!TIsSame<U, FString::ElementType>::Value, "Can only use FString with a TStringViewImpl of the same type!");
+	}
+
 	inline const C& operator[](SizeType Pos) const	{ return DataPtr[Pos]; }
 
 	inline const C* Data() const					{ return DataPtr; }
@@ -82,127 +116,37 @@ public:
 
 	// Operations
 
-	inline SIZE_T			CopyString(C* Dest, SizeType CharCount, SizeType Position = 0)	{ memcpy(Dest, DataPtr + Position, FMath::Min(Size - Position, CharCount)); }
-	inline TStringViewImpl& SubStr(SizeType Position, SizeType CharCount)					{ return TStringViewImpl(DataPtr + Position, FMath::Min(Size - Position, CharCount)); }
+	inline SIZE_T			CopyString(C* Dest, SizeType CharCount, SizeType Position = 0) const	{ memcpy(Dest, DataPtr + Position, FMath::Min(Size - Position, CharCount)); }
+	inline TStringViewImpl	SubStr(SizeType Position, SizeType CharCount) const						{ check(Position <= Len()); return TStringViewImpl(DataPtr + Position, FMath::Min(Size - Position, CharCount)); }
 
+	// Maintain compatibility with FString 
+	inline TStringViewImpl	Left(SizeType CharCount) const											{ return TStringViewImpl(Data(), FMath::Clamp(CharCount, 0, Len())); }
+	inline TStringViewImpl  LeftChop(SizeType CharCount) const										{ return TStringViewImpl(Data(), FMath::Clamp(Len() - CharCount, 0, Len())); }
+	inline TStringViewImpl  Right(SizeType CharCount) const											{ return TStringViewImpl(Data() + Len() - FMath::Clamp(CharCount, 0, Len())); }
+	inline TStringViewImpl	RightChop(SizeType CharCount) const										{ return TStringViewImpl(Data() + Len() - FMath::Clamp(Len() - CharCount, 0, Len())); }
+	inline TStringViewImpl  Mid(SizeType Position, SizeType CharCount) const						{ return SubStr(Position, CharCount); /*This is just a wrapper around SubStr to keep compatibility with the FString interface*/} 
+	
 	// Comparison
+	inline bool				Equals(const TStringViewImpl Other, ESearchCase::Type SearchCase) const	{ return Compare(Other, SearchCase) == 0; }
 
-	// Compare
-
-	inline bool				StartsWith(const TStringViewImpl& Prefix) const				{ return 0 == TCString<C>::Strnicmp(Prefix.Data(), this->Data(), Prefix.Len()); }
-	inline bool				StartsWith(C Prefix) const									{ return Size >= 1 && DataPtr[0] == Prefix; }
-
-	// EndsWith
-
-	inline bool				EndsWith(const TStringViewImpl& Suffix) const 
-	{ 
-		const SIZE_T SuffixLen = Suffix.Len();
-		if (SuffixLen > this->Len())
-		{
-			return false;
-		}
-
-		return 0 == TCString<C>::Strnicmp(Suffix.Data(), this->Data() + this->Len() - SuffixLen, SuffixLen);
-	}
+	inline int32			Compare(const TStringViewImpl Other, ESearchCase::Type SearchCase) const;
 
-	inline bool FindChar(C InChar, int32& OutIndex) const
-	{
-		for (SizeType i = 0; i < Size; ++i)
-		{
-			if (DataPtr[i] == InChar)
-			{
-				OutIndex = (int32) i;
-				return true;
-			}
-		}
-
-		return false;
-	}
+	inline bool				StartsWith(C Prefix) const												{ return Size >= 1 && DataPtr[0] == Prefix; }
+	inline bool				StartsWith(const TStringViewImpl& Prefix) const							{ return 0 == TCString<C>::Strnicmp(Prefix.Data(), this->Data(), Prefix.Len()); }
+	
+	inline bool				EndsWith(C Suffix) const												{ return Size >= 1 && DataPtr[Size-1] == Suffix; }
+	inline bool				EndsWith(const TStringViewImpl& Suffix) const;
 
-	inline bool FindLastChar(C InChar, int32& OutIndex) const
-	{
-		if (Size == 0)
-		{
-			return false;
-		}
-
-		for (SizeType i = Size - 1; i >= 0; --i)
-		{
-			if (DataPtr[i] == InChar)
-			{
-				OutIndex = (int32)i;
-				return true;
-			}
-		}
-
-		return false;
-	}
+	// Searching/Finding
+	inline bool FindChar(C InChar, int32& OutIndex) const;
+	inline bool FindLastChar(C InChar, int32& OutIndex) const;
 
-	// FindLastOf
-	// FindFirstNotOf
-	// FindLastNotOf
+	template <typename Predicate> inline SizeType FindLastCharByPredicate(Predicate Pred, SizeType Count) const;
+	template <typename Predicate> inline SizeType FindLastCharByPredicate(Predicate Pred) const		{ return FindLastCharByPredicate(Pred, Len()); }
 
 	inline const C* begin() const					{ return DataPtr; }
 	inline const C* end() const						{ return DataPtr + Size; }
 
-protected:
-	const C*	DataPtr;
-	SizeType	Size;
-};
-
-// Case-insensitive comparison operators
-
-template<typename C>
-inline bool operator==(const TStringViewImpl<C>& lhs, const TStringViewImpl<C>& rhs)
-{
-	if (rhs.Len() != lhs.Len())
-	{
-		return false;
-	}
-
-	return lhs.StartsWith(rhs);
-}
-
-template<typename C>
-inline bool operator!=(const TStringViewImpl<C>& lhs, const TStringViewImpl<C>& rhs)
-{
-	return !operator==(lhs, rhs);
-}
-
-//////////////////////////////////////////////////////////////////////////
-
-class FAnsiStringView : public TStringViewImpl<ANSICHAR>
-{
-public:
-	FAnsiStringView(const ANSICHAR* InString, SizeType InStrlen)
-	:	TStringViewImpl<ANSICHAR>(InString, InStrlen)
-	{
-	}
-
-	FAnsiStringView(const ANSICHAR* InString)
-	:	TStringViewImpl<ANSICHAR>(InString)
-	{
-	}
-};
-
-class FStringView : public TStringViewImpl<TCHAR>
-{
-public:
-	FStringView(const FString& InString)
-	:	TStringViewImpl<TCHAR>(*InString, InString.Len())
-	{
-	}
-
-	FStringView(const TCHAR* InString)
-	:	TStringViewImpl<TCHAR>(InString)
-	{
-	}
-
-	FStringView(const TCHAR* InString, SizeType InStrlen)
-	:	TStringViewImpl<TCHAR>(InString, InStrlen)
-	{
-	}
-
 	/** Convert to a dynamic string
 
 		Remember that a string view is not necessarily
@@ -219,5 +163,19 @@ public:
 		the heap is always more expensive even if it
 		seems seductively convenient.
 	  */
-	CORE_API FString ToString() const;
+	FString ToString() const
+	{
+		return FString(Len(), Data());
+	}
+protected:
+	const C*	DataPtr;
+	SizeType	Size;
 };
+
+//////////////////////////////////////////////////////////////////////////
+
+using FAnsiStringView = TStringViewImpl<ANSICHAR>;
+using FWideStringView = TStringViewImpl<WIDECHAR>;
+using FStringView = TStringViewImpl<TCHAR>;
+
+#include "StringView.inl"
diff --git a/Engine/Source/Runtime/Core/Public/Misc/StringView.inl b/Engine/Source/Runtime/Core/Public/Misc/StringView.inl
new file mode 100644
--- /dev/null
+++ b/Engine/Source/Runtime/Core/Public/Misc/StringView.inl
@@ -0,0 +1,109 @@
+// Copyright 1998-2019 Epic Games, Inc. All Rights Reserved.
+
+// Case-insensitive comparison operators
+template<typename C>
+inline bool operator==(const TStringViewImpl<C> Lhs, typename TIdentity<TStringViewImpl<C>>::Type Rhs)
+{
+	if (Rhs.Len() != Lhs.Len())
+	{
+		return false;
+	}
+
+	return FPlatformString::Strnicmp(Lhs.Data(), Rhs.Data(), Lhs.Len()) == 0;
+}
+
+template<typename C>
+inline bool operator!=(const TStringViewImpl<C> Lhs, typename TIdentity<TStringViewImpl<C>>::Type Rhs)
+{
+	return !operator==(Lhs, Rhs);
+}
+
+template<typename C>
+inline int32 TStringViewImpl<C>::Compare(const TStringViewImpl Other, ESearchCase::Type SearchCase) const
+{
+	const SizeType ShortestLength = FMath::Min(Len(), Other.Len());
+
+	int Result;
+	if (SearchCase == ESearchCase::CaseSensitive)
+	{
+		Result = FPlatformString::Strncmp(Data(), Other.Data(), ShortestLength);
+	}
+	else
+	{
+		Result = FPlatformString::Strnicmp(Data(), Other.Data(), ShortestLength);
+	}
+
+	if (Result == 0 && Len() != Other.Len())
+	{
+		return Len() > Other.Len() ? 1 : -1;
+	}
+	else
+	{
+		return Result;
+	}
+}
+
+template<typename C>
+inline bool TStringViewImpl<C>::EndsWith(const TStringViewImpl& Suffix) const
+{
+	const SizeType SuffixLen = Suffix.Len();
+	if (SuffixLen > this->Len())
+	{
+		return false;
+	}
+
+	return 0 == TCString<C>::Strnicmp(Suffix.Data(), this->Data() + this->Len() - SuffixLen, SuffixLen);
+}
+
+template<typename C>
+inline bool TStringViewImpl<C>::FindChar(C InChar, int32& OutIndex) const
+{
+	for (SizeType i = 0; i < Size; ++i)
+	{
+		if (DataPtr[i] == InChar)
+		{
+			OutIndex = (int32)i;
+			return true;
+		}
+	}
+
+	return false;
+}
+
+template<typename C>
+inline bool TStringViewImpl<C>::FindLastChar(C InChar, int32& OutIndex) const
+{
+	if (Size == 0)
+	{
+		return false;
+	}
+
+	for (SizeType i = Size - 1; i >= 0; --i)
+	{
+		if (DataPtr[i] == InChar)
+		{
+			OutIndex = (int32)i;
+			return true;
+		}
+	}
+
+	return false;
+}
+
+template<typename C>
+template<typename Predicate>
+inline typename TStringViewImpl<C>::SizeType TStringViewImpl<C>::FindLastCharByPredicate(Predicate Pred, TStringViewImpl<C>::SizeType Count) const
+{
+	check(Count >= 0 && Count <= Len());
+
+	for (const C* RESTRICT Start = Data(), *RESTRICT Data = Start + Count; Data != Start; )
+	{
+		--Data;
+		if (Pred(*Data))
+		{
+			return static_cast<SizeType>(Data - Start);
+		}
+	}
+
+	return INDEX_NONE;
+}
