Change unicode-literal symbols into \u code points to fix compile error.

diff --git a/Engine/Source/Runtime/Core/Private/Tests/Misc/StringViewTest.cpp b/Engine/Source/Runtime/Core/Private/Tests/Misc/StringViewTest.cpp
--- a/Engine/Source/Runtime/Core/Private/Tests/Misc/StringViewTest.cpp
+++ b/Engine/Source/Runtime/Core/Private/Tests/Misc/StringViewTest.cpp
@@ -627,7 +627,10 @@ bool FStringViewTestFindLastChar::RunTest(const FString& Parameters)
 IMPLEMENT_SIMPLE_AUTOMATION_TEST(FStringViewTestFindLastCharByPredicate, TEST_NAME_ROOT ".FindLastCharByPredicate", TestFlags)
 bool FStringViewTestFindLastCharByPredicate::RunTest(const FString& Parameters)
 {
-	auto AcceptMoneyCharacters = [](TCHAR InputChar) { return InputChar == TEXT('£') || InputChar == TEXT('$') || InputChar == TEXT('€'); };
+	auto AcceptMoneyCharacters = [](TCHAR InputChar) { return InputChar == TEXT('\u00a3') // pound sign
+															|| InputChar == TEXT('$')
+															|| InputChar == TEXT('\u20ac'); // euro sign
+	};
 
 	FStringView ViewDollars			= TEXT("The currency we have is in $");
 	FStringView ViewPounds			= TEXT("Another currency £ could be found?");
