Added initial capacity for container allocators so container can start with inline allocator capacity

Avoids redundant non-inlined grow calls to inline allocators. TArray and TBitArray make use of this. TBitArray's default allocator is inline so this will improve performance of almost all TBitArray constructors.

Also optimized TBitArray constructor. This should improve performance of TMap / TSet constructors since they use a TSparseArray that internally uses a TBitArray,

diff --git a/Engine/Source/Developer/ShaderCompilerCommon/Private/HlslUtils.h b/Engine/Source/Developer/ShaderCompilerCommon/Private/HlslUtils.h
--- a/Engine/Source/Developer/ShaderCompilerCommon/Private/HlslUtils.h
+++ b/Engine/Source/Developer/ShaderCompilerCommon/Private/HlslUtils.h
@@ -172,6 +172,11 @@ namespace CrossCompiler
 				return !!Data;
 			}
 
+			SizeType GetInitialCapacity() const
+			{
+				return 0;
+			}
+
 			FLinearAllocator* LinearAllocator;
 
 		private:
diff --git a/Engine/Source/Runtime/Core/Private/HAL/Allocators/AnsiAllocator.h b/Engine/Source/Runtime/Core/Private/HAL/Allocators/AnsiAllocator.h
--- a/Engine/Source/Runtime/Core/Private/HAL/Allocators/AnsiAllocator.h
+++ b/Engine/Source/Runtime/Core/Private/HAL/Allocators/AnsiAllocator.h
@@ -100,6 +100,10 @@ public:
 			return !!Data;
 		}
 
+		SizeType GetInitialCapacity() const
+		{
+			return 0;
+		}
 
 	private:
 		ForAnyElementType(const ForAnyElementType&);
diff --git a/Engine/Source/Runtime/Core/Public/Containers/Array.h b/Engine/Source/Runtime/Core/Public/Containers/Array.h
--- a/Engine/Source/Runtime/Core/Public/Containers/Array.h
+++ b/Engine/Source/Runtime/Core/Public/Containers/Array.h
@@ -293,7 +293,7 @@ public:
 	 */
 	FORCEINLINE TArray()
 		: ArrayNum(0)
-		, ArrayMax(0)
+		, ArrayMax(AllocatorInstance.GetInitialCapacity())
 	{}
 
 	/**
@@ -459,7 +459,7 @@ private:
 		checkf(ToArray.ArrayNum == FromArray.ArrayNum && ToArray.ArrayMax == FromArray.ArrayMax, TEXT("Data lost when moving to a container with a more constrained size type"));
 
 		FromArray.ArrayNum = 0;
-		FromArray.ArrayMax = 0;
+		FromArray.ArrayMax = FromArray.AllocatorInstance.GetInitialCapacity();
 	}
 
 	/**
@@ -2568,7 +2568,7 @@ private:
 		}
 		else
 		{
-			ArrayMax = 0;
+			ArrayMax = AllocatorInstance.GetInitialCapacity();
 		}
 	}
 
diff --git a/Engine/Source/Runtime/Core/Public/Containers/BitArray.h b/Engine/Source/Runtime/Core/Public/Containers/BitArray.h
--- a/Engine/Source/Runtime/Core/Public/Containers/BitArray.h
+++ b/Engine/Source/Runtime/Core/Public/Containers/BitArray.h
@@ -185,16 +185,45 @@ public:
 	template<typename,typename>
 	friend class TConstDualSetBitIterator;
 
+	TBitArray()
+	:	NumBits(0)
+	,	MaxBits(AllocatorInstance.GetInitialCapacity() * NumBitsPerDWORD)
+	{}
+
 	/**
 	 * Minimal initialization constructor.
 	 * @param Value - The value to initial the bits to.
 	 * @param InNumBits - The initial number of bits in the array.
 	 */
-	explicit TBitArray( const bool Value = false, const int32 InNumBits = 0 )
-	:	NumBits(0)
-	,	MaxBits(0)
+	explicit TBitArray(bool bValue, int32 InNumBits)
+	:	NumBits(InNumBits)
+	,	MaxBits(AllocatorInstance.GetInitialCapacity() * NumBitsPerDWORD)
 	{
-		Init(Value,InNumBits);
+		if (NumBits > 0)
+		{
+			const int32 Words = FMath::DivideAndRoundUp(NumBits, NumBitsPerDWORD);
+
+			if (Words > FMath::DivideAndRoundUp(MaxBits, NumBitsPerDWORD))
+			{
+				AllocatorInstance.ResizeAllocation(0, Words, sizeof(uint32));
+				MaxBits = Words * NumBitsPerDWORD;
+			}
+			
+			if (Words > 8)
+			{
+				FMemory::Memset(GetData(), bValue ? 0xff : 0, Words * sizeof(uint32));
+			}
+			else
+			{
+				uint32 Word = bValue ? ~0u : 0;
+				int32 WordIdx = 0;
+				do
+				{
+					GetData()[WordIdx] = Word;
+				}
+				while (++WordIdx < Words);
+			}
+		}
 	}
 
 	/**
diff --git a/Engine/Source/Runtime/Core/Public/Containers/ContainerAllocationPolicies.h b/Engine/Source/Runtime/Core/Public/Containers/ContainerAllocationPolicies.h
--- a/Engine/Source/Runtime/Core/Public/Containers/ContainerAllocationPolicies.h
+++ b/Engine/Source/Runtime/Core/Public/Containers/ContainerAllocationPolicies.h
@@ -236,6 +236,9 @@ public:
 
 		/** Returns true if the allocator has made any heap allocations */
 		bool HasAllocation() const;
+
+		/** Returns number of pre-allocated elements the container can use before allocating more space */
+		SizeType GetInitialCapacity() const;
 	};
 
 	/**
@@ -332,6 +335,11 @@ public:
 			return !!Data;
 		}
 
+		SizeType GetInitialCapacity() const
+		{
+			return 0;
+		}
+
 	private:
 		ForAnyElementType(const ForAnyElementType&);
 		ForAnyElementType& operator=(const ForAnyElementType&);
@@ -472,6 +480,11 @@ public:
 		{
 			return !!Data;
 		}
+		
+		SizeType GetInitialCapacity() const
+		{
+			return 0;
+		}
 
 	private:
 		ForAnyElementType(const ForAnyElementType&);
@@ -631,6 +644,11 @@ public:
 			return SecondaryData.HasAllocation();
 		}
 
+		SizeType GetInitialCapacity() const
+		{
+			return NumInlineElements;
+		}
+
 	private:
 		ForElementType(const ForElementType&);
 		ForElementType& operator=(const ForElementType&);
@@ -773,6 +791,11 @@ public:
 			return Data != GetInlineElements();
 		}
 
+		SizeType GetInitialCapacity() const
+		{
+			return NumInlineElements;
+		}
+
 	private:
 		ForElementType(const ForElementType&) = delete;
 		ForElementType& operator=(const ForElementType&) = delete;
@@ -875,7 +898,11 @@ public:
 		{
 			return false;
 		}
-
+		
+		SizeType GetInitialCapacity() const
+		{
+			return NumInlineElements;
+		}
 
 	private:
 		ForElementType(const ForElementType&);
diff --git a/Engine/Source/Runtime/Core/Public/Misc/MemStack.h b/Engine/Source/Runtime/Core/Public/Misc/MemStack.h
--- a/Engine/Source/Runtime/Core/Public/Misc/MemStack.h
+++ b/Engine/Source/Runtime/Core/Public/Misc/MemStack.h
@@ -345,6 +345,11 @@ public:
 		{
 			return !!Data;
 		}
+
+		SizeType GetInitialCapacity() const
+		{
+			return 0;
+		}
 			
 	private:
 
diff --git a/Engine/Source/Runtime/Engine/Public/Animation/AnimCompressionTypes.h b/Engine/Source/Runtime/Engine/Public/Animation/AnimCompressionTypes.h
--- a/Engine/Source/Runtime/Engine/Public/Animation/AnimCompressionTypes.h
+++ b/Engine/Source/Runtime/Engine/Public/Animation/AnimCompressionTypes.h
@@ -560,6 +560,11 @@ public:
 			return !!Data;
 		}
 
+		SizeType GetInitialCapacity() const
+		{
+			return 0;
+		}
+
 		void AcceptFileMapping(IMappedFileHandle* InMappedHandle, IMappedFileRegion* InMappedRegion, void *MallocPtr)
 		{
 			check(!MappedHandle && !Data); // we could support stuff like this, but that usually isn't what we want for streamlined loading
