PR #6351: Delete problematic copy constructors on named pipe (Contributed by [REDACTED_PERSON])

diff --git a/Engine/Source/Runtime/Core/Public/GenericPlatform/GenericPlatformNamedPipe.h b/Engine/Source/Runtime/Core/Public/GenericPlatform/GenericPlatformNamedPipe.h
--- a/Engine/Source/Runtime/Core/Public/GenericPlatform/GenericPlatformNamedPipe.h
+++ b/Engine/Source/Runtime/Core/Public/GenericPlatform/GenericPlatformNamedPipe.h
@@ -20,6 +20,9 @@ public:
 	/** Virtual destructor. */
 	virtual ~FGenericPlatformNamedPipe();
 
+	FGenericPlatformNamedPipe(const FGenericPlatformNamedPipe&) = delete;
+	FGenericPlatformNamedPipe& operator=(const FGenericPlatformNamedPipe&) = delete;
+
 public:
 
 	/** Create a named pipe as a server or client, using overlapped IO if bAsync=1. */
diff --git a/Engine/Source/Runtime/Core/Public/Windows/WindowsPlatformNamedPipe.h b/Engine/Source/Runtime/Core/Public/Windows/WindowsPlatformNamedPipe.h
--- a/Engine/Source/Runtime/Core/Public/Windows/WindowsPlatformNamedPipe.h
+++ b/Engine/Source/Runtime/Core/Public/Windows/WindowsPlatformNamedPipe.h
@@ -21,6 +21,9 @@ public:
 	FWindowsPlatformNamedPipe();
 	virtual ~FWindowsPlatformNamedPipe();
 
+	FWindowsPlatformNamedPipe(const FWindowsPlatformNamedPipe&) = delete;
+	FWindowsPlatformNamedPipe& operator=(const FWindowsPlatformNamedPipe&) = delete;
+
 public:
 
 	// FGenericPlatformNamedPipe overrides
