Implement better defined behavior for non-append operations on keyed struct arrays in config files.

diff --git a/Engine/Source/Runtime/Core/Private/Misc/ConfigCacheIni.cpp b/Engine/Source/Runtime/Core/Private/Misc/ConfigCacheIni.cpp
--- a/Engine/Source/Runtime/Core/Private/Misc/ConfigCacheIni.cpp
+++ b/Engine/Source/Runtime/Core/Private/Misc/ConfigCacheIni.cpp
@@ -351,6 +351,21 @@ static void ExtractPropertyValue(const FString& FullStructValue, const FString&
 }
 
 void FConfigSection::HandleAddCommand(FName Key, FString&& Value, bool bAppendValueIfNotArrayOfStructsKeyUsed)
+{
+	if (!HandleArrayOfKeyedStructsCommand(Key, Forward<FString&&>(Value)))
+	{
+		if (bAppendValueIfNotArrayOfStructsKeyUsed)
+		{
+			Add(Key, MoveTemp(Value));
+		}
+		else
+		{
+			AddUnique(Key, MoveTemp(Value));
+		}
+	}
+}
+
+bool FConfigSection::HandleArrayOfKeyedStructsCommand(FName Key, FString&& Value)
 {
 	FString* StructKey = ArrayOfStructKeys.Find(Key);
 	bool bHandledWithKey = false;
@@ -389,17 +404,7 @@ void FConfigSection::HandleAddCommand(FName Key, FString&& Value, bool bAppendVa
 		}
 	}
 
-	if (!bHandledWithKey)
-	{
-		if (bAppendValueIfNotArrayOfStructsKeyUsed)
-		{
-			Add(Key, MoveTemp(Value));
-		}
-		else
-		{
-			AddUnique(Key, MoveTemp(Value));
-		}
-	}
+	return bHandledWithKey;
 }
 
 // Look through the file's per object config ArrayOfStruct keys and see if this section matches
@@ -841,15 +846,19 @@ void FConfigFile::CombineFromBuffer(const FString& Buffer)
 				}
 				else
 				{
-					// Add if not present and replace if present.
-					FConfigValue* ConfigValue = CurrentSection->Find( Start );
-					if( !ConfigValue )
+					// First see if this can be processed as an array of keyed structs command
+					if (!CurrentSection->HandleArrayOfKeyedStructsCommand(Start, MoveTemp(ProcessedValue)))
 					{
-						CurrentSection->Add( Start, MoveTemp(ProcessedValue) );
-					}
-					else
-					{
-						*ConfigValue = FConfigValue(MoveTemp(ProcessedValue));
+						// Add if not present and replace if present.
+						FConfigValue* ConfigValue = CurrentSection->Find(Start);
+						if (!ConfigValue)
+						{
+							CurrentSection->Add(Start, MoveTemp(ProcessedValue));
+						}
+						else
+						{
+							*ConfigValue = FConfigValue(MoveTemp(ProcessedValue));
+						}
 					}
 				}
 
diff --git a/Engine/Source/Runtime/Core/Public/Misc/ConfigCacheIni.h b/Engine/Source/Runtime/Core/Public/Misc/ConfigCacheIni.h
--- a/Engine/Source/Runtime/Core/Public/Misc/ConfigCacheIni.h
+++ b/Engine/Source/Runtime/Core/Public/Misc/ConfigCacheIni.h
@@ -223,6 +223,8 @@ public:
 	// process the '+' and '.' commands, takingf into account ArrayOfStruct unique keys
 	void CORE_API HandleAddCommand(FName Key, FString&& Value, bool bAppendValueIfNotArrayOfStructsKeyUsed);
 
+	bool HandleArrayOfKeyedStructsCommand(FName Key, FString&& Value);
+
 	template<typename Allocator> 
 	void MultiFind(const FName Key, TArray<FConfigValue, Allocator>& OutValues, const bool bMaintainOrder = false) const
 	{
