Added SetNumUninitialized() to TBitArray

diff --git a/Engine/Source/Runtime/Core/Public/Containers/BitArray.h b/Engine/Source/Runtime/Core/Public/Containers/BitArray.h
--- a/Engine/Source/Runtime/Core/Public/Containers/BitArray.h
+++ b/Engine/Source/Runtime/Core/Public/Containers/BitArray.h
@@ -460,6 +460,24 @@ public:
 		}
 	}
 
+	/** Sets number of bits without initializing new bits. */
+	void SetNumUninitialized(int32 InNumBits)
+	{
+		int32 PreviousNumBits = NumBits;
+		NumBits = InNumBits;
+
+		if (InNumBits > MaxBits)
+		{
+			const int32 PreviousNumDWORDs = FMath::DivideAndRoundUp(PreviousNumBits, NumBitsPerDWORD);
+			const uint32 MaxDWORDs = AllocatorInstance.CalculateSlackReserve(
+				FMath::DivideAndRoundUp(InNumBits, NumBitsPerDWORD), sizeof(uint32));
+			
+			AllocatorInstance.ResizeAllocation(PreviousNumDWORDs, MaxDWORDs, sizeof(uint32));	
+
+			MaxBits = MaxDWORDs * NumBitsPerDWORD;
+		}
+	}
+
 	/**
 	 * Sets or unsets a range of bits within the array.
 	 * @param  Index  The index of the first bit to set.
