UnrealInsights:
Fixed cpu profiler trace buffer flushing logic
Fixed issue where scope ids of zero could sometimes be traced out from dynamic scopes
Fixed assert in debug builds

diff --git a/Engine/Source/Developer/TraceServices/Private/Analyzers/CpuProfilerTraceAnalysis.cpp b/Engine/Source/Developer/TraceServices/Private/Analyzers/CpuProfilerTraceAnalysis.cpp
--- a/Engine/Source/Developer/TraceServices/Private/Analyzers/CpuProfilerTraceAnalysis.cpp
+++ b/Engine/Source/Developer/TraceServices/Private/Analyzers/CpuProfilerTraceAnalysis.cpp
@@ -11,6 +11,15 @@ FCpuProfilerAnalyzer::FCpuProfilerAnalyzer(Trace::IAnalysisSession& InSession, T
 
 }
 
+FCpuProfilerAnalyzer::~FCpuProfilerAnalyzer()
+{
+	for (auto& KV : ThreadStatesMap)
+	{
+		FThreadState* ThreadState = KV.Value;
+		delete ThreadState;
+	}
+}
+
 void FCpuProfilerAnalyzer::OnAnalysisBegin(const FOnAnalysisContext& Context)
 {
 	auto& Builder = Context.InterfaceBuilder;
@@ -46,8 +55,8 @@ bool FCpuProfilerAnalyzer::OnEvent(uint16 RouteId, const FOnEventContext& Contex
 	{
 		TotalEventSize += EventData.GetAttachmentSize();
 		uint32 ThreadId = EventData.GetValue<uint32>("ThreadId");
-		TSharedRef<FThreadState> ThreadState = GetThreadState(ThreadId);
-		uint64 LastCycle = ThreadState->LastCycle;
+		FThreadState& ThreadState = GetThreadState(ThreadId);
+		uint64 LastCycle = ThreadState.LastCycle;
 		uint64 BufferSize = EventData.GetAttachmentSize();
 		const uint8* BufferPtr = EventData.GetAttachment();
 		const uint8* BufferEnd = BufferPtr + BufferSize;
@@ -58,7 +67,7 @@ bool FCpuProfilerAnalyzer::OnEvent(uint16 RouteId, const FOnEventContext& Contex
 			LastCycle = ActualCycle;
 			if (DecodedCycle & 1ull)
 			{
-				EventScopeState& ScopeState = ThreadState->ScopeStack.AddDefaulted_GetRef();
+				EventScopeState& ScopeState = ThreadState.ScopeStack.AddDefaulted_GetRef();
 				ScopeState.StartCycle = ActualCycle;
 				uint32 SpecId = FTraceAnalyzerUtils::Decode7bit(BufferPtr);
 				uint32* FindIt = ScopeIdToEventIdMap.Find(SpecId);
@@ -72,13 +81,13 @@ bool FCpuProfilerAnalyzer::OnEvent(uint16 RouteId, const FOnEventContext& Contex
 				}
 				Trace::FTimingProfilerEvent Event;
 				Event.TimerIndex = ScopeState.EventTypeId;
-				ThreadState->Timeline->AppendBeginEvent(Context.SessionContext.TimestampFromCycle(ScopeState.StartCycle), Event);
+				ThreadState.Timeline->AppendBeginEvent(Context.SessionContext.TimestampFromCycle(ScopeState.StartCycle), Event);
 				++TotalScopeCount;
 			}
-			else if (ThreadState->ScopeStack.Num())
+			else if (ThreadState.ScopeStack.Num())
 			{
-				ThreadState->ScopeStack.Pop();
-				ThreadState->Timeline->AppendEndEvent(Context.SessionContext.TimestampFromCycle(ActualCycle));
+				ThreadState.ScopeStack.Pop();
+				ThreadState.Timeline->AppendEndEvent(Context.SessionContext.TimestampFromCycle(ActualCycle));
 			}
 		}
 		check(BufferPtr == BufferEnd);
@@ -89,14 +98,14 @@ bool FCpuProfilerAnalyzer::OnEvent(uint16 RouteId, const FOnEventContext& Contex
 			if (RouteId == RouteId_EndCapture)
 			{
 
-				while (ThreadState->ScopeStack.Num())
+				while (ThreadState.ScopeStack.Num())
 				{
-					ThreadState->ScopeStack.Pop();
-					ThreadState->Timeline->AppendEndEvent(LastTimestamp);
+					ThreadState.ScopeStack.Pop();
+					ThreadState.Timeline->AppendEndEvent(LastTimestamp);
 				}
 			}
 		}
-		ThreadState->LastCycle = LastCycle;
+		ThreadState.LastCycle = LastCycle;
 		BytesPerScope = double(TotalEventSize) / double(TotalScopeCount);
 		break;
 	}
@@ -110,6 +119,7 @@ void FCpuProfilerAnalyzer::DefineScope(uint32 Id, const TCHAR* Name)
 	if (ScopeIdToEventIdMap.Contains(Id))
 	{
 		TimingProfilerProvider.SetTimerName(ScopeIdToEventIdMap[Id], Name);
+		ScopeNameToEventIdMap.Add(Name, Id);
 	}
 	else
 	{
@@ -127,18 +137,15 @@ void FCpuProfilerAnalyzer::DefineScope(uint32 Id, const TCHAR* Name)
 	}
 }
 
-TSharedRef<FCpuProfilerAnalyzer::FThreadState> FCpuProfilerAnalyzer::GetThreadState(uint32 ThreadId)
+FCpuProfilerAnalyzer::FThreadState& FCpuProfilerAnalyzer::GetThreadState(uint32 ThreadId)
 {
-	if (!ThreadStatesMap.Contains(ThreadId))
+	FThreadState* ThreadState = ThreadStatesMap.FindRef(ThreadId);
+	if (!ThreadState)
 	{
-		TSharedRef<FThreadState> ThreadState = MakeShared<FThreadState>();
+		ThreadState = new FThreadState();
 		ThreadState->Timeline = &TimingProfilerProvider.EditCpuThreadTimeline(ThreadId);
 		ThreadStatesMap.Add(ThreadId, ThreadState);
-		return ThreadState;
-	}
-	else
-	{
-		return ThreadStatesMap[ThreadId];
 	}
+	return *ThreadState;
 }
 
diff --git a/Engine/Source/Developer/TraceServices/Private/Analyzers/CpuProfilerTraceAnalysis.h b/Engine/Source/Developer/TraceServices/Private/Analyzers/CpuProfilerTraceAnalysis.h
--- a/Engine/Source/Developer/TraceServices/Private/Analyzers/CpuProfilerTraceAnalysis.h
+++ b/Engine/Source/Developer/TraceServices/Private/Analyzers/CpuProfilerTraceAnalysis.h
@@ -16,6 +16,7 @@ class FCpuProfilerAnalyzer
 {
 public:
 	FCpuProfilerAnalyzer(Trace::IAnalysisSession& Session, Trace::FTimingProfilerProvider& TimingProfilerProvider);
+	~FCpuProfilerAnalyzer();
 	virtual void OnAnalysisBegin(const FOnAnalysisContext& Context) override;
 	virtual bool OnEvent(uint16 RouteId, const FOnEventContext& Context) override;
 
@@ -34,7 +35,7 @@ private:
 	};
 
 	void DefineScope(uint32 Id, const TCHAR* ScopeName);
-	TSharedRef<FThreadState> GetThreadState(uint32 ThreadId);
+	FThreadState& GetThreadState(uint32 ThreadId);
 
 	enum : uint16
 	{
@@ -45,7 +46,7 @@ private:
 
 	Trace::IAnalysisSession& Session;
 	Trace::FTimingProfilerProvider& TimingProfilerProvider;
-	TMap<uint32, TSharedRef<FThreadState>> ThreadStatesMap;
+	TMap<uint32, FThreadState*> ThreadStatesMap;
 	TMap<uint32, uint32> ScopeIdToEventIdMap;
 	TMap<const TCHAR*, uint32> ScopeNameToEventIdMap;
 	uint64 TotalEventSize = 0;
diff --git a/Engine/Source/Runtime/Core/Private/ProfilingDebugging/CpuProfilerTrace.cpp b/Engine/Source/Runtime/Core/Private/ProfilingDebugging/CpuProfilerTrace.cpp
--- a/Engine/Source/Runtime/Core/Private/ProfilingDebugging/CpuProfilerTrace.cpp
+++ b/Engine/Source/Runtime/Core/Private/ProfilingDebugging/CpuProfilerTrace.cpp
@@ -63,6 +63,11 @@ struct FCpuProfilerTraceInternal
 	struct FThreadBuffer
 		: public FTlsAutoCleanup
 	{
+		FThreadBuffer()
+			: DynamicScopeNamesMemory(0)
+		{
+		}
+
 		uint64 LastCycle = 0;
 		uint16 BufferSize = 0;
 		uint8 Buffer[MaxBufferSize];
@@ -108,99 +113,77 @@ void FCpuProfilerTraceInternal::EndCapture(FThreadBuffer* InThreadBuffer)
 	InThreadBuffer->BufferSize = 0;
 }
 
-void FCpuProfilerTrace::OutputBeginEvent(uint32 SpecId)
-{
-	++FCpuProfilerTraceInternal::ThreadDepth;
-	bool bEventEnabled = UE_TRACE_EVENT_IS_ENABLED(CpuProfiler, EventBatch);
-	if ((!bEventEnabled) & (bEventEnabled == FCpuProfilerTraceInternal::bThreadEnabled))
-	{
-		return;
-	}
-	FCpuProfilerTraceInternal::bThreadEnabled = bEventEnabled;
-	FCpuProfilerTraceInternal::FThreadBuffer* ThreadBuffer = FCpuProfilerTraceInternal::ThreadBuffer;
-	if (!ThreadBuffer)
-	{
-		ThreadBuffer = FCpuProfilerTraceInternal::CreateThreadBuffer();
+#define CPUPROFILERTRACE_OUTPUTBEGINEVENT_PROLOGUE() \
+	++FCpuProfilerTraceInternal::ThreadDepth; \
+	bool bEventEnabled = UE_TRACE_EVENT_IS_ENABLED(CpuProfiler, EventBatch); \
+	if ((!bEventEnabled) & (bEventEnabled == FCpuProfilerTraceInternal::bThreadEnabled)) \
+	{ \
+		return; \
+	} \
+	FCpuProfilerTraceInternal::bThreadEnabled = bEventEnabled; \
+	FCpuProfilerTraceInternal::FThreadBuffer* ThreadBuffer = FCpuProfilerTraceInternal::ThreadBuffer; \
+	if (!ThreadBuffer) \
+	{ \
+		ThreadBuffer = FCpuProfilerTraceInternal::CreateThreadBuffer(); \
+	} \
+	if (!bEventEnabled) \
+	{ \
+		FCpuProfilerTraceInternal::EndCapture(ThreadBuffer); \
+		return; \
 	}
-	if (!bEventEnabled)
-	{
-		FCpuProfilerTraceInternal::EndCapture(ThreadBuffer);
-		return;
-	}
-	if (ThreadBuffer->BufferSize >= FCpuProfilerTraceInternal::FullBufferThreshold)
-	{
-		FCpuProfilerTraceInternal::FlushThreadBuffer(ThreadBuffer);
+
+#define CPUPROFILERTRACE_OUTPUTBEGINEVENT_EPILOGUE() \
+	uint64 Cycle = FPlatformTime::Cycles64(); \
+	uint64 CycleDiff = Cycle - ThreadBuffer->LastCycle; \
+	ThreadBuffer->LastCycle = Cycle; \
+	uint8* BufferPtr = ThreadBuffer->Buffer + ThreadBuffer->BufferSize; \
+	FTraceUtils::Encode7bit((CycleDiff << 1) | 1ull, BufferPtr); \
+	FTraceUtils::Encode7bit(SpecId, BufferPtr); \
+	ThreadBuffer->BufferSize = BufferPtr - ThreadBuffer->Buffer; \
+	if (ThreadBuffer->BufferSize >= FCpuProfilerTraceInternal::FullBufferThreshold) \
+	{ \
+		FCpuProfilerTraceInternal::FlushThreadBuffer(ThreadBuffer); \
 	}
-	uint64 Cycle = FPlatformTime::Cycles64();
-	uint64 CycleDiff = Cycle - ThreadBuffer->LastCycle;
-	ThreadBuffer->LastCycle = Cycle;
-	uint8* BufferPtr = ThreadBuffer->Buffer + ThreadBuffer->BufferSize;
-	FTraceUtils::Encode7bit((CycleDiff << 1) | 1ull, BufferPtr);
-	FTraceUtils::Encode7bit(SpecId, BufferPtr);
-	ThreadBuffer->BufferSize = BufferPtr - ThreadBuffer->Buffer;
+
+void FCpuProfilerTrace::OutputBeginEvent(uint32 SpecId)
+{
+	CPUPROFILERTRACE_OUTPUTBEGINEVENT_PROLOGUE();
+	CPUPROFILERTRACE_OUTPUTBEGINEVENT_EPILOGUE();
 }
 
 void FCpuProfilerTrace::OutputBeginDynamicEvent(const ANSICHAR* Name, ECpuProfilerGroup Group)
 {
-	if (UE_TRACE_EVENT_IS_ENABLED(CpuProfiler, EventBatch))
-	{
-		FCpuProfilerTraceInternal::FThreadBuffer* ThreadBuffer = FCpuProfilerTraceInternal::ThreadBuffer;
-		if (!ThreadBuffer)
-		{
-			ThreadBuffer = FCpuProfilerTraceInternal::CreateThreadBuffer();
-		}
-		uint32* FindSpecId = ThreadBuffer->DynamicAnsiScopeNamesMap.Find(Name);
-		if (!FindSpecId)
-		{
-			int32 NameSize = strlen(Name) + 1;
-			ANSICHAR* NameCopy = reinterpret_cast<ANSICHAR*>(ThreadBuffer->DynamicScopeNamesMemory.Alloc(NameSize, alignof(ANSICHAR)));
-			FMemory::Memmove(NameCopy, Name, NameSize);
-			uint32 SpecId = OutputEventType(NameCopy, Group);
-			ThreadBuffer->DynamicAnsiScopeNamesMap.Add(NameCopy, SpecId);
-			OutputBeginEvent(SpecId);
-		}
-		else
-		{
-			OutputBeginEvent(*FindSpecId);
-		}
-		
-	}
-	else
+	CPUPROFILERTRACE_OUTPUTBEGINEVENT_PROLOGUE();
+	uint32 SpecId = ThreadBuffer->DynamicAnsiScopeNamesMap.FindRef(Name);
+	if (!SpecId)
 	{
-		OutputBeginEvent(0);
+		int32 NameSize = strlen(Name) + 1;
+		ANSICHAR* NameCopy = reinterpret_cast<ANSICHAR*>(ThreadBuffer->DynamicScopeNamesMemory.Alloc(NameSize, alignof(ANSICHAR)));
+		FMemory::Memmove(NameCopy, Name, NameSize);
+		SpecId = OutputEventType(NameCopy, Group);
+		ThreadBuffer->DynamicAnsiScopeNamesMap.Add(NameCopy, SpecId);
 	}
+	CPUPROFILERTRACE_OUTPUTBEGINEVENT_EPILOGUE();
 }
 
 void FCpuProfilerTrace::OutputBeginDynamicEvent(const TCHAR* Name, ECpuProfilerGroup Group)
 {
-	if (UE_TRACE_EVENT_IS_ENABLED(CpuProfiler, EventBatch))
-	{
-		FCpuProfilerTraceInternal::FThreadBuffer* ThreadBuffer = FCpuProfilerTraceInternal::ThreadBuffer;
-		if (!ThreadBuffer)
-		{
-			ThreadBuffer = FCpuProfilerTraceInternal::CreateThreadBuffer();
-		}
-		uint32* FindSpecId = ThreadBuffer->DynamicTCharScopeNamesMap.Find(Name);
-		if (!FindSpecId)
-		{
-			int32 NameSize = (FCString::Strlen(Name) + 1) * sizeof(TCHAR);
-			TCHAR* NameCopy = reinterpret_cast<TCHAR*>(ThreadBuffer->DynamicScopeNamesMemory.Alloc(NameSize, alignof(TCHAR)));
-			FMemory::Memmove(NameCopy, Name, NameSize);
-			uint32 SpecId = OutputEventType(NameCopy, Group);
-			ThreadBuffer->DynamicTCharScopeNamesMap.Add(NameCopy, SpecId);
-			OutputBeginEvent(SpecId);
-		}
-		else
-		{
-			OutputBeginEvent(*FindSpecId);
-		}
-	}
-	else
+	CPUPROFILERTRACE_OUTPUTBEGINEVENT_PROLOGUE();
+	uint32 SpecId = ThreadBuffer->DynamicTCharScopeNamesMap.FindRef(Name);
+	if (!SpecId)
 	{
-		OutputBeginEvent(0);
+		int32 NameSize = (FCString::Strlen(Name) + 1) * sizeof(TCHAR);
+		TCHAR* NameCopy = reinterpret_cast<TCHAR*>(ThreadBuffer->DynamicScopeNamesMemory.Alloc(NameSize, alignof(TCHAR)));
+		FMemory::Memmove(NameCopy, Name, NameSize);
+		SpecId = OutputEventType(NameCopy, Group);
+		ThreadBuffer->DynamicTCharScopeNamesMap.Add(NameCopy, SpecId);
 	}
+	CPUPROFILERTRACE_OUTPUTBEGINEVENT_EPILOGUE();
 }
 
+#undef CPUPROFILERTRACE_OUTPUTBEGINEVENT_PROLOGUE
+#undef CPUPROFILERTRACE_OUTPUTBEGINEVENT_EPILOGUE
+
 void FCpuProfilerTrace::OutputEndEvent()
 {
 	--FCpuProfilerTraceInternal::ThreadDepth;
@@ -220,16 +203,16 @@ void FCpuProfilerTrace::OutputEndEvent()
 		FCpuProfilerTraceInternal::EndCapture(ThreadBuffer);
 		return;
 	}
-	if ((FCpuProfilerTraceInternal::ThreadDepth == 0) | (ThreadBuffer->BufferSize >= FCpuProfilerTraceInternal::FullBufferThreshold))
-	{
-		FCpuProfilerTraceInternal::FlushThreadBuffer(ThreadBuffer);
-	}
 	uint64 Cycle = FPlatformTime::Cycles64();
 	uint64 CycleDiff = Cycle - ThreadBuffer->LastCycle;
 	ThreadBuffer->LastCycle = Cycle;
 	uint8* BufferPtr = ThreadBuffer->Buffer + ThreadBuffer->BufferSize;
 	FTraceUtils::Encode7bit(CycleDiff << 1, BufferPtr);
 	ThreadBuffer->BufferSize = BufferPtr - ThreadBuffer->Buffer;
+	if ((FCpuProfilerTraceInternal::ThreadDepth == 0) | (ThreadBuffer->BufferSize >= FCpuProfilerTraceInternal::FullBufferThreshold))
+	{
+		FCpuProfilerTraceInternal::FlushThreadBuffer(ThreadBuffer);
+	}
 }
 
 uint32 FCpuProfilerTraceInternal::GetNextSpecId()
