UnrealInsights:
Refinements to CountersTrace API

diff --git a/Engine/Source/Runtime/Core/Public/ProfilingDebugging/CountersTrace.h b/Engine/Source/Runtime/Core/Public/ProfilingDebugging/CountersTrace.h
--- a/Engine/Source/Runtime/Core/Public/ProfilingDebugging/CountersTrace.h
+++ b/Engine/Source/Runtime/Core/Public/ProfilingDebugging/CountersTrace.h
@@ -31,152 +31,101 @@ struct FCountersTrace
 	CORE_API static void OutputSetValue(uint16 CounterId, int64 Value);
 	CORE_API static void OutputSetValue(uint16 CounterId, double Value);
 
-	class FCounter
+	template<typename ValueType, ETraceCounterType CounterType>
+	class TCounter
 	{
 	public:
-		FCounter()
-			: IntValue(0)
+		TCounter(const TCHAR* InCounterName, ETraceCounterDisplayHint InCounterDisplayHint)
+			: Value(0)
 		{
+			CounterId = OutputInitCounter(InCounterName, CounterType, InCounterDisplayHint);
 		}
 
-		FCounter(const TCHAR* InCounterName, ETraceCounterType InCounterType, ETraceCounterDisplayHint InCounterDisplayHint)
-			: IntValue(0)
-			, CounterType(InCounterType)
+		void Set(ValueType InValue)
 		{
-			CounterId = OutputInitCounter(InCounterName, InCounterType, InCounterDisplayHint);
-		}
-
-		void Init(const TCHAR* InCounterName, ETraceCounterType InCounterType, ETraceCounterDisplayHint InCounterDisplayHint)
-		{
-			if (CounterId == 0)
+			if (Value != InValue)
 			{
-				CounterType = InCounterType;
-				CounterId = OutputInitCounter(InCounterName, InCounterType, InCounterDisplayHint);
+				Value = InValue;
+				OutputSetValue(CounterId, Value);
 			}
 		}
 
-		void Set(int64 Value)
+		void Add(ValueType InValue)
 		{
-			if (CounterType == TraceCounterType_Int)
-			{
-				if (IntValue != Value)
-				{
-					IntValue = Value;
-					OutputSetValue(CounterId, IntValue);
-				}
-			}
-			else
+			if (InValue != 0)
 			{
-				double ValueAsDouble = static_cast<double>(Value);
-				if (FloatValue != ValueAsDouble)
-				{
-					FloatValue = ValueAsDouble;
-					OutputSetValue(CounterId, FloatValue);
-				}
+				Value += InValue;
+				OutputSetValue(CounterId, Value);
 			}
 		}
 
-		void Set(double Value)
+		void Subtract(ValueType InValue)
 		{
-			if (CounterType == TraceCounterType_Int)
+			if (InValue != 0)
 			{
-				int64 ValueAsInt = static_cast<int64>(Value);
-				if (IntValue != ValueAsInt)
-				{
-					IntValue = ValueAsInt;
-					OutputSetValue(CounterId, IntValue);
-				}
-			}
-			else
-			{
-				if (FloatValue != Value)
-				{
-					FloatValue = Value;
-					OutputSetValue(CounterId, FloatValue);
-				}
+				Value -= InValue;
+				OutputSetValue(CounterId, Value);
 			}
 		}
 
-		void Add(int64 Value)
+		void Increment()
 		{
-			if (Value != 0)
-			{
-				if (CounterType == TraceCounterType_Int)
-				{
-					IntValue += Value;
-					OutputSetValue(CounterId, IntValue);
-				}
-				else
-				{
-					FloatValue += static_cast<double>(Value);
-					OutputSetValue(CounterId, FloatValue);
-				}
-			}
+			++Value;
+			OutputSetValue(CounterId, Value);
 		}
-
-		void Add(double Value)
+		
+		void Decrement()
 		{
-			if (Value != 0)
-			{
-				if (CounterType == TraceCounterType_Int)
-				{
-					IntValue += static_cast<int64>(Value);
-					OutputSetValue(CounterId, IntValue);
-				}
-				else
-				{
-					FloatValue += Value;
-					OutputSetValue(CounterId, FloatValue);
-				}
-			}
+			--Value;
+			OutputSetValue(CounterId, Value);
 		}
-
+		
 	private:
-		union
-		{
-			int64 IntValue;
-			double FloatValue;
-		};
+		ValueType Value;
 		uint16 CounterId;
-		ETraceCounterType CounterType;
 	};
 
+	using FCounterInt = TCounter<int64, TraceCounterType_Int>;
+	using FCounterFloat = TCounter<double, TraceCounterType_Float>;
+
 	CORE_API static void Init(const TCHAR* CmdLine);
 };
 
 #define TRACE_COUNTERS_INIT(CmdLine) \
 	FCountersTrace::Init(CmdLine);
 
-#define __TRACE_DECLARE_COUNTER(CounterName, CounterType, CounterDisplayHint) \
-	FCountersTrace::FCounter PREPROCESSOR_JOIN(__GTraceCounter, CounterName)(TEXT(#CounterName), CounterType, CounterDisplayHint);
+#define __TRACE_DECLARE_INLINE_COUNTER(CounterDisplayName, CounterType, CounterDisplayHint) \
+	static FCountersTrace::CounterType PREPROCESSOR_JOIN(__TraceCounter, __LINE__)(CounterDisplayName, CounterDisplayHint);
 
-#define __TRACE_DECLARE_INLINE_COUNTER(CounterName, CounterType, CounterDisplayHint) \
-	static FCountersTrace::FCounter PREPROCESSOR_JOIN(__TraceCounter, __LINE__); \
-	PREPROCESSOR_JOIN(__TraceCounter, __LINE__).Init(TEXT(#CounterName), CounterType, CounterDisplayHint); \
+#define TRACE_INT_VALUE(CounterDisplayName, Value) \
+	__TRACE_DECLARE_INLINE_COUNTER(CounterDisplayName, FCounterInt, TraceCounterDisplayHint_None) \
+	PREPROCESSOR_JOIN(__TraceCounter, __LINE__).Set(Value);
 
-#define TRACE_INT_VALUE(CounterName, Value) \
-	__TRACE_DECLARE_INLINE_COUNTER(CounterName, TraceCounterType_Int, TraceCounterDisplayHint_None) \
-	PREPROCESSOR_JOIN(__TraceCounter, __LINE__).Set(int64(Value));
+#define TRACE_FLOAT_VALUE(CounterDisplayName, Value) \
+	__TRACE_DECLARE_INLINE_COUNTER(CounterDisplayName, FCounterFloat, TraceCounterDisplayHint_None) \
+	PREPROCESSOR_JOIN(__TraceCounter, __LINE__).Set(Value);
 
-#define TRACE_FLOAT_VALUE(CounterName, Value) \
-	__TRACE_DECLARE_INLINE_COUNTER(CounterName, TraceCounterType_Float, TraceCounterDisplayHint_None) \
-	PREPROCESSOR_JOIN(__TraceCounter, __LINE__).Set(double(Value));
+#define TRACE_MEMORY_VALUE(CounterDisplayName, Value) \
+	__TRACE_DECLARE_INLINE_COUNTER(CounterDisplayName, FCounterInt, TraceCounterDisplayHint_Memory) \
+	PREPROCESSOR_JOIN(__TraceCounter, __LINE__).Set(Value);
 
-#define TRACE_MEMORY_VALUE(CounterName, Value) \
-	__TRACE_DECLARE_INLINE_COUNTER(CounterName, TraceCounterType_Int, TraceCounterDisplayHint_Memory) \
-	PREPROCESSOR_JOIN(__TraceCounter, __LINE__).Set(int64(Value));
+#define TRACE_DECLARE_INT_COUNTER(CounterName, CounterDisplayName) \
+	FCountersTrace::FCounterInt PREPROCESSOR_JOIN(__GTraceCounter, CounterName)(CounterDisplayName, TraceCounterDisplayHint_None);
 
-#define TRACE_DECLARE_INT_COUNTER(CounterName) \
-	__TRACE_DECLARE_COUNTER(CounterName, TraceCounterType_Int, TraceCounterDisplayHint_None)
+#define TRACE_DECLARE_INT_COUNTER_EXTERN(CounterName) \
+	extern FCountersTrace::FCounterInt PREPROCESSOR_JOIN(__GTraceCounter, CounterName);
 
-#define TRACE_DECLARE_FLOAT_COUNTER(CounterName) \
-	__TRACE_DECLARE_COUNTER(CounterName, TraceCounterType_Float, TraceCounterDisplayHint_None)
+#define TRACE_DECLARE_FLOAT_COUNTER(CounterName, CounterDisplayName) \
+	FCountersTrace::FCounterFloat PREPROCESSOR_JOIN(__GTraceCounter, CounterName)(CounterDisplayName, TraceCounterDisplayHint_None);
 
-#define TRACE_DECLARE_MEMORY_COUNTER(CounterName) \
-	__TRACE_DECLARE_COUNTER(CounterName, TraceCounterType_Int, TraceCounterDisplayHint_Memory)
+#define TRACE_DECLARE_FLOAT_COUNTER_EXTERN(CounterName) \
+	extern FCountersTrace::FCounterFloat PREPROCESSOR_JOIN(__GTraceCounter, CounterName);
 
-#define TRACE_DECLARE_COUNTER_EXTERN(CounterName) \
-	extern FCountersTrace::FCounter PREPROCESSOR_JOIN(__GTraceCounter, CounterName);
+#define TRACE_DECLARE_MEMORY_COUNTER(CounterName, CounterDisplayName) \
+	FCountersTrace::FCounterInt PREPROCESSOR_JOIN(__GTraceCounter, CounterName)(CounterDisplayName, TraceCounterDisplayHint_Memory);
+
+#define TRACE_DECLARE_MEMORY_COUNTER_EXTERN(CounterName) \
+	TRACE_DECLARE_INT_COUNTER_EXTERN(CounterName)
 
 #define TRACE_COUNTER_SET(CounterName, Value) \
 	PREPROCESSOR_JOIN(__GTraceCounter, CounterName).Set(Value);
@@ -184,24 +133,30 @@ struct FCountersTrace
 #define TRACE_COUNTER_ADD(CounterName, Value) \
 	PREPROCESSOR_JOIN(__GTraceCounter, CounterName).Add(Value);
 
+#define TRACE_COUNTER_SUBTRACT(CounterName, Value) \
+	PREPROCESSOR_JOIN(__GTraceCounter, CounterName).Subtract(Value);
+
 #define TRACE_COUNTER_INCREMENT(CounterName) \
-	PREPROCESSOR_JOIN(__GTraceCounter, CounterName).Add(1);
+	PREPROCESSOR_JOIN(__GTraceCounter, CounterName).Increment();
 
 #define TRACE_COUNTER_DECREMENT(CounterName) \
-	PREPROCESSOR_JOIN(__GTraceCounter, CounterName).Add(-1);
+	PREPROCESSOR_JOIN(__GTraceCounter, CounterName).Decrement();
 
 #else
 
 #define TRACE_COUNTERS_INIT(CmdLine)
-#define TRACE_INT_VALUE(CounterName, Value)
-#define TRACE_FLOAT_VALUE(CounterName, Value)
-#define TRACE_MEMORY_VALUE(CounterName, Value)
-#define TRACE_DECLARE_INT_COUNTER(CounterName)
-#define TRACE_DECLARE_FLOAT_COUNTER(CounterName)
-#define TRACE_DECLARE_MEMORY_COUNTER(CounterName)
-#define TRACE_DECLARE_COUNTER_EXTERN(CounterName)
+#define TRACE_INT_VALUE(CounterDisplayName, Value)
+#define TRACE_FLOAT_VALUE(CounterDisplayName, Value)
+#define TRACE_MEMORY_VALUE(CounterDisplayName, Value)
+#define TRACE_DECLARE_INT_COUNTER(CounterName, CounterDisplayName)
+#define TRACE_DECLARE_INT_COUNTER_EXTERN(CounterName)
+#define TRACE_DECLARE_FLOAT_COUNTER(CounterName, CounterDisplayName)
+#define TRACE_DECLARE_FLOAT_COUNTER_EXTERN(CounterName)
+#define TRACE_DECLARE_MEMORY_COUNTER(CounterName, CounterDisplayName)
+#define TRACE_DECLARE_MEMORY_COUNTER_EXTERN(CounterName)
 #define TRACE_COUNTER_SET(CounterName, Value)
 #define TRACE_COUNTER_ADD(CounterName, Value)
+#define TRACE_COUNTER_SUBTRACT(CounterName, Value)
 #define TRACE_COUNTER_INCREMENT(CounterName)
 #define TRACE_COUNTER_DECREMENT(CounterName)
 
