Allow LoadTimeTrace's define to be set externally.

#rnx

diff --git a/Engine/Source/Runtime/Core/Public/Serialization/LoadTimeTrace.h b/Engine/Source/Runtime/Core/Public/Serialization/LoadTimeTrace.h
--- a/Engine/Source/Runtime/Core/Public/Serialization/LoadTimeTrace.h
+++ b/Engine/Source/Runtime/Core/Public/Serialization/LoadTimeTrace.h
@@ -6,11 +6,13 @@
 #include "Trace/Config.h"
 #include "ProfilingDebugging/FormatArgsTrace.h"
 
+#if !defined(LOADTIMEPROFILERTRACE_ENABLED)
 #if UE_TRACE_ENABLED && !UE_BUILD_SHIPPING
 #define LOADTIMEPROFILERTRACE_ENABLED 1
 #else
 #define LOADTIMEPROFILERTRACE_ENABLED 0
 #endif
+#endif
 
 #if LOADTIMEPROFILERTRACE_ENABLED
 
