Add FIoDispatcher::DoesChunkExists  that allows the user to poll if a ChunkId has been mounted yet or not.

The main use case for this will be optional bulk data files, which will have valid FChunkIds but might exist in FioStoreReaders that have not yet been mounted.
In future work when we extend the system to actually have multiple stores for optional data we could consider baking the store id into the FChunkId which would allow for much faster checks?

#rnx

diff --git a/Engine/Source/Runtime/Core/Private/IO/IoDispatcher.cpp b/Engine/Source/Runtime/Core/Private/IO/IoDispatcher.cpp
--- a/Engine/Source/Runtime/Core/Private/IO/IoDispatcher.cpp
+++ b/Engine/Source/Runtime/Core/Private/IO/IoDispatcher.cpp
@@ -171,6 +171,11 @@ public:
 		return FIoStatus(EIoErrorCode::NotFound);
 	}
 
+	bool DoesChunkExists(const FIoChunkId& ChunkId) const
+	{
+		return Toc.Find(ChunkId) != nullptr;
+	}
+
 	TIoStatusOr<uint64> GetSizeForChunk(const FIoChunkId& ChunkId) const
 	{
 		const FIoStoreTocEntry* Entry = Toc.Find(ChunkId);
@@ -430,6 +435,19 @@ public:
 		return FIoStatus(EIoErrorCode::NotFound);
 	}
 
+	bool DoesChunkExists(const FIoChunkId& ChunkId) const
+	{
+		for (const auto& IoStore : IoStores)
+		{
+			if (IoStore->Impl->DoesChunkExists(ChunkId))
+			{
+				return true;
+			}
+		}
+
+		return false;
+	}
+
 	TIoStatusOr<uint64> GetSizeForChunk(const FIoChunkId& ChunkId) const
 	{
 		for (const auto& IoStore : IoStores)
@@ -538,6 +556,11 @@ public:
 		FreeRequest(Request);
 	}
 
+	bool DoesChunkExists(const FIoChunkId& ChunkId) const
+	{
+		return IoStore->DoesChunkExists(ChunkId);
+	}
+
 	TIoStatusOr<uint64> GetSizeForChunk(const FIoChunkId& ChunkId) const
 	{
 		return IoStore->GetSizeForChunk(ChunkId);
@@ -683,6 +706,12 @@ FIoDispatcher::CloseFileChunk(const FIoChunkId& FileChunkId)
 }
 
 // Polling methods
+bool 
+FIoDispatcher::DoesChunkExists(const FIoChunkId& ChunkId) const
+{
+	return Impl->DoesChunkExists(ChunkId);
+}
+
 TIoStatusOr<uint64>	
 FIoDispatcher::GetSizeForChunk(const FIoChunkId& ChunkId) const
 {
diff --git a/Engine/Source/Runtime/Core/Public/IO/IoDispatcher.h b/Engine/Source/Runtime/Core/Public/IO/IoDispatcher.h
--- a/Engine/Source/Runtime/Core/Public/IO/IoDispatcher.h
+++ b/Engine/Source/Runtime/Core/Public/IO/IoDispatcher.h
@@ -743,6 +743,7 @@ public:
 	CORE_API void					ReadWithCallback(const FIoChunkId& Chunk, const FIoReadOptions& Options, TFunction<void(TIoStatusOr<FIoBuffer>)>&& Callback);
 
 	// Polling methods
+	CORE_API bool					DoesChunkExists(const FIoChunkId& ChunkId) const;
 	CORE_API TIoStatusOr<uint64>	GetSizeForChunk(const FIoChunkId& ChunkId) const;
 
 	CORE_API TIoStatusOr<FIoChunkId>	OpenFileChunk(FStringView Filename);
