Added UE_ASSUME macro that works on both MSVC and clang

Deprecated old ASSUME macro, that only worked for MSVC, and updated all call sites.

diff --git a/Engine/Source/Programs/UnrealHeaderTool/Private/HeaderParser.cpp b/Engine/Source/Programs/UnrealHeaderTool/Private/HeaderParser.cpp
--- a/Engine/Source/Programs/UnrealHeaderTool/Private/HeaderParser.cpp
+++ b/Engine/Source/Programs/UnrealHeaderTool/Private/HeaderParser.cpp
@@ -8670,7 +8670,7 @@ bool ShouldKeepBlockContents(EBlockDirectiveType DirectiveType)
 	}
 
 	check(false);
-	ASSUME(false);
+	UE_ASSUME(false);
 }
 
 bool ShouldKeepDirective(EBlockDirectiveType DirectiveType)
@@ -8691,7 +8691,7 @@ bool ShouldKeepDirective(EBlockDirectiveType DirectiveType)
 	}
 
 	check(false);
-	ASSUME(false);
+	UE_ASSUME(false);
 }
 
 EBlockDirectiveType ParseCommandToBlockDirectiveType(const TCHAR** Str)
@@ -8749,7 +8749,7 @@ const TCHAR* GetBlockDirectiveTypeString(EBlockDirectiveType DirectiveType)
 	}
 
 	check(false);
-	ASSUME(false);
+	UE_ASSUME(false);
 }
 
 // Performs a preliminary parse of the text in the specified buffer, pulling out useful information for the header generation process
diff --git a/Engine/Source/Runtime/Core/Public/Containers/BitArray.h b/Engine/Source/Runtime/Core/Public/Containers/BitArray.h
--- a/Engine/Source/Runtime/Core/Public/Containers/BitArray.h
+++ b/Engine/Source/Runtime/Core/Public/Containers/BitArray.h
@@ -643,7 +643,7 @@ public:
 		{
 			// If we're looking for a false, then we flip the bits - then we only need to find the first one bit
 			const uint32 Bits = bValue ? (DwordArray[DwordIndex]) : ~(DwordArray[DwordIndex]);
-			ASSUME(Bits != 0);
+			UE_ASSUME(Bits != 0);
 			const int32 LowestBitIndex = FMath::CountTrailingZeros(Bits) + (DwordIndex << NumBitsPerDWORDLogTwo);
 			if (LowestBitIndex < LocalNumBits)
 			{
@@ -686,7 +686,7 @@ public:
 
 		// Flip the bits, then we only need to find the first one bit -- easy.
 		const uint32 Bits = (bValue ? DwordArray[DwordIndex] : ~DwordArray[DwordIndex]) & Mask;
-		ASSUME(Bits != 0);
+		UE_ASSUME(Bits != 0);
 
 		uint32 BitIndex = (NumBitsPerDWORD - 1) - FMath::CountLeadingZeros(Bits);
 
@@ -719,7 +719,7 @@ public:
 		{
 			// Flip the bits, then we only need to find the first one bit -- easy.
 			const uint32 Bits = ~(DwordArray[DwordIndex]);
-			ASSUME(Bits != 0);
+			UE_ASSUME(Bits != 0);
 			const uint32 LowestBit = (Bits) & (-(int32)Bits);
 			const int32 LowestBitIndex = FMath::CountTrailingZeros(Bits) + (DwordIndex << NumBitsPerDWORDLogTwo);
 			if (LowestBitIndex < LocalNumBits)
@@ -763,7 +763,7 @@ public:
 
 		// Flip the bits, then we only need to find the first one bit -- easy.
 		const uint32 Bits = ~DwordArray[DwordIndex] & Mask;
-		ASSUME(Bits != 0);
+		UE_ASSUME(Bits != 0);
 
 		uint32 BitIndex = (NumBitsPerDWORD - 1) - FMath::CountLeadingZeros(Bits);
 		DwordArray[DwordIndex] |= 1u << BitIndex;
diff --git a/Engine/Source/Runtime/Core/Public/Containers/SparseArray.h b/Engine/Source/Runtime/Core/Public/Containers/SparseArray.h
--- a/Engine/Source/Runtime/Core/Public/Containers/SparseArray.h
+++ b/Engine/Source/Runtime/Core/Public/Containers/SparseArray.h
@@ -1281,6 +1281,6 @@ struct TIsZeroConstructType<FScriptSparseArray>
  */
 inline void* operator new(size_t Size,const FSparseArrayAllocationInfo& Allocation)
 {
-	ASSUME(Allocation.Pointer);
+	UE_ASSUME(Allocation.Pointer);
 	return Allocation.Pointer;
 }
diff --git a/Engine/Source/Runtime/Core/Public/Containers/StaticBitArray.h b/Engine/Source/Runtime/Core/Public/Containers/StaticBitArray.h
--- a/Engine/Source/Runtime/Core/Public/Containers/StaticBitArray.h
+++ b/Engine/Source/Runtime/Core/Public/Containers/StaticBitArray.h
@@ -295,7 +295,7 @@ public:
 		{
 			// Flip the bits, then we only need to find the first one bit -- easy.
 			const WordType Bits = ~(Words[WordIndex]);
-			ASSUME(Bits != 0);
+			UE_ASSUME(Bits != 0);
 
 			const int32 LowestBitIndex = FMath::CountTrailingZeros64(Bits) + (WordIndex << NumBitsPerWordLog2);
 			if (LowestBitIndex < LocalNumBits)
@@ -327,7 +327,7 @@ public:
 		if (WordIndex < NumWords)
 		{
 			const WordType Bits = Words[WordIndex];
-			ASSUME(Bits != 0);
+			UE_ASSUME(Bits != 0);
 
 			const int32 LowestBitIndex = (int32)FMath::CountTrailingZeros64(Bits) + (WordIndex << NumBitsPerWordLog2);
 			if (LowestBitIndex < LocalNumBits)
diff --git a/Engine/Source/Runtime/Core/Public/HAL/Platform.h b/Engine/Source/Runtime/Core/Public/HAL/Platform.h
--- a/Engine/Source/Runtime/Core/Public/HAL/Platform.h
+++ b/Engine/Source/Runtime/Core/Public/HAL/Platform.h
@@ -615,12 +615,25 @@
 	#define FUNCTION_CHECK_RETURN(...) DEPRECATED_MACRO(4.12, "FUNCTION_CHECK_RETURN has been deprecated and should be replaced with FUNCTION_CHECK_RETURN_START and FUNCTION_CHECK_RETURN_END.") FUNCTION_CHECK_RETURN_START __VA_ARGS__ FUNCTION_CHECK_RETURN_END
 #endif
 
-#ifndef ASSUME										/* Hints compiler that expression is true; generally restricted to comparisons against constants */
-	#define ASSUME(...)
+/** Promise expression is true. Compiler can optimize accordingly with undefined behavior if wrong. Static analyzers understand this.  */
+#ifndef UE_ASSUME
+	#if defined(__clang__) 
+		#if __has_builtin(__builtin_assume)
+			#define UE_ASSUME(x) __builtin_assume(x)
+		#else
+			#define UE_ASSUME(x)
+		#endif
+	#elif defined(_MSC_VER)
+		#define UE_ASSUME(x) __assume(x)
+	#else
+		#define UE_ASSUME(x)
+	#endif
 #endif
 
+#define ASSUME(x) UE_ASSUME(x) DEPRECATED_MACRO(4.25, "Please use UE_ASSUME instead.")
+
 /** Branch prediction hints */
-#ifndef LIKELY						/* Hints compiler that expression is likely to be true, much softer than ASSUME - allows (penalized by worse performance) expression to be false */
+#ifndef LIKELY						/* Hints compiler that expression is likely to be true, much softer than UE_ASSUME - allows (penalized by worse performance) expression to be false */
 	#if ( defined(__clang__) || defined(__GNUC__) ) && (PLATFORM_UNIX)	// effect of these on non-Linux platform has not been analyzed as of 2016-03-21
 		#define LIKELY(x)			__builtin_expect(!!(x), 1)
 	#else
diff --git a/Engine/Source/Runtime/Core/Public/HoloLens/HoloLensPlatform.h b/Engine/Source/Runtime/Core/Public/HoloLens/HoloLensPlatform.h
--- a/Engine/Source/Runtime/Core/Public/HoloLens/HoloLensPlatform.h
+++ b/Engine/Source/Runtime/Core/Public/HoloLens/HoloLensPlatform.h
@@ -91,9 +91,6 @@ typedef FHoloLensTypes FPlatformTypes;
 #define FORCEINLINE		__forceinline			/* Force code to be inline */
 #define FORCENOINLINE	__declspec(noinline)	/* Force code to NOT be inline */
 
-// Hints compiler that expression is true; generally restricted to comparisons against constants
-#define ASSUME(expr)	__assume(expr)
-
 #define DECLARE_UINT64(x)	x
 
 // Optimization macros (uses __pragma to enable inside a #define).
diff --git a/Engine/Source/Runtime/Core/Public/Windows/WIndowsPlatform.h b/Engine/Source/Runtime/Core/Public/Windows/WIndowsPlatform.h
--- a/Engine/Source/Runtime/Core/Public/Windows/WIndowsPlatform.h
+++ b/Engine/Source/Runtime/Core/Public/Windows/WIndowsPlatform.h
@@ -96,11 +96,6 @@ typedef FWindowsPlatformTypes FPlatformTypes;
 #define FUNCTION_NO_RETURN_START __declspec(noreturn)				/* Indicate that the function never returns. */
 #define FUNCTION_NON_NULL_RETURN_START _Ret_notnull_				/* Indicate that the function never returns nullptr. */
 
-// Hints compiler that expression is true; generally restricted to comparisons against constants
-#if !defined(__clang__) || defined(_MSC_VER)	// Clang only supports __assume when using -fms-extensions
-	#define ASSUME(expr) __assume(expr)
-#endif
-
 #define DECLARE_UINT64(x)	x
 
 // Optimization macros (uses __pragma to enable inside a #define).
diff --git a/Engine/Source/Runtime/CoreUObject/Private/Serialization/UnversionedPropertySerialization.cpp b/Engine/Source/Runtime/CoreUObject/Private/Serialization/UnversionedPropertySerialization.cpp
--- a/Engine/Source/Runtime/CoreUObject/Private/Serialization/UnversionedPropertySerialization.cpp
+++ b/Engine/Source/Runtime/CoreUObject/Private/Serialization/UnversionedPropertySerialization.cpp
@@ -112,7 +112,7 @@ public:
 			case EIntegerType::Uint16: MemZeroRange<uint16>(GetValue(Data), FastZeroIntNum); break;
 			case EIntegerType::Uint32: MemZeroRange<uint32>(GetValue(Data), FastZeroIntNum); break;
 			case EIntegerType::Uint64: MemZeroRange<uint64>(GetValue(Data), FastZeroIntNum); break;
-			default: ASSUME(0);
+			default: UE_ASSUME(0);
 		}
 	}
 
@@ -196,7 +196,7 @@ private:
 			 case EIntegerType::Uint16: Slot << *reinterpret_cast<uint16*>(Value); break;
 			 case EIntegerType::Uint32: Slot << *reinterpret_cast<uint32*>(Value); break;
 			 case EIntegerType::Uint64: Slot << *reinterpret_cast<uint64*>(Value); break;
-			 default: ASSUME(0);
+			 default: UE_ASSUME(0);
 		 }
 	}
 
@@ -232,7 +232,7 @@ private:
 	template<typename T>
 	FORCEINLINE static void MemZeroRange(void* Value, uint32 Num)
 	{
-		ASSUME(Num > 0);
+		UE_ASSUME(Num > 0);
 		for (uint32 Idx = 0; Idx < Num; ++Idx)
 		{
 			reinterpret_cast<T*>(Value)[Idx] = 0;
diff --git a/Engine/Source/Runtime/CoreUObject/Public/UObject/UObjectBaseUtility.h b/Engine/Source/Runtime/CoreUObject/Public/UObject/UObjectBaseUtility.h
--- a/Engine/Source/Runtime/CoreUObject/Public/UObject/UObjectBaseUtility.h
+++ b/Engine/Source/Runtime/CoreUObject/Public/UObject/UObjectBaseUtility.h
@@ -490,8 +490,8 @@ public:
 		const UClass* ThisClass = GetClass();
 
 		// Stop the compiler doing some unnecessary branching for nullptr checks
-		ASSUME(SomeBaseClass);
-		ASSUME(ThisClass);
+		UE_ASSUME(SomeBaseClass);
+		UE_ASSUME(ThisClass);
 
 		return IsChildOfWorkaround(ThisClass, SomeBaseClass);
 	}
diff --git a/Engine/Source/Runtime/D3D12RHI/Private/D3D12View.h b/Engine/Source/Runtime/D3D12RHI/Private/D3D12View.h
--- a/Engine/Source/Runtime/D3D12RHI/Private/D3D12View.h
+++ b/Engine/Source/Runtime/D3D12RHI/Private/D3D12View.h
@@ -38,7 +38,7 @@ public:
 	{
 		switch (Desc.ViewDimension)
 		{
-		default: ASSUME(0 && "Corrupt Resource Type on Shader Resource View"); break;
+		default: UE_ASSUME(0 && "Corrupt Resource Type on Shader Resource View"); break;
 
 		case (D3D12_SRV_DIMENSION_BUFFER) :
 			break;
@@ -126,7 +126,7 @@ public:
 	{
 		switch (Desc.ViewDimension)
 		{
-		default: ASSUME(0 && "Corrupt Resource Type on Unordered Access View"); break;
+		default: UE_ASSUME(0 && "Corrupt Resource Type on Unordered Access View"); break;
 
 		case (D3D12_UAV_DIMENSION_BUFFER) : break;
 
@@ -172,7 +172,7 @@ public:
 	{
 		switch (Desc.ViewDimension)
 		{
-		default: ASSUME(0 && "Corrupt Resource Type on Render Target View"); break;
+		default: UE_ASSUME(0 && "Corrupt Resource Type on Render Target View"); break;
 
 		case (D3D12_RTV_DIMENSION_BUFFER) : break;
 
@@ -226,7 +226,7 @@ public:
 	{
 		switch (Desc.ViewDimension)
 		{
-		default: ASSUME(0 && "Corrupt Resource Type on Depth Stencil View"); break;
+		default: UE_ASSUME(0 && "Corrupt Resource Type on Depth Stencil View"); break;
 
 		case (D3D12_DSV_DIMENSION_TEXTURE1D) :
 			m_BeginMip = uint8(Desc.Texture1D.MipSlice);
diff --git a/Engine/Source/Runtime/Renderer/Private/ShadowSetup.cpp b/Engine/Source/Runtime/Renderer/Private/ShadowSetup.cpp
--- a/Engine/Source/Runtime/Renderer/Private/ShadowSetup.cpp
+++ b/Engine/Source/Runtime/Renderer/Private/ShadowSetup.cpp
@@ -399,8 +399,8 @@ static bool GetBestShadowTransform(const FVector& ZAxis,const FBoundingBoxVertex
 	const int32 NumEdges = Edges.Num();
 
 	// We're always dealing with box geometry here, so we can hint the compiler
-	ASSUME( NumPoints == 8 );
-	ASSUME( NumEdges == 12 );
+	UE_ASSUME( NumPoints == 8 );
+	UE_ASSUME( NumEdges == 12 );
 
 	for(int32 EdgeIndex = 0;EdgeIndex < NumEdges; ++EdgeIndex)
 	{
