Fix a typo that I introduced when renaming a method during code review.

#rnx

diff --git a/Engine/Source/Runtime/Core/Private/IO/IoDispatcher.cpp b/Engine/Source/Runtime/Core/Private/IO/IoDispatcher.cpp
--- a/Engine/Source/Runtime/Core/Private/IO/IoDispatcher.cpp
+++ b/Engine/Source/Runtime/Core/Private/IO/IoDispatcher.cpp
@@ -171,7 +171,7 @@ public:
 		return FIoStatus(EIoErrorCode::NotFound);
 	}
 
-	bool DoesChunkExists(const FIoChunkId& ChunkId) const
+	bool DoesChunkExist(const FIoChunkId& ChunkId) const
 	{
 		return Toc.Find(ChunkId) != nullptr;
 	}
@@ -435,11 +435,11 @@ public:
 		return FIoStatus(EIoErrorCode::NotFound);
 	}
 
-	bool DoesChunkExists(const FIoChunkId& ChunkId) const
+	bool DoesChunkExist(const FIoChunkId& ChunkId) const
 	{
 		for (const auto& IoStore : IoStores)
 		{
-			if (IoStore->Impl->DoesChunkExists(ChunkId))
+			if (IoStore->Impl->DoesChunkExist(ChunkId))
 			{
 				return true;
 			}
@@ -556,9 +556,9 @@ public:
 		FreeRequest(Request);
 	}
 
-	bool DoesChunkExists(const FIoChunkId& ChunkId) const
+	bool DoesChunkExist(const FIoChunkId& ChunkId) const
 	{
-		return IoStore->DoesChunkExists(ChunkId);
+		return IoStore->DoesChunkExist(ChunkId);
 	}
 
 	TIoStatusOr<uint64> GetSizeForChunk(const FIoChunkId& ChunkId) const
@@ -707,9 +707,9 @@ FIoDispatcher::CloseFileChunk(const FIoChunkId& FileChunkId)
 
 // Polling methods
 bool 
-FIoDispatcher::DoesChunkExists(const FIoChunkId& ChunkId) const
+FIoDispatcher::DoesChunkExist(const FIoChunkId& ChunkId) const
 {
-	return Impl->DoesChunkExists(ChunkId);
+	return Impl->DoesChunkExist(ChunkId);
 }
 
 TIoStatusOr<uint64>	
diff --git a/Engine/Source/Runtime/Core/Public/IO/IoDispatcher.h b/Engine/Source/Runtime/Core/Public/IO/IoDispatcher.h
--- a/Engine/Source/Runtime/Core/Public/IO/IoDispatcher.h
+++ b/Engine/Source/Runtime/Core/Public/IO/IoDispatcher.h
@@ -743,7 +743,7 @@ public:
 	CORE_API void					ReadWithCallback(const FIoChunkId& Chunk, const FIoReadOptions& Options, TFunction<void(TIoStatusOr<FIoBuffer>)>&& Callback);
 
 	// Polling methods
-	CORE_API bool					DoesChunkExists(const FIoChunkId& ChunkId) const;
+	CORE_API bool					DoesChunkExist(const FIoChunkId& ChunkId) const;
 	CORE_API TIoStatusOr<uint64>	GetSizeForChunk(const FIoChunkId& ChunkId) const;
 
 	CORE_API TIoStatusOr<FIoChunkId>	OpenFileChunk(FStringView Filename);
