Replace auto with the actual type (resolved from TRefCountPtr<FIoStoreReader>)

#rnx

diff --git a/Engine/Source/Runtime/Core/Private/IO/IoDispatcher.cpp b/Engine/Source/Runtime/Core/Private/IO/IoDispatcher.cpp
--- a/Engine/Source/Runtime/Core/Private/IO/IoDispatcher.cpp
+++ b/Engine/Source/Runtime/Core/Private/IO/IoDispatcher.cpp
@@ -422,7 +422,7 @@ public:
 
 	TIoStatusOr<FIoBuffer> Resolve(FIoChunkId& ChunkId, FIoReadOptions Options)
 	{
-		for (const auto& IoStore : IoStores)
+		for (const FIoStoreReader* IoStore : IoStores)
 		{
 			TIoStatusOr<FIoBuffer> Result = IoStore->Impl->Lookup(ChunkId, Options);
 
@@ -437,7 +437,7 @@ public:
 
 	bool DoesChunkExist(const FIoChunkId& ChunkId) const
 	{
-		for (const auto& IoStore : IoStores)
+		for (const FIoStoreReader* IoStore : IoStores)
 		{
 			if (IoStore->Impl->DoesChunkExist(ChunkId))
 			{
@@ -450,7 +450,7 @@ public:
 
 	TIoStatusOr<uint64> GetSizeForChunk(const FIoChunkId& ChunkId) const
 	{
-		for (const auto& IoStore : IoStores)
+		for (const FIoStoreReader* IoStore : IoStores)
 		{
 			TIoStatusOr<uint64> Result = IoStore->Impl->GetSizeForChunk(ChunkId);
 
