Iteration on StringView

- Add a default constructor to the string view types.
- Add an implicit string view constructor from any contiguous container of compatible characters.
- Add GetData, GetNum, and specialize TIsContiguousContainer to support TStringViewImpl as a contiguous container.
- Extend equality and inequality operators to allow more flexible comparison of string views.
- Restore derived types for F*StringView to permit forward declarations.
- Remove ToString from string views to avoid a dependency on FString.
- Remove the header dependency between StringView.h and UnrealString.h.
- Remove FindLastCharByPredicate in favor of using algorithms.

- Add an explicit constructor of strings from string views.
- Add assignment of strings from string views.

- Move ESearchCase and ESearchDir to CString.h to enable their use without including FString.
- Add Algo::FindLast[ByPredicate] to permit FindLastCharByPredicate functionality on string views.
- Add a missing #include in ElementType.h.

diff --git a/Engine/Source/Runtime/Core/Private/Containers/String.cpp b/Engine/Source/Runtime/Core/Private/Containers/String.cpp
--- a/Engine/Source/Runtime/Core/Private/Containers/String.cpp
+++ b/Engine/Source/Runtime/Core/Private/Containers/String.cpp
@@ -1,16 +1,18 @@
 // Copyright 1998-2019 Epic Games, Inc. All Rights Reserved.
 
+#include "Containers/UnrealString.h"
+
+#include "CoreGlobals.h"
 #include "CoreTypes.h"
-#include "Misc/AssertionMacros.h"
-#include "Misc/VarArgs.h"
-#include "Math/NumericLimits.h"
-#include "Math/UnrealMathUtility.h"
 #include "HAL/UnrealMemory.h"
-#include "Templates/UnrealTemplate.h"
-#include "Containers/UnrealString.h"
 #include "Logging/LogMacros.h"
-#include "CoreGlobals.h"
+#include "Math/NumericLimits.h"
+#include "Math/UnrealMathUtility.h"
+#include "Misc/AssertionMacros.h"
 #include "Misc/ByteSwap.h"
+#include "Misc/StringView.h"
+#include "Misc/VarArgs.h"
+#include "Templates/UnrealTemplate.h"
 
 /* FString implementation
  *****************************************************************************/
@@ -116,6 +118,44 @@ namespace UE4String_Private
 	}
 }
 
+FString::FString(const FStringView& Other)
+{
+	if (const FStringView::SizeType OtherLen = Other.Len())
+	{
+		Reserve(OtherLen);
+		Append(Other.Data(), OtherLen);
+		AppendChar(TEXT('\0'));
+	}
+}
+
+FString& FString::operator=(const FStringView& Other)
+{
+	const TCHAR* const OtherData = Other.Data();
+	const FStringView::SizeType OtherLen = Other.Len();
+	if (OtherLen == 0)
+	{
+		Data.Empty();
+	}
+	else if (OtherData < Data.GetData() + Data.Num() && Data.GetData() < OtherData + OtherLen)
+	{
+		*this = FString(Other);
+	}
+	else
+	{
+		Data.Empty(OtherLen + 1);
+		Data.AddUninitialized(OtherLen + 1);
+		TCHAR* DataPtr = Data.GetData();
+		CopyAssignItems(DataPtr, OtherData, OtherLen);
+		DataPtr[OtherLen] = TEXT('\0');
+	}
+	return *this;
+}
+
+FString::operator FStringView() const
+{
+	return FStringView(Data.GetData(), Len());
+}
+
 void FString::TrimToNullTerminator()
 {
 	if( Data.Num() )
diff --git a/Engine/Source/Runtime/Core/Private/Containers/StringTest.cpp b/Engine/Source/Runtime/Core/Private/Containers/StringTest.cpp
--- a/Engine/Source/Runtime/Core/Private/Containers/StringTest.cpp
+++ b/Engine/Source/Runtime/Core/Private/Containers/StringTest.cpp
@@ -3,6 +3,7 @@
 #include "CoreTypes.h"
 #include "Misc/AutomationTest.h"
 #include "Misc/AssertionMacros.h"
+#include "Misc/StringView.h"
 #include "Containers/UnrealString.h"
 #include "Serialization/MemoryReader.h"
 #include "Serialization/MemoryWriter.h"
@@ -238,4 +239,46 @@ bool FLexTryParseStringTest::RunTest(const FString& Parameters)
 	return true;
 }
 
+IMPLEMENT_SIMPLE_AUTOMATION_TEST(FStringFromStringViewTest, "System.Core.String.FromStringView", EAutomationTestFlags::ApplicationContextMask | EAutomationTestFlags::SmokeFilter)
+bool FStringFromStringViewTest::RunTest(const FString& Parameters)
+{
+	// Verify basic construction and assignment from a string view.
+	{
+		const TCHAR* Literal = TEXT("Literal");
+		TestEqual(TEXT("String(StringView)"), FString(FStringView(Literal)), Literal);
+		TestEqual(TEXT("String = StringView"), FString(TEXT("Temp")) = FStringView(Literal), Literal);
+
+		FStringView EmptyStringView;
+		FString EmptyString(EmptyStringView);
+		TestTrue(TEXT("String(EmptyStringView)"), EmptyString.IsEmpty());
+		TestTrue(TEXT("String(EmptyStringView) (No Allocation)"), EmptyString.GetAllocatedSize() == 0);
+
+		EmptyString = TEXT("Temp");
+		EmptyString = EmptyStringView;
+		TestTrue(TEXT("String = EmptyStringView"), EmptyString.IsEmpty());
+		TestTrue(TEXT("String = EmptyStringView (No Allocation)"), EmptyString.GetAllocatedSize() == 0);
+	}
+
+	// Verify assignment from a view of itself.
+	{
+		FString AssignEntireString(TEXT("AssignEntireString"));
+		AssignEntireString = FStringView(AssignEntireString);
+		TestEqual(TEXT("String = StringView(String)"), AssignEntireString, TEXT("AssignEntireString"));
+
+		FString AssignStartOfString(TEXT("AssignStartOfString"));
+		AssignStartOfString = FStringView(AssignStartOfString).Left(11);
+		TestEqual(TEXT("String = StringView(String).Left"), AssignStartOfString, TEXT("AssignStart"));
+
+		FString AssignEndOfString(TEXT("AssignEndOfString"));
+		AssignEndOfString = FStringView(AssignEndOfString).Right(11);
+		TestEqual(TEXT("String = StringView(String).Right"), AssignEndOfString, TEXT("EndOfString"));
+
+		FString AssignMiddleOfString(TEXT("AssignMiddleOfString"));
+		AssignMiddleOfString = FStringView(AssignMiddleOfString).Mid(6, 6);
+		TestEqual(TEXT("String = StringView(String).Mid"), AssignMiddleOfString, TEXT("Middle"));
+	}
+
+	return true;
+}
+
 #endif // WITH_DEV_AUTOMATION_TESTS
diff --git a/Engine/Source/Runtime/Core/Private/IO/IoDispatcher.cpp b/Engine/Source/Runtime/Core/Private/IO/IoDispatcher.cpp
--- a/Engine/Source/Runtime/Core/Private/IO/IoDispatcher.cpp
+++ b/Engine/Source/Runtime/Core/Private/IO/IoDispatcher.cpp
@@ -192,7 +192,7 @@ public:
 
 	void SetUniqueId(FStringView InUniqueId)
 	{
-		UniqueId = InUniqueId.ToString();
+		UniqueId = InUniqueId;
 	}
 
 	TIoStatusOr<FIoChunkId> OpenFileChunk(FStringView Filename);
@@ -207,9 +207,9 @@ private:
 
 	FIoChunkId GetFileChunkId(FStringView Filename)
 	{
-		check(Filename.Data() && Filename.Len());
+		check(Filename.GetData() && Filename.Len());
 
-		uint32 Hash = FCrc::MemCrc32(Filename.Data(), Filename.Len());
+		uint32 Hash = FCrc::MemCrc32(Filename.GetData(), Filename.Len());
 
 		uint8 Data[12] = {0};
 		*reinterpret_cast<uint32*>(&Data[0]) = Hash;
@@ -234,7 +234,7 @@ FIoStatus FIoStoreReaderImpl::Initialize(FStringView InUniqueId)
 {
 	IPlatformFile& Ipf = FPlatformFileManager::Get().GetPlatformFile();
 
-	UniqueId = InUniqueId.ToString();
+	UniqueId = InUniqueId;
 
 	const FString& RootPath = Environment.GetRootPath();
 
@@ -337,7 +337,7 @@ FIoStoreReaderImpl::OpenFileChunk(FStringView Filename)
 	IPlatformFile& Ipf = FPlatformFileManager::Get().GetPlatformFile();
 
 	TUniquePtr<IMappedFileHandle> FileHandle;
-	FileHandle.Reset(Ipf.OpenMapped(Filename.Data()));
+	FileHandle.Reset(Ipf.OpenMapped(Filename.GetData()));
 
 	if (!FileHandle)
 	{
diff --git a/Engine/Source/Runtime/Core/Private/IO/IoStore.cpp b/Engine/Source/Runtime/Core/Private/IO/IoStore.cpp
--- a/Engine/Source/Runtime/Core/Private/IO/IoStore.cpp
+++ b/Engine/Source/Runtime/Core/Private/IO/IoStore.cpp
@@ -21,7 +21,7 @@ FIoStoreEnvironment::~FIoStoreEnvironment()
 
 void FIoStoreEnvironment::InitializeFileEnvironment(FStringView InRootPath)
 {
-	RootPath = InRootPath.ToString();
+	RootPath = InRootPath;
 }
 
 //////////////////////////////////////////////////////////////////////////
diff --git a/Engine/Source/Runtime/Core/Private/Misc/StringView.cpp b/Engine/Source/Runtime/Core/Private/Misc/StringView.cpp
--- a/Engine/Source/Runtime/Core/Private/Misc/StringView.cpp
+++ b/Engine/Source/Runtime/Core/Private/Misc/StringView.cpp
@@ -1,3 +1,70 @@
 // Copyright 1998-2019 Epic Games, Inc. All Rights Reserved.
 
 #include "Misc/StringView.h"
+
+#include "Containers/UnrealString.h"
+#include "HAL/PlatformString.h"
+
+template<typename CharType, typename ViewType>
+int32 TStringViewImpl<CharType, ViewType>::Compare(const TStringViewImpl& Other, ESearchCase::Type SearchCase) const
+{
+	const SizeType SelfLen = Len();
+	const SizeType OtherLen = Other.Len();
+	const SizeType ShortestLength = FMath::Min(SelfLen, OtherLen);
+
+	int Result;
+	if (SearchCase == ESearchCase::CaseSensitive)
+	{
+		Result = FPlatformString::Strncmp(GetData(), Other.GetData(), ShortestLength);
+	}
+	else
+	{
+		Result = FPlatformString::Strnicmp(GetData(), Other.GetData(), ShortestLength);
+	}
+
+	if (Result != 0 || SelfLen == OtherLen)
+	{
+		return Result;
+	}
+
+	return SelfLen < OtherLen ? -1 : 1;
+}
+
+template<typename CharType, typename ViewType>
+bool TStringViewImpl<CharType, ViewType>::FindChar(CharType InChar, SizeType& OutIndex) const
+{
+	for (SizeType i = 0; i < Size; ++i)
+	{
+		if (DataPtr[i] == InChar)
+		{
+			OutIndex = i;
+			return true;
+		}
+	}
+
+	return false;
+}
+
+template<typename CharType, typename ViewType>
+bool TStringViewImpl<CharType, ViewType>::FindLastChar(CharType InChar, SizeType& OutIndex) const
+{
+	if (Size == 0)
+	{
+		return false;
+	}
+
+	for (SizeType i = Size - 1; i >= 0; --i)
+	{
+		if (DataPtr[i] == InChar)
+		{
+			OutIndex = i;
+			return true;
+		}
+	}
+
+	return false;
+}
+
+template class TStringViewImpl<TCHAR, FStringView>;
+template class TStringViewImpl<ANSICHAR, FAnsiStringView>;
+template class TStringViewImpl<WIDECHAR, FWideStringView>;
diff --git a/Engine/Source/Runtime/Core/Private/Tests/Misc/StringViewTest.cpp b/Engine/Source/Runtime/Core/Private/Tests/Misc/StringViewTest.cpp
--- a/Engine/Source/Runtime/Core/Private/Tests/Misc/StringViewTest.cpp
+++ b/Engine/Source/Runtime/Core/Private/Tests/Misc/StringViewTest.cpp
@@ -1,6 +1,8 @@
 // Copyright 1998-2019 Epic Games, Inc. All Rights Reserved.
 
 #include "Misc/StringView.h"
+
+#include "Containers/UnrealString.h"
 #include "Misc/AutomationTest.h"
 
 #if WITH_DEV_AUTOMATION_TESTS 
@@ -109,7 +111,7 @@ bool FStringViewTestEquality::RunTest(const FString& Parameters)
 	FStringView WideView(WideStringLiteralSrc);
 
 	TestTrue(TEXT("Equality(0)"), WideView == WideStringLiteralSrc);
-	TestFalse(TEXT("Equality(1)"), WideView != WideStringLiteralSrc);	
+	TestFalse(TEXT("Equality(1)"), WideView != WideStringLiteralSrc);
 	TestTrue(TEXT("Equality(2)"), WideView == WideStringLiteralLower);
 	TestFalse(TEXT("Equality(3)"), WideView != WideStringLiteralLower);
 	TestTrue(TEXT("Equality(4)"), WideView == WideStringLiteralUpper);
@@ -119,27 +121,91 @@ bool FStringViewTestEquality::RunTest(const FString& Parameters)
 	TestFalse(TEXT("Equality(8)"), WideView == WideStringLiteralLonger);
 	TestTrue(TEXT("Equality(9)"), WideView != WideStringLiteralLonger);
 
+	TestTrue(TEXT("Equality(10)"), WideStringLiteralSrc == WideView);
+	TestFalse(TEXT("Equality(11)"), WideStringLiteralSrc != WideView);
+	TestTrue(TEXT("Equality(12)"), WideStringLiteralLower == WideView);
+	TestFalse(TEXT("Equality(13)"), WideStringLiteralLower != WideView);
+	TestTrue(TEXT("Equality(14)"), WideStringLiteralUpper == WideView);
+	TestFalse(TEXT("Equality(15)"), WideStringLiteralUpper != WideView);
+	TestFalse(TEXT("Equality(16)"), WideStringLiteralShort == WideView);
+	TestTrue(TEXT("Equality(17)"), WideStringLiteralShort != WideView);
+	TestFalse(TEXT("Equality(18)"), WideStringLiteralLonger == WideView);
+	TestTrue(TEXT("Equality(19)"), WideStringLiteralLonger != WideView);
+
 	FString WideStringSrc = WideStringLiteralSrc;
 	FString WideStringLower = WideStringLiteralLower;
 	FString WideStringUpper = WideStringLiteralUpper;
 	FString WideStringShort = WideStringLiteralShort;
 	FString WideStringLonger = WideStringLiteralLonger;
 
-	TestTrue(TEXT("Equality(10)"), WideView == WideStringSrc);
-	TestFalse(TEXT("Equality(11)"), WideView != WideStringSrc);
-	TestTrue(TEXT("Equality(12)"), WideView == WideStringLower);
-	TestFalse(TEXT("Equality(13)"), WideView != WideStringLower);
-	TestTrue(TEXT("Equality(14)"), WideView == WideStringUpper);
-	TestFalse(TEXT("Equality(15)"), WideView != WideStringUpper);
-	TestFalse(TEXT("Equality(16)"), WideView == WideStringShort);
-	TestTrue(TEXT("Equality(17)"), WideView != WideStringShort);
-	TestFalse(TEXT("Equality(18)"), WideView == WideStringLonger);
-	TestTrue(TEXT("Equality(19)"), WideView != WideStringLonger);
+	TestTrue(TEXT("Equality(20)"), WideView == WideStringSrc);
+	TestFalse(TEXT("Equality(21)"), WideView != WideStringSrc);
+	TestTrue(TEXT("Equality(22)"), WideView == WideStringLower);
+	TestFalse(TEXT("Equality(23)"), WideView != WideStringLower);
+	TestTrue(TEXT("Equality(24)"), WideView == WideStringUpper);
+	TestFalse(TEXT("Equality(25)"), WideView != WideStringUpper);
+	TestFalse(TEXT("Equality(26)"), WideView == WideStringShort);
+	TestTrue(TEXT("Equality(27)"), WideView != WideStringShort);
+	TestFalse(TEXT("Equality(28)"), WideView == WideStringLonger);
+	TestTrue(TEXT("Equality(29)"), WideView != WideStringLonger);
+
+	TestTrue(TEXT("Equality(30)"), WideStringSrc == WideView);
+	TestFalse(TEXT("Equality(31)"), WideStringSrc != WideView);
+	TestTrue(TEXT("Equality(32)"), WideStringLower == WideView);
+	TestFalse(TEXT("Equality(33)"), WideStringLower != WideView);
+	TestTrue(TEXT("Equality(34)"), WideStringUpper == WideView);
+	TestFalse(TEXT("Equality(35)"), WideStringUpper != WideView);
+	TestFalse(TEXT("Equality(36)"), WideStringShort == WideView);
+	TestTrue(TEXT("Equality(37)"), WideStringShort != WideView);
+	TestFalse(TEXT("Equality(38)"), WideStringLonger == WideView);
+	TestTrue(TEXT("Equality(39)"), WideStringLonger != WideView);
 
 	FStringView IdenticalView(WideStringLiteralSrc);
 
-	TestTrue(TEXT("Equality(20)"), WideView == IdenticalView);
-	TestTrue(TEXT("Equality(21)"), IdenticalView == WideView);
+	TestTrue(TEXT("Equality(40)"), WideView == IdenticalView);
+	TestTrue(TEXT("Equality(41)"), IdenticalView == WideView);
+
+	// Views without null termination
+
+	FStringView ShortViewNoNull = WideView.Left(FStringView(WideStringLiteralShort).Len());
+
+	TestTrue(TEXT("Equality(42)"), ShortViewNoNull == WideStringLiteralShort);
+	TestFalse(TEXT("Equality(43)"), ShortViewNoNull != WideStringLiteralShort);
+	TestTrue(TEXT("Equality(44)"), WideStringLiteralShort == ShortViewNoNull);
+	TestFalse(TEXT("Equality(45)"), WideStringLiteralShort != ShortViewNoNull);
+	TestFalse(TEXT("Equality(46)"), ShortViewNoNull == WideStringLiteralSrc);
+	TestTrue(TEXT("Equality(47)"), ShortViewNoNull != WideStringLiteralSrc);
+	TestFalse(TEXT("Equality(48)"), WideStringLiteralSrc == ShortViewNoNull);
+	TestTrue(TEXT("Equality(49)"), WideStringLiteralSrc != ShortViewNoNull);
+
+	TestTrue(TEXT("Equality(50)"), ShortViewNoNull == WideStringShort);
+	TestFalse(TEXT("Equality(51)"), ShortViewNoNull != WideStringShort);
+	TestTrue(TEXT("Equality(52)"), WideStringShort == ShortViewNoNull);
+	TestFalse(TEXT("Equality(53)"), WideStringShort != ShortViewNoNull);
+	TestFalse(TEXT("Equality(54)"), ShortViewNoNull == WideStringSrc);
+	TestTrue(TEXT("Equality(55)"), ShortViewNoNull != WideStringSrc);
+	TestFalse(TEXT("Equality(56)"), WideStringSrc == ShortViewNoNull);
+	TestTrue(TEXT("Equality(57)"), WideStringSrc != ShortViewNoNull);
+
+	FStringView WideViewNoNull = FStringView(WideStringLiteralLonger).Left(WideView.Len());
+
+	TestTrue(TEXT("Equality(58)"), WideViewNoNull == WideStringLiteralSrc);
+	TestFalse(TEXT("Equality(59)"), WideViewNoNull != WideStringLiteralSrc);
+	TestTrue(TEXT("Equality(60)"), WideStringLiteralSrc == WideViewNoNull);
+	TestFalse(TEXT("Equality(61)"), WideStringLiteralSrc != WideViewNoNull);
+	TestFalse(TEXT("Equality(62)"), WideViewNoNull == WideStringLiteralLonger);
+	TestTrue(TEXT("Equality(63)"), WideViewNoNull != WideStringLiteralLonger);
+	TestFalse(TEXT("Equality(64)"), WideStringLiteralLonger == WideViewNoNull);
+	TestTrue(TEXT("Equality(65)"), WideStringLiteralLonger != WideViewNoNull);
+
+	TestTrue(TEXT("Equality(66)"), WideViewNoNull == WideStringLiteralSrc);
+	TestFalse(TEXT("Equality(67)"), WideViewNoNull != WideStringLiteralSrc);
+	TestTrue(TEXT("Equality(68)"), WideStringLiteralSrc == WideViewNoNull);
+	TestFalse(TEXT("Equality(69)"), WideStringLiteralSrc != WideViewNoNull);
+	TestFalse(TEXT("Equality(70)"), WideViewNoNull == WideStringLiteralLonger);
+	TestTrue(TEXT("Equality(71)"), WideViewNoNull != WideStringLiteralLonger);
+	TestFalse(TEXT("Equality(72)"), WideStringLiteralLonger == WideViewNoNull);
+	TestTrue(TEXT("Equality(73)"), WideStringLiteralLonger != WideViewNoNull);
 
 	return true;
 }
@@ -624,61 +690,6 @@ bool FStringViewTestFindLastChar::RunTest(const FString& Parameters)
 	return true; 
 }
 
-IMPLEMENT_SIMPLE_AUTOMATION_TEST(FStringViewTestFindLastCharByPredicate, TEST_NAME_ROOT ".FindLastCharByPredicate", TestFlags)
-bool FStringViewTestFindLastCharByPredicate::RunTest(const FString& Parameters)
-{
-	auto AcceptMoneyCharacters = [](TCHAR InputChar) { return InputChar == TEXT('\u00a3') // pound sign
-															|| InputChar == TEXT('$')
-															|| InputChar == TEXT('\u20ac'); // euro sign
-	};
-
-	FStringView ViewDollars			= TEXT("The currency we have is in $");
-	FStringView ViewPounds			= TEXT("Another currency £ could be found?");
-	FStringView ViewMixedCurrency	= TEXT("We have both £,$ and € to spend");
-	FStringView ViewNoCurreny		= TEXT("We are a day late and a dollar short");
-	FStringView ViewEmpty			= TEXT("");
-
-	{
-		FStringView::SizeType Result = ViewDollars.FindLastCharByPredicate(AcceptMoneyCharacters);
-		TestTrue(TEXT("FStringView::FindLastCharByPredicate(0)"), Result == 27);
-
-		FStringView::SizeType ResultWithLimits = ViewDollars.FindLastCharByPredicate(AcceptMoneyCharacters, 27);
-		TestTrue(TEXT("FStringView::FindLastCharByPredicate(1)"), ResultWithLimits == INDEX_NONE);
-	}
-
-	{
-		FStringView::SizeType Result = ViewPounds.FindLastCharByPredicate(AcceptMoneyCharacters);
-		TestTrue(TEXT("FStringView::FindLastCharByPredicate(2)"), Result == 17);
-
-		FStringView::SizeType ResultWithLimits = ViewPounds.FindLastCharByPredicate(AcceptMoneyCharacters, 17);
-		TestTrue(TEXT("FStringView::FindLastCharByPredicate(3)"), ResultWithLimits == INDEX_NONE);
-	}
-
-	{
-		FStringView::SizeType EuroPosition = ViewMixedCurrency.FindLastCharByPredicate(AcceptMoneyCharacters);
-		TestTrue(TEXT("FStringView::FindLastCharByPredicate(4)"), EuroPosition == 21);
-
-		FStringView::SizeType DollarPosition = ViewMixedCurrency.FindLastCharByPredicate(AcceptMoneyCharacters, 21);
-		TestTrue(TEXT("FStringView::FindLastCharByPredicate(5)"), DollarPosition == 15);
-
-		FStringView::SizeType PoundPosition = ViewMixedCurrency.FindLastCharByPredicate(AcceptMoneyCharacters, 15);
-		TestTrue(TEXT("FStringView::FindLastCharByPredicate(6)"), PoundPosition == 13);
-
-		FStringView::SizeType Result = ViewMixedCurrency.FindLastCharByPredicate(AcceptMoneyCharacters, 13);
-		TestTrue(TEXT("FStringView::FindLastCharByPredicate(7)"), Result == INDEX_NONE);
-	}
-
-	{
-		FStringView::SizeType Result = ViewNoCurreny.FindLastCharByPredicate(AcceptMoneyCharacters);
-		TestTrue(TEXT("FStringView::FindLastCharByPredicate(8)"), Result == INDEX_NONE);
-	}
-	
-	{
-		FStringView::SizeType Result = ViewEmpty.FindLastCharByPredicate(AcceptMoneyCharacters);
-		TestTrue(TEXT("FStringView::FindLastCharByPredicate(9)"), Result == INDEX_NONE);
-	}
-
-	return true;
-}
+#undef TEST_NAME_ROOT
 
 #endif //WITH_DEV_AUTOMATION_TESTS
diff --git a/Engine/Source/Runtime/Core/Public/Algo/FindLast.h b/Engine/Source/Runtime/Core/Public/Algo/FindLast.h
new file mode 100644
--- /dev/null
+++ b/Engine/Source/Runtime/Core/Public/Algo/FindLast.h
@@ -0,0 +1,87 @@
+// Copyright 1998-2019 Epic Games, Inc. All Rights Reserved.
+
+#pragma once
+
+#include "Templates/IdentityFunctor.h"
+#include "Templates/Invoke.h"
+#include "Templates/UnrealTemplate.h" // For MoveTemp
+
+namespace AlgoImpl
+{
+	template <typename T, typename ValueType, typename ProjectionType>
+	T* FindLastBy(T* First, SIZE_T Num, const ValueType& Value, ProjectionType Proj)
+	{
+		for (T* Last = First + Num; First != Last;)
+		{
+			if (Invoke(Proj, *--Last) == Value)
+			{
+				return Last;
+			}
+		}
+
+		return nullptr;
+	}
+
+	template <typename T, typename PredicateType>
+	T* FindLastByPredicate(T* First, SIZE_T Num, PredicateType Pred)
+	{
+		for (T* Last = First + Num; First != Last;)
+		{
+			if (Invoke(Pred, *--Last))
+			{
+				return Last;
+			}
+		}
+
+		return nullptr;
+	}
+}
+
+namespace Algo
+{
+	/**
+	 * Returns a pointer to the last element in the range which is equal to the given value.
+	 *
+	 * @param  Range  The range to search.
+	 * @param  Value  The value to search for.
+	 *
+	 * @return A pointer to the last element found, or nullptr if none was found.
+	 */
+	template <typename RangeType, typename ValueType>
+	FORCEINLINE auto FindLast(RangeType&& Range, const ValueType& Value)
+		-> decltype(AlgoImpl::FindLastBy(GetData(Range), GetNum(Range), Value, FIdentityFunctor()))
+	{
+		return AlgoImpl::FindLastBy(GetData(Range), GetNum(Range), Value, FIdentityFunctor());
+	}
+
+	/**
+	 * Returns a pointer to the last element in the range whose projection is equal to the given value.
+	 *
+	 * @param  Range  The range to search.
+	 * @param  Value  The value to search for.
+	 * @param  Proj   The projection to apply to the element.
+	 *
+	 * @return A pointer to the last element found, or nullptr if none was found.
+	 */
+	template <typename RangeType, typename ValueType, typename ProjectionType>
+	FORCEINLINE auto FindLastBy(RangeType&& Range, const ValueType& Value, ProjectionType Proj)
+		-> decltype(AlgoImpl::FindLastBy(GetData(Range), GetNum(Range), Value, MoveTemp(Proj)))
+	{
+		return AlgoImpl::FindLastBy(GetData(Range), GetNum(Range), Value, MoveTemp(Proj));
+	}
+
+	/**
+	 * Returns a pointer to the last element in the range which matches the predicate.
+	 *
+	 * @param  Range  The range to search.
+	 * @param  Pred   The predicate to search for.
+	 *
+	 * @return A pointer to the last element found, or nullptr if none was found.
+	 */
+	template <typename RangeType, typename PredicateType>
+	FORCEINLINE auto FindLastByPredicate(RangeType&& Range, PredicateType Pred)
+		-> decltype(AlgoImpl::FindLastByPredicate(GetData(Range), GetNum(Range), MoveTemp(Pred)))
+	{
+		return AlgoImpl::FindLastByPredicate(GetData(Range), GetNum(Range), MoveTemp(Pred));
+	}
+}
diff --git a/Engine/Source/Runtime/Core/Public/Containers/UnrealString.h b/Engine/Source/Runtime/Core/Public/Containers/UnrealString.h
--- a/Engine/Source/Runtime/Core/Public/Containers/UnrealString.h
+++ b/Engine/Source/Runtime/Core/Public/Containers/UnrealString.h
@@ -25,34 +25,20 @@
 #include "Templates/TypeHash.h"
 #include "Templates/IsFloatingPoint.h"
 
+class FString;
+class FStringView;
 struct FStringFormatArg;
 template<typename KeyType,typename ValueType,typename SetAllocator ,typename KeyFuncs > class TMap;
 
-/** Determines case sensitivity options for string comparisons. */
-namespace ESearchCase
+template <>
+struct TIsContiguousContainer<FString>
 {
-	enum Type
-	{
-		/** Case sensitive. Upper/lower casing must match for strings to be considered equal. */
-		CaseSensitive,
-
-		/** Ignore case. Upper/lower casing does not matter when making a comparison. */
-		IgnoreCase,
-	};
+	enum { Value = true };
 };
 
-/** Determines search direction for string operations. */
-namespace ESearchDir
-{
-	enum Type
-	{
-		/** Search from the start, moving forward through the string. */
-		FromStart,
-
-		/** Search from the end, moving backward through the string. */
-		FromEnd,
-	};
-}
+TCHAR* GetData(FString& String);
+const TCHAR* GetData(const FString& String);
+SIZE_T GetNum(const FString& String);
 
 /**
  * A dynamically sizeable string.
@@ -144,6 +130,8 @@ public:
 		}
 	}
 
+	explicit FString(const FStringView& Other);
+
 #ifdef __OBJC__
 	/** Convert Objective-C NSString* to FString */
 	FORCEINLINE FString(const NSString* In)
@@ -192,6 +180,16 @@ public:
 		return *this;
 	}
 
+	/**
+	 * Copy assignment from a string view
+	 */
+	FString& operator=(const FStringView& Other);
+
+	/**
+	 * Implicit conversion to a string view
+	 */
+	operator FStringView() const;
+
 	/**
 	 * Return specific character from this string
 	 *
@@ -1904,12 +1902,6 @@ struct TContainerTraits<FString> : public TContainerTraitsBase<FString>
 template<> struct TIsZeroConstructType<FString> { enum { Value = true }; };
 Expose_TNameOf(FString)
 
-template <>
-struct TIsContiguousContainer<FString>
-{
-	enum { Value = true };
-};
-
 inline TCHAR* GetData(FString& String)
 {
 	return String.GetCharArray().GetData();
diff --git a/Engine/Source/Runtime/Core/Public/Misc/CString.h b/Engine/Source/Runtime/Core/Public/Misc/CString.h
--- a/Engine/Source/Runtime/Core/Public/Misc/CString.h
+++ b/Engine/Source/Runtime/Core/Public/Misc/CString.h
@@ -13,6 +13,32 @@
 
 #define MAX_SPRINTF 1024
 
+/** Determines case sensitivity options for string comparisons. */
+namespace ESearchCase
+{
+	enum Type
+	{
+		/** Case sensitive. Upper/lower casing must match for strings to be considered equal. */
+		CaseSensitive,
+
+		/** Ignore case. Upper/lower casing does not matter when making a comparison. */
+		IgnoreCase,
+	};
+};
+
+/** Determines search direction for string operations. */
+namespace ESearchDir
+{
+	enum Type
+	{
+		/** Search from the start, moving forward through the string. */
+		FromStart,
+
+		/** Search from the end, moving backward through the string. */
+		FromEnd,
+	};
+}
+
 /** Helper class used to convert CString into a boolean value. */
 struct CORE_API FToBoolHelper
 {
diff --git a/Engine/Source/Runtime/Core/Public/Misc/StringBuilder.h b/Engine/Source/Runtime/Core/Public/Misc/StringBuilder.h
--- a/Engine/Source/Runtime/Core/Public/Misc/StringBuilder.h
+++ b/Engine/Source/Runtime/Core/Public/Misc/StringBuilder.h
@@ -93,7 +93,7 @@ public:
 
 	inline TStringBuilderImpl& AppendAnsi(const FAnsiStringView& AnsiString)
 	{
-		return AppendAnsi(AnsiString.Data(), AnsiString.Len());
+		return AppendAnsi(AnsiString.GetData(), AnsiString.Len());
 	}
 
 	inline TStringBuilderImpl& AppendAnsi(const ANSICHAR* NulTerminatedString, const int32 Length)
@@ -285,7 +285,7 @@ public:
 		EnsureCapacity(StringLength);
 		TCHAR* RESTRICT Dest = CurPos;
 
-		FMemory::Memcpy(CurPos, StringView.Data(), StringLength * sizeof(TCHAR));
+		FMemory::Memcpy(CurPos, StringView.GetData(), StringLength * sizeof(TCHAR));
 
 		CurPos += StringLength;
 
@@ -354,4 +354,4 @@ inline FStringBuilderBase&					operator<<(FStringBuilderBase& Builder, const ANS
 
 #endif
 
-inline FStringBuilderBase&					operator<<(FStringBuilderBase& Builder, const FStringView& Str)					{ Builder.Append(Str.Data(), Str.Len()); return Builder; }
+inline FStringBuilderBase&					operator<<(FStringBuilderBase& Builder, const FStringView& Str)					{ Builder.Append(Str.GetData(), Str.Len()); return Builder; }
diff --git a/Engine/Source/Runtime/Core/Public/Misc/StringView.h b/Engine/Source/Runtime/Core/Public/Misc/StringView.h
--- a/Engine/Source/Runtime/Core/Public/Misc/StringView.h
+++ b/Engine/Source/Runtime/Core/Public/Misc/StringView.h
@@ -3,19 +3,36 @@
 #pragma once
 
 #include "CoreTypes.h"
+#include "Math/UnrealMathUtility.h"
 #include "Misc/CString.h"
-#include "Containers/UnrealString.h"
+#include "Templates/AndOrNot.h"
+#include "Templates/EnableIf.h"
+#include "Templates/RemoveCV.h"
+#include "Templates/UnrealTemplate.h"
+#include "Traits/IsContiguousContainer.h"
+
+namespace StringViewPrivate
+{
+	/**
+	 * Trait testing whether a range type is compatible with the view type
+	 */
+	template <typename RangeType, typename ElementType>
+	struct TIsCompatibleRangeType : TIsSame<ElementType, typename TRemoveCV<typename TRemovePointer<decltype(GetData(DeclVal<RangeType&>()))>::Type>::Type>
+	{
+	};
+}
 
 /** String View
 
-	* A string view is implicitly constructible from FString and const char* style strings
+	* A string view is implicitly constructible from const char* style strings and
+	  from character ranges such as FString.
 
 	* A string view does not own any data nor does it attempt to control any lifetimes, it 
 	  merely points at a subrange of characters in some other string. It's up to the user
 	  to ensure the underlying string stays valid for the lifetime of the string view.
 
 	* A string view does not represent a NUL terminated string and therefore you should
-	  never pass in the pointer returned by Data() into a C-string API accepting only a
+	  never pass in the pointer returned by GetData() into a C-string API accepting only a
 	  pointer. You must either use a string builder to make a properly terminated string,
 	  or the ToString() function below if you absolutely must, and can live with the
 	  knowledge that you just added yet one more memory allocation even though memory
@@ -48,66 +65,45 @@
 
   */
 
-template<typename C>
+template <typename CharType, typename DerivedViewType>
 class TStringViewImpl
 {
 public:
-	typedef int32 SizeType;
+	using ElementType = CharType;
+	using SizeType = int32;
+	using ViewType = DerivedViewType;
 
-	inline TStringViewImpl(const C* InData)
-	:	DataPtr(InData)
-	,	Size(TCString<C>::Strlen(InData))
-	{
-	}
+	template <typename RangeType>
+	using TIsCompatibleRangeType = TAnd<TIsContiguousContainer<RangeType>, StringViewPrivate::TIsCompatibleRangeType<RangeType, ElementType>>;
 
-	template<	typename CharType,
-				typename = typename TEnableIf<TIsCharType<CharType>::Value>::Type>
-	TStringViewImpl(const CharType * InData)
-	{
-		static_assert(TIsCharType<CharType>::Value, "Trying to assign an incompatible string type to the TStringViewImpl!");
-	}
+	template <typename RangeType>
+	using TEnableIfCompatibleRangeType = typename TEnableIf<TIsCompatibleRangeType<RangeType>::Value>::Type;
 
-	inline TStringViewImpl(const C* InData, SizeType InSize)
-	:	DataPtr(InData)
-	,	Size(InSize)
-	{
-	}
+	TStringViewImpl() = default;
 
-	template<	typename CharType,
-				typename = typename TEnableIf<TIsCharType<CharType>::Value>::Type>
-	TStringViewImpl(const CharType * InData, SizeType InSize)
+	TStringViewImpl(const CharType* InData)
+		: DataPtr(InData)
+		, Size(TCString<CharType>::Strlen(InData))
 	{
-		static_assert(TIsCharType<CharType>::Value, "Trying to assign an incompatible string type to the TStringViewImpl!");
 	}
 
-	// Constructor for FString where the StringView is of the correct type
-	template<	typename U = C,
-				typename = typename TEnableIf< TIsSame<U, FString::ElementType>::Value >::Type >
-	TStringViewImpl(const FString& InString)
-	:	DataPtr(*InString)
-	,	Size(InString.Len())
+	TStringViewImpl(const CharType* InData, SizeType InSize)
+		: DataPtr(InData)
+		, Size(InSize)
 	{
-
 	}
 
-	// Constructor for FString where the StringView is of the incorrect type
-	// In this case we give a clear static assert as to why the code will not compile.
-	template<	typename U = C,
-				typename = typename TEnableIf< !TIsSame<U, FString::ElementType>::Value >::Type,
-				typename = U >
-	TStringViewImpl(const FString & InString)
-	{
-		static_assert(!TIsSame<U, FString::ElementType>::Value, "Can only use FString with a TStringViewImpl of the same type!");
-	}
+	inline const CharType& operator[](SizeType Pos) const { return DataPtr[Pos]; }
 
-	inline const C& operator[](SizeType Pos) const	{ return DataPtr[Pos]; }
+	inline const CharType* GetData() const { return DataPtr; }
 
-	inline const C* Data() const					{ return DataPtr; }
+	UE_DEPRECATED(4.25, "'Data' is deprecated. Please use 'GetData' instead!")
+	inline const CharType* Data() const { return DataPtr; }
 
 	// Capacity
 
-	inline SizeType	Len() const						{ return Size; }
-	inline bool		IsEmpty() const					{ return Size == 0; }
+	inline SizeType Len() const { return Size; }
+	inline bool IsEmpty() const { return Size == 0; }
 
 	// Modifiers
 
@@ -116,66 +112,73 @@ public:
 
 	// Operations
 
-	inline SIZE_T			CopyString(C* Dest, SizeType CharCount, SizeType Position = 0) const	{ memcpy(Dest, DataPtr + Position, FMath::Min(Size - Position, CharCount)); }
-	inline TStringViewImpl	SubStr(SizeType Position, SizeType CharCount) const						{ check(Position <= Len()); return TStringViewImpl(DataPtr + Position, FMath::Min(Size - Position, CharCount)); }
+	inline SizeType CopyString(CharType* Dest, SizeType CharCount, SizeType Position = 0) const;
+	inline ViewType SubStr(SizeType Position, SizeType CharCount) const						{ check(Position <= Len()); return ViewType(DataPtr + Position, FMath::Min(Size - Position, CharCount)); }
 
 	// Maintain compatibility with FString 
-	inline TStringViewImpl	Left(SizeType CharCount) const											{ return TStringViewImpl(Data(), FMath::Clamp(CharCount, 0, Len())); }
-	inline TStringViewImpl  LeftChop(SizeType CharCount) const										{ return TStringViewImpl(Data(), FMath::Clamp(Len() - CharCount, 0, Len())); }
-	inline TStringViewImpl  Right(SizeType CharCount) const											{ return TStringViewImpl(Data() + Len() - FMath::Clamp(CharCount, 0, Len())); }
-	inline TStringViewImpl	RightChop(SizeType CharCount) const										{ return TStringViewImpl(Data() + Len() - FMath::Clamp(Len() - CharCount, 0, Len())); }
-	inline TStringViewImpl  Mid(SizeType Position, SizeType CharCount) const						{ return SubStr(Position, CharCount); /*This is just a wrapper around SubStr to keep compatibility with the FString interface*/} 
-	
+	inline ViewType Left(SizeType CharCount) const											{ return ViewType(GetData(), FMath::Clamp(CharCount, 0, Len())); }
+	inline ViewType LeftChop(SizeType CharCount) const										{ return ViewType(GetData(), FMath::Clamp(Len() - CharCount, 0, Len())); }
+	inline ViewType Right(SizeType CharCount) const											{ return ViewType(GetData() + Len() - FMath::Clamp(CharCount, 0, Len())); }
+	inline ViewType RightChop(SizeType CharCount) const										{ return ViewType(GetData() + Len() - FMath::Clamp(Len() - CharCount, 0, Len())); }
+	inline ViewType Mid(SizeType Position, SizeType CharCount) const						{ return SubStr(Position, CharCount); /*This is just a wrapper around SubStr to keep compatibility with the FString interface*/} 
+
 	// Comparison
-	inline bool				Equals(const TStringViewImpl Other, ESearchCase::Type SearchCase) const	{ return Compare(Other, SearchCase) == 0; }
+	inline bool				Equals(const TStringViewImpl& Other, ESearchCase::Type SearchCase) const { return Len() == Other.Len() && Compare(Other, SearchCase) == 0; }
+	CORE_API int32			Compare(const TStringViewImpl& Other, ESearchCase::Type SearchCase) const;
 
-	inline int32			Compare(const TStringViewImpl Other, ESearchCase::Type SearchCase) const;
+	inline bool StartsWith(CharType Prefix) const										{ return Size >= 1 && DataPtr[0] == Prefix; }
+	inline bool StartsWith(const TStringViewImpl& Prefix) const							{ return Len() >= Prefix.Len() && TCString<CharType>::Strnicmp(GetData(), Prefix.GetData(), Prefix.Len()) == 0; }
 
-	inline bool				StartsWith(C Prefix) const												{ return Size >= 1 && DataPtr[0] == Prefix; }
-	inline bool				StartsWith(const TStringViewImpl& Prefix) const							{ return 0 == TCString<C>::Strnicmp(Prefix.Data(), this->Data(), Prefix.Len()); }
-	
-	inline bool				EndsWith(C Suffix) const												{ return Size >= 1 && DataPtr[Size-1] == Suffix; }
-	inline bool				EndsWith(const TStringViewImpl& Suffix) const;
+	inline bool EndsWith(CharType Suffix) const											{ return Size >= 1 && DataPtr[Size-1] == Suffix; }
+	inline bool EndsWith(const TStringViewImpl& Suffix) const							{ return Len() >= Suffix.Len() && TCString<CharType>::Strnicmp(GetData() + Len() - Suffix.Len(), Suffix.GetData(), Suffix.Len()) == 0; }
 
 	// Searching/Finding
-	inline bool FindChar(C InChar, int32& OutIndex) const;
-	inline bool FindLastChar(C InChar, int32& OutIndex) const;
-
-	template <typename Predicate> inline SizeType FindLastCharByPredicate(Predicate Pred, SizeType Count) const;
-	template <typename Predicate> inline SizeType FindLastCharByPredicate(Predicate Pred) const		{ return FindLastCharByPredicate(Pred, Len()); }
-
-	inline const C* begin() const					{ return DataPtr; }
-	inline const C* end() const						{ return DataPtr + Size; }
-
-	/** Convert to a dynamic string
-
-		Remember that a string view is not necessarily
-		NUL terminated so whenever you end up passing a
-		string into an API which accepts a plain C string
-		then you will need to ensure it's NUL terminated
-		first. This is one way of doing that, involving
-		a memory allocation as well as a string copy,
-		along with a memory deallocation when the string
-		goes away.
-
-		I would encourage you to use a TStringBuilder<>
-		instead since this avoids hitting the heap, and
-		the heap is always more expensive even if it
-		seems seductively convenient.
-	  */
-	FString ToString() const
-	{
-		return FString(Len(), Data());
-	}
+	CORE_API bool FindChar(CharType InChar, SizeType& OutIndex) const;
+	CORE_API bool FindLastChar(CharType InChar, SizeType& OutIndex) const;
+
+public:
+	/**
+	 * DO NOT USE DIRECTLY
+	 * STL-like iterators to enable range-based for loop support.
+	 */
+	inline const CharType* begin() const { return DataPtr; }
+	inline const CharType* end() const { return DataPtr + Size; }
+
 protected:
-	const C*	DataPtr;
-	SizeType	Size;
+	const CharType* DataPtr = nullptr;
+	SizeType Size = 0;
 };
 
+template <typename CharType, typename ViewType>
+struct TIsContiguousContainer<TStringViewImpl<CharType, ViewType>>
+{
+	enum { Value = true };
+};
+
+template <typename CharType, typename ViewType>
+inline SIZE_T GetNum(const TStringViewImpl<CharType, ViewType>& String)
+{
+	return String.Len();
+}
+
 //////////////////////////////////////////////////////////////////////////
 
-using FAnsiStringView = TStringViewImpl<ANSICHAR>;
-using FWideStringView = TStringViewImpl<WIDECHAR>;
-using FStringView = TStringViewImpl<TCHAR>;
+class FStringView : public TStringViewImpl<TCHAR, FStringView>
+{
+public:
+	using TStringViewImpl<ElementType, FStringView>::TStringViewImpl;
+};
+
+class FAnsiStringView : public TStringViewImpl<ANSICHAR, FAnsiStringView>
+{
+public:
+	using TStringViewImpl<ElementType, FAnsiStringView>::TStringViewImpl;
+};
+
+class FWideStringView : public TStringViewImpl<WIDECHAR, FWideStringView>
+{
+public:
+	using TStringViewImpl<ElementType, FWideStringView>::TStringViewImpl;
+};
 
 #include "StringView.inl"
diff --git a/Engine/Source/Runtime/Core/Public/Misc/StringView.inl b/Engine/Source/Runtime/Core/Public/Misc/StringView.inl
--- a/Engine/Source/Runtime/Core/Public/Misc/StringView.inl
+++ b/Engine/Source/Runtime/Core/Public/Misc/StringView.inl
@@ -1,109 +1,78 @@
 // Copyright 1998-2019 Epic Games, Inc. All Rights Reserved.
 
-// Case-insensitive comparison operators
-template<typename C>
-inline bool operator==(const TStringViewImpl<C> Lhs, typename TIdentity<TStringViewImpl<C>>::Type Rhs)
-{
-	if (Rhs.Len() != Lhs.Len())
-	{
-		return false;
-	}
+#include "HAL/PlatformString.h"
+#include "HAL/UnrealMemory.h"
 
-	return FPlatformString::Strnicmp(Lhs.Data(), Rhs.Data(), Lhs.Len()) == 0;
+template <typename CharType, typename ViewType>
+inline typename TStringViewImpl<CharType, ViewType>::SizeType TStringViewImpl<CharType, ViewType>::CopyString(CharType* Dest, SizeType CharCount, SizeType Position) const
+{
+	SizeType CopyCount = FMath::Min(Size - Position, CharCount);
+	FMemory::Memcpy(Dest, DataPtr + Position, CopyCount);
+	return CopyCount;
 }
 
-template<typename C>
-inline bool operator!=(const TStringViewImpl<C> Lhs, typename TIdentity<TStringViewImpl<C>>::Type Rhs)
+// Case-insensitive comparison operators
+
+template <typename CharType, typename ViewType>
+inline bool operator==(const TStringViewImpl<CharType, ViewType>& Lhs, const TStringViewImpl<CharType, ViewType>& Rhs)
 {
-	return !operator==(Lhs, Rhs);
+	return Lhs.Len() == Rhs.Len() && FPlatformString::Strnicmp(Lhs.GetData(), Rhs.GetData(), Lhs.Len()) == 0;
 }
 
-template<typename C>
-inline int32 TStringViewImpl<C>::Compare(const TStringViewImpl Other, ESearchCase::Type SearchCase) const
+template <typename CharType, typename ViewType>
+inline bool operator!=(const TStringViewImpl<CharType, ViewType>& Lhs, const TStringViewImpl<CharType, ViewType>& Rhs)
 {
-	const SizeType ShortestLength = FMath::Min(Len(), Other.Len());
+	return !operator==(Lhs, Rhs);
+}
 
-	int Result;
-	if (SearchCase == ESearchCase::CaseSensitive)
-	{
-		Result = FPlatformString::Strncmp(Data(), Other.Data(), ShortestLength);
-	}
-	else
-	{
-		Result = FPlatformString::Strnicmp(Data(), Other.Data(), ShortestLength);
-	}
+// Case-insensitive range comparison operators
 
-	if (Result == 0 && Len() != Other.Len())
-	{
-		return Len() > Other.Len() ? 1 : -1;
-	}
-	else
-	{
-		return Result;
-	}
+template <typename CharType, typename ViewType, typename RangeType, typename = typename TStringViewImpl<CharType, ViewType>::template TEnableIfCompatibleRangeType<RangeType>>
+inline bool operator==(const TStringViewImpl<CharType, ViewType>& Lhs, const RangeType& Rhs)
+{
+	return Lhs == TStringViewImpl<CharType, ViewType>(Rhs);
 }
 
-template<typename C>
-inline bool TStringViewImpl<C>::EndsWith(const TStringViewImpl& Suffix) const
+template <typename CharType, typename ViewType, typename RangeType, typename = typename TStringViewImpl<CharType, ViewType>::template TEnableIfCompatibleRangeType<RangeType>>
+inline bool operator==(const RangeType& Lhs, const TStringViewImpl<CharType, ViewType>& Rhs)
 {
-	const SizeType SuffixLen = Suffix.Len();
-	if (SuffixLen > this->Len())
-	{
-		return false;
-	}
-
-	return 0 == TCString<C>::Strnicmp(Suffix.Data(), this->Data() + this->Len() - SuffixLen, SuffixLen);
+	return TStringViewImpl<CharType, ViewType>(Lhs) == Rhs;
 }
 
-template<typename C>
-inline bool TStringViewImpl<C>::FindChar(C InChar, int32& OutIndex) const
+template <typename CharType, typename ViewType, typename RangeType, typename = typename TStringViewImpl<CharType, ViewType>::template TEnableIfCompatibleRangeType<RangeType>>
+inline bool operator!=(const TStringViewImpl<CharType, ViewType>& Lhs, const RangeType& Rhs)
 {
-	for (SizeType i = 0; i < Size; ++i)
-	{
-		if (DataPtr[i] == InChar)
-		{
-			OutIndex = (int32)i;
-			return true;
-		}
-	}
-
-	return false;
+	return !operator==(Lhs, Rhs);
 }
 
-template<typename C>
-inline bool TStringViewImpl<C>::FindLastChar(C InChar, int32& OutIndex) const
+template <typename CharType, typename ViewType, typename RangeType, typename = typename TStringViewImpl<CharType, ViewType>::template TEnableIfCompatibleRangeType<RangeType>>
+inline bool operator!=(const RangeType& Lhs, const TStringViewImpl<CharType, ViewType>& Rhs)
 {
-	if (Size == 0)
-	{
-		return false;
-	}
+	return !operator==(Lhs, Rhs);
+}
 
-	for (SizeType i = Size - 1; i >= 0; --i)
-	{
-		if (DataPtr[i] == InChar)
-		{
-			OutIndex = (int32)i;
-			return true;
-		}
-	}
+// Case-insensitive C-string comparison operators
 
-	return false;
+template <typename CharType, typename ViewType>
+inline bool operator==(const TStringViewImpl<CharType, ViewType>& Lhs, const CharType* Rhs)
+{
+	return TCString<CharType>::Strnicmp(Lhs.GetData(), Rhs, Lhs.Len()) == 0 && !Rhs[Lhs.Len()];
 }
 
-template<typename C>
-template<typename Predicate>
-inline typename TStringViewImpl<C>::SizeType TStringViewImpl<C>::FindLastCharByPredicate(Predicate Pred, TStringViewImpl<C>::SizeType Count) const
+template <typename CharType, typename ViewType>
+inline bool operator==(const CharType* Lhs, const TStringViewImpl<CharType, ViewType>& Rhs)
 {
-	check(Count >= 0 && Count <= Len());
+	return TCString<CharType>::Strnicmp(Lhs, Rhs.GetData(), Rhs.Len()) == 0 && !Lhs[Rhs.Len()];
+}
 
-	for (const C* RESTRICT Start = Data(), *RESTRICT Data = Start + Count; Data != Start; )
-	{
-		--Data;
-		if (Pred(*Data))
-		{
-			return static_cast<SizeType>(Data - Start);
-		}
-	}
+template <typename CharType, typename ViewType>
+inline bool operator!=(const TStringViewImpl<CharType, ViewType>& Lhs, const CharType* Rhs)
+{
+	return !operator==(Lhs, Rhs);
+}
 
-	return INDEX_NONE;
+template <typename CharType, typename ViewType>
+inline bool operator!=(const CharType* Lhs, const TStringViewImpl<CharType, ViewType>& Rhs)
+{
+	return !operator==(Lhs, Rhs);
 }
diff --git a/Engine/Source/Runtime/Core/Public/Traits/ElementType.h b/Engine/Source/Runtime/Core/Public/Traits/ElementType.h
--- a/Engine/Source/Runtime/Core/Public/Traits/ElementType.h
+++ b/Engine/Source/Runtime/Core/Public/Traits/ElementType.h
@@ -3,6 +3,7 @@
 #pragma once
 
 #include <initializer_list>
+#include "Templates/RemoveCV.h"
 
 /**
  * Traits class which gets the element type of a container.
