Update remaining uses of the newly deprecated FStringView::Data()

#rnx

diff --git a/Engine/Source/Runtime/Core/Private/Containers/String.cpp b/Engine/Source/Runtime/Core/Private/Containers/String.cpp
--- a/Engine/Source/Runtime/Core/Private/Containers/String.cpp
+++ b/Engine/Source/Runtime/Core/Private/Containers/String.cpp
@@ -123,14 +123,14 @@ FString::FString(const FStringView& Other)
 	if (const FStringView::SizeType OtherLen = Other.Len())
 	{
 		Reserve(OtherLen);
-		Append(Other.Data(), OtherLen);
+		Append(Other.GetData(), OtherLen);
 		AppendChar(TEXT('\0'));
 	}
 }
 
 FString& FString::operator=(const FStringView& Other)
 {
-	const TCHAR* const OtherData = Other.Data();
+	const TCHAR* const OtherData = Other.GetData();
 	const FStringView::SizeType OtherLen = Other.Len();
 	if (OtherLen == 0)
 	{
diff --git a/Engine/Source/Runtime/Core/Private/IO/IoStatus.cpp b/Engine/Source/Runtime/Core/Private/IO/IoStatus.cpp
--- a/Engine/Source/Runtime/Core/Private/IO/IoStatus.cpp
+++ b/Engine/Source/Runtime/Core/Private/IO/IoStatus.cpp
@@ -50,7 +50,7 @@ FIoStatus::FIoStatus(EIoErrorCode Code, const FStringView& InErrorMessage)
 : ErrorCode(Code)
 {
 	const int32 ErrorMessageLength = FMath::Min(MaxErrorMessageLength - 1, InErrorMessage.Len());
-	FPlatformString::Convert(ErrorMessage, ErrorMessageLength, InErrorMessage.Data(), ErrorMessageLength);
+	FPlatformString::Convert(ErrorMessage, ErrorMessageLength, InErrorMessage.GetData(), ErrorMessageLength);
 	ErrorMessage[ErrorMessageLength] = 0;
 }
 
@@ -96,7 +96,7 @@ FIoStatusBuilder::FIoStatusBuilder(EIoErrorCode InStatusCode)
 FIoStatusBuilder::FIoStatusBuilder(const FIoStatus& InStatus, FStringView String)
 :	StatusCode(InStatus.ErrorCode)
 {
-	Message.Append(String.Data(), String.Len());
+	Message.Append(String.GetData(), String.Len());
 }
 
 FIoStatusBuilder::~FIoStatusBuilder()
@@ -110,7 +110,7 @@ FIoStatusBuilder::operator FIoStatus()
 
 FIoStatusBuilder& FIoStatusBuilder::operator<<(FStringView String)
 {
-	Message.Append(String.Data(), String.Len());
+	Message.Append(String.GetData(), String.Len());
 
 	return *this;
 }
diff --git a/Engine/Source/Runtime/Core/Private/Tests/Misc/StringViewTest.cpp b/Engine/Source/Runtime/Core/Private/Tests/Misc/StringViewTest.cpp
--- a/Engine/Source/Runtime/Core/Private/Tests/Misc/StringViewTest.cpp
+++ b/Engine/Source/Runtime/Core/Private/Tests/Misc/StringViewTest.cpp
@@ -24,7 +24,7 @@ bool FStringViewTestCtor::RunTest(const FString& Parameters)
 	{
 		FStringView View(TEXT("Test Ctor"));
 		TestEqual(TEXT("View length"), View.Len(), FCStringWide::Strlen(TEXT("Test Ctor")));
-		TestEqual(TEXT("The result of Strncmp"), FCStringWide::Strncmp(View.Data(), TEXT("Test Ctor"), View.Len()), 0);
+		TestEqual(TEXT("The result of Strncmp"), FCStringWide::Strncmp(View.GetData(), TEXT("Test Ctor"), View.Len()), 0);
 		TestFalse("View.IsEmpty", View.IsEmpty());
 	}
 
@@ -32,7 +32,7 @@ bool FStringViewTestCtor::RunTest(const FString& Parameters)
 	{
 		FStringView View(TEXT("Test SubSection Ctor"), 4);
 		TestEqual(TEXT("View length"), View.Len(), 4);
-		TestEqual(TEXT("The result of Strncmp"), FCStringWide::Strncmp(View.Data(), TEXT("Test"), View.Len()), 0);
+		TestEqual(TEXT("The result of Strncmp"), FCStringWide::Strncmp(View.GetData(), TEXT("Test"), View.Len()), 0);
 		TestFalse("View.IsEmpty", View.IsEmpty());
 	}
 
@@ -42,7 +42,7 @@ bool FStringViewTestCtor::RunTest(const FString& Parameters)
 		FStringView View(String);
 
 		TestEqual(TEXT("View length"), View.Len(), String.Len());
-		TestEqual(TEXT("The result of Strncmp"), FCStringWide::Strncmp(View.Data(), *String, View.Len()), 0);
+		TestEqual(TEXT("The result of Strncmp"), FCStringWide::Strncmp(View.GetData(), *String, View.Len()), 0);
 		TestFalse("View.IsEmpty", View.IsEmpty());
 	}
 
@@ -50,7 +50,7 @@ bool FStringViewTestCtor::RunTest(const FString& Parameters)
 	{
 		FAnsiStringView View("Test Ctor");
 		TestEqual(TEXT("View length"), View.Len(), FCStringAnsi::Strlen("Test Ctor"));
-		TestEqual(TEXT("The result of Strncmp"), FCStringAnsi::Strncmp(View.Data(), "Test Ctor", View.Len()), 0);
+		TestEqual(TEXT("The result of Strncmp"), FCStringAnsi::Strncmp(View.GetData(), "Test Ctor", View.Len()), 0);
 		TestFalse("View.IsEmpty", View.IsEmpty());
 	}
 
@@ -58,7 +58,7 @@ bool FStringViewTestCtor::RunTest(const FString& Parameters)
 	{
 		FAnsiStringView View("Test SubSection Ctor", 4);
 		TestEqual(TEXT("View length"), View.Len(), 4);
-		TestEqual(TEXT("The result of Strncmp"), FCStringAnsi::Strncmp(View.Data(), "Test", View.Len()), 0);
+		TestEqual(TEXT("The result of Strncmp"), FCStringAnsi::Strncmp(View.GetData(), "Test", View.Len()), 0);
 		TestFalse("View.IsEmpty", View.IsEmpty());
 	}
 
@@ -345,7 +345,7 @@ bool FStringViewTestArrayModifiers::RunTest(const FString& Parameters)
 		View.RemovePrefix(FCStringWide::Strlen(Prefix));
 
 		TestEqual(TEXT("View length"), View.Len(), FCStringWide::Strlen(Suffix));
-		TestEqual(TEXT("The result of Strncmp"), FCStringWide::Strncmp(View.Data(), Suffix, View.Len()), 0);
+		TestEqual(TEXT("The result of Strncmp"), FCStringWide::Strncmp(View.GetData(), Suffix, View.Len()), 0);
 	}
 
 	// Remove suffix
@@ -354,7 +354,7 @@ bool FStringViewTestArrayModifiers::RunTest(const FString& Parameters)
 		View.RemoveSuffix(FCStringWide::Strlen(Suffix));
 
 		TestEqual(TEXT("View length"), View.Len(), FCStringWide::Strlen(Prefix));
-		TestEqual(TEXT("The result of Strncmp"), FCStringWide::Strncmp(View.Data(), Prefix, View.Len()), 0);
+		TestEqual(TEXT("The result of Strncmp"), FCStringWide::Strncmp(View.GetData(), Prefix, View.Len()), 0);
 	}
 
 	return true;
@@ -457,9 +457,9 @@ bool FStringViewTestSubStr::RunTest(const FString& Parameters)
 																	// string since the null terminator is still valid
 		FStringView OutofBoundsResult = View.SubStr(0, 1024);
 
-		TestTrue(TEXT("FStringView::SubStr(2)"), FCString::Strncmp(Word0.Data(), TEXT("A"), Word0.Len()) == 0);
-		TestTrue(TEXT("FStringView::SubStr(3)"), FCString::Strncmp(Word1.Data(), TEXT("test"), Word1.Len()) == 0);
-		TestTrue(TEXT("FStringView::SubStr(4)"), FCString::Strncmp(Word2.Data(), TEXT("string"), Word2.Len()) == 0);
+		TestTrue(TEXT("FStringView::SubStr(2)"), FCString::Strncmp(Word0.GetData(), TEXT("A"), Word0.Len()) == 0);
+		TestTrue(TEXT("FStringView::SubStr(3)"), FCString::Strncmp(Word1.GetData(), TEXT("test"), Word1.Len()) == 0);
+		TestTrue(TEXT("FStringView::SubStr(4)"), FCString::Strncmp(Word2.GetData(), TEXT("string"), Word2.Len()) == 0);
 		TestTrue(TEXT("FStringView::SubStr(5)"), NullTerminatorResult.IsEmpty());
 		TestTrue(TEXT("FStringView::SubStr(6)"), View == OutofBoundsResult);
 	}
@@ -483,10 +483,10 @@ bool FStringViewTestLeft::RunTest(const FString& Parameters)
 		FStringView View(TEXT("A test string"));
 		FStringView Result = View.Left(8);
 
-		TestTrue(TEXT("FStringView::Left"), FCString::Strncmp(Result.Data(), TEXT("A test s"), Result.Len()) == 0);
+		TestTrue(TEXT("FStringView::Left"), FCString::Strncmp(Result.GetData(), TEXT("A test s"), Result.Len()) == 0);
 
 		FStringView OutofBoundsResult = View.Left(1024);
-		TestTrue(TEXT("FStringView::Left"), FCString::Strncmp(OutofBoundsResult.Data(), TEXT("A test string"), OutofBoundsResult.Len()) == 0);
+		TestTrue(TEXT("FStringView::Left"), FCString::Strncmp(OutofBoundsResult.GetData(), TEXT("A test string"), OutofBoundsResult.Len()) == 0);
 	}
 
 	return true;
@@ -508,10 +508,10 @@ bool FStringViewTestLeftChop::RunTest(const FString& Parameters)
 		FStringView View(TEXT("A test string"));
 		FStringView Result = View.LeftChop(5);
 
-		TestTrue(TEXT("FStringView::LeftChop"), FCString::Strncmp(Result.Data(), TEXT("A test s"), Result.Len()) == 0);
+		TestTrue(TEXT("FStringView::LeftChop"), FCString::Strncmp(Result.GetData(), TEXT("A test s"), Result.Len()) == 0);
 
 		FStringView OutofBoundsResult = View.LeftChop(1024);
-		TestTrue(TEXT("FStringView::LeftChop"), FCString::Strncmp(OutofBoundsResult.Data(), TEXT("A test string"), OutofBoundsResult.Len()) == 0);
+		TestTrue(TEXT("FStringView::LeftChop"), FCString::Strncmp(OutofBoundsResult.GetData(), TEXT("A test string"), OutofBoundsResult.Len()) == 0);
 	}
 
 	return true;
@@ -533,10 +533,10 @@ bool FStringViewTestRight::RunTest(const FString& Parameters)
 		FStringView View(TEXT("A test string"));
 		FStringView Result = View.Right(8);
 
-		TestTrue(TEXT("FStringView::Right"), FCString::Strncmp(Result.Data(), TEXT("t string"), Result.Len()) == 0);
+		TestTrue(TEXT("FStringView::Right"), FCString::Strncmp(Result.GetData(), TEXT("t string"), Result.Len()) == 0);
 
 		FStringView OutofBoundsResult = View.Right(1024);
-		TestTrue(TEXT("FStringView::Right"), FCString::Strncmp(OutofBoundsResult.Data(), TEXT("A test string"), OutofBoundsResult.Len()) == 0);
+		TestTrue(TEXT("FStringView::Right"), FCString::Strncmp(OutofBoundsResult.GetData(), TEXT("A test string"), OutofBoundsResult.Len()) == 0);
 	}
 
 	return true;
@@ -558,10 +558,10 @@ bool FStringViewTestRightChop::RunTest(const FString& Parameters)
 		FStringView View(TEXT("A test string"));
 		FStringView Result = View.RightChop(3);
 
-		TestTrue(TEXT("FStringView::RightChop"), FCString::Strncmp(Result.Data(), TEXT("est string"), Result.Len()) == 0);
+		TestTrue(TEXT("FStringView::RightChop"), FCString::Strncmp(Result.GetData(), TEXT("est string"), Result.Len()) == 0);
 
 		FStringView OutofBoundsResult = View.RightChop(1024);
-		TestTrue(TEXT("FStringView::RightChop"), FCString::Strncmp(OutofBoundsResult.Data(), TEXT("A test string"), OutofBoundsResult.Len()) == 0);
+		TestTrue(TEXT("FStringView::RightChop"), FCString::Strncmp(OutofBoundsResult.GetData(), TEXT("A test string"), OutofBoundsResult.Len()) == 0);
 	}
 
 	return true;
@@ -590,9 +590,9 @@ bool FStringViewTestMid::RunTest(const FString& Parameters)
 																// string since the null terminator is still valid
 		FStringView OutofBoundsResult = View.Mid(0, 1024);
 
-		TestTrue(TEXT("FStringView::Mid(2)"), FCString::Strncmp(Word0.Data(), TEXT("A"), Word0.Len()) == 0);
-		TestTrue(TEXT("FStringView::Mid(3)"), FCString::Strncmp(Word1.Data(), TEXT("test"), Word1.Len()) == 0);
-		TestTrue(TEXT("FStringView::Mid(4)"), FCString::Strncmp(Word2.Data(), TEXT("string"), Word2.Len()) == 0);
+		TestTrue(TEXT("FStringView::Mid(2)"), FCString::Strncmp(Word0.GetData(), TEXT("A"), Word0.Len()) == 0);
+		TestTrue(TEXT("FStringView::Mid(3)"), FCString::Strncmp(Word1.GetData(), TEXT("test"), Word1.Len()) == 0);
+		TestTrue(TEXT("FStringView::Mid(4)"), FCString::Strncmp(Word2.GetData(), TEXT("string"), Word2.Len()) == 0);
 		TestTrue(TEXT("FStringView::Mid(5)"), NullTerminatorResult.IsEmpty());
 		TestTrue(TEXT("FStringView::Mid(6)"), View == OutofBoundsResult);
 	}
