Another potential workaround for IncludeTool limitations

Assume all our clang platforms has __builtin_assume() since IncludeTools fails to parse __has_builtin(__builtin_assume). Can override UE_ASSUME macro in the platform specific headers if some clang platform lacks support.

diff --git a/Engine/Source/Runtime/Core/Public/HAL/Platform.h b/Engine/Source/Runtime/Core/Public/HAL/Platform.h
--- a/Engine/Source/Runtime/Core/Public/HAL/Platform.h
+++ b/Engine/Source/Runtime/Core/Public/HAL/Platform.h
@@ -618,11 +618,7 @@
 /** Promise expression is true. Compiler can optimize accordingly with undefined behavior if wrong. Static analyzers understand this.  */
 #ifndef UE_ASSUME
 	#if defined(__clang__)
-		#if __has_builtin(__builtin_assume)
-			#define UE_ASSUME(x) __builtin_assume(x)
-		#else
-			#define UE_ASSUME(x)
-		#endif
+		#define UE_ASSUME(x) __builtin_assume(x)
 	#elif defined(_MSC_VER)
 		#define UE_ASSUME(x) __assume(x)
 	#else
