Optimize compile time of AlgosTest.cpp

MSVC: Before -> After
c1xx.dll: 0.24204s -> 0.24124s
c2.dll: 5.74921s -> 0.68202s
Elapsed Time: 5.718s -> 0.660s
Total Compilation Time: 6.462s -> 1.846s

#rnx

diff --git a/Engine/Source/Runtime/Core/Private/Containers/Algo/AlgosTest.cpp b/Engine/Source/Runtime/Core/Private/Containers/Algo/AlgosTest.cpp
--- a/Engine/Source/Runtime/Core/Private/Containers/Algo/AlgosTest.cpp
+++ b/Engine/Source/Runtime/Core/Private/Containers/Algo/AlgosTest.cpp
@@ -17,14 +17,13 @@
 #include "Algo/LevenshteinDistance.h"
 #include "Templates/Greater.h"
 
-IMPLEMENT_SIMPLE_AUTOMATION_TEST(FAlgosTest, "System.Core.Misc.Algos", EAutomationTestFlags::ApplicationContextMask | EAutomationTestFlags::SmokeFilter)
-
-#define NUM_TEST_OBJECTS      32
-
+class FAlgosTestBase : public FAutomationTestBase
+{
+public:
+	using FAutomationTestBase::FAutomationTestBase;
 
+	static constexpr int32 NumTestObjects = 32;
 
-bool FAlgosTest::RunTest(const FString& Parameters)
-{
 	struct FTestData
 	{
 		FTestData(FString&& InName, int32 InAge, bool bInRetired = false)
@@ -49,19 +48,25 @@ bool FAlgosTest::RunTest(const FString& Parameters)
 		bool bRetired;
 	};
 
-	// setup
-	TArray<int> TestData;
-	for (int i = 0; i < NUM_TEST_OBJECTS; ++i)
+	void Initialize()
 	{
-		TestData.Add(i);
+		for (int i = 0; i < NumTestObjects; ++i)
+		{
+			TestData.Add(i);
+		}
+		for (int i = 0; i < NumTestObjects; ++i)
+		{
+			TestData2.Add(FMath::Rand());
+		}
 	}
-	TArray<int> TestData2;
-	for (int i = 0; i < NUM_TEST_OBJECTS; ++i)
+
+	void Cleanup()
 	{
-		TestData2.Add(FMath::Rand());
+		TestData2.Empty();
+		TestData.Empty();
 	}
 
-	// copy
+	void TestCopy()
 	{
 		TArray<int> TestArray;
 		// empty array
@@ -69,24 +74,24 @@ bool FAlgosTest::RunTest(const FString& Parameters)
 		check(TestArray == TestData);
 		// existing data
 		Algo::Copy(TestData2, TestArray);
-		check(TestArray.Num() == NUM_TEST_OBJECTS * 2);
-		for (int i = 0; i < NUM_TEST_OBJECTS; ++i)
+		check(TestArray.Num() == NumTestObjects * 2);
+		for (int i = 0; i < NumTestObjects; ++i)
 		{
 			check(TestArray[i] == TestData[i]);
 		}
-		for (int i = 0; i < NUM_TEST_OBJECTS; ++i)
+		for (int i = 0; i < NumTestObjects; ++i)
 		{
-			check(TestArray[i + NUM_TEST_OBJECTS] == TestData2[i]);
+			check(TestArray[i + NumTestObjects] == TestData2[i]);
 		}
 	}
 
-	// copy if
+	void TestCopyIf()
 	{
 		TArray<int> TestArray;
 		// empty array
 		Algo::CopyIf(TestData, TestArray, [](int i) { return (i % 2) == 0; });
 		int j = 0;
-		for (int i = 0; i < NUM_TEST_OBJECTS; ++i)
+		for (int i = 0; i < NumTestObjects; ++i)
 		{
 			if (TestData[i] % 2 == 0)
 			{
@@ -97,7 +102,7 @@ bool FAlgosTest::RunTest(const FString& Parameters)
 		// existing data
 		Algo::CopyIf(TestData2, TestArray, [](int i) { return (i % 2) == 0; });
 		j = 0;
-		for (int i = 0; i < NUM_TEST_OBJECTS; ++i)
+		for (int i = 0; i < NumTestObjects; ++i)
 		{
 			if (TestData[i] % 2 == 0)
 			{
@@ -105,7 +110,7 @@ bool FAlgosTest::RunTest(const FString& Parameters)
 				++j;
 			}
 		}
-		for (int i = 0; i < NUM_TEST_OBJECTS; ++i)
+		for (int i = 0; i < NumTestObjects; ++i)
 		{
 			if (TestData2[i] % 2 == 0)
 			{
@@ -116,14 +121,14 @@ bool FAlgosTest::RunTest(const FString& Parameters)
 		check(j == TestArray.Num())
 	}
 
-	// transform
+	void TestTransform()
 	{
 		TArray<float> TestArray;
 
 		// empty array
 		{
 			Algo::Transform(TestData, TestArray, [](int i) { return FMath::DegreesToRadians(i); });
-			check(TestArray.Num() == NUM_TEST_OBJECTS);
+			check(TestArray.Num() == NumTestObjects);
 			for (int i = 0; i < TestArray.Num(); ++i)
 			{
 				check(TestArray[i] == FMath::DegreesToRadians(TestData[i]));
@@ -133,14 +138,14 @@ bool FAlgosTest::RunTest(const FString& Parameters)
 		// existing data
 		{
 			Algo::Transform(TestData2, TestArray, [](int i) { return FMath::DegreesToRadians(i); });
-			check(TestArray.Num() == NUM_TEST_OBJECTS * 2);
-			for (int i = 0; i < NUM_TEST_OBJECTS; ++i)
+			check(TestArray.Num() == NumTestObjects * 2);
+			for (int i = 0; i < NumTestObjects; ++i)
 			{
 				check(TestArray[i] == FMath::DegreesToRadians(TestData[i]));
 			}
-			for (int i = 0; i < NUM_TEST_OBJECTS; ++i)
+			for (int i = 0; i < NumTestObjects; ++i)
 			{
-				check(TestArray[i + NUM_TEST_OBJECTS] == FMath::DegreesToRadians(TestData2[i]));
+				check(TestArray[i + NumTestObjects] == FMath::DegreesToRadians(TestData2[i]));
 			}
 		}
 
@@ -192,7 +197,7 @@ bool FAlgosTest::RunTest(const FString& Parameters)
 		}
 	}
 
-	// transform if
+	void TestTransformIf()
 	{
 		TArray<float> TestArray;
 
@@ -200,7 +205,7 @@ bool FAlgosTest::RunTest(const FString& Parameters)
 		{
 			Algo::TransformIf(TestData, TestArray, [](int i) { return (i % 2) == 0; }, [](int i) { return FMath::DegreesToRadians(i); });
 			int j = 0;
-			for (int i = 0; i < NUM_TEST_OBJECTS; ++i)
+			for (int i = 0; i < NumTestObjects; ++i)
 			{
 				if (TestData[i] % 2 == 0)
 				{
@@ -214,7 +219,7 @@ bool FAlgosTest::RunTest(const FString& Parameters)
 		{
 			Algo::TransformIf(TestData2, TestArray, [](int i) { return (i % 2) == 0; }, [](int i) { return FMath::DegreesToRadians(i); });
 			int j = 0;
-			for (int i = 0; i < NUM_TEST_OBJECTS; ++i)
+			for (int i = 0; i < NumTestObjects; ++i)
 			{
 				if (TestData[i] % 2 == 0)
 				{
@@ -222,7 +227,7 @@ bool FAlgosTest::RunTest(const FString& Parameters)
 					++j;
 				}
 			}
-			for (int i = 0; i < NUM_TEST_OBJECTS; ++i)
+			for (int i = 0; i < NumTestObjects; ++i)
 			{
 				if (TestData2[i] % 2 == 0)
 				{
@@ -257,7 +262,7 @@ bool FAlgosTest::RunTest(const FString& Parameters)
 		}
 	}
 
-	// binary search
+	void TestBinarySearch()
 	{
 		// Verify static array case
 		int StaticArray[] = { 2,4,6,6,6,8 };
@@ -293,7 +298,7 @@ bool FAlgosTest::RunTest(const FString& Parameters)
 		check(Algo::UpperBoundBy(IntArray, 6, FIdentityFunctor()) == 7);
 	}
 
-	// heapify
+	void TestHeapify()
 	{
 		TArray<int> TestArray = TestData2;
 		Algo::Heapify(TestArray);
@@ -301,7 +306,7 @@ bool FAlgosTest::RunTest(const FString& Parameters)
 		check(Algo::IsHeap(TestArray));
 	}
 
-	// heap sort
+	void TestHeapSort()
 	{
 		TArray<int> TestArray = TestData2;
 		Algo::HeapSort(TestArray);
@@ -311,7 +316,7 @@ bool FAlgosTest::RunTest(const FString& Parameters)
 		check(Algo::IsSorted(TestArray));
 	}
 
-	// intro sort
+	void TestIntroSort()
 	{
 		TArray<int> TestArray = TestData2;
 		Algo::IntroSort(TestArray);
@@ -319,7 +324,7 @@ bool FAlgosTest::RunTest(const FString& Parameters)
 		check(Algo::IsSorted(TestArray));
 	}
 
-	// sort
+	void TestSort()
 	{
 		// regular Sort
 		TArray<int> TestArray = TestData2;
@@ -355,106 +360,158 @@ bool FAlgosTest::RunTest(const FString& Parameters)
 		check(Algo::IsSortedBy(TestArray, Projection, Predicate));
 	}
 
-	// Edit distance test
+	void TestEditDistance()
 	{
-		auto RunEditDistanceTest = [this](const FString& A, const FString& B, const ESearchCase::Type SearchCase, const int32 ExpectedResultDistance)
+		struct FEditDistanceTestData
 		{
-			// Run test
-			int32 ResultDistance = MAX_int32;
-			if (SearchCase == ESearchCase::IgnoreCase)
-			{
-				ResultDistance = Algo::LevenshteinDistance(A.ToLower(), B.ToLower());
-			}
-			else
-			{
-				ResultDistance = Algo::LevenshteinDistance(A, B);
-			}
+			const TCHAR* A;
+			const TCHAR* B;
+			ESearchCase::Type SearchCase;
+			int32 ExpectedResultDistance;
+		};
 
-			if (ResultDistance != ExpectedResultDistance)
-			{
-				FString SearchCaseStr = SearchCase == ESearchCase::CaseSensitive ? TEXT("CaseSensitive") : TEXT("IgnoreCase");
-				AddError(FString::Printf(TEXT("Algo::EditDistance return the wrong distance between 2 string (A '%s', B '%s', case '%s', result '%d', expected '%d')."), *A, *B, *SearchCaseStr, ResultDistance, ExpectedResultDistance));
-			}
+		const FEditDistanceTestData EditDistanceTests[] =
+		{
+			//Empty tests
+			{ TEXT(""), TEXT("Saturday"), ESearchCase::CaseSensitive, 8 },
+			{ TEXT(""), TEXT("Saturday"), ESearchCase::IgnoreCase, 8 },
+			{ TEXT("Saturday"), TEXT(""), ESearchCase::CaseSensitive, 8 },
+			{ TEXT("Saturday"), TEXT(""), ESearchCase::IgnoreCase, 8 },
+			//One letter tests
+			{ TEXT("a"), TEXT("a"), ESearchCase::CaseSensitive, 0 },
+			{ TEXT("a"), TEXT("b"), ESearchCase::CaseSensitive, 1 },
+			//Equal tests
+			{ TEXT("Saturday"), TEXT("Saturday"), ESearchCase::CaseSensitive, 0 },
+			{ TEXT("Saturday"), TEXT("Saturday"), ESearchCase::IgnoreCase, 0 },
+			//Simple casing test
+			{ TEXT("Saturday"), TEXT("saturday"), ESearchCase::CaseSensitive, 1 },
+			{ TEXT("Saturday"), TEXT("saturday"), ESearchCase::IgnoreCase, 0 },
+			{ TEXT("saturday"), TEXT("Saturday"), ESearchCase::CaseSensitive, 1 },
+			{ TEXT("saturday"), TEXT("Saturday"), ESearchCase::IgnoreCase, 0 },
+			{ TEXT("SaturdaY"), TEXT("saturday"), ESearchCase::CaseSensitive, 2 },
+			{ TEXT("SaturdaY"), TEXT("saturday"), ESearchCase::IgnoreCase, 0 },
+			{ TEXT("saturdaY"), TEXT("Saturday"), ESearchCase::CaseSensitive, 2 },
+			{ TEXT("saturdaY"), TEXT("Saturday"), ESearchCase::IgnoreCase, 0 },
+			{ TEXT("SATURDAY"), TEXT("saturday"), ESearchCase::CaseSensitive, 8 },
+			{ TEXT("SATURDAY"), TEXT("saturday"), ESearchCase::IgnoreCase, 0 },
+			//First char diff
+			{ TEXT("Saturday"), TEXT("baturday"), ESearchCase::CaseSensitive, 1 },
+			{ TEXT("Saturday"), TEXT("baturday"), ESearchCase::IgnoreCase, 1 },
+			//Last char diff
+			{ TEXT("Saturday"), TEXT("Saturdai"), ESearchCase::CaseSensitive, 1 },
+			{ TEXT("Saturday"), TEXT("Saturdai"), ESearchCase::IgnoreCase, 1 },
+			//Middle char diff
+			{ TEXT("Satyrday"), TEXT("Saturday"), ESearchCase::CaseSensitive, 1 },
+			{ TEXT("Satyrday"), TEXT("Saturday"), ESearchCase::IgnoreCase, 1 },
+			//Real cases
+			{ TEXT("Copy_Body"), TEXT("Body"), ESearchCase::CaseSensitive, 5 },
+			{ TEXT("Copy_Body"), TEXT("Body"), ESearchCase::IgnoreCase, 5 },
+			{ TEXT("copy_Body"), TEXT("Paste_Body"), ESearchCase::CaseSensitive, 5 },
+			{ TEXT("copy_Body"), TEXT("Paste_Body"), ESearchCase::IgnoreCase, 5 },
+			{ TEXT("legs"), TEXT("Legs_1"), ESearchCase::CaseSensitive, 3 },
+			{ TEXT("legs"), TEXT("Legs_1"), ESearchCase::IgnoreCase, 2 },
+			{ TEXT("arms"), TEXT("Arms"), ESearchCase::CaseSensitive, 1 },
+			{ TEXT("arms"), TEXT("Arms"), ESearchCase::IgnoreCase, 0 },
+			{ TEXT("Saturday"), TEXT("Sunday"), ESearchCase::CaseSensitive, 3 },
+			{ TEXT("Saturday"), TEXT("Sunday"), ESearchCase::IgnoreCase, 3 },
+			{ TEXT("Saturday"), TEXT("suNday"), ESearchCase::CaseSensitive, 4 },
+			{ TEXT("Saturday"), TEXT("suNday"), ESearchCase::IgnoreCase, 3 },
+			{ TEXT("Saturday"), TEXT("sUnday"), ESearchCase::CaseSensitive, 5 },
+			{ TEXT("Saturday"), TEXT("sUnday"), ESearchCase::IgnoreCase, 3 },
 		};
-		//Empty tests
-		RunEditDistanceTest(TEXT(""), TEXT("Saturday"), ESearchCase::CaseSensitive, 8);
-		RunEditDistanceTest(TEXT(""), TEXT("Saturday"), ESearchCase::IgnoreCase, 8);
-		RunEditDistanceTest(TEXT("Saturday"), TEXT(""), ESearchCase::CaseSensitive, 8);
-		RunEditDistanceTest(TEXT("Saturday"), TEXT(""), ESearchCase::IgnoreCase, 8);
-		//One letter tests
-		RunEditDistanceTest(TEXT("a"), TEXT("a"), ESearchCase::CaseSensitive, 0);
-		RunEditDistanceTest(TEXT("a"), TEXT("b"), ESearchCase::CaseSensitive, 1);
-		//Equal tests
-		RunEditDistanceTest(TEXT("Saturday"), TEXT("Saturday"), ESearchCase::CaseSensitive, 0);
-		RunEditDistanceTest(TEXT("Saturday"), TEXT("Saturday"), ESearchCase::IgnoreCase, 0);
-		//Simple casing test
-		RunEditDistanceTest(TEXT("Saturday"), TEXT("saturday"), ESearchCase::CaseSensitive, 1);
-		RunEditDistanceTest(TEXT("Saturday"), TEXT("saturday"), ESearchCase::IgnoreCase, 0);
-		RunEditDistanceTest(TEXT("saturday"), TEXT("Saturday"), ESearchCase::CaseSensitive, 1);
-		RunEditDistanceTest(TEXT("saturday"), TEXT("Saturday"), ESearchCase::IgnoreCase, 0);
-		RunEditDistanceTest(TEXT("SaturdaY"), TEXT("saturday"), ESearchCase::CaseSensitive, 2);
-		RunEditDistanceTest(TEXT("SaturdaY"), TEXT("saturday"), ESearchCase::IgnoreCase, 0);
-		RunEditDistanceTest(TEXT("saturdaY"), TEXT("Saturday"), ESearchCase::CaseSensitive, 2);
-		RunEditDistanceTest(TEXT("saturdaY"), TEXT("Saturday"), ESearchCase::IgnoreCase, 0);
-		RunEditDistanceTest(TEXT("SATURDAY"), TEXT("saturday"), ESearchCase::CaseSensitive, 8);
-		RunEditDistanceTest(TEXT("SATURDAY"), TEXT("saturday"), ESearchCase::IgnoreCase, 0);
-		//First char diff
-		RunEditDistanceTest(TEXT("Saturday"), TEXT("baturday"), ESearchCase::CaseSensitive, 1);
-		RunEditDistanceTest(TEXT("Saturday"), TEXT("baturday"), ESearchCase::IgnoreCase, 1);
-		//Last char diff
-		RunEditDistanceTest(TEXT("Saturday"), TEXT("Saturdai"), ESearchCase::CaseSensitive, 1);
-		RunEditDistanceTest(TEXT("Saturday"), TEXT("Saturdai"), ESearchCase::IgnoreCase, 1);
-		//Middle char diff
-		RunEditDistanceTest(TEXT("Satyrday"), TEXT("Saturday"), ESearchCase::CaseSensitive, 1);
-		RunEditDistanceTest(TEXT("Satyrday"), TEXT("Saturday"), ESearchCase::IgnoreCase, 1);
-		//Real cases
-		RunEditDistanceTest(TEXT("Copy_Body"), TEXT("Body"), ESearchCase::CaseSensitive, 5);
-		RunEditDistanceTest(TEXT("Copy_Body"), TEXT("Body"), ESearchCase::IgnoreCase, 5);
-		RunEditDistanceTest(TEXT("copy_Body"), TEXT("Paste_Body"), ESearchCase::CaseSensitive, 5);
-		RunEditDistanceTest(TEXT("copy_Body"), TEXT("Paste_Body"), ESearchCase::IgnoreCase, 5);
-		RunEditDistanceTest(TEXT("legs"), TEXT("Legs_1"), ESearchCase::CaseSensitive, 3);
-		RunEditDistanceTest(TEXT("legs"), TEXT("Legs_1"), ESearchCase::IgnoreCase, 2);
-		RunEditDistanceTest(TEXT("arms"), TEXT("Arms"), ESearchCase::CaseSensitive, 1);
-		RunEditDistanceTest(TEXT("arms"), TEXT("Arms"), ESearchCase::IgnoreCase, 0);
-		RunEditDistanceTest(TEXT("Saturday"), TEXT("Sunday"), ESearchCase::CaseSensitive, 3);
-		RunEditDistanceTest(TEXT("Saturday"), TEXT("Sunday"), ESearchCase::IgnoreCase, 3);
-		RunEditDistanceTest(TEXT("Saturday"), TEXT("suNday"), ESearchCase::CaseSensitive, 4);
-		RunEditDistanceTest(TEXT("Saturday"), TEXT("suNday"), ESearchCase::IgnoreCase, 3);
-		RunEditDistanceTest(TEXT("Saturday"), TEXT("sUnday"), ESearchCase::CaseSensitive, 5);
-		RunEditDistanceTest(TEXT("Saturday"), TEXT("sUnday"), ESearchCase::IgnoreCase, 3);
-
-		auto RunEditDistanceTestArray = [this](const FString& ArrayDescriptionA, const FString& ArrayDescriptionB, const TArray<int32>& A, const TArray<int32>& B, const int32 ExpectedResultDistance)
+
+		for (const FEditDistanceTestData& Test : EditDistanceTests)
 		{
-			// Run test
-			int32 ResultDistance = Algo::LevenshteinDistance(A, B);
+			RunEditDistanceTest(Test.A, Test.B, Test.SearchCase, Test.ExpectedResultDistance);
+		}
+	}
 
-			if (ResultDistance != ExpectedResultDistance)
-			{
-				AddError(FString::Printf(TEXT("Algo::EditDistance return the wrong distance between 2 array (A '%s', B '%s', result '%d', expected '%d')."), *ArrayDescriptionA, *ArrayDescriptionB, ResultDistance, ExpectedResultDistance));
-			}
+	void RunEditDistanceTest(const FString& A, const FString& B, const ESearchCase::Type SearchCase, const int32 ExpectedResultDistance)
+	{
+		// Run test
+		int32 ResultDistance = MAX_int32;
+		if (SearchCase == ESearchCase::IgnoreCase)
+		{
+			ResultDistance = Algo::LevenshteinDistance(A.ToLower(), B.ToLower());
+		}
+		else
+		{
+			ResultDistance = Algo::LevenshteinDistance(A, B);
+		}
+
+		if (ResultDistance != ExpectedResultDistance)
+		{
+			FString SearchCaseStr = SearchCase == ESearchCase::CaseSensitive ? TEXT("CaseSensitive") : TEXT("IgnoreCase");
+			AddError(FString::Printf(TEXT("Algo::EditDistance return the wrong distance between 2 string (A '%s', B '%s', case '%s', result '%d', expected '%d')."), *A, *B, *SearchCaseStr, ResultDistance, ExpectedResultDistance));
+		}
+	}
+
+	void TestEditDistanceArray()
+	{
+		struct FEditDistanceArrayTestData
+		{
+			const TCHAR* ArrayDescriptionA;
+			const TCHAR* ArrayDescriptionB;
+			TArray<int32> A;
+			TArray<int32> B;
+			int32 ExpectedResultDistance;
 		};
 
-		TArray<int32> A = { 1, 2, 3, 4 };
-		TArray<int32> B = { 1, 2, 3, 4 };
-		//Identical array
-		RunEditDistanceTestArray(FString("{1, 2, 3, 4}"), FString("{1, 2, 3, 4}"), A, B, 0);
-		//1 difference
-		B[3] = 10;
-		RunEditDistanceTestArray(FString("{1, 2, 3, 4}"), FString("{1, 2, 3, 10}"), A, B, 1);
-		//1 character less
-		B.RemoveAt(3);
-		RunEditDistanceTestArray(FString("{1, 2, 3, 4}"), FString("{1, 2, 3}"), A, B, 1);
-		//1 character more
-		B.Add(4);
-		B.Add(5);
-		RunEditDistanceTestArray(FString("{1, 2, 3, 4}"), FString("{1, 2, 3, 4, 5}"), A, B, 1);
-		//2 character more
-		B.Add(6);
-		RunEditDistanceTestArray(FString("{1, 2, 3, 4}"), FString("{1, 2, 3, 4, 5, 6}"), A, B, 2);
-		//B string empty 
-		B.Empty();
-		RunEditDistanceTestArray(FString("{1, 2, 3, 4}"), FString("{}"), A, B, A.Num());
+		const FEditDistanceArrayTestData EditDistanceArrayTests[] =
+		{
+			//Identical array
+			{ TEXT("{1, 2, 3, 4}"), TEXT("{1, 2, 3, 4}"), {1, 2, 3, 4}, {1, 2, 3, 4}, 0 },
+			//1 difference
+			{ TEXT("{1, 2, 3, 4}"), TEXT("{1, 2, 3, 10}"), {1, 2, 3, 4}, {1, 2, 3, 10}, 1 },
+			//1 character less
+			{ TEXT("{1, 2, 3, 4}"), TEXT("{1, 2, 3}"), {1, 2, 3, 4}, {1, 2, 3}, 1 },
+			//1 character more
+			{ TEXT("{1, 2, 3, 4}"), TEXT("{1, 2, 3, 4, 5}"), {1, 2, 3, 4}, {1, 2, 3, 4, 5}, 1 },
+			//2 character more
+			{ TEXT("{1, 2, 3, 4}"), TEXT("{1, 2, 3, 4, 5, 6}"), {1, 2, 3, 4}, {1, 2, 3, 4, 5, 6}, 2 },
+			//B string empty
+			{ TEXT("{1, 2, 3, 4}"), TEXT("{}"), {1, 2, 3, 4}, {}, 4 },
+		};
+
+		for (const FEditDistanceArrayTestData& Test : EditDistanceArrayTests)
+		{
+			RunEditDistanceTestArray(Test.ArrayDescriptionA, Test.ArrayDescriptionB, Test.A, Test.B, Test.ExpectedResultDistance);
+		}
 	}
 
+	void RunEditDistanceTestArray(const FString& ArrayDescriptionA, const FString& ArrayDescriptionB, const TArray<int32>& A, const TArray<int32>& B, const int32 ExpectedResultDistance)
+	{
+		// Run test
+		int32 ResultDistance = Algo::LevenshteinDistance(A, B);
+
+		if (ResultDistance != ExpectedResultDistance)
+		{
+			AddError(FString::Printf(TEXT("Algo::EditDistance return the wrong distance between 2 array (A '%s', B '%s', result '%d', expected '%d')."), *ArrayDescriptionA, *ArrayDescriptionB, ResultDistance, ExpectedResultDistance));
+		}
+	}
+
+private:
+	TArray<int> TestData;
+	TArray<int> TestData2;
+};
+
+IMPLEMENT_CUSTOM_SIMPLE_AUTOMATION_TEST(FAlgosTest, FAlgosTestBase, "System.Core.Misc.Algos", EAutomationTestFlags::ApplicationContextMask | EAutomationTestFlags::SmokeFilter)
+
+bool FAlgosTest::RunTest(const FString& Parameters)
+{
+	Initialize();
+	TestCopy();
+	TestCopyIf();
+	TestTransform();
+	TestTransformIf();
+	TestBinarySearch();
+	TestHeapify();
+	TestHeapSort();
+	TestIntroSort();
+	TestSort();
+	TestEditDistance();
+	TestEditDistanceArray();
+	Cleanup();
+
 	return true;
 }
