[I/O Dipsatcher] Removed OpenFileChunk as it's no longer used

#rnx

diff --git a/Engine/Source/Runtime/Core/Private/IO/IoDispatcher.cpp b/Engine/Source/Runtime/Core/Private/IO/IoDispatcher.cpp
--- a/Engine/Source/Runtime/Core/Private/IO/IoDispatcher.cpp
+++ b/Engine/Source/Runtime/Core/Private/IO/IoDispatcher.cpp
@@ -17,7 +17,6 @@
 #include "Misc/LazySingleton.h"
 #include "Misc/CoreDelegates.h"
 #include "Trace/Trace.h"
-#include "Misc/Crc.h"
 
 #define IODISPATCHER_TRACE_ENABLED !UE_BUILD_SHIPPING
 
@@ -163,10 +162,6 @@ public:
 		{
 			return FIoBuffer(FIoBuffer::Wrap, MappedRegion->GetMappedPtr() + Entry->GetOffset(), Entry->GetLength());
 		}
-		else if (FFileChunkHandle* FileChunk = FileChunkHandles.Find(ChunkId))
-		{
-			return FIoBuffer(FIoBuffer::Wrap, FileChunk->MappedRegion->GetMappedPtr() + Options.GetOffset(), Options.GetSize());
-		}
 
 		return FIoStatus(EIoErrorCode::NotFound);
 	}
@@ -190,44 +185,13 @@ public:
 		}
 	}
 
-	void SetUniqueId(FStringView InUniqueId)
-	{
-		UniqueId = InUniqueId;
-	}
-
-	TIoStatusOr<FIoChunkId> OpenFileChunk(FStringView Filename);
-	FIoStatus CloseFileChunk(const FIoChunkId& FileChunkId);
-
 private:
-	struct FFileChunkHandle
-	{
-		TUniquePtr<IMappedFileHandle> Handle;
-		TUniquePtr<IMappedFileRegion> MappedRegion;
-	};
-
-	FIoChunkId GetFileChunkId(FStringView Filename)
-	{
-		check(Filename.GetData() && Filename.Len());
-
-		uint32 Hash = FCrc::MemCrc32(Filename.GetData(), Filename.Len());
-
-		uint8 Data[12] = {0};
-		*reinterpret_cast<uint32*>(&Data[0]) = Hash;
-		*reinterpret_cast<uint8*>(&Data[10]) = 255; // i.e EChunkType::LooseFile;
-
-		FIoChunkId FileChunkId;
-		FileChunkId.Set(Data, 12);
-
-		return FileChunkId;
-	}
-
 	FIoStoreEnvironment&				Environment;
 	FString								UniqueId;
 	TMap<FIoChunkId, FIoStoreTocEntry>	Toc;
 	TUniquePtr<IFileHandle>				ContainerFileHandle;
 	TUniquePtr<IMappedFileHandle>		ContainerMappedFileHandle;
 	TUniquePtr<IMappedFileRegion>		MappedRegion;
-	TMap<FIoChunkId, FFileChunkHandle>	FileChunkHandles;
 };
 
 FIoStatus FIoStoreReaderImpl::Initialize(FStringView InUniqueId)
@@ -324,48 +288,6 @@ FIoStatus FIoStoreReaderImpl::Initialize(FStringView InUniqueId)
 	return FIoStatus::Ok;
 }
 
-TIoStatusOr<FIoChunkId>
-FIoStoreReaderImpl::OpenFileChunk(FStringView Filename)
-{
-	FIoChunkId FileChunkId = GetFileChunkId(Filename);
-
-	if (FileChunkHandles.Contains(FileChunkId))
-	{
-		return FileChunkId;
-	}
-
-	IPlatformFile& Ipf = FPlatformFileManager::Get().GetPlatformFile();
-
-	TUniquePtr<IMappedFileHandle> FileHandle;
-	FileHandle.Reset(Ipf.OpenMapped(Filename.GetData()));
-
-	if (!FileHandle)
-	{
-		return FIoStatus(EIoErrorCode::FileNotOpen);
-	}
-
-	TUniquePtr<IMappedFileRegion> MappedFileRegion;
-	MappedFileRegion.Reset(FileHandle->MapRegion());
-
-	if (!MappedFileRegion)
-	{
-		return FIoStatus(EIoErrorCode::FileNotOpen);
-	}
-
-	FFileChunkHandle& FileChunk = FileChunkHandles.Add(FileChunkId);
-	FileChunk.Handle = MoveTemp(FileHandle);
-	FileChunk.MappedRegion = MoveTemp(MappedFileRegion);
-
-	return FileChunkId;
-}
-
-FIoStatus
-FIoStoreReaderImpl::CloseFileChunk(const FIoChunkId& FileChunkId)
-{
-	const uint32 NumRemoved = FileChunkHandles.Remove(FileChunkId);
-	return NumRemoved ? FIoStatus::Ok : FIoStatus(EIoErrorCode::FileNotOpen);
-}
-
 FIoStoreReader::FIoStoreReader(FIoStoreEnvironment& Environment)
 :	Impl(new FIoStoreReaderImpl(Environment))
 {
@@ -398,7 +320,6 @@ public:
 	FIoChunkId				ChunkId;
 	FIoReadOptions			Options;
 	TIoStatusOr<FIoBuffer>	Result;
-	uint64					UserData	= 0;
 	FIoRequestImpl*			NextRequest = nullptr;
 };
 
@@ -482,29 +403,23 @@ public:
 			RequestAllocator.Trim();
 			BatchAllocator.Trim();
 		});
-
-		LooseFileReader = new FIoStoreReader(GIoStoreEnvironment);
-		LooseFileReader->Impl->SetUniqueId(TEXT("LooseFiles"));
-
-		Mount(LooseFileReader);
 	}
 
-	FIoRequestImpl* AllocRequest(const FIoChunkId& ChunkId, FIoReadOptions Options, uint64 UserData = 0)
+	FIoRequestImpl* AllocRequest(const FIoChunkId& ChunkId, FIoReadOptions Options)
 	{
 		FIoRequestImpl* Request = RequestAllocator.Construct();
 
 		Request->ChunkId		= ChunkId;
 		Request->Options		= Options;
 		Request->Result			= FIoStatus::Unknown;
-		Request->UserData		= UserData;
 		Request->NextRequest	= nullptr;
 
 		return Request;
 	}
 
-	FIoRequestImpl* AllocRequest(FIoBatchImpl* Batch, const FIoChunkId& ChunkId, FIoReadOptions Options, uint64 UserData = 0)
+	FIoRequestImpl* AllocRequest(FIoBatchImpl* Batch, const FIoChunkId& ChunkId, FIoReadOptions Options)
 	{
-		FIoRequestImpl* Request	= AllocRequest(ChunkId, Options, UserData);
+		FIoRequestImpl* Request	= AllocRequest(ChunkId, Options);
 
 		Request->NextRequest	= Batch->FirstRequest;
 		Batch->FirstRequest		= Request;
@@ -631,16 +546,6 @@ public:
 		IoStore->Unmount(IoStoreReader);
 	}
 
-	TIoStatusOr<FIoChunkId> OpenFileChunk(FStringView Filename)
-	{
-		return LooseFileReader->Impl->OpenFileChunk(Filename);
-	}
-
-	FIoStatus CloseFileChunk(const FIoChunkId& FileChunkId)
-	{
-		return LooseFileReader->Impl->CloseFileChunk(FileChunkId);
-	}
-
 private:
 	using FRequestAllocator		= TBlockAllocator<FIoRequestImpl, 4096>;
 	using FBatchAllocator		= TBlockAllocator<FIoBatchImpl, 4096>;
@@ -648,7 +553,6 @@ private:
 	TRefCountPtr<FIoStoreImpl>	IoStore;
 	FRequestAllocator			RequestAllocator;
 	FBatchAllocator				BatchAllocator;
-	FIoStoreReader*				LooseFileReader = nullptr;
 };
 
 //////////////////////////////////////////////////////////////////////////
@@ -693,18 +597,6 @@ FIoDispatcher::ReadWithCallback(const FIoChunkId& Chunk, const FIoReadOptions& O
 	Impl->ReadWithCallback(Chunk, Options, MoveTemp(Callback));
 }
 
-TIoStatusOr<FIoChunkId>
-FIoDispatcher::OpenFileChunk(FStringView Filename)
-{
-	return Impl->OpenFileChunk(Filename);
-}
-
-FIoStatus
-FIoDispatcher::CloseFileChunk(const FIoChunkId& FileChunkId)
-{
-	return Impl->CloseFileChunk(FileChunkId);
-}
-
 // Polling methods
 bool 
 FIoDispatcher::DoesChunkExist(const FIoChunkId& ChunkId) const
@@ -761,9 +653,9 @@ FIoBatch::FIoBatch(FIoDispatcherImpl* InDispatcher, FIoBatchImpl* InImpl)
 }
 
 FIoRequest
-FIoBatch::Read(const FIoChunkId& ChunkId, FIoReadOptions Options, uint64 UserData)
+FIoBatch::Read(const FIoChunkId& ChunkId, FIoReadOptions Options)
 {
-	return FIoRequest(Dispatcher->AllocRequest(Impl, ChunkId, Options, UserData));
+	return FIoRequest(Dispatcher->AllocRequest(Impl, ChunkId, Options));
 }
 
 void 
@@ -820,10 +712,4 @@ FIoRequest::GetResult() const
 	return Impl->Result;
 }
 
-uint64
-FIoRequest::GetUserData() const
-{
-	return Impl->UserData;
-}
-
 #endif // PLATFORM_IMPLEMENTS_IO
diff --git a/Engine/Source/Runtime/Core/Public/IO/IoDispatcher.h b/Engine/Source/Runtime/Core/Public/IO/IoDispatcher.h
--- a/Engine/Source/Runtime/Core/Public/IO/IoDispatcher.h
+++ b/Engine/Source/Runtime/Core/Public/IO/IoDispatcher.h
@@ -688,7 +688,6 @@ public:
 	CORE_API FIoStatus						Status() const;
 	CORE_API const FIoChunkId&				GetChunkId() const;
 	CORE_API TIoStatusOr<FIoBuffer>			GetResult() const;
-	CORE_API uint64							GetUserData() const;
 
 private:
 	FIoRequestImpl* Impl = nullptr;
@@ -715,7 +714,7 @@ class FIoBatch
 	FIoBatch(FIoDispatcherImpl* InDispatcher, FIoBatchImpl* InImpl);
 
 public:
-	CORE_API FIoRequest Read(const FIoChunkId& Chunk, FIoReadOptions Options, uint64 UserData = 0);
+	CORE_API FIoRequest Read(const FIoChunkId& Chunk, FIoReadOptions Options);
 
 	CORE_API void ForEachRequest(TFunction<bool(FIoRequest&)>&& Callback);
 
@@ -749,9 +748,6 @@ public:
 	CORE_API bool					DoesChunkExist(const FIoChunkId& ChunkId) const;
 	CORE_API TIoStatusOr<uint64>	GetSizeForChunk(const FIoChunkId& ChunkId) const;
 
-	CORE_API TIoStatusOr<FIoChunkId>	OpenFileChunk(FStringView Filename);
-	CORE_API FIoStatus					CloseFileChunk(const FIoChunkId& FileChunkId);
-
 	FIoDispatcher(const FIoDispatcher&) = default;
 	FIoDispatcher& operator=(const FIoDispatcher&) = delete;
 
