Properly handle dispatch of the serialized property on FArchiveProxy to stop crashes in the cooker when doing non-determinism checks on arrays which differ.

diff --git a/Engine/Source/Runtime/Core/Public/Serialization/Archive.h b/Engine/Source/Runtime/Core/Public/Serialization/Archive.h
--- a/Engine/Source/Runtime/Core/Public/Serialization/Archive.h
+++ b/Engine/Source/Runtime/Core/Public/Serialization/Archive.h
@@ -503,7 +503,7 @@ public:
 	 *
 	 * @param InProperty Pointer to the property that is currently being serialized
 	 */
-	FORCEINLINE void SetSerializedProperty(UProperty* InProperty)
+	virtual void SetSerializedProperty(UProperty* InProperty)
 	{
 		SerializedProperty = InProperty;
 	}
@@ -536,7 +536,7 @@ public:
 	/**
 	 * Set the raw serialized property chain for this archive, optionally overriding the serialized property too (or null to use the head of the property chain)
 	 */
-	void SetSerializedPropertyChain(const FArchiveSerializedPropertyChain* InSerializedPropertyChain, class UProperty* InSerializedPropertyOverride = nullptr);
+	virtual void SetSerializedPropertyChain(const FArchiveSerializedPropertyChain* InSerializedPropertyChain, class UProperty* InSerializedPropertyOverride = nullptr);
 
 #if WITH_EDITORONLY_DATA
 	/** Returns true if the stack of currently serialized properties contains an editor-only property */
@@ -1593,57 +1593,356 @@ public:
 		return this;
 	}
 
-	using FArchiveState::IsLoading;
-	using FArchiveState::IsSaving;
-	using FArchiveState::IsTransacting;
-	using FArchiveState::IsTextFormat;
-	using FArchiveState::WantBinaryPropertySerialization;
-	using FArchiveState::UseUnversionedPropertySerialization;
-	using FArchiveState::IsForcingUnicode;
-	using FArchiveState::IsPersistent;
-	using FArchiveState::IsError;
-	using FArchiveState::IsCriticalError;
-	using FArchiveState::ContainsCode;
-	using FArchiveState::ContainsMap;
-	using FArchiveState::RequiresLocalizationGather;
-	using FArchiveState::ForceByteSwapping;
-	using FArchiveState::IsSerializingDefaults;
-	using FArchiveState::IsIgnoringArchetypeRef;
-	using FArchiveState::DoDelta;
-	using FArchiveState::DoIntraPropertyDelta;
-	using FArchiveState::IsIgnoringOuterRef;
-	using FArchiveState::IsIgnoringClassGeneratedByRef;
-	using FArchiveState::IsIgnoringClassRef;
-	using FArchiveState::IsAllowingLazyLoading;
-	using FArchiveState::IsObjectReferenceCollector;
-	using FArchiveState::IsModifyingWeakAndStrongReferences;
-	using FArchiveState::IsCountingMemory;
-	using FArchiveState::GetPortFlags;
-	using FArchiveState::HasAnyPortFlags;
-	using FArchiveState::HasAllPortFlags;
-	using FArchiveState::GetDebugSerializationFlags;
-	using FArchiveState::ShouldSkipBulkData;
-	using FArchiveState::GetMaxSerializeSize;
-	using FArchiveState::GetCustomVersions;
-	using FArchiveState::SetCustomVersions;
-	using FArchiveState::ResetCustomVersions;
-	using FArchiveState::SetCustomVersion;
-	using FArchiveState::SetByteSwapping;
-	using FArchiveState::SetPortFlags;
-	using FArchiveState::SetDebugSerializationFlags;
-	using FArchiveState::IsFilterEditorOnly;
-	using FArchiveState::SetFilterEditorOnly;
-	using FArchiveState::IsSaveGame;
-	using FArchiveState::IsNetArchive;
-	using FArchiveState::IsCooking;
-	using FArchiveState::CookingTarget;
-	using FArchiveState::SetCookingTarget;
-	using FArchiveState::UseToResolveEnumerators;
-	using FArchiveState::ShouldSkipProperty;
-	using FArchiveState::SetSerializedProperty;
-	using FArchiveState::GetSerializedProperty;
-	using FArchiveState::GetSerializedPropertyChain;
-	using FArchiveState::SetSerializedPropertyChain;
+	FORCEINLINE bool IsLoading() const
+	{
+		return ArIsLoading;
+	}
+
+	FORCEINLINE bool IsSaving() const
+	{
+		return ArIsSaving;
+	}
+
+	FORCEINLINE bool IsTransacting() const
+	{
+		if (FPlatformProperties::HasEditorOnlyData())
+		{
+			return ArIsTransacting;
+		}
+		else
+		{
+			return false;
+		}
+	}
+
+	FORCEINLINE bool IsTextFormat() const
+	{
+		return (ArIsTextFormat && WITH_TEXT_ARCHIVE_SUPPORT);
+	}
+
+	FORCEINLINE bool WantBinaryPropertySerialization() const
+	{
+		return ArWantBinaryPropertySerialization;
+	}
+
+	FORCEINLINE bool IsForcingUnicode() const
+	{
+		return ArForceUnicode;
+	}
+
+	FORCEINLINE bool IsPersistent() const
+	{
+		return ArIsPersistent;
+	}
+
+	FORCEINLINE bool IsError() const
+	{
+		return ArIsError;
+	}
+
+	FORCEINLINE bool IsCriticalError() const
+	{
+		return ArIsCriticalError;
+	}
+
+	FORCEINLINE bool ContainsCode() const
+	{
+		return ArContainsCode;
+	}
+
+	FORCEINLINE bool ContainsMap() const
+	{
+		return ArContainsMap;
+	}
+
+	FORCEINLINE bool RequiresLocalizationGather() const
+	{
+		return ArRequiresLocalizationGather;
+	}
+
+	FORCEINLINE bool ForceByteSwapping() const
+	{
+		return ArForceByteSwapping;
+	}
+
+	FORCEINLINE bool IsSerializingDefaults() const
+	{
+		return (ArSerializingDefaults > 0) ? true : false;
+	}
+
+	FORCEINLINE bool IsIgnoringArchetypeRef() const
+	{
+		return ArIgnoreArchetypeRef;
+	}
+
+	FORCEINLINE bool DoDelta() const
+	{
+		return !ArNoDelta;
+	}
+
+	FORCEINLINE bool DoIntraPropertyDelta() const
+	{
+		return !ArNoIntraPropertyDelta;
+	}
+
+	FORCEINLINE bool IsIgnoringOuterRef() const
+	{
+		return ArIgnoreOuterRef;
+	}
+
+	FORCEINLINE bool IsIgnoringClassGeneratedByRef() const
+	{
+		return ArIgnoreClassGeneratedByRef;
+	}
+
+	FORCEINLINE bool IsIgnoringClassRef() const
+	{
+		return ArIgnoreClassRef;
+	}
+
+	FORCEINLINE bool IsAllowingLazyLoading() const
+	{
+		return ArAllowLazyLoading;
+	}
+
+	FORCEINLINE bool IsObjectReferenceCollector() const
+	{
+		return ArIsObjectReferenceCollector;
+	}
+
+	FORCEINLINE bool IsModifyingWeakAndStrongReferences() const
+	{
+		return ArIsModifyingWeakAndStrongReferences;
+	}
+
+	FORCEINLINE bool IsCountingMemory() const
+	{
+		return ArIsCountingMemory;
+	}
+
+	FORCEINLINE uint32 GetPortFlags() const
+	{
+		return ArPortFlags;
+	}
+
+	FORCEINLINE bool HasAnyPortFlags(uint32 Flags) const
+	{
+		return ((ArPortFlags & Flags) != 0);
+	}
+
+	FORCEINLINE bool HasAllPortFlags(uint32 Flags) const
+	{
+		return ((ArPortFlags & Flags) == Flags);
+	}
+
+	FORCEINLINE uint32 GetDebugSerializationFlags() const
+	{
+#if WITH_EDITOR
+		return ArDebugSerializationFlags;
+#else
+		return 0;
+#endif
+	}
+
+	FORCEINLINE bool ShouldSkipBulkData() const
+	{
+		return ArShouldSkipBulkData;
+	}
+
+	FORCEINLINE int64 GetMaxSerializeSize() const
+	{
+		return ArMaxSerializeSize;
+	}
+
+	/**
+	 * Gets the custom version numbers for this archive.
+	 *
+	 * @return The container of custom versions in the archive.
+	 */
+	virtual const FCustomVersionContainer& GetCustomVersions() const;
+
+	/**
+	 * Sets the custom version numbers for this archive.
+	 *
+	 * @param CustomVersionContainer - The container of custom versions to copy into the archive.
+	 */
+	virtual void SetCustomVersions(const FCustomVersionContainer& CustomVersionContainer);
+
+	/** Resets the custom version numbers for this archive. */
+	virtual void ResetCustomVersions();
+
+	/**
+	 * Sets a specific custom version
+	 *
+	 * @param Key - The guid of the custom version to query.
+	 * @param Version - The version number to set key to
+	 * @param FriendlyName - Friendly name corresponding to the key
+	 */
+	void SetCustomVersion(const struct FGuid& Key, int32 Version, FName FriendlyName);
+
+	/**
+	 * Toggle byte order swapping. This is needed in rare cases when we already know that the data
+	 * swapping has already occurred or if we know that it will be handled later.
+	 *
+	 * @param Enabled	set to true to enable byte order swapping
+	 */
+	void SetByteSwapping(bool Enabled)
+	{
+		ArForceByteSwapping = Enabled;
+	}
+
+	/**
+	 * Sets the archive's property serialization modifier flags
+	 *
+	 * @param	InPortFlags		the new flags to use for property serialization
+	 */
+	void SetPortFlags(uint32 InPortFlags)
+	{
+		ArPortFlags = InPortFlags;
+	}
+
+	/**
+	 * Sets the archives custom serialization modifier flags (nothing to do with PortFlags or Custom versions)
+	 * 
+	 * @param InCustomFlags the new flags to use for custom serialization
+	 */
+	void SetDebugSerializationFlags(uint32 InCustomFlags)
+	{
+#if WITH_EDITOR
+		ArDebugSerializationFlags = InCustomFlags;
+#endif
+	}
+
+	/**
+	 * Indicates whether this archive is filtering editor-only on save or contains data that had editor-only content stripped.
+	 *
+	 * @return true if the archive filters editor-only content, false otherwise.
+	 */
+	bool IsFilterEditorOnly()
+	{
+		return ArIsFilterEditorOnly;
+	}
+
+	/**
+	 * Sets a flag indicating that this archive needs to filter editor-only content.
+	 *
+	 * @param InFilterEditorOnly Whether to filter editor-only content.
+	 */
+	virtual void SetFilterEditorOnly(bool InFilterEditorOnly)
+	{
+		ArIsFilterEditorOnly = InFilterEditorOnly;
+	}
+
+	/**
+	 * Indicates whether this archive is saving or loading game state
+	 *
+	 * @note This is intended for game-specific archives and is not true for any of the build in save methods
+	 * @return true if the archive is dealing with save games, false otherwise.
+	 */
+	bool IsSaveGame()
+	{
+		return ArIsSaveGame;
+	}
+
+	/**
+	 * Whether or not this archive is serializing data being sent/received by the netcode
+	 */
+	FORCEINLINE bool IsNetArchive()
+	{
+		return ArIsNetArchive;
+	}
+
+	/**
+	 * Checks whether the archive is used for cooking.
+	 *
+	 * @return true if the archive is used for cooking, false otherwise.
+	 */
+	FORCEINLINE bool IsCooking() const	
+	{
+		check(!CookingTargetPlatform || (!IsLoading() && !IsTransacting() && IsSaving()));
+
+		return !!CookingTargetPlatform;
+	}
+
+	/**
+	 * Returns the cooking target platform.
+	 *
+	 * @return Target platform.
+	 */
+	FORCEINLINE const ITargetPlatform* CookingTarget() const 
+	{
+		return CookingTargetPlatform;
+	}
+
+	/**
+	 * Sets the cooking target platform.
+	 *
+	 * @param InCookingTarget The target platform to set.
+	 */
+	FORCEINLINE void SetCookingTarget(const ITargetPlatform* InCookingTarget) 
+	{
+		CookingTargetPlatform = InCookingTarget;
+	}
+
+	/**
+	 * Checks whether the archive is used to resolve out-of-date enum indexes
+	 * If function returns true, the archive should be called only for objects containing user defined enum
+	 *
+	 * @return true if the archive is used to resolve out-of-date enum indexes
+	 */
+	virtual bool UseToResolveEnumerators() const
+	{
+		return false;
+	}
+
+	/**
+	 * Checks whether the archive wants to skip the property independent of the other flags
+	 */
+	virtual bool ShouldSkipProperty(const UProperty* InProperty) const
+	{
+		return false;
+	}
+
+	/**
+	 * Overrides the property that is currently being serialized
+	 * @note: You likely want to call PushSerializedProperty/PopSerializedProperty instead
+	 * 
+	 * @param InProperty Pointer to the property that is currently being serialized
+	 */
+	FORCEINLINE void SetSerializedProperty(UProperty* InProperty)
+	{
+		SerializedProperty = InProperty;
+	}
+
+	/**
+	 * Gets the property that is currently being serialized
+	 *
+	 * @return Pointer to the property that is currently being serialized
+	 */
+	FORCEINLINE class UProperty* GetSerializedProperty() const
+	{
+		return SerializedProperty;
+	}
+
+	/**
+	 * Gets the chain of properties that are currently being serialized
+	 * @note This populates the array in stack order, so the 0th entry in the array is the top of the stack of properties
+	 */
+	void GetSerializedPropertyChain(TArray<class UProperty*>& OutProperties) const;
+
+	/**
+	 * Get the raw serialized property chain for this archive
+	 * @note Accessing this directly can avoid an array allocation depending on your use-case
+	 */
+	FORCEINLINE const FArchiveSerializedPropertyChain* GetSerializedPropertyChain() const
+	{
+		return SerializedPropertyChain;
+	}
+
+	/**
+	 * Set the raw serialized property chain for this archive, optionally overriding the serialized property too (or null to use the head of the property chain)
+	 */
+	void SetSerializedPropertyChain(const FArchiveSerializedPropertyChain* InSerializedPropertyChain, class UProperty* InSerializedPropertyOverride = nullptr);
 
 	/**
 	 * Push a property that is currently being serialized onto the property stack
@@ -1850,7 +2149,7 @@ public:
 	/** Custom serialization modifier flags can be used for anything */
 	using FArchiveState::ArDebugSerializationFlags;
 	/** Debug stack storage if you want to add data to the archive for usage further down the serialization stack this should be used in conjunction with the FScopeAddDebugData struct */
-
+	
 	virtual void PushDebugDataString(const FName& DebugData);
 	virtual void PopDebugDataString() { }
 
diff --git a/Engine/Source/Runtime/Core/Public/Serialization/ArchiveProxy.h b/Engine/Source/Runtime/Core/Public/Serialization/ArchiveProxy.h
--- a/Engine/Source/Runtime/Core/Public/Serialization/ArchiveProxy.h
+++ b/Engine/Source/Runtime/Core/Public/Serialization/ArchiveProxy.h
@@ -236,14 +236,28 @@ public:
 	}
 #endif
 
+	FORCEINLINE void SetSerializedProperty(UProperty* InProperty) override
+	{
+		FArchive::SetSerializedProperty(InProperty);
+		InnerArchive.SetSerializedProperty(InProperty);
+	}
+
+	void SetSerializedPropertyChain(const FArchiveSerializedPropertyChain* InSerializedPropertyChain, class UProperty* InSerializedPropertyOverride = nullptr) override
+	{
+		FArchive::SetSerializedPropertyChain(InSerializedPropertyChain, InSerializedPropertyOverride);
+		InnerArchive.SetSerializedPropertyChain(InSerializedPropertyChain, InSerializedPropertyOverride);
+	}
+
 	/** Pushes editor-only marker to the stack of currently serialized properties */
 	virtual FORCEINLINE void PushSerializedProperty(class UProperty* InProperty, const bool bIsEditorOnlyProperty)
 	{
+		FArchive::PushSerializedProperty(InProperty, bIsEditorOnlyProperty);
 		InnerArchive.PushSerializedProperty(InProperty, bIsEditorOnlyProperty);
 	}
 	/** Pops editor-only marker from the stack of currently serialized properties */
 	virtual FORCEINLINE void PopSerializedProperty(class UProperty* InProperty, const bool bIsEditorOnlyProperty)
 	{
+		FArchive::PopSerializedProperty(InProperty, bIsEditorOnlyProperty);
 		InnerArchive.PopSerializedProperty(InProperty, bIsEditorOnlyProperty);
 	}
 #if WITH_EDITORONLY_DATA
