Gracefully handle ue4stats files without frames

diff --git a/Engine/Source/Developer/Profiler/Private/ProfilerRawStatsForThreadView.cpp b/Engine/Source/Developer/Profiler/Private/ProfilerRawStatsForThreadView.cpp
--- a/Engine/Source/Developer/Profiler/Private/ProfilerRawStatsForThreadView.cpp
+++ b/Engine/Source/Developer/Profiler/Private/ProfilerRawStatsForThreadView.cpp
@@ -190,6 +190,7 @@ void FRawProfilerSession::PrepareLoading()
 		// !!CAUTION!! Frame number in the raw stats is pointless, because it is time based, not frame based.
 		// Background threads usually execute time consuming operations, so the frame number won't be valid.
 		// Needs to be combined by the thread and the time, not by the frame number.
+		if (Stream.FramesInfo.Num() > 0)
 		{
 			int64 FrameOffset0 = Stream.FramesInfo[0].FrameFileOffset;
 			FileReader->Seek( FrameOffset0 );
diff --git a/Engine/Source/Developer/Profiler/Private/ProfilerSession.cpp b/Engine/Source/Developer/Profiler/Private/ProfilerSession.cpp
--- a/Engine/Source/Developer/Profiler/Private/ProfilerSession.cpp
+++ b/Engine/Source/Developer/Profiler/Private/ProfilerSession.cpp
@@ -527,6 +527,10 @@ bool FProfilerSession::HandleTicker( float DeltaTime )
 			CompletionSyncAggregatedEventGraphData();
 
 			// Advance event graphs.
+			if (DataProvider->GetNumFrames() == 0)
+			{
+				EventGraphDataTotal = MakeShareable(new FEventGraphData());
+			}
 			UpdateAllEventGraphs( DataProvider->GetNumFrames() );
 
 			// Broadcast that a capture file has been fully processed.
diff --git a/Engine/Source/Runtime/Core/Private/Stats/StatsFile.cpp b/Engine/Source/Runtime/Core/Private/Stats/StatsFile.cpp
--- a/Engine/Source/Runtime/Core/Private/Stats/StatsFile.cpp
+++ b/Engine/Source/Runtime/Core/Private/Stats/StatsFile.cpp
@@ -632,8 +632,15 @@ bool FStatsReadFile::PrepareLoading()
 	Stream.ReadFramesOffsets( *Reader );
 
 	// Move file pointer to the first frame or first stat packet.
-	const int64 FrameOffset0 = Stream.FramesInfo[0].FrameFileOffset;
-	Reader->Seek( FrameOffset0 );
+	if (Stream.FramesInfo.Num() > 0)
+	{
+		const int64 FrameOffset0 = Stream.FramesInfo[0].FrameFileOffset;
+		Reader->Seek( FrameOffset0 );
+	}
+	else
+	{
+		Reader->Seek(Reader->TotalSize());
+	}
 
 	const double TotalTime = FPlatformTime::Seconds() - StartTime;
 	UE_LOG( LogStats, Log, TEXT( "Prepare loading took %.2f sec(s)" ), TotalTime );
