[UE-84223] Change float and double versions of trunc,round,ceil,floor to use the stdlibrary functions instead of using the int version and casting to a float.
This fixes the incorrect values when running these functions on reals > MAX_INT32

diff --git a/Engine/Source/Runtime/Core/Private/Tests/Math/UnrealMathTest.cpp b/Engine/Source/Runtime/Core/Private/Tests/Math/UnrealMathTest.cpp
--- a/Engine/Source/Runtime/Core/Private/Tests/Math/UnrealMathTest.cpp
+++ b/Engine/Source/Runtime/Core/Private/Tests/Math/UnrealMathTest.cpp
@@ -1763,4 +1763,155 @@ bool FMathRoundHalfFromZeroTests::RunTest(const FString& Parameters)
 	return true;
 }
 
+IMPLEMENT_SIMPLE_AUTOMATION_TEST(FMathTruncationTests, "System.Core.Math.TruncationFunctions", EAutomationTestFlags::ApplicationContextMask | EAutomationTestFlags::SmokeFilter)
+bool FMathTruncationTests::RunTest(const FString& Parameters)
+{
+	// Float: 1-bit Sign, 8-bit exponent, 23-bit mantissa, implicit 1
+	float FloatTestCases[][5] {
+		//Value				Trunc				Ceil				Floor				Round		
+		{-1.5f,				-1.0f,				-1.0f,				-2.0f,				-1.0f,				}, // We do not use round half to even, we always round .5 up (towards +inf)
+		{-1.0f,				-1.0f,				-1.0f,				-1.0f,				-1.0f,				},
+		{-0.75f,			-0.0f,				-0.0f,				-1.0f,				-1.0f,				},
+		{-0.5f,				-0.0f,				-0.0f,				-1.0f,				-0.0f,				}, // We do not use round half to even, we always round .5 up (towards +inf)
+		{-0.25f,			-0.0f,				-0.0f,				-1.0f,				-0.0f,				},
+		{0.0f,				0.0f,				0.0f,				0.0f,				0.0f,				},
+		{0.25f,				0.0f,				1.0f,				0.0f,				0.0f,				},
+		{0.5f,				0.0f,				1.0f,				0.0f,				1.0f,				}, // We do not use round half to even, we always round .5 up (towards +inf)
+		{0.75f,				0.0f,				1.0f,				0.0f,				1.0f,				},
+		{1.0f,				1.0f,				1.0f,				1.0f,				1.0f,				},
+		{1.5f,				1.0f,				2.0f,				1.0f,				2.0f,				},
+		{17179869184.0f,	17179869184.0f,		17179869184.0f,		17179869184.0f,		17179869184.0f,		}, // 2^34.  Note that 2^34 + 1 is not representable, but 2^34 is the string of bits 0, 00100010, 10000000000000000000000,
+		{-17179869184.0f,	-17179869184.0f,	-17179869184.0f,	-17179869184.0f,	-17179869184.0f,	}, // -2^34
+		{1048576.6f,		1048576.0f,			1048577.0f,			1048576.0f,			1048577.0f,			}, // 2^20 + 0.6
+		{-1048576.6f,		-1048576.0f,		-1048576.0f,		-1048577.0f,		-1048577.0f,		}, // -2^20 - 0.6
+	};
+	int IntTestCases[][4]{
+		//					Trunc				Ceil				Floor				Round
+		{					-1,					-1,					-2,					-1,					},
+		{					-1,					-1,					-1,					-1,					},
+		{					0,					0,					-1,					-1,					},
+		{					0,					0,					-1,					0,					},
+		{					0,					0,					-1,					0,					},
+		{					0,					0,					0,					0,					},
+		{					0,					1,					0,					0,					},
+		{					0,					1,					0,					1,					},
+		{					0,					1,					0,					1,					},
+		{					1,					1,					1,					1,					},
+		{					1,					2,					1,					2,					},
+		{					0,					0,					0,					0,					}, // undefined, > MAX_INT32
+		{					0,					0,					0,					0,					}, // undefined, < MIN_INT32
+		{					1048576,			1048577,			1048576,			1048577,			},
+		{					-1048576,			-1048576,			-1048577,			-1048577,			},
+	};
+	static_assert(UE_ARRAY_COUNT(FloatTestCases) == UE_ARRAY_COUNT(IntTestCases), "IntTestCases use the value from FloatTestCases and must be the same length");
+
+	TCHAR TestNameBuffer[128];
+	auto SubTestName = [&TestNameBuffer] (const TCHAR* FunctionName, double Input) {
+		FCString::Snprintf(TestNameBuffer, UE_ARRAY_COUNT(TestNameBuffer), TEXT("%s(%lf)"), FunctionName, Input);
+		return TestNameBuffer;
+	};
+
+	for (uint32 TestCaseIndex = 0; TestCaseIndex < UE_ARRAY_COUNT(FloatTestCases); TestCaseIndex++)
+	{
+		float* FloatValues = FloatTestCases[TestCaseIndex];
+		float Input = FloatValues[0];
+
+		TestEqual(SubTestName(TEXT("TruncToFloat"), Input), FMath::TruncToFloat(Input), FloatValues[1]);
+		TestEqual(SubTestName(TEXT("CeilToFloat"), Input), FMath::CeilToFloat(Input), FloatValues[2]);
+		TestEqual(SubTestName(TEXT("FloorToFloat"), Input), FMath::FloorToFloat(Input), FloatValues[3]);
+		TestEqual(SubTestName(TEXT("RoundToFloat"), Input), FMath::RoundToFloat(Input), FloatValues[4]);
+
+		int* IntValues = IntTestCases[TestCaseIndex];
+		if ((float)MIN_int32 <= Input && Input <= (float)MAX_int32)
+		{
+			TestEqual(SubTestName(TEXT("TruncToInt"), Input), FMath::TruncToInt(Input), IntValues[0]);
+			TestEqual(SubTestName(TEXT("CeilToInt"), Input), FMath::CeilToInt(Input), IntValues[1]);
+			TestEqual(SubTestName(TEXT("FloorToInt"), Input), FMath::FloorToInt(Input), IntValues[2]);
+			TestEqual(SubTestName(TEXT("RoundToInt"), Input), FMath::RoundToInt(Input), IntValues[3]);
+		}
+	}
+
+	// Double: 1-bit sign, 11-bit exponent, 52-bit mantissa, implicit 1
+	double DoubleTestCases[][5]{
+		//Value						Trunc					Ceil					Floor					Round		
+		{-1.5,						-1.0,					-1.0,					-2.0,					-1.0,					}, // We do not use round half to even, we always round .5 up (towards +inf)
+		{-1.0,						-1.0,					-1.0,					-1.0,					-1.0,					},
+		{-0.75,						-0.0,					-0.0,					-1.0,					-1.0,					},
+		{-0.5,						-0.0,					-0.0,					-1.0,					-0.0,					}, // We do not use round half to even, we always round .5 up (towards +inf)
+		{-0.25,						-0.0,					-0.0,					-1.0,					-0.0,					},
+		{0.0,						0.0,					0.0,					0.0,					0.0,					},
+		{0.25,						0.0,					1.0,					0.0,					0.0,					},
+		{0.5,						0.0,					1.0,					0.0,					1.0,					}, // We do not use round half to even, we always round .5 up (towards +inf)
+		{0.75,						0.0,					1.0,					0.0,					1.0,					},
+		{1.0,						1.0,					1.0,					1.0,					1.0,					},
+		{1.5,						1.0,					2.0,					1.0,					2.0,					},
+		{17179869184.0,				17179869184.0,			17179869184.0,			17179869184.0,			17179869184.0			}, // 2^34
+		{-17179869184.0,			-17179869184.0,			-17179869184.0,			-17179869184.0,			-17179869184.0			},
+		{1048576.6,					1048576.0,				1048577.0,				1048576.0,				1048577.0,				},
+		{-1048576.6,				-1048576.0,				-1048576.0,				-1048577.0,				-1048577.0,				},
+		{73786976294838206464.,		73786976294838206464.,	73786976294838206464.,	73786976294838206464.,	73786976294838206464.	}, // 2^66
+		{-73786976294838206464.,	-73786976294838206464.,	-73786976294838206464.,	-73786976294838206464.,	-73786976294838206464.	},
+		{281474976710656.6,			281474976710656.0,		281474976710657.0,		281474976710656.0,		281474976710657.0		}, // 2^48 + 0.6
+		{-281474976710656.6,		-281474976710656.0,		-281474976710656.0,		-281474976710657.0,		-281474976710657.0		},
+	};
+
+	for (uint32 TestCaseIndex = 0; TestCaseIndex < UE_ARRAY_COUNT(DoubleTestCases); TestCaseIndex++)
+	{
+		double* DoubleValues = DoubleTestCases[TestCaseIndex];
+		double Input = DoubleValues[0];
+
+		TestEqual(SubTestName(TEXT("TruncToDouble"), Input), FMath::TruncToDouble(Input), DoubleValues[1]);
+		TestEqual(SubTestName(TEXT("CeilToDouble"), Input), FMath::CeilToDouble(Input), DoubleValues[2]);
+		TestEqual(SubTestName(TEXT("FloorToDouble"), Input), FMath::FloorToDouble(Input), DoubleValues[3]);
+		TestEqual(SubTestName(TEXT("RoundToDouble"), Input), FMath::RoundToDouble(Input), DoubleValues[4]);
+	}
+
+#define MATH_TRUNCATION_SPEED_TEST
+#ifdef MATH_TRUNCATION_SPEED_TEST
+	volatile static float ForceCompileFloat;
+	volatile static int ForceCompileInt;
+
+	auto TimeIt = [](const TCHAR* SubFunctionName, float(*ComputeMath)(float Input), float(*ComputeGeneric)(float Input))
+	{
+		double StartTime = FPlatformTime::Seconds();
+		const float StartInput = 0.6f;
+		const float NumTrials = 10.f*1000.f*1000.f;
+		const float MicroSecondsPerSecond = 1000 * 1000.f;
+		for (float Input = StartInput; Input < NumTrials; Input += 1.0f)
+		{
+			ForceCompileFloat += ComputeMath(Input);
+		}
+		double EndTime = FPlatformTime::Seconds();
+		double FMathDuration = EndTime - StartTime;
+		StartTime = FPlatformTime::Seconds();
+		for (float Input = 0.6f; Input < 10 * 1000 * 1000; Input += 1.0f)
+		{
+			ForceCompileFloat += ComputeGeneric(Input);
+		}
+		EndTime = FPlatformTime::Seconds();
+		double GenericDuration = EndTime - StartTime;
+
+		UE_LOG(LogInit, Log, TEXT("%s: FMath time: %lfus, Generic: %lfus"), SubFunctionName, FMathDuration*MicroSecondsPerSecond / NumTrials, GenericDuration*MicroSecondsPerSecond / NumTrials);
+	};
+
+	TimeIt(TEXT("TruncToInt"), [](float Input) { return (float)FMath::TruncToInt(Input); }, [](float Input) { return (float)FGenericPlatformMath::TruncToInt(Input); });
+	TimeIt(TEXT("CeilToInt"), [](float Input) { return (float)FMath::CeilToInt(Input); }, [](float Input) { return (float)FGenericPlatformMath::CeilToInt(Input); });
+	TimeIt(TEXT("FloorToInt"), [](float Input) { return (float)FMath::FloorToInt(Input); }, [](float Input) { return (float)FGenericPlatformMath::FloorToInt(Input); });
+	TimeIt(TEXT("RoundToInt"), [](float Input) { return (float)FMath::RoundToInt(Input); }, [](float Input) { return (float)FGenericPlatformMath::RoundToInt(Input); });
+
+	TimeIt(TEXT("TruncToFloat"), [](float Input) { return FMath::TruncToFloat(Input); }, [](float Input) { return FGenericPlatformMath::TruncToFloat(Input); });
+	TimeIt(TEXT("CeilToFloat"), [](float Input) { return FMath::CeilToFloat(Input); }, [](float Input) { return FGenericPlatformMath::CeilToFloat(Input); });
+	TimeIt(TEXT("FloorToFloat"), [](float Input) { return FMath::FloorToFloat(Input); }, [](float Input) { return FGenericPlatformMath::FloorToFloat(Input); });
+	TimeIt(TEXT("RoundToFloat"), [](float Input) { return FMath::RoundToFloat(Input); }, [](float Input) { return FGenericPlatformMath::RoundToFloat(Input); });
+
+	TimeIt(TEXT("TruncToDouble"), [](float Input) { return (float)FMath::TruncToDouble((double)Input); }, [](float Input) { return (float)FGenericPlatformMath::TruncToDouble((double)Input); });
+	TimeIt(TEXT("CeilToDouble"), [](float Input) { return (float)FMath::CeilToDouble((double)Input); }, [](float Input) { return (float)FGenericPlatformMath::CeilToDouble((double)Input); });
+	TimeIt(TEXT("FloorToDouble"), [](float Input) { return (float)FMath::FloorToDouble((double)Input); }, [](float Input) { return (float)FGenericPlatformMath::FloorToDouble((double)Input); });
+	TimeIt(TEXT("RoundToDouble"), [](float Input) { return (float)FMath::RoundToDouble((double)Input); }, [](float Input) { return (float)FGenericPlatformMath::RoundToDouble((double)Input); });
+#endif
+
+	return true;
+}
+
+
 #endif //WITH_DEV_AUTOMATION_TESTS
diff --git a/Engine/Source/Runtime/Core/Public/GenericPlatform/GenericPlatformMath.h b/Engine/Source/Runtime/Core/Public/GenericPlatform/GenericPlatformMath.h
--- a/Engine/Source/Runtime/Core/Public/GenericPlatform/GenericPlatformMath.h
+++ b/Engine/Source/Runtime/Core/Public/GenericPlatform/GenericPlatformMath.h
@@ -31,9 +31,19 @@ struct FGenericPlatformMath
 	 * @param F		Floating point value to convert
 	 * @return		Truncated integer value.
 	 */
-	static CONSTEXPR FORCEINLINE float TruncToFloat(float F)
+	static FORCEINLINE float TruncToFloat(float F)
 	{
-		return (float)TruncToInt(F);
+		return truncf(F);
+	}
+
+	/**
+	 * Converts a double to an integer value with truncation towards zero.
+	 * @param F		Floating point value to convert
+	 * @return		Truncated integer value.
+	 */
+	static FORCEINLINE double TruncToDouble(double F)
+	{
+		return trunc(F);
 	}
 
 	/**
diff --git a/Engine/Source/Runtime/Core/Public/Mac/MacPlatformMath.h b/Engine/Source/Runtime/Core/Public/Mac/MacPlatformMath.h
--- a/Engine/Source/Runtime/Core/Public/Mac/MacPlatformMath.h
+++ b/Engine/Source/Runtime/Core/Public/Mac/MacPlatformMath.h
@@ -16,39 +16,62 @@ struct FMacPlatformMath : public FClangPlatformMath
 {
 	static FORCEINLINE int32 TruncToInt(float F)
 	{
-		return _mm_cvtt_ss2si(_mm_set_ss(F));
+		return UnrealPlatformMathSSE::TruncToInt(F);
 	}
 
 	static FORCEINLINE float TruncToFloat(float F)
 	{
-		return (float)TruncToInt(F); // same as generic implementation, but this will call the faster trunc
+		return UnrealPlatformMathSSE::TruncToFloat(F);
+	}
+
+	static FORCEINLINE double TruncToDouble(double F)
+	{
+		return UnrealPlatformMathSSE::TruncToDouble(F);
 	}
 
 	static FORCEINLINE int32 RoundToInt(float F)
 	{
-		// Note: the x2 is to workaround the rounding-to-nearest-even-number issue when the fraction is .5
-		return _mm_cvt_ss2si(_mm_set_ss(F + F + 0.5f)) >> 1;
+		return UnrealPlatformMathSSE::RoundToInt(F);
 	}
 
 	static FORCEINLINE float RoundToFloat(float F)
 	{
-		return (float)RoundToInt(F);
+		return UnrealPlatformMathSSE::RoundToFloat(F);
+	}
+
+	static FORCEINLINE double RoundToDouble(double F)
+	{
+		return UnrealPlatformMathSSE::RoundToDouble(F);
 	}
 
 	static FORCEINLINE int32 FloorToInt(float F)
 	{
-		return _mm_cvt_ss2si(_mm_set_ss(F + F - 0.5f)) >> 1;
+		return UnrealPlatformMathSSE::FloorToInt(F);
 	}
 
 	static FORCEINLINE float FloorToFloat(float F)
 	{
-		return (float)FloorToInt(F);
+		return UnrealPlatformMathSSE::FloorToFloat(F);
+	}
+
+	static FORCEINLINE double FloorToDouble(double F)
+	{
+		return UnrealPlatformMathSSE::FloorToDouble(F);
 	}
 
 	static FORCEINLINE int32 CeilToInt(float F)
 	{
-		// Note: the x2 is to workaround the rounding-to-nearest-even-number issue when the fraction is .5
-		return -(_mm_cvt_ss2si(_mm_set_ss(-0.5f - (F + F))) >> 1);
+		return UnrealPlatformMathSSE::CeilToInt(F);
+	}
+
+	static FORCEINLINE float CeilToFloat(float F)
+	{
+		return UnrealPlatformMathSSE::CeilToFloat(F);
+	}
+
+	static FORCEINLINE double CeilToDouble(double F)
+	{
+		return UnrealPlatformMathSSE::CeilToDouble(F);
 	}
 
 #if PLATFORM_ENABLE_POPCNT_INTRINSIC
@@ -61,11 +84,6 @@ struct FMacPlatformMath : public FClangPlatformMath
 	}
 #endif
 
-	static FORCEINLINE float CeilToFloat(float F)
-	{
-		return (float)CeilToInt(F);
-	}
-
 	static FORCEINLINE bool IsNaN( float A ) { return isnan(A) != 0; }
 	static FORCEINLINE bool IsFinite( float A ) { return isfinite(A); }
 
diff --git a/Engine/Source/Runtime/Core/Public/Math/UnrealPlatformMathSSE.h b/Engine/Source/Runtime/Core/Public/Math/UnrealPlatformMathSSE.h
--- a/Engine/Source/Runtime/Core/Public/Math/UnrealPlatformMathSSE.h
+++ b/Engine/Source/Runtime/Core/Public/Math/UnrealPlatformMathSSE.h
@@ -68,4 +68,73 @@ namespace UnrealPlatformMathSSE
 		_mm_store_ss(&temp, X1);
 		return temp;
 	}
+
+	static FORCEINLINE int32 TruncToInt(float F)
+	{
+		return _mm_cvtt_ss2si(_mm_set_ss(F));
+	}
+
+	static FORCEINLINE float TruncToFloat(float F)
+	{
+		return _mm_cvtss_f32(_mm_round_ps(_mm_set_ss(F), 3));
+	}
+
+	static FORCEINLINE double TruncToDouble(double F)
+	{
+		return _mm_cvtsd_f64(_mm_round_pd(_mm_set_sd(F), 3));
+	}
+
+	static FORCEINLINE int32 FloorToInt(float F)
+	{
+		// Note: unlike the Generic solution and the float solution, we implement FloorToInt using a rounding instruction, rather than implementing RoundToInt using a floor instruction.  
+		// We therefore need to do the same times-2 transform (with a slighly different formula) that RoundToInt does; see the note on RoundToInt
+		return _mm_cvt_ss2si(_mm_set_ss(F + F - 0.5f)) >> 1;
+	}
+
+	static FORCEINLINE float FloorToFloat(float F)
+	{
+		return _mm_cvtss_f32(_mm_floor_ps(_mm_set_ss(F)));
+	}
+
+	static FORCEINLINE double FloorToDouble(double F)
+	{
+		return _mm_cvtsd_f64(_mm_floor_pd(_mm_set_sd(F)));
+	}
+
+	static FORCEINLINE int32 RoundToInt(float F)
+	{
+		// Note: the times-2 is to remove the rounding-to-nearest-even-number behavior that mm_cvt_ss2si uses when the fraction is .5
+		// The formula we uses causes the round instruction to always be applied to a an odd integer when the original value was 0.5, and eliminates the rounding-to-nearest-even-number behavior
+		// Input -> multiply by two and add .5 -> Round to nearest whole -> divide by two and truncate
+		// N -> (2N) + .5 -> 2N (or possibly 2N+1) -> N
+		// N + .5 -> 2N + 1.5 -> (round towards even now always means round up) -> 2N + 2 -> N + 1
+		return _mm_cvt_ss2si(_mm_set_ss(F + F + 0.5f)) >> 1;
+	}
+
+	static FORCEINLINE float RoundToFloat(float F)
+	{
+		return FloorToFloat(F + 0.5f);
+	}
+
+	static FORCEINLINE double RoundToDouble(double F)
+	{
+		return FloorToDouble(F + 0.5);
+	}
+
+	static FORCEINLINE int32 CeilToInt(float F)
+	{
+		// Note: unlike the Generic solution and the float solution, we implement CeilToInt using a rounding instruction, rather than a dedicated ceil instruction
+		// We therefore need to do the same times-2 transform (with a slighly different formula) that RoundToInt does; see the note on RoundToInt
+		return -(_mm_cvt_ss2si(_mm_set_ss(-0.5f - (F + F))) >> 1);
+	}
+
+	static FORCEINLINE float CeilToFloat(float F)
+	{
+		return _mm_cvtss_f32(_mm_ceil_ps(_mm_set_ss(F)));
+	}
+
+	static FORCEINLINE double CeilToDouble(double F)
+	{
+		return _mm_cvtsd_f64(_mm_ceil_pd(_mm_set_sd(F)));
+	}
 }
diff --git a/Engine/Source/Runtime/Core/Public/Unix/UnixPlatformMath.h b/Engine/Source/Runtime/Core/Public/Unix/UnixPlatformMath.h
--- a/Engine/Source/Runtime/Core/Public/Unix/UnixPlatformMath.h
+++ b/Engine/Source/Runtime/Core/Public/Unix/UnixPlatformMath.h
@@ -21,45 +21,62 @@ struct FLinuxPlatformMath : public FClangPlatformMath
 #if PLATFORM_ENABLE_VECTORINTRINSICS
 	static FORCEINLINE int32 TruncToInt(float F)
 	{
-		return _mm_cvtt_ss2si(_mm_set_ss(F));
+		return UnrealPlatformMathSSE::TruncToInt(F);
 	}
 
 	static FORCEINLINE float TruncToFloat(float F)
 	{
-		return (float)TruncToInt(F); // same as generic implementation, but this will call the faster trunc
+		return UnrealPlatformMathSSE::TruncToFloat(F);
+	}
+
+	static FORCEINLINE double TruncToDouble(double F)
+	{
+		return UnrealPlatformMathSSE::TruncToDouble(F);
 	}
 
 	static FORCEINLINE int32 RoundToInt(float F)
 	{
-		// Note: the x2 is to workaround the rounding-to-nearest-even-number issue when the fraction is .5
-		return _mm_cvt_ss2si(_mm_set_ss(F + F + 0.5f)) >> 1;
+		return UnrealPlatformMathSSE::RoundToInt(F);
 	}
 
 	static FORCEINLINE float RoundToFloat(float F)
 	{
-		return (float)RoundToInt(F);
+		return UnrealPlatformMathSSE::RoundToFloat(F);
+	}
+
+	static FORCEINLINE double RoundToDouble(double F)
+	{
+		return UnrealPlatformMathSSE::RoundToDouble(F);
 	}
 
 	static FORCEINLINE int32 FloorToInt(float F)
 	{
-		return _mm_cvt_ss2si(_mm_set_ss(F + F - 0.5f)) >> 1;
+		return UnrealPlatformMathSSE::FloorToInt(F);
 	}
 
 	static FORCEINLINE float FloorToFloat(float F)
 	{
-		return (float)FloorToInt(F);
+		return UnrealPlatformMathSSE::FloorToFloat(F);
+	}
+
+	static FORCEINLINE double FloorToDouble(double F)
+	{
+		return UnrealPlatformMathSSE::FloorToDouble(F);
 	}
 
 	static FORCEINLINE int32 CeilToInt(float F)
 	{
-		// Note: the x2 is to workaround the rounding-to-nearest-even-number issue when the fraction is .5
-		return -(_mm_cvt_ss2si(_mm_set_ss(-0.5f - (F + F))) >> 1);
+		return UnrealPlatformMathSSE::CeilToInt(F);
 	}
 
 	static FORCEINLINE float CeilToFloat(float F)
 	{
-		// Note: the x2 is to workaround the rounding-to-nearest-even-number issue when the fraction is .5
-		return (float)CeilToInt(F);
+		return UnrealPlatformMathSSE::CeilToFloat(F);
+	}
+
+	static FORCEINLINE double CeilToDouble(double F)
+	{
+		return UnrealPlatformMathSSE::CeilToDouble(F);
 	}
 
 	static FORCEINLINE bool IsNaN( float A ) { return isnan(A) != 0; }
diff --git a/Engine/Source/Runtime/Core/Public/Windows/WindowsPlatformMath.h b/Engine/Source/Runtime/Core/Public/Windows/WindowsPlatformMath.h
--- a/Engine/Source/Runtime/Core/Public/Windows/WindowsPlatformMath.h
+++ b/Engine/Source/Runtime/Core/Public/Windows/WindowsPlatformMath.h
@@ -24,45 +24,62 @@ struct FWindowsPlatformMath : public FGenericPlatformMath
 #if PLATFORM_ENABLE_VECTORINTRINSICS
 	static FORCEINLINE int32 TruncToInt(float F)
 	{
-		return _mm_cvtt_ss2si(_mm_set_ss(F));
+		return UnrealPlatformMathSSE::TruncToInt(F);
 	}
 
-	static FORCEINLINE float TruncToFloat( float F )
+	static FORCEINLINE float TruncToFloat(float F)
 	{
-		return (float)TruncToInt(F); // same as generic implementation, but this will call the faster trunc
+		return UnrealPlatformMathSSE::TruncToFloat(F);
+	}
+
+	static FORCEINLINE double TruncToDouble(double F)
+	{
+		return UnrealPlatformMathSSE::TruncToDouble(F);
 	}
 
 	static FORCEINLINE int32 RoundToInt( float F )
 	{
-		// Note: the x2 is to workaround the rounding-to-nearest-even-number issue when the fraction is .5
-		return _mm_cvt_ss2si( _mm_set_ss(F + F + 0.5f) ) >> 1;
+		return UnrealPlatformMathSSE::RoundToInt(F);
 	}
 
 	static FORCEINLINE float RoundToFloat(float F)
 	{
-		return (float)RoundToInt(F);
+		return UnrealPlatformMathSSE::RoundToFloat(F);
+	}
+
+	static FORCEINLINE double RoundToDouble(double F)
+	{
+		return UnrealPlatformMathSSE::RoundToDouble(F);
 	}
 
 	static FORCEINLINE int32 FloorToInt(float F)
 	{
-		return _mm_cvt_ss2si(_mm_set_ss(F + F - 0.5f)) >> 1;
+		return UnrealPlatformMathSSE::FloorToInt(F);
 	}
 
 	static FORCEINLINE float FloorToFloat(float F)
 	{
-		return (float)FloorToInt(F);
+		return UnrealPlatformMathSSE::FloorToFloat(F);
+	}
+
+	static FORCEINLINE double FloorToDouble(double F)
+	{
+		return UnrealPlatformMathSSE::FloorToDouble(F);
 	}
 
 	static FORCEINLINE int32 CeilToInt(float F)
 	{
-		// Note: the x2 is to workaround the rounding-to-nearest-even-number issue when the fraction is .5
-		return -(_mm_cvt_ss2si(_mm_set_ss(-0.5f - (F + F))) >> 1);
+		return UnrealPlatformMathSSE::CeilToInt(F);
 	}
 
 	static FORCEINLINE float CeilToFloat(float F)
 	{
-		// Note: the x2 is to workaround the rounding-to-nearest-even-number issue when the fraction is .5
-		return (float)CeilToInt(F);
+		return UnrealPlatformMathSSE::CeilToFloat(F);
+	}
+
+	static FORCEINLINE double CeilToDouble(double F)
+	{
+		return UnrealPlatformMathSSE::CeilToDouble(F);
 	}
 
 	static FORCEINLINE bool IsNaN( float A ) { return _isnan(A) != 0; }
