Fix linux compile: Turn off _mm_round_ps intrinsics for clang; clang supports _mm_cvtt_ss2si but does support __mm_round_ps

diff --git a/Engine/Source/Runtime/Core/Public/Math/UnrealPlatformMathSSE.h b/Engine/Source/Runtime/Core/Public/Math/UnrealPlatformMathSSE.h
--- a/Engine/Source/Runtime/Core/Public/Math/UnrealPlatformMathSSE.h
+++ b/Engine/Source/Runtime/Core/Public/Math/UnrealPlatformMathSSE.h
@@ -74,6 +74,7 @@ namespace UnrealPlatformMathSSE
 		return _mm_cvtt_ss2si(_mm_set_ss(F));
 	}
 
+#if !defined(__clang__) // Clang does not yet support _mm_round_ps/_mm_round_pd
 	static FORCEINLINE float TruncToFloat(float F)
 	{
 		return _mm_cvtss_f32(_mm_round_ps(_mm_set_ss(F), 3));
@@ -83,6 +84,7 @@ namespace UnrealPlatformMathSSE
 	{
 		return _mm_cvtsd_f64(_mm_round_pd(_mm_set_sd(F), 3));
 	}
+#endif
 
 	static FORCEINLINE int32 FloorToInt(float F)
 	{
@@ -91,6 +93,7 @@ namespace UnrealPlatformMathSSE
 		return _mm_cvt_ss2si(_mm_set_ss(F + F - 0.5f)) >> 1;
 	}
 
+#if !defined(__clang__) // Clang does not yet support _mm_floor_ps/_mm_floor_pd
 	static FORCEINLINE float FloorToFloat(float F)
 	{
 		return _mm_cvtss_f32(_mm_floor_ps(_mm_set_ss(F)));
@@ -100,6 +103,7 @@ namespace UnrealPlatformMathSSE
 	{
 		return _mm_cvtsd_f64(_mm_floor_pd(_mm_set_sd(F)));
 	}
+#endif
 
 	static FORCEINLINE int32 RoundToInt(float F)
 	{
@@ -111,6 +115,7 @@ namespace UnrealPlatformMathSSE
 		return _mm_cvt_ss2si(_mm_set_ss(F + F + 0.5f)) >> 1;
 	}
 
+#if !defined(__clang__) // Clang does not yet support _mm_floor_ps/_mm_floor_pd
 	static FORCEINLINE float RoundToFloat(float F)
 	{
 		return FloorToFloat(F + 0.5f);
@@ -120,7 +125,7 @@ namespace UnrealPlatformMathSSE
 	{
 		return FloorToDouble(F + 0.5);
 	}
-
+#endif
 	static FORCEINLINE int32 CeilToInt(float F)
 	{
 		// Note: unlike the Generic solution and the float solution, we implement CeilToInt using a rounding instruction, rather than a dedicated ceil instruction
@@ -128,6 +133,7 @@ namespace UnrealPlatformMathSSE
 		return -(_mm_cvt_ss2si(_mm_set_ss(-0.5f - (F + F))) >> 1);
 	}
 
+#if !defined(__clang__) // Clang does not yet support _mm_ceil_ps/_mm_ceil_pd
 	static FORCEINLINE float CeilToFloat(float F)
 	{
 		return _mm_cvtss_f32(_mm_ceil_ps(_mm_set_ss(F)));
@@ -137,4 +143,5 @@ namespace UnrealPlatformMathSSE
 	{
 		return _mm_cvtsd_f64(_mm_ceil_pd(_mm_set_sd(F)));
 	}
+#endif
 }
diff --git a/Engine/Source/Runtime/Core/Public/Unix/UnixPlatformMath.h b/Engine/Source/Runtime/Core/Public/Unix/UnixPlatformMath.h
--- a/Engine/Source/Runtime/Core/Public/Unix/UnixPlatformMath.h
+++ b/Engine/Source/Runtime/Core/Public/Unix/UnixPlatformMath.h
@@ -24,6 +24,7 @@ struct FLinuxPlatformMath : public FClangPlatformMath
 		return UnrealPlatformMathSSE::TruncToInt(F);
 	}
 
+#if !defined(__clang__) // Clang does not yet support _mm_round_ps/_mm_round_pd
 	static FORCEINLINE float TruncToFloat(float F)
 	{
 		return UnrealPlatformMathSSE::TruncToFloat(F);
@@ -33,12 +34,14 @@ struct FLinuxPlatformMath : public FClangPlatformMath
 	{
 		return UnrealPlatformMathSSE::TruncToDouble(F);
 	}
+#endif
 
 	static FORCEINLINE int32 RoundToInt(float F)
 	{
 		return UnrealPlatformMathSSE::RoundToInt(F);
 	}
 
+#if !defined(__clang__) // Clang does not yet support _mm_floor_ps/_mm_floor_pd
 	static FORCEINLINE float RoundToFloat(float F)
 	{
 		return UnrealPlatformMathSSE::RoundToFloat(F);
@@ -48,12 +51,14 @@ struct FLinuxPlatformMath : public FClangPlatformMath
 	{
 		return UnrealPlatformMathSSE::RoundToDouble(F);
 	}
+#endif
 
 	static FORCEINLINE int32 FloorToInt(float F)
 	{
 		return UnrealPlatformMathSSE::FloorToInt(F);
 	}
 
+#if !defined(__clang__) // Clang does not yet support _mm_floor_ps/_mm_floor_pd
 	static FORCEINLINE float FloorToFloat(float F)
 	{
 		return UnrealPlatformMathSSE::FloorToFloat(F);
@@ -63,12 +68,14 @@ struct FLinuxPlatformMath : public FClangPlatformMath
 	{
 		return UnrealPlatformMathSSE::FloorToDouble(F);
 	}
+#endif
 
 	static FORCEINLINE int32 CeilToInt(float F)
 	{
 		return UnrealPlatformMathSSE::CeilToInt(F);
 	}
 
+#if !defined(__clang__) // Clang does not yet support _mm_ceil_ps/_mm_ceil_pd
 	static FORCEINLINE float CeilToFloat(float F)
 	{
 		return UnrealPlatformMathSSE::CeilToFloat(F);
@@ -78,6 +85,7 @@ struct FLinuxPlatformMath : public FClangPlatformMath
 	{
 		return UnrealPlatformMathSSE::CeilToDouble(F);
 	}
+#endif
 
 	static FORCEINLINE bool IsNaN( float A ) { return isnan(A) != 0; }
 	static FORCEINLINE bool IsFinite( float A ) { return isfinite(A); }
