Fix Linux/WindowsClang compile: need to use PLATFORM_LINUX rather than __clang__ to turn off the intrinsics on linux.

diff --git a/Engine/Source/Runtime/Core/Public/Math/UnrealPlatformMathSSE.h b/Engine/Source/Runtime/Core/Public/Math/UnrealPlatformMathSSE.h
--- a/Engine/Source/Runtime/Core/Public/Math/UnrealPlatformMathSSE.h
+++ b/Engine/Source/Runtime/Core/Public/Math/UnrealPlatformMathSSE.h
@@ -74,7 +74,7 @@ namespace UnrealPlatformMathSSE
 		return _mm_cvtt_ss2si(_mm_set_ss(F));
 	}
 
-#if !defined(__clang__) // Clang does not yet support _mm_round_ps/_mm_round_pd
+#if !PLATFORM_LINUX // Linux Clang does not yet support _mm_round_ps/_mm_round_pd
 	static FORCEINLINE float TruncToFloat(float F)
 	{
 		return _mm_cvtss_f32(_mm_round_ps(_mm_set_ss(F), 3));
@@ -93,7 +93,7 @@ namespace UnrealPlatformMathSSE
 		return _mm_cvt_ss2si(_mm_set_ss(F + F - 0.5f)) >> 1;
 	}
 
-#if !defined(__clang__) // Clang does not yet support _mm_floor_ps/_mm_floor_pd
+#if !PLATFORM_LINUX // Linux Clang does not yet support _mm_round_ps/_mm_round_pd
 	static FORCEINLINE float FloorToFloat(float F)
 	{
 		return _mm_cvtss_f32(_mm_floor_ps(_mm_set_ss(F)));
@@ -115,7 +115,7 @@ namespace UnrealPlatformMathSSE
 		return _mm_cvt_ss2si(_mm_set_ss(F + F + 0.5f)) >> 1;
 	}
 
-#if !defined(__clang__) // Clang does not yet support _mm_floor_ps/_mm_floor_pd
+#if !PLATFORM_LINUX // Linux Clang does not yet support _mm_round_ps/_mm_round_pd
 	static FORCEINLINE float RoundToFloat(float F)
 	{
 		return FloorToFloat(F + 0.5f);
@@ -133,7 +133,7 @@ namespace UnrealPlatformMathSSE
 		return -(_mm_cvt_ss2si(_mm_set_ss(-0.5f - (F + F))) >> 1);
 	}
 
-#if !defined(__clang__) // Clang does not yet support _mm_ceil_ps/_mm_ceil_pd
+#if !PLATFORM_LINUX // Linux Clang does not yet support _mm_round_ps/_mm_round_pd
 	static FORCEINLINE float CeilToFloat(float F)
 	{
 		return _mm_cvtss_f32(_mm_ceil_ps(_mm_set_ss(F)));
diff --git a/Engine/Source/Runtime/Core/Public/Unix/UnixPlatformMath.h b/Engine/Source/Runtime/Core/Public/Unix/UnixPlatformMath.h
--- a/Engine/Source/Runtime/Core/Public/Unix/UnixPlatformMath.h
+++ b/Engine/Source/Runtime/Core/Public/Unix/UnixPlatformMath.h
@@ -24,7 +24,7 @@ struct FLinuxPlatformMath : public FClangPlatformMath
 		return UnrealPlatformMathSSE::TruncToInt(F);
 	}
 
-#if !defined(__clang__) // Clang does not yet support _mm_round_ps/_mm_round_pd
+#if !PLATFORM_LINUX // Linux Clang does not yet support _mm_round_ps/_mm_round_pd
 	static FORCEINLINE float TruncToFloat(float F)
 	{
 		return UnrealPlatformMathSSE::TruncToFloat(F);
@@ -41,7 +41,7 @@ struct FLinuxPlatformMath : public FClangPlatformMath
 		return UnrealPlatformMathSSE::RoundToInt(F);
 	}
 
-#if !defined(__clang__) // Clang does not yet support _mm_floor_ps/_mm_floor_pd
+#if !PLATFORM_LINUX // Linux Clang does not yet support _mm_round_ps/_mm_round_pd
 	static FORCEINLINE float RoundToFloat(float F)
 	{
 		return UnrealPlatformMathSSE::RoundToFloat(F);
@@ -58,7 +58,7 @@ struct FLinuxPlatformMath : public FClangPlatformMath
 		return UnrealPlatformMathSSE::FloorToInt(F);
 	}
 
-#if !defined(__clang__) // Clang does not yet support _mm_floor_ps/_mm_floor_pd
+#if !PLATFORM_LINUX // Linux Clang does not yet support _mm_round_ps/_mm_round_pd
 	static FORCEINLINE float FloorToFloat(float F)
 	{
 		return UnrealPlatformMathSSE::FloorToFloat(F);
@@ -75,7 +75,7 @@ struct FLinuxPlatformMath : public FClangPlatformMath
 		return UnrealPlatformMathSSE::CeilToInt(F);
 	}
 
-#if !defined(__clang__) // Clang does not yet support _mm_ceil_ps/_mm_ceil_pd
+#if !PLATFORM_LINUX // Linux Clang does not yet support _mm_round_ps/_mm_round_pd
 	static FORCEINLINE float CeilToFloat(float F)
 	{
 		return UnrealPlatformMathSSE::CeilToFloat(F);
