Handle attempting to read past the end of a file as failure so that it behaves the same way as our other platform implementations.

It looks like this behaviour got lost when the read/write functions were changed to overlapped implementations in CL 3407495, so this change will just restore the lost functionality.
Note that we will attempt to record what ever data was written/read before returning false.

diff --git a/Engine/Source/Runtime/Core/Private/Windows/WindowsPlatformFile.cpp b/Engine/Source/Runtime/Core/Private/Windows/WindowsPlatformFile.cpp
--- a/Engine/Source/Runtime/Core/Private/Windows/WindowsPlatformFile.cpp
+++ b/Engine/Source/Runtime/Core/Private/Windows/WindowsPlatformFile.cpp
@@ -590,6 +590,13 @@ public:
 			// Update where we are in the file
 			FilePos += NumRead;
 			UpdateOverlappedPos();
+			
+			// Early out as a failure case if we did not read all of the bytes that we expected to read
+			if (BytesToRead32 != NumRead)
+			{
+				return false; 
+			}	
+					
 		} while (BytesToRead > 0);
 		TRACE_PLATFORMFILE_END_READ(&OverlappedIO, TotalNumRead);
 		return true;
@@ -624,6 +631,7 @@ public:
 					return false;
 				}
 			}
+	
 			BytesToWrite -= BytesToWrite32;
 			Source += BytesToWrite32;
 			TotalNumWritten += NumWritten;
@@ -631,6 +639,13 @@ public:
 			FilePos += NumWritten;
 			UpdateOverlappedPos();
 			FileSize = FMath::Max(FilePos, FileSize);
+			
+			// Early out as a failure case if we didn't write all of the data we expected
+			if (BytesToWrite32 != NumWritten)
+			{
+				return false;
+			}
+			
 		} while (BytesToWrite > 0);
 
 		TRACE_PLATFORMFILE_END_WRITE(this, TotalNumWritten);
