Prevent the LLM internal arrays from tempoarily overflowing their 32bit bounds while either calculating the new capacity or determining if we should trim the internal storage capacity. We will now only overflow if the final new capacity is too large for 32bit.

diff --git a/Engine/Source/Runtime/Core/Private/HAL/LowLevelMemoryUtils.h b/Engine/Source/Runtime/Core/Private/HAL/LowLevelMemoryUtils.h
--- a/Engine/Source/Runtime/Core/Private/HAL/LowLevelMemoryUtils.h
+++ b/Engine/Source/Runtime/Core/Private/HAL/LowLevelMemoryUtils.h
@@ -57,7 +57,7 @@ public:
 			uint32 NewCapacity = DefaultCapacity;
 			if (Capacity)
 			{
-				NewCapacity = (15U * Capacity) / 10U;
+				NewCapacity = Capacity + (Capacity / 2);
 				ensureMsgf(NewCapacity > Capacity, TEXT("Unsigned integer overflow."));
 			}
 			Reserve(NewCapacity);
@@ -145,10 +145,11 @@ public:
 
 	void Trim()
 	{
-		if (Array != StaticArray && (Count * 128) / Capacity < 100)
+		// Trim if usage has dropped below 3/4 of the total capacity
+		if (Array != StaticArray && Count < (Capacity - (Capacity / 4)))
 		{
 			Reserve(Count);
-		}
+		} 
 	}
 
 private:
