Some traits moved out into their own headers.

#rnx

diff --git a/Engine/Source/Runtime/Core/Public/Templates/CopyQualifiersFromTo.h b/Engine/Source/Runtime/Core/Public/Templates/CopyQualifiersFromTo.h
new file mode 100644
--- /dev/null
+++ b/Engine/Source/Runtime/Core/Public/Templates/CopyQualifiersFromTo.h
@@ -0,0 +1,14 @@
+// Copyright 1998-2019 Epic Games, Inc. All Rights Reserved.
+
+#pragma once
+
+/**
+ * Copies the cv-qualifiers from one type to another, e.g.:
+ *
+ * TCopyQualifiersFromTo<const    T1,       T2>::Type == const T2
+ * TCopyQualifiersFromTo<volatile T1, const T2>::Type == const volatile T2
+ */
+template <typename From, typename To> struct TCopyQualifiersFromTo                          { typedef                To Type; };
+template <typename From, typename To> struct TCopyQualifiersFromTo<const          From, To> { typedef const          To Type; };
+template <typename From, typename To> struct TCopyQualifiersFromTo<      volatile From, To> { typedef       volatile To Type; };
+template <typename From, typename To> struct TCopyQualifiersFromTo<const volatile From, To> { typedef const volatile To Type; };
diff --git a/Engine/Source/Runtime/Core/Public/Templates/LosesQualifiersFromTo.h b/Engine/Source/Runtime/Core/Public/Templates/LosesQualifiersFromTo.h
new file mode 100644
--- /dev/null
+++ b/Engine/Source/Runtime/Core/Public/Templates/LosesQualifiersFromTo.h
@@ -0,0 +1,18 @@
+// Copyright 1998-2019 Epic Games, Inc. All Rights Reserved.
+
+#pragma once
+
+#include "Templates/AreTypesEqual.h"
+#include "Templates/CopyQualifiersFromTo.h"
+
+/**
+ * Tests if qualifiers are lost between one type and another, e.g.:
+ *
+ * TLosesQualifiersFromTo<const    T1,                T2>::Value == true
+ * TLosesQualifiersFromTo<volatile T1, const volatile T2>::Value == false
+ */
+template <typename From, typename To>
+struct TLosesQualifiersFromTo
+{
+	enum { Value = !TAreTypesEqual<typename TCopyQualifiersFromTo<From, To>::Type, To>::Value };
+};
diff --git a/Engine/Source/Runtime/Core/Public/Templates/UnrealTemplate.h b/Engine/Source/Runtime/Core/Public/Templates/UnrealTemplate.h
--- a/Engine/Source/Runtime/Core/Public/Templates/UnrealTemplate.h
+++ b/Engine/Source/Runtime/Core/Public/Templates/UnrealTemplate.h
@@ -548,29 +548,6 @@ struct FNoopStruct
 	{}
 };
 
-/**
- * Copies the cv-qualifiers from one type to another, e.g.:
- *
- * TCopyQualifiersFromTo<const    T1,       T2>::Type == const T2
- * TCopyQualifiersFromTo<volatile T1, const T2>::Type == const volatile T2
- */
-template <typename From, typename To> struct TCopyQualifiersFromTo                          { typedef                To Type; };
-template <typename From, typename To> struct TCopyQualifiersFromTo<const          From, To> { typedef const          To Type; };
-template <typename From, typename To> struct TCopyQualifiersFromTo<      volatile From, To> { typedef       volatile To Type; };
-template <typename From, typename To> struct TCopyQualifiersFromTo<const volatile From, To> { typedef const volatile To Type; };
-
-/**
- * Tests if qualifiers are lost between one type and another, e.g.:
- *
- * TCopyQualifiersFromTo<const    T1,                T2>::Value == true
- * TCopyQualifiersFromTo<volatile T1, const volatile T2>::Value == false
- */
-template <typename From, typename To>
-struct TLosesQualifiersFromTo
-{
-	enum { Value = !TAreTypesEqual<typename TCopyQualifiersFromTo<From, To>::Type, To>::Value };
-};
-
 /**
  * Equivalent to std::declval.
  *
diff --git a/Engine/Source/Runtime/Core/Public/Templates/UnrealTypeTraits.h b/Engine/Source/Runtime/Core/Public/Templates/UnrealTypeTraits.h
--- a/Engine/Source/Runtime/Core/Public/Templates/UnrealTypeTraits.h
+++ b/Engine/Source/Runtime/Core/Public/Templates/UnrealTypeTraits.h
@@ -16,6 +16,7 @@
 #include "Templates/IsArithmetic.h"
 #include "Templates/IsEnum.h"
 #include "Templates/RemoveCV.h"
+#include "Traits/IsVoidType.h"
 
 #include "Templates/IsPODType.h"
 #include "Templates/IsTriviallyCopyConstructible.h"
@@ -135,15 +136,6 @@ template<typename T> struct TIsLValueReferenceType<T&> { enum { Value = true  };
 template<typename T> struct TIsRValueReferenceType      { enum { Value = false }; };
 template<typename T> struct TIsRValueReferenceType<T&&> { enum { Value = true  }; };
 
-/**
- * TIsVoidType
- */
-template<typename T> struct TIsVoidType { enum { Value = false }; };
-template<> struct TIsVoidType<void> { enum { Value = true }; };
-template<> struct TIsVoidType<void const> { enum { Value = true }; };
-template<> struct TIsVoidType<void volatile> { enum { Value = true }; };
-template<> struct TIsVoidType<void const volatile> { enum { Value = true }; };
-
 /**
  * TIsFundamentalType
  */
diff --git a/Engine/Source/Runtime/Core/Public/Traits/IsVoidType.h b/Engine/Source/Runtime/Core/Public/Traits/IsVoidType.h
new file mode 100644
--- /dev/null
+++ b/Engine/Source/Runtime/Core/Public/Traits/IsVoidType.h
@@ -0,0 +1,12 @@
+// Copyright 1998-2019 Epic Games, Inc. All Rights Reserved.
+
+#pragma once
+
+/**
+ * TIsVoidType
+ */
+template<typename T> struct TIsVoidType { enum { Value = false }; };
+template<> struct TIsVoidType<void> { enum { Value = true }; };
+template<> struct TIsVoidType<void const> { enum { Value = true }; };
+template<> struct TIsVoidType<void volatile> { enum { Value = true }; };
+template<> struct TIsVoidType<void const volatile> { enum { Value = true }; };
diff --git a/Engine/Source/Runtime/Core/Public/UObject/WeakObjectPtrTemplates.h b/Engine/Source/Runtime/Core/Public/UObject/WeakObjectPtrTemplates.h
--- a/Engine/Source/Runtime/Core/Public/UObject/WeakObjectPtrTemplates.h
+++ b/Engine/Source/Runtime/Core/Public/UObject/WeakObjectPtrTemplates.h
@@ -7,6 +7,7 @@
 #include "Templates/PointerIsConvertibleFromTo.h"
 #include "Templates/AreTypesEqual.h"
 #include "Templates/AndOrNot.h"
+#include "Templates/LosesQualifiersFromTo.h"
 #include "Containers/Map.h"
 
 /**
diff --git a/Engine/Source/Runtime/CoreUObject/Public/Templates/Casts.h b/Engine/Source/Runtime/CoreUObject/Public/Templates/Casts.h
--- a/Engine/Source/Runtime/CoreUObject/Public/Templates/Casts.h
+++ b/Engine/Source/Runtime/CoreUObject/Public/Templates/Casts.h
@@ -6,6 +6,8 @@
 #include "UObject/ObjectMacros.h"
 #include "UObject/Object.h"
 #include "UObject/Class.h"
+#include "Templates/LosesQualifiersFromTo.h"
+#include "Traits/IsVoidType.h"
 
 class AActor;
 class APawn;
diff --git a/Engine/Source/Runtime/Engine/Public/Net/RepLayout.h b/Engine/Source/Runtime/Engine/Public/Net/RepLayout.h
--- a/Engine/Source/Runtime/Engine/Public/Net/RepLayout.h
+++ b/Engine/Source/Runtime/Engine/Public/Net/RepLayout.h
@@ -21,6 +21,7 @@
 #include "UObject/GCObject.h"
 #include "Containers/StaticBitArray.h"
 #include "Net/GuidReferences.h"
+#include "Templates/CopyQualifiersFromTo.h"
 
 class FGuidReferences;
 class FNetFieldExportGroup;
diff --git a/Engine/Source/Runtime/MeshDescription/Public/MeshAttributeArray.h b/Engine/Source/Runtime/MeshDescription/Public/MeshAttributeArray.h
--- a/Engine/Source/Runtime/MeshDescription/Public/MeshAttributeArray.h
+++ b/Engine/Source/Runtime/MeshDescription/Public/MeshAttributeArray.h
@@ -6,6 +6,7 @@
 #include "MeshTypes.h"
 #include "MeshElementRemappings.h"
 #include "Misc/TVariant.h"
+#include "Templates/CopyQualifiersFromTo.h"
 #include "UObject/ReleaseObjectVersion.h"
 
 
