Add ESearchCase parameter to StartsWith and EndsWith on StringView

#rnx

diff --git a/Engine/Source/Runtime/Core/Public/Misc/StringView.h b/Engine/Source/Runtime/Core/Public/Misc/StringView.h
--- a/Engine/Source/Runtime/Core/Public/Misc/StringView.h
+++ b/Engine/Source/Runtime/Core/Public/Misc/StringView.h
@@ -128,10 +128,10 @@ public:
 	CORE_API int32 Compare(const TStringViewImpl& Other, ESearchCase::Type SearchCase = ESearchCase::CaseSensitive) const;
 
 	inline bool StartsWith(CharType Prefix) const { return Size >= 1 && DataPtr[0] == Prefix; }
-	inline bool StartsWith(const TStringViewImpl& Prefix) const;
+	inline bool StartsWith(const TStringViewImpl& Prefix, ESearchCase::Type SearchCase = ESearchCase::IgnoreCase) const;
 
 	inline bool EndsWith(CharType Suffix) const { return Size >= 1 && DataPtr[Size-1] == Suffix; }
-	inline bool EndsWith(const TStringViewImpl& Suffix) const;
+	inline bool EndsWith(const TStringViewImpl& Suffix, ESearchCase::Type SearchCase = ESearchCase::IgnoreCase) const;
 
 	// Searching/Finding
 	CORE_API bool FindChar(CharType InChar, SizeType& OutIndex) const;
diff --git a/Engine/Source/Runtime/Core/Public/Misc/StringView.inl b/Engine/Source/Runtime/Core/Public/Misc/StringView.inl
--- a/Engine/Source/Runtime/Core/Public/Misc/StringView.inl
+++ b/Engine/Source/Runtime/Core/Public/Misc/StringView.inl
@@ -64,15 +64,15 @@ inline bool TStringViewImpl<CharType, ViewType>::Equals(const TStringViewImpl& O
 }
 
 template <typename CharType, typename ViewType>
-inline bool TStringViewImpl<CharType, ViewType>::StartsWith(const TStringViewImpl& Prefix) const
+inline bool TStringViewImpl<CharType, ViewType>::StartsWith(const TStringViewImpl& Prefix, ESearchCase::Type SearchCase) const
 {
-	return Size >= Prefix.Size && TCString<CharType>::Strnicmp(DataPtr, Prefix.DataPtr, Prefix.Size) == 0;
+	return Prefix.Equals(Left(Prefix.Len()), SearchCase);
 }
 
 template <typename CharType, typename ViewType>
-inline bool TStringViewImpl<CharType, ViewType>::EndsWith(const TStringViewImpl& Suffix) const
+inline bool TStringViewImpl<CharType, ViewType>::EndsWith(const TStringViewImpl& Suffix, ESearchCase::Type SearchCase) const
 {
-	return Size >= Suffix.Size && TCString<CharType>::Strnicmp(DataPtr + Size - Suffix.Size, Suffix.DataPtr, Suffix.Size) == 0;
+	return Suffix.Equals(Right(Suffix.Len()), SearchCase);
 }
 
 // Case-insensitive comparison operators
