Iteration on StringBuilder

- Add TStringBuilderImpl::Appendf to allow formatting values into a string builder.
- Add TStringBuilderImpl::Reset to allow reusing a string builder.
- Add operator<< for writing integers to string builders. These are implemented inefficiently using Appendf for now and can be optimized later if necessary.

diff --git a/Engine/Source/Runtime/Core/Private/Misc/StringBuilder.cpp b/Engine/Source/Runtime/Core/Private/Misc/StringBuilder.cpp
--- a/Engine/Source/Runtime/Core/Private/Misc/StringBuilder.cpp
+++ b/Engine/Source/Runtime/Core/Private/Misc/StringBuilder.cpp
@@ -65,6 +65,31 @@ TStringBuilderImpl<C>::FreeBuffer(void* Buffer, int32 ByteCount)
 	FMemory::Free(Buffer);
 }
 
+template<typename C>
+void
+TStringBuilderImpl<C>::AppendfImpl(TStringBuilderImpl& Self, const C* Fmt, ...)
+{
+	for (;;)
+	{
+		va_list ArgPack;
+		va_start(ArgPack, Fmt);
+		const int32 RemainingSize = Self.End - Self.CurPos;
+		const int32 Result = TCString<C>::GetVarArgs(Self.CurPos, RemainingSize, Fmt, ArgPack);
+		va_end(ArgPack);
+
+		if (Result >= 0 && Result < RemainingSize)
+		{
+			Self.CurPos += Result;
+			return;
+		}
+		else
+		{
+			// Total size will be rounded up to the next power of two. Start with at least 64.
+			Self.Extend(64);
+		}
+	}
+}
+
 // Instantiate templates once
 
 template class TStringBuilderImpl<ANSICHAR>;
diff --git a/Engine/Source/Runtime/Core/Public/Misc/StringBuilder.h b/Engine/Source/Runtime/Core/Public/Misc/StringBuilder.h
--- a/Engine/Source/Runtime/Core/Public/Misc/StringBuilder.h
+++ b/Engine/Source/Runtime/Core/Public/Misc/StringBuilder.h
@@ -5,6 +5,10 @@
 #include "CoreTypes.h"
 #include "Misc/CString.h"
 #include "Misc/StringView.h"
+#include "Templates/AndOrNot.h"
+#include "Templates/EnableIf.h"
+#include "Templates/IsArrayOrRefOfType.h"
+#include "Templates/IsValidVariadicFunctionArg.h"
 
 #define USE_STRING_LITERAL_PATH 1
 
@@ -61,6 +65,14 @@ public:
 
 	inline const C	LastChar() const	{ return *(CurPos - 1); }
 
+	/**
+	 * Empties the string builder, but doesn't change memory allocation.
+	 */
+	inline void Reset()
+	{
+		CurPos = Base;
+	}
+
 	inline TStringBuilderImpl& Append(C Char)
 	{
 		EnsureCapacity(1);
@@ -180,7 +192,22 @@ public:
 	}
 #endif
 
+	/**
+	 * Appends to the string builder similarly to how classic sprintf works.
+	 *
+	 * @param Format A format string that specifies how to format the additional arguments. Refer to standard printf format.
+	 */
+	template <typename FmtType, typename... Types>
+	typename TEnableIf<TIsArrayOrRefOfType<FmtType, C>::Value>::Type Appendf(const FmtType& Fmt, Types... Args)
+	{
+		static_assert(TIsArrayOrRefOfType<FmtType, C>::Value, "Formatting string must be a character array.");
+		static_assert(TAnd<TIsValidVariadicFunctionArg<Types>...>::Value, "Invalid argument(s) passed to Appendf.");
+		AppendfImpl(*this, Fmt, Forward<Types>(Args)...);
+	}
+
 protected:
+	CORE_API static void VARARGS AppendfImpl(TStringBuilderImpl& Self, const C* Fmt, ...);
+
 	inline void Initialize(C* InBase, int32 InCapacity)
 	{
 		Base	= InBase;
@@ -355,3 +382,15 @@ inline FStringBuilderBase&					operator<<(FStringBuilderBase& Builder, const ANS
 #endif
 
 inline FStringBuilderBase&					operator<<(FStringBuilderBase& Builder, const FStringView& Str)					{ Builder.Append(Str.GetData(), Str.Len()); return Builder; }
+
+inline FStringBuilderBase&					operator<<(FStringBuilderBase& Builder, int32 Value)							{ Builder.Appendf(TEXT("%d"), Value); return Builder; }
+inline FStringBuilderBase&					operator<<(FStringBuilderBase& Builder, uint32 Value)							{ Builder.Appendf(TEXT("%u"), Value); return Builder; }
+
+inline FStringBuilderBase&					operator<<(FStringBuilderBase& Builder, int64 Value)							{ Builder.Appendf(TEXT(INT64_FMT), Value); return Builder; }
+inline FStringBuilderBase&					operator<<(FStringBuilderBase& Builder, uint64 Value)							{ Builder.Appendf(TEXT(UINT64_FMT), Value); return Builder; }
+
+inline FStringBuilderBase&					operator<<(FStringBuilderBase& Builder, int8 Value)								{ return Builder << int32(Value); }
+inline FStringBuilderBase&					operator<<(FStringBuilderBase& Builder, uint8 Value)							{ return Builder << uint32(Value); }
+
+inline FStringBuilderBase&					operator<<(FStringBuilderBase& Builder, int16 Value)							{ return Builder << int32(Value); }
+inline FStringBuilderBase&					operator<<(FStringBuilderBase& Builder, uint16 Value)							{ return Builder << uint32(Value); }
