Add FStringBuilderBase& overloads for FName to string conversions

- Optimize FNameEntry::GetPlainNameString() to avoid an unnecessary strlen and copy of the string.

diff --git a/Engine/Source/Runtime/Core/Private/UObject/UnrealNames.cpp b/Engine/Source/Runtime/Core/Private/UObject/UnrealNames.cpp
--- a/Engine/Source/Runtime/Core/Private/UObject/UnrealNames.cpp
+++ b/Engine/Source/Runtime/Core/Private/UObject/UnrealNames.cpp
@@ -9,6 +9,7 @@
 #include "Templates/UnrealTemplate.h"
 #include "Misc/CString.h"
 #include "Misc/Crc.h"
+#include "Misc/StringBuilder.h"
 #include "Containers/UnrealString.h"
 #include "UObject/NameTypes.h"
 #include "Logging/LogMacros.h"
@@ -1161,7 +1162,14 @@ static const TCHAR* EntryToCString(const FNameEntry& Entry, FNameBuffer& Temp)
 FString FNameEntry::GetPlainNameString() const
 {
 	FNameBuffer Temp;
-	return FString(EntryToCString(*this, Temp));
+	if (Header.bIsWide)
+	{
+		return FString(Header.Len, GetUnterminatedName(Temp.WideName));
+	}
+	else
+	{
+		return FString(Header.Len, GetUnterminatedName(Temp.AnsiName));
+	}
 }
 
 void FNameEntry::AppendNameToString(FString& Out) const
@@ -1170,12 +1178,39 @@ void FNameEntry::AppendNameToString(FString& Out) const
 	Out.Append(EntryToCString(*this, Temp), Header.Len);
 }
 
+void FNameEntry::AppendNameToString(FStringBuilderBase& Out) const
+{
+	FNameBuffer Temp;
+	if (Header.bIsWide)
+	{
+		Out.Append(GetUnterminatedName(Temp.WideName), Header.Len);
+	}
+	else
+	{
+		Out.AppendAnsi(GetUnterminatedName(Temp.AnsiName), Header.Len);
+	}
+}
+
 void FNameEntry::AppendNameToPathString(FString& Out) const
 {
 	FNameBuffer Temp;
 	Out.PathAppend(EntryToCString(*this, Temp), Header.Len);
 }
 
+void FNameEntry::AppendNameToPathString(FStringBuilderBase& Out) const
+{
+	if (Out.Len() > 0)
+	{
+		const TCHAR LastChar = Out.LastChar();
+		if (LastChar != TEXT('/') && LastChar != TEXT('\\'))
+		{
+			Out.Append(TEXT('/'));
+		}
+	}
+
+	AppendNameToString(Out);
+}
+
 int32 FNameEntry::GetSize(const TCHAR* Name)
 {
 	return FNameEntry::GetSize(FCString::Strlen(Name), FCString::IsPureAnsi(Name));
@@ -1808,6 +1843,12 @@ void FName::ToString(FString& Out) const
 	}
 }
 
+void FName::ToString(FStringBuilderBase& Out) const
+{
+	Out.Reset();
+	AppendString(Out);
+}
+
 uint32 FName::GetStringLength() const
 {
 	const FNameEntry& Entry = *GetDisplayNameEntry();
@@ -1863,6 +1904,17 @@ void FName::AppendString(FString& Out) const
 	}
 }
 
+void FName::AppendString(FStringBuilderBase& Out) const
+{
+	GetDisplayNameEntry()->AppendNameToString(Out);
+
+	const int32 InternalNumber = GetNumber();
+	if (InternalNumber != NAME_NO_NUMBER_INTERNAL)
+	{
+		Out << TEXT('_') << NAME_INTERNAL_TO_EXTERNAL(InternalNumber);
+	}
+}
+
 void FName::DisplayHash(FOutputDevice& Ar)
 {
 	GetNamePool().LogStats(Ar);
diff --git a/Engine/Source/Runtime/Core/Public/UObject/NameTypes.h b/Engine/Source/Runtime/Core/Public/UObject/NameTypes.h
--- a/Engine/Source/Runtime/Core/Public/UObject/NameTypes.h
+++ b/Engine/Source/Runtime/Core/Public/UObject/NameTypes.h
@@ -19,7 +19,7 @@
 
 /** 
  * Do we want to support case-variants for FName?
- * This will add an extra NAME_INunrealDEX variable to FName, but means that ToString() will return you the exact same 
+ * This will add an extra NAME_INDEX variable to FName, but means that ToString() will return you the exact same 
  * string that FName::Init was called with (which is useful if your FNames are shown to the end user)
  * Currently this is enabled for the Editor and any Programs (such as UHT), but not the Runtime
  */
@@ -27,6 +27,7 @@
 	#define WITH_CASE_PRESERVING_NAME WITH_EDITORONLY_DATA
 #endif
 
+class FStringBuilderBase;
 class FText;
 
 /** Maximum size of name. */
@@ -223,12 +224,21 @@ public:
 	/** Copy name to a dynamically allocated FString. */
 	CORE_API FString GetPlainNameString() const;
 
+	/** Copy name to a FStringBuilderBase. */
+	CORE_API void GetPlainNameString(FStringBuilderBase& OutString) const;
+
 	/** Appends name to string. May allocate. */
 	CORE_API void AppendNameToString(FString& OutString) const;
 
+	/** Appends name to string builder. */
+	CORE_API void AppendNameToString(FStringBuilderBase& OutString) const;
+
 	/** Appends name to string with path separator using FString::PathAppend(). */
 	CORE_API void AppendNameToPathString(FString& OutString) const;
 
+	/** Appends name to string builder with path separator. */
+	CORE_API void AppendNameToPathString(FStringBuilderBase& OutString) const;
+
 
 	/**
 	 * Returns the size in bytes for FNameEntry structure. This is != sizeof(FNameEntry) as we only allocated as needed.
@@ -437,6 +447,13 @@ public:
 	 */
 	void ToString(FString& Out) const;
 
+	/**
+	 * Converts an FName to a readable format, in place
+	 * 
+	 * @param Out StringBuilder to fill with the string representation of the name
+	 */
+	void ToString(FStringBuilderBase& Out) const;
+
 	/**
 	 * Get the number of characters, excluding null-terminator, that ToString() would yield
 	 */
@@ -469,6 +486,13 @@ public:
 	 */
 	void AppendString(FString& Out) const;
 
+	/**
+	 * Converts an FName to a readable format, in place, appending to an existing string (ala GetFullName)
+	 * 
+	 * @param Out StringBuilder to append with the string representation of the name
+	 */
+	void AppendString(FStringBuilderBase& Out) const;
+
 	/**
 	 * Check to see if this FName matches the other FName, potentially also checking for any case variations
 	 */
