Usage of TPointerIsConvertibleFromTo removed to prevent unintentional memoization when types are undefined, producing incorrect overloading.
Deprecated overloads removed.

diff --git a/Engine/Source/Runtime/Core/Public/Templates/SharedPointer.h b/Engine/Source/Runtime/Core/Public/Templates/SharedPointer.h
--- a/Engine/Source/Runtime/Core/Public/Templates/SharedPointer.h
+++ b/Engine/Source/Runtime/Core/Public/Templates/SharedPointer.h
@@ -137,6 +137,17 @@ namespace UE4SharedPointer_Private
 }
 
 
+// TSharedPtr of one mode to a type which has a TSharedFromThis only of another mode is illegal.
+// A type which does not inherit TSharedFromThis at all is ok.
+// We only check this inside the constructor because we don't necessarily have the full type of T when we declare a TSharedPtr<T>.
+#define UE_TSHAREDPTR_STATIC_ASSERT_VALID_MODE(ObjectType, Mode) \
+	enum \
+	{ \
+		ObjectTypeHasSameModeSharedFromThis = TPointerIsConvertibleFromTo<ObjectType, TSharedFromThis<ObjectType, Mode>>::Value, \
+		ObjectTypeHasOppositeModeSharedFromThis = TPointerIsConvertibleFromTo<ObjectType, TSharedFromThis<ObjectType, (Mode == ESPMode::NotThreadSafe) ? ESPMode::ThreadSafe : ESPMode::NotThreadSafe>>::Value \
+	}; \
+	static_assert(ObjectTypeHasSameModeSharedFromThis || !ObjectTypeHasOppositeModeSharedFromThis, "You cannot use a TSharedPtr of one mode with a type which inherits TSharedFromThis of another mode.");
+
 /**
  * TSharedRef is a non-nullable, non-intrusive reference-counted authoritative object reference.
  *
@@ -159,12 +170,14 @@ public:
 	 */
 	template <
 		typename OtherType,
-		typename = typename TEnableIf<TPointerIsConvertibleFromTo<OtherType, ObjectType>::Value>::Type
+		typename = decltype(ImplicitConv<ObjectType*>((OtherType*)nullptr))
 	>
 	FORCEINLINE explicit TSharedRef( OtherType* InObject )
 		: Object( InObject )
 		, SharedReferenceCount( SharedPointerInternals::NewDefaultReferenceController( InObject ) )
 	{
+		UE_TSHAREDPTR_STATIC_ASSERT_VALID_MODE(ObjectType, Mode)
+
 		Init(InObject);
 	}
 
@@ -177,12 +190,14 @@ public:
 	template <
 		typename OtherType,
 		typename DeleterType,
-		typename = typename TEnableIf<TPointerIsConvertibleFromTo<OtherType, ObjectType>::Value>::Type
+		typename = decltype(ImplicitConv<ObjectType*>((OtherType*)nullptr))
 	>
 	FORCEINLINE TSharedRef( OtherType* InObject, DeleterType&& InDeleter )
 		: Object( InObject )
 		, SharedReferenceCount( SharedPointerInternals::NewCustomReferenceController( InObject, Forward< DeleterType >( InDeleter ) ) )
 	{
+		UE_TSHAREDPTR_STATIC_ASSERT_VALID_MODE(ObjectType, Mode)
+
 		Init(InObject);
 	}
 
@@ -208,12 +223,14 @@ public:
 	// NOTE: The following is an Unreal extension to standard shared_ptr behavior
 	template <
 		typename OtherType,
-		typename = typename TEnableIf<TPointerIsConvertibleFromTo<OtherType, ObjectType>::Value>::Type
+		typename = decltype(ImplicitConv<ObjectType*>((OtherType*)nullptr))
 	>
 	FORCEINLINE TSharedRef( SharedPointerInternals::FRawPtrProxy< OtherType > const& InRawPtrProxy )
 		: Object( InRawPtrProxy.Object )
 		, SharedReferenceCount( InRawPtrProxy.ReferenceController )
 	{
+		UE_TSHAREDPTR_STATIC_ASSERT_VALID_MODE(ObjectType, Mode)
+
 		// If the following assert goes off, it means a TSharedRef was initialized from a nullptr object pointer.
 		// Shared references must never be nullptr, so either pass a valid object or consider using TSharedPtr instead.
 		check( InRawPtrProxy.Object != nullptr );
@@ -231,12 +248,13 @@ public:
 	 */
 	template <
 		typename OtherType,
-		typename = typename TEnableIf<TPointerIsConvertibleFromTo<OtherType, ObjectType>::Value>::Type
+		typename = decltype(ImplicitConv<ObjectType*>((OtherType*)nullptr))
 	>
 	FORCEINLINE TSharedRef( TSharedRef< OtherType, Mode > const& InSharedRef )
 		: Object( InSharedRef.Object )
 		, SharedReferenceCount( InSharedRef.SharedReferenceCount )
-	{ }
+	{
+	}
 
 	/**
 	 * Special constructor used internally to statically cast one shared reference type to another.  You
@@ -250,8 +268,9 @@ public:
 	FORCEINLINE TSharedRef( TSharedRef< OtherType, Mode > const& InSharedRef, SharedPointerInternals::FStaticCastTag )
 		: Object( static_cast< ObjectType* >( InSharedRef.Object ) )
 		, SharedReferenceCount( InSharedRef.SharedReferenceCount )
-	{ }
-  
+	{
+	}
+
 	/**
 	 * Special constructor used internally to cast a 'const' shared reference a 'mutable' reference.  You
 	 * should never call this constructor directly.  Instead, use the ConstCastSharedRef() function.
@@ -264,7 +283,8 @@ public:
 	FORCEINLINE TSharedRef( TSharedRef< OtherType, Mode > const& InSharedRef, SharedPointerInternals::FConstCastTag )
 		: Object( const_cast< ObjectType* >( InSharedRef.Object ) )
 		, SharedReferenceCount( InSharedRef.SharedReferenceCount )
-	{ }
+	{
+	}
 
 	/**
 	 * Aliasing constructor used to create a shared reference which shares its reference count with
@@ -327,7 +347,7 @@ public:
 	// NOTE: The following is an Unreal extension to standard shared_ptr behavior
 	template <
 		typename OtherType,
-		typename = typename TEnableIf<TPointerIsConvertibleFromTo<OtherType, ObjectType>::Value>::Type
+		typename = decltype(ImplicitConv<ObjectType*>((OtherType*)nullptr))
 	>
 	FORCEINLINE TSharedRef& operator=( SharedPointerInternals::FRawPtrProxy< OtherType > const& InRawPtrProxy )
 	{
@@ -419,7 +439,7 @@ private:
 	 */
 	template <
 		typename OtherType,
-		typename = typename TEnableIf<TPointerIsConvertibleFromTo<OtherType, ObjectType>::Value>::Type
+		typename = decltype(ImplicitConv<ObjectType*>((OtherType*)nullptr))
 	>
 	FORCEINLINE explicit TSharedRef( TSharedPtr< OtherType, Mode > const& InSharedPtr )
 		: Object( InSharedPtr.Object )
@@ -432,7 +452,7 @@ private:
 
 	template <
 		typename OtherType,
-		typename = typename TEnableIf<TPointerIsConvertibleFromTo<OtherType, ObjectType>::Value>::Type
+		typename = decltype(ImplicitConv<ObjectType*>((OtherType*)nullptr))
 	>
 	FORCEINLINE explicit TSharedRef( TSharedPtr< OtherType, Mode >&& InSharedPtr )
 		: Object( InSharedPtr.Object )
@@ -500,6 +520,8 @@ private:
 		: Object(InObject)
 		, SharedReferenceCount(InSharedReferenceCount)
 	{
+		UE_TSHAREDPTR_STATIC_ASSERT_VALID_MODE(ObjectType, Mode)
+
 		Init(InObject);
 	}
 };
@@ -536,16 +558,6 @@ struct FMakeReferenceTo<const void>
 template< class ObjectType, ESPMode Mode >
 class TSharedPtr
 {
-	enum
-	{
-		ObjectTypeHasSameModeSharedFromThis     = TPointerIsConvertibleFromTo<ObjectType, TSharedFromThis<ObjectType, Mode>>::Value,
-		ObjectTypeHasOppositeModeSharedFromThis = TPointerIsConvertibleFromTo<ObjectType, TSharedFromThis<ObjectType, (Mode == ESPMode::NotThreadSafe) ? ESPMode::ThreadSafe : ESPMode::NotThreadSafe>>::Value
-	};
-
-	// TSharedPtr of one mode to a type which has a TSharedFromThis only of another mode is illegal.
-	// A type which does not inherit TSharedFromThis at all is ok.
-	static_assert( TSharedPtr::ObjectTypeHasSameModeSharedFromThis || !TSharedPtr::ObjectTypeHasOppositeModeSharedFromThis, "You cannot use a TSharedPtr of one mode with a type which inherits TSharedFromThis of another mode.");
-
 public:
 	using ElementType = ObjectType;
 
@@ -556,7 +568,8 @@ public:
 	FORCEINLINE TSharedPtr( SharedPointerInternals::FNullTag* = nullptr )
 		: Object( nullptr )
 		, SharedReferenceCount()
-	{ }
+	{
+	}
 
 	/**
 	 * Constructs a shared pointer that owns the specified object.  Note that passing nullptr here will
@@ -566,12 +579,14 @@ public:
 	 */
 	template <
 		typename OtherType,
-		typename = typename TEnableIf<TPointerIsConvertibleFromTo<OtherType, ObjectType>::Value>::Type
+		typename = decltype(ImplicitConv<ObjectType*>((OtherType*)nullptr))
 	>
 	FORCEINLINE explicit TSharedPtr( OtherType* InObject )
 		: Object( InObject )
 		, SharedReferenceCount( SharedPointerInternals::NewDefaultReferenceController( InObject ) )
 	{
+		UE_TSHAREDPTR_STATIC_ASSERT_VALID_MODE(ObjectType, Mode)
+
 		// If the object happens to be derived from TSharedFromThis, the following method
 		// will prime the object with a weak pointer to itself.
 		SharedPointerInternals::EnableSharedFromThis( this, InObject, InObject );
@@ -587,12 +602,14 @@ public:
 	template <
 		typename OtherType,
 		typename DeleterType,
-		typename = typename TEnableIf<TPointerIsConvertibleFromTo<OtherType, ObjectType>::Value>::Type
+		typename = decltype(ImplicitConv<ObjectType*>((OtherType*)nullptr))
 	>
 	FORCEINLINE TSharedPtr( OtherType* InObject, DeleterType&& InDeleter )
 		: Object( InObject )
 		, SharedReferenceCount( SharedPointerInternals::NewCustomReferenceController( InObject, Forward< DeleterType >( InDeleter ) ) )
 	{
+		UE_TSHAREDPTR_STATIC_ASSERT_VALID_MODE(ObjectType, Mode)
+
 		// If the object happens to be derived from TSharedFromThis, the following method
 		// will prime the object with a weak pointer to itself.
 		SharedPointerInternals::EnableSharedFromThis( this, InObject, InObject );
@@ -606,12 +623,14 @@ public:
 	// NOTE: The following is an Unreal extension to standard shared_ptr behavior
 	template <
 		typename OtherType,
-		typename = typename TEnableIf<TPointerIsConvertibleFromTo<OtherType, ObjectType>::Value>::Type
+		typename = decltype(ImplicitConv<ObjectType*>((OtherType*)nullptr))
 	>
 	FORCEINLINE TSharedPtr( SharedPointerInternals::FRawPtrProxy< OtherType > const& InRawPtrProxy )
 		: Object( InRawPtrProxy.Object )
 		, SharedReferenceCount( InRawPtrProxy.ReferenceController )
 	{
+		UE_TSHAREDPTR_STATIC_ASSERT_VALID_MODE(ObjectType, Mode)
+
 		// If the object happens to be derived from TSharedFromThis, the following method
 		// will prime the object with a weak pointer to itself.
 		SharedPointerInternals::EnableSharedFromThis( this, InRawPtrProxy.Object, InRawPtrProxy.Object );
@@ -625,17 +644,19 @@ public:
 	 */
 	template <
 		typename OtherType,
-		typename = typename TEnableIf<TPointerIsConvertibleFromTo<OtherType, ObjectType>::Value>::Type
+		typename = decltype(ImplicitConv<ObjectType*>((OtherType*)nullptr))
 	>
 	FORCEINLINE TSharedPtr( TSharedPtr< OtherType, Mode > const& InSharedPtr )
 		: Object( InSharedPtr.Object )
 		, SharedReferenceCount( InSharedPtr.SharedReferenceCount )
-	{ }
+	{
+	}
 
 	FORCEINLINE TSharedPtr( TSharedPtr const& InSharedPtr )
 		: Object( InSharedPtr.Object )
 		, SharedReferenceCount( InSharedPtr.SharedReferenceCount )
-	{ }
+	{
+	}
 
 	FORCEINLINE TSharedPtr( TSharedPtr&& InSharedPtr )
 		: Object( InSharedPtr.Object )
@@ -653,7 +674,7 @@ public:
 	// NOTE: The following is an Unreal extension to standard shared_ptr behavior
 	template <
 		typename OtherType,
-		typename = typename TEnableIf<TPointerIsConvertibleFromTo<OtherType, ObjectType>::Value>::Type
+		typename = decltype(ImplicitConv<ObjectType*>((OtherType*)nullptr))
 	>
 	FORCEINLINE TSharedPtr( TSharedRef< OtherType, Mode > const& InSharedRef )
 		: Object( InSharedRef.Object )
@@ -675,7 +696,8 @@ public:
 	FORCEINLINE TSharedPtr( TSharedPtr< OtherType, Mode > const& InSharedPtr, SharedPointerInternals::FStaticCastTag )
 		: Object( static_cast< ObjectType* >( InSharedPtr.Object ) )
 		, SharedReferenceCount( InSharedPtr.SharedReferenceCount )
-	{ }
+	{
+	}
 	
 	/**
 	 * Special constructor used internally to cast a 'const' shared pointer a 'mutable' pointer.  You
@@ -689,7 +711,8 @@ public:
 	FORCEINLINE TSharedPtr( TSharedPtr< OtherType, Mode > const& InSharedPtr, SharedPointerInternals::FConstCastTag )
 		: Object( const_cast< ObjectType* >( InSharedPtr.Object ) )
 		, SharedReferenceCount( InSharedPtr.SharedReferenceCount )
-	{ }
+	{
+	}
 
 	/**
 	 * Aliasing constructor used to create a shared pointer which shares its reference count with
@@ -702,7 +725,8 @@ public:
 	FORCEINLINE TSharedPtr( TSharedPtr< OtherType, Mode > const& OtherSharedPtr, ObjectType* InObject )
 		: Object( InObject )
 		, SharedReferenceCount( OtherSharedPtr.SharedReferenceCount )
-	{ }
+	{
+	}
 
 	/**
 	 * Aliasing constructor used to create a shared pointer which shares its reference count with
@@ -730,7 +754,8 @@ public:
 	FORCEINLINE TSharedPtr( TSharedRef< OtherType, Mode > const& OtherSharedRef, ObjectType* InObject )
 		: Object( InObject )
 		, SharedReferenceCount( OtherSharedRef.SharedReferenceCount )
-	{ }
+	{
+	}
 
 	/**
 	 * Assignment to a nullptr pointer.  The object currently referenced by this shared pointer will no longer be
@@ -778,7 +803,7 @@ public:
 	// NOTE: The following is an Unreal extension to standard shared_ptr behavior
 	template <
 		typename OtherType,
-		typename = typename TEnableIf<TPointerIsConvertibleFromTo<OtherType, ObjectType>::Value>::Type
+		typename = decltype(ImplicitConv<ObjectType*>((OtherType*)nullptr))
 	>
 	FORCEINLINE TSharedPtr& operator=( SharedPointerInternals::FRawPtrProxy< OtherType > const& InRawPtrProxy )
 	{
@@ -911,7 +936,7 @@ private:
 	 */
 	template <
 		typename OtherType,
-		typename = typename TEnableIf<TPointerIsConvertibleFromTo<OtherType, ObjectType>::Value>::Type
+		typename = decltype(ImplicitConv<ObjectType*>((OtherType*)nullptr))
 	>
 	FORCEINLINE explicit TSharedPtr( TWeakPtr< OtherType, Mode > const& InWeakPtr )
 		: Object( nullptr )
@@ -1000,7 +1025,8 @@ public:
 	FORCEINLINE TWeakPtr( SharedPointerInternals::FNullTag* = nullptr )
 		: Object( nullptr )
 		, WeakReferenceCount()
-	{ }
+	{
+	}
 
 	/**
 	 * Constructs a weak pointer from a shared reference
@@ -1010,12 +1036,13 @@ public:
 	// NOTE: The following is an Unreal extension to standard shared_ptr behavior
 	template <
 		typename OtherType,
-		typename = typename TEnableIf<TPointerIsConvertibleFromTo<OtherType, ObjectType>::Value>::Type
+		typename = decltype(ImplicitConv<ObjectType*>((OtherType*)nullptr))
 	>
 	FORCEINLINE TWeakPtr( TSharedRef< OtherType, Mode > const& InSharedRef )
 		: Object( InSharedRef.Object )
 		, WeakReferenceCount( InSharedRef.SharedReferenceCount )
-	{ }
+	{
+	}
 
 	/**
 	 * Constructs a weak pointer from a shared pointer
@@ -1024,12 +1051,13 @@ public:
 	 */
 	template <
 		typename OtherType,
-		typename = typename TEnableIf<TPointerIsConvertibleFromTo<OtherType, ObjectType>::Value>::Type
+		typename = decltype(ImplicitConv<ObjectType*>((OtherType*)nullptr))
 	>
 	FORCEINLINE TWeakPtr( TSharedPtr< OtherType, Mode > const& InSharedPtr )
 		: Object( InSharedPtr.Object )
 		, WeakReferenceCount( InSharedPtr.SharedReferenceCount )
-	{ }
+	{
+	}
 
 	/**
 	 * Constructs a weak pointer from a weak pointer of another type.
@@ -1039,16 +1067,17 @@ public:
 	 */
 	template <
 		typename OtherType,
-		typename = typename TEnableIf<TPointerIsConvertibleFromTo<OtherType, ObjectType>::Value>::Type
+		typename = decltype(ImplicitConv<ObjectType*>((OtherType*)nullptr))
 	>
 	FORCEINLINE TWeakPtr( TWeakPtr< OtherType, Mode > const& InWeakPtr )
 		: Object( InWeakPtr.Object )
 		, WeakReferenceCount( InWeakPtr.WeakReferenceCount )
-	{ }
+	{
+	}
 
 	template <
 		typename OtherType,
-		typename = typename TEnableIf<TPointerIsConvertibleFromTo<OtherType, ObjectType>::Value>::Type
+		typename = decltype(ImplicitConv<ObjectType*>((OtherType*)nullptr))
 	>
 	FORCEINLINE TWeakPtr( TWeakPtr< OtherType, Mode >&& InWeakPtr )
 		: Object( InWeakPtr.Object )
@@ -1060,7 +1089,8 @@ public:
 	FORCEINLINE TWeakPtr( TWeakPtr const& InWeakPtr )
 		: Object( InWeakPtr.Object )
 		, WeakReferenceCount( InWeakPtr.WeakReferenceCount )
-	{ }
+	{
+	}
 
 	FORCEINLINE TWeakPtr( TWeakPtr&& InWeakPtr )
 		: Object( InWeakPtr.Object )
@@ -1110,7 +1140,7 @@ public:
 	 */
 	template <
 		typename OtherType,
-		typename = typename TEnableIf<TPointerIsConvertibleFromTo<OtherType, ObjectType>::Value>::Type
+		typename = decltype(ImplicitConv<ObjectType*>((OtherType*)nullptr))
 	>
 	FORCEINLINE TWeakPtr& operator=( TWeakPtr<OtherType, Mode> const& InWeakPtr )
 	{
@@ -1121,7 +1151,7 @@ public:
 
 	template <
 		typename OtherType,
-		typename = typename TEnableIf<TPointerIsConvertibleFromTo<OtherType, ObjectType>::Value>::Type
+		typename = decltype(ImplicitConv<ObjectType*>((OtherType*)nullptr))
 	>
 	FORCEINLINE TWeakPtr& operator=( TWeakPtr<OtherType, Mode>&& InWeakPtr )
 	{
@@ -1139,7 +1169,7 @@ public:
 	// NOTE: The following is an Unreal extension to standard shared_ptr behavior
 	template <
 		typename OtherType,
-		typename = typename TEnableIf<TPointerIsConvertibleFromTo<OtherType, ObjectType>::Value>::Type
+		typename = decltype(ImplicitConv<ObjectType*>((OtherType*)nullptr))
 	>
 	FORCEINLINE TWeakPtr& operator=( TSharedRef< OtherType, Mode > const& InSharedRef )
 	{
@@ -1155,7 +1185,7 @@ public:
 	 */
 	template <
 		typename OtherType,
-		typename = typename TEnableIf<TPointerIsConvertibleFromTo<OtherType, ObjectType>::Value>::Type
+		typename = decltype(ImplicitConv<ObjectType*>((OtherType*)nullptr))
 	>
 	FORCEINLINE TWeakPtr& operator=( TSharedPtr< OtherType, Mode > const& InSharedPtr )
 	{
diff --git a/Engine/Source/Runtime/Core/Public/Templates/UniquePtr.h b/Engine/Source/Runtime/Core/Public/Templates/UniquePtr.h
--- a/Engine/Source/Runtime/Core/Public/Templates/UniquePtr.h
+++ b/Engine/Source/Runtime/Core/Public/Templates/UniquePtr.h
@@ -3,7 +3,6 @@
 #pragma once
 
 #include "CoreTypes.h"
-#include "Templates/PointerIsConvertibleFromTo.h"
 #include "Templates/UnrealTemplate.h"
 #include "Templates/IsArray.h"
 #include "Templates/RemoveExtent.h"
@@ -27,7 +26,7 @@ struct TDefaultDelete
 
 	template <
 		typename U,
-		typename = typename TEnableIf<TPointerIsConvertibleFromTo<U, T>::Value>::Type
+		typename = decltype(ImplicitConv<T*>((U*)nullptr))
 	>
 	TDefaultDelete(const TDefaultDelete<U>&)
 	{
@@ -35,7 +34,7 @@ struct TDefaultDelete
 
 	template <
 		typename U,
-		typename = typename TEnableIf<TPointerIsConvertibleFromTo<U, T>::Value>::Type
+		typename = decltype(ImplicitConv<T*>((U*)nullptr))
 	>
 	TDefaultDelete& operator=(const TDefaultDelete<U>&)
 	{
@@ -58,7 +57,7 @@ struct TDefaultDelete<T[]>
 
 	template <
 		typename U,
-		typename = typename TEnableIf<TPointerIsConvertibleFromTo<U[], T[]>::Value>::Type
+		typename = decltype(ImplicitConv<T(*)[]>((U(*)[])nullptr))
 	>
 	TDefaultDelete(const TDefaultDelete<U[]>&)
 	{
@@ -66,15 +65,18 @@ struct TDefaultDelete<T[]>
 
 	template <
 		typename U,
-		typename = typename TEnableIf<TPointerIsConvertibleFromTo<U[], T[]>::Value>::Type
+		typename = decltype(ImplicitConv<T(*)[]>((U(*)[])nullptr))
 	>
 	TDefaultDelete& operator=(const TDefaultDelete<U[]>&)
 	{
 		return *this;
 	}
 
-	template <typename U>
-	typename TEnableIf<TPointerIsConvertibleFromTo<U[], T[]>::Value>::Type operator()(U* Ptr) const
+	template <
+		typename U,
+		typename = decltype(ImplicitConv<T(*)[]>((U(*)[])nullptr))
+	>
+	void operator()(U* Ptr) const
 	{
 		delete [] Ptr;
 	}
@@ -131,7 +133,8 @@ public:
 	template <
 		typename OtherT,
 		typename OtherDeleter,
-		typename = typename TEnableIf<!TIsArray<OtherT>::Value>::Type
+		typename = typename TEnableIf<!TIsArray<OtherT>::Value>::Type,
+		typename = decltype(ImplicitConv<T*>((OtherT*)nullptr))
 	>
 	FORCEINLINE TUniquePtr(TUniquePtr<OtherT, OtherDeleter>&& Other)
 		: Deleter(MoveTemp(Other.GetDeleter()))
@@ -162,8 +165,13 @@ public:
 	/**
 	 * Assignment operator for rvalues of other (usually derived) types
 	 */
-	template <typename OtherT, typename OtherDeleter>
-	FORCEINLINE typename TEnableIf<!TIsArray<OtherT>::Value, TUniquePtr&>::Type operator=(TUniquePtr<OtherT, OtherDeleter>&& Other)
+	template <
+		typename OtherT,
+		typename OtherDeleter,
+		typename = typename TEnableIf<!TIsArray<OtherT>::Value>::Type,
+		typename = decltype(ImplicitConv<T*>((OtherT*)nullptr))
+	>
+	FORCEINLINE TUniquePtr& operator=(TUniquePtr<OtherT, OtherDeleter>&& Other)
 	{
 		// We delete last, because we don't want odd side effects if the destructor of T relies on the state of this or Other
 		T* OldPtr = Ptr;
@@ -337,7 +345,7 @@ public:
 	 */
 	template <
 		typename U,
-		typename = typename TEnableIf<TPointerIsConvertibleFromTo<U[], T[]>::Value>::Type
+		typename = decltype(ImplicitConv<T(*)[]>((U(*)[])nullptr))
 	>
 	explicit FORCEINLINE TUniquePtr(U* InPtr)
 		: Ptr(InPtr)
@@ -368,7 +376,7 @@ public:
 	template <
 		typename OtherT,
 		typename OtherDeleter,
-		typename = typename TEnableIf<TPointerIsConvertibleFromTo<OtherT[], T[]>::Value>::Type
+		typename = decltype(ImplicitConv<T(*)[]>((OtherT(*)[])nullptr))
 	>
 	FORCEINLINE TUniquePtr(TUniquePtr<OtherT, OtherDeleter>&& Other)
 		: Deleter(MoveTemp(Other.GetDeleter()))
@@ -399,8 +407,12 @@ public:
 	/**
 	 * Assignment operator for rvalues of other (usually less qualified) types
 	 */
-	template <typename OtherT, typename OtherDeleter>
-	FORCEINLINE typename TEnableIf<TPointerIsConvertibleFromTo<OtherT[], T[]>::Value, TUniquePtr&>::Type operator=(TUniquePtr<OtherT>&& Other)
+	template <
+		typename OtherT,
+		typename OtherDeleter,
+		typename = decltype(ImplicitConv<T(*)[]>((OtherT(*)[])nullptr))
+	>
+	FORCEINLINE TUniquePtr& operator=(TUniquePtr<OtherT>&& Other)
 	{
 		// We delete last, because we don't want odd side effects if the destructor of T relies on the state of this or Other
 		T* OldPtr = Ptr;
@@ -501,8 +513,11 @@ public:
 	 *
 	 * @param InPtr A pointer to the array to take ownership of.
 	 */
-	template <typename U>
-	FORCEINLINE typename TEnableIf<TPointerIsConvertibleFromTo<U[], T[]>::Value>::Type Reset(U* InPtr)
+	template <
+		typename U,
+		typename = decltype(ImplicitConv<T(*)[]>((U(*)[])nullptr))
+	>
+	FORCEINLINE void Reset(U* InPtr)
 	{
 		// We delete last, because we don't want odd side effects if the destructor of T relies on the state of this
 		T* OldPtr = Ptr;
diff --git a/Engine/Source/Runtime/Core/Public/UObject/WeakObjectPtrTemplates.h b/Engine/Source/Runtime/Core/Public/UObject/WeakObjectPtrTemplates.h
--- a/Engine/Source/Runtime/Core/Public/UObject/WeakObjectPtrTemplates.h
+++ b/Engine/Source/Runtime/Core/Public/UObject/WeakObjectPtrTemplates.h
@@ -53,28 +53,11 @@ public:
 	 */
 	template <
 		typename U,
-		typename = typename TEnableIf<TPointerIsConvertibleFromTo<const volatile U, const volatile T>::Value>::Type
+		typename = decltype(ImplicitConv<T*>((U*)nullptr))
 	>
 	FORCEINLINE TWeakObjectPtr(U* Object) :
-		TWeakObjectPtrBase((UObject*)Object)
+		TWeakObjectPtrBase((const UObject*)Object)
 	{
-		// When removing this deprecation warning, change:
-		//
-		// TWeakObjectPtrBase((UObject*)Object)
-		// to:
-		// TWeakObjectPtrBase((const UObject*)Object)
-		//
-		// and:
-		//
-		// typename TEnableIf<TPointerIsConvertibleFromTo<const volatile U, const volatile T>::Value>::Type
-		// to:
-		// typename TEnableIf<TPointerIsConvertibleFromTo<U, T>::Value>::Type
-		UE_STATIC_DEPRECATE(
-			4.19,
-			(TLosesQualifiersFromTo<U, T>::Value),
-			"Implicit conversions from const pointers to non-const TWeakObjectPtrs has been deprecated as it is not const-correct."
-		);
-
 		// This static assert is in here rather than in the body of the class because we want
 		// to be able to define TWeakObjectPtr<UUndefinedClass>.
 		static_assert(TPointerIsConvertibleFromTo<T, const volatile UObject>::Value, "TWeakObjectPtr can only be constructed with UObject types");
@@ -84,13 +67,13 @@ public:
 	 * Construct from another weak pointer of another type, intended for derived-to-base conversions
 	 * @param Other weak pointer to copy from
 	 */
-	template <typename OtherT>
+	template <
+		typename OtherT,
+		typename = decltype(ImplicitConv<T*>((OtherT*)nullptr))
+	>
 	FORCEINLINE TWeakObjectPtr(const TWeakObjectPtr<OtherT, TWeakObjectPtrBase>& Other) :
 		TWeakObjectPtrBase(*(TWeakObjectPtrBase*)&Other) // we do a C-style cast to private base here to avoid clang 3.6.0 compilation problems with friend declarations
 	{
-		// It's also possible that this static_assert may fail for valid conversions because
-		// one or both of the types have only been forward-declared.
-		static_assert(TPointerIsConvertibleFromTo<OtherT, T>::Value, "Unable to convert TWeakObjectPtr - types are incompatible");
 	}
 
 	/**
@@ -125,13 +108,12 @@ public:
 	 * Assign from another weak pointer, intended for derived-to-base conversions
 	 * @param Other weak pointer to copy from
 	 */
-	template <typename OtherT>
+	template <
+		typename OtherT,
+		typename = decltype(ImplicitConv<T*>((OtherT*)nullptr))
+	>
 	FORCEINLINE TWeakObjectPtr& operator=(const TWeakObjectPtr<OtherT, TWeakObjectPtrBase>& Other)
 	{
-		// It's also possible that this static_assert may fail for valid conversions because
-		// one or both of the types have only been forward-declared.
-		static_assert(TPointerIsConvertibleFromTo<OtherT, T>::Value, "Unable to convert TWeakObjectPtr - types are incompatible");
-
 		*(TWeakObjectPtrBase*)this = *(TWeakObjectPtrBase*)&Other; // we do a C-style cast to private base here to avoid clang 3.6.0 compilation problems with friend declarations
 
 		return *this;
@@ -255,36 +237,22 @@ FORCEINLINE TWeakObjectPtr<T> MakeWeakObjectPtr(T* Ptr)
  * If both pointers would return nullptr from Get() they count as equal even if they were not initialized to the same object.
  * @param Other weak pointer to compare to
  */
-template <typename LhsT, typename RhsT, typename OtherTWeakObjectPtrBase>
+template <typename LhsT, typename RhsT, typename OtherTWeakObjectPtrBase, typename = decltype((LhsT*)nullptr == (RhsT*)nullptr)>
 FORCENOINLINE bool operator==(const TWeakObjectPtr<LhsT, OtherTWeakObjectPtrBase>& Lhs, const TWeakObjectPtr<RhsT, OtherTWeakObjectPtrBase>& Rhs)
 {
-	// It's also possible that this static_assert may fail for valid conversions because
-	// one or both of the types have only been forward-declared.
-	static_assert(TPointerIsConvertibleFromTo<LhsT, RhsT>::Value || TPointerIsConvertibleFromTo<RhsT, LhsT>::Value, "Unable to compare TWeakObjectPtrs - types are incompatible");
-
 	return (const OtherTWeakObjectPtrBase&)Lhs == (const OtherTWeakObjectPtrBase&)Rhs;
 }
 
-template <typename LhsT, typename RhsT, typename OtherTWeakObjectPtrBase>
+template <typename LhsT, typename RhsT, typename OtherTWeakObjectPtrBase, typename = decltype((LhsT*)nullptr == (RhsT*)nullptr)>
 FORCENOINLINE bool operator==(const TWeakObjectPtr<LhsT, OtherTWeakObjectPtrBase>& Lhs, const RhsT* Rhs)
 {
-	// It's also possible that these static_asserts may fail for valid conversions because
-	// one or both of the types have only been forward-declared.
-	static_assert(TPointerIsConvertibleFromTo<RhsT, UObject>::Value, "TWeakObjectPtr can only be compared with UObject types");
-	static_assert(TPointerIsConvertibleFromTo<LhsT, RhsT>::Value || TPointerIsConvertibleFromTo<RhsT, LhsT>::Value, "Unable to compare TWeakObjectPtr with raw pointer - types are incompatible");
-
 	// NOTE: this constructs a TWeakObjectPtrBase, which has some amount of overhead, so this may not be an efficient operation
 	return (const OtherTWeakObjectPtrBase&)Lhs == OtherTWeakObjectPtrBase(Rhs);
 }
 
-template <typename LhsT, typename RhsT, typename OtherTWeakObjectPtrBase>
+template <typename LhsT, typename RhsT, typename OtherTWeakObjectPtrBase, typename = decltype((LhsT*)nullptr == (RhsT*)nullptr)>
 FORCENOINLINE bool operator==(const LhsT* Lhs, const TWeakObjectPtr<RhsT, OtherTWeakObjectPtrBase>& Rhs)
 {
-	// It's also possible that these static_asserts may fail for valid conversions because
-	// one or both of the types have only been forward-declared.
-	static_assert(TPointerIsConvertibleFromTo<LhsT, UObject>::Value, "TWeakObjectPtr can only be compared with UObject types");
-	static_assert(TPointerIsConvertibleFromTo<LhsT, RhsT>::Value || TPointerIsConvertibleFromTo<RhsT, LhsT>::Value, "Unable to compare TWeakObjectPtr with raw pointer - types are incompatible");
-
 	// NOTE: this constructs a TWeakObjectPtrBase, which has some amount of overhead, so this may not be an efficient operation
 	return OtherTWeakObjectPtrBase(Lhs) == (const OtherTWeakObjectPtrBase&)Rhs;
 }
@@ -305,36 +273,22 @@ FORCENOINLINE bool operator==(TYPE_OF_NULLPTR, const TWeakObjectPtr<RhsT, OtherT
  * Compare weak pointers for inequality
  * @param Other weak pointer to compare to
  */
-template <typename LhsT, typename RhsT, typename OtherTWeakObjectPtrBase>
+template <typename LhsT, typename RhsT, typename OtherTWeakObjectPtrBase, typename = decltype((LhsT*)nullptr != (RhsT*)nullptr)>
 FORCENOINLINE bool operator!=(const TWeakObjectPtr<LhsT, OtherTWeakObjectPtrBase>& Lhs, const TWeakObjectPtr<RhsT, OtherTWeakObjectPtrBase>& Rhs)
 {
-	// It's also possible that this static_assert may fail for valid conversions because
-	// one or both of the types have only been forward-declared.
-	static_assert(TPointerIsConvertibleFromTo<LhsT, RhsT>::Value || TPointerIsConvertibleFromTo<RhsT, LhsT>::Value, "Unable to compare TWeakObjectPtrs - types are incompatible");
-
 	return (const OtherTWeakObjectPtrBase&)Lhs != (const OtherTWeakObjectPtrBase&)Rhs;
 }
 
-template <typename LhsT, typename RhsT, typename OtherTWeakObjectPtrBase>
+template <typename LhsT, typename RhsT, typename OtherTWeakObjectPtrBase, typename = decltype((LhsT*)nullptr != (RhsT*)nullptr)>
 FORCENOINLINE bool operator!=(const TWeakObjectPtr<LhsT, OtherTWeakObjectPtrBase>& Lhs, const RhsT* Rhs)
 {
-	// It's also possible that these static_asserts may fail for valid conversions because
-	// one or both of the types have only been forward-declared.
-	static_assert(TPointerIsConvertibleFromTo<RhsT, UObject>::Value, "TWeakObjectPtr can only be compared with UObject types");
-	static_assert(TPointerIsConvertibleFromTo<LhsT, RhsT>::Value || TPointerIsConvertibleFromTo<RhsT, LhsT>::Value, "Unable to compare TWeakObjectPtr with raw pointer - types are incompatible");
-
 	// NOTE: this constructs a TWeakObjectPtrBase, which has some amount of overhead, so this may not be an efficient operation
 	return (const OtherTWeakObjectPtrBase&)Lhs != OtherTWeakObjectPtrBase(Rhs);
 }
 
-template <typename LhsT, typename RhsT, typename OtherTWeakObjectPtrBase>
+template <typename LhsT, typename RhsT, typename OtherTWeakObjectPtrBase, typename = decltype((LhsT*)nullptr != (RhsT*)nullptr)>
 FORCENOINLINE bool operator!=(const LhsT* Lhs, const TWeakObjectPtr<RhsT, OtherTWeakObjectPtrBase>& Rhs)
 {
-	// It's also possible that these static_asserts may fail for valid conversions because
-	// one or both of the types have only been forward-declared.
-	static_assert(TPointerIsConvertibleFromTo<LhsT, UObject>::Value, "TWeakObjectPtr can only be compared with UObject types");
-	static_assert(TPointerIsConvertibleFromTo<LhsT, RhsT>::Value || TPointerIsConvertibleFromTo<RhsT, LhsT>::Value, "Unable to compare TWeakObjectPtr with raw pointer - types are incompatible");
-
 	// NOTE: this constructs a TWeakObjectPtrBase, which has some amount of overhead, so this may not be an efficient operation
 	return OtherTWeakObjectPtrBase(Lhs) != (const OtherTWeakObjectPtrBase&)Rhs;
 }
diff --git a/Engine/Source/Runtime/CoreUObject/Public/Templates/SubclassOf.h b/Engine/Source/Runtime/CoreUObject/Public/Templates/SubclassOf.h
--- a/Engine/Source/Runtime/CoreUObject/Public/Templates/SubclassOf.h
+++ b/Engine/Source/Runtime/CoreUObject/Public/Templates/SubclassOf.h
@@ -28,14 +28,14 @@ public:
 	}
 
 	/** Copy Constructor, will only compile if types are compatible */
-	template <class TClassA, class = typename TEnableIf<TPointerIsConvertibleFromTo<TClassA, TClass>::Value>::Type>
+	template <class TClassA, class = decltype(ImplicitConv<TClass*>((TClassA*)nullptr))>
 	FORCEINLINE TSubclassOf(const TSubclassOf<TClassA>& From) :
 		Class(*From)
 	{
 	}
 
 	/** Assignment operator, will only compile if types are compatible */
-	template <class TClassA, class = typename TEnableIf<TPointerIsConvertibleFromTo<TClassA, TClass>::Value>::Type>
+	template <class TClassA, class = decltype(ImplicitConv<TClass*>((TClassA*)nullptr))>
 	FORCEINLINE TSubclassOf& operator=(const TSubclassOf<TClassA>& From)
 	{
 		Class = *From;
diff --git a/Engine/Source/Runtime/CoreUObject/Public/UObject/LazyObjectPtr.h b/Engine/Source/Runtime/CoreUObject/Public/UObject/LazyObjectPtr.h
--- a/Engine/Source/Runtime/CoreUObject/Public/UObject/LazyObjectPtr.h
+++ b/Engine/Source/Runtime/CoreUObject/Public/UObject/LazyObjectPtr.h
@@ -175,16 +175,15 @@ public:
 	TLazyObjectPtr<T>& operator=(const TLazyObjectPtr<T>&) = default;
 
 	/** Construct from another lazy pointer with implicit upcasting allowed */
-	template<typename U, typename = typename TEnableIf<TPointerIsConvertibleFromTo<U, T>::Value>::Type>
+	template<typename U, typename = decltype(ImplicitConv<T*>((U*)nullptr))>
 	FORCEINLINE TLazyObjectPtr(const TLazyObjectPtr<U>& Other) :
 		FLazyObjectPtr((const FLazyObjectPtr&)Other)
 	{
 	}
 	
 	/** Assign from another lazy pointer with implicit upcasting allowed */
-	template<typename U>
-	FORCEINLINE typename TEnableIf<TPointerIsConvertibleFromTo<U, T>::Value, TLazyObjectPtr<T>&>::Type
-		operator=(const TLazyObjectPtr<U>& Other)
+	template<typename U, typename = decltype(ImplicitConv<T*>((U*)nullptr))>
+	FORCEINLINE TLazyObjectPtr<T>& operator=(const TLazyObjectPtr<U>& Other)
 	{
 		FLazyObjectPtr::operator=((const FLazyObjectPtr&)Other);
 		return *this;
@@ -313,30 +312,26 @@ public:
 // The reason these aren't inside the class (above) is because Visual Studio 2012-2013 crashes when compiling them :D
 
 /** Compare with another TLazyObjectPtr of related type */
-template<typename T, typename U>
+template<typename T, typename U, typename = decltype((T*)nullptr == (U*)nullptr)>
 FORCEINLINE bool operator==(const TLazyObjectPtr<T>& Lhs, const TLazyObjectPtr<U>& Rhs)
 {
-	static_assert(TPointerIsConvertibleFromTo<U, T>::Value || TPointerIsConvertibleFromTo<T, U>::Value, "Unable to compare TLazyObjectPtr - types are incompatible");
 	return (const FLazyObjectPtr&)Lhs == (const FLazyObjectPtr&)Rhs;
 }
-template<typename T, typename U>
+template<typename T, typename U, typename = decltype((T*)nullptr != (U*)nullptr)>
 FORCEINLINE bool operator!=(const TLazyObjectPtr<T>& Lhs, const TLazyObjectPtr<U>& Rhs)
 {
-	static_assert(TPointerIsConvertibleFromTo<U, T>::Value || TPointerIsConvertibleFromTo<T, U>::Value, "Unable to compare TLazyObjectPtr - types are incompatible");
 	return (const FLazyObjectPtr&)Lhs != (const FLazyObjectPtr&)Rhs;
 }
 
 /** Compare for equality with a raw pointer **/
-template<typename T, typename U>
+template<typename T, typename U, typename = decltype((T*)nullptr == (U*)nullptr)>
 FORCEINLINE bool operator==(const TLazyObjectPtr<T>& Lhs, const U* Rhs)
 {
-	static_assert(TPointerIsConvertibleFromTo<U, T>::Value || TPointerIsConvertibleFromTo<T, U>::Value, "Unable to compare TLazyObjectPtr with raw pointer - types are incompatible");
 	return Lhs.Get() == Rhs;
 }
-template<typename T, typename U>
+template<typename T, typename U, typename = decltype((T*)nullptr == (U*)nullptr)>
 FORCEINLINE bool operator==(const U* Lhs, const TLazyObjectPtr<T>& Rhs)
 {
-	static_assert(TPointerIsConvertibleFromTo<U, T>::Value || TPointerIsConvertibleFromTo<T, U>::Value, "Unable to compare TLazyObjectPtr with raw pointer - types are incompatible");
 	return Lhs == Rhs.Get();
 }
 
@@ -353,16 +348,14 @@ FORCEINLINE bool operator==(TYPE_OF_NULLPTR, const TLazyObjectPtr<T>& Rhs)
 }
 
 /** Compare for inequality with a raw pointer	**/
-template<typename T, typename U>
+template<typename T, typename U, typename = decltype((T*)nullptr != (U*)nullptr)>
 FORCEINLINE bool operator!=(const TLazyObjectPtr<T>& Lhs, const U* Rhs)
 {
-	static_assert(TPointerIsConvertibleFromTo<U, T>::Value || TPointerIsConvertibleFromTo<T, U>::Value, "Unable to compare TLazyObjectPtr - types are incompatible");
 	return Lhs.Get() != Rhs;
 }
-template<typename T, typename U>
+template<typename T, typename U, typename = decltype((T*)nullptr != (U*)nullptr)>
 FORCEINLINE bool operator!=(const U* Lhs, const TLazyObjectPtr<T>& Rhs)
 {
-	static_assert(TPointerIsConvertibleFromTo<U, T>::Value || TPointerIsConvertibleFromTo<T, U>::Value, "Unable to compare TLazyObjectPtr - types are incompatible");
 	return Lhs != Rhs.Get();
 }
 
diff --git a/Engine/Source/Runtime/CoreUObject/Public/UObject/SoftObjectPtr.h b/Engine/Source/Runtime/CoreUObject/Public/UObject/SoftObjectPtr.h
--- a/Engine/Source/Runtime/CoreUObject/Public/UObject/SoftObjectPtr.h
+++ b/Engine/Source/Runtime/CoreUObject/Public/UObject/SoftObjectPtr.h
@@ -124,14 +124,14 @@ public:
 	}
 	
 	/** Construct from another soft pointer */
-	template <class U, class = typename TEnableIf<TPointerIsConvertibleFromTo<U, T>::Value>::Type>
+	template <class U, class = decltype(ImplicitConv<T*>((U*)nullptr))>
 	FORCEINLINE TSoftObjectPtr(const TSoftObjectPtr<U>& Other)
 		: SoftObjectPtr(Other.SoftObjectPtr)
 	{
 	}
 
 	/** Construct from a moveable soft pointer */
-	template <class U, class = typename TEnableIf<TPointerIsConvertibleFromTo<U, T>::Value>::Type>
+	template <class U, class = decltype(ImplicitConv<T*>((U*)nullptr))>
 	FORCEINLINE TSoftObjectPtr(TSoftObjectPtr<U>&& Other)
 		: SoftObjectPtr(MoveTemp(Other.SoftObjectPtr))
 	{
@@ -191,7 +191,7 @@ public:
 	}
 
 	/** Copy from a weak pointer to an object already in memory */
-	template <class U, class = typename TEnableIf<TPointerIsConvertibleFromTo<U, T>::Value>::Type>
+	template <class U, class = decltype(ImplicitConv<T*>((U*)nullptr))>
 	FORCEINLINE TSoftObjectPtr& operator=(const TWeakObjectPtr<U>& Other)
 	{
 		SoftObjectPtr = Other;
@@ -199,7 +199,7 @@ public:
 	}
 
 	/** Copy from another soft pointer */
-	template <class U, class = typename TEnableIf<TPointerIsConvertibleFromTo<U, T>::Value>::Type>
+	template <class U, class = decltype(ImplicitConv<T*>((U*)nullptr))>
 	FORCEINLINE TSoftObjectPtr& operator=(TSoftObjectPtr<U> Other)
 	{
 		SoftObjectPtr = MoveTemp(Other.SoftObjectPtr);
@@ -409,7 +409,7 @@ public:
 	}
 		
 	/** Construct from another soft pointer */
-	template <class TClassA, class = typename TEnableIf<TPointerIsConvertibleFromTo<TClassA, TClass>::Value>::Type>
+	template <class TClassA, class = decltype(ImplicitConv<TClass*>((TClassA*)nullptr))>
 	FORCEINLINE TSoftClassPtr(const TSoftClassPtr<TClassA>& Other)
 		: SoftObjectPtr(Other.SoftObjectPtr)
 	{
@@ -452,7 +452,7 @@ public:
 	}
 
 	/** Copy from a weak pointer already in memory */
-	template<class TClassA, class = typename TEnableIf<TPointerIsConvertibleFromTo<TClassA, UClass>::Value>::Type>
+	template<class TClassA, class = decltype(ImplicitConv<TClass*>((TClassA*)nullptr))>
 	FORCEINLINE TSoftClassPtr& operator=(const TWeakObjectPtr<TClassA>& Other)
 	{
 		SoftObjectPtr = Other;
@@ -460,7 +460,7 @@ public:
 	}
 
 	/** Copy from another soft pointer */
-	template<class TClassA, class = typename TEnableIf<TPointerIsConvertibleFromTo<TClassA, TClass>::Value>::Type>
+	template<class TClassA, class = decltype(ImplicitConv<TClass*>((TClassA*)nullptr))>
 	FORCEINLINE TSoftClassPtr& operator=(const TSoftObjectPtr<TClassA>& Other)
 	{
 		SoftObjectPtr = Other.SoftObjectPtr;
diff --git a/Engine/Source/Runtime/CoreUObject/Public/UObject/StrongObjectPtr.h b/Engine/Source/Runtime/CoreUObject/Public/UObject/StrongObjectPtr.h
--- a/Engine/Source/Runtime/CoreUObject/Public/UObject/StrongObjectPtr.h
+++ b/Engine/Source/Runtime/CoreUObject/Public/UObject/StrongObjectPtr.h
@@ -76,7 +76,7 @@ public:
 
 	template <
 		typename OtherObjectType,
-		typename = typename TEnableIf<TPointerIsConvertibleFromTo<OtherObjectType, ObjectType>::Value>::Type
+		typename = decltype(ImplicitConv<ObjectType*>((OtherObjectType*)nullptr))
 	>
 	FORCEINLINE_DEBUGGABLE TStrongObjectPtr(const TStrongObjectPtr<OtherObjectType>& InOther)
 		: ReferenceCollector(InOther.Get())
@@ -90,11 +90,11 @@ public:
 		return *this;
 	}
 
-	template <typename OtherObjectType>
-	FORCEINLINE_DEBUGGABLE typename TEnableIf<
-		TPointerIsConvertibleFromTo<OtherObjectType, ObjectType>::Value,
-		TStrongObjectPtr&
-	>::Type operator=(const TStrongObjectPtr<OtherObjectType>& InOther)
+	template <
+		typename OtherObjectType,
+		typename = decltype(ImplicitConv<ObjectType*>((OtherObjectType*)nullptr))
+	>
+	FORCEINLINE_DEBUGGABLE TStrongObjectPtr& operator=(const TStrongObjectPtr<OtherObjectType>& InOther)
 	{
 		ReferenceCollector->Set(InOther.Get());
 		return *this;
