Temporary fix for PlatformMath intrinics - global disable - while waiting on approval for the better solution.

diff --git a/Engine/Source/Runtime/Core/Public/Math/UnrealPlatformMathSSE.h b/Engine/Source/Runtime/Core/Public/Math/UnrealPlatformMathSSE.h
--- a/Engine/Source/Runtime/Core/Public/Math/UnrealPlatformMathSSE.h
+++ b/Engine/Source/Runtime/Core/Public/Math/UnrealPlatformMathSSE.h
@@ -7,6 +7,7 @@
 
 // Code including this header is responsible for including the correct platform-specific header for SSE intrinsics.
 
+#define PLATFORM_ENABLE_SSE4_MATH 0 // Temporarily disable the intrinsics on all platforms to fix the build
 
 namespace UnrealPlatformMathSSE
 {
@@ -74,17 +75,23 @@ namespace UnrealPlatformMathSSE
 		return _mm_cvtt_ss2si(_mm_set_ss(F));
 	}
 
-#if !PLATFORM_LINUX // Linux Clang does not yet support _mm_round_ps/_mm_round_pd
 	static FORCEINLINE float TruncToFloat(float F)
 	{
+#if PLATFORM_ENABLE_SSE4_MATH
 		return _mm_cvtss_f32(_mm_round_ps(_mm_set_ss(F), 3));
+#else
+		return truncf(F);
+#endif
 	}
 
 	static FORCEINLINE double TruncToDouble(double F)
 	{
+#if PLATFORM_ENABLE_SSE4_MATH
 		return _mm_cvtsd_f64(_mm_round_pd(_mm_set_sd(F), 3));
-	}
+#else
+		return trunc(F);
 #endif
+	}
 
 	static FORCEINLINE int32 FloorToInt(float F)
 	{
@@ -93,17 +100,23 @@ namespace UnrealPlatformMathSSE
 		return _mm_cvt_ss2si(_mm_set_ss(F + F - 0.5f)) >> 1;
 	}
 
-#if !PLATFORM_LINUX // Linux Clang does not yet support _mm_round_ps/_mm_round_pd
 	static FORCEINLINE float FloorToFloat(float F)
 	{
+#if PLATFORM_ENABLE_SSE4_MATH
 		return _mm_cvtss_f32(_mm_floor_ps(_mm_set_ss(F)));
+#else
+		return floorf(F);
+#endif
 	}
 
 	static FORCEINLINE double FloorToDouble(double F)
 	{
+#if PLATFORM_ENABLE_SSE4_MATH
 		return _mm_cvtsd_f64(_mm_floor_pd(_mm_set_sd(F)));
-	}
+#else
+		return floor(F);
 #endif
+	}
 
 	static FORCEINLINE int32 RoundToInt(float F)
 	{
@@ -115,7 +128,6 @@ namespace UnrealPlatformMathSSE
 		return _mm_cvt_ss2si(_mm_set_ss(F + F + 0.5f)) >> 1;
 	}
 
-#if !PLATFORM_LINUX // Linux Clang does not yet support _mm_round_ps/_mm_round_pd
 	static FORCEINLINE float RoundToFloat(float F)
 	{
 		return FloorToFloat(F + 0.5f);
@@ -125,7 +137,7 @@ namespace UnrealPlatformMathSSE
 	{
 		return FloorToDouble(F + 0.5);
 	}
-#endif
+
 	static FORCEINLINE int32 CeilToInt(float F)
 	{
 		// Note: unlike the Generic solution and the float solution, we implement CeilToInt using a rounding instruction, rather than a dedicated ceil instruction
@@ -133,15 +145,22 @@ namespace UnrealPlatformMathSSE
 		return -(_mm_cvt_ss2si(_mm_set_ss(-0.5f - (F + F))) >> 1);
 	}
 
-#if !PLATFORM_LINUX // Linux Clang does not yet support _mm_round_ps/_mm_round_pd
 	static FORCEINLINE float CeilToFloat(float F)
 	{
+#if PLATFORM_ENABLE_SSE4_MATH
 		return _mm_cvtss_f32(_mm_ceil_ps(_mm_set_ss(F)));
+#else
+		return ceilf(F);
+#endif
+
 	}
 
 	static FORCEINLINE double CeilToDouble(double F)
 	{
+#if PLATFORM_ENABLE_SSE4_MATH
 		return _mm_cvtsd_f64(_mm_ceil_pd(_mm_set_sd(F)));
-	}
+#else
+		return ceil(F);
 #endif
+	}
 }
