TIsPointerConvertibleFromTo removed from delegates when UObject checking to avoid incorrect evaluations when incomplete types are involved.

diff --git a/Engine/Source/Runtime/Core/Public/Delegates/DelegateInstancesImpl.h b/Engine/Source/Runtime/Core/Public/Delegates/DelegateInstancesImpl.h
--- a/Engine/Source/Runtime/Core/Public/Delegates/DelegateInstancesImpl.h
+++ b/Engine/Source/Runtime/Core/Public/Delegates/DelegateInstancesImpl.h
@@ -9,7 +9,6 @@
 #pragma once
 #include "CoreTypes.h"
 #include "Misc/AssertionMacros.h"
-#include "Templates/PointerIsConvertibleFromTo.h"
 #include "Templates/TypeWrapper.h"
 #include "Templates/AreTypesEqual.h"
 #include "Templates/UnrealTypeTraits.h"
@@ -22,6 +21,12 @@ class FDelegateBase;
 class FDelegateHandle;
 enum class ESPMode;
 
+namespace UE4Delegates_Private
+{
+	constexpr bool IsUObjectPtr(const volatile UObjectBase*) { return true; }
+	constexpr bool IsUObjectPtr(...)                         { return false; }
+}
+
 /* Macros for function parameter and delegate payload lists
  *****************************************************************************/
 
@@ -43,7 +48,7 @@ private:
 	typedef IBaseDelegateInstance<RetValType (ParamTypes...)>                                  Super;
 	typedef TBaseUFunctionDelegateInstance<UserClass, RetValType (ParamTypes...), VarTypes...> UnwrappedThisType;
 
-	static_assert(TPointerIsConvertibleFromTo<UserClass, const UObjectBase>::Value, "You cannot use UFunction delegates with non UObject classes.");
+	static_assert(UE4Delegates_Private::IsUObjectPtr((UserClass*)nullptr), "You cannot use UFunction delegates with non UObject classes.");
 
 public:
 	TBaseUFunctionDelegateInstance(UserClass* InUserObject, const FName& InFunctionName, VarTypes... Vars)
@@ -389,7 +394,7 @@ public:
 	typedef typename TUnwrapType<WrappedRetValType>::Type RetValType;
 
 private:
-	static_assert(!TPointerIsConvertibleFromTo<UserClass, const UObjectBase>::Value, "You cannot use raw method delegates with UObjects.");
+	static_assert(!UE4Delegates_Private::IsUObjectPtr((UserClass*)nullptr), "You cannot use raw method delegates with UObjects.");
 
 	typedef IBaseDelegateInstance<RetValType (ParamTypes...)>                                          Super;
 	typedef TBaseRawMethodDelegateInstance<bConst, UserClass, RetValType (ParamTypes...), VarTypes...> UnwrappedThisType;
@@ -559,7 +564,7 @@ private:
 	typedef IBaseDelegateInstance<RetValType (ParamTypes...)>                                              Super;
 	typedef TBaseUObjectMethodDelegateInstance<bConst, UserClass, RetValType (ParamTypes...), VarTypes...> UnwrappedThisType;
 
-	static_assert(TPointerIsConvertibleFromTo<UserClass, const UObjectBase>::Value, "You cannot use UObject method delegates with raw pointers.");
+	static_assert(UE4Delegates_Private::IsUObjectPtr((UserClass*)nullptr), "You cannot use UObject method delegates with raw pointers.");
 
 public:
 	typedef typename TMemFunPtrType<bConst, UserClass, RetValType (ParamTypes..., VarTypes...)>::Type FMethodPtr;
