Map GetTypeHash checking changed to not use templates when it might to instantiate correctly later.

diff --git a/Engine/Source/Runtime/Core/Public/Containers/Map.h b/Engine/Source/Runtime/Core/Public/Containers/Map.h
--- a/Engine/Source/Runtime/Core/Public/Containers/Map.h
+++ b/Engine/Source/Runtime/Core/Public/Containers/Map.h
@@ -5,13 +5,11 @@
 #include "CoreTypes.h"
 
 #include "Algo/Reverse.h"
-#include "Concepts/GetTypeHashable.h"
 #include "Containers/Set.h"
 #include "Containers/UnrealString.h"
 #include "Misc/AssertionMacros.h"
 #include "Misc/StructBuilder.h"
 #include "Templates/Function.h"
-#include "Templates/Models.h"
 #include "Templates/Sorting.h"
 #include "Templates/Tuple.h"
 #include "Templates/UnrealTemplate.h"
@@ -96,7 +94,11 @@ struct TDefaultMapKeyFuncs : BaseKeyFuncs<TPair<KeyType,ValueType>,KeyType,bInAl
 template<typename KeyType, typename ValueType, bool bInAllowDuplicateKeys>
 struct TDefaultMapHashableKeyFuncs : TDefaultMapKeyFuncs<KeyType, ValueType, bInAllowDuplicateKeys>
 {
-	static_assert(TModels<CGetTypeHashable, KeyType>::Value, "TMap must have a hashable KeyType unless a custom key func is provided.");
+	// Check that the key type is actually hashable
+	//
+	// If this line fails to compile then your key doesn't have
+	// a GetTypeHash() overload.
+	using HashabilityCheck = decltype(GetTypeHash(DeclVal<const KeyType>()));
 };
 
 /** 
