Iteration on StringView

- Add string view literals using ""_SV / ""_ASV / ""_WSV.
- Add constexpr to a handful of functions. Adding greater constexpr support will require edits elsewhere.
- Set OutIndex to INDEX_NONE when Find[Last]Char cannot find the character, for consistency with FString.
- Add the new *Inline versions of Left*, Right*, Mid functions, for consistency with FString.
- Move from Misc to Containers for consistency with FString.
- Document the view API.

diff --git a/Engine/Source/Runtime/Core/Private/Containers/StringView.cpp b/Engine/Source/Runtime/Core/Private/Containers/StringView.cpp
new file mode 100644
--- /dev/null
+++ b/Engine/Source/Runtime/Core/Private/Containers/StringView.cpp
@@ -0,0 +1,73 @@
+// Copyright 1998-2019 Epic Games, Inc. All Rights Reserved.
+
+#include "Misc/StringView.h"
+
+#include "Containers/UnrealString.h"
+#include "HAL/PlatformString.h"
+
+template<typename CharType, typename ViewType>
+int32 TStringViewImpl<CharType, ViewType>::Compare(const TStringViewImpl& Other, ESearchCase::Type SearchCase) const
+{
+	const SizeType SelfLen = Len();
+	const SizeType OtherLen = Other.Len();
+	const SizeType ShortestLength = FMath::Min(SelfLen, OtherLen);
+
+	int Result;
+	if (SearchCase == ESearchCase::CaseSensitive)
+	{
+		Result = FPlatformString::Strncmp(GetData(), Other.GetData(), ShortestLength);
+	}
+	else
+	{
+		Result = FPlatformString::Strnicmp(GetData(), Other.GetData(), ShortestLength);
+	}
+
+	if (Result != 0 || SelfLen == OtherLen)
+	{
+		return Result;
+	}
+
+	return SelfLen < OtherLen ? -1 : 1;
+}
+
+template<typename CharType, typename ViewType>
+bool TStringViewImpl<CharType, ViewType>::FindChar(CharType InChar, SizeType& OutIndex) const
+{
+	for (SizeType i = 0; i < Size; ++i)
+	{
+		if (DataPtr[i] == InChar)
+		{
+			OutIndex = i;
+			return true;
+		}
+	}
+
+	OutIndex = INDEX_NONE;
+	return false;
+}
+
+template<typename CharType, typename ViewType>
+bool TStringViewImpl<CharType, ViewType>::FindLastChar(CharType InChar, SizeType& OutIndex) const
+{
+	if (Size == 0)
+	{
+		OutIndex = INDEX_NONE;
+		return false;
+	}
+
+	for (SizeType i = Size - 1; i >= 0; --i)
+	{
+		if (DataPtr[i] == InChar)
+		{
+			OutIndex = i;
+			return true;
+		}
+	}
+
+	OutIndex = INDEX_NONE;
+	return false;
+}
+
+template class TStringViewImpl<TCHAR, FStringView>;
+template class TStringViewImpl<ANSICHAR, FAnsiStringView>;
+template class TStringViewImpl<WIDECHAR, FWideStringView>;
diff --git a/Engine/Source/Runtime/Core/Private/Misc/StringView.cpp b/Engine/Source/Runtime/Core/Private/Misc/StringView.cpp
deleted file mode 100644
diff --git a/Engine/Source/Runtime/Core/Public/Containers/StringView.h b/Engine/Source/Runtime/Core/Public/Containers/StringView.h
new file mode 100644
--- /dev/null
+++ b/Engine/Source/Runtime/Core/Public/Containers/StringView.h
@@ -0,0 +1,276 @@
+// Copyright 1998-2019 Epic Games, Inc. All Rights Reserved.
+
+#pragma once
+
+#include "CoreTypes.h"
+#include "Math/NumericLimits.h"
+#include "Math/UnrealMathUtility.h"
+#include "Misc/CString.h"
+#include "Templates/AndOrNot.h"
+#include "Templates/EnableIf.h"
+#include "Templates/RemoveCV.h"
+#include "Templates/UnrealTemplate.h"
+#include "Traits/IsContiguousContainer.h"
+
+namespace StringViewPrivate
+{
+	/**
+	 * Trait testing whether a range type is compatible with the view type
+	 */
+	template <typename RangeType, typename ElementType>
+	struct TIsCompatibleRangeType : TIsSame<ElementType, typename TRemoveCV<typename TRemovePointer<decltype(GetData(DeclVal<RangeType&>()))>::Type>::Type>
+	{
+	};
+}
+
+/** String View
+
+	* A string view is implicitly constructible from const char* style strings and
+	  from character ranges such as FString.
+
+	* A string view does not own any data nor does it attempt to control any lifetimes, it
+	  merely points at a subrange of characters in some other string. It's up to the user
+	  to ensure the underlying string stays valid for the lifetime of the string view.
+
+	* A string view does not represent a NUL terminated string and therefore you should
+	  never pass in the pointer returned by GetData() into a C-string API accepting only a
+	  pointer. You must either use a string builder to make a properly terminated string,
+	  or use an API that accepts a length argument in addition to the C-string.
+
+	String views are a good fit for arguments to functions which don't wish to care
+	which style of string construction is used by the caller. If you accept strings via
+	string views then the caller is free to use FString, FStringBuilder, raw C strings,
+	or any other type which can be converted into a string view.
+
+	i.e., a function such as:
+
+	void DoFoo(const FStringView& InString);
+
+	May be called as:
+
+	void MultiFoo()
+	{
+		FString MyFoo(TEXT("Zoo"));
+		const TCHAR* MyFooStr = *MyFoo;
+
+		TStringBuilder<64> BuiltFoo;
+		BuiltFoo.Append(TEXT("ABC"));
+
+		DoFoo(MyFoo);
+		DoFoo(MyFooStr);
+		DoFoo(TEXT("ABC"));
+		DoFoo(BuiltFoo);
+	}
+
+  */
+
+template <typename CharType, typename DerivedViewType>
+class TStringViewImpl
+{
+public:
+	using ElementType = CharType;
+	using SizeType = int32;
+	using ViewType = DerivedViewType;
+
+	template <typename RangeType>
+	using TIsCompatibleRangeType = TAnd<TIsContiguousContainer<RangeType>, StringViewPrivate::TIsCompatibleRangeType<RangeType, ElementType>>;
+
+	template <typename RangeType>
+	using TEnableIfCompatibleRangeType = typename TEnableIf<TIsCompatibleRangeType<RangeType>::Value>::Type;
+
+	/** Construct an empty view. */
+	constexpr TStringViewImpl() = default;
+
+	/**
+	 * Construct a view of the null-terminated string pointed to by InData.
+	 *
+	 * The caller is responsible for ensuring that the provided character range remains valid for the lifetime of the view.
+	 */
+	inline TStringViewImpl(const CharType* InData)
+		: DataPtr(InData)
+		, Size(TCString<CharType>::Strlen(InData))
+	{
+	}
+
+	/**
+	 * Construct a view of InSize characters beginning at InData.
+	 *
+	 * The caller is responsible for ensuring that the provided character range remains valid for the lifetime of the view.
+	 */
+	constexpr inline TStringViewImpl(const CharType* InData, SizeType InSize)
+		: DataPtr(InData)
+		, Size(InSize)
+	{
+	}
+
+	/** Access the character at the given index in the view. */
+	inline const CharType& operator[](SizeType Index) const;
+
+	/** Returns a pointer to the start of the view. This is NOT guaranteed to be null-terminated! */
+	constexpr inline const CharType* GetData() const { return DataPtr; }
+
+	/** Returns a pointer to the start of the view. This is NOT guaranteed to be null-terminated! */
+	UE_DEPRECATED(4.25, "'Data' is deprecated. Please use 'GetData' instead!")
+	constexpr inline const CharType* Data() const { return DataPtr; }
+
+	// Capacity
+
+	/** Returns the length of the string view. */
+	constexpr inline SizeType Len() const { return Size; }
+
+	/** Returns whether the string view is empty. */
+	constexpr inline bool IsEmpty() const { return Size == 0; }
+
+	// Modifiers
+
+	/** Modifies the view to remove the given number of characters from the start. */
+	inline void		RemovePrefix(SizeType CharCount)	{ DataPtr += CharCount; Size -= CharCount; }
+	/** Modifies the view to remove the given number of characters from the end. */
+	inline void		RemoveSuffix(SizeType CharCount)	{ Size -= CharCount; }
+
+	// Operations
+
+	/**
+	 * Copy characters from the view into a destination buffer without null termination.
+	 *
+	 * @param Dest Buffer to write into. Must have space for at least CharCount characters.
+	 * @param CharCount The maximum number of characters to copy.
+	 * @param Position The offset into the view from which to start copying.
+	 *
+	 * @return The number of characters written to the destination buffer.
+	 */
+	inline SizeType CopyString(CharType* Dest, SizeType CharCount, SizeType Position = 0) const;
+
+	/** Alias for Mid. */
+	inline ViewType SubStr(SizeType Position, SizeType CharCount) const { return Mid(Position, CharCount); }
+
+	/** Returns the left-most part of the view by taking the given number of characters from the left. */
+	inline ViewType Left(SizeType CharCount) const;
+	/** Returns the left-most part of the view by chopping the given number of characters from the right. */
+	inline ViewType LeftChop(SizeType CharCount) const;
+	/** Returns the right-most part of the view by taking the given number of characters from the right. */
+	inline ViewType Right(SizeType CharCount) const;
+	/** Returns the right-most part of the view by chopping the given number of characters from the left. */
+	inline ViewType RightChop(SizeType CharCount) const;
+	/** Returns the middle part of the view by taking up to the given number of characters from the given position. */
+	inline ViewType Mid(SizeType Position, SizeType CharCount = TNumericLimits<SizeType>::Max()) const;
+
+	/** Modifies the view to be the given number of characters from the left. */
+	inline void LeftInlineInline(SizeType CharCount) { *this = Left(CharCount); }
+	/** Modifies the view by chopping the given number of characters from the right. */
+	inline void LeftChopInline(SizeType CharCount) { *this = LeftChop(CharCount); }
+	/** Modifies the view to be the given number of characters from the right. */
+	inline void RightInline(SizeType CharCount) { *this = Right(CharCount); }
+	/** Modifies the view by chopping the given number of characters from the left. */
+	inline void RightChopInline(SizeType CharCount) { *this = RightChop(CharCount); }
+	/** Modifies the view to be the middle part by taking up to the given number of characters from the given position. */
+	inline void MidInline(SizeType Position, SizeType CharCount = TNumericLimits<SizeType>::Max()) { *this = Mid(Position, CharCount); }
+
+	// Comparison
+
+	/**
+	 * Check whether this view is lexicographically equivalent to another view.
+	 *
+	 * @param SearchCase Whether the comparison should ignore case.
+	 */
+	bool Equals(const TStringViewImpl& Other, ESearchCase::Type SearchCase = ESearchCase::CaseSensitive) const;
+
+	/**
+	 * Compare this view lexicographically with another view.
+	 *
+	 * @param SearchCase Whether the comparison should ignore case.
+	 *
+	 * @return 0 is equal, negative if this view is less, positive if this view is greater.
+	 */
+	CORE_API int32 Compare(const TStringViewImpl& Other, ESearchCase::Type SearchCase = ESearchCase::CaseSensitive) const;
+
+	/** Returns whether this view starts with the prefix character compared case-sensitively. */
+	inline bool StartsWith(CharType Prefix) const { return Size >= 1 && DataPtr[0] == Prefix; }
+	/** Returns whether this view starts with the prefix with optional case sensitivity. */
+	inline bool StartsWith(const TStringViewImpl& Prefix, ESearchCase::Type SearchCase = ESearchCase::IgnoreCase) const;
+
+	/** Returns whether this view ends with the suffix character compared case-sensitively. */
+	inline bool EndsWith(CharType Suffix) const { return Size >= 1 && DataPtr[Size-1] == Suffix; }
+	/** Returns whether this view ends with the suffix with optional case sensitivity. */
+	inline bool EndsWith(const TStringViewImpl& Suffix, ESearchCase::Type SearchCase = ESearchCase::IgnoreCase) const;
+
+	// Searching/Finding
+
+	/**
+	 * Search the view for the first occurrence of a character.
+	 *
+	 * @param InChar The character to search for. Comparison is lexicographic.
+	 * @param OutIndex [out] The position at which the character was found, or INDEX_NONE if not found.
+	 *
+	 * @return Whether the character was found in the view.
+	 */
+	CORE_API bool FindChar(CharType InChar, SizeType& OutIndex) const;
+
+	/**
+	 * Search the view for the last occurrence of a character.
+	 *
+	 * @param InChar The character to search for. Comparison is lexicographic.
+	 * @param OutIndex [out] The position at which the character was found, or INDEX_NONE if not found.
+	 *
+	 * @return Whether the character was found in the view.
+	 */
+	CORE_API bool FindLastChar(CharType InChar, SizeType& OutIndex) const;
+
+public:
+	/**
+	 * DO NOT USE DIRECTLY
+	 * STL-like iterators to enable range-based for loop support.
+	 */
+	constexpr inline const CharType* begin() const { return DataPtr; }
+	constexpr inline const CharType* end() const { return DataPtr + Size; }
+
+protected:
+	const CharType* DataPtr = nullptr;
+	SizeType Size = 0;
+};
+
+template <typename CharType, typename ViewType>
+struct TIsContiguousContainer<TStringViewImpl<CharType, ViewType>>
+{
+	enum { Value = true };
+};
+
+template <typename CharType, typename ViewType>
+constexpr inline SIZE_T GetNum(const TStringViewImpl<CharType, ViewType>& String)
+{
+	return String.Len();
+}
+
+//////////////////////////////////////////////////////////////////////////
+
+class FStringView : public TStringViewImpl<TCHAR, FStringView>
+{
+public:
+	using TStringViewImpl<ElementType, FStringView>::TStringViewImpl;
+};
+
+class FAnsiStringView : public TStringViewImpl<ANSICHAR, FAnsiStringView>
+{
+public:
+	using TStringViewImpl<ElementType, FAnsiStringView>::TStringViewImpl;
+};
+
+class FWideStringView : public TStringViewImpl<WIDECHAR, FWideStringView>
+{
+public:
+	using TStringViewImpl<ElementType, FWideStringView>::TStringViewImpl;
+};
+
+template <> struct TIsContiguousContainer<FStringView> { enum { Value = true }; };
+template <> struct TIsContiguousContainer<FAnsiStringView> { enum { Value = true }; };
+template <> struct TIsContiguousContainer<FWideStringView> { enum { Value = true }; };
+
+constexpr inline SIZE_T GetNum(const FStringView& String) { return String.Len(); }
+constexpr inline SIZE_T GetNum(const FAnsiStringView& String) { return String.Len(); }
+constexpr inline SIZE_T GetNum(const FWideStringView& String) { return String.Len(); }
+
+constexpr inline FStringView operator "" _SV(const TCHAR* String, size_t Size) { return FStringView(String, Size); }
+constexpr inline FAnsiStringView operator "" _ASV(const ANSICHAR* String, size_t Size) { return FAnsiStringView(String, Size); }
+constexpr inline FWideStringView operator "" _WSV(const WIDECHAR* String, size_t Size) { return FWideStringView(String, Size); }
+
+#include "StringView.inl"
diff --git a/Engine/Source/Runtime/Core/Public/Containers/StringView.inl b/Engine/Source/Runtime/Core/Public/Containers/StringView.inl
new file mode 100644
--- /dev/null
+++ b/Engine/Source/Runtime/Core/Public/Containers/StringView.inl
@@ -0,0 +1,131 @@
+// Copyright 1998-2019 Epic Games, Inc. All Rights Reserved.
+
+#include "HAL/PlatformString.h"
+#include "HAL/UnrealMemory.h"
+#include "Templates/Decay.h"
+#include "Templates/UnrealTypeTraits.h"
+#include "Traits/IntType.h"
+
+template <typename CharType, typename ViewType>
+inline const CharType& TStringViewImpl<CharType, ViewType>::operator[](SizeType Index) const
+{
+	checkf(Index >= 0 && Index < Size, TEXT("Index out of bounds on StringView: index %i on a view with a length of %i"), Index, Size);
+	return DataPtr[Index];
+}
+
+template <typename CharType, typename ViewType>
+inline typename TStringViewImpl<CharType, ViewType>::SizeType TStringViewImpl<CharType, ViewType>::CopyString(CharType* Dest, SizeType CharCount, SizeType Position) const
+{
+	const  SizeType CopyCount = FMath::Min(Size - Position, CharCount);
+	FMemory::Memcpy(Dest, DataPtr + Position, CopyCount);
+	return CopyCount;
+}
+
+template <typename CharType, typename ViewType>
+inline ViewType TStringViewImpl<CharType, ViewType>::Left(SizeType CharCount) const
+{
+	return ViewType(DataPtr, FMath::Clamp(CharCount, 0, Size));
+}
+
+template <typename CharType, typename ViewType>
+inline ViewType TStringViewImpl<CharType, ViewType>::LeftChop(SizeType CharCount) const
+{
+	return ViewType(DataPtr, FMath::Clamp(Size - CharCount, 0, Size));
+}
+
+template <typename CharType, typename ViewType>
+inline ViewType TStringViewImpl<CharType, ViewType>::Right(SizeType CharCount) const
+{
+	const SizeType OutLen = FMath::Clamp(CharCount, 0, Size);
+	return ViewType(DataPtr + Size - OutLen, OutLen);
+}
+
+template <typename CharType, typename ViewType>
+inline ViewType TStringViewImpl<CharType, ViewType>::RightChop(SizeType CharCount) const
+{
+	const SizeType OutLen = FMath::Clamp(Size - CharCount, 0, Size);
+	return ViewType(DataPtr + Size - OutLen, OutLen);
+}
+
+template <typename CharType, typename ViewType>
+inline ViewType TStringViewImpl<CharType, ViewType>::Mid(SizeType Position, SizeType CharCount) const
+{
+	check(CharCount >= 0);
+	using USizeType = TUnsignedIntType_T<sizeof(SizeType)>;
+	Position = FMath::Clamp<USizeType>(Position, 0, Size);
+	CharCount = FMath::Clamp<USizeType>(CharCount, 0, Size - Position);
+	return ViewType(DataPtr + Position, CharCount);
+}
+
+template <typename CharType, typename ViewType>
+inline bool TStringViewImpl<CharType, ViewType>::Equals(const TStringViewImpl& Other, ESearchCase::Type SearchCase) const
+{
+	return Size == Other.Size && Compare(Other, SearchCase) == 0;
+}
+
+template <typename CharType, typename ViewType>
+inline bool TStringViewImpl<CharType, ViewType>::StartsWith(const TStringViewImpl& Prefix, ESearchCase::Type SearchCase) const
+{
+	return Prefix.Equals(Left(Prefix.Len()), SearchCase);
+}
+
+template <typename CharType, typename ViewType>
+inline bool TStringViewImpl<CharType, ViewType>::EndsWith(const TStringViewImpl& Suffix, ESearchCase::Type SearchCase) const
+{
+	return Suffix.Equals(Right(Suffix.Len()), SearchCase);
+}
+
+// Case-insensitive comparison operators
+
+template <typename CharType, typename ViewType, typename RangeType,
+	typename = typename TStringViewImpl<CharType, ViewType>::template TEnableIfCompatibleRangeType<RangeType>>
+inline bool operator==(const TStringViewImpl<CharType, ViewType>& Lhs, RangeType&& Rhs)
+{
+	return Lhs.Equals(TStringViewImpl<CharType, ViewType>(Forward<RangeType>(Rhs)), ESearchCase::IgnoreCase);
+}
+
+template <typename CharType, typename ViewType, typename RangeType,
+	typename = typename TEnableIf<TNot<TIsDerivedFrom<typename TDecay<RangeType>::Type, TStringViewImpl<CharType, ViewType>>>::Value>::Type>
+inline auto operator==(RangeType&& Lhs, const TStringViewImpl<CharType, ViewType>& Rhs) -> decltype(Rhs == Forward<RangeType>(Lhs))
+{
+	return Rhs == Forward<RangeType>(Lhs);
+}
+
+template <typename CharType, typename ViewType, typename RangeType>
+inline auto operator!=(const TStringViewImpl<CharType, ViewType>& Lhs, RangeType&& Rhs) -> decltype(!(Lhs == Forward<RangeType>(Rhs)))
+{
+	return !(Lhs == Forward<RangeType>(Rhs));
+}
+
+template <typename CharType, typename ViewType, typename RangeType,
+	typename = typename TEnableIf<TNot<TIsDerivedFrom<typename TDecay<RangeType>::Type, TStringViewImpl<CharType, ViewType>>>::Value>::Type>
+inline auto operator!=(RangeType&& Lhs, const TStringViewImpl<CharType, ViewType>& Rhs) -> decltype(!(Rhs == Forward<RangeType>(Lhs)))
+{
+	return !(Rhs == Forward<RangeType>(Lhs));
+}
+
+// Case-insensitive C-string comparison operators
+
+template <typename CharType, typename ViewType>
+inline bool operator==(const TStringViewImpl<CharType, ViewType>& Lhs, const CharType* Rhs)
+{
+	return TCString<CharType>::Strnicmp(Lhs.GetData(), Rhs, Lhs.Len()) == 0 && !Rhs[Lhs.Len()];
+}
+
+template <typename CharType, typename ViewType>
+inline bool operator==(const CharType* Lhs, const TStringViewImpl<CharType, ViewType>& Rhs)
+{
+	return Rhs == Lhs;
+}
+
+template <typename CharType, typename ViewType>
+inline bool operator!=(const TStringViewImpl<CharType, ViewType>& Lhs, const CharType* Rhs)
+{
+	return !(Lhs == Rhs);
+}
+
+template <typename CharType, typename ViewType>
+inline bool operator!=(const CharType* Lhs, const TStringViewImpl<CharType, ViewType>& Rhs)
+{
+	return !(Lhs == Rhs);
+}
diff --git a/Engine/Source/Runtime/Core/Public/Misc/StringView.h b/Engine/Source/Runtime/Core/Public/Misc/StringView.h
--- a/Engine/Source/Runtime/Core/Public/Misc/StringView.h
+++ b/Engine/Source/Runtime/Core/Public/Misc/StringView.h
@@ -2,192 +2,10 @@
 
 #pragma once
 
-#include "CoreTypes.h"
-#include "Math/NumericLimits.h"
-#include "Math/UnrealMathUtility.h"
-#include "Misc/CString.h"
-#include "Templates/AndOrNot.h"
-#include "Templates/EnableIf.h"
-#include "Templates/RemoveCV.h"
-#include "Templates/UnrealTemplate.h"
-#include "Traits/IsContiguousContainer.h"
+//#ifdef _MSC_VER
+//#pragma message(__FILE__"(6): warning : #include Containers/StringView.h instead of Misc/StringView.h")
+//#else
+//#pragma message("#include Containers/StringView.h instead of Misc/StringView.h")
+//#endif
 
-namespace StringViewPrivate
-{
-	/**
-	 * Trait testing whether a range type is compatible with the view type
-	 */
-	template <typename RangeType, typename ElementType>
-	struct TIsCompatibleRangeType : TIsSame<ElementType, typename TRemoveCV<typename TRemovePointer<decltype(GetData(DeclVal<RangeType&>()))>::Type>::Type>
-	{
-	};
-}
-
-/** String View
-
-	* A string view is implicitly constructible from const char* style strings and
-	  from character ranges such as FString.
-
-	* A string view does not own any data nor does it attempt to control any lifetimes, it 
-	  merely points at a subrange of characters in some other string. It's up to the user
-	  to ensure the underlying string stays valid for the lifetime of the string view.
-
-	* A string view does not represent a NUL terminated string and therefore you should
-	  never pass in the pointer returned by GetData() into a C-string API accepting only a
-	  pointer. You must either use a string builder to make a properly terminated string,
-	  or the ToString() function below if you absolutely must, and can live with the
-	  knowledge that you just added yet one more memory allocation even though memory
-	  allocations are not cheap.
-
-	String views are a good fit for arguments to functions which don't wish to care
-	which style of string construction is used by the caller. If you accept strings via
-	string views then the caller is free to use FString, FStringBuilder or raw C strings
-	or any other type which can be converted into a string builder.
-
-	I.e a function such as
-
-	void DoFoo(const FStringView& InString);
-
-	May be called as:
-
-	void MultiFoo()
-	{
-	     FString MyFoo(TEXT("Zoo"));
-		 const TCHAR* MyFooStr = *MyFoo;
-
-		 TStringBuilder<64> BuiltFoo;
-		 BuiltFoo.Append(TEXT("ABC"));
-
-		 DoFoo(MyFoo);
-		 DoFoo(MyFooStr);
-		 DoFoo(TEXT("ABC"));
-		 DoFoo(BuiltFoo);
-	}
-
-  */
-
-template <typename CharType, typename DerivedViewType>
-class TStringViewImpl
-{
-public:
-	using ElementType = CharType;
-	using SizeType = int32;
-	using ViewType = DerivedViewType;
-
-	template <typename RangeType>
-	using TIsCompatibleRangeType = TAnd<TIsContiguousContainer<RangeType>, StringViewPrivate::TIsCompatibleRangeType<RangeType, ElementType>>;
-
-	template <typename RangeType>
-	using TEnableIfCompatibleRangeType = typename TEnableIf<TIsCompatibleRangeType<RangeType>::Value>::Type;
-
-	TStringViewImpl() = default;
-
-	inline TStringViewImpl(const CharType* InData)
-		: DataPtr(InData)
-		, Size(TCString<CharType>::Strlen(InData))
-	{
-	}
-
-	inline TStringViewImpl(const CharType* InData, SizeType InSize)
-		: DataPtr(InData)
-		, Size(InSize)
-	{
-	}
-
-	inline const CharType& operator[](SizeType Index) const;
-
-	inline const CharType* GetData() const { return DataPtr; }
-
-	UE_DEPRECATED(4.25, "'Data' is deprecated. Please use 'GetData' instead!")
-	inline const CharType* Data() const { return DataPtr; }
-
-	// Capacity
-
-	inline SizeType Len() const { return Size; }
-	inline bool IsEmpty() const { return Size == 0; }
-
-	// Modifiers
-
-	inline void		RemovePrefix(SizeType CharCount)	{ DataPtr += CharCount; Size -= CharCount; }
-	inline void		RemoveSuffix(SizeType CharCount)	{ Size -= CharCount; }
-
-	// Operations
-
-	inline SizeType CopyString(CharType* Dest, SizeType CharCount, SizeType Position = 0) const;
-	inline ViewType SubStr(SizeType Position, SizeType CharCount) const { return Mid(Position, CharCount); }
-
-	// Maintain compatibility with FString 
-	inline ViewType Left(SizeType CharCount) const;
-	inline ViewType LeftChop(SizeType CharCount) const;
-	inline ViewType Right(SizeType CharCount) const;
-	inline ViewType RightChop(SizeType CharCount) const;
-	inline ViewType Mid(SizeType Position, SizeType CharCount = TNumericLimits<SizeType>::Max()) const;
-
-	// Comparison
-	bool Equals(const TStringViewImpl& Other, ESearchCase::Type SearchCase = ESearchCase::CaseSensitive) const;
-	CORE_API int32 Compare(const TStringViewImpl& Other, ESearchCase::Type SearchCase = ESearchCase::CaseSensitive) const;
-
-	inline bool StartsWith(CharType Prefix) const { return Size >= 1 && DataPtr[0] == Prefix; }
-	inline bool StartsWith(const TStringViewImpl& Prefix, ESearchCase::Type SearchCase = ESearchCase::IgnoreCase) const;
-
-	inline bool EndsWith(CharType Suffix) const { return Size >= 1 && DataPtr[Size-1] == Suffix; }
-	inline bool EndsWith(const TStringViewImpl& Suffix, ESearchCase::Type SearchCase = ESearchCase::IgnoreCase) const;
-
-	// Searching/Finding
-	CORE_API bool FindChar(CharType InChar, SizeType& OutIndex) const;
-	CORE_API bool FindLastChar(CharType InChar, SizeType& OutIndex) const;
-
-public:
-	/**
-	 * DO NOT USE DIRECTLY
-	 * STL-like iterators to enable range-based for loop support.
-	 */
-	inline const CharType* begin() const { return DataPtr; }
-	inline const CharType* end() const { return DataPtr + Size; }
-
-protected:
-	const CharType* DataPtr = nullptr;
-	SizeType Size = 0;
-};
-
-template <typename CharType, typename ViewType>
-struct TIsContiguousContainer<TStringViewImpl<CharType, ViewType>>
-{
-	enum { Value = true };
-};
-
-template <typename CharType, typename ViewType>
-inline SIZE_T GetNum(const TStringViewImpl<CharType, ViewType>& String)
-{
-	return String.Len();
-}
-
-//////////////////////////////////////////////////////////////////////////
-
-class FStringView : public TStringViewImpl<TCHAR, FStringView>
-{
-public:
-	using TStringViewImpl<ElementType, FStringView>::TStringViewImpl;
-};
-
-class FAnsiStringView : public TStringViewImpl<ANSICHAR, FAnsiStringView>
-{
-public:
-	using TStringViewImpl<ElementType, FAnsiStringView>::TStringViewImpl;
-};
-
-class FWideStringView : public TStringViewImpl<WIDECHAR, FWideStringView>
-{
-public:
-	using TStringViewImpl<ElementType, FWideStringView>::TStringViewImpl;
-};
-
-template <> struct TIsContiguousContainer<FStringView> { enum { Value = true }; };
-template <> struct TIsContiguousContainer<FAnsiStringView> { enum { Value = true }; };
-template <> struct TIsContiguousContainer<FWideStringView> { enum { Value = true }; };
-
-inline SIZE_T GetNum(const FStringView& String) { return String.Len(); }
-inline SIZE_T GetNum(const FAnsiStringView& String) { return String.Len(); }
-inline SIZE_T GetNum(const FWideStringView& String) { return String.Len(); }
-
-#include "StringView.inl"
+#include "Containers/StringView.h"
diff --git a/Engine/Source/Runtime/Core/Public/Misc/StringView.inl b/Engine/Source/Runtime/Core/Public/Misc/StringView.inl
deleted file mode 100644
