Add FPathViews as an analog for FPaths that uses FStringView and FStringBuilderBase

Includes implementations and tests for: GetCleanFilename, GetBaseFilename, GetBaseFilenameWithPath, GetPath, GetExtension, GetPathLeaf, Split.

diff --git a/Engine/Source/Runtime/Core/Private/Misc/PathViews.cpp b/Engine/Source/Runtime/Core/Private/Misc/PathViews.cpp
new file mode 100644
--- /dev/null
+++ b/Engine/Source/Runtime/Core/Private/Misc/PathViews.cpp
@@ -0,0 +1,76 @@
+// Copyright 1998-2019 Epic Games, Inc. All Rights Reserved.
+
+#include "Misc/PathViews.h"
+
+#include "Algo/FindLast.h"
+#include "Misc/StringView.h"
+
+namespace UE4PathViews_Private
+{
+	static bool IsSlashOrBackslash(TCHAR C) { return C == TEXT('/') || C == TEXT('\\'); }
+	static bool IsNotSlashOrBackslash(TCHAR C) { return C != TEXT('/') && C != TEXT('\\'); }
+}
+
+FStringView FPathViews::GetCleanFilename(const FStringView& InPath)
+{
+	if (const TCHAR* StartPos = Algo::FindLastByPredicate(InPath, UE4PathViews_Private::IsSlashOrBackslash))
+	{
+		return InPath.RightChop(StartPos - InPath.GetData() + 1);
+	}
+	return InPath;
+}
+
+FStringView FPathViews::GetBaseFilename(const FStringView& InPath)
+{
+	const FStringView CleanPath = GetCleanFilename(InPath);
+	return CleanPath.LeftChop(GetExtension(CleanPath, /*bIncludeDot*/ true).Len());
+}
+
+FStringView FPathViews::GetBaseFilenameWithPath(const FStringView& InPath)
+{
+	return InPath.LeftChop(GetExtension(InPath, /*bIncludeDot*/ true).Len());
+}
+
+FStringView FPathViews::GetBaseFilename(const FStringView& InPath, bool bRemovePath)
+{
+	return bRemovePath ? GetBaseFilename(InPath) : GetBaseFilenameWithPath(InPath);
+}
+
+FStringView FPathViews::GetPath(const FStringView& InPath)
+{
+	if (const TCHAR* EndPos = Algo::FindLastByPredicate(InPath, UE4PathViews_Private::IsSlashOrBackslash))
+	{
+		return InPath.Left(EndPos - InPath.GetData());
+	}
+	return FStringView();
+}
+
+FStringView FPathViews::GetExtension(const FStringView& InPath, bool bIncludeDot)
+{
+	if (const TCHAR* Dot = Algo::FindLast(GetCleanFilename(InPath), TEXT('.')))
+	{
+		const TCHAR* Extension = bIncludeDot ? Dot : Dot + 1;
+		return FStringView(Extension, InPath.GetData() + InPath.Len() - Extension);
+	}
+	return FStringView();
+}
+
+FStringView FPathViews::GetPathLeaf(const FStringView& InPath)
+{
+	if (const TCHAR* EndPos = Algo::FindLastByPredicate(InPath, UE4PathViews_Private::IsNotSlashOrBackslash))
+	{
+		++EndPos;
+		return GetCleanFilename(InPath.Left(EndPos - InPath.GetData()));
+	}
+	return FStringView();
+}
+
+void FPathViews::Split(const FStringView& InPath, FStringView& OutPath, FStringView& OutName, FStringView& OutExt)
+{
+	const FStringView CleanName = GetCleanFilename(InPath);
+	const TCHAR* DotPos = Algo::FindLast(CleanName, TEXT('.'));
+	const FStringView::SizeType NameLen = DotPos ? DotPos - CleanName.GetData() : CleanName.Len();
+	OutPath = InPath.LeftChop(CleanName.Len() + 1);
+	OutName = CleanName.Left(NameLen);
+	OutExt = CleanName.RightChop(NameLen + 1);
+}
diff --git a/Engine/Source/Runtime/Core/Private/Tests/Misc/PathViewsTest.cpp b/Engine/Source/Runtime/Core/Private/Tests/Misc/PathViewsTest.cpp
new file mode 100644
--- /dev/null
+++ b/Engine/Source/Runtime/Core/Private/Tests/Misc/PathViewsTest.cpp
@@ -0,0 +1,235 @@
+// Copyright 1998-2019 Epic Games, Inc. All Rights Reserved.
+
+#include "Misc/PathViews.h"
+
+#include "Misc/AutomationTest.h"
+#include "Misc/StringView.h"
+
+#if WITH_DEV_AUTOMATION_TESTS
+
+class FPathViewsTest : public FAutomationTestBase
+{
+protected:
+	using FAutomationTestBase::FAutomationTestBase;
+
+	void TestViewTransform(FStringView (*InFunction)(const FStringView& InPath), const FStringView& InPath, const TCHAR* InExpected)
+	{
+		const FStringView Actual = InFunction(InPath);
+		if (Actual != InExpected)
+		{
+			AddError(FString::Printf(TEXT("Failed on path '%.*s' (got '%.*s', expected '%s')."),
+				InPath.Len(), InPath.GetData(), Actual.Len(), Actual.GetData(), InExpected));
+		}
+	}
+};
+
+IMPLEMENT_CUSTOM_SIMPLE_AUTOMATION_TEST(FPathViewsGetCleanFilenameTest, FPathViewsTest, "System.Core.Misc.PathViews.GetCleanFilename", EAutomationTestFlags::ApplicationContextMask | EAutomationTestFlags::SmokeFilter)
+bool FPathViewsGetCleanFilenameTest::RunTest(const FString& InParameters)
+{
+	auto RunGetCleanFilenameTest = [this](const TCHAR* InPath, const TCHAR* InExpected)
+	{
+		TestViewTransform(FPathViews::GetCleanFilename, InPath, InExpected);
+	};
+
+	RunGetCleanFilenameTest(TEXT(""), TEXT(""));
+	RunGetCleanFilenameTest(TEXT(".txt"), TEXT(".txt"));
+	RunGetCleanFilenameTest(TEXT(".tar.gz"), TEXT(".tar.gz"));
+	RunGetCleanFilenameTest(TEXT(".tar.gz/"), TEXT(""));
+	RunGetCleanFilenameTest(TEXT(".tar.gz\\"), TEXT(""));
+	RunGetCleanFilenameTest(TEXT("File"), TEXT("File"));
+	RunGetCleanFilenameTest(TEXT("File.tar.gz"), TEXT("File.tar.gz"));
+	RunGetCleanFilenameTest(TEXT("File.tar.gz/"), TEXT(""));
+	RunGetCleanFilenameTest(TEXT("File.tar.gz\\"), TEXT(""));
+	RunGetCleanFilenameTest(TEXT("C:/Folder/"), TEXT(""));
+	RunGetCleanFilenameTest(TEXT("C:/Folder/File"), TEXT("File"));
+	RunGetCleanFilenameTest(TEXT("C:/Folder/File.tar.gz"), TEXT("File.tar.gz"));
+	RunGetCleanFilenameTest(TEXT("C:/Folder/First.Last/File"), TEXT("File"));
+	RunGetCleanFilenameTest(TEXT("C:/Folder/First.Last/File.tar.gz"), TEXT("File.tar.gz"));
+	RunGetCleanFilenameTest(TEXT("C:\\Folder\\"), TEXT(""));
+	RunGetCleanFilenameTest(TEXT("C:\\Folder\\File"), TEXT("File"));
+	RunGetCleanFilenameTest(TEXT("C:\\Folder\\First.Last\\"), TEXT(""));
+	RunGetCleanFilenameTest(TEXT("C:\\Folder\\First.Last\\File"), TEXT("File"));
+	RunGetCleanFilenameTest(TEXT("C:\\Folder\\First.Last\\File.tar.gz"), TEXT("File.tar.gz"));
+
+	return true;
+}
+
+IMPLEMENT_CUSTOM_SIMPLE_AUTOMATION_TEST(FPathViewsGetBaseFilenameTest, FPathViewsTest, "System.Core.Misc.PathViews.GetBaseFilename", EAutomationTestFlags::ApplicationContextMask | EAutomationTestFlags::SmokeFilter)
+bool FPathViewsGetBaseFilenameTest::RunTest(const FString& InParameters)
+{
+	auto RunGetBaseFilenameTest = [this](const TCHAR* InPath, const TCHAR* InExpected, const TCHAR* InExpectedWithPath)
+	{
+		const FStringView Path = InPath;
+		TestViewTransform(FPathViews::GetBaseFilename, Path, InExpected);
+		TestViewTransform(FPathViews::GetBaseFilenameWithPath, Path, InExpectedWithPath);
+	};
+
+	RunGetBaseFilenameTest(TEXT(""), TEXT(""), TEXT(""));
+	RunGetBaseFilenameTest(TEXT(".txt"), TEXT(""), TEXT(""));
+	RunGetBaseFilenameTest(TEXT(".tar.gz"), TEXT(".tar"), TEXT(".tar"));
+	RunGetBaseFilenameTest(TEXT(".tar.gz/"), TEXT(""), TEXT(".tar.gz/"));
+	RunGetBaseFilenameTest(TEXT(".tar.gz\\"), TEXT(""), TEXT(".tar.gz\\"));
+	RunGetBaseFilenameTest(TEXT("File"), TEXT("File"), TEXT("File"));
+	RunGetBaseFilenameTest(TEXT("File.txt"), TEXT("File"), TEXT("File"));
+	RunGetBaseFilenameTest(TEXT("File.tar.gz"), TEXT("File.tar"), TEXT("File.tar"));
+	RunGetBaseFilenameTest(TEXT("File.tar.gz/"), TEXT(""), TEXT("File.tar.gz/"));
+	RunGetBaseFilenameTest(TEXT("File.tar.gz\\"), TEXT(""), TEXT("File.tar.gz\\"));
+	RunGetBaseFilenameTest(TEXT("C:/Folder/"), TEXT(""), TEXT("C:/Folder/"));
+	RunGetBaseFilenameTest(TEXT("C:/Folder/File"), TEXT("File"), TEXT("C:/Folder/File"));
+	RunGetBaseFilenameTest(TEXT("C:/Folder/File.tar.gz"), TEXT("File.tar"), TEXT("C:/Folder/File.tar"));
+	RunGetBaseFilenameTest(TEXT("C:/Folder/First.Last/File"), TEXT("File"), TEXT("C:/Folder/First.Last/File"));
+	RunGetBaseFilenameTest(TEXT("C:/Folder/First.Last/File.txt"), TEXT("File"), TEXT("C:/Folder/First.Last/File"));
+	RunGetBaseFilenameTest(TEXT("C:/Folder/First.Last/File.tar.gz"), TEXT("File.tar"), TEXT("C:/Folder/First.Last/File.tar"));
+	RunGetBaseFilenameTest(TEXT("C:\\Folder\\"), TEXT(""), TEXT("C:\\Folder\\"));
+	RunGetBaseFilenameTest(TEXT("C:\\Folder\\File"), TEXT("File"), TEXT("C:\\Folder\\File"));
+	RunGetBaseFilenameTest(TEXT("C:\\Folder\\First.Last\\"), TEXT(""), TEXT("C:\\Folder\\First.Last\\"));
+	RunGetBaseFilenameTest(TEXT("C:\\Folder\\First.Last\\File"), TEXT("File"), TEXT("C:\\Folder\\First.Last\\File"));
+	RunGetBaseFilenameTest(TEXT("C:\\Folder\\First.Last\\File.txt"), TEXT("File"), TEXT("C:\\Folder\\First.Last\\File"));
+	RunGetBaseFilenameTest(TEXT("C:\\Folder\\First.Last\\File.tar.gz"), TEXT("File.tar"), TEXT("C:\\Folder\\First.Last\\File.tar"));
+
+	return true;
+}
+
+IMPLEMENT_CUSTOM_SIMPLE_AUTOMATION_TEST(FPathViewsGetPathTest, FPathViewsTest, "System.Core.Misc.PathViews.GetPath", EAutomationTestFlags::ApplicationContextMask | EAutomationTestFlags::SmokeFilter)
+bool FPathViewsGetPathTest::RunTest(const FString& InParameters)
+{
+	auto RunGetPathTest = [this](const TCHAR* InPath, const TCHAR* InExpected)
+	{
+		TestViewTransform(FPathViews::GetPath, InPath, InExpected);
+	};
+
+	RunGetPathTest(TEXT(""), TEXT(""));
+	RunGetPathTest(TEXT(".txt"), TEXT(""));
+	RunGetPathTest(TEXT(".tar.gz"), TEXT(""));
+	RunGetPathTest(TEXT(".tar.gz/"), TEXT(".tar.gz"));
+	RunGetPathTest(TEXT(".tar.gz\\"), TEXT(".tar.gz"));
+	RunGetPathTest(TEXT("File"), TEXT(""));
+	RunGetPathTest(TEXT("File.txt"), TEXT(""));
+	RunGetPathTest(TEXT("File.tar.gz"), TEXT(""));
+	RunGetPathTest(TEXT("File.tar.gz/"), TEXT("File.tar.gz"));
+	RunGetPathTest(TEXT("File.tar.gz\\"), TEXT("File.tar.gz"));
+	RunGetPathTest(TEXT("C:/Folder/"), TEXT("C:/Folder"));
+	RunGetPathTest(TEXT("C:/Folder/File"), TEXT("C:/Folder"));
+	RunGetPathTest(TEXT("C:/Folder/File.tar.gz"), TEXT("C:/Folder"));
+	RunGetPathTest(TEXT("C:/Folder/First.Last/File"), TEXT("C:/Folder/First.Last"));
+	RunGetPathTest(TEXT("C:/Folder/First.Last/File.tar.gz"), TEXT("C:/Folder/First.Last"));
+	RunGetPathTest(TEXT("C:\\Folder\\"), TEXT("C:\\Folder"));
+	RunGetPathTest(TEXT("C:\\Folder\\File"), TEXT("C:\\Folder"));
+	RunGetPathTest(TEXT("C:\\Folder\\First.Last\\"), TEXT("C:\\Folder\\First.Last"));
+	RunGetPathTest(TEXT("C:\\Folder\\First.Last\\File"), TEXT("C:\\Folder\\First.Last"));
+	RunGetPathTest(TEXT("C:\\Folder\\First.Last\\File.tar.gz"), TEXT("C:\\Folder\\First.Last"));
+
+	return true;
+}
+
+IMPLEMENT_CUSTOM_SIMPLE_AUTOMATION_TEST(FPathViewsGetExtensionTest, FPathViewsTest, "System.Core.Misc.PathViews.GetExtension", EAutomationTestFlags::ApplicationContextMask | EAutomationTestFlags::SmokeFilter)
+bool FPathViewsGetExtensionTest::RunTest(const FString& InParameters)
+{
+	auto RunGetExtensionTest = [this](const TCHAR* InPath, const TCHAR* InExpectedExt, const TCHAR* InExpectedExtDot)
+	{
+		const FStringView Path = InPath;
+		TestViewTransform([](const FStringView& InPath) { return FPathViews::GetExtension(InPath, /*bIncludeDot*/ false); }, Path, InExpectedExt);
+		TestViewTransform([](const FStringView& InPath) { return FPathViews::GetExtension(InPath, /*bIncludeDot*/ true); }, Path, InExpectedExtDot);
+	};
+
+	RunGetExtensionTest(TEXT(""), TEXT(""), TEXT(""));
+	RunGetExtensionTest(TEXT(".txt"), TEXT("txt"), TEXT(".txt"));
+	RunGetExtensionTest(TEXT(".tar.gz"), TEXT("gz"), TEXT(".gz"));
+	RunGetExtensionTest(TEXT(".tar.gz/"), TEXT(""), TEXT(""));
+	RunGetExtensionTest(TEXT(".tar.gz\\"), TEXT(""), TEXT(""));
+	RunGetExtensionTest(TEXT("File"), TEXT(""), TEXT(""));
+	RunGetExtensionTest(TEXT("File.txt"), TEXT("txt"), TEXT(".txt"));
+	RunGetExtensionTest(TEXT("File.tar.gz"), TEXT("gz"), TEXT(".gz"));
+	RunGetExtensionTest(TEXT("File.tar.gz/"), TEXT(""), TEXT(""));
+	RunGetExtensionTest(TEXT("File.tar.gz\\"), TEXT(""), TEXT(""));
+	RunGetExtensionTest(TEXT("C:/Folder/File"), TEXT(""), TEXT(""));
+	RunGetExtensionTest(TEXT("C:\\Folder\\File"), TEXT(""), TEXT(""));
+	RunGetExtensionTest(TEXT("C:/Folder/File.txt"), TEXT("txt"), TEXT(".txt"));
+	RunGetExtensionTest(TEXT("C:\\Folder\\File.txt"), TEXT("txt"), TEXT(".txt"));
+	RunGetExtensionTest(TEXT("C:/Folder/File.tar.gz"), TEXT("gz"), TEXT(".gz"));
+	RunGetExtensionTest(TEXT("C:\\Folder\\File.tar.gz"), TEXT("gz"), TEXT(".gz"));
+	RunGetExtensionTest(TEXT("C:/Folder/First.Last/File"), TEXT(""), TEXT(""));
+	RunGetExtensionTest(TEXT("C:\\Folder\\First.Last\\File"), TEXT(""), TEXT(""));
+	RunGetExtensionTest(TEXT("C:/Folder/First.Last/File.txt"), TEXT("txt"), TEXT(".txt"));
+	RunGetExtensionTest(TEXT("C:\\Folder\\First.Last\\File.txt"), TEXT("txt"), TEXT(".txt"));
+	RunGetExtensionTest(TEXT("C:/Folder/First.Last/File.tar.gz"), TEXT("gz"), TEXT(".gz"));
+	RunGetExtensionTest(TEXT("C:\\Folder\\First.Last\\File.tar.gz"), TEXT("gz"), TEXT(".gz"));
+
+	return true;
+}
+
+IMPLEMENT_CUSTOM_SIMPLE_AUTOMATION_TEST(FPathViewsGetPathLeafTest, FPathViewsTest, "System.Core.Misc.PathViews.GetPathLeaf", EAutomationTestFlags::ApplicationContextMask | EAutomationTestFlags::SmokeFilter)
+bool FPathViewsGetPathLeafTest::RunTest(const FString& InParameters)
+{
+	auto RunGetPathLeafTest = [this](const TCHAR* InPath, const TCHAR* InExpected)
+	{
+		TestViewTransform(FPathViews::GetPathLeaf, InPath, InExpected);
+	};
+
+	RunGetPathLeafTest(TEXT(""), TEXT(""));
+	RunGetPathLeafTest(TEXT(".txt"), TEXT(".txt"));
+	RunGetPathLeafTest(TEXT(".tar.gz"), TEXT(".tar.gz"));
+	RunGetPathLeafTest(TEXT(".tar.gz/"), TEXT(".tar.gz"));
+	RunGetPathLeafTest(TEXT(".tar.gz\\"), TEXT(".tar.gz"));
+	RunGetPathLeafTest(TEXT("File"), TEXT("File"));
+	RunGetPathLeafTest(TEXT("File.txt"), TEXT("File.txt"));
+	RunGetPathLeafTest(TEXT("File.tar.gz"), TEXT("File.tar.gz"));
+	RunGetPathLeafTest(TEXT("File.tar.gz/"), TEXT("File.tar.gz"));
+	RunGetPathLeafTest(TEXT("File.tar.gz\\"), TEXT("File.tar.gz"));
+	RunGetPathLeafTest(TEXT("C:/Folder/"), TEXT("Folder"));
+	RunGetPathLeafTest(TEXT("C:/Folder/File"), TEXT("File"));
+	RunGetPathLeafTest(TEXT("C:/Folder/File.tar.gz"), TEXT("File.tar.gz"));
+	RunGetPathLeafTest(TEXT("C:/Folder/First.Last/File"), TEXT("File"));
+	RunGetPathLeafTest(TEXT("C:/Folder/First.Last/File.tar.gz"), TEXT("File.tar.gz"));
+	RunGetPathLeafTest(TEXT("C:\\Folder\\"), TEXT("Folder"));
+	RunGetPathLeafTest(TEXT("C:\\Folder\\File"), TEXT("File"));
+	RunGetPathLeafTest(TEXT("C:\\Folder\\First.Last\\"), TEXT("First.Last"));
+	RunGetPathLeafTest(TEXT("C:\\Folder\\First.Last\\File"), TEXT("File"));
+	RunGetPathLeafTest(TEXT("C:\\Folder\\First.Last\\File.tar.gz"), TEXT("File.tar.gz"));
+
+	return true;
+}
+
+IMPLEMENT_CUSTOM_SIMPLE_AUTOMATION_TEST(FPathViewsSplitTest, FPathViewsTest, "System.Core.Misc.PathViews.Split", EAutomationTestFlags::ApplicationContextMask | EAutomationTestFlags::SmokeFilter)
+bool FPathViewsSplitTest::RunTest(const FString& InParameters)
+{
+	auto RunSplitTest = [this](const TCHAR* InPath, const TCHAR* InExpectedPath, const TCHAR* InExpectedName, const TCHAR* InExpectedExt)
+	{
+		FStringView SplitPath, SplitName, SplitExt;
+		FPathViews::Split(InPath, SplitPath, SplitName, SplitExt);
+		if (SplitPath != InExpectedPath || SplitName != InExpectedName || SplitExt != InExpectedExt)
+		{
+			AddError(FString::Printf(TEXT("Failed to split path '%s' (got ('%.*s', '%.*s', '%.*s'), expected ('%s', '%s', '%s'))."), InPath,
+				SplitPath.Len(), SplitPath.GetData(), SplitName.Len(), SplitName.GetData(), SplitExt.Len(), SplitExt.GetData(),
+				InExpectedPath, InExpectedName, InExpectedExt));
+		}
+	};
+
+	RunSplitTest(TEXT(""), TEXT(""), TEXT(""), TEXT(""));
+	RunSplitTest(TEXT(".txt"), TEXT(""), TEXT(""), TEXT("txt"));
+	RunSplitTest(TEXT(".tar.gz"), TEXT(""), TEXT(".tar"), TEXT("gz"));
+	RunSplitTest(TEXT(".tar.gz/"), TEXT(".tar.gz"), TEXT(""), TEXT(""));
+	RunSplitTest(TEXT(".tar.gz\\"), TEXT(".tar.gz"), TEXT(""), TEXT(""));
+	RunSplitTest(TEXT("File"), TEXT(""), TEXT("File"), TEXT(""));
+	RunSplitTest(TEXT("File.txt"), TEXT(""), TEXT("File"), TEXT("txt"));
+	RunSplitTest(TEXT("File.tar.gz"), TEXT(""), TEXT("File.tar"), TEXT("gz"));
+	RunSplitTest(TEXT("File.tar.gz/"), TEXT("File.tar.gz"), TEXT(""), TEXT(""));
+	RunSplitTest(TEXT("File.tar.gz\\"), TEXT("File.tar.gz"), TEXT(""), TEXT(""));
+	RunSplitTest(TEXT("C:/Folder/"), TEXT("C:/Folder"), TEXT(""), TEXT(""));
+	RunSplitTest(TEXT("C:/Folder/File"), TEXT("C:/Folder"), TEXT("File"), TEXT(""));
+	RunSplitTest(TEXT("C:/Folder/File.txt"), TEXT("C:/Folder"), TEXT("File"), TEXT("txt"));
+	RunSplitTest(TEXT("C:/Folder/File.tar.gz"), TEXT("C:/Folder"), TEXT("File.tar"), TEXT("gz"));
+	RunSplitTest(TEXT("C:/Folder/First.Last/File"), TEXT("C:/Folder/First.Last"), TEXT("File"), TEXT(""));
+	RunSplitTest(TEXT("C:/Folder/First.Last/File.txt"), TEXT("C:/Folder/First.Last"), TEXT("File"), TEXT("txt"));
+	RunSplitTest(TEXT("C:/Folder/First.Last/File.tar.gz"), TEXT("C:/Folder/First.Last"), TEXT("File.tar"), TEXT("gz"));
+	RunSplitTest(TEXT("C:\\Folder\\"), TEXT("C:\\Folder"), TEXT(""), TEXT(""));
+	RunSplitTest(TEXT("C:\\Folder\\File"), TEXT("C:\\Folder"), TEXT("File"), TEXT(""));
+	RunSplitTest(TEXT("C:\\Folder\\First.Last\\"), TEXT("C:\\Folder\\First.Last"), TEXT(""), TEXT(""));
+	RunSplitTest(TEXT("C:\\Folder\\First.Last\\File"), TEXT("C:\\Folder\\First.Last"), TEXT("File"), TEXT(""));
+	RunSplitTest(TEXT("C:\\Folder\\First.Last\\File.txt"), TEXT("C:\\Folder\\First.Last"), TEXT("File"), TEXT("txt"));
+	RunSplitTest(TEXT("C:\\Folder\\First.Last\\File.tar.gz"), TEXT("C:\\Folder\\First.Last"), TEXT("File.tar"), TEXT("gz"));
+
+	return true;
+}
+
+#endif
diff --git a/Engine/Source/Runtime/Core/Public/Misc/PathViews.h b/Engine/Source/Runtime/Core/Public/Misc/PathViews.h
new file mode 100644
--- /dev/null
+++ b/Engine/Source/Runtime/Core/Public/Misc/PathViews.h
@@ -0,0 +1,119 @@
+// Copyright 1998-2019 Epic Games, Inc. All Rights Reserved.
+
+#pragma once
+
+#include "CoreTypes.h"
+
+class FStringView;
+
+class CORE_API FPathViews
+{
+public:
+	/**
+	 * Returns the portion of the path after the last separator.
+	 *
+	 * Examples: (Using '/' but '\' is valid too.)
+	 * "A/B/C.D" -> "C.D"
+	 * "A/B/C"   -> "C"
+	 * "A/B/"    -> ""
+	 * "A"       -> "A"
+	 *
+	 * @return The portion of the path after the last separator.
+	 */
+	static FStringView GetCleanFilename(const FStringView& InPath);
+
+	/**
+	 * Returns the portion of the path after the last separator and before the last dot.
+	 *
+	 * Examples: (Using '/' but '\' is valid too.)
+	 * "A/B/C.D" -> "C"
+	 * "A/B/C"   -> "C"
+	 * "A/B/"    -> ""
+	 * "A"       -> "A"
+	 *
+	 * @return The portion of the path after the last separator and before the last dot.
+	 */
+	static FStringView GetBaseFilename(const FStringView& InPath);
+
+	/**
+	 * Returns the portion of the path before the last dot.
+	 *
+	 * Examples: (Using '/' but '\' is valid too.)
+	 * "A/B/C.D" -> "A/B/C"
+	 * "A/B/C"   -> "A/B/C"
+	 * "A/B/"    -> "A/B/"
+	 * "A"       -> "A"
+	 *
+	 * @return The portion of the path before the last dot.
+	 */
+	static FStringView GetBaseFilenameWithPath(const FStringView& InPath);
+
+	/** Returns the portion of the path before the last dot and optionally after the last separator. */
+	UE_DEPRECATED(4.25, "FPathViews::GetBaseFilename(InPath, bRemovePath) has been superseded by "
+		"FPathViews::GetBaseFilename(InPath) and FPathViews::GetBaseFilenameWithPath(InPath).")
+	static FStringView GetBaseFilename(const FStringView& InPath, bool bRemovePath);
+
+	/**
+	 * Returns the portion of the path before the last separator.
+	 *
+	 * Examples: (Using '/' but '\' is valid too.)
+	 * "A/B/C.D" -> "A/B"
+	 * "A/B/C"   -> "A/B"
+	 * "A/B/"    -> "A/B"
+	 * "A"       -> ""
+	 *
+	 * @return The portion of the path before the last separator.
+	 */
+	static FStringView GetPath(const FStringView& InPath);
+
+	/**
+	 * Returns the portion of the path after the last dot following the last separator, optionally including the dot.
+	 *
+	 * Examples: (Using '/' but '\' is valid too.)
+	 * "A/B.C.D" -> "D" (bIncludeDot=false) or ".D" (bIncludeDot=true)
+	 * "A/B/C.D" -> "D" (bIncludeDot=false) or ".D" (bIncludeDot=true)
+	 * "A/B/.D"  -> "D" (bIncludeDot=false) or ".D" (bIncludeDot=true)
+	 * ".D"      -> "D" (bIncludeDot=false) or ".D" (bIncludeDot=true)
+	 * "A/B/C"   -> ""
+	 * "A.B/C"   -> ""
+	 * "A.B/"    -> ""
+	 * "A"       -> ""
+	 *
+	 * @param bIncludeDot Whether to include the leading dot in the returned view.
+	 *
+	 * @return The portion of the path after the last dot following the last separator, optionally including the dot.
+	 */
+	static FStringView GetExtension(const FStringView& InPath, bool bIncludeDot=false);
+
+	/**
+	 * Returns the last non-empty path component.
+	 *
+	 * Examples: (Using '/' but '\' is valid too.)
+	 * "A/B/C.D" -> "C.D"
+	 * "A/B/C"   -> "C"
+	 * "A/B/"    -> "B"
+	 * "A"       -> "A"
+	 *
+	 * @return The last non-empty path component.
+	 */
+	static FStringView GetPathLeaf(const FStringView& InPath);
+
+	/**
+	 * Splits a path into three parts, any of which may be empty: the path, the clean name, and the extension.
+	 *
+	 * Examples: (Using '/' but '\' is valid too.)
+	 * "A/B/C.D" -> ("A/B", "C",   "D")
+	 * "A/B/C"   -> ("A/B", "C",   "")
+	 * "A/B/"    -> ("A/B", "",    "")
+	 * "A/B/.D"  -> ("A/B", "",    "D")
+	 * "A/B.C.D" -> ("A",   "B.C", "D")
+	 * "A"       -> ("",    "A",   "")
+	 * "A.D"     -> ("",    "A",   "D")
+	 * ".D"      -> ("",    "",    "D")
+	 *
+	 * @param OutPath [out] Receives the path portion of the input string, excluding the trailing separator.
+	 * @param OutName [out] Receives the name portion of the input string.
+	 * @param OutExt  [out] Receives the extension portion of the input string, excluding the dot.
+	 */
+	static void Split(const FStringView& InPath, FStringView& OutPath, FStringView& OutName, FStringView& OutExt);
+};
