Fix shadowed variable compile error on android.

diff --git a/Engine/Source/Runtime/Core/Private/Tests/Misc/PathViewsTest.cpp b/Engine/Source/Runtime/Core/Private/Tests/Misc/PathViewsTest.cpp
--- a/Engine/Source/Runtime/Core/Private/Tests/Misc/PathViewsTest.cpp
+++ b/Engine/Source/Runtime/Core/Private/Tests/Misc/PathViewsTest.cpp
@@ -128,8 +128,8 @@ bool FPathViewsGetExtensionTest::RunTest(const FString& InParameters)
 	auto RunGetExtensionTest = [this](const TCHAR* InPath, const TCHAR* InExpectedExt, const TCHAR* InExpectedExtDot)
 	{
 		const FStringView Path = InPath;
-		TestViewTransform([](const FStringView& InPath) { return FPathViews::GetExtension(InPath, /*bIncludeDot*/ false); }, Path, InExpectedExt);
-		TestViewTransform([](const FStringView& InPath) { return FPathViews::GetExtension(InPath, /*bIncludeDot*/ true); }, Path, InExpectedExtDot);
+		TestViewTransform([](const FStringView& InPath2) { return FPathViews::GetExtension(InPath2, /*bIncludeDot*/ false); }, Path, InExpectedExt);
+		TestViewTransform([](const FStringView& InPath2) { return FPathViews::GetExtension(InPath2, /*bIncludeDot*/ true); }, Path, InExpectedExtDot);
 	};
 
 	RunGetExtensionTest(TEXT(""), TEXT(""), TEXT(""));
