Removed FPaths functions which were marked as deprecated for 4.18

diff --git a/Engine/Source/Runtime/Core/Public/Misc/Paths.h b/Engine/Source/Runtime/Core/Public/Misc/Paths.h
--- a/Engine/Source/Runtime/Core/Public/Misc/Paths.h
+++ b/Engine/Source/Runtime/Core/Public/Misc/Paths.h
@@ -155,9 +155,6 @@ public:
 	 */
 	static FString ProjectDir();
 
-	UE_DEPRECATED(4.18, "FPaths::GameDir() has been superseded by FPaths::ProjectDir().")
-	static FORCEINLINE FString GameDir() { return ProjectDir(); }
-
 	/**
 	* Returns the root directory for user-specific game files.
 	*
@@ -165,9 +162,6 @@ public:
 	*/
 	static FString ProjectUserDir();
 
-	UE_DEPRECATED(4.18, "FPaths::GameUserDir() has been superseded by FPaths::ProjectUserDir().")
-	static FORCEINLINE FString GameUserDir() { return ProjectUserDir(); }
-
 	/**
 	 * Returns the content directory of the current game by looking at FApp::GetProjectName().
 	 *
@@ -175,9 +169,6 @@ public:
 	 */
 	static FString ProjectContentDir();
 
-	UE_DEPRECATED(4.18, "FPaths::GameContentDir() has been superseded by FPaths::ProjectContentDir().")
-	static FORCEINLINE FString GameContentDir() { return ProjectContentDir(); }
-
 	/**
 	* Returns the directory the root configuration files are located.
 	*
@@ -185,9 +176,6 @@ public:
 	*/
 	static FString ProjectConfigDir();
 
-	UE_DEPRECATED(4.18, "FPaths::GameConfigDir() has been superseded by FPaths::ProjectConfigDir().")
-	static FORCEINLINE FString GameConfigDir() { return ProjectConfigDir(); }
-
 	/**
 	 * Returns the saved directory of the current game by looking at FApp::GetProjectName().
 	 *
@@ -195,9 +183,6 @@ public:
 	 */
 	static const FString& ProjectSavedDir();
 
-	UE_DEPRECATED(4.18, "FPaths::GameSavedDir() has been superseded by FPaths::ProjectSavedDir().")
-	static FORCEINLINE FString GameSavedDir() { return ProjectSavedDir(); }
-
 	/**
 	 * Returns the intermediate directory of the current game by looking at FApp::GetProjectName().
 	 *
@@ -207,9 +192,6 @@ public:
 
 	static FString ShaderWorkingDir();
 
-	UE_DEPRECATED(4.18, "FPaths::GameIntermediateDir() has been superseded by FPaths::ProjectIntermediateDir().")
-	static FORCEINLINE FString GameIntermediateDir() { return ProjectIntermediateDir(); }
-
 	/**
 	 * Returns the plugins directory of the current game by looking at FApp::GetProjectName().
 	 *
@@ -217,9 +199,6 @@ public:
 	 */
 	static FString ProjectPluginsDir();
 
-	UE_DEPRECATED(4.18, "FPaths::GamePluginsDir() has been superseded by FPaths::ProjectPluginsDir().")
-	static FORCEINLINE FString GamePluginsDir() { return ProjectPluginsDir(); }
-
 	/**
 	 * Returns the mods directory of the current project by looking at FApp::GetProjectName().
 	 *
@@ -237,9 +216,6 @@ public:
 	*/
 	static FString ProjectPersistentDownloadDir();
 
-	UE_DEPRECATED(4.18, "FPaths::GamePersistentDownloadDir() has been superseded by FPaths::ProjectPersistentDownloadDir().")
-	static FORCEINLINE FString GamePersistentDownloadDir() { return ProjectPersistentDownloadDir(); }
-
 	/**
 	 * Returns the directory the engine uses to look for the source leaf ini files. This
 	 * can't be an .ini variable for obvious reasons.
@@ -299,9 +275,6 @@ public:
 	 */
 	static FString ProjectLogDir();
 
-	UE_DEPRECATED(4.18, "FPaths::GameLogDir() has been superseded by FPaths::ProjectLogDir().")
-	static FORCEINLINE FString GameLogDir() { return ProjectLogDir(); }
-
 	/** Returns the directory for automation save files */
 	static FString AutomationDir();
 
