Deleted deprecated functions from UnrealString.h

diff --git a/Engine/Source/Runtime/Core/Public/Containers/UnrealString.h b/Engine/Source/Runtime/Core/Public/Containers/UnrealString.h
--- a/Engine/Source/Runtime/Core/Public/Containers/UnrealString.h
+++ b/Engine/Source/Runtime/Core/Public/Containers/UnrealString.h
@@ -1442,18 +1442,6 @@ public:
 		return PrintfImpl(Fmt, Args...);
 	}
 
-	template <typename FmtType, typename... Types>
-	UE_DEPRECATED(4.20, "The formatting string must now be a TCHAR string literal.")
-	static typename TEnableIf<!TIsArrayOrRefOfType<FmtType, TCHAR>::Value, FString>::Type Printf(const FmtType& Fmt, Types... Args)
-	{
-		// NOTE: When this deprecated function is removed, the return type of the overload above
-		//       should be set to simply FString.
-
-		static_assert(TAnd<TIsValidVariadicFunctionArg<Types>...>::Value, "Invalid argument(s) passed to FString::Printf");
-
-		return PrintfImpl(Fmt, Args...);
-	}
-
 	/**
 	 * Just like Printf, but appends the formatted text to the existing FString instead.
 	 * @return a reference to the modified string, so that it can be chained
@@ -1568,18 +1556,6 @@ public:
 		return MatchesWildcard(*Wildcard, SearchCase);
 	}
 
-	/**
-	 * Removes whitespace characters from the front of this string.
-	 */
-	UE_DEPRECATED(4.18, "FString::Trim() has been split into separate functions for copy and mutate semantics. Call FString::TrimStart() to return a copy of the string with whitespace trimmed from the start, or FString::TrimStartInline() to modify an FString object in-place.")
-	FString Trim();
-
-	/**
-	 * Removes trailing whitespace characters
-	 */
-	UE_DEPRECATED(4.18, "FString::TrimTrailing() has been split into separate functions for copy and mutate semantics. Call FString::TrimEnd() to return a copy of the string with whitespace trimmed from the start, or FString::TrimEndInline() to modify an FString object in-place.")
-	FString TrimTrailing();
-
 	/**
 	 * Removes whitespace characters from the start and end of this string. Modifies the string in-place.
 	 */
@@ -2246,51 +2222,6 @@ static bool LexTryParseString(bool& OutValue, const TCHAR* Buffer)
 }
 
 
-/** Deprecated Lex namespace. Forwards on to Lex prefixed equivalents for backwards compatibility. See Lex-prefixed functions above. */
-namespace Lex
-{
-	
-	template<typename T> 
-	UE_DEPRECATED(4.20, "Lex::FromString has been deprecated. Please use LexFromString instead.")
-	void FromString(T& OutValue, const TCHAR* Buffer) 
-	{ 
-		LexFromString(OutValue, Buffer); 
-	}
-
-	template<typename T> 
-	UE_DEPRECATED(4.20, "Lex::ToString has been deprecated. Please use LexToString instead.")
-#if PLATFORM_COMPILER_HAS_DECLTYPE_AUTO
-	decltype(auto) ToString(T&& Value)
-#else
-	auto ToString(T&& Value) -> decltype(LexToString(Forward<T>(Value)))
-#endif
-	{
-		return LexToString(Forward<T>(Value)); 
-	}
-
-	template<typename T>
-	UE_DEPRECATED(4.20, "Lex::ToSanitizedString has been deprecated. Please use LexToSanitizedString instead.")
-#if PLATFORM_COMPILER_HAS_DECLTYPE_AUTO
-	decltype(auto) ToSanitizedString(T&& Value)
-#else
-	auto ToSanitizedString(T&& Value) -> decltype(LexToSanitizedString(Forward<T>(Value)))
-#endif
-	{
-		return LexToSanitizedString(Forward<T>(Value)); 
-	}
-
-	template<typename T>
-	UE_DEPRECATED(4.20, "Lex::TryParseString has been deprecated. Please use LexTryParseString instead.")
-	bool TryParseString(T& OutValue, const TCHAR* Buffer) 
-	{ 
-		return LexTryParseString(OutValue, Buffer); 
-	}
-}
-
-// Deprecated alias for old LexicalConversion namespace.
-// Namespace alias deprecation doesn't exist, so we can't actually mark it with the UE_DEPRECATED() macro.
-namespace LexicalConversion = Lex;
-
 /** Shorthand legacy use for Lex functions */
 template<typename T>
 struct TTypeToString
