Removed implementations of functions deleted in CL10482240 (accidental split CL submit :()

diff --git a/Engine/Source/Runtime/Core/Private/Containers/String.cpp b/Engine/Source/Runtime/Core/Private/Containers/String.cpp
--- a/Engine/Source/Runtime/Core/Private/Containers/String.cpp
+++ b/Engine/Source/Runtime/Core/Private/Containers/String.cpp
@@ -484,44 +484,6 @@ void FString::ReplaceCharInlineIgnoreCase(const TCHAR SearchChar, const TCHAR Re
 	ReplaceCharInlineCaseSensitive(SearchChar, ReplacementChar);
 }
 
-FString FString::Trim()
-{
-	int32 Pos = 0;
-	while(Pos < Len())
-	{
-		if( FChar::IsWhitespace( (*this)[Pos] ) )
-		{
-			Pos++;
-		}
-		else
-		{
-			break;
-		}
-	}
-
-	*this = Right( Len()-Pos );
-
-	return *this;
-}
-
-FString FString::TrimTrailing( void )
-{
-	int32 Pos = Len() - 1;
-	while( Pos >= 0 )
-	{
-		if( !FChar::IsWhitespace( ( *this )[Pos] ) )
-		{
-			break;
-		}
-
-		Pos--;
-	}
-
-	*this = Left( Pos + 1 );
-
-	return( *this );
-}
-
 void FString::TrimStartAndEndInline()
 {
 	TrimEndInline();
