Removed deprecated functions in Compression.(h|cpp)

diff --git a/Engine/Source/Runtime/Core/Private/Misc/Compression.cpp b/Engine/Source/Runtime/Core/Private/Misc/Compression.cpp
--- a/Engine/Source/Runtime/Core/Private/Misc/Compression.cpp
+++ b/Engine/Source/Runtime/Core/Private/Misc/Compression.cpp
@@ -793,54 +793,3 @@ bool FCompression::VerifyCompressionFlagsValid(int32 InCompressionFlags)
 /***********************
   Deprecated functions
 ***********************/
-
-int32 FCompression::CompressMemoryBound(ECompressionFlags Flags, int32 UncompressedSize, int32 BitWindow)
-{
-	switch (Flags & COMPRESS_DeprecatedFormatFlagsMask)
-	{
-		case COMPRESS_ZLIB:
-			return CompressMemoryBound(NAME_Zlib, UncompressedSize, (ECompressionFlags)(Flags & COMPRESS_OptionsFlagsMask), BitWindow);
-		case COMPRESS_GZIP:
-			return CompressMemoryBound(NAME_Gzip, UncompressedSize, (ECompressionFlags)(Flags & COMPRESS_OptionsFlagsMask));
-		case COMPRESS_Custom:
-			return CompressMemoryBound(TEXT("Oodle"), UncompressedSize, (ECompressionFlags)(Flags & COMPRESS_OptionsFlagsMask));
-	}
-
-	return 0;
-}
-
-bool FCompression::CompressMemory(ECompressionFlags Flags, void* CompressedBuffer, int32& CompressedSize, const void* UncompressedBuffer, int32 UncompressedSize, int32 BitWindow)
-{
-	switch (Flags & COMPRESS_DeprecatedFormatFlagsMask)
-	{
-		case COMPRESS_ZLIB:
-			return CompressMemory(NAME_Zlib, CompressedBuffer, CompressedSize, UncompressedBuffer, UncompressedSize, (ECompressionFlags)(Flags & COMPRESS_OptionsFlagsMask), BitWindow);
-		case COMPRESS_GZIP:
-			return CompressMemory(NAME_Gzip, CompressedBuffer, CompressedSize, UncompressedBuffer, UncompressedSize, (ECompressionFlags)(Flags & COMPRESS_OptionsFlagsMask));
-		case COMPRESS_Custom:
-			return CompressMemory(TEXT("Oodle"), CompressedBuffer, CompressedSize, UncompressedBuffer, UncompressedSize, (ECompressionFlags)(Flags & COMPRESS_OptionsFlagsMask));
-	}
-
-	return false;
-}
-
-bool FCompression::UncompressMemory(ECompressionFlags Flags, void* UncompressedBuffer, int32 UncompressedSize, const void* CompressedBuffer, int32 CompressedSize, bool bIsSourcePadded, int32 BitWindow)
-{
-	if (bIsSourcePadded)
-	{
-		Flags = (ECompressionFlags)(Flags | COMPRESS_SourceIsPadded);
-	}
-
-	switch (Flags & COMPRESS_DeprecatedFormatFlagsMask)
-	{
-		case COMPRESS_ZLIB:
-			return UncompressMemory(NAME_Zlib, UncompressedBuffer, UncompressedSize, CompressedBuffer, CompressedSize, (ECompressionFlags)(Flags & COMPRESS_OptionsFlagsMask), BitWindow);
-		case COMPRESS_GZIP:
-			return UncompressMemory(NAME_Gzip, UncompressedBuffer, UncompressedSize, CompressedBuffer, CompressedSize, (ECompressionFlags)(Flags & COMPRESS_OptionsFlagsMask));
-		case COMPRESS_Custom:
-			return UncompressMemory(TEXT("Oodle"), UncompressedBuffer, UncompressedSize, CompressedBuffer, CompressedSize, (ECompressionFlags)(Flags & COMPRESS_OptionsFlagsMask));
-	}
-
-	return false;
-}
-
diff --git a/Engine/Source/Runtime/Core/Public/Misc/Compression.h b/Engine/Source/Runtime/Core/Public/Misc/Compression.h
--- a/Engine/Source/Runtime/Core/Public/Misc/Compression.h
+++ b/Engine/Source/Runtime/Core/Public/Misc/Compression.h
@@ -100,17 +100,6 @@ struct FCompression
 	 */
 	CORE_API static bool VerifyCompressionFlagsValid(int32 InCompressionFlags);
 
-
-
-
-
-	UE_DEPRECATED(4.20, "Use the FName based version of CompressMemoryBound")
-	CORE_API static int32 CompressMemoryBound(ECompressionFlags Flags, int32 UncompressedSize, int32 BitWindow = DEFAULT_ZLIB_BIT_WINDOW);
-	UE_DEPRECATED(4.20, "Use the FName based version of CompressMemory")
-	CORE_API static bool CompressMemory(ECompressionFlags Flags, void* CompressedBuffer, int32& CompressedSize, const void* UncompressedBuffer, int32 UncompressedSize, int32 BitWindow = DEFAULT_ZLIB_BIT_WINDOW);
-	UE_DEPRECATED(4.20, "Use the FName based version of UncompressMemory")
-	CORE_API static bool UncompressMemory(ECompressionFlags Flags, void* UncompressedBuffer, int32 UncompressedSize, const void* CompressedBuffer, int32 CompressedSize, bool bIsSourcePadded = false, int32 BitWindow = DEFAULT_ZLIB_BIT_WINDOW);
-
 	CORE_API static FName GetCompressionFormatFromDeprecatedFlags(ECompressionFlags DeprecatedFlags);
 
 private:
