Removed deprecated functions in OutputDevice.(cpp|h)

diff --git a/Engine/Source/Runtime/Core/Private/Misc/OutputDevice.cpp b/Engine/Source/Runtime/Core/Private/Misc/OutputDevice.cpp
--- a/Engine/Source/Runtime/Core/Private/Misc/OutputDevice.cpp
+++ b/Engine/Source/Runtime/Core/Private/Misc/OutputDevice.cpp
@@ -10,16 +10,6 @@
 
 DEFINE_LOG_CATEGORY(LogOutputDevice);
 
-const TCHAR* FOutputDevice::VerbosityToString(ELogVerbosity::Type Verbosity)
-{
-	return FOutputDeviceHelper::VerbosityToString(Verbosity);
-}
-
-FString FOutputDevice::FormatLogLine( ELogVerbosity::Type Verbosity, const class FName& Category, const TCHAR* Message /*= nullptr*/, ELogTimes::Type LogTime /*= ELogTimes::None*/, const double Time /*= -1.0*/ )
-{
-	return FOutputDeviceHelper::FormatLogLine(Verbosity, Category, Message, LogTime, Time);
-}
-
 void FOutputDevice::Log( ELogVerbosity::Type Verbosity, const TCHAR* Str )
 {
 	Serialize( Str, Verbosity, NAME_None );
diff --git a/Engine/Source/Runtime/Core/Public/Misc/OutputDevice.h b/Engine/Source/Runtime/Core/Public/Misc/OutputDevice.h
--- a/Engine/Source/Runtime/Core/Public/Misc/OutputDevice.h
+++ b/Engine/Source/Runtime/Core/Public/Misc/OutputDevice.h
@@ -132,14 +132,6 @@ public:
 
 	virtual ~FOutputDevice() = default;
 
-	// static helpers
-	UE_DEPRECATED(4.12, "Please use FOutputDeviceHelper::VerbosityToString.")
-	static const TCHAR* VerbosityToString(ELogVerbosity::Type Verbosity);
-
-	UE_DEPRECATED(4.12, "Please use FOutputDeviceHelper::FormatLogLine.")
-	static FString FormatLogLine(ELogVerbosity::Type Verbosity, const FName& Category, const TCHAR* Message = nullptr, ELogTimes::Type LogTime = ELogTimes::None, const double Time = -1.0);
-
-
 	// FOutputDevice interface.
 	virtual void Serialize( const TCHAR* V, ELogVerbosity::Type Verbosity, const FName& Category ) = 0;
 	virtual void Serialize( const TCHAR* V, ELogVerbosity::Type Verbosity, const FName& Category, const double Time )
