Removed deprecated FArchiveLoadCompressedProxy constructor

diff --git a/Engine/Source/Runtime/Core/Private/Serialization/ArchiveLoadCompressedProxy.cpp b/Engine/Source/Runtime/Core/Private/Serialization/ArchiveLoadCompressedProxy.cpp
--- a/Engine/Source/Runtime/Core/Private/Serialization/ArchiveLoadCompressedProxy.cpp
+++ b/Engine/Source/Runtime/Core/Private/Serialization/ArchiveLoadCompressedProxy.cpp
@@ -9,11 +9,6 @@
 	FArchiveLoadCompressedProxy
 ----------------------------------------------------------------------------*/
 
-FArchiveLoadCompressedProxy::FArchiveLoadCompressedProxy(FArchiveLoadCompressedProxy::EVS2015Redirector, const TArray<uint8>& InCompressedData, ECompressionFlags InCompressionFlags)
-	: FArchiveLoadCompressedProxy(InCompressedData, FCompression::GetCompressionFormatFromDeprecatedFlags(InCompressionFlags))
-{
-}
-
 FArchiveLoadCompressedProxy::FArchiveLoadCompressedProxy(const TArray<uint8>& InCompressedData, FName InCompressionFormat, ECompressionFlags InCompressionFlags)
 	:	CompressedData(InCompressedData)
 	,	CompressionFormat(InCompressionFormat)
diff --git a/Engine/Source/Runtime/Core/Public/Serialization/ArchiveLoadCompressedProxy.h b/Engine/Source/Runtime/Core/Public/Serialization/ArchiveLoadCompressedProxy.h
--- a/Engine/Source/Runtime/Core/Public/Serialization/ArchiveLoadCompressedProxy.h
+++ b/Engine/Source/Runtime/Core/Public/Serialization/ArchiveLoadCompressedProxy.h
@@ -16,29 +16,7 @@
  */
 class CORE_API FArchiveLoadCompressedProxy : public FArchive
 {
-private:
-	/**
-	 * This enum and the following constructor is a workaround for VC compiler bug that prevents using attributes
-	 * on constructors without inline implementation. This should be removed when the deprecated ctor is removed.
-	 */
-	enum EVS2015Redirector
-	{
-		Redirect
-	};
-	FArchiveLoadCompressedProxy(EVS2015Redirector InUnused, const TArray<uint8>& InCompressedData, ECompressionFlags InCompressionFlags);
-
 public:
-	/** 
-	 * Constructor, initializing all member variables and allocating temp memory.
-	 *
-	 * @param	InCompressedData	Array of bytes that is holding compressed data
-	 * @param	InCompressionFormat	What format to compress with
-	 */
-	UE_DEPRECATED(4.21, "Use the FName version of FArchiveLoadCompressedProxy constructor")
-	FArchiveLoadCompressedProxy(const TArray<uint8>& InCompressedData, ECompressionFlags InCompressionFlags)
-		// Make sure to remove the EVS2015Redirector constructor when this constructor is removed
-		: FArchiveLoadCompressedProxy(EVS2015Redirector::Redirect, InCompressedData, InCompressionFlags)
-	{}
 	FArchiveLoadCompressedProxy(const TArray<uint8>& InCompressedData, FName CompressionFormat, ECompressionFlags InCompressionFlags=COMPRESS_NoFlags);
 
 	/** Destructor, freeing temporary memory. */
