Removed deprecated Histogram functionality

diff --git a/Engine/Source/Runtime/Core/Private/ProfilingDebugging/Histogram.cpp b/Engine/Source/Runtime/Core/Private/ProfilingDebugging/Histogram.cpp
--- a/Engine/Source/Runtime/Core/Private/ProfilingDebugging/Histogram.cpp
+++ b/Engine/Source/Runtime/Core/Private/ProfilingDebugging/Histogram.cpp
@@ -123,69 +123,6 @@ void FHistogram::AddMeasurement(double ValueForBinning, double MeasurementValue)
 	}
 }
 
-void FHistogram::DumpToAnalytics(const FString& ParamNamePrefix, TArray<TPair<FString, double>>& OutParamArray)
-{
-	double TotalSum = 0;
-	double TotalObservations = 0;
-	double AverageObservation = 0;
-
-	if (LIKELY(Bins.Num()))
-	{
-		for (int BinIdx = 0, LastBinIdx = Bins.Num() - 1; BinIdx < LastBinIdx; ++BinIdx)
-		{
-			FBin& Bin = Bins[BinIdx];
-			FString ParamName = FString::Printf(TEXT("_%.0f_%.0f"), Bin.MinValue, Bin.UpperBound);
-			OutParamArray.Add(TPairInitializer<FString, double>(ParamNamePrefix + ParamName + TEXT("_Count"), Bin.Count));
-			OutParamArray.Add(TPairInitializer<FString, double>(ParamNamePrefix + ParamName + TEXT("_Sum"), Bin.Sum));
-
-			TotalObservations += Bin.Count;
-			TotalSum += Bin.Sum;
-		}
-
-		FBin& LastBin = Bins.Last();
-		FString ParamName = FString::Printf(TEXT("_%.0f_AndAbove"), LastBin.MinValue);
-		OutParamArray.Add(TPairInitializer<FString, double>(ParamNamePrefix + ParamName + TEXT("_Count"), LastBin.Count));
-		OutParamArray.Add(TPairInitializer<FString, double>(ParamNamePrefix + ParamName + TEXT("_Sum"), LastBin.Sum));
-
-		TotalObservations += LastBin.Count;
-		TotalSum += LastBin.Sum;
-	}
-
-	if (TotalObservations > 0)
-	{
-		AverageObservation = TotalSum / TotalObservations;
-	}
-
-	// add an average value for ease of monitoring/analyzing
-	OutParamArray.Add(TPairInitializer<FString, double>(ParamNamePrefix + TEXT("_Average"), AverageObservation));
-}
-
-FString FHistogram::DumpToAnalyticsString() const
-{
-	// Bucket:Count:Sum;Bucket:Count:Sum;...
-	FString Result;
-	if (LIKELY(Bins.Num()))
-	{
-		for (int BinIdx = 0, LastBinIdx = Bins.Num() - 1; BinIdx < LastBinIdx+1; ++BinIdx)
-		{
-			const FBin& Bin = Bins[BinIdx];
-			if (BinIdx != 0)
-			{
-				Result += TEXT(';');
-			}
-			if (BinIdx != LastBinIdx)
-			{
-				Result += FString::Printf(TEXT("%d_%d:%d:%.5f"), (int64)Bin.MinValue, (int64)Bin.UpperBound, Bin.Count, Bin.Sum);
-			}
-			else
-			{
-				Result += FString::Printf(TEXT("%d_AndAbove:%d:%.5f"), (int64)Bin.MinValue, Bin.Count, Bin.Sum);
-			}
-		}
-	}
-	return Result;
-}
-
 FString FHistogram::DumpToJsonString(TFunctionRef<FString (double, double)> ConvertBinToLabel) const
 {
 	// {"Bin":"Name","Count":Count,"Sum":Sum}
diff --git a/Engine/Source/Runtime/Core/Public/ProfilingDebugging/Histogram.h b/Engine/Source/Runtime/Core/Public/ProfilingDebugging/Histogram.h
--- a/Engine/Source/Runtime/Core/Public/ProfilingDebugging/Histogram.h
+++ b/Engine/Source/Runtime/Core/Public/ProfilingDebugging/Histogram.h
@@ -56,14 +56,6 @@ struct CORE_API FHistogram
 	/** Prints histogram contents to the log. */
 	void DumpToLog(const FString& HistogramName);
 
-	/** Populates array commonly used in analytics events, adding two pairs per bin (count and sum). */
-	UE_DEPRECATED(4.21, "This function is deprecated. DumpToJsonString instead")
-	void DumpToAnalytics(const FString& ParamNamePrefix, TArray<TPair<FString, double>>& OutParamArray);
-
-	/** Returns a string in a format that can be consumed by an analytics event in the format expected by the analytics backends. Bucket:Count:Sum;Bucket:Count:Sum;...  */
-	UE_DEPRECATED(4.21, "This function is deprecated. DumpToJsonString instead")
-	FString DumpToAnalyticsString() const;
-
 	/** 
 	 * Returns a string in a Json format: [{"Bin":"BinName","Count":Count,"Sum":Sum},...]. 
 	 * Bucket name is constructed by calling ConvertBinToLabel on the MinValue and UpperBound for each bin. 
