Removed deprecated TWeakObjectPtr operator=

diff --git a/Engine/Source/Runtime/Core/Public/UObject/WeakObjectPtrTemplates.h b/Engine/Source/Runtime/Core/Public/UObject/WeakObjectPtrTemplates.h
--- a/Engine/Source/Runtime/Core/Public/UObject/WeakObjectPtrTemplates.h
+++ b/Engine/Source/Runtime/Core/Public/UObject/WeakObjectPtrTemplates.h
@@ -95,14 +95,6 @@ public:
 		TWeakObjectPtrBase::operator=(TempObject);
 		return *this;
 	}
-	template<class U>
-	UE_DEPRECATED(4.19, "Implicit conversions from const pointers to non-const TWeakObjectPtrs has been deprecated. Please const-correct this usage.")
-	FORCEINLINE typename TEnableIf<TLosesQualifiersFromTo<U, T>::Value, TWeakObjectPtr&>::Type operator=(U* Object)
-	{
-		const T* TempObject = Object;
-		TWeakObjectPtrBase::operator=(TempObject);
-		return *this;
-	}
 
 	/**  
 	 * Assign from another weak pointer, intended for derived-to-base conversions
