Removed misc deprecated functions

diff --git a/Engine/Source/Runtime/Core/Public/Misc/App.h b/Engine/Source/Runtime/Core/Public/Misc/App.h
--- a/Engine/Source/Runtime/Core/Public/Misc/App.h
+++ b/Engine/Source/Runtime/Core/Public/Misc/App.h
@@ -104,9 +104,6 @@ public:
 		return GInternalProjectName;
 	}
 
-	UE_DEPRECATED(4.18, "GetGameName() has been superseded by GetProjectName().")
-	FORCEINLINE static const TCHAR* GetGameName() { return GetProjectName(); }
-
 	/**
 	 * Gets the name of the application, i.e. "UE4" or "Rocket".
 	 *
@@ -142,9 +139,6 @@ public:
 		return (IsProjectNameEmpty() == false) && (FCString::Stricmp(GInternalProjectName, TEXT("None")) != 0);
 	}
 
-	UE_DEPRECATED(4.18, "HasGameName() has been superseded by HasProjectName().")
-	FORCEINLINE static bool HasGameName() { return HasProjectName(); }
-
 	/**
 	 * Checks whether this application is a game.
 	 *
@@ -173,9 +167,6 @@ public:
 		return (GInternalProjectName[0] == 0);
 	}
 
-	UE_DEPRECATED(4.18, "IsGameNameEmpty() has been superseded by IsProjectNameEmpty().")
-	FORCEINLINE static bool IsGameNameEmpty() { return IsProjectNameEmpty(); }
-
 	/**
 	* Sets the name of the current project.
 	*
@@ -189,9 +180,6 @@ public:
 		GInternalProjectName[UE_ARRAY_COUNT(GInternalProjectName) - 1] = 0;
 	}
 
-	UE_DEPRECATED(4.18, "SetGameName() has been superseded by SetProjectName().")
-	FORCEINLINE static void SetGameName(const TCHAR* InGameName) { SetProjectName(InGameName); }
-
 public:
 
 	/**
@@ -616,13 +604,6 @@ public:
 		return TimecodeFrameRate;
 	}
 	
-	UE_DEPRECATED(4.21, "Please use the version of SetTimecodeAndFrameRate")
-	static void SetTimecode(FTimecode InTimecode)
-	{
-		Timecode = InTimecode;
-		TimecodeFrameRate = FFrameRate(60, 1);
-	}
-
 	/**
 	 * Sets the current timecode, and the frame rate to which it's relative.
 	 *
diff --git a/Engine/Source/Runtime/Core/Public/Misc/AssertionMacros.h b/Engine/Source/Runtime/Core/Public/Misc/AssertionMacros.h
--- a/Engine/Source/Runtime/Core/Public/Misc/AssertionMacros.h
+++ b/Engine/Source/Runtime/Core/Public/Misc/AssertionMacros.h
@@ -107,18 +107,6 @@ public:
 		return OptionallyLogFormattedEnsureMessageReturningFalseImpl(bLog, Expr, File, Line, FormattedMsg, Args...);
 	}
 
-	template <typename FmtType, typename... Types>
-	UE_DEPRECATED(4.20, "The formatting string must now be a TCHAR string literal.")
-	static FORCEINLINE typename TEnableIf<!TIsArrayOrRefOfType<FmtType, TCHAR>::Value, bool>::Type OptionallyLogFormattedEnsureMessageReturningFalse(bool bLog, const ANSICHAR* Expr, const ANSICHAR* File, int32 Line, const FmtType& FormattedMsg, Types... Args)
-	{
-		// NOTE: When this deprecated function is removed, the return type of the overload above
-		//       should be set to simply bool.
-
-		static_assert(TAnd<TIsValidVariadicFunctionArg<Types>...>::Value, "Invalid argument(s) passed to ensureMsgf");
-
-		return OptionallyLogFormattedEnsureMessageReturningFalseImpl(bLog, Expr, File, Line, FormattedMsg, Args...);
-	}
-
 #endif // DO_CHECK || DO_GUARD_SLOW
 
 	/**
diff --git a/Engine/Source/Runtime/Core/Public/Misc/MonitoredProcess.h b/Engine/Source/Runtime/Core/Public/Misc/MonitoredProcess.h
--- a/Engine/Source/Runtime/Core/Public/Misc/MonitoredProcess.h
+++ b/Engine/Source/Runtime/Core/Public/Misc/MonitoredProcess.h
@@ -88,17 +88,6 @@ public:
 		return ProcessHandle;
 	}
 
-	/**
-	 * Checks whether the process is still running.
-	 *
-	 * @return true if the process is running, false otherwise.
-	 */
-	UE_DEPRECATED(4.16, "IsRunning() is deprecated because it doesn't support -nothreading. Please use Update()")
-	bool IsRunning() const
-	{
-		return bIsRunning;
-	}
-
 	/**
 	* Checks whether the process is still running. In single threaded mode, this will tick the thread processing
 	*
diff --git a/Engine/Source/Runtime/Core/Public/Serialization/CustomVersion.h b/Engine/Source/Runtime/Core/Public/Serialization/CustomVersion.h
--- a/Engine/Source/Runtime/Core/Public/Serialization/CustomVersion.h
+++ b/Engine/Source/Runtime/Core/Public/Serialization/CustomVersion.h
@@ -83,9 +83,6 @@ private:
 
 typedef TArray<FCustomVersion> FCustomVersionArray;
 
-UE_DEPRECATED(4.19, "FCustomVersionSet renamed to FCustomVersionArray")
-typedef FCustomVersionArray FCustomVersionSet;
-
 class CORE_API FCustomVersionRegistration;
 
 
