UE_STATIC_ASSERT_COMPLETE_TYPE macro added to cause a compile error if a type is incomplete.  This needs to be a macro because a trait would be memoized and could give incorrect behavior if used in compile-time logic like SFINAE for function overloading.

diff --git a/Engine/Source/Runtime/Core/Public/Misc/StaticAssertCompleteType.h b/Engine/Source/Runtime/Core/Public/Misc/StaticAssertCompleteType.h
new file mode 100644
--- /dev/null
+++ b/Engine/Source/Runtime/Core/Public/Misc/StaticAssertCompleteType.h
@@ -0,0 +1,48 @@
+// Copyright 1998-2019 Epic Games, Inc. All Rights Reserved.
+
+#pragma once
+
+namespace UE4StaticAssertCompleteType_Private
+{
+	class FIncompleteType;
+
+	template <typename T> struct TUEStaticAssertTypeChecker       { static T& Func(); };
+
+	// Voids will give an unhelpful error message when trying to take a reference to them, so work around it
+	template <> struct TUEStaticAssertTypeChecker<               void> { static FIncompleteType Func(); };
+	template <> struct TUEStaticAssertTypeChecker<const          void> { static FIncompleteType Func(); };
+	template <> struct TUEStaticAssertTypeChecker<volatile       void> { static FIncompleteType Func(); };
+	template <> struct TUEStaticAssertTypeChecker<const volatile void> { static FIncompleteType Func(); };
+
+	// References are always complete types, but as we're using sizeof to check and sizeof ignores
+	// references, let's just return a known complete type to make it work.
+	template <typename T> struct TUEStaticAssertTypeChecker<T&>  { static int Func(); };
+	template <typename T> struct TUEStaticAssertTypeChecker<T&&> { static int Func(); };
+
+	// Function types are complete types, but we can't form a reference to one, so let's just make those work too
+	template <typename RetType, typename... ArgTypes> struct TUEStaticAssertTypeChecker<RetType(ArgTypes...)> { static int Func(); };
+}
+
+// Causes a compile error if a type is incomplete
+#define UE_STATIC_ASSERT_COMPLETE_TYPE(TypeToCheck, Message) static_assert(sizeof(UE4StaticAssertCompleteType_Private::TUEStaticAssertTypeChecker<TypeToCheck>::Func()), Message)
+
+// Tests
+
+// Each of these should fail to compile
+#if 0
+	UE_STATIC_ASSERT_COMPLETE_TYPE(               void,                                     "CV void is incomplete");
+	UE_STATIC_ASSERT_COMPLETE_TYPE(      volatile void,                                     "CV void is incomplete");
+	UE_STATIC_ASSERT_COMPLETE_TYPE(const          void,                                     "CV void is incomplete");
+	UE_STATIC_ASSERT_COMPLETE_TYPE(const volatile void,                                     "CV void is incomplete");
+	UE_STATIC_ASSERT_COMPLETE_TYPE(UE4StaticAssertCompleteType_Private::FIncompleteType,    "A forward-declared-but-undefined class is incomplete");
+	UE_STATIC_ASSERT_COMPLETE_TYPE(UE4StaticAssertCompleteType_Private::FIncompleteType[2], "An array of an incomplete class is incomplete");
+	UE_STATIC_ASSERT_COMPLETE_TYPE(int[],                                                   "An array of a complete type of unspecified bound is incomplete");
+#endif
+
+// Each of these should pass
+#if 0
+	UE_STATIC_ASSERT_COMPLETE_TYPE(UE4StaticAssertCompleteType_Private::FIncompleteType*,                                                         "A pointer to an incomplete type is complete");
+	UE_STATIC_ASSERT_COMPLETE_TYPE(UE4StaticAssertCompleteType_Private::FIncompleteType&,                                                         "A reference to an incomplete type is complete");
+	UE_STATIC_ASSERT_COMPLETE_TYPE(UE4StaticAssertCompleteType_Private::FIncompleteType   (UE4StaticAssertCompleteType_Private::FIncompleteType), "A function type is not incomplete, even if it returns or takes an incomplete type");
+	UE_STATIC_ASSERT_COMPLETE_TYPE(UE4StaticAssertCompleteType_Private::FIncompleteType(&)(UE4StaticAssertCompleteType_Private::FIncompleteType), "References to function types must give a good error");
+#endif
