Static checks added to TPointerIsConvertibleFromTo to ensure that it's never instantiated with incomplete types, which would likely give the wrong answer.  Specializations have also been added to not require the static check when the test is always apparent, e.g. only changes to qualifiers or when one of the types is void.

diff --git a/Engine/Source/Runtime/Core/Public/Templates/PointerIsConvertibleFromTo.h b/Engine/Source/Runtime/Core/Public/Templates/PointerIsConvertibleFromTo.h
--- a/Engine/Source/Runtime/Core/Public/Templates/PointerIsConvertibleFromTo.h
+++ b/Engine/Source/Runtime/Core/Public/Templates/PointerIsConvertibleFromTo.h
@@ -3,19 +3,61 @@
 #pragma once
 
 #include "CoreTypes.h"
+#include "Misc/StaticAssertCompleteType.h"
+#include "Templates/RemoveCV.h"
+#include "Templates/LosesQualifiersFromTo.h"
+
+namespace UE4PointerIsConvertibleFromTo_Private
+{
+	template <typename From, typename To, typename NoCVFrom = typename TRemoveCV<From>::Type, typename NoCVTo = typename TRemoveCV<To>::Type>
+	struct TImpl
+	{
+	private:
+		static uint8  Test(...);
+		static uint16 Test(To*);
+
+		UE_STATIC_ASSERT_COMPLETE_TYPE(From, "TPointerIsConvertibleFromTo must not be instantiated with incomplete types");
+		UE_STATIC_ASSERT_COMPLETE_TYPE(To,   "TPointerIsConvertibleFromTo must not be instantiated with incomplete types");
+
+	public:
+		enum { Value = sizeof(Test((From*)nullptr)) - 1 };
+	};
+
+	template <typename From, typename To, typename NoCVFrom>
+	struct TImpl<From, To, NoCVFrom, NoCVFrom>
+	{
+		// cv T* to cv T* conversions are always allowed as long as no CVs are lost
+		enum { Value = !TLosesQualifiersFromTo<From, To>::Value };
+	};
+
+	template <typename From, typename To, typename NoCVFrom>
+	struct TImpl<From, To, NoCVFrom, void>
+	{
+		// cv T* to cv void* conversions are always allowed as long as no CVs are lost
+		enum { Value = !TLosesQualifiersFromTo<From, To>::Value };
+	};
+
+	template <typename From, typename To>
+	struct TImpl<From, To, void, void>
+	{
+		// cv void* to cv void* conversions are always allowed as long as no CVs are lost
+		enum { Value = !TLosesQualifiersFromTo<From, To>::Value };
+	};
+
+	template <typename From, typename To, typename NoCVTo>
+	struct TImpl<From, To, void, NoCVTo>
+	{
+		// cv void* to cv not_void* conversions are never legal
+		enum { Value = false };
+	};
+}
 
 /**
  * Tests if a From* is convertible to a To*
  **/
 template <typename From, typename To>
-struct TPointerIsConvertibleFromTo
+struct TPointerIsConvertibleFromTo : UE4PointerIsConvertibleFromTo_Private::TImpl<From, To>
 {
-private:
-	static uint8  Test(...);
-	static uint16 Test(To*);
-
-public:
-	enum { Value  = sizeof(Test((From*)nullptr)) - 1 };
 };
 
 
