Use default constructors for EngineVersion and EngineVersionBase.

diff --git a/Engine/Source/Runtime/Core/Private/Misc/EngineVersion.cpp b/Engine/Source/Runtime/Core/Private/Misc/EngineVersion.cpp
--- a/Engine/Source/Runtime/Core/Private/Misc/EngineVersion.cpp
+++ b/Engine/Source/Runtime/Core/Private/Misc/EngineVersion.cpp
@@ -8,14 +8,6 @@
 #include "UObject/ReleaseObjectVersion.h"
 #include "BuildSettings.h"
 
-FEngineVersionBase::FEngineVersionBase()
-: Major(0)
-, Minor(0)
-, Patch(0)
-, Changelist(0)
-{
-}
-
 FEngineVersionBase::FEngineVersionBase(uint16 InMajor, uint16 InMinor, uint16 InPatch, uint32 InChangelist)
 : Major(InMajor)
 , Minor(InMinor)
@@ -89,11 +81,6 @@ uint32 FEngineVersionBase::EncodeLicenseeChangelist(uint32 Changelist)
 }
 
 
-FEngineVersion::FEngineVersion()
-{
-	Empty();
-}
-
 FEngineVersion::FEngineVersion(uint16 InMajor, uint16 InMinor, uint16 InPatch, uint32 InChangelist, const FString &InBranch)
 {
 	Set(InMajor, InMinor, InPatch, InChangelist, InBranch);
diff --git a/Engine/Source/Runtime/Core/Public/Misc/EngineVersion.h b/Engine/Source/Runtime/Core/Public/Misc/EngineVersion.h
--- a/Engine/Source/Runtime/Core/Public/Misc/EngineVersion.h
+++ b/Engine/Source/Runtime/Core/Public/Misc/EngineVersion.h
@@ -13,7 +13,7 @@ class CORE_API FEngineVersion : public FEngineVersionBase
 public:
 
 	/** Empty constructor. Initializes the version to 0.0.0-0. */
-	FEngineVersion();
+	FEngineVersion() = default;
 
 	/** Constructs a version from the given components. */
 	FEngineVersion(uint16 InMajor, uint16 InMinor, uint16 InPatch, uint32 InChangelist, const FString &InBranch);
diff --git a/Engine/Source/Runtime/Core/Public/Misc/EngineVersionBase.h b/Engine/Source/Runtime/Core/Public/Misc/EngineVersionBase.h
--- a/Engine/Source/Runtime/Core/Public/Misc/EngineVersionBase.h
+++ b/Engine/Source/Runtime/Core/Public/Misc/EngineVersionBase.h
@@ -30,7 +30,7 @@ class CORE_API FEngineVersionBase
 public:
 
 	/** Empty constructor. Initializes the version to 0.0.0-0. */
-	FEngineVersionBase();
+	FEngineVersionBase() = default;
 
 	/** Constructs a version from the given components. */
 	FEngineVersionBase(uint16 InMajor, uint16 InMinor, uint16 InPatch = 0, uint32 InChangelist = 0);
@@ -74,14 +74,14 @@ public:
 protected:
 
 	/** Major version number. */
-	uint16 Major;
+	uint16 Major = 0;
 
 	/** Minor version number. */
-	uint16 Minor;
+	uint16 Minor = 0;
 
 	/** Patch version number. */
-	uint16 Patch;
+	uint16 Patch = 0;
 
 	/** Changelist number. This is used to arbitrate when Major/Minor/Patch version numbers match. Use GetChangelist() instead of using this member directly. */
-	uint32 Changelist;
+	uint32 Changelist = 0;
 };
