Use a fixed pimpl to avoid frequent allocations while serializing adapted types.

diff --git a/Engine/Source/Runtime/Core/Private/Serialization/StructuredArchiveFromArchive.cpp b/Engine/Source/Runtime/Core/Private/Serialization/StructuredArchiveFromArchive.cpp
--- a/Engine/Source/Runtime/Core/Private/Serialization/StructuredArchiveFromArchive.cpp
+++ b/Engine/Source/Runtime/Core/Private/Serialization/StructuredArchiveFromArchive.cpp
@@ -19,13 +19,19 @@ struct FStructuredArchiveFromArchive::FImpl
 };
 
 FStructuredArchiveFromArchive::FStructuredArchiveFromArchive(FArchive& Ar)
-	: Pimpl(Ar)
 {
+	static_assert(FStructuredArchiveFromArchive::ImplSize >= sizeof(FStructuredArchiveFromArchive::FImpl), "FStructuredArchiveFromArchive::ImplSize must fit in the size of FStructuredArchiveFromArchive::Impl");
+	static_assert(FStructuredArchiveFromArchive::ImplAlignment >= alignof(FStructuredArchiveFromArchive::FImpl), "FStructuredArchiveFromArchive::ImplAlignment must be compatible with the alignment of FStructuredArchiveFromArchive::Impl");
+
+	new (ImplStorage) FImpl(Ar);
 }
 
-FStructuredArchiveFromArchive::~FStructuredArchiveFromArchive() = default;
+FStructuredArchiveFromArchive::~FStructuredArchiveFromArchive()
+{
+	DestructItem((FImpl*)ImplStorage);
+}
 
 FStructuredArchive::FSlot FStructuredArchiveFromArchive::GetSlot()
 {
-	return Pimpl->Slot;
+	return ((FImpl*)ImplStorage)->Slot;
 }
diff --git a/Engine/Source/Runtime/Core/Public/Serialization/StructuredArchive.h b/Engine/Source/Runtime/Core/Public/Serialization/StructuredArchive.h
--- a/Engine/Source/Runtime/Core/Public/Serialization/StructuredArchive.h
+++ b/Engine/Source/Runtime/Core/Public/Serialization/StructuredArchive.h
@@ -623,6 +623,9 @@ class CORE_API FStructuredArchiveFromArchive
 {
 	UE_NONCOPYABLE(FStructuredArchiveFromArchive)
 
+	static constexpr uint32 ImplSize      = 400;
+	static constexpr uint32 ImplAlignment = 8;
+
 	struct FImpl;
 
 public:
@@ -632,8 +635,8 @@ public:
 	FStructuredArchiveSlot GetSlot();
 
 private:
-	// Implmented as a pimpl in order to reduce dependencies
-	TUniqueObj<FImpl> Pimpl;
+	// Implmented as a pimpl in order to reduce dependencies, but an inline one to avoid heap allocations
+	alignas(ImplAlignment) uint8 ImplStorage[ImplSize];
 };
 
 #if WITH_TEXT_ARCHIVE_SUPPORT
