Added FName batch serialization and aligned up IoStore chunks to 16B

The batch serialization is used by the new IoStore and AsyncLoading2.

This path loads almost 10M  FNames / CPU seconds for an internal project on my beefy workstation.

diff --git a/Engine/Source/Developer/IoStoreUtilities/Private/IoStoreUtilities.cpp b/Engine/Source/Developer/IoStoreUtilities/Private/IoStoreUtilities.cpp
--- a/Engine/Source/Developer/IoStoreUtilities/Private/IoStoreUtilities.cpp
+++ b/Engine/Source/Developer/IoStoreUtilities/Private/IoStoreUtilities.cpp
@@ -19,6 +19,7 @@
 #include "Serialization/JsonSerializer.h"
 #include "Serialization/BufferWriter.h"
 #include "Serialization/LargeMemoryWriter.h"
+#include "UObject/NameBatchSerialization.h"
 #include "UObject/PackageFileSummary.h"
 #include "UObject/ObjectResource.h"
 #include "UObject/Package.h"
@@ -92,24 +93,14 @@ public:
 		A << NameIndex << NameNumber;
 	}
 
-	int32 Num()
+
+	const TArray<FNameEntryId>& GetNameMap() const
 	{
-		return NameIndices.Num();
+		return NameMap;
 	}
 
-	int32 Save(FArchive& Ar, const FString& CsvFilePath)
-	{
-		int32 TotalSize = 0;
-		{
-			int32 NameCount = NameMap.Num();
-			Ar << NameCount;
-			for (int32 I = 0; I < NameCount; ++I)
-			{
-				FName::GetEntry(NameMap[I])->Write(Ar);
-			}
-			TotalSize = Ar.TotalSize();
-		}
 #if OUTPUT_NAMEMAP_CSV
+	void SaveCsv(const FString& CsvFilePath)
 		{
 			TUniquePtr<FArchive> CsvArchive(IFileManager::Get().CreateFileWriter(*CsvFilePath));
 			if (CsvArchive)
@@ -134,9 +125,8 @@ public:
 				}
 			}
 		}
-#endif
-		return TotalSize;
 	}
+#endif
 
 private:
 	TMap<FNameEntryId, int32> NameIndices;
@@ -1916,17 +1906,32 @@ int32 CreateTarget(const FContainerTarget& Target)
 			UE_LOG(LogIoStore, Error, TEXT("Failed to save initial load meta data to container file"));
 		}
 	}
-	
-	uint64 GlobalNameMapSize = 0;
+
+	uint64 GlobalNamesMB = 0;
+	uint64 GlobalNameHashesMB = 0;
 	{
 		UE_LOG(LogIoStore, Display, TEXT("Saving global name map to container file"));
-		FLargeMemoryWriter GlobalNameMapArchive(0, true);
-		GlobalNameMapSize = NameMapBuilder.Save(GlobalNameMapArchive, OutputDir / TEXT("Container.namemap.csv"));
-		const FIoStatus Status = AppendToContainer(*IoStoreWriter, GlobalNameMapArchive, CreateIoChunkId(0, 0, EIoChunkType::LoaderGlobalNameMap));
-		if (!Status.IsOk())
+
+		TArray<uint8> Names;
+		TArray<uint8> Hashes;
+		SaveNameBatch(NameMapBuilder.GetNameMap(), /* out */ Names, /* out */ Hashes);
+
+		GlobalNamesMB = Names.Num() >> 20;
+		GlobalNameHashesMB = Hashes.Num() >> 20;
+
+		FIoStatus NameStatus = IoStoreWriter->Append(CreateIoChunkId(0, 0, EIoChunkType::LoaderGlobalNames), 
+													 FIoBuffer(FIoBuffer::Wrap, Names.GetData(), Names.Num()));
+		FIoStatus HashStatus = IoStoreWriter->Append(CreateIoChunkId(0, 0, EIoChunkType::LoaderGlobalNameHashes),
+													 FIoBuffer(FIoBuffer::Wrap, Hashes.GetData(), Hashes.Num()));
+		
+		if (!NameStatus.IsOk() || !HashStatus.IsOk())
 		{
 			UE_LOG(LogIoStore, Error, TEXT("Failed to save global name map to container file"));
 		}
+
+#if OUTPUT_NAMEMAP_CSV
+		NameMapBuilder.SaveCsv(OutputDir / TEXT("Container.namemap.csv"));
+#endif
 	}
 
 	UE_LOG(LogIoStore, Display, TEXT("Calculating stats..."));
@@ -2012,7 +2017,8 @@ int32 CreateTarget(const FContainerTarget& Target)
 	UE_LOG(LogIoStore, Display, TEXT("Pak: %12.2f MB PreloadDependencies, %d imports + %d exports = %d"),
 		(double)PreloadDependenciesSize / 1024.0 / 1024.0, ImportPreloadCount, ExportPreloadCount, PreloadDependencies.Num());
 
-	UE_LOG(LogIoStore, Display, TEXT("IoStore: %8.2f MB GlobalNameMap, %d unique names"), (double)GlobalNameMapSize / 1024.0 / 1024.0, NameMapBuilder.Num());
+	UE_LOG(LogIoStore, Display, TEXT("IoStore: %8.2f MB GlobalNames, %d unique names"), (double)GlobalNamesMB, NameMapBuilder.GetNameMap().Num());
+	UE_LOG(LogIoStore, Display, TEXT("IoStore: %8.2f MB GlobalNameHashes"), (double)GlobalNameHashesMB);
 	UE_LOG(LogIoStore, Display, TEXT("IoStore: %8.2f MB GlobalPackageData"), (double)StoreTocByteCount / 1024.0 / 1024.0);
 	UE_LOG(LogIoStore, Display, TEXT("IoStore: %8.2f MB GlobalImportedPackages, %d imported packages"), (double)ImportedPackagesByteCount / 1024.0 / 1024.0, ImportedPackagesCount);
 	UE_LOG(LogIoStore, Display, TEXT("IoStore: %8.2f MB GlobalBundleOrders, %d bundles"), (double)ExportBundleOrderByteCount / 1024.0 / 1024.0, ExportBundleOrderEntries.Num());
diff --git a/Engine/Source/Runtime/Core/Private/IO/IoDispatcher.cpp b/Engine/Source/Runtime/Core/Private/IO/IoDispatcher.cpp
--- a/Engine/Source/Runtime/Core/Private/IO/IoDispatcher.cpp
+++ b/Engine/Source/Runtime/Core/Private/IO/IoDispatcher.cpp
@@ -676,8 +676,14 @@ FIoBatch::Issue()
 void 
 FIoBatch::Wait()
 {
-	//TODO: Create synchronization event here when it's actually needed
-	unimplemented();
+	// Temporary spinning implementation
+	for (FIoRequestImpl* Request = Impl->FirstRequest; Request; Request = Request->NextRequest)
+	{
+		if (!Request->Result.Status().IsCompleted())
+		{
+			FPlatformProcess::Sleep(0);
+		}
+	}
 }
 
 void 
diff --git a/Engine/Source/Runtime/Core/Private/IO/IoStore.cpp b/Engine/Source/Runtime/Core/Private/IO/IoStore.cpp
--- a/Engine/Source/Runtime/Core/Private/IO/IoStore.cpp
+++ b/Engine/Source/Runtime/Core/Private/IO/IoStore.cpp
@@ -8,6 +8,7 @@
 //////////////////////////////////////////////////////////////////////////
 
 constexpr char FIoStoreTocHeader::TocMagicImg[];
+constexpr uint32 IoChunkAlignment = 16;
 
 //////////////////////////////////////////////////////////////////////////
 
@@ -88,13 +89,22 @@ public:
 
 		FIoStoreTocEntry TocEntry;
 
+		check(ContainerFileHandle->Tell() % IoChunkAlignment == 0);
+
 		TocEntry.SetOffset(ContainerFileHandle->Tell());
 		TocEntry.SetLength(Chunk.DataSize());
 		TocEntry.ChunkId = ChunkId;
 
 		IsMetadataDirty = true;
 
-		const bool Success = ContainerFileHandle->Write(Chunk.Data(), Chunk.DataSize());
+		bool Success = ContainerFileHandle->Write(Chunk.Data(), Chunk.DataSize());
+
+		if (uint32 UnpaddedBytes = Chunk.DataSize() % IoChunkAlignment)
+		{
+			static constexpr uint8 Zeroes[IoChunkAlignment] = {};
+			uint32 Padding = (IoChunkAlignment - UnpaddedBytes) % IoChunkAlignment;
+			Success &= ContainerFileHandle->Write(Zeroes, Padding);
+		}
 
 		if (Success)
 		{
diff --git a/Engine/Source/Runtime/Core/Private/UObject/UnrealNames.cpp b/Engine/Source/Runtime/Core/Private/UObject/UnrealNames.cpp
--- a/Engine/Source/Runtime/Core/Private/UObject/UnrealNames.cpp
+++ b/Engine/Source/Runtime/Core/Private/UObject/UnrealNames.cpp
@@ -1,6 +1,7 @@
 // Copyright 1998-2019 Epic Games, Inc. All Rights Reserved.
 
 #include "UObject/UnrealNames.h"
+#include "UObject/NameBatchSerialization.h"
 #include "Misc/AssertionMacros.h"
 #include "Misc/MessageDialog.h"
 #include "Math/NumericLimits.h"
@@ -24,6 +25,7 @@
 #include "HAL/IConsoleManager.h"
 #include "HAL/LowLevelMemTracker.h"
 #include "Hash/CityHash.h"
+#include "Templates/AlignmentTemplates.h"
 
 // Page protection to catch FNameEntry stomps
 #ifndef FNAME_WRITE_PROTECT_PAGES
@@ -277,18 +279,30 @@ public:
 		}
 	}
 
+	void ReserveBlocks(uint32 Num)
+	{
+		FWriteScopeLock _(Lock);
+
+		for (uint32 Idx = Num - 1; Idx > CurrentBlock && Blocks[Idx] == nullptr; --Idx)
+		{
+			Blocks[Idx] = AllocBlock();
+		}
+	}
+
+
 	/**
 	 * Allocates the requested amount of bytes and returns an id that can be used to access them
 	 *
 	 * @param   Size  Size in bytes to allocate, 
 	 * @return  Allocation of passed in size cast to a FNameEntry pointer.
 	 */
+	template <class ScopeLock>
 	FNameEntryHandle Allocate(uint32 Bytes)
 	{
 		Bytes = Align(Bytes, alignof(FNameEntry));
 		check(Bytes <= BlockSizeBytes);
 
-		FRWScopeLock _(Lock, FRWScopeLockType::SLT_Write);
+		ScopeLock _(Lock);
 
 		// Allocate a new pool if current one is exhausted. We don't worry about a little bit
 		// of waste at the end given the relative size of pool to average and max allocation.
@@ -306,9 +320,10 @@ public:
 		return FNameEntryHandle(CurrentBlock, ByteOffset / Stride);
 	}
 
+	template<class ScopeLock>
 	FNameEntryHandle Create(FNameStringView Name, TOptional<FNameEntryId> ComparisonId, FNameEntryHeader Header)
 	{
-		FNameEntryHandle Handle = Allocate(FNameEntry::GetDataOffset() + Name.BytesWithoutTerminator());
+		FNameEntryHandle Handle = Allocate<ScopeLock>(FNameEntry::GetDataOffset() + Name.BytesWithoutTerminator());
 		FNameEntry& Entry = Resolve(Handle);
 
 #if WITH_CASE_PRESERVING_NAME
@@ -335,6 +350,16 @@ public:
 		return *reinterpret_cast<FNameEntry*>(Blocks[Handle.Block] + Stride * Handle.Offset);
 	}
 
+	void BatchLock() const
+	{
+		Lock.WriteLock();
+	}
+
+	void BatchUnlock() const
+	{
+		Lock.WriteUnlock();
+	}
+
 	/** Returns the number of blocks that have been allocated so far for names. */
 	uint32 NumBlocks() const
 	{
@@ -374,7 +399,11 @@ private:
 		}
 	}
 
-	/** Allocates a new pool. */
+	static uint8* AllocBlock()
+	{
+		return (uint8*)FMemory::Malloc(BlockSizeBytes, FPlatformMemory::GetConstants().PageSize);
+	}
+	
 	void AllocateNewBlock()
 	{
 		LLM_SCOPE(ELLMTag::FName);
@@ -392,9 +421,12 @@ private:
 		CurrentByteCursor = 0;
 
 		check(CurrentBlock < FNameMaxBlocks);
-		check(Blocks[CurrentBlock] == nullptr);
 
-		Blocks[CurrentBlock] = (uint8*)FMemory::Malloc(BlockSizeBytes, FPlatformMemory::GetConstants().PageSize);
+		// Allocate block unless it's already reserved
+		if (Blocks[CurrentBlock] == nullptr)
+		{
+			Blocks[CurrentBlock] = AllocBlock();
+		}
 	}
 
 	mutable FRWLock Lock;
@@ -419,7 +451,7 @@ enum { FNamePoolInitialSlotsPerShard = 1 << FNamePoolInitialSlotBits };
 
 /** Hashes name into 64 bits that determines shard and slot index.
  *	
- *	Small parts	of the hash is also stored in unused bits of the slot and entry. 
+ *	A small part of the hash is also stored in unused bits of the slot and entry. 
  *	The former optimizes linear probing by accessing less entry data.
  *	The latter optimizes linear probing by avoiding copying and deobfuscating entry data.
  *
@@ -435,11 +467,25 @@ struct FNameHash
 	uint32 SlotProbeHash; // Helps cull equality checks (decode + strnicmp) when probing slots
 	FNameEntryHeader EntryProbeHeader; // Helps cull equality checks when probing inspects entries
 
+	static constexpr uint64 AlgorithmId = 0xC1640000;
+
+	template<class CharType>
+	static uint64 GenerateHash(const CharType* Str, int32 Len)
+	{
+		return CityHash64(reinterpret_cast<const char*>(Str), Len * sizeof(CharType));
+	}
+
+	template<class CharType>
+	static uint64 GenerateLowerCaseHash(const CharType* Str, uint32 Len);
 
 	template<class CharType>
 	FNameHash(const CharType* Str, int32 Len)
+		: FNameHash(Str, Len, GenerateHash(Str, Len))
+	{}
+
+	template<class CharType>
+	FNameHash(const CharType* Str, int32 Len, uint64 Hash)
 	{
-		uint64 Hash = CityHash64(reinterpret_cast<const char*>(Str), Len * sizeof(CharType));
 		uint32 Hi = static_cast<uint32>(Hash >> 32);
 		uint32 Lo = static_cast<uint32>(Hash);
 
@@ -498,8 +544,27 @@ struct FNameHash
 		uint32 FourChars = FPlatformMemory::ReadUnaligned<uint32>(Str);
 		return (FourChars & ToUpperMask) == NoneAsInt;
 	}
+
+	bool operator==(const FNameHash& Rhs) const
+	{
+		return  ShardIndex == Rhs.ShardIndex &&
+				UnmaskedSlotIndex == Rhs.UnmaskedSlotIndex &&
+				SlotProbeHash == Rhs.SlotProbeHash &&
+				EntryProbeHeader == Rhs.EntryProbeHeader;
+	}
 };
 
+template<class CharType>
+FORCENOINLINE uint64 FNameHash::GenerateLowerCaseHash(const CharType* Str, uint32 Len)
+{
+	CharType LowerStr[NAME_SIZE];
+	for (uint32 I = 0; I < Len; ++I)
+	{
+		LowerStr[I] = TChar<CharType>::ToLower(Str[I]);
+	}
+
+	return FNameHash::GenerateHash(LowerStr, Len);
+}
 
 template<class CharType>
 FORCENOINLINE FNameHash HashLowerCase(const CharType* Str, uint32 Len)
@@ -529,11 +594,16 @@ FNameHash HashName<ENameCase::CaseSensitive>(FNameStringView Name)
 template<ENameCase Sensitivity>
 struct FNameValue
 {
-	FNameValue(FNameStringView InName)
+	explicit FNameValue(FNameStringView InName)
 		: Name(InName)
 		, Hash(HashName<Sensitivity>(InName))
 	{}
 
+	FNameValue(FNameStringView InName, FNameHash InHash)
+		: Name(InName)
+		, Hash(InHash)
+	{}
+
 	FNameStringView Name;
 	FNameHash Hash;
 	TOptional<FNameEntryId> ComparisonId;
@@ -544,6 +614,11 @@ using FNameComparisonValue = FNameValue<ENameCase::IgnoreCase>;
 using FNameDisplayValue = FNameValue<ENameCase::CaseSensitive>;
 #endif
 
+// For prelocked batch insertions
+struct FNullScopeLock
+{
+	FNullScopeLock(FRWLock&) {}
+};
 
 class alignas(PLATFORM_CACHE_LINE_SIZE) FNamePoolShardBase : FNoncopyable
 {
@@ -566,9 +641,18 @@ public:
 		UsedSlots = 0;
 		CapacityMask = 0;
 		Slots = nullptr;
+		NumCreatedEntries = 0;
+		NumCreatedWideEntries = 0;
 	}
 
-	uint32 Capacity() const { return CapacityMask + 1; }
+	uint32 Capacity() const	{ return CapacityMask + 1; }
+
+	uint32 NumCreated() const { return NumCreatedEntries; }
+	uint32 NumCreatedWide() const { return NumCreatedWideEntries; }
+
+	// Used for batch insertion together with Insert<FNullScopeLock>()
+	void BatchLock() const	 { Lock.WriteLock(); }
+	void BatchUnlock() const { Lock.WriteUnlock(); }
 
 protected:
 	enum { LoadFactorQuotient = 9, LoadFactorDivisor = 10 }; // I.e. realloc slots when 90% full
@@ -578,6 +662,8 @@ protected:
 	uint32 CapacityMask = 0;
 	FNameSlot* Slots = nullptr;
 	FNameEntryAllocator* Entries = nullptr;
+	uint32 NumCreatedEntries = 0;
+	uint32 NumCreatedWideEntries = 0;
 
 
 	template<ENameCase Sensitivity>
@@ -598,10 +684,10 @@ public:
 		return Probe(Value).GetId();
 	}
 
-	FNameEntryId Insert(const FNameValue<Sensitivity>& Value, bool& bCreatedNewEntry)
+	template<class ScopeLock = FWriteScopeLock>
+	FORCEINLINE FNameEntryId Insert(const FNameValue<Sensitivity>& Value, bool& bCreatedNewEntry)
 	{
-		FRWScopeLock _(Lock, FRWScopeLockType::SLT_Write);
-
+		ScopeLock _(Lock);
 		FNameSlot& Slot = Probe(Value);
 
 		if (Slot.Used())
@@ -609,10 +695,12 @@ public:
 			return Slot.GetId();
 		}
 
-		FNameEntryId NewEntryId = Entries->Create(Value.Name, Value.ComparisonId, Value.Hash.EntryProbeHeader);
+		FNameEntryId NewEntryId = Entries->Create<ScopeLock>(Value.Name, Value.ComparisonId, Value.Hash.EntryProbeHeader);
 
 		ClaimSlot(Slot, FNameSlot(NewEntryId, Value.Hash.SlotProbeHash));
 
+		++NumCreatedEntries;
+		NumCreatedWideEntries += Value.Name.bIsWide;
 		bCreatedNewEntry = true;
 
 		return NewEntryId;
@@ -631,6 +719,16 @@ public:
 		}
 	}
 
+	void Reserve(uint32 Num)
+	{
+		uint32 WantedCapacity = FMath::RoundUpToPowerOfTwo(Num * LoadFactorDivisor / LoadFactorQuotient);
+
+		if (WantedCapacity > Capacity())
+		{
+			Grow(WantedCapacity);
+		}
+	}
+
 private:
 	void ClaimSlot(FNameSlot& UnusedSlot, FNameSlot NewValue)
 	{
@@ -644,12 +742,16 @@ private:
 	}
 
 	void Grow()
+	{
+		Grow(Capacity() * 2);
+	}
+
+	void Grow(const uint32 NewCapacity)
 	{
 		LLM_SCOPE(ELLMTag::FName);
 		FNameSlot* const OldSlots = Slots;
 		const uint32 OldUsedSlots = UsedSlots;
 		const uint32 OldCapacity = Capacity();
-		const uint32 NewCapacity = OldCapacity * 2;
 
 		Slots = (FNameSlot*)FMemory::Malloc(NewCapacity * sizeof(FNameSlot), alignof(FNameSlot));
 		memset(Slots, 0, NewCapacity * sizeof(FNameSlot));
@@ -675,7 +777,7 @@ private:
 	}
 
 	/** Find slot containing value or the first free slot that should be used to store it  */
-	FNameSlot& Probe(const FNameValue<Sensitivity>& Value) const
+	FORCEINLINE FNameSlot& Probe(const FNameValue<Sensitivity>& Value) const
 	{
 		return Probe(Value.Hash.UnmaskedSlotIndex, 
 			[&](FNameSlot Slot)	{ return Slot.GetProbeHash() == Value.Hash.SlotProbeHash && 
@@ -684,7 +786,7 @@ private:
 
 	/** Find slot that fulfills predicate or the first free slot  */
 	template<class PredicateFn>
-	FNameSlot& Probe(uint32 UnmaskedSlotIndex, PredicateFn Predicate) const
+	FORCEINLINE FNameSlot& Probe(uint32 UnmaskedSlotIndex, PredicateFn Predicate) const
 	{
 		const uint32 Mask = CapacityMask;
 		for (uint32 I = FNameHash::GetProbeStart(UnmaskedSlotIndex, Mask); true; I = (I + 1) & Mask)
@@ -710,7 +812,8 @@ class FNamePool
 {
 public:
 	FNamePool();
-
+	
+	void			Reserve(uint32 NumBlocks, uint32 NumEntries);
 	FNameEntryId	Store(FNameStringView View);
 	FNameEntryId	Find(FNameStringView View) const;
 	FNameEntryId	Find(EName Ename) const;
@@ -720,11 +823,16 @@ public:
 	FNameEntry&		Resolve(FNameEntryHandle Handle) const { return Entries.Resolve(Handle); }
 
 	bool			IsValid(FNameEntryHandle Handle) const;
-		
+
+	void			BatchLock();
+	FNameEntryId	BatchStore(const FNameComparisonValue& ComparisonValue);
+	void			BatchUnlock();
+
 	/// Stats and debug related functions ///
 
-	uint32			NumAnsiEntries() const { return AnsiCount; }
-	uint32			NumWideEntries() const { return WideCount; }
+	uint32			NumEntries() const;
+	uint32			NumAnsiEntries() const;
+	uint32			NumWideEntries() const;
 	uint32			NumBlocks() const { return Entries.NumBlocks(); }
 	uint32			NumSlots() const;
 	void			LogStats(FOutputDevice& Ar) const;
@@ -735,8 +843,6 @@ private:
 	enum { MaxENames = 512 };
 
 	FNameEntryAllocator Entries;
-	TAtomic<uint32> AnsiCount;
-	TAtomic<uint32> WideCount;
 
 #if WITH_CASE_PRESERVING_NAME
 	FNamePoolShard<ENameCase::CaseSensitive> DisplayShards[FNamePoolShards];
@@ -750,8 +856,6 @@ private:
 };
 
 FNamePool::FNamePool()
-	: AnsiCount(0)
-	, WideCount(0)
 {
 	for (FNamePoolShardBase& Shard : ComparisonShards)
 	{
@@ -840,13 +944,11 @@ FNameEntryId FNamePool::Store(FNameStringView Name)
 	}
 #endif
 
-	TAtomic<uint32>& EntryCount = Name.IsAnsi() ? AnsiCount : WideCount;
 	bool bAdded = false;
 
 	// Insert comparison name first since display value must contain comparison name
 	FNameComparisonValue ComparisonValue(Name);
 	FNameEntryId ComparisonId = ComparisonShards[ComparisonValue.Hash.ShardIndex].Insert(ComparisonValue, bAdded);
-	EntryCount += bAdded;
 
 #if WITH_CASE_PRESERVING_NAME
 	// Check if ComparisonId can be used as DisplayId
@@ -857,18 +959,80 @@ FNameEntryId FNamePool::Store(FNameStringView Name)
 	}
 	else
 	{
-		bAdded = false;
 		DisplayValue.ComparisonId = ComparisonId;
-		FNameEntryId DisplayId = DisplayShard.Insert(DisplayValue, bAdded);
-		EntryCount += bAdded;
-
-		return DisplayId;
+		return DisplayShard.Insert(DisplayValue, bAdded);
 	}
 #else
 	return ComparisonId;
 #endif
 }
 
+void FNamePool::BatchLock()
+{
+	for (const FNamePoolShardBase& Shard : ComparisonShards)
+	{
+		Shard.BatchLock();
+	}
+
+	// Acquire entry allocator lock after shard locks
+	Entries.BatchLock();
+}
+
+FORCEINLINE FNameEntryId FNamePool::BatchStore(const FNameComparisonValue& ComparisonValue)
+{
+	bool bCreatedNewEntry;
+	return ComparisonShards[ComparisonValue.Hash.ShardIndex].Insert<FNullScopeLock>(ComparisonValue, bCreatedNewEntry);
+}
+
+void FNamePool::BatchUnlock()
+{
+	Entries.BatchUnlock();
+
+	for (int32 Idx = FNamePoolShards - 1; Idx >= 0; --Idx)
+	{
+		ComparisonShards[Idx].BatchUnlock();
+	}
+}
+
+uint32 FNamePool::NumEntries() const
+{
+	uint32 Out = 0;
+#if WITH_CASE_PRESERVING_NAME
+	for (const FNamePoolShardBase& Shard : DisplayShards)
+	{
+		Out += Shard.NumCreated();
+	}
+#endif
+	for (const FNamePoolShardBase& Shard : ComparisonShards)
+	{
+		Out += Shard.NumCreated();
+	}
+
+	return Out;
+}
+
+uint32 FNamePool::NumAnsiEntries() const
+{
+	return NumEntries() - NumWideEntries();
+}
+
+uint32 FNamePool::NumWideEntries() const
+{
+	uint32 Out = 0;
+#if WITH_CASE_PRESERVING_NAME
+	for (const FNamePoolShardBase& Shard : DisplayShards)
+	{
+		Out += Shard.NumCreatedWide();
+	}
+#endif
+	for (const FNamePoolShardBase& Shard : ComparisonShards)
+	{
+		Out += Shard.NumCreatedWide();
+	}
+
+	return Out;
+}
+
 uint32 FNamePool::NumSlots() const
 {
 	uint32 SlotCapacity = 0;
@@ -888,13 +1052,13 @@ uint32 FNamePool::NumSlots() const
 
 void FNamePool::LogStats(FOutputDevice& Ar) const
 {
-	Ar.Logf(TEXT("%i FNames using in %ikB + %ikB"), WideCount + AnsiCount, sizeof(FNamePool), Entries.NumBlocks() * FNameEntryAllocator::BlockSizeBytes / 1024);
+	Ar.Logf(TEXT("%i FNames using in %ikB + %ikB"), NumEntries(), sizeof(FNamePool), Entries.NumBlocks() * FNameEntryAllocator::BlockSizeBytes / 1024);
 }
 
 TArray<const FNameEntry*> FNamePool::DebugDump() const
 {
 	TArray<const FNameEntry*> Out;
-	Out.Reserve(WideCount + AnsiCount);
+	Out.Reserve(NumEntries());
 	Entries.DebugDump(Out);
 	return Out;
 }
@@ -909,6 +1073,28 @@ const EName* FNamePool::FindEName(FNameEntryId Id) const
 	return Id.ToUnstableInt() > LargestEnameUnstableId ? nullptr : EntryToEName.Find(Id);
 }
 
+void FNamePool::Reserve(uint32 NumBytes, uint32 InNumEntries)
+{
+	uint32 NumBlocks = NumBytes / FNameEntryAllocator::BlockSizeBytes + 1;
+	Entries.ReserveBlocks(NumBlocks);
+
+	if (NumEntries() < InNumEntries)
+	{
+		uint32 NumEntriesPerShard = InNumEntries / FNamePoolShards + 1;
+
+	#if WITH_CASE_PRESERVING_NAME
+		for (FNamePoolShard<ENameCase::CaseSensitive>& Shard : DisplayShards)
+		{
+			Shard.Reserve(NumEntriesPerShard);
+		}
+	#endif
+		for (FNamePoolShard<ENameCase::IgnoreCase>& Shard : ComparisonShards)
+		{
+			Shard.Reserve(NumEntriesPerShard);
+		}
+	}
+}
+
 static bool bNamePoolInitialized;
 alignas(FNamePool) static uint8 NamePoolData[sizeof(FNamePool)];
 
@@ -1972,6 +2158,8 @@ void CheckLazyName(const CharType(&Literal)[N])
 	check(FLazyName(Literal) == FLazyName(Literal2));
 }
 
+static void TestNameBatch();
+
 void FName::AutoTest()
 {
 #if DO_CHECK
@@ -2159,6 +2347,7 @@ void FName::AutoTest()
 	CheckLazyName("None_0");
 	CheckLazyName("None_1");
 
+	TestNameBatch();
 
 #if 0
 	// Check hash table growth still yields the same unique FName ids
@@ -2362,6 +2551,462 @@ bool operator==(const FLazyName& A, const FLazyName& B)
 	}
 }
 
+/*-----------------------------------------------------------------------------
+	FName batch serialization
+-----------------------------------------------------------------------------*/
+
+/**
+ * FNameStringView sibling with UTF16 Little-Endian wide strings instead of WIDECHAR 
+ *
+ * View into serialized data instead of how it will be stored in memory once loaded.
+ */
+struct FNameSerializedView
+{
+	FNameSerializedView(const ANSICHAR* InStr, uint32 InLen)
+	: Ansi(InStr)
+	, Len(InLen)
+	, bIsUtf16(false)
+	{}
+	
+	FNameSerializedView(const UTF16CHAR* InStr, uint32 InLen)
+	: Utf16(InStr)
+	, Len(InLen)
+	, bIsUtf16(true)
+	{}
+
+	FNameSerializedView(const uint8* InData, uint32 InLen, bool bInUtf16)
+	: Data(InData)
+	, Len(InLen)
+	, bIsUtf16(bInUtf16)
+	{}
+
+	union
+	{
+		const uint8* Data;
+		const ANSICHAR* Ansi;
+		const UTF16CHAR* Utf16;
+	};
+
+	uint32 Len;
+	bool bIsUtf16;
+};
+
+static uint8* AddUninitializedBytes(TArray<uint8>& Out, uint32 Bytes)
+{
+	uint32 OldNum = Out.AddUninitialized(Bytes);
+	return Out.GetData() + OldNum;
+}
+
+template<typename T>
+static T* AddUninitializedElements(TArray<uint8>& Out, uint32 Num)
+{
+	check(Out.Num() %  alignof(T) == 0);
+	return reinterpret_cast<T*>(AddUninitializedBytes(Out, Num * sizeof(T)));
+}
+
+template<typename T>
+static void AddValue(TArray<uint8>& Out, T Value)
+{
+	*AddUninitializedElements<T>(Out, 1) = Value;
+}
+
+template<typename T>
+static void AlignTo(TArray<uint8>& Out)
+{
+	if (uint32 UnpaddedBytes = Out.Num() % sizeof(T))
+	{
+		Out.AddZeroed(sizeof(T) - UnpaddedBytes);
+	}
+}
+
+static uint32 GetRequiredUtf16Padding(const uint8* Ptr)
+{
+	return UPTRINT(Ptr) & 1u;
+}
+
+struct FSerializedNameHeader
+{
+	FSerializedNameHeader(uint32 Len, bool bIsUtf16)
+	{
+		static_assert(NAME_SIZE < 0x8000u, "");
+		check(Len <= NAME_SIZE);
+
+		Data[0] = uint8(bIsUtf16) << 7 | static_cast<uint8>(Len >> 8);
+		Data[1] = static_cast<uint8>(Len);
+	}
+
+	uint8 IsUtf16() const
+	{
+		return Data[0] & 0x80u;
+	}
+
+	uint32 Len() const
+	{
+		return ((Data[0] & 0x7Fu) << 8) + Data[1];
+	}
+
+	uint8 Data[2];
+};
+
+FNameSerializedView LoadNameHeader(const uint8*& InOutIt)
+{
+	const FSerializedNameHeader& Header = *reinterpret_cast<const FSerializedNameHeader*>(InOutIt);
+	const uint8* NameData = InOutIt + sizeof(FSerializedNameHeader);
+	const uint32 Len = Header.Len();
+
+	if (Header.IsUtf16())
+	{
+		NameData += GetRequiredUtf16Padding(NameData);
+		InOutIt = NameData + Len * sizeof(UTF16CHAR);
+		return FNameSerializedView(NameData, Len, /* UTF16 */ true);
+	}
+	else
+	{
+		InOutIt = NameData + Len * sizeof(ANSICHAR);
+		return FNameSerializedView(NameData, Len, /* UTF16 */ false);
+	}
+}
+
+#if ALLOW_NAME_BATCH_SAVING
+
+static FNameSerializedView SaveAnsiName(TArray<uint8>& Out, const ANSICHAR* Src, uint32 Len)
+{
+	ANSICHAR* Dst = AddUninitializedElements<ANSICHAR>(Out, Len);
+	FMemory::Memcpy(Dst, Src, Len * sizeof(ANSICHAR));
+
+	return FNameSerializedView(Dst, Len);
+}
+
+static FNameSerializedView SaveUtf16Name(TArray<uint8>& Out, const WIDECHAR* Src, uint32 Len)
+{
+	// Align to UTF16CHAR after header
+	AlignTo<UTF16CHAR>(Out);
+	
+#if !PLATFORM_LITTLE_ENDIAN
+	#error TODO: Implement saving code units as Little-Endian on Big-Endian platforms
+#endif
+
+	// This is a no-op when sizeof(UTF16CHAR) == sizeof(WIDECHAR), which it usually is
+	FTCHARToUTF16 Utf16String(Src, Len);
+
+	UTF16CHAR* Dst = AddUninitializedElements<UTF16CHAR>(Out, Utf16String.Length());
+	FMemory::Memcpy(Dst, Utf16String.Get(), Utf16String.Length() * sizeof(UTF16CHAR));
+
+	return FNameSerializedView(Dst, Len);
+}
+
+static FNameSerializedView SaveAnsiOrUtf16Name(TArray<uint8>& Out, FNameStringView Name)
+{
+	void* HeaderData = AddUninitializedBytes(Out, sizeof(FSerializedNameHeader));
+	new (HeaderData) FSerializedNameHeader(Name.Len, Name.bIsWide);
+
+	if (Name.bIsWide)
+	{
+		return SaveUtf16Name(Out, Name.Wide, Name.Len);
+	}
+	else
+	{
+		return SaveAnsiName(Out, Name.Ansi, Name.Len);
+	}
+}
+
+void SaveNameBatch(TArrayView<const FNameEntryId> Names, TArray<uint8>& OutNameData, TArray<uint8>& OutHashData)
+{
+	OutNameData.Empty(/* average bytes per name guesstimate */ 40 * Names.Num());
+	OutHashData.Empty((/* hash version */ 1 + Names.Num()) * sizeof(uint64));
+
+	// Save hash algorithm version
+	AddValue(OutHashData, INTEL_ORDER64(FNameHash::AlgorithmId));
+
+	// Save names and hashes
+	FNameBuffer CustomDecodeBuffer;
+	for (FNameEntryId EntryId : Names)
+	{
+		FNameStringView InMemoryName = GetNamePoolPostInit().Resolve(EntryId).MakeView(CustomDecodeBuffer);
+		FNameSerializedView SavedName = SaveAnsiOrUtf16Name(OutNameData, InMemoryName);
+
+		uint64 LowerHash = SavedName.bIsUtf16 ? FNameHash::GenerateLowerCaseHash(SavedName.Utf16, SavedName.Len)
+											  : FNameHash::GenerateLowerCaseHash(SavedName.Ansi,  SavedName.Len);
+
+		AddValue(OutHashData, INTEL_ORDER64(LowerHash));
+	}
+}
+
+#endif // WITH_EDITOR
+
+FORCENOINLINE void ReserveNameBatch(uint32 NameDataBytes, uint32 HashDataBytes)
+{
+	uint32 NumEntries = HashDataBytes / sizeof(uint64) - 1;
+	// Add 20% slack to reduce probing costs
+	auto AddSlack = [](uint64 In){ return static_cast<uint32>(In * 6 / 5); };
+	GetNamePoolPostInit().Reserve(AddSlack(NameDataBytes), AddSlack(NumEntries));
+}
+
+static FNameEntryId BatchLoadNameWithoutHash(const UTF16CHAR* Str, uint32 Len)
+{
+	WIDECHAR Temp[NAME_SIZE];
+	for (uint32 Idx = 0; Idx < Len; ++Idx)
+	{
+		Temp[Idx] = INTEL_ORDER16(Str[Idx]);
+	}
+	
+#if PLATFORM_TCHAR_IS_4_BYTES
+	// Inline combine any surrogate pairs in the data when loading into a UTF-32 string
+	Len = StringConv::InlineCombineSurrogates_Buffer(Temp, Len);
+#endif
+
+	FNameStringView Name(Temp, Len);
+	FNameHash Hash = HashName<ENameCase::IgnoreCase>(Name);
+	return GetNamePoolPostInit().BatchStore(FNameComparisonValue(Name, Hash));
+}
+
+static FNameEntryId BatchLoadNameWithoutHash(const ANSICHAR* Str, uint32 Len)
+{
+	FNameStringView Name(Str, Len);
+	FNameHash Hash = HashName<ENameCase::IgnoreCase>(Name);
+	return GetNamePoolPostInit().BatchStore(FNameComparisonValue(Name, Hash));
+}
+
+static FNameEntryId BatchLoadNameWithoutHash(const FNameSerializedView& Name)
+{
+	return Name.bIsUtf16 ? BatchLoadNameWithoutHash(Name.Utf16, Name.Len)
+						 : BatchLoadNameWithoutHash(Name.Ansi, Name.Len);
+}
+
+template<typename CharType>
+FNameEntryId BatchLoadNameWithHash(const CharType* Str, uint32 Len, uint64 InHash)
+{
+	FNameStringView Name(Str, Len);
+	FNameHash Hash(Str, Len, InHash);
+	checkfSlow(Hash == HashName<ENameCase::IgnoreCase>(Name), TEXT("Precalculated hash was wrong"));
+	return GetNamePoolPostInit().BatchStore(FNameComparisonValue(Name, Hash));
+}
+
+static FNameEntryId BatchLoadNameWithHash(const FNameSerializedView& InName, uint64 InHash)
+{
+	if (InName.bIsUtf16)
+	{
+		// Wide names and hashes are currently stored as UTF16 Little-Endian 
+		// regardless of target architecture. 
+#if PLATFORM_LITTLE_ENDIAN
+		if (sizeof(UTF16CHAR) == sizeof(WIDECHAR))
+		{
+			return BatchLoadNameWithHash(reinterpret_cast<const WIDECHAR*>(InName.Utf16), InName.Len, InHash);
+		}
+#endif
+
+		return BatchLoadNameWithoutHash(InName.Utf16, InName.Len);
+	}
+	else
+	{
+		return BatchLoadNameWithHash(InName.Ansi, InName.Len, InHash);
+	}
+}
+
+void LoadNameBatch(TArray<FNameEntryId>& OutNames, TArrayView<const uint8> NameData, TArrayView<const uint8> HashData)
+{
+	check(IsAligned(NameData.GetData(), sizeof(uint64)));
+	check(IsAligned(HashData.GetData(), sizeof(uint64)));
+	check(IsAligned(HashData.Num(), sizeof(uint64)));
+	check(HashData.Num() > 0);
+
+	const uint8* NameIt = NameData.GetData();
+	const uint8* NameEnd = NameData.GetData() + NameData.Num();
+
+	const uint64* HashDataIt = reinterpret_cast<const uint64*>(HashData.GetData());
+	uint64 HashVersion = INTEL_ORDER64(HashDataIt[0]);
+	TArrayView<const uint64> Hashes = MakeArrayView(HashDataIt + 1, HashData.Num() / sizeof(uint64) - 1);
+
+	OutNames.Empty(Hashes.Num());
+
+	GetNamePoolPostInit().BatchLock();
+
+	if (HashVersion == FNameHash::AlgorithmId)
+	{
+		for (uint64 Hash : Hashes)
+		{
+			check(NameIt < NameEnd);
+			FNameSerializedView Name = LoadNameHeader(/* in-out */ NameIt);
+			OutNames.Add(BatchLoadNameWithHash(Name, INTEL_ORDER64(Hash)));
+		}
+	}
+	else
+	{
+		while (NameIt < NameEnd)
+		{
+			FNameSerializedView Name = LoadNameHeader(/* in-out */ NameIt);
+			OutNames.Add(BatchLoadNameWithoutHash(Name));
+		}
+	
+	}
+
+	GetNamePoolPostInit().BatchUnlock();
+
+	check(NameIt == NameEnd);
+}
+
+#if 0 && ALLOW_NAME_BATCH_SAVING  
+
+FORCENOINLINE void PerfTestLoadNameBatch(TArray<FNameEntryId>& OutNames, TArrayView<const uint8> NameData, TArrayView<const uint8> HashData)
+{
+	LoadNameBatch(OutNames, NameData, HashData);
+}
+
+static bool WithinBlock(uint8* BlockBegin, const FNameEntry* Entry)
+{
+	return UPTRINT(Entry) >= UPTRINT(BlockBegin) && UPTRINT(Entry) < UPTRINT(BlockBegin) + FNameEntryAllocator::BlockSizeBytes;
+}
+
+#include "GenericPlatform\GenericPlatformFile.h"
+
+static void WriteBlobFile(const TCHAR* FileName, const TArray<uint8>& Blob)
+{
+	TUniquePtr<IFileHandle> FileHandle(IPlatformFile::GetPlatformPhysical().OpenWrite(FileName));
+	FileHandle->Write(Blob.GetData(), Blob.Num());
+}
+
+static TArray<uint8> ReadBlobFile(const TCHAR* FileName)
+{
+	TArray<uint8> Out;
+	TUniquePtr<IFileHandle> FileHandle(IPlatformFile::GetPlatformPhysical().OpenRead(FileName));
+	if (FileHandle)
+	{
+		Out.AddUninitialized(FileHandle->Size());
+		FileHandle->Read(Out.GetData(), Out.Num());
+	}
+
+	return Out;
+}
+
+CORE_API int SaveNameBatchTestFiles();
+CORE_API int LoadNameBatchTestFiles();
+
+int SaveNameBatchTestFiles()
+{
+	uint8** Blocks = GetNamePool().GetBlocksForDebugVisualizer();
+	uint32 BlockIdx = 0;
+	TArray<FNameEntryId> NameEntries;
+	for (const FNameEntry* Entry : FName::DebugDump())
+	{
+		BlockIdx += !WithinBlock(Blocks[BlockIdx], Entry);
+		check(WithinBlock(Blocks[BlockIdx], Entry));
+
+		FNameEntryHandle Handle(BlockIdx, (UPTRINT(Entry) - UPTRINT(Blocks[BlockIdx]))/FNameEntryAllocator::Stride);
+		NameEntries.Add(Handle);
+	}
+
+	TArray<uint8> NameData;
+	TArray<uint8> HashData;
+	SaveNameBatch(MakeArrayView(NameEntries), NameData, HashData);
+
+	WriteBlobFile(TEXT("TestNameBatch.Names"), NameData);
+	WriteBlobFile(TEXT("TestNameBatch.Hashes"), HashData);
+
+	return NameEntries.Num();
+}
+
+int LoadNameBatchTestFiles()
+{
+	TArray<uint8> NameData = ReadBlobFile(TEXT("TestNameBatch.Names"));
+	TArray<uint8> HashData = ReadBlobFile(TEXT("TestNameBatch.Hashes"));
+
+	TArray<FNameEntryId> NameEntries;
+	if (HashData.Num())
+	{
+		ReserveNameBatch(NameData.Num(), HashData.Num());
+		PerfTestLoadNameBatch(NameEntries, MakeArrayView(NameData), MakeArrayView(HashData));
+	}
+	return NameEntries.Num();
+}
+
+#endif
+
+static void TestNameBatch()
+{
+#if ALLOW_NAME_BATCH_SAVING
+
+	TArray<FNameEntryId> Names;
+	TArray<uint8> NameData;
+	TArray<uint8> HashData;
+
+	// Test empty batch
+	SaveNameBatch(MakeArrayView(Names), NameData, HashData);
+	check(NameData.Num() == 0);
+	LoadNameBatch(Names, MakeArrayView(NameData), MakeArrayView(HashData));
+	check(Names.Num() == 0);
+
+	// Test empty / "None" name and another EName
+	Names.Add(FName().GetComparisonIndex());
+	Names.Add(FName(NAME_Box).GetComparisonIndex());
+
+	// Test long strings
+	FString MaxLengthAnsi;
+	MaxLengthAnsi.Reserve(NAME_SIZE);
+	while (MaxLengthAnsi.Len() < NAME_SIZE)
+	{
+		MaxLengthAnsi.Append("0123456789ABCDEF");
+	}
+	MaxLengthAnsi.RemoveAt(NAME_SIZE, MaxLengthAnsi.Len() % NAME_SIZE);
+
+	FString MaxLengthWide = MaxLengthAnsi;
+	MaxLengthWide[200] = 500;
+
+	for (const FString& MaxLength : {MaxLengthAnsi, MaxLengthWide})
+	{
+		Names.Add(FName(*MaxLength).GetComparisonIndex());
+		Names.Add(FName(*MaxLength + NAME_SIZE - 255).GetComparisonIndex());
+		Names.Add(FName(*MaxLength + NAME_SIZE - 256).GetComparisonIndex());
+		Names.Add(FName(*MaxLength + NAME_SIZE - 257).GetComparisonIndex());
+	}
+
+	// Test UTF-16 alignment
+	FString Wide("Wide ");
+	Wide[4] = 60000;
+
+	Names.Add(FName(*Wide).GetComparisonIndex());
+	Names.Add(FName("odd").GetComparisonIndex());
+	Names.Add(FName(*Wide).GetComparisonIndex());
+	Names.Add(FName("even").GetComparisonIndex());
+	Names.Add(FName(*Wide).GetComparisonIndex());
+
+	// Roundtrip names
+	SaveNameBatch(MakeArrayView(Names), NameData, HashData);
+	check(NameData.Num() > 0);
+	TArray<FNameEntryId> LoadedNames;
+	LoadNameBatch(LoadedNames, MakeArrayView(NameData), MakeArrayView(HashData));
+	check(LoadedNames == Names);
+
+	// Test changing hash version
+	HashData[0] = 0xba;
+	HashData[1] = 0xad;
+	LoadNameBatch(LoadedNames, MakeArrayView(NameData), MakeArrayView(HashData));
+	check(LoadedNames == Names);
+
+	// Test determinism
+	TArray<uint8> NameData2;
+	TArray<uint8> HashData2;
+
+	auto ClearAndReserveMemoryWithBytePattern = 
+		[](TArray<uint8>& Out, uint8 Pattern, uint32 Num)
+	{
+		Out.Init(Pattern, Num);
+		Out.Empty(Out.Max());
+	};
+	
+	ClearAndReserveMemoryWithBytePattern(NameData2, 0xaa, NameData.Num());
+	ClearAndReserveMemoryWithBytePattern(HashData2, 0xaa, HashData.Num());
+	ClearAndReserveMemoryWithBytePattern(NameData,	0xbb, NameData.Num());
+	ClearAndReserveMemoryWithBytePattern(HashData,	0xbb, HashData.Num());
+	
+	SaveNameBatch(MakeArrayView(Names), NameData, HashData);
+	SaveNameBatch(MakeArrayView(Names), NameData2, HashData2);
+	
+	check(NameData == NameData2);
+	check(HashData == HashData2);
+
+#endif // ALLOW_NAME_BATCH_SAVING
+}
+
 #if !UE_BUILD_SHIPPING && !UE_BUILD_TEST
 
 #include "Containers/StackTracker.h"
diff --git a/Engine/Source/Runtime/Core/Public/IO/IoDispatcher.h b/Engine/Source/Runtime/Core/Public/IO/IoDispatcher.h
--- a/Engine/Source/Runtime/Core/Public/IO/IoDispatcher.h
+++ b/Engine/Source/Runtime/Core/Public/IO/IoDispatcher.h
@@ -584,7 +584,8 @@ enum class EIoChunkType : uint8
 	OptionalBulkData,
 	LoaderGlobalMeta,
 	LoaderInitialLoadMeta,
-	LoaderGlobalNameMap
+	LoaderGlobalNames,
+	LoaderGlobalNameHashes
 };
 
 /**
diff --git a/Engine/Source/Runtime/Core/Public/UObject/NameBatchSerialization.h b/Engine/Source/Runtime/Core/Public/UObject/NameBatchSerialization.h
new file mode 100644
--- /dev/null
+++ b/Engine/Source/Runtime/Core/Public/UObject/NameBatchSerialization.h
@@ -0,0 +1,25 @@
+// Copyright 1998-2019 Epic Games, Inc. All Rights Reserved.
+
+#pragma once
+
+#include "Containers/ArrayView.h"
+#include "UObject/NameTypes.h"
+
+#define ALLOW_NAME_BATCH_SAVING (UE_BUILD_DEBUG || UE_BUILD_DEVELOPMENT) && PLATFORM_LITTLE_ENDIAN
+
+#if ALLOW_NAME_BATCH_SAVING
+// Save comparison entries in given order to a name blob and a versioned hash blob.
+CORE_API void SaveNameBatch(TArrayView<const FNameEntryId> Names, TArray<uint8>& OutNameData, TArray<uint8>& OutHashData);
+#endif
+
+// Reserve memory in preparation for batch loading
+//
+// @param Bytes for existing and new names.
+CORE_API void ReserveNameBatch(uint32 NameDataBytes, uint32 HashDataBytes);
+
+// Load a name blob with precalculated hashes.
+//
+// Names are rehased if hash algorithm version doesn't match.
+//
+// @param NameData, HashData must be 8-byte aligned.
+CORE_API void LoadNameBatch(TArray<FNameEntryId>& OutNames, TArrayView<const uint8> NameData, TArrayView<const uint8> HashData);
\ No newline at end of file
diff --git a/Engine/Source/Runtime/CoreUObject/Private/Serialization/AsyncLoading2.cpp b/Engine/Source/Runtime/CoreUObject/Private/Serialization/AsyncLoading2.cpp
--- a/Engine/Source/Runtime/CoreUObject/Private/Serialization/AsyncLoading2.cpp
+++ b/Engine/Source/Runtime/CoreUObject/Private/Serialization/AsyncLoading2.cpp
@@ -20,6 +20,7 @@
 #include "Misc/PackageName.h"
 #include "UObject/ObjectResource.h"
 #include "UObject/SoftObjectPath.h"
+#include "UObject/NameBatchSerialization.h"
 #include "Serialization/DeferredMessageLog.h"
 #include "UObject/UObjectThreadContext.h"
 #include "Misc/Paths.h"
@@ -271,33 +272,29 @@ class FGlobalNameMap
 public:
 	void Load(FIoDispatcher& IoDispatcher)
 	{
-		FEvent* Event = FPlatformProcess::GetSynchEventFromPool();
-		FIoBuffer IoBuffer;
+		check(NameEntries.Num() == 0);
 
-		IoDispatcher.ReadWithCallback(
-			CreateIoChunkId(0, 0, EIoChunkType::LoaderGlobalNameMap),
-			FIoReadOptions(),
-			[Event, &IoBuffer](TIoStatusOr<FIoBuffer> Result)
-			{
-				IoBuffer = Result.ConsumeValueOrDie();
-				Event->Trigger();
-			});
+		FIoChunkId NamesId = CreateIoChunkId(0, 0, EIoChunkType::LoaderGlobalNames);
+		FIoChunkId HashesId = CreateIoChunkId(0, 0, EIoChunkType::LoaderGlobalNameHashes);
 
-		Event->Wait();
-		FPlatformProcess::ReturnSynchEventToPool(Event);
+		FIoBatch Batch = IoDispatcher.NewBatch();
+		FIoRequest NameRequest = Batch.Read(NamesId, FIoReadOptions());
+		FIoRequest HashRequest = Batch.Read(HashesId, FIoReadOptions());
+		Batch.Issue();
 
-		FLargeMemoryReader Ar(IoBuffer.Data(), IoBuffer.DataSize());
+		ReserveNameBatch(	IoDispatcher.GetSizeForChunk(NamesId).ValueOrDie(),
+							IoDispatcher.GetSizeForChunk(HashesId).ValueOrDie());
 
-		int32 NameCount;
-		Ar << NameCount;
-		NameEntries.Reserve(NameCount);
-		FNameEntrySerialized SerializedNameEntry(ENAME_LinkerConstructor);
+		Batch.Wait();
 
-		for (int32 I = 0; I < NameCount; ++I)
-		{
-			Ar << SerializedNameEntry;
-			NameEntries.Emplace(FName(SerializedNameEntry).GetDisplayIndex());
-		}
+		FIoBuffer NameBuffer = NameRequest.GetResult().ConsumeValueOrDie();
+		FIoBuffer HashBuffer = HashRequest.GetResult().ConsumeValueOrDie();
+
+		LoadNameBatch(/* Out */ NameEntries, 
+						MakeArrayView(NameBuffer.Data(), NameBuffer.DataSize()),
+						MakeArrayView(HashBuffer.Data(), HashBuffer.DataSize()));
+
+		IoDispatcher.FreeBatch(Batch);
 	}
 
 	FName GetName(const uint32 NameIndex, const uint32 NameNumber) const
