Tagged property serialization optimizations

Compare FNames using a single branch in 64-bit test/shipping on x64 and avoid needless magic static. Both yielded ~1%  speedup in the older versioned tagged property loading.

diff --git a/Engine/Source/Runtime/Core/Public/UObject/NameTypes.h b/Engine/Source/Runtime/Core/Public/UObject/NameTypes.h
--- a/Engine/Source/Runtime/Core/Public/UObject/NameTypes.h
+++ b/Engine/Source/Runtime/Core/Public/UObject/NameTypes.h
@@ -502,12 +502,16 @@ public:
 			&& (!bCompareNumber || GetNumber() == Other.GetNumber());
 	}
 
-	FORCEINLINE bool operator==(const FName& Other) const
+	FORCEINLINE bool operator==(FName Other) const
 	{
+#if PLATFORM_64BITS && !WITH_CASE_PRESERVING_NAME
+		return ToComparableInt() == Other.ToComparableInt();
+#else
 		return (ComparisonIndex == Other.ComparisonIndex) & (GetNumber() == Other.GetNumber());
+#endif
 	}
 
-	FORCEINLINE bool operator!=(const FName& Other) const
+	FORCEINLINE bool operator!=(FName Other) const
 	{
 		return !(*this == Other);
 	}
@@ -550,7 +554,11 @@ public:
 	/** True for FName(), FName(NAME_None) and FName("None") */
 	FORCEINLINE bool IsNone() const
 	{
+#if PLATFORM_64BITS && !WITH_CASE_PRESERVING_NAME
+		return ToComparableInt() == 0;
+#else
 		return !ComparisonIndex && GetNumber() == NAME_NO_NUMBER_INTERNAL;
+#endif
 	}
 
 	/**
@@ -868,6 +876,15 @@ private:
 	/** Number portion of the string/number pair (stored internally as 1 more than actual, so zero'd memory will be the default, no-instance case) */
 	uint32			Number;
 
+#if PLATFORM_64BITS && !WITH_CASE_PRESERVING_NAME
+	FORCEINLINE uint64 ToComparableInt() const
+	{
+		static_assert(sizeof(*this) == sizeof(uint64), "");
+		alignas(uint64) FName AlignedCopy = *this;
+		return reinterpret_cast<uint64&>(AlignedCopy);
+	}
+#endif
+
 	friend const TCHAR* DebugFName(int32);
 	friend const TCHAR* DebugFName(int32, int32);
 	friend const TCHAR* DebugFName(FName&);
@@ -881,6 +898,10 @@ private:
 #endif
 	}
 
+
+
+
+
 	static bool IsWithinBounds(FNameEntryId Id);
 };
 
diff --git a/Engine/Source/Runtime/CoreUObject/Private/UObject/Class.cpp b/Engine/Source/Runtime/CoreUObject/Private/UObject/Class.cpp
--- a/Engine/Source/Runtime/CoreUObject/Private/UObject/Class.cpp
+++ b/Engine/Source/Runtime/CoreUObject/Private/UObject/Class.cpp
@@ -1451,8 +1451,9 @@ void UStruct::SerializeVersionedTaggedProperties(FStructuredArchive::FSlot Slot,
 
 		if (!UnderlyingArchive.IsTextFormat())
 		{
-			static FName Temp(NAME_None);
-			UnderlyingArchive << Temp;
+			// Add an empty FName that serves as a null-terminator
+			FName NoneTerminator;
+			UnderlyingArchive << NoneTerminator;
 		}
 	}
 }
