Iteration on StringView

- Fix a linker error for some targets by moving CORE_API functions outside of TStringViewImpl.
- Change members of TStringViewImpl to take ViewType parameters instead of TStringViewImpl.
- Use TCString instead of using FPlatformString directly.

#rnx

diff --git a/Engine/Source/Runtime/Core/Private/Containers/StringView.cpp b/Engine/Source/Runtime/Core/Private/Containers/StringView.cpp
--- a/Engine/Source/Runtime/Core/Private/Containers/StringView.cpp
+++ b/Engine/Source/Runtime/Core/Private/Containers/StringView.cpp
@@ -2,38 +2,46 @@
 
 #include "Misc/StringView.h"
 
-#include "Containers/UnrealString.h"
-#include "HAL/PlatformString.h"
+template class TStringViewImpl<TCHAR, FStringView>;
+template class TStringViewImpl<ANSICHAR, FAnsiStringView>;
+template class TStringViewImpl<WIDECHAR, FWideStringView>;
 
-template<typename CharType, typename ViewType>
-int32 TStringViewImpl<CharType, ViewType>::Compare(const TStringViewImpl& Other, ESearchCase::Type SearchCase) const
+template <typename ViewType>
+int32 StringViewPrivate::Compare(const ViewType& Lhs, const ViewType& Rhs, ESearchCase::Type SearchCase)
 {
-	const SizeType SelfLen = Len();
-	const SizeType OtherLen = Other.Len();
-	const SizeType ShortestLength = FMath::Min(SelfLen, OtherLen);
+	const typename ViewType::SizeType LhsLen = Lhs.Len();
+	const typename ViewType::SizeType RhsLen = Rhs.Len();
+	const typename ViewType::SizeType MinLen = FMath::Min(LhsLen, RhsLen);
 
 	int Result;
 	if (SearchCase == ESearchCase::CaseSensitive)
 	{
-		Result = FPlatformString::Strncmp(GetData(), Other.GetData(), ShortestLength);
+		Result = TCString<typename ViewType::ElementType>::Strncmp(Lhs.GetData(), Rhs.GetData(), MinLen);
 	}
 	else
 	{
-		Result = FPlatformString::Strnicmp(GetData(), Other.GetData(), ShortestLength);
+		Result = TCString<typename ViewType::ElementType>::Strnicmp(Lhs.GetData(), Rhs.GetData(), MinLen);
 	}
 
-	if (Result != 0 || SelfLen == OtherLen)
+	if (Result != 0 || LhsLen == RhsLen)
 	{
 		return Result;
 	}
 
-	return SelfLen < OtherLen ? -1 : 1;
+	return LhsLen < RhsLen ? -1 : 1;
 }
 
-template<typename CharType, typename ViewType>
-bool TStringViewImpl<CharType, ViewType>::FindChar(CharType InChar, SizeType& OutIndex) const
+template int32 CORE_API StringViewPrivate::Compare(const FStringView& Lhs, const FStringView& Rhs, ESearchCase::Type SearchCase);
+template int32 CORE_API StringViewPrivate::Compare(const FAnsiStringView& Lhs, const FAnsiStringView& Rhs, ESearchCase::Type SearchCase);
+template int32 CORE_API StringViewPrivate::Compare(const FWideStringView& Lhs, const FWideStringView& Rhs, ESearchCase::Type SearchCase);
+
+template <typename ViewType>
+bool StringViewPrivate::FindChar(const ViewType& InView, typename ViewType::ElementType InChar, typename ViewType::SizeType& OutIndex)
 {
-	for (SizeType i = 0; i < Size; ++i)
+	const typename ViewType::ElementType* DataPtr = InView.GetData();
+	const typename ViewType::SizeType Size = InView.Len();
+
+	for (typename ViewType::SizeType i = 0; i < Size; ++i)
 	{
 		if (DataPtr[i] == InChar)
 		{
@@ -46,16 +54,23 @@ bool TStringViewImpl<CharType, ViewType>::FindChar(CharType InChar, SizeType& Ou
 	return false;
 }
 
-template<typename CharType, typename ViewType>
-bool TStringViewImpl<CharType, ViewType>::FindLastChar(CharType InChar, SizeType& OutIndex) const
+template bool CORE_API StringViewPrivate::FindChar(const FStringView& InView, TCHAR InChar, typename FStringView::SizeType& OutIndex);
+template bool CORE_API StringViewPrivate::FindChar(const FAnsiStringView& InView, ANSICHAR InChar, typename FAnsiStringView::SizeType& OutIndex);
+template bool CORE_API StringViewPrivate::FindChar(const FWideStringView& InView, WIDECHAR InChar, typename FWideStringView::SizeType& OutIndex);
+
+template <typename ViewType>
+bool StringViewPrivate::FindLastChar(const ViewType& InView, typename ViewType::ElementType InChar, typename ViewType::SizeType& OutIndex)
 {
+	const typename ViewType::ElementType* DataPtr = InView.GetData();
+	const typename ViewType::SizeType Size = InView.Len();
+
 	if (Size == 0)
 	{
 		OutIndex = INDEX_NONE;
 		return false;
 	}
 
-	for (SizeType i = Size - 1; i >= 0; --i)
+	for (typename ViewType::SizeType i = Size - 1; i >= 0; --i)
 	{
 		if (DataPtr[i] == InChar)
 		{
@@ -68,6 +83,6 @@ bool TStringViewImpl<CharType, ViewType>::FindLastChar(CharType InChar, SizeType
 	return false;
 }
 
-template class TStringViewImpl<TCHAR, FStringView>;
-template class TStringViewImpl<ANSICHAR, FAnsiStringView>;
-template class TStringViewImpl<WIDECHAR, FWideStringView>;
+template bool CORE_API StringViewPrivate::FindLastChar(const FStringView& InView, TCHAR InChar, typename FStringView::SizeType& OutIndex);
+template bool CORE_API StringViewPrivate::FindLastChar(const FAnsiStringView& InView, ANSICHAR InChar, typename FAnsiStringView::SizeType& OutIndex);
+template bool CORE_API StringViewPrivate::FindLastChar(const FWideStringView& InView, WIDECHAR InChar, typename FWideStringView::SizeType& OutIndex);
diff --git a/Engine/Source/Runtime/Core/Private/Tests/Misc/StringViewTest.cpp b/Engine/Source/Runtime/Core/Private/Tests/Misc/StringViewTest.cpp
--- a/Engine/Source/Runtime/Core/Private/Tests/Misc/StringViewTest.cpp
+++ b/Engine/Source/Runtime/Core/Private/Tests/Misc/StringViewTest.cpp
@@ -615,37 +615,37 @@ bool FStringViewTestFindChar::RunTest(const FString& Parameters)
 
 	{
 		int32 Index = INDEX_NONE;
-		TestFalse(TEXT("FStringView::FindChar-Return(0)"), EmptyView.FindChar('a', Index));
+		TestFalse(TEXT("FStringView::FindChar-Return(0)"), EmptyView.FindChar(TEXT('a'), Index));
 		TestEqual(TEXT("FStringView::FindChar-Index(0)"), Index, INDEX_NONE);
 	}
 
 	{
 		int32 Index = INDEX_NONE;
-		TestTrue(TEXT("FStringView::FindChar-Return(1)"), View.FindChar('a', Index));
+		TestTrue(TEXT("FStringView::FindChar-Return(1)"), View.FindChar(TEXT('a'), Index));
 		TestEqual(TEXT("FStringView::FindChar-Index(1)"), Index, 0);
 	}
 
 	{
 		int32 Index = INDEX_NONE;
-		TestTrue(TEXT("FStringView::FindChar-Return(2)"), View.FindChar('F', Index));
+		TestTrue(TEXT("FStringView::FindChar-Return(2)"), View.FindChar(TEXT('F'), Index));
 		TestEqual(TEXT("FStringView::FindChar-Index(2)"), Index, 5);
 	}
 
 	{
 		int32 Index = INDEX_NONE;
-		TestFalse(TEXT("FStringView::FindChar-Return(3)"), View.FindChar('A', Index));
+		TestFalse(TEXT("FStringView::FindChar-Return(3)"), View.FindChar(TEXT('A'), Index));
 		TestEqual(TEXT("FStringView::FindChar-Index(3)"), Index, INDEX_NONE);
 	}
 
 	{
 		int32 Index = INDEX_NONE;
-		TestFalse(TEXT("FStringView::FindChar-Return(4)"), View.FindChar('d', Index));
+		TestFalse(TEXT("FStringView::FindChar-Return(4)"), View.FindChar(TEXT('d'), Index));
 		TestEqual(TEXT("FStringView::FindChar-Index(4)"), Index, INDEX_NONE);
 	}
 
 	{
 		int32 Index = INDEX_NONE;
-		TestTrue(TEXT("FStringView::FindChar-Return(5)"), View.FindChar(' ', Index));
+		TestTrue(TEXT("FStringView::FindChar-Return(5)"), View.FindChar(TEXT(' '), Index));
 		TestEqual(TEXT("FStringView::FindChar-Index(5)"), Index, 4);
 	}
 
@@ -660,37 +660,37 @@ bool FStringViewTestFindLastChar::RunTest(const FString& Parameters)
 
 	{
 		int32 Index = INDEX_NONE;
-		TestFalse(TEXT("FStringView::FindChar-Return(0)"), EmptyView.FindLastChar('a', Index));
+		TestFalse(TEXT("FStringView::FindChar-Return(0)"), EmptyView.FindLastChar(TEXT('a'), Index));
 		TestEqual(TEXT("FStringView::FindChar-Index(0)"), Index, INDEX_NONE);
 	}
 
 	{
 		int32 Index = INDEX_NONE;
-		TestTrue(TEXT("FStringView::FindChar-Return(1)"), View.FindLastChar('a', Index));
+		TestTrue(TEXT("FStringView::FindChar-Return(1)"), View.FindLastChar(TEXT('a'), Index));
 		TestEqual(TEXT("FStringView::FindChar-Index(1)"), Index, 7);
 	}
 
 	{
 		int32 Index = INDEX_NONE;
-		TestTrue(TEXT("FStringView::FindChar-Return(2)"), View.FindLastChar('B', Index));
+		TestTrue(TEXT("FStringView::FindChar-Return(2)"), View.FindLastChar(TEXT('B'), Index));
 		TestEqual(TEXT("FStringView::FindChar-Index(2)"), Index, 1);
 	}
 
 	{
 		int32 Index = INDEX_NONE;
-		TestFalse(TEXT("FStringView::FindChar-Return(3)"), View.FindLastChar('A', Index));
+		TestFalse(TEXT("FStringView::FindChar-Return(3)"), View.FindLastChar(TEXT('A'), Index));
 		TestEqual(TEXT("FStringView::FindChar-Index(3)"), Index, INDEX_NONE);
 	}
 
 	{
 		int32 Index = INDEX_NONE;
-		TestFalse(TEXT("FStringView::FindChar-Return(4)"), View.FindLastChar('d', Index));
+		TestFalse(TEXT("FStringView::FindChar-Return(4)"), View.FindLastChar(TEXT('d'), Index));
 		TestEqual(TEXT("FStringView::FindChar-Index(4)"), Index, INDEX_NONE);
 	}
 
 	{
 		int32 Index = INDEX_NONE;
-		TestTrue(TEXT("FStringView::FindChar-Return(5)"), View.FindLastChar(' ', Index));
+		TestTrue(TEXT("FStringView::FindChar-Return(5)"), View.FindLastChar(TEXT(' '), Index));
 		TestEqual(TEXT("FStringView::FindChar-Index(5)"), Index, 4);
 	}
 
diff --git a/Engine/Source/Runtime/Core/Public/Containers/StringView.h b/Engine/Source/Runtime/Core/Public/Containers/StringView.h
--- a/Engine/Source/Runtime/Core/Public/Containers/StringView.h
+++ b/Engine/Source/Runtime/Core/Public/Containers/StringView.h
@@ -173,7 +173,7 @@ public:
 	 *
 	 * @param SearchCase Whether the comparison should ignore case.
 	 */
-	bool Equals(const TStringViewImpl& Other, ESearchCase::Type SearchCase = ESearchCase::CaseSensitive) const;
+	inline bool Equals(const ViewType& Other, ESearchCase::Type SearchCase = ESearchCase::CaseSensitive) const;
 
 	/**
 	 * Compare this view lexicographically with another view.
@@ -182,17 +182,17 @@ public:
 	 *
 	 * @return 0 is equal, negative if this view is less, positive if this view is greater.
 	 */
-	CORE_API int32 Compare(const TStringViewImpl& Other, ESearchCase::Type SearchCase = ESearchCase::CaseSensitive) const;
+	inline int32 Compare(const ViewType& Other, ESearchCase::Type SearchCase = ESearchCase::CaseSensitive) const;
 
 	/** Returns whether this view starts with the prefix character compared case-sensitively. */
 	inline bool StartsWith(CharType Prefix) const { return Size >= 1 && DataPtr[0] == Prefix; }
 	/** Returns whether this view starts with the prefix with optional case sensitivity. */
-	inline bool StartsWith(const TStringViewImpl& Prefix, ESearchCase::Type SearchCase = ESearchCase::IgnoreCase) const;
+	inline bool StartsWith(const ViewType& Prefix, ESearchCase::Type SearchCase = ESearchCase::IgnoreCase) const;
 
 	/** Returns whether this view ends with the suffix character compared case-sensitively. */
 	inline bool EndsWith(CharType Suffix) const { return Size >= 1 && DataPtr[Size-1] == Suffix; }
 	/** Returns whether this view ends with the suffix with optional case sensitivity. */
-	inline bool EndsWith(const TStringViewImpl& Suffix, ESearchCase::Type SearchCase = ESearchCase::IgnoreCase) const;
+	inline bool EndsWith(const ViewType& Suffix, ESearchCase::Type SearchCase = ESearchCase::IgnoreCase) const;
 
 	// Searching/Finding
 
@@ -204,7 +204,7 @@ public:
 	 *
 	 * @return Whether the character was found in the view.
 	 */
-	CORE_API bool FindChar(CharType InChar, SizeType& OutIndex) const;
+	inline bool FindChar(CharType InChar, SizeType& OutIndex) const;
 
 	/**
 	 * Search the view for the last occurrence of a character.
@@ -214,7 +214,7 @@ public:
 	 *
 	 * @return Whether the character was found in the view.
 	 */
-	CORE_API bool FindLastChar(CharType InChar, SizeType& OutIndex) const;
+	inline bool FindLastChar(CharType InChar, SizeType& OutIndex) const;
 
 public:
 	/**
diff --git a/Engine/Source/Runtime/Core/Public/Containers/StringView.inl b/Engine/Source/Runtime/Core/Public/Containers/StringView.inl
--- a/Engine/Source/Runtime/Core/Public/Containers/StringView.inl
+++ b/Engine/Source/Runtime/Core/Public/Containers/StringView.inl
@@ -1,11 +1,25 @@
 // Copyright 1998-2019 Epic Games, Inc. All Rights Reserved.
 
-#include "HAL/PlatformString.h"
 #include "HAL/UnrealMemory.h"
 #include "Templates/Decay.h"
 #include "Templates/UnrealTypeTraits.h"
 #include "Traits/IntType.h"
 
+namespace StringViewPrivate
+{
+	// These exported functions are implemented outside of the class to avoid the complexity
+	// of correctly exporting functions using many toolchains from a type that uses the CRTP.
+
+	template<class ViewType>
+	CORE_API int32 Compare(const ViewType& Lhs, const ViewType& Rhs, ESearchCase::Type SearchCase);
+
+	template<class ViewType>
+	CORE_API bool FindChar(const ViewType& InView, typename ViewType::ElementType InChar, typename ViewType::SizeType& OutIndex);
+
+	template<class ViewType>
+	CORE_API bool FindLastChar(const ViewType& InView, typename ViewType::ElementType InChar, typename ViewType::SizeType& OutIndex);
+}
+
 template <typename CharType, typename ViewType>
 inline const CharType& TStringViewImpl<CharType, ViewType>::operator[](SizeType Index) const
 {
@@ -58,30 +72,48 @@ inline ViewType TStringViewImpl<CharType, ViewType>::Mid(SizeType Position, Size
 }
 
 template <typename CharType, typename ViewType>
-inline bool TStringViewImpl<CharType, ViewType>::Equals(const TStringViewImpl& Other, ESearchCase::Type SearchCase) const
+inline bool TStringViewImpl<CharType, ViewType>::Equals(const ViewType& Other, ESearchCase::Type SearchCase) const
 {
 	return Size == Other.Size && Compare(Other, SearchCase) == 0;
 }
 
 template <typename CharType, typename ViewType>
-inline bool TStringViewImpl<CharType, ViewType>::StartsWith(const TStringViewImpl& Prefix, ESearchCase::Type SearchCase) const
+inline int32 TStringViewImpl<CharType, ViewType>::Compare(const ViewType& Other, ESearchCase::Type SearchCase) const
+{
+	return StringViewPrivate::Compare(static_cast<const ViewType&>(*this), Other, SearchCase);
+}
+
+template <typename CharType, typename ViewType>
+inline bool TStringViewImpl<CharType, ViewType>::StartsWith(const ViewType& Prefix, ESearchCase::Type SearchCase) const
 {
 	return Prefix.Equals(Left(Prefix.Len()), SearchCase);
 }
 
 template <typename CharType, typename ViewType>
-inline bool TStringViewImpl<CharType, ViewType>::EndsWith(const TStringViewImpl& Suffix, ESearchCase::Type SearchCase) const
+inline bool TStringViewImpl<CharType, ViewType>::EndsWith(const ViewType& Suffix, ESearchCase::Type SearchCase) const
 {
 	return Suffix.Equals(Right(Suffix.Len()), SearchCase);
 }
 
+template <typename CharType, typename ViewType>
+inline bool TStringViewImpl<CharType, ViewType>::FindChar(CharType InChar, SizeType& OutIndex) const
+{
+	return StringViewPrivate::FindChar(static_cast<const ViewType&>(*this), InChar, OutIndex);
+}
+
+template <typename CharType, typename ViewType>
+inline bool TStringViewImpl<CharType, ViewType>::FindLastChar(CharType InChar, SizeType& OutIndex) const
+{
+	return StringViewPrivate::FindLastChar(static_cast<const ViewType&>(*this), InChar, OutIndex);
+}
+
 // Case-insensitive comparison operators
 
 template <typename CharType, typename ViewType, typename RangeType,
 	typename = typename TStringViewImpl<CharType, ViewType>::template TEnableIfCompatibleRangeType<RangeType>>
 inline bool operator==(const TStringViewImpl<CharType, ViewType>& Lhs, RangeType&& Rhs)
 {
-	return Lhs.Equals(TStringViewImpl<CharType, ViewType>(Forward<RangeType>(Rhs)), ESearchCase::IgnoreCase);
+	return Lhs.Equals(ViewType(Forward<RangeType>(Rhs)), ESearchCase::IgnoreCase);
 }
 
 template <typename CharType, typename ViewType, typename RangeType,
