[UE-85352] Fix spurious index out of range assert in FPaths::IsUnderDirectory

diff --git a/Engine/Source/Runtime/Core/Private/Misc/Paths.cpp b/Engine/Source/Runtime/Core/Private/Misc/Paths.cpp
--- a/Engine/Source/Runtime/Core/Private/Misc/Paths.cpp
+++ b/Engine/Source/Runtime/Core/Private/Misc/Paths.cpp
@@ -1489,7 +1489,7 @@ bool FPaths::IsUnderDirectory(const FString& InPath, const FString& InDirectory)
 	int Compare = FCString::Strncmp(*Path, *Directory, Directory.Len());
 #endif
 
-	return Compare == 0 && (Path[Directory.Len()] == 0 || Path[Directory.Len()] == '/');
+	return Compare == 0 && (Path.Len() == Directory.Len() || Path[Directory.Len()] == '/');
 }
 
 
