FPaths tests: Add tests for IsUnderDirectory, remove VS2015 workaround, and reduce string construction to improve compile time and (hopefully) fix the problem VS2015 compiler was having.
#rnx

diff --git a/Engine/Source/Runtime/Core/Private/Tests/Misc/PathsTest.cpp b/Engine/Source/Runtime/Core/Private/Tests/Misc/PathsTest.cpp
--- a/Engine/Source/Runtime/Core/Private/Tests/Misc/PathsTest.cpp
+++ b/Engine/Source/Runtime/Core/Private/Tests/Misc/PathsTest.cpp
@@ -10,14 +10,11 @@
 
 IMPLEMENT_SIMPLE_AUTOMATION_TEST(FPathTests, "System.Core.Misc.Paths", EAutomationTestFlags::ApplicationContextMask | EAutomationTestFlags::SmokeFilter)
 
-// MSVC 2015 generates bad code here due to the number of string literals, in a cooked build it was crashing on RunCollapseRelativeDirectoriesTest(TEXT("C:/."), TEXT("C:/."));
-PRAGMA_DISABLE_OPTIMIZATION
-
 bool FPathTests::RunTest( const FString& Parameters )
 {
 	// Directory collapsing
 	{
-		auto RunCollapseRelativeDirectoriesTest = [this](const FString& InPath, const TCHAR* InResult)
+		auto RunCollapseRelativeDirectoriesTest = [this](const TCHAR* InPath, const TCHAR* InResult)
 		{
 			// Run test
 			FString CollapsedPath = InPath;
@@ -28,7 +25,7 @@ bool FPathTests::RunTest( const FString& Parameters )
 				// If we're looking for a result, make sure it was returned correctly
 				if (!bValid || CollapsedPath != InResult)
 				{
-					AddError(FString::Printf(TEXT("Path '%s' failed to collapse correctly (got '%s', expected '%s')."), *InPath, *CollapsedPath, InResult));
+					AddError(FString::Printf(TEXT("Path '%s' failed to collapse correctly (got '%s', expected '%s')."), InPath, *CollapsedPath, InResult));
 				}
 			}
 			else
@@ -36,7 +33,7 @@ bool FPathTests::RunTest( const FString& Parameters )
 				// Otherwise, make sure it failed
 				if (bValid)
 				{
-					AddError(FString::Printf(TEXT("Path '%s' collapsed unexpectedly."), *InPath));
+					AddError(FString::Printf(TEXT("Path '%s' collapsed unexpectedly."), InPath));
 				}
 			}
 		};
@@ -101,13 +98,13 @@ bool FPathTests::RunTest( const FString& Parameters )
 
 	// Extension texts
 	{
-		auto RunGetExtensionTest = [this](const FString& InPath, const FString& InExpectedExt)
+		auto RunGetExtensionTest = [this](const TCHAR* InPath, const TCHAR* InExpectedExt)
 		{
 			// Run test
-			const FString Ext = FPaths::GetExtension(InPath);
+			const FString Ext = FPaths::GetExtension(FString(InPath));
 			if (Ext != InExpectedExt)
 			{
-				AddError(FString::Printf(TEXT("Path '%s' failed to get the extension (got '%s', expected '%s')."), *InPath, *Ext, *InExpectedExt));
+				AddError(FString::Printf(TEXT("Path '%s' failed to get the extension (got '%s', expected '%s')."), InPath, *Ext, InExpectedExt));
 			}
 		};
 
@@ -121,13 +118,13 @@ bool FPathTests::RunTest( const FString& Parameters )
 		RunGetExtensionTest(TEXT("C:/Folder/First.Last/file.txt"),			TEXT("txt"));
 		RunGetExtensionTest(TEXT("C:/Folder/First.Last/file.tar.gz"),		TEXT("gz"));
 
-		auto RunSetExtensionTest = [this](const FString& InPath, const FString& InNewExt, const FString& InExpectedPath)
+		auto RunSetExtensionTest = [this](const TCHAR* InPath, const TCHAR* InNewExt, const FString& InExpectedPath)
 		{
 			// Run test
-			const FString NewPath = FPaths::SetExtension(InPath, InNewExt);
+			const FString NewPath = FPaths::SetExtension(FString(InPath), FString(InNewExt));
 			if (NewPath != InExpectedPath)
 			{
-				AddError(FString::Printf(TEXT("Path '%s' failed to set the extension (got '%s', expected '%s')."), *InPath, *NewPath, *InExpectedPath));
+				AddError(FString::Printf(TEXT("Path '%s' failed to set the extension (got '%s', expected '%s')."), InPath, *NewPath, *InExpectedPath));
 			}
 		};
 
@@ -141,13 +138,13 @@ bool FPathTests::RunTest( const FString& Parameters )
 		RunSetExtensionTest(TEXT("C:/Folder/First.Last/file.txt"),			TEXT("log"),	TEXT("C:/Folder/First.Last/file.log"));
 		RunSetExtensionTest(TEXT("C:/Folder/First.Last/file.tar.gz"),		TEXT("gz2"),	TEXT("C:/Folder/First.Last/file.tar.gz2"));
 
-		auto RunChangeExtensionTest = [this](const FString& InPath, const FString& InNewExt, const FString& InExpectedPath)
+		auto RunChangeExtensionTest = [this](const TCHAR* InPath, const TCHAR* InNewExt, const FString& InExpectedPath)
 		{
 			// Run test
-			const FString NewPath = FPaths::ChangeExtension(InPath, InNewExt);
+			const FString NewPath = FPaths::ChangeExtension(FString(InPath), FString(InNewExt));
 			if (NewPath != InExpectedPath)
 			{
-				AddError(FString::Printf(TEXT("Path '%s' failed to change the extension (got '%s', expected '%s')."), *InPath, *NewPath, *InExpectedPath));
+				AddError(FString::Printf(TEXT("Path '%s' failed to change the extension (got '%s', expected '%s')."), InPath, *NewPath, *InExpectedPath));
 			}
 		};
 
@@ -162,9 +159,32 @@ bool FPathTests::RunTest( const FString& Parameters )
 		RunChangeExtensionTest(TEXT("C:/Folder/First.Last/file.tar.gz"),	TEXT("gz2"),	TEXT("C:/Folder/First.Last/file.tar.gz2"));
 	}
 
+	// IsUnderDirectory
+	{
+		auto RunIsUnderDirectoryTest = [this](const TCHAR* InPath1, const TCHAR* InPath2, bool ExpectedResult)
+		{
+			// Run test
+			bool Result = FPaths::IsUnderDirectory(FString(InPath1), FString(InPath2));
+			if (Result != ExpectedResult)
+			{
+				AddError(FString::Printf(TEXT("FPaths::IsUnderDirectory('%s', '%s') != %s."), InPath1, InPath2, ExpectedResult ? TEXT("true") : TEXT("false")));
+			}
+		};
+
+		RunIsUnderDirectoryTest(TEXT("C:/Folder"),			TEXT("C:/FolderN"), false);
+		RunIsUnderDirectoryTest(TEXT("C:/Folder1"),			TEXT("C:/Folder2"), false);
+		RunIsUnderDirectoryTest(TEXT("C:/Folder"),			TEXT("C:/Folder/SubDir"), false);
+
+		RunIsUnderDirectoryTest(TEXT("C:/Folder"),			TEXT("C:/Folder"), true);
+		RunIsUnderDirectoryTest(TEXT("C:/Folder/File"),		TEXT("C:/Folder"), true);
+		RunIsUnderDirectoryTest(TEXT("C:/Folder/File"),		TEXT("C:/Folder/"), true);
+		RunIsUnderDirectoryTest(TEXT("C:/Folder/"),			TEXT("C:/Folder"), true);
+		RunIsUnderDirectoryTest(TEXT("C:/Folder/"),			TEXT("C:/Folder/"), true);
+		RunIsUnderDirectoryTest(TEXT("C:/Folder/Subdir/"),	TEXT("C:/Folder"), true);
+		RunIsUnderDirectoryTest(TEXT("C:/Folder/Subdir/"),	TEXT("C:/Folder/"), true);
+	}
+
 	return true;
 }
 
-PRAGMA_ENABLE_OPTIMIZATION
-
 #endif //WITH_DEV_AUTOMATION_TESTS
