Change FStringView constructor from TCHAR* to interpret nullptr as empty string, to match FString.
Add testcase, and fix some missing TEXT macros in the test file.
#rnx

diff --git a/Engine/Source/Runtime/Core/Private/Tests/Misc/StringViewTest.cpp b/Engine/Source/Runtime/Core/Private/Tests/Misc/StringViewTest.cpp
--- a/Engine/Source/Runtime/Core/Private/Tests/Misc/StringViewTest.cpp
+++ b/Engine/Source/Runtime/Core/Private/Tests/Misc/StringViewTest.cpp
@@ -20,8 +20,15 @@ bool FStringViewTestCtor::RunTest(const FString& Parameters)
 	// Empty View
 	{
 		FStringView View( TEXT("") );
-		TestEqual("", View.Len(), 0);
-		TestTrue("View.IsEmpty", View.IsEmpty());
+		TestEqual(TEXT(""), View.Len(), 0);
+		TestTrue(TEXT("View.IsEmpty"), View.IsEmpty());
+	}
+
+	// Constructing from nullptr is supported; nullptr interpreted as empty string
+	{
+		FStringView View(nullptr);
+		TestEqual(TEXT(""), View.Len(), 0);
+		TestTrue(TEXT("View.IsEmpty"), View.IsEmpty());
 	}
 
 	// Create from a wchar literal
@@ -29,7 +36,7 @@ bool FStringViewTestCtor::RunTest(const FString& Parameters)
 		FStringView View(TEXT("Test Ctor"));
 		TestEqual(TEXT("View length"), View.Len(), FCStringWide::Strlen(TEXT("Test Ctor")));
 		TestEqual(TEXT("The result of Strncmp"), FCStringWide::Strncmp(View.GetData(), TEXT("Test Ctor"), View.Len()), 0);
-		TestFalse("View.IsEmpty", View.IsEmpty());
+		TestFalse(TEXT("View.IsEmpty"), View.IsEmpty());
 	}
 
 	// Create from a sub section of a wchar literal
@@ -37,7 +44,7 @@ bool FStringViewTestCtor::RunTest(const FString& Parameters)
 		FStringView View(TEXT("Test SubSection Ctor"), 4);
 		TestEqual(TEXT("View length"), View.Len(), 4);
 		TestEqual(TEXT("The result of Strncmp"), FCStringWide::Strncmp(View.GetData(), TEXT("Test"), View.Len()), 0);
-		TestFalse("View.IsEmpty", View.IsEmpty());
+		TestFalse(TEXT("View.IsEmpty"), View.IsEmpty());
 	}
 
 	// Create from a FString
@@ -47,7 +54,7 @@ bool FStringViewTestCtor::RunTest(const FString& Parameters)
 
 		TestEqual(TEXT("View length"), View.Len(), String.Len());
 		TestEqual(TEXT("The result of Strncmp"), FCStringWide::Strncmp(View.GetData(), *String, View.Len()), 0);
-		TestFalse("View.IsEmpty", View.IsEmpty());
+		TestFalse(TEXT("View.IsEmpty"), View.IsEmpty());
 	}
 
 	// Create from a ansi literal
@@ -55,7 +62,7 @@ bool FStringViewTestCtor::RunTest(const FString& Parameters)
 		FAnsiStringView View("Test Ctor");
 		TestEqual(TEXT("View length"), View.Len(), FCStringAnsi::Strlen("Test Ctor"));
 		TestEqual(TEXT("The result of Strncmp"), FCStringAnsi::Strncmp(View.GetData(), "Test Ctor", View.Len()), 0);
-		TestFalse("View.IsEmpty", View.IsEmpty());
+		TestFalse(TEXT("View.IsEmpty"), View.IsEmpty());
 	}
 
 	// Create from a sub section of an ansi literal
@@ -63,7 +70,7 @@ bool FStringViewTestCtor::RunTest(const FString& Parameters)
 		FAnsiStringView View("Test SubSection Ctor", 4);
 		TestEqual(TEXT("View length"), View.Len(), 4);
 		TestEqual(TEXT("The result of Strncmp"), FCStringAnsi::Strncmp(View.GetData(), "Test", View.Len()), 0);
-		TestFalse("View.IsEmpty", View.IsEmpty());
+		TestFalse(TEXT("View.IsEmpty"), View.IsEmpty());
 	}
 
 	return true;
@@ -372,32 +379,32 @@ bool FStringViewTestStartsWith::RunTest(const FString& Parameters)
 	// Test an empty view
 	{
 		FStringView View( TEXT(""));
-		TestTrue(" View.StartsWith", View.StartsWith(TEXT("")));
-		TestFalse(" View.StartsWith", View.StartsWith(TEXT("Text")));
-		TestFalse(" View.StartsWith", View.StartsWith('A'));
+		TestTrue(TEXT(" View.StartsWith"), View.StartsWith(TEXT("")));
+		TestFalse(TEXT(" View.StartsWith"), View.StartsWith(TEXT("Text")));
+		TestFalse(TEXT(" View.StartsWith"), View.StartsWith('A'));
 	}
 
 	// Test a valid view with the correct text
 	{
 		FStringView View(TEXT("String to test"));
-		TestTrue(" View.StartsWith", View.StartsWith(TEXT("String")));
-		TestTrue(" View.StartsWith", View.StartsWith('S'));
+		TestTrue(TEXT(" View.StartsWith"), View.StartsWith(TEXT("String")));
+		TestTrue(TEXT(" View.StartsWith"), View.StartsWith('S'));
 	}
 
 	// Test a valid view with incorrect text
 	{
 		FStringView View(TEXT("String to test"));
-		TestFalse(" View.StartsWith", View.StartsWith(TEXT("test")));
-		TestFalse(" View.StartsWith", View.StartsWith('t'));
+		TestFalse(TEXT(" View.StartsWith"), View.StartsWith(TEXT("test")));
+		TestFalse(TEXT(" View.StartsWith"), View.StartsWith('t'));
 	}
 
 	// Test a valid view with the correct text but with different case
 	{
 		FStringView View(TEXT("String to test"));
-		TestTrue(" View.StartsWith", View.StartsWith(TEXT("sTrInG")));
+		TestTrue(TEXT(" View.StartsWith"), View.StartsWith(TEXT("sTrInG")));
 
 		// Searching by char is case sensitive to keep compatibility with FString
-		TestFalse(" View.StartsWith", View.StartsWith('s'));
+		TestFalse(TEXT(" View.StartsWith"), View.StartsWith('s'));
 	}
 
 	return true;
@@ -409,32 +416,32 @@ bool FStringViewTestEndsWith::RunTest(const FString& Parameters)
 	// Test an empty view
 	{
 		FStringView View(TEXT(""));
-		TestTrue(" View.EndsWith", View.EndsWith(TEXT("")));
-		TestFalse(" View.EndsWith", View.EndsWith(TEXT("Text")));
-		TestFalse(" View.EndsWith", View.EndsWith('A'));
+		TestTrue(TEXT(" View.EndsWith"), View.EndsWith(TEXT("")));
+		TestFalse(TEXT(" View.EndsWith"), View.EndsWith(TEXT("Text")));
+		TestFalse(TEXT(" View.EndsWith"), View.EndsWith('A'));
 	}
 
 	// Test a valid view with the correct text
 	{
 		FStringView View(TEXT("String to test"));
-		TestTrue(" View.EndsWith", View.EndsWith(TEXT("test")));
-		TestTrue(" View.EndsWith", View.EndsWith('t'));
+		TestTrue(TEXT(" View.EndsWith"), View.EndsWith(TEXT("test")));
+		TestTrue(TEXT(" View.EndsWith"), View.EndsWith('t'));
 	}
 
 	// Test a valid view with incorrect text
 	{
 		FStringView View(TEXT("String to test"));
-		TestFalse(" View.EndsWith", View.EndsWith(TEXT("String")));
-		TestFalse(" View.EndsWith", View.EndsWith('S'));
+		TestFalse(TEXT(" View.EndsWith"), View.EndsWith(TEXT("String")));
+		TestFalse(TEXT(" View.EndsWith"), View.EndsWith('S'));
 	}
 
 	// Test a valid view with the correct text but with different case
 	{
 		FStringView View(TEXT("String to test"));
-		TestTrue(" View.EndsWith", View.EndsWith(TEXT("TeST")));
+		TestTrue(TEXT(" View.EndsWith"), View.EndsWith(TEXT("TeST")));
 
 		// Searching by char is case sensitive to keep compatibility with FString
-		TestFalse(" View.EndsWith", View.EndsWith('T')); 
+		TestFalse(TEXT(" View.EndsWith"), View.EndsWith('T')); 
 	}
 
 	return true;
@@ -479,10 +486,10 @@ bool FStringViewTestLeft::RunTest(const FString& Parameters)
 	{
 		FStringView EmptyView(TEXT(""));
 		FStringView EmptyResult = EmptyView.Left(0);
-		TestTrue("FStringView::Left", EmptyResult.IsEmpty());
+		TestTrue(TEXT("FStringView::Left"), EmptyResult.IsEmpty());
 
 		FStringView OutofBoundsResult = EmptyView.Left(1024);
-		TestTrue("FStringView::Left", OutofBoundsResult.IsEmpty());
+		TestTrue(TEXT("FStringView::Left"), OutofBoundsResult.IsEmpty());
 	}
 	
 	{
@@ -504,10 +511,10 @@ bool FStringViewTestLeftChop::RunTest(const FString& Parameters)
 	{
 		FStringView EmptyView(TEXT(""));
 		FStringView EmptyResult = EmptyView.LeftChop(0);
-		TestTrue("FStringView::LeftChop", EmptyResult.IsEmpty());
+		TestTrue(TEXT("FStringView::LeftChop"), EmptyResult.IsEmpty());
 
 		FStringView OutofBoundsResult = EmptyView.LeftChop(1024);
-		TestTrue("FStringView::LeftChop", OutofBoundsResult.IsEmpty());
+		TestTrue(TEXT("FStringView::LeftChop"), OutofBoundsResult.IsEmpty());
 	}
 
 	{
@@ -529,10 +536,10 @@ bool FStringViewTestRight::RunTest(const FString& Parameters)
 	{
 		FStringView EmptyView(TEXT(""));
 		FStringView EmptyResult = EmptyView.Right(0);
-		TestTrue("FStringView::Right", EmptyResult.IsEmpty());
+		TestTrue(TEXT("FStringView::Right"), EmptyResult.IsEmpty());
 
 		FStringView OutofBoundsResult = EmptyView.Right(1024);
-		TestTrue("FStringView::Right", OutofBoundsResult.IsEmpty());
+		TestTrue(TEXT("FStringView::Right"), OutofBoundsResult.IsEmpty());
 	}
 
 	{
@@ -554,10 +561,10 @@ bool FStringViewTestRightChop::RunTest(const FString& Parameters)
 	{
 		FStringView EmptyView(TEXT(""));
 		FStringView EmptyResult = EmptyView.RightChop(0);
-		TestTrue("FStringView::RightChop", EmptyResult.IsEmpty());
+		TestTrue(TEXT("FStringView::RightChop"), EmptyResult.IsEmpty());
 
 		FStringView OutofBoundsResult = EmptyView.RightChop(1024);
-		TestTrue("FStringView::RightChop", OutofBoundsResult.IsEmpty());
+		TestTrue(TEXT("FStringView::RightChop"), OutofBoundsResult.IsEmpty());
 	}
 
 	{
diff --git a/Engine/Source/Runtime/Core/Public/Containers/StringView.h b/Engine/Source/Runtime/Core/Public/Containers/StringView.h
--- a/Engine/Source/Runtime/Core/Public/Containers/StringView.h
+++ b/Engine/Source/Runtime/Core/Public/Containers/StringView.h
@@ -88,7 +88,7 @@ public:
 	 */
 	inline TStringViewImpl(const CharType* InData)
 		: DataPtr(InData)
-		, Size(TCString<CharType>::Strlen(InData))
+		, Size(InData ? TCString<CharType>::Strlen(InData) : 0)
 	{
 	}
 
