Fix FName batch serialization unit test and add max FName length check

diff --git a/Engine/Source/Runtime/Core/Private/UObject/UnrealNames.cpp b/Engine/Source/Runtime/Core/Private/UObject/UnrealNames.cpp
--- a/Engine/Source/Runtime/Core/Private/UObject/UnrealNames.cpp
+++ b/Engine/Source/Runtime/Core/Private/UObject/UnrealNames.cpp
@@ -1807,6 +1807,8 @@ struct FNameHelper
 
 	static FName Make(FNameStringView View, EFindName FindType, int32 InternalNumber)
 	{
+		checkf(View.Len < NAME_SIZE, TEXT("FName's %d max length exceeded. Got %d characters excluding null-terminator."), NAME_SIZE - 1, View.Len);
+
 		FNamePool& Pool = GetNamePool();
 
 		FNameEntryId DisplayId, ComparisonId;
@@ -2946,7 +2948,7 @@ static void TestNameBatch()
 	{
 		MaxLengthAnsi.Append("0123456789ABCDEF");
 	}
-	MaxLengthAnsi.RemoveAt(NAME_SIZE, MaxLengthAnsi.Len() % NAME_SIZE);
+	MaxLengthAnsi = MaxLengthAnsi.Left(NAME_SIZE - 1);
 
 	FString MaxLengthWide = MaxLengthAnsi;
 	MaxLengthWide[200] = 500;
