Iteration on StringBuilder

- Support string builders as contiguous containers.
- Add append operators for FAnsiStringBuilderBase.
- Add AddUninitialized and RemoveSuffix to facilitate complex append operations without a temporary buffer.
- Treat null C-style strings as empty for consistency with FString.

diff --git a/Engine/Source/Runtime/Core/Public/Misc/StringBuilder.h b/Engine/Source/Runtime/Core/Public/Misc/StringBuilder.h
--- a/Engine/Source/Runtime/Core/Public/Misc/StringBuilder.h
+++ b/Engine/Source/Runtime/Core/Public/Misc/StringBuilder.h
@@ -9,6 +9,7 @@
 #include "Templates/EnableIf.h"
 #include "Templates/IsArrayOrRefOfType.h"
 #include "Templates/IsValidVariadicFunctionArg.h"
+#include "Traits/IsContiguousContainer.h"
 
 #define USE_STRING_LITERAL_PATH 1
 
@@ -60,6 +61,8 @@ public:
 	const TStringBuilderImpl& operator=(TStringBuilderImpl&&) = delete;
 
 	inline int32	Len() const			{ return int32(CurPos - Base); }
+	inline C* GetData()					{ return Base; }
+	inline const C* GetData() const		{ return Base; }
 	inline const C* ToString() const	{ EnsureNulTerminated(); return Base; }
 	inline const C* operator*() const	{ EnsureNulTerminated(); return Base; }
 
@@ -73,6 +76,30 @@ public:
 		CurPos = Base;
 	}
 
+	/**
+	 * Adds a given number of uninitialized characters into the string builder.
+	 *
+	 * @param InCount The number of uninitialized characters to add.
+	 *
+	 * @return The number of characters in the string builder before adding new characters.
+	 */
+	inline int32 AddUninitialized(int32 InCount)
+	{
+		EnsureCapacity(InCount);
+		const int32 OldCount = Len();
+		CurPos += InCount;
+		return OldCount;
+	}
+
+	/**
+	 * Modifies the string builder to remove the given number of characters from the end.
+	 */
+	inline void RemoveSuffix(int32 InCount)
+	{
+		check(InCount <= Len());
+		CurPos -= InCount;
+	}
+
 	inline TStringBuilderImpl& Append(C Char)
 	{
 		EnsureCapacity(1);
@@ -88,6 +115,11 @@ public:
 	inline TStringBuilderImpl& AppendAnsi(const ANSICHAR* NulTerminatedString)
 #endif
 	{
+		if (!NulTerminatedString)
+		{
+			return *this;
+		}
+
 		const int32 Length = TCString<ANSICHAR>::Strlen(NulTerminatedString);
 
 		EnsureCapacity(Length);
@@ -129,6 +161,11 @@ public:
 	inline TStringBuilderImpl& Append(const C* NulTerminatedString)
 #endif
 	{
+		if (!NulTerminatedString)
+		{
+			return *this;
+		}
+
 		const int32 Length = TCString<C>::Strlen(NulTerminatedString);
 
 		EnsureCapacity(Length);
@@ -245,6 +282,15 @@ protected:
 	bool		bIsExtendable = false;
 };
 
+template <typename CharType>
+struct TIsContiguousContainer<TStringBuilderImpl<CharType>> { enum { Value = true }; };
+
+template <typename CharType>
+constexpr inline SIZE_T GetNum(const TStringBuilderImpl<CharType>& Builder)
+{
+	return Builder.Len();
+}
+
 //////////////////////////////////////////////////////////////////////////
 
 extern template class TStringBuilderImpl<ANSICHAR>;
@@ -263,6 +309,9 @@ protected:
 	~FAnsiStringBuilderBase() = default;
 };
 
+template <>
+struct TIsContiguousContainer<FAnsiStringBuilderBase> { enum { Value = true }; };
+
 template<int32 N>
 class TFixedAnsiStringBuilder : public FAnsiStringBuilderBase
 {
@@ -278,6 +327,9 @@ private:
 	ANSICHAR	StringBuffer[N];
 };
 
+template <int32 N>
+struct TIsContiguousContainer<TFixedAnsiStringBuilder<N>> { enum { Value = true }; };
+
 template<int32 N>
 class TAnsiStringBuilder : public FAnsiStringBuilderBase
 {
@@ -294,6 +346,9 @@ private:
 	ANSICHAR	StringBuffer[N];
 };
 
+template <int32 N>
+struct TIsContiguousContainer<TAnsiStringBuilder<N>> { enum { Value = true }; };
+
 //////////////////////////////////////////////////////////////////////////
 
 extern template class TStringBuilderImpl<TCHAR>;
@@ -328,6 +383,9 @@ protected:
 	~FStringBuilderBase() = default;
 };
 
+template <>
+struct TIsContiguousContainer<FStringBuilderBase> { enum { Value = true }; };
+
 template<int32 N>
 class TFixedStringBuilder : public FStringBuilderBase
 {
@@ -343,6 +401,9 @@ private:
 	TCHAR	StringBuffer[N];
 };
 
+template <int32 N>
+struct TIsContiguousContainer<TFixedStringBuilder<N>> { enum { Value = true }; };
+
 template<int32 N>
 class TStringBuilder : public FStringBuilderBase
 {
@@ -359,6 +420,9 @@ private:
 	TCHAR	StringBuffer[N];
 };
 
+template <int32 N>
+struct TIsContiguousContainer<TStringBuilder<N>> { enum { Value = true }; };
+
 // Append operator implementations
 
 inline FStringBuilderBase&					operator<<(FStringBuilderBase& Builder, const TCHAR Char)						{ Builder.Append(Char); return Builder; }
@@ -394,3 +458,33 @@ inline FStringBuilderBase&					operator<<(FStringBuilderBase& Builder, uint8 Val
 
 inline FStringBuilderBase&					operator<<(FStringBuilderBase& Builder, int16 Value)							{ return Builder << int32(Value); }
 inline FStringBuilderBase&					operator<<(FStringBuilderBase& Builder, uint16 Value)							{ return Builder << uint32(Value); }
+
+// Append operator implementations for FAnsiStringBuilderBase
+
+inline FAnsiStringBuilderBase&				operator<<(FAnsiStringBuilderBase& Builder, const ANSICHAR Char)				{ Builder.Append(Char); return Builder; }
+
+#if USE_STRING_LITERAL_PATH
+
+inline FAnsiStringBuilderBase&				operator<<(FAnsiStringBuilderBase& Builder, const ANSICHAR*& NulTerminatedString)	{ Builder.Append(NulTerminatedString); return Builder; }
+template<int32 N> FAnsiStringBuilderBase&	operator<<(FAnsiStringBuilderBase& Builder, const ANSICHAR (&StringLiteral)[N])		{ Builder.Append(StringLiteral); return Builder; }
+template<int32 N> FAnsiStringBuilderBase&	operator<<(FAnsiStringBuilderBase& Builder, ANSICHAR (&NulTerminatedString)[N])		{ Builder.Append(&NulTerminatedString[0], TCString<ANSICHAR>::Strlen(NulTerminatedString)); return Builder; }
+
+#else
+
+inline FAnsiStringBuilderBase&				operator<<(FAnsiStringBuilderBase& Builder, const ANSICHAR* NulTerminatedString)	{ Builder.Append(NulTerminatedString); return Builder; }
+
+#endif
+
+inline FAnsiStringBuilderBase&				operator<<(FAnsiStringBuilderBase& Builder, const FAnsiStringView& Str)				{ Builder.Append(Str.GetData(), Str.Len()); return Builder; }
+
+inline FAnsiStringBuilderBase&				operator<<(FAnsiStringBuilderBase& Builder, int32 Value)							{ Builder.Appendf("%d", Value); return Builder; }
+inline FAnsiStringBuilderBase&				operator<<(FAnsiStringBuilderBase& Builder, uint32 Value)							{ Builder.Appendf("%u", Value); return Builder; }
+
+inline FAnsiStringBuilderBase&				operator<<(FAnsiStringBuilderBase& Builder, int64 Value)							{ Builder.Appendf(INT64_FMT, Value); return Builder; }
+inline FAnsiStringBuilderBase&				operator<<(FAnsiStringBuilderBase& Builder, uint64 Value)							{ Builder.Appendf(UINT64_FMT, Value); return Builder; }
+
+inline FAnsiStringBuilderBase&				operator<<(FAnsiStringBuilderBase& Builder, int8 Value)								{ return Builder << int32(Value); }
+inline FAnsiStringBuilderBase&				operator<<(FAnsiStringBuilderBase& Builder, uint8 Value)							{ return Builder << uint32(Value); }
+
+inline FAnsiStringBuilderBase&				operator<<(FAnsiStringBuilderBase& Builder, int16 Value)							{ return Builder << int32(Value); }
+inline FAnsiStringBuilderBase&				operator<<(FAnsiStringBuilderBase& Builder, uint16 Value)							{ return Builder << uint32(Value); }
