Cleanup of IoDispatcher platform abstraction
Generic IoDispatcher implementation using Async IO
Made IoDispatcher less file centric
Implemented support for using IoDispatcher ReadOptions

diff --git a/Engine/Source/Runtime/Core/Private/GenericPlatform/GenericPlatformIoDispatcher.cpp b/Engine/Source/Runtime/Core/Private/GenericPlatform/GenericPlatformIoDispatcher.cpp
new file mode 100644
--- /dev/null
+++ b/Engine/Source/Runtime/Core/Private/GenericPlatform/GenericPlatformIoDispatcher.cpp
@@ -0,0 +1,116 @@
+// Copyright 1998-2019 Epic Games, Inc. All Rights Reserved.
+
+#include "GenericPlatformIoDispatcher.h"
+#include "IO/IoDispatcherFileBackend.h"
+#include "Async/AsyncFileHandle.h"
+#include "HAL/Event.h"
+
+FGenericIoDispatcherEventQueue::FGenericIoDispatcherEventQueue()
+	: Event(FPlatformProcess::GetSynchEventFromPool())
+{
+	
+}
+
+FGenericIoDispatcherEventQueue::~FGenericIoDispatcherEventQueue()
+{
+	FPlatformProcess::ReturnSynchEventToPool(Event);
+}
+
+void FGenericIoDispatcherEventQueue::Notify()
+{
+	Event->Trigger();
+}
+
+void FGenericIoDispatcherEventQueue::ProcessEvents()
+{
+	Event->Wait();
+}
+
+FGenericFileIoStoreImpl::FGenericFileIoStoreImpl(FGenericIoDispatcherEventQueue& InEventQueue)
+	: EventQueue(InEventQueue)
+{
+	
+}
+
+bool FGenericFileIoStoreImpl::OpenContainer(const TCHAR* ContainerFilePath, uint64& ContainerFileHandle, uint64& ContainerFileSize)
+{
+	IPlatformFile& Ipf = IPlatformFile::GetPlatformPhysical();
+	IAsyncReadFileHandle* FileHandle = Ipf.OpenAsyncRead(ContainerFilePath);
+	if (!FileHandle)
+	{
+		return false;
+	}
+	int64 SizeResult = -1;
+	FEvent* Event = FPlatformProcess::GetSynchEventFromPool();
+	FAsyncFileCallBack SizeCallback = [Event, &SizeResult](bool bWasCancelled, IAsyncReadRequest* Request)
+	{
+		check(!bWasCancelled);
+		SizeResult = Request->GetSizeResults();
+		Event->Trigger();
+	};
+	TUniquePtr<IAsyncReadRequest> SizeRequest(FileHandle->SizeRequest(&SizeCallback));
+	Event->Wait();
+	FPlatformProcess::ReturnSynchEventToPool(Event);
+	if (SizeResult > 0)
+	{
+		ContainerFileHandle = reinterpret_cast<UPTRINT>(FileHandle);
+		ContainerFileSize = SizeResult;
+		return true;
+	}
+	else
+	{
+		return false;
+	}
+}
+
+void FGenericFileIoStoreImpl::BeginReadsForRequest(FFileIoStoreResolvedRequest& ResolvedRequest)
+{
+	if (!ResolvedRequest.Request->IoBuffer.DataSize())
+	{
+		ResolvedRequest.Request->IoBuffer = FIoBuffer(ResolvedRequest.ResolvedSize);
+	}
+}
+
+void FGenericFileIoStoreImpl::ReadBlockFromFile(FFileIoStoreReadBlock* Block, uint8* Destination, uint64 FileHandle, uint64 Size, uint64 Offset)
+{
+	FAsyncFileCallBack ReadCallback = [this, Block](bool bWasCancelled, IAsyncReadRequest* Request)
+	{
+		check(!bWasCancelled);
+		{
+			FScopeLock _(&CompletedReadRequestsCritical);
+			CompletedReadRequests.AddTail(MakeTuple(Request, Block));
+		}
+		EventQueue.Notify();
+	};
+
+	IAsyncReadFileHandle* AsyncReadFileHandle = reinterpret_cast<IAsyncReadFileHandle*>(static_cast<UPTRINT>(FileHandle));
+	AsyncReadFileHandle->ReadRequest(Offset, Size, AIOP_Normal, &ReadCallback, Destination);
+}
+
+void FGenericFileIoStoreImpl::EndReadsForRequest()
+{
+
+}
+
+FFileIoStoreReadBlock* FGenericFileIoStoreImpl::GetNextCompletedBlock()
+{
+	IAsyncReadRequest* AsyncReadRequest;
+	FFileIoStoreReadBlock* CompletedBlock;
+	TDoubleLinkedList<TTuple<IAsyncReadRequest*, FFileIoStoreReadBlock*>>::TDoubleLinkedListNode* CompletedNode;
+	{
+		FScopeLock _(&CompletedReadRequestsCritical);
+		CompletedNode = CompletedReadRequests.GetHead();
+		if (!CompletedNode)
+		{
+			return nullptr;
+		}
+		TTuple<IAsyncReadRequest*, FFileIoStoreReadBlock*> NodeValue = CompletedNode->GetValue();
+		AsyncReadRequest = NodeValue.Get<0>();
+		CompletedBlock = NodeValue.Get<1>();
+		CompletedReadRequests.RemoveNode(CompletedNode);
+	}
+	AsyncReadRequest->WaitCompletion();
+	delete AsyncReadRequest;
+	return CompletedBlock;
+}
+
diff --git a/Engine/Source/Runtime/Core/Private/GenericPlatform/GenericPlatformIoDispatcher.h b/Engine/Source/Runtime/Core/Private/GenericPlatform/GenericPlatformIoDispatcher.h
new file mode 100644
--- /dev/null
+++ b/Engine/Source/Runtime/Core/Private/GenericPlatform/GenericPlatformIoDispatcher.h
@@ -0,0 +1,40 @@
+// Copyright 1998-2019 Epic Games, Inc. All Rights Reserved.
+
+#pragma once
+
+#include "Containers/List.h"
+#include "HAL/CriticalSection.h"
+#include "Templates/Tuple.h"
+
+struct FFileIoStoreReadBlock;
+struct FFileIoStoreResolvedRequest;
+class IAsyncReadRequest;
+
+class FGenericIoDispatcherEventQueue
+{
+public:
+	FGenericIoDispatcherEventQueue();
+	~FGenericIoDispatcherEventQueue();
+	void Notify();
+	void ProcessEvents();
+
+private:
+	FEvent* Event;
+};
+
+class FGenericFileIoStoreImpl
+{
+public:
+	FGenericFileIoStoreImpl(FGenericIoDispatcherEventQueue& InEventQueue);
+	bool OpenContainer(const TCHAR* ContainerFilePath, uint64& ContainerFileHandle, uint64& ContainerFileSize);
+	void BeginReadsForRequest(FFileIoStoreResolvedRequest& ResolvedRequest);
+	void ReadBlockFromFile(FFileIoStoreReadBlock* Block, uint8* Destination, uint64 FileHandle, uint64 Size, uint64 Offset);
+	void EndReadsForRequest();
+	FFileIoStoreReadBlock* GetNextCompletedBlock();
+
+private:
+	FGenericIoDispatcherEventQueue& EventQueue;
+	FCriticalSection CompletedReadRequestsCritical;
+	TDoubleLinkedList<TTuple<IAsyncReadRequest*, FFileIoStoreReadBlock*>> CompletedReadRequests;
+};
+
diff --git a/Engine/Source/Runtime/Core/Private/IO/IoDispatcher.cpp b/Engine/Source/Runtime/Core/Private/IO/IoDispatcher.cpp
--- a/Engine/Source/Runtime/Core/Private/IO/IoDispatcher.cpp
+++ b/Engine/Source/Runtime/Core/Private/IO/IoDispatcher.cpp
@@ -1,45 +1,26 @@
 // Copyright 1998-2019 Epic Games, Inc. All Rights Reserved.
 
 #include "IO/IoDispatcher.h"
+#include "IO/IoDispatcherPrivate.h"
 #include "IO/IoStore.h"
-#include "Async/MappedFileHandle.h"
-#include "HAL/CriticalSection.h"
-#include "HAL/FileManager.h"
-#include "HAL/PlatformFilemanager.h"
-#include "HAL/Runnable.h"
-#include "HAL/RunnableThread.h"
-#include "Misc/CString.h"
-#include "Misc/EventPool.h"
+#include "IO/IoDispatcherFileBackend.h"
 #include "Misc/ScopeRWLock.h"
-#include "Misc/ScopeLock.h"
-#include "Misc/SecureHash.h"
-#include "Misc/StringBuilder.h"
-#include "Misc/LazySingleton.h"
 #include "Misc/CoreDelegates.h"
-#include "Trace/Trace.h"
-
-#define IODISPATCHER_TRACE_ENABLED !UE_BUILD_SHIPPING
+#include "Math/RandomStream.h"
+#include "Misc/ScopeLock.h"
+#include "HAL/Runnable.h"
+#include "HAL/RunnableThread.h"
+#include "HAL/PlatformProcess.h"
 
 DEFINE_LOG_CATEGORY(LogIoDispatcher);
 
 const FIoChunkId FIoChunkId::InvalidChunkId = FIoChunkId::CreateEmptyId();
 
-#if !defined(PLATFORM_IMPLEMENTS_IO)
-//////////////////////////////////////////////////////////////////////////
-
 TUniquePtr<FIoDispatcher> GIoDispatcher;
-FIoStoreEnvironment GIoStoreEnvironment;
-
-UE_TRACE_EVENT_BEGIN(IoDispatcher, BatchIssued, Always)
-	UE_TRACE_EVENT_FIELD(uint64, Cycle)
-	UE_TRACE_EVENT_FIELD(uint64, BatchId)
-UE_TRACE_EVENT_END()
 
-UE_TRACE_EVENT_BEGIN(IoDispatcher, BatchResolved, Always)
-	UE_TRACE_EVENT_FIELD(uint64, Cycle)
-	UE_TRACE_EVENT_FIELD(uint64, BatchId)
-	UE_TRACE_EVENT_FIELD(uint64, TotalSize)
-UE_TRACE_EVENT_END()
+#if UE_BUILD_DEVELOPMENT || UE_BUILD_DEBUG
+//PRAGMA_DISABLE_OPTIMIZATION
+#endif
 
 template <typename T, uint32 BlockSize = 128>
 class TBlockAllocator
@@ -58,21 +39,21 @@ public:
 		{
 			//TODO: Virtual alloc
 			FBlock* Block = new FBlock;
-			
+
 			for (int32 ElementIndex = 0; ElementIndex < BlockSize; ++ElementIndex)
 			{
-				FElement* Element	= &Block->Elements[ElementIndex];
-				Element->Next		= NextFree;
-				NextFree			= Element;
+				FElement* Element = &Block->Elements[ElementIndex];
+				Element->Next = NextFree;
+				NextFree = Element;
 			}
 
-			Block->Next	= Blocks;
-			Blocks		= Block;
+			Block->Next = Blocks;
+			Blocks = Block;
 		}
 
-		FElement* Element	= NextFree;
-		NextFree			= Element->Next;
-		
+		FElement* Element = NextFree;
+		NextFree = Element->Next;
+
 		++NumElements;
 
 		return Element->Buffer.GetTypedPtr();
@@ -82,9 +63,9 @@ public:
 	{
 		FScopeLock _(&CriticalSection);
 
-		FElement* Element	= reinterpret_cast<FElement*>(Ptr);
-		Element->Next		= NextFree;
-		NextFree			= Element;
+		FElement* Element = reinterpret_cast<FElement*>(Ptr);
+		Element->Next = NextFree;
+		NextFree = Element;
 
 		--NumElements;
 	}
@@ -121,8 +102,8 @@ private:
 			delete Tmp;
 		}
 
-		Blocks		= nullptr;
-		NextFree	= nullptr;
+		Blocks = nullptr;
+		NextFree = nullptr;
 		NumElements = 0;
 	}
 
@@ -138,291 +119,45 @@ private:
 		FBlock* Next = nullptr;
 	};
 
-	FBlock*				Blocks		= nullptr;
-	FElement*			NextFree	= nullptr;
+	FBlock*				Blocks = nullptr;
+	FElement*			NextFree = nullptr;
 	int32				NumElements = 0;
 	FCriticalSection	CriticalSection;
 };
 
-//////////////////////////////////////////////////////////////////////////
-
-class FIoStoreReaderImpl
-{
-public:
-	FIoStoreReaderImpl(FIoStoreEnvironment& InEnvironment)
-		: Environment(InEnvironment)
-	{
-	}
-
-	FIoStatus Initialize(FStringView UniqueId);
-
-	TIoStatusOr<FIoBuffer> Lookup(const FIoChunkId& ChunkId, FIoReadOptions Options)
-	{
-		if (const FIoStoreTocEntry* Entry = Toc.Find(ChunkId))
-		{
-			return FIoBuffer(FIoBuffer::Wrap, MappedRegion->GetMappedPtr() + Entry->GetOffset(), Entry->GetLength());
-		}
-
-		return FIoStatus(EIoErrorCode::NotFound);
-	}
-
-	bool DoesChunkExist(const FIoChunkId& ChunkId) const
-	{
-		return Toc.Find(ChunkId) != nullptr;
-	}
-
-	TIoStatusOr<uint64> GetSizeForChunk(const FIoChunkId& ChunkId) const
-	{
-		const FIoStoreTocEntry* Entry = Toc.Find(ChunkId);
-
-		if (Entry != nullptr)
-		{
-			return Entry->GetLength();
-		}
-		else
-		{
-			return FIoStatus(EIoErrorCode::NotFound);
-		}
-	}
-
-private:
-	FIoStoreEnvironment&				Environment;
-	FString								UniqueId;
-	TMap<FIoChunkId, FIoStoreTocEntry>	Toc;
-	TUniquePtr<IFileHandle>				ContainerFileHandle;
-	TUniquePtr<IMappedFileHandle>		ContainerMappedFileHandle;
-	TUniquePtr<IMappedFileRegion>		MappedRegion;
-};
-
-FIoStatus FIoStoreReaderImpl::Initialize(FStringView InUniqueId)
-{
-	IPlatformFile& Ipf = FPlatformFileManager::Get().GetPlatformFile();
-
-	UniqueId = InUniqueId;
-
-	const FString& RootPath = Environment.GetRootPath();
-
-	TStringBuilder<256> ContainerFilePath;
-	ContainerFilePath.Append(RootPath);
-	if (ContainerFilePath.LastChar() != '/')
-		ContainerFilePath.Append(TEXT('/'));
-
-	TStringBuilder<256> TocFilePath;
-	TocFilePath.Append(ContainerFilePath);
-	TocFilePath.Append(TEXT("Container.utoc"));
-
-	ContainerFilePath.Append(TEXT("Container.ucas"));
-	ContainerFileHandle.Reset(Ipf.OpenRead(*ContainerFilePath, /* allowwrite */ false));
-	
-	if (!ContainerFileHandle)
-	{
-		return FIoStatusBuilder(EIoErrorCode::FileOpenFailed) << TEXT("Failed to open IoStore container file '") << *ContainerFilePath << TEXT("'");
-	}
-
-	const uint64 ContainerSize = ContainerFileHandle->Size();
-
-	ContainerMappedFileHandle.Reset(Ipf.OpenMapped(*ContainerFilePath));
-	MappedRegion.Reset(ContainerMappedFileHandle->MapRegion());
-
-	if (!ContainerMappedFileHandle)
-	{
-		return FIoStatusBuilder(EIoErrorCode::FileOpenFailed) << TEXT("Failed to memory map IoStore container file '") << *ContainerFilePath << TEXT("'");
-	}
-
-	TUniquePtr<uint8[]> TocBuffer;
-	bool bTocReadOk = false;
-
-	{
-		TUniquePtr<IFileHandle>	TocFileHandle(Ipf.OpenRead(*TocFilePath, /* allowwrite */ false));
-
-		if (!TocFileHandle)
-		{
-			return FIoStatusBuilder(EIoErrorCode::FileOpenFailed) << TEXT("Failed to open IoStore TOC file '") << *TocFilePath << TEXT("'");
-		}
-
-		const int64 TocSize = TocFileHandle->Size();
-		TocBuffer = MakeUnique<uint8[]>(TocSize);
-		bTocReadOk = TocFileHandle->Read(TocBuffer.Get(), TocSize);
-	}
-
-	if (!bTocReadOk)
-	{
-		return FIoStatusBuilder(EIoErrorCode::CorruptToc) << TEXT("Failed to read IoStore TOC file '") << *TocFilePath << TEXT("'");
-	}
-
-	const FIoStoreTocHeader* Header = reinterpret_cast<const FIoStoreTocHeader*>(TocBuffer.Get());
-
-	if (!Header->CheckMagic())
-	{
-		return FIoStatusBuilder(EIoErrorCode::CorruptToc) << TEXT("TOC header magic mismatch while reading '") << *TocFilePath << TEXT("'");
-	}
-
-	if (Header->TocHeaderSize != sizeof(FIoStoreTocHeader))
-	{
-		return FIoStatusBuilder(EIoErrorCode::CorruptToc) << TEXT("TOC header size mismatch while reading '") << *TocFilePath << TEXT("'");
-	}
-
-	if (Header->TocEntrySize != sizeof(FIoStoreTocEntry))
-	{
-		return FIoStatusBuilder(EIoErrorCode::CorruptToc) << TEXT("TOC entry size mismatch while reading '") << *TocFilePath << TEXT("'");
-	}
-
-	const FIoStoreTocEntry* Entry = reinterpret_cast<const FIoStoreTocEntry*>(TocBuffer.Get() + sizeof(FIoStoreTocHeader));
-	uint32 EntryCount = Header->TocEntryCount;
-
-	Toc.Reserve(EntryCount);
-
-	while(EntryCount--)
-	{
-		if ((Entry->GetOffset() + Entry->GetLength()) > ContainerSize)
-		{
-			// TODO: add details
-			return FIoStatusBuilder(EIoErrorCode::CorruptToc) << TEXT("TOC entry out of container bounds while reading '") << *TocFilePath << TEXT("'");
-		}
-
-		Toc.Add(Entry->ChunkId, *Entry);
-
-		++Entry;
-	}
-
-	return FIoStatus::Ok;
-}
-
-FIoStoreReader::FIoStoreReader(FIoStoreEnvironment& Environment)
-:	Impl(new FIoStoreReaderImpl(Environment))
-{
-}
-
-FIoStoreReader::~FIoStoreReader()
-{
-	delete Impl;
-}
-
-FIoStatus 
-FIoStoreReader::Initialize(FStringView UniqueId)
-{
-	return Impl->Initialize(UniqueId);
-}
-
-//////////////////////////////////////////////////////////////////////////
-
-class FIoBatchImpl
-{
-public:
-	FIoRequestImpl* FirstRequest	= nullptr;
-	FIoBatchImpl*	NextBatch		= nullptr;
-	TAtomic<uint32>	OutstandingRequests { 0 };
-};
-
-class FIoRequestImpl
-{
-public:
-	FIoChunkId				ChunkId;
-	FIoReadOptions			Options;
-	TIoStatusOr<FIoBuffer>	Result;
-	FIoRequestImpl*			NextRequest = nullptr;
-};
-
-//////////////////////////////////////////////////////////////////////////
-
-class FIoStoreImpl : public FRefCountBase
-{
-public:
-	void Mount(FIoStoreReader* IoStore)
-	{
-		FWriteScopeLock _(RwLockIoStores);
-		// TODO prevent duplicates?
-		IoStores.Add(IoStore);
-	}
-
-	void Unmount(FIoStoreReader* IoStore)
-	{
-		FWriteScopeLock _(RwLockIoStores);
-		IoStores.Remove(IoStore);
-	}
-
-	TIoStatusOr<FIoBuffer> Resolve(FIoChunkId& ChunkId, FIoReadOptions Options)
-	{
-		for (const FIoStoreReader* IoStore : IoStores)
-		{
-			TIoStatusOr<FIoBuffer> Result = IoStore->Impl->Lookup(ChunkId, Options);
-
-			if (Result.IsOk())
-			{
-				return Result;
-			}
-		}
-
-		return FIoStatus(EIoErrorCode::NotFound);
-	}
-
-	bool DoesChunkExist(const FIoChunkId& ChunkId) const
-	{
-		for (const FIoStoreReader* IoStore : IoStores)
-		{
-			if (IoStore->Impl->DoesChunkExist(ChunkId))
-			{
-				return true;
-			}
-		}
-
-		return false;
-	}
-
-	TIoStatusOr<uint64> GetSizeForChunk(const FIoChunkId& ChunkId) const
-	{
-		for (const FIoStoreReader* IoStore : IoStores)
-		{
-			TIoStatusOr<uint64> Result = IoStore->Impl->GetSizeForChunk(ChunkId);
-
-			if (Result.IsOk())
-			{
-				return Result;
-			}
-		}
-
-		return FIoStatus(EIoErrorCode::NotFound);
-	}
-
-private:
-	FRWLock									RwLockIoStores;
-	TArray<TRefCountPtr<FIoStoreReader>>	IoStores;
-};
-
-//////////////////////////////////////////////////////////////////////////
-
 class FIoDispatcherImpl
+	: public FRunnable
 {
 public:
 	FIoDispatcherImpl()
+		: FileIoStore(EventQueue)
 	{
-		IoStore = new FIoStoreImpl;
-
 		FCoreDelegates::GetMemoryTrimDelegate().AddLambda([this]()
 		{
 			RequestAllocator.Trim();
 			BatchAllocator.Trim();
 		});
+
+		Thread = FRunnableThread::Create(this, TEXT("IoDispatcher"), 0, TPri_AboveNormal);
 	}
 
 	FIoRequestImpl* AllocRequest(const FIoChunkId& ChunkId, FIoReadOptions Options)
 	{
 		FIoRequestImpl* Request = RequestAllocator.Construct();
 
-		Request->ChunkId		= ChunkId;
-		Request->Options		= Options;
-		Request->Result			= FIoStatus::Unknown;
-		Request->NextRequest	= nullptr;
+		Request->ChunkId = ChunkId;
+		Request->Options = Options;
+		Request->Status = FIoStatus::Unknown;
 
 		return Request;
 	}
 
 	FIoRequestImpl* AllocRequest(FIoBatchImpl* Batch, const FIoChunkId& ChunkId, FIoReadOptions Options)
 	{
-		FIoRequestImpl* Request	= AllocRequest(ChunkId, Options);
-
-		Request->NextRequest	= Batch->FirstRequest;
-		Batch->FirstRequest		= Request;
+		FIoRequestImpl* Request = AllocRequest(ChunkId, Options);
+		Request->Batch = Batch;
+		Request->BatchNextRequest = Batch->FirstRequest;
+		Batch->FirstRequest = Request;
 
 		return Request;
 	}
@@ -432,13 +167,10 @@ public:
 		RequestAllocator.Destroy(Request);
 	}
 
-	FIoBatchImpl* AllocBatch(FIoRequestImpl* FirstRequest = nullptr)
+	FIoBatchImpl* AllocBatch()
 	{
-		FIoBatchImpl* Batch			= BatchAllocator.Construct();
+		FIoBatchImpl* Batch = BatchAllocator.Construct();
 
-		Batch->FirstRequest			= FirstRequest;
-		Batch->OutstandingRequests	= 0;
-	
 		return Batch;
 	}
 
@@ -448,8 +180,8 @@ public:
 
 		while (Request)
 		{
-			FIoRequestImpl* Tmp = Request;	
-			Request = Request->NextRequest;
+			FIoRequestImpl* Tmp = Request;
+			Request = Request->BatchNextRequest;
 			FreeRequest(Tmp);
 		}
 
@@ -459,26 +191,36 @@ public:
 	void ReadWithCallback(const FIoChunkId& Chunk, const FIoReadOptions& Options, TFunction<void(TIoStatusOr<FIoBuffer>)>&& Callback)
 	{
 		FIoRequestImpl* Request = AllocRequest(Chunk, Options);
-		Request->Result = IoStore->Resolve(Request->ChunkId, Request->Options);
-		if (Request->Result.IsOk())
+		Request->Callback = MoveTemp(Callback);
 		{
-			Callback(Request->Result.ConsumeValueOrDie());
-		}
-		else
-		{
-			Callback(TIoStatusOr<FIoBuffer>(Request->Result.Status()));
+			FScopeLock _(&WaitingLock);
+			Request->NextRequest = nullptr;
+			if (!WaitingRequestsTail)
+			{
+				WaitingRequestsHead = WaitingRequestsTail = Request;
+			}
+			else
+			{
+				WaitingRequestsTail->NextRequest = Request;
+				WaitingRequestsTail = Request;
+			}
 		}
-		FreeRequest(Request);
+		EventQueue.Notify();
+	}
+
+	FIoStatus Mount(FIoStoreEnvironment& Environment)
+	{
+		return FileIoStore.Mount(*Environment.GetRootPath());
 	}
 
 	bool DoesChunkExist(const FIoChunkId& ChunkId) const
 	{
-		return IoStore->DoesChunkExist(ChunkId);
+		return FileIoStore.DoesChunkExist(ChunkId);
 	}
 
 	TIoStatusOr<uint64> GetSizeForChunk(const FIoChunkId& ChunkId) const
 	{
-		return IoStore->GetSizeForChunk(ChunkId);
+		return FileIoStore.GetSizeForChunk(ChunkId);
 	}
 
 	template<typename Func>
@@ -490,75 +232,169 @@ public:
 		{
 			const bool bDoContinue = InCallbackFunction(Request);
 
-			Request = bDoContinue ? Request->NextRequest : nullptr;
+			Request = bDoContinue ? Request->BatchNextRequest : nullptr;
 		}
 	}
 
 	void IssueBatch(const FIoBatchImpl* Batch)
 	{
-		// At this point the batch is immutable and we should start
-		// doing the work.
-
-#if IODISPATCHER_TRACE_ENABLED
-		UE_TRACE_LOG(IoDispatcher, BatchIssued)
-			<< BatchIssued.Cycle(FPlatformTime::Cycles64())
-			<< BatchIssued.BatchId(uint64(Batch));
-#endif
-		uint64 TotalBatchSize = 0;
-		IterateBatch(Batch, [this, &TotalBatchSize](FIoRequestImpl* Request)
 		{
-			Request->Result = IoStore->Resolve(Request->ChunkId, Request->Options);
+			FScopeLock _(&WaitingLock);
+			if (!WaitingRequestsHead)
+			{
+				WaitingRequestsHead = Batch->FirstRequest;
+			}
+			else
+			{
+				WaitingRequestsTail->NextRequest = Batch->FirstRequest;
+			}
+			WaitingRequestsTail = Batch->FirstRequest;
+			while (WaitingRequestsTail->BatchNextRequest)
+			{
+				WaitingRequestsTail->NextRequest = WaitingRequestsTail->BatchNextRequest;
+				WaitingRequestsTail = WaitingRequestsTail->BatchNextRequest;
+			}
+			WaitingRequestsTail->NextRequest = nullptr;
+		}
+		EventQueue.Notify();
+	}
 
-#if IODISPATCHER_TRACE_ENABLED
-			TotalBatchSize += Request->Result.ValueOrDie().DataSize();
-#endif
-			return true;
-		});
+private:
+	friend class FIoBatch;
 
-#if IODISPATCHER_TRACE_ENABLED
-		UE_TRACE_LOG(IoDispatcher, BatchResolved)
-			<< BatchResolved.Cycle(FPlatformTime::Cycles64())
-			<< BatchResolved.BatchId(uint64(Batch))
-			<< BatchResolved.TotalSize(TotalBatchSize);
-#endif
+	void CompleteRequests()
+	{
+		while (SubmittedRequestsHead && SubmittedRequestsHead->UnfinishedReadsCount == 0)
+		{
+			FIoRequestImpl* NextRequest = SubmittedRequestsHead->NextRequest;
+			CompleteRequest(SubmittedRequestsHead);
+			SubmittedRequestsHead = NextRequest;
+		}
+		if (!SubmittedRequestsHead)
+		{
+			SubmittedRequestsTail = nullptr;
+		}
+	}
+
+	void CompleteRequest(FIoRequestImpl* Request)
+	{
+		if (!Request->Status.IsCompleted())
+		{
+			Request->Status = EIoErrorCode::Ok;
+			if (Request->Callback)
+			{
+				Request->Callback(Request->IoBuffer);
+			}
+		}
+		else
+		{
+			if (Request->Callback)
+			{
+				Request->Callback(Request->Status);
+			}
+		}
+		if (!Request->Batch)
+		{
+			FreeRequest(Request);
+		}
 	}
 
-	bool IsBatchReady(const FIoBatchImpl* Batch)
+	void ProcessIncomingRequests()
 	{
-		bool bIsReady = true;
+		FIoRequestImpl* WaitingRequest;
+		{
+			FScopeLock _(&WaitingLock);
+			WaitingRequest = WaitingRequestsHead;
+			WaitingRequestsHead = WaitingRequestsTail = nullptr;
+		}
+		while (WaitingRequest)
+		{
+			RequestsToSubmit.Add(WaitingRequest);
+			WaitingRequest = WaitingRequest->NextRequest;
+		}
 
-		IterateBatch(Batch, [&bIsReady](FIoRequestImpl* Request)
+		int32 RequestsToSubmitCount = RequestsToSubmit.Num();
+		if (!RequestsToSubmitCount)
 		{
-			bIsReady &= Request->Result.Status().IsCompleted();
-			return bIsReady;
-		});
+			return;
+		}
+
+		//TRACE_CPUPROFILER_EVENT_SCOPE(IoDispatcherSubmitRequests);
+
+		for (; CurrentRequestsToSubmitIndex < RequestsToSubmitCount; ++CurrentRequestsToSubmitIndex)
+		{
+			FIoRequestImpl* Request = RequestsToSubmit[CurrentRequestsToSubmitIndex];
 
-		return bIsReady;
+			EIoStoreResolveResult Result = FileIoStore.Resolve(Request);
+			if (Result == IoStoreResolveResult_Stalled)
+			{
+				break;
+			}
+			if (Result == IoStoreResolveResult_NotFound)
+			{
+				Request->Status = FIoStatus(EIoErrorCode::NotFound);
+			}
+			if (!SubmittedRequestsTail)
+			{
+				SubmittedRequestsHead = SubmittedRequestsTail = Request;
+			}
+			else
+			{
+				SubmittedRequestsTail->NextRequest = Request;
+				SubmittedRequestsTail = Request;
+			}
+			Request->NextRequest = nullptr;
+		}
+
+		if (CurrentRequestsToSubmitIndex == RequestsToSubmitCount)
+		{
+			RequestsToSubmit.Reset(false);
+			CurrentRequestsToSubmitIndex = 0;
+		}
 	}
 
-	void Mount(FIoStoreReader* IoStoreReader)
+	virtual bool Init()
 	{
-		IoStore->Mount(IoStoreReader);
+		return true;
 	}
 
-	void Unmount(FIoStoreReader* IoStoreReader)
+	virtual uint32 Run()
 	{
-		IoStore->Unmount(IoStoreReader);
+		while (true)
+		{
+			EventQueue.ProcessEvents();
+
+			FileIoStore.ProcessIncomingBlocks();
+			ProcessIncomingRequests();
+			CompleteRequests();
+		}
+		return 0;
 	}
 
-private:
-	using FRequestAllocator		= TBlockAllocator<FIoRequestImpl, 4096>;
-	using FBatchAllocator		= TBlockAllocator<FIoBatchImpl, 4096>;
+	virtual void Stop()
+	{
+	}
 
-	TRefCountPtr<FIoStoreImpl>	IoStore;
-	FRequestAllocator			RequestAllocator;
-	FBatchAllocator				BatchAllocator;
-};
+	using FRequestAllocator = TBlockAllocator<FIoRequestImpl, 4096>;
+	using FBatchAllocator = TBlockAllocator<FIoBatchImpl, 4096>;
 
-//////////////////////////////////////////////////////////////////////////
+	FIoDispatcherEventQueue EventQueue;
+
+	FFileIoStore FileIoStore;
+	FRequestAllocator RequestAllocator;
+	FBatchAllocator BatchAllocator;
+	FRunnableThread* Thread = nullptr;
+	FCriticalSection WaitingLock;
+	FIoRequestImpl* WaitingRequestsHead = nullptr;
+	FIoRequestImpl* WaitingRequestsTail = nullptr;
+	TArray<FIoRequestImpl*> RequestsToSubmit;
+	int32 CurrentRequestsToSubmitIndex = 0;
+	FIoRequestImpl* SubmittedRequestsHead = nullptr;
+	FIoRequestImpl* SubmittedRequestsTail = nullptr;
+};
 
 FIoDispatcher::FIoDispatcher()
-:	Impl(new FIoDispatcherImpl())
+	: Impl(new FIoDispatcherImpl())
 {
 }
 
@@ -567,16 +403,9 @@ FIoDispatcher::~FIoDispatcher()
 	delete Impl;
 }
 
-void		
-FIoDispatcher::Mount(FIoStoreReader* IoStore)
-{
-	Impl->Mount(IoStore);
-}
-
-void		
-FIoDispatcher::Unmount(FIoStoreReader* IoStore)
+FIoStatus FIoDispatcher::Mount(FIoStoreEnvironment& Environment)
 {
-	Impl->Unmount(IoStore);
+	return Impl->Mount(Environment);
 }
 
 FIoBatch
@@ -598,28 +427,21 @@ FIoDispatcher::ReadWithCallback(const FIoChunkId& Chunk, const FIoReadOptions& O
 }
 
 // Polling methods
-bool 
+bool
 FIoDispatcher::DoesChunkExist(const FIoChunkId& ChunkId) const
 {
 	return Impl->DoesChunkExist(ChunkId);
 }
 
-TIoStatusOr<uint64>	
+TIoStatusOr<uint64>
 FIoDispatcher::GetSizeForChunk(const FIoChunkId& ChunkId) const
 {
 	return Impl->GetSizeForChunk(ChunkId);
 }
 
-bool
-FIoDispatcher::IsInitialized()
-{
-	return GIoDispatcher.IsValid();
-}
-
 FIoStatus
-FIoDispatcher::Initialize(const FString& Directory)
+FIoDispatcher::Initialize()
 {
-	GIoStoreEnvironment.InitializeFileEnvironment(Directory);
 	GIoDispatcher = MakeUnique<FIoDispatcher>();
 
 	return EIoErrorCode::Ok;
@@ -637,18 +459,10 @@ FIoDispatcher::Get()
 	return *GIoDispatcher;
 }
 
-FIoStoreEnvironment&
-FIoDispatcher::GetEnvironment()
-{
-	return GIoStoreEnvironment;
-}
-
-//////////////////////////////////////////////////////////////////////////
-
 FIoBatch::FIoBatch(FIoDispatcherImpl* InDispatcher, FIoBatchImpl* InImpl)
-:	Dispatcher(InDispatcher)
-,	Impl(InImpl)
-,	CompletionEvent()
+	: Dispatcher(InDispatcher)
+	, Impl(InImpl)
+	, CompletionEvent()
 {
 }
 
@@ -658,7 +472,7 @@ FIoBatch::Read(const FIoChunkId& ChunkId, FIoReadOptions Options)
 	return FIoRequest(Dispatcher->AllocRequest(Impl, ChunkId, Options));
 }
 
-void 
+void
 FIoBatch::ForEachRequest(TFunction<bool(FIoRequest&)>&& Callback)
 {
 	Dispatcher->IterateBatch(Impl, [&](FIoRequestImpl* InRequest) {
@@ -667,26 +481,25 @@ FIoBatch::ForEachRequest(TFunction<bool(FIoRequest&)>&& Callback)
 	});
 }
 
-void 
+void
 FIoBatch::Issue()
 {
 	Dispatcher->IssueBatch(Impl);
 }
 
-void 
+void
 FIoBatch::Wait()
 {
-	// Temporary spinning implementation
-	for (FIoRequestImpl* Request = Impl->FirstRequest; Request; Request = Request->NextRequest)
+	for (FIoRequestImpl* Request = Impl->FirstRequest; Request; Request = Request->BatchNextRequest)
 	{
-		if (!Request->Result.Status().IsCompleted())
+		while (!Request->Status.IsCompleted())
 		{
 			FPlatformProcess::Sleep(0);
 		}
 	}
 }
 
-void 
+void
 FIoBatch::Cancel()
 {
 	unimplemented();
@@ -694,16 +507,16 @@ FIoBatch::Cancel()
 
 //////////////////////////////////////////////////////////////////////////
 
-bool		
-FIoRequest::IsOk() const 
-{ 
-	return Impl->Result.IsOk(); 
+bool
+FIoRequest::IsOk() const
+{
+	return Impl->Status.IsOk();
 }
 
-FIoStatus	
+FIoStatus
 FIoRequest::Status() const
-{ 
-	return Impl->Result.Status();
+{
+	return Impl->Status;
 }
 
 const FIoChunkId&
@@ -715,7 +528,12 @@ FIoRequest::GetChunkId() const
 TIoStatusOr<FIoBuffer>
 FIoRequest::GetResult() const
 {
-	return Impl->Result;
+	if (Impl->Status.IsOk())
+	{
+		return Impl->IoBuffer;
+	}
+	else
+	{
+		return Impl->Status;
+	}
 }
-
-#endif // PLATFORM_IMPLEMENTS_IO
diff --git a/Engine/Source/Runtime/Core/Private/IO/IoDispatcherFileBackend.cpp b/Engine/Source/Runtime/Core/Private/IO/IoDispatcherFileBackend.cpp
new file mode 100644
--- /dev/null
+++ b/Engine/Source/Runtime/Core/Private/IO/IoDispatcherFileBackend.cpp
@@ -0,0 +1,345 @@
+// Copyright 1998-2019 Epic Games, Inc. All Rights Reserved.
+
+#include "IO/IoDispatcherFileBackend.h"
+#include "Misc/ScopeRWLock.h"
+#include "ProfilingDebugging/CountersTrace.h"
+#include "HAL/PlatformFile.h"
+
+TRACE_DECLARE_INT_COUNTER(IoDispatcherCacheHitsCold, TEXT("IoDispatcher/CacheHitsCold"));
+TRACE_DECLARE_INT_COUNTER(IoDispatcherCacheHitsHot, TEXT("IoDispatcher/CacheHitsHot"));
+TRACE_DECLARE_INT_COUNTER(IoDispatcherCacheMisses, TEXT("IoDispatcher/CacheMisses"));
+
+FFileIoStoreReader::FFileIoStoreReader(FFileIoStoreImpl& InPlatformImpl)
+	: PlatformImpl(InPlatformImpl)
+{
+}
+
+FIoStatus FFileIoStoreReader::Initialize(const TCHAR* RootPath)
+{
+	IPlatformFile& Ipf = IPlatformFile::GetPlatformPhysical();
+
+	TStringBuilder<256> ContainerFilePath;
+	ContainerFilePath.Append(RootPath);
+	if (ContainerFilePath.LastChar() != '/')
+	{
+		ContainerFilePath.Append(TEXT('/'));
+	}
+
+	TStringBuilder<256> TocFilePath;
+	TocFilePath.Append(ContainerFilePath);
+	TocFilePath.Append(TEXT("Container.utoc"));
+
+	ContainerFilePath.Append(TEXT("Container.ucas"));
+
+	if (!PlatformImpl.OpenContainer(*ContainerFilePath, ContainerFileHandle, ContainerFileSize))
+	{
+		return FIoStatusBuilder(EIoErrorCode::FileOpenFailed) << TEXT("Failed to open IoStore container file '") << *ContainerFilePath << TEXT("'");
+	}
+
+	TUniquePtr<uint8[]> TocBuffer;
+	bool bTocReadOk = false;
+
+	{
+		TUniquePtr<IFileHandle>	TocFileHandle(Ipf.OpenRead(*TocFilePath, /* allowwrite */ false));
+
+		if (!TocFileHandle)
+		{
+			return FIoStatusBuilder(EIoErrorCode::FileOpenFailed) << TEXT("Failed to open IoStore TOC file '") << *TocFilePath << TEXT("'");
+		}
+
+		const int64 TocSize = TocFileHandle->Size();
+		TocBuffer = MakeUnique<uint8[]>(TocSize);
+		bTocReadOk = TocFileHandle->Read(TocBuffer.Get(), TocSize);
+	}
+
+	if (!bTocReadOk)
+	{
+		return FIoStatusBuilder(EIoErrorCode::CorruptToc) << TEXT("Failed to read IoStore TOC file '") << *TocFilePath << TEXT("'");
+	}
+
+	const FIoStoreTocHeader* Header = reinterpret_cast<const FIoStoreTocHeader*>(TocBuffer.Get());
+
+	if (!Header->CheckMagic())
+	{
+		return FIoStatusBuilder(EIoErrorCode::CorruptToc) << TEXT("TOC header magic mismatch while reading '") << *TocFilePath << TEXT("'");
+	}
+
+	if (Header->TocHeaderSize != sizeof(FIoStoreTocHeader))
+	{
+		return FIoStatusBuilder(EIoErrorCode::CorruptToc) << TEXT("TOC header size mismatch while reading '") << *TocFilePath << TEXT("'");
+	}
+
+	if (Header->TocEntrySize != sizeof(FIoStoreTocEntry))
+	{
+		return FIoStatusBuilder(EIoErrorCode::CorruptToc) << TEXT("TOC entry size mismatch while reading '") << *TocFilePath << TEXT("'");
+	}
+
+	const FIoStoreTocEntry* Entry = reinterpret_cast<const FIoStoreTocEntry*>(TocBuffer.Get() + sizeof(FIoStoreTocHeader));
+	uint32 EntryCount = Header->TocEntryCount;
+
+	Toc.Reserve(EntryCount);
+	while (EntryCount--)
+	{
+		if ((Entry->GetOffset() + Entry->GetLength()) > ContainerFileSize)
+		{
+			// TODO: add details
+			return FIoStatusBuilder(EIoErrorCode::CorruptToc) << TEXT("TOC entry out of container bounds while reading '") << *TocFilePath << TEXT("'");
+		}
+
+		Toc.Add(Entry->ChunkId, *Entry);
+		++Entry;
+	}
+
+	return FIoStatus::Ok;
+}
+
+bool FFileIoStoreReader::DoesChunkExist(const FIoChunkId& ChunkId) const
+{
+	return Toc.Find(ChunkId) != nullptr;
+}
+
+TIoStatusOr<uint64> FFileIoStoreReader::GetSizeForChunk(const FIoChunkId& ChunkId) const
+{
+	const FIoStoreTocEntry* Entry = Toc.Find(ChunkId);
+
+	if (Entry != nullptr)
+	{
+		return Entry->GetLength();
+	}
+	else
+	{
+		return FIoStatus(EIoErrorCode::NotFound);
+	}
+}
+
+bool FFileIoStoreReader::Resolve(FFileIoStoreResolvedRequest& ResolvedRequest)
+{
+	const FIoStoreTocEntry* Entry = Toc.Find(ResolvedRequest.Request->ChunkId);
+
+	if (!Entry)
+	{
+		return false;
+	}
+
+	ResolvedRequest.ResolvedFileHandle = ContainerFileHandle;
+	uint64 FileEndOffset = Entry->GetOffset() + Entry->GetLength();
+	uint64 RequestedBeginOffset = Entry->GetOffset() + ResolvedRequest.Request->Options.GetOffset();
+	uint64 RequestedEndOffset = FMath::Min(FileEndOffset, RequestedBeginOffset + ResolvedRequest.Request->Options.GetSize());
+	ResolvedRequest.ResolvedFileSize = ContainerFileSize;
+	ResolvedRequest.ResolvedOffset = RequestedBeginOffset;
+	if (RequestedEndOffset > RequestedBeginOffset)
+	{
+		ResolvedRequest.ResolvedSize = RequestedEndOffset - RequestedBeginOffset;
+	}
+	else
+	{
+		ResolvedRequest.ResolvedSize = 0;
+	}
+	return true;
+}
+
+FFileIoStore::FFileIoStore(FIoDispatcherEventQueue& InEventQueue)
+	: PlatformImpl(InEventQueue)
+{
+	InitCache();
+}
+
+FIoStatus FFileIoStore::Mount(const TCHAR* RootPath)
+{
+	TUniquePtr<FFileIoStoreReader> Reader(new FFileIoStoreReader(PlatformImpl));
+	FIoStatus IoStatus = Reader->Initialize(RootPath);
+	if (IoStatus.IsOk())
+	{
+		FWriteScopeLock _(IoStoreReadersLock);
+		IoStoreReaders.Add(Reader.Release());
+	}
+	return IoStatus;
+}
+
+EIoStoreResolveResult FFileIoStore::Resolve(FIoRequestImpl* Request)
+{
+	FReadScopeLock _(IoStoreReadersLock);
+	FFileIoStoreResolvedRequest ResolvedRequest;
+	ResolvedRequest.Request = Request;
+	for (FFileIoStoreReader* Reader : IoStoreReaders)
+	{
+		if (Reader->Resolve(ResolvedRequest))
+		{
+			// Halt IO request submission if the two least recently used blocks are still waiting for their IO
+			if (!(LruTail.LruPrev->bIsReady && LruTail.LruPrev->LruPrev->bIsReady))
+			{
+				return IoStoreResolveResult_Stalled;
+			}
+
+			Request->UnfinishedReadsCount = 0;
+			if (ResolvedRequest.ResolvedSize > 0)
+			{
+				if (void* TargetVa = Request->Options.GetTargetVa())
+				{
+					ResolvedRequest.Request->IoBuffer = FIoBuffer(FIoBuffer::Wrap, TargetVa, ResolvedRequest.ResolvedSize);
+				}
+				PlatformImpl.BeginReadsForRequest(ResolvedRequest);
+				uint64 RequestBeginBlockIndex = ResolvedRequest.ResolvedOffset / CacheBlockSize;
+				uint64 RequestEndBlockIndex = (ResolvedRequest.ResolvedOffset + ResolvedRequest.ResolvedSize - 1) / CacheBlockSize + 1;
+				uint64 BlockCount = RequestEndBlockIndex - RequestBeginBlockIndex;
+				check(BlockCount > 0);
+				ReadBlockCached(RequestBeginBlockIndex, ResolvedRequest);
+				if (BlockCount > 1)
+				{
+					if (BlockCount > 2)
+					{
+						ReadBlocksUncached(RequestBeginBlockIndex + 1, BlockCount - 2, ResolvedRequest);
+					}
+					ReadBlockCached(RequestEndBlockIndex - 1, ResolvedRequest);
+				}
+				PlatformImpl.EndReadsForRequest();
+			}
+
+			return IoStoreResolveResult_OK;
+		}
+	}
+
+	return IoStoreResolveResult_NotFound;
+}
+
+bool FFileIoStore::DoesChunkExist(const FIoChunkId& ChunkId) const
+{
+	FReadScopeLock _(IoStoreReadersLock);
+	for (FFileIoStoreReader* Reader : IoStoreReaders)
+	{
+		if (Reader->DoesChunkExist(ChunkId))
+		{
+			return true;
+		}
+	}
+	return false;
+}
+
+TIoStatusOr<uint64> FFileIoStore::GetSizeForChunk(const FIoChunkId& ChunkId) const
+{
+	FReadScopeLock _(IoStoreReadersLock);
+	for (FFileIoStoreReader* Reader : IoStoreReaders)
+	{
+		TIoStatusOr<uint64> ReaderResult = Reader->GetSizeForChunk(ChunkId);
+		if (ReaderResult.IsOk())
+		{
+			return ReaderResult;
+		}
+	}
+	return FIoStatus(EIoErrorCode::NotFound);
+}
+
+void FFileIoStore::ProcessIncomingBlocks()
+{
+	while (FFileIoStoreReadBlock* CompletedBlock = PlatformImpl.GetNextCompletedBlock())
+	{
+		check(!CompletedBlock->bIsReady);
+		CompletedBlock->bIsReady = true;
+		for (FFileIoStoreReadBlockScatter& Scatter : CompletedBlock->ScatterList)
+		{
+			FMemory::Memcpy(Scatter.Dst, Scatter.Src, Scatter.Size);
+			check(Scatter.Request->UnfinishedReadsCount > 0);
+			--Scatter.Request->UnfinishedReadsCount;
+		}
+		if (CacheBlocks.GetData() <= CompletedBlock && CompletedBlock < CacheBlocks.GetData() + CacheBlockCount)
+		{
+			CompletedBlock->ScatterList.Empty();
+		}
+		else
+		{
+			delete CompletedBlock;
+		}
+	}
+}
+
+void FFileIoStore::InitCache()
+{
+	uint8* CacheMemoryBlock = static_cast<uint8*>(FMemory::Malloc(CacheMemorySize));
+	CacheBlocks.AddDefaulted(CacheBlockCount);
+	FFileIoStoreReadBlock* PreviousBlock = &LruHead;
+	for (uint32 BlockIndex = 0; BlockIndex < CacheBlockCount; ++BlockIndex)
+	{
+		FFileIoStoreReadBlock& ReadBlock = CacheBlocks[BlockIndex];
+		PreviousBlock->LruNext = &ReadBlock;
+		ReadBlock.LruPrev = PreviousBlock;
+		PreviousBlock = &ReadBlock;
+
+		ReadBlock.bIsReady = true;
+		ReadBlock.Buffer = CacheMemoryBlock;
+		CacheMemoryBlock += CacheBlockSize;
+	}
+	PreviousBlock->LruNext = &LruTail;
+	LruTail.LruPrev = PreviousBlock;
+}
+
+void FFileIoStore::ReadBlockCached(uint32 BlockIndex, const FFileIoStoreResolvedRequest& ResolvedRequest)
+{
+	FFileIoStoreCacheBlockKey Key;
+	Key.FileHandle = ResolvedRequest.ResolvedFileHandle;
+	Key.BlockIndex = BlockIndex;
+	Key.Hash = HashCombine(ResolvedRequest.ResolvedFileHandle, BlockIndex);
+	uint64 BlockOffset = uint64(BlockIndex) * uint64(CacheBlockSize);
+	FFileIoStoreReadBlock* CachedBlock = CachedBlocksMap.FindRef(Key);
+	if (!CachedBlock)
+	{
+		CachedBlock = LruTail.LruPrev;
+		check(CachedBlock->bIsReady);
+		CachedBlocksMap.Remove(CachedBlock->Key);
+		CachedBlock->Key = Key;
+		CachedBlock->bIsReady = false;
+		CachedBlocksMap.Add(CachedBlock->Key, CachedBlock);
+
+		uint64 ReadSize = FMath::Min(ResolvedRequest.ResolvedFileSize, BlockOffset + CacheBlockSize) - BlockOffset;
+		PlatformImpl.ReadBlockFromFile(CachedBlock, CachedBlock->Buffer, ResolvedRequest.ResolvedFileHandle, ReadSize, BlockOffset);
+
+		TRACE_COUNTER_INCREMENT(IoDispatcherCacheMisses);
+	}
+
+	CachedBlock->LruPrev->LruNext = CachedBlock->LruNext;
+	CachedBlock->LruNext->LruPrev = CachedBlock->LruPrev;
+
+	CachedBlock->LruNext = LruHead.LruNext;
+	CachedBlock->LruPrev = &LruHead;
+	CachedBlock->LruNext->LruPrev = CachedBlock;
+	LruHead.LruNext = CachedBlock;
+
+	uint64 RequestStartOffsetInBlock = FMath::Max<int64>(0, int64(ResolvedRequest.ResolvedOffset) - BlockOffset);
+	uint64 RequestEndOffsetInBlock = FMath::Min<uint64>(CacheBlockSize, ResolvedRequest.ResolvedOffset + ResolvedRequest.ResolvedSize - BlockOffset);
+	uint64 BlockOffsetInRequest = FMath::Max<int64>(0, int64(BlockOffset) - ResolvedRequest.ResolvedOffset);
+	uint64 RequestSizeInBlock = RequestEndOffsetInBlock - RequestStartOffsetInBlock;
+	check(RequestSizeInBlock <= ResolvedRequest.Request->IoBuffer.DataSize());
+	check(RequestStartOffsetInBlock + RequestSizeInBlock <= CacheBlockSize);
+	check(ResolvedRequest.Request->IoBuffer.Data() + BlockOffsetInRequest + RequestSizeInBlock <= ResolvedRequest.Request->IoBuffer.Data() + ResolvedRequest.Request->IoBuffer.DataSize());
+	if (CachedBlock->bIsReady)
+	{
+		TRACE_COUNTER_INCREMENT(IoDispatcherCacheHitsHot);
+		FMemory::Memcpy(ResolvedRequest.Request->IoBuffer.Data() + BlockOffsetInRequest, CachedBlock->Buffer + RequestStartOffsetInBlock, RequestSizeInBlock);
+	}
+	else
+	{
+		TRACE_COUNTER_INCREMENT(IoDispatcherCacheHitsCold);
+		++ResolvedRequest.Request->UnfinishedReadsCount;
+		FFileIoStoreReadBlockScatter& Scatter = CachedBlock->ScatterList.AddDefaulted_GetRef();
+		Scatter.Request = ResolvedRequest.Request;
+		Scatter.Dst = ResolvedRequest.Request->IoBuffer.Data() + BlockOffsetInRequest;
+		Scatter.Src = CachedBlock->Buffer + RequestStartOffsetInBlock;
+		Scatter.Size = RequestSizeInBlock;
+	}
+}
+
+void FFileIoStore::ReadBlocksUncached(uint32 BeginBlockIndex, uint32 BlockCount, FFileIoStoreResolvedRequest& ResolvedRequest)
+{
+	uint64 BlockOffset = uint64(BeginBlockIndex) * uint64(CacheBlockSize);
+	uint64 BlockOffsetInRequest = BlockOffset - ResolvedRequest.ResolvedOffset;
+	uint64 ReadSize = BlockCount * CacheBlockSize;
+	check(ResolvedRequest.Request->IoBuffer.Data() + BlockOffsetInRequest + ReadSize <= ResolvedRequest.Request->IoBuffer.Data() + ResolvedRequest.Request->IoBuffer.DataSize());
+
+	FFileIoStoreReadBlock* UncachedBlock = new FFileIoStoreReadBlock();
+	++ResolvedRequest.Request->UnfinishedReadsCount;
+	FFileIoStoreReadBlockScatter& Scatter = UncachedBlock->ScatterList.AddDefaulted_GetRef();
+	Scatter.Request = ResolvedRequest.Request;
+	Scatter.Dst = nullptr;
+	Scatter.Src = nullptr;
+	Scatter.Size = 0;
+	PlatformImpl.ReadBlockFromFile(UncachedBlock, ResolvedRequest.Request->IoBuffer.Data() + BlockOffsetInRequest, ResolvedRequest.ResolvedFileHandle, ReadSize, BlockOffset);
+}
diff --git a/Engine/Source/Runtime/Core/Private/IO/IoDispatcherFileBackend.h b/Engine/Source/Runtime/Core/Private/IO/IoDispatcherFileBackend.h
new file mode 100644
--- /dev/null
+++ b/Engine/Source/Runtime/Core/Private/IO/IoDispatcherFileBackend.h
@@ -0,0 +1,99 @@
+// Copyright 1998-2019 Epic Games, Inc. All Rights Reserved.
+
+#pragma once
+
+#include"IO/IoDispatcherPrivate.h"
+#include "IO/IoStore.h"
+#include "Containers/Array.h"
+#include "Containers/Map.h"
+
+struct FFileIoStoreCacheBlockKey
+{
+	uint64 FileHandle;
+	uint32 BlockIndex;
+	uint32 Hash;
+
+	friend bool operator==(const FFileIoStoreCacheBlockKey& A, const FFileIoStoreCacheBlockKey& B)
+	{
+		return A.BlockIndex == B.BlockIndex && A.FileHandle == B.FileHandle;
+	}
+
+	friend uint32 GetTypeHash(const FFileIoStoreCacheBlockKey& Key)
+	{
+		return Key.Hash;
+	}
+};
+
+struct FFileIoStoreReadBlockScatter
+{
+	FIoRequestImpl* Request = nullptr;
+	uint8* Dst;
+	const uint8* Src;
+	uint32 Size;
+};
+
+struct FFileIoStoreReadBlock
+{
+	FFileIoStoreReadBlock* Next = nullptr;
+	FFileIoStoreReadBlock* LruPrev = nullptr;
+	FFileIoStoreReadBlock* LruNext = nullptr;
+	FFileIoStoreCacheBlockKey Key;
+	uint8* Buffer = nullptr;
+	TArray<FFileIoStoreReadBlockScatter> ScatterList;
+	bool bIsReady = false;
+};
+
+struct FFileIoStoreResolvedRequest
+{
+	FIoRequestImpl* Request;
+	uint64 ResolvedFileHandle;
+	uint64 ResolvedOffset;
+	uint64 ResolvedSize;
+	uint64 ResolvedFileSize;
+};
+
+class FFileIoStoreReader
+{
+public:
+	FFileIoStoreReader(FFileIoStoreImpl& InPlatformImpl);
+	FIoStatus Initialize(const TCHAR* RootPath);
+	bool DoesChunkExist(const FIoChunkId& ChunkId) const;
+	TIoStatusOr<uint64> GetSizeForChunk(const FIoChunkId& ChunkId) const;
+	bool Resolve(FFileIoStoreResolvedRequest& ResolvedRequest);
+
+private:
+	FFileIoStoreImpl& PlatformImpl;
+
+	TMap<FIoChunkId, FIoStoreTocEntry> Toc;
+	uint64 ContainerFileHandle;
+	uint64 ContainerFileSize;
+};
+
+class FFileIoStore
+{
+public:
+	FFileIoStore(FIoDispatcherEventQueue& InEventQueue);
+	FIoStatus Mount(const TCHAR* RootPath);
+	EIoStoreResolveResult Resolve(FIoRequestImpl* Request);
+	bool DoesChunkExist(const FIoChunkId& ChunkId) const;
+	TIoStatusOr<uint64> GetSizeForChunk(const FIoChunkId& ChunkId) const;
+	void ProcessIncomingBlocks();
+
+private:
+	void InitCache();
+	void ReadBlockCached(uint32 BlockIndex, const FFileIoStoreResolvedRequest& ResolvedRequest);
+	void ReadBlocksUncached(uint32 BeginBlockIndex, uint32 BlockCount, FFileIoStoreResolvedRequest& ResolvedRequest);
+
+	static constexpr uint32 CacheMemorySize = 32 << 20;
+	static constexpr uint32 CacheBlockSize = 256 << 10;
+	static constexpr uint32 CacheBlockCount = CacheMemorySize / CacheBlockSize;
+
+	FFileIoStoreImpl PlatformImpl;
+
+	mutable FRWLock IoStoreReadersLock;
+	TArray<FFileIoStoreReader*> IoStoreReaders;
+	TArray<FFileIoStoreReadBlock> CacheBlocks;
+	TMap<FFileIoStoreCacheBlockKey, FFileIoStoreReadBlock*> CachedBlocksMap;
+	FFileIoStoreReadBlock LruHead;
+	FFileIoStoreReadBlock LruTail;
+};
diff --git a/Engine/Source/Runtime/Core/Private/IO/IoDispatcherPrivate.h b/Engine/Source/Runtime/Core/Private/IO/IoDispatcherPrivate.h
new file mode 100644
--- /dev/null
+++ b/Engine/Source/Runtime/Core/Private/IO/IoDispatcherPrivate.h
@@ -0,0 +1,47 @@
+// Copyright 1998-2019 Epic Games, Inc. All Rights Reserved.
+
+#pragma once
+
+#include "IO/IoDispatcher.h"
+
+#ifndef PLATFORM_IMPLEMENTS_IO
+#define PLATFORM_IMPLEMENTS_IO 0
+#endif
+
+#if PLATFORM_IMPLEMENTS_IO
+#include COMPILED_PLATFORM_HEADER(PlatformIoDispatcher.h)
+#else
+#include "GenericPlatform/GenericPlatformIoDispatcher.h"
+typedef FGenericIoDispatcherEventQueue FIoDispatcherEventQueue;
+typedef FGenericFileIoStoreImpl FFileIoStoreImpl;
+#endif
+
+class FIoRequestImpl;
+
+enum EIoStoreResolveResult
+{
+	IoStoreResolveResult_OK,
+	IoStoreResolveResult_Stalled,
+	IoStoreResolveResult_NotFound,
+};
+
+class FIoBatchImpl
+{
+public:
+	FIoRequestImpl* FirstRequest = nullptr;
+};
+
+class FIoRequestImpl
+{
+public:
+	FIoBatchImpl* Batch = nullptr;
+	FIoRequestImpl* NextRequest = nullptr;
+	FIoRequestImpl* BatchNextRequest = nullptr;
+	FIoStatus Status;
+	FIoChunkId ChunkId;
+	FIoReadOptions Options;
+	FIoBuffer IoBuffer;
+	uint32 UnfinishedReadsCount;
+	TFunction<void(TIoStatusOr<FIoBuffer>)> Callback;
+};
+
diff --git a/Engine/Source/Runtime/Core/Public/IO/IoDispatcher.h b/Engine/Source/Runtime/Core/Public/IO/IoDispatcher.h
--- a/Engine/Source/Runtime/Core/Public/IO/IoDispatcher.h
+++ b/Engine/Source/Runtime/Core/Public/IO/IoDispatcher.h
@@ -22,13 +22,12 @@
 
 class FIoRequest;
 class FIoDispatcher;
-class FIoStoreReader;
 class FIoStoreWriter;
+class FIoStoreEnvironment;
 
 class FIoRequestImpl;
 class FIoBatchImpl;
 class FIoDispatcherImpl;
-class FIoStoreReaderImpl;
 class FIoStoreWriterImpl;
 
 CORE_API DECLARE_LOG_CATEGORY_EXTERN(LogIoDispatcher, Log, All);
@@ -629,27 +628,22 @@ class FIoReadOptions
 public:
 	FIoReadOptions() = default;
 
-	FIoReadOptions(uint64 InOffset, uint32 InSize)
+	FIoReadOptions(uint64 InOffset, uint64 InSize)
 		: RequestedOffset(InOffset)
 		, RequestedSize(InSize)
 	{ }
 
 	~FIoReadOptions() = default;
 
-	void SetRange(uint64 Offset, uint32 Size)
+	void SetRange(uint64 Offset, uint64 Size)
 	{
 		RequestedOffset = Offset;
 		RequestedSize	= Size;
 	}
 
-	void SetTargetVa(uint64 VaTargetAddress)
+	void SetTargetVa(void* InTargetVa)
 	{
-		TargetVa = VaTargetAddress;
-	}
-
-	void ForGPU()
-	{
-		Flags |= EFlags::GPUMemory;
+		TargetVa = InTargetVa;
 	}
 
 	uint64 GetOffset() const
@@ -662,16 +656,16 @@ public:
 		return RequestedSize;
 	}
 
+	void* GetTargetVa() const
+	{
+		return TargetVa;
+	}
+
 private:
-	uint64	TargetVa		= 0;
 	uint64	RequestedOffset = 0;
-	uint32	RequestedSize	= ~uint32(0);
-	uint32	Flags			= 0;
-
-	enum EFlags
-	{
-		GPUMemory = 1 << 0,
-	};
+	uint64	RequestedSize = ~uint64(0);
+	void* TargetVa = nullptr;
+	uint32	Flags = 0;
 };
 
 //////////////////////////////////////////////////////////////////////////
@@ -738,8 +732,7 @@ public:
 	CORE_API						FIoDispatcher();
 	CORE_API virtual				~FIoDispatcher();
 
-	CORE_API void					Mount(FIoStoreReader* IoStore);
-	CORE_API void					Unmount(FIoStoreReader* IoStore);
+	CORE_API FIoStatus				Mount(FIoStoreEnvironment& Environment);
 
 	CORE_API FIoBatch				NewBatch();
 	CORE_API void					FreeBatch(FIoBatch Batch);
@@ -754,10 +747,9 @@ public:
 	FIoDispatcher& operator=(const FIoDispatcher&) = delete;
 
 	static CORE_API bool IsInitialized();
-	static CORE_API FIoStatus Initialize(const FString& Directory);
+	static CORE_API FIoStatus Initialize();
 	static CORE_API void Shutdown();
 	static CORE_API FIoDispatcher& Get();
-	static CORE_API class FIoStoreEnvironment& GetEnvironment();
 
 private:
 	FIoDispatcherImpl* Impl = nullptr;
@@ -785,26 +777,6 @@ private:
 	FString			RootPath;
 };
 
-class FIoStoreReader : public FRefCountBase
-{
-public:
-	CORE_API FIoStoreReader(FIoStoreEnvironment& Environment);
-	CORE_API ~FIoStoreReader();
-
-	/** This will parse the manifests in the environment and
-		populate the table of contents. To be useful the IO dispatcher
-		needs to have access to the information, which is what
-		the I/O dispatcher Mount()/Unmount() calls are for.
-	  */
-	CORE_API FIoStatus Initialize(FStringView UniqueId);
-
-private:
-	FIoStoreReaderImpl* Impl;
-
-	friend class FIoDispatcherImpl;
-	friend class FIoStoreImpl;
-};
-
 //////////////////////////////////////////////////////////////////////////
 
 class FIoStoreWriter
diff --git a/Engine/Source/Runtime/CoreUObject/Private/Serialization/AsyncLoading2.cpp b/Engine/Source/Runtime/CoreUObject/Private/Serialization/AsyncLoading2.cpp
--- a/Engine/Source/Runtime/CoreUObject/Private/Serialization/AsyncLoading2.cpp
+++ b/Engine/Source/Runtime/CoreUObject/Private/Serialization/AsyncLoading2.cpp
@@ -1876,12 +1876,10 @@ void FAsyncLoadingThread2Impl::InitializeLoading()
 	{
 		TRACE_CPUPROFILER_EVENT_SCOPE(InitIoDispatcher);
 
-		TUniquePtr<FIoStoreReader> IoStoreReader = MakeUnique<FIoStoreReader>(FIoDispatcher::GetEnvironment());
-		FIoStatus ReaderStatus = IoStoreReader->Initialize(TEXT("PackageLoader"));
-
-		UE_CLOG(!ReaderStatus.IsOk(), LogStreaming, Error, TEXT("Failed to initialize I/O dispatcher: '%s'"), *ReaderStatus.ToString());
-
-		IoDispatcher.Mount(IoStoreReader.Release());
+		FIoStoreEnvironment IoStoreEnvironment;
+		IoStoreEnvironment.InitializeFileEnvironment(FPaths::ProjectDir());
+		FIoStatus ReaderStatus = IoDispatcher.Mount(IoStoreEnvironment);
+		UE_CLOG(!ReaderStatus.IsOk(), LogStreaming, Fatal, TEXT("Failed to initialize I/O dispatcher: '%s'"), *ReaderStatus.ToString());
 
 #if USE_NEW_BULKDATA
 		FBulkDataBase::SetIODispatcher(&IoDispatcher);
diff --git a/Engine/Source/Runtime/CoreUObject/Private/Serialization/BulkData2.cpp b/Engine/Source/Runtime/CoreUObject/Private/Serialization/BulkData2.cpp
--- a/Engine/Source/Runtime/CoreUObject/Private/Serialization/BulkData2.cpp
+++ b/Engine/Source/Runtime/CoreUObject/Private/Serialization/BulkData2.cpp
@@ -270,7 +270,7 @@ public:
 		if (FPlatformProcess::SupportsMultithreading())
 		{
 			Task = new FAsyncTask<FBulkDataIoDispatcherRequestWorker>(*this);
-			Task->StartBackgroundTask(GIOThreadPool);
+			Task->StartBackgroundTask(GBackgroundPriorityThreadPool);
 		}
 		else
 		{
diff --git a/Engine/Source/Runtime/Launch/Private/LaunchEngineLoop.cpp b/Engine/Source/Runtime/Launch/Private/LaunchEngineLoop.cpp
--- a/Engine/Source/Runtime/Launch/Private/LaunchEngineLoop.cpp
+++ b/Engine/Source/Runtime/Launch/Private/LaunchEngineLoop.cpp
@@ -1447,7 +1447,7 @@ int32 FEngineLoop::PreInitPreStartupScreen(const TCHAR* CmdLine)
 	// Initialize I/O dispatcher when using the new package loader
 	if (FParse::Param(FCommandLine::Get(), TEXT("zenloader")))
 	{
-		FIoDispatcher::Initialize(FPaths::ProjectDir());
+		FIoDispatcher::Initialize();
 	}
 #endif
 
