Add GetLastError information to FFileManagerGeneric::Move
#rnx

diff --git a/Engine/Source/Runtime/Core/Private/HAL/FileManagerGeneric.cpp b/Engine/Source/Runtime/Core/Private/HAL/FileManagerGeneric.cpp
--- a/Engine/Source/Runtime/Core/Private/HAL/FileManagerGeneric.cpp
+++ b/Engine/Source/Runtime/Core/Private/HAL/FileManagerGeneric.cpp
@@ -327,10 +327,11 @@ bool FFileManagerGeneric::Move( const TCHAR* Dest, const TCHAR* Src, bool Replac
 
 		int32 RetryCount = 10;
 		bool bSuccess = false;
+		uint32 ErrorCode = FPlatformMisc::GetLastError();
 		while (RetryCount--)
 		{
 			// If the move failed, throw a warning but retry before we throw an error
-			UE_LOG(LogFileManager, Warning, TEXT("MoveFile was unable to move '%s' to '%s', retrying in .5s..."), Src, Dest);
+			UE_LOG(LogFileManager, Warning, TEXT("MoveFile was unable to move '%s' to '%s' (Error Code %i), retrying in .5s..."), Src, Dest, ErrorCode);
 
 			// Wait just a little bit( i.e. a totally arbitrary amount )...
 			FPlatformProcess::Sleep(0.5f);
@@ -342,6 +343,7 @@ bool FFileManagerGeneric::Move( const TCHAR* Dest, const TCHAR* Src, bool Replac
 				UE_LOG(LogFileManager, Warning, TEXT("MoveFile recovered during retry!"));
 				break;
 			}
+			ErrorCode = FPlatformMisc::GetLastError();
 		}
 		if (!bSuccess)
 		{
