Add FName::TryAppendAnsiString as an optimization because most names are not wide

- Optimize appending to FStringBuilderBase by using its new AddUninitialized to avoid temporary buffers.
- Remove AppendNameToPathString until there is a use case for it. It hasn't been copied up from from Dev-Core yet.

diff --git a/Engine/Source/Runtime/Core/Private/UObject/UnrealNames.cpp b/Engine/Source/Runtime/Core/Private/UObject/UnrealNames.cpp
--- a/Engine/Source/Runtime/Core/Private/UObject/UnrealNames.cpp
+++ b/Engine/Source/Runtime/Core/Private/UObject/UnrealNames.cpp
@@ -1366,35 +1366,31 @@ void FNameEntry::AppendNameToString(FString& Out) const
 
 void FNameEntry::AppendNameToString(FStringBuilderBase& Out) const
 {
-	FNameBuffer Temp;
+	const int32 Offset = Out.AddUninitialized(Header.Len);
+	TCHAR* OutChars = Out.GetData() + Offset;
 	if (Header.bIsWide)
 	{
-		Out.Append(GetUnterminatedName(Temp.WideName), Header.Len);
+		CopyUnterminatedName(reinterpret_cast<WIDECHAR*>(OutChars));
+		ConvertInPlace<WIDECHAR, TCHAR>(reinterpret_cast<WIDECHAR*>(OutChars), Header.Len);
 	}
 	else
 	{
-		Out.AppendAnsi(GetUnterminatedName(Temp.AnsiName), Header.Len);
+		CopyUnterminatedName(reinterpret_cast<ANSICHAR*>(OutChars));
+		ConvertInPlace<ANSICHAR, TCHAR>(reinterpret_cast<ANSICHAR*>(OutChars), Header.Len);
 	}
 }
 
-void FNameEntry::AppendNameToPathString(FString& Out) const
+void FNameEntry::AppendAnsiNameToString(FAnsiStringBuilderBase& Out) const
 {
-	FNameBuffer Temp;
-	Out.PathAppend(EntryToCString(*this, Temp), Header.Len);
+	check(!IsWide());
+	const int32 Offset = Out.AddUninitialized(Header.Len);
+	CopyUnterminatedName(Out.GetData() + Offset);
 }
 
-void FNameEntry::AppendNameToPathString(FStringBuilderBase& Out) const
+void FNameEntry::AppendNameToPathString(FString& Out) const
 {
-	if (Out.Len() > 0)
-	{
-		const TCHAR LastChar = Out.LastChar();
-		if (LastChar != TEXT('/') && LastChar != TEXT('\\'))
-		{
-			Out.Append(TEXT('/'));
-		}
-	}
-
-	AppendNameToString(Out);
+	FNameBuffer Temp;
+	Out.PathAppend(EntryToCString(*this, Temp), Header.Len);
 }
 
 int32 FNameEntry::GetSize(const TCHAR* Name)
@@ -2103,6 +2099,26 @@ void FName::AppendString(FStringBuilderBase& Out) const
 	}
 }
 
+bool FName::TryAppendAnsiString(FAnsiStringBuilderBase& Out) const
+{
+	const FNameEntry* const NameEntry = GetDisplayNameEntry();
+
+	if (NameEntry->IsWide())
+	{
+		return false;
+	}
+
+	NameEntry->AppendAnsiNameToString(Out);
+
+	const int32 InternalNumber = GetNumber();
+	if (InternalNumber != NAME_NO_NUMBER_INTERNAL)
+	{
+		Out << '_' << NAME_INTERNAL_TO_EXTERNAL(InternalNumber);
+	}
+
+	return true;
+}
+
 void FName::DisplayHash(FOutputDevice& Ar)
 {
 	GetNamePool().LogStats(Ar);
diff --git a/Engine/Source/Runtime/Core/Public/UObject/NameTypes.h b/Engine/Source/Runtime/Core/Public/UObject/NameTypes.h
--- a/Engine/Source/Runtime/Core/Public/UObject/NameTypes.h
+++ b/Engine/Source/Runtime/Core/Public/UObject/NameTypes.h
@@ -27,6 +27,7 @@
 	#define WITH_CASE_PRESERVING_NAME WITH_EDITORONLY_DATA
 #endif
 
+class FAnsiStringBuilderBase;
 class FStringBuilderBase;
 class FText;
 
@@ -233,12 +234,12 @@ public:
 	/** Appends name to string builder. */
 	CORE_API void AppendNameToString(FStringBuilderBase& OutString) const;
 
+	/** Appends name to string builder. Entry must not be wide. */
+	CORE_API void AppendAnsiNameToString(FAnsiStringBuilderBase& OutString) const;
+
 	/** Appends name to string with path separator using FString::PathAppend(). */
 	CORE_API void AppendNameToPathString(FString& OutString) const;
 
-	/** Appends name to string builder with path separator. */
-	CORE_API void AppendNameToPathString(FStringBuilderBase& OutString) const;
-
 
 	/**
 	 * Returns the size in bytes for FNameEntry structure. This is != sizeof(FNameEntry) as we only allocated as needed.
@@ -493,6 +494,15 @@ public:
 	 */
 	void AppendString(FStringBuilderBase& Out) const;
 
+	/**
+	 * Converts an ANSI FName to a readable format appended to the string builder.
+	 *
+	 * @param Out A string builder to write the readable representation of the name into.
+	 *
+	 * @return Whether the string is ANSI. A return of false indicates that the string was wide and was not written.
+	 */
+	bool TryAppendAnsiString(FAnsiStringBuilderBase& Out) const;
+
 	/**
 	 * Check to see if this FName matches the other FName, potentially also checking for any case variations
 	 */
