Iteration on StringView

- Add TStringView<CharType> alias to facilitate string view usage in code templated on the character type.
- Make FWideStringView an alias for FStringView to facilitate the creation of TStringView. This seems backward but the explicitly wide version is likely rarely used and this still permits easy forward declaration of FStringView.
- Add TrimStart/TrimEnd/TrimStartAndEnd and inline variants for compatibility with FString.

diff --git a/Engine/Source/Runtime/Core/Private/Containers/StringView.cpp b/Engine/Source/Runtime/Core/Private/Containers/StringView.cpp
--- a/Engine/Source/Runtime/Core/Private/Containers/StringView.cpp
+++ b/Engine/Source/Runtime/Core/Private/Containers/StringView.cpp
@@ -1,88 +1,156 @@
 // Copyright 1998-2019 Epic Games, Inc. All Rights Reserved.
 
-#include "Misc/StringView.h"
+#include "Containers/StringView.h"
 
-template class TStringViewImpl<TCHAR, FStringView>;
 template class TStringViewImpl<ANSICHAR, FAnsiStringView>;
 template class TStringViewImpl<WIDECHAR, FWideStringView>;
 
-template <typename ViewType>
-int32 StringViewPrivate::Compare(const ViewType& Lhs, const ViewType& Rhs, ESearchCase::Type SearchCase)
+namespace StringViewPrivate
 {
-	const typename ViewType::SizeType LhsLen = Lhs.Len();
-	const typename ViewType::SizeType RhsLen = Rhs.Len();
-	const typename ViewType::SizeType MinLen = FMath::Min(LhsLen, RhsLen);
-
-	int Result;
-	if (SearchCase == ESearchCase::CaseSensitive)
+	template <typename ViewType>
+	inline int32 CompareImpl(const ViewType& Lhs, const ViewType& Rhs, ESearchCase::Type SearchCase)
 	{
-		Result = TCString<typename ViewType::ElementType>::Strncmp(Lhs.GetData(), Rhs.GetData(), MinLen);
+		const typename ViewType::SizeType LhsLen = Lhs.Len();
+		const typename ViewType::SizeType RhsLen = Rhs.Len();
+		const typename ViewType::SizeType MinLen = FMath::Min(LhsLen, RhsLen);
+
+		int Result;
+		if (SearchCase == ESearchCase::CaseSensitive)
+		{
+			Result = TCString<typename ViewType::ElementType>::Strncmp(Lhs.GetData(), Rhs.GetData(), MinLen);
+		}
+		else
+		{
+			Result = TCString<typename ViewType::ElementType>::Strnicmp(Lhs.GetData(), Rhs.GetData(), MinLen);
+		}
+
+		if (Result != 0 || LhsLen == RhsLen)
+		{
+			return Result;
+		}
+
+		return LhsLen < RhsLen ? -1 : 1;
 	}
-	else
+
+	int32 Compare(const FAnsiStringView& Lhs, const FAnsiStringView& Rhs, ESearchCase::Type SearchCase)
 	{
-		Result = TCString<typename ViewType::ElementType>::Strnicmp(Lhs.GetData(), Rhs.GetData(), MinLen);
+		return CompareImpl(Lhs, Rhs, SearchCase);
 	}
 
-	if (Result != 0 || LhsLen == RhsLen)
+	int32 Compare(const FWideStringView& Lhs, const FWideStringView& Rhs, ESearchCase::Type SearchCase)
 	{
-		return Result;
+		return CompareImpl(Lhs, Rhs, SearchCase);
 	}
 
-	return LhsLen < RhsLen ? -1 : 1;
-}
+	template <typename ViewType>
+	inline bool FindCharImpl(const ViewType& InView, typename ViewType::ElementType InChar, typename ViewType::SizeType& OutIndex)
+	{
+		const typename ViewType::ElementType* DataPtr = InView.GetData();
+		const typename ViewType::SizeType Size = InView.Len();
 
-template int32 CORE_API StringViewPrivate::Compare(const FStringView& Lhs, const FStringView& Rhs, ESearchCase::Type SearchCase);
-template int32 CORE_API StringViewPrivate::Compare(const FAnsiStringView& Lhs, const FAnsiStringView& Rhs, ESearchCase::Type SearchCase);
-template int32 CORE_API StringViewPrivate::Compare(const FWideStringView& Lhs, const FWideStringView& Rhs, ESearchCase::Type SearchCase);
+		for (typename ViewType::SizeType i = 0; i < Size; ++i)
+		{
+			if (DataPtr[i] == InChar)
+			{
+				OutIndex = i;
+				return true;
+			}
+		}
 
-template <typename ViewType>
-bool StringViewPrivate::FindChar(const ViewType& InView, typename ViewType::ElementType InChar, typename ViewType::SizeType& OutIndex)
-{
-	const typename ViewType::ElementType* DataPtr = InView.GetData();
-	const typename ViewType::SizeType Size = InView.Len();
+		OutIndex = INDEX_NONE;
+		return false;
+	}
 
-	for (typename ViewType::SizeType i = 0; i < Size; ++i)
+	bool FindChar(const FAnsiStringView& InView, ANSICHAR InChar, typename FAnsiStringView::SizeType& OutIndex)
 	{
-		if (DataPtr[i] == InChar)
+		return FindCharImpl(InView, InChar, OutIndex);
+	}
+
+	bool FindChar(const FWideStringView& InView, WIDECHAR InChar, typename FWideStringView::SizeType& OutIndex)
+	{
+		return FindCharImpl(InView, InChar, OutIndex);
+	}
+
+	template <typename ViewType>
+	bool FindLastCharImpl(const ViewType& InView, typename ViewType::ElementType InChar, typename ViewType::SizeType& OutIndex)
+	{
+		const typename ViewType::ElementType* DataPtr = InView.GetData();
+		const typename ViewType::SizeType Size = InView.Len();
+
+		if (Size == 0)
 		{
-			OutIndex = i;
-			return true;
+			OutIndex = INDEX_NONE;
+			return false;
 		}
+
+		for (typename ViewType::SizeType i = Size - 1; i >= 0; --i)
+		{
+			if (DataPtr[i] == InChar)
+			{
+				OutIndex = i;
+				return true;
+			}
+		}
+
+		OutIndex = INDEX_NONE;
+		return false;
 	}
 
-	OutIndex = INDEX_NONE;
-	return false;
-}
+	bool FindLastChar(const FAnsiStringView& InView, ANSICHAR InChar, typename FAnsiStringView::SizeType& OutIndex)
+	{
+		return FindLastCharImpl(InView, InChar, OutIndex);
+	}
 
-template bool CORE_API StringViewPrivate::FindChar(const FStringView& InView, TCHAR InChar, typename FStringView::SizeType& OutIndex);
-template bool CORE_API StringViewPrivate::FindChar(const FAnsiStringView& InView, ANSICHAR InChar, typename FAnsiStringView::SizeType& OutIndex);
-template bool CORE_API StringViewPrivate::FindChar(const FWideStringView& InView, WIDECHAR InChar, typename FWideStringView::SizeType& OutIndex);
+	bool FindLastChar(const FWideStringView& InView, WIDECHAR InChar, typename FWideStringView::SizeType& OutIndex)
+	{
+		return FindLastCharImpl(InView, InChar, OutIndex);
+	}
 
-template <typename ViewType>
-bool StringViewPrivate::FindLastChar(const ViewType& InView, typename ViewType::ElementType InChar, typename ViewType::SizeType& OutIndex)
-{
-	const typename ViewType::ElementType* DataPtr = InView.GetData();
-	const typename ViewType::SizeType Size = InView.Len();
+	template <typename ViewType>
+	inline ViewType TrimStartImpl(ViewType InView)
+	{
+		typename ViewType::SizeType SpaceCount = 0;
+		for (typename ViewType::ElementType Char : InView)
+		{
+			if (!TChar<typename ViewType::ElementType>::IsWhitespace(Char))
+			{
+				break;
+			}
+			++SpaceCount;
+		}
+		InView.RemovePrefix(SpaceCount);
+		return InView;
+	}
 
-	if (Size == 0)
+	FAnsiStringView TrimStart(const FAnsiStringView& InView)
 	{
-		OutIndex = INDEX_NONE;
-		return false;
+		return TrimStartImpl(InView);
+	}
+
+	FWideStringView TrimStart(const FWideStringView& InView)
+	{
+		return TrimStartImpl(InView);
 	}
 
-	for (typename ViewType::SizeType i = Size - 1; i >= 0; --i)
+	template <typename ViewType>
+	inline ViewType TrimEndImpl(const ViewType& InView)
 	{
-		if (DataPtr[i] == InChar)
+		const typename ViewType::ElementType* DataPtr = InView.GetData();
+		typename ViewType::SizeType Len = InView.Len();
+		while (Len && TChar<typename ViewType::ElementType>::IsWhitespace(DataPtr[Len - 1]))
 		{
-			OutIndex = i;
-			return true;
+			--Len;
 		}
+		return ViewType(DataPtr, Len);
 	}
 
-	OutIndex = INDEX_NONE;
-	return false;
-}
+	FAnsiStringView TrimEnd(const FAnsiStringView& InView)
+	{
+		return TrimEndImpl(InView);
+	}
 
-template bool CORE_API StringViewPrivate::FindLastChar(const FStringView& InView, TCHAR InChar, typename FStringView::SizeType& OutIndex);
-template bool CORE_API StringViewPrivate::FindLastChar(const FAnsiStringView& InView, ANSICHAR InChar, typename FAnsiStringView::SizeType& OutIndex);
-template bool CORE_API StringViewPrivate::FindLastChar(const FWideStringView& InView, WIDECHAR InChar, typename FWideStringView::SizeType& OutIndex);
+	FWideStringView TrimEnd(const FWideStringView& InView)
+	{
+		return TrimEndImpl(InView);
+	}
+}
diff --git a/Engine/Source/Runtime/Core/Private/Tests/Misc/StringViewTest.cpp b/Engine/Source/Runtime/Core/Private/Tests/Misc/StringViewTest.cpp
--- a/Engine/Source/Runtime/Core/Private/Tests/Misc/StringViewTest.cpp
+++ b/Engine/Source/Runtime/Core/Private/Tests/Misc/StringViewTest.cpp
@@ -1,17 +1,25 @@
 // Copyright 1998-2019 Epic Games, Inc. All Rights Reserved.
 
-#include "Misc/StringView.h"
+#include "Containers/StringView.h"
 
 #include "Containers/UnrealString.h"
 #include "Misc/AutomationTest.h"
 
-#if WITH_DEV_AUTOMATION_TESTS 
+#if WITH_DEV_AUTOMATION_TESTS
+
+static_assert(TIsSame<typename FStringView::ElementType, TCHAR>::Value, "FStringView must use TCHAR.");
+static_assert(TIsSame<typename FAnsiStringView::ElementType, ANSICHAR>::Value, "FAnsiStringView must use ANSICHAR.");
+static_assert(TIsSame<typename FWideStringView::ElementType, WIDECHAR>::Value, "FWideStringView must use WIDECHAR.");
+
+static_assert(TIsSame<FStringView, TStringView<TCHAR>>::Value, "FStringView must be the same as TStringView<TCHAR>.");
+static_assert(TIsSame<FAnsiStringView, TStringView<ANSICHAR>>::Value, "FAnsiStringView must be the same as TStringView<ANSICHAR>.");
+static_assert(TIsSame<FWideStringView, TStringView<WIDECHAR>>::Value, "FWideStringView must be the same as TStringView<WIDECHAR>.");
 
 static_assert(TIsContiguousContainer<FStringView>::Value, "FStringView must be a contiguous container.");
 static_assert(TIsContiguousContainer<FAnsiStringView>::Value, "FAnsiStringView must be a contiguous container.");
 static_assert(TIsContiguousContainer<FWideStringView>::Value, "FWideStringView must be a contiguous container.");
 
-#define TEST_NAME_ROOT "System.Core.Misc.StringView"
+#define TEST_NAME_ROOT "System.Core.StringView"
 constexpr const uint32 TestFlags = EAutomationTestFlags::ApplicationContextMask | EAutomationTestFlags::EngineFilter;
 
 IMPLEMENT_SIMPLE_AUTOMATION_TEST(FStringViewTestCtor, TEST_NAME_ROOT ".Ctor", TestFlags)
@@ -73,6 +81,13 @@ bool FStringViewTestCtor::RunTest(const FString& Parameters)
 		TestFalse(TEXT("View.IsEmpty"), View.IsEmpty());
 	}
 
+	// Create using string view literals
+	{
+		FStringView View = TEXT("Test"_SV);
+		FAnsiStringView ViewAnsi = "Test"_ASV;
+		FWideStringView ViewWide = L"Test"_WSV;
+	}
+
 	return true;
 }
 
@@ -614,6 +629,51 @@ bool FStringViewTestMid::RunTest(const FString& Parameters)
 	return true;
 }
 
+IMPLEMENT_SIMPLE_AUTOMATION_TEST(FStringViewTestTrimStartAndEnd, TEST_NAME_ROOT ".TrimStartAndEnd", TestFlags)
+bool FStringViewTestTrimStartAndEnd::RunTest(const FString& Parameters)
+{
+	TestTrue(TEXT("FStringView::TrimStartAndEnd(\"\")"), TEXT(""_SV).TrimStartAndEnd().IsEmpty());
+	TestTrue(TEXT("FStringView::TrimStartAndEnd(\" \")"), TEXT(" "_SV).TrimStartAndEnd().IsEmpty());
+	TestTrue(TEXT("FStringView::TrimStartAndEnd(\"  \")"), TEXT("  "_SV).TrimStartAndEnd().IsEmpty());
+	TestTrue(TEXT("FStringView::TrimStartAndEnd(\" \\t\\r\\n\")"), TEXT(" \t\r\n"_SV).TrimStartAndEnd().IsEmpty());
+
+	TestEqual(TEXT("FStringView::TrimStartAndEnd(\"ABC123\")"), TEXT("ABC123"_SV).TrimStartAndEnd(), TEXT("ABC123"_SV));
+	TestEqual(TEXT("FStringView::TrimStartAndEnd(\"A \\t\\r\\nB\")"), TEXT("A \t\r\nB"_SV).TrimStartAndEnd(), TEXT("A \t\r\nB"_SV));
+	TestEqual(TEXT("FStringView::TrimStartAndEnd(\" \\t\\r\\nABC123\\n\\r\\t \")"), TEXT(" \t\r\nABC123\n\r\t "_SV).TrimStartAndEnd(), TEXT("ABC123"_SV));
+
+	return true;
+}
+
+IMPLEMENT_SIMPLE_AUTOMATION_TEST(FStringViewTestTrimStart, TEST_NAME_ROOT ".TrimStart", TestFlags)
+bool FStringViewTestTrimStart::RunTest(const FString& Parameters)
+{
+	TestTrue(TEXT("FStringView::TrimStart(\"\")"), TEXT(""_SV).TrimStart().IsEmpty());
+	TestTrue(TEXT("FStringView::TrimStart(\" \")"), TEXT(" "_SV).TrimStart().IsEmpty());
+	TestTrue(TEXT("FStringView::TrimStart(\"  \")"), TEXT("  "_SV).TrimStart().IsEmpty());
+	TestTrue(TEXT("FStringView::TrimStart(\" \\t\\r\\n\")"), TEXT(" \t\r\n"_SV).TrimStart().IsEmpty());
+
+	TestEqual(TEXT("FStringView::TrimStart(\"ABC123\")"), TEXT("ABC123"_SV).TrimStart(), TEXT("ABC123"_SV));
+	TestEqual(TEXT("FStringView::TrimStart(\"A \\t\\r\\nB\")"), TEXT("A \t\r\nB"_SV).TrimStart(), TEXT("A \t\r\nB"_SV));
+	TestEqual(TEXT("FStringView::TrimStart(\" \\t\\r\\nABC123\\n\\r\\t \")"), TEXT(" \t\r\nABC123\n\r\t "_SV).TrimStart(), TEXT("ABC123\n\r\t "_SV));
+
+	return true;
+}
+
+IMPLEMENT_SIMPLE_AUTOMATION_TEST(FStringViewTestTrimEnd, TEST_NAME_ROOT ".TrimEnd", TestFlags)
+bool FStringViewTestTrimEnd::RunTest(const FString& Parameters)
+{
+	TestTrue(TEXT("FStringView::TrimEnd(\"\")"), TEXT(""_SV).TrimEnd().IsEmpty());
+	TestTrue(TEXT("FStringView::TrimEnd(\" \")"), TEXT(" "_SV).TrimEnd().IsEmpty());
+	TestTrue(TEXT("FStringView::TrimEnd(\"  \")"), TEXT("  "_SV).TrimEnd().IsEmpty());
+	TestTrue(TEXT("FStringView::TrimEnd(\" \\t\\r\\n\")"), TEXT(" \t\r\n"_SV).TrimEnd().IsEmpty());
+
+	TestEqual(TEXT("FStringView::TrimEnd(\"ABC123\")"), TEXT("ABC123"_SV).TrimEnd(), TEXT("ABC123"_SV));
+	TestEqual(TEXT("FStringView::TrimEnd(\"A \\t\\r\\nB\")"), TEXT("A \t\r\nB"_SV).TrimEnd(), TEXT("A \t\r\nB"_SV));
+	TestEqual(TEXT("FStringView::TrimEnd(\" \\t\\r\\nABC123\\n\\r\\t \")"), TEXT(" \t\r\nABC123\n\r\t "_SV).TrimEnd(), TEXT(" \t\r\nABC123"_SV));
+
+	return true;
+}
+
 IMPLEMENT_SIMPLE_AUTOMATION_TEST(FStringViewTestFindChar, TEST_NAME_ROOT ".FindChar", TestFlags)
 bool FStringViewTestFindChar::RunTest(const FString& Parameters)
 {
diff --git a/Engine/Source/Runtime/Core/Public/Containers/StringView.h b/Engine/Source/Runtime/Core/Public/Containers/StringView.h
--- a/Engine/Source/Runtime/Core/Public/Containers/StringView.h
+++ b/Engine/Source/Runtime/Core/Public/Containers/StringView.h
@@ -154,6 +154,12 @@ public:
 	inline ViewType RightChop(SizeType CharCount) const;
 	/** Returns the middle part of the view by taking up to the given number of characters from the given position. */
 	inline ViewType Mid(SizeType Position, SizeType CharCount = TNumericLimits<SizeType>::Max()) const;
+	/** Returns the middle part of the view between any whitespace at the start and end. */
+	inline ViewType TrimStartAndEnd() const;
+	/** Returns the right part of the view after any whitespace at the start. */
+	inline ViewType TrimStart() const;
+	/** Returns the left part of the view before any whitespace at the end. */
+	inline ViewType TrimEnd() const;
 
 	/** Modifies the view to be the given number of characters from the left. */
 	inline void LeftInlineInline(SizeType CharCount) { *this = Left(CharCount); }
@@ -165,6 +171,12 @@ public:
 	inline void RightChopInline(SizeType CharCount) { *this = RightChop(CharCount); }
 	/** Modifies the view to be the middle part by taking up to the given number of characters from the given position. */
 	inline void MidInline(SizeType Position, SizeType CharCount = TNumericLimits<SizeType>::Max()) { *this = Mid(Position, CharCount); }
+	/** Modifies the view to be the middle part between any whitespace at the start and end. */
+	inline void TrimStartAndEndInline() { *this = TrimStartAndEnd(); }
+	/** Modifies the view to be the right part after any whitespace at the start. */
+	inline void TrimStartInline() { *this = TrimStart(); }
+	/** Modifies the view to be the left part before any whitespace at the end. */
+	inline void TrimEndInline() { *this = TrimEnd(); }
 
 	// Comparison
 
@@ -255,22 +267,28 @@ public:
 	using TStringViewImpl<ElementType, FAnsiStringView>::TStringViewImpl;
 };
 
-class FWideStringView : public TStringViewImpl<WIDECHAR, FWideStringView>
-{
-public:
-	using TStringViewImpl<ElementType, FWideStringView>::TStringViewImpl;
-};
+using FWideStringView = FStringView;
 
 template <> struct TIsContiguousContainer<FStringView> { enum { Value = true }; };
 template <> struct TIsContiguousContainer<FAnsiStringView> { enum { Value = true }; };
-template <> struct TIsContiguousContainer<FWideStringView> { enum { Value = true }; };
 
 constexpr inline SIZE_T GetNum(const FStringView& String) { return String.Len(); }
 constexpr inline SIZE_T GetNum(const FAnsiStringView& String) { return String.Len(); }
-constexpr inline SIZE_T GetNum(const FWideStringView& String) { return String.Len(); }
 
 constexpr inline FStringView operator "" _SV(const TCHAR* String, size_t Size) { return FStringView(String, Size); }
 constexpr inline FAnsiStringView operator "" _ASV(const ANSICHAR* String, size_t Size) { return FAnsiStringView(String, Size); }
 constexpr inline FWideStringView operator "" _WSV(const WIDECHAR* String, size_t Size) { return FWideStringView(String, Size); }
 
+//////////////////////////////////////////////////////////////////////////
+
+namespace StringViewPrivate
+{
+	template <typename CharType> struct TStringViewType;
+	template <> struct TStringViewType<ANSICHAR> { using Type = FAnsiStringView; };
+	template <> struct TStringViewType<WIDECHAR> { using Type = FWideStringView; };
+}
+
+template <typename CharType>
+using TStringView = typename StringViewPrivate::TStringViewType<CharType>::Type;
+
 #include "StringView.inl"
diff --git a/Engine/Source/Runtime/Core/Public/Containers/StringView.inl b/Engine/Source/Runtime/Core/Public/Containers/StringView.inl
--- a/Engine/Source/Runtime/Core/Public/Containers/StringView.inl
+++ b/Engine/Source/Runtime/Core/Public/Containers/StringView.inl
@@ -10,14 +10,20 @@ namespace StringViewPrivate
 	// These exported functions are implemented outside of the class to avoid the complexity
 	// of correctly exporting functions using many toolchains from a type that uses the CRTP.
 
-	template<class ViewType>
-	CORE_API int32 Compare(const ViewType& Lhs, const ViewType& Rhs, ESearchCase::Type SearchCase);
+	CORE_API int32 Compare(const FAnsiStringView& Lhs, const FAnsiStringView& Rhs, ESearchCase::Type SearchCase);
+	CORE_API int32 Compare(const FWideStringView& Lhs, const FWideStringView& Rhs, ESearchCase::Type SearchCase);
 
-	template<class ViewType>
-	CORE_API bool FindChar(const ViewType& InView, typename ViewType::ElementType InChar, typename ViewType::SizeType& OutIndex);
+	CORE_API bool FindChar(const FAnsiStringView& InView, ANSICHAR InChar, typename FAnsiStringView::SizeType& OutIndex);
+	CORE_API bool FindChar(const FWideStringView& InView, WIDECHAR InChar, typename FWideStringView::SizeType& OutIndex);
 
-	template<class ViewType>
-	CORE_API bool FindLastChar(const ViewType& InView, typename ViewType::ElementType InChar, typename ViewType::SizeType& OutIndex);
+	CORE_API bool FindLastChar(const FAnsiStringView& InView, ANSICHAR InChar, typename FAnsiStringView::SizeType& OutIndex);
+	CORE_API bool FindLastChar(const FWideStringView& InView, WIDECHAR InChar, typename FWideStringView::SizeType& OutIndex);
+
+	CORE_API FAnsiStringView TrimStart(const FAnsiStringView& InView);
+	CORE_API FWideStringView TrimStart(const FWideStringView& InView);
+
+	CORE_API FAnsiStringView TrimEnd(const FAnsiStringView& InView);
+	CORE_API FWideStringView TrimEnd(const FWideStringView& InView);
 }
 
 template <typename CharType, typename ViewType>
@@ -71,6 +77,24 @@ inline ViewType TStringViewImpl<CharType, ViewType>::Mid(SizeType Position, Size
 	return ViewType(DataPtr + Position, CharCount);
 }
 
+template <typename CharType, typename ViewType>
+inline ViewType TStringViewImpl<CharType, ViewType>::TrimStartAndEnd() const
+{
+	return TrimStart().TrimEnd();
+}
+
+template <typename CharType, typename ViewType>
+inline ViewType TStringViewImpl<CharType, ViewType>::TrimStart() const
+{
+	return StringViewPrivate::TrimStart(static_cast<const ViewType&>(*this));
+}
+
+template <typename CharType, typename ViewType>
+inline ViewType TStringViewImpl<CharType, ViewType>::TrimEnd() const
+{
+	return StringViewPrivate::TrimEnd(static_cast<const ViewType&>(*this));
+}
+
 template <typename CharType, typename ViewType>
 inline bool TStringViewImpl<CharType, ViewType>::Equals(const ViewType& Other, ESearchCase::Type SearchCase) const
 {
