Fix a compile error in StringViewTest when WIDECHAR is not wchar_t

#rnx

diff --git a/Engine/Source/Runtime/Core/Private/Tests/Containers/StringViewTest.cpp b/Engine/Source/Runtime/Core/Private/Tests/Containers/StringViewTest.cpp
new file mode 100644
--- /dev/null
+++ b/Engine/Source/Runtime/Core/Private/Tests/Containers/StringViewTest.cpp
@@ -0,0 +1,769 @@
+// Copyright 1998-2019 Epic Games, Inc. All Rights Reserved.
+
+#include "Containers/StringView.h"
+
+#include "Containers/UnrealString.h"
+#include "Misc/AutomationTest.h"
+
+#if WITH_DEV_AUTOMATION_TESTS
+
+static_assert(TIsSame<typename FStringView::ElementType, TCHAR>::Value, "FStringView must use TCHAR.");
+static_assert(TIsSame<typename FAnsiStringView::ElementType, ANSICHAR>::Value, "FAnsiStringView must use ANSICHAR.");
+static_assert(TIsSame<typename FWideStringView::ElementType, WIDECHAR>::Value, "FWideStringView must use WIDECHAR.");
+
+static_assert(TIsSame<FStringView, TStringView<TCHAR>>::Value, "FStringView must be the same as TStringView<TCHAR>.");
+static_assert(TIsSame<FAnsiStringView, TStringView<ANSICHAR>>::Value, "FAnsiStringView must be the same as TStringView<ANSICHAR>.");
+static_assert(TIsSame<FWideStringView, TStringView<WIDECHAR>>::Value, "FWideStringView must be the same as TStringView<WIDECHAR>.");
+
+static_assert(TIsContiguousContainer<FStringView>::Value, "FStringView must be a contiguous container.");
+static_assert(TIsContiguousContainer<FAnsiStringView>::Value, "FAnsiStringView must be a contiguous container.");
+static_assert(TIsContiguousContainer<FWideStringView>::Value, "FWideStringView must be a contiguous container.");
+
+#define TEST_NAME_ROOT "System.Core.StringView"
+constexpr const uint32 TestFlags = EAutomationTestFlags::ApplicationContextMask | EAutomationTestFlags::EngineFilter;
+
+IMPLEMENT_SIMPLE_AUTOMATION_TEST(FStringViewTestCtor, TEST_NAME_ROOT ".Ctor", TestFlags)
+bool FStringViewTestCtor::RunTest(const FString& Parameters)
+{
+	// Empty View
+	{
+		FStringView View( TEXT("") );
+		TestEqual(TEXT(""), View.Len(), 0);
+		TestTrue(TEXT("View.IsEmpty"), View.IsEmpty());
+	}
+
+	// Constructing from nullptr is supported; nullptr interpreted as empty string
+	{
+		FStringView View(nullptr);
+		TestEqual(TEXT(""), View.Len(), 0);
+		TestTrue(TEXT("View.IsEmpty"), View.IsEmpty());
+	}
+
+	// Create from a wchar literal
+	{
+		FStringView View(TEXT("Test Ctor"));
+		TestEqual(TEXT("View length"), View.Len(), FCStringWide::Strlen(TEXT("Test Ctor")));
+		TestEqual(TEXT("The result of Strncmp"), FCStringWide::Strncmp(View.GetData(), TEXT("Test Ctor"), View.Len()), 0);
+		TestFalse(TEXT("View.IsEmpty"), View.IsEmpty());
+	}
+
+	// Create from a sub section of a wchar literal
+	{
+		FStringView View(TEXT("Test SubSection Ctor"), 4);
+		TestEqual(TEXT("View length"), View.Len(), 4);
+		TestEqual(TEXT("The result of Strncmp"), FCStringWide::Strncmp(View.GetData(), TEXT("Test"), View.Len()), 0);
+		TestFalse(TEXT("View.IsEmpty"), View.IsEmpty());
+	}
+
+	// Create from a FString
+	{
+		FString String(TEXT("String Object"));
+		FStringView View(String);
+
+		TestEqual(TEXT("View length"), View.Len(), String.Len());
+		TestEqual(TEXT("The result of Strncmp"), FCStringWide::Strncmp(View.GetData(), *String, View.Len()), 0);
+		TestFalse(TEXT("View.IsEmpty"), View.IsEmpty());
+	}
+
+	// Create from a ansi literal
+	{
+		FAnsiStringView View("Test Ctor");
+		TestEqual(TEXT("View length"), View.Len(), FCStringAnsi::Strlen("Test Ctor"));
+		TestEqual(TEXT("The result of Strncmp"), FCStringAnsi::Strncmp(View.GetData(), "Test Ctor", View.Len()), 0);
+		TestFalse(TEXT("View.IsEmpty"), View.IsEmpty());
+	}
+
+	// Create from a sub section of an ansi literal
+	{
+		FAnsiStringView View("Test SubSection Ctor", 4);
+		TestEqual(TEXT("View length"), View.Len(), 4);
+		TestEqual(TEXT("The result of Strncmp"), FCStringAnsi::Strncmp(View.GetData(), "Test", View.Len()), 0);
+		TestFalse(TEXT("View.IsEmpty"), View.IsEmpty());
+	}
+
+	// Create using string view literals
+	{
+		FStringView View = TEXT("Test"_SV);
+		FAnsiStringView ViewAnsi = "Test"_ASV;
+		FWideStringView ViewWide = TEXT("Test"_WSV);
+	}
+
+	return true;
+}
+
+IMPLEMENT_SIMPLE_AUTOMATION_TEST(FStringViewTestIterators, TEST_NAME_ROOT ".Iterators", TestFlags)
+bool FStringViewTestIterators::RunTest(const FString& Parameters)
+{
+	// Iterate over a string view
+	{
+		const TCHAR* StringLiteralSrc = TEXT("Iterator!");
+		FStringView View(StringLiteralSrc);
+
+		for (TCHAR C : View)
+		{
+			TestTrue(TEXT("Iterators(0)-Iteration"), C == *StringLiteralSrc++);
+		}
+
+		// Make sure we iterated over the entire string
+		TestTrue(TEXT("Iterators(0-EndCheck"), *StringLiteralSrc == '\0');
+	}
+
+	// Iterate over a partial string view
+	{
+		const TCHAR* StringLiteralSrc = TEXT("Iterator|with extras!");
+		FStringView View(StringLiteralSrc, 8);
+
+		for (TCHAR C : View)
+		{
+			TestTrue(TEXT("Iterators(1)-Iteration"), C == *StringLiteralSrc++);
+		}
+
+		// Make sure we only iterated over the part of the string that the view represents
+		TestTrue(TEXT("Iterators(1)-EndCheck"), *StringLiteralSrc == '|');
+	}
+
+	return true;
+}
+
+IMPLEMENT_SIMPLE_AUTOMATION_TEST(FStringViewTestEquality, TEST_NAME_ROOT ".Equality", TestFlags)
+bool FStringViewTestEquality::RunTest(const FString& Parameters)
+{
+	const TCHAR* WideStringLiteralSrc = TEXT("String To Test!");
+	const TCHAR* WideStringLiteralLower = TEXT("string to test!");
+	const TCHAR* WideStringLiteralUpper = TEXT("STRING TO TEST!");
+	const TCHAR* WideStringLiteralShort = TEXT("String To");
+	const TCHAR* WideStringLiteralLonger = TEXT("String To Test! Extended");
+
+	FStringView WideView(WideStringLiteralSrc);
+
+	TestTrue(TEXT("Equality(0)"), WideView == WideStringLiteralSrc);
+	TestFalse(TEXT("Equality(1)"), WideView != WideStringLiteralSrc);
+	TestTrue(TEXT("Equality(2)"), WideView == WideStringLiteralLower);
+	TestFalse(TEXT("Equality(3)"), WideView != WideStringLiteralLower);
+	TestTrue(TEXT("Equality(4)"), WideView == WideStringLiteralUpper);
+	TestFalse(TEXT("Equality(5)"), WideView != WideStringLiteralUpper);
+	TestFalse(TEXT("Equality(6)"), WideView == WideStringLiteralShort);
+	TestTrue(TEXT("Equality(7)"), WideView != WideStringLiteralShort);
+	TestFalse(TEXT("Equality(8)"), WideView == WideStringLiteralLonger);
+	TestTrue(TEXT("Equality(9)"), WideView != WideStringLiteralLonger);
+
+	TestTrue(TEXT("Equality(10)"), WideStringLiteralSrc == WideView);
+	TestFalse(TEXT("Equality(11)"), WideStringLiteralSrc != WideView);
+	TestTrue(TEXT("Equality(12)"), WideStringLiteralLower == WideView);
+	TestFalse(TEXT("Equality(13)"), WideStringLiteralLower != WideView);
+	TestTrue(TEXT("Equality(14)"), WideStringLiteralUpper == WideView);
+	TestFalse(TEXT("Equality(15)"), WideStringLiteralUpper != WideView);
+	TestFalse(TEXT("Equality(16)"), WideStringLiteralShort == WideView);
+	TestTrue(TEXT("Equality(17)"), WideStringLiteralShort != WideView);
+	TestFalse(TEXT("Equality(18)"), WideStringLiteralLonger == WideView);
+	TestTrue(TEXT("Equality(19)"), WideStringLiteralLonger != WideView);
+
+	FString WideStringSrc = WideStringLiteralSrc;
+	FString WideStringLower = WideStringLiteralLower;
+	FString WideStringUpper = WideStringLiteralUpper;
+	FString WideStringShort = WideStringLiteralShort;
+	FString WideStringLonger = WideStringLiteralLonger;
+
+	TestTrue(TEXT("Equality(20)"), WideView == WideStringSrc);
+	TestFalse(TEXT("Equality(21)"), WideView != WideStringSrc);
+	TestTrue(TEXT("Equality(22)"), WideView == WideStringLower);
+	TestFalse(TEXT("Equality(23)"), WideView != WideStringLower);
+	TestTrue(TEXT("Equality(24)"), WideView == WideStringUpper);
+	TestFalse(TEXT("Equality(25)"), WideView != WideStringUpper);
+	TestFalse(TEXT("Equality(26)"), WideView == WideStringShort);
+	TestTrue(TEXT("Equality(27)"), WideView != WideStringShort);
+	TestFalse(TEXT("Equality(28)"), WideView == WideStringLonger);
+	TestTrue(TEXT("Equality(29)"), WideView != WideStringLonger);
+
+	TestTrue(TEXT("Equality(30)"), WideStringSrc == WideView);
+	TestFalse(TEXT("Equality(31)"), WideStringSrc != WideView);
+	TestTrue(TEXT("Equality(32)"), WideStringLower == WideView);
+	TestFalse(TEXT("Equality(33)"), WideStringLower != WideView);
+	TestTrue(TEXT("Equality(34)"), WideStringUpper == WideView);
+	TestFalse(TEXT("Equality(35)"), WideStringUpper != WideView);
+	TestFalse(TEXT("Equality(36)"), WideStringShort == WideView);
+	TestTrue(TEXT("Equality(37)"), WideStringShort != WideView);
+	TestFalse(TEXT("Equality(38)"), WideStringLonger == WideView);
+	TestTrue(TEXT("Equality(39)"), WideStringLonger != WideView);
+
+	FStringView IdenticalView(WideStringLiteralSrc);
+
+	TestTrue(TEXT("Equality(40a)"), WideView == IdenticalView);
+	TestFalse(TEXT("Equality(40b)"), WideView != IdenticalView);
+	TestTrue(TEXT("Equality(41a)"), IdenticalView == WideView);
+	TestFalse(TEXT("Equality(41b)"), IdenticalView != WideView);
+
+	// Views without null termination
+
+	FStringView ShortViewNoNull = WideView.Left(FStringView(WideStringLiteralShort).Len());
+
+	TestTrue(TEXT("Equality(42)"), ShortViewNoNull == WideStringLiteralShort);
+	TestFalse(TEXT("Equality(43)"), ShortViewNoNull != WideStringLiteralShort);
+	TestTrue(TEXT("Equality(44)"), WideStringLiteralShort == ShortViewNoNull);
+	TestFalse(TEXT("Equality(45)"), WideStringLiteralShort != ShortViewNoNull);
+	TestFalse(TEXT("Equality(46)"), ShortViewNoNull == WideStringLiteralSrc);
+	TestTrue(TEXT("Equality(47)"), ShortViewNoNull != WideStringLiteralSrc);
+	TestFalse(TEXT("Equality(48)"), WideStringLiteralSrc == ShortViewNoNull);
+	TestTrue(TEXT("Equality(49)"), WideStringLiteralSrc != ShortViewNoNull);
+
+	TestTrue(TEXT("Equality(50)"), ShortViewNoNull == WideStringShort);
+	TestFalse(TEXT("Equality(51)"), ShortViewNoNull != WideStringShort);
+	TestTrue(TEXT("Equality(52)"), WideStringShort == ShortViewNoNull);
+	TestFalse(TEXT("Equality(53)"), WideStringShort != ShortViewNoNull);
+	TestFalse(TEXT("Equality(54)"), ShortViewNoNull == WideStringSrc);
+	TestTrue(TEXT("Equality(55)"), ShortViewNoNull != WideStringSrc);
+	TestFalse(TEXT("Equality(56)"), WideStringSrc == ShortViewNoNull);
+	TestTrue(TEXT("Equality(57)"), WideStringSrc != ShortViewNoNull);
+
+	FStringView WideViewNoNull = FStringView(WideStringLiteralLonger).Left(WideView.Len());
+
+	TestTrue(TEXT("Equality(58)"), WideViewNoNull == WideStringLiteralSrc);
+	TestFalse(TEXT("Equality(59)"), WideViewNoNull != WideStringLiteralSrc);
+	TestTrue(TEXT("Equality(60)"), WideStringLiteralSrc == WideViewNoNull);
+	TestFalse(TEXT("Equality(61)"), WideStringLiteralSrc != WideViewNoNull);
+	TestFalse(TEXT("Equality(62)"), WideViewNoNull == WideStringLiteralLonger);
+	TestTrue(TEXT("Equality(63)"), WideViewNoNull != WideStringLiteralLonger);
+	TestFalse(TEXT("Equality(64)"), WideStringLiteralLonger == WideViewNoNull);
+	TestTrue(TEXT("Equality(65)"), WideStringLiteralLonger != WideViewNoNull);
+
+	TestTrue(TEXT("Equality(66)"), WideViewNoNull == WideStringLiteralSrc);
+	TestFalse(TEXT("Equality(67)"), WideViewNoNull != WideStringLiteralSrc);
+	TestTrue(TEXT("Equality(68)"), WideStringLiteralSrc == WideViewNoNull);
+	TestFalse(TEXT("Equality(69)"), WideStringLiteralSrc != WideViewNoNull);
+	TestFalse(TEXT("Equality(70)"), WideViewNoNull == WideStringLiteralLonger);
+	TestTrue(TEXT("Equality(71)"), WideViewNoNull != WideStringLiteralLonger);
+	TestFalse(TEXT("Equality(72)"), WideStringLiteralLonger == WideViewNoNull);
+	TestTrue(TEXT("Equality(73)"), WideStringLiteralLonger != WideViewNoNull);
+
+	return true;
+}
+
+IMPLEMENT_SIMPLE_AUTOMATION_TEST(FStringViewTestComparisonCaseSensitive, TEST_NAME_ROOT ".ComparisonCaseSensitive", TestFlags)
+bool FStringViewTestComparisonCaseSensitive::RunTest(const FString& Parameters)
+{
+	// Basic comparisons involving case
+	{
+		const TCHAR* WideStringLiteralSrc = TEXT("String To Test!");
+		const TCHAR* WideStringLiteralLower = TEXT("string to test!");
+		const TCHAR* WideStringLiteralUpper = TEXT("STRING TO TEST!");
+
+		FStringView WideView(WideStringLiteralSrc);
+
+		TestTrue(TEXT("ComparisonCaseSensitive(0)"), WideView.Compare(WideStringLiteralSrc, ESearchCase::CaseSensitive) == 0);
+		TestFalse(TEXT("ComparisonCaseSensitive(1)"), WideView.Compare(WideStringLiteralLower, ESearchCase::CaseSensitive) > 0);
+		TestFalse(TEXT("ComparisonCaseSensitive(2)"), WideView.Compare(WideStringLiteralUpper, ESearchCase::CaseSensitive) < 0);
+
+		FStringView EmptyView(TEXT(""));
+		TestTrue(TEXT("ComparisonCaseSensitive(3)"), WideView.Compare(EmptyView, ESearchCase::CaseSensitive) > 0);
+
+		FStringView IdenticalView(WideStringLiteralSrc);
+		TestTrue(TEXT("ComparisonCaseSensitive(4)"), WideView.Compare(IdenticalView, ESearchCase::CaseSensitive) == 0);
+
+	}
+
+	// Test comparisons of different lengths
+	{
+		const TCHAR* WideStringLiteralUpper = TEXT("ABCDEF");
+		const TCHAR* WideStringLiteralLower = TEXT("abcdef");
+
+		const TCHAR* WideStringLiteralUpperFirst = TEXT("ABCdef");
+		const TCHAR* WideStringLiteralLowerFirst = TEXT("abcDEF");
+
+		FStringView ViewLongUpper(WideStringLiteralUpper);
+		FStringView ViewLongLower(WideStringLiteralLower);
+
+		// Note that the characters after these views are in a different case, this will help catch over read issues
+		FStringView ViewShortUpper(WideStringLiteralUpperFirst, 3);
+		FStringView ViewShortLower(WideStringLiteralLowerFirst, 3);
+
+		// Same length, different cases
+		TestTrue(TEXT("ComparisonCaseSensitive(5)"), ViewLongUpper.Compare(ViewLongLower, ESearchCase::CaseSensitive) < 0);
+		TestTrue(TEXT("ComparisonCaseSensitive(6)"), ViewLongLower.Compare(ViewLongUpper, ESearchCase::CaseSensitive) > 0);
+
+		// Same case, different lengths
+		TestTrue(TEXT("ComparisonCaseSensitive(7)"), ViewLongUpper.Compare(ViewShortUpper, ESearchCase::CaseSensitive) > 0);
+		TestTrue(TEXT("ComparisonCaseSensitive(8)"), ViewShortUpper.Compare(ViewLongUpper, ESearchCase::CaseSensitive) < 0);
+
+		// Different length, different cases
+		TestTrue(TEXT("ComparisonCaseSensitive(9)"), ViewLongUpper.Compare(ViewShortLower, ESearchCase::CaseSensitive) < 0);
+		TestTrue(TEXT("ComparisonCaseSensitive(10)"), ViewShortLower.Compare(ViewLongUpper, ESearchCase::CaseSensitive) > 0);
+	}
+
+	return true;
+}
+
+IMPLEMENT_SIMPLE_AUTOMATION_TEST(FStringViewTestComparisonCaseInsensitive, TEST_NAME_ROOT ".ComparisonCaseInsensitive", TestFlags)
+bool FStringViewTestComparisonCaseInsensitive::RunTest(const FString& Parameters)
+{
+	// Basic comparisons involving case
+	{
+		const TCHAR* WideStringLiteralSrc = TEXT("String To Test!");
+		const TCHAR* WideStringLiteralLower = TEXT("string to test!");
+		const TCHAR* WideStringLiteralUpper = TEXT("STRING TO TEST!");
+
+		FStringView WideView(WideStringLiteralSrc);
+
+		TestTrue(TEXT("ComparisonCaseInsensitive(0)"), WideView.Compare(WideStringLiteralSrc, ESearchCase::IgnoreCase) == 0);
+		TestTrue(TEXT("ComparisonCaseInsensitive(1)"), WideView.Compare(WideStringLiteralLower, ESearchCase::IgnoreCase) == 0);
+		TestTrue(TEXT("ComparisonCaseInsensitive(2)"), WideView.Compare(WideStringLiteralUpper, ESearchCase::IgnoreCase) == 0);
+
+		FStringView EmptyView(TEXT(""));
+		TestTrue(TEXT("ComparisonCaseInsensitive(3)"), WideView.Compare(EmptyView, ESearchCase::IgnoreCase) > 0);
+
+		FStringView IdenticalView(WideStringLiteralSrc);
+		TestTrue(TEXT("ComparisonCaseInsensitive(4)"), WideView.Compare(IdenticalView, ESearchCase::IgnoreCase) == 0);
+
+	}
+
+	// Test comparisons of different lengths
+	{
+		const TCHAR* WideStringLiteralUpper = TEXT("ABCDEF");
+		const TCHAR* WideStringLiteralLower = TEXT("abcdef");
+
+		const TCHAR* WideStringLiteralUpperFirst = TEXT("ABCdef");
+		const TCHAR* WideStringLiteralLowerFirst = TEXT("abcDEF");
+
+		FStringView ViewLongUpper(WideStringLiteralUpper);
+		FStringView ViewLongLower(WideStringLiteralLower);
+
+		// Note that the characters after these views are in a different case, this will help catch over read issues
+		FStringView ViewShortUpper(WideStringLiteralUpperFirst, 3);
+		FStringView ViewShortLower(WideStringLiteralLowerFirst, 3);
+
+		// Same length, different cases
+		TestTrue(TEXT("ComparisonCaseInsensitive(5)"), ViewLongUpper.Compare(ViewLongLower, ESearchCase::IgnoreCase) == 0);
+		TestTrue(TEXT("ComparisonCaseInsensitive(6)"), ViewLongLower.Compare(ViewLongUpper, ESearchCase::IgnoreCase) == 0);
+
+		// Same case, different lengths
+		TestTrue(TEXT("ComparisonCaseInsensitive(7)"), ViewLongUpper.Compare(ViewShortUpper, ESearchCase::IgnoreCase) > 0);
+		TestTrue(TEXT("ComparisonCaseInsensitive(8)"), ViewShortUpper.Compare(ViewLongUpper, ESearchCase::IgnoreCase) < 0);
+
+		// Different length, different cases
+		TestTrue(TEXT("ComparisonCaseInsensitive(9)"), ViewLongUpper.Compare(ViewShortLower, ESearchCase::IgnoreCase) > 0);
+		TestTrue(TEXT("ComparisonCaseInsensitive(10)"), ViewShortLower.Compare(ViewLongUpper, ESearchCase::IgnoreCase) < 0);
+	}
+
+	return true;
+}
+
+IMPLEMENT_SIMPLE_AUTOMATION_TEST(FStringViewTestArrayAccessor, TEST_NAME_ROOT ".ArrayAccessor", TestFlags)
+bool FStringViewTestArrayAccessor::RunTest(const FString& Parameters)
+{
+	const TCHAR* SrcString = TEXT("String To Test");
+	FStringView View(SrcString);
+
+	for( int32 i = 0; i < View.Len(); ++i )
+	{
+		TestEqual(TEXT("the character accessed"), View[i], SrcString[i]);
+	}
+
+	return true;
+}
+
+IMPLEMENT_SIMPLE_AUTOMATION_TEST(FStringViewTestArrayModifiers, TEST_NAME_ROOT ".Modifiers", TestFlags)
+bool FStringViewTestArrayModifiers::RunTest(const FString& Parameters)
+{
+	const TCHAR* FullText = TEXT("PrefixSuffix");
+	const TCHAR* Prefix = TEXT("Prefix");
+	const TCHAR* Suffix = TEXT("Suffix");
+
+	// Remove prefix
+	{
+		FStringView View(FullText);
+		View.RemovePrefix(FCStringWide::Strlen(Prefix));
+
+		TestEqual(TEXT("View length"), View.Len(), FCStringWide::Strlen(Suffix));
+		TestEqual(TEXT("The result of Strncmp"), FCStringWide::Strncmp(View.GetData(), Suffix, View.Len()), 0);
+	}
+
+	// Remove suffix
+	{
+		FStringView View(FullText);
+		View.RemoveSuffix(FCStringWide::Strlen(Suffix));
+
+		TestEqual(TEXT("View length"), View.Len(), FCStringWide::Strlen(Prefix));
+		TestEqual(TEXT("The result of Strncmp"), FCStringWide::Strncmp(View.GetData(), Prefix, View.Len()), 0);
+	}
+
+	return true;
+}
+
+IMPLEMENT_SIMPLE_AUTOMATION_TEST(FStringViewTestStartsWith, TEST_NAME_ROOT ".StartsWith", TestFlags)
+bool FStringViewTestStartsWith::RunTest(const FString& Parameters)
+{
+	// Test an empty view
+	{
+		FStringView View( TEXT(""));
+		TestTrue(TEXT(" View.StartsWith"), View.StartsWith(TEXT("")));
+		TestFalse(TEXT(" View.StartsWith"), View.StartsWith(TEXT("Text")));
+		TestFalse(TEXT(" View.StartsWith"), View.StartsWith('A'));
+	}
+
+	// Test a valid view with the correct text
+	{
+		FStringView View(TEXT("String to test"));
+		TestTrue(TEXT(" View.StartsWith"), View.StartsWith(TEXT("String")));
+		TestTrue(TEXT(" View.StartsWith"), View.StartsWith('S'));
+	}
+
+	// Test a valid view with incorrect text
+	{
+		FStringView View(TEXT("String to test"));
+		TestFalse(TEXT(" View.StartsWith"), View.StartsWith(TEXT("test")));
+		TestFalse(TEXT(" View.StartsWith"), View.StartsWith('t'));
+	}
+
+	// Test a valid view with the correct text but with different case
+	{
+		FStringView View(TEXT("String to test"));
+		TestTrue(TEXT(" View.StartsWith"), View.StartsWith(TEXT("sTrInG")));
+
+		// Searching by char is case sensitive to keep compatibility with FString
+		TestFalse(TEXT(" View.StartsWith"), View.StartsWith('s'));
+	}
+
+	return true;
+}
+
+IMPLEMENT_SIMPLE_AUTOMATION_TEST(FStringViewTestEndsWith, TEST_NAME_ROOT ".EndsWith", TestFlags)
+bool FStringViewTestEndsWith::RunTest(const FString& Parameters)
+{
+	// Test an empty view
+	{
+		FStringView View(TEXT(""));
+		TestTrue(TEXT(" View.EndsWith"), View.EndsWith(TEXT("")));
+		TestFalse(TEXT(" View.EndsWith"), View.EndsWith(TEXT("Text")));
+		TestFalse(TEXT(" View.EndsWith"), View.EndsWith('A'));
+	}
+
+	// Test a valid view with the correct text
+	{
+		FStringView View(TEXT("String to test"));
+		TestTrue(TEXT(" View.EndsWith"), View.EndsWith(TEXT("test")));
+		TestTrue(TEXT(" View.EndsWith"), View.EndsWith('t'));
+	}
+
+	// Test a valid view with incorrect text
+	{
+		FStringView View(TEXT("String to test"));
+		TestFalse(TEXT(" View.EndsWith"), View.EndsWith(TEXT("String")));
+		TestFalse(TEXT(" View.EndsWith"), View.EndsWith('S'));
+	}
+
+	// Test a valid view with the correct text but with different case
+	{
+		FStringView View(TEXT("String to test"));
+		TestTrue(TEXT(" View.EndsWith"), View.EndsWith(TEXT("TeST")));
+
+		// Searching by char is case sensitive to keep compatibility with FString
+		TestFalse(TEXT(" View.EndsWith"), View.EndsWith('T')); 
+	}
+
+	return true;
+}
+
+IMPLEMENT_SIMPLE_AUTOMATION_TEST(FStringViewTestSubStr, TEST_NAME_ROOT ".SubStr", TestFlags)
+bool FStringViewTestSubStr::RunTest(const FString& Parameters)
+{
+	{
+		FStringView EmptyView(TEXT(""));
+		FStringView EmptyResult = EmptyView.SubStr(0, 10);
+		TestTrue(TEXT("FStringView::SubStr(0)"), EmptyResult.IsEmpty());
+
+		// The following line is commented out as it would fail an assert and currently we cannot test for this in unit tests 
+		// FStringView OutofBoundsResult = EmptyView.SubStr(1000, 10000); 
+		FStringView OutofBoundsResult = EmptyView.SubStr(0, 10000);
+		TestTrue(TEXT("FStringView::SubStr(1)"), OutofBoundsResult.IsEmpty());
+	}
+
+	{
+		FStringView View(TEXT("A test string|"));
+		FStringView Word0 = View.SubStr(0, 1);
+		FStringView Word1 = View.SubStr(2, 4);
+		FStringView Word2 = View.SubStr(7, 6);
+		FStringView NullTerminatorResult = View.SubStr(14, 1024);	// We can create a substr that starts at the end of the 
+																	// string since the null terminator is still valid
+		FStringView OutofBoundsResult = View.SubStr(0, 1024);
+
+		TestTrue(TEXT("FStringView::SubStr(2)"), FCString::Strncmp(Word0.GetData(), TEXT("A"), Word0.Len()) == 0);
+		TestTrue(TEXT("FStringView::SubStr(3)"), FCString::Strncmp(Word1.GetData(), TEXT("test"), Word1.Len()) == 0);
+		TestTrue(TEXT("FStringView::SubStr(4)"), FCString::Strncmp(Word2.GetData(), TEXT("string"), Word2.Len()) == 0);
+		TestTrue(TEXT("FStringView::SubStr(5)"), NullTerminatorResult.IsEmpty());
+		TestTrue(TEXT("FStringView::SubStr(6)"), View == OutofBoundsResult);
+	}
+
+	return true;
+}
+
+IMPLEMENT_SIMPLE_AUTOMATION_TEST(FStringViewTestLeft, TEST_NAME_ROOT ".Left", TestFlags)
+bool FStringViewTestLeft::RunTest(const FString& Parameters)
+{
+	{
+		FStringView EmptyView(TEXT(""));
+		FStringView EmptyResult = EmptyView.Left(0);
+		TestTrue(TEXT("FStringView::Left"), EmptyResult.IsEmpty());
+
+		FStringView OutofBoundsResult = EmptyView.Left(1024);
+		TestTrue(TEXT("FStringView::Left"), OutofBoundsResult.IsEmpty());
+	}
+	
+	{
+		FStringView View(TEXT("A test string padded"), 13); // "A test string" without null termination
+		FStringView Result = View.Left(8);
+
+		TestTrue(TEXT("FStringView::Left"), FCString::Strncmp(Result.GetData(), TEXT("A test s"), Result.Len()) == 0);
+
+		FStringView OutofBoundsResult = View.Left(1024);
+		TestTrue(TEXT("FStringView::Left"), FCString::Strncmp(OutofBoundsResult.GetData(), TEXT("A test string"), OutofBoundsResult.Len()) == 0);
+	}
+
+	return true;
+}
+
+IMPLEMENT_SIMPLE_AUTOMATION_TEST(FStringViewTestLeftChop, TEST_NAME_ROOT ".LeftChop", TestFlags)
+bool FStringViewTestLeftChop::RunTest(const FString& Parameters)
+{
+	{
+		FStringView EmptyView(TEXT(""));
+		FStringView EmptyResult = EmptyView.LeftChop(0);
+		TestTrue(TEXT("FStringView::LeftChop"), EmptyResult.IsEmpty());
+
+		FStringView OutofBoundsResult = EmptyView.LeftChop(1024);
+		TestTrue(TEXT("FStringView::LeftChop"), OutofBoundsResult.IsEmpty());
+	}
+
+	{
+		FStringView View(TEXT("A test string padded"), 13); // "A test string" without null termination
+		FStringView Result = View.LeftChop(5);
+
+		TestTrue(TEXT("FStringView::LeftChop"), FCString::Strncmp(Result.GetData(), TEXT("A test s"), Result.Len()) == 0);
+
+		FStringView OutofBoundsResult = View.LeftChop(1024);
+		TestTrue(TEXT("FStringView::LeftChop"), FCString::Strncmp(OutofBoundsResult.GetData(), TEXT("A test string"), OutofBoundsResult.Len()) == 0);
+	}
+
+	return true;
+}
+
+IMPLEMENT_SIMPLE_AUTOMATION_TEST(FStringViewTestRight, TEST_NAME_ROOT ".Right", TestFlags)
+bool FStringViewTestRight::RunTest(const FString& Parameters)
+{
+	{
+		FStringView EmptyView(TEXT(""));
+		FStringView EmptyResult = EmptyView.Right(0);
+		TestTrue(TEXT("FStringView::Right"), EmptyResult.IsEmpty());
+
+		FStringView OutofBoundsResult = EmptyView.Right(1024);
+		TestTrue(TEXT("FStringView::Right"), OutofBoundsResult.IsEmpty());
+	}
+
+	{
+		FStringView View(TEXT("A test string padded"), 13); // "A test string" without null termination
+		FStringView Result = View.Right(8);
+
+		TestTrue(TEXT("FStringView::Right"), FCString::Strncmp(Result.GetData(), TEXT("t string"), Result.Len()) == 0);
+
+		FStringView OutofBoundsResult = View.Right(1024);
+		TestTrue(TEXT("FStringView::Right"), FCString::Strncmp(OutofBoundsResult.GetData(), TEXT("A test string"), OutofBoundsResult.Len()) == 0);
+	}
+
+	return true;
+}
+
+IMPLEMENT_SIMPLE_AUTOMATION_TEST(FStringViewTestRightChop, TEST_NAME_ROOT ".RightChop", TestFlags)
+bool FStringViewTestRightChop::RunTest(const FString& Parameters)
+{
+	{
+		FStringView EmptyView(TEXT(""));
+		FStringView EmptyResult = EmptyView.RightChop(0);
+		TestTrue(TEXT("FStringView::RightChop"), EmptyResult.IsEmpty());
+
+		FStringView OutofBoundsResult = EmptyView.RightChop(1024);
+		TestTrue(TEXT("FStringView::RightChop"), OutofBoundsResult.IsEmpty());
+	}
+
+	{
+		FStringView View(TEXT("A test string padded"), 13); // "A test string" without null termination
+		FStringView Result = View.RightChop(3);
+
+		TestTrue(TEXT("FStringView::RightChop"), FCString::Strncmp(Result.GetData(), TEXT("est string"), Result.Len()) == 0);
+
+		FStringView OutofBoundsResult = View.RightChop(1024);
+		TestTrue(TEXT("FStringView::RightChop"), FCString::Strncmp(OutofBoundsResult.GetData(), TEXT("A test string"), OutofBoundsResult.Len()) == 0);
+	}
+
+	return true;
+}
+
+IMPLEMENT_SIMPLE_AUTOMATION_TEST(FStringViewTestMid, TEST_NAME_ROOT ".Mid", TestFlags)
+bool FStringViewTestMid::RunTest(const FString& Parameters)
+{
+	{
+		FStringView EmptyView(TEXT(""));
+		FStringView EmptyResult = EmptyView.Mid(0, 10);
+		TestTrue(TEXT("FStringView::Mid(0)"), EmptyResult.IsEmpty());
+
+		// The following line is commented out as it would fail an assert and currently we cannot test for this in unit tests 
+		// FStringView OutofBoundsResult = EmptyView.Mid(1000, 10000); 
+		FStringView OutofBoundsResult = EmptyView.Mid(0, 10000);
+		TestTrue(TEXT("FStringView::Mid(1)"), OutofBoundsResult.IsEmpty());
+	}
+
+	{
+		FStringView View(TEXT("A test string|"));
+		FStringView Word0 = View.Mid(0, 1);
+		FStringView Word1 = View.Mid(2, 4);
+		FStringView Word2 = View.Mid(7, 6);
+		FStringView NullTerminatorResult = View.Mid(14, 1024);	// We can call Mid with a position that starts at the end of the 
+																// string since the null terminator is still valid
+		FStringView OutofBoundsResult = View.Mid(0, 1024);
+
+		TestTrue(TEXT("FStringView::Mid(2)"), FCString::Strncmp(Word0.GetData(), TEXT("A"), Word0.Len()) == 0);
+		TestTrue(TEXT("FStringView::Mid(3)"), FCString::Strncmp(Word1.GetData(), TEXT("test"), Word1.Len()) == 0);
+		TestTrue(TEXT("FStringView::Mid(4)"), FCString::Strncmp(Word2.GetData(), TEXT("string"), Word2.Len()) == 0);
+		TestTrue(TEXT("FStringView::Mid(5)"), NullTerminatorResult.IsEmpty());
+		TestTrue(TEXT("FStringView::Mid(6)"), View == OutofBoundsResult);
+		TestTrue(TEXT("FStringView::Mid(7)"), View.Mid(512, 1024).IsEmpty());
+	}
+
+	return true;
+}
+
+IMPLEMENT_SIMPLE_AUTOMATION_TEST(FStringViewTestTrimStartAndEnd, TEST_NAME_ROOT ".TrimStartAndEnd", TestFlags)
+bool FStringViewTestTrimStartAndEnd::RunTest(const FString& Parameters)
+{
+	TestTrue(TEXT("FStringView::TrimStartAndEnd(\"\")"), TEXT(""_SV).TrimStartAndEnd().IsEmpty());
+	TestTrue(TEXT("FStringView::TrimStartAndEnd(\" \")"), TEXT(" "_SV).TrimStartAndEnd().IsEmpty());
+	TestTrue(TEXT("FStringView::TrimStartAndEnd(\"  \")"), TEXT("  "_SV).TrimStartAndEnd().IsEmpty());
+	TestTrue(TEXT("FStringView::TrimStartAndEnd(\" \\t\\r\\n\")"), TEXT(" \t\r\n"_SV).TrimStartAndEnd().IsEmpty());
+
+	TestEqual(TEXT("FStringView::TrimStartAndEnd(\"ABC123\")"), TEXT("ABC123"_SV).TrimStartAndEnd(), TEXT("ABC123"_SV));
+	TestEqual(TEXT("FStringView::TrimStartAndEnd(\"A \\t\\r\\nB\")"), TEXT("A \t\r\nB"_SV).TrimStartAndEnd(), TEXT("A \t\r\nB"_SV));
+	TestEqual(TEXT("FStringView::TrimStartAndEnd(\" \\t\\r\\nABC123\\n\\r\\t \")"), TEXT(" \t\r\nABC123\n\r\t "_SV).TrimStartAndEnd(), TEXT("ABC123"_SV));
+
+	return true;
+}
+
+IMPLEMENT_SIMPLE_AUTOMATION_TEST(FStringViewTestTrimStart, TEST_NAME_ROOT ".TrimStart", TestFlags)
+bool FStringViewTestTrimStart::RunTest(const FString& Parameters)
+{
+	TestTrue(TEXT("FStringView::TrimStart(\"\")"), TEXT(""_SV).TrimStart().IsEmpty());
+	TestTrue(TEXT("FStringView::TrimStart(\" \")"), TEXT(" "_SV).TrimStart().IsEmpty());
+	TestTrue(TEXT("FStringView::TrimStart(\"  \")"), TEXT("  "_SV).TrimStart().IsEmpty());
+	TestTrue(TEXT("FStringView::TrimStart(\" \\t\\r\\n\")"), TEXT(" \t\r\n"_SV).TrimStart().IsEmpty());
+
+	TestEqual(TEXT("FStringView::TrimStart(\"ABC123\")"), TEXT("ABC123"_SV).TrimStart(), TEXT("ABC123"_SV));
+	TestEqual(TEXT("FStringView::TrimStart(\"A \\t\\r\\nB\")"), TEXT("A \t\r\nB"_SV).TrimStart(), TEXT("A \t\r\nB"_SV));
+	TestEqual(TEXT("FStringView::TrimStart(\" \\t\\r\\nABC123\\n\\r\\t \")"), TEXT(" \t\r\nABC123\n\r\t "_SV).TrimStart(), TEXT("ABC123\n\r\t "_SV));
+
+	return true;
+}
+
+IMPLEMENT_SIMPLE_AUTOMATION_TEST(FStringViewTestTrimEnd, TEST_NAME_ROOT ".TrimEnd", TestFlags)
+bool FStringViewTestTrimEnd::RunTest(const FString& Parameters)
+{
+	TestTrue(TEXT("FStringView::TrimEnd(\"\")"), TEXT(""_SV).TrimEnd().IsEmpty());
+	TestTrue(TEXT("FStringView::TrimEnd(\" \")"), TEXT(" "_SV).TrimEnd().IsEmpty());
+	TestTrue(TEXT("FStringView::TrimEnd(\"  \")"), TEXT("  "_SV).TrimEnd().IsEmpty());
+	TestTrue(TEXT("FStringView::TrimEnd(\" \\t\\r\\n\")"), TEXT(" \t\r\n"_SV).TrimEnd().IsEmpty());
+
+	TestEqual(TEXT("FStringView::TrimEnd(\"ABC123\")"), TEXT("ABC123"_SV).TrimEnd(), TEXT("ABC123"_SV));
+	TestEqual(TEXT("FStringView::TrimEnd(\"A \\t\\r\\nB\")"), TEXT("A \t\r\nB"_SV).TrimEnd(), TEXT("A \t\r\nB"_SV));
+	TestEqual(TEXT("FStringView::TrimEnd(\" \\t\\r\\nABC123\\n\\r\\t \")"), TEXT(" \t\r\nABC123\n\r\t "_SV).TrimEnd(), TEXT(" \t\r\nABC123"_SV));
+
+	return true;
+}
+
+IMPLEMENT_SIMPLE_AUTOMATION_TEST(FStringViewTestFindChar, TEST_NAME_ROOT ".FindChar", TestFlags)
+bool FStringViewTestFindChar::RunTest(const FString& Parameters)
+{
+	FStringView EmptyView = TEXT("");
+	FStringView View = TEXT("aBce Fga");
+
+	{
+		int32 Index = INDEX_NONE;
+		TestFalse(TEXT("FStringView::FindChar-Return(0)"), EmptyView.FindChar(TEXT('a'), Index));
+		TestEqual(TEXT("FStringView::FindChar-Index(0)"), Index, INDEX_NONE);
+	}
+
+	{
+		int32 Index = INDEX_NONE;
+		TestTrue(TEXT("FStringView::FindChar-Return(1)"), View.FindChar(TEXT('a'), Index));
+		TestEqual(TEXT("FStringView::FindChar-Index(1)"), Index, 0);
+	}
+
+	{
+		int32 Index = INDEX_NONE;
+		TestTrue(TEXT("FStringView::FindChar-Return(2)"), View.FindChar(TEXT('F'), Index));
+		TestEqual(TEXT("FStringView::FindChar-Index(2)"), Index, 5);
+	}
+
+	{
+		int32 Index = INDEX_NONE;
+		TestFalse(TEXT("FStringView::FindChar-Return(3)"), View.FindChar(TEXT('A'), Index));
+		TestEqual(TEXT("FStringView::FindChar-Index(3)"), Index, INDEX_NONE);
+	}
+
+	{
+		int32 Index = INDEX_NONE;
+		TestFalse(TEXT("FStringView::FindChar-Return(4)"), View.FindChar(TEXT('d'), Index));
+		TestEqual(TEXT("FStringView::FindChar-Index(4)"), Index, INDEX_NONE);
+	}
+
+	{
+		int32 Index = INDEX_NONE;
+		TestTrue(TEXT("FStringView::FindChar-Return(5)"), View.FindChar(TEXT(' '), Index));
+		TestEqual(TEXT("FStringView::FindChar-Index(5)"), Index, 4);
+	}
+
+	return true; 
+}
+
+IMPLEMENT_SIMPLE_AUTOMATION_TEST(FStringViewTestFindLastChar, TEST_NAME_ROOT ".FindLastChar", TestFlags)
+bool FStringViewTestFindLastChar::RunTest(const FString& Parameters)
+{
+	FStringView EmptyView = TEXT("");
+	FStringView View = TEXT("aBce Fga");
+
+	{
+		int32 Index = INDEX_NONE;
+		TestFalse(TEXT("FStringView::FindChar-Return(0)"), EmptyView.FindLastChar(TEXT('a'), Index));
+		TestEqual(TEXT("FStringView::FindChar-Index(0)"), Index, INDEX_NONE);
+	}
+
+	{
+		int32 Index = INDEX_NONE;
+		TestTrue(TEXT("FStringView::FindChar-Return(1)"), View.FindLastChar(TEXT('a'), Index));
+		TestEqual(TEXT("FStringView::FindChar-Index(1)"), Index, 7);
+	}
+
+	{
+		int32 Index = INDEX_NONE;
+		TestTrue(TEXT("FStringView::FindChar-Return(2)"), View.FindLastChar(TEXT('B'), Index));
+		TestEqual(TEXT("FStringView::FindChar-Index(2)"), Index, 1);
+	}
+
+	{
+		int32 Index = INDEX_NONE;
+		TestFalse(TEXT("FStringView::FindChar-Return(3)"), View.FindLastChar(TEXT('A'), Index));
+		TestEqual(TEXT("FStringView::FindChar-Index(3)"), Index, INDEX_NONE);
+	}
+
+	{
+		int32 Index = INDEX_NONE;
+		TestFalse(TEXT("FStringView::FindChar-Return(4)"), View.FindLastChar(TEXT('d'), Index));
+		TestEqual(TEXT("FStringView::FindChar-Index(4)"), Index, INDEX_NONE);
+	}
+
+	{
+		int32 Index = INDEX_NONE;
+		TestTrue(TEXT("FStringView::FindChar-Return(5)"), View.FindLastChar(TEXT(' '), Index));
+		TestEqual(TEXT("FStringView::FindChar-Index(5)"), Index, 4);
+	}
+
+	return true; 
+}
+
+#undef TEST_NAME_ROOT
+
+#endif //WITH_DEV_AUTOMATION_TESTS
diff --git a/Engine/Source/Runtime/Core/Private/Tests/Misc/StringViewTest.cpp b/Engine/Source/Runtime/Core/Private/Tests/Misc/StringViewTest.cpp
deleted file mode 100644
