Compile fix for nonunity compile builds (missing header include)

#rnx

diff --git a/Engine/Source/Runtime/Core/Private/GenericPlatform/GenericPlatformIoDispatcher.cpp b/Engine/Source/Runtime/Core/Private/GenericPlatform/GenericPlatformIoDispatcher.cpp
--- a/Engine/Source/Runtime/Core/Private/GenericPlatform/GenericPlatformIoDispatcher.cpp
+++ b/Engine/Source/Runtime/Core/Private/GenericPlatform/GenericPlatformIoDispatcher.cpp
@@ -4,6 +4,7 @@
 #include "IO/IoDispatcherFileBackend.h"
 #include "Async/AsyncFileHandle.h"
 #include "HAL/Event.h"
+#include "Misc/ScopeLock.h"
 
 FGenericIoDispatcherEventQueue::FGenericIoDispatcherEventQueue()
 	: Event(FPlatformProcess::GetSynchEventFromPool())
