GitHub #6415 : Added ToString() to TInterval

diff --git a/Engine/Source/Runtime/Core/Public/Math/Interval.h b/Engine/Source/Runtime/Core/Public/Math/Interval.h
--- a/Engine/Source/Runtime/Core/Public/Math/Interval.h
+++ b/Engine/Source/Runtime/Core/Public/Math/Interval.h
@@ -179,6 +179,16 @@ public:
 		return ElementType();
 	}
 
+	/**
+	 * Get a textual representation of the interval.
+	 *
+	 * @return Text describing the interval.
+	 */
+	FString ToString() const
+	{
+		return FString::Printf(TEXT("Min=%s Max=%s"), *LexToString(Min), *LexToString(Max));
+	}
+
 public:
 
 	/**
@@ -267,3 +277,17 @@ public:
 
 DEFINE_INTERVAL_WRAPPER_STRUCT(FFloatInterval, float)
 DEFINE_INTERVAL_WRAPPER_STRUCT(FInt32Interval, int32)
+
+template<>
+inline FString TInterval<float>::ToString() const
+{
+	// This is more efficient than going through LexToString
+	return FString::Printf(TEXT("Min=%f Max=%f"), Min, Max);
+}
+
+template<>
+inline FString TInterval<int32>::ToString() const
+{
+	// This is more efficient than going through LexToString
+	return FString::Printf(TEXT("Min=%d Max=%d"), Min, Max);
+}
