Fix for missing include.

#rnx

diff --git a/Engine/Source/Runtime/Core/Public/Math/Interval.h b/Engine/Source/Runtime/Core/Public/Math/Interval.h
--- a/Engine/Source/Runtime/Core/Public/Math/Interval.h
+++ b/Engine/Source/Runtime/Core/Public/Math/Interval.h
@@ -7,6 +7,7 @@
 #include "Templates/UnrealTypeTraits.h"
 #include "Math/NumericLimits.h"
 #include "Math/UnrealMathUtility.h"
+#include "Containers/UnrealString.h"
 
 /**
  * Type traits for Arithmetic interval.
