Fix IoDispatcher not being able to find the ucas/utoc files if a sandbox file system is enabled.

- Change from using IPlatformFile::GetPlatformPhysical to FPlatformFileManager::Get().GetPlatformFile which will allow us to use the Sandbox file system if it is enabled.
- Delete the file handle if we are not going to return it (due to failure case) to avoid leaking the handle.

#rnx

diff --git a/Engine/Source/Runtime/Core/Private/GenericPlatform/GenericPlatformIoDispatcher.cpp b/Engine/Source/Runtime/Core/Private/GenericPlatform/GenericPlatformIoDispatcher.cpp
--- a/Engine/Source/Runtime/Core/Private/GenericPlatform/GenericPlatformIoDispatcher.cpp
+++ b/Engine/Source/Runtime/Core/Private/GenericPlatform/GenericPlatformIoDispatcher.cpp
@@ -4,6 +4,7 @@
 #include "IO/IoDispatcherFileBackend.h"
 #include "Async/AsyncFileHandle.h"
 #include "HAL/Event.h"
+#include "HAL/PlatformFileManager.h"
 #include "Misc/ScopeLock.h"
 
 FGenericIoDispatcherEventQueue::FGenericIoDispatcherEventQueue()
@@ -35,12 +36,13 @@ FGenericFileIoStoreImpl::FGenericFileIoStoreImpl(FGenericIoDispatcherEventQueue&
 
 bool FGenericFileIoStoreImpl::OpenContainer(const TCHAR* ContainerFilePath, uint64& ContainerFileHandle, uint64& ContainerFileSize)
 {
-	IPlatformFile& Ipf = IPlatformFile::GetPlatformPhysical();
+	IPlatformFile& Ipf = FPlatformFileManager::Get().GetPlatformFile();
 	IAsyncReadFileHandle* FileHandle = Ipf.OpenAsyncRead(ContainerFilePath);
 	if (!FileHandle)
 	{
 		return false;
 	}
+
 	int64 SizeResult = -1;
 	FEvent* Event = FPlatformProcess::GetSynchEventFromPool();
 	FAsyncFileCallBack SizeCallback = [Event, &SizeResult](bool bWasCancelled, IAsyncReadRequest* Request)
@@ -60,6 +62,7 @@ bool FGenericFileIoStoreImpl::OpenContainer(const TCHAR* ContainerFilePath, uint
 	}
 	else
 	{
+		delete FileHandle;
 		return false;
 	}
 }
diff --git a/Engine/Source/Runtime/Core/Private/IO/IoDispatcherFileBackend.cpp b/Engine/Source/Runtime/Core/Private/IO/IoDispatcherFileBackend.cpp
--- a/Engine/Source/Runtime/Core/Private/IO/IoDispatcherFileBackend.cpp
+++ b/Engine/Source/Runtime/Core/Private/IO/IoDispatcherFileBackend.cpp
@@ -3,7 +3,8 @@
 #include "IO/IoDispatcherFileBackend.h"
 #include "Misc/ScopeRWLock.h"
 #include "ProfilingDebugging/CountersTrace.h"
-#include "HAL/PlatformFile.h"
+#include "HAL/PlatformFileManager.h"
+#include "GenericPlatform/GenericPlatformFile.h"
 
 TRACE_DECLARE_INT_COUNTER(IoDispatcherCacheHitsCold, TEXT("IoDispatcher/CacheHitsCold"));
 TRACE_DECLARE_INT_COUNTER(IoDispatcherCacheHitsHot, TEXT("IoDispatcher/CacheHitsHot"));
@@ -16,7 +17,7 @@ FFileIoStoreReader::FFileIoStoreReader(FFileIoStoreImpl& InPlatformImpl)
 
 FIoStatus FFileIoStoreReader::Initialize(const TCHAR* RootPath)
 {
-	IPlatformFile& Ipf = IPlatformFile::GetPlatformPhysical();
+	IPlatformFile& Ipf = FPlatformFileManager::Get().GetPlatformFile();
 
 	TStringBuilder<256> ContainerFilePath;
 	ContainerFilePath.Append(RootPath);
