Fixed static analysis warning

diff --git a/Engine/Source/Runtime/Core/Private/GenericPlatform/GenericPlatformIoDispatcher.cpp b/Engine/Source/Runtime/Core/Private/GenericPlatform/GenericPlatformIoDispatcher.cpp
--- a/Engine/Source/Runtime/Core/Private/GenericPlatform/GenericPlatformIoDispatcher.cpp
+++ b/Engine/Source/Runtime/Core/Private/GenericPlatform/GenericPlatformIoDispatcher.cpp
@@ -42,18 +42,20 @@ bool FGenericFileIoStoreImpl::OpenContainer(const TCHAR* ContainerFilePath, uint
 	{
 		return false;
 	}
-
-	int64 SizeResult = -1;
 	FEvent* Event = FPlatformProcess::GetSynchEventFromPool();
-	FAsyncFileCallBack SizeCallback = [Event, &SizeResult](bool bWasCancelled, IAsyncReadRequest* Request)
+	FAsyncFileCallBack SizeCallback = [Event](bool bWasCancelled, IAsyncReadRequest*)
 	{
 		check(!bWasCancelled);
-		SizeResult = Request->GetSizeResults();
 		Event->Trigger();
 	};
 	TUniquePtr<IAsyncReadRequest> SizeRequest(FileHandle->SizeRequest(&SizeCallback));
+	if (!SizeRequest.IsValid())
+	{
+		return false;
+	}
 	Event->Wait();
 	FPlatformProcess::ReturnSynchEventToPool(Event);
+	int64 SizeResult = SizeRequest->GetSizeResults();
 	if (SizeResult > 0)
 	{
 		ContainerFileHandle = reinterpret_cast<UPTRINT>(FileHandle);
