Remove the CRTP from TStringViewImpl

Use TStringView<T> to determine the required view type without it being provided as a template parameter to TStringViewImpl.

diff --git a/Engine/Source/Runtime/Core/Private/Containers/StringView.cpp b/Engine/Source/Runtime/Core/Private/Containers/StringView.cpp
--- a/Engine/Source/Runtime/Core/Private/Containers/StringView.cpp
+++ b/Engine/Source/Runtime/Core/Private/Containers/StringView.cpp
@@ -2,155 +2,95 @@
 
 #include "Containers/StringView.h"
 
-template class TStringViewImpl<ANSICHAR, FAnsiStringView>;
-template class TStringViewImpl<WIDECHAR, FWideStringView>;
-
-namespace StringViewPrivate
+template <typename CharType>
+int32 TStringViewImpl<CharType>::Compare(const ViewType& Rhs, const ESearchCase::Type SearchCase) const
 {
-	template <typename ViewType>
-	inline int32 CompareImpl(const ViewType& Lhs, const ViewType& Rhs, ESearchCase::Type SearchCase)
-	{
-		const typename ViewType::SizeType LhsLen = Lhs.Len();
-		const typename ViewType::SizeType RhsLen = Rhs.Len();
-		const typename ViewType::SizeType MinLen = FMath::Min(LhsLen, RhsLen);
-
-		int Result;
-		if (SearchCase == ESearchCase::CaseSensitive)
-		{
-			Result = TCString<typename ViewType::ElementType>::Strncmp(Lhs.GetData(), Rhs.GetData(), MinLen);
-		}
-		else
-		{
-			Result = TCString<typename ViewType::ElementType>::Strnicmp(Lhs.GetData(), Rhs.GetData(), MinLen);
-		}
+	const SizeType LhsLen = Len();
+	const SizeType RhsLen = Rhs.Len();
+	const SizeType MinLen = FMath::Min(LhsLen, RhsLen);
 
-		if (Result != 0 || LhsLen == RhsLen)
-		{
-			return Result;
-		}
-
-		return LhsLen < RhsLen ? -1 : 1;
+	int Result;
+	if (SearchCase == ESearchCase::CaseSensitive)
+	{
+		Result = TCString<typename ViewType::ElementType>::Strncmp(GetData(), Rhs.GetData(), MinLen);
 	}
-
-	int32 Compare(const FAnsiStringView& Lhs, const FAnsiStringView& Rhs, ESearchCase::Type SearchCase)
+	else
 	{
-		return CompareImpl(Lhs, Rhs, SearchCase);
+		Result = TCString<typename ViewType::ElementType>::Strnicmp(GetData(), Rhs.GetData(), MinLen);
 	}
 
-	int32 Compare(const FWideStringView& Lhs, const FWideStringView& Rhs, ESearchCase::Type SearchCase)
+	if (Result != 0 || LhsLen == RhsLen)
 	{
-		return CompareImpl(Lhs, Rhs, SearchCase);
+		return Result;
 	}
 
-	template <typename ViewType>
-	inline bool FindCharImpl(const ViewType& InView, typename ViewType::ElementType InChar, typename ViewType::SizeType& OutIndex)
-	{
-		const typename ViewType::ElementType* DataPtr = InView.GetData();
-		const typename ViewType::SizeType Size = InView.Len();
+	return LhsLen < RhsLen ? -1 : 1;
+}
 
-		for (typename ViewType::SizeType i = 0; i < Size; ++i)
+template <typename CharType>
+bool TStringViewImpl<CharType>::FindChar(const ElementType InChar, SizeType& OutIndex) const
+{
+	for (SizeType i = 0; i < Size; ++i)
+	{
+		if (DataPtr[i] == InChar)
 		{
-			if (DataPtr[i] == InChar)
-			{
-				OutIndex = i;
-				return true;
-			}
+			OutIndex = i;
+			return true;
 		}
-
-		OutIndex = INDEX_NONE;
-		return false;
 	}
 
-	bool FindChar(const FAnsiStringView& InView, ANSICHAR InChar, typename FAnsiStringView::SizeType& OutIndex)
-	{
-		return FindCharImpl(InView, InChar, OutIndex);
-	}
+	OutIndex = INDEX_NONE;
+	return false;
+}
 
-	bool FindChar(const FWideStringView& InView, WIDECHAR InChar, typename FWideStringView::SizeType& OutIndex)
-	{
-		return FindCharImpl(InView, InChar, OutIndex);
-	}
 
-	template <typename ViewType>
-	bool FindLastCharImpl(const ViewType& InView, typename ViewType::ElementType InChar, typename ViewType::SizeType& OutIndex)
+template <typename CharType>
+bool TStringViewImpl<CharType>::FindLastChar(const ElementType InChar, SizeType& OutIndex) const
+{
+	if (Size == 0)
 	{
-		const typename ViewType::ElementType* DataPtr = InView.GetData();
-		const typename ViewType::SizeType Size = InView.Len();
-
-		if (Size == 0)
-		{
-			OutIndex = INDEX_NONE;
-			return false;
-		}
-
-		for (typename ViewType::SizeType i = Size - 1; i >= 0; --i)
-		{
-			if (DataPtr[i] == InChar)
-			{
-				OutIndex = i;
-				return true;
-			}
-		}
-
 		OutIndex = INDEX_NONE;
 		return false;
 	}
 
-	bool FindLastChar(const FAnsiStringView& InView, ANSICHAR InChar, typename FAnsiStringView::SizeType& OutIndex)
+	for (SizeType i = Size - 1; i >= 0; --i)
 	{
-		return FindLastCharImpl(InView, InChar, OutIndex);
-	}
-
-	bool FindLastChar(const FWideStringView& InView, WIDECHAR InChar, typename FWideStringView::SizeType& OutIndex)
-	{
-		return FindLastCharImpl(InView, InChar, OutIndex);
-	}
-
-	template <typename ViewType>
-	inline ViewType TrimStartImpl(ViewType InView)
-	{
-		typename ViewType::SizeType SpaceCount = 0;
-		for (typename ViewType::ElementType Char : InView)
+		if (DataPtr[i] == InChar)
 		{
-			if (!TChar<typename ViewType::ElementType>::IsWhitespace(Char))
-			{
-				break;
-			}
-			++SpaceCount;
+			OutIndex = i;
+			return true;
 		}
-		InView.RemovePrefix(SpaceCount);
-		return InView;
 	}
 
-	FAnsiStringView TrimStart(const FAnsiStringView& InView)
-	{
-		return TrimStartImpl(InView);
-	}
-
-	FWideStringView TrimStart(const FWideStringView& InView)
-	{
-		return TrimStartImpl(InView);
-	}
+	OutIndex = INDEX_NONE;
+	return false;
+}
 
-	template <typename ViewType>
-	inline ViewType TrimEndImpl(const ViewType& InView)
+template <typename CharType>
+TStringView<CharType> TStringViewImpl<CharType>::TrimStart() const
+{
+	SizeType SpaceCount = 0;
+	for (ElementType Char : *this)
 	{
-		const typename ViewType::ElementType* DataPtr = InView.GetData();
-		typename ViewType::SizeType Len = InView.Len();
-		while (Len && TChar<typename ViewType::ElementType>::IsWhitespace(DataPtr[Len - 1]))
+		if (!TChar<ElementType>::IsWhitespace(Char))
 		{
-			--Len;
+			break;
 		}
-		return ViewType(DataPtr, Len);
-	}
-
-	FAnsiStringView TrimEnd(const FAnsiStringView& InView)
-	{
-		return TrimEndImpl(InView);
+		++SpaceCount;
 	}
+	return ViewType(DataPtr + SpaceCount, Size - SpaceCount);
+}
 
-	FWideStringView TrimEnd(const FWideStringView& InView)
+template <typename CharType>
+TStringView<CharType> TStringViewImpl<CharType>::TrimEnd() const
+{
+	SizeType NewSize = Size;
+	while (NewSize && TChar<ElementType>::IsWhitespace(DataPtr[NewSize - 1]))
 	{
-		return TrimEndImpl(InView);
+		--NewSize;
 	}
+	return ViewType(DataPtr, NewSize);
 }
+
+template class TStringViewImpl<ANSICHAR>;
+template class TStringViewImpl<WIDECHAR>;
diff --git a/Engine/Source/Runtime/Core/Public/Containers/StringView.h b/Engine/Source/Runtime/Core/Public/Containers/StringView.h
--- a/Engine/Source/Runtime/Core/Public/Containers/StringView.h
+++ b/Engine/Source/Runtime/Core/Public/Containers/StringView.h
@@ -12,6 +12,9 @@
 #include "Templates/UnrealTemplate.h"
 #include "Traits/IsContiguousContainer.h"
 
+class FAnsiStringView;
+class FStringView;
+
 namespace StringViewPrivate
 {
 	/**
@@ -21,8 +24,16 @@ namespace StringViewPrivate
 	struct TIsCompatibleRangeType : TIsSame<ElementType, typename TRemoveCV<typename TRemovePointer<decltype(GetData(DeclVal<RangeType&>()))>::Type>::Type>
 	{
 	};
+
+	template <typename CharType> struct TStringViewType;
+	template <> struct TStringViewType<ANSICHAR> { using Type = FAnsiStringView; };
+	template <> struct TStringViewType<TCHAR> { using Type = FStringView; };
 }
 
+/** The string view type for a given character type. */
+template <typename CharType>
+using TStringView = typename StringViewPrivate::TStringViewType<CharType>::Type;
+
 /** String View
 
 	* A string view is implicitly constructible from const char* style strings and
@@ -64,13 +75,13 @@ namespace StringViewPrivate
 
   */
 
-template <typename CharType, typename DerivedViewType>
+template <typename CharType>
 class TStringViewImpl
 {
 public:
 	using ElementType = CharType;
 	using SizeType = int32;
-	using ViewType = DerivedViewType;
+	using ViewType = TStringView<ElementType>;
 
 	template <typename RangeType>
 	using TIsCompatibleRangeType = TAnd<TIsContiguousContainer<RangeType>, StringViewPrivate::TIsCompatibleRangeType<RangeType, ElementType>>;
@@ -157,9 +168,9 @@ public:
 	/** Returns the middle part of the view between any whitespace at the start and end. */
 	inline ViewType TrimStartAndEnd() const;
 	/** Returns the right part of the view after any whitespace at the start. */
-	inline ViewType TrimStart() const;
+	CORE_API ViewType TrimStart() const;
 	/** Returns the left part of the view before any whitespace at the end. */
-	inline ViewType TrimEnd() const;
+	CORE_API ViewType TrimEnd() const;
 
 	/** Modifies the view to be the given number of characters from the left. */
 	inline void LeftInlineInline(SizeType CharCount) { *this = Left(CharCount); }
@@ -194,7 +205,7 @@ public:
 	 *
 	 * @return 0 is equal, negative if this view is less, positive if this view is greater.
 	 */
-	inline int32 Compare(const ViewType& Other, ESearchCase::Type SearchCase = ESearchCase::CaseSensitive) const;
+	CORE_API int32 Compare(const ViewType& Other, ESearchCase::Type SearchCase = ESearchCase::CaseSensitive) const;
 
 	/** Returns whether this view starts with the prefix character compared case-sensitively. */
 	inline bool StartsWith(CharType Prefix) const { return Size >= 1 && DataPtr[0] == Prefix; }
@@ -216,7 +227,7 @@ public:
 	 *
 	 * @return Whether the character was found in the view.
 	 */
-	inline bool FindChar(CharType InChar, SizeType& OutIndex) const;
+	CORE_API bool FindChar(CharType InChar, SizeType& OutIndex) const;
 
 	/**
 	 * Search the view for the last occurrence of a character.
@@ -226,7 +237,7 @@ public:
 	 *
 	 * @return Whether the character was found in the view.
 	 */
-	inline bool FindLastChar(CharType InChar, SizeType& OutIndex) const;
+	CORE_API bool FindLastChar(CharType InChar, SizeType& OutIndex) const;
 
 public:
 	/**
@@ -241,54 +252,48 @@ protected:
 	SizeType Size = 0;
 };
 
-template <typename CharType, typename ViewType>
-struct TIsContiguousContainer<TStringViewImpl<CharType, ViewType>>
+template <typename CharType>
+struct TIsContiguousContainer<TStringViewImpl<CharType>>
 {
 	enum { Value = true };
 };
 
-template <typename CharType, typename ViewType>
-constexpr inline SIZE_T GetNum(const TStringViewImpl<CharType, ViewType>& String)
+template <typename CharType>
+constexpr inline SIZE_T GetNum(const TStringViewImpl<CharType>& String)
 {
 	return String.Len();
 }
 
 //////////////////////////////////////////////////////////////////////////
 
-class FStringView : public TStringViewImpl<TCHAR, FStringView>
+class FStringView : public TStringViewImpl<TCHAR>
 {
 public:
-	using TStringViewImpl<ElementType, FStringView>::TStringViewImpl;
+	using TStringViewImpl<ElementType>::TStringViewImpl;
 };
 
-class FAnsiStringView : public TStringViewImpl<ANSICHAR, FAnsiStringView>
+template <> struct TIsContiguousContainer<FStringView> { enum { Value = true }; };
+
+constexpr inline FStringView operator "" _SV(const TCHAR* String, size_t Size) { return FStringView(String, Size); }
+
+//////////////////////////////////////////////////////////////////////////
+
+class FAnsiStringView : public TStringViewImpl<ANSICHAR>
 {
 public:
-	using TStringViewImpl<ElementType, FAnsiStringView>::TStringViewImpl;
+	using TStringViewImpl<ElementType>::TStringViewImpl;
 };
 
-using FWideStringView = FStringView;
-
-template <> struct TIsContiguousContainer<FStringView> { enum { Value = true }; };
 template <> struct TIsContiguousContainer<FAnsiStringView> { enum { Value = true }; };
 
-constexpr inline SIZE_T GetNum(const FStringView& String) { return String.Len(); }
-constexpr inline SIZE_T GetNum(const FAnsiStringView& String) { return String.Len(); }
-
-constexpr inline FStringView operator "" _SV(const TCHAR* String, size_t Size) { return FStringView(String, Size); }
 constexpr inline FAnsiStringView operator "" _ASV(const ANSICHAR* String, size_t Size) { return FAnsiStringView(String, Size); }
-constexpr inline FWideStringView operator "" _WSV(const WIDECHAR* String, size_t Size) { return FWideStringView(String, Size); }
 
 //////////////////////////////////////////////////////////////////////////
 
-namespace StringViewPrivate
-{
-	template <typename CharType> struct TStringViewType;
-	template <> struct TStringViewType<ANSICHAR> { using Type = FAnsiStringView; };
-	template <> struct TStringViewType<WIDECHAR> { using Type = FWideStringView; };
-}
+using FWideStringView = FStringView;
 
-template <typename CharType>
-using TStringView = typename StringViewPrivate::TStringViewType<CharType>::Type;
+constexpr inline FWideStringView operator "" _WSV(const WIDECHAR* String, size_t Size) { return FWideStringView(String, Size); }
+
+//////////////////////////////////////////////////////////////////////////
 
 #include "StringView.inl"
diff --git a/Engine/Source/Runtime/Core/Public/Containers/StringView.inl b/Engine/Source/Runtime/Core/Public/Containers/StringView.inl
--- a/Engine/Source/Runtime/Core/Public/Containers/StringView.inl
+++ b/Engine/Source/Runtime/Core/Public/Containers/StringView.inl
@@ -5,70 +5,49 @@
 #include "Templates/UnrealTypeTraits.h"
 #include "Traits/IntType.h"
 
-namespace StringViewPrivate
-{
-	// These exported functions are implemented outside of the class to avoid the complexity
-	// of correctly exporting functions using many toolchains from a type that uses the CRTP.
-
-	CORE_API int32 Compare(const FAnsiStringView& Lhs, const FAnsiStringView& Rhs, ESearchCase::Type SearchCase);
-	CORE_API int32 Compare(const FWideStringView& Lhs, const FWideStringView& Rhs, ESearchCase::Type SearchCase);
-
-	CORE_API bool FindChar(const FAnsiStringView& InView, ANSICHAR InChar, typename FAnsiStringView::SizeType& OutIndex);
-	CORE_API bool FindChar(const FWideStringView& InView, WIDECHAR InChar, typename FWideStringView::SizeType& OutIndex);
-
-	CORE_API bool FindLastChar(const FAnsiStringView& InView, ANSICHAR InChar, typename FAnsiStringView::SizeType& OutIndex);
-	CORE_API bool FindLastChar(const FWideStringView& InView, WIDECHAR InChar, typename FWideStringView::SizeType& OutIndex);
-
-	CORE_API FAnsiStringView TrimStart(const FAnsiStringView& InView);
-	CORE_API FWideStringView TrimStart(const FWideStringView& InView);
-
-	CORE_API FAnsiStringView TrimEnd(const FAnsiStringView& InView);
-	CORE_API FWideStringView TrimEnd(const FWideStringView& InView);
-}
-
-template <typename CharType, typename ViewType>
-inline const CharType& TStringViewImpl<CharType, ViewType>::operator[](SizeType Index) const
+template <typename CharType>
+inline const CharType& TStringViewImpl<CharType>::operator[](SizeType Index) const
 {
 	checkf(Index >= 0 && Index < Size, TEXT("Index out of bounds on StringView: index %i on a view with a length of %i"), Index, Size);
 	return DataPtr[Index];
 }
 
-template <typename CharType, typename ViewType>
-inline typename TStringViewImpl<CharType, ViewType>::SizeType TStringViewImpl<CharType, ViewType>::CopyString(CharType* Dest, SizeType CharCount, SizeType Position) const
+template <typename CharType>
+inline typename TStringViewImpl<CharType>::SizeType TStringViewImpl<CharType>::CopyString(CharType* Dest, SizeType CharCount, SizeType Position) const
 {
 	const  SizeType CopyCount = FMath::Min(Size - Position, CharCount);
 	FMemory::Memcpy(Dest, DataPtr + Position, CopyCount);
 	return CopyCount;
 }
 
-template <typename CharType, typename ViewType>
-inline ViewType TStringViewImpl<CharType, ViewType>::Left(SizeType CharCount) const
+template <typename CharType>
+inline TStringView<CharType> TStringViewImpl<CharType>::Left(SizeType CharCount) const
 {
 	return ViewType(DataPtr, FMath::Clamp(CharCount, 0, Size));
 }
 
-template <typename CharType, typename ViewType>
-inline ViewType TStringViewImpl<CharType, ViewType>::LeftChop(SizeType CharCount) const
+template <typename CharType>
+inline TStringView<CharType> TStringViewImpl<CharType>::LeftChop(SizeType CharCount) const
 {
 	return ViewType(DataPtr, FMath::Clamp(Size - CharCount, 0, Size));
 }
 
-template <typename CharType, typename ViewType>
-inline ViewType TStringViewImpl<CharType, ViewType>::Right(SizeType CharCount) const
+template <typename CharType>
+inline TStringView<CharType> TStringViewImpl<CharType>::Right(SizeType CharCount) const
 {
 	const SizeType OutLen = FMath::Clamp(CharCount, 0, Size);
 	return ViewType(DataPtr + Size - OutLen, OutLen);
 }
 
-template <typename CharType, typename ViewType>
-inline ViewType TStringViewImpl<CharType, ViewType>::RightChop(SizeType CharCount) const
+template <typename CharType>
+inline TStringView<CharType> TStringViewImpl<CharType>::RightChop(SizeType CharCount) const
 {
 	const SizeType OutLen = FMath::Clamp(Size - CharCount, 0, Size);
 	return ViewType(DataPtr + Size - OutLen, OutLen);
 }
 
-template <typename CharType, typename ViewType>
-inline ViewType TStringViewImpl<CharType, ViewType>::Mid(SizeType Position, SizeType CharCount) const
+template <typename CharType>
+inline TStringView<CharType> TStringViewImpl<CharType>::Mid(SizeType Position, SizeType CharCount) const
 {
 	check(CharCount >= 0);
 	using USizeType = TUnsignedIntType_T<sizeof(SizeType)>;
@@ -77,111 +56,81 @@ inline ViewType TStringViewImpl<CharType, ViewType>::Mid(SizeType Position, Size
 	return ViewType(DataPtr + Position, CharCount);
 }
 
-template <typename CharType, typename ViewType>
-inline ViewType TStringViewImpl<CharType, ViewType>::TrimStartAndEnd() const
+template <typename CharType>
+inline TStringView<CharType> TStringViewImpl<CharType>::TrimStartAndEnd() const
 {
 	return TrimStart().TrimEnd();
 }
 
-template <typename CharType, typename ViewType>
-inline ViewType TStringViewImpl<CharType, ViewType>::TrimStart() const
-{
-	return StringViewPrivate::TrimStart(static_cast<const ViewType&>(*this));
-}
-
-template <typename CharType, typename ViewType>
-inline ViewType TStringViewImpl<CharType, ViewType>::TrimEnd() const
-{
-	return StringViewPrivate::TrimEnd(static_cast<const ViewType&>(*this));
-}
-
-template <typename CharType, typename ViewType>
-inline bool TStringViewImpl<CharType, ViewType>::Equals(const ViewType& Other, ESearchCase::Type SearchCase) const
+template <typename CharType>
+inline bool TStringViewImpl<CharType>::Equals(const ViewType& Other, ESearchCase::Type SearchCase) const
 {
 	return Size == Other.Size && Compare(Other, SearchCase) == 0;
 }
 
-template <typename CharType, typename ViewType>
-inline int32 TStringViewImpl<CharType, ViewType>::Compare(const ViewType& Other, ESearchCase::Type SearchCase) const
-{
-	return StringViewPrivate::Compare(static_cast<const ViewType&>(*this), Other, SearchCase);
-}
-
-template <typename CharType, typename ViewType>
-inline bool TStringViewImpl<CharType, ViewType>::StartsWith(const ViewType& Prefix, ESearchCase::Type SearchCase) const
+template <typename CharType>
+inline bool TStringViewImpl<CharType>::StartsWith(const ViewType& Prefix, ESearchCase::Type SearchCase) const
 {
 	return Prefix.Equals(Left(Prefix.Len()), SearchCase);
 }
 
-template <typename CharType, typename ViewType>
-inline bool TStringViewImpl<CharType, ViewType>::EndsWith(const ViewType& Suffix, ESearchCase::Type SearchCase) const
+template <typename CharType>
+inline bool TStringViewImpl<CharType>::EndsWith(const ViewType& Suffix, ESearchCase::Type SearchCase) const
 {
 	return Suffix.Equals(Right(Suffix.Len()), SearchCase);
 }
 
-template <typename CharType, typename ViewType>
-inline bool TStringViewImpl<CharType, ViewType>::FindChar(CharType InChar, SizeType& OutIndex) const
-{
-	return StringViewPrivate::FindChar(static_cast<const ViewType&>(*this), InChar, OutIndex);
-}
-
-template <typename CharType, typename ViewType>
-inline bool TStringViewImpl<CharType, ViewType>::FindLastChar(CharType InChar, SizeType& OutIndex) const
-{
-	return StringViewPrivate::FindLastChar(static_cast<const ViewType&>(*this), InChar, OutIndex);
-}
-
 // Case-insensitive comparison operators
 
-template <typename CharType, typename ViewType, typename RangeType,
-	typename = typename TStringViewImpl<CharType, ViewType>::template TEnableIfCompatibleRangeType<RangeType>>
-inline bool operator==(const TStringViewImpl<CharType, ViewType>& Lhs, RangeType&& Rhs)
+template <typename CharType, typename RangeType,
+	typename = typename TStringViewImpl<CharType>::template TEnableIfCompatibleRangeType<RangeType>>
+inline bool operator==(const TStringViewImpl<CharType>& Lhs, RangeType&& Rhs)
 {
-	return Lhs.Equals(ViewType(Forward<RangeType>(Rhs)), ESearchCase::IgnoreCase);
+	return Lhs.Equals(TStringView<CharType>(Forward<RangeType>(Rhs)), ESearchCase::IgnoreCase);
 }
 
-template <typename CharType, typename ViewType, typename RangeType,
-	typename = typename TEnableIf<TNot<TIsDerivedFrom<typename TDecay<RangeType>::Type, TStringViewImpl<CharType, ViewType>>>::Value>::Type>
-inline auto operator==(RangeType&& Lhs, const TStringViewImpl<CharType, ViewType>& Rhs) -> decltype(Rhs == Forward<RangeType>(Lhs))
+template <typename CharType, typename RangeType,
+	typename = typename TEnableIf<TNot<TIsDerivedFrom<typename TDecay<RangeType>::Type, TStringViewImpl<CharType>>>::Value>::Type>
+inline auto operator==(RangeType&& Lhs, const TStringViewImpl<CharType>& Rhs) -> decltype(Rhs == Forward<RangeType>(Lhs))
 {
 	return Rhs == Forward<RangeType>(Lhs);
 }
 
-template <typename CharType, typename ViewType, typename RangeType>
-inline auto operator!=(const TStringViewImpl<CharType, ViewType>& Lhs, RangeType&& Rhs) -> decltype(!(Lhs == Forward<RangeType>(Rhs)))
+template <typename CharType, typename RangeType>
+inline auto operator!=(const TStringViewImpl<CharType>& Lhs, RangeType&& Rhs) -> decltype(!(Lhs == Forward<RangeType>(Rhs)))
 {
 	return !(Lhs == Forward<RangeType>(Rhs));
 }
 
-template <typename CharType, typename ViewType, typename RangeType,
-	typename = typename TEnableIf<TNot<TIsDerivedFrom<typename TDecay<RangeType>::Type, TStringViewImpl<CharType, ViewType>>>::Value>::Type>
-inline auto operator!=(RangeType&& Lhs, const TStringViewImpl<CharType, ViewType>& Rhs) -> decltype(!(Rhs == Forward<RangeType>(Lhs)))
+template <typename CharType, typename RangeType,
+	typename = typename TEnableIf<TNot<TIsDerivedFrom<typename TDecay<RangeType>::Type, TStringViewImpl<CharType>>>::Value>::Type>
+inline auto operator!=(RangeType&& Lhs, const TStringViewImpl<CharType>& Rhs) -> decltype(!(Rhs == Forward<RangeType>(Lhs)))
 {
 	return !(Rhs == Forward<RangeType>(Lhs));
 }
 
 // Case-insensitive C-string comparison operators
 
-template <typename CharType, typename ViewType>
-inline bool operator==(const TStringViewImpl<CharType, ViewType>& Lhs, const CharType* Rhs)
+template <typename CharType>
+inline bool operator==(const TStringViewImpl<CharType>& Lhs, const CharType* Rhs)
 {
 	return TCString<CharType>::Strnicmp(Lhs.GetData(), Rhs, Lhs.Len()) == 0 && !Rhs[Lhs.Len()];
 }
 
-template <typename CharType, typename ViewType>
-inline bool operator==(const CharType* Lhs, const TStringViewImpl<CharType, ViewType>& Rhs)
+template <typename CharType>
+inline bool operator==(const CharType* Lhs, const TStringViewImpl<CharType>& Rhs)
 {
 	return Rhs == Lhs;
 }
 
-template <typename CharType, typename ViewType>
-inline bool operator!=(const TStringViewImpl<CharType, ViewType>& Lhs, const CharType* Rhs)
+template <typename CharType>
+inline bool operator!=(const TStringViewImpl<CharType>& Lhs, const CharType* Rhs)
 {
 	return !(Lhs == Rhs);
 }
 
-template <typename CharType, typename ViewType>
-inline bool operator!=(const CharType* Lhs, const TStringViewImpl<CharType, ViewType>& Rhs)
+template <typename CharType>
+inline bool operator!=(const CharType* Lhs, const TStringViewImpl<CharType>& Rhs)
 {
 	return !(Lhs == Rhs);
 }
