Iteration on StringBuilder

- Add Join to assist migration from FString.
- Add JoinQuoted because it is a common use case that requires a lot of boilerplate code to implement.
- Add ElementType for compatibility with code that relies on TElementType<T>.
- Add TStringBuilderBase<T> to alias FStringBuilderBase or FAnsiStringBuilderBase as appropriate.
- Change TStringBuilderImpl& return types to TStringBuilderBase<T>& to avoid exposing the implementation type to callers.

diff --git a/Engine/Source/Runtime/Core/Private/Misc/StringBuilder.cpp b/Engine/Source/Runtime/Core/Private/Misc/StringBuilder.cpp
--- a/Engine/Source/Runtime/Core/Private/Misc/StringBuilder.cpp
+++ b/Engine/Source/Runtime/Core/Private/Misc/StringBuilder.cpp
@@ -1,6 +1,7 @@
 // Copyright 1998-2019 Epic Games, Inc. All Rights Reserved.
 
 #include "Misc/StringBuilder.h"
+
 #include "Misc/AssertionMacros.h"
 #include "HAL/PlatformMath.h"
 #include "HAL/UnrealMemory.h"
@@ -66,8 +67,8 @@ TStringBuilderImpl<C>::FreeBuffer(void* Buffer, int32 ByteCount)
 }
 
 template<typename C>
-void
-TStringBuilderImpl<C>::AppendfImpl(TStringBuilderImpl& Self, const C* Fmt, ...)
+typename TStringBuilderImpl<C>::BuilderType&
+TStringBuilderImpl<C>::AppendfImpl(BuilderType& Self, const C* Fmt, ...)
 {
 	for (;;)
 	{
@@ -80,7 +81,7 @@ TStringBuilderImpl<C>::AppendfImpl(TStringBuilderImpl& Self, const C* Fmt, ...)
 		if (Result >= 0 && Result < RemainingSize)
 		{
 			Self.CurPos += Result;
-			return;
+			return Self;
 		}
 		else
 		{
diff --git a/Engine/Source/Runtime/Core/Public/Misc/StringBuilder.h b/Engine/Source/Runtime/Core/Public/Misc/StringBuilder.h
--- a/Engine/Source/Runtime/Core/Public/Misc/StringBuilder.h
+++ b/Engine/Source/Runtime/Core/Public/Misc/StringBuilder.h
@@ -2,9 +2,9 @@
 
 #pragma once
 
+#include "Containers/StringView.h"
 #include "CoreTypes.h"
 #include "Misc/CString.h"
-#include "Misc/StringView.h"
 #include "Templates/AndOrNot.h"
 #include "Templates/EnableIf.h"
 #include "Templates/IsArrayOrRefOfType.h"
@@ -13,6 +13,20 @@
 
 #define USE_STRING_LITERAL_PATH 1
 
+class FAnsiStringBuilderBase;
+class FStringBuilderBase;
+
+namespace StringBuilderPrivate
+{
+	template <typename CharType> struct TStringBuilderBaseType;
+	template <> struct TStringBuilderBaseType<ANSICHAR> { using Type = FAnsiStringBuilderBase; };
+	template <> struct TStringBuilderBaseType<WIDECHAR> { using Type = FStringBuilderBase; };
+}
+
+/** The string builder base type to be used by append operators and function output parameters for a given character type. */
+template <typename CharType>
+using TStringBuilderBase = typename StringBuilderPrivate::TStringBuilderBaseType<CharType>::Type;
+
 /** String Builder implementation
 
 	This class helps with the common task of constructing new strings. 
@@ -51,6 +65,21 @@ template<typename C>
 class TStringBuilderImpl
 {
 public:
+	/** The character type that this builder operates on. */
+	using ElementType = C;
+	/** The string builder base type to be used by append operators and function output parameters. */
+	using BuilderType = TStringBuilderBase<ElementType>;
+	/** The string view type that this builder is compatible with. */
+	using ViewType = TStringView<ElementType>;
+
+	/** Whether the given type can be appended to this builder using the append operator. */
+	template <typename AppendType>
+	using TCanAppend = TIsSame<BuilderType&, decltype(DeclVal<BuilderType&>() << DeclVal<AppendType>())>;
+
+	/** Whether the given range type can have its elements appended to the builder using the append operator. */
+	template <typename RangeType>
+	using TCanAppendRange = TAnd<TIsContiguousContainer<RangeType>, TCanAppend<decltype(*::GetData(DeclVal<RangeType>()))>>;
+
 				TStringBuilderImpl() = default;
 	CORE_API	~TStringBuilderImpl();
 
@@ -66,6 +95,8 @@ public:
 	inline const C* ToString() const	{ EnsureNulTerminated(); return Base; }
 	inline const C* operator*() const	{ EnsureNulTerminated(); return Base; }
 
+	inline operator ViewType() const	{ return ViewType(Base, CurPos - Base); }
+
 	inline const C	LastChar() const	{ return *(CurPos - 1); }
 
 	/**
@@ -81,7 +112,7 @@ public:
 	 *
 	 * @param InCount The number of uninitialized characters to add.
 	 *
-	 * @return The number of characters in the string builder before adding new characters.
+	 * @return The number of characters in the string builder before adding the new characters.
 	 */
 	inline int32 AddUninitialized(int32 InCount)
 	{
@@ -100,47 +131,35 @@ public:
 		CurPos -= InCount;
 	}
 
-	inline TStringBuilderImpl& Append(C Char)
+	inline BuilderType& Append(C Char)
 	{
 		EnsureCapacity(1);
 
 		*CurPos++ = Char;
 
-		return *this;
+		return BaseBuilder();
 	}
 
 #if USE_STRING_LITERAL_PATH
-	inline TStringBuilderImpl& AppendAnsi(const ANSICHAR*& NulTerminatedString)
+	inline BuilderType& AppendAnsi(const ANSICHAR*& NulTerminatedString)
 #else
-	inline TStringBuilderImpl& AppendAnsi(const ANSICHAR* NulTerminatedString)
+	inline BuilderType& AppendAnsi(const ANSICHAR* NulTerminatedString)
 #endif
 	{
 		if (!NulTerminatedString)
 		{
-			return *this;
-		}
-
-		const int32 Length = TCString<ANSICHAR>::Strlen(NulTerminatedString);
-
-		EnsureCapacity(Length);
-
-		C* RESTRICT Dest = CurPos;
-		CurPos += Length;
-
-		for (int32 i = 0; i < Length; ++i)
-		{
-			Dest[i] = NulTerminatedString[i];
+			return BaseBuilder();
 		}
 
-		return *this;
+		return AppendAnsi(NulTerminatedString, TCString<ANSICHAR>::Strlen(NulTerminatedString));
 	}
 
-	inline TStringBuilderImpl& AppendAnsi(const FAnsiStringView& AnsiString)
+	inline BuilderType& AppendAnsi(const FAnsiStringView& AnsiString)
 	{
 		return AppendAnsi(AnsiString.GetData(), AnsiString.Len());
 	}
 
-	inline TStringBuilderImpl& AppendAnsi(const ANSICHAR* NulTerminatedString, const int32 Length)
+	inline BuilderType& AppendAnsi(const ANSICHAR* String, const int32 Length)
 	{
 		EnsureCapacity(Length);
 
@@ -149,85 +168,132 @@ public:
 
 		for (int32 i = 0; i < Length; ++i)
 		{
-			Dest[i] = NulTerminatedString[i];
+			Dest[i] = String[i];
 		}
 
-		return *this;
+		return BaseBuilder();
 	}
 
 #if USE_STRING_LITERAL_PATH
-	inline TStringBuilderImpl& Append(const C*& NulTerminatedString)
+	inline BuilderType& Append(const C*& NulTerminatedString)
 #else
-	inline TStringBuilderImpl& Append(const C* NulTerminatedString)
+	inline BuilderType& Append(const C* NulTerminatedString)
 #endif
 	{
 		if (!NulTerminatedString)
 		{
-			return *this;
+			return BaseBuilder();
 		}
 
-		const int32 Length = TCString<C>::Strlen(NulTerminatedString);
-
-		EnsureCapacity(Length);
-		C* RESTRICT Dest = CurPos;
-		CurPos += Length;
-
-		FMemory::Memcpy(Dest, NulTerminatedString, Length * sizeof(C));
+		return Append(NulTerminatedString, TCString<C>::Strlen(NulTerminatedString));
+	}
 
-		return *this;
+	inline BuilderType& Append(const ViewType& StringView)
+	{
+		return Append(StringView.GetData(), StringView.Len());
 	}
 
-	inline TStringBuilderImpl& Append(const C* NulTerminatedString, int32 MaxChars)
+	inline BuilderType& Append(const C* String, int32 MaxChars)
 	{
-		const int32 StringLength = TCString<C>::Strlen(NulTerminatedString);
+		const int32 StringLength = TCString<C>::Strlen(String);
 		const int32 Length = FPlatformMath::Min<int32>(MaxChars, StringLength);
 
 		EnsureCapacity(Length);
 		C* RESTRICT Dest = CurPos;
 		CurPos += Length;
 
-		FMemory::Memcpy(Dest, NulTerminatedString, Length * sizeof(C));
+		FMemory::Memcpy(Dest, String, Length * sizeof(C));
 
-		return *this;
+		return BaseBuilder();
 	}
 
 #if USE_STRING_LITERAL_PATH
 	template<int32 ArrayLength>
-	inline TStringBuilderImpl& Append(C(&CharArray)[ArrayLength])
+	inline BuilderType& Append(C(&CharArray)[ArrayLength])
 	{
 		return Append(&CharArray[0]);
 	}
 
 	template<int32 LiteralLength>
-	inline TStringBuilderImpl& Append(const C (&StringLiteral)[LiteralLength])
+	inline BuilderType& Append(const C (&StringLiteral)[LiteralLength])
 	{
 		constexpr int32 Length = LiteralLength - 1;
-		EnsureCapacity(Length);
-		C* RESTRICT Dest = CurPos;
-		CurPos += Length;
-
-		FMemory::Memcpy(Dest, StringLiteral, Length * sizeof(C));
-
-		return *this;
+		return Append(StringLiteral, Length);
 	}
 
 	template<int32 LiteralLength>
-	inline TStringBuilderImpl& AppendAnsi(const ANSICHAR(&StringLiteral)[LiteralLength])
+	inline BuilderType& AppendAnsi(const ANSICHAR (&StringLiteral)[LiteralLength])
 	{
 		constexpr int32 Length = LiteralLength - 1;
-		EnsureCapacity(Length);
-
-		C* RESTRICT Dest = CurPos;
-		CurPos += Length;
+		return AppendAnsi(StringLiteral, Length);
+	}
+#endif
 
-		for (int32 i = 0; i < Length; ++i)
+	/**
+	 * Append every element of the range to the builder, separating the elements by the delimiter.
+	 *
+	 * This function is only available when the elements of the range and the delimiter can both be
+	 * written to the builder using the append operator.
+	 *
+	 * @param InRange The range of elements to join and append.
+	 * @param InDelimiter The delimiter to append as a separator for the elements.
+	 *
+	 * @return The builder, to allow additional operations to be composed with this one.
+	 */
+	template <typename RangeType, typename DelimiterType,
+		typename = typename TEnableIf<TAnd<TCanAppendRange<RangeType&&>, TCanAppend<DelimiterType&&>>::Value>::Type>
+	inline BuilderType& Join(RangeType&& InRange, DelimiterType&& InDelimiter)
+	{
+		BuilderType& Self = BaseBuilder();
+		bool bFirst = true;
+		for (auto&& Elem : Forward<RangeType>(InRange))
 		{
-			Dest[i] = StringLiteral[i];
+			if (bFirst)
+			{
+				bFirst = false;
+			}
+			else
+			{
+				Self << InDelimiter;
+			}
+			Self << Elem;
 		}
+		return Self;
+	}
 
-		return *this;
+	/**
+	 * Append every element of the range to the builder, separating the elements by the delimiter, and
+	 * surrounding every element on each side with the given quote.
+	 *
+	 * This function is only available when the elements of the range, the delimiter, and the quote can be
+	 * written to the builder using the append operator.
+	 *
+	 * @param InRange The range of elements to join and append.
+	 * @param InDelimiter The delimiter to append as a separator for the elements.
+	 * @param InQuote The quote to append on both sides of each element.
+	 *
+	 * @return The builder, to allow additional operations to be composed with this one.
+	 */
+	template <typename RangeType, typename DelimiterType, typename QuoteType,
+		typename = typename TEnableIf<TAnd<TCanAppendRange<RangeType>, TCanAppend<DelimiterType&&>, TCanAppend<QuoteType&&>>::Value>::Type>
+	inline BuilderType& JoinQuoted(RangeType&& InRange, DelimiterType&& InDelimiter, QuoteType&& InQuote)
+	{
+		BuilderType& Self = BaseBuilder();
+		bool bFirst = true;
+		for (auto&& Elem : Forward<RangeType>(InRange))
+		{
+			if (bFirst)
+			{
+				bFirst = false;
+			}
+			else
+			{
+				Self << InDelimiter;
+			}
+			Self << InQuote << Elem << InQuote;
+		}
+		return Self;
 	}
-#endif
 
 	/**
 	 * Appends to the string builder similarly to how classic sprintf works.
@@ -235,16 +301,19 @@ public:
 	 * @param Format A format string that specifies how to format the additional arguments. Refer to standard printf format.
 	 */
 	template <typename FmtType, typename... Types>
-	typename TEnableIf<TIsArrayOrRefOfType<FmtType, C>::Value>::Type Appendf(const FmtType& Fmt, Types... Args)
+	typename TEnableIf<TIsArrayOrRefOfType<FmtType, C>::Value, BuilderType&>::Type Appendf(const FmtType& Fmt, Types... Args)
 	{
 		static_assert(TIsArrayOrRefOfType<FmtType, C>::Value, "Formatting string must be a character array.");
 		static_assert(TAnd<TIsValidVariadicFunctionArg<Types>...>::Value, "Invalid argument(s) passed to Appendf.");
-		AppendfImpl(*this, Fmt, Forward<Types>(Args)...);
+		return AppendfImpl(BaseBuilder(), Fmt, Forward<Types>(Args)...);
 	}
 
-protected:
-	CORE_API static void VARARGS AppendfImpl(TStringBuilderImpl& Self, const C* Fmt, ...);
+private:
+	CORE_API static BuilderType& VARARGS AppendfImpl(BuilderType& Self, const C* Fmt, ...);
 
+	inline BuilderType& BaseBuilder() { return static_cast<BuilderType&>(*this); }
+
+protected:
 	inline void Initialize(C* InBase, int32 InCapacity)
 	{
 		Base	= InBase;
@@ -297,9 +366,6 @@ extern template class TStringBuilderImpl<ANSICHAR>;
 
 class FAnsiStringBuilderBase : public TStringBuilderImpl<ANSICHAR>
 {
-public:
-	operator FAnsiStringView() const { return FAnsiStringView(Base, CurPos - Base); }
-
 protected:
 	inline FAnsiStringBuilderBase(ANSICHAR* BufferPointer, int32 BufferCapacity)
 	{
@@ -355,25 +421,6 @@ extern template class TStringBuilderImpl<TCHAR>;
 
 class FStringBuilderBase : public TStringBuilderImpl<TCHAR>
 {
-public:
-	operator FStringView() const { return FStringView(Base, CurPos - Base); }
-
-	using TStringBuilderImpl<TCHAR>::Append;
-
-	inline TStringBuilderImpl& Append(const FStringView& StringView)
-	{
-		const int32 StringLength = StringView.Len();
-
-		EnsureCapacity(StringLength);
-		TCHAR* RESTRICT Dest = CurPos;
-
-		FMemory::Memcpy(CurPos, StringView.GetData(), StringLength * sizeof(TCHAR));
-
-		CurPos += StringLength;
-
-		return *this;
-	}
-
 protected:
 	FStringBuilderBase(TCHAR* BufferPointer, int32 BufferCapacity)
 	{
@@ -425,33 +472,33 @@ struct TIsContiguousContainer<TStringBuilder<N>> { enum { Value = true }; };
 
 // Append operator implementations
 
-inline FStringBuilderBase&					operator<<(FStringBuilderBase& Builder, const TCHAR Char)						{ Builder.Append(Char); return Builder; }
+inline FStringBuilderBase&					operator<<(FStringBuilderBase& Builder, const TCHAR Char)						{ return Builder.Append(Char); }
 
 #if USE_STRING_LITERAL_PATH
 
-inline FStringBuilderBase&					operator<<(FStringBuilderBase& Builder, const TCHAR*& NulTerminatedString)		{ Builder.Append(NulTerminatedString); return Builder; }
-inline FStringBuilderBase&					operator<<(FStringBuilderBase& Builder, const ANSICHAR*& NulTerminatedString)	{ Builder.AppendAnsi(NulTerminatedString); return Builder; }
+inline FStringBuilderBase&					operator<<(FStringBuilderBase& Builder, const TCHAR*& NulTerminatedString)		{ return Builder.Append(NulTerminatedString); }
+inline FStringBuilderBase&					operator<<(FStringBuilderBase& Builder, const ANSICHAR*& NulTerminatedString)	{ return Builder.AppendAnsi(NulTerminatedString); }
 
-template<int32 N> FStringBuilderBase&		operator<<(FStringBuilderBase& Builder, const TCHAR (&StringLiteral)[N])		{ Builder.Append(StringLiteral); return Builder; }
-template<int32 N> FStringBuilderBase&		operator<<(FStringBuilderBase& Builder, const ANSICHAR (&StringLiteral)[N])		{ Builder.AppendAnsi(StringLiteral); return Builder; }
+template<int32 N> FStringBuilderBase&		operator<<(FStringBuilderBase& Builder, const TCHAR (&StringLiteral)[N])		{ return Builder.Append(StringLiteral); }
+template<int32 N> FStringBuilderBase&		operator<<(FStringBuilderBase& Builder, const ANSICHAR (&StringLiteral)[N])		{ return Builder.AppendAnsi(StringLiteral); }
 
-template<int32 N> FStringBuilderBase&		operator<<(FStringBuilderBase& Builder, TCHAR (&NulTerminatedString)[N])		{ Builder.Append(&NulTerminatedString[0], TCString<TCHAR>::Strlen(NulTerminatedString)); return Builder; }
-template<int32 N> FStringBuilderBase&		operator<<(FStringBuilderBase& Builder, ANSICHAR (&NulTerminatedString)[N])		{ Builder.AppendAnsi(&NulTerminatedString[0], TCString<ANSICHAR>::Strlen(NulTerminatedString)); return Builder; }
+template<int32 N> FStringBuilderBase&		operator<<(FStringBuilderBase& Builder, TCHAR (&NulTerminatedString)[N])		{ return Builder.Append(&NulTerminatedString[0], TCString<TCHAR>::Strlen(NulTerminatedString)); }
+template<int32 N> FStringBuilderBase&		operator<<(FStringBuilderBase& Builder, ANSICHAR (&NulTerminatedString)[N])		{ return Builder.AppendAnsi(&NulTerminatedString[0], TCString<ANSICHAR>::Strlen(NulTerminatedString)); }
 
 #else
 
-inline FStringBuilderBase&					operator<<(FStringBuilderBase& Builder, const TCHAR* NulTerminatedString)		{ Builder.Append(NulTerminatedString); return Builder; }
-inline FStringBuilderBase&					operator<<(FStringBuilderBase& Builder, const ANSICHAR* NulTerminatedString)	{ Builder.AppendAnsi(NulTerminatedString); return Builder; }
+inline FStringBuilderBase&					operator<<(FStringBuilderBase& Builder, const TCHAR* NulTerminatedString)		{ return Builder.Append(NulTerminatedString); }
+inline FStringBuilderBase&					operator<<(FStringBuilderBase& Builder, const ANSICHAR* NulTerminatedString)	{ return Builder.AppendAnsi(NulTerminatedString); }
 
 #endif
 
-inline FStringBuilderBase&					operator<<(FStringBuilderBase& Builder, const FStringView& Str)					{ Builder.Append(Str.GetData(), Str.Len()); return Builder; }
+inline FStringBuilderBase&					operator<<(FStringBuilderBase& Builder, const FStringView& Str)					{ return Builder.Append(Str); }
 
-inline FStringBuilderBase&					operator<<(FStringBuilderBase& Builder, int32 Value)							{ Builder.Appendf(TEXT("%d"), Value); return Builder; }
-inline FStringBuilderBase&					operator<<(FStringBuilderBase& Builder, uint32 Value)							{ Builder.Appendf(TEXT("%u"), Value); return Builder; }
+inline FStringBuilderBase&					operator<<(FStringBuilderBase& Builder, int32 Value)							{ return Builder.Appendf(TEXT("%d"), Value); }
+inline FStringBuilderBase&					operator<<(FStringBuilderBase& Builder, uint32 Value)							{ return Builder.Appendf(TEXT("%u"), Value); }
 
-inline FStringBuilderBase&					operator<<(FStringBuilderBase& Builder, int64 Value)							{ Builder.Appendf(TEXT(INT64_FMT), Value); return Builder; }
-inline FStringBuilderBase&					operator<<(FStringBuilderBase& Builder, uint64 Value)							{ Builder.Appendf(TEXT(UINT64_FMT), Value); return Builder; }
+inline FStringBuilderBase&					operator<<(FStringBuilderBase& Builder, int64 Value)							{ return Builder.Appendf(TEXT(INT64_FMT), Value); }
+inline FStringBuilderBase&					operator<<(FStringBuilderBase& Builder, uint64 Value)							{ return Builder.Appendf(TEXT(UINT64_FMT), Value); }
 
 inline FStringBuilderBase&					operator<<(FStringBuilderBase& Builder, int8 Value)								{ return Builder << int32(Value); }
 inline FStringBuilderBase&					operator<<(FStringBuilderBase& Builder, uint8 Value)							{ return Builder << uint32(Value); }
@@ -461,27 +508,27 @@ inline FStringBuilderBase&					operator<<(FStringBuilderBase& Builder, uint16 Va
 
 // Append operator implementations for FAnsiStringBuilderBase
 
-inline FAnsiStringBuilderBase&				operator<<(FAnsiStringBuilderBase& Builder, const ANSICHAR Char)				{ Builder.Append(Char); return Builder; }
+inline FAnsiStringBuilderBase&				operator<<(FAnsiStringBuilderBase& Builder, const ANSICHAR Char)					{ return Builder.Append(Char); }
 
 #if USE_STRING_LITERAL_PATH
 
-inline FAnsiStringBuilderBase&				operator<<(FAnsiStringBuilderBase& Builder, const ANSICHAR*& NulTerminatedString)	{ Builder.Append(NulTerminatedString); return Builder; }
-template<int32 N> FAnsiStringBuilderBase&	operator<<(FAnsiStringBuilderBase& Builder, const ANSICHAR (&StringLiteral)[N])		{ Builder.Append(StringLiteral); return Builder; }
-template<int32 N> FAnsiStringBuilderBase&	operator<<(FAnsiStringBuilderBase& Builder, ANSICHAR (&NulTerminatedString)[N])		{ Builder.Append(&NulTerminatedString[0], TCString<ANSICHAR>::Strlen(NulTerminatedString)); return Builder; }
+inline FAnsiStringBuilderBase&				operator<<(FAnsiStringBuilderBase& Builder, const ANSICHAR*& NulTerminatedString)	{ return Builder.Append(NulTerminatedString); }
+template<int32 N> FAnsiStringBuilderBase&	operator<<(FAnsiStringBuilderBase& Builder, const ANSICHAR (&StringLiteral)[N])		{ return Builder.Append(StringLiteral); }
+template<int32 N> FAnsiStringBuilderBase&	operator<<(FAnsiStringBuilderBase& Builder, ANSICHAR (&NulTerminatedString)[N])		{ return Builder.Append(&NulTerminatedString[0], TCString<ANSICHAR>::Strlen(NulTerminatedString)); }
 
 #else
 
-inline FAnsiStringBuilderBase&				operator<<(FAnsiStringBuilderBase& Builder, const ANSICHAR* NulTerminatedString)	{ Builder.Append(NulTerminatedString); return Builder; }
+inline FAnsiStringBuilderBase&				operator<<(FAnsiStringBuilderBase& Builder, const ANSICHAR* NulTerminatedString)	{ return Builder.Append(NulTerminatedString); }
 
 #endif
 
-inline FAnsiStringBuilderBase&				operator<<(FAnsiStringBuilderBase& Builder, const FAnsiStringView& Str)				{ Builder.Append(Str.GetData(), Str.Len()); return Builder; }
+inline FAnsiStringBuilderBase&				operator<<(FAnsiStringBuilderBase& Builder, const FAnsiStringView& Str)				{ return Builder.Append(Str); }
 
-inline FAnsiStringBuilderBase&				operator<<(FAnsiStringBuilderBase& Builder, int32 Value)							{ Builder.Appendf("%d", Value); return Builder; }
-inline FAnsiStringBuilderBase&				operator<<(FAnsiStringBuilderBase& Builder, uint32 Value)							{ Builder.Appendf("%u", Value); return Builder; }
+inline FAnsiStringBuilderBase&				operator<<(FAnsiStringBuilderBase& Builder, int32 Value)							{ return Builder.Appendf("%d", Value); }
+inline FAnsiStringBuilderBase&				operator<<(FAnsiStringBuilderBase& Builder, uint32 Value)							{ return Builder.Appendf("%u", Value); }
 
-inline FAnsiStringBuilderBase&				operator<<(FAnsiStringBuilderBase& Builder, int64 Value)							{ Builder.Appendf(INT64_FMT, Value); return Builder; }
-inline FAnsiStringBuilderBase&				operator<<(FAnsiStringBuilderBase& Builder, uint64 Value)							{ Builder.Appendf(UINT64_FMT, Value); return Builder; }
+inline FAnsiStringBuilderBase&				operator<<(FAnsiStringBuilderBase& Builder, int64 Value)							{ return Builder.Appendf(INT64_FMT, Value); }
+inline FAnsiStringBuilderBase&				operator<<(FAnsiStringBuilderBase& Builder, uint64 Value)							{ return Builder.Appendf(UINT64_FMT, Value); }
 
 inline FAnsiStringBuilderBase&				operator<<(FAnsiStringBuilderBase& Builder, int8 Value)								{ return Builder << int32(Value); }
 inline FAnsiStringBuilderBase&				operator<<(FAnsiStringBuilderBase& Builder, uint8 Value)							{ return Builder << uint32(Value); }
