Compile fix for Linux Editor

The case for the file was wrong (manager is in lower case)

#rnx

diff --git a/Engine/Source/Runtime/Core/Private/GenericPlatform/GenericPlatformIoDispatcher.cpp b/Engine/Source/Runtime/Core/Private/GenericPlatform/GenericPlatformIoDispatcher.cpp
--- a/Engine/Source/Runtime/Core/Private/GenericPlatform/GenericPlatformIoDispatcher.cpp
+++ b/Engine/Source/Runtime/Core/Private/GenericPlatform/GenericPlatformIoDispatcher.cpp
@@ -4,7 +4,7 @@
 #include "IO/IoDispatcherFileBackend.h"
 #include "Async/AsyncFileHandle.h"
 #include "HAL/Event.h"
-#include "HAL/PlatformFileManager.h"
+#include "HAL/PlatformFilemanager.h"
 #include "Misc/ScopeLock.h"
 
 FGenericIoDispatcherEventQueue::FGenericIoDispatcherEventQueue()
