Fix compile error on linux due to mismatched case in header.
#rnx

diff --git a/Engine/Source/Runtime/Core/Private/IO/IoDispatcherFileBackend.cpp b/Engine/Source/Runtime/Core/Private/IO/IoDispatcherFileBackend.cpp
--- a/Engine/Source/Runtime/Core/Private/IO/IoDispatcherFileBackend.cpp
+++ b/Engine/Source/Runtime/Core/Private/IO/IoDispatcherFileBackend.cpp
@@ -3,7 +3,7 @@
 #include "IO/IoDispatcherFileBackend.h"
 #include "Misc/ScopeRWLock.h"
 #include "ProfilingDebugging/CountersTrace.h"
-#include "HAL/PlatformFileManager.h"
+#include "HAL/PlatformFilemanager.h"
 #include "GenericPlatform/GenericPlatformFile.h"
 
 TRACE_DECLARE_INT_COUNTER(IoDispatcherCacheHitsCold, TEXT("IoDispatcher/CacheHitsCold"));
