Add TEqualTo as a binary predicate to perform equality comparison of its arguments

diff --git a/Engine/Source/Runtime/Core/Public/Templates/EqualTo.h b/Engine/Source/Runtime/Core/Public/Templates/EqualTo.h
new file mode 100644
--- /dev/null
+++ b/Engine/Source/Runtime/Core/Public/Templates/EqualTo.h
@@ -0,0 +1,31 @@
+// Copyright 1998-2019 Epic Games, Inc. All Rights Reserved.
+
+#pragma once
+
+#include "UnrealTemplate.h"
+
+/**
+ * Binary predicate class for performing equality comparisons.
+ *
+ * Uses operator== when available.
+ *
+ * See: https://en.cppreference.com/w/cpp/utility/functional/equal_to
+ */
+template <typename T = void>
+struct TEqualTo
+{
+	constexpr auto operator()(const T& Lhs, const T& Rhs) const -> decltype(Lhs == Rhs)
+	{
+		return Lhs == Rhs;
+	}
+};
+
+template <>
+struct TEqualTo<void>
+{
+	template <typename T, typename U>
+	constexpr auto operator()(T&& Lhs, U&& Rhs) const -> decltype(Forward<T>(Lhs) == Forward<U>(Rhs))
+	{
+		return Forward<T>(Lhs) == Forward<U>(Rhs);
+	}
+};
