Static analysis fix

diff --git a/Engine/Source/Runtime/Core/Private/IO/IoDispatcher.cpp b/Engine/Source/Runtime/Core/Private/IO/IoDispatcher.cpp
--- a/Engine/Source/Runtime/Core/Private/IO/IoDispatcher.cpp
+++ b/Engine/Source/Runtime/Core/Private/IO/IoDispatcher.cpp
@@ -324,6 +324,7 @@ private:
 		for (; CurrentRequestsToSubmitIndex < RequestsToSubmitCount; ++CurrentRequestsToSubmitIndex)
 		{
 			FIoRequestImpl* Request = RequestsToSubmit[CurrentRequestsToSubmitIndex];
+			check(Request);
 
 			EIoStoreResolveResult Result = FileIoStore.Resolve(Request);
 			if (Result == IoStoreResolveResult_Stalled)
