EngineTest runtime improvement: Cap the number of threads used in System.Core.Misc.LockFree to reduce the run time on machines with many processors.

32 core machine, ue4 creates 24 worker threads
Debug configuration, attached in visual studio
Time before: 539 seconds
Time after: 67 seconds

#rn minor Automation

diff --git a/Engine/Source/Runtime/Core/Private/Async/TaskGraph.cpp b/Engine/Source/Runtime/Core/Private/Async/TaskGraph.cpp
--- a/Engine/Source/Runtime/Core/Private/Async/TaskGraph.cpp
+++ b/Engine/Source/Runtime/Core/Private/Async/TaskGraph.cpp
@@ -29,6 +29,7 @@
 #include "HAL/ThreadHeartBeat.h"
 #include "ProfilingDebugging/ExternalProfiler.h"
 #include "ProfilingDebugging/MiscTrace.h"
+#include "ProfilingDebugging/ScopedTimers.h"
 
 DEFINE_LOG_CATEGORY_STATIC(LogTaskGraph, Log, All);
 
@@ -1842,13 +1843,16 @@ public:
 		{
 			if (StallForTaskThread->Decrement())
 			{
-				QUICK_SCOPE_CYCLE_COUNTER(STAT_Broadcast_WaitForOthers);
-				TaskEvent->Wait();
+				if (TaskEvent)
 				{
-					float ThisTime = FPlatformTime::Seconds() - StartTime;
-					if (ThisTime > 0.02f)
+					QUICK_SCOPE_CYCLE_COUNTER(STAT_Broadcast_WaitForOthers);
+					TaskEvent->Wait();
 					{
-						UE_CLOG(GPrintBroadcastWarnings, LogTaskGraph, Warning, TEXT("Task graph took %6.2fms for %s to recieve broadcast do processing and wait for other task threads."), ThisTime * 1000.0f, Name);
+						float ThisTime = FPlatformTime::Seconds() - StartTime;
+						if (ThisTime > 0.02f)
+						{
+							UE_CLOG(GPrintBroadcastWarnings, LogTaskGraph, Warning, TEXT("Task graph took %6.2fms for %s to recieve broadcast do processing and wait for other task threads."), ThisTime * 1000.0f, Name);
+						}
 					}
 				}
 			}
@@ -2505,6 +2509,36 @@ static void TestLockFree(int32 OuterIters = 3)
 		return;
 	}
 
+	const int32 NumWorkers = FTaskGraphInterface::Get().GetNumWorkerThreads();
+	// If we have too many threads active at once, they become too slow due to contention.  Set a reasonable maximum for how many are required to guarantee correctness of our LockFreePointers.
+	const int32 MaxWorkersForTest = 5; 
+	const int32 MinWorkersForTest = 2; // With less than two threads we're not testing threading at all, so the test is pointless.
+	if (NumWorkers < MinWorkersForTest)
+	{
+		UE_LOG(LogConsoleResponse, Display, TEXT("WARNING: TestLockFree disabled for current machine because of not enough worker threads.  Need %d, have %d."), MinWorkersForTest, NumWorkers);
+		return;
+	}
+
+	FScopedDurationTimeLogger DurationLogger(TEXT("TestLockFree Runtime"));
+	const uint32 NumWorkersForTest = static_cast<uint32>(FMath::Clamp(NumWorkers, MinWorkersForTest, MaxWorkersForTest));
+	auto RunWorkersSynchronous = [NumWorkersForTest](const TFunction<void(uint32)>& WorkerTask)
+	{
+		const bool bIsManualReset = true;
+		FEvent* AllDoneEvent = FPlatformProcess::GetSynchEventFromPool(bIsManualReset);
+		FThreadSafeCounter UnfinishedWorkers;
+		FGraphEventArray Tasks;
+		const double StartTime = FPlatformTime::Seconds();
+		UnfinishedWorkers.Add(NumWorkersForTest);
+		for (uint32 Index = 0; Index < NumWorkersForTest; Index++)
+		{
+			TFunction<void(ENamedThreads::Type)> WorkerTaskWithIndex{ [Index, &WorkerTask](ENamedThreads::Type CurrentThread) { WorkerTask(Index); } };
+			Tasks.Add(TGraphTask<FBroadcastTask>::CreateTask().ConstructAndDispatchWhenReady(WorkerTaskWithIndex, StartTime,
+				TEXT("TestLockFreeWorker"), ENamedThreads::AnyNormalThreadHiPriTask, &UnfinishedWorkers, nullptr, AllDoneEvent));
+		}
+		AllDoneEvent->Wait();
+		FPlatformProcess::ReturnSynchEventToPool(AllDoneEvent);
+	};
+
 	for (int32 Iter = 0; Iter < OuterIters; Iter++)
 	{
 		{
@@ -2514,15 +2548,15 @@ static void TestLockFree(int32 OuterIters = 3)
 			{
 				Rig.Test1.Push(new FTestStruct(Index));
 			}
-			TFunction<void(ENamedThreads::Type CurrentThread)> Broadcast =
-				[&Rig](ENamedThreads::Type MyThread)
+			TFunction<void(uint32)> Broadcast =
+				[&Rig](uint32 WorkerIndex)
 			{
-				FRandomStream Stream(((int32)MyThread) * 7 + 13);
+				FRandomStream Stream(((int32)WorkerIndex) * 7 + 13);
 				for (int32 Index = 0; Index < 1000000; Index++)
 				{
 					if (Index % 200000 == 1)
 					{
-						//UE_LOG(LogTemp, Log, TEXT("%8d iters thread=%d"), Index, int32(MyThread));
+						//UE_LOG(LogTemp, Log, TEXT("%8d iters thread=%d"), Index, int32(WorkerIndex));
 					}
 					if (Stream.FRand() < .03f)
 					{
@@ -2596,7 +2630,7 @@ static void TestLockFree(int32 OuterIters = 3)
 					}
 				}
 			};
-			FTaskGraphInterface::BroadcastSlow_OnlyUseForSpecialPurposes(true, false, Broadcast);
+			RunWorkersSynchronous(Broadcast);
 
 			TArray<FTestStruct*> Items;
 			Rig.Test1.PopAll(Items);
@@ -2633,15 +2667,15 @@ static void TestLockFree(int32 OuterIters = 3)
 			{
 				Rig.Test1.Push(new FTestStruct(Index));
 			}
-			TFunction<void(ENamedThreads::Type CurrentThread)> Broadcast =
-				[&Rig](ENamedThreads::Type MyThread)
+			TFunction<void(uint32)> Broadcast =
+				[&Rig](uint32 WorkerIndex)
 			{
-				FRandomStream Stream(((int32)MyThread) * 7 + 13);
+				FRandomStream Stream(((int32)WorkerIndex) * 7 + 13);
 				for (int32 Index = 0; Index < 1000000; Index++)
 				{
 					if (Index % 200000 == 1)
 					{
-						//UE_LOG(LogTemp, Log, TEXT("%8d iters thread=%d"), Index, int32(MyThread));
+						//UE_LOG(LogTemp, Log, TEXT("%8d iters thread=%d"), Index, int32(WorkerIndex));
 					}
 					if (Stream.FRand() < .03f)
 					{
@@ -2715,7 +2749,7 @@ static void TestLockFree(int32 OuterIters = 3)
 					}
 				}
 			};
-			FTaskGraphInterface::BroadcastSlow_OnlyUseForSpecialPurposes(true, false, Broadcast);
+			RunWorkersSynchronous(Broadcast);
 
 			TArray<FTestStruct*> Items;
 			Rig.Test1.PopAll(Items);
