Unique algo to remove duplicates from a range, similar to `std::unique`

diff --git a/Engine/Source/Runtime/Core/Private/Tests/Algo/UniqueTest.cpp b/Engine/Source/Runtime/Core/Private/Tests/Algo/UniqueTest.cpp
new file mode 100644
--- /dev/null
+++ b/Engine/Source/Runtime/Core/Private/Tests/Algo/UniqueTest.cpp
@@ -0,0 +1,56 @@
+// Copyright 1998-2019 Epic Games, Inc. All Rights Reserved.
+
+#include "Containers/UnrealString.h"
+#include "Algo/Unique.h"
+#include "Containers/Array.h"
+#include "Misc/AutomationTest.h"
+
+#if WITH_DEV_AUTOMATION_TESTS
+
+IMPLEMENT_SIMPLE_AUTOMATION_TEST(FUniqueTest, "System.Core.Algo.Unique", EAutomationTestFlags::ApplicationContextMask | EAutomationTestFlags::SmokeFilter)
+
+bool FUniqueTest::RunTest(const FString& Parameters)
+{
+	using namespace Algo;
+
+	{
+		TArray<int32> Array;
+		int32 RemoveFrom = Unique(Array);
+		TestEqual(TEXT("`Unique` must handle an empty container"), RemoveFrom, 0);
+	}
+	{
+		TArray<int32> Array{ 1, 2, 3 };
+		Array.SetNum(Unique(Array));
+		TestEqual(TEXT("Uniqued container with no duplicates must remain unchanged"), Array, TArray<int32>{1, 2, 3});
+	}
+	{
+		TArray<int32> Array{ 1, 1, 2, 2, 2, 3, 3, 3, 3};
+		Array.SetNum(Unique(Array));
+		TestEqual(TEXT("`Unique` with multiple duplicates must return correct result"), Array, TArray<int32>{1, 2, 3});
+	}
+	{
+		TArray<int32> Array{ 1, 1, 2, 3, 3, 3 };
+		Array.SetNum(Unique(Array));
+		TestEqual(TEXT("`Unique` with duplicates and unique items must return correct result"), Array, TArray<int32>{1, 2, 3});
+	}
+	{
+		FString Str = TEXT("aa");
+		Str = Str.Mid(0, Unique(Str));
+		TestEqual(TEXT("`Unique` on `FString` as an example of arbitrary random-access container must compile and return correct result"), Str, FString(TEXT("a")));
+	}
+	{
+		int32 Array[] = {1};
+		int32 NewSize = Unique(Array);
+		TestEqual(TEXT("`Unique` must support C arrays"), NewSize, 1);
+	}
+	{
+		TArray<int32> Array = { 1, 1 };
+		int32 NewSize = Unique(MakeArrayView(Array.GetData() + 1, 1));
+		TestEqual(TEXT("`Unique` must support ranges"), NewSize, 1);
+	}
+
+	return true;
+}
+
+
+#endif //WITH_DEV_AUTOMATION_TESTS
diff --git a/Engine/Source/Runtime/Core/Public/Algo/Unique.h b/Engine/Source/Runtime/Core/Public/Algo/Unique.h
new file mode 100644
--- /dev/null
+++ b/Engine/Source/Runtime/Core/Public/Algo/Unique.h
@@ -0,0 +1,64 @@
+// Copyright 1998-2019 Epic Games, Inc. All Rights Reserved.
+
+#pragma once
+
+#include "Templates/UnrealTemplate.h"
+#include "Templates/Invoke.h"
+#include "Templates/EqualTo.h"
+
+namespace AlgoImpl
+{
+	template <typename T, typename SizeType, typename BinaryPredicate>
+	SizeType Unique(T* Array, SizeType ArraySize, BinaryPredicate Predicate)
+	{
+		if (ArraySize <= 1)
+		{
+			return ArraySize;
+		}
+
+		T* Result = Array;
+		for (T* Iter = Array + 1; Iter != Array + ArraySize; ++Iter)
+		{
+			if (!Invoke(Predicate, *Result, *Iter)) 
+			{
+				++Result;
+				if (Result != Iter)
+				{
+					*Result = MoveTemp(*Iter);
+				}
+			}
+		}
+
+		return Result + 1 - Array;
+	}
+}
+
+namespace Algo
+{
+	/**
+	 * Eliminates all but the first element from every consecutive group of equivalent elements and 
+	 * returns past-the-end index of unique elements for the new logical end of the range.
+	 * 
+	 * Removing is done by shifting the elements in the range in such a way that elements to be erased are overwritten. 
+	 * Relative order of the elements that remain is preserved and the physical size of the range is unchanged. 
+	 * References to an element between the new logical end and the physical end of the range are still 
+	 * dereferenceable, but the elements themselves have unspecified values. A call to `Unique` is typically followed by a call 
+	 * to a container's `RemoveAt` method, which erases the unspecified values and reduces the physical size of the container 
+	 * to match its new logical size.
+	 * 
+	 * Elements are compared using operator== or given binary predicate. The behavior is undefined if it is not an equivalence relation.
+	 *
+	 * See https://en.cppreference.com/w/cpp/algorithm/unique
+	 */
+	template<typename RangeType>
+	auto Unique(RangeType&& Range) -> decltype(AlgoImpl::Unique(GetData(Range), GetNum(Range), TEqualTo<>{}))
+	{
+		return AlgoImpl::Unique(GetData(Range), GetNum(Range), TEqualTo<>{});
+	}
+
+	template<typename RangeType, typename BinaryPredicate>
+	auto Unique(RangeType&& Range, BinaryPredicate Predicate) -> decltype(AlgoImpl::Unique(GetData(Range), GetNum(Range), MoveTemp(Predicate)))
+	{
+		return AlgoImpl::Unique(GetData(Range), GetNum(Range), MoveTemp(Predicate));
+	}
+}
