Add string view constructors and a string builder append operator to FName

diff --git a/Engine/Source/Runtime/Core/Private/UObject/UnrealNames.cpp b/Engine/Source/Runtime/Core/Private/UObject/UnrealNames.cpp
--- a/Engine/Source/Runtime/Core/Private/UObject/UnrealNames.cpp
+++ b/Engine/Source/Runtime/Core/Private/UObject/UnrealNames.cpp
@@ -1912,6 +1912,14 @@ FName::FName(int32 Len, const ANSICHAR* Name, EFindName FindType)
 	: FName(FNameHelper::MakeDetectNumber(MakeUnconvertedView(Name, Len), FindType))
 {}
 
+FName::FName(const FStringView& Name, EFindName FindType)
+	: FName(FNameHelper::MakeDetectNumber(MakeUnconvertedView(Name.GetData(), Name.Len()), FindType))
+{}
+
+FName::FName(const FAnsiStringView& Name, EFindName FindType)
+	: FName(FNameHelper::MakeDetectNumber(MakeUnconvertedView(Name.GetData(), Name.Len()), FindType))
+{}
+
 FName::FName(const WIDECHAR* Name, int32 InNumber, EFindName FindType)
 	: FName(FNameHelper::MakeWithNumber(MakeUnconvertedView(Name), FindType, InNumber))
 {}
@@ -1930,6 +1938,14 @@ FName::FName(int32 Len, const ANSICHAR* Name, int32 InNumber, EFindName FindType
 												: FNameHelper::MakeDetectNumber(MakeUnconvertedView(Name, Len), FindType))
 {}
 
+FName::FName(const FStringView& Name, int32 InNumber, EFindName FindType)
+	: FName(Name.Len(), Name.GetData(), InNumber, FindType)
+{}
+
+FName::FName(const FAnsiStringView& Name, int32 InNumber, EFindName FindType)
+	: FName(Name.Len(), Name.GetData(), InNumber, FindType)
+{}
+
 FName::FName(const TCHAR* Name, int32 InNumber, EFindName FindType, bool bSplitName)
 	: FName(InNumber == NAME_NO_NUMBER_INTERNAL && bSplitName 
 			? FNameHelper::MakeDetectNumber(MakeUnconvertedView(Name), FindType)
diff --git a/Engine/Source/Runtime/Core/Public/UObject/NameTypes.h b/Engine/Source/Runtime/Core/Public/UObject/NameTypes.h
--- a/Engine/Source/Runtime/Core/Public/UObject/NameTypes.h
+++ b/Engine/Source/Runtime/Core/Public/UObject/NameTypes.h
@@ -28,7 +28,9 @@
 #endif
 
 class FAnsiStringBuilderBase;
+class FAnsiStringView;
 class FStringBuilderBase;
+class FStringView;
 class FText;
 
 /** Maximum size of name. */
@@ -767,6 +769,8 @@ public:
 	/** Create FName from non-null string with known length  */
 	FName(int32 Len, const WIDECHAR* Name, EFindName FindType=FNAME_Add);
 	FName(int32 Len, const ANSICHAR* Name, EFindName FindType=FNAME_Add);
+	FName(const FStringView& Name, EFindName FindType=FNAME_Add);
+	FName(const FAnsiStringView& Name, EFindName FindType=FNAME_Add);
 
 	/**
 	 * Create an FName. If FindType is FNAME_Find, and the string part of the name 
@@ -781,6 +785,8 @@ public:
 	FName(const ANSICHAR* Name, int32 InNumber, EFindName FindType = FNAME_Add);
 	FName(int32 Len, const WIDECHAR* Name, int32 Number, EFindName FindType = FNAME_Add);
 	FName(int32 Len, const ANSICHAR* Name, int32 InNumber, EFindName FindType = FNAME_Add);
+	FName(const FStringView& Name, int32 InNumber, EFindName FindType=FNAME_Add);
+	FName(const FAnsiStringView& Name, int32 InNumber, EFindName FindType=FNAME_Add);
 
 	/**
 	 * Create an FName. If FindType is FNAME_Find, and the string part of the name 
@@ -954,6 +960,11 @@ FORCEINLINE FName ScriptNameToName(const FScriptName& InName)
 	return FName(InName.ComparisonIndex, InName.DisplayIndex, InName.Number);
 }
 
+inline FStringBuilderBase& operator<<(FStringBuilderBase& Builder, const FName& Name)
+{
+	Name.AppendString(Builder);
+	return Builder;
+}
 
 /**
  * Equality operator with CharType* on left hand side and FName on right hand side
