PR #6337: Minimize changes when editing INIs
Change FConfigFile::UpdateSections to preserve relative location of modified sections in the ConfigFile.
Fix UpdateSinglePropertyInSection to keep a blank line at the end.

Contributed by [REDACTED_PERSON]
#rn minor

diff --git a/Engine/Source/Runtime/Core/Private/Misc/ConfigCacheIni.cpp b/Engine/Source/Runtime/Core/Private/Misc/ConfigCacheIni.cpp
--- a/Engine/Source/Runtime/Core/Private/Misc/ConfigCacheIni.cpp
+++ b/Engine/Source/Runtime/Core/Private/Misc/ConfigCacheIni.cpp
@@ -1,6 +1,8 @@
 // Copyright 1998-2019 Epic Games, Inc. All Rights Reserved.
 
 #include "Misc/ConfigCacheIni.h"
+
+#include "Containers/StringView.h"
 #include "Misc/DateTime.h"
 #include "Misc/MessageDialog.h"
 #include "HAL/FileManager.h"
@@ -16,6 +18,7 @@
 #include "Misc/DefaultValueHelper.h"
 #include "Misc/ConfigManifest.h"
 #include "Misc/DataDrivenPlatformInfoRegistry.h"
+#include "Misc/StringBuilder.h"
 
 #if WITH_EDITOR
 	#define INI_CACHE 1
@@ -1439,26 +1442,39 @@ static void ClearHierarchyCache( const TCHAR* BaseIniName )
 #endif
 }
 
-bool FConfigFile::Write( const FString& Filename, bool bDoRemoteWrite/* = true*/, const FString& InitialText/*=FString()*/ )
+bool FConfigFile::Write(const FString& Filename, bool bDoRemoteWrite/* = true*/, const FString& PrefixText/*=FString()*/)
+{
+	TMap<FString, FString> SectionTexts;
+	TArray<FString> SectionOrder;
+	if (!PrefixText.IsEmpty())
+	{
+		SectionTexts.Add(FString(), PrefixText);
+	}
+	return Write(Filename, bDoRemoteWrite, SectionTexts, SectionOrder);
+}
+
+bool FConfigFile::Write(const FString& Filename, bool bDoRemoteWrite, TMap<FString, FString>& InOutSectionTexts, const TArray<FString>& InSectionOrder)
 {
 	if( !Dirty || NoSave || FParse::Param( FCommandLine::Get(), TEXT("nowrite")) || 
 		(FParse::Param( FCommandLine::Get(), TEXT("Multiprocess"))  && !FParse::Param( FCommandLine::Get(), TEXT("MultiprocessSaveConfig"))) // Is can be useful to save configs with multiprocess if they are given INI overrides
 		) 
 		return true;
 
-	FString Text = InitialText;
-
 	bool bAcquiredIniCombineThreshold = false;	// avoids extra work when writing multiple properties
 	int32 IniCombineThreshold = -1;
 
+	TStringBuilder<128> Text;
+	FStringView BlankLine(LINE_TERMINATOR LINE_TERMINATOR);
+	TArray<FString> SectionOrder;
+	SectionOrder.Reserve(InSectionOrder.Num() + this->Num());
+	SectionOrder.Append(InSectionOrder);
+
 	for( TIterator SectionIterator(*this); SectionIterator; ++SectionIterator )
 	{
 		const FString& SectionName = SectionIterator.Key();
 		const FConfigSection& Section = SectionIterator.Value();
 
-		// Flag to check whether a property was written on this section, 
-		// if none we do not want to make any changes to the destination file on this round.
-		bool bWroteASectionProperty = false;
+		Text.Reset();
 
 		TSet< FName > PropertiesAddedLookup;
 
@@ -1504,15 +1520,14 @@ bool FConfigFile::Write( const FString& Filename, bool bDoRemoteWrite/* = true*/
 				if ((bIsADefaultIniWrite || bDifferentNumberOfElements || !DoesConfigPropertyValueMatch(SourceConfigFile, SectionName, PropertyName, PropertyValue)) && !bOptionIsFromCommandline)
 				{
 					// If this is the first property we are writing of this section, then print the section name
-					if( !bWroteASectionProperty )
+					if( Text.Len() == 0 )
 					{
-						Text += FString::Printf( TEXT("[%s]") LINE_TERMINATOR, *SectionName);
-						bWroteASectionProperty = true;
+						Text.Append(FString::Printf( TEXT("[%s]") LINE_TERMINATOR, *SectionName));
 
-						// and if the sectrion has any array of struct uniqueness keys, add them here
+						// and if the section has any array of struct uniqueness keys, add them here
 						for (auto It = Section.ArrayOfStructKeys.CreateConstIterator(); It; ++It)
 						{
-							Text += FString::Printf(TEXT("@%s=%s") LINE_TERMINATOR, *It.Key().ToString(), *It.Value());
+							Text.Append(FString::Printf(TEXT("@%s=%s") LINE_TERMINATOR, *It.Key().ToString(), *It.Value()));
 						}
 					}
 
@@ -1543,7 +1558,7 @@ bool FConfigFile::Write( const FString& Filename, bool bDoRemoteWrite/* = true*/
 					{
 						for (const FConfigValue& ConfigValue : CompletePropertyToWrite)
 						{
-							Text += GenerateExportedPropertyLine(PropertyName.ToString(), ConfigValue.GetSavedValue());
+							Text.Append(GenerateExportedPropertyLine(PropertyName.ToString(), ConfigValue.GetSavedValue()));
 						}
 					}
 
@@ -1552,24 +1567,88 @@ bool FConfigFile::Write( const FString& Filename, bool bDoRemoteWrite/* = true*/
 			}
 		}
 
-		// If we wrote any part of the section, then add some whitespace after the section.
-		if( bWroteASectionProperty )
+		// If we didn't decide to write any properties on this section, then we don't add the section
+		// to the destination file
+		if (Text.Len() > 0)
 		{
-			Text += LINE_TERMINATOR;
+			InOutSectionTexts.FindOrAdd(SectionName) = FString(Text.ToString(), Text.Len());
+
+			// Add the Section to SectionOrder in case it's not already there
+			SectionOrder.Add(SectionName);
 		}
+		else
+		{
+			InOutSectionTexts.Remove(SectionName);
+		}
+	}
 
+	// Join all of the sections together
+	Text.Reset();
+	TSet<FString> SectionNamesLeftToWrite;
+	SectionNamesLeftToWrite.Reserve(InOutSectionTexts.Num());
+	for (TPair<FString,FString>& kvpair : InOutSectionTexts)
+	{
+		SectionNamesLeftToWrite.Add(kvpair.Key);
+	}
+
+	auto AddSectionToText = [&Text, &InOutSectionTexts, &SectionNamesLeftToWrite, &BlankLine](const FString& SectionName)
+	{
+		FString* SectionText = InOutSectionTexts.Find(SectionName);
+		if (!SectionText)
+		{
+			return;
+		}
+		if (SectionNamesLeftToWrite.Remove(SectionName) == 0)
+		{
+			// We already wrote this section
+			return;
+		}
+		Text.Append(*SectionText);
+		if (!FStringView(Text.ToString()).EndsWith(BlankLine))
+		{
+			Text.Append(LINE_TERMINATOR);
+		}
+	};
+
+	// First add the empty section
+	AddSectionToText(FString());
+
+	// Second add all the sections in SectionOrder; this includes any sections in *this that were not in InSectionOrder, because we added them during the loop
+	for (FString& SectionName : SectionOrder)
+	{
+		AddSectionToText(SectionName);
+	}
+
+	// Third add any remaining sections that were passed in in InOutSectionTexts but were not specified in InSectionOrder and were not in *this
+	if (SectionNamesLeftToWrite.Num() > 0)
+	{
+		TArray<FString> RemainingNames;
+		for (FString& SectionName : SectionNamesLeftToWrite)
+		{
+			RemainingNames.Add(SectionName);
+		}
+		RemainingNames.Sort();
+		for (FString& SectionName : RemainingNames)
+		{
+			AddSectionToText(SectionName);
+		}
 	}
 
 	// Ensure We have at least something to write
-	Text += LINE_TERMINATOR;
+	if (Text.Len() == 0)
+	{
+		Text.Append(LINE_TERMINATOR);
+	}
+
 
+	FString TextAsString(Text.ToString(), Text.Len());
 	if (bDoRemoteWrite)
 	{
 		// Write out the remote version (assuming it was loaded)
-		FRemoteConfig::Get()->Write(*Filename, Text);
+		FRemoteConfig::Get()->Write(*Filename, TextAsString);
 	}
 
-	bool bResult = SaveConfigFileWrapper(*Filename, Text);
+	bool bResult = SaveConfigFileWrapper(*Filename, TextAsString);
 
 #if INI_CACHE
 	// if we wrote the config successfully
@@ -1867,7 +1946,14 @@ void FConfigFile::ProcessSourceAndCheckAgainstBackup()
 	}
 }
 
-void FConfigFile::ProcessPropertyAndWriteForDefaults( int IniCombineThreshold, const TArray< FConfigValue >& InCompletePropertyToProcess, FString& OutText, const FString& SectionName, const FString& PropertyName )
+void FConfigFile::ProcessPropertyAndWriteForDefaults(int IniCombineThreshold, const TArray< FConfigValue >& InCompletePropertyToProcess, FString& OutText, const FString& SectionName, const FString& PropertyName)
+{
+	TStringBuilder<128> PropertyText;
+	ProcessPropertyAndWriteForDefaults(IniCombineThreshold, InCompletePropertyToProcess, PropertyText, SectionName, PropertyName);
+	OutText.Append(PropertyText.ToString(), PropertyText.Len());
+}
+
+void FConfigFile::ProcessPropertyAndWriteForDefaults( int IniCombineThreshold, const TArray< FConfigValue >& InCompletePropertyToProcess, FStringBuilderBase& OutText, const FString& SectionName, const FString& PropertyName )
 {
 	// Only process against a hierarchy if this config file has one.
 	if (SourceIniHierarchy.Num() > 0)
@@ -1909,7 +1995,7 @@ void FConfigFile::ProcessPropertyAndWriteForDefaults( int IniCombineThreshold, c
 			for (const FString& NextElement : ArrayProperties)
 			{
 				FString PropertyNameWithRemoveOp = PropertyName.Replace(TEXT("+"), TEXT("-"));
-				OutText += GenerateExportedPropertyLine(PropertyNameWithRemoveOp, NextElement);
+				OutText.Append(GenerateExportedPropertyLine(PropertyNameWithRemoveOp, NextElement));
 			}
 		}
 	}
@@ -1917,7 +2003,7 @@ void FConfigFile::ProcessPropertyAndWriteForDefaults( int IniCombineThreshold, c
 	// Write the properties out to a file.
 	for ( const FConfigValue& PropertyIt : InCompletePropertyToProcess)
 	{
-		OutText += GenerateExportedPropertyLine(PropertyName, PropertyIt.GetSavedValue());
+		OutText.Append(GenerateExportedPropertyLine(PropertyName, PropertyIt.GetSavedValue()));
 	}
 }
 
@@ -3862,10 +3948,52 @@ FString FConfigCacheIni::GetGameUserSettingsDir()
 
 void FConfigFile::UpdateSections(const TCHAR* DiskFilename, const TCHAR* IniRootName/*=nullptr*/, const TCHAR* OverridePlatform/*=nullptr*/)
 {
-	// since we don't want any modifications to other sections, we manually process the file, not read into sections, etc
+	// Since we don't want any modifications to other sections, we manually process the file, not read into sections, etc
+	// Keep track of existing SectionTexts and orders so that we can preserve the order of the sections in Write to reduce the diff we make to the file on disk
 	FString DiskFile;
 	FString NewFile;
-	bool bIsLastLineEmtpy = false;
+	TStringBuilder<128> SectionText;
+	TMap<FString, FString> SectionTexts;
+	TArray<FString> SectionOrder;
+	FString SectionName;
+
+	auto AddSectionText = [this, &SectionTexts, &SectionOrder, &SectionName, &SectionText]()
+	{
+		if (SectionText.Len() == 0)
+		{
+			// No text in the section, not even a section header, e.g. because this is the prefix section and there was no prefix.
+			// Skip the section - add it to SectionTexts or SectionOrder
+		}
+		else
+		{
+			if (Contains(SectionName))
+			{
+				// Do not add the section to SectionTexts, so that Write will skip writing it at all if it is empty
+				// But do add it to SectionOrder, so that Write will write it to the right location if it is non-empty
+			}
+			else
+			{
+				// Check for duplicate sections in the file-on-disk; handle these by combining them.  This will modify the file, but will guarantee that we don't lose data
+				FString* ExistingSectionText = SectionTexts.Find(SectionName);
+				if (ExistingSectionText)
+				{
+					ExistingSectionText->Append(SectionText.ToString(), SectionText.Len());
+				}
+				else
+				{
+					SectionTexts.Add(SectionName, FString(SectionText.ToString()));
+				}
+			}
+
+			SectionOrder.Add(SectionName);
+		}
+
+		// Clear the name and text for the next section
+		SectionName.Reset();
+		SectionText.Reset();
+	};
+
+	SectionName = FString(); // The lines we read before we encounter a section header should be preserved as prefix lines; we implement this by storing them in an empty SectionName
 	if (LoadConfigFileWrapper(DiskFilename, DiskFile))
 	{
 		// walk each line
@@ -3885,23 +4013,22 @@ void FConfigFile::UpdateSections(const TCHAR* DiskFilename, const TCHAR* IniRoot
 				// is this line a section? (must be at least [x])
 				if (TheLine.Len() > 3 && TheLine[0] == '[' && TheLine[TheLine.Len() - 1] == ']')
 				{
-					// look to see if this section is one we are going to update; if so, then skip lines until a new section
-					FString SectionName = TheLine.Mid(1, TheLine.Len() - 2);
-					bIsSkippingSection = Contains(SectionName);
-				}
-
-				// if we aren't skipping, then write out the line
-				if (!bIsSkippingSection)
-				{
-					NewFile += TheLine + LINE_TERMINATOR;
+					// Add the old section we just finished reading
+					AddSectionText();
 
-					// track if the last line written was empty
-					bIsLastLineEmtpy = TheLine.Len() == 0;
+					// Set SectionName to the name of new section we are about to read
+					SectionName = TheLine.Mid(1, TheLine.Len() - 2);
 				}
+
+				SectionText.Append(TheLine);
+				SectionText.Append(LINE_TERMINATOR);
 			}
 		}
 	}
 
+	// Add the last section we read
+	AddSectionText();
+
 	// load the hierarchy up to right before this file
 	if (IniRootName != nullptr)
 	{
@@ -3928,7 +4055,7 @@ void FConfigFile::UpdateSections(const TCHAR* DiskFilename, const TCHAR* IniRoot
 
 		ClearHierarchyCache(IniRootName);
 
-		// Get a collection of the source hierachy properties
+		// Get a collection of the source hierarchy properties
 		if (SourceConfigFile)
 		{
 			delete SourceConfigFile;
@@ -3939,13 +4066,7 @@ void FConfigFile::UpdateSections(const TCHAR* DiskFilename, const TCHAR* IniRoot
 		LoadIniFileHierarchy(SourceIniHierarchy, *SourceConfigFile, true);
 	}
 
-	// take what we got above (which has the sections skipped), and then append the new sections
-	if (Num() > 0 && !bIsLastLineEmtpy)
-	{
-		// add a blank link between old sections and new (if there are any new sections)
-		NewFile += LINE_TERMINATOR;
-	}
-	Write(DiskFilename, true, NewFile);
+	Write(DiskFilename, true, SectionTexts, SectionOrder);
 }
 
 
@@ -3982,7 +4103,12 @@ public:
 	{
 		UpdatePropertyInSection();
 		// Rebuild the file with the updated section.
-		const FString NewFile = IniFileMakeup.BeforeSection + IniFileMakeup.Section + IniFileMakeup.AfterSection;
+
+		FString NewFile = IniFileMakeup.BeforeSection + IniFileMakeup.Section + IniFileMakeup.AfterSection;
+		if (!NewFile.EndsWith(LINE_TERMINATOR LINE_TERMINATOR))
+		{
+			NewFile.AppendChars(LINE_TERMINATOR, TCString<TCHAR>::Strlen(LINE_TERMINATOR));
+		}
 		return SaveConfigFileWrapper(*IniFilename, NewFile);
 	}
 
diff --git a/Engine/Source/Runtime/Core/Public/Misc/ConfigCacheIni.h b/Engine/Source/Runtime/Core/Public/Misc/ConfigCacheIni.h
--- a/Engine/Source/Runtime/Core/Public/Misc/ConfigCacheIni.h
+++ b/Engine/Source/Runtime/Core/Public/Misc/ConfigCacheIni.h
@@ -20,6 +20,8 @@
 #include "Misc/Paths.h"
 #include "Serialization/StructuredArchive.h"
 
+class FStringBuilderBase;
+
 CORE_API DECLARE_LOG_CATEGORY_EXTERN(LogConfig, Log, All);
 
 // Server builds should be tweakable even in Shipping
@@ -342,7 +344,23 @@ public:
 	CORE_API bool Combine( const FString& Filename);
 	CORE_API void CombineFromBuffer(const FString& Buffer);
 	CORE_API void Read( const FString& Filename );
-	CORE_API bool Write( const FString& Filename, bool bDoRemoteWrite=true, const FString& InitialText=FString() );
+
+	/** Write this ConfigFile to the given Filename, constructed the text from the config sections in *this, prepended by the optional PrefixText */
+	CORE_API bool Write( const FString& Filename, bool bDoRemoteWrite=true, const FString& PrefixText=FString());
+
+	/** Write a ConfigFile to the ginve Filename, constructed from the given SectionTexts, in the given order, with sections in *this overriding sections in SectionTexts
+	 * @param Filename - The file to write to
+	 * @param bDoRemoteWrite - If true, also write the file to FRemoteConfig::Get()
+	 * @param InOutSectionTexts - A map from section name to existing text for that section; text does not include the name of the section.
+	 *  Entries in the TMap that also exist in *this will be updated.
+	 *  If the empty string is present, it will be written out first (it is interpreted as a prefix before the first section)
+	 * @param InSectionOrder - List of section names in the order in which each section should be written to disk, from e.g. the existing file.
+	 *  Any section in this array that is not found in InOutSectionTexts will be ignored.
+	 *  Any section in InOutSectionTexts that is not in this array will be appended to the end.
+	 *  Duplicate entries are ignored; the first found index is used.
+	 * @return TRUE if the write was successful
+	 */
+	CORE_API bool Write(const FString& Filename, bool bDoRemoteWrite, TMap<FString, FString>& InOutSectionTexts, const TArray<FString>& InSectionOrder);
 	CORE_API void Dump(FOutputDevice& Ar);
 
 	CORE_API bool GetString( const TCHAR* Section, const TCHAR* Key, FString& Value ) const;
@@ -370,8 +388,8 @@ public:
 	CORE_API void AddMissingProperties(const FConfigFile& InSourceFile);
 
 	/**
-	 * Saves only the sections in this FConfigFile its source files. All other sections in the file are left alone. The sections in this
-	 * file are completely replaced. If IniRootName is specified, the saved settings are the diffed against the file in the hierarchy up
+	 * Saves only the sections in this FConfigFile into the given file. All other sections in the file are left alone. The sections in this
+	 * file are completely replaced. If IniRootName is specified, the current section settings are diffed against the file in the hierarchy up
 	 * to right before this file (so, if you are saving DefaultEngine.ini, and IniRootName is "Engine", then Base.ini and BaseEngine.ini
 	 * will be loaded, and only differences against that will be saved into DefaultEngine.ini)
 	 *
@@ -432,6 +450,9 @@ private:
 	 * @param SectionName - The section name the array property is being written to
 	 * @param PropertyName - The property name of the array
 	 */
+	void ProcessPropertyAndWriteForDefaults(int32 IniCombineThreshold, const TArray<FConfigValue>& InCompletePropertyToProcess, FStringBuilderBase& OutText, const FString& SectionName, const FString& PropertyName);
+
+	/** Version of ProcessPropertyAndWriteForDefaults that takes an FString to append to rather than a StringBuilder */
 	void ProcessPropertyAndWriteForDefaults(int32 IniCombineThreshold, const TArray<FConfigValue>& InCompletePropertyToProcess, FString& OutText, const FString& SectionName, const FString& PropertyName);
 
 	/**
