Removed unnecessary ChunkId's in I/O store runtime TOC.

diff --git a/Engine/Source/Runtime/Core/Private/IO/IoDispatcherFileBackend.cpp b/Engine/Source/Runtime/Core/Private/IO/IoDispatcherFileBackend.cpp
--- a/Engine/Source/Runtime/Core/Private/IO/IoDispatcherFileBackend.cpp
+++ b/Engine/Source/Runtime/Core/Private/IO/IoDispatcherFileBackend.cpp
@@ -94,7 +94,7 @@ FIoStatus FFileIoStoreReader::Initialize(const FIoStoreEnvironment& Environment)
 			return FIoStatusBuilder(EIoErrorCode::CorruptToc) << TEXT("TOC entry out of container bounds while reading '") << *TocFilePath << TEXT("'");
 		}
 
-		Toc.Add(Entry->ChunkId, *Entry);
+		Toc.Add(Entry->ChunkId, Entry->OffsetAndLength);
 		++Entry;
 	}
 
@@ -108,11 +108,11 @@ bool FFileIoStoreReader::DoesChunkExist(const FIoChunkId& ChunkId) const
 
 TIoStatusOr<uint64> FFileIoStoreReader::GetSizeForChunk(const FIoChunkId& ChunkId) const
 {
-	const FIoStoreTocEntry* Entry = Toc.Find(ChunkId);
+	const FIoOffsetAndLength* OffsetAndLength = Toc.Find(ChunkId);
 
-	if (Entry != nullptr)
+	if (OffsetAndLength != nullptr)
 	{
-		return Entry->GetLength();
+		return OffsetAndLength->GetLength();
 	}
 	else
 	{
@@ -122,16 +122,16 @@ TIoStatusOr<uint64> FFileIoStoreReader::GetSizeForChunk(const FIoChunkId& ChunkI
 
 bool FFileIoStoreReader::Resolve(FFileIoStoreResolvedRequest& ResolvedRequest)
 {
-	const FIoStoreTocEntry* Entry = Toc.Find(ResolvedRequest.Request->ChunkId);
+	const FIoOffsetAndLength* OffsetAndLength = Toc.Find(ResolvedRequest.Request->ChunkId);
 
-	if (!Entry)
+	if (!OffsetAndLength)
 	{
 		return false;
 	}
 
 	ResolvedRequest.ResolvedFileHandle = ContainerFileHandle;
-	uint64 FileEndOffset = Entry->GetOffset() + Entry->GetLength();
-	uint64 RequestedBeginOffset = Entry->GetOffset() + ResolvedRequest.Request->Options.GetOffset();
+	uint64 FileEndOffset = OffsetAndLength->GetOffset() + OffsetAndLength->GetLength();
+	uint64 RequestedBeginOffset = OffsetAndLength->GetOffset() + ResolvedRequest.Request->Options.GetOffset();
 	uint64 RequestedEndOffset = FMath::Min(FileEndOffset, RequestedBeginOffset + ResolvedRequest.Request->Options.GetSize());
 	ResolvedRequest.ResolvedFileSize = ContainerFileSize;
 	ResolvedRequest.ResolvedOffset = RequestedBeginOffset;
diff --git a/Engine/Source/Runtime/Core/Private/IO/IoDispatcherFileBackend.h b/Engine/Source/Runtime/Core/Private/IO/IoDispatcherFileBackend.h
--- a/Engine/Source/Runtime/Core/Private/IO/IoDispatcherFileBackend.h
+++ b/Engine/Source/Runtime/Core/Private/IO/IoDispatcherFileBackend.h
@@ -64,7 +64,7 @@ public:
 private:
 	FFileIoStoreImpl& PlatformImpl;
 
-	TMap<FIoChunkId, FIoStoreTocEntry> Toc;
+	TMap<FIoChunkId, FIoOffsetAndLength> Toc;
 	uint64 ContainerFileHandle;
 	uint64 ContainerFileSize;
 };
diff --git a/Engine/Source/Runtime/Core/Private/IO/IoStore.h b/Engine/Source/Runtime/Core/Private/IO/IoStore.h
--- a/Engine/Source/Runtime/Core/Private/IO/IoStore.h
+++ b/Engine/Source/Runtime/Core/Private/IO/IoStore.h
@@ -29,18 +29,11 @@ struct FIoStoreTocHeader
 };
 
 /**
- * I/O Store TOC entry.
+ * Combined offset and length.
  */
-struct FIoStoreTocEntry
+struct FIoOffsetAndLength
 {
-	// We use 5 bytes for offset and size, this is enough to represent 
-	// an offset and size of 1PB
-	uint8		OffsetAndLength[5 + 5];
-
-	// TBD: should the chunk ID use content addressing, or names, or a
-	//      mix of both?
-	FIoChunkId	ChunkId;
-
+public:
 	inline uint64 GetOffset() const
 	{
 		return OffsetAndLength[4]
@@ -78,4 +71,41 @@ struct FIoStoreTocEntry
 		OffsetAndLength[8] = uint8(Length >> 8);
 		OffsetAndLength[9] = uint8(Length >> 0);
 	}
+
+private:
+	// We use 5 bytes for offset and size, this is enough to represent
+	// an offset and size of 1PB
+	uint8 OffsetAndLength[5 + 5];
+};
+
+/**
+ * I/O Store TOC entry.
+ */
+struct FIoStoreTocEntry
+{
+	FIoOffsetAndLength	OffsetAndLength;
+
+	// TBD: should the chunk ID use content addressing, or names, or a
+	//      mix of both?
+	FIoChunkId	ChunkId;
+
+	inline uint64 GetOffset() const
+	{
+		return OffsetAndLength.GetOffset();
+	}
+
+	inline uint64 GetLength() const
+	{
+		return OffsetAndLength.GetLength();
+	}
+
+	inline void SetOffset(uint64 Offset)
+	{
+		OffsetAndLength.SetOffset(Offset);
+	}
+
+	inline void SetLength(uint64 Length)
+	{
+		OffsetAndLength.SetLength(Length);
+	}
 };
