Added missing include

diff --git a/Engine/Source/Runtime/Core/Private/IO/IoDispatcher.cpp b/Engine/Source/Runtime/Core/Private/IO/IoDispatcher.cpp
--- a/Engine/Source/Runtime/Core/Private/IO/IoDispatcher.cpp
+++ b/Engine/Source/Runtime/Core/Private/IO/IoDispatcher.cpp
@@ -13,6 +13,7 @@
 #include "HAL/PlatformProcess.h"
 #include "Serialization/LargeMemoryReader.h"
 #include "GenericPlatform/GenericPlatformChunkInstall.h"
+#include "HAL/Event.h"
 
 DEFINE_LOG_CATEGORY(LogIoDispatcher);
 
