Avoid calling strlen when TStringBuilderImpl::Append is provided a length

- This was inconsistent with the behavior of AppendAnsi and now recently starting causing problems with appending empty views because Append(const TStringView&) was changed to call the pointer/size overload and the pointer may be null. Operations such as Builder << MutableArrayOfTCHAR were previously calling strlen twice.

#rnx

diff --git a/Engine/Source/Runtime/Core/Public/Misc/StringBuilder.h b/Engine/Source/Runtime/Core/Public/Misc/StringBuilder.h
--- a/Engine/Source/Runtime/Core/Public/Misc/StringBuilder.h
+++ b/Engine/Source/Runtime/Core/Public/Misc/StringBuilder.h
@@ -9,6 +9,7 @@
 #include "Templates/EnableIf.h"
 #include "Templates/IsArrayOrRefOfType.h"
 #include "Templates/IsValidVariadicFunctionArg.h"
+#include "Templates/UnrealTemplate.h"
 #include "Traits/IsContiguousContainer.h"
 
 #define USE_STRING_LITERAL_PATH 1
@@ -193,11 +194,8 @@ public:
 		return Append(StringView.GetData(), StringView.Len());
 	}
 
-	inline BuilderType& Append(const C* String, int32 MaxChars)
+	inline BuilderType& Append(const C* String, int32 Length)
 	{
-		const int32 StringLength = TCString<C>::Strlen(String);
-		const int32 Length = FPlatformMath::Min<int32>(MaxChars, StringLength);
-
 		EnsureCapacity(Length);
 		C* RESTRICT Dest = CurPos;
 		CurPos += Length;
@@ -522,7 +520,8 @@ inline FAnsiStringBuilderBase&				operator<<(FAnsiStringBuilderBase& Builder, co
 
 #endif
 
-inline FAnsiStringBuilderBase&				operator<<(FAnsiStringBuilderBase& Builder, const FAnsiStringView& Str)				{ return Builder.Append(Str); }
+template <typename T, typename = decltype(ImplicitConv<FAnsiStringView>(DeclVal<T>()))>
+inline FAnsiStringBuilderBase&				operator<<(FAnsiStringBuilderBase& Builder, T&& Str)								{ return Builder.Append(ImplicitConv<FAnsiStringView>(Forward<T>(Str))); }
 
 inline FAnsiStringBuilderBase&				operator<<(FAnsiStringBuilderBase& Builder, int32 Value)							{ return Builder.Appendf("%d", Value); }
 inline FAnsiStringBuilderBase&				operator<<(FAnsiStringBuilderBase& Builder, uint32 Value)							{ return Builder.Appendf("%u", Value); }
