Static analysis warning fix

diff --git a/Engine/Source/Runtime/Core/Private/IO/IoDispatcher.cpp b/Engine/Source/Runtime/Core/Private/IO/IoDispatcher.cpp
--- a/Engine/Source/Runtime/Core/Private/IO/IoDispatcher.cpp
+++ b/Engine/Source/Runtime/Core/Private/IO/IoDispatcher.cpp
@@ -185,7 +185,10 @@ public:
 		InstallManifestAr << InstallManifest;
 
 		IPlatformChunkInstall* ChunkInstall = FPlatformMisc::GetPlatformChunkInstall();
-		ChunkDownloadedDelegateHandle = ChunkInstall->AddChunkInstallDelegate(FPlatformChunkInstallDelegate::CreateRaw(this, &FIoDispatcherImpl::OnChunkDownloaded));
+		if (ChunkInstall)
+		{
+			ChunkDownloadedDelegateHandle = ChunkInstall->AddChunkInstallDelegate(FPlatformChunkInstallDelegate::CreateRaw(this, &FIoDispatcherImpl::OnChunkDownloaded));
+		}
 		for (const FIoStoreInstallManifest::FEntry& InstallManifestEntry : InstallManifest.ReadEntries())
 		{
 			FIoStoreEnvironment PartitionEnvironment(InitialEnvironment, InstallManifestEntry.PartitionName);
