Remove System.Core.Math.TruncationFunctions from SmokeTests since it writes to log and therefore is spamming every engine startup.

#rnx

diff --git a/Engine/Source/Runtime/Core/Private/Tests/Math/UnrealMathTest.cpp b/Engine/Source/Runtime/Core/Private/Tests/Math/UnrealMathTest.cpp
--- a/Engine/Source/Runtime/Core/Private/Tests/Math/UnrealMathTest.cpp
+++ b/Engine/Source/Runtime/Core/Private/Tests/Math/UnrealMathTest.cpp
@@ -1763,7 +1763,7 @@ bool FMathRoundHalfFromZeroTests::RunTest(const FString& Parameters)
 	return true;
 }
 
-IMPLEMENT_SIMPLE_AUTOMATION_TEST(FMathTruncationTests, "System.Core.Math.TruncationFunctions", EAutomationTestFlags::ApplicationContextMask | EAutomationTestFlags::SmokeFilter)
+IMPLEMENT_SIMPLE_AUTOMATION_TEST(FMathTruncationTests, "System.Core.Math.TruncationFunctions", EAutomationTestFlags::ApplicationContextMask | EAutomationTestFlags::EngineFilter)
 bool FMathTruncationTests::RunTest(const FString& Parameters)
 {
 	// Float: 1-bit Sign, 8-bit exponent, 23-bit mantissa, implicit 1
