Add FString += FStringView to facilitate conversion to FStringView

diff --git a/Engine/Source/Runtime/Core/Private/Containers/String.cpp b/Engine/Source/Runtime/Core/Private/Containers/String.cpp
--- a/Engine/Source/Runtime/Core/Private/Containers/String.cpp
+++ b/Engine/Source/Runtime/Core/Private/Containers/String.cpp
@@ -2,6 +2,7 @@
 
 #include "Containers/UnrealString.h"
 
+#include "Containers/StringView.h"
 #include "CoreGlobals.h"
 #include "CoreTypes.h"
 #include "HAL/UnrealMemory.h"
@@ -10,7 +11,6 @@
 #include "Math/UnrealMathUtility.h"
 #include "Misc/AssertionMacros.h"
 #include "Misc/ByteSwap.h"
-#include "Misc/StringView.h"
 #include "Misc/VarArgs.h"
 #include "Templates/UnrealTemplate.h"
 
@@ -151,6 +151,13 @@ FString& FString::operator=(const FStringView& Other)
 	return *this;
 }
 
+FString& FString::operator+=(const FStringView& Str)
+{
+	CheckInvariants();
+	AppendChars(Str.GetData(), Str.Len());
+	return *this;
+}
+
 FString::operator FStringView() const
 {
 	return FStringView(Data.GetData(), Len());
diff --git a/Engine/Source/Runtime/Core/Public/Containers/UnrealString.h b/Engine/Source/Runtime/Core/Public/Containers/UnrealString.h
--- a/Engine/Source/Runtime/Core/Public/Containers/UnrealString.h
+++ b/Engine/Source/Runtime/Core/Public/Containers/UnrealString.h
@@ -385,6 +385,14 @@ public:
 		*(EndPtr + Count) = 0;
 	}
 
+	/**
+	 * Concatenate this with the given string
+	 *
+	 * @param Str string view to be concatenated onto the end of this
+	 * @return reference to this
+	 */
+	FString& operator+=(const FStringView& Str);
+
 	/**
 	 * Concatenate this with given string
 	 * 
