Add Algo::Replace, String::ParseLines, String::ParseTokens[Multiple]

- Algo::Replace can be used to replace FString::ReplaceCharInline on array views of characters such as on a view of a string builder.
- String::ParseLines can be used to split a string view on line endings to find the non-empty lines in the view. Intended to replace functions like FString::ParseIntoArrayLines and FFileHelper::LoadFileToStringArray.
- String::ParseTokens[Multiple] can be used to split a string view on one or more single-character or multi-character delimiters. Intended to replace FString::ParseIntoArray.
- Extend Algo::CompareByPredicate to allow the input ranges to be different types and to take them by universal reference.
- Extend TArrayView and MakeArrayView to allow array views to be created from initializer lists.

diff --git a/Engine/Source/Runtime/Core/Private/String/ParseLines.cpp b/Engine/Source/Runtime/Core/Private/String/ParseLines.cpp
new file mode 100644
--- /dev/null
+++ b/Engine/Source/Runtime/Core/Private/String/ParseLines.cpp
@@ -0,0 +1,40 @@
+// Copyright 1998-2019 Epic Games, Inc. All Rights Reserved.
+
+#include "String/ParseLines.h"
+
+#include "Containers/StringView.h"
+
+namespace String
+{
+	void ParseLines(const FStringView& View, const TFunctionRef<void(FStringView)>& Visitor)
+	{
+		const TCHAR* ViewIt = View.GetData();
+		const TCHAR* ViewEnd = ViewIt + View.Len();
+		do
+		{
+			const TCHAR* LineStart = ViewIt;
+			const TCHAR* LineEnd = ViewEnd;
+			for (; ViewIt != ViewEnd; ++ViewIt)
+			{
+				const TCHAR CurrentChar = *ViewIt;
+				if (CurrentChar == TEXT('\n'))
+				{
+					LineEnd = ViewIt++;
+					break;
+				}
+				if (CurrentChar == TEXT('\r'))
+				{
+					LineEnd = ViewIt++;
+					if (ViewIt != ViewEnd && *ViewIt == TEXT('\n'))
+					{
+						++ViewIt;
+					}
+					break;
+				}
+			}
+
+			Visitor(FStringView(LineStart, LineEnd - LineStart));
+		}
+		while (ViewIt != ViewEnd);
+	}
+}
diff --git a/Engine/Source/Runtime/Core/Private/String/ParseTokens.cpp b/Engine/Source/Runtime/Core/Private/String/ParseTokens.cpp
new file mode 100644
--- /dev/null
+++ b/Engine/Source/Runtime/Core/Private/String/ParseTokens.cpp
@@ -0,0 +1,58 @@
+// Copyright 1998-2019 Epic Games, Inc. All Rights Reserved.
+
+#include "String/ParseTokens.h"
+
+#include "Algo/Find.h"
+#include "Algo/NoneOf.h"
+#include "Containers/StringView.h"
+
+namespace String
+{
+	template <typename DelimiterType>
+	inline static void ParseTokensMultipleImpl(const FStringView& View, TArrayView<const DelimiterType> Delimiters, TFunctionRef<void(FStringView)> Visitor)
+	{
+		// This is a naive implementation that takes time proportional to InView.Len() * TotalDelimiterLen.
+		// If this function becomes a bottleneck, it can be specialized separately for single-character delimiters
+		// and multi-character delimiters, as well as for one and many delimiters. There are algorithms for each of
+		// these four combinations of requirements that are linear or sub-linear in the length of string to search.
+
+		const FStringView::SizeType ViewLen = View.Len();
+		FStringView::SizeType NextTokenIndex = 0;
+
+		for (FStringView::SizeType ViewIndex = 0; ViewIndex != ViewLen;)
+		{
+			struct
+			{
+				FStringView RemainingView;
+				inline bool operator()(const FStringView& InDelim) const { return RemainingView.StartsWith(InDelim, ESearchCase::CaseSensitive); }
+				inline bool operator()(const TCHAR InDelim) const { return RemainingView.StartsWith(InDelim); }
+				inline FStringView::SizeType Len(const FStringView& InDelim) const { return InDelim.Len(); }
+				inline int32 Len(const TCHAR InDelim) const { return 1; }
+			} MatchDelim{FStringView(View.GetData() + ViewIndex, ViewLen - ViewIndex)};
+
+			if (const DelimiterType* Delim = Algo::FindByPredicate(Delimiters, MatchDelim))
+			{
+				Visitor(FStringView(View.GetData() + NextTokenIndex, ViewIndex - NextTokenIndex));
+				ViewIndex += MatchDelim.Len(*Delim);
+				NextTokenIndex = ViewIndex;
+			}
+			else
+			{
+				++ViewIndex;
+			}
+		}
+
+		Visitor(FStringView(View.GetData() + NextTokenIndex, ViewLen - NextTokenIndex));
+	}
+
+	void ParseTokensMultiple(const FStringView& View, TArrayView<const FStringView> Delimiters, TFunctionRef<void(FStringView)> Visitor)
+	{
+		check(Algo::NoneOf(Delimiters, &FStringView::IsEmpty));
+		ParseTokensMultipleImpl(View, Delimiters, Visitor);
+	}
+
+	void ParseTokensMultiple(const FStringView& View, TArrayView<const TCHAR> Delimiters, TFunctionRef<void(FStringView)> Visitor)
+	{
+		ParseTokensMultipleImpl(View, Delimiters, Visitor);
+	}
+}
diff --git a/Engine/Source/Runtime/Core/Private/Tests/String/ParseLinesTest.cpp b/Engine/Source/Runtime/Core/Private/Tests/String/ParseLinesTest.cpp
new file mode 100644
--- /dev/null
+++ b/Engine/Source/Runtime/Core/Private/Tests/String/ParseLinesTest.cpp
@@ -0,0 +1,54 @@
+// Copyright 1998-2019 Epic Games, Inc. All Rights Reserved.
+
+#include "String/ParseLines.h"
+
+#include "Algo/Compare.h"
+#include "Containers/StringView.h"
+#include "Misc/AutomationTest.h"
+#include "Misc/StringBuilder.h"
+#include "Templates/EqualTo.h"
+
+#if WITH_DEV_AUTOMATION_TESTS
+
+IMPLEMENT_SIMPLE_AUTOMATION_TEST(StringParseLinesTest, "System.Core.String.ParseLines", EAutomationTestFlags::ApplicationContextMask | EAutomationTestFlags::EngineFilter)
+bool StringParseLinesTest::RunTest(const FString& Parameters)
+{
+	auto RunParseLinesTest = [this](const FStringView& View, std::initializer_list<FStringView> ExpectedLines)
+	{
+		TArray<FStringView, TInlineAllocator<8>> ResultLines;
+		String::ParseLines(View, [&ResultLines](FStringView Line) { ResultLines.Add(Line); });
+		if (!Algo::CompareByPredicate(ResultLines, ExpectedLines, TEqualTo<>()))
+		{
+			TStringBuilder<512> Error;
+			Error << TEXT("String::ParseLines failed to parse \"") << FString(View).ReplaceCharWithEscapedChar() << TEXT("\" result {");
+			Error.JoinQuoted(ResultLines, TEXT(", "_SV), TEXT("\""_SV));
+			Error << TEXT("} expected {");
+			Error.JoinQuoted(ExpectedLines, TEXT(", "_SV), TEXT("\""_SV));
+			Error << TEXT("}");
+			AddError(Error.ToString());
+		}
+	};
+
+	RunParseLinesTest(TEXT(""_SV),                        {TEXT(""_SV)});
+	RunParseLinesTest(TEXT("\n"_SV),                      {TEXT(""_SV)});
+	RunParseLinesTest(TEXT("\r"_SV),                      {TEXT(""_SV)});
+	RunParseLinesTest(TEXT("\r\n"_SV),                    {TEXT(""_SV)});
+	RunParseLinesTest(TEXT("\n\n"_SV),                    {TEXT(""_SV), TEXT(""_SV)});
+	RunParseLinesTest(TEXT("\r\r"_SV),                    {TEXT(""_SV), TEXT(""_SV)});
+	RunParseLinesTest(TEXT("\r\n\r\n"_SV),                {TEXT(""_SV), TEXT(""_SV)});
+	RunParseLinesTest(TEXT("\r\nABC"_SV).Left(2),         {TEXT(""_SV)});
+	RunParseLinesTest(TEXT("\r\nABC\r\nDEF"_SV).Left(5),  {TEXT(""_SV), TEXT("ABC"_SV)});
+	RunParseLinesTest(TEXT("ABC DEF"_SV),                 {TEXT("ABC DEF"_SV)});
+	RunParseLinesTest(TEXT("\nABC DEF\n"_SV),             {TEXT(""_SV), TEXT("ABC DEF"_SV)});
+	RunParseLinesTest(TEXT("\rABC DEF\r"_SV),             {TEXT(""_SV), TEXT("ABC DEF"_SV)});
+	RunParseLinesTest(TEXT("\r\nABC DEF\r\n"_SV),         {TEXT(""_SV), TEXT("ABC DEF"_SV)});
+	RunParseLinesTest(TEXT("\r\n\r\nABC DEF\r\n\r\n"_SV), {TEXT(""_SV), TEXT(""_SV), TEXT("ABC DEF"_SV), TEXT(""_SV)});
+	RunParseLinesTest(TEXT("ABC\nDEF"_SV),                {TEXT("ABC"_SV), TEXT("DEF"_SV)});
+	RunParseLinesTest(TEXT("ABC\rDEF"_SV),                {TEXT("ABC"_SV), TEXT("DEF"_SV)});
+	RunParseLinesTest(TEXT("\r\nABC\r\nDEF\r\n"_SV),      {TEXT(""_SV), TEXT("ABC"_SV), TEXT("DEF"_SV)});
+	RunParseLinesTest(TEXT("\r\nABC\r\n\r\nDEF\r\n"_SV),  {TEXT(""_SV), TEXT("ABC"_SV), TEXT(""_SV), TEXT("DEF"_SV)});
+
+	return true;
+}
+
+#endif // WITH_DEV_AUTOMATION_TESTS
diff --git a/Engine/Source/Runtime/Core/Private/Tests/String/ParseTokensTest.cpp b/Engine/Source/Runtime/Core/Private/Tests/String/ParseTokensTest.cpp
new file mode 100644
--- /dev/null
+++ b/Engine/Source/Runtime/Core/Private/Tests/String/ParseTokensTest.cpp
@@ -0,0 +1,91 @@
+// Copyright 1998-2019 Epic Games, Inc. All Rights Reserved.
+
+#include "String/ParseTokens.h"
+
+#include "Algo/Compare.h"
+#include "Misc/AutomationTest.h"
+#include "Misc/StringBuilder.h"
+#include "Templates/EqualTo.h"
+
+#if WITH_DEV_AUTOMATION_TESTS
+
+IMPLEMENT_SIMPLE_AUTOMATION_TEST(StringParseTokensByStringTest, "System.Core.String.ParseTokens.ByString", EAutomationTestFlags::ApplicationContextMask | EAutomationTestFlags::EngineFilter)
+bool StringParseTokensByStringTest::RunTest(const FString& Parameters)
+{
+	auto RunParseTokensTest = [this](const FStringView& View, std::initializer_list<FStringView> Delimiters, std::initializer_list<FStringView> ExpectedTokens)
+	{
+		TArray<FStringView, TInlineAllocator<8>> ResultTokens;
+		String::ParseTokensMultiple(View, Delimiters, [&ResultTokens](FStringView Token) { ResultTokens.Add(Token); });
+		if (!Algo::CompareByPredicate(ResultTokens, ExpectedTokens, TEqualTo<>()))
+		{
+			TStringBuilder<512> Error;
+			Error << TEXT("String::ParseTokensMultiple failed to parse \"") << View << TEXT("\" with delimiters {");
+			Error.JoinQuoted(Delimiters, TEXT(", "_SV), TEXT("\""_SV));
+			Error << TEXT("} result {");
+			Error.JoinQuoted(ResultTokens, TEXT(", "_SV), TEXT("\""_SV));
+			Error << TEXT("} expected {");
+			Error.JoinQuoted(ExpectedTokens, TEXT(", "_SV), TEXT("\""_SV));
+			Error << TEXT("}");
+			AddError(Error.ToString());
+		}
+	};
+
+	RunParseTokensTest(TEXT(""),         {},                       {TEXT("")});
+	RunParseTokensTest(TEXT("ABC"),      {},                       {TEXT("ABC")});
+
+	RunParseTokensTest(TEXT(""),         {TEXT(",")},              {TEXT("")});
+	RunParseTokensTest(TEXT(","),        {TEXT(",")},              {TEXT(""), TEXT("")});
+	RunParseTokensTest(TEXT(",,"),       {TEXT(",")},              {TEXT(""), TEXT(""), TEXT("")});
+	RunParseTokensTest(TEXT("ABC"),      {TEXT(",")},              {TEXT("ABC")});
+	RunParseTokensTest(TEXT("A,,C"),     {TEXT(",")},              {TEXT("A"), TEXT(""), TEXT("C")});
+	RunParseTokensTest(TEXT("A,B,C"),    {TEXT(",")},              {TEXT("A"), TEXT("B"), TEXT("C")});
+	RunParseTokensTest(TEXT(",A,B,C,"),  {TEXT(",")},              {TEXT(""), TEXT("A"), TEXT("B"), TEXT("C"), TEXT("")});
+
+	RunParseTokensTest(TEXT("ABCDABCD"), {TEXT("AB")},             {TEXT(""), TEXT("CD"), TEXT("CD")});
+	RunParseTokensTest(TEXT("ABCDABCD"), {TEXT("ABCD")},           {TEXT(""), TEXT(""), TEXT("")});
+	RunParseTokensTest(TEXT("ABCDABCD"), {TEXT("DA")},             {TEXT("ABC"), TEXT("BCD")});
+
+	RunParseTokensTest(TEXT("ABCDABCD"), {TEXT("B"),  TEXT("D")},  {TEXT("A"), TEXT("C"), TEXT("A"), TEXT("C"), TEXT("")});
+	RunParseTokensTest(TEXT("ABCDABCD"), {TEXT("BC"), TEXT("DA")}, {TEXT("A"), TEXT(""), TEXT(""), TEXT("D")});
+
+	return true;
+}
+
+IMPLEMENT_SIMPLE_AUTOMATION_TEST(StringParseTokensByCharTest, "System.Core.String.ParseTokens.ByChar", EAutomationTestFlags::ApplicationContextMask | EAutomationTestFlags::EngineFilter)
+bool StringParseTokensByCharTest::RunTest(const FString& Parameters)
+{
+	auto RunParseTokensTest = [this](const FStringView& View, std::initializer_list<TCHAR> Delimiters, std::initializer_list<FStringView> ExpectedTokens)
+	{
+		TArray<FStringView, TInlineAllocator<8>> ResultTokens;
+		String::ParseTokensMultiple(View, Delimiters, [&ResultTokens](FStringView Token) { ResultTokens.Add(Token); });
+		if (!Algo::CompareByPredicate(ResultTokens, ExpectedTokens, TEqualTo<>()))
+		{
+			TStringBuilder<512> Error;
+			Error << TEXT("String::ParseTokensMultiple failed to parse \"") << View << TEXT("\" with delimiters {");
+			Error.JoinQuoted(Delimiters, TEXT(", "_SV), TEXT("'"_SV));
+			Error << TEXT("} result {");
+			Error.JoinQuoted(ResultTokens, TEXT(", "_SV), TEXT("\""_SV));
+			Error << TEXT("} expected {");
+			Error.JoinQuoted(ExpectedTokens, TEXT(", "_SV), TEXT("\""_SV));
+			Error << TEXT("}");
+			AddError(Error.ToString());
+		}
+	};
+
+	RunParseTokensTest(TEXT(""),         {},                       {TEXT("")});
+	RunParseTokensTest(TEXT("ABC"),      {},                       {TEXT("ABC")});
+
+	RunParseTokensTest(TEXT(""),         {TEXT(',')},              {TEXT("")});
+	RunParseTokensTest(TEXT(","),        {TEXT(',')},              {TEXT(""), TEXT("")});
+	RunParseTokensTest(TEXT(",,"),       {TEXT(',')},              {TEXT(""), TEXT(""), TEXT("")});
+	RunParseTokensTest(TEXT("ABC"),      {TEXT(',')},              {TEXT("ABC")});
+	RunParseTokensTest(TEXT("A,,C"),     {TEXT(',')},              {TEXT("A"), TEXT(""), TEXT("C")});
+	RunParseTokensTest(TEXT("A,B,C"),    {TEXT(',')},              {TEXT("A"), TEXT("B"), TEXT("C")});
+	RunParseTokensTest(TEXT(",A,B,C,"),  {TEXT(',')},              {TEXT(""), TEXT("A"), TEXT("B"), TEXT("C"), TEXT("")});
+
+	RunParseTokensTest(TEXT("ABCDABCD"), {TEXT('B'),  TEXT('D')},  {TEXT("A"), TEXT("C"), TEXT("A"), TEXT("C"), TEXT("")});
+
+	return true;
+}
+
+#endif // WITH_DEV_AUTOMATION_TESTS
diff --git a/Engine/Source/Runtime/Core/Public/Algo/Compare.h b/Engine/Source/Runtime/Core/Public/Algo/Compare.h
--- a/Engine/Source/Runtime/Core/Public/Algo/Compare.h
+++ b/Engine/Source/Runtime/Core/Public/Algo/Compare.h
@@ -8,37 +8,36 @@
 namespace Algo
 {
 	/**
-	* Compares entries in the container using a user-defined predicate container
+	* Compares two contiguous containers using a predicate to compare pairs of elements.
 	*
-	* @param  InputA     Container used as baseline
-	* @param  InputB     Container to compare against
-	* @param  Predicate  Condition which returns true for elements which are deemed equal
+	* @param  InputA     Container of elements that are used as the first argument to the predicate.
+	* @param  InputB     Container of elements that are used as the second argument to the predicate.
+	* @param  Predicate  Condition which returns true for elements which are deemed equal.
+	*
+	* @return Whether the containers are the same size and the predicate returned true for every pair of elements.
 	*/
-	template <typename InT, typename PredicateT>
-	FORCEINLINE bool CompareByPredicate(const InT& InputA, const InT& InputB, PredicateT Predicate)
+	template <typename InAT, typename InBT, typename PredicateT>
+	constexpr bool CompareByPredicate(InAT&& InputA, InBT&& InputB, PredicateT Predicate)
 	{
-		if (InputA.Num() == InputB.Num())
+		const SIZE_T SizeA = GetNum(InputA);
+		const SIZE_T SizeB = GetNum(InputB);
+
+		if (SizeA != SizeB)
 		{
-			uint32 Count = GetNum(InputA);
+			return false;
+		}
 
-			auto* A = GetData(InputA);
-			auto* B = GetData(InputB);
+		auto* A = GetData(InputA);
+		auto* B = GetData(InputB);
 
-			while (Count)
+		for (SIZE_T Count = SizeA; Count; --Count)
+		{
+			if (!Invoke(Predicate, *A++, *B++))
 			{
-				if (!(Invoke(Predicate, *A, *B)))
-				{
-					return false;
-				}
-
-				++A;
-				++B;
-				--Count;
+				return false;
 			}
-
-			return true;
 		}
 
-		return false;
+		return true;
 	}
 }
diff --git a/Engine/Source/Runtime/Core/Public/Algo/Replace.h b/Engine/Source/Runtime/Core/Public/Algo/Replace.h
new file mode 100644
--- /dev/null
+++ b/Engine/Source/Runtime/Core/Public/Algo/Replace.h
@@ -0,0 +1,47 @@
+// Copyright 1998-2019 Epic Games, Inc. All Rights Reserved.
+
+#pragma once
+
+#include "Templates/Invoke.h"
+#include "Templates/UnrealTemplate.h"
+
+namespace Algo
+{
+	/**
+	 * Replaces all elements that compare equal to one value with a new value.
+	 *
+	 * @param Range The range to search and modify in-place.
+	 * @param InOld The value to search for.
+	 * @param InNew The value to copy in as a replacement.
+	 */
+	template <typename RangeType, typename ValueType>
+	FORCEINLINE void Replace(RangeType&& Range, const ValueType& InOld, const ValueType& InNew)
+	{
+		for (auto& Elem : Forward<RangeType>(Range))
+		{
+			if (Elem == InOld)
+			{
+				Elem = InNew;
+			}
+		}
+	}
+
+	/**
+	 * Replaces all elements that satisfy the predicate with the given value.
+	 *
+	 * @param Range The range to search and modify in-place.
+	 * @param InPred The predicate to apply to each element.
+	 * @param InNew The value to copy in as a replacement for each element satisfying the predicate.
+	 */
+	template <typename RangeType, typename ValueType, typename PredicateType>
+	FORCEINLINE void ReplaceIf(RangeType&& Range, PredicateType InPred, const ValueType& InNew)
+	{
+		for (auto& Elem : Forward<RangeType>(Range))
+		{
+			if (Invoke(InPred, Elem))
+			{
+				Elem = InNew;
+			}
+		}
+	}
+}
diff --git a/Engine/Source/Runtime/Core/Public/Containers/ArrayView.h b/Engine/Source/Runtime/Core/Public/Containers/ArrayView.h
--- a/Engine/Source/Runtime/Core/Public/Containers/ArrayView.h
+++ b/Engine/Source/Runtime/Core/Public/Containers/ArrayView.h
@@ -140,6 +140,17 @@ public:
 		check(ArrayNum >= 0);
 	}
 
+	/**
+	 * Construct a view of an initializer list.
+	 *
+	 * The caller is responsible for ensuring that the view does not outlive the initializer list.
+	 */
+	FORCEINLINE TArrayView(std::initializer_list<ElementType> List)
+		: DataPtr(ArrayViewPrivate::GetDataHelper(List))
+		, ArrayNum(GetNum(List))
+	{
+	}
+
 public:
 
 	/**
@@ -556,3 +567,9 @@ auto MakeArrayView(ElementType* Pointer, int32 Size)
 {
 	return TArrayView<ElementType>(Pointer, Size);
 }
+
+template <typename T>
+TArrayView<const T> MakeArrayView(std::initializer_list<T> List)
+{
+	return TArrayView<const T>(List);
+}
diff --git a/Engine/Source/Runtime/Core/Public/String/ParseLines.h b/Engine/Source/Runtime/Core/Public/String/ParseLines.h
new file mode 100644
--- /dev/null
+++ b/Engine/Source/Runtime/Core/Public/String/ParseLines.h
@@ -0,0 +1,21 @@
+// Copyright 1998-2019 Epic Games, Inc. All Rights Reserved.
+
+#pragma once
+
+#include "CoreTypes.h"
+#include "Templates/Function.h"
+
+class FStringView;
+
+namespace String
+{
+	/**
+	 * Visit every line in the input view as terminated by any of CRLF, CR, LF.
+	 *
+	 * Empty lines are visited by this function.
+	 *
+	 * @param View The string view to split into lines.
+	 * @param Visitor A function that is called for each line.
+	 */
+	CORE_API void ParseLines(const FStringView& View, const TFunctionRef<void(FStringView)>& Visitor);
+}
diff --git a/Engine/Source/Runtime/Core/Public/String/ParseTokens.h b/Engine/Source/Runtime/Core/Public/String/ParseTokens.h
new file mode 100644
--- /dev/null
+++ b/Engine/Source/Runtime/Core/Public/String/ParseTokens.h
@@ -0,0 +1,64 @@
+// Copyright 1998-2019 Epic Games, Inc. All Rights Reserved.
+
+#pragma once
+
+#include "CoreTypes.h"
+#include "Containers/ArrayView.h"
+#include "Templates/Function.h"
+
+class FStringView;
+
+namespace String
+{
+	/**
+	 * Visit every token in the input view as separated by any of the delimiters.
+	 *
+	 * Comparisons with delimiters are case-sensitive and empty tokens are visited.
+	 * Behavior is undefined when delimiters overlap each other, such as the delimiters
+	 * ("AB, "BC") and the input view "1ABC2".
+	 *
+	 * @param View The string view to split into tokens.
+	 * @param Delimiters An array of non-overlapping non-empty delimiters to split the input string on.
+	 * @param Visitor A function that is called for each token.
+	 */
+	CORE_API void ParseTokensMultiple(const FStringView& View, TArrayView<const FStringView> Delimiters, TFunctionRef<void(FStringView)> Visitor);
+
+	/**
+	 * Visit every token in the input view as separated by any of the delimiters.
+	 *
+	 * Comparisons with delimiters are case-sensitive and empty tokens are visited.
+	 *
+	 * @param View The string view to split into tokens.
+	 * @param Delimiters An array of delimiter characters to split the input string on.
+	 * @param Visitor A function that is called for each token.
+	 */
+	CORE_API void ParseTokensMultiple(const FStringView& View, TArrayView<const TCHAR> Delimiters, TFunctionRef<void(FStringView)> Visitor);
+
+	/**
+	 * Visit every token in the input view as separated by the delimiter.
+	 *
+	 * Comparisons with the delimiter are case-sensitive and empty tokens are visited.
+	 *
+	 * @param View The string view to split into tokens.
+	 * @param Delimiter The non-empty delimiter to split the input string on.
+	 * @param Visitor A function that is called for each token.
+	 */
+	inline void ParseTokens(const FStringView& View, const FStringView& Delimiter, TFunctionRef<void(FStringView)> Visitor)
+	{
+		return ParseTokensMultiple(View, MakeArrayView(&Delimiter, 1), Visitor);
+	}
+
+	/**
+	 * Visit every token in the input view as separated by the delimiter.
+	 *
+	 * Comparisons with the delimiter are case-sensitive and empty tokens are visited.
+	 *
+	 * @param View The string view to split into tokens.
+	 * @param Delimiter The delimiter character to split the input string on.
+	 * @param Visitor A function that is called for each token.
+	 */
+	inline void ParseTokens(const FStringView& View, TCHAR Delimiter, TFunctionRef<void(FStringView)> Visitor)
+	{
+		return ParseTokensMultiple(View, MakeArrayView(&Delimiter, 1), Visitor);
+	}
+}
