Change String namespace to UEString to avoid name conflicts

#rnx

diff --git a/Engine/Source/Runtime/Core/Private/Containers/String.cpp b/Engine/Source/Runtime/Core/Private/Containers/String.cpp
--- a/Engine/Source/Runtime/Core/Private/Containers/String.cpp
+++ b/Engine/Source/Runtime/Core/Private/Containers/String.cpp
@@ -1482,7 +1482,7 @@ FArchive& operator<<( FArchive& Ar, FString& A )
 
 int32 HexToBytes(const FString& HexString, uint8* OutBytes)
 {
-	return String::HexToBytes(HexString, OutBytes);
+	return UEString::HexToBytes(HexString, OutBytes);
 }
 
 int32 FindMatchingClosingParenthesis(const FString& TargetString, const int32 StartSearch)
diff --git a/Engine/Source/Runtime/Core/Private/String/BytesToHex.cpp b/Engine/Source/Runtime/Core/Private/String/BytesToHex.cpp
--- a/Engine/Source/Runtime/Core/Private/String/BytesToHex.cpp
+++ b/Engine/Source/Runtime/Core/Private/String/BytesToHex.cpp
@@ -4,7 +4,7 @@
 
 #include "Misc/StringBuilder.h"
 
-namespace String
+namespace UEString
 {
 	template <typename CharType>
 	void BytesToHexImpl(TArrayView<const uint8> Bytes, CharType* OutHex)
diff --git a/Engine/Source/Runtime/Core/Private/String/HexToBytes.cpp b/Engine/Source/Runtime/Core/Private/String/HexToBytes.cpp
--- a/Engine/Source/Runtime/Core/Private/String/HexToBytes.cpp
+++ b/Engine/Source/Runtime/Core/Private/String/HexToBytes.cpp
@@ -5,7 +5,7 @@
 #include "Containers/StringView.h"
 #include "Containers/UnrealString.h"
 
-namespace String
+namespace UEString
 {
 	template <typename CharType>
 	int32 HexToBytesImpl(const TStringView<CharType>& Hex, uint8* const OutBytes)
diff --git a/Engine/Source/Runtime/Core/Private/String/ParseLines.cpp b/Engine/Source/Runtime/Core/Private/String/ParseLines.cpp
--- a/Engine/Source/Runtime/Core/Private/String/ParseLines.cpp
+++ b/Engine/Source/Runtime/Core/Private/String/ParseLines.cpp
@@ -4,7 +4,7 @@
 
 #include "Containers/StringView.h"
 
-namespace String
+namespace UEString
 {
 	void ParseLines(const FStringView& View, const TFunctionRef<void(FStringView)>& Visitor)
 	{
diff --git a/Engine/Source/Runtime/Core/Private/String/ParseTokens.cpp b/Engine/Source/Runtime/Core/Private/String/ParseTokens.cpp
--- a/Engine/Source/Runtime/Core/Private/String/ParseTokens.cpp
+++ b/Engine/Source/Runtime/Core/Private/String/ParseTokens.cpp
@@ -6,7 +6,7 @@
 #include "Algo/NoneOf.h"
 #include "Containers/StringView.h"
 
-namespace String
+namespace UEString
 {
 	template <typename DelimiterType>
 	inline static void ParseTokensMultipleImpl(const FStringView& View, TArrayView<const DelimiterType> Delimiters, TFunctionRef<void(FStringView)> Visitor)
diff --git a/Engine/Source/Runtime/Core/Private/Tests/String/ParseLinesTest.cpp b/Engine/Source/Runtime/Core/Private/Tests/String/ParseLinesTest.cpp
--- a/Engine/Source/Runtime/Core/Private/Tests/String/ParseLinesTest.cpp
+++ b/Engine/Source/Runtime/Core/Private/Tests/String/ParseLinesTest.cpp
@@ -16,11 +16,11 @@ bool StringParseLinesTest::RunTest(const FString& Parameters)
 	auto RunParseLinesTest = [this](const FStringView& View, std::initializer_list<FStringView> ExpectedLines)
 	{
 		TArray<FStringView, TInlineAllocator<8>> ResultLines;
-		String::ParseLines(View, [&ResultLines](FStringView Line) { ResultLines.Add(Line); });
+		UEString::ParseLines(View, [&ResultLines](FStringView Line) { ResultLines.Add(Line); });
 		if (!Algo::CompareByPredicate(ResultLines, ExpectedLines, TEqualTo<>()))
 		{
 			TStringBuilder<512> Error;
-			Error << TEXT("String::ParseLines failed to parse \"") << FString(View).ReplaceCharWithEscapedChar() << TEXT("\" result {");
+			Error << TEXT("UEString::ParseLines failed to parse \"") << FString(View).ReplaceCharWithEscapedChar() << TEXT("\" result {");
 			Error.JoinQuoted(ResultLines, TEXT(", "_SV), TEXT("\""_SV));
 			Error << TEXT("} expected {");
 			Error.JoinQuoted(ExpectedLines, TEXT(", "_SV), TEXT("\""_SV));
diff --git a/Engine/Source/Runtime/Core/Private/Tests/String/ParseTokensTest.cpp b/Engine/Source/Runtime/Core/Private/Tests/String/ParseTokensTest.cpp
--- a/Engine/Source/Runtime/Core/Private/Tests/String/ParseTokensTest.cpp
+++ b/Engine/Source/Runtime/Core/Private/Tests/String/ParseTokensTest.cpp
@@ -15,11 +15,11 @@ bool StringParseTokensByStringTest::RunTest(const FString& Parameters)
 	auto RunParseTokensTest = [this](const FStringView& View, std::initializer_list<FStringView> Delimiters, std::initializer_list<FStringView> ExpectedTokens)
 	{
 		TArray<FStringView, TInlineAllocator<8>> ResultTokens;
-		String::ParseTokensMultiple(View, Delimiters, [&ResultTokens](FStringView Token) { ResultTokens.Add(Token); });
+		UEString::ParseTokensMultiple(View, Delimiters, [&ResultTokens](FStringView Token) { ResultTokens.Add(Token); });
 		if (!Algo::CompareByPredicate(ResultTokens, ExpectedTokens, TEqualTo<>()))
 		{
 			TStringBuilder<512> Error;
-			Error << TEXT("String::ParseTokensMultiple failed to parse \"") << View << TEXT("\" with delimiters {");
+			Error << TEXT("UEString::ParseTokensMultiple failed to parse \"") << View << TEXT("\" with delimiters {");
 			Error.JoinQuoted(Delimiters, TEXT(", "_SV), TEXT("\""_SV));
 			Error << TEXT("} result {");
 			Error.JoinQuoted(ResultTokens, TEXT(", "_SV), TEXT("\""_SV));
@@ -57,11 +57,11 @@ bool StringParseTokensByCharTest::RunTest(const FString& Parameters)
 	auto RunParseTokensTest = [this](const FStringView& View, std::initializer_list<TCHAR> Delimiters, std::initializer_list<FStringView> ExpectedTokens)
 	{
 		TArray<FStringView, TInlineAllocator<8>> ResultTokens;
-		String::ParseTokensMultiple(View, Delimiters, [&ResultTokens](FStringView Token) { ResultTokens.Add(Token); });
+		UEString::ParseTokensMultiple(View, Delimiters, [&ResultTokens](FStringView Token) { ResultTokens.Add(Token); });
 		if (!Algo::CompareByPredicate(ResultTokens, ExpectedTokens, TEqualTo<>()))
 		{
 			TStringBuilder<512> Error;
-			Error << TEXT("String::ParseTokensMultiple failed to parse \"") << View << TEXT("\" with delimiters {");
+			Error << TEXT("UEString::ParseTokensMultiple failed to parse \"") << View << TEXT("\" with delimiters {");
 			Error.JoinQuoted(Delimiters, TEXT(", "_SV), TEXT("'"_SV));
 			Error << TEXT("} result {");
 			Error.JoinQuoted(ResultTokens, TEXT(", "_SV), TEXT("\""_SV));
diff --git a/Engine/Source/Runtime/Core/Public/Misc/SecureHash.h b/Engine/Source/Runtime/Core/Public/Misc/SecureHash.h
--- a/Engine/Source/Runtime/Core/Public/Misc/SecureHash.h
+++ b/Engine/Source/Runtime/Core/Public/Misc/SecureHash.h
@@ -214,7 +214,7 @@ public:
 	inline void FromString(const FStringView& Src)
 	{
 		check(Src.Len() == 40);
-		String::HexToBytes(Src, Hash);
+		UEString::HexToBytes(Src, Hash);
 	}
 
 	friend bool operator==(const FSHAHash& X, const FSHAHash& Y)
@@ -238,8 +238,8 @@ public:
 	friend CORE_API void LexFromString(FSHAHash& Hash, const TCHAR*);
 };
 
-inline FStringBuilderBase& operator<<(FStringBuilderBase& Builder, const FSHAHash& Hash) { String::BytesToHex(Hash.Hash, Builder); return Builder; }
-inline FAnsiStringBuilderBase& operator<<(FAnsiStringBuilderBase& Builder, const FSHAHash& Hash) { String::BytesToHex(Hash.Hash, Builder); return Builder; }
+inline FStringBuilderBase& operator<<(FStringBuilderBase& Builder, const FSHAHash& Hash) { UEString::BytesToHex(Hash.Hash, Builder); return Builder; }
+inline FAnsiStringBuilderBase& operator<<(FAnsiStringBuilderBase& Builder, const FSHAHash& Hash) { UEString::BytesToHex(Hash.Hash, Builder); return Builder; }
 
 class CORE_API FSHA1
 {
diff --git a/Engine/Source/Runtime/Core/Public/String/BytesToHex.h b/Engine/Source/Runtime/Core/Public/String/BytesToHex.h
--- a/Engine/Source/Runtime/Core/Public/String/BytesToHex.h
+++ b/Engine/Source/Runtime/Core/Public/String/BytesToHex.h
@@ -8,7 +8,7 @@
 class FStringBuilderBase;
 class FAnsiStringBuilderBase;
 
-namespace String
+namespace UEString
 {
 	/**
 	 * Convert an array of bytes to a non-null-terminated string of hex digits.
diff --git a/Engine/Source/Runtime/Core/Public/String/HexToBytes.h b/Engine/Source/Runtime/Core/Public/String/HexToBytes.h
--- a/Engine/Source/Runtime/Core/Public/String/HexToBytes.h
+++ b/Engine/Source/Runtime/Core/Public/String/HexToBytes.h
@@ -7,7 +7,7 @@
 class FStringView;
 class FAnsiStringView;
 
-namespace String
+namespace UEString
 {
 	/** 
 	 * Convert an array of hex digits into an array of bytes.
diff --git a/Engine/Source/Runtime/Core/Public/String/ParseLines.h b/Engine/Source/Runtime/Core/Public/String/ParseLines.h
--- a/Engine/Source/Runtime/Core/Public/String/ParseLines.h
+++ b/Engine/Source/Runtime/Core/Public/String/ParseLines.h
@@ -7,7 +7,7 @@
 
 class FStringView;
 
-namespace String
+namespace UEString
 {
 	/**
 	 * Visit every line in the input view as terminated by any of CRLF, CR, LF.
diff --git a/Engine/Source/Runtime/Core/Public/String/ParseTokens.h b/Engine/Source/Runtime/Core/Public/String/ParseTokens.h
--- a/Engine/Source/Runtime/Core/Public/String/ParseTokens.h
+++ b/Engine/Source/Runtime/Core/Public/String/ParseTokens.h
@@ -8,7 +8,7 @@
 
 class FStringView;
 
-namespace String
+namespace UEString
 {
 	/**
 	 * Visit every token in the input view as separated by any of the delimiters.
