Change the temporary UEString namespace to UE::String

#rnx

diff --git a/Engine/Source/Runtime/Core/Private/Containers/String.cpp b/Engine/Source/Runtime/Core/Private/Containers/String.cpp
--- a/Engine/Source/Runtime/Core/Private/Containers/String.cpp
+++ b/Engine/Source/Runtime/Core/Private/Containers/String.cpp
@@ -1482,7 +1482,7 @@ FArchive& operator<<( FArchive& Ar, FString& A )
 
 int32 HexToBytes(const FString& HexString, uint8* OutBytes)
 {
-	return UEString::HexToBytes(HexString, OutBytes);
+	return UE::String::HexToBytes(HexString, OutBytes);
 }
 
 int32 FindMatchingClosingParenthesis(const FString& TargetString, const int32 StartSearch)
diff --git a/Engine/Source/Runtime/Core/Private/String/BytesToHex.cpp b/Engine/Source/Runtime/Core/Private/String/BytesToHex.cpp
--- a/Engine/Source/Runtime/Core/Private/String/BytesToHex.cpp
+++ b/Engine/Source/Runtime/Core/Private/String/BytesToHex.cpp
@@ -4,7 +4,9 @@
 
 #include "Misc/StringBuilder.h"
 
-namespace UEString
+namespace UE
+{
+namespace String
 {
 	template <typename CharType>
 	void BytesToHexImpl(TArrayView<const uint8> Bytes, CharType* OutHex)
@@ -41,3 +43,4 @@ namespace UEString
 		BytesToHexImpl(Bytes, GetData(Builder) + Offset);
 	}
 }
+}
diff --git a/Engine/Source/Runtime/Core/Private/String/HexToBytes.cpp b/Engine/Source/Runtime/Core/Private/String/HexToBytes.cpp
--- a/Engine/Source/Runtime/Core/Private/String/HexToBytes.cpp
+++ b/Engine/Source/Runtime/Core/Private/String/HexToBytes.cpp
@@ -5,7 +5,9 @@
 #include "Containers/StringView.h"
 #include "Containers/UnrealString.h"
 
-namespace UEString
+namespace UE
+{
+namespace String
 {
 	template <typename CharType>
 	int32 HexToBytesImpl(const TStringView<CharType>& Hex, uint8* const OutBytes)
@@ -36,3 +38,4 @@ namespace UEString
 		return HexToBytesImpl<ANSICHAR>(Hex, OutBytes);
 	}
 }
+}
diff --git a/Engine/Source/Runtime/Core/Private/String/ParseLines.cpp b/Engine/Source/Runtime/Core/Private/String/ParseLines.cpp
--- a/Engine/Source/Runtime/Core/Private/String/ParseLines.cpp
+++ b/Engine/Source/Runtime/Core/Private/String/ParseLines.cpp
@@ -4,7 +4,9 @@
 
 #include "Containers/StringView.h"
 
-namespace UEString
+namespace UE
+{
+namespace String
 {
 	void ParseLines(const FStringView& View, const TFunctionRef<void(FStringView)>& Visitor)
 	{
@@ -38,3 +40,4 @@ namespace UEString
 		while (ViewIt != ViewEnd);
 	}
 }
+}
diff --git a/Engine/Source/Runtime/Core/Private/String/ParseTokens.cpp b/Engine/Source/Runtime/Core/Private/String/ParseTokens.cpp
--- a/Engine/Source/Runtime/Core/Private/String/ParseTokens.cpp
+++ b/Engine/Source/Runtime/Core/Private/String/ParseTokens.cpp
@@ -6,7 +6,9 @@
 #include "Algo/NoneOf.h"
 #include "Containers/StringView.h"
 
-namespace UEString
+namespace UE
+{
+namespace String
 {
 	template <typename DelimiterType>
 	inline static void ParseTokensMultipleImpl(const FStringView& View, TArrayView<const DelimiterType> Delimiters, TFunctionRef<void(FStringView)> Visitor)
@@ -56,3 +58,4 @@ namespace UEString
 		ParseTokensMultipleImpl(View, Delimiters, Visitor);
 	}
 }
+}
diff --git a/Engine/Source/Runtime/Core/Private/Tests/String/ParseLinesTest.cpp b/Engine/Source/Runtime/Core/Private/Tests/String/ParseLinesTest.cpp
--- a/Engine/Source/Runtime/Core/Private/Tests/String/ParseLinesTest.cpp
+++ b/Engine/Source/Runtime/Core/Private/Tests/String/ParseLinesTest.cpp
@@ -16,11 +16,11 @@ bool StringParseLinesTest::RunTest(const FString& Parameters)
 	auto RunParseLinesTest = [this](const FStringView& View, std::initializer_list<FStringView> ExpectedLines)
 	{
 		TArray<FStringView, TInlineAllocator<8>> ResultLines;
-		UEString::ParseLines(View, [&ResultLines](FStringView Line) { ResultLines.Add(Line); });
+		UE::String::ParseLines(View, [&ResultLines](FStringView Line) { ResultLines.Add(Line); });
 		if (!Algo::CompareByPredicate(ResultLines, ExpectedLines, TEqualTo<>()))
 		{
 			TStringBuilder<512> Error;
-			Error << TEXT("UEString::ParseLines failed to parse \"") << FString(View).ReplaceCharWithEscapedChar() << TEXT("\" result {");
+			Error << TEXT("UE::String::ParseLines failed to parse \"") << FString(View).ReplaceCharWithEscapedChar() << TEXT("\" result {");
 			Error.JoinQuoted(ResultLines, TEXT(", "_SV), TEXT("\""_SV));
 			Error << TEXT("} expected {");
 			Error.JoinQuoted(ExpectedLines, TEXT(", "_SV), TEXT("\""_SV));
diff --git a/Engine/Source/Runtime/Core/Private/Tests/String/ParseTokensTest.cpp b/Engine/Source/Runtime/Core/Private/Tests/String/ParseTokensTest.cpp
--- a/Engine/Source/Runtime/Core/Private/Tests/String/ParseTokensTest.cpp
+++ b/Engine/Source/Runtime/Core/Private/Tests/String/ParseTokensTest.cpp
@@ -15,11 +15,11 @@ bool StringParseTokensByStringTest::RunTest(const FString& Parameters)
 	auto RunParseTokensTest = [this](const FStringView& View, std::initializer_list<FStringView> Delimiters, std::initializer_list<FStringView> ExpectedTokens)
 	{
 		TArray<FStringView, TInlineAllocator<8>> ResultTokens;
-		UEString::ParseTokensMultiple(View, Delimiters, [&ResultTokens](FStringView Token) { ResultTokens.Add(Token); });
+		UE::String::ParseTokensMultiple(View, Delimiters, [&ResultTokens](FStringView Token) { ResultTokens.Add(Token); });
 		if (!Algo::CompareByPredicate(ResultTokens, ExpectedTokens, TEqualTo<>()))
 		{
 			TStringBuilder<512> Error;
-			Error << TEXT("UEString::ParseTokensMultiple failed to parse \"") << View << TEXT("\" with delimiters {");
+			Error << TEXT("UE::String::ParseTokensMultiple failed to parse \"") << View << TEXT("\" with delimiters {");
 			Error.JoinQuoted(Delimiters, TEXT(", "_SV), TEXT("\""_SV));
 			Error << TEXT("} result {");
 			Error.JoinQuoted(ResultTokens, TEXT(", "_SV), TEXT("\""_SV));
@@ -57,11 +57,11 @@ bool StringParseTokensByCharTest::RunTest(const FString& Parameters)
 	auto RunParseTokensTest = [this](const FStringView& View, std::initializer_list<TCHAR> Delimiters, std::initializer_list<FStringView> ExpectedTokens)
 	{
 		TArray<FStringView, TInlineAllocator<8>> ResultTokens;
-		UEString::ParseTokensMultiple(View, Delimiters, [&ResultTokens](FStringView Token) { ResultTokens.Add(Token); });
+		UE::String::ParseTokensMultiple(View, Delimiters, [&ResultTokens](FStringView Token) { ResultTokens.Add(Token); });
 		if (!Algo::CompareByPredicate(ResultTokens, ExpectedTokens, TEqualTo<>()))
 		{
 			TStringBuilder<512> Error;
-			Error << TEXT("UEString::ParseTokensMultiple failed to parse \"") << View << TEXT("\" with delimiters {");
+			Error << TEXT("UE::String::ParseTokensMultiple failed to parse \"") << View << TEXT("\" with delimiters {");
 			Error.JoinQuoted(Delimiters, TEXT(", "_SV), TEXT("'"_SV));
 			Error << TEXT("} result {");
 			Error.JoinQuoted(ResultTokens, TEXT(", "_SV), TEXT("\""_SV));
diff --git a/Engine/Source/Runtime/Core/Public/Misc/SecureHash.h b/Engine/Source/Runtime/Core/Public/Misc/SecureHash.h
--- a/Engine/Source/Runtime/Core/Public/Misc/SecureHash.h
+++ b/Engine/Source/Runtime/Core/Public/Misc/SecureHash.h
@@ -214,7 +214,7 @@ public:
 	inline void FromString(const FStringView& Src)
 	{
 		check(Src.Len() == 40);
-		UEString::HexToBytes(Src, Hash);
+		UE::String::HexToBytes(Src, Hash);
 	}
 
 	friend bool operator==(const FSHAHash& X, const FSHAHash& Y)
@@ -238,8 +238,8 @@ public:
 	friend CORE_API void LexFromString(FSHAHash& Hash, const TCHAR*);
 };
 
-inline FStringBuilderBase& operator<<(FStringBuilderBase& Builder, const FSHAHash& Hash) { UEString::BytesToHex(Hash.Hash, Builder); return Builder; }
-inline FAnsiStringBuilderBase& operator<<(FAnsiStringBuilderBase& Builder, const FSHAHash& Hash) { UEString::BytesToHex(Hash.Hash, Builder); return Builder; }
+inline FStringBuilderBase& operator<<(FStringBuilderBase& Builder, const FSHAHash& Hash) { UE::String::BytesToHex(Hash.Hash, Builder); return Builder; }
+inline FAnsiStringBuilderBase& operator<<(FAnsiStringBuilderBase& Builder, const FSHAHash& Hash) { UE::String::BytesToHex(Hash.Hash, Builder); return Builder; }
 
 class CORE_API FSHA1
 {
diff --git a/Engine/Source/Runtime/Core/Public/String/BytesToHex.h b/Engine/Source/Runtime/Core/Public/String/BytesToHex.h
--- a/Engine/Source/Runtime/Core/Public/String/BytesToHex.h
+++ b/Engine/Source/Runtime/Core/Public/String/BytesToHex.h
@@ -8,7 +8,9 @@
 class FStringBuilderBase;
 class FAnsiStringBuilderBase;
 
-namespace UEString
+namespace UE
+{
+namespace String
 {
 	/**
 	 * Convert an array of bytes to a non-null-terminated string of hex digits.
@@ -28,3 +30,4 @@ namespace UEString
 	CORE_API void BytesToHex(TArrayView<const uint8> Bytes, FStringBuilderBase& Builder);
 	CORE_API void BytesToHex(TArrayView<const uint8> Bytes, FAnsiStringBuilderBase& Builder);
 }
+}
diff --git a/Engine/Source/Runtime/Core/Public/String/HexToBytes.h b/Engine/Source/Runtime/Core/Public/String/HexToBytes.h
--- a/Engine/Source/Runtime/Core/Public/String/HexToBytes.h
+++ b/Engine/Source/Runtime/Core/Public/String/HexToBytes.h
@@ -7,7 +7,9 @@
 class FStringView;
 class FAnsiStringView;
 
-namespace UEString
+namespace UE
+{
+namespace String
 {
 	/** 
 	 * Convert an array of hex digits into an array of bytes.
@@ -20,3 +22,4 @@ namespace UEString
 	CORE_API int32 HexToBytes(const FStringView& Hex, uint8* OutBytes);
 	CORE_API int32 HexToBytes(const FAnsiStringView& Hex, uint8* OutBytes);
 }
+}
diff --git a/Engine/Source/Runtime/Core/Public/String/ParseLines.h b/Engine/Source/Runtime/Core/Public/String/ParseLines.h
--- a/Engine/Source/Runtime/Core/Public/String/ParseLines.h
+++ b/Engine/Source/Runtime/Core/Public/String/ParseLines.h
@@ -7,7 +7,9 @@
 
 class FStringView;
 
-namespace UEString
+namespace UE
+{
+namespace String
 {
 	/**
 	 * Visit every line in the input view as terminated by any of CRLF, CR, LF.
@@ -19,3 +21,4 @@ namespace UEString
 	 */
 	CORE_API void ParseLines(const FStringView& View, const TFunctionRef<void(FStringView)>& Visitor);
 }
+}
diff --git a/Engine/Source/Runtime/Core/Public/String/ParseTokens.h b/Engine/Source/Runtime/Core/Public/String/ParseTokens.h
--- a/Engine/Source/Runtime/Core/Public/String/ParseTokens.h
+++ b/Engine/Source/Runtime/Core/Public/String/ParseTokens.h
@@ -8,7 +8,9 @@
 
 class FStringView;
 
-namespace UEString
+namespace UE
+{
+namespace String
 {
 	/**
 	 * Visit every token in the input view as separated by any of the delimiters.
@@ -62,3 +64,4 @@ namespace UEString
 		return ParseTokensMultiple(View, MakeArrayView(&Delimiter, 1), Visitor);
 	}
 }
+}
