Optimize FPaths::RemoveDuplicateSlashes by removing unnecessary allocations.
#rnx

diff --git a/Engine/Source/Runtime/Core/Private/Misc/Paths.cpp b/Engine/Source/Runtime/Core/Private/Misc/Paths.cpp
--- a/Engine/Source/Runtime/Core/Private/Misc/Paths.cpp
+++ b/Engine/Source/Runtime/Core/Private/Misc/Paths.cpp
@@ -1074,10 +1074,40 @@ bool FPaths::CollapseRelativeDirectories(FString& InPath)
 
 void FPaths::RemoveDuplicateSlashes(FString& InPath)
 {
-	while (InPath.Contains(TEXT("//"), ESearchCase::CaseSensitive))
+	TCHAR* Text = InPath.GetCharArray().GetData();
+	if (!Text)
 	{
-		InPath = InPath.Replace(TEXT("//"), TEXT("/"), ESearchCase::CaseSensitive);
+		return;
+	}
+	const TCHAR* const TwoSlashStr = TEXT("//");
+	const TCHAR SlashChr = TEXT('/');
+
+	TCHAR* const EditStart = TCString<TCHAR>::Strstr(Text, TwoSlashStr);
+	if (!EditStart)
+	{
+		return;
 	}
+	const TCHAR* const TextEnd = Text + InPath.Len();
+
+	// Since we know we've found TwoSlashes, point the initial Write head at the spot where the second slash is (which we shall skip), and point the Read head at the next character after the second slash
+	TCHAR* Write = EditStart + 1;	// The position to write the next character of the output
+	const TCHAR* Read = Write + 1;	// The next character of the input to read
+
+	for (; Read < TextEnd; ++Read)
+	{
+		TCHAR NextChar = *Read;
+		if (Write[-1] != SlashChr || NextChar != SlashChr)
+		{
+			*Write++ = NextChar;
+		}
+		else
+		{
+			// Skip the NextChar when adding on a slash after an existing slash, e.g // or more generally before/////after
+			//                                                                       WR                         W  R
+		}
+	}
+	*Write = TEXT('\0');
+	InPath.TrimToNullTerminator();
 }
 
 void FPaths::MakeStandardFilename(FString& InPath)
diff --git a/Engine/Source/Runtime/Core/Private/Tests/Misc/PathsTest.cpp b/Engine/Source/Runtime/Core/Private/Tests/Misc/PathsTest.cpp
--- a/Engine/Source/Runtime/Core/Private/Tests/Misc/PathsTest.cpp
+++ b/Engine/Source/Runtime/Core/Private/Tests/Misc/PathsTest.cpp
@@ -184,6 +184,36 @@ bool FPathTests::RunTest( const FString& Parameters )
 		RunIsUnderDirectoryTest(TEXT("C:/Folder/Subdir/"),	TEXT("C:/Folder/"), true);
 	}
 
+	// RemoveDuplicateSlashes
+	{
+		auto RunRemoveDuplicateSlashesTest = [this](const TCHAR* InPath, const TCHAR* InExpectedResult)
+		{
+			FString ReplacementPath(InPath);
+			const FString ExpectedResult(InExpectedResult);
+			FPaths::RemoveDuplicateSlashes(ReplacementPath);
+			if (!ReplacementPath.Equals(ExpectedResult, ESearchCase::CaseSensitive))
+			{
+				AddError(FString::Printf(TEXT("FPaths::RemoveDuplicateSlashes('%s') != '%s'."), InPath, InExpectedResult));
+			}
+			else if (ReplacementPath.GetCharArray().Num() != ExpectedResult.Len())
+			{
+				AddError(FString::Printf(TEXT("FPaths::RemoveDuplicateSlashes('%s') returned a result with extra space still allocated after the null terminator."), InPath));
+			}
+		};
+
+		RunRemoveDuplicateSlashesTest(TEXT(""), TEXT(""));
+		RunRemoveDuplicateSlashesTest(TEXT("C:/Folder/File.txt"), TEXT("C:/Folder/File.txt"));
+		RunRemoveDuplicateSlashesTest(TEXT("C:/Folder/File/"), TEXT("C:/Folder/File/"));
+		RunRemoveDuplicateSlashesTest(TEXT("/"), TEXT("/"));
+		RunRemoveDuplicateSlashesTest(TEXT("//"), TEXT("/"));
+		RunRemoveDuplicateSlashesTest(TEXT("////"), TEXT("/"));
+		RunRemoveDuplicateSlashesTest(TEXT("/Folder/File"), TEXT("/Folder/File"));
+		RunRemoveDuplicateSlashesTest(TEXT("//Folder/File"), TEXT("/Folder/File")); // Don't use on //UNC paths; it will be stripped!
+		RunRemoveDuplicateSlashesTest(TEXT("/////Folder//////File/////"), TEXT("/Folder/File/"));
+		RunRemoveDuplicateSlashesTest(TEXT("\\\\Folder\\\\File\\\\"), TEXT("\\\\Folder\\\\File\\\\")); // It doesn't strip backslash, and we rely on that in some places
+		RunRemoveDuplicateSlashesTest(TEXT("//\\\\//Folder//\\\\//File//\\\\//"), TEXT("/\\\\/Folder/\\\\/File/\\\\/"));
+	}
+
 	return true;
 }
 
