Fix incorrect test code.
#rnx

diff --git a/Engine/Source/Runtime/Core/Private/Tests/Misc/PathsTest.cpp b/Engine/Source/Runtime/Core/Private/Tests/Misc/PathsTest.cpp
--- a/Engine/Source/Runtime/Core/Private/Tests/Misc/PathsTest.cpp
+++ b/Engine/Source/Runtime/Core/Private/Tests/Misc/PathsTest.cpp
@@ -195,7 +195,7 @@ bool FPathTests::RunTest( const FString& Parameters )
 			{
 				AddError(FString::Printf(TEXT("FPaths::RemoveDuplicateSlashes('%s') != '%s'."), InPath, InExpectedResult));
 			}
-			else if (ReplacementPath.GetCharArray().Num() != ExpectedResult.Len())
+			else if (ReplacementPath.GetCharArray().Num() != 0 && ReplacementPath.GetCharArray().Num() != ExpectedResult.Len() + 1)
 			{
 				AddError(FString::Printf(TEXT("FPaths::RemoveDuplicateSlashes('%s') returned a result with extra space still allocated after the null terminator."), InPath));
 			}
