Optimize ParseTokens[Multiple] for single-character delimiters

Reduced the cost of ParseTokens[Multiple] by 27% in the case being profiled.

diff --git a/Engine/Source/Runtime/Core/Private/String/ParseTokens.cpp b/Engine/Source/Runtime/Core/Private/String/ParseTokens.cpp
--- a/Engine/Source/Runtime/Core/Private/String/ParseTokens.cpp
+++ b/Engine/Source/Runtime/Core/Private/String/ParseTokens.cpp
@@ -2,40 +2,121 @@
 
 #include "String/ParseTokens.h"
 
+#include "Algo/AllOf.h"
 #include "Algo/Find.h"
 #include "Algo/NoneOf.h"
+#include "Containers/BitArray.h"
 #include "Containers/StringView.h"
 
 namespace UE
 {
 namespace String
 {
-	template <typename DelimiterType>
-	inline static void ParseTokensMultipleImpl(const FStringView& View, TArrayView<const DelimiterType> Delimiters, TFunctionRef<void(FStringView)> Visitor)
+	/** Parse tokens with one single-character delimiter. */
+	inline static void ParseTokens1Delim1Char(const FStringView& View, const TCHAR Delimiter, TFunctionRef<void(FStringView)> Visitor)
 	{
-		// This is a naive implementation that takes time proportional to InView.Len() * TotalDelimiterLen.
-		// If this function becomes a bottleneck, it can be specialized separately for single-character delimiters
-		// and multi-character delimiters, as well as for one and many delimiters. There are algorithms for each of
-		// these four combinations of requirements that are linear or sub-linear in the length of string to search.
+		const TCHAR* ViewIt = View.GetData();
+		const TCHAR* const ViewEnd = ViewIt + View.Len();
+		const TCHAR* NextToken = ViewIt;
+
+		for (;;)
+		{
+			if (ViewIt == ViewEnd)
+			{
+				break;
+			}
+			if (*ViewIt != Delimiter)
+			{
+				++ViewIt;
+				continue;
+			}
+			Visitor(FStringView(NextToken, ViewIt - NextToken));
+			NextToken = ++ViewIt;
+		}
+
+		Visitor(FStringView(NextToken, ViewIt - NextToken));
+	}
+
+	/** Parse tokens with multiple single-character Basic Latin delimiters. */
+	inline static void ParseTokensNDelim1CharBasicLatin(const FStringView& View, TArrayView<const TCHAR> Delimiters, TFunctionRef<void(FStringView)> Visitor)
+	{
+		TBitArray<> DelimiterMask(false, 128);
+		for (uint32 Delimiter : Delimiters)
+		{
+			DelimiterMask[Delimiter] = true;
+		}
+
+		const TCHAR* ViewIt = View.GetData();
+		const TCHAR* const ViewEnd = ViewIt + View.Len();
+		const TCHAR* NextToken = ViewIt;
+
+		for (;;)
+		{
+			if (ViewIt == ViewEnd)
+			{
+				break;
+			}
+			const uint32 CodePoint = *ViewIt;
+			if (CodePoint >= 128 || !DelimiterMask[CodePoint])
+			{
+				++ViewIt;
+				continue;
+			}
+			Visitor(FStringView(NextToken, ViewIt - NextToken));
+			NextToken = ++ViewIt;
+		}
+
+		Visitor(FStringView(NextToken, ViewIt - NextToken));
+	}
+
+	/** Parse tokens with multiple single-character delimiters in the Basic Multilingual Plane. */
+	inline static void ParseTokensNDelim1Char(const FStringView& View, TArrayView<const TCHAR> Delimiters, TFunctionRef<void(FStringView)> Visitor)
+	{
+		if (Algo::AllOf(Delimiters, [](TCHAR Delimiter) { return Delimiter < 128; }))
+		{
+			return ParseTokensNDelim1CharBasicLatin(View, Delimiters, Visitor);
+		}
+
+		const TCHAR* ViewIt = View.GetData();
+		const TCHAR* const ViewEnd = ViewIt + View.Len();
+		const TCHAR* NextToken = ViewIt;
+
+		for (;;)
+		{
+			if (ViewIt == ViewEnd)
+			{
+				break;
+			}
+			if (!Algo::Find(Delimiters, *ViewIt))
+			{
+				++ViewIt;
+				continue;
+			}
+			Visitor(FStringView(NextToken, ViewIt - NextToken));
+			NextToken = ++ViewIt;
+		}
+
+		Visitor(FStringView(NextToken, ViewIt - NextToken));
+	}
+
+	/** Parse tokens with multiple multi-character delimiters. */
+	inline static void ParseTokensNDelimNChar(const FStringView& View, TArrayView<const FStringView> Delimiters, TFunctionRef<void(FStringView)> Visitor)
+	{
+		// This is a naive implementation that takes time proportional to View.Len() * TotalDelimiterLen.
+		// If this function becomes a bottleneck, it can be specialized separately for one and many delimiters.
+		// There are algorithms for each are linear or sub-linear in the length of string to search.
 
 		const FStringView::SizeType ViewLen = View.Len();
 		FStringView::SizeType NextTokenIndex = 0;
 
 		for (FStringView::SizeType ViewIndex = 0; ViewIndex != ViewLen;)
 		{
-			struct
-			{
-				FStringView RemainingView;
-				inline bool operator()(const FStringView& InDelim) const { return RemainingView.StartsWith(InDelim, ESearchCase::CaseSensitive); }
-				inline bool operator()(const TCHAR InDelim) const { return RemainingView.StartsWith(InDelim); }
-				inline FStringView::SizeType Len(const FStringView& InDelim) const { return InDelim.Len(); }
-				inline int32 Len(const TCHAR InDelim) const { return 1; }
-			} MatchDelim{FStringView(View.GetData() + ViewIndex, ViewLen - ViewIndex)};
-
-			if (const DelimiterType* Delim = Algo::FindByPredicate(Delimiters, MatchDelim))
+			const FStringView RemainingView(View.GetData() + ViewIndex, ViewLen - ViewIndex);
+			auto MatchDelimiter = [&RemainingView](const FStringView& Delimiter) { return RemainingView.StartsWith(Delimiter, ESearchCase::CaseSensitive); };
+			if (const FStringView* Delimiter = Algo::FindByPredicate(Delimiters, MatchDelimiter))
 			{
 				Visitor(FStringView(View.GetData() + NextTokenIndex, ViewIndex - NextTokenIndex));
-				ViewIndex += MatchDelim.Len(*Delim);
+				ViewIndex += Delimiter->Len();
 				NextTokenIndex = ViewIndex;
 			}
 			else
@@ -50,12 +131,45 @@ namespace String
 	void ParseTokensMultiple(const FStringView& View, TArrayView<const FStringView> Delimiters, TFunctionRef<void(FStringView)> Visitor)
 	{
 		check(Algo::NoneOf(Delimiters, &FStringView::IsEmpty));
-		ParseTokensMultipleImpl(View, Delimiters, Visitor);
+		switch (Delimiters.Num())
+		{
+		case 0:
+			return Visitor(View);
+		case 1:
+			if (Delimiters[0].Len() == 1)
+			{
+				return ParseTokens1Delim1Char(View, Delimiters[0][0], MoveTemp(Visitor));
+			}
+			return ParseTokensNDelimNChar(View, Delimiters, MoveTemp(Visitor));
+		default:
+			if (Algo::AllOf(Delimiters, [](const FStringView& Delimiter) { return Delimiter.Len() == 1; }))
+			{
+				TArray<TCHAR, TInlineAllocator<32>> DelimiterChars;
+				DelimiterChars.Reserve(Delimiters.Num());
+				for (const FStringView& Delimiter : Delimiters)
+				{
+					DelimiterChars.Add(Delimiter[0]);
+				}
+				return ParseTokensNDelim1Char(View, DelimiterChars, MoveTemp(Visitor));
+			}
+			else
+			{
+				return ParseTokensNDelimNChar(View, Delimiters, MoveTemp(Visitor));
+			}
+		}
 	}
 
 	void ParseTokensMultiple(const FStringView& View, TArrayView<const TCHAR> Delimiters, TFunctionRef<void(FStringView)> Visitor)
 	{
-		ParseTokensMultipleImpl(View, Delimiters, Visitor);
+		switch (Delimiters.Num())
+		{
+		case 0:
+			return Visitor(View);
+		case 1:
+			return ParseTokens1Delim1Char(View, Delimiters[0], MoveTemp(Visitor));
+		default:
+			return ParseTokensNDelim1Char(View, Delimiters, MoveTemp(Visitor));
+		}
 	}
 }
 }
diff --git a/Engine/Source/Runtime/Core/Private/Tests/String/ParseTokensTest.cpp b/Engine/Source/Runtime/Core/Private/Tests/String/ParseTokensTest.cpp
--- a/Engine/Source/Runtime/Core/Private/Tests/String/ParseTokensTest.cpp
+++ b/Engine/Source/Runtime/Core/Private/Tests/String/ParseTokensTest.cpp
@@ -40,6 +40,7 @@ bool StringParseTokensByStringTest::RunTest(const FString& Parameters)
 	RunParseTokensTest(TEXT("A,,C"),     {TEXT(",")},              {TEXT("A"), TEXT(""), TEXT("C")});
 	RunParseTokensTest(TEXT("A,B,C"),    {TEXT(",")},              {TEXT("A"), TEXT("B"), TEXT("C")});
 	RunParseTokensTest(TEXT(",A,B,C,"),  {TEXT(",")},              {TEXT(""), TEXT("A"), TEXT("B"), TEXT("C"), TEXT("")});
+	RunParseTokensTest(TEXT("A\u2022B\u2022C"), {TEXT("\u2022")},  {TEXT("A"), TEXT("B"), TEXT("C")});
 
 	RunParseTokensTest(TEXT("ABCDABCD"), {TEXT("AB")},             {TEXT(""), TEXT("CD"), TEXT("CD")});
 	RunParseTokensTest(TEXT("ABCDABCD"), {TEXT("ABCD")},           {TEXT(""), TEXT(""), TEXT("")});
@@ -48,6 +49,9 @@ bool StringParseTokensByStringTest::RunTest(const FString& Parameters)
 	RunParseTokensTest(TEXT("ABCDABCD"), {TEXT("B"),  TEXT("D")},  {TEXT("A"), TEXT("C"), TEXT("A"), TEXT("C"), TEXT("")});
 	RunParseTokensTest(TEXT("ABCDABCD"), {TEXT("BC"), TEXT("DA")}, {TEXT("A"), TEXT(""), TEXT(""), TEXT("D")});
 
+	RunParseTokensTest(TEXT("A\u2022\u2022B,,C"), {TEXT(",,"), TEXT("\u2022\u2022")},                      {TEXT("A"), TEXT("B"), TEXT("C")});
+	RunParseTokensTest(TEXT("A\u2022\u2022B\u0085\u0085C"), {TEXT("\u0085\u0085"), TEXT("\u2022\u2022")},  {TEXT("A"), TEXT("B"), TEXT("C")});
+
 	return true;
 }
 
@@ -82,8 +86,11 @@ bool StringParseTokensByCharTest::RunTest(const FString& Parameters)
 	RunParseTokensTest(TEXT("A,,C"),     {TEXT(',')},              {TEXT("A"), TEXT(""), TEXT("C")});
 	RunParseTokensTest(TEXT("A,B,C"),    {TEXT(',')},              {TEXT("A"), TEXT("B"), TEXT("C")});
 	RunParseTokensTest(TEXT(",A,B,C,"),  {TEXT(',')},              {TEXT(""), TEXT("A"), TEXT("B"), TEXT("C"), TEXT("")});
+	RunParseTokensTest(TEXT("A\u2022B\u2022C"), {TEXT('\u2022')},  {TEXT("A"), TEXT("B"), TEXT("C")});
 
-	RunParseTokensTest(TEXT("ABCDABCD"), {TEXT('B'),  TEXT('D')},  {TEXT("A"), TEXT("C"), TEXT("A"), TEXT("C"), TEXT("")});
+	RunParseTokensTest(TEXT("ABCDABCD"), {TEXT('B'),  TEXT('D')},                  {TEXT("A"), TEXT("C"), TEXT("A"), TEXT("C"), TEXT("")});
+	RunParseTokensTest(TEXT("A\u2022B,C"), {TEXT(','), TEXT('\u2022')},            {TEXT("A"), TEXT("B"), TEXT("C")});
+	RunParseTokensTest(TEXT("A\u2022B\u0085C"), {TEXT('\u0085'), TEXT('\u2022')},  {TEXT("A"), TEXT("B"), TEXT("C")});
 
 	return true;
 }
