PR #6457: Minor Fix for CString<ANSICHAR>::Strspn and Strcspn (Contributed by [REDACTED_PERSON])

diff --git a/Engine/Source/Runtime/Core/Public/Misc/CString.h b/Engine/Source/Runtime/Core/Public/Misc/CString.h
--- a/Engine/Source/Runtime/Core/Public/Misc/CString.h
+++ b/Engine/Source/Runtime/Core/Public/Misc/CString.h
@@ -766,10 +766,10 @@ typename TCString<T>::CharType* TCString<T>::Strrstr( CharType* String, const Ch
 template <typename T> FORCEINLINE
 int32 TCString<T>::Strspn( const CharType* String, const CharType* Mask )
 {
-	const TCHAR* StringIt = String;
+	const CharType* StringIt = String;
 	while (*StringIt)
 	{
-		for (const TCHAR* MaskIt = Mask; *MaskIt; ++MaskIt)
+		for (const CharType* MaskIt = Mask; *MaskIt; ++MaskIt)
 		{
 			if (*StringIt == *MaskIt)
 			{
@@ -789,10 +789,10 @@ int32 TCString<T>::Strspn( const CharType* String, const CharType* Mask )
 template <typename T> FORCEINLINE
 int32 TCString<T>::Strcspn( const CharType* String, const CharType* Mask )
 {
-	const TCHAR* StringIt = String;
+	const CharType* StringIt = String;
 	while (*StringIt)
 	{
-		for (const TCHAR* MaskIt = Mask; *MaskIt; ++MaskIt)
+		for (const CharType* MaskIt = Mask; *MaskIt; ++MaskIt)
 		{
 			if (*StringIt == *MaskIt)
 			{
