Fix ParseIntoArray to work on TCHAR* Tokens[] = {...}
#rnx

diff --git a/Engine/Source/Runtime/Core/Private/Containers/String.cpp b/Engine/Source/Runtime/Core/Private/Containers/String.cpp
--- a/Engine/Source/Runtime/Core/Private/Containers/String.cpp
+++ b/Engine/Source/Runtime/Core/Private/Containers/String.cpp
@@ -990,7 +990,7 @@ int32 FString::ParseIntoArrayLines(TArray<FString>& OutArray, bool InCullEmpty)
 	return ParseIntoArray(OutArray, LineEndings, NumLineEndings, InCullEmpty);
 }
 
-int32 FString::ParseIntoArray(TArray<FString>& OutArray, const TCHAR** DelimArray, int32 NumDelims, bool InCullEmpty) const
+int32 FString::ParseIntoArray(TArray<FString>& OutArray, const TCHAR* const * DelimArray, int32 NumDelims, bool InCullEmpty) const
 {
 	// Make sure the delimit string is not null or empty
 	check(DelimArray);
diff --git a/Engine/Source/Runtime/Core/Public/Containers/UnrealString.h b/Engine/Source/Runtime/Core/Public/Containers/UnrealString.h
--- a/Engine/Source/Runtime/Core/Public/Containers/UnrealString.h
+++ b/Engine/Source/Runtime/Core/Public/Containers/UnrealString.h
@@ -1709,7 +1709,7 @@ public:
 	*
 	* @return	The number of elements in InArray
 	*/
-	int32 ParseIntoArray(TArray<FString>& OutArray, const TCHAR** DelimArray, int32 NumDelims, bool InCullEmpty = true) const;
+	int32 ParseIntoArray(TArray<FString>& OutArray, const TCHAR*const* DelimArray, int32 NumDelims, bool InCullEmpty = true) const;
 
 	/**
 	 * Takes an array of strings and removes any zero length entries.
