Fix FName construction to avoid stomping the stack for long names

diff --git a/Engine/Source/Runtime/Core/Private/UObject/UnrealNames.cpp b/Engine/Source/Runtime/Core/Private/UObject/UnrealNames.cpp
--- a/Engine/Source/Runtime/Core/Private/UObject/UnrealNames.cpp
+++ b/Engine/Source/Runtime/Core/Private/UObject/UnrealNames.cpp
@@ -1789,7 +1789,7 @@ struct FNameHelper
 		{
 			// Consider _mm_packus_epi16 or similar if this proves too slow
 			ANSICHAR AnsiName[NAME_SIZE];
-			for (int32 I = 0; I < View.Len; ++I)
+			for (int32 I = 0, Len = FMath::Min<int32>(View.Len, NAME_SIZE); I < Len; ++I)
 			{
 				AnsiName[I] = View.Str[I];
 			}
