Check for FName construction from an ASCII string that exceeds NAME_SIZE

#rnx

diff --git a/Engine/Source/Runtime/Core/Private/UObject/UnrealNames.cpp b/Engine/Source/Runtime/Core/Private/UObject/UnrealNames.cpp
--- a/Engine/Source/Runtime/Core/Private/UObject/UnrealNames.cpp
+++ b/Engine/Source/Runtime/Core/Private/UObject/UnrealNames.cpp
@@ -1788,8 +1788,9 @@ struct FNameHelper
 		if (!View.bIsWide)
 		{
 			// Consider _mm_packus_epi16 or similar if this proves too slow
+			checkf(View.Len < NAME_SIZE, TEXT("FName's %d max length exceeded. Got %d characters excluding null-terminator."), NAME_SIZE - 1, View.Len);
 			ANSICHAR AnsiName[NAME_SIZE];
-			for (int32 I = 0, Len = FMath::Min<int32>(View.Len, NAME_SIZE); I < Len; ++I)
+			for (int32 I = 0, Len = View.Len; I < Len; ++I)
 			{
 				AnsiName[I] = View.Str[I];
 			}
