Update copyright for a few new files in Dev-Core

#rnx

diff --git a/Engine/Source/Runtime/Core/Private/String/BytesToHex.cpp b/Engine/Source/Runtime/Core/Private/String/BytesToHex.cpp
--- a/Engine/Source/Runtime/Core/Private/String/BytesToHex.cpp
+++ b/Engine/Source/Runtime/Core/Private/String/BytesToHex.cpp
@@ -1,4 +1,4 @@
-// Copyright 1998-2019 Epic Games, Inc. All Rights Reserved.
+// Copyright Epic Games, Inc. All Rights Reserved.
 
 #include "String/BytesToHex.h"
 
diff --git a/Engine/Source/Runtime/Core/Private/String/HexToBytes.cpp b/Engine/Source/Runtime/Core/Private/String/HexToBytes.cpp
--- a/Engine/Source/Runtime/Core/Private/String/HexToBytes.cpp
+++ b/Engine/Source/Runtime/Core/Private/String/HexToBytes.cpp
@@ -1,4 +1,4 @@
-// Copyright 1998-2019 Epic Games, Inc. All Rights Reserved.
+// Copyright Epic Games, Inc. All Rights Reserved.
 
 #include "String/HexToBytes.h"
 
diff --git a/Engine/Source/Runtime/Core/Private/String/LexFromString.cpp b/Engine/Source/Runtime/Core/Private/String/LexFromString.cpp
--- a/Engine/Source/Runtime/Core/Private/String/LexFromString.cpp
+++ b/Engine/Source/Runtime/Core/Private/String/LexFromString.cpp
@@ -1,4 +1,4 @@
-// Copyright 1998-2019 Epic Games, Inc. All Rights Reserved.
+// Copyright Epic Games, Inc. All Rights Reserved.
 
 #include "String/LexFromString.h"
 
diff --git a/Engine/Source/Runtime/Core/Public/String/BytesToHex.h b/Engine/Source/Runtime/Core/Public/String/BytesToHex.h
--- a/Engine/Source/Runtime/Core/Public/String/BytesToHex.h
+++ b/Engine/Source/Runtime/Core/Public/String/BytesToHex.h
@@ -1,4 +1,4 @@
-// Copyright 1998-2019 Epic Games, Inc. All Rights Reserved.
+// Copyright Epic Games, Inc. All Rights Reserved.
 
 #pragma once
 
diff --git a/Engine/Source/Runtime/Core/Public/String/HexToBytes.h b/Engine/Source/Runtime/Core/Public/String/HexToBytes.h
--- a/Engine/Source/Runtime/Core/Public/String/HexToBytes.h
+++ b/Engine/Source/Runtime/Core/Public/String/HexToBytes.h
@@ -1,4 +1,4 @@
-// Copyright 1998-2019 Epic Games, Inc. All Rights Reserved.
+// Copyright Epic Games, Inc. All Rights Reserved.
 
 #pragma once
 
diff --git a/Engine/Source/Runtime/Core/Public/String/LexFromString.h b/Engine/Source/Runtime/Core/Public/String/LexFromString.h
--- a/Engine/Source/Runtime/Core/Public/String/LexFromString.h
+++ b/Engine/Source/Runtime/Core/Public/String/LexFromString.h
@@ -1,4 +1,4 @@
-// Copyright 1998-2019 Epic Games, Inc. All Rights Reserved.
+// Copyright Epic Games, Inc. All Rights Reserved.
 
 #pragma once
 
