Added FName length bounds-checking on top of existing check() to guard against malicious input in shipping builds

diff --git a/Engine/Source/Runtime/Core/Private/UObject/UnrealNames.cpp b/Engine/Source/Runtime/Core/Private/UObject/UnrealNames.cpp
--- a/Engine/Source/Runtime/Core/Private/UObject/UnrealNames.cpp
+++ b/Engine/Source/Runtime/Core/Private/UObject/UnrealNames.cpp
@@ -1788,9 +1788,8 @@ struct FNameHelper
 		if (!View.bIsWide)
 		{
 			// Consider _mm_packus_epi16 or similar if this proves too slow
-			checkf(View.Len < NAME_SIZE, TEXT("FName's %d max length exceeded. Got %d characters excluding null-terminator."), NAME_SIZE - 1, View.Len);
 			ANSICHAR AnsiName[NAME_SIZE];
-			for (int32 I = 0, Len = View.Len; I < Len; ++I)
+			for (int32 I = 0, Len = FMath::Min<int32>(View.Len, NAME_SIZE); I < Len; ++I)
 			{
 				AnsiName[I] = View.Str[I];
 			}
@@ -1804,7 +1803,11 @@ struct FNameHelper
 
 	static FName Make(FNameStringView View, EFindName FindType, int32 InternalNumber)
 	{
-		checkf(View.Len < NAME_SIZE, TEXT("FName's %d max length exceeded. Got %d characters excluding null-terminator."), NAME_SIZE - 1, View.Len);
+		if (View.Len >= NAME_SIZE)
+		{
+			checkf(false, TEXT("FName's %d max length exceeded. Got %d characters excluding null-terminator."), NAME_SIZE - 1, View.Len);
+			return FName("ERROR_NAME_SIZE_EXCEEDED");
+		}
 
 		FNamePool& Pool = GetNamePool();
 
