Allow TruncToFloat intrinsics used on platforms that support them.
Add SSE and AVX defines to allow specifying availability on PC, MAC, and Linux.
SSE and AVX defines have ALWAYS_HAS and MAYBE_HAS variants to distinguish between available to compile and available to run without the need to check cpuid.
Add templated base classes TUnrealPlatformMathSSEBase and TUnrealPlatformMathSSE4Base to remove the boilerplate overrides added by each of the platformmaths that use the sse overrides.
#rn Minor

diff --git a/Engine/Source/Programs/UnrealBuildTool/Configuration/ModuleRules.cs b/Engine/Source/Programs/UnrealBuildTool/Configuration/ModuleRules.cs
--- a/Engine/Source/Programs/UnrealBuildTool/Configuration/ModuleRules.cs
+++ b/Engine/Source/Programs/UnrealBuildTool/Configuration/ModuleRules.cs
@@ -567,7 +567,8 @@ namespace UnrealBuildTool
 		public bool bUseRTTI = false;
 
 		/// <summary>
-		/// Use AVX instructions
+		/// Direct the compiler to generate AVX instructions wherever SSE or AVX intrinsics are used, on the platforms that support it.
+		/// Note that by enabling this you are changing the minspec for the PC platform, and the resultant executable will crash on machines without AVX support.
 		/// </summary>
 		public bool bUseAVX = false;
 
diff --git a/Engine/Source/Programs/UnrealBuildTool/Platform/Windows/VCToolChain.cs b/Engine/Source/Programs/UnrealBuildTool/Platform/Windows/VCToolChain.cs
--- a/Engine/Source/Programs/UnrealBuildTool/Platform/Windows/VCToolChain.cs
+++ b/Engine/Source/Programs/UnrealBuildTool/Platform/Windows/VCToolChain.cs
@@ -348,8 +348,17 @@ namespace UnrealBuildTool
 			{
 				if (CompileEnvironment.bUseAVX)
 				{
-					// Allow the compiler to generate AVX instructions.
+					// Define /arch:AVX for the current compilation unit.  Machines without AVX support will crash on any SSE/AVX instructions if they run this compilation unit.
 					Arguments.Add("/arch:AVX");
+
+					// AVX available implies sse4 and sse2 available.
+					// Inform Unreal code that we have sse2, sse4, and AVX, both available to compile and available to run
+					// By setting the ALWAYS_HAS defines, we we direct Unreal code to skip cpuid checks to verify that the running hardware supports sse/avx.
+					AddDefinition(Arguments, "PLATFORM_ENABLE_VECTORINTRINSICS=1");
+					AddDefinition(Arguments, "PLATFORM_MAYBE_HAS_SSE4_1=1");
+					AddDefinition(Arguments, "PLATFORM_ALWAYS_HAS_SSE4_1=1");
+					AddDefinition(Arguments, "PLATFORM_MAYBE_HAS_AVX=1");
+					AddDefinition(Arguments, "PLATFORM_ALWAYS_HAS_AVX=1");
 				}
 				// SSE options are not allowed when using the 64 bit toolchain
 				// (enables SSE2 automatically)
diff --git a/Engine/Source/Programs/UnrealBuildTool/System/CppCompileEnvironment.cs b/Engine/Source/Programs/UnrealBuildTool/System/CppCompileEnvironment.cs
--- a/Engine/Source/Programs/UnrealBuildTool/System/CppCompileEnvironment.cs
+++ b/Engine/Source/Programs/UnrealBuildTool/System/CppCompileEnvironment.cs
@@ -139,7 +139,8 @@ namespace UnrealBuildTool
 		public bool bCompileISPC = false;
 
 		/// <summary>
-		/// Use AVX instructions
+		/// Direct the compiler to generate AVX instructions wherever SSE or AVX intrinsics are used.
+		/// Note that by enabling this you are changing the minspec for the PC platform, and the resultant executable will crash on machines without AVX support.
 		/// </summary>
 		public bool bUseAVX = false;
 
diff --git a/Engine/Source/Runtime/Core/Public/HAL/Platform.h b/Engine/Source/Runtime/Core/Public/HAL/Platform.h
--- a/Engine/Source/Runtime/Core/Public/HAL/Platform.h
+++ b/Engine/Source/Runtime/Core/Public/HAL/Platform.h
@@ -162,9 +162,29 @@
 #ifndef PLATFORM_SUPPORTS_PRAGMA_PACK
 	#define PLATFORM_SUPPORTS_PRAGMA_PACK		0
 #endif
+
+// Defines for the availibility of the various levels of vector intrinsics.
+// These may be set from UnrealBuildTool, otherwise each platform-specific platform.h is expected to set them appropriately.
 #ifndef PLATFORM_ENABLE_VECTORINTRINSICS
 	#define PLATFORM_ENABLE_VECTORINTRINSICS	0
 #endif
+// If PLATFORM_MAYBE_HAS_### is 1, then ### intrinsics are compilable.
+// This does not guarantee that the intrinsics are runnable on all instances of the platform however; a runtime check such as cpuid may be required to confirm availability.
+// If PLATFORM_ALWAYS_HAS_### is 1, then ## intrinsics will compile and run on all instances of the platform.  PLATFORM_ALWAYS_HAS_### == 1 implies PLATFORM_MAYBE_HAS_### == 1.
+#ifndef PLATFORM_MAYBE_HAS_SSE4_1
+	#define PLATFORM_MAYBE_HAS_SSE4_1			0
+#endif
+#ifndef PLATFORM_ALWAYS_HAS_SSE4_1
+	#define PLATFORM_ALWAYS_HAS_SSE4_1			0
+#endif
+#ifndef PLATFORM_MAYBE_HAS_AVX
+	#define PLATFORM_MAYBE_HAS_AVX				0
+#endif
+#ifndef PLATFORM_ALWAYS_HAS_AVX
+	#define PLATFORM_ALWAYS_HAS_AVX				0
+#endif
+
+
 #ifndef PLATFORM_HAS_CPUID
 	#if defined(_M_IX86) || defined(__i386__) || defined(_M_X64) || defined(__x86_64__) || defined (__amd64__)
 		#define PLATFORM_HAS_CPUID				1
diff --git a/Engine/Source/Runtime/Core/Public/Mac/MacPlatform.h b/Engine/Source/Runtime/Core/Public/Mac/MacPlatform.h
--- a/Engine/Source/Runtime/Core/Public/Mac/MacPlatform.h
+++ b/Engine/Source/Runtime/Core/Public/Mac/MacPlatform.h
@@ -41,6 +41,15 @@ typedef FMacPlatformTypes FPlatformTypes;
 //#define PLATFORM_EXCEPTIONS_DISABLED				!PLATFORM_DESKTOP
 #define PLATFORM_SUPPORTS_PRAGMA_PACK				1
 #define PLATFORM_ENABLE_VECTORINTRINSICS			1
+#ifndef PLATFORM_MAYBE_HAS_SSE4_1 // May be set from UnrealBuildTool
+	#define PLATFORM_MAYBE_HAS_SSE4_1				1
+#endif
+// Current unreal minspec is sse2, not sse4, so on mac any calling code must check _cpuid before calling SSE4 instructions
+// If called on a platform for which _cpuid for SSE4 returns false, attempting to call SSE4 intrinsics will crash
+// If your title has raised the minspec to sse4, you can define PLATFORM_ALWAYS_HAS_SSE4_1 to 1
+#ifndef PLATFORM_ALWAYS_HAS_SSE4_1 // May be set from UnrealBuildTool
+	#define PLATFORM_ALWAYS_HAS_SSE4_1				0
+#endif
 //#define PLATFORM_USE_LS_SPEC_FOR_WIDECHAR			1
 #define PLATFORM_USE_SYSTEM_VSWPRINTF				0
 #define PLATFORM_COMPILER_DISTINGUISHES_INT_AND_LONG			1
diff --git a/Engine/Source/Runtime/Core/Public/Mac/MacPlatformMath.h b/Engine/Source/Runtime/Core/Public/Mac/MacPlatformMath.h
--- a/Engine/Source/Runtime/Core/Public/Mac/MacPlatformMath.h
+++ b/Engine/Source/Runtime/Core/Public/Mac/MacPlatformMath.h
@@ -7,73 +7,14 @@
 #pragma once
 #include "Clang/ClangPlatformMath.h"
 #include "Mac/MacSystemIncludes.h"
-#include "Math/UnrealPlatformMathSSE.h"
+#include <smmintrin.h>
+#include "Math/UnrealPlatformMathSSE4.h"
 
 /**
 * Mac implementation of the Math OS functions
 **/
-struct FMacPlatformMath : public FClangPlatformMath
+struct FMacPlatformMath : public TUnrealPlatformMathSSE4Base<FClangPlatformMath>
 {
-	static FORCEINLINE int32 TruncToInt(float F)
-	{
-		return UnrealPlatformMathSSE::TruncToInt(F);
-	}
-
-	static FORCEINLINE float TruncToFloat(float F)
-	{
-		return UnrealPlatformMathSSE::TruncToFloat(F);
-	}
-
-	static FORCEINLINE double TruncToDouble(double F)
-	{
-		return UnrealPlatformMathSSE::TruncToDouble(F);
-	}
-
-	static FORCEINLINE int32 RoundToInt(float F)
-	{
-		return UnrealPlatformMathSSE::RoundToInt(F);
-	}
-
-	static FORCEINLINE float RoundToFloat(float F)
-	{
-		return UnrealPlatformMathSSE::RoundToFloat(F);
-	}
-
-	static FORCEINLINE double RoundToDouble(double F)
-	{
-		return UnrealPlatformMathSSE::RoundToDouble(F);
-	}
-
-	static FORCEINLINE int32 FloorToInt(float F)
-	{
-		return UnrealPlatformMathSSE::FloorToInt(F);
-	}
-
-	static FORCEINLINE float FloorToFloat(float F)
-	{
-		return UnrealPlatformMathSSE::FloorToFloat(F);
-	}
-
-	static FORCEINLINE double FloorToDouble(double F)
-	{
-		return UnrealPlatformMathSSE::FloorToDouble(F);
-	}
-
-	static FORCEINLINE int32 CeilToInt(float F)
-	{
-		return UnrealPlatformMathSSE::CeilToInt(F);
-	}
-
-	static FORCEINLINE float CeilToFloat(float F)
-	{
-		return UnrealPlatformMathSSE::CeilToFloat(F);
-	}
-
-	static FORCEINLINE double CeilToDouble(double F)
-	{
-		return UnrealPlatformMathSSE::CeilToDouble(F);
-	}
-
 #if PLATFORM_ENABLE_POPCNT_INTRINSIC
 	/**
 	 * Use the SSE instruction to count bits
@@ -86,18 +27,6 @@ struct FMacPlatformMath : public FClangPlatformMath
 
 	static FORCEINLINE bool IsNaN( float A ) { return isnan(A) != 0; }
 	static FORCEINLINE bool IsFinite( float A ) { return isfinite(A); }
-
-#if PLATFORM_ENABLE_VECTORINTRINSICS
-	static FORCEINLINE float InvSqrt(float F)
-	{
-		return UnrealPlatformMathSSE::InvSqrt(F);
-	}
-
-	static FORCEINLINE float InvSqrtEst( float F )
-	{
-		return UnrealPlatformMathSSE::InvSqrtEst(F);
-	}
-#endif
 };
 
 typedef FMacPlatformMath FPlatformMath;
diff --git a/Engine/Source/Runtime/Core/Public/Math/UnrealPlatformMathSSE.h b/Engine/Source/Runtime/Core/Public/Math/UnrealPlatformMathSSE.h
--- a/Engine/Source/Runtime/Core/Public/Math/UnrealPlatformMathSSE.h
+++ b/Engine/Source/Runtime/Core/Public/Math/UnrealPlatformMathSSE.h
@@ -4,12 +4,20 @@
 
 #include "CoreTypes.h"
 #include "HAL/PlatformMisc.h"
-
 // Code including this header is responsible for including the correct platform-specific header for SSE intrinsics.
 
-#define PLATFORM_ENABLE_SSE4_MATH 0 // Temporarily disable the intrinsics on all platforms to fix the build
+#if !PLATFORM_ENABLE_VECTORINTRINSICS
+
+template<class Base>
+struct TUnrealPlatformMathSSEBase : public Base
+{
+};
+
+#else
 
-namespace UnrealPlatformMathSSE
+namespace UE4
+{
+namespace SSE
 {
 	static FORCEINLINE float InvSqrt(float F)
 	{
@@ -75,49 +83,13 @@ namespace UnrealPlatformMathSSE
 		return _mm_cvtt_ss2si(_mm_set_ss(F));
 	}
 
-	static FORCEINLINE float TruncToFloat(float F)
-	{
-#if PLATFORM_ENABLE_SSE4_MATH
-		return _mm_cvtss_f32(_mm_round_ps(_mm_set_ss(F), 3));
-#else
-		return truncf(F);
-#endif
-	}
-
-	static FORCEINLINE double TruncToDouble(double F)
-	{
-#if PLATFORM_ENABLE_SSE4_MATH
-		return _mm_cvtsd_f64(_mm_round_pd(_mm_set_sd(F), 3));
-#else
-		return trunc(F);
-#endif
-	}
-
 	static FORCEINLINE int32 FloorToInt(float F)
 	{
-		// Note: unlike the Generic solution and the float solution, we implement FloorToInt using a rounding instruction, rather than implementing RoundToInt using a floor instruction.  
+		// Note: unlike the Generic solution and the SSE4 float solution, we implement FloorToInt using a rounding instruction, rather than implementing RoundToInt using a floor instruction.  
 		// We therefore need to do the same times-2 transform (with a slighly different formula) that RoundToInt does; see the note on RoundToInt
 		return _mm_cvt_ss2si(_mm_set_ss(F + F - 0.5f)) >> 1;
 	}
 
-	static FORCEINLINE float FloorToFloat(float F)
-	{
-#if PLATFORM_ENABLE_SSE4_MATH
-		return _mm_cvtss_f32(_mm_floor_ps(_mm_set_ss(F)));
-#else
-		return floorf(F);
-#endif
-	}
-
-	static FORCEINLINE double FloorToDouble(double F)
-	{
-#if PLATFORM_ENABLE_SSE4_MATH
-		return _mm_cvtsd_f64(_mm_floor_pd(_mm_set_sd(F)));
-#else
-		return floor(F);
-#endif
-	}
-
 	static FORCEINLINE int32 RoundToInt(float F)
 	{
 		// Note: the times-2 is to remove the rounding-to-nearest-even-number behavior that mm_cvt_ss2si uses when the fraction is .5
@@ -128,39 +100,48 @@ namespace UnrealPlatformMathSSE
 		return _mm_cvt_ss2si(_mm_set_ss(F + F + 0.5f)) >> 1;
 	}
 
-	static FORCEINLINE float RoundToFloat(float F)
+	static FORCEINLINE int32 CeilToInt(float F)
 	{
-		return FloorToFloat(F + 0.5f);
+		// Note: unlike the Generic solution and the SSE4 float solution, we implement CeilToInt using a rounding instruction, rather than a dedicated ceil instruction
+		// We therefore need to do the same times-2 transform (with a slighly different formula) that RoundToInt does; see the note on RoundToInt
+		return -(_mm_cvt_ss2si(_mm_set_ss(-0.5f - (F + F))) >> 1);
 	}
+}
+}
 
-	static FORCEINLINE double RoundToDouble(double F)
+template<class Base>
+struct TUnrealPlatformMathSSEBase : public Base
+{
+	static FORCEINLINE int32 TruncToInt(float F)
 	{
-		return FloorToDouble(F + 0.5);
+		return UE4::SSE::TruncToInt(F);
 	}
 
-	static FORCEINLINE int32 CeilToInt(float F)
+	static FORCEINLINE int32 RoundToInt(float F)
 	{
-		// Note: unlike the Generic solution and the float solution, we implement CeilToInt using a rounding instruction, rather than a dedicated ceil instruction
-		// We therefore need to do the same times-2 transform (with a slighly different formula) that RoundToInt does; see the note on RoundToInt
-		return -(_mm_cvt_ss2si(_mm_set_ss(-0.5f - (F + F))) >> 1);
+		return UE4::SSE::RoundToInt(F);
 	}
 
-	static FORCEINLINE float CeilToFloat(float F)
+	static FORCEINLINE int32 FloorToInt(float F)
 	{
-#if PLATFORM_ENABLE_SSE4_MATH
-		return _mm_cvtss_f32(_mm_ceil_ps(_mm_set_ss(F)));
-#else
-		return ceilf(F);
-#endif
+		return UE4::SSE::FloorToInt(F);
+	}
 
+	static FORCEINLINE int32 CeilToInt(float F)
+	{
+		return UE4::SSE::CeilToInt(F);
 	}
 
-	static FORCEINLINE double CeilToDouble(double F)
+	static FORCEINLINE float InvSqrt(float F)
 	{
-#if PLATFORM_ENABLE_SSE4_MATH
-		return _mm_cvtsd_f64(_mm_ceil_pd(_mm_set_sd(F)));
-#else
-		return ceil(F);
-#endif
+		return UE4::SSE::InvSqrt(F);
 	}
-}
+
+	static FORCEINLINE float InvSqrtEst(float F)
+	{
+		return UE4::SSE::InvSqrtEst(F);
+	}
+
+};
+
+#endif // PLATFORM_ENABLE_VECTORINTRINSICS
\ No newline at end of file
diff --git a/Engine/Source/Runtime/Core/Public/Math/UnrealPlatformMathSSE4.h b/Engine/Source/Runtime/Core/Public/Math/UnrealPlatformMathSSE4.h
new file mode 100644
--- /dev/null
+++ b/Engine/Source/Runtime/Core/Public/Math/UnrealPlatformMathSSE4.h
@@ -0,0 +1,109 @@
+// Copyright Epic Games, Inc. All Rights Reserved.
+
+#pragma once
+
+#include "HAL/Platform.h"
+#include "Math/UnrealPlatformMathSSE.h"
+// Code including this header is responsible for including the correct platform-specific header for SSE intrinsics.
+
+#if PLATFORM_MAYBE_HAS_SSE4_1
+
+// SSE4_1 instrinsics are only available for compile if PLATFORM_MAYBE_HAS_SSE4_1 is 1, and even if compiled, should not
+// be called at runtime unless PLATFORM_ALWAYS_HAS_SSE4_1 is 1 or cpuid has been checked to verify that the current hardware
+// instance supports it.
+namespace UE4
+{
+namespace SSE4
+{
+	static FORCEINLINE float TruncToFloat(float F)
+	{
+		return _mm_cvtss_f32(_mm_round_ps(_mm_set_ss(F), 3));
+	}
+
+	static FORCEINLINE double TruncToDouble(double F)
+	{
+		return _mm_cvtsd_f64(_mm_round_pd(_mm_set_sd(F), 3));
+	}
+
+	static FORCEINLINE float FloorToFloat(float F)
+	{
+		return _mm_cvtss_f32(_mm_floor_ps(_mm_set_ss(F)));
+	}
+
+	static FORCEINLINE double FloorToDouble(double F)
+	{
+		return _mm_cvtsd_f64(_mm_floor_pd(_mm_set_sd(F)));
+	}
+
+	static FORCEINLINE float RoundToFloat(float F)
+	{
+		return FloorToFloat(F + 0.5f);
+	}
+
+	static FORCEINLINE double RoundToDouble(double F)
+	{
+		return FloorToDouble(F + 0.5);
+	}
+
+	static FORCEINLINE float CeilToFloat(float F)
+	{
+		return _mm_cvtss_f32(_mm_ceil_ps(_mm_set_ss(F)));
+	}
+
+	static FORCEINLINE double CeilToDouble(double F)
+	{
+		return _mm_cvtsd_f64(_mm_ceil_pd(_mm_set_sd(F)));
+	}
+}
+}
+
+#endif // PLATFORM_MAYBE_HAS_SSE4_1
+
+// We currently don't have any runtime checks in FPlatforMath for whether the SSE4_1 intrinsics are available,
+// so we have to not call them from FPlatformMath classes unless PLATFORM_ALWAYS_HAS_SSE4_1
+#define UNREALPLATFORMMATH_SSE4_1_ENABLED PLATFORM_ALWAYS_HAS_SSE4_1
+
+template<class Base>
+struct TUnrealPlatformMathSSE4Base : public TUnrealPlatformMathSSEBase<Base>
+{
+#if UNREALPLATFORMMATH_SSE4_1_ENABLED
+	static FORCEINLINE float TruncToFloat(float F)
+	{
+		return UE4::SSE4::TruncToFloat(F);
+	}
+
+	static FORCEINLINE double TruncToDouble(double F)
+	{
+		return UE4::SSE4::TruncToDouble(F);
+	}
+	static FORCEINLINE float RoundToFloat(float F)
+	{
+		return UE4::SSE4::RoundToFloat(F);
+	}
+
+	static FORCEINLINE double RoundToDouble(double F)
+	{
+		return UE4::SSE4::RoundToDouble(F);
+	}
+
+	static FORCEINLINE float FloorToFloat(float F)
+	{
+		return UE4::SSE4::FloorToFloat(F);
+	}
+
+	static FORCEINLINE double FloorToDouble(double F)
+	{
+		return UE4::SSE4::FloorToDouble(F);
+	}
+
+	static FORCEINLINE float CeilToFloat(float F)
+	{
+		return UE4::SSE4::CeilToFloat(F);
+	}
+
+	static FORCEINLINE double CeilToDouble(double F)
+	{
+		return UE4::SSE4::CeilToDouble(F);
+	}
+#endif // UNREALPLATFORMMATH_SSE4_ENABLED
+};
diff --git a/Engine/Source/Runtime/Core/Public/Unix/UnixPlatform.h b/Engine/Source/Runtime/Core/Public/Unix/UnixPlatform.h
--- a/Engine/Source/Runtime/Core/Public/Unix/UnixPlatform.h
+++ b/Engine/Source/Runtime/Core/Public/Unix/UnixPlatform.h
@@ -93,6 +93,14 @@ typedef FUnixPlatformTypes FPlatformTypes;
 	#define PLATFORM_ENABLE_VECTORINTRINSICS		0
 #endif // defined(_M_IX86) || defined(__i386__) || defined(_M_X64) || defined(__x86_64__) || defined (__amd64__) 
 
+// We do not currently compile with -msse4 or higher on Unix or Linux
+#ifndef PLATFORM_MAYBE_HAS_SSE4_1 // May be set from UnrealBuildTool
+	#define PLATFORM_MAYBE_HAS_SSE4_1							0
+#endif
+#ifndef PLATFORM_ALWAYS_HAS_SSE4_1 // May be set from UnrealBuildTool
+	#define PLATFORM_ALWAYS_HAS_SSE4_1							0
+#endif
+
 // Function type macros.
 #define VARARGS															/* Functions with variable arguments */
 #define CDECL															/* Standard C function */
diff --git a/Engine/Source/Runtime/Core/Public/Unix/UnixPlatformMath.h b/Engine/Source/Runtime/Core/Public/Unix/UnixPlatformMath.h
--- a/Engine/Source/Runtime/Core/Public/Unix/UnixPlatformMath.h
+++ b/Engine/Source/Runtime/Core/Public/Unix/UnixPlatformMath.h
@@ -9,97 +9,17 @@
 #include "CoreTypes.h"
 #include "Clang/ClangPlatformMath.h"
 #include "Linux/LinuxSystemIncludes.h"
-#if PLATFORM_ENABLE_VECTORINTRINSICS
-	#include "Math/UnrealPlatformMathSSE.h"
-#endif // PLATFORM_ENABLE_VECTORINTRINSICS
+#include "Math/UnrealPlatformMathSSE4.h"
 
 /**
 * Linux implementation of the Math OS functions
 **/
-struct FLinuxPlatformMath : public FClangPlatformMath
+struct FLinuxPlatformMath : public TUnrealPlatformMathSSE4Base<FClangPlatformMath>
 {
 #if PLATFORM_ENABLE_VECTORINTRINSICS
-	static FORCEINLINE int32 TruncToInt(float F)
-	{
-		return UnrealPlatformMathSSE::TruncToInt(F);
-	}
-
-#if !PLATFORM_LINUX // Linux Clang does not yet support _mm_round_ps/_mm_round_pd
-	static FORCEINLINE float TruncToFloat(float F)
-	{
-		return UnrealPlatformMathSSE::TruncToFloat(F);
-	}
-
-	static FORCEINLINE double TruncToDouble(double F)
-	{
-		return UnrealPlatformMathSSE::TruncToDouble(F);
-	}
-#endif
-
-	static FORCEINLINE int32 RoundToInt(float F)
-	{
-		return UnrealPlatformMathSSE::RoundToInt(F);
-	}
-
-#if !PLATFORM_LINUX // Linux Clang does not yet support _mm_round_ps/_mm_round_pd
-	static FORCEINLINE float RoundToFloat(float F)
-	{
-		return UnrealPlatformMathSSE::RoundToFloat(F);
-	}
-
-	static FORCEINLINE double RoundToDouble(double F)
-	{
-		return UnrealPlatformMathSSE::RoundToDouble(F);
-	}
-#endif
-
-	static FORCEINLINE int32 FloorToInt(float F)
-	{
-		return UnrealPlatformMathSSE::FloorToInt(F);
-	}
-
-#if !PLATFORM_LINUX // Linux Clang does not yet support _mm_round_ps/_mm_round_pd
-	static FORCEINLINE float FloorToFloat(float F)
-	{
-		return UnrealPlatformMathSSE::FloorToFloat(F);
-	}
-
-	static FORCEINLINE double FloorToDouble(double F)
-	{
-		return UnrealPlatformMathSSE::FloorToDouble(F);
-	}
-#endif
-
-	static FORCEINLINE int32 CeilToInt(float F)
-	{
-		return UnrealPlatformMathSSE::CeilToInt(F);
-	}
-
-#if !PLATFORM_LINUX // Linux Clang does not yet support _mm_round_ps/_mm_round_pd
-	static FORCEINLINE float CeilToFloat(float F)
-	{
-		return UnrealPlatformMathSSE::CeilToFloat(F);
-	}
-
-	static FORCEINLINE double CeilToDouble(double F)
-	{
-		return UnrealPlatformMathSSE::CeilToDouble(F);
-	}
-#endif
-
 	static FORCEINLINE bool IsNaN( float A ) { return isnan(A) != 0; }
 	static FORCEINLINE bool IsFinite( float A ) { return isfinite(A); }
 
-	static FORCEINLINE float InvSqrt( float F )
-	{
-		return UnrealPlatformMathSSE::InvSqrt(F);
-	}
-
-	static FORCEINLINE float InvSqrtEst( float F )
-	{
-		return UnrealPlatformMathSSE::InvSqrtEst(F);
-	}
-
 #if PLATFORM_ENABLE_POPCNT_INTRINSIC
 	/**
 	 * Use the SSE instruction to count bits
diff --git a/Engine/Source/Runtime/Core/Public/Windows/WIndowsPlatform.h b/Engine/Source/Runtime/Core/Public/Windows/WIndowsPlatform.h
--- a/Engine/Source/Runtime/Core/Public/Windows/WIndowsPlatform.h
+++ b/Engine/Source/Runtime/Core/Public/Windows/WIndowsPlatform.h
@@ -50,6 +50,16 @@ typedef FWindowsPlatformTypes FPlatformTypes;
 
 #define PLATFORM_SUPPORTS_PRAGMA_PACK						1
 #define PLATFORM_ENABLE_VECTORINTRINSICS					1
+#ifndef PLATFORM_MAYBE_HAS_SSE4_1 // May be set from UnrealBuildTool
+	#define PLATFORM_MAYBE_HAS_SSE4_1						1
+#endif
+// Current unreal minspec is sse2, not sse4, so on windows any calling code must check _cpuid before calling SSE4 instructions;
+// If called on a platform for which _cpuid for SSE4 returns false, attempting to call SSE4 intrinsics will crash
+// If your title has raised the minspec to sse4, you can define PLATFORM_ALWAYS_HAS_SSE4_1 to 1
+#ifndef PLATFORM_ALWAYS_HAS_SSE4_1 // May be set from UnrealBuildTool
+	#define PLATFORM_ALWAYS_HAS_SSE4_1						0
+#endif
+
 #define PLATFORM_USE_LS_SPEC_FOR_WIDECHAR					0
 //#define PLATFORM_USE_SYSTEM_VSWPRINTF						1
 //#define PLATFORM_TCHAR_IS_4_BYTES							0
diff --git a/Engine/Source/Runtime/Core/Public/Windows/WindowsPlatformMath.h b/Engine/Source/Runtime/Core/Public/Windows/WindowsPlatformMath.h
--- a/Engine/Source/Runtime/Core/Public/Windows/WindowsPlatformMath.h
+++ b/Engine/Source/Runtime/Core/Public/Windows/WindowsPlatformMath.h
@@ -12,90 +12,17 @@
 #include "XboxOne/XboxOneSystemIncludes.h"
 #endif
 
-#if PLATFORM_ENABLE_VECTORINTRINSICS
-#include "Math/UnrealPlatformMathSSE.h"
-#endif
+#include "Math/UnrealPlatformMathSSE4.h"
 
 /**
 * Windows implementation of the Math OS functions
 **/
-struct FWindowsPlatformMath : public FGenericPlatformMath
+struct FWindowsPlatformMath : public TUnrealPlatformMathSSE4Base<FGenericPlatformMath>
 {
 #if PLATFORM_ENABLE_VECTORINTRINSICS
-	static FORCEINLINE int32 TruncToInt(float F)
-	{
-		return UnrealPlatformMathSSE::TruncToInt(F);
-	}
-
-	static FORCEINLINE float TruncToFloat(float F)
-	{
-		return UnrealPlatformMathSSE::TruncToFloat(F);
-	}
-
-	static FORCEINLINE double TruncToDouble(double F)
-	{
-		return UnrealPlatformMathSSE::TruncToDouble(F);
-	}
-
-	static FORCEINLINE int32 RoundToInt( float F )
-	{
-		return UnrealPlatformMathSSE::RoundToInt(F);
-	}
-
-	static FORCEINLINE float RoundToFloat(float F)
-	{
-		return UnrealPlatformMathSSE::RoundToFloat(F);
-	}
-
-	static FORCEINLINE double RoundToDouble(double F)
-	{
-		return UnrealPlatformMathSSE::RoundToDouble(F);
-	}
-
-	static FORCEINLINE int32 FloorToInt(float F)
-	{
-		return UnrealPlatformMathSSE::FloorToInt(F);
-	}
-
-	static FORCEINLINE float FloorToFloat(float F)
-	{
-		return UnrealPlatformMathSSE::FloorToFloat(F);
-	}
-
-	static FORCEINLINE double FloorToDouble(double F)
-	{
-		return UnrealPlatformMathSSE::FloorToDouble(F);
-	}
-
-	static FORCEINLINE int32 CeilToInt(float F)
-	{
-		return UnrealPlatformMathSSE::CeilToInt(F);
-	}
-
-	static FORCEINLINE float CeilToFloat(float F)
-	{
-		return UnrealPlatformMathSSE::CeilToFloat(F);
-	}
-
-	static FORCEINLINE double CeilToDouble(double F)
-	{
-		return UnrealPlatformMathSSE::CeilToDouble(F);
-	}
-
 	static FORCEINLINE bool IsNaN( float A ) { return _isnan(A) != 0; }
 	static FORCEINLINE bool IsFinite( float A ) { return _finite(A) != 0; }
 
-	static FORCEINLINE float InvSqrt(float F)
-	{
-		return UnrealPlatformMathSSE::InvSqrt(F);
-	}
-
-	static FORCEINLINE float InvSqrtEst( float F )
-	{
-		return UnrealPlatformMathSSE::InvSqrtEst(F);
-	}
-		
-
 	#pragma intrinsic( _BitScanReverse )
 	static FORCEINLINE uint32 FloorLog2(uint32 Value) 
 	{
