if constexpr enabled on more compiler versions.

diff --git a/Engine/Source/Runtime/Core/Public/Clang/ClangPlatform.h b/Engine/Source/Runtime/Core/Public/Clang/ClangPlatform.h
--- a/Engine/Source/Runtime/Core/Public/Clang/ClangPlatform.h
+++ b/Engine/Source/Runtime/Core/Public/Clang/ClangPlatform.h
@@ -6,7 +6,7 @@
 
 #pragma once
 
-#if __has_feature(cxx_if_constexpr)
+#if defined(__cpp_if_constexpr)
 	#define PLATFORM_COMPILER_HAS_IF_CONSTEXPR 1
 #else
 	#define PLATFORM_COMPILER_HAS_IF_CONSTEXPR 0
diff --git a/Engine/Source/Runtime/Core/Public/MSVC/MSVCPlatform.h b/Engine/Source/Runtime/Core/Public/MSVC/MSVCPlatform.h
--- a/Engine/Source/Runtime/Core/Public/MSVC/MSVCPlatform.h
+++ b/Engine/Source/Runtime/Core/Public/MSVC/MSVCPlatform.h
@@ -6,7 +6,7 @@
 
 #pragma once
 
-#if _MSC_FULL_VER >= 191125507 && defined(__cpp_if_constexpr)
+#if _MSC_VER >= 1920 || defined(__cpp_if_constexpr)
 	#define PLATFORM_COMPILER_HAS_IF_CONSTEXPR 1
 #else
 	#define PLATFORM_COMPILER_HAS_IF_CONSTEXPR 0
