Only support if constexpr in MSVC >= 1920 after recommendation from Microsoft.

diff --git a/Engine/Source/Runtime/Core/Public/MSVC/MSVCPlatform.h b/Engine/Source/Runtime/Core/Public/MSVC/MSVCPlatform.h
--- a/Engine/Source/Runtime/Core/Public/MSVC/MSVCPlatform.h
+++ b/Engine/Source/Runtime/Core/Public/MSVC/MSVCPlatform.h
@@ -6,7 +6,7 @@
 
 #pragma once
 
-#if _MSC_VER >= 1920 || defined(__cpp_if_constexpr)
+#if _MSC_VER >= 1920
 	#define PLATFORM_COMPILER_HAS_IF_CONSTEXPR 1
 #else
 	#define PLATFORM_COMPILER_HAS_IF_CONSTEXPR 0
