Update copyright for a few new files in Dev-Core

#rnx

diff --git a/Engine/Source/Runtime/Core/Private/String/ParseLines.cpp b/Engine/Source/Runtime/Core/Private/String/ParseLines.cpp
--- a/Engine/Source/Runtime/Core/Private/String/ParseLines.cpp
+++ b/Engine/Source/Runtime/Core/Private/String/ParseLines.cpp
@@ -1,4 +1,4 @@
-// Copyright 1998-2019 Epic Games, Inc. All Rights Reserved.
+// Copyright Epic Games, Inc. All Rights Reserved.
 
 #include "String/ParseLines.h"
 
diff --git a/Engine/Source/Runtime/Core/Private/String/ParseTokens.cpp b/Engine/Source/Runtime/Core/Private/String/ParseTokens.cpp
--- a/Engine/Source/Runtime/Core/Private/String/ParseTokens.cpp
+++ b/Engine/Source/Runtime/Core/Private/String/ParseTokens.cpp
@@ -1,4 +1,4 @@
-// Copyright 1998-2019 Epic Games, Inc. All Rights Reserved.
+// Copyright Epic Games, Inc. All Rights Reserved.
 
 #include "String/ParseTokens.h"
 
diff --git a/Engine/Source/Runtime/Core/Private/Tests/String/ParseLinesTest.cpp b/Engine/Source/Runtime/Core/Private/Tests/String/ParseLinesTest.cpp
--- a/Engine/Source/Runtime/Core/Private/Tests/String/ParseLinesTest.cpp
+++ b/Engine/Source/Runtime/Core/Private/Tests/String/ParseLinesTest.cpp
@@ -1,4 +1,4 @@
-// Copyright 1998-2019 Epic Games, Inc. All Rights Reserved.
+// Copyright Epic Games, Inc. All Rights Reserved.
 
 #include "String/ParseLines.h"
 
diff --git a/Engine/Source/Runtime/Core/Private/Tests/String/ParseTokensTest.cpp b/Engine/Source/Runtime/Core/Private/Tests/String/ParseTokensTest.cpp
--- a/Engine/Source/Runtime/Core/Private/Tests/String/ParseTokensTest.cpp
+++ b/Engine/Source/Runtime/Core/Private/Tests/String/ParseTokensTest.cpp
@@ -1,4 +1,4 @@
-// Copyright 1998-2019 Epic Games, Inc. All Rights Reserved.
+// Copyright Epic Games, Inc. All Rights Reserved.
 
 #include "String/ParseTokens.h"
 
diff --git a/Engine/Source/Runtime/Core/Public/Algo/Replace.h b/Engine/Source/Runtime/Core/Public/Algo/Replace.h
--- a/Engine/Source/Runtime/Core/Public/Algo/Replace.h
+++ b/Engine/Source/Runtime/Core/Public/Algo/Replace.h
@@ -1,4 +1,4 @@
-// Copyright 1998-2019 Epic Games, Inc. All Rights Reserved.
+// Copyright Epic Games, Inc. All Rights Reserved.
 
 #pragma once
 
diff --git a/Engine/Source/Runtime/Core/Public/String/ParseLines.h b/Engine/Source/Runtime/Core/Public/String/ParseLines.h
--- a/Engine/Source/Runtime/Core/Public/String/ParseLines.h
+++ b/Engine/Source/Runtime/Core/Public/String/ParseLines.h
@@ -1,4 +1,4 @@
-// Copyright 1998-2019 Epic Games, Inc. All Rights Reserved.
+// Copyright Epic Games, Inc. All Rights Reserved.
 
 #pragma once
 
diff --git a/Engine/Source/Runtime/Core/Public/String/ParseTokens.h b/Engine/Source/Runtime/Core/Public/String/ParseTokens.h
--- a/Engine/Source/Runtime/Core/Public/String/ParseTokens.h
+++ b/Engine/Source/Runtime/Core/Public/String/ParseTokens.h
@@ -1,4 +1,4 @@
-// Copyright 1998-2019 Epic Games, Inc. All Rights Reserved.
+// Copyright Epic Games, Inc. All Rights Reserved.
 
 #pragma once
 
