Fix missing include for sse4.1 intrinsics for clang compiler in windows platform.
#rnx

diff --git a/Engine/Source/Runtime/Core/Public/Windows/WindowsPlatformMath.h b/Engine/Source/Runtime/Core/Public/Windows/WindowsPlatformMath.h
--- a/Engine/Source/Runtime/Core/Public/Windows/WindowsPlatformMath.h
+++ b/Engine/Source/Runtime/Core/Public/Windows/WindowsPlatformMath.h
@@ -8,6 +8,7 @@
 // This implementation is used by both Windows and XBoxOne
 #if PLATFORM_WINDOWS
 #include <intrin.h>
+#include <smmintrin.h>
 #elif PLATFORM_XBOXONE
 #include "XboxOne/XboxOneSystemIncludes.h"
 #endif
