Fix for StringView copying wide strings.

diff --git a/Engine/Source/Runtime/Core/Public/Containers/StringView.inl b/Engine/Source/Runtime/Core/Public/Containers/StringView.inl
--- a/Engine/Source/Runtime/Core/Public/Containers/StringView.inl
+++ b/Engine/Source/Runtime/Core/Public/Containers/StringView.inl
@@ -16,7 +16,7 @@ template <typename CharType>
 inline typename TStringViewImpl<CharType>::SizeType TStringViewImpl<CharType>::CopyString(CharType* Dest, SizeType CharCount, SizeType Position) const
 {
 	const  SizeType CopyCount = FMath::Min(Size - Position, CharCount);
-	FMemory::Memcpy(Dest, DataPtr + Position, CopyCount);
+	FMemory::Memcpy(Dest, DataPtr + Position, CopyCount * sizeof(CharType));
 	return CopyCount;
 }
 
