Fix for PURE_VIRTUAL containing code with commas in it.

diff --git a/Engine/Source/Runtime/Core/Public/Misc/CoreMiscDefines.h b/Engine/Source/Runtime/Core/Public/Misc/CoreMiscDefines.h
--- a/Engine/Source/Runtime/Core/Public/Misc/CoreMiscDefines.h
+++ b/Engine/Source/Runtime/Core/Public/Misc/CoreMiscDefines.h
@@ -56,9 +56,9 @@
 #define RETURN_VAL_IF_EXIT_REQUESTED(x)
 
 #if CHECK_PUREVIRTUALS
-#define PURE_VIRTUAL(func,extra) =0;
+#define PURE_VIRTUAL(func,...) =0;
 #else
-#define PURE_VIRTUAL(func,extra) { LowLevelFatalError(TEXT("Pure virtual not implemented (%s)"), TEXT(#func)); extra }
+#define PURE_VIRTUAL(func,...) { LowLevelFatalError(TEXT("Pure virtual not implemented (%s)"), TEXT(#func)); __VA_ARGS__ }
 #endif
 
 
