Fix for TArray (and therefore FString) assignment ignoring previously reserved memory.

diff --git a/Engine/Source/Runtime/Core/Public/Containers/Array.h b/Engine/Source/Runtime/Core/Public/Containers/Array.h
--- a/Engine/Source/Runtime/Core/Public/Containers/Array.h
+++ b/Engine/Source/Runtime/Core/Public/Containers/Array.h
@@ -2535,11 +2535,15 @@ private:
 		{
 			NewMax = AllocatorInstance.CalculateSlackReserve(NewMax, sizeof(ElementType));
 		}
-		if (NewMax != PrevMax)
+		if (NewMax > PrevMax)
 		{
 			AllocatorInstance.ResizeAllocation(0, NewMax, sizeof(ElementType));
+			ArrayMax = NewMax;
+		}
+		else
+		{
+			ArrayMax = PrevMax;
 		}
-		ArrayMax = NewMax;
 	}
 
 
