Fix build error for the Incremental EC job (-disableunity -nopch)

#rnx

diff --git a/Engine/Source/Runtime/Core/Private/IO/IoStore.cpp b/Engine/Source/Runtime/Core/Private/IO/IoStore.cpp
--- a/Engine/Source/Runtime/Core/Private/IO/IoStore.cpp
+++ b/Engine/Source/Runtime/Core/Private/IO/IoStore.cpp
@@ -4,6 +4,7 @@
 #include "Containers/Map.h"
 #include "HAL/FileManager.h"
 #include "Templates/UniquePtr.h"
+#include "Misc/Paths.h"
 
 //////////////////////////////////////////////////////////////////////////
 
