Fix nested cpu scopes with different channels.

Nested CpuProfiler scopes were ended too early in certain nested scenarios. This was solved by introducing a channel gate to OutputEndEvent. This means in order to see the cpu scopes from for example RHICommands channel, users also need to enable CpuProfilerChannel. Also removed bThreadEnabled which wasn't needed.

diff --git a/Engine/Source/Runtime/Core/Private/ProfilingDebugging/CpuProfilerTrace.cpp b/Engine/Source/Runtime/Core/Private/ProfilingDebugging/CpuProfilerTrace.cpp
--- a/Engine/Source/Runtime/Core/Private/ProfilingDebugging/CpuProfilerTrace.cpp
+++ b/Engine/Source/Runtime/Core/Private/ProfilingDebugging/CpuProfilerTrace.cpp
@@ -84,12 +84,10 @@ struct FCpuProfilerTraceInternal
 	FORCENOINLINE static void EndCapture(FThreadBuffer* ThreadBuffer);
 
 	static thread_local uint32 ThreadDepth;
-	static thread_local bool bThreadEnabled;
 	static thread_local FThreadBuffer* ThreadBuffer;
 };
 
 thread_local uint32 FCpuProfilerTraceInternal::ThreadDepth = 0;
-thread_local bool FCpuProfilerTraceInternal::bThreadEnabled = false;
 thread_local FCpuProfilerTraceInternal::FThreadBuffer* FCpuProfilerTraceInternal::ThreadBuffer = nullptr;
 
 FCpuProfilerTraceInternal::FThreadBuffer* FCpuProfilerTraceInternal::CreateThreadBuffer()
@@ -107,33 +105,18 @@ void FCpuProfilerTraceInternal::FlushThreadBuffer(FThreadBuffer* InThreadBuffer)
 	InThreadBuffer->BufferSize = 0;
 }
 
-void FCpuProfilerTraceInternal::EndCapture(FThreadBuffer* InThreadBuffer)
-{
-	UE_TRACE_LOG(CpuProfiler, EndCapture, CpuProfilerChannel, InThreadBuffer->BufferSize)
-		<< EndCapture.ThreadId(FPlatformTLS::GetCurrentThreadId())
-		<< EndCapture.Attachment(InThreadBuffer->Buffer, InThreadBuffer->BufferSize);
-	InThreadBuffer->BufferSize = 0;
-}
-
 #define CPUPROFILERTRACE_OUTPUTBEGINEVENT_PROLOGUE() \
 	++FCpuProfilerTraceInternal::ThreadDepth; \
-	const bool bEventEnabled = UE_TRACE_CHANNELEXPR_IS_ENABLED(Channel); \
-	if ((!bEventEnabled) & (bEventEnabled == FCpuProfilerTraceInternal::bThreadEnabled)) \
+	const bool bEventEnabled = UE_TRACE_CHANNELEXPR_IS_ENABLED(Channel | CpuProfilerChannel); \
+	if (!bEventEnabled) \
 	{ \
 		return; \
 	} \
-	check(FCpuProfilerTraceInternal::ThreadDepth < 1024); \
-	FCpuProfilerTraceInternal::bThreadEnabled = bEventEnabled; \
 	FCpuProfilerTraceInternal::FThreadBuffer* ThreadBuffer = FCpuProfilerTraceInternal::ThreadBuffer; \
 	if (!ThreadBuffer) \
 	{ \
 		ThreadBuffer = FCpuProfilerTraceInternal::CreateThreadBuffer(); \
 	} \
-	if (!bEventEnabled) \
-	{ \
-		FCpuProfilerTraceInternal::EndCapture(ThreadBuffer); \
-		return; \
-	}
 
 #define CPUPROFILERTRACE_OUTPUTBEGINEVENT_EPILOGUE() \
 	uint64 Cycle = FPlatformTime::Cycles64(); \
@@ -187,25 +170,19 @@ void FCpuProfilerTrace::OutputBeginDynamicEvent(const TCHAR* Name, const Trace::
 #undef CPUPROFILERTRACE_OUTPUTBEGINEVENT_PROLOGUE
 #undef CPUPROFILERTRACE_OUTPUTBEGINEVENT_EPILOGUE
 
-void FCpuProfilerTrace::OutputEndEvent()
+void FCpuProfilerTrace::OutputEndEvent(const Trace::FChannel& Channel)
 {
 	--FCpuProfilerTraceInternal::ThreadDepth;
-	const bool bEventEnabled = UE_TRACE_CHANNELEXPR_IS_ENABLED(CpuProfilerChannel);
-	if ((!bEventEnabled) & (bEventEnabled == FCpuProfilerTraceInternal::bThreadEnabled))
+	const bool bEventEnabled = UE_TRACE_CHANNELEXPR_IS_ENABLED(Channel | CpuProfilerChannel);
+	if (!bEventEnabled)
 	{
 		return;
 	}
-	FCpuProfilerTraceInternal::bThreadEnabled = bEventEnabled;
 	FCpuProfilerTraceInternal::FThreadBuffer* ThreadBuffer = FCpuProfilerTraceInternal::ThreadBuffer;
 	if (!ThreadBuffer)
 	{
 		ThreadBuffer = FCpuProfilerTraceInternal::CreateThreadBuffer();
 	}
-	if (!bEventEnabled)
-	{
-		FCpuProfilerTraceInternal::EndCapture(ThreadBuffer);
-		return;
-	}
 	uint64 Cycle = FPlatformTime::Cycles64();
 	uint64 CycleDiff = Cycle - ThreadBuffer->LastCycle;
 	ThreadBuffer->LastCycle = Cycle;
diff --git a/Engine/Source/Runtime/Core/Private/Windows/WindowsPlatformMisc.cpp b/Engine/Source/Runtime/Core/Private/Windows/WindowsPlatformMisc.cpp
--- a/Engine/Source/Runtime/Core/Private/Windows/WindowsPlatformMisc.cpp
+++ b/Engine/Source/Runtime/Core/Private/Windows/WindowsPlatformMisc.cpp
@@ -904,7 +904,7 @@ void FWindowsPlatformMisc::EndNamedEvent()
 	}
 #endif
 #if CPUPROFILERTRACE_ENABLED
-	FCpuProfilerTrace::OutputEndEvent();
+	FCpuProfilerTrace::OutputEndEvent(NamedEventsChannel);
 #endif
 }
 #endif // STATS || ENABLE_STATNAMEDEVENTS
diff --git a/Engine/Source/Runtime/Core/Public/ProfilingDebugging/CpuProfilerTrace.h b/Engine/Source/Runtime/Core/Public/ProfilingDebugging/CpuProfilerTrace.h
--- a/Engine/Source/Runtime/Core/Public/ProfilingDebugging/CpuProfilerTrace.h
+++ b/Engine/Source/Runtime/Core/Public/ProfilingDebugging/CpuProfilerTrace.h
@@ -31,37 +31,44 @@ struct FCpuProfilerTrace
 	CORE_API static void OutputBeginEvent(uint32 SpecId, const Trace::FChannel& Channel);
 	CORE_API static void OutputBeginDynamicEvent(const ANSICHAR* Name, const Trace::FChannel& Channel);
 	CORE_API static void OutputBeginDynamicEvent(const TCHAR* Name, const Trace::FChannel& Channel);
-	CORE_API static void OutputEndEvent();
+	CORE_API static void OutputEndEvent(const Trace::FChannel& Channel);
 
 	struct FEventScope
 	{
 		FEventScope(uint32 InSpecId, const Trace::FChannel& Channel)
+			: Channel(Channel)
 		{
 			OutputBeginEvent(InSpecId, Channel); 
 		}
 
 		~FEventScope()
 		{
-			OutputEndEvent();
+			OutputEndEvent(Channel);
 		}
+
+		const Trace::FChannel& Channel;
 	};
 
 	struct FDynamicEventScope
 	{
 		FDynamicEventScope(const ANSICHAR* InEventName, const Trace::FChannel& Channel)
+			: Channel(Channel)
 		{
 			OutputBeginDynamicEvent(InEventName, Channel);
 		}
 
 		FDynamicEventScope(const TCHAR* InEventName, const Trace::FChannel& Channel)
+			: Channel(Channel)
 		{
 			OutputBeginDynamicEvent(InEventName, Channel);
 		}
 
 		~FDynamicEventScope()
 		{
-			OutputEndEvent();
+			OutputEndEvent(Channel);
 		}
+
+		const Trace::FChannel& Channel;
 	};
 };
 
