IoDispatcher:
Fixed some initialization order problems

diff --git a/Engine/Source/Runtime/Core/Private/GenericPlatform/GenericPlatformIoDispatcher.cpp b/Engine/Source/Runtime/Core/Private/GenericPlatform/GenericPlatformIoDispatcher.cpp
--- a/Engine/Source/Runtime/Core/Private/GenericPlatform/GenericPlatformIoDispatcher.cpp
+++ b/Engine/Source/Runtime/Core/Private/GenericPlatform/GenericPlatformIoDispatcher.cpp
@@ -37,36 +37,21 @@ FGenericFileIoStoreImpl::FGenericFileIoStoreImpl(FGenericIoDispatcherEventQueue&
 bool FGenericFileIoStoreImpl::OpenContainer(const TCHAR* ContainerFilePath, uint64& ContainerFileHandle, uint64& ContainerFileSize)
 {
 	IPlatformFile& Ipf = FPlatformFileManager::Get().GetPlatformFile();
-	IAsyncReadFileHandle* FileHandle = Ipf.OpenAsyncRead(ContainerFilePath);
-	if (!FileHandle)
 	{
-		return false;
-	}
-	FEvent* Event = FPlatformProcess::GetSynchEventFromPool();
-	FAsyncFileCallBack SizeCallback = [Event](bool bWasCancelled, IAsyncReadRequest*)
-	{
-		check(!bWasCancelled);
-		Event->Trigger();
-	};
-	TUniquePtr<IAsyncReadRequest> SizeRequest(FileHandle->SizeRequest(&SizeCallback));
-	if (!SizeRequest.IsValid())
-	{
-		return false;
-	}
-	Event->Wait();
-	FPlatformProcess::ReturnSynchEventToPool(Event);
-	int64 SizeResult = SizeRequest->GetSizeResults();
-	if (SizeResult > 0)
-	{
-		ContainerFileHandle = reinterpret_cast<UPTRINT>(FileHandle);
-		ContainerFileSize = SizeResult;
-		return true;
+		TUniquePtr<IFileHandle> FileHandle(Ipf.OpenRead(ContainerFilePath));
+		if (!FileHandle)
+		{
+			return false;
+		}
+		ContainerFileSize = FileHandle->Size();
 	}
-	else
+	IAsyncReadFileHandle* AsyncFileHandle = Ipf.OpenAsyncRead(ContainerFilePath);
+	if (!AsyncFileHandle)
 	{
-		delete FileHandle;
 		return false;
 	}
+	ContainerFileHandle = reinterpret_cast<UPTRINT>(AsyncFileHandle);
+	return true;
 }
 
 void FGenericFileIoStoreImpl::BeginReadsForRequest(FFileIoStoreResolvedRequest& ResolvedRequest)
diff --git a/Engine/Source/Runtime/Core/Private/IO/IoDispatcher.cpp b/Engine/Source/Runtime/Core/Private/IO/IoDispatcher.cpp
--- a/Engine/Source/Runtime/Core/Private/IO/IoDispatcher.cpp
+++ b/Engine/Source/Runtime/Core/Private/IO/IoDispatcher.cpp
@@ -149,13 +149,8 @@ public:
 		delete Thread;
 	}
 
-	FIoStatus Initialize(const FIoStoreEnvironment& InitialEnvironment)
+	FIoStatus Initialize()
 	{
-		FIoStatus IoStatus = Mount(InitialEnvironment);
-		if (!IoStatus.IsOk())
-		{
-			return IoStatus;
-		}
 		return FIoStatus::Ok;
 	}
 
@@ -474,11 +469,11 @@ FIoDispatcher::IsValidEnvironment(const FIoStoreEnvironment& Environment)
 }
 
 FIoStatus
-FIoDispatcher::Initialize(const FIoStoreEnvironment& InitialEnvironment)
+FIoDispatcher::Initialize()
 {
 	GIoDispatcher = MakeUnique<FIoDispatcher>();
 
-	return GIoDispatcher->Impl->Initialize(InitialEnvironment);
+	return GIoDispatcher->Impl->Initialize();
 }
 
 void
diff --git a/Engine/Source/Runtime/Core/Public/IO/IoDispatcher.h b/Engine/Source/Runtime/Core/Public/IO/IoDispatcher.h
--- a/Engine/Source/Runtime/Core/Public/IO/IoDispatcher.h
+++ b/Engine/Source/Runtime/Core/Public/IO/IoDispatcher.h
@@ -747,7 +747,7 @@ public:
 
 	static CORE_API bool IsValidEnvironment(const FIoStoreEnvironment& Environment);
 	static CORE_API bool IsInitialized();
-	static CORE_API FIoStatus Initialize(const FIoStoreEnvironment& InitialEnvironment);
+	static CORE_API FIoStatus Initialize();
 	static CORE_API void Shutdown();
 	static CORE_API FIoDispatcher& Get();
 
diff --git a/Engine/Source/Runtime/CoreUObject/Private/Serialization/AsyncLoading2.cpp b/Engine/Source/Runtime/CoreUObject/Private/Serialization/AsyncLoading2.cpp
--- a/Engine/Source/Runtime/CoreUObject/Private/Serialization/AsyncLoading2.cpp
+++ b/Engine/Source/Runtime/CoreUObject/Private/Serialization/AsyncLoading2.cpp
@@ -1886,6 +1886,11 @@ struct FAsyncLoadingTickScope2
 
 void FAsyncLoadingThread2::InitializeLoading()
 {
+	FIoStoreEnvironment GlobalEnvironment;
+	GlobalEnvironment.InitializeFileEnvironment(FPaths::ProjectDir() / TEXT("global"));
+	FIoStatus MountStatus = IoDispatcher.Mount(GlobalEnvironment);
+	UE_CLOG(!MountStatus.IsOk(), LogInit, Fatal, TEXT("Failed to initialize I/O dispatcher: '%s'"), *MountStatus.ToString());
+
 #if USE_NEW_BULKDATA
 	FBulkDataBase::SetIoDispatcher(&IoDispatcher);
 #endif
diff --git a/Engine/Source/Runtime/Launch/Private/LaunchEngineLoop.cpp b/Engine/Source/Runtime/Launch/Private/LaunchEngineLoop.cpp
--- a/Engine/Source/Runtime/Launch/Private/LaunchEngineLoop.cpp
+++ b/Engine/Source/Runtime/Launch/Private/LaunchEngineLoop.cpp
@@ -1587,15 +1587,28 @@ int32 FEngineLoop::PreInitPreStartupScreen(const TCHAR* CmdLine)
 
 #if !(IS_PROGRAM || WITH_EDITOR)
 	// Initialize I/O dispatcher if available
-	if (!FParse::Param(CmdLine, TEXT("noiodispatcher")))
 	{
 		TRACE_CPUPROFILER_EVENT_SCOPE(InitIoDispatcher);
-
 		FIoStoreEnvironment IoStoreEnvironment;
 		IoStoreEnvironment.InitializeFileEnvironment(FPaths::ProjectDir() / TEXT("global"));
+		bool bEnableIoDispatcher = false;
 		if (FIoDispatcher::IsValidEnvironment(IoStoreEnvironment))
 		{
-			FIoStatus IoDispatcherInitStatus = FIoDispatcher::Initialize(IoStoreEnvironment);
+			bEnableIoDispatcher = true;
+		}
+#if !UE_BUILD_SHIPPING
+		if (FParse::Param(CmdLine, TEXT("forceiodispatcher")))
+		{
+			bEnableIoDispatcher = true;
+		}
+		else if (FParse::Param(CmdLine, TEXT("noiodispatcher")))
+		{
+			bEnableIoDispatcher = false;
+		}
+#endif
+		if (bEnableIoDispatcher)
+		{
+			FIoStatus IoDispatcherInitStatus = FIoDispatcher::Initialize();
 			UE_CLOG(!IoDispatcherInitStatus.IsOk(), LogInit, Fatal, TEXT("Failed to initialize I/O dispatcher: '%s'"), *IoDispatcherInitStatus.ToString());
 		}
 	}
