Fix FName batch loading check OldUsedSlots == UsedSlots in EngineTest

Very early batch loading reserve didn't acquire lock and in the EngineTest case other threads were creating FNames in parallel

diff --git a/Engine/Source/Runtime/Core/Private/UObject/UnrealNames.cpp b/Engine/Source/Runtime/Core/Private/UObject/UnrealNames.cpp
--- a/Engine/Source/Runtime/Core/Private/UObject/UnrealNames.cpp
+++ b/Engine/Source/Runtime/Core/Private/UObject/UnrealNames.cpp
@@ -723,6 +723,7 @@ public:
 	{
 		uint32 WantedCapacity = FMath::RoundUpToPowerOfTwo(Num * LoadFactorDivisor / LoadFactorQuotient);
 
+		FWriteScopeLock _(Lock);
 		if (WantedCapacity > Capacity())
 		{
 			Grow(WantedCapacity);
