Fix frame cycle counting when enabling channels during run.

diff --git a/Engine/Source/Runtime/Core/Private/ProfilingDebugging/MiscTrace.cpp b/Engine/Source/Runtime/Core/Private/ProfilingDebugging/MiscTrace.cpp
--- a/Engine/Source/Runtime/Core/Private/ProfilingDebugging/MiscTrace.cpp
+++ b/Engine/Source/Runtime/Core/Private/ProfilingDebugging/MiscTrace.cpp
@@ -126,6 +126,10 @@ void FMiscTrace::OutputBookmarkInternal(const void* BookmarkPoint, uint16 Encode
 
 void FMiscTrace::OutputBeginFrame(ETraceFrameType FrameType)
 {
+	if (!UE_TRACE_CHANNELEXPR_IS_ENABLED(FrameChannel))
+	{
+		return;
+	}
 	uint64 Cycle = FPlatformTime::Cycles64();
 	uint64 CycleDiff = Cycle - FMiscTraceInternal::LastFrameCycle[FrameType];
 	FMiscTraceInternal::LastFrameCycle[FrameType] = Cycle;
@@ -147,6 +151,10 @@ void FMiscTrace::OutputBeginFrame(ETraceFrameType FrameType)
 
 void FMiscTrace::OutputEndFrame(ETraceFrameType FrameType)
 {
+	if (!UE_TRACE_CHANNELEXPR_IS_ENABLED(FrameChannel))
+	{
+		return;
+	}
 	uint64 Cycle = FPlatformTime::Cycles64();
 	uint64 CycleDiff = Cycle - FMiscTraceInternal::LastFrameCycle[FrameType];
 	FMiscTraceInternal::LastFrameCycle[FrameType] = Cycle;
