Tiny optimization when comparing FName to nullptr or empty C string

diff --git a/Engine/Source/Runtime/Core/Private/UObject/UnrealNames.cpp b/Engine/Source/Runtime/Core/Private/UObject/UnrealNames.cpp
--- a/Engine/Source/Runtime/Core/Private/UObject/UnrealNames.cpp
+++ b/Engine/Source/Runtime/Core/Private/UObject/UnrealNames.cpp
@@ -1865,7 +1865,7 @@ struct FNameHelper
 		// Make NAME_None == TEXT("") or nullptr consistent with NAME_None == FName(TEXT("")) or FName(nullptr)
 		if (Str == nullptr || Str[0] == '\0')
 		{
-			return Name == NAME_None;
+			return Name.IsNone();
 		}
 
 		const FNameEntry& Entry = *Name.GetComparisonNameEntry();
