Fix an issue that was causing occasional duplicate FIoChunkIds when generating them for Bulkdata.

Changed the way we generate the FIoChunkId by reducing the chance of losing data when discarding the top byte of the BulkDataChunkId.
Add a log error if we find we have discarded data

#rnx

diff --git a/Engine/Source/Developer/IoStoreUtilities/Private/IoStoreUtilities.cpp b/Engine/Source/Developer/IoStoreUtilities/Private/IoStoreUtilities.cpp
--- a/Engine/Source/Developer/IoStoreUtilities/Private/IoStoreUtilities.cpp
+++ b/Engine/Source/Developer/IoStoreUtilities/Private/IoStoreUtilities.cpp
@@ -205,7 +205,16 @@ static FIoChunkId CreateChunkId(int32 GlobalPackageId, uint16 ChunkIndex, EIoChu
 	return ChunkId;
 }
 
-static FIoChunkId CreateChunkIdForBulkData(int32 GlobalPackageId, uint64 BulkdataOffset, EIoChunkType ChunkType, const TCHAR* DebugString)
+static FIoChunkId CreateChunkIdForBulkData(int32 GlobalPackageId,EIoChunkType ChunkType, const TCHAR* DebugString)
+{
+	FIoChunkId ChunkId = CreateIoChunkId(GlobalPackageId, 0, ChunkType);
+#if OUTPUT_CHUNKID_DIRECTORY
+	ChunkIdCsv.AddChunk(GlobalPackageId, 0, 0, (uint8)ChunkType, GetTypeHash(ChunkId), DebugString);
+#endif
+	return ChunkId;
+}
+
+static FIoChunkId CreateChunkIdForBulkData(int32 GlobalPackageId, int64 BulkdataOffset, EIoChunkType ChunkType, const TCHAR* DebugString)
 {
 	FIoChunkId ChunkId = CreateBulkdataChunkId(GlobalPackageId, BulkdataOffset, ChunkType);
 #if OUTPUT_CHUNKID_DIRECTORY
@@ -860,7 +869,7 @@ static bool WriteBulkData(	const FString& Filename, FPackageStoreBulkDataManifes
 	{
 		const EIoChunkType ChunkIdType = BulkdataTypeToChunkIdType(Type);
 
-		const FIoChunkId BulkDataChunkId = CreateChunkIdForBulkData(Package.GlobalPackageId, TNumericLimits<uint64>::Max()-1, ChunkIdType, *Package.FileName);
+		const FIoChunkId BulkDataChunkId = CreateChunkIdForBulkData(Package.GlobalPackageId, ChunkIdType, *Package.FileName);
 
 #if !SKIP_WRITE_CONTAINER		
 		FIoBuffer IoBuffer;
diff --git a/Engine/Source/Runtime/Core/Public/IO/IoDispatcher.h b/Engine/Source/Runtime/Core/Public/IO/IoDispatcher.h
--- a/Engine/Source/Runtime/Core/Public/IO/IoDispatcher.h
+++ b/Engine/Source/Runtime/Core/Public/IO/IoDispatcher.h
@@ -604,14 +604,32 @@ static FIoChunkId CreateIoChunkId(uint32 GlobalPackageId, uint16 ChunkIndex, EIo
 }
 
 /**
- * Creates a chunk identifier for BulkData
+ * Creates a FIoChunkId in the format that Bulkdata expects.
+ *
+ * @param GlobalPackageId	The identifier for the package that the bulkdata object is owned by
+ * @param BulkDataChunkId	A unique id for the bulkdata (commonly the bulkdata offset value is used) 
+ * @param ChunkType			The chunk type commonly 'BulkData' or 'OptionalBulkData'
+ *
+ * @return A valid FIoChunkId
  */
-static FIoChunkId CreateBulkdataChunkId(int32 GlobalPackageId, uint64 BulkDataChunkId, EIoChunkType ChunkType)
-{
+static FIoChunkId CreateBulkdataChunkId(int32 GlobalPackageId, int64 BulkDataChunkId, EIoChunkType ChunkType)
+{
+	// We need to be able to call this in the data pipeline and at runtime but currently are unable change the 
+	// file format we cannot generate this during cook and pass it to runtime.
+	// The offset in file is the only unique value we can easily obtain at runtime but it can be negative,
+	// which is a problem because we will only store the first 7 bytes and a negative value will have the 
+	// top bit set. 
+	// We adjust the id and cast to unsigned so that the top byte is very unlikely to have data in it (and log 
+	// it as an error if it does)
+
+	const uint64 Offset = ((uint64_t)1 << 56) / 2;
+	const uint64 AdjustedChunkId = BulkDataChunkId + Offset;
 	uint8 Data[12] = { 0 };
+	
+	UE_CLOG((AdjustedChunkId & 0xF000000000000000) != 0, LogIoDispatcher, Error, TEXT("The BulkDataChunkId (%lld) being used to create a BulkdataChunkId is too large and will lose data, this might create unintended duplicate ids!"), BulkDataChunkId);
 
 	*reinterpret_cast<int32*>(&Data[0]) = GlobalPackageId;
-	*reinterpret_cast<uint64*>(&Data[4]) = BulkDataChunkId; // Top byte will get overwritten!
+	*reinterpret_cast<uint64*>(&Data[4]) = AdjustedChunkId; // Top byte will get overwritten!
 	*reinterpret_cast<uint8*>(&Data[11]) = static_cast<uint8>(ChunkType);
 
 	FIoChunkId ChunkId;
