Fixed compile error when EXPERIMENTAL_STATSTRACE_ENABLED is 1.

diff --git a/Engine/Source/Runtime/Core/Private/Stats/Stats2.cpp b/Engine/Source/Runtime/Core/Private/Stats/Stats2.cpp
--- a/Engine/Source/Runtime/Core/Private/Stats/Stats2.cpp
+++ b/Engine/Source/Runtime/Core/Private/Stats/Stats2.cpp
@@ -508,6 +508,7 @@ public:
 	{
 		LLM_SCOPE(ELLMTag::Stats);
 
+		FName Stat;
 		TStatIdData* Result;
 		FString StatDescription;
 		{
@@ -515,7 +516,7 @@ public:
 
 			FStatNameAndInfo LongName(StatShortName, InGroup, InCategory, InDescription, InStatType, bShouldClearEveryFrame, bCycleStat, bSortByName, MemoryRegion);
 
-			FName Stat = LongName.GetEncodedName();
+			Stat = LongName.GetEncodedName();
 
 			FName Group(InGroup);
 			FGroupEnable* Found = HighPerformanceEnable.Find(Group);
