TNthTypeFromParameterPack moved to UnrealTypeTraits.h.

#rnx

diff --git a/Engine/Source/Runtime/Core/Public/Templates/Tuple.h b/Engine/Source/Runtime/Core/Public/Templates/Tuple.h
--- a/Engine/Source/Runtime/Core/Public/Templates/Tuple.h
+++ b/Engine/Source/Runtime/Core/Public/Templates/Tuple.h
@@ -17,21 +17,6 @@ struct TTuple;
 
 namespace UE4Tuple_Private
 {
-	template <int32 N, typename... Types>
-	struct TNthTypeFromParameterPack;
-
-	template <int32 N, typename T, typename... OtherTypes>
-	struct TNthTypeFromParameterPack<N, T, OtherTypes...>
-	{
-		typedef typename TNthTypeFromParameterPack<N - 1, OtherTypes...>::Type Type;
-	};
-
-	template <typename T, typename... OtherTypes>
-	struct TNthTypeFromParameterPack<0, T, OtherTypes...>
-	{
-		typedef T Type;
-	};
-
 	template <typename T, typename... Types>
 	struct TDecayedFrontOfParameterPackIsSameType
 	{
diff --git a/Engine/Source/Runtime/Core/Public/Templates/UnrealTypeTraits.h b/Engine/Source/Runtime/Core/Public/Templates/UnrealTypeTraits.h
--- a/Engine/Source/Runtime/Core/Public/Templates/UnrealTypeTraits.h
+++ b/Engine/Source/Runtime/Core/Public/Templates/UnrealTypeTraits.h
@@ -68,6 +68,22 @@ struct TIsDerivedFrom
 template<typename A, typename B>	struct TIsSame			{ enum { Value = false	}; };
 template<typename T>				struct TIsSame<T, T>	{ enum { Value = true	}; };
 
+/** Gets the Nth type in a template parameter pack. N must be less than sizeof...(Types) */
+template <int32 N, typename... Types>
+struct TNthTypeFromParameterPack;
+
+template <int32 N, typename T, typename... OtherTypes>
+struct TNthTypeFromParameterPack<N, T, OtherTypes...>
+{
+	using Type = typename TNthTypeFromParameterPack<N - 1, OtherTypes...>::Type;
+};
+
+template <typename T, typename... OtherTypes>
+struct TNthTypeFromParameterPack<0, T, OtherTypes...>
+{
+	using Type = T;
+};
+
 /**
  * TIsCharType
  */
