TVariant metaprogramming moved.

#rnx

diff --git a/Engine/Source/Runtime/Core/Public/Misc/TVariant.h b/Engine/Source/Runtime/Core/Public/Misc/TVariant.h
--- a/Engine/Source/Runtime/Core/Public/Misc/TVariant.h
+++ b/Engine/Source/Runtime/Core/Public/Misc/TVariant.h
@@ -2,257 +2,7 @@
 
 #pragma once
 
-#include "Templates/ChooseClass.h"
-#include "Templates/EnableIf.h"
-#include "Templates/IsClass.h"
-#include "Templates/IsConstructible.h"
-#include "Templates/MemoryOps.h"
-#include "Templates/RemoveReference.h"
-#include "Templates/TypeCompatibleBytes.h"
-#include "Templates/UnrealTemplate.h"
-#include "Templates/UnrealTypeTraits.h"
-
-#include "Misc/AssertionMacros.h"
-
-// Due to a bug in Visual Studio, we must use a recursive template to determine max sizeof() and alignof() of types in a template parameter pack.
-// On other compilers, we use a constexpr array and pluck out the largest
-#if defined(_MSC_VER) && !defined(__clang__)
-#define TVARIANT_STORAGE_USE_RECURSIVE_TEMPLATE 1
-#else
-#define TVARIANT_STORAGE_USE_RECURSIVE_TEMPLATE 0
-#endif
-
-namespace UE4Variant_Details
-{
-	/** Determine if all the types in a template parameter pack has duplicate types */
-	template <typename...>
-	struct TTypePackContainsDuplicates;
-
-	/** A template parameter pack containing a single type has no duplicates */
-	template <typename T>
-	struct TTypePackContainsDuplicates<T>
-	{
-		static constexpr bool Value = false;
-	};
-
-	/**
-	 * A template parameter pack containing the same type adjacently contains duplicate types.
-	 * The next structure ensures that we check all pairs of types in a template parameter pack.
-	 */
-	template <typename T, typename... Ts>
-	struct TTypePackContainsDuplicates<T, T, Ts...>
-	{
-		static constexpr bool Value = true;
-	};
-
-	/** Check all pairs of types in a template parameter pack to determine if any type is duplicated */
-	template <typename T, typename U, typename... Rest>
-	struct TTypePackContainsDuplicates<T, U, Rest...>
-	{
-		static constexpr bool Value = TTypePackContainsDuplicates<T, Rest...>::Value || TTypePackContainsDuplicates<U, Rest...>::Value;
-	};
-
-	/** Determine if any of the types in a template parameter pack are references */
-	template <typename... Ts>
-	struct TContainsReferenceType
-	{
-		static constexpr bool Value = TOr<TIsReferenceType<Ts>...>::Value;
-	};
-
-#if TVARIANT_STORAGE_USE_RECURSIVE_TEMPLATE
-	/** Determine the max alignof and sizeof of all types in a template parameter pack */
-	template <typename... Ts>
-	struct TVariantStorageTraits;
-
-	template <typename T, typename... Ts>
-	struct TVariantStorageTraits<T, Ts...>
-	{
-		static constexpr SIZE_T MaxSizeof(SIZE_T CurrentSize)
-		{
-			return TVariantStorageTraits<Ts...>::MaxSizeof(TVariantStorageTraits<T>::MaxSizeof(CurrentSize));
-		}
-
-		static constexpr SIZE_T MaxAlignof(SIZE_T CurrentSize)
-		{
-			return TVariantStorageTraits<Ts...>::MaxAlignof(TVariantStorageTraits<T>::MaxAlignof(CurrentSize));
-		}
-	};
-
-	template <typename T>
-	struct TVariantStorageTraits<T>
-	{
-		static constexpr SIZE_T MaxSizeof(SIZE_T CurrentSize)
-		{
-			return CurrentSize > sizeof(T) ? CurrentSize : sizeof(T);
-		}
-
-		static constexpr SIZE_T MaxAlignof(SIZE_T CurrentSize)
-		{
-			return CurrentSize > alignof(T) ? CurrentSize : alignof(T);
-		}
-	};
-
-	/** Expose a type that is suitable for storing any of the types in a template parameter pack */
-	template <typename T, typename... Ts>
-	struct TVariantStorage
-	{
-		static constexpr SIZE_T SizeofValue = TVariantStorageTraits<T, Ts...>::MaxSizeof(0);
-		static constexpr SIZE_T AlignofValue = TVariantStorageTraits<T, Ts...>::MaxAlignof(0);
-		static_assert(SizeofValue > 0, "MaxSizeof must be greater than 0");
-		static_assert(AlignofValue > 0, "MaxAlignof must be greater than 0");
-		using Type = TAlignedBytes<SizeofValue, AlignofValue>;
-	};
-#else
-	/** Determine the max alignof and sizeof of all types in a template parameter pack and provide a type that is compatible with those sizes */
-	template <typename... Ts>
-	struct TVariantStorage
-	{
-		static constexpr SIZE_T MaxOf(const SIZE_T Sizes[])
-		{
-			SIZE_T MaxSize = Sizes[0];
-			for (int32 Itr = 1; Itr < sizeof...(Ts); ++Itr)
-			{
-				if (Sizes[Itr] > MaxSize)
-				{
-					MaxSize = Sizes[Itr];
-				}
-			}
-			return MaxSize;
-		}
-		static constexpr SIZE_T MaxSizeof()
-		{
-			constexpr SIZE_T Sizes[] = { sizeof(Ts)... };
-			return MaxOf(Sizes);
-		}
-		static constexpr SIZE_T MaxAlignof()
-		{
-			constexpr SIZE_T Sizes[] = { alignof(Ts)... };
-			return MaxOf(Sizes);
-		}
-
-		static constexpr SIZE_T SizeofValue = MaxSizeof();
-		static constexpr SIZE_T AlignofValue = MaxAlignof();
-		static_assert(SizeofValue > 0, "MaxSizeof must be greater than 0");
-		static_assert(AlignofValue > 0, "MaxAlignof must be greater than 0");
-		using Type = TAlignedBytes<SizeofValue, AlignofValue>;
-	};
-#endif
-
-	/** Helper to lookup indices of each type in a template parameter pack */
-	template <SIZE_T N, typename LookupType, typename... Ts>
-	struct TParameterPackTypeIndexHelper
-	{
-		static constexpr SIZE_T Value = (SIZE_T)-1;
-	};
-
-	/** When the type we're looking up bubbles up to the top, we return the current index */
-	template <SIZE_T N, typename T, typename... Ts>
-	struct TParameterPackTypeIndexHelper<N, T, T, Ts...>
-	{
-		static constexpr SIZE_T Value = N;
-	};
-
-	/** When different type than the lookup is at the front of the parameter pack, we increase the index and move to the next type */
-	template <SIZE_T N, typename LookupType, typename T, typename... Ts>
-	struct TParameterPackTypeIndexHelper<N, LookupType, T, Ts...>
-	{
-		static constexpr SIZE_T Value = TParameterPackTypeIndexHelper<N + 1, LookupType, Ts...>::Value;
-	};
-
-	/** Entry-point for looking up the index of a type in a template parameter pack */
-	template <typename LookupType, typename... Ts>
-	struct TParameterPackTypeIndex
-	{
-		static constexpr SIZE_T Value = TParameterPackTypeIndexHelper<0, LookupType, Ts...>::Value;
-	};
-
-	/** An adapter for calling DestructItem */
-	template <typename T>
-	struct TDestructorCaller
-	{
-		static constexpr void Destruct(void* Storage)
-		{
-			DestructItem(static_cast<T*>(Storage));
-		}
-	};
-
-	/** Lookup a type in a template parameter pack by its index and call the destructor */
-	template <typename... Ts>
-	struct TDestructorLookup
-	{
-		/** If the index matches, call the destructor, otherwise call with the next index and type in the parameter pack*/
-		static void Destruct(SIZE_T TypeIndex, void* Value)
-		{
-			static constexpr void(*Destructors[])(void*) = { &TDestructorCaller<Ts>::Destruct... };
-			check(TypeIndex < UE_ARRAY_COUNT(Destructors));
-			Destructors[TypeIndex](Value);
-		}
-	};
-
-	/** An adapter for calling a copy constructor of a type */
-	template <typename T>
-	struct TCopyConstructorCaller
-	{
-		/** Call the copy constructor of a type with the provided memory location and value */
-		static void Construct(void* Storage, const void* Value)
-		{
-			new(Storage) T(*static_cast<const T*>(Value));
-		}
-	};
-
-	/** A utility for calling a type's copy constructor based on an index into a template parameter pack */
-	template <typename... Ts>
-	struct TCopyConstructorLookup
-	{
-		/** Construct the type at the index in the template parameter pack with the provided memory location and value */
-		static void Construct(SIZE_T TypeIndex, void* Storage, const void* Value)
-		{
-			static constexpr void(*CopyConstructors[])(void*, const void*) = { &TCopyConstructorCaller<Ts>::Construct... };
-			check(TypeIndex < UE_ARRAY_COUNT(CopyConstructors));
-			CopyConstructors[TypeIndex](Storage, Value);
-		}
-	};
-
-
-	/** A utility for calling a type's move constructor based on an index into a template parameter pack */
-	template <typename T>
-	struct TMoveConstructorCaller
-	{
-		/** Call the move constructor of a type with the provided memory location and value */
-		static void Construct(void* Storage, void* Value)
-		{
-			new(Storage) T(MoveTemp(*static_cast<T*>(Value)));
-		}
-	};
-
-	/** A utility for calling a type's move constructor based on an index into a template parameter pack */
-	template <typename... Ts>
-	struct TMoveConstructorLookup
-	{
-		/** Construct the type at the index in the template parameter pack with the provided memory location and value */
-		static void Construct(SIZE_T TypeIndex, void* Target, void* Source)
-		{
-			static constexpr void(*MoveConstructors[])(void*, void*) = { &TMoveConstructorCaller<Ts>::Construct... };
-			check(TypeIndex < UE_ARRAY_COUNT(MoveConstructors));
-			MoveConstructors[TypeIndex](Target, Source);
-		}
-	};
-
-	/** Determine if the type with the provided index in the template parameter pack is the same */
-	template <typename LookupType, typename... Ts>
-	struct TIsType
-	{
-		/** Check if the type at the provided index is the lookup type */
-		static bool IsSame(SIZE_T TypeIndex)
-		{
-			static constexpr bool bIsSameType[] = { TIsSame<Ts, LookupType>::Value... };
-			check(TypeIndex < UE_ARRAY_COUNT(bIsSameType));
-			return bIsSameType[TypeIndex];
-		}
-	};
-}
-
-#undef TVARIANT_STORAGE_USE_RECURSIVE_TEMPLATE
+#include "TVariantMeta.h"
 
 /**
  * A special tag used to indicate that in-place construction of a variant should take place.
@@ -290,7 +40,7 @@ public:
 	{
 		constexpr SIZE_T Index = UE4Variant_Details::TParameterPackTypeIndex<U, T, Ts...>::Value;
 		static_assert(Index != (SIZE_T)-1, "The TVariant is not declared to hold the type being constructed");
-		
+
 		new(&Storage) U(Forward<TArgs>(Args)...);
 		TypeIndex = Index;
 	}
@@ -400,7 +150,7 @@ public:
 	{
 		constexpr SIZE_T Index = UE4Variant_Details::TParameterPackTypeIndex<U, T, Ts...>::Value;
 		static_assert(Index != (SIZE_T)-1, "The TVariant is not declared to hold the type passed to Emplace<>");
-		
+
 		UE4Variant_Details::TDestructorLookup<T, Ts...>::Destruct(TypeIndex, &Storage);
 		new(&Storage) U(Forward<TArgs>(Args)...);
 		TypeIndex = Index;
@@ -408,7 +158,7 @@ public:
 
 	/** Lookup the index of a type in the template parameter pack at compile time. */
 	template <typename U>
-	static constexpr SIZE_T IndexOfType() 
+	static constexpr SIZE_T IndexOfType()
 	{
 		constexpr SIZE_T Index = UE4Variant_Details::TParameterPackTypeIndex<U, T, Ts...>::Value;
 		static_assert(Index != (SIZE_T)-1, "The TVariant is not declared to hold the type passed to IndexOfType<>");
@@ -426,4 +176,4 @@ private:
 	typename UE4Variant_Details::TVariantStorage<T, Ts...>::Type Storage;
 	/** Index into the template parameter pack for the type held. */
 	SIZE_T TypeIndex;
-};
+};
\ No newline at end of file
diff --git a/Engine/Source/Runtime/Core/Public/Misc/TVariantMeta.h b/Engine/Source/Runtime/Core/Public/Misc/TVariantMeta.h
new file mode 100644
--- /dev/null
+++ b/Engine/Source/Runtime/Core/Public/Misc/TVariantMeta.h
@@ -0,0 +1,256 @@
+// Copyright Epic Games, Inc. All Rights Reserved.
+
+#pragma once
+
+#include "Templates/ChooseClass.h"
+#include "Templates/EnableIf.h"
+#include "Templates/IsClass.h"
+#include "Templates/IsConstructible.h"
+#include "Templates/MemoryOps.h"
+#include "Templates/RemoveReference.h"
+#include "Templates/TypeCompatibleBytes.h"
+#include "Templates/UnrealTemplate.h"
+#include "Templates/UnrealTypeTraits.h"
+
+#include "Misc/AssertionMacros.h"
+
+// Due to a bug in Visual Studio, we must use a recursive template to determine max sizeof() and alignof() of types in a template parameter pack.
+// On other compilers, we use a constexpr array and pluck out the largest
+// Bug reported to Microsoft https://developercommunity.visualstudio.com/content/problem/528990/constexpr-expansion-inside-a-lambda-fails-to-be-ev.html
+#if defined(_MSC_VER) && !defined(__clang__)
+#define TVARIANT_STORAGE_USE_RECURSIVE_TEMPLATE 1
+#else
+#define TVARIANT_STORAGE_USE_RECURSIVE_TEMPLATE 0
+#endif
+
+namespace UE4Variant_Details
+{
+	/** Determine if all the types in a template parameter pack has duplicate types */
+	template <typename...>
+	struct TTypePackContainsDuplicates;
+
+	/** A template parameter pack containing a single type has no duplicates */
+	template <typename T>
+	struct TTypePackContainsDuplicates<T>
+	{
+		static constexpr bool Value = false;
+	};
+
+	/**
+	 * A template parameter pack containing the same type adjacently contains duplicate types.
+	 * The next structure ensures that we check all pairs of types in a template parameter pack.
+	 */
+	template <typename T, typename... Ts>
+	struct TTypePackContainsDuplicates<T, T, Ts...>
+	{
+		static constexpr bool Value = true;
+	};
+
+	/** Check all pairs of types in a template parameter pack to determine if any type is duplicated */
+	template <typename T, typename U, typename... Rest>
+	struct TTypePackContainsDuplicates<T, U, Rest...>
+	{
+		static constexpr bool Value = TTypePackContainsDuplicates<T, Rest...>::Value || TTypePackContainsDuplicates<U, Rest...>::Value;
+	};
+
+	/** Determine if any of the types in a template parameter pack are references */
+	template <typename... Ts>
+	struct TContainsReferenceType
+	{
+		static constexpr bool Value = TOr<TIsReferenceType<Ts>...>::Value;
+	};
+
+#if TVARIANT_STORAGE_USE_RECURSIVE_TEMPLATE
+	/** Determine the max alignof and sizeof of all types in a template parameter pack */
+	template <typename... Ts>
+	struct TVariantStorageTraits;
+
+	template <typename T, typename... Ts>
+	struct TVariantStorageTraits<T, Ts...>
+	{
+		static constexpr SIZE_T MaxSizeof(SIZE_T CurrentSize)
+		{
+			return TVariantStorageTraits<Ts...>::MaxSizeof(TVariantStorageTraits<T>::MaxSizeof(CurrentSize));
+		}
+
+		static constexpr SIZE_T MaxAlignof(SIZE_T CurrentSize)
+		{
+			return TVariantStorageTraits<Ts...>::MaxAlignof(TVariantStorageTraits<T>::MaxAlignof(CurrentSize));
+		}
+	};
+
+	template <typename T>
+	struct TVariantStorageTraits<T>
+	{
+		static constexpr SIZE_T MaxSizeof(SIZE_T CurrentSize)
+		{
+			return CurrentSize > sizeof(T) ? CurrentSize : sizeof(T);
+		}
+
+		static constexpr SIZE_T MaxAlignof(SIZE_T CurrentSize)
+		{
+			return CurrentSize > alignof(T) ? CurrentSize : alignof(T);
+		}
+	};
+
+	/** Expose a type that is suitable for storing any of the types in a template parameter pack */
+	template <typename T, typename... Ts>
+	struct TVariantStorage
+	{
+		static constexpr SIZE_T SizeofValue = TVariantStorageTraits<T, Ts...>::MaxSizeof(0);
+		static constexpr SIZE_T AlignofValue = TVariantStorageTraits<T, Ts...>::MaxAlignof(0);
+		static_assert(SizeofValue > 0, "MaxSizeof must be greater than 0");
+		static_assert(AlignofValue > 0, "MaxAlignof must be greater than 0");
+		using Type = TAlignedBytes<SizeofValue, AlignofValue>;
+	};
+#else
+	/** Determine the max alignof and sizeof of all types in a template parameter pack and provide a type that is compatible with those sizes */
+	template <typename... Ts>
+	struct TVariantStorage
+	{
+		static constexpr SIZE_T MaxOf(const SIZE_T Sizes[])
+		{
+			SIZE_T MaxSize = Sizes[0];
+			for (int32 Itr = 1; Itr < sizeof...(Ts); ++Itr)
+			{
+				if (Sizes[Itr] > MaxSize)
+				{
+					MaxSize = Sizes[Itr];
+				}
+			}
+			return MaxSize;
+		}
+		static constexpr SIZE_T MaxSizeof()
+		{
+			constexpr SIZE_T Sizes[] = { sizeof(Ts)... };
+			return MaxOf(Sizes);
+		}
+		static constexpr SIZE_T MaxAlignof()
+		{
+			constexpr SIZE_T Sizes[] = { alignof(Ts)... };
+			return MaxOf(Sizes);
+		}
+
+		static constexpr SIZE_T SizeofValue = MaxSizeof();
+		static constexpr SIZE_T AlignofValue = MaxAlignof();
+		static_assert(SizeofValue > 0, "MaxSizeof must be greater than 0");
+		static_assert(AlignofValue > 0, "MaxAlignof must be greater than 0");
+		using Type = TAlignedBytes<SizeofValue, AlignofValue>;
+	};
+#endif
+
+	/** Helper to lookup indices of each type in a template parameter pack */
+	template <SIZE_T N, typename LookupType, typename... Ts>
+	struct TParameterPackTypeIndexHelper
+	{
+		static constexpr SIZE_T Value = (SIZE_T)-1;
+	};
+
+	/** When the type we're looking up bubbles up to the top, we return the current index */
+	template <SIZE_T N, typename T, typename... Ts>
+	struct TParameterPackTypeIndexHelper<N, T, T, Ts...>
+	{
+		static constexpr SIZE_T Value = N;
+	};
+
+	/** When different type than the lookup is at the front of the parameter pack, we increase the index and move to the next type */
+	template <SIZE_T N, typename LookupType, typename T, typename... Ts>
+	struct TParameterPackTypeIndexHelper<N, LookupType, T, Ts...>
+	{
+		static constexpr SIZE_T Value = TParameterPackTypeIndexHelper<N + 1, LookupType, Ts...>::Value;
+	};
+
+	/** Entry-point for looking up the index of a type in a template parameter pack */
+	template <typename LookupType, typename... Ts>
+	struct TParameterPackTypeIndex
+	{
+		static constexpr SIZE_T Value = TParameterPackTypeIndexHelper<0, LookupType, Ts...>::Value;
+	};
+
+	/** An adapter for calling DestructItem */
+	template <typename T>
+	struct TDestructorCaller
+	{
+		static constexpr void Destruct(void* Storage)
+		{
+			DestructItem(static_cast<T*>(Storage));
+		}
+	};
+
+	/** Lookup a type in a template parameter pack by its index and call the destructor */
+	template <typename... Ts>
+	struct TDestructorLookup
+	{
+		/** If the index matches, call the destructor, otherwise call with the next index and type in the parameter pack*/
+		static void Destruct(SIZE_T TypeIndex, void* Value)
+		{
+			static constexpr void(*Destructors[])(void*) = { &TDestructorCaller<Ts>::Destruct... };
+			check(TypeIndex < UE_ARRAY_COUNT(Destructors));
+			Destructors[TypeIndex](Value);
+		}
+	};
+
+	/** An adapter for calling a copy constructor of a type */
+	template <typename T>
+	struct TCopyConstructorCaller
+	{
+		/** Call the copy constructor of a type with the provided memory location and value */
+		static void Construct(void* Storage, const void* Value)
+		{
+			new(Storage) T(*static_cast<const T*>(Value));
+		}
+	};
+
+	/** A utility for calling a type's copy constructor based on an index into a template parameter pack */
+	template <typename... Ts>
+	struct TCopyConstructorLookup
+	{
+		/** Construct the type at the index in the template parameter pack with the provided memory location and value */
+		static void Construct(SIZE_T TypeIndex, void* Storage, const void* Value)
+		{
+			static constexpr void(*CopyConstructors[])(void*, const void*) = { &TCopyConstructorCaller<Ts>::Construct... };
+			check(TypeIndex < UE_ARRAY_COUNT(CopyConstructors));
+			CopyConstructors[TypeIndex](Storage, Value);
+		}
+	};
+
+
+	/** A utility for calling a type's move constructor based on an index into a template parameter pack */
+	template <typename T>
+	struct TMoveConstructorCaller
+	{
+		/** Call the move constructor of a type with the provided memory location and value */
+		static void Construct(void* Storage, void* Value)
+		{
+			new(Storage) T(MoveTemp(*static_cast<T*>(Value)));
+		}
+	};
+
+	/** A utility for calling a type's move constructor based on an index into a template parameter pack */
+	template <typename... Ts>
+	struct TMoveConstructorLookup
+	{
+		/** Construct the type at the index in the template parameter pack with the provided memory location and value */
+		static void Construct(SIZE_T TypeIndex, void* Target, void* Source)
+		{
+			static constexpr void(*MoveConstructors[])(void*, void*) = { &TMoveConstructorCaller<Ts>::Construct... };
+			check(TypeIndex < UE_ARRAY_COUNT(MoveConstructors));
+			MoveConstructors[TypeIndex](Target, Source);
+		}
+	};
+
+	/** Determine if the type with the provided index in the template parameter pack is the same */
+	template <typename LookupType, typename... Ts>
+	struct TIsType
+	{
+		/** Check if the type at the provided index is the lookup type */
+		static bool IsSame(SIZE_T TypeIndex)
+		{
+			static constexpr bool bIsSameType[] = { TIsSame<Ts, LookupType>::Value... };
+			check(TypeIndex < UE_ARRAY_COUNT(bIsSameType));
+			return bIsSameType[TypeIndex];
+		}
+	};
+}
+
+#undef TVARIANT_STORAGE_USE_RECURSIVE_TEMPLATE
