Visit support for TVariant, copied from CL#s 11085384 and 11096421.
PLATFORM_COMPILER_HAS_FOLD_EXPRESSIONS support for Clang and MSVC.

diff --git a/Engine/Source/Runtime/Core/Public/Clang/ClangPlatform.h b/Engine/Source/Runtime/Core/Public/Clang/ClangPlatform.h
--- a/Engine/Source/Runtime/Core/Public/Clang/ClangPlatform.h
+++ b/Engine/Source/Runtime/Core/Public/Clang/ClangPlatform.h
@@ -11,3 +11,9 @@
 #else
 	#define PLATFORM_COMPILER_HAS_IF_CONSTEXPR 0
 #endif
+
+#if defined(__cpp_fold_expressions)
+	#define PLATFORM_COMPILER_HAS_FOLD_EXPRESSIONS 1
+#else
+	#define PLATFORM_COMPILER_HAS_FOLD_EXPRESSIONS 0
+#endif
diff --git a/Engine/Source/Runtime/Core/Public/HAL/Platform.h b/Engine/Source/Runtime/Core/Public/HAL/Platform.h
--- a/Engine/Source/Runtime/Core/Public/HAL/Platform.h
+++ b/Engine/Source/Runtime/Core/Public/HAL/Platform.h
@@ -240,6 +240,9 @@
 #ifndef PLATFORM_COMPILER_HAS_IF_CONSTEXPR
 	#define PLATFORM_COMPILER_HAS_IF_CONSTEXPR 1
 #endif
+#ifndef PLATFORM_COMPILER_HAS_FOLD_EXPRESSIONS
+	#define PLATFORM_COMPILER_HAS_FOLD_EXPRESSIONS 0
+#endif
 #ifndef PLATFORM_TCHAR_IS_1_BYTE
 	#define PLATFORM_TCHAR_IS_1_BYTE			0
 #endif
diff --git a/Engine/Source/Runtime/Core/Public/MSVC/MSVCPlatform.h b/Engine/Source/Runtime/Core/Public/MSVC/MSVCPlatform.h
--- a/Engine/Source/Runtime/Core/Public/MSVC/MSVCPlatform.h
+++ b/Engine/Source/Runtime/Core/Public/MSVC/MSVCPlatform.h
@@ -11,3 +11,9 @@
 #else
 	#define PLATFORM_COMPILER_HAS_IF_CONSTEXPR 0
 #endif
+
+#if defined(__cpp_fold_expressions)
+	#define PLATFORM_COMPILER_HAS_FOLD_EXPRESSIONS 1
+#else
+	#define PLATFORM_COMPILER_HAS_FOLD_EXPRESSIONS 0
+#endif
diff --git a/Engine/Source/Runtime/Core/Public/Misc/TVariant.h b/Engine/Source/Runtime/Core/Public/Misc/TVariant.h
--- a/Engine/Source/Runtime/Core/Public/Misc/TVariant.h
+++ b/Engine/Source/Runtime/Core/Public/Misc/TVariant.h
@@ -3,6 +3,9 @@
 #pragma once
 
 #include "TVariantMeta.h"
+#include "Templates/EnableIf.h"
+#include "Templates/IsConstructible.h"
+#include "Templates/Decay.h"
 
 /**
  * A special tag used to indicate that in-place construction of a variant should take place.
@@ -20,7 +23,7 @@ struct FEmptyVariantState {};
  * Attempting to use the value of a Get() when the underlying type is different leads to undefined behavior.
  */
 template <typename T, typename... Ts>
-class TVariant
+class TVariant final : private UE4Variant_Details::TVariantStorage<T, Ts...>
 {
 	static_assert(!UE4Variant_Details::TTypePackContainsDuplicates<T, Ts...>::Value, "All the types used in TVariant should be unique");
 	static_assert(!UE4Variant_Details::TContainsReferenceType<T, Ts...>::Value, "TVariant cannot hold reference types");
@@ -30,7 +33,7 @@ public:
 	TVariant()
 	{
 		static_assert(TIsConstructible<T>::Value, "To default-initialize a TVariant, the first type in the parameter pack must be default constructible. Use FEmptyVariantState as the first type if none of the other types can be listed first.");
-		new(&Storage) T();
+		new(&UE4Variant_Details::CastToStorage(*this).Storage) T();
 		TypeIndex = 0;
 	}
 
@@ -41,7 +44,7 @@ public:
 		constexpr SIZE_T Index = UE4Variant_Details::TParameterPackTypeIndex<U, T, Ts...>::Value;
 		static_assert(Index != (SIZE_T)-1, "The TVariant is not declared to hold the type being constructed");
 
-		new(&Storage) U(Forward<TArgs>(Args)...);
+		new(&UE4Variant_Details::CastToStorage(*this).Storage) U(Forward<TArgs>(Args)...);
 		TypeIndex = Index;
 	}
 
@@ -49,14 +52,14 @@ public:
 	TVariant(const TVariant& Other)
 		: TypeIndex(Other.TypeIndex)
 	{
-		UE4Variant_Details::TCopyConstructorLookup<T, Ts...>::Construct(TypeIndex, &Storage, &Other.Storage);
+		UE4Variant_Details::TCopyConstructorLookup<T, Ts...>::Construct(TypeIndex, &UE4Variant_Details::CastToStorage(*this).Storage, &UE4Variant_Details::CastToStorage(Other).Storage);
 	}
 
 	/** Move construct the variant from another variant of the same type */
 	TVariant(TVariant&& Other)
 		: TypeIndex(Other.TypeIndex)
 	{
-		UE4Variant_Details::TMoveConstructorLookup<T, Ts...>::Construct(TypeIndex, &Storage, &Other.Storage);
+		UE4Variant_Details::TMoveConstructorLookup<T, Ts...>::Construct(TypeIndex, &UE4Variant_Details::CastToStorage(*this).Storage, &UE4Variant_Details::CastToStorage(Other).Storage);
 	}
 
 	/** Copy assign a variant from another variant of the same type */
@@ -84,7 +87,7 @@ public:
 	/** Destruct the underlying type (if appropriate) */
 	~TVariant()
 	{
-		UE4Variant_Details::TDestructorLookup<T, Ts...>::Destruct(TypeIndex, &Storage);
+		UE4Variant_Details::TDestructorLookup<T, Ts...>::Destruct(TypeIndex, &UE4Variant_Details::CastToStorage(*this).Storage);
 	}
 
 	/** Determine if the variant holds the specific type */
@@ -102,7 +105,7 @@ public:
 		static_assert(Index != (SIZE_T)-1, "The TVariant is not declared to hold the type passed to Get<>");
 
 		check(Index == TypeIndex);
-		return *reinterpret_cast<U*>(&Storage);
+		return *reinterpret_cast<U*>(&UE4Variant_Details::CastToStorage(*this).Storage);
 	}
 
 	/** Get a reference to the held value. Bad things can happen if this is called on a variant that does not hold the type asked for */
@@ -119,14 +122,14 @@ public:
 	{
 		constexpr SIZE_T Index = UE4Variant_Details::TParameterPackTypeIndex<U, T, Ts...>::Value;
 		static_assert(Index != (SIZE_T)-1, "The TVariant is not declared to hold the type passed to TryGet<>");
-		return Index == TypeIndex ? reinterpret_cast<U*>(&Storage) : nullptr;
+		return Index == TypeIndex ? reinterpret_cast<U*>(&UE4Variant_Details::CastToStorage(*this).Storage) : nullptr;
 	}
 
 	/** Get a pointer to the held value if the held type is the same as the one specified */
 	template <typename U>
 	const U* TryGet() const
 	{
-		// Temporarily add the const qualifier so we can implement TrGet in one location.
+		// Temporarily remove the const qualifier so we can implement TryGet in one location.
 		return const_cast<TVariant*>(this)->template TryGet<U>();
 	}
 
@@ -151,8 +154,8 @@ public:
 		constexpr SIZE_T Index = UE4Variant_Details::TParameterPackTypeIndex<U, T, Ts...>::Value;
 		static_assert(Index != (SIZE_T)-1, "The TVariant is not declared to hold the type passed to Emplace<>");
 
-		UE4Variant_Details::TDestructorLookup<T, Ts...>::Destruct(TypeIndex, &Storage);
-		new(&Storage) U(Forward<TArgs>(Args)...);
+		UE4Variant_Details::TDestructorLookup<T, Ts...>::Destruct(TypeIndex, &UE4Variant_Details::CastToStorage(*this).Storage);
+		new(&UE4Variant_Details::CastToStorage(*this).Storage) U(Forward<TArgs>(Args)...);
 		TypeIndex = Index;
 	}
 
@@ -172,8 +175,60 @@ public:
 	}
 
 private:
-	/** Memory location for the value */
-	typename UE4Variant_Details::TVariantStorage<T, Ts...>::Type Storage;
 	/** Index into the template parameter pack for the type held. */
 	SIZE_T TypeIndex;
-};
\ No newline at end of file
+};
+
+/** Determine if a type is a variant */
+template <typename T>
+struct TIsVariant
+{
+	static constexpr bool Value = false;
+};
+
+template <typename... Ts>
+struct TIsVariant<TVariant<Ts...>>
+{
+	static constexpr bool Value = true;
+};
+
+template <typename T> struct TIsVariant<T&> : public TIsVariant<T> {};
+template <typename T> struct TIsVariant<T&&> : public TIsVariant<T> {};
+template <typename T> struct TIsVariant<const T> : public TIsVariant<T> {};
+
+/** Determine the number of types in a TVariant */
+template <typename> struct TVariantSize;
+
+template <typename... Ts>
+struct TVariantSize<TVariant<Ts...>>
+{
+	static constexpr SIZE_T Value = sizeof...(Ts);
+};
+
+template <typename T> struct TVariantSize<T&> : public TVariantSize<T> {};
+template <typename T> struct TVariantSize<T&&> : public TVariantSize<T> {};
+template <typename T> struct TVariantSize<const T> : public TVariantSize<T> {};
+
+/** Apply a visitor function to the list of variants */
+template <
+	typename Func,
+	typename... Variants,
+	typename = typename TEnableIf<UE4Variant_Details::TIsAllVariant<typename TDecay<Variants>::Type...>::Value>::Type
+>
+decltype(auto) Visit(Func&& Callable, Variants&&... Args)
+{
+#if PLATFORM_COMPILER_HAS_FOLD_EXPRESSIONS
+	constexpr SIZE_T NumPermutations = (1 * ... * (TVariantSize<Variants>::Value));
+#else
+	constexpr SIZE_T VariantSizes[] = { TVariantSize<Variants>::Value... };
+	constexpr SIZE_T NumPermutations = UE4Variant_Details::Multiply(VariantSizes, sizeof...(Variants));
+#endif
+
+	return UE4Variant_Details::VisitImpl(
+		UE4Variant_Details::EncodeIndices(Args...),
+		Forward<Func>(Callable),
+		TMakeIntegerSequence<SIZE_T, NumPermutations>{},
+		TMakeIntegerSequence<SIZE_T, sizeof...(Variants)>{},
+		Forward<Variants>(Args)...
+	);
+}
\ No newline at end of file
diff --git a/Engine/Source/Runtime/Core/Public/Misc/TVariantMeta.h b/Engine/Source/Runtime/Core/Public/Misc/TVariantMeta.h
--- a/Engine/Source/Runtime/Core/Public/Misc/TVariantMeta.h
+++ b/Engine/Source/Runtime/Core/Public/Misc/TVariantMeta.h
@@ -2,15 +2,12 @@
 
 #pragma once
 
-#include "Templates/ChooseClass.h"
-#include "Templates/EnableIf.h"
-#include "Templates/IsClass.h"
-#include "Templates/IsConstructible.h"
 #include "Templates/MemoryOps.h"
-#include "Templates/RemoveReference.h"
 #include "Templates/TypeCompatibleBytes.h"
 #include "Templates/UnrealTemplate.h"
 #include "Templates/UnrealTypeTraits.h"
+#include "Delegates/IntegerSequence.h"
+#include "Templates/AndOrNot.h"
 
 #include "Misc/AssertionMacros.h"
 
@@ -23,6 +20,15 @@
 #define TVARIANT_STORAGE_USE_RECURSIVE_TEMPLATE 0
 #endif
 
+template <typename T, typename... Ts>
+class TVariant;
+
+template <typename T>
+struct TIsVariant;
+
+template <typename T>
+struct TVariantSize;
+
 namespace UE4Variant_Details
 {
 	/** Determine if all the types in a template parameter pack has duplicate types */
@@ -101,7 +107,24 @@ namespace UE4Variant_Details
 		static constexpr SIZE_T AlignofValue = TVariantStorageTraits<T, Ts...>::MaxAlignof(0);
 		static_assert(SizeofValue > 0, "MaxSizeof must be greater than 0");
 		static_assert(AlignofValue > 0, "MaxAlignof must be greater than 0");
-		using Type = TAlignedBytes<SizeofValue, AlignofValue>;
+
+		/** Interpret the underlying data as the type in the variant parameter pack at the compile-time index. This function is used to implement Visit and should not be used directly */
+		template <SIZE_T N>
+		auto& GetValueAsIndexedType()
+		{
+			using ReturnType = typename TNthTypeFromParameterPack<N, T, Ts...>::Type;
+			return *reinterpret_cast<ReturnType*>(&Storage);
+		}
+
+		/** Interpret the underlying data as the type in the variant parameter pack at the compile-time index. This function is used to implement Visit and should not be used directly */
+		template <SIZE_T N>
+		const auto& GetValueAsIndexedType() const
+		{
+			// Temporarily remove the const qualifier so we can implement GetValueAsIndexedType in one location.
+			return const_cast<TVariantStorage*>(this)->template GetValueAsIndexedType<N>();
+		}
+
+		TAlignedBytes<SizeofValue, AlignofValue> Storage;
 	};
 #else
 	/** Determine the max alignof and sizeof of all types in a template parameter pack and provide a type that is compatible with those sizes */
@@ -135,7 +158,24 @@ namespace UE4Variant_Details
 		static constexpr SIZE_T AlignofValue = MaxAlignof();
 		static_assert(SizeofValue > 0, "MaxSizeof must be greater than 0");
 		static_assert(AlignofValue > 0, "MaxAlignof must be greater than 0");
-		using Type = TAlignedBytes<SizeofValue, AlignofValue>;
+
+		/** Interpret the underlying data as the type in the variant parameter pack at the compile-time index. This function is used to implement Visit and should not be used directly */
+		template <SIZE_T N>
+		auto& GetValueAsIndexedType()
+		{
+			using ReturnType = typename TNthTypeFromParameterPack<N, Ts...>::Type;
+			return *reinterpret_cast<ReturnType*>(&Storage);
+		}
+
+		/** Interpret the underlying data as the type in the variant parameter pack at the compile-time index. This function is used to implement Visit and should not be used directly */
+		template <SIZE_T N>
+		const auto& GetValueAsIndexedType() const
+		{
+			// Temporarily remove the const qualifier so we can implement GetValueAsIndexedType in one location.
+			return const_cast<TVariantStorage*>(this)->template GetValueAsIndexedType<N>();
+		}
+
+		TAlignedBytes<SizeofValue, AlignofValue> Storage;
 	};
 #endif
 
@@ -251,6 +291,96 @@ namespace UE4Variant_Details
 			return bIsSameType[TypeIndex];
 		}
 	};
+
+	/** Determine if all the types are TVariant<...> */
+	template <typename... Ts>
+	using TIsAllVariant = TAnd<TIsVariant<Ts>...>;
+
+	/** Encode the stored index of a bunch of variants into a single value used to lookup a Visit invocation function */
+	template <typename T>
+	inline SIZE_T EncodeIndices(const T& Variant)
+	{
+		return Variant.GetIndex();
+	}
+
+	template <typename Variant0, typename... Variants>
+	inline SIZE_T EncodeIndices(const Variant0& First, const Variants&... Rest)
+	{
+		return First.GetIndex() + TVariantSize<Variant0>::Value * EncodeIndices(Rest...);
+	}
+
+	/** Inverse operation of EncodeIndices. Decodes an encoded index into the individual index for the specified variant index. */
+	constexpr SIZE_T DecodeIndex(SIZE_T EncodedIndex, SIZE_T VariantIndex, const SIZE_T* VariantSizes)
+	{
+		while (VariantIndex)
+		{
+			EncodedIndex /= *VariantSizes;
+			--VariantIndex;
+			++VariantSizes;
+		}
+		return EncodedIndex % *VariantSizes;
+	}
+
+#if !PLATFORM_COMPILER_HAS_FOLD_EXPRESSIONS
+	/** Used to determine the total number of possible Visit invocations when fold expressions are not available. */
+	constexpr SIZE_T Multiply(const SIZE_T* Args, SIZE_T Num)
+	{
+		SIZE_T Result = 1;
+		while (Num)
+		{
+			Result *= *Args++;
+			--Num;
+		}
+		return Result;
+	}
+#endif
+
+	/** Cast a TVariant to its private base */
+	template <typename... Ts>
+	FORCEINLINE TVariantStorage<Ts...>& CastToStorage(TVariant<Ts...>& Variant)
+	{
+		return *(TVariantStorage<Ts...>*)(&Variant);
+	}
+
+	template <typename... Ts>
+	FORCEINLINE const TVariantStorage<Ts...>& CastToStorage(const TVariant<Ts...>& Variant)
+	{
+		return *(const TVariantStorage<Ts...>*)(&Variant);
+	}
+
+	/** Invocation detail for a single combination of stored variant indices */
+	template <SIZE_T EncodedIndex, SIZE_T... VariantIndices, typename Func, typename... Variants>
+	inline decltype(auto) VisitApplyEncoded(Func&& Callable, Variants&&... Args)
+	{
+		constexpr SIZE_T VariantSizes[] = { TVariantSize<Variants>::Value... };
+		return Callable(CastToStorage(Args).template GetValueAsIndexedType<DecodeIndex(EncodedIndex, VariantIndices, VariantSizes)>()...);
+	}
+
+	/**
+	 * Work around used to separate pack expansion of EncodedIndices and VariantIndices in VisitImpl below when defining the Invokers array.
+	 *
+	 * Ideally the line below would only need to be written as:
+	 * constexpr InvokeFn Invokers[] = { &VisitApplyEncoded<EncodedIndices, VariantIndices...>... };
+	 *
+	 * Due to what appears to be a lexing bug, MSVC 2017 tries to expand EncodedIndices and VariantIndices together
+	 */
+	template <typename InvokeFn, SIZE_T... VariantIndices>
+	struct TWrapper
+	{
+		template <SIZE_T EncodedIndex>
+		static constexpr InvokeFn FuncPtr = &VisitApplyEncoded<EncodedIndex, VariantIndices...>;
+	};
+
+	/** Implementation detail for Visit(Callable, Variants...). Builds an array of invokers, and forwards the variants to the callable for the specific EncodedIndex */
+	template <typename Func, SIZE_T... EncodedIndices, SIZE_T... VariantIndices, typename... Variants>
+	decltype(auto) VisitImpl(SIZE_T EncodedIndex, Func&& Callable, TIntegerSequence<SIZE_T, EncodedIndices...>&&, TIntegerSequence<SIZE_T, VariantIndices...>&& VariantIndicesSeq, Variants&&... Args)
+	{
+		using ReturnType = decltype(VisitApplyEncoded<0, VariantIndices...>(Forward<Func>(Callable), Forward<Variants>(Args)...));
+		using InvokeFn = ReturnType(*)(Func&&, Variants&&...);
+		using WrapperType = TWrapper<InvokeFn, VariantIndices...>;
+		static constexpr InvokeFn Invokers[] = { WrapperType::template FuncPtr<EncodedIndices>... };
+		return Invokers[EncodedIndex](Forward<Func>(Callable), Forward<Variants>(Args)...);
+	}
 }
 
 #undef TVARIANT_STORAGE_USE_RECURSIVE_TEMPLATE
