Rename "CpuProfiler" channel to "Cpu" and disable frame names with indices.

* Rename CpuProfilerChannel to CpuChannel.
* Removed NamedEventsChannel and use CpuChannel.

diff --git a/Engine/Source/Runtime/Core/Private/ProfilingDebugging/CpuProfilerTrace.cpp b/Engine/Source/Runtime/Core/Private/ProfilingDebugging/CpuProfilerTrace.cpp
--- a/Engine/Source/Runtime/Core/Private/ProfilingDebugging/CpuProfilerTrace.cpp
+++ b/Engine/Source/Runtime/Core/Private/ProfilingDebugging/CpuProfilerTrace.cpp
@@ -12,8 +12,7 @@
 
 #if CPUPROFILERTRACE_ENABLED
 
-UE_TRACE_CHANNEL_DEFINE(CpuProfilerChannel)
-UE_TRACE_CHANNEL_DEFINE(NamedEventsChannel)
+UE_TRACE_CHANNEL_DEFINE(CpuChannel)
 
 UE_TRACE_EVENT_BEGIN(CpuProfiler, EventSpec, Important)
 	UE_TRACE_EVENT_FIELD(uint32, Id)
@@ -99,7 +98,7 @@ FCpuProfilerTraceInternal::FThreadBuffer* FCpuProfilerTraceInternal::CreateThrea
 
 void FCpuProfilerTraceInternal::FlushThreadBuffer(FThreadBuffer* InThreadBuffer)
 {
-	UE_TRACE_LOG(CpuProfiler, EventBatch, CpuProfilerChannel, InThreadBuffer->BufferSize)
+	UE_TRACE_LOG(CpuProfiler, EventBatch, CpuChannel, InThreadBuffer->BufferSize)
 		<< EventBatch.ThreadId(FPlatformTLS::GetCurrentThreadId())
 		<< EventBatch.Attachment(InThreadBuffer->Buffer, InThreadBuffer->BufferSize);
 	InThreadBuffer->BufferSize = 0;
@@ -107,7 +106,7 @@ void FCpuProfilerTraceInternal::FlushThreadBuffer(FThreadBuffer* InThreadBuffer)
 
 #define CPUPROFILERTRACE_OUTPUTBEGINEVENT_PROLOGUE() \
 	++FCpuProfilerTraceInternal::ThreadDepth; \
-	const bool bEventEnabled = UE_TRACE_CHANNELEXPR_IS_ENABLED(Channel | CpuProfilerChannel); \
+	const bool bEventEnabled = UE_TRACE_CHANNELEXPR_IS_ENABLED(Channel | CpuChannel); \
 	if (!bEventEnabled) \
 	{ \
 		return; \
@@ -173,7 +172,7 @@ void FCpuProfilerTrace::OutputBeginDynamicEvent(const TCHAR* Name, const Trace::
 void FCpuProfilerTrace::OutputEndEvent(const Trace::FChannel& Channel)
 {
 	--FCpuProfilerTraceInternal::ThreadDepth;
-	const bool bEventEnabled = UE_TRACE_CHANNELEXPR_IS_ENABLED(Channel | CpuProfilerChannel);
+	const bool bEventEnabled = UE_TRACE_CHANNELEXPR_IS_ENABLED(Channel | CpuChannel);
 	if (!bEventEnabled)
 	{
 		return;
@@ -205,7 +204,7 @@ uint32 FCpuProfilerTrace::OutputEventType(const TCHAR* Name)
 {
 	uint32 SpecId = FCpuProfilerTraceInternal::GetNextSpecId();
 	uint16 NameSize = (FCString::Strlen(Name) + 1) * sizeof(TCHAR);
-	UE_TRACE_LOG(CpuProfiler, EventSpec, CpuProfilerChannel, NameSize)
+	UE_TRACE_LOG(CpuProfiler, EventSpec, CpuChannel, NameSize)
 		<< EventSpec.Id(SpecId)
 		<< EventSpec.CharSize(sizeof(TCHAR))
 		<< EventSpec.Attachment(Name, NameSize);
@@ -216,7 +215,7 @@ uint32 FCpuProfilerTrace::OutputEventType(const ANSICHAR* Name)
 {
 	uint32 SpecId = FCpuProfilerTraceInternal::GetNextSpecId();
 	uint16 NameSize = strlen(Name) + 1;
-	UE_TRACE_LOG(CpuProfiler, EventSpec, CpuProfilerChannel, NameSize)
+	UE_TRACE_LOG(CpuProfiler, EventSpec, CpuChannel, NameSize)
 		<< EventSpec.Id(SpecId)
 		<< EventSpec.CharSize(1)
 		<< EventSpec.Attachment(Name, NameSize);
@@ -227,7 +226,7 @@ void FCpuProfilerTrace::Init(const TCHAR* CmdLine)
 {
 	if (FParse::Param(CmdLine, TEXT("cpuprofilertrace")))
 	{
-		Trace::ToggleChannel(CpuProfilerChannel, true);
+		Trace::ToggleChannel(CpuChannel, true);
 	}
 }
 
diff --git a/Engine/Source/Runtime/Core/Private/Windows/WindowsPlatformMisc.cpp b/Engine/Source/Runtime/Core/Private/Windows/WindowsPlatformMisc.cpp
--- a/Engine/Source/Runtime/Core/Private/Windows/WindowsPlatformMisc.cpp
+++ b/Engine/Source/Runtime/Core/Private/Windows/WindowsPlatformMisc.cpp
@@ -872,7 +872,7 @@ void FWindowsPlatformMisc::BeginNamedEvent(const struct FColor& Color, const TCH
 	}
 #endif
 #if CPUPROFILERTRACE_ENABLED
-	FCpuProfilerTrace::OutputBeginDynamicEvent(Text, NamedEventsChannel);
+	FCpuProfilerTrace::OutputBeginDynamicEvent(Text, CpuChannel);
 #endif
 }
 
@@ -888,7 +888,7 @@ void FWindowsPlatformMisc::BeginNamedEvent(const struct FColor& Color, const ANS
 	}
 #endif
 #if CPUPROFILERTRACE_ENABLED
-	FCpuProfilerTrace::OutputBeginDynamicEvent(Text, NamedEventsChannel);
+	FCpuProfilerTrace::OutputBeginDynamicEvent(Text, CpuChannel);
 #endif
 }
 
@@ -904,7 +904,7 @@ void FWindowsPlatformMisc::EndNamedEvent()
 	}
 #endif
 #if CPUPROFILERTRACE_ENABLED
-	FCpuProfilerTrace::OutputEndEvent(NamedEventsChannel);
+	FCpuProfilerTrace::OutputEndEvent(CpuChannel);
 #endif
 }
 #endif // STATS || ENABLE_STATNAMEDEVENTS
diff --git a/Engine/Source/Runtime/Core/Public/ProfilingDebugging/CpuProfilerTrace.h b/Engine/Source/Runtime/Core/Public/ProfilingDebugging/CpuProfilerTrace.h
--- a/Engine/Source/Runtime/Core/Public/ProfilingDebugging/CpuProfilerTrace.h
+++ b/Engine/Source/Runtime/Core/Public/ProfilingDebugging/CpuProfilerTrace.h
@@ -19,8 +19,7 @@
 
 // @note Cannot use the declare macros in this header since including
 // Trace.h will result in a circular dependency.
-CORE_API extern Trace::FChannel CpuProfilerChannel;
-CORE_API extern Trace::FChannel NamedEventsChannel;
+CORE_API extern Trace::FChannel CpuChannel;
 
 struct FCpuProfilerTrace
 {
@@ -89,13 +88,13 @@ struct FCpuProfilerTrace
 	TRACE_CPUPROFILER_EVENT_SCOPE_ON_CHANNEL_STR(#Name, Channel)
 
 #define TRACE_CPUPROFILER_EVENT_SCOPE(Name) \
-	TRACE_CPUPROFILER_EVENT_SCOPE_ON_CHANNEL(Name, CpuProfilerChannel)
+	TRACE_CPUPROFILER_EVENT_SCOPE_ON_CHANNEL(Name, CpuChannel)
 
 #define TRACE_CPUPROFILER_EVENT_SCOPE_TEXT_ON_CHANNEL(Name, Channel) \
 	FCpuProfilerTrace::FDynamicEventScope PREPROCESSOR_JOIN(__CpuProfilerEventScope, __LINE__)(Name, Channel);
 	
 #define TRACE_CPUPROFILER_EVENT_SCOPE_TEXT(Name) \
-	TRACE_CPUPROFILER_EVENT_SCOPE_TEXT_ON_CHANNEL(Name, CpuProfilerChannel)
+	TRACE_CPUPROFILER_EVENT_SCOPE_TEXT_ON_CHANNEL(Name, CpuChannel)
 
 #else
 
