Fix Linux C4244 warnings: '<usage>': conversion from '<biggertype>' to '<smallertype>', possible loss of data
#rnx

diff --git a/Engine/Source/Runtime/Core/Private/String/HexToBytes.cpp b/Engine/Source/Runtime/Core/Private/String/HexToBytes.cpp
--- a/Engine/Source/Runtime/Core/Private/String/HexToBytes.cpp
+++ b/Engine/Source/Runtime/Core/Private/String/HexToBytes.cpp
@@ -25,7 +25,7 @@ namespace String
 			const uint8 HiNibble = TCharToNibble(*HexPos++) << 4;
 			*OutPos++ = HiNibble | TCharToNibble(*HexPos++);
 		}
-		return OutPos - OutBytes;
+		return static_cast<int32>(OutPos - OutBytes);
 	}
 
 	int32 HexToBytes(const FStringView& Hex, uint8* OutBytes)
diff --git a/Engine/Source/Runtime/Core/Private/String/ParseLines.cpp b/Engine/Source/Runtime/Core/Private/String/ParseLines.cpp
--- a/Engine/Source/Runtime/Core/Private/String/ParseLines.cpp
+++ b/Engine/Source/Runtime/Core/Private/String/ParseLines.cpp
@@ -35,7 +35,7 @@ namespace String
 				}
 			}
 
-			Visitor(FStringView(LineStart, LineEnd - LineStart));
+			Visitor(FStringView(LineStart, static_cast<FStringView::SizeType>(LineEnd - LineStart)));
 		}
 		while (ViewIt != ViewEnd);
 	}
diff --git a/Engine/Source/Runtime/Core/Private/String/ParseTokens.cpp b/Engine/Source/Runtime/Core/Private/String/ParseTokens.cpp
--- a/Engine/Source/Runtime/Core/Private/String/ParseTokens.cpp
+++ b/Engine/Source/Runtime/Core/Private/String/ParseTokens.cpp
@@ -30,11 +30,11 @@ namespace String
 				++ViewIt;
 				continue;
 			}
-			Visitor(FStringView(NextToken, ViewIt - NextToken));
+			Visitor(FStringView(NextToken, static_cast<FStringView::SizeType>(ViewIt - NextToken)));
 			NextToken = ++ViewIt;
 		}
 
-		Visitor(FStringView(NextToken, ViewIt - NextToken));
+		Visitor(FStringView(NextToken, static_cast<FStringView::SizeType>(ViewIt - NextToken)));
 	}
 
 	/** Parse tokens with multiple single-character Basic Latin delimiters. */
@@ -62,11 +62,11 @@ namespace String
 				++ViewIt;
 				continue;
 			}
-			Visitor(FStringView(NextToken, ViewIt - NextToken));
+			Visitor(FStringView(NextToken, static_cast<FStringView::SizeType>(ViewIt - NextToken)));
 			NextToken = ++ViewIt;
 		}
 
-		Visitor(FStringView(NextToken, ViewIt - NextToken));
+		Visitor(FStringView(NextToken, static_cast<FStringView::SizeType>(ViewIt - NextToken)));
 	}
 
 	/** Parse tokens with multiple single-character delimiters in the Basic Multilingual Plane. */
@@ -92,11 +92,11 @@ namespace String
 				++ViewIt;
 				continue;
 			}
-			Visitor(FStringView(NextToken, ViewIt - NextToken));
+			Visitor(FStringView(NextToken, static_cast<FStringView::SizeType>(ViewIt - NextToken)));
 			NextToken = ++ViewIt;
 		}
 
-		Visitor(FStringView(NextToken, ViewIt - NextToken));
+		Visitor(FStringView(NextToken, static_cast<FStringView::SizeType>(ViewIt - NextToken)));
 	}
 
 	/** Parse tokens with multiple multi-character delimiters. */
diff --git a/Engine/Source/Runtime/Core/Public/Containers/UnrealString.h b/Engine/Source/Runtime/Core/Public/Containers/UnrealString.h
--- a/Engine/Source/Runtime/Core/Public/Containers/UnrealString.h
+++ b/Engine/Source/Runtime/Core/Public/Containers/UnrealString.h
@@ -2113,15 +2113,15 @@ inline const uint8 TCharToNibble(const TCHAR Hex)
 {
 	if (Hex >= '0' && Hex <= '9')
 	{
-		return Hex - '0';
+		return static_cast<uint8>(Hex - '0');
 	}
 	if (Hex >= 'A' && Hex <= 'F')
 	{
-		return Hex - 'A' + 10;
+		return static_cast<uint8>(Hex - 'A' + 10);
 	}
 	if (Hex >= 'a' && Hex <= 'f')
 	{
-		return Hex - 'a' + 10;
+		return static_cast<uint8>(Hex - 'a' + 10);
 	}
 	checkf(false, TEXT("'%c' (0x%02X) is not a valid hexadecimal digit"), Hex, Hex);
 	return 0;
