IoDispatcher:
Fixed broken request size logic

diff --git a/Engine/Source/Runtime/Core/Private/IO/IoDispatcherFileBackend.cpp b/Engine/Source/Runtime/Core/Private/IO/IoDispatcherFileBackend.cpp
--- a/Engine/Source/Runtime/Core/Private/IO/IoDispatcherFileBackend.cpp
+++ b/Engine/Source/Runtime/Core/Private/IO/IoDispatcherFileBackend.cpp
@@ -119,18 +119,16 @@ bool FFileIoStoreReader::Resolve(FFileIoStoreResolvedRequest& ResolvedRequest)
 	}
 
 	ResolvedRequest.ResolvedFileHandle = ContainerFileHandle;
-	uint64 FileEndOffset = OffsetAndLength->GetOffset() + OffsetAndLength->GetLength();
-	uint64 RequestedBeginOffset = OffsetAndLength->GetOffset() + ResolvedRequest.Request->Options.GetOffset();
-	uint64 RequestedEndOffset = FMath::Min(FileEndOffset, RequestedBeginOffset + ResolvedRequest.Request->Options.GetSize());
 	ResolvedRequest.ResolvedFileSize = ContainerFileSize;
-	ResolvedRequest.ResolvedOffset = RequestedBeginOffset;
-	if (RequestedEndOffset > RequestedBeginOffset)
+	uint64 RequestedOffset = ResolvedRequest.Request->Options.GetOffset();
+	ResolvedRequest.ResolvedOffset = OffsetAndLength->GetOffset() + RequestedOffset;
+	if (RequestedOffset > OffsetAndLength->GetLength())
 	{
-		ResolvedRequest.ResolvedSize = RequestedEndOffset - RequestedBeginOffset;
+		ResolvedRequest.ResolvedSize = 0;
 	}
 	else
 	{
-		ResolvedRequest.ResolvedSize = 0;
+		ResolvedRequest.ResolvedSize = FMath::Min(ResolvedRequest.Request->Options.GetSize(), OffsetAndLength->GetLength() - RequestedOffset);
 	}
 	return true;
 }
