Fix a race condition in FOutputDeviceRedirector that is corrupting logs

Calling EmptyBufferedLines() before serializing the buffered lines was allowing their memory (owned by FLogAllocator) to be reused by buffered lines from other threads before the main thread had consumed them.

This fix swaps between two inline allocators to allow the main thread to consume one buffer while other threads can fill the other buffer. A better solution would be to store buffered logs in a ring buffer that is consumed by a dedicated logging thread, to reduce logging overhead on the main thread.

diff --git a/Engine/Source/Runtime/Core/Private/Misc/OutputDeviceRedirector.cpp b/Engine/Source/Runtime/Core/Private/Misc/OutputDeviceRedirector.cpp
--- a/Engine/Source/Runtime/Core/Private/Misc/OutputDeviceRedirector.cpp
+++ b/Engine/Source/Runtime/Core/Private/Misc/OutputDeviceRedirector.cpp
@@ -1,10 +1,12 @@
 // Copyright Epic Games, Inc. All Rights Reserved.
 
 #include "Misc/OutputDeviceRedirector.h"
+
+#include "Containers/BitArray.h"
+#include "HAL/PlatformProcess.h"
+#include "Misc/CoreStats.h"
 #include "Misc/ScopeLock.h"
 #include "Stats/Stats.h"
-#include "Misc/CoreStats.h"
-#include "HAL/PlatformProcess.h"
 
 /*-----------------------------------------------------------------------------
 	FOutputDeviceRedirector.
@@ -15,24 +17,47 @@ class FLogAllocator
 public:
 	bool HasSpace(int32 NumChars) const
 	{
-		return Data.Num() + NumChars <= MaxChars;
+		return Data[BufferIndex].Num() + NumChars <= BufferSize;
 	}
 
 	TCHAR* Alloc(int32 NumChars)
 	{
 		check(HasSpace(NumChars));
-		Data.AddUninitialized(NumChars);
-		return Data.GetData() + Data.Num() - NumChars;
+		int32 DataIndex = Data[BufferIndex].AddUninitialized(NumChars);
+		return Data[BufferIndex].GetData() + DataIndex;
+	}
+
+	/**
+	 * Swap which buffer is being used for new allocations.
+	 *
+	 * Buffer remains valid until BufferCount calls are made to this function.
+	 */
+	void SwapBuffers()
+	{
+		BufferIndex = (BufferIndex + 1) % BufferCount;
+		check(!DataLocked[BufferIndex]);
+		Data[BufferIndex].Empty();
+	}
+
+	struct FBufferLock { int32 Index = -1; };
+
+	FBufferLock LockBuffer()
+	{
+		DataLocked[BufferIndex] = true;
+		return {BufferIndex};
 	}
 
-	void Reset()
+	void UnlockBuffer(FBufferLock Lock)
 	{
-		Data.Empty();
+		DataLocked[Lock.Index] = false;
 	}
 
 private:
-	enum { MaxChars = 4096 };
-	TArray<TCHAR, TInlineAllocator<MaxChars>> Data;
+	static constexpr int32 BufferSize = 4096;
+	static constexpr int32 BufferCount = 2;
+	TArray<TCHAR, TInlineAllocator<BufferSize>> Data[BufferCount];
+	TBitArray<TInlineAllocator<1>> DataLocked{false, BufferCount};
+	int32 BufferIndex = 0;
 };
 
 FBufferedLine::FBufferedLine(const TCHAR* InData, const FName& InCategory, ELogVerbosity::Type InVerbosity, double InTime, FLogAllocator* ExternalAllocator)
@@ -166,6 +191,7 @@ void FOutputDeviceRedirector::InternalFlushThreadedLogs(TLocalOutputDevicesArray
 	if (BufferedLines.Num())
 	{
 		TArray<FBufferedLine, TInlineAllocator<64>> LocalBufferedLines;
+		FLogAllocator::FBufferLock BufferLock;
 		{
 			FScopeLock ScopeLock(&BufferSynchronizationObject);
 			LocalBufferedLines.AddUninitialized(BufferedLines.Num());
@@ -173,6 +199,10 @@ void FOutputDeviceRedirector::InternalFlushThreadedLogs(TLocalOutputDevicesArray
 			{
 				new(&LocalBufferedLines[LineIndex]) FBufferedLine(BufferedLines[LineIndex], FBufferedLine::EMoveCtor);
 			}
+			if (BufferedLinesAllocator)
+			{
+				BufferLock = BufferedLinesAllocator->LockBuffer();
+			}
 			EmptyBufferedLines();
 		}
 
@@ -191,6 +221,12 @@ void FOutputDeviceRedirector::InternalFlushThreadedLogs(TLocalOutputDevicesArray
 				}
 			}
 		}
+
+		if (BufferedLinesAllocator)
+		{
+			FScopeLock ScopeLock(&BufferSynchronizationObject);
+			BufferedLinesAllocator->UnlockBuffer(BufferLock);
+		}
 	}
 }
 
@@ -200,7 +236,7 @@ void FOutputDeviceRedirector::EmptyBufferedLines()
 
 	if (BufferedLinesAllocator)
 	{
-		BufferedLinesAllocator->Reset();
+		BufferedLinesAllocator->SwapBuffers();
 	}
 }
 
@@ -312,7 +348,7 @@ void FOutputDeviceRedirector::SerializeImpl(const TCHAR* Data, ELogVerbosity::Ty
 
 #if PLATFORM_DESKTOP
 	// this is for errors which occur after shutdown we might be able to salvage information from stdout 
-	if ((BufferedOutputDevices.Num() == 0) && IsEngineExitRequested())
+	if ((LocalBufferedDevices.Num() == 0) && IsEngineExitRequested())
 	{
 #if PLATFORM_WINDOWS
 		_tprintf(_T("%s\n"), Data);
