IoDispatcher:
Use blocking reads instead of async reads
Don't let the block cache limit the number of active IO requests

diff --git a/Engine/Source/Runtime/Core/Private/GenericPlatform/GenericPlatformIoDispatcher.cpp b/Engine/Source/Runtime/Core/Private/GenericPlatform/GenericPlatformIoDispatcher.cpp
--- a/Engine/Source/Runtime/Core/Private/GenericPlatform/GenericPlatformIoDispatcher.cpp
+++ b/Engine/Source/Runtime/Core/Private/GenericPlatform/GenericPlatformIoDispatcher.cpp
@@ -6,11 +6,20 @@
 #include "HAL/Event.h"
 #include "HAL/PlatformFilemanager.h"
 #include "Misc/ScopeLock.h"
+#include "ProfilingDebugging/CountersTrace.h"
+#include "HAL/RunnableThread.h"
+
+//PRAGMA_DISABLE_OPTIMIZATION
+
+TRACE_DECLARE_INT_COUNTER(IoDispatcherSequentialReads, TEXT("IoDispatcher/SequentialReads"));
+TRACE_DECLARE_INT_COUNTER(IoDispatcherForwardSeeks, TEXT("IoDispatcher/ForwardSeeks"));
+TRACE_DECLARE_INT_COUNTER(IoDispatcherBackwardSeeks, TEXT("IoDispatcher/BackwardSeeks"));
+TRACE_DECLARE_MEMORY_COUNTER(IoDispatcherTotalSeekDistance, TEXT("IoDispatcher/TotalSeekDistance"));
+TRACE_DECLARE_INT_COUNTER(IoDispatcherPendingBlocksCount, TEXT("IoDispatcher/PendingBlocksCount"));
 
 FGenericIoDispatcherEventQueue::FGenericIoDispatcherEventQueue()
 	: Event(FPlatformProcess::GetSynchEventFromPool())
 {
-	
 }
 
 FGenericIoDispatcherEventQueue::~FGenericIoDispatcherEventQueue()
@@ -23,34 +32,34 @@ void FGenericIoDispatcherEventQueue::Notify()
 	Event->Trigger();
 }
 
-void FGenericIoDispatcherEventQueue::ProcessEvents()
+void FGenericIoDispatcherEventQueue::Wait()
 {
 	Event->Wait();
 }
 
 FGenericFileIoStoreImpl::FGenericFileIoStoreImpl(FGenericIoDispatcherEventQueue& InEventQueue)
 	: EventQueue(InEventQueue)
+	, PendingBlockEvent(FPlatformProcess::GetSynchEventFromPool())
+{
+	Thread = FRunnableThread::Create(this, TEXT("IoService"), 0, TPri_AboveNormal);
+}
+
+FGenericFileIoStoreImpl::~FGenericFileIoStoreImpl()
 {
-	
+	delete Thread;
+	FPlatformProcess::ReturnSynchEventToPool(PendingBlockEvent);
 }
 
 bool FGenericFileIoStoreImpl::OpenContainer(const TCHAR* ContainerFilePath, uint64& ContainerFileHandle, uint64& ContainerFileSize)
 {
 	IPlatformFile& Ipf = FPlatformFileManager::Get().GetPlatformFile();
-	{
-		TUniquePtr<IFileHandle> FileHandle(Ipf.OpenRead(ContainerFilePath));
-		if (!FileHandle)
-		{
-			return false;
-		}
-		ContainerFileSize = FileHandle->Size();
-	}
-	IAsyncReadFileHandle* AsyncFileHandle = Ipf.OpenAsyncRead(ContainerFilePath);
-	if (!AsyncFileHandle)
+	IFileHandle* FileHandle = Ipf.OpenReadNoBuffering(ContainerFilePath);
+	if (!FileHandle)
 	{
 		return false;
 	}
-	ContainerFileHandle = reinterpret_cast<UPTRINT>(AsyncFileHandle);
+	ContainerFileHandle = reinterpret_cast<UPTRINT>(FileHandle);
+	ContainerFileSize = FileHandle->Size();
 	return true;
 }
 
@@ -62,20 +71,23 @@ void FGenericFileIoStoreImpl::BeginReadsForRequest(FFileIoStoreResolvedRequest&
 	}
 }
 
-void FGenericFileIoStoreImpl::ReadBlockFromFile(FFileIoStoreReadBlock* Block, uint8* Destination, uint64 FileHandle, uint64 Size, uint64 Offset)
+void FGenericFileIoStoreImpl::ReadBlockFromFile(FFileIoStoreReadBlock* Block)
 {
-	FAsyncFileCallBack ReadCallback = [this, Block](bool bWasCancelled, IAsyncReadRequest* Request)
 	{
-		check(!bWasCancelled);
+		FScopeLock Lock(&PendingBlocksCritical);
+		if (!PendingBlocksHead)
 		{
-			FScopeLock _(&CompletedReadRequestsCritical);
-			CompletedReadRequests.AddTail(MakeTuple(Request, Block));
+			PendingBlocksHead = PendingBlocksTail = Block;
 		}
-		EventQueue.Notify();
-	};
-
-	IAsyncReadFileHandle* AsyncReadFileHandle = reinterpret_cast<IAsyncReadFileHandle*>(static_cast<UPTRINT>(FileHandle));
-	AsyncReadFileHandle->ReadRequest(Offset, Size, AIOP_Normal, &ReadCallback, Destination);
+		else
+		{
+			PendingBlocksTail->Next = Block;
+			PendingBlocksTail = Block;
+		}
+		Block->Next = nullptr;
+		TRACE_COUNTER_INCREMENT(IoDispatcherPendingBlocksCount);
+	}
+	PendingBlockEvent->Trigger();
 }
 
 void FGenericFileIoStoreImpl::EndReadsForRequest()
@@ -85,23 +97,120 @@ void FGenericFileIoStoreImpl::EndReadsForRequest()
 
 FFileIoStoreReadBlock* FGenericFileIoStoreImpl::GetNextCompletedBlock()
 {
-	IAsyncReadRequest* AsyncReadRequest;
-	FFileIoStoreReadBlock* CompletedBlock;
-	TDoubleLinkedList<TTuple<IAsyncReadRequest*, FFileIoStoreReadBlock*>>::TDoubleLinkedListNode* CompletedNode;
+	FScopeLock _(&CompletedBlocksCritical);
+	FFileIoStoreReadBlock* CompletedBlock = CompletedBlocksHead;
+	if (!CompletedBlocksHead)
 	{
-		FScopeLock _(&CompletedReadRequestsCritical);
-		CompletedNode = CompletedReadRequests.GetHead();
-		if (!CompletedNode)
-		{
-			return nullptr;
-		}
-		TTuple<IAsyncReadRequest*, FFileIoStoreReadBlock*> NodeValue = CompletedNode->GetValue();
-		AsyncReadRequest = NodeValue.Get<0>();
-		CompletedBlock = NodeValue.Get<1>();
-		CompletedReadRequests.RemoveNode(CompletedNode);
+		return nullptr;
+	}
+	CompletedBlocksHead = CompletedBlocksHead->Next;
+	if (!CompletedBlocksHead)
+	{
+		CompletedBlocksTail = nullptr;
 	}
-	AsyncReadRequest->WaitCompletion();
-	delete AsyncReadRequest;
 	return CompletedBlock;
 }
 
+bool FGenericFileIoStoreImpl::Init()
+{
+	return true;
+}
+
+void FGenericFileIoStoreImpl::Stop()
+{
+	bStopRequested = true;
+	PendingBlockEvent->Trigger();
+}
+
+uint32 FGenericFileIoStoreImpl::Run()
+{
+	FFileIoStoreReadBlock* ScheduledBlocksHead = nullptr;
+	FFileIoStoreReadBlock* ScheduledBlocksTail = nullptr;
+
+	for (;;)
+	{
+		for (;;)
+		{
+			if (bStopRequested)
+			{
+				break;
+			}
+			{
+				FScopeLock PendingBlocksLock(&PendingBlocksCritical);
+				if (PendingBlocksHead)
+				{
+					if (!ScheduledBlocksTail)
+					{
+						ScheduledBlocksHead = PendingBlocksHead;
+						ScheduledBlocksTail = PendingBlocksTail;
+					}
+					else
+					{
+						ScheduledBlocksTail->Next = PendingBlocksHead;
+						ScheduledBlocksTail = PendingBlocksTail;
+					}
+					PendingBlocksHead = PendingBlocksTail = nullptr;
+				}
+			}
+
+			if (!ScheduledBlocksHead)
+			{
+				break;
+			}
+
+			FFileIoStoreReadBlock* BlockToRead = ScheduledBlocksHead;
+			ScheduledBlocksHead = ScheduledBlocksHead->Next;
+			if (!ScheduledBlocksHead)
+			{
+				ScheduledBlocksTail = nullptr;
+			}
+			
+			if (!BlockToRead->Buffer.DataSize())
+			{
+				BlockToRead->Buffer = FIoBuffer(BlockToRead->Size);
+			}
+
+			/*FString ScopeNameString;
+			ScopeNameString.Appendf(TEXT("ReadBlockFromFile: %lld (%lld)"), BlockToRead->Key.BlockIndex, BlockToRead->Size);
+			TRACE_CPUPROFILER_EVENT_SCOPE_TEXT(*ScopeNameString);
+			*/
+			TRACE_CPUPROFILER_EVENT_SCOPE(ReadBlockFromFile);
+			IFileHandle* FileHandle = reinterpret_cast<IFileHandle*>(static_cast<UPTRINT>(BlockToRead->Key.FileHandle));
+			if (FileHandle->Tell() != BlockToRead->Offset)
+			{
+				if (uint64(FileHandle->Tell()) > BlockToRead->Offset)
+				{
+					TRACE_COUNTER_INCREMENT(IoDispatcherBackwardSeeks);
+				}
+				else
+				{
+					TRACE_COUNTER_INCREMENT(IoDispatcherForwardSeeks);
+				}
+				TRACE_COUNTER_ADD(IoDispatcherTotalSeekDistance, FMath::Abs(FileHandle->Tell() - int64(BlockToRead->Offset)));
+				FileHandle->Seek(BlockToRead->Offset);
+			}
+			else
+			{
+				TRACE_COUNTER_INCREMENT(IoDispatcherSequentialReads);
+			}
+			FileHandle->Read(BlockToRead->Buffer.Data(), BlockToRead->Size);
+			TRACE_COUNTER_DECREMENT(IoDispatcherPendingBlocksCount);
+			{
+				FScopeLock _(&CompletedBlocksCritical);
+				if (!CompletedBlocksHead)
+				{
+					CompletedBlocksHead = CompletedBlocksTail = BlockToRead;
+				}
+				else
+				{
+					CompletedBlocksTail->Next = BlockToRead;
+					CompletedBlocksTail = BlockToRead;
+				}
+				BlockToRead->Next = nullptr;
+			}
+			EventQueue.Notify();
+		}
+		PendingBlockEvent->Wait();
+	}
+	return 0;
+}
diff --git a/Engine/Source/Runtime/Core/Private/GenericPlatform/GenericPlatformIoDispatcher.h b/Engine/Source/Runtime/Core/Private/GenericPlatform/GenericPlatformIoDispatcher.h
--- a/Engine/Source/Runtime/Core/Private/GenericPlatform/GenericPlatformIoDispatcher.h
+++ b/Engine/Source/Runtime/Core/Private/GenericPlatform/GenericPlatformIoDispatcher.h
@@ -5,6 +5,7 @@
 #include "Containers/List.h"
 #include "HAL/CriticalSection.h"
 #include "Templates/Tuple.h"
+#include "HAL/Runnable.h"
 
 struct FFileIoStoreReadBlock;
 struct FFileIoStoreResolvedRequest;
@@ -16,25 +17,39 @@ public:
 	FGenericIoDispatcherEventQueue();
 	~FGenericIoDispatcherEventQueue();
 	void Notify();
-	void ProcessEvents();
+	void Wait();
+	void Poll() {};
 
 private:
 	FEvent* Event;
 };
 
 class FGenericFileIoStoreImpl
+	: public FRunnable
 {
 public:
 	FGenericFileIoStoreImpl(FGenericIoDispatcherEventQueue& InEventQueue);
+	~FGenericFileIoStoreImpl();
 	bool OpenContainer(const TCHAR* ContainerFilePath, uint64& ContainerFileHandle, uint64& ContainerFileSize);
 	void BeginReadsForRequest(FFileIoStoreResolvedRequest& ResolvedRequest);
-	void ReadBlockFromFile(FFileIoStoreReadBlock* Block, uint8* Destination, uint64 FileHandle, uint64 Size, uint64 Offset);
+	void ReadBlockFromFile(FFileIoStoreReadBlock* Block);
 	void EndReadsForRequest();
 	FFileIoStoreReadBlock* GetNextCompletedBlock();
+	virtual bool Init() override;
+	virtual uint32 Run() override;
+	virtual void Stop() override;
 
 private:
 	FGenericIoDispatcherEventQueue& EventQueue;
-	FCriticalSection CompletedReadRequestsCritical;
-	TDoubleLinkedList<TTuple<IAsyncReadRequest*, FFileIoStoreReadBlock*>> CompletedReadRequests;
+
+	FCriticalSection PendingBlocksCritical;
+	FFileIoStoreReadBlock* PendingBlocksHead = nullptr;
+	FFileIoStoreReadBlock* PendingBlocksTail = nullptr;
+	FCriticalSection CompletedBlocksCritical;
+	FFileIoStoreReadBlock* CompletedBlocksHead = nullptr;
+	FFileIoStoreReadBlock* CompletedBlocksTail = nullptr;
+	FEvent* PendingBlockEvent;
+	FRunnableThread* Thread;
+	TAtomic<bool> bStopRequested{ false };
 };
 
diff --git a/Engine/Source/Runtime/Core/Private/IO/IoDispatcher.cpp b/Engine/Source/Runtime/Core/Private/IO/IoDispatcher.cpp
--- a/Engine/Source/Runtime/Core/Private/IO/IoDispatcher.cpp
+++ b/Engine/Source/Runtime/Core/Private/IO/IoDispatcher.cpp
@@ -275,10 +275,20 @@ public:
 private:
 	friend class FIoBatch;
 
-	void CompleteRequests()
+	void ProcessCompletedBlocks()
+	{
+		EventQueue.Poll();
+		while (FileIoStore.ProcessCompletedBlock())
+		{
+			ProcessCompletedRequests();
+		}
+	}
+
+	void ProcessCompletedRequests()
 	{
 		while (SubmittedRequestsHead && SubmittedRequestsHead->UnfinishedReadsCount == 0)
 		{
+			//TRACE_CPUPROFILER_EVENT_SCOPE(CompleteRequest);
 			FIoRequestImpl* NextRequest = SubmittedRequestsHead->NextRequest;
 			CompleteRequest(SubmittedRequestsHead);
 			SubmittedRequestsHead = NextRequest;
@@ -314,36 +324,43 @@ private:
 
 	void ProcessIncomingRequests()
 	{
-		FIoRequestImpl* WaitingRequest;
-		{
-			FScopeLock _(&WaitingLock);
-			WaitingRequest = WaitingRequestsHead;
-			WaitingRequestsHead = WaitingRequestsTail = nullptr;
-		}
-		while (WaitingRequest)
-		{
-			RequestsToSubmit.Add(WaitingRequest);
-			WaitingRequest = WaitingRequest->NextRequest;
-		}
-
-		int32 RequestsToSubmitCount = RequestsToSubmit.Num();
-		if (!RequestsToSubmitCount)
+		FIoRequestImpl* RequestsToSubmitHead = nullptr;
+		FIoRequestImpl* RequestsToSubmitTail = nullptr;
+		//TRACE_CPUPROFILER_EVENT_SCOPE(ProcessIncomingRequests);
+		for (;;)
 		{
-			return;
-		}
+			{
+				FScopeLock _(&WaitingLock);
+				if (WaitingRequestsHead)
+				{
+					if (RequestsToSubmitTail)
+					{
+						RequestsToSubmitTail->NextRequest = WaitingRequestsHead;
+						RequestsToSubmitTail = WaitingRequestsTail;
+					}
+					else
+					{
+						RequestsToSubmitHead = WaitingRequestsHead;
+						RequestsToSubmitTail = WaitingRequestsTail;
+					}
+					WaitingRequestsHead = WaitingRequestsTail = nullptr;
+				}
+			}
+			if (!RequestsToSubmitHead)
+			{
+				return;
+			}
 
-		//TRACE_CPUPROFILER_EVENT_SCOPE(IoDispatcherSubmitRequests);
+			FIoRequestImpl* Request = RequestsToSubmitHead;
+			RequestsToSubmitHead = RequestsToSubmitHead->NextRequest;
+			if (!RequestsToSubmitHead)
+			{
+				RequestsToSubmitTail = nullptr;
+			}
 
-		for (; CurrentRequestsToSubmitIndex < RequestsToSubmitCount; ++CurrentRequestsToSubmitIndex)
-		{
-			FIoRequestImpl* Request = RequestsToSubmit[CurrentRequestsToSubmitIndex];
-			check(Request);
+			//TRACE_CPUPROFILER_EVENT_SCOPE(ResolveRequest);
 
 			EIoStoreResolveResult Result = FileIoStore.Resolve(Request);
-			if (Result == IoStoreResolveResult_Stalled)
-			{
-				break;
-			}
 			if (Result == IoStoreResolveResult_NotFound)
 			{
 				Request->Status = FIoStatus(EIoErrorCode::NotFound);
@@ -358,12 +375,8 @@ private:
 				SubmittedRequestsTail = Request;
 			}
 			Request->NextRequest = nullptr;
-		}
 
-		if (CurrentRequestsToSubmitIndex == RequestsToSubmitCount)
-		{
-			RequestsToSubmit.Reset(false);
-			CurrentRequestsToSubmitIndex = 0;
+			ProcessCompletedBlocks();
 		}
 	}
 
@@ -376,11 +389,12 @@ private:
 	{
 		while (!bStopRequested)
 		{
-			EventQueue.ProcessEvents();
-
-			FileIoStore.ProcessIncomingBlocks();
+			EventQueue.Wait();
+			
+			TRACE_CPUPROFILER_EVENT_SCOPE(ProcessEventQueue);
 			ProcessIncomingRequests();
-			CompleteRequests();
+			ProcessCompletedBlocks();
+			ProcessCompletedRequests();
 		}
 		return 0;
 	}
@@ -403,8 +417,6 @@ private:
 	FCriticalSection WaitingLock;
 	FIoRequestImpl* WaitingRequestsHead = nullptr;
 	FIoRequestImpl* WaitingRequestsTail = nullptr;
-	TArray<FIoRequestImpl*> RequestsToSubmit;
-	int32 CurrentRequestsToSubmitIndex = 0;
 	FIoRequestImpl* SubmittedRequestsHead = nullptr;
 	FIoRequestImpl* SubmittedRequestsTail = nullptr;
 	TAtomic<bool> bStopRequested { false };
diff --git a/Engine/Source/Runtime/Core/Private/IO/IoDispatcherFileBackend.cpp b/Engine/Source/Runtime/Core/Private/IO/IoDispatcherFileBackend.cpp
--- a/Engine/Source/Runtime/Core/Private/IO/IoDispatcherFileBackend.cpp
+++ b/Engine/Source/Runtime/Core/Private/IO/IoDispatcherFileBackend.cpp
@@ -5,11 +5,30 @@
 #include "ProfilingDebugging/CountersTrace.h"
 #include "HAL/PlatformFilemanager.h"
 #include "GenericPlatform/GenericPlatformFile.h"
+#include "HAL/IConsoleManager.h"
 
+TRACE_DECLARE_MEMORY_COUNTER(IoDispatcherTotalBytesRead, TEXT("IoDispatcher/TotalBytesRead"));
+TRACE_DECLARE_MEMORY_COUNTER(IoDispatcherTotalBytesScattered, TEXT("IoDispatcher/TotalBytesScattered"));
 TRACE_DECLARE_INT_COUNTER(IoDispatcherCacheHitsCold, TEXT("IoDispatcher/CacheHitsCold"));
 TRACE_DECLARE_INT_COUNTER(IoDispatcherCacheHitsHot, TEXT("IoDispatcher/CacheHitsHot"));
 TRACE_DECLARE_INT_COUNTER(IoDispatcherCacheMisses, TEXT("IoDispatcher/CacheMisses"));
 
+//PRAGMA_DISABLE_OPTIMIZATION
+
+int32 GIoDispatcherBlockSizeKB = 256;
+/*static FAutoConsoleVariableRef CVar_IoDispatcherBlockSizeKB(
+	TEXT("s.IoDispatcherBlockSizeKB"),
+	GIoDispatcherBlockSizeKB,
+	TEXT("IoDispatcher read block size (in kilobytes).")
+);*/
+
+int32 GIoDispatcherCacheSizeMB = 32;
+static FAutoConsoleVariableRef CVar_IoDispatcherCacheSizeMB(
+	TEXT("s.IoDispatcherCacheSizeMB"),
+	GIoDispatcherCacheSizeMB,
+	TEXT("IoDispatcher cache size (in megabytes).")
+);
+
 FFileIoStoreReader::FFileIoStoreReader(FFileIoStoreImpl& InPlatformImpl)
 	: PlatformImpl(InPlatformImpl)
 {
@@ -135,8 +154,10 @@ bool FFileIoStoreReader::Resolve(FFileIoStoreResolvedRequest& ResolvedRequest)
 
 FFileIoStore::FFileIoStore(FIoDispatcherEventQueue& InEventQueue)
 	: PlatformImpl(InEventQueue)
+	, CacheBlockSize(GIoDispatcherBlockSizeKB > 0 ? uint64(GIoDispatcherBlockSizeKB) << 10 : 256 << 10)
 {
-	InitCache();
+	LruHead.LruNext = &LruTail;
+	LruTail.LruPrev = &LruHead;
 }
 
 FIoStatus FFileIoStore::Mount(const FIoStoreEnvironment& Environment)
@@ -160,12 +181,6 @@ EIoStoreResolveResult FFileIoStore::Resolve(FIoRequestImpl* Request)
 	{
 		if (Reader->Resolve(ResolvedRequest))
 		{
-			// Halt IO request submission if the two least recently used blocks are still waiting for their IO
-			if (!(LruTail.LruPrev->bIsReady && LruTail.LruPrev->LruPrev->bIsReady))
-			{
-				return IoStoreResolveResult_Stalled;
-			}
-
 			Request->UnfinishedReadsCount = 0;
 			if (ResolvedRequest.ResolvedSize > 0)
 			{
@@ -232,47 +247,51 @@ bool FFileIoStore::IsValidEnvironment(const FIoStoreEnvironment& Environment)
 	return FPlatformFileManager::Get().GetPlatformFile().FileExists(*TocFilePath);
 }
 
-void FFileIoStore::ProcessIncomingBlocks()
+bool FFileIoStore::ProcessCompletedBlock()
 {
-	while (FFileIoStoreReadBlock* CompletedBlock = PlatformImpl.GetNextCompletedBlock())
+	const uint64 CacheMemorySize = GIoDispatcherCacheSizeMB > 0 ? uint64(GIoDispatcherCacheSizeMB) << 20 : 0;
+	FFileIoStoreReadBlock* CompletedBlock = PlatformImpl.GetNextCompletedBlock();
+	if (!CompletedBlock)
 	{
-		check(!CompletedBlock->bIsReady);
-		CompletedBlock->bIsReady = true;
-		for (FFileIoStoreReadBlockScatter& Scatter : CompletedBlock->ScatterList)
-		{
-			FMemory::Memcpy(Scatter.Dst, Scatter.Src, Scatter.Size);
-			check(Scatter.Request->UnfinishedReadsCount > 0);
-			--Scatter.Request->UnfinishedReadsCount;
-		}
-		if (CacheBlocks.GetData() <= CompletedBlock && CompletedBlock < CacheBlocks.GetData() + CacheBlockCount)
+		return false;
+	}
+	check(!CompletedBlock->bIsReady);
+	CompletedBlock->bIsReady = true;
+	TRACE_COUNTER_ADD(IoDispatcherTotalBytesRead, CompletedBlock->Size);
+	for (FFileIoStoreReadBlockScatter& Scatter : CompletedBlock->ScatterList)
+	{
+		if (Scatter.DstOffset != MAX_uint64)
 		{
-			CompletedBlock->ScatterList.Empty();
+			FMemory::Memcpy(Scatter.Request->IoBuffer.Data() + Scatter.DstOffset, CompletedBlock->Buffer.Data() + Scatter.SrcOffset, Scatter.Size);
 		}
-		else
+		TRACE_COUNTER_ADD(IoDispatcherTotalBytesScattered, Scatter.Size);
+		check(Scatter.Request->UnfinishedReadsCount > 0);
+		--Scatter.Request->UnfinishedReadsCount;
+	}
+	//CompletedBlock->ScatterList.Empty();
+
+	if (CompletedBlock->LruPrev)
+	{
+		CurrentCacheUsage += CompletedBlock->Size;
+
+		FFileIoStoreReadBlock* EvictionCandidate = LruTail.LruPrev;
+		while (CurrentCacheUsage > CacheMemorySize && EvictionCandidate != &LruHead && EvictionCandidate->bIsReady)
 		{
-			delete CompletedBlock;
+			FFileIoStoreReadBlock* NextEvictionCandidate = EvictionCandidate->LruPrev;
+			EvictionCandidate->LruNext->LruPrev = EvictionCandidate->LruPrev;
+			EvictionCandidate->LruPrev->LruNext = EvictionCandidate->LruNext;
+			CurrentCacheUsage -= EvictionCandidate->Size;
+			CachedBlocksMap.Remove(EvictionCandidate->Key);
+			delete EvictionCandidate;
+			EvictionCandidate = NextEvictionCandidate;
 		}
 	}
-}
-
-void FFileIoStore::InitCache()
-{
-	uint8* CacheMemoryBlock = static_cast<uint8*>(FMemory::Malloc(CacheMemorySize));
-	CacheBlocks.AddDefaulted(CacheBlockCount);
-	FFileIoStoreReadBlock* PreviousBlock = &LruHead;
-	for (uint32 BlockIndex = 0; BlockIndex < CacheBlockCount; ++BlockIndex)
+	else
 	{
-		FFileIoStoreReadBlock& ReadBlock = CacheBlocks[BlockIndex];
-		PreviousBlock->LruNext = &ReadBlock;
-		ReadBlock.LruPrev = PreviousBlock;
-		PreviousBlock = &ReadBlock;
-
-		ReadBlock.bIsReady = true;
-		ReadBlock.Buffer = CacheMemoryBlock;
-		CacheMemoryBlock += CacheBlockSize;
+		CachedBlocksMap.Remove(CompletedBlock->Key);
+		delete CompletedBlock;
 	}
-	PreviousBlock->LruNext = &LruTail;
-	LruTail.LruPrev = PreviousBlock;
+	return true;
 }
 
 void FFileIoStore::ReadBlockCached(uint32 BlockIndex, const FFileIoStoreResolvedRequest& ResolvedRequest)
@@ -280,30 +299,44 @@ void FFileIoStore::ReadBlockCached(uint32 BlockIndex, const FFileIoStoreResolved
 	FFileIoStoreCacheBlockKey Key;
 	Key.FileHandle = ResolvedRequest.ResolvedFileHandle;
 	Key.BlockIndex = BlockIndex;
-	Key.Hash = HashCombine(GetTypeHash(ResolvedRequest.ResolvedFileHandle), GetTypeHash(BlockIndex));
 	uint64 BlockOffset = uint64(BlockIndex) * uint64(CacheBlockSize);
 	FFileIoStoreReadBlock* CachedBlock = CachedBlocksMap.FindRef(Key);
 	if (!CachedBlock)
 	{
-		CachedBlock = LruTail.LruPrev;
-		check(CachedBlock->bIsReady);
-		CachedBlocksMap.Remove(CachedBlock->Key);
+		uint64 ReadSize = FMath::Min(ResolvedRequest.ResolvedFileSize, BlockOffset + CacheBlockSize) - BlockOffset;
+
+		CachedBlock = new FFileIoStoreReadBlock();
 		CachedBlock->Key = Key;
 		CachedBlock->bIsReady = false;
+		CachedBlock->Offset = BlockOffset;
+		CachedBlock->Size = ReadSize;
 		CachedBlocksMap.Add(CachedBlock->Key, CachedBlock);
 
-		uint64 ReadSize = FMath::Min(ResolvedRequest.ResolvedFileSize, BlockOffset + CacheBlockSize) - BlockOffset;
-		PlatformImpl.ReadBlockFromFile(CachedBlock, CachedBlock->Buffer, ResolvedRequest.ResolvedFileHandle, ReadSize, BlockOffset);
-
+		PlatformImpl.ReadBlockFromFile(CachedBlock);
 		TRACE_COUNTER_INCREMENT(IoDispatcherCacheMisses);
 	}
+	else
+	{
+		if (CachedBlock->bIsReady)
+		{
+			TRACE_COUNTER_INCREMENT(IoDispatcherCacheHitsHot);
+		}
+		else
+		{
+			TRACE_COUNTER_INCREMENT(IoDispatcherCacheHitsCold);
+		}
+	}
 
-	CachedBlock->LruPrev->LruNext = CachedBlock->LruNext;
-	CachedBlock->LruNext->LruPrev = CachedBlock->LruPrev;
+	if (CachedBlock->LruPrev)
+	{
+		check(CachedBlock->LruNext);
+		CachedBlock->LruPrev->LruNext = CachedBlock->LruNext;
+		CachedBlock->LruNext->LruPrev = CachedBlock->LruPrev;
+	}
 
 	CachedBlock->LruNext = LruHead.LruNext;
 	CachedBlock->LruPrev = &LruHead;
-	CachedBlock->LruNext->LruPrev = CachedBlock;
+	LruHead.LruNext->LruPrev = CachedBlock;
 	LruHead.LruNext = CachedBlock;
 
 	uint64 RequestStartOffsetInBlock = FMath::Max<int64>(0, int64(ResolvedRequest.ResolvedOffset) - BlockOffset);
@@ -315,17 +348,16 @@ void FFileIoStore::ReadBlockCached(uint32 BlockIndex, const FFileIoStoreResolved
 	check(ResolvedRequest.Request->IoBuffer.Data() + BlockOffsetInRequest + RequestSizeInBlock <= ResolvedRequest.Request->IoBuffer.Data() + ResolvedRequest.Request->IoBuffer.DataSize());
 	if (CachedBlock->bIsReady)
 	{
-		TRACE_COUNTER_INCREMENT(IoDispatcherCacheHitsHot);
-		FMemory::Memcpy(ResolvedRequest.Request->IoBuffer.Data() + BlockOffsetInRequest, CachedBlock->Buffer + RequestStartOffsetInBlock, RequestSizeInBlock);
+		FMemory::Memcpy(ResolvedRequest.Request->IoBuffer.Data() + BlockOffsetInRequest, CachedBlock->Buffer.Data() + RequestStartOffsetInBlock, RequestSizeInBlock);
+		TRACE_COUNTER_ADD(IoDispatcherTotalBytesScattered, RequestSizeInBlock);
 	}
 	else
 	{
-		TRACE_COUNTER_INCREMENT(IoDispatcherCacheHitsCold);
 		++ResolvedRequest.Request->UnfinishedReadsCount;
 		FFileIoStoreReadBlockScatter& Scatter = CachedBlock->ScatterList.AddDefaulted_GetRef();
 		Scatter.Request = ResolvedRequest.Request;
-		Scatter.Dst = ResolvedRequest.Request->IoBuffer.Data() + BlockOffsetInRequest;
-		Scatter.Src = CachedBlock->Buffer + RequestStartOffsetInBlock;
+		Scatter.DstOffset = BlockOffsetInRequest;
+		Scatter.SrcOffset = RequestStartOffsetInBlock;
 		check(RequestSizeInBlock <= TNumericLimits<uint32>::Max());
 		Scatter.Size = (uint32)RequestSizeInBlock;
 	}
@@ -339,11 +371,17 @@ void FFileIoStore::ReadBlocksUncached(uint32 BeginBlockIndex, uint32 BlockCount,
 	check(ResolvedRequest.Request->IoBuffer.Data() + BlockOffsetInRequest + ReadSize <= ResolvedRequest.Request->IoBuffer.Data() + ResolvedRequest.Request->IoBuffer.DataSize());
 
 	FFileIoStoreReadBlock* UncachedBlock = new FFileIoStoreReadBlock();
+	UncachedBlock->Size = ReadSize;
+	UncachedBlock->Offset = BlockOffset;
+	UncachedBlock->Key.FileHandle = ResolvedRequest.ResolvedFileHandle;
+	UncachedBlock->Key.BlockIndex = BlockOffset / CacheBlockSize;
+	UncachedBlock->Buffer = FIoBuffer(FIoBuffer::Wrap, ResolvedRequest.Request->IoBuffer.Data() + BlockOffsetInRequest, ReadSize);
+
 	++ResolvedRequest.Request->UnfinishedReadsCount;
 	FFileIoStoreReadBlockScatter& Scatter = UncachedBlock->ScatterList.AddDefaulted_GetRef();
 	Scatter.Request = ResolvedRequest.Request;
-	Scatter.Dst = nullptr;
-	Scatter.Src = nullptr;
-	Scatter.Size = 0;
-	PlatformImpl.ReadBlockFromFile(UncachedBlock, ResolvedRequest.Request->IoBuffer.Data() + BlockOffsetInRequest, ResolvedRequest.ResolvedFileHandle, ReadSize, BlockOffset);
+	Scatter.DstOffset = MAX_uint64;
+	Scatter.SrcOffset = MAX_uint64;
+	Scatter.Size = ReadSize;
+	PlatformImpl.ReadBlockFromFile(UncachedBlock);
 }
diff --git a/Engine/Source/Runtime/Core/Private/IO/IoDispatcherFileBackend.h b/Engine/Source/Runtime/Core/Private/IO/IoDispatcherFileBackend.h
--- a/Engine/Source/Runtime/Core/Private/IO/IoDispatcherFileBackend.h
+++ b/Engine/Source/Runtime/Core/Private/IO/IoDispatcherFileBackend.h
@@ -10,8 +10,7 @@
 struct FFileIoStoreCacheBlockKey
 {
 	uint64 FileHandle;
-	uint32 BlockIndex;
-	uint32 Hash;
+	uint64 BlockIndex;
 
 	friend bool operator==(const FFileIoStoreCacheBlockKey& A, const FFileIoStoreCacheBlockKey& B)
 	{
@@ -20,16 +19,16 @@ struct FFileIoStoreCacheBlockKey
 
 	friend uint32 GetTypeHash(const FFileIoStoreCacheBlockKey& Key)
 	{
-		return Key.Hash;
+		return HashCombine(GetTypeHash(Key.FileHandle), GetTypeHash(Key.BlockIndex));
 	}
 };
 
 struct FFileIoStoreReadBlockScatter
 {
 	FIoRequestImpl* Request = nullptr;
-	uint8* Dst;
-	const uint8* Src;
-	uint32 Size;
+	uint64 DstOffset;
+	uint64 SrcOffset;
+	uint64 Size;
 };
 
 struct FFileIoStoreReadBlock
@@ -38,7 +37,9 @@ struct FFileIoStoreReadBlock
 	FFileIoStoreReadBlock* LruPrev = nullptr;
 	FFileIoStoreReadBlock* LruNext = nullptr;
 	FFileIoStoreCacheBlockKey Key;
-	uint8* Buffer = nullptr;
+	FIoBuffer Buffer;
+	uint64 Size = 0;
+	uint64 Offset = 0;
 	TArray<FFileIoStoreReadBlockScatter> ScatterList;
 	bool bIsReady = false;
 };
@@ -77,7 +78,7 @@ public:
 	EIoStoreResolveResult Resolve(FIoRequestImpl* Request);
 	bool DoesChunkExist(const FIoChunkId& ChunkId) const;
 	TIoStatusOr<uint64> GetSizeForChunk(const FIoChunkId& ChunkId) const;
-	void ProcessIncomingBlocks();
+	bool ProcessCompletedBlock();
 
 	static bool IsValidEnvironment(const FIoStoreEnvironment& Environment);
 
@@ -86,16 +87,13 @@ private:
 	void ReadBlockCached(uint32 BlockIndex, const FFileIoStoreResolvedRequest& ResolvedRequest);
 	void ReadBlocksUncached(uint32 BeginBlockIndex, uint32 BlockCount, FFileIoStoreResolvedRequest& ResolvedRequest);
 
-	static constexpr uint32 CacheMemorySize = 32 << 20;
-	static constexpr uint32 CacheBlockSize = 256 << 10;
-	static constexpr uint32 CacheBlockCount = CacheMemorySize / CacheBlockSize;
-
 	FFileIoStoreImpl PlatformImpl;
 
 	mutable FRWLock IoStoreReadersLock;
 	TArray<FFileIoStoreReader*> IoStoreReaders;
-	TArray<FFileIoStoreReadBlock> CacheBlocks;
 	TMap<FFileIoStoreCacheBlockKey, FFileIoStoreReadBlock*> CachedBlocksMap;
 	FFileIoStoreReadBlock LruHead;
 	FFileIoStoreReadBlock LruTail;
+	const uint64 CacheBlockSize;
+	uint64 CurrentCacheUsage = 0;
 };
diff --git a/Engine/Source/Runtime/Core/Private/IO/IoDispatcherPrivate.h b/Engine/Source/Runtime/Core/Private/IO/IoDispatcherPrivate.h
--- a/Engine/Source/Runtime/Core/Private/IO/IoDispatcherPrivate.h
+++ b/Engine/Source/Runtime/Core/Private/IO/IoDispatcherPrivate.h
@@ -21,7 +21,6 @@ class FIoRequestImpl;
 enum EIoStoreResolveResult
 {
 	IoStoreResolveResult_OK,
-	IoStoreResolveResult_Stalled,
 	IoStoreResolveResult_NotFound,
 };
 
