Add FAndroidMisc::IsSupportedAndroidDevice()

#android

diff --git a/Engine/Build/Android/Java/src/com/epicgames/ue4/GameActivity.java.template b/Engine/Build/Android/Java/src/com/epicgames/ue4/GameActivity.java.template
--- a/Engine/Build/Android/Java/src/com/epicgames/ue4/GameActivity.java.template
+++ b/Engine/Build/Android/Java/src/com/epicgames/ue4/GameActivity.java.template
@@ -477,7 +477,7 @@ public class GameActivity extends $${gameActivitySuperClass}$$ implements Surfac
 	private String EGLErrorStatus = "UNSET";
 	
 	/** Flag for not officially supported devices */
-	public boolean bSupportedDevice = true;
+	public static boolean bSupportedDevice = true;
 
 	/** Used for LocalNotification support*/
 	private boolean localNotificationAppLaunched = false;
diff --git a/Engine/Source/Runtime/Core/Private/Android/AndroidPlatformMisc.cpp b/Engine/Source/Runtime/Core/Private/Android/AndroidPlatformMisc.cpp
--- a/Engine/Source/Runtime/Core/Private/Android/AndroidPlatformMisc.cpp
+++ b/Engine/Source/Runtime/Core/Private/Android/AndroidPlatformMisc.cpp
@@ -1537,6 +1537,40 @@ int32 FAndroidMisc::GetAndroidBuildVersion()
 }
 #endif
 
+#if USE_ANDROID_JNI
+bool FAndroidMisc::IsSupportedAndroidDevice()
+{
+	static bool bChecked = false;
+	static bool bSupported = true;
+
+	if (!bChecked)
+	{
+		bChecked = true;
+
+		JNIEnv* JEnv = AndroidJavaEnv::GetJavaEnv();
+		if (nullptr != JEnv)
+		{
+			jclass Class = AndroidJavaEnv::FindJavaClassGlobalRef("com/epicgames/ue4/GameActivity");
+			if (nullptr != Class)
+			{
+				jfieldID Field = JEnv->GetStaticFieldID(Class, "bSupportedDevice", "Z");
+				if (nullptr != Field)
+				{
+					bSupported = (bool)JEnv->GetStaticBooleanField(Class, Field);
+				}
+				JEnv->DeleteGlobalRef(Class);
+			}
+		}
+	}
+	return bSupported;
+}
+#else
+bool FAndroidMisc::IsSupportedAndroidDevice()
+{
+	return true;
+}
+#endif
+
 bool FAndroidMisc::ShouldDisablePluginAtRuntime(const FString& PluginName)
 {
 #if PLATFORM_ANDROID_ARM64 || PLATFORM_ANDROID_X64
