Add FAndroidMisc::IsSupportedAndroidDevice()

#android

diff --git a/Engine/Source/Runtime/Core/Public/Android/AndroidPlatformMisc.h b/Engine/Source/Runtime/Core/Public/Android/AndroidPlatformMisc.h
--- a/Engine/Source/Runtime/Core/Public/Android/AndroidPlatformMisc.h
+++ b/Engine/Source/Runtime/Core/Public/Android/AndroidPlatformMisc.h
@@ -182,6 +182,7 @@ public:
 #if USE_ANDROID_JNI
 	static int GetAndroidBuildVersion();
 #endif
+	static bool IsSupportedAndroidDevice();
 	static TMap<FString, FString> GetConfigRulesTMap();
 	static FString* GetConfigRulesVariable(const FString& Key);
 
