Add support for throttling Disk IO throughput via the ThrottleDiskIOMBS command line argument when DiskUtilizationTracker is enabled.

diff --git a/Engine/Source/Runtime/Core/Private/HAL/DiskUtilizationTracker.cpp b/Engine/Source/Runtime/Core/Private/HAL/DiskUtilizationTracker.cpp
--- a/Engine/Source/Runtime/Core/Private/HAL/DiskUtilizationTracker.cpp
+++ b/Engine/Source/Runtime/Core/Private/HAL/DiskUtilizationTracker.cpp
@@ -26,6 +26,113 @@ void FDiskUtilizationTracker::UtilizationStats::Dump() const
 	UE_LOG(LogDiskIO, Display, TEXT("Read Throughput: %fMB/s Pct Time Idle: %f%%"), GetReadThrougputMBS(), GetPercentTimeIdle());
 }
 
+float FDiskUtilizationTracker::GetThrottleRateMBS()
+{
+	float ThrottledThroughputMBs = 0.0f;
+	FParse::Value(FCommandLine::Get(), TEXT("-ThrottleDiskIOMBS="), ThrottledThroughputMBs);
+	if (ThrottledThroughputMBs > 0.0f)
+	{
+		UE_LOG(LogDiskIO, Warning, TEXT("Disk IO will be throttled to %fMB/s"), ThrottledThroughputMBs);
+	}
+
+	return ThrottledThroughputMBs;
+}
+
+void FDiskUtilizationTracker::StartRead(uint64 InReadBytes, uint64 InSeekDistance)
+{
+	static bool bBreak = false;
+
+	bool bReset = bResetShortTermStats.AtomicSet(false);
+
+	if (bReset)
+	{
+		ShortTermStats.Reset();
+		bBreak = true;
+	}
+
+	// update total reads
+	LongTermStats.TotalReads++;
+	ShortTermStats.TotalReads++;
+
+	// update seek data
+	if (InSeekDistance > 0)
+	{
+		LongTermStats.TotalSeeks++;
+		ShortTermStats.TotalSeeks++;
+
+		LongTermStats.TotalSeekDistance += InSeekDistance;
+		ShortTermStats.TotalSeekDistance += InSeekDistance;
+	}
+
+	{
+		FScopeLock Lock(&CriticalSection);
+
+		if (InFlightReads == 0)
+		{
+			// if this is the first started read from idle start 
+			ReadStartCycle = FPlatformTime::Cycles64();
+
+			// update idle time (if we've been idle)
+			if (IdleStartCycle > 0)
+			{
+				const double IdleTime = double(ReadStartCycle - IdleStartCycle) * FPlatformTime::GetSecondsPerCycle64();
+
+				LongTermStats.TotalIdleTime += IdleTime;
+				ShortTermStats.TotalIdleTime += bReset ? 0 : IdleTime;
+
+				CSV_CUSTOM_STAT(DiskIO, AccumulatedIdleTime, float(IdleTime), ECsvCustomStatOp::Accumulate);
+			}
+		}
+
+		InFlightBytes += InReadBytes;
+		InFlightReads++;
+	}
+}
+
+void FDiskUtilizationTracker::FinishRead()
+{
+	// if we're the last in flight read update the start idle counter
+	{
+		FScopeLock Lock(&CriticalSection);
+		check(InFlightReads > 0);
+
+		if (--InFlightReads == 0)
+		{
+#if !UE_BUILD_SHIPPING
+			static const float ThrottledThroughputBS = GetThrottleRateMBS() * 1024 * 1024;
+
+			if ((ThrottledThroughputBS > 0.0f) && (LongTermStats.GetReadThrougputBS() > ThrottledThroughputBS))
+			{
+				const double IOTime = double(FPlatformTime::Cycles64() - ReadStartCycle) * FPlatformTime::GetSecondsPerCycle64();
+				const double ThrottledIOTime = ((LongTermStats.TotalBytesRead + InFlightBytes) / ThrottledThroughputBS) - LongTermStats.TotalIOTime;
+
+				if (IOTime < ThrottledIOTime)
+				{
+					FPlatformProcess::Sleep(ThrottledIOTime - IOTime);
+				}
+			}
+#endif // !UE_BUILD_SHIPPING
+
+			IdleStartCycle = FPlatformTime::Cycles64();
+
+			// update our read counters
+			const double IOTime = double(IdleStartCycle - ReadStartCycle) * FPlatformTime::GetSecondsPerCycle64();
+
+			LongTermStats.TotalIOTime += IOTime;
+			ShortTermStats.TotalIOTime += IOTime;
+
+			LongTermStats.TotalBytesRead += InFlightBytes;
+			ShortTermStats.TotalBytesRead += InFlightBytes;
+
+			CSV_CUSTOM_STAT(DiskIO, AccumulatedIOTime, float(IOTime), ECsvCustomStatOp::Accumulate);
+
+			InFlightBytes = 0;
+		}
+
+	}
+	MaybePrint();
+}
+
 void FDiskUtilizationTracker::MaybePrint()
 {
 #if !UE_BUILD_SHIPPING && SPEW_DISK_UTILIZATION
diff --git a/Engine/Source/Runtime/Core/Public/HAL/DiskUtilizationTracker.h b/Engine/Source/Runtime/Core/Public/HAL/DiskUtilizationTracker.h
--- a/Engine/Source/Runtime/Core/Public/HAL/DiskUtilizationTracker.h
+++ b/Engine/Source/Runtime/Core/Public/HAL/DiskUtilizationTracker.h
@@ -34,14 +34,24 @@ struct FDiskUtilizationTracker
 			TotalIdleTime(0.0)
 		{}
 
+		double GetOverallThroughputBS() const
+		{
+			return (TotalIOTime + TotalIdleTime) > 0.0 ? double(TotalBytesRead) / (TotalIOTime + TotalIdleTime) : 0.0;
+		}
+
 		double GetOverallThroughputMBS() const
 		{
-			return double(TotalBytesRead) / (TotalIOTime + TotalIdleTime) / (1024 * 1024);
+			return GetOverallThroughputBS() / (1024.0 * 1024.0);
+		}
+
+		double GetReadThrougputBS() const
+		{
+			return TotalIOTime > 0.0 ? double(TotalBytesRead) / TotalIOTime : 0.0;
 		}
 
 		double GetReadThrougputMBS() const
 		{
-			return double(TotalBytesRead) / (TotalIOTime) / (1024 * 1024);
+			return GetReadThrougputBS() / (1024.0 * 1024.0);
 		}
 
 		double GetTotalIdleTimeInSeconds() const
@@ -104,98 +114,20 @@ struct FDiskUtilizationTracker
 	{
 	}
 
-	void StartRead(uint64 InReadBytes, uint64 InSeekDistance = 0)
-	{
-		static bool bBreak = false;
-
-		bool bReset = bResetShortTermStats.AtomicSet(false);
-
-		if (bReset)
-		{
-			ShortTermStats.Reset();
-			bBreak = true;
-		}
-
-		// update total reads
-		LongTermStats.TotalReads++;
-		ShortTermStats.TotalReads++;
-
-		// update seek data
-		if (InSeekDistance > 0)
-		{
-			LongTermStats.TotalSeeks++;
-			ShortTermStats.TotalSeeks++;
-
-			LongTermStats.TotalSeekDistance += InSeekDistance;
-			ShortTermStats.TotalSeekDistance += InSeekDistance;
-		}
-
-		{
-			FScopeLock Lock(&CriticalSection);
-
-			if (InFlightReads == 0)
-			{
-				// if this is the first started read from idle start 
-				ReadStartCycle = FPlatformTime::Cycles64();
-
-				// update idle time (if we've been idle)
-				if (IdleStartCycle > 0)
-				{
-					const double IdleTime = double(ReadStartCycle - IdleStartCycle) * FPlatformTime::GetSecondsPerCycle64();
-
-					LongTermStats.TotalIdleTime += IdleTime;
-					ShortTermStats.TotalIdleTime += bReset ? 0 : IdleTime;
-
-					CSV_CUSTOM_STAT(DiskIO, AccumulatedIdleTime, float(IdleTime), ECsvCustomStatOp::Accumulate);
-				}
-			}
-
-			InFlightBytes += InReadBytes;
-			InFlightReads++;
-		}
-	}
-
-	void FinishRead()
-	{
-
-		// if we're the last in flight read update the start idle counter
-		{
-			FScopeLock Lock(&CriticalSection);
-			check(InFlightReads > 0);
-
-			if (--InFlightReads == 0)
-			{
-				IdleStartCycle = FPlatformTime::Cycles64();
-
-				// update our read counters
-				const double IOTime = double(IdleStartCycle - ReadStartCycle) * FPlatformTime::GetSecondsPerCycle64();
-
-				LongTermStats.TotalIOTime += IOTime;
-				ShortTermStats.TotalIOTime += IOTime;
-
-				LongTermStats.TotalBytesRead += InFlightBytes;
-				ShortTermStats.TotalBytesRead += InFlightBytes;
-
-				CSV_CUSTOM_STAT(DiskIO, AccumulatedIOTime, float(IOTime), ECsvCustomStatOp::Accumulate);
-
-				InFlightBytes = 0;
-			}
-
-		}
-		MaybePrint();
-	}
+	void StartRead(uint64 InReadBytes, uint64 InSeekDistance = 0);
+	void FinishRead();
 
 	uint32 GetOutstandingRequests() const
 	{
 		return InFlightReads;
 	}
 
-	struct UtilizationStats& GetLongTermStats()
+	const struct UtilizationStats& GetLongTermStats() const
 	{
 		return LongTermStats;
 	}
 
-	struct UtilizationStats& GetShortTermStats()
+	const struct UtilizationStats& GetShortTermStats() const
 	{
 		return ShortTermStats;
 	}
@@ -205,6 +137,8 @@ struct FDiskUtilizationTracker
 		bResetShortTermStats = true;
 	}
 
+private:
+	static float GetThrottleRateMBS();
 	static constexpr float PrintFrequencySeconds = 0.5f;
 
 	void MaybePrint();
