Allow non-native Android callstacks to be stored in crash context's <callstack> property.
#rnx

diff --git a/Engine/Source/Runtime/Core/Private/Android/AndroidPlatformCrashContext.cpp b/Engine/Source/Runtime/Core/Private/Android/AndroidPlatformCrashContext.cpp
--- a/Engine/Source/Runtime/Core/Private/Android/AndroidPlatformCrashContext.cpp
+++ b/Engine/Source/Runtime/Core/Private/Android/AndroidPlatformCrashContext.cpp
@@ -276,6 +276,26 @@ FAndroidCrashContext::FAndroidCrashContext(ECrashContextType InType, const TCHAR
 	}
 }
 
+
+void FAndroidCrashContext::SetOverrideCallstack(const FString& OverrideCallstackIN)
+{
+	OverrideCallstack.Reset();
+	TArray<FString> OutArray;
+	OverrideCallstackIN.ParseIntoArrayLines(OutArray);
+
+	for (const FString& Line : OutArray)
+	{
+		AppendEscapedXMLString(OverrideCallstack, *Line);
+		OverrideCallstack += TEXT("&#xA;");
+		OverrideCallstack += LINE_TERMINATOR;
+	}
+}
+
+const TCHAR* FAndroidCrashContext::GetCallstackProperty() const
+{
+	return *OverrideCallstack;
+}
+
 void FAndroidCrashContext::CaptureCrashInfo()
 {
 	CapturePortableCallStack(0, Context);
diff --git a/Engine/Source/Runtime/Core/Private/Android/AndroidPlatformMisc.cpp b/Engine/Source/Runtime/Core/Private/Android/AndroidPlatformMisc.cpp
--- a/Engine/Source/Runtime/Core/Private/Android/AndroidPlatformMisc.cpp
+++ b/Engine/Source/Runtime/Core/Private/Android/AndroidPlatformMisc.cpp
@@ -1220,13 +1220,16 @@ void FAndroidMisc::TriggerNonFatalCrashHandler(ECrashContextType InType, const F
 	}
 }
 
-void FAndroidMisc::TriggerCrashHandler(const TCHAR* InErrorMessage, const TMap<FString, FString>& AdditionalProperties)
+void FAndroidMisc::TriggerCrashHandler(const TCHAR* InErrorMessage, const TCHAR* OverrideCallstack)
 {
 	FAndroidCrashContext CrashContext(ECrashContextType::Crash, InErrorMessage);
-	CrashContext.CaptureCrashInfo();
-	for (TMap<FString, FString>::TConstIterator Iter(AdditionalProperties); Iter; ++Iter)
+	if(OverrideCallstack)
 	{
-		CrashContext.AddAndroidCrashProperty(*Iter.Key(), *Iter.Value());
+		CrashContext.SetOverrideCallstack(OverrideCallstack);
+	}
+	else
+	{
+		CrashContext.CaptureCrashInfo();
 	}
 
 	if (GCrashHandlerPointer)
diff --git a/Engine/Source/Runtime/Core/Private/GenericPlatform/GenericPlatformCrashContext.cpp b/Engine/Source/Runtime/Core/Private/GenericPlatform/GenericPlatformCrashContext.cpp
--- a/Engine/Source/Runtime/Core/Private/GenericPlatform/GenericPlatformCrashContext.cpp
+++ b/Engine/Source/Runtime/Core/Private/GenericPlatform/GenericPlatformCrashContext.cpp
@@ -575,7 +575,13 @@ void FGenericCrashContext::SerializeContentToBuffer() const
 
 	// Legacy callstack element for current crash reporter
 	AddCrashProperty( TEXT( "NumMinidumpFramesToIgnore"), NumMinidumpFramesToIgnore );
-	AddCrashProperty( TEXT( "CallStack" ), TEXT("") );
+	// Allow platforms to override callstack property, on some platforms the callstack is not captured by native code, those callstacks can be substituted by platform code here.
+	{
+		CommonBuffer += TEXT("<CallStack>");
+		CommonBuffer += GetCallstackProperty();
+		CommonBuffer += TEXT("</CallStack>");
+		CommonBuffer += LINE_TERMINATOR;
+	}
 
 	// Add new portable callstack element with crash stack
 	AddPortableCallStack();
@@ -636,6 +642,11 @@ void FGenericCrashContext::SerializeContentToBuffer() const
 	AddFooter(CommonBuffer);
 }
 
+const TCHAR* FGenericCrashContext::GetCallstackProperty() const
+{
+	return TEXT("");
+}
+
 void FGenericCrashContext::SetNumMinidumpFramesToIgnore(int InNumMinidumpFramesToIgnore)
 {
 	NumMinidumpFramesToIgnore = InNumMinidumpFramesToIgnore;
diff --git a/Engine/Source/Runtime/Core/Public/Android/AndroidPlatformCrashContext.h b/Engine/Source/Runtime/Core/Public/Android/AndroidPlatformCrashContext.h
--- a/Engine/Source/Runtime/Core/Public/Android/AndroidPlatformCrashContext.h
+++ b/Engine/Source/Runtime/Core/Public/Android/AndroidPlatformCrashContext.h
@@ -50,6 +50,8 @@ struct CORE_API FAndroidCrashContext : public FGenericCrashContext
 
 	void AddAndroidCrashProperty(const FString& Key, const FString& Value);
 
+	void SetOverrideCallstack(const FString& OverrideCallstackIN);
+
 	// generate an absolute path to a crash report folder.
 	static void GenerateReportDirectoryName(char(&DirectoryNameOUT)[CrashReportMaxPathSize]);
 
@@ -58,7 +60,14 @@ struct CORE_API FAndroidCrashContext : public FGenericCrashContext
 	/** Async-safe ItoA */
 	static const ANSICHAR* ItoANSI(uint64 Val, uint64 Base, uint32 Len = 0);
 
+protected:
+
+	/** Allow platform implementations to provide a callstack property. Primarily used when non-native code triggers a crash. */
+	virtual const TCHAR* GetCallstackProperty() const;
+
 private:
+	FString OverrideCallstack;
+
 	TMap<FString, FString> AdditionalProperties;
 
 	// The path used by this instance to store the report.
diff --git a/Engine/Source/Runtime/Core/Public/Android/AndroidPlatformMisc.h b/Engine/Source/Runtime/Core/Public/Android/AndroidPlatformMisc.h
--- a/Engine/Source/Runtime/Core/Public/Android/AndroidPlatformMisc.h
+++ b/Engine/Source/Runtime/Core/Public/Android/AndroidPlatformMisc.h
@@ -215,7 +215,7 @@ public:
 	typedef TFunction<void()> OnPauseCallBackType;
 	static OnPauseCallBackType GetOnPauseCallback();
 	static void SetOnPauseCallback(OnPauseCallBackType InOnPauseCallback);
-	static void TriggerCrashHandler(const TCHAR* InErrorMessage, const TMap<FString, FString>& AdditionalProperties);
+	static void TriggerCrashHandler(const TCHAR* InErrorMessage, const TCHAR* OverrideCallstack);
 	static void TriggerNonFatalCrashHandler(enum class ECrashContextType InType, const FString& Message);
 
 	static bool IsInSignalHandler();
diff --git a/Engine/Source/Runtime/Core/Public/GenericPlatform/GenericPlatformCrashContext.h b/Engine/Source/Runtime/Core/Public/GenericPlatform/GenericPlatformCrashContext.h
--- a/Engine/Source/Runtime/Core/Public/GenericPlatform/GenericPlatformCrashContext.h
+++ b/Engine/Source/Runtime/Core/Public/GenericPlatform/GenericPlatformCrashContext.h
@@ -432,6 +432,9 @@ protected:
 	TArray<FCrashStackFrame> CallStack;
 	TArray<FThreadStackFrames> ThreadCallStacks;
 
+	/** Allow platform implementations to provide a callstack property. Primarily used when non-native code triggers a crash. */
+	virtual const TCHAR* GetCallstackProperty() const;
+
 private:
 
 	/** Serializes the session context section of the crash context to a temporary file. */
