Static analysis fix:

warning V792: The 'IsAligned' function located to the right of the operator '&' will be called regardless of the value of the left operand. Perhaps, it is better to use '&&'

diff --git a/Engine/Source/Runtime/Core/Private/HAL/MallocBinnedArena.cpp b/Engine/Source/Runtime/Core/Private/HAL/MallocBinnedArena.cpp
--- a/Engine/Source/Runtime/Core/Private/HAL/MallocBinnedArena.cpp
+++ b/Engine/Source/Runtime/Core/Private/HAL/MallocBinnedArena.cpp
@@ -937,7 +937,7 @@ void* FMallocBinnedArena::ReallocExternal(void* Ptr, SIZE_T NewSize, uint32 Alig
 		check(Ptr); // null is an OS allocation because it will not fall in our VM block
 		uint32 BlockSize = PoolIndexToBlockSize(PoolIndex);
 		if (
-			((NewSize <= BlockSize) & IsAligned(Ptr, Alignment)) && //-V792
+			((int)(NewSize <= BlockSize) & (int)IsAligned(Ptr, Alignment)) &&
 			(PoolIndex == 0 || NewSize > PoolIndexToBlockSize(PoolIndex - 1)))
 		{
 			return Ptr;
