Copying Private-LoadTimes-4.24 stream to Main. Biggest changes are in Materials/Shader memory freezing.

diff --git a/Engine/Source/Runtime/Core/Public/Containers/Map.h b/Engine/Source/Runtime/Core/Public/Containers/Map.h
--- a/Engine/Source/Runtime/Core/Public/Containers/Map.h
+++ b/Engine/Source/Runtime/Core/Public/Containers/Map.h
@@ -17,7 +17,7 @@
 
 #define ExchangeB(A,B) {bool T=A; A=B; B=T;}
 
-template <typename KeyType,typename ValueType>
+template <typename KeyType, typename ValueType>
 using TPair = TTuple<KeyType, ValueType>;
 
 /** An initializer type for pairs that's passed to the pair set when adding a new pair. */
@@ -29,8 +29,8 @@ public:
 	typename TRValueToLValueReference<ValueInitType>::Type Value;
 
 	/** Initialization constructor. */
-	FORCEINLINE TPairInitializer(KeyInitType InKey,ValueInitType InValue)
-		: Key  (InKey  )
+	FORCEINLINE TPairInitializer(KeyInitType InKey, ValueInitType InValue)
+		: Key(InKey)
 		, Value(InValue)
 	{
 	}
@@ -72,7 +72,7 @@ public:
 
 /** Defines how the map's pairs are hashed. */
 template<typename KeyType, typename ValueType, bool bInAllowDuplicateKeys>
-struct TDefaultMapKeyFuncs : BaseKeyFuncs<TPair<KeyType,ValueType>,KeyType,bInAllowDuplicateKeys>
+struct TDefaultMapKeyFuncs : BaseKeyFuncs<TPair<KeyType, ValueType>, KeyType, bInAllowDuplicateKeys>
 {
 	typedef typename TTypeTraits<KeyType>::ConstPointerType KeyInitType;
 	typedef const TPairInitializer<typename TTypeTraits<KeyType>::ConstInitType, typename TTypeTraits<ValueType>::ConstInitType>& ElementInitType;
@@ -81,7 +81,7 @@ struct TDefaultMapKeyFuncs : BaseKeyFuncs<TPair<KeyType,ValueType>,KeyType,bInAl
 	{
 		return Element.Key;
 	}
-	static FORCEINLINE bool Matches(KeyInitType A,KeyInitType B)
+	static FORCEINLINE bool Matches(KeyInitType A, KeyInitType B)
 	{
 		return A == B;
 	}
@@ -103,7 +103,7 @@ struct TDefaultMapHashableKeyFuncs : TDefaultMapKeyFuncs<KeyType, ValueType, bIn
 
 /** 
  * The base class of maps from keys to values.  Implemented using a TSet of key-value pairs with a custom KeyFuncs, 
- * with the same O(1) addition, removal, and finding.
+ * with the same O(1) addition, removal, and finding. 
  *
  * The ByHash() functions are somewhat dangerous but particularly useful in two scenarios:
  * -- Heterogeneous lookup to avoid creating expensive keys like FString when looking up by const TCHAR*.
@@ -122,6 +122,8 @@ class TMapBase
 	friend struct TContainerTraits<TMapBase>;
 
 public:
+	static const bool SupportsFreezeMemoryImage = TAllocatorTraits<SetAllocator>::SupportsFreezeMemoryImage;
+
 	typedef typename TTypeTraits<KeyType  >::ConstPointerType KeyConstPointerType;
 	typedef typename TTypeTraits<KeyType  >::ConstInitType    KeyInitType;
 	typedef typename TTypeTraits<ValueType>::ConstInitType    ValueInitType;
@@ -164,13 +166,13 @@ protected:
 
 public:
 	// Legacy comparison operators.  Note that these also test whether the map's key-value pairs were added in the same order!
-	friend bool LegacyCompareEqual(const TMapBase& A,const TMapBase& B)
+	friend bool LegacyCompareEqual(const TMapBase& A, const TMapBase& B)
 	{
-		return LegacyCompareEqual(A.Pairs,B.Pairs);
+		return LegacyCompareEqual(A.Pairs, B.Pairs);
 	}
-	friend bool LegacyCompareNotEqual(const TMapBase& A,const TMapBase& B)
+	friend bool LegacyCompareNotEqual(const TMapBase& A, const TMapBase& B)
 	{
-		return LegacyCompareNotEqual(A.Pairs,B.Pairs);
+		return LegacyCompareNotEqual(A.Pairs, B.Pairs);
 	}
 
 	/**
@@ -220,11 +222,11 @@ public:
 		Pairs.Empty(ExpectedNumElements);
 	}
 
-    /** Efficiently empties out the map but preserves all allocations and capacities */
-    FORCEINLINE void Reset()
-    {
-        Pairs.Reset();
-    }
+	/** Efficiently empties out the map but preserves all allocations and capacities */
+	FORCEINLINE void Reset()
+	{
+		Pairs.Reset();
+	}
 
 	/** Shrinks the pair set to avoid slack. */
 	FORCEINLINE void Shrink()
@@ -265,9 +267,9 @@ public:
 	template<typename Allocator> int32 GetKeys(TArray<KeyType, Allocator>& OutKeys) const
 	{
 		TSet<KeyType> VisitedKeys;
-		for(typename ElementSetType::TConstIterator It(Pairs);It;++It)
+		for (typename ElementSetType::TConstIterator It(Pairs); It; ++It)
 		{
-			if ( !VisitedKeys.Contains(It->Key) )
+			if (!VisitedKeys.Contains(It->Key))
 			{
 				OutKeys.Add(It->Key);
 				VisitedKeys.Add(It->Key);
@@ -306,10 +308,10 @@ public:
 	 * @param InValue The value to associate with the key.
 	 * @return A reference to the value as stored in the map. The reference is only valid until the next change to any key in the map.
 	 */
-	FORCEINLINE ValueType& Add(const KeyType&  InKey, const ValueType&  InValue) { return Emplace(                   InKey ,                    InValue ); }
-	FORCEINLINE ValueType& Add(const KeyType&  InKey,       ValueType&& InValue) { return Emplace(                   InKey , MoveTempIfPossible(InValue)); }
-	FORCEINLINE ValueType& Add(      KeyType&& InKey, const ValueType&  InValue) { return Emplace(MoveTempIfPossible(InKey),                    InValue ); }
-	FORCEINLINE ValueType& Add(      KeyType&& InKey,       ValueType&& InValue) { return Emplace(MoveTempIfPossible(InKey), MoveTempIfPossible(InValue)); }
+	FORCEINLINE ValueType& Add(const KeyType&  InKey, const ValueType&  InValue) { return Emplace(InKey, InValue); }
+	FORCEINLINE ValueType& Add(const KeyType&  InKey,		ValueType&& InValue) { return Emplace(InKey, MoveTempIfPossible(InValue)); }
+	FORCEINLINE ValueType& Add(		 KeyType&& InKey, const ValueType&  InValue) { return Emplace(MoveTempIfPossible(InKey), InValue); }
+	FORCEINLINE ValueType& Add(		 KeyType&& InKey,		ValueType&& InValue) { return Emplace(MoveTempIfPossible(InKey), MoveTempIfPossible(InValue)); }
 
 	/** See Add() and class documentation section on ByHash() functions */
 	FORCEINLINE ValueType& AddByHash(uint32 KeyHash, const KeyType&  InKey, const ValueType&  InValue) { return EmplaceByHash(KeyHash, InKey, InValue); }
@@ -323,8 +325,8 @@ public:
 	 * @param InKey The key to associate the value with.
 	 * @return A reference to the value as stored in the map. The reference is only valid until the next change to any key in the map.
 	 */
-	FORCEINLINE ValueType& Add(const KeyType&  InKey) { return Emplace(                   InKey ); }
-	FORCEINLINE ValueType& Add(      KeyType&& InKey) { return Emplace(MoveTempIfPossible(InKey)); }
+	FORCEINLINE ValueType& Add(const KeyType&  InKey) { return Emplace(InKey); }
+	FORCEINLINE ValueType& Add(		 KeyType&& InKey) { return Emplace(MoveTempIfPossible(InKey)); }
 
 	/** See Add() and class documentation section on ByHash() functions */
 	FORCEINLINE ValueType& AddByHash(uint32 KeyHash, const KeyType&  InKey) { return EmplaceByHash(KeyHash, InKey); }
@@ -336,8 +338,8 @@ public:
 	 * @param InKeyValue A Tuple containing the Key and Value to associate together
 	 * @return A reference to the value as stored in the map. The reference is only valid until the next change to any key in the map.
 	 */
-	FORCEINLINE ValueType& Add(const TTuple<KeyType,ValueType>&  InKeyValue) { return Emplace(                   InKeyValue.Key ,                    InKeyValue.Value ); }
-	FORCEINLINE ValueType& Add(      TTuple<KeyType,ValueType>&& InKeyValue) { return Emplace(MoveTempIfPossible(InKeyValue.Key), MoveTempIfPossible(InKeyValue.Value)); }
+	FORCEINLINE ValueType& Add(const TTuple<KeyType, ValueType>&  InKeyValue) { return Emplace(InKeyValue.Key, InKeyValue.Value); }
+	FORCEINLINE ValueType& Add(		 TTuple<KeyType, ValueType>&& InKeyValue) { return Emplace(MoveTempIfPossible(InKeyValue.Key), MoveTempIfPossible(InKeyValue.Value)); }
 
 	/**
 	 * Sets the value associated with a key.
@@ -417,9 +419,9 @@ public:
 	 */
 	const KeyType* FindKey(ValueInitType Value) const
 	{
-		for(typename ElementSetType::TConstIterator PairIt(Pairs);PairIt;++PairIt)
+		for (typename ElementSetType::TConstIterator PairIt(Pairs); PairIt; ++PairIt)
 		{
-			if(PairIt->Value == Value)
+			if (PairIt->Value == Value)
 			{
 				return &PairIt->Key;
 			}
@@ -552,7 +554,7 @@ public:
 	FORCEINLINE const ValueType& FindChecked(KeyConstPointerType Key) const
 	{
 		const auto* Pair = Pairs.Find(Key);
-		check( Pair != nullptr );
+		check(Pair != nullptr);
 		return Pair->Value;
 	}
 
@@ -565,7 +567,7 @@ public:
 	FORCEINLINE ValueType& FindChecked(KeyConstPointerType Key)
 	{
 		auto* Pair = Pairs.Find(Key);
-		check( Pair != nullptr );
+		check(Pair != nullptr);
 		return Pair->Value;
 	}
 
@@ -611,7 +613,7 @@ public:
 	template<typename Allocator> void GenerateKeyArray(TArray<KeyType, Allocator>& OutArray) const
 	{
 		OutArray.Empty(Pairs.Num());
-		for(typename ElementSetType::TConstIterator PairIt(Pairs);PairIt;++PairIt)
+		for (typename ElementSetType::TConstIterator PairIt(Pairs); PairIt; ++PairIt)
 		{
 			new(OutArray) KeyType(PairIt->Key);
 		}
@@ -625,14 +627,14 @@ public:
 	template<typename Allocator> void GenerateValueArray(TArray<ValueType, Allocator>& OutArray) const
 	{
 		OutArray.Empty(Pairs.Num());
-		for(typename ElementSetType::TConstIterator PairIt(Pairs);PairIt;++PairIt)
+		for (typename ElementSetType::TConstIterator PairIt(Pairs); PairIt; ++PairIt)
 		{
 			new(OutArray) ValueType(PairIt->Value);
 		}
 	}
 
 	/** Serializer. */
-	FORCEINLINE friend FArchive& operator<<(FArchive& Ar,TMapBase& Map)
+	FORCEINLINE friend FArchive& operator<<(FArchive& Ar, TMapBase& Map)
 	{
 		return Ar << Map.Pairs;
 	}
@@ -697,13 +699,13 @@ protected:
 		typedef typename TChooseClass<
 			bConst,
 			typename TChooseClass<bRangedFor, typename ElementSetType::TRangedForConstIterator, typename ElementSetType::TConstIterator>::Result,
-			typename TChooseClass<bRangedFor, typename ElementSetType::TRangedForIterator,      typename ElementSetType::TIterator     >::Result
+			typename TChooseClass<bRangedFor, typename ElementSetType::TRangedForIterator, typename ElementSetType::TIterator     >::Result
 		>::Result PairItType;
 	private:
-		typedef typename TChooseClass<bConst,const TMapBase,TMapBase>::Result MapType;
-		typedef typename TChooseClass<bConst,const KeyType,KeyType>::Result ItKeyType;
-		typedef typename TChooseClass<bConst,const ValueType,ValueType>::Result ItValueType;
-		typedef typename TChooseClass<bConst,const typename ElementSetType::ElementType, typename ElementSetType::ElementType>::Result PairType;
+		typedef typename TChooseClass<bConst, const TMapBase, TMapBase>::Result MapType;
+		typedef typename TChooseClass<bConst, const KeyType, KeyType>::Result ItKeyType;
+		typedef typename TChooseClass<bConst, const ValueType, ValueType>::Result ItValueType;
+		typedef typename TChooseClass<bConst, const typename ElementSetType::ElementType, typename ElementSetType::ElementType>::Result PairType;
 
 	public:
 		FORCEINLINE TBaseIterator(const PairItType& InElementIt)
@@ -746,9 +748,9 @@ protected:
 	class TBaseKeyIterator
 	{
 	private:
-		typedef typename TChooseClass<bConst,typename ElementSetType::TConstKeyIterator,typename ElementSetType::TKeyIterator>::Result SetItType;
-		typedef typename TChooseClass<bConst,const KeyType,KeyType>::Result ItKeyType;
-		typedef typename TChooseClass<bConst,const ValueType,ValueType>::Result ItValueType;
+		typedef typename TChooseClass<bConst, typename ElementSetType::TConstKeyIterator, typename ElementSetType::TKeyIterator>::Result SetItType;
+		typedef typename TChooseClass<bConst, const KeyType, KeyType>::Result ItKeyType;
+		typedef typename TChooseClass<bConst, const ValueType, ValueType>::Result ItValueType;
 
 	public:
 		/** Initialization constructor. */
@@ -774,7 +776,7 @@ protected:
 			return !(bool)*this;
 		}
 
-		FORCEINLINE ItKeyType&   Key  () const { return SetIt->Key; }
+		FORCEINLINE ItKeyType&   Key() const { return SetIt->Key; }
 		FORCEINLINE ItValueType& Value() const { return SetIt->Value; }
 
 	protected:
@@ -785,6 +787,20 @@ protected:
 	ElementSetType Pairs;
 
 public:
+	void WriteMemoryImage(FMemoryImageWriter& Writer) const
+	{
+		Pairs.WriteMemoryImage(Writer);
+	}
+
+	void CopyUnfrozen(const FMemoryUnfreezeContent& Context, void* Dst) const
+	{
+		Pairs.CopyUnfrozen(Context, Dst);
+	}
+
+	static void AppendHash(const FPlatformTypeLayoutParameters& LayoutParams, FSHA1& Hasher)
+	{
+		ElementSetType::AppendHash(LayoutParams, Hasher);
+	}
 
 	/** Map iterator. */
 	class TIterator : public TBaseIterator<false>
@@ -793,8 +809,8 @@ public:
 
 		/** Initialization constructor. */
 		FORCEINLINE TIterator(TMapBase& InMap, bool bInRequiresRehashOnRemoval = false)
-			: TBaseIterator<false>    (InMap.Pairs.CreateIterator())
-			, Map                     (InMap)
+			: TBaseIterator<false>(InMap.Pairs.CreateIterator())
+			, Map(InMap)
 			, bElementsHaveBeenRemoved(false)
 			, bRequiresRehashOnRemoval(bInRequiresRehashOnRemoval)
 		{
@@ -803,7 +819,7 @@ public:
 		/** Destructor. */
 		FORCEINLINE ~TIterator()
 		{
-			if(bElementsHaveBeenRemoved && bRequiresRehashOnRemoval)
+			if (bElementsHaveBeenRemoved && bRequiresRehashOnRemoval)
 			{
 				Map.Pairs.Relax();
 			}
@@ -832,15 +848,15 @@ public:
 		}
 	};
 
-	using TRangedForIterator      = TBaseIterator<false, true>;
+	using TRangedForIterator = TBaseIterator<false, true>;
 	using TRangedForConstIterator = TBaseIterator<true, true>;
 
 	/** Iterates over values associated with a specified key in a const map. */
 	class TConstKeyIterator : public TBaseKeyIterator<true>
 	{
 	public:
-		FORCEINLINE TConstKeyIterator(const TMapBase& InMap,KeyInitType InKey)
-		:	TBaseKeyIterator<true>(typename ElementSetType::TConstKeyIterator(InMap.Pairs,InKey))
+		FORCEINLINE TConstKeyIterator(const TMapBase& InMap, KeyInitType InKey)
+			: TBaseKeyIterator<true>(typename ElementSetType::TConstKeyIterator(InMap.Pairs, InKey))
 		{}
 	};
 
@@ -848,8 +864,8 @@ public:
 	class TKeyIterator : public TBaseKeyIterator<false>
 	{
 	public:
-		FORCEINLINE TKeyIterator(TMapBase& InMap,KeyInitType InKey)
-		:	TBaseKeyIterator<false>(typename ElementSetType::TKeyIterator(InMap.Pairs,InKey))
+		FORCEINLINE TKeyIterator(TMapBase& InMap, KeyInitType InKey)
+			: TBaseKeyIterator<false>(typename ElementSetType::TKeyIterator(InMap.Pairs, InKey))
 		{}
 
 		/** Removes the current key-value pair from the map. */
@@ -888,13 +904,12 @@ public:
 	 * DO NOT USE DIRECTLY
 	 * STL-like iterators to enable range-based for loop support.
 	 */
-	FORCEINLINE TRangedForIterator      begin()       { return TRangedForIterator     (Pairs.begin()); }
+	FORCEINLINE TRangedForIterator      begin() { return TRangedForIterator(Pairs.begin()); }
 	FORCEINLINE TRangedForConstIterator begin() const { return TRangedForConstIterator(Pairs.begin()); }
-	FORCEINLINE TRangedForIterator      end  ()       { return TRangedForIterator     (Pairs.end());   }
-	FORCEINLINE TRangedForConstIterator end  () const { return TRangedForConstIterator(Pairs.end());   }
+	FORCEINLINE TRangedForIterator      end() { return TRangedForIterator(Pairs.end()); }
+	FORCEINLINE TRangedForConstIterator end() const { return TRangedForConstIterator(Pairs.end()); }
 };
 
-
 /** The base type of sortable maps. */
 template <typename KeyType, typename ValueType, typename SetAllocator, typename KeyFuncs>
 class TSortableMapBase : public TMapBase<KeyType, ValueType, SetAllocator, KeyFuncs>
@@ -946,9 +961,9 @@ public:
 	 * Invoked using "MyMapVar.KeySort( PREDICATE_CLASS() );"
 	 */
 	template<typename PREDICATE_CLASS>
-	FORCEINLINE void KeySort( const PREDICATE_CLASS& Predicate )
+	FORCEINLINE void KeySort(const PREDICATE_CLASS& Predicate)
 	{
-		Super::Pairs.Sort( FKeyComparisonClass<PREDICATE_CLASS>( Predicate ) );
+		Super::Pairs.Sort(FKeyComparisonClass<PREDICATE_CLASS>(Predicate));
 	}
 
 	/**
@@ -966,9 +981,9 @@ public:
 	 * Invoked using "MyMapVar.ValueSort( PREDICATE_CLASS() );"
 	 */
 	template<typename PREDICATE_CLASS>
-	FORCEINLINE void ValueSort( const PREDICATE_CLASS& Predicate )
+	FORCEINLINE void ValueSort(const PREDICATE_CLASS& Predicate)
 	{
-		Super::Pairs.Sort( FValueComparisonClass<PREDICATE_CLASS>( Predicate ) );
+		Super::Pairs.Sort(FValueComparisonClass<PREDICATE_CLASS>(Predicate));
 	}
 
 	/**
@@ -992,13 +1007,13 @@ private:
 
 	public:
 
-		FORCEINLINE FKeyComparisonClass( const PREDICATE_CLASS& InPredicate )
-			: Predicate( InPredicate )
+		FORCEINLINE FKeyComparisonClass(const PREDICATE_CLASS& InPredicate)
+			: Predicate(InPredicate)
 		{}
 
-		FORCEINLINE bool operator()( const typename Super::ElementType& A, const typename Super::ElementType& B ) const
+		FORCEINLINE bool operator()(const typename Super::ElementType& A, const typename Super::ElementType& B) const
 		{
-			return Predicate( A.Key, B.Key );
+			return Predicate(A.Key, B.Key);
 		}
 	};
 
@@ -1010,25 +1025,28 @@ private:
 
 	public:
 
-		FORCEINLINE FValueComparisonClass( const PREDICATE_CLASS& InPredicate )
-			: Predicate( InPredicate )
+		FORCEINLINE FValueComparisonClass(const PREDICATE_CLASS& InPredicate)
+			: Predicate(InPredicate)
 		{}
 
-		FORCEINLINE bool operator()( const typename Super::ElementType& A, const typename Super::ElementType& B ) const
+		FORCEINLINE bool operator()(const typename Super::ElementType& A, const typename Super::ElementType& B) const
 		{
-			return Predicate( A.Value, B.Value );
+			return Predicate(A.Value, B.Value);
 		}
 	};
 };
 
-class FScriptMap;
+template <typename AllocatorType, typename InDerivedType = void>
+class TScriptMap;
 
 /** A TMapBase specialization that only allows a single value associated with each key.*/
-template<typename KeyType,typename ValueType,typename SetAllocator /*= FDefaultSetAllocator*/,typename KeyFuncs /*= TDefaultMapHashableKeyFuncs<KeyType,ValueType,false>*/>
+template<typename KeyType, typename ValueType, typename SetAllocator /*= FDefaultSetAllocator*/, typename KeyFuncs /*= TDefaultMapHashableKeyFuncs<KeyType,ValueType,false>*/>
 class TMap : public TSortableMapBase<KeyType, ValueType, SetAllocator, KeyFuncs>
 {
 	friend struct TContainerTraits<TMap>;
-	friend class  FScriptMap;
+
+	template <typename, typename>
+	friend class TScriptMap;
 
 	static_assert(!KeyFuncs::bAllowDuplicateKeys, "TMap cannot be instantiated with a KeyFuncs which allows duplicate keys");
 
@@ -1102,10 +1120,10 @@ public:
 	 * @param OutRemovedValue If found, the value that was removed (not modified if the key was not found)
 	 * @return whether or not the key was found
 	 */
-	FORCEINLINE bool RemoveAndCopyValue(KeyInitType Key,ValueType& OutRemovedValue)
+	FORCEINLINE bool RemoveAndCopyValue(KeyInitType Key, ValueType& OutRemovedValue)
 	{
 		const FSetElementId PairId = Super::Pairs.FindId(Key);
-		if(!PairId.IsValidId())
+		if (!PairId.IsValidId())
 			return false;
 
 		OutRemovedValue = MoveTempIfPossible(Super::Pairs[PairId].Value);
@@ -1164,13 +1182,36 @@ public:
 		}
 	}
 
-	FORCEINLINE       ValueType& operator[](KeyConstPointerType Key)       { return this->FindChecked(Key); }
+	FORCEINLINE       ValueType& operator[](KeyConstPointerType Key) { return this->FindChecked(Key); }
 	FORCEINLINE const ValueType& operator[](KeyConstPointerType Key) const { return this->FindChecked(Key); }
 };
 
+namespace Freeze
+{
+	template <typename KeyType, typename ValueType, typename SetAllocator, typename KeyFuncs>
+	void IntrinsicWriteMemoryImage(FMemoryImageWriter& Writer, const TMap<KeyType, ValueType, SetAllocator, KeyFuncs>& Object, const FTypeLayoutDesc&)
+	{
+		Object.WriteMemoryImage(Writer);
+	}
+
+	template <typename KeyType, typename ValueType, typename SetAllocator, typename KeyFuncs>
+	void IntrinsicUnfrozenCopy(const FMemoryUnfreezeContent& Context, const TMap<KeyType, ValueType, SetAllocator, KeyFuncs>& Object, void* OutDst)
+	{
+		Object.CopyUnfrozen(Context, OutDst);
+	}
+
+	template <typename KeyType, typename ValueType, typename SetAllocator, typename KeyFuncs>
+	uint32 IntrinsicAppendHash(const TMap<KeyType, ValueType, SetAllocator, KeyFuncs>* DummyObject, const FTypeLayoutDesc& TypeDesc, const FPlatformTypeLayoutParameters& LayoutParams, FSHA1& Hasher)
+	{
+		TMap<KeyType, ValueType, SetAllocator, KeyFuncs>::AppendHash(LayoutParams, Hasher);
+		return DefaultAppendHash(TypeDesc, LayoutParams, Hasher);
+	}
+}
+
+DECLARE_TEMPLATE_INTRINSIC_TYPE_LAYOUT((template <typename KeyType, typename ValueType, typename SetAllocator, typename KeyFuncs>), (TMap<KeyType, ValueType, SetAllocator, KeyFuncs>));
 
 /** A TMapBase specialization that allows multiple values to be associated with each key. */
-template<typename KeyType,typename ValueType,typename SetAllocator /* = FDefaultSetAllocator */,typename KeyFuncs /*= TDefaultMapHashableKeyFuncs<KeyType,ValueType,true>*/>
+template<typename KeyType, typename ValueType, typename SetAllocator /* = FDefaultSetAllocator */, typename KeyFuncs /*= TDefaultMapHashableKeyFuncs<KeyType,ValueType,true>*/>
 class TMultiMap : public TSortableMapBase<KeyType, ValueType, SetAllocator, KeyFuncs>
 {
 	friend struct TContainerTraits<TMultiMap>;
@@ -1247,14 +1288,14 @@ public:
 	 * @param OutValues Upon return, contains the values associated with the key.
 	 * @param bMaintainOrder true if the Values array should be in the same order as the map's pairs.
 	 */
-	template<typename Allocator> void MultiFind(KeyInitType Key,TArray<ValueType, Allocator>& OutValues,bool bMaintainOrder = false) const
+	template<typename Allocator> void MultiFind(KeyInitType Key, TArray<ValueType, Allocator>& OutValues, bool bMaintainOrder = false) const
 	{
-		for(typename Super::ElementSetType::TConstKeyIterator It(Super::Pairs,Key);It;++It)
+		for (typename Super::ElementSetType::TConstKeyIterator It(Super::Pairs, Key); It; ++It)
 		{
 			new(OutValues) ValueType(It->Value);
 		}
 
-		if(bMaintainOrder)
+		if (bMaintainOrder)
 		{
 			Algo::Reverse(OutValues);
 		}
@@ -1268,14 +1309,14 @@ public:
 	 *					Pointers are only valid until the next change to any key in the map.
 	 * @param bMaintainOrder true if the Values array should be in the same order as the map's pairs.
 	 */
-	template<typename Allocator> void MultiFindPointer(KeyInitType Key,TArray<const ValueType*, Allocator>& OutValues,bool bMaintainOrder = false) const
+	template<typename Allocator> void MultiFindPointer(KeyInitType Key, TArray<const ValueType*, Allocator>& OutValues, bool bMaintainOrder = false) const
 	{
-		for(typename Super::ElementSetType::TConstKeyIterator It(Super::Pairs,Key);It;++It)
+		for (typename Super::ElementSetType::TConstKeyIterator It(Super::Pairs, Key); It; ++It)
 		{
 			OutValues.Add(&It->Value);
 		}
 
-		if(bMaintainOrder)
+		if (bMaintainOrder)
 		{
 			Algo::Reverse(OutValues);
 		}
@@ -1302,10 +1343,10 @@ public:
 	 * @param InValue The value to associate.
 	 * @return A reference to the value as stored in the map; the reference is only valid until the next change to any key in the map.
 	 */
-	FORCEINLINE ValueType& AddUnique(const KeyType&  InKey, const ValueType&  InValue) { return EmplaceUnique(                   InKey ,                    InValue ); }
-	FORCEINLINE ValueType& AddUnique(const KeyType&  InKey,       ValueType&& InValue) { return EmplaceUnique(                   InKey , MoveTempIfPossible(InValue)); }
-	FORCEINLINE ValueType& AddUnique(      KeyType&& InKey, const ValueType&  InValue) { return EmplaceUnique(MoveTempIfPossible(InKey),                    InValue ); }
-	FORCEINLINE ValueType& AddUnique(      KeyType&& InKey,       ValueType&& InValue) { return EmplaceUnique(MoveTempIfPossible(InKey), MoveTempIfPossible(InValue)); }
+	FORCEINLINE ValueType& AddUnique(const KeyType&  InKey, const ValueType&  InValue) { return EmplaceUnique(InKey, InValue); }
+	FORCEINLINE ValueType& AddUnique(const KeyType&  InKey, ValueType&& InValue) { return EmplaceUnique(InKey, MoveTempIfPossible(InValue)); }
+	FORCEINLINE ValueType& AddUnique(KeyType&& InKey, const ValueType&  InValue) { return EmplaceUnique(MoveTempIfPossible(InKey), InValue); }
+	FORCEINLINE ValueType& AddUnique(KeyType&& InKey, ValueType&& InValue) { return EmplaceUnique(MoveTempIfPossible(InKey), MoveTempIfPossible(InValue)); }
 
 	/**
 	 * Add a key-value association to the map.
@@ -1348,14 +1389,14 @@ public:
 	 * @param InValue The value part of the pair to remove.
 	 * @return The number of associations removed.
 	 */
-	int32 Remove(KeyInitType InKey,ValueInitType InValue)
+	int32 Remove(KeyInitType InKey, ValueInitType InValue)
 	{
 		// Iterate over pairs with a matching key.
 		int32 NumRemovedPairs = 0;
-		for(typename Super::ElementSetType::TKeyIterator It(Super::Pairs,InKey);It;++It)
+		for (typename Super::ElementSetType::TKeyIterator It(Super::Pairs, InKey); It; ++It)
 		{
 			// If this pair has a matching value as well, remove it.
-			if(It->Value == InValue)
+			if (It->Value == InValue)
 			{
 				It.RemoveCurrent();
 				++NumRemovedPairs;
@@ -1371,14 +1412,14 @@ public:
 	 * @param InValue The value part of the pair to remove.
 	 * @return The number of associations removed.
 	 */
-	int32 RemoveSingle(KeyInitType InKey,ValueInitType InValue)
+	int32 RemoveSingle(KeyInitType InKey, ValueInitType InValue)
 	{
 		// Iterate over pairs with a matching key.
 		int32 NumRemovedPairs = 0;
-		for(typename Super::ElementSetType::TKeyIterator It(Super::Pairs,InKey);It;++It)
+		for (typename Super::ElementSetType::TKeyIterator It(Super::Pairs, InKey); It; ++It)
 		{
 			// If this pair has a matching value as well, remove it.
-			if(It->Value == InValue)
+			if (It->Value == InValue)
 			{
 				It.RemoveCurrent();
 				++NumRemovedPairs;
@@ -1398,7 +1439,7 @@ public:
 	 * @return If the map contains a matching association, a pointer to the value in the map is returned.  Otherwise nullptr is returned.
 	 *			The pointer is only valid as long as the map isn't changed.
 	 */
-	FORCEINLINE const ValueType* FindPair(KeyInitType Key,ValueInitType Value) const
+	FORCEINLINE const ValueType* FindPair(KeyInitType Key, ValueInitType Value) const
 	{
 		return const_cast<TMultiMap*>(this)->FindPair(Key, Value);
 	}
@@ -1411,13 +1452,13 @@ public:
 	 * @return If the map contains a matching association, a pointer to the value in the map is returned.  Otherwise nullptr is returned.
 	 *			The pointer is only valid as long as the map isn't changed.
 	 */
-	ValueType* FindPair(KeyInitType Key,ValueInitType Value)
+	ValueType* FindPair(KeyInitType Key, ValueInitType Value)
 	{
 		// Iterate over pairs with a matching key.
-		for(typename Super::ElementSetType::TKeyIterator It(Super::Pairs,Key);It;++It)
+		for (typename Super::ElementSetType::TKeyIterator It(Super::Pairs, Key); It; ++It)
 		{
 			// If the pair's value matches, return a pointer to it.
-			if(It->Value == Value)
+			if (It->Value == Value)
 			{
 				return &It->Value;
 			}
@@ -1431,7 +1472,7 @@ public:
 	{
 		// Iterate over pairs with a matching key.
 		int32 NumMatchingPairs = 0;
-		for(typename Super::ElementSetType::TConstKeyIterator It(Super::Pairs,Key);It;++It)
+		for (typename Super::ElementSetType::TConstKeyIterator It(Super::Pairs, Key); It; ++It)
 		{
 			++NumMatchingPairs;
 		}
@@ -1479,6 +1520,30 @@ public:
 	}
 };
 
+namespace Freeze
+{
+	template <typename KeyType, typename ValueType, typename SetAllocator, typename KeyFuncs>
+	void IntrinsicWriteMemoryImage(FMemoryImageWriter& Writer, const TMultiMap<KeyType, ValueType, SetAllocator, KeyFuncs>& Object, const FTypeLayoutDesc&)
+	{
+		Object.WriteMemoryImage(Writer);
+	}
+
+	template <typename KeyType, typename ValueType, typename SetAllocator, typename KeyFuncs>
+	void IntrinsicUnfrozenCopy(const FMemoryUnfreezeContent& Context, const TMultiMap<KeyType, ValueType, SetAllocator, KeyFuncs>& Object, void* OutDst)
+	{
+		Object.CopyUnfrozen(Context, OutDst);
+	}
+
+	template <typename KeyType, typename ValueType, typename SetAllocator, typename KeyFuncs>
+	uint32 IntrinsicAppendHash(const TMultiMap<KeyType, ValueType, SetAllocator, KeyFuncs>* DummyObject, const FTypeLayoutDesc& TypeDesc, const FPlatformTypeLayoutParameters& LayoutParams, FSHA1& Hasher)
+	{
+		TMultiMap<KeyType, ValueType, SetAllocator, KeyFuncs>::AppendHash(LayoutParams, Hasher);
+		return DefaultAppendHash(TypeDesc, LayoutParams, Hasher);
+	}
+}
+
+DECLARE_TEMPLATE_INTRINSIC_TYPE_LAYOUT((template <typename KeyType, typename ValueType, typename SetAllocator, typename KeyFuncs>), (TMultiMap<KeyType, ValueType, SetAllocator, KeyFuncs>));
+
 
 struct FScriptMapLayout
 {
@@ -1491,8 +1556,11 @@ struct FScriptMapLayout
 
 // Untyped map type for accessing TMap data, like FScriptArray for TArray.
 // Must have the same memory representation as a TMap.
-class FScriptMap
+template <typename AllocatorType, typename InDerivedType>
+class TScriptMap
 {
+	using DerivedType = typename TChooseClass<TIsVoidType<InDerivedType>::Value, TScriptMap, InDerivedType>::Result;
+
 public:
 	static FScriptMapLayout GetScriptLayout(int32 KeySize, int32 KeyAlignment, int32 ValueSize, int32 ValueAlignment)
 	{
@@ -1500,16 +1568,16 @@ public:
 
 		// TPair<Key, Value>
 		FStructBuilder PairStruct;
-		int32 KeyOffset    = PairStruct.AddMember(KeySize,   KeyAlignment);
+		int32 KeyOffset = PairStruct.AddMember(KeySize, KeyAlignment);
 		Result.ValueOffset = PairStruct.AddMember(ValueSize, ValueAlignment);
-		Result.SetLayout   = FScriptSet::GetScriptLayout(PairStruct.GetSize(),  PairStruct.GetAlignment());
+		Result.SetLayout = FScriptSet::GetScriptLayout(PairStruct.GetSize(), PairStruct.GetAlignment());
 
 		checkf(KeyOffset == 0, TEXT("The key inside the TPair is expected to be at the start of the struct"));
 
 		return Result;
 	}
 
-	FScriptMap()
+	TScriptMap()
 	{
 	}
 
@@ -1538,11 +1606,13 @@ public:
 		return Pairs.GetData(Index, Layout.SetLayout);
 	}
 
-	void MoveAssign(FScriptMap& Other, const FScriptMapLayout& Layout)
+	void MoveAssign(DerivedType& Other, const FScriptMapLayout& Layout)
 	{
-		checkSlow(this != &Other);
+		TScriptMap* TypedOther = (TScriptMap*)&Other;
+
+		checkSlow(this != TypedOther);
 		Empty(0, Layout);
-		Pairs.MoveAssign(Other.Pairs, Layout.SetLayout);
+		Pairs.MoveAssign(TypedOther->Pairs, Layout.SetLayout);
 	}
 
 	void Empty(int32 Slack, const FScriptMapLayout& Layout)
@@ -1566,13 +1636,13 @@ public:
 		return Pairs.AddUninitialized(Layout.SetLayout);
 	}
 
-	void Rehash(const FScriptMapLayout& Layout, TFunctionRef<uint32 (const void*)> GetKeyHash)
+	void Rehash(const FScriptMapLayout& Layout, TFunctionRef<uint32(const void*)> GetKeyHash)
 	{
 		Pairs.Rehash(Layout.SetLayout, GetKeyHash);
 	}
 	
 	/** Finds the associated key, value from hash of Key, rather than linearly searching */
-	int32 FindPairIndex(const void* Key, const FScriptMapLayout& MapLayout, TFunctionRef<uint32 (const void*)> GetKeyHash, TFunctionRef<bool (const void*, const void*)> KeyEqualityFn)
+	int32 FindPairIndex(const void* Key, const FScriptMapLayout& MapLayout, TFunctionRef<uint32(const void*)> GetKeyHash, TFunctionRef<bool(const void*, const void*)> KeyEqualityFn)
 	{
 		if (Pairs.Num())
 		{
@@ -1594,7 +1664,7 @@ public:
 	}
 
 	/** Finds the associated value from hash of Key, rather than linearly searching */
-	uint8* FindValue(const void* Key, const FScriptMapLayout& MapLayout, TFunctionRef<uint32 (const void*)> GetKeyHash, TFunctionRef<bool (const void*, const void*)> KeyEqualityFn)
+	uint8* FindValue(const void* Key, const FScriptMapLayout& MapLayout, TFunctionRef<uint32(const void*)> GetKeyHash, TFunctionRef<bool(const void*, const void*)> KeyEqualityFn)
 	{
 		int32 FoundIndex = FindPairIndex(Key, MapLayout, GetKeyHash, KeyEqualityFn);
 		if (FoundIndex != INDEX_NONE)
@@ -1625,15 +1695,15 @@ public:
 			GetKeyHash,
 			KeyEqualityFn,
 			[KeyConstructAndAssignFn, ValueConstructAndAssignFn, Layout](void* NewPair)
-			{
-				KeyConstructAndAssignFn((uint8*)NewPair);
-				ValueConstructAndAssignFn((uint8*)NewPair + Layout.ValueOffset);
-			},
+		{
+			KeyConstructAndAssignFn((uint8*)NewPair);
+			ValueConstructAndAssignFn((uint8*)NewPair + Layout.ValueOffset);
+		},
 			[DestructKeyFn, DestructValueFn, Layout](void* NewPair)
-			{
-				DestructValueFn((uint8*)NewPair + Layout.ValueOffset);
-				DestructKeyFn((uint8*)NewPair);
-			}
+		{
+			DestructValueFn((uint8*)NewPair + Layout.ValueOffset);
+			DestructKeyFn((uint8*)NewPair);
+		}
 		);
 	}
 
@@ -1666,48 +1736,55 @@ public:
 	}
 
 private:
-	FScriptSet Pairs;
+	TScriptSet<AllocatorType> Pairs;
 
 	// This function isn't intended to be called, just to be compiled to validate the correctness of the type.
 	static void CheckConstraints()
 	{
-		typedef FScriptMap        ScriptType;
+		typedef TScriptMap        ScriptType;
 		typedef TMap<int32, int8> RealType;
 
 		// Check that the class footprint is the same
-		static_assert(sizeof (ScriptType) == sizeof (RealType), "FScriptMap's size doesn't match TMap");
-		static_assert(alignof(ScriptType) == alignof(RealType), "FScriptMap's alignment doesn't match TMap");
+		static_assert(sizeof(ScriptType) == sizeof(RealType), "TScriptMap's size doesn't match TMap");
+		static_assert(alignof(ScriptType) == alignof(RealType), "TScriptMap's alignment doesn't match TMap");
 
 		// Check member sizes
-		static_assert(sizeof(DeclVal<ScriptType>().Pairs) == sizeof(DeclVal<RealType>().Pairs), "FScriptMap's Pairs member size does not match TMap's");
+		static_assert(sizeof(DeclVal<ScriptType>().Pairs) == sizeof(DeclVal<RealType>().Pairs), "TScriptMap's Pairs member size does not match TMap's");
 
 		// Check member offsets
-		static_assert(STRUCT_OFFSET(ScriptType, Pairs) == STRUCT_OFFSET(RealType, Pairs), "FScriptMap's Pairs member offset does not match TMap's");
+		static_assert(STRUCT_OFFSET(ScriptType, Pairs) == STRUCT_OFFSET(RealType, Pairs), "TScriptMap's Pairs member offset does not match TMap's");
 	}
 
 public:
 	// These should really be private, because they shouldn't be called, but there's a bunch of code
 	// that needs to be fixed first.
-	FScriptMap(const FScriptMap&) { check(false); }
-	void operator=(const FScriptMap&) { check(false); }
+	TScriptMap(const TScriptMap&) { check(false); }
+	void operator=(const TScriptMap&) { check(false); }
 };
 
 
-template <>
-struct TIsZeroConstructType<FScriptMap>
+template <typename AllocatorType>
+struct TIsZeroConstructType<TScriptMap<AllocatorType>>
 {
 	enum { Value = true };
 };
 
+class FScriptMap : public TScriptMap<FDefaultSetAllocator, FScriptMap>
+{
+	using Super = TScriptMap<FDefaultSetAllocator, FScriptMap>;
+
+public:
+	using Super::Super;
+};
 
-template <typename KeyType, typename ValueType, typename SetAllocator,typename KeyFuncs>
+template <typename KeyType, typename ValueType, typename SetAllocator, typename KeyFuncs>
 struct TContainerTraits<TMap<KeyType, ValueType, SetAllocator, KeyFuncs>> : public TContainerTraitsBase<TMap<KeyType, ValueType, SetAllocator, KeyFuncs>>
 {
 	enum { MoveWillEmptyContainer = TContainerTraits<typename TMap<KeyType, ValueType, SetAllocator, KeyFuncs>::ElementSetType>::MoveWillEmptyContainer };
 };
 
 
-template <typename KeyType, typename ValueType, typename SetAllocator,typename KeyFuncs>
+template <typename KeyType, typename ValueType, typename SetAllocator, typename KeyFuncs>
 struct TContainerTraits<TMultiMap<KeyType, ValueType, SetAllocator, KeyFuncs>> : public TContainerTraitsBase<TMultiMap<KeyType, ValueType, SetAllocator, KeyFuncs>>
 {
 	enum { MoveWillEmptyContainer = TContainerTraits<typename TMultiMap<KeyType, ValueType, SetAllocator, KeyFuncs>::ElementSetType>::MoveWillEmptyContainer };
diff --git a/Engine/Source/Runtime/Core/Public/Containers/ResourceArray.h b/Engine/Source/Runtime/Core/Public/Containers/ResourceArray.h
--- a/Engine/Source/Runtime/Core/Public/Containers/ResourceArray.h
+++ b/Engine/Source/Runtime/Core/Public/Containers/ResourceArray.h
@@ -3,14 +3,18 @@
 #pragma once
 
 #include "CoreTypes.h"
+#include "Serialization/MemoryLayout.h"
 
 /**
  * An element type independent interface to the resource array.
  */
 class FResourceArrayInterface
 {
+	DECLARE_EXPORTED_TYPE_LAYOUT(FResourceArrayInterface, CORE_API, Abstract);
 public:
 
+	virtual ~FResourceArrayInterface() {}
+
 	/**
 	 * @return A pointer to the resource data.
 	 */
diff --git a/Engine/Source/Runtime/Core/Public/Containers/ScriptArray.h b/Engine/Source/Runtime/Core/Public/Containers/ScriptArray.h
--- a/Engine/Source/Runtime/Core/Public/Containers/ScriptArray.h
+++ b/Engine/Source/Runtime/Core/Public/Containers/ScriptArray.h
@@ -14,8 +14,9 @@
  * Base dynamic array.
  * An untyped data array; mirrors a TArray's members, but doesn't need an exact C++ type for its elements.
  **/
-class FScriptArray
-	: protected FHeapAllocator::ForAnyElementType
+template <typename AllocatorType>
+class TScriptArray
+	: protected AllocatorType::ForAnyElementType
 {
 public:
 
@@ -91,11 +92,11 @@ public:
 			ResizeTo(ArrayNum, NumBytesPerElement);
 		}
 	}
-	void MoveAssign(FScriptArray& Other, int32 NumBytesPerElement)
+	void MoveAssign(TScriptArray& Other, int32 NumBytesPerElement)
 	{
 		checkSlow(this != &Other);
 		Empty(0, NumBytesPerElement);
-		MoveToEmpty(Other);
+		this->MoveToEmpty(Other);
 		ArrayNum = Other.ArrayNum; Other.ArrayNum = 0;
 		ArrayMax = Other.ArrayMax; Other.ArrayMax = 0;
 	}
@@ -116,7 +117,7 @@ public:
 			NumBytesPerElement
 			);
 	}
-	FScriptArray()
+	TScriptArray()
 	:   ArrayNum( 0 )
 	,	ArrayMax( 0 )
 	{
@@ -163,7 +164,7 @@ public:
 
 protected:
 
-	FScriptArray( int32 InNum, int32 NumBytesPerElement  )
+	TScriptArray( int32 InNum, int32 NumBytesPerElement  )
 	:   ArrayNum( 0 )
 	,	ArrayMax( InNum )
 
@@ -208,6 +209,33 @@ protected:
 			this->ResizeAllocation(ArrayNum, ArrayMax, NumBytesPerElement);
 		}
 	}
+public:
+	// These should really be private, because they shouldn't be called, but there's a bunch of code
+	// that needs to be fixed first.
+	TScriptArray(const TScriptArray&) { check(false); }
+	void operator=(const TScriptArray&) { check(false); }
+};
+
+template<typename AllocatorType> struct TIsZeroConstructType<TScriptArray<AllocatorType>> { enum { Value = true }; };
+
+class FScriptArray : public TScriptArray<FHeapAllocator>
+{
+	using Super = TScriptArray<FHeapAllocator>;
+
+public:
+	FScriptArray() = default;
+
+	void MoveAssign(FScriptArray& Other, int32 NumBytesPerElement)
+	{
+		Super::MoveAssign(Other, NumBytesPerElement);
+	}
+
+protected:
+	FScriptArray(int32 InNum, int32 NumBytesPerElement)
+		: TScriptArray<FHeapAllocator>(InNum, NumBytesPerElement)
+	{
+	}
+
 public:
 	// These should really be private, because they shouldn't be called, but there's a bunch of code
 	// that needs to be fixed first.
@@ -215,5 +243,4 @@ public:
 	void operator=(const FScriptArray&) { check(false); }
 };
 
-
 template<> struct TIsZeroConstructType<FScriptArray> { enum { Value = true }; };
diff --git a/Engine/Source/Runtime/Core/Public/Containers/Set.h b/Engine/Source/Runtime/Core/Public/Containers/Set.h
--- a/Engine/Source/Runtime/Core/Public/Containers/Set.h
+++ b/Engine/Source/Runtime/Core/Public/Containers/Set.h
@@ -18,6 +18,7 @@
 #include "Templates/AreTypesEqual.h"
 #include "Templates/Decay.h"
 #include "Serialization/StructuredArchive.h"
+#include "Serialization/MemoryImageWriter.h"
 #include "ContainersFwd.h"
 
 /**
@@ -80,6 +81,9 @@ FORCEINLINE void MoveByRelocate(T& A, T& B)
 	RelocateConstructItems<T>(&A, &B, 1);
 }
 
+template <typename AllocatorType, typename InDerivedType = void>
+class TScriptSet;
+
 /** Either NULL or an identifier for an element of a set. */
 class FSetElementId
 {
@@ -88,7 +92,8 @@ public:
 	template<typename,typename,typename>
 	friend class TSet;
 
-	friend class FScriptSet;
+	template <typename, typename>
+	friend class TScriptSet;
 
 	/** Default constructor. */
 	FORCEINLINE FSetElementId():
@@ -142,13 +147,25 @@ private:
 	}
 };
 
-/** An element in the set. */
-template <typename InElementType>
-class TSetElement
+// This is just an int32
+DECLARE_INTRINSIC_TYPE_LAYOUT(FSetElementId);
+
+template<typename InElementType, bool bTypeLayout>
+class TSetElementBase
 {
 public:
 	typedef InElementType ElementType;
 
+	FORCEINLINE TSetElementBase() {}
+
+	/** Initialization constructor. */
+	template <typename InitType, typename = typename TEnableIf<!TAreTypesEqual<TSetElementBase, typename TDecay<InitType>::Type>::Value>::Type> explicit FORCEINLINE TSetElementBase(InitType&& InValue) : Value(Forward<InitType>(InValue)) {}
+
+	TSetElementBase(TSetElementBase&&) = default;
+	TSetElementBase(const TSetElementBase&) = default;
+	TSetElementBase& operator=(TSetElementBase&&) = default;
+	TSetElementBase& operator=(const TSetElementBase&) = default;
+
 	/** The element's value. */
 	ElementType Value;
 
@@ -157,13 +174,47 @@ public:
 
 	/** The hash bucket that the element is currently linked to. */
 	mutable int32 HashIndex;
+};
+
+template<typename InElementType>
+class TSetElementBase<InElementType, true>
+{
+	DECLARE_INLINE_TYPE_LAYOUT(TSetElementBase, NonVirtual);
+public:
+	typedef InElementType ElementType;
+
+	FORCEINLINE TSetElementBase() {}
+
+	/** Initialization constructor. */
+	template <typename InitType, typename = typename TEnableIf<!TAreTypesEqual<TSetElementBase, typename TDecay<InitType>::Type>::Value>::Type> explicit FORCEINLINE TSetElementBase(InitType&& InValue) : Value(Forward<InitType>(InValue)) {}
+
+	TSetElementBase(TSetElementBase&&) = default;
+	TSetElementBase(const TSetElementBase&) = default;
+	TSetElementBase& operator=(TSetElementBase&&) = default;
+	TSetElementBase& operator=(const TSetElementBase&) = default;
 
+	/** The element's value. */
+	LAYOUT_FIELD(ElementType, Value);
+
+	/** The id of the next element in the same hash bucket. */
+	LAYOUT_MUTABLE_FIELD(FSetElementId, HashNextId);
+
+	/** The hash bucket that the element is currently linked to. */
+	LAYOUT_MUTABLE_FIELD(int32, HashIndex);
+};
+
+/** An element in the set. */
+template <typename InElementType>
+class TSetElement : public TSetElementBase<InElementType, THasTypeLayout<InElementType>::Value>
+{
+	using Super = TSetElementBase<InElementType, THasTypeLayout<InElementType>::Value>;
+public:
 	/** Default constructor. */
 	FORCEINLINE TSetElement()
 	{}
 
 	/** Initialization constructor. */
-	template <typename InitType, typename = typename TEnableIf<!TAreTypesEqual<TSetElement, typename TDecay<InitType>::Type>::Value>::Type> explicit FORCEINLINE TSetElement(InitType&& InValue) : Value(Forward<InitType>(InValue)) {}
+	template <typename InitType, typename = typename TEnableIf<!TAreTypesEqual<TSetElement, typename TDecay<InitType>::Type>::Value>::Type> explicit FORCEINLINE TSetElement(InitType&& InValue) : Super(Forward<InitType>(InValue)) {}
 
 	TSetElement(TSetElement&&) = default;
 	TSetElement(const TSetElement&) = default;
@@ -185,11 +236,11 @@ public:
 	// Comparison operators
 	FORCEINLINE bool operator==(const TSetElement& Other) const
 	{
-		return Value == Other.Value;
+		return this->Value == Other.Value;
 	}
 	FORCEINLINE bool operator!=(const TSetElement& Other) const
 	{
-		return Value != Other.Value;
+		return this->Value != Other.Value;
 	}
 };
 
@@ -215,8 +266,14 @@ template<
 	>
 class TSet
 {
+public:
+	static const bool SupportsFreezeMemoryImage = TAllocatorTraits<Allocator>::SupportsFreezeMemoryImage && THasTypeLayout<InElementType>::Value;
+
+private:
 	friend struct TContainerTraits<TSet>;
-	friend class  FScriptSet;
+
+	template <typename, typename>
+	friend class TScriptSet;
 
 	typedef typename KeyFuncs::KeyInitType     KeyInitType;
 	typedef typename KeyFuncs::ElementInitType ElementInitType;
@@ -1152,6 +1209,54 @@ private:
 	mutable HashType Hash;
 	mutable int32	 HashSize;
 
+	template<bool bFreezeMemoryImage, typename Dummy=void>
+	struct TSupportsFreezeMemoryImageHelper
+	{
+		static void WriteMemoryImage(FMemoryImageWriter& Writer, const TSet&) { Writer.WriteBytes(TSet()); }
+		static void CopyUnfrozen(const FMemoryUnfreezeContent& Context, const TSet&, void* Dst) { new(Dst) TSet(); }
+	};
+
+	template<typename Dummy>
+	struct TSupportsFreezeMemoryImageHelper<true, Dummy>
+	{
+		static void WriteMemoryImage(FMemoryImageWriter& Writer, const TSet& Object)
+		{
+			Object.Elements.WriteMemoryImage(Writer);
+			Object.Hash.WriteMemoryImage(Writer, StaticGetTypeLayoutDesc<FSetElementId>(), Object.HashSize);
+			Writer.WriteBytes(Object.HashSize);
+		}
+
+		static void CopyUnfrozen(const FMemoryUnfreezeContent& Context, const TSet& Object, void* Dst)
+		{
+			TSet* DstObject = static_cast<TSet*>(Dst);
+			Object.Elements.CopyUnfrozen(Context, &DstObject->Elements);
+
+			new(&DstObject->Hash) HashType();
+			DstObject->Hash.ResizeAllocation(0, Object.HashSize, sizeof(FSetElementId));
+			FMemory::Memcpy(DstObject->Hash.GetAllocation(), Object.Hash.GetAllocation(), sizeof(FSetElementId) * Object.HashSize);
+			DstObject->HashSize = Object.HashSize;
+		}
+	};
+
+public:
+	void WriteMemoryImage(FMemoryImageWriter& Writer) const
+	{
+		checkf(!Writer.GetTargetLayoutParams().b32Bit, TEXT("TSet does not currently support freezing for 32bits"));
+		TSupportsFreezeMemoryImageHelper<SupportsFreezeMemoryImage>::WriteMemoryImage(Writer, *this);
+	}
+
+	void CopyUnfrozen(const FMemoryUnfreezeContent& Context, void* Dst) const
+	{
+		TSupportsFreezeMemoryImageHelper<SupportsFreezeMemoryImage>::CopyUnfrozen(Context, *this, Dst);
+	}
+
+	static void AppendHash(const FPlatformTypeLayoutParameters& LayoutParams, FSHA1& Hasher)
+	{
+		ElementArrayType::AppendHash(LayoutParams, Hasher);
+	}
+
+private:
+
 	FORCEINLINE FSetElementId& GetTypedHash(int32 HashIndex) const
 	{
 		return ((FSetElementId*)Hash.GetAllocation())[HashIndex & (HashSize - 1)];
@@ -1469,6 +1574,30 @@ public:
 	FORCEINLINE TRangedForConstIterator end() const   { return TRangedForConstIterator(Elements.end());   }
 };
 
+namespace Freeze
+{
+	template<typename ElementType, typename KeyFuncs, typename Allocator>
+	void IntrinsicWriteMemoryImage(FMemoryImageWriter& Writer, const TSet<ElementType, KeyFuncs, Allocator>& Object, const FTypeLayoutDesc&)
+	{
+		Object.WriteMemoryImage(Writer);
+	}
+
+	template<typename ElementType, typename KeyFuncs, typename Allocator>
+	void IntrinsicUnfrozenCopy(const FMemoryUnfreezeContent& Context, const TSet<ElementType, KeyFuncs, Allocator>& Object, void* OutDst)
+	{
+		Object.CopyUnfrozen(Context, OutDst);
+	}
+
+	template<typename ElementType, typename KeyFuncs, typename Allocator>
+	uint32 IntrinsicAppendHash(const TSet<ElementType, KeyFuncs, Allocator>* DummyObject, const FTypeLayoutDesc& TypeDesc, const FPlatformTypeLayoutParameters& LayoutParams, FSHA1& Hasher)
+	{
+		TSet<ElementType, KeyFuncs, Allocator>::AppendHash(LayoutParams, Hasher);
+		return DefaultAppendHash(TypeDesc, LayoutParams, Hasher);
+	}
+}
+
+DECLARE_TEMPLATE_INTRINSIC_TYPE_LAYOUT((template <typename ElementType, typename KeyFuncs, typename Allocator>), (TSet<ElementType, KeyFuncs, Allocator>));
+
 template<typename ElementType, typename KeyFuncs, typename Allocator>
 struct TContainerTraits<TSet<ElementType, KeyFuncs, Allocator> > : public TContainerTraitsBase<TSet<ElementType, KeyFuncs, Allocator> >
 {
@@ -1490,8 +1619,11 @@ struct FScriptSetLayout
 
 // Untyped set type for accessing TSet data, like FScriptArray for TArray.
 // Must have the same memory representation as a TSet.
-class FScriptSet
+template <typename Allocator, typename InDerivedType>
+class TScriptSet
 {
+	using DerivedType = typename TChooseClass<TIsVoidType<InDerivedType>::Value, TScriptSet, InDerivedType>::Result;
+
 public:
 	static FScriptSetLayout GetScriptLayout(int32 ElementSize, int32 ElementAlignment)
 	{
@@ -1510,7 +1642,7 @@ public:
 		return Result;
 	}
 
-	FScriptSet()
+	TScriptSet()
 		: HashSize(0)
 	{
 	}
@@ -1540,7 +1672,7 @@ public:
 		return Elements.GetData(Index, Layout.SparseArrayLayout);
 	}
 
-	void MoveAssign(FScriptSet& Other, const FScriptSetLayout& Layout)
+	void MoveAssign(DerivedType& Other, const FScriptSetLayout& Layout)
 	{
 		checkSlow(this != &Other);
 		Empty(0, Layout);
@@ -1747,12 +1879,12 @@ private:
 		return NewElementIndex;
 	}
 
-	typedef FDefaultSetAllocator Allocator;
-	typedef Allocator::HashAllocator::ForElementType<FSetElementId> HashType;
+	typedef TScriptSparseArray<typename Allocator::SparseArrayAllocator> ElementArrayType;
+	typedef typename Allocator::HashAllocator::template ForElementType<FSetElementId> HashType;
 
-	FScriptSparseArray Elements;
-	mutable HashType   Hash;
-	mutable int32      HashSize;
+	ElementArrayType Elements;
+	mutable HashType Hash;
+	mutable int32    HashSize;
 
 	FORCEINLINE FSetElementId& GetTypedHash(int32 HashIndex) const
 	{
@@ -1772,33 +1904,41 @@ private:
 	// This function isn't intended to be called, just to be compiled to validate the correctness of the type.
 	static void CheckConstraints()
 	{
-		typedef FScriptSet  ScriptType;
+		typedef TScriptSet  ScriptType;
 		typedef TSet<int32> RealType;
 
 		// Check that the class footprint is the same
-		static_assert(sizeof (ScriptType) == sizeof (RealType), "FScriptSet's size doesn't match TSet");
-		static_assert(alignof(ScriptType) == alignof(RealType), "FScriptSet's alignment doesn't match TSet");
+		static_assert(sizeof (ScriptType) == sizeof (RealType), "TScriptSet's size doesn't match TSet");
+		static_assert(alignof(ScriptType) == alignof(RealType), "TScriptSet's alignment doesn't match TSet");
 
 		// Check member sizes
-		static_assert(sizeof(DeclVal<ScriptType>().Elements) == sizeof(DeclVal<RealType>().Elements), "FScriptSet's Elements member size does not match TSet's");
-		static_assert(sizeof(DeclVal<ScriptType>().Hash)     == sizeof(DeclVal<RealType>().Hash),     "FScriptSet's Hash member size does not match TSet's");
-		static_assert(sizeof(DeclVal<ScriptType>().HashSize) == sizeof(DeclVal<RealType>().HashSize), "FScriptSet's HashSize member size does not match TSet's");
+		static_assert(sizeof(DeclVal<ScriptType>().Elements) == sizeof(DeclVal<RealType>().Elements), "TScriptSet's Elements member size does not match TSet's");
+		static_assert(sizeof(DeclVal<ScriptType>().Hash)     == sizeof(DeclVal<RealType>().Hash),     "TScriptSet's Hash member size does not match TSet's");
+		static_assert(sizeof(DeclVal<ScriptType>().HashSize) == sizeof(DeclVal<RealType>().HashSize), "TScriptSet's HashSize member size does not match TSet's");
 
 		// Check member offsets
-		static_assert(STRUCT_OFFSET(ScriptType, Elements) == STRUCT_OFFSET(RealType, Elements), "FScriptSet's Elements member offset does not match TSet's");
-		static_assert(STRUCT_OFFSET(ScriptType, Hash)     == STRUCT_OFFSET(RealType, Hash),     "FScriptSet's Hash member offset does not match TSet's");
-		static_assert(STRUCT_OFFSET(ScriptType, HashSize) == STRUCT_OFFSET(RealType, HashSize), "FScriptSet's FirstFreeIndex member offset does not match TSet's");
+		static_assert(STRUCT_OFFSET(ScriptType, Elements) == STRUCT_OFFSET(RealType, Elements), "TScriptSet's Elements member offset does not match TSet's");
+		static_assert(STRUCT_OFFSET(ScriptType, Hash)     == STRUCT_OFFSET(RealType, Hash),     "TScriptSet's Hash member offset does not match TSet's");
+		static_assert(STRUCT_OFFSET(ScriptType, HashSize) == STRUCT_OFFSET(RealType, HashSize), "TScriptSet's FirstFreeIndex member offset does not match TSet's");
 	}
 
 public:
 	// These should really be private, because they shouldn't be called, but there's a bunch of code
 	// that needs to be fixed first.
-	FScriptSet(const FScriptSet&) { check(false); }
-	void operator=(const FScriptSet&) { check(false); }
+	TScriptSet(const TScriptSet&) { check(false); }
+	void operator=(const TScriptSet&) { check(false); }
 };
 
-template <>
-struct TIsZeroConstructType<FScriptSet>
+template <typename AllocatorType, typename InDerivedType>
+struct TIsZeroConstructType<TScriptSet<AllocatorType, InDerivedType>>
 {
 	enum { Value = true };
 };
+
+class FScriptSet : public TScriptSet<FDefaultSetAllocator, FScriptSet>
+{
+	using Super = TScriptSet<FDefaultSetAllocator, FScriptSet>;
+
+public:
+	using Super::Super;
+};
diff --git a/Engine/Source/Runtime/Core/Public/Containers/SortedMap.h b/Engine/Source/Runtime/Core/Public/Containers/SortedMap.h
--- a/Engine/Source/Runtime/Core/Public/Containers/SortedMap.h
+++ b/Engine/Source/Runtime/Core/Public/Containers/SortedMap.h
@@ -677,6 +677,7 @@ public:
 	FORCEINLINE RangedForConstIteratorType	end() const   { return Pairs.end(); }
 };
 
+DECLARE_TEMPLATE_INTRINSIC_TYPE_LAYOUT((template <typename KeyType, typename ValueType, typename ArrayAllocator, typename SortPredicate>), (TSortedMap<KeyType, ValueType, ArrayAllocator, SortPredicate>));
 
 template <typename KeyType, typename ValueType, typename ArrayAllocator, typename SortPredicate>
 struct TContainerTraits<TSortedMap<KeyType, ValueType, ArrayAllocator, SortPredicate>> : public TContainerTraitsBase<TSortedMap<KeyType, ValueType, ArrayAllocator, SortPredicate>>
diff --git a/Engine/Source/Runtime/Core/Public/Containers/SparseArray.h b/Engine/Source/Runtime/Core/Public/Containers/SparseArray.h
--- a/Engine/Source/Runtime/Core/Public/Containers/SparseArray.h
+++ b/Engine/Source/Runtime/Core/Public/Containers/SparseArray.h
@@ -16,6 +16,7 @@
 #include "Containers/ScriptArray.h"
 #include "Containers/BitArray.h"
 #include "Serialization/StructuredArchive.h"
+#include "Serialization/MemoryImageWriter.h"
 
 
 #if UE_BUILD_SHIPPING || UE_BUILD_TEST
@@ -53,7 +54,10 @@ union TSparseArrayElementOrFreeListLink
 	};
 };
 
-class FScriptSparseArray;
+DECLARE_TEMPLATE_INTRINSIC_TYPE_LAYOUT(template<typename ElementType>, TSparseArrayElementOrFreeListLink<ElementType>);
+
+template <typename AllocatorType, typename InDerivedType = void>
+class TScriptSparseArray;
 
 /**
  * A dynamically sized array where element indices aren't necessarily contiguous.  Memory is allocated for all 
@@ -68,9 +72,12 @@ class TSparseArray
 	using ElementType = InElementType;
 
 	friend struct TContainerTraits<TSparseArray>;
-	friend class  FScriptSparseArray;
+
+	template <typename, typename>
+	friend class TScriptSparseArray;
 
 public:
+	static const bool SupportsFreezeMemoryImage = TAllocatorTraits<Allocator>::SupportsFreezeMemoryImage && THasTypeLayout<InElementType>::Value;
 
 	/** Destructor. */
 	~TSparseArray()
@@ -1063,8 +1070,131 @@ private:
 
 	/** The number of elements in the free list. */
 	int32 NumFreeIndices;
+
+	template<bool bFreezeMemoryImage, typename Dummy=void>
+	struct TSupportsFreezeMemoryImageHelper
+	{
+		static void WriteMemoryImage(FMemoryImageWriter& Writer, const TSparseArray&) { Writer.WriteBytes(TSparseArray()); }
+		static void CopyUnfrozen(const FMemoryUnfreezeContent& Context, const TSparseArray&, void* Dst) { new(Dst) TSparseArray(); }
+		static void AppendHash(const FPlatformTypeLayoutParameters& LayoutParams, FSHA1& Hasher) {}
+	};
+
+	template<typename Dummy>
+	struct TSupportsFreezeMemoryImageHelper<true, Dummy>
+	{
+		static void WriteMemoryImage(FMemoryImageWriter& Writer, const TSparseArray& Object)
+		{
+			// Write Data
+			const int32 NumElements = Object.Data.Num();
+			if (NumElements > 0)
+			{
+				const FTypeLayoutDesc& ElementTypeDesc = StaticGetTypeLayoutDesc<ElementType>();
+				FMemoryImageWriter ArrayWriter = Writer.WritePointer(FString::Printf(TEXT("TSparseArray<%s>"), ElementTypeDesc.Name));
+				for (int32 i = 0; i < NumElements; ++i)
+				{
+					const FElementOrFreeListLink& Elem = Object.Data[i];
+					const uint32 StartOffset = ArrayWriter.WriteAlignment<FElementOrFreeListLink>();
+					if (Object.AllocationFlags[i])
+					{
+						ArrayWriter.WriteObject(&Elem.ElementData, ElementTypeDesc);
+					}
+					else
+					{
+						ArrayWriter.WriteBytes(Elem.PrevFreeIndex);
+						ArrayWriter.WriteBytes(Elem.NextFreeIndex);
+					}
+					ArrayWriter.WritePaddingToSize(StartOffset + sizeof(FElementOrFreeListLink));
+				}
+			}
+			else
+			{
+				Writer.WriteMemoryImagePointerSizedBytes(0u);
+			}
+			Writer.WriteBytes(NumElements);
+			Writer.WriteBytes(NumElements);
+
+			//
+			Object.AllocationFlags.WriteMemoryImage(Writer);
+			Writer.WriteBytes(Object.FirstFreeIndex);
+			Writer.WriteBytes(Object.NumFreeIndices);
+		}
+
+		static void CopyUnfrozen(const FMemoryUnfreezeContent& Context, const TSparseArray& Object, void* Dst)
+		{
+			const FTypeLayoutDesc& ElementTypeDesc = StaticGetTypeLayoutDesc<ElementType>();
+			TSparseArray* DstObject = (TSparseArray*)Dst;
+			{
+				new(&DstObject->Data) DataType();
+				DstObject->Data.SetNumUninitialized(Object.Data.Num());
+				for (int32 i = 0; i < Object.Data.Num(); ++i)
+				{
+					const FElementOrFreeListLink& Elem = Object.Data[i];
+					FElementOrFreeListLink& DstElem = DstObject->Data[i];
+					if (Object.AllocationFlags[i])
+					{
+						Context.UnfreezeObject(&Elem.ElementData, ElementTypeDesc, &DstElem.ElementData);
+					}
+					else
+					{
+						DstElem.PrevFreeIndex = Elem.PrevFreeIndex;
+						DstElem.NextFreeIndex = Elem.NextFreeIndex;
+					}
+				}
+			}
+
+			new(&DstObject->AllocationFlags) AllocationBitArrayType(Object.AllocationFlags);
+			DstObject->FirstFreeIndex = Object.FirstFreeIndex;
+			DstObject->NumFreeIndices = Object.NumFreeIndices;
+		}
+
+		static void AppendHash(const FPlatformTypeLayoutParameters& LayoutParams, FSHA1& Hasher)
+		{
+			Freeze::AppendHash(StaticGetTypeLayoutDesc<ElementType>(), LayoutParams, Hasher);
+		}
+	};
+
+public:
+	void WriteMemoryImage(FMemoryImageWriter& Writer) const
+	{
+		checkf(!Writer.GetTargetLayoutParams().b32Bit, TEXT("TSparseArray does not currently support freezing for 32bits"));
+		TSupportsFreezeMemoryImageHelper<SupportsFreezeMemoryImage>::WriteMemoryImage(Writer, *this);
+	}
+
+	void CopyUnfrozen(const FMemoryUnfreezeContent& Context, void* Dst) const
+	{
+		TSupportsFreezeMemoryImageHelper<SupportsFreezeMemoryImage>::CopyUnfrozen(Context, *this, Dst);
+	}
+
+	static void AppendHash(const FPlatformTypeLayoutParameters& LayoutParams, FSHA1& Hasher)
+	{
+		TSupportsFreezeMemoryImageHelper<SupportsFreezeMemoryImage>::AppendHash(LayoutParams, Hasher);
+	}
 };
 
+namespace Freeze
+{
+	template<typename ElementType, typename Allocator>
+	void IntrinsicWriteMemoryImage(FMemoryImageWriter& Writer, const TSparseArray<ElementType, Allocator>& Object, const FTypeLayoutDesc&)
+	{
+		Object.WriteMemoryImage(Writer);
+	}
+
+	template<typename ElementType, typename Allocator>
+	void IntrinsicUnfrozenCopy(const FMemoryUnfreezeContent& Context, const TSparseArray<ElementType, Allocator>& Object, void* OutDst)
+	{
+		Object.CopyUnfrozen(Context, OutDst);
+	}
+
+	template<typename ElementType, typename Allocator>
+	uint32 IntrinsicAppendHash(const TSparseArray<ElementType, Allocator>* DummyObject, const FTypeLayoutDesc& TypeDesc, const FPlatformTypeLayoutParameters& LayoutParams, FSHA1& Hasher)
+	{
+		TSparseArray<ElementType, Allocator>::AppendHash(LayoutParams, Hasher);
+		return DefaultAppendHash(TypeDesc, LayoutParams, Hasher);
+	}
+}
+
+DECLARE_TEMPLATE_INTRINSIC_TYPE_LAYOUT((template <typename ElementType, typename Allocator>), (TSparseArray<ElementType, Allocator>));
+
 template<typename ElementType, typename Allocator>
 struct TContainerTraits<TSparseArray<ElementType, Allocator> > : public TContainerTraitsBase<TSparseArray<ElementType, Allocator> >
 {
@@ -1098,8 +1228,11 @@ struct FScriptSparseArrayLayout
 
 // Untyped sparse array type for accessing TSparseArray data, like FScriptArray for TArray.
 // Must have the same memory representation as a TSet.
-class FScriptSparseArray
+template <typename AllocatorType, typename InDerivedType>
+class TScriptSparseArray
 {
+	using DerivedType = typename TChooseClass<TIsVoidType<InDerivedType>::Value, TScriptSparseArray, InDerivedType>::Result;
+
 public:
 	static FScriptSparseArrayLayout GetScriptLayout(int32 ElementSize, int32 ElementAlignment)
 	{
@@ -1110,7 +1243,7 @@ public:
 		return Result;
 	}
 
-	FScriptSparseArray()
+	TScriptSparseArray()
 		: FirstFreeIndex(-1)
 		, NumFreeIndices(0)
 	{
@@ -1141,7 +1274,7 @@ public:
 		return (const uint8*)Data.GetData() + Layout.Size * Index;
 	}
 
-	void MoveAssign(FScriptSparseArray& Other, const FScriptSparseArrayLayout& Layout)
+	void MoveAssign(DerivedType& Other, const FScriptSparseArrayLayout& Layout)
 	{
 		checkSlow(this != &Other);
 		Empty(0, Layout);
@@ -1216,36 +1349,36 @@ public:
 	}
 
 private:
-	FScriptArray    Data;
-	FScriptBitArray AllocationFlags;
-	int32           FirstFreeIndex;
-	int32           NumFreeIndices;
+	TScriptArray   <typename AllocatorType::ElementAllocator>  Data;
+	TScriptBitArray<typename AllocatorType::BitArrayAllocator> AllocationFlags;
+	int32                                                      FirstFreeIndex;
+	int32                                                      NumFreeIndices;
 
 	// This function isn't intended to be called, just to be compiled to validate the correctness of the type.
 	static void CheckConstraints()
 	{
-		typedef FScriptSparseArray  ScriptType;
+		typedef TScriptSparseArray  ScriptType;
 		typedef TSparseArray<int32> RealType;
 
 		// Check that the class footprint is the same
-		static_assert(sizeof (ScriptType) == sizeof (RealType), "FScriptSparseArray's size doesn't match TSparseArray");
-		static_assert(alignof(ScriptType) == alignof(RealType), "FScriptSparseArray's alignment doesn't match TSparseArray");
+		static_assert(sizeof (ScriptType) == sizeof (RealType), "TScriptSparseArray's size doesn't match TSparseArray");
+		static_assert(alignof(ScriptType) == alignof(RealType), "TScriptSparseArray's alignment doesn't match TSparseArray");
 
 		// Check member sizes
-		static_assert(sizeof(DeclVal<ScriptType>().Data)            == sizeof(DeclVal<RealType>().Data),            "FScriptSparseArray's Data member size does not match TSparseArray's");
-		static_assert(sizeof(DeclVal<ScriptType>().AllocationFlags) == sizeof(DeclVal<RealType>().AllocationFlags), "FScriptSparseArray's AllocationFlags member size does not match TSparseArray's");
-		static_assert(sizeof(DeclVal<ScriptType>().FirstFreeIndex)  == sizeof(DeclVal<RealType>().FirstFreeIndex),  "FScriptSparseArray's FirstFreeIndex member size does not match TSparseArray's");
-		static_assert(sizeof(DeclVal<ScriptType>().NumFreeIndices)  == sizeof(DeclVal<RealType>().NumFreeIndices),  "FScriptSparseArray's NumFreeIndices member size does not match TSparseArray's");
+		static_assert(sizeof(DeclVal<ScriptType>().Data)            == sizeof(DeclVal<RealType>().Data),            "TScriptSparseArray's Data member size does not match TSparseArray's");
+		static_assert(sizeof(DeclVal<ScriptType>().AllocationFlags) == sizeof(DeclVal<RealType>().AllocationFlags), "TScriptSparseArray's AllocationFlags member size does not match TSparseArray's");
+		static_assert(sizeof(DeclVal<ScriptType>().FirstFreeIndex)  == sizeof(DeclVal<RealType>().FirstFreeIndex),  "TScriptSparseArray's FirstFreeIndex member size does not match TSparseArray's");
+		static_assert(sizeof(DeclVal<ScriptType>().NumFreeIndices)  == sizeof(DeclVal<RealType>().NumFreeIndices),  "TScriptSparseArray's NumFreeIndices member size does not match TSparseArray's");
 
 		// Check member offsets
-		static_assert(STRUCT_OFFSET(ScriptType, Data)            == STRUCT_OFFSET(RealType, Data),            "FScriptSparseArray's Data member offset does not match TSparseArray's");
-		static_assert(STRUCT_OFFSET(ScriptType, AllocationFlags) == STRUCT_OFFSET(RealType, AllocationFlags), "FScriptSparseArray's AllocationFlags member offset does not match TSparseArray's");
-		static_assert(STRUCT_OFFSET(ScriptType, FirstFreeIndex)  == STRUCT_OFFSET(RealType, FirstFreeIndex),  "FScriptSparseArray's FirstFreeIndex member offset does not match TSparseArray's");
-		static_assert(STRUCT_OFFSET(ScriptType, NumFreeIndices)  == STRUCT_OFFSET(RealType, NumFreeIndices),  "FScriptSparseArray's NumFreeIndices member offset does not match TSparseArray's");
+		static_assert(STRUCT_OFFSET(ScriptType, Data)            == STRUCT_OFFSET(RealType, Data),            "TScriptSparseArray's Data member offset does not match TSparseArray's");
+		static_assert(STRUCT_OFFSET(ScriptType, AllocationFlags) == STRUCT_OFFSET(RealType, AllocationFlags), "TScriptSparseArray's AllocationFlags member offset does not match TSparseArray's");
+		static_assert(STRUCT_OFFSET(ScriptType, FirstFreeIndex)  == STRUCT_OFFSET(RealType, FirstFreeIndex),  "TScriptSparseArray's FirstFreeIndex member offset does not match TSparseArray's");
+		static_assert(STRUCT_OFFSET(ScriptType, NumFreeIndices)  == STRUCT_OFFSET(RealType, NumFreeIndices),  "TScriptSparseArray's NumFreeIndices member offset does not match TSparseArray's");
 
 		// Check free index offsets
-		static_assert(STRUCT_OFFSET(ScriptType::FFreeListLink, PrevFreeIndex) == STRUCT_OFFSET(RealType::FElementOrFreeListLink, PrevFreeIndex), "FScriptSparseArray's FFreeListLink's PrevFreeIndex member offset does not match TSparseArray's");
-		static_assert(STRUCT_OFFSET(ScriptType::FFreeListLink, NextFreeIndex) == STRUCT_OFFSET(RealType::FElementOrFreeListLink, NextFreeIndex), "FScriptSparseArray's FFreeListLink's NextFreeIndex member offset does not match TSparseArray's");
+		static_assert(STRUCT_OFFSET(ScriptType::FFreeListLink, PrevFreeIndex) == STRUCT_OFFSET(RealType::FElementOrFreeListLink, PrevFreeIndex), "TScriptSparseArray's FFreeListLink's PrevFreeIndex member offset does not match TSparseArray's");
+		static_assert(STRUCT_OFFSET(ScriptType::FFreeListLink, NextFreeIndex) == STRUCT_OFFSET(RealType::FElementOrFreeListLink, NextFreeIndex), "TScriptSparseArray's FFreeListLink's NextFreeIndex member offset does not match TSparseArray's");
 	}
 
 	struct FFreeListLink
@@ -1266,16 +1399,24 @@ private:
 public:
 	// These should really be private, because they shouldn't be called, but there's a bunch of code
 	// that needs to be fixed first.
-	FScriptSparseArray(const FScriptSparseArray&) { check(false); }
-	void operator=(const FScriptSparseArray&) { check(false); }
+	TScriptSparseArray(const TScriptSparseArray&) { check(false); }
+	void operator=(const TScriptSparseArray&) { check(false); }
 };
 
-template <>
-struct TIsZeroConstructType<FScriptSparseArray>
+template <typename AllocatorType, typename InDerivedType>
+struct TIsZeroConstructType<TScriptSparseArray<AllocatorType, InDerivedType>>
 {
 	enum { Value = true };
 };
 
+class FScriptSparseArray : public TScriptSparseArray<FDefaultSparseArrayAllocator, FScriptSparseArray>
+{
+	using Super = TScriptSparseArray<FDefaultSparseArrayAllocator, FScriptSparseArray>;
+
+public:
+	using Super::Super;
+};
+
 /**
  * A placement new operator which constructs an element in a sparse array allocation.
  */
diff --git a/Engine/Source/Runtime/Core/Public/CoreGlobals.h b/Engine/Source/Runtime/Core/Public/CoreGlobals.h
--- a/Engine/Source/Runtime/Core/Public/CoreGlobals.h
+++ b/Engine/Source/Runtime/Core/Public/CoreGlobals.h
@@ -304,6 +304,7 @@ extern CORE_API FString GGameIni;
 extern CORE_API FString GGameUserSettingsIni;
 extern CORE_API FString GRuntimeOptionsIni;
 extern CORE_API FString GInstallBundleIni;
+extern CORE_API FString GDeviceProfilesIni;
 
 extern CORE_API float GNearClippingPlane;
 
diff --git a/Engine/Source/Runtime/Core/Public/Delegates/DelegateBase.h b/Engine/Source/Runtime/Core/Public/Delegates/DelegateBase.h
--- a/Engine/Source/Runtime/Core/Public/Delegates/DelegateBase.h
+++ b/Engine/Source/Runtime/Core/Public/Delegates/DelegateBase.h
@@ -33,6 +33,8 @@ struct FWeakObjectPtr;
 template <typename ObjectPtrType>
 class FMulticastDelegateBase;
 
+ALIAS_TEMPLATE_TYPE_LAYOUT(template<typename ElementType>, FDelegateAllocatorType::ForElementType<ElementType>, void*);
+
 /**
  * Base class for unicast delegates.
  */
diff --git a/Engine/Source/Runtime/Core/Public/Experimental/Containers/RobinHoodHashTable.h b/Engine/Source/Runtime/Core/Public/Experimental/Containers/RobinHoodHashTable.h
--- a/Engine/Source/Runtime/Core/Public/Experimental/Containers/RobinHoodHashTable.h
+++ b/Engine/Source/Runtime/Core/Public/Experimental/Containers/RobinHoodHashTable.h
@@ -6,6 +6,14 @@
 #include "Containers/Map.h"
 #include "Templates/UnrealTemplate.h"
 
+#define RUN_HASHTABLE_CONCURENCY_CHECKS UE_BUILD_DEBUG
+
+#if RUN_HASHTABLE_CONCURENCY_CHECKS
+#define CHECK_CONCURRENT_ACCESS(x) check(x)
+#else
+#define CHECK_CONCURRENT_ACCESS(x)
+#endif
+
 namespace Experimental
 {
 
@@ -91,9 +99,10 @@ namespace RobinHoodHashTable_Private
 		friend class TRobinHoodHashTable;
 
 		template<typename DeducedKeyType, typename DeducedValueType>
-		inline TKeyValue(DeducedKeyType&& InKey, DeducedValueType&& InVal) : Pair(Forward<DeducedKeyType>(InKey), Forward<DeducedValueType>(InVal)) {}
+		inline TKeyValue(DeducedKeyType&& InKey, DeducedValueType&& InVal, FHashType InHash) : Pair(Forward<DeducedKeyType>(InKey), Forward<DeducedValueType>(InVal)), Hash(InHash) {}
 
 		ElementType Pair;
+		FHashType Hash;
 
 		inline FindValueType FindImpl()
 		{
@@ -126,9 +135,10 @@ namespace RobinHoodHashTable_Private
 		friend class TRobinHoodHashTable;
 
 		template<typename DeducedKeyType>
-		inline TKeyValue(DeducedKeyType&& InKey, FUnitType&&) : Key(Forward<DeducedKeyType>(InKey)) {}
+		inline TKeyValue(DeducedKeyType&& InKey, FUnitType&&, FHashType InHash) : Key(Forward<DeducedKeyType>(InKey)), Hash(InHash) {}
 
 		ElementType Key;
+		FHashType Hash;
 
 		inline FindValueType FindImpl() const
 		{
@@ -178,25 +188,31 @@ namespace RobinHoodHashTable_Private
 			}
 
 			template<typename DeducedKeyType, typename DeducedValueType>
-			inline IndexType Allocate(DeducedKeyType&& Key, DeducedValueType&& Val)
+			inline IndexType Allocate(DeducedKeyType&& Key, DeducedValueType&& Val, FHashType Hash)
 			{
+				CHECK_CONCURRENT_ACCESS(FPlatformAtomics::InterlockedIncrement(&ConcurrentWriters) == 1);
+				CHECK_CONCURRENT_ACCESS(FPlatformAtomics::InterlockedIncrement(&ConcurrentReaders) == 1);
 				IndexType Index;
 				if (FreeList.Num() > 0)
 				{
 					Index = FreeList.Pop();
-					new (&KeyVals[Index]) KeyValueType{ Forward<DeducedKeyType>(Key), Forward<DeducedValueType>(Val) };
+					new (&KeyVals[Index]) KeyValueType{ Forward<DeducedKeyType>(Key), Forward<DeducedValueType>(Val), Hash };
 				}
 				else
 				{
 					Index = KeyVals.Num();
-					KeyVals.Push(KeyValueType{ Forward<DeducedKeyType>(Key), Forward<DeducedValueType>(Val) });
+					KeyVals.Push(KeyValueType{ Forward<DeducedKeyType>(Key), Forward<DeducedValueType>(Val), Hash });
 				}
 				checkSlow(Index != InvalidIndex);
+				CHECK_CONCURRENT_ACCESS(FPlatformAtomics::InterlockedDecrement(&ConcurrentReaders) == 0);
+				CHECK_CONCURRENT_ACCESS(FPlatformAtomics::InterlockedDecrement(&ConcurrentWriters) == 0);
 				return Index;
 			}
 
 			inline void Deallocate(IndexType Index)
 			{
+				CHECK_CONCURRENT_ACCESS(FPlatformAtomics::InterlockedIncrement(&ConcurrentWriters) == 1);
+				CHECK_CONCURRENT_ACCESS(FPlatformAtomics::InterlockedIncrement(&ConcurrentReaders) == 1);
 				int InsertIndex = 0;
 				for (; InsertIndex < FreeList.Num(); InsertIndex++)
 				{
@@ -205,8 +221,11 @@ namespace RobinHoodHashTable_Private
 						break;
 					}
 				}
+
 				FreeList.EmplaceAt(InsertIndex, Index);
 				KeyVals[Index].~KeyValueType();
+				CHECK_CONCURRENT_ACCESS(FPlatformAtomics::InterlockedDecrement(&ConcurrentReaders) == 0);
+				CHECK_CONCURRENT_ACCESS(FPlatformAtomics::InterlockedDecrement(&ConcurrentWriters) == 0);
 				//TODO shrink KeyValue Array.
 				//As the FreeList is sorted backwards this would work as follows:
 				// - go though the FreeList front to back as long as the FreeList entry is the last
@@ -216,23 +235,31 @@ namespace RobinHoodHashTable_Private
 
 			inline const KeyValueType& Get(IndexType Index) const
 			{
+				CHECK_CONCURRENT_ACCESS(ConcurrentWriters == 0);
+				CHECK_CONCURRENT_ACCESS(FPlatformAtomics::InterlockedIncrement(&ConcurrentReaders) >= 1);
 #if DO_CHECK
 				for (IndexType Idx : FreeList)
 				{
 					check(Idx != Index);
 				}
 #endif
+				CHECK_CONCURRENT_ACCESS(FPlatformAtomics::InterlockedDecrement(&ConcurrentReaders) >= 0);
+				CHECK_CONCURRENT_ACCESS(ConcurrentWriters == 0);
 				return KeyVals[Index];
 			}
 
 			inline KeyValueType& Get(IndexType Index)
 			{
+				CHECK_CONCURRENT_ACCESS(ConcurrentWriters == 0);
+				CHECK_CONCURRENT_ACCESS(FPlatformAtomics::InterlockedIncrement(&ConcurrentReaders) >= 1);
 #if DO_CHECK
 				for (IndexType Idx : FreeList)
 				{
 					check(Idx != Index);
 				}
 #endif
+				CHECK_CONCURRENT_ACCESS(FPlatformAtomics::InterlockedDecrement(&ConcurrentReaders) >= 0);
+				CHECK_CONCURRENT_ACCESS(ConcurrentWriters == 0);
 				return KeyVals[Index];
 			}
 
@@ -308,6 +335,10 @@ namespace RobinHoodHashTable_Private
 				KeyVals.Reserve(ReserveNum);
 			}
 
+#if RUN_HASHTABLE_CONCURENCY_CHECKS
+			mutable int ConcurrentReaders = 0;
+			mutable int ConcurrentWriters = 0;
+#endif
 			TArray<KeyValueType, HashMapAllocator> KeyVals;
 			TArray<IndexType, HashMapAllocator> FreeList;
 		};
@@ -364,8 +395,11 @@ namespace RobinHoodHashTable_Private
 
 	protected:
 		template<typename DeducedKeyType, typename DeducedValueType>
-		inline FHashElementId FindOrAddIdByHash(FHashType HashValue, DeducedKeyType&& Key, DeducedValueType&& Val)
+		inline FHashElementId FindOrAddIdByHash(FHashType HashValue, DeducedKeyType&& Key, DeducedValueType&& Val, bool& bIsAlreadyInMap)
 		{
+			CHECK_CONCURRENT_ACCESS(FPlatformAtomics::InterlockedIncrement(&ConcurrentWriters) == 1);
+			CHECK_CONCURRENT_ACCESS(FPlatformAtomics::InterlockedIncrement(&ConcurrentReaders) == 1);
+
 			checkSlow(HashValue == ComputeHash(Key));
 			IndexType BucketIndex = ModTableSize(HashValue.AsUInt());
 			const IndexType EndBucketIndex = ModTableSize(HashValue.AsUInt() + MaximumDistance + 1);
@@ -375,6 +409,9 @@ namespace RobinHoodHashTable_Private
 				{
 					if (Hasher::Matches(Key, KeyValueData.Get(IndexData[BucketIndex]).GetKey()))
 					{
+						CHECK_CONCURRENT_ACCESS(FPlatformAtomics::InterlockedDecrement(&ConcurrentReaders) == 0);
+						CHECK_CONCURRENT_ACCESS(FPlatformAtomics::InterlockedDecrement(&ConcurrentWriters) == 0);
+						bIsAlreadyInMap = true;
 						return IndexData[BucketIndex];
 					}
 				}
@@ -403,22 +440,27 @@ namespace RobinHoodHashTable_Private
 				}
 			}
 
-			IndexType InsertIndex = KeyValueData.Allocate(Forward<DeducedKeyType>(Key), Forward<DeducedValueType>(Val));
+			IndexType InsertIndex = KeyValueData.Allocate(Forward<DeducedKeyType>(Key), Forward<DeducedValueType>(Val), HashValue);
 			InsertIntoTable(InsertIndex, HashValue);
+
+			CHECK_CONCURRENT_ACCESS(FPlatformAtomics::InterlockedDecrement(&ConcurrentReaders) == 0);
+			CHECK_CONCURRENT_ACCESS(FPlatformAtomics::InterlockedDecrement(&ConcurrentWriters) == 0);
+
+			bIsAlreadyInMap = false;
 			return FHashElementId(InsertIndex);
 		}
 
 		template<typename DeducedKeyType, typename DeducedValueType>
-		inline FHashElementId FindOrAddId(DeducedKeyType&& Key, DeducedValueType&& Val)
+		inline FHashElementId FindOrAddId(DeducedKeyType&& Key, DeducedValueType&& Val, bool& bIsAlreadyInMap)
 		{
 			FHashType HashValue = ComputeHash(Key);
-			return FindOrAddIdByHash(HashValue, Forward<DeducedKeyType>(Key), Forward<DeducedValueType>(Val));
+			return FindOrAddIdByHash(HashValue, Forward<DeducedKeyType>(Key), Forward<DeducedValueType>(Val), bIsAlreadyInMap);
 		}
 
 		template<typename DeducedKeyType, typename DeducedValueType>
-		inline FindValueType FindOrAdd(DeducedKeyType&& Key, DeducedValueType&& Val)
+		inline FindValueType FindOrAdd(DeducedKeyType&& Key, DeducedValueType&& Val, bool& bIsAlreadyInMap)
 		{
-			FHashElementId Id = FindOrAddId(Forward<DeducedKeyType>(Key), Forward<DeducedValueType>(Val));
+			FHashElementId Id = FindOrAddId(Forward<DeducedKeyType>(Key), Forward<DeducedValueType>(Val), bIsAlreadyInMap);
 			return KeyValueData.Get(Id.GetIndex()).FindImpl();
 		}
 
@@ -639,6 +681,9 @@ namespace RobinHoodHashTable_Private
 
 		inline FHashElementId FindId(const KeyType& Key) const
 		{
+			CHECK_CONCURRENT_ACCESS(ConcurrentWriters == 0);
+			CHECK_CONCURRENT_ACCESS(FPlatformAtomics::InterlockedIncrement(&ConcurrentReaders) >= 1);
+
 			FHashType HashValue = ComputeHash(Key);
 			IndexType BucketIndex = ModTableSize(HashValue.AsUInt());
 			const IndexType EndBucketIndex = ModTableSize(HashValue.AsUInt() + MaximumDistance + 1);
@@ -648,6 +693,8 @@ namespace RobinHoodHashTable_Private
 				{
 					if (Hasher::Matches(Key, KeyValueData.Get(IndexData[BucketIndex]).GetKey()))
 					{
+						CHECK_CONCURRENT_ACCESS(FPlatformAtomics::InterlockedDecrement(&ConcurrentReaders) >= 0);
+						CHECK_CONCURRENT_ACCESS(ConcurrentWriters == 0);
 						return FHashElementId(IndexData[BucketIndex]);
 					}
 				}
@@ -655,6 +702,9 @@ namespace RobinHoodHashTable_Private
 				BucketIndex = ModTableSize(BucketIndex + 1);
 			} while (BucketIndex != EndBucketIndex);
 
+			CHECK_CONCURRENT_ACCESS(FPlatformAtomics::InterlockedDecrement(&ConcurrentReaders) >= 0);
+			CHECK_CONCURRENT_ACCESS(ConcurrentWriters == 0);
+
 			return FHashElementId();
 		}
 
@@ -682,7 +732,10 @@ namespace RobinHoodHashTable_Private
 
 		bool Remove(const KeyType& Key)
 		{
-			bool bRetValue = false;
+			CHECK_CONCURRENT_ACCESS(FPlatformAtomics::InterlockedIncrement(&ConcurrentWriters) == 1);
+			CHECK_CONCURRENT_ACCESS(FPlatformAtomics::InterlockedIncrement(&ConcurrentReaders) == 1);
+
+			bool bIsFoundInMap = false;
 			FHashType HashValue = ComputeHash(Key);
 			IndexType BucketIndex = ModTableSize(HashValue.AsUInt());
 			const IndexType EndBucketIndex = ModTableSize(HashValue.AsUInt() + MaximumDistance + 1);
@@ -694,7 +747,7 @@ namespace RobinHoodHashTable_Private
 					{
 						KeyValueData.Deallocate(IndexData[BucketIndex]);
 						HashData[BucketIndex] = FHashType();
-						bRetValue = true;
+						bIsFoundInMap = true;
 						break;
 					}
 				}
@@ -702,7 +755,7 @@ namespace RobinHoodHashTable_Private
 				BucketIndex = ModTableSize(BucketIndex + 1);
 			} while (BucketIndex != EndBucketIndex);
 
-			if (bRetValue && (KeyValueData.Num() * LoadFactorQuotient * 4) < (SizePow2Minus1 * LoadFactorDivisor))
+			if (bIsFoundInMap && (KeyValueData.Num() * LoadFactorQuotient * 4) < (SizePow2Minus1 * LoadFactorDivisor))
 			{
 				TArray<IndexType> IndexDataOld = MoveTemp(IndexData);
 				TArray<FHashType> HashDataOld = MoveTemp(HashData);
@@ -723,16 +776,67 @@ namespace RobinHoodHashTable_Private
 				}
 			}
 
-			return bRetValue;
+			CHECK_CONCURRENT_ACCESS(FPlatformAtomics::InterlockedDecrement(&ConcurrentReaders) == 0);
+			CHECK_CONCURRENT_ACCESS(FPlatformAtomics::InterlockedDecrement(&ConcurrentWriters) == 0);
+
+			return bIsFoundInMap;
 		}
 
 		inline bool RemoveByElementId(FHashElementId Id)
 		{
-			return Remove(KeyValueData.Get(Id.GetIndex()).GetKey());
+			CHECK_CONCURRENT_ACCESS(FPlatformAtomics::InterlockedIncrement(&ConcurrentWriters) == 1);
+			CHECK_CONCURRENT_ACCESS(FPlatformAtomics::InterlockedIncrement(&ConcurrentReaders) == 1);
+			
+			bool bIsFoundInMap = false;
+			FHashType HashValue = KeyValueData.Get(Id.GetIndex()).Hash;
+			IndexType BucketIndex = ModTableSize(HashValue.AsUInt());
+			const IndexType EndBucketIndex = ModTableSize(HashValue.AsUInt() + MaximumDistance + 1);
+			do
+			{
+				if (Id.GetIndex() == IndexData[BucketIndex])
+				{
+					checkSlow(HashData[BucketIndex] == HashValue);
+					KeyValueData.Deallocate(IndexData[BucketIndex]);
+					HashData[BucketIndex] = FHashType();
+					bIsFoundInMap = true;
+					break;
+				}
+
+				BucketIndex = ModTableSize(BucketIndex + 1);
+			} while (BucketIndex != EndBucketIndex);
+
+			if (bIsFoundInMap && (KeyValueData.Num() * LoadFactorQuotient * 4) < (SizePow2Minus1 * LoadFactorDivisor))
+			{
+				TArray<IndexType> IndexDataOld = MoveTemp(IndexData);
+				TArray<FHashType> HashDataOld = MoveTemp(HashData);
+				IndexType OldSizePow2Minus1 = SizePow2Minus1;
+				SizePow2Minus1 = SizePow2Minus1 / 2;
+				MaximumDistance = 0;
+				IndexData.Reserve(SizePow2Minus1 + 1);
+				IndexData.AddUninitialized(SizePow2Minus1 + 1);
+				HashData.Reserve(SizePow2Minus1 + 1);
+				HashData.AddDefaulted(SizePow2Minus1 + 1);
+
+				for (IndexType Index = 0; Index <= OldSizePow2Minus1; Index++)
+				{
+					if (HashDataOld[Index].IsOccupied())
+					{
+						InsertIntoTable(IndexDataOld[Index], HashDataOld[Index]);
+					}
+				}
+			}
+
+			CHECK_CONCURRENT_ACCESS(FPlatformAtomics::InterlockedDecrement(&ConcurrentReaders) == 0);
+			CHECK_CONCURRENT_ACCESS(FPlatformAtomics::InterlockedDecrement(&ConcurrentWriters) == 0);
+
+			return bIsFoundInMap;
 		}
 
 		void Empty()
 		{
+			CHECK_CONCURRENT_ACCESS(FPlatformAtomics::InterlockedIncrement(&ConcurrentWriters) == 1);
+			CHECK_CONCURRENT_ACCESS(FPlatformAtomics::InterlockedIncrement(&ConcurrentReaders) == 1);
+
 			IndexData.Empty(1); 
 			IndexData.AddUninitialized();
 			HashData.Empty(1); 
@@ -740,10 +844,16 @@ namespace RobinHoodHashTable_Private
 			KeyValueData.Empty();
 			SizePow2Minus1 = 0;
 			MaximumDistance = 0;
+
+			CHECK_CONCURRENT_ACCESS(FPlatformAtomics::InterlockedDecrement(&ConcurrentReaders) == 0);
+			CHECK_CONCURRENT_ACCESS(FPlatformAtomics::InterlockedDecrement(&ConcurrentWriters) == 0);
 		}
 
 		void Reserve(SizeType ReserveNum)
 		{
+			CHECK_CONCURRENT_ACCESS(FPlatformAtomics::InterlockedIncrement(&ConcurrentWriters) == 1);
+			CHECK_CONCURRENT_ACCESS(FPlatformAtomics::InterlockedIncrement(&ConcurrentReaders) == 1);
+
 			if (ReserveNum > KeyValueData.Num())
 			{
 				KeyValueData.Reserve(ReserveNum);
@@ -775,6 +885,9 @@ namespace RobinHoodHashTable_Private
 					}
 				}
 			}
+
+			CHECK_CONCURRENT_ACCESS(FPlatformAtomics::InterlockedDecrement(&ConcurrentReaders) == 0);
+			CHECK_CONCURRENT_ACCESS(FPlatformAtomics::InterlockedDecrement(&ConcurrentWriters) == 0);
 		}
 
 	private:
@@ -785,6 +898,11 @@ namespace RobinHoodHashTable_Private
 
 		IndexType SizePow2Minus1 = 0;
 		IndexType MaximumDistance = 0;
+
+#if RUN_HASHTABLE_CONCURENCY_CHECKS
+		mutable int ConcurrentReaders = 0;
+		mutable int ConcurrentWriters = 0;
+#endif
 	};
 }
 
@@ -806,64 +924,136 @@ public:
 	TRobinHoodHashMap(TRobinHoodHashMap&& Other) = default;
 	TRobinHoodHashMap& operator=(TRobinHoodHashMap&& Other) = default;
 
+	FHashElementId FindOrAddIdByHash(FHashType HashValue, const KeyType& Key, const ValueType& Val, bool& bIsAlreadyInMap)
+	{
+		return Base::FindOrAddIdByHash(HashValue, Key, Val, bIsAlreadyInMap);
+	}
+
+	FHashElementId FindOrAddIdByHash(FHashType HashValue, const KeyType& Key, ValueType&& Val, bool& bIsAlreadyInMap)
+	{
+		return Base::FindOrAddIdByHash(HashValue, Key, MoveTemp(Val), bIsAlreadyInMap);
+	}
+
+	FHashElementId FindOrAddIdByHash(FHashType HashValue, KeyType&& Key, const ValueType& Val, bool& bIsAlreadyInMap)
+	{
+		return Base::FindOrAddIdByHash(HashValue, MoveTemp(Key), Val, bIsAlreadyInMap);
+	}
+
+	FHashElementId FindOrAddIdByHash(FHashType HashValue, KeyType&& Key, ValueType&& Val, bool& bIsAlreadyInMap)
+	{
+		return Base::FindOrAddIdByHash(HashValue, MoveTemp(Key), MoveTemp(Val), bIsAlreadyInMap);
+	}
+
+	FHashElementId FindOrAddId(const KeyType& Key, const ValueType& Val, bool& bIsAlreadyInMap)
+	{
+		return Base::FindOrAddId(Key, Val, bIsAlreadyInMap);
+	}
+
+	FHashElementId FindOrAddId(const KeyType& Key, ValueType&& Val, bool& bIsAlreadyInMap)
+	{
+		return Base::FindOrAddId(Key, MoveTemp(Val), bIsAlreadyInMap);
+	}
+
+	FHashElementId FindOrAddId(KeyType&& Key, const ValueType& Val, bool& bIsAlreadyInMap)
+	{
+		return Base::FindOrAddId(MoveTemp(Key), Val);
+	}
+
+	FHashElementId FindOrAddId(KeyType&& Key, ValueType&& Val, bool& bIsAlreadyInMap)
+	{
+		return Base::FindOrAddId(MoveTemp(Key), MoveTemp(Val), bIsAlreadyInMap);
+	}
+
+	FindValueType FindOrAdd(const KeyType& Key, const ValueType& Val, bool& bIsAlreadyInMap)
+	{
+		return Base::FindOrAdd(Key, Val, bIsAlreadyInMap);
+	}
+
+	FindValueType FindOrAdd(const KeyType& Key, ValueType&& Val, bool& bIsAlreadyInMap)
+	{
+		return Base::FindOrAdd(Key, MoveTemp(Val), bIsAlreadyInMap);
+	}
+
+	FindValueType FindOrAdd(KeyType&& Key, const ValueType& Val, bool& bIsAlreadyInMap)
+	{
+		return Base::FindOrAdd(MoveTemp(Key), Val, bIsAlreadyInMap);
+	}
+
+	FindValueType FindOrAdd(KeyType&& Key, ValueType&& Val, bool& bIsAlreadyInMap)
+	{
+		return Base::FindOrAdd(MoveTemp(Key), MoveTemp(Val), bIsAlreadyInMap);
+	}
+
 	FHashElementId FindOrAddIdByHash(FHashType HashValue, const KeyType& Key, const ValueType& Val)
 	{
-		return Base::FindOrAddIdByHash(HashValue, Key, Val);
+		bool bIsAlreadyInMap;
+		return Base::FindOrAddIdByHash(HashValue, Key, Val, bIsAlreadyInMap);
 	}
 
 	FHashElementId FindOrAddIdByHash(FHashType HashValue, const KeyType& Key, ValueType&& Val)
 	{
-		return Base::FindOrAddIdByHash(HashValue, Key, MoveTemp(Val));
+		bool bIsAlreadyInMap;
+		return Base::FindOrAddIdByHash(HashValue, Key, MoveTemp(Val), bIsAlreadyInMap);
 	}
 
 	FHashElementId FindOrAddIdByHash(FHashType HashValue, KeyType&& Key, const ValueType& Val)
 	{
-		return Base::FindOrAddIdByHash(HashValue, MoveTemp(Key), Val);
+		bool bIsAlreadyInMap;
+		return Base::FindOrAddIdByHash(HashValue, MoveTemp(Key), Val, bIsAlreadyInMap);
 	}
 
 	FHashElementId FindOrAddIdByHash(FHashType HashValue, KeyType&& Key, ValueType&& Val)
 	{
-		return Base::FindOrAddIdByHash(HashValue, MoveTemp(Key), MoveTemp(Val));
+		bool bIsAlreadyInMap;
+		return Base::FindOrAddIdByHash(HashValue, MoveTemp(Key), MoveTemp(Val), bIsAlreadyInMap);
 	}
 
 	FHashElementId FindOrAddId(const KeyType& Key, const ValueType& Val)
 	{
-		return Base::FindOrAddId(Key, Val);
+		bool bIsAlreadyInMap;
+		return Base::FindOrAddId(Key, Val, bIsAlreadyInMap);
 	}
 
 	FHashElementId FindOrAddId(const KeyType& Key, ValueType&& Val)
 	{
-		return Base::FindOrAddId(Key, MoveTemp(Val));
+		bool bIsAlreadyInMap;
+		return Base::FindOrAddId(Key, MoveTemp(Val), bIsAlreadyInMap);
 	}
 
 	FHashElementId FindOrAddId(KeyType&& Key, const ValueType& Val)
 	{
+		bool bIsAlreadyInMap;
 		return Base::FindOrAddId(MoveTemp(Key), Val);
 	}
 
 	FHashElementId FindOrAddId(KeyType&& Key, ValueType&& Val)
 	{
-		return Base::FindOrAddId(MoveTemp(Key), MoveTemp(Val));
+		bool bIsAlreadyInMap;
+		return Base::FindOrAddId(MoveTemp(Key), MoveTemp(Val), bIsAlreadyInMap);
 	}
 
 	FindValueType FindOrAdd(const KeyType& Key, const ValueType& Val)
 	{
-		return Base::FindOrAdd(Key, Val);
+		bool bIsAlreadyInMap;
+		return Base::FindOrAdd(Key, Val, bIsAlreadyInMap);
 	}
 
 	FindValueType FindOrAdd(const KeyType& Key, ValueType&& Val)
 	{
-		return Base::FindOrAdd(Key, MoveTemp(Val));
+		bool bIsAlreadyInMap;
+		return Base::FindOrAdd(Key, MoveTemp(Val), bIsAlreadyInMap);
 	}
 
 	FindValueType FindOrAdd(KeyType&& Key, const ValueType& Val)
 	{
-		return Base::FindOrAdd(MoveTemp(Key), Val);
+		bool bIsAlreadyInMap;
+		return Base::FindOrAdd(MoveTemp(Key), Val, bIsAlreadyInMap);
 	}
 
 	FindValueType FindOrAdd(KeyType&& Key, ValueType&& Val)
 	{
-		return Base::FindOrAdd(MoveTemp(Key), MoveTemp(Val));
+		bool bIsAlreadyInMap;
+		return Base::FindOrAdd(MoveTemp(Key), MoveTemp(Val), bIsAlreadyInMap);
 	}
 };
 
@@ -886,35 +1076,74 @@ public:
 	TRobinHoodHashSet(TRobinHoodHashSet&& Other) = default;
 	TRobinHoodHashSet& operator=(TRobinHoodHashSet&& Other) = default;
 
+	FHashElementId FindOrAddIdByHash(FHashType HashValue, const KeyType& Key, bool& bIsAlreadyInSet)
+	{
+		return Base::FindOrAddIdByHash(HashValue, Key, Unit(), bIsAlreadyInSet);
+	}
+
+	FHashElementId FindOrAddIdByHash(FHashType HashValue, KeyType&& Key, bool& bIsAlreadyInSet)
+	{
+		return Base::FindOrAddIdByHash(HashValue, MoveTemp(Key), Unit(), bIsAlreadyInSet);
+	}
+
+	FHashElementId FindOrAddId(const KeyType& Key, bool& bIsAlreadyInSet)
+	{
+		return Base::FindOrAddId(Key, Unit(), bIsAlreadyInSet);
+	}
+
+	FHashElementId FindOrAddId(KeyType&& Key, bool& bIsAlreadyInSet)
+	{
+		return Base::FindOrAddId(MoveTemp(Key), Unit(), bIsAlreadyInSet);
+	}
+
+	FindValueType FindOrAdd(const KeyType& Key, bool& bIsAlreadyInSet)
+	{
+		return Base::FindOrAdd(Key, Unit(), bIsAlreadyInSet);
+	}
+
+	FindValueType FindOrAdd(KeyType&& Key, bool& bIsAlreadyInSet)
+	{
+		return Base::FindOrAdd(MoveTemp(Key), Unit(), bIsAlreadyInSet);
+	}
+
 	FHashElementId FindOrAddIdByHash(FHashType HashValue, const KeyType& Key)
 	{
-		return Base::FindOrAddIdByHash(HashValue, Key, Unit());
+		bool bIsAlreadyInSet;
+		return Base::FindOrAddIdByHash(HashValue, Key, Unit(), bIsAlreadyInSet);
 	}
 
 	FHashElementId FindOrAddIdByHash(FHashType HashValue, KeyType&& Key)
 	{
-		return Base::FindOrAddIdByHash(HashValue, MoveTemp(Key), Unit());
+		bool bIsAlreadyInSet;
+		return Base::FindOrAddIdByHash(HashValue, MoveTemp(Key), Unit(), bIsAlreadyInSet);
 	}
 
 	FHashElementId FindOrAddId(const KeyType& Key)
 	{
-		return Base::FindOrAddId(Key, Unit());
+		bool bIsAlreadyInSet;
+		return Base::FindOrAddId(Key, Unit(), bIsAlreadyInSet);
 	}
 
 	FHashElementId FindOrAddId(KeyType&& Key)
 	{
-		return Base::FindOrAddId(MoveTemp(Key), Unit());
+		bool bIsAlreadyInSet;
+		return Base::FindOrAddId(MoveTemp(Key), Unit(), bIsAlreadyInSet);
 	}
 
 	FindValueType FindOrAdd(const KeyType& Key)
 	{
-		return Base::FindOrAdd(Key, Unit());
+		bool bIsAlreadyInSet;
+		return Base::FindOrAdd(Key, Unit(), bIsAlreadyInSet);
 	}
 
 	FindValueType FindOrAdd(KeyType&& Key)
 	{
-		return Base::FindOrAdd(MoveTemp(Key), Unit());
+		bool bIsAlreadyInSet;
+		return Base::FindOrAdd(MoveTemp(Key), Unit(), bIsAlreadyInSet);
 	}
 };
 
-};
\ No newline at end of file
+};
+
+#undef CHECK_CONCURRENT_ACCESS
+#undef RUN_HASHTABLE_CONCURENCY_CHECKS
\ No newline at end of file
diff --git a/Engine/Source/Runtime/Core/Public/FramePro/FramePro.h b/Engine/Source/Runtime/Core/Public/FramePro/FramePro.h
--- a/Engine/Source/Runtime/Core/Public/FramePro/FramePro.h
+++ b/Engine/Source/Runtime/Core/Public/FramePro/FramePro.h
@@ -915,7 +915,11 @@ namespace FramePro
 	#define FRAMEPRO_PLATFORM_UE4 1
 	#define FRAMEPRO_PLATFORM_XBOXONE 0
 	#define FRAMEPRO_PLATFORM_HOLOLENS 0
+#if defined(_WIN32) || defined(_WIN64) || defined(WIN32) || defined(WIN64) || defined(__WIN32__) || defined(__WINDOWS__)
+	#define FRAMEPRO_PLATFORM_WIN 1
+#else
 	#define FRAMEPRO_PLATFORM_WIN 0
+#endif
 	#define FRAMEPRO_PLATFORM_LINUX 0
 	#define FRAMEPRO_PLATFORM_PS4 0
 	#define FRAMEPRO_PLATFORM_ANDROID 0
diff --git a/Engine/Source/Runtime/Core/Public/HAL/PlatformMisc.h b/Engine/Source/Runtime/Core/Public/HAL/PlatformMisc.h
--- a/Engine/Source/Runtime/Core/Public/HAL/PlatformMisc.h
+++ b/Engine/Source/Runtime/Core/Public/HAL/PlatformMisc.h
@@ -25,6 +25,10 @@
 #endif
 #endif
 
+#ifndef PLATFORM_EMPTY_BASES
+#define PLATFORM_EMPTY_BASES
+#endif
+
 // Master switch for scoped named events
 #define ENABLE_NAMED_EVENTS (!UE_BUILD_SHIPPING && 1)
 
diff --git a/Engine/Source/Runtime/Core/Public/HAL/PreprocessorHelpers.h b/Engine/Source/Runtime/Core/Public/HAL/PreprocessorHelpers.h
--- a/Engine/Source/Runtime/Core/Public/HAL/PreprocessorHelpers.h
+++ b/Engine/Source/Runtime/Core/Public/HAL/PreprocessorHelpers.h
@@ -10,6 +10,10 @@
 #define PREPROCESSOR_JOIN(x, y) PREPROCESSOR_JOIN_INNER(x, y)
 #define PREPROCESSOR_JOIN_INNER(x, y) x##y
 
+// Concatenates the first two preprocessor tokens of a variadic list, after performing macro expansion on them
+#define PREPROCESSOR_JOIN_FIRST(x, ...) PREPROCESSOR_JOIN_FIRST_INNER(x, __VA_ARGS__)
+#define PREPROCESSOR_JOIN_FIRST_INNER(x, ...) x##__VA_ARGS__
+
 // Expands to the second argument or the third argument if the first argument is 1 or 0 respectively
 #define PREPROCESSOR_IF(cond, x, y) PREPROCESSOR_JOIN(PREPROCESSOR_IF_INNER_, cond)(x, y)
 #define PREPROCESSOR_IF_INNER_1(x, y) x
@@ -21,7 +25,16 @@
 // Expands to nothing - used as a placeholder
 #define PREPROCESSOR_NOTHING
 
-
+// Removes a single layer of parentheses from a macro argument if they are present - used to allow
+// brackets to be optionally added when the argument contains commas, e.g.:
+//
+// #define DEFINE_VARIABLE(Type, Name) PREPROCESSOR_REMOVE_OPTIONAL_PARENS(Type) Name;
+//
+// DEFINE_VARIABLE(int, IntVar)                  // expands to: int IntVar;
+// DEFINE_VARIABLE((TPair<int, float>), PairVar) // expands to: TPair<int, float> PairVar;
+#define PREPROCESSOR_REMOVE_OPTIONAL_PARENS(...) PREPROCESSOR_JOIN_FIRST(PREPROCESSOR_REMOVE_OPTIONAL_PARENS_IMPL,PREPROCESSOR_REMOVE_OPTIONAL_PARENS_IMPL __VA_ARGS__)
+#define PREPROCESSOR_REMOVE_OPTIONAL_PARENS_IMPL(...) PREPROCESSOR_REMOVE_OPTIONAL_PARENS_IMPL __VA_ARGS__
+#define PREPROCESSOR_REMOVE_OPTIONAL_PARENS_IMPLPREPROCESSOR_REMOVE_OPTIONAL_PARENS_IMPL
 
 // setup standardized way of including platform headers from the "uber-platform" headers like PlatformFile.h
 #ifdef OVERRIDE_PLATFORM_HEADER_NAME
diff --git a/Engine/Source/Runtime/Core/Public/HAL/UnrealMemory.h b/Engine/Source/Runtime/Core/Public/HAL/UnrealMemory.h
--- a/Engine/Source/Runtime/Core/Public/HAL/UnrealMemory.h
+++ b/Engine/Source/Runtime/Core/Public/HAL/UnrealMemory.h
@@ -167,6 +167,14 @@ struct CORE_API FMemory
 	static void* Realloc(void* Original, SIZE_T Count, uint32 Alignment = DEFAULT_ALIGNMENT);
 	static void Free(void* Original);
 	static SIZE_T GetAllocSize(void* Original);
+
+	static FORCEINLINE_DEBUGGABLE void* MallocZeroed(SIZE_T Count, uint32 Alignment = DEFAULT_ALIGNMENT)
+	{
+		void* Memory = Malloc(Count, Alignment);
+		Memzero(Memory, Count);
+		return Memory;
+	}
+
 	/**
 	* For some allocators this will return the actual size that should be requested to eliminate
 	* internal fragmentation. The return value will always be >= Count. This can be used to grow
diff --git a/Engine/Source/Runtime/Core/Public/Logging/LogMacros.h b/Engine/Source/Runtime/Core/Public/Logging/LogMacros.h
--- a/Engine/Source/Runtime/Core/Public/Logging/LogMacros.h
+++ b/Engine/Source/Runtime/Core/Public/Logging/LogMacros.h
@@ -112,6 +112,9 @@ private:
 			UE_LOG_EXPAND_IS_FATAL(Verbosity, CA_ASSUME(false);, PREPROCESSOR_NOTHING) \
 		} \
 	}
+
+	#define UE_LOG_CLINKAGE(CategoryName, Verbosity, Format, ...) UE_LOG(CategoryName, Verbosity, Format, __VA_ARGS__ )
+
 	// Conditional logging (fatal errors only).
 	#define UE_CLOG(Condition, CategoryName, Verbosity, Format, ...) \
 	{ \
@@ -187,6 +190,40 @@ private:
 	 * @param Format, format text
 	 ***/
 	#define UE_LOG(CategoryName, Verbosity, Format, ...) \
+	{ \
+		static_assert(TIsArrayOrRefOfType<decltype(Format), TCHAR>::Value, "Formatting string must be a TCHAR array."); \
+		static_assert((ELogVerbosity::Verbosity & ELogVerbosity::VerbosityMask) < ELogVerbosity::NumVerbosity && ELogVerbosity::Verbosity > 0, "Verbosity must be constant and in range."); \
+		CA_CONSTANT_IF((ELogVerbosity::Verbosity & ELogVerbosity::VerbosityMask) <= ELogVerbosity::COMPILED_IN_MINIMUM_VERBOSITY && (ELogVerbosity::Warning & ELogVerbosity::VerbosityMask) <= FLogCategory##CategoryName::CompileTimeVerbosity) \
+		{ \
+			UE_LOG_EXPAND_IS_FATAL(Verbosity, PREPROCESSOR_NOTHING, if (!CategoryName.IsSuppressed(ELogVerbosity::Verbosity))) \
+			{ \
+				auto UE_LOG_noinline_lambda = [](const auto& LCategoryName, const auto& LFormat, const auto&... UE_LOG_Args) FORCENOINLINE \
+				{ \
+					TRACE_LOG_MESSAGE(LCategoryName, Verbosity, LFormat, UE_LOG_Args...) \
+					UE_LOG_EXPAND_IS_FATAL(Verbosity, \
+						{ \
+							FMsg::Logf_Internal(UE_LOG_SOURCE_FILE(__FILE__), __LINE__, LCategoryName.GetCategoryName(), ELogVerbosity::Verbosity, LFormat, UE_LOG_Args...); \
+							_DebugBreakAndPromptForRemote(); \
+							FDebug::ProcessFatalError(); \
+							CA_ASSUME(false); \
+						}, \
+						{ \
+							FMsg::Logf_Internal(nullptr, 0, LCategoryName.GetCategoryName(), ELogVerbosity::Verbosity, LFormat, UE_LOG_Args...); \
+						} \
+					) \
+				}; \
+				UE_LOG_noinline_lambda(CategoryName, Format, ##__VA_ARGS__); \
+			} \
+		} \
+	}
+
+	/** 
+	 * A  macro that outputs a formatted message to log if a given logging category is active at a given verbosity level
+	 * @param CategoryName name of the logging category
+	 * @param Verbosity, verbosity level to test against
+	 * @param Format, format text
+	 ***/
+	#define UE_LOG_CLINKAGE(CategoryName, Verbosity, Format, ...) \
 	{ \
 		static_assert(TIsArrayOrRefOfType<decltype(Format), TCHAR>::Value, "Formatting string must be a TCHAR array."); \
 		static_assert((ELogVerbosity::Verbosity & ELogVerbosity::VerbosityMask) < ELogVerbosity::NumVerbosity && ELogVerbosity::Verbosity > 0, "Verbosity must be constant and in range."); \
@@ -197,19 +234,18 @@ private:
 				TRACE_LOG_MESSAGE(CategoryName, Verbosity, Format, ##__VA_ARGS__) \
 				UE_LOG_EXPAND_IS_FATAL(Verbosity, \
 					{ \
-						FMsg::Logf_Internal(UE_LOG_SOURCE_FILE(__FILE__), __LINE__, CategoryName.GetCategoryName(), ELogVerbosity::Verbosity, Format, ##__VA_ARGS__); \
+						FMsg::Logf_Internal(UE_LOG_SOURCE_FILE(__FILE__), __LINE__, CategoryName.GetCategoryName(), ELogVerbosity::Verbosity, Format,  ##__VA_ARGS__); \
 						_DebugBreakAndPromptForRemote(); \
 						FDebug::ProcessFatalError(); \
 						CA_ASSUME(false); \
 					}, \
 					{ \
-						FMsg::Logf_Internal(nullptr, 0, CategoryName.GetCategoryName(), ELogVerbosity::Verbosity, Format, ##__VA_ARGS__); \
+						FMsg::Logf_Internal(nullptr, 0, CategoryName.GetCategoryName(), ELogVerbosity::Verbosity, Format,  ##__VA_ARGS__); \
 					} \
 				) \
 			} \
 		} \
 	}
-
 	/**
 	* A  macro that outputs a formatted message to the log specifically used for security events
 	* @param NetConnection, a valid UNetConnection
@@ -240,19 +276,23 @@ private:
 			{ \
 				if (Condition) \
 				{ \
-					TRACE_LOG_MESSAGE(CategoryName, Verbosity, Format, ##__VA_ARGS__) \
-					UE_LOG_EXPAND_IS_FATAL(Verbosity, \
-						{ \
-							FMsg::Logf_Internal(UE_LOG_SOURCE_FILE(__FILE__), __LINE__, CategoryName.GetCategoryName(), ELogVerbosity::Verbosity, Format, ##__VA_ARGS__); \
-							_DebugBreakAndPromptForRemote(); \
-							FDebug::ProcessFatalError(); \
-							CA_ASSUME(false); \
-						}, \
-						{ \
-							FMsg::Logf_Internal(nullptr, 0, CategoryName.GetCategoryName(), ELogVerbosity::Verbosity, Format, ##__VA_ARGS__); \
-						} \
-					) \
-					CA_ASSUME(true); \
+					auto UE_LOG_noinline_lambda = [](const auto& LCategoryName, const auto& LFormat, const auto&... UE_LOG_Args) FORCENOINLINE \
+					{ \
+						TRACE_LOG_MESSAGE(LCategoryName, Verbosity, LFormat, UE_LOG_Args...) \
+						UE_LOG_EXPAND_IS_FATAL(Verbosity, \
+							{ \
+								FMsg::Logf_Internal(UE_LOG_SOURCE_FILE(__FILE__), __LINE__, LCategoryName.GetCategoryName(), ELogVerbosity::Verbosity, LFormat, UE_LOG_Args...); \
+								_DebugBreakAndPromptForRemote(); \
+								FDebug::ProcessFatalError(); \
+								CA_ASSUME(false); \
+							}, \
+							{ \
+								FMsg::Logf_Internal(nullptr, 0, LCategoryName.GetCategoryName(), ELogVerbosity::Verbosity, LFormat, UE_LOG_Args...); \
+							} \
+						) \
+						CA_ASSUME(true); \
+					}; \
+					UE_LOG_noinline_lambda(CategoryName, Format, ##__VA_ARGS__); \
 				} \
 			} \
 		} \
diff --git a/Engine/Source/Runtime/Core/Public/Math/Color.h b/Engine/Source/Runtime/Core/Public/Math/Color.h
--- a/Engine/Source/Runtime/Core/Public/Math/Color.h
+++ b/Engine/Source/Runtime/Core/Public/Math/Color.h
@@ -8,6 +8,7 @@
 #include "Containers/UnrealString.h"
 #include "Misc/Parse.h"
 #include "Serialization/StructuredArchive.h"
+#include "Serialization/MemoryLayout.h"
 
 class FFloat16Color;
 
@@ -391,6 +392,7 @@ struct FLinearColor
 	static CORE_API const FLinearColor Blue;
 	static CORE_API const FLinearColor Yellow;
 };
+DECLARE_INTRINSIC_TYPE_LAYOUT(FLinearColor);
 
 FORCEINLINE FLinearColor operator*(float Scalar,const FLinearColor& Color)
 {
@@ -638,7 +640,7 @@ private:
 	 */
 	explicit FColor(const FLinearColor& LinearColor);
 };
-
+DECLARE_INTRINSIC_TYPE_LAYOUT(FColor);
 
 FORCEINLINE uint32 GetTypeHash( const FColor& Color )
 {
