Copying Private-LoadTimes-4.24 stream to Main. Biggest changes are in Materials/Shader memory freezing.

diff --git a/Engine/Source/Runtime/Core/Public/Math/Float16.h b/Engine/Source/Runtime/Core/Public/Math/Float16.h
--- a/Engine/Source/Runtime/Core/Public/Math/Float16.h
+++ b/Engine/Source/Runtime/Core/Public/Math/Float16.h
@@ -6,6 +6,7 @@
 #include "Serialization/Archive.h"
 #include "Math/UnrealMathUtility.h"
 #include "Math/Float32.h"
+#include "Serialization/MemoryLayout.h"
 
 /**
 * 16 bit float components and conversion
@@ -91,7 +92,9 @@ public:
 		return Ar << V.Encoded;
 	}
 };
+template<> struct TCanBulkSerialize<FFloat16> { enum { Value = true }; };
 
+DECLARE_INTRINSIC_TYPE_LAYOUT(FFloat16);
 
 FORCEINLINE FFloat16::FFloat16()
 	:	Encoded(0)
diff --git a/Engine/Source/Runtime/Core/Public/Math/Vector.h b/Engine/Source/Runtime/Core/Public/Math/Vector.h
--- a/Engine/Source/Runtime/Core/Public/Math/Vector.h
+++ b/Engine/Source/Runtime/Core/Public/Math/Vector.h
@@ -18,6 +18,7 @@
 #include "Internationalization/Internationalization.h"
 #include "Math/IntVector.h"
 #include "Math/Axis.h"
+#include "Serialization/MemoryLayout.h"
 
 #if PLATFORM_VECTOR_CUBIC_INTERP_SSE
 #include "Math/UnrealMathSSE.h"
@@ -1012,6 +1013,9 @@ public:
 	 */
 	CORE_API bool NetSerialize(FArchive& Ar, class UPackageMap* Map, bool& bOutSuccess);
 };
+template<> struct TCanBulkSerialize<FVector> { enum { Value = true }; };
+
+DECLARE_INTRINSIC_TYPE_LAYOUT(FVector);
 
 
 /* FVector inline functions
diff --git a/Engine/Source/Runtime/Core/Public/Math/Vector4.h b/Engine/Source/Runtime/Core/Public/Math/Vector4.h
--- a/Engine/Source/Runtime/Core/Public/Math/Vector4.h
+++ b/Engine/Source/Runtime/Core/Public/Math/Vector4.h
@@ -10,6 +10,7 @@
 #include "Logging/LogMacros.h"
 #include "Math/Vector2D.h"
 #include "Math/Vector.h"
+#include "Serialization/MemoryLayout.h"
 
 /**
  * A 4D homogeneous vector, 4x1 FLOATs, 16-byte aligned.
@@ -433,6 +434,8 @@ public:
 
 } GCC_ALIGN(16);
 
+DECLARE_INTRINSIC_TYPE_LAYOUT(FVector4);
+
 
 /**
  * Creates a hash value from a FVector4.
diff --git a/Engine/Source/Runtime/Core/Public/Misc/DelayedAutoRegister.h b/Engine/Source/Runtime/Core/Public/Misc/DelayedAutoRegister.h
new file mode 100644
--- /dev/null
+++ b/Engine/Source/Runtime/Core/Public/Misc/DelayedAutoRegister.h
@@ -0,0 +1,28 @@
+// Copyright 1998-2019 Epic Games, Inc. All Rights Reserved.
+
+#pragma once
+
+#include "Templates/Function.h"
+
+enum class EDelayedRegisterRunPhase : uint8
+{
+	StartOfEnginePreInit,
+	FileSystemReady,
+	StatSystemReady,
+	IniSystemReady,
+	TaskGraphSystemReady,
+	ShaderTypesReady,
+	PreObjectSystemReady,
+	ObjectSystemReady,
+	EndOfEngineInit,
+
+	NumPhases,
+};
+
+struct CORE_API FDelayedAutoRegisterHelper
+{
+
+	FDelayedAutoRegisterHelper(EDelayedRegisterRunPhase RunPhase, TFunction<void()> RegistrationFunction);
+
+	static void RunAndClearDelayedAutoRegisterDelegates(EDelayedRegisterRunPhase RunPhase);
+};
diff --git a/Engine/Source/Runtime/Core/Public/Misc/Guid.h b/Engine/Source/Runtime/Core/Public/Misc/Guid.h
--- a/Engine/Source/Runtime/Core/Public/Misc/Guid.h
+++ b/Engine/Source/Runtime/Core/Public/Misc/Guid.h
@@ -7,6 +7,8 @@
 #include "Misc/Crc.h"
 #include "Containers/UnrealString.h"
 #include "Serialization/StructuredArchive.h"
+#include "Serialization/MemoryLayout.h"
+#include "Hash/CityHash.h"
 
 class FArchive;
 class FOutputDevice;
@@ -299,7 +301,7 @@ public:
 	 */
 	friend uint32 GetTypeHash(const FGuid& Guid)
 	{
-		return FCrc::MemCrc_DEPRECATED(&Guid, sizeof(FGuid));
+		return uint32(CityHash64((char*)&Guid, sizeof(FGuid)));
 	}
 
 public:
@@ -347,6 +349,7 @@ public:
 	/** Holds the fourth component. */
 	uint32 D;
 };
-
+template<> struct TCanBulkSerialize<FGuid> { enum { Value = true }; };
+DECLARE_INTRINSIC_TYPE_LAYOUT(FGuid);
 
 template <> struct TIsPODType<FGuid> { enum { Value = true }; };
diff --git a/Engine/Source/Runtime/Core/Public/Misc/PreLoadFile.h b/Engine/Source/Runtime/Core/Public/Misc/PreLoadFile.h
new file mode 100644
--- /dev/null
+++ b/Engine/Source/Runtime/Core/Public/Misc/PreLoadFile.h
@@ -0,0 +1,24 @@
+// Copyright 1998-2019 Epic Games, Inc. All Rights Reserved.
+
+#pragma once
+
+#include "Misc/DelayedAutoRegister.h"
+
+struct CORE_API FPreLoadFile : public FDelayedAutoRegisterHelper
+{
+	FPreLoadFile(const TCHAR* InPath);
+
+	void* TakeOwnershipOfLoadedData(int64* OutFileSize=nullptr);
+	static void* TakeOwnershipOfLoadedDataByPath(const TCHAR* Filename, int64* OutFileSize);
+
+protected:
+	void KickOffRead();
+
+	bool bIsComplete;
+	static bool bSystemNoLongerTakingRequests;
+	void* Data;
+	int64 FileSize;
+	FString Path;
+	FEvent* CompletionEvent;
+	class IAsyncReadFileHandle* AsyncReadHandle;
+};
diff --git a/Engine/Source/Runtime/Core/Public/Misc/ScopeLock.h b/Engine/Source/Runtime/Core/Public/Misc/ScopeLock.h
--- a/Engine/Source/Runtime/Core/Public/Misc/ScopeLock.h
+++ b/Engine/Source/Runtime/Core/Public/Misc/ScopeLock.h
@@ -42,9 +42,18 @@ public:
 	/** Destructor that performs a release on the synchronization object. */
 	~FScopeLock()
 	{
-		check(SynchObject);
-		SynchObject->Unlock();
+		Unlock();
 	}
+
+	void Unlock()
+	{
+		if(SynchObject)
+		{
+			SynchObject->Unlock();
+			SynchObject = nullptr;
+		}
+	}
+
 private:
 
 	/** Default constructor (hidden on purpose). */
diff --git a/Engine/Source/Runtime/Core/Public/Misc/SecureHash.h b/Engine/Source/Runtime/Core/Public/Misc/SecureHash.h
--- a/Engine/Source/Runtime/Core/Public/Misc/SecureHash.h
+++ b/Engine/Source/Runtime/Core/Public/Misc/SecureHash.h
@@ -13,6 +13,7 @@
 #include "Serialization/BufferReader.h"
 #include "String/BytesToHex.h"
 #include "String/HexToBytes.h"
+#include "Serialization/MemoryLayout.h"
 
 class FAnsiStringBuilderBase;
 class FStringBuilderBase;
@@ -227,6 +228,11 @@ public:
 		return FMemory::Memcmp(&X.Hash, &Y.Hash, sizeof(X.Hash)) != 0;
 	}
 
+	friend bool operator<(const FSHAHash& X, const FSHAHash& Y)
+	{
+		return FMemory::Memcmp(&X.Hash, &Y.Hash, sizeof(X.Hash)) < 0;
+	}
+
 	friend CORE_API FArchive& operator<<( FArchive& Ar, FSHAHash& G );
 	
 	friend uint32 GetTypeHash(const FSHAHash& InKey)
@@ -237,6 +243,7 @@ public:
 	friend CORE_API FString LexToString(const FSHAHash&);
 	friend CORE_API void LexFromString(FSHAHash& Hash, const TCHAR*);
 };
+DECLARE_INTRINSIC_TYPE_LAYOUT(FSHAHash);
 
 inline FStringBuilderBase& operator<<(FStringBuilderBase& Builder, const FSHAHash& Hash) { UE::String::BytesToHex(Hash.Hash, Builder); return Builder; }
 inline FAnsiStringBuilderBase& operator<<(FAnsiStringBuilderBase& Builder, const FSHAHash& Hash) { UE::String::BytesToHex(Hash.Hash, Builder); return Builder; }
diff --git a/Engine/Source/Runtime/Core/Public/ProfilingDebugging/LoadTimeTracker.h b/Engine/Source/Runtime/Core/Public/ProfilingDebugging/LoadTimeTracker.h
--- a/Engine/Source/Runtime/Core/Public/ProfilingDebugging/LoadTimeTracker.h
+++ b/Engine/Source/Runtime/Core/Public/ProfilingDebugging/LoadTimeTracker.h
@@ -13,11 +13,16 @@
 #include "ProfilingDebugging/ScopedTimers.h"
 #include "Trace/Config.h"
 #include "Serialization/LoadTimeTrace.h"
+#include "Stats/Stats.h"
 
 #ifndef ENABLE_LOADTIME_TRACKING
 	#define ENABLE_LOADTIME_TRACKING 0
 #endif
 
+#ifndef ENABLE_LOADTIME_TRACKING_WITH_STATS
+	#define ENABLE_LOADTIME_TRACKING_WITH_STATS 0
+#endif
+
 #define ENABLE_LOADTIME_RAW_TIMINGS 0
 
 /** High level load time tracker utility (such as initial engine startup or game specific timings) */
@@ -217,3 +222,11 @@ struct CORE_API FScopedLoadTimeAccumulatorTimer : public FScopedDurationTimer
 #define SCOPED_LOADTIMER(TimerName) TRACE_CPUPROFILER_EVENT_SCOPE_ON_CHANNEL(TimerName, LoadTimeChannel)
 #define SCOPED_LOADTIMER_CNT(TimerName)
 #endif
+
+#if ENABLE_LOADTIME_TRACKING_WITH_STATS && STATS
+	#define SCOPED_ACCUM_LOADTIME_STAT(InstanceName) FSimpleScopeSecondsStat ScopeTimer(FDynamicStats::CreateStatIdDouble<FStatGroup_STATGROUP_LoadTimeClass>(InstanceName, true), 1000.0);
+	#define ACCUM_LOADTIMECOUNT_STAT(InstanceName) INC_DWORD_STAT_FNAME_BY(FDynamicStats::CreateStatIdInt64<FStatGroup_STATGROUP_LoadTimeClassCount>(InstanceName+TEXT("_Count"), true).GetName(), 1);
+#else
+	#define SCOPED_ACCUM_LOADTIME_STAT(...)
+	#define ACCUM_LOADTIMECOUNT_STAT(...)
+#endif
diff --git a/Engine/Source/Runtime/Core/Public/Serialization/MemoryImage.h b/Engine/Source/Runtime/Core/Public/Serialization/MemoryImage.h
new file mode 100644
--- /dev/null
+++ b/Engine/Source/Runtime/Core/Public/Serialization/MemoryImage.h
@@ -0,0 +1,930 @@
+#pragma once
+
+#include "CoreTypes.h"
+#include "Containers/Array.h"
+#include "Containers/Set.h"
+#include "Containers/UnrealString.h"
+#include "Misc/SecureHash.h"
+#include "Serialization/Archive.h"
+#include "Serialization/MemoryLayout.h"
+#include "Serialization/MemoryImageWriter.h"
+#include "Templates/RefCounting.h"
+
+#if defined(WITH_RTTI) || defined(_CPPRTTI) || defined(__GXX_RTTI) || WITH_EDITOR
+#include <typeinfo>
+#endif
+
+class FMemoryImage;
+
+class FPointerTableBase
+{
+public:
+	virtual ~FPointerTableBase() {}
+	virtual int32 AddIndexedPointer(const FTypeLayoutDesc& TypeDesc, void* Ptr) = 0;
+	virtual void* GetIndexedPointer(const FTypeLayoutDesc& TypeDesc, uint32 i) const = 0;
+};
+
+struct FMemoryImageVTablePointer
+{
+	uint64 TypeNameHash;
+	uint32 VTableOffset;
+	uint32 Offset;
+};
+inline bool operator==(const FMemoryImageVTablePointer& Lhs, const FMemoryImageVTablePointer& Rhs)
+{
+	return Lhs.TypeNameHash == Rhs.TypeNameHash && Lhs.VTableOffset == Rhs.VTableOffset && Lhs.Offset == Rhs.Offset;
+}
+inline bool operator!=(const FMemoryImageVTablePointer& Lhs, const FMemoryImageVTablePointer& Rhs)
+{
+	return !operator==(Lhs, Rhs);
+}
+inline bool operator<(const FMemoryImageVTablePointer& Lhs, const FMemoryImageVTablePointer& Rhs)
+{
+	if (Lhs.TypeNameHash != Rhs.TypeNameHash) return Lhs.TypeNameHash < Rhs.TypeNameHash;
+	if (Lhs.VTableOffset != Rhs.VTableOffset) return Lhs.VTableOffset < Rhs.VTableOffset;
+	return Lhs.Offset < Rhs.Offset;
+}
+
+struct FMemoryImageNamePointer
+{
+	FName Name;
+	uint32 Offset;
+};
+inline bool operator==(const FMemoryImageNamePointer& Lhs, const FMemoryImageNamePointer& Rhs)
+{
+	return Lhs.Offset == Rhs.Offset && Lhs.Name == Rhs.Name;
+}
+inline bool operator!=(const FMemoryImageNamePointer& Lhs, const FMemoryImageNamePointer& Rhs)
+{
+	return !operator==(Lhs, Rhs);
+}
+inline bool operator<(const FMemoryImageNamePointer& Lhs, const FMemoryImageNamePointer& Rhs)
+{
+	if (Lhs.Name != Rhs.Name) return Lhs.Name.FastLess(Rhs.Name);
+	return Lhs.Offset < Rhs.Offset;
+}
+
+struct FMemoryImageResult
+{
+	TArray<uint8> Bytes;
+	TArray<FMemoryImageVTablePointer> VTables;
+	TArray<FMemoryImageNamePointer> Names;
+
+	CORE_API void SaveToArchive(FArchive& Ar) const;
+	CORE_API void ApplyPatches(void* FrozenObject) const;
+	CORE_API static void ApplyPatchesFromArchive(void* FrozenObject, FArchive& Ar);
+};
+
+class CORE_API FMemoryImageSection : public FRefCountedObject
+{
+public:
+	struct FSectionPointer
+	{
+		uint32 SectionIndex;
+		uint32 PointerOffset;
+		uint32 Offset;
+	};
+
+	FMemoryImageSection(FMemoryImage* InImage, FString InName)
+		: ParentImage(InImage)
+		, DebugName(InName)
+		, MaxAlignment(1u)
+	{}
+
+	uint32 GetOffset() const { return Bytes.Num(); }
+
+	uint32 WriteAlignment(uint32 Alignment)
+	{
+		const uint32 PrevSize = Bytes.Num();
+		const uint32 Offset = Align(PrevSize, Alignment);
+		Bytes.SetNumZeroed(Offset);
+		MaxAlignment = FMath::Max(MaxAlignment, Alignment);
+		return Offset;
+	}
+
+	void WritePaddingToSize(uint32 Offset)
+	{
+		check(Offset >= (uint32)Bytes.Num());
+		Bytes.SetNumZeroed(Offset);
+	}
+
+	uint32 WriteBytes(const void* Data, uint32 Size)
+	{
+		const uint32 Offset = GetOffset();
+		Bytes.SetNumUninitialized(Offset + Size);
+		FMemory::Memcpy(Bytes.GetData() + Offset, Data, Size);
+		return Offset;
+	}
+
+	template<typename T>
+	uint32 WriteBytes(const T& Data) { return WriteBytes(&Data, sizeof(T)); }
+
+	FMemoryImageSection* WritePointer(const FString& SectionName, uint32 Offset = 0u);
+	uint32 WriteMemoryImagePointerSizedBytes(uint64 PointerValue);
+	uint32 WriteRawPointerSizedBytes(uint64 PointerValue);
+	uint32 WriteVTable(const FTypeLayoutDesc& TypeDesc, const FTypeLayoutDesc& DerivedTypeDesc);
+	uint32 WriteFName(const FName& Name);
+	uint32 Flatten(FMemoryImageResult& OutResult) const;
+
+	void ComputeHash();
+
+	FMemoryImage* ParentImage;
+	FString DebugName;
+	TArray<uint8> Bytes;
+	TArray<FSectionPointer> Pointers;
+	TArray<FMemoryImageVTablePointer> VTables;
+	TArray<FMemoryImageNamePointer> Names;
+	FSHAHash Hash;
+	uint32 MaxAlignment;
+};
+
+class CORE_API FMemoryImage
+{
+public:
+	FMemoryImage()
+		: PointerTable(nullptr)
+		, PrevPointerTable(nullptr)
+		, CurrentStruct(nullptr)
+	{
+		HostLayoutParameters.InitializeForCurrent();
+	}
+
+	FPointerTableBase& GetPointerTable() const { check(PointerTable); return *PointerTable; }
+	const FPointerTableBase& GetPrevPointerTable() const { check(PrevPointerTable); return *PrevPointerTable; }
+
+	FMemoryImageSection* AllocateSection(const FString& Name)
+	{
+		FMemoryImageSection* Section = new FMemoryImageSection(this, Name);
+		Sections.Add(Section);
+		return Section;
+	}
+
+	void AddDependency(const FTypeLayoutDesc& TypeDesc);
+
+	/** Merging duplicate sections will make the resulting memory image smaller.
+	 * This will only work for data that is expected to be read-only after freezing.  Merging sections will break any manual fix-ups applied to the frozen data
+	 */
+	void Flatten(FMemoryImageResult& OutResult, bool bMergeDuplicateSections = false);
+
+	TArray<TRefCountPtr<FMemoryImageSection>> Sections;
+	TArray<const FTypeLayoutDesc*> TypeDependencies;
+	FPointerTableBase* PointerTable;
+	const FPointerTableBase* PrevPointerTable;
+	FPlatformTypeLayoutParameters HostLayoutParameters;
+	FPlatformTypeLayoutParameters TargetLayoutParameters;
+	const class UStruct* CurrentStruct;
+};
+
+template<typename T>
+class TMemoryImagePtr
+{
+public:
+	inline bool IsFrozen() const { return OffsetFromThis & IsFrozenMask; }
+	inline bool IsValid() const { return Ptr != nullptr; }
+	inline bool IsNull() const { return Ptr == nullptr; }
+
+	inline TMemoryImagePtr(T* InPtr = nullptr) : Ptr(InPtr) {}
+	inline TMemoryImagePtr(const TMemoryImagePtr<T>& InPtr) : Ptr(InPtr.Get()) {}
+	inline TMemoryImagePtr& operator=(T* InPtr) { Ptr = InPtr; check((OffsetFromThis & AllFlags) == 0u); return *this; }
+	inline TMemoryImagePtr& operator=(const TMemoryImagePtr<T>& InPtr) { Ptr = InPtr.Get(); check((OffsetFromThis & AllFlags) == 0u); return *this; }
+
+	inline ~TMemoryImagePtr() 
+	{
+
+	}
+
+	inline FMemoryImagePtrInt GetFrozenOffsetFromThis() const { check(IsFrozen()); return (OffsetFromThis >> OffsetShift); }
+
+	inline T* Get() const
+	{
+		return IsFrozen() ? GetFrozenPtrInternal() : Ptr;
+	}
+
+	inline T* GetChecked() const { T* Value = Get(); check(Value); return Value; }
+	inline T* operator->() const { return GetChecked(); }
+	inline T& operator*() const { return *GetChecked(); }
+	inline operator T*() const { return Get(); }
+
+	void SafeDelete(FPointerTableBase* PtrTable = nullptr)
+	{
+		T* RawPtr = Get();
+		if (RawPtr)
+		{
+			DeleteObjectFromLayout(RawPtr, PtrTable, IsFrozen());
+			Ptr = nullptr;
+		}
+	}
+
+	void WriteMemoryImageWithDerivedType(FMemoryImageWriter& Writer, const FTypeLayoutDesc* DerivedTypeDesc) const
+	{
+		const T* RawPtr = Get();
+		if (RawPtr)
+		{
+			check(DerivedTypeDesc);
+			// Compile-time type of the thing we're pointing to
+			const FTypeLayoutDesc& StaticTypeDesc = StaticGetTypeLayoutDesc<T>();
+			// 'this' offset to adjust from the compile-time type to the run-time type
+			const uint32 OffsetToBase = DerivedTypeDesc->GetOffsetToBase(StaticTypeDesc);
+
+			FMemoryImageWriter PointerWriter = Writer.WritePointer(FString::Printf(TEXT("TMemoryImagePtr<%s>"), DerivedTypeDesc->Name), OffsetToBase);
+			PointerWriter.WriteObject((uint8*)RawPtr - OffsetToBase, *DerivedTypeDesc);
+		}
+		else
+		{
+			Writer.WriteMemoryImagePointerSizedBytes(0u);
+		}
+	}
+
+private:
+	inline T* GetFrozenPtrInternal() const
+	{
+		return (T*)((char*)this + (OffsetFromThis >> OffsetShift));
+	}
+
+	enum
+	{
+		IsFrozenMask = (1 << 0),
+		AllFlags = IsFrozenMask,
+		OffsetShift = 1u,
+	};
+
+protected:
+	union
+	{
+		T* Ptr;
+		FMemoryImagePtrInt OffsetFromThis;
+	};
+};
+
+namespace Freeze
+{
+	template<typename T>
+	void IntrinsicWriteMemoryImage(FMemoryImageWriter& Writer, const TMemoryImagePtr<T>& Object, const FTypeLayoutDesc&)
+	{
+		T* RawPtr = Object.Get();
+		if (RawPtr)
+		{
+			// Compile-time type of the thing we're pointing to
+			const FTypeLayoutDesc& StaticTypeDesc = StaticGetTypeLayoutDesc<T>();
+			// Actual run-time type of the thing we're pointing to
+			const FTypeLayoutDesc& DerivedTypeDesc = GetTypeLayoutDesc(Writer.TryGetPrevPointerTable(), *RawPtr);
+			// 'this' offset to adjust from the compile-time type to the run-time type
+			const uint32 OffsetToBase = DerivedTypeDesc.GetOffsetToBase(StaticTypeDesc);
+
+			FMemoryImageWriter PointerWriter = Writer.WritePointer(FString::Printf(TEXT("TMemoryImagePtr<%s>"), DerivedTypeDesc.Name), OffsetToBase);
+			PointerWriter.WriteObject((uint8*)RawPtr - OffsetToBase, DerivedTypeDesc);
+		}
+		else
+		{
+			Writer.WriteMemoryImagePointerSizedBytes(0u);
+		}
+	}
+
+	template<typename T>
+	void IntrinsicUnfrozenCopy(const FMemoryUnfreezeContent& Context, const TMemoryImagePtr<T>& Object, void* OutDst)
+	{
+		T* RawPtr = Object.Get();
+		if (RawPtr)
+		{
+			// Compile-time type of the thing we're pointing to
+			const FTypeLayoutDesc& StaticTypeDesc = StaticGetTypeLayoutDesc<T>();
+			// Actual run-time type of the thing we're pointing to
+			const FTypeLayoutDesc& DerivedTypeDesc = GetTypeLayoutDesc(Context.TryGetPrevPointerTable(), *RawPtr);
+			// 'this' offset to adjust from the compile-time type to the run-time type
+			const uint32 OffsetToBase = DerivedTypeDesc.GetOffsetToBase(StaticTypeDesc);
+
+			void* UnfrozenMemory = FMemory::Malloc(DerivedTypeDesc.Size, DerivedTypeDesc.Alignment);
+			Context.UnfreezeObject((uint8*)RawPtr - OffsetToBase, DerivedTypeDesc, UnfrozenMemory);
+			T* UnfrozenPtr = (T*)((uint8*)UnfrozenMemory + OffsetToBase);
+			new(OutDst) TMemoryImagePtr<T>(UnfrozenPtr);
+		}
+		else
+		{
+			new(OutDst) TMemoryImagePtr<T>(nullptr);
+		}
+	}
+
+	template<typename T>
+	uint32 IntrinsicAppendHash(const TMemoryImagePtr<T>* DummyObject, const FTypeLayoutDesc& TypeDesc, const FPlatformTypeLayoutParameters& LayoutParams, FSHA1& Hasher)
+	{
+		return AppendHashForNameAndSize(TypeDesc.Name, LayoutParams.GetMemoryImagePointerSize(), Hasher);
+	}
+
+	template<typename T>
+	inline uint32 IntrinsicGetTargetAlignment(const TMemoryImagePtr<T>* DummyObject, const FTypeLayoutDesc& TypeDesc, const FPlatformTypeLayoutParameters& LayoutParams)
+	{
+		return FMath::Min(LayoutParams.GetMemoryImagePointerSize(), LayoutParams.MaxFieldAlignment);
+	}
+}
+
+DECLARE_TEMPLATE_INTRINSIC_TYPE_LAYOUT(template<typename T>, TMemoryImagePtr<T>);
+
+template<typename T>
+class TUniqueMemoryImagePtr : public TMemoryImagePtr<T>
+{
+public:
+	inline TUniqueMemoryImagePtr()
+		: TMemoryImagePtr<T>(nullptr)
+	{}
+	explicit inline TUniqueMemoryImagePtr(T* InPtr)
+		: TMemoryImagePtr<T>(InPtr)
+	{ }
+
+	inline TUniqueMemoryImagePtr(TUniqueMemoryImagePtr&& Other)
+	{
+		this->SafeDelete();
+		this->Ptr = Other.Ptr;
+		Other.Ptr = nullptr;
+	}
+	inline ~TUniqueMemoryImagePtr()
+	{
+		this->SafeDelete();
+	}
+	inline TUniqueMemoryImagePtr& operator=(T* InPtr)
+	{
+		this->SafeDelete();
+		this->Ptr = InPtr;
+	}
+	inline TUniqueMemoryImagePtr& operator=(TUniqueMemoryImagePtr&& Other)
+	{
+		if (this != &Other)
+		{
+			// We _should_ delete last like TUniquePtr, but we have issues with SafeDelete, and being Frozen or not
+			this->SafeDelete();
+			this->Ptr = Other.Ptr;
+			Other.Ptr = nullptr;
+		}
+
+		return *this;
+	}
+
+};
+
+class CORE_API FMemoryImageAllocatorBase
+{
+	UE_NONCOPYABLE(FMemoryImageAllocatorBase);
+public:
+	FMemoryImageAllocatorBase() = default;
+
+	/**
+	* Moves the state of another allocator into this one.
+	* Assumes that the allocator is currently empty, i.e. memory may be allocated but any existing elements have already been destructed (if necessary).
+	* @param Other - The allocator to move the state from.  This allocator should be left in a valid empty state.
+	*/
+	void MoveToEmpty(FMemoryImageAllocatorBase& Other);
+
+	/** Destructor. */
+	~FMemoryImageAllocatorBase();
+
+	// FContainerAllocatorInterface
+	FORCEINLINE FScriptContainerElement* GetAllocation() const
+	{
+		return Data.Get();
+	}
+
+	FORCEINLINE SIZE_T GetAllocatedSize(int32 NumAllocatedElements, SIZE_T NumBytesPerElement) const
+	{
+		return NumAllocatedElements * NumBytesPerElement;
+	}
+	FORCEINLINE bool HasAllocation()
+	{
+		return Data.IsValid();
+	}
+	FORCEINLINE FMemoryImagePtrInt GetFrozenOffsetFromThis() const { return Data.GetFrozenOffsetFromThis(); }
+
+	void ResizeAllocation(int32 PreviousNumElements, int32 NumElements, SIZE_T NumBytesPerElement, uint32 Alignment);
+	void WriteMemoryImage(FMemoryImageWriter& Writer, const FTypeLayoutDesc& TypeDesc, int32 NumAllocatedElements, uint32 Alignment) const;
+	void CopyUnfrozen(const FMemoryUnfreezeContent& Context, const FTypeLayoutDesc& TypeDesc, int32 NumAllocatedElements, void* Dst) const;
+
+private:
+	TMemoryImagePtr<FScriptContainerElement> Data;
+};
+
+template<uint32 Alignment = DEFAULT_ALIGNMENT>
+class TMemoryImageAllocator
+{
+public:
+	using SizeType = int32;
+
+	enum { NeedsElementType = false };
+	enum { RequireRangeCheck = true };
+	enum { SupportsFreezeMemoryImage = true };
+
+	class ForAnyElementType : public FMemoryImageAllocatorBase
+	{
+	public:
+		/** Default constructor. */
+		ForAnyElementType() = default;
+
+		FORCEINLINE SizeType GetInitialCapacity() const
+		{
+			return 0;
+		}
+		FORCEINLINE int32 CalculateSlackReserve(int32 NumElements, int32 NumBytesPerElement) const
+		{
+			return DefaultCalculateSlackReserve(NumElements, NumBytesPerElement, true, Alignment);
+		}
+		FORCEINLINE int32 CalculateSlackShrink(int32 NumElements, int32 NumAllocatedElements, int32 NumBytesPerElement) const
+		{
+			return DefaultCalculateSlackShrink(NumElements, NumAllocatedElements, NumBytesPerElement, true, Alignment);
+		}
+		FORCEINLINE int32 CalculateSlackGrow(int32 NumElements, int32 NumAllocatedElements, int32 NumBytesPerElement) const
+		{
+			return DefaultCalculateSlackGrow(NumElements, NumAllocatedElements, NumBytesPerElement, true, Alignment);
+		}
+		FORCEINLINE void ResizeAllocation(int32 PreviousNumElements, int32 NumElements, SIZE_T NumBytesPerElement)
+		{
+			FMemoryImageAllocatorBase::ResizeAllocation(PreviousNumElements, NumElements, NumBytesPerElement, Alignment);
+		}
+
+		FORCEINLINE void WriteMemoryImage(FMemoryImageWriter& Writer, const FTypeLayoutDesc& TypeDesc, int32 NumAllocatedElements) const
+		{
+			FMemoryImageAllocatorBase::WriteMemoryImage(Writer, TypeDesc, NumAllocatedElements, Alignment);
+		}
+	};
+
+	template<typename ElementType>
+	class ForElementType : public ForAnyElementType
+	{
+	public:
+		ForElementType() {}
+		FORCEINLINE ElementType* GetAllocation() const { return (ElementType*)ForAnyElementType::GetAllocation(); }
+	};
+};
+
+//@todo stever
+/*static_assert(
+	sizeof(TMemoryImageAllocator<>::ForAnyElementType) == sizeof(FDefaultAllocator::ForAnyElementType) && alignof(TMemoryImageAllocator<>::ForAnyElementType) == alignof(FDefaultAllocator::ForAnyElementType),
+	"TMemoryImageAllocator must be the same layout as FDefaultAllocator for our FScriptArray hacks to work"
+);*/
+
+template <uint32 Alignment>
+struct TAllocatorTraits<TMemoryImageAllocator<Alignment>> : TAllocatorTraitsBase<TMemoryImageAllocator<Alignment>>
+{
+	enum { SupportsMove = true };
+	enum { IsZeroConstruct = true };
+	enum { SupportsFreezeMemoryImage = true };
+};
+
+using FMemoryImageAllocator = TMemoryImageAllocator<>;
+
+using FMemoryImageSparseArrayAllocator = TSparseArrayAllocator<FMemoryImageAllocator, FMemoryImageAllocator>;
+using FMemoryImageSetAllocator = TSetAllocator<FMemoryImageSparseArrayAllocator, FMemoryImageAllocator>;
+
+template<typename T>
+using TMemoryImageArray = TArray<T, FMemoryImageAllocator>;
+
+template<typename ElementType, typename KeyFuncs = DefaultKeyFuncs<ElementType>>
+using TMemoryImageSet = TSet<ElementType, KeyFuncs, FMemoryImageSetAllocator>;
+
+template <typename KeyType, typename ValueType, typename KeyFuncs = TDefaultMapHashableKeyFuncs<KeyType, ValueType, false>>
+using TMemoryImageMap = TMap<KeyType, ValueType, FMemoryImageSetAllocator, KeyFuncs>;
+
+class FMemoryImageString
+{
+	DECLARE_EXPORTED_TYPE_LAYOUT(FMemoryImageString, CORE_API, NonVirtual);
+private:
+	/** Array holding the character data */
+	using DataType = TMemoryImageArray<TCHAR>;
+	LAYOUT_FIELD(DataType, Data)
+public:
+	using ElementType = TCHAR;
+
+	FMemoryImageString() = default;
+	FMemoryImageString(FMemoryImageString&&) = default;
+	FMemoryImageString(const FMemoryImageString&) = default;
+	FMemoryImageString& operator=(FMemoryImageString&&) = default;
+	FMemoryImageString& operator=(const FMemoryImageString&) = default;
+
+	FORCEINLINE FMemoryImageString(const FString& Other) : Data(Other.GetCharArray()) {}
+
+	template <
+		typename CharType,
+		typename = typename TEnableIf<TIsCharType<CharType>::Value>::Type // This TEnableIf is to ensure we don't instantiate this constructor for non-char types, like id* in Obj-C
+	>
+		FORCEINLINE FMemoryImageString(const CharType* Src)
+	{
+		if (Src && *Src)
+		{
+			int32 SrcLen = TCString<CharType>::Strlen(Src) + 1;
+			int32 DestLen = FPlatformString::ConvertedLength<TCHAR>(Src, SrcLen);
+			Data.AddUninitialized(DestLen);
+			FPlatformString::Convert(Data.GetData(), DestLen, Src, SrcLen);
+		}
+	}
+
+	FORCEINLINE operator FString() const { return FString(Len(), Data.GetData()); }
+
+	FORCEINLINE const TCHAR* operator*() const
+	{
+		return Data.Num() ? Data.GetData() : TEXT("");
+	}
+
+	FORCEINLINE bool IsEmpty() const { return Data.Num() <= 1; }
+	FORCEINLINE SIZE_T GetAllocatedSize() const { return Data.GetAllocatedSize(); }
+
+	FORCEINLINE int32 Len() const
+	{
+		return Data.Num() ? Data.Num() - 1 : 0;
+	}
+
+	friend inline FArchive& operator<<(FArchive& Ar, FMemoryImageString& Ref)
+	{
+		Ar << Ref.Data;
+		return Ar;
+	}
+
+	friend inline bool operator==(const FMemoryImageString& Lhs, const FMemoryImageString& Rhs)
+	{
+		return FCString::Stricmp(*Lhs, *Rhs) == 0;
+	}
+
+	friend inline bool operator!=(const FMemoryImageString& Lhs, const FMemoryImageString& Rhs)
+	{
+		return !operator==(Lhs, Rhs);
+	}
+
+	friend inline bool operator==(const FMemoryImageString& Lhs, const FString& Rhs)
+	{
+		return FCString::Stricmp(*Lhs, *Rhs) == 0;
+	}
+
+	friend inline bool operator!=(const FMemoryImageString& Lhs, const FString& Rhs)
+	{
+		return !operator==(Lhs, Rhs);
+	}
+
+	inline DataType::ElementAllocatorType& GetAllocatorInstance() { return Data.GetAllocatorInstance(); }
+};
+
+/** Case insensitive string hash function. */
+FORCEINLINE uint32 GetTypeHash(const FMemoryImageString& S)
+{
+	return FCrc::Strihash_DEPRECATED(*S);
+}
+
+
+class FHashedName
+{
+public:
+	inline FHashedName() : Hash(0u) {}
+	inline explicit FHashedName(uint64 InHash) : Hash(InHash) {}
+	CORE_API FHashedName(const TCHAR* InString);
+	CORE_API FHashedName(const FString& InString);
+	CORE_API FHashedName(const FName& InName);
+
+	inline uint64 GetHash() const { return Hash; }
+	inline bool IsNone() const { return Hash == 0u; }
+
+	friend inline bool operator==(const FHashedName& Lhs, const FHashedName& Rhs) { return Lhs.Hash == Rhs.Hash; }
+	friend inline bool operator!=(const FHashedName& Lhs, const FHashedName& Rhs) { return Lhs.Hash != Rhs.Hash; }
+
+	/** For sorting by name */
+	friend inline bool operator<(const FHashedName& Lhs, const FHashedName& Rhs) { return Lhs.Hash < Rhs.Hash; }
+
+	friend inline FArchive& operator<<(FArchive& Ar, FHashedName& String)
+	{
+		Ar << String.Hash;
+		return Ar;
+	}
+
+	friend inline uint32 GetTypeHash(const FHashedName& Name)
+	{
+		return GetTypeHash(Name.Hash);
+	}
+
+	/*inline FString ToString() const
+	{
+		return FString::Printf(TEXT("0x%016X"), Hash);
+	}*/
+
+private:
+	uint64 Hash;
+};
+DECLARE_INTRINSIC_TYPE_LAYOUT(FHashedName);
+
+class CORE_API FPtrTableBase
+{
+public:
+	template<typename PtrType>
+	static void LoadAndApplyPatchesFromArchive(FArchive& Ar, void* FrozenBase, const PtrType& Ptr)
+	{
+		int32 NumOffsets = 0;
+		Ar << NumOffsets;
+		for (int32 OffsetIndex = 0; OffsetIndex < NumOffsets; ++OffsetIndex)
+		{
+			uint32 Offset = 0u;
+			Ar << Offset;
+			new((char*)FrozenBase + Offset) PtrType(Ptr);
+		}
+	}
+
+	void SavePatchesToArchive(FArchive& Ar, uint32 PtrIndex) const;
+
+protected:
+	struct FPatchOffset
+	{
+		uint32 Offset;
+		uint32 NextIndex;
+	};
+
+	struct FPatchOffsetList
+	{
+		FPatchOffsetList() : FirstIndex(~0u), NumOffsets(0u) {}
+		uint32 FirstIndex;
+		uint32 NumOffsets;
+	};
+
+	void AddPatchedPointerBase(uint32 PtrIndex, uint64 Offset);
+
+	TArray<FPatchOffsetList> PatchLists;
+	TArray<FPatchOffset> PatchOffsets;
+};
+
+template<typename T, typename PtrType>
+class TPtrTableBase : public FPtrTableBase
+{
+public:
+	static const FTypeLayoutDesc& StaticGetPtrTypeLayoutDesc();
+
+	void Empty(int32 NewSize = 0)
+	{
+		Pointers.Reset(NewSize);
+	}
+
+	uint32 Num() const { return Pointers.Num(); }
+	uint32 AddIndexedPointer(T* Ptr) { check(Ptr); return Pointers.AddUnique(Ptr); }
+
+	bool TryAddIndexedPtr(const FTypeLayoutDesc& TypeDesc, void* Ptr, int32& OutIndex)
+	{
+		if (TypeDesc == StaticGetPtrTypeLayoutDesc())
+		{
+			OutIndex = AddIndexedPointer(static_cast<T*>(Ptr));
+			return true;
+		}
+		return false;
+	}
+
+	void LoadIndexedPointer(T* Ptr)
+	{
+		if (Ptr)
+		{
+			checkSlow(!Pointers.Contains(Ptr));
+			Pointers.Add(Ptr);
+		}
+		else
+		{
+			// allow duplicate nullptrs
+			// pointers that were valid when saving may not be found when loading, need to preserve indices
+			Pointers.Add(nullptr);
+		}
+	}
+
+	void AddPatchedPointer(T* Ptr, uint64 Offset)
+	{
+		const uint32 PtrIndex = AddIndexedPointer(Ptr);
+		FPtrTableBase::AddPatchedPointerBase(PtrIndex, Offset);
+	}
+
+	T* GetIndexedPointer(uint32 i) const { return Pointers[i]; }
+
+	bool TryGetIndexedPtr(const FTypeLayoutDesc& TypeDesc, uint32 i, void*& OutPtr) const
+	{
+		if (TypeDesc == StaticGetPtrTypeLayoutDesc())
+		{
+			OutPtr = GetIndexedPointer(i);
+			return true;
+		}
+		return false;
+	}
+
+	void ApplyPointerPatches(void* FrozenBase) const
+	{
+		for (int32 PtrIndex = 0; PtrIndex < PatchLists.Num(); ++PtrIndex)
+		{
+			uint32 PatchIndex = PatchLists[PtrIndex].FirstIndex;
+			while (PatchIndex != ~0u)
+			{
+				const FPatchOffset& Patch = PatchOffsets[PatchIndex];
+				new((char*)FrozenBase + Patch.Offset) PtrType(Pointers[PtrIndex]);
+				PatchIndex = Patch.NextIndex;
+			}
+		}
+	}
+
+	inline typename TArray<PtrType>::RangedForIteratorType begin() { return Pointers.begin(); }
+	inline typename TArray<PtrType>::RangedForIteratorType end() { return Pointers.end(); }
+	inline typename TArray<PtrType>::RangedForConstIteratorType begin() const { return Pointers.begin(); }
+	inline typename TArray<PtrType>::RangedForConstIteratorType end() const { return Pointers.end(); }
+private:
+	TArray<PtrType> Pointers;
+};
+
+template<typename T>
+class TPtrTable : public TPtrTableBase<T, T*> {};
+
+template<typename T>
+class TRefCountPtrTable : public TPtrTableBase<T, TRefCountPtr<T>>
+{
+	using Super = TPtrTableBase<T, TRefCountPtr<T>>;
+};
+
+class FVoidPtrTable : public TPtrTableBase<void, void*> {};
+
+#ifdef _MSC_VER
+#pragma warning(push)
+#pragma warning(disable : 4583) // destructor is not implicitly called
+#endif
+
+/**
+ * Note that IndexedPtr types don't declare a default intrinsic type layout.
+ * Instead any required pointer types need to be declared/implemented using DECLARE_EXPORTED_TEMPLATE_INTRINSIC_TYPE_LAYOUT/IMPLEMENT_TEMPLATE_INTRINSIC_TYPE_LAYOUT.
+ * The TypeDesc of indexed pointers are compared for equality when adding to pointer tables,
+ * and it's possible for inline type layouts to generate multiple copies when referenced from multiple modules
+ */
+template<typename T, typename PtrType>
+class TIndexedPtrBase
+{
+public:
+	using FPtrTable = TPtrTableBase<T, PtrType>;
+
+	inline TIndexedPtrBase(T* InPtr = nullptr) : Ptr(InPtr) {}
+	inline ~TIndexedPtrBase() { if(!IsFrozen()) Ptr.~PtrType(); }
+
+	// Copy constructor requires an unfrozen source
+	inline TIndexedPtrBase(const TIndexedPtrBase<T, PtrType>& Rhs) : Ptr(Rhs.GetUnfrozen()) {}
+
+	inline TIndexedPtrBase(const TIndexedPtrBase<T, PtrType>& Rhs, const FPtrTable& InTable) : Ptr(Rhs.Get(InTable)) {}
+
+	inline TIndexedPtrBase& operator=(T* Rhs)
+	{
+		// If not currently frozen, invoke the standard assignment operator for the underlying pointer type
+		// If frozen, construct a new (non-frozen) pointer over the existing frozen offset
+		if (!IsFrozen()) Ptr = Rhs;
+		else new(&Ptr) PtrType(Rhs);
+		check(!IsFrozen());
+		return *this;
+	}
+
+	inline TIndexedPtrBase& operator=(const PtrType& Rhs)
+	{
+		if (!IsFrozen()) Ptr = Rhs;
+		else new(&Ptr) PtrType(Rhs);
+		check(!IsFrozen());
+		return *this;
+	}
+
+	inline TIndexedPtrBase& operator=(PtrType&& Rhs)
+	{
+		if (!IsFrozen()) Ptr = Rhs;
+		else new(&Ptr) PtrType(Rhs);
+		check(!IsFrozen());
+		return *this;
+	}
+	
+	inline bool IsFrozen() const { return PackedIndex & IsFrozenMask; }
+	inline bool IsValid() const { return PackedIndex != 0u; } // works for both frozen/unfrozen cases
+	inline bool IsNull() const { return PackedIndex == 0u; }
+
+	inline void SafeRelease()
+	{
+		if (!IsFrozen())
+		{
+			SafeReleaseImpl(Ptr);
+		}
+	}
+
+	inline T* Get(const FPtrTable& PtrTable) const
+	{
+		if (IsFrozen())
+		{
+			return PtrTable.GetIndexedPointer(PackedIndex >> IndexShift);
+		}
+		return Ptr;
+	}
+
+	inline T* Get(const FPointerTableBase* PtrTable) const
+	{
+		if (IsFrozen())
+		{
+			check(PtrTable);
+			const FTypeLayoutDesc& TypeDesc = StaticGetTypeLayoutDesc<TIndexedPtrBase<T, PtrType>>();
+			return static_cast<T*>(PtrTable->GetIndexedPointer(TypeDesc, PackedIndex >> IndexShift));
+		}
+		return Ptr;
+	}
+
+	inline T* GetUnfrozen() const { check(!IsFrozen()); return Ptr; }
+
+private:
+	enum
+	{
+		IsFrozenMask = (1 << 0),
+		IndexShift = 1,
+	};
+
+	static void SafeReleaseImpl(T*& InPtr)
+	{
+		if (InPtr)
+		{
+			delete InPtr;
+			InPtr = nullptr;
+		}
+	}
+
+	static void SafeReleaseImpl(TRefCountPtr<T>& InPtr)
+	{
+		InPtr.SafeRelease();
+	}
+
+	static_assert(sizeof(PtrType) <= sizeof(FMemoryImageUPtrInt), "PtrType must fit within a standard pointer");
+	union
+	{
+		PtrType Ptr;
+		FMemoryImageUPtrInt PackedIndex;
+	};
+};
+
+#ifdef _MSC_VER
+#pragma warning(pop)
+#endif
+
+template<typename T, typename PtrType>
+inline const FTypeLayoutDesc& TPtrTableBase<T, PtrType>::StaticGetPtrTypeLayoutDesc()
+{
+	return StaticGetTypeLayoutDesc<TIndexedPtrBase<T, PtrType>>();
+}
+
+namespace Freeze
+{
+	template<typename T, typename PtrType>
+	void IntrinsicWriteMemoryImage(FMemoryImageWriter& Writer, const TIndexedPtrBase<T, PtrType>& Object, const FTypeLayoutDesc& TypeDesc)
+	{
+		T* RawPtr = Object.Get(Writer.TryGetPrevPointerTable());
+		if (RawPtr)
+		{
+			const uint32 Index = Writer.GetPointerTable().AddIndexedPointer(TypeDesc, RawPtr);
+			check(Index != (uint32)INDEX_NONE);
+			const uint64 FrozenPackedIndex = (Index << 1u) | 1u;
+			Writer.WriteMemoryImagePointerSizedBytes(FrozenPackedIndex);
+		}
+		else
+		{
+			Writer.WriteMemoryImagePointerSizedBytes(0u);
+		}
+	}
+
+	template<typename T, typename PtrType>
+	void IntrinsicUnfrozenCopy(const FMemoryUnfreezeContent& Context, const TIndexedPtrBase<T, PtrType>& Object, void* OutDst)
+	{
+		new(OutDst) TIndexedPtrBase<T, PtrType>(Object.Get(Context.TryGetPrevPointerTable()));
+	}
+
+	template<typename T, typename PtrType>
+	uint32 IntrinsicAppendHash(const TIndexedPtrBase<T, PtrType>* DummyObject, const FTypeLayoutDesc& TypeDesc, const FPlatformTypeLayoutParameters& LayoutParams, FSHA1& Hasher)
+	{
+		return AppendHashForNameAndSize(TypeDesc.Name, LayoutParams.GetMemoryImagePointerSize(), Hasher);
+	}
+
+	template<typename T, typename PtrType>
+	inline uint32 IntrinsicGetTargetAlignment(const TIndexedPtrBase<T, PtrType>* DummyObject, const FTypeLayoutDesc& TypeDesc, const FPlatformTypeLayoutParameters& LayoutParams)
+	{
+		return FMath::Min(LayoutParams.GetMemoryImagePointerSize(), LayoutParams.MaxFieldAlignment);
+	}
+}
+
+template<typename T>
+using TIndexedPtr = TIndexedPtrBase<T, T*>;
+
+template<typename T>
+using TIndexedRefCountPtr = TIndexedPtrBase<T, TRefCountPtr<T>>;
+
+template<typename T, typename PtrType>
+class TPatchedPtrBase
+{
+public:
+	using FPtrTable = TPtrTableBase<T, PtrType>;
+
+	inline TPatchedPtrBase(T* InPtr = nullptr) : Ptr(InPtr) {}
+
+	inline T* Get() const
+	{
+		return Ptr;
+	}
+
+	inline T* GetChecked() const { T* Value = Get(); check(Value); return Value; }
+	inline T* operator->() const { return GetChecked(); }
+	inline T& operator*() const { return *GetChecked(); }
+	inline operator T*() const { return Get(); }
+
+private:
+	static_assert(sizeof(PtrType) == sizeof(void*), "PtrType must be a standard pointer");
+	PtrType Ptr;
+};
+
+template<typename T>
+using TPatchedPtr = TPatchedPtrBase<T, T*>;
+
+template<typename T>
+using TPatchedRefCountPtr = TPatchedPtrBase<T, TRefCountPtr<T>>;
+
diff --git a/Engine/Source/Runtime/Core/Public/Serialization/MemoryImageWriter.h b/Engine/Source/Runtime/Core/Public/Serialization/MemoryImageWriter.h
new file mode 100644
--- /dev/null
+++ b/Engine/Source/Runtime/Core/Public/Serialization/MemoryImageWriter.h
@@ -0,0 +1,90 @@
+#pragma once
+
+#include "CoreTypes.h"
+#include "Serialization/MemoryLayout.h"
+
+class FMemoryImage;
+class FMemoryImageSection;
+class FPointerTableBase;
+
+class CORE_API FMemoryImageWriter
+{
+public:
+	explicit FMemoryImageWriter(FMemoryImage& InImage);
+	explicit FMemoryImageWriter(FMemoryImageSection* InSection);
+	~FMemoryImageWriter();
+
+	FMemoryImage& GetImage() const;
+	const FPlatformTypeLayoutParameters& GetHostLayoutParams() const;
+	const FPlatformTypeLayoutParameters& GetTargetLayoutParams() const;
+	FPointerTableBase& GetPointerTable() const;
+	const FPointerTableBase* TryGetPrevPointerTable() const;
+
+	inline bool Is32BitTarget() const { return GetTargetLayoutParams().b32Bit; }
+	inline bool Is64BitTarget() const { return !Is32BitTarget(); }
+
+	void AddDependency(const FTypeLayoutDesc& TypeDesc);
+
+	void WriteObject(const void* Object, const FTypeLayoutDesc& TypeDesc);
+	void WriteObjectArray(const void* Object, const FTypeLayoutDesc& TypeDesc, uint32_t NumArray);
+
+	uint32 GetOffset() const;
+	uint32 WriteAlignment(uint32 Alignment);
+	void WritePaddingToSize(uint32 Offset);
+	uint32 WriteBytes(const void* Data, uint32 Size);
+	FMemoryImageWriter WritePointer(const FString& SectionName, uint32 Offset = 0u);
+	uint32 WriteRawPointerSizedBytes(uint64 PointerValue);
+	uint32 WriteMemoryImagePointerSizedBytes(uint64 PointerValue);
+	uint32 WriteVTable(const FTypeLayoutDesc& TypeDesc, const FTypeLayoutDesc& DerivedTypeDesc);
+	uint32 WriteFName(const FName& Name);
+
+	template<typename T>
+	void WriteObjectArray(const T* Object, uint32 NumArray)
+	{
+		const FTypeLayoutDesc& TypeDesc = StaticGetTypeLayoutDesc<T>();
+		WriteObjectArray(Object, TypeDesc, TypeDesc, NumArray);
+	}
+
+	template<typename T>
+	void WriteObject(const T& Object)
+	{
+		WriteObject(&Object, GetTypeLayoutDesc(TryGetPrevPointerTable(), Object));
+	}
+
+	template<typename T>
+	uint32 WriteAlignment()
+	{
+		return WriteAlignment(alignof(T));
+	}
+
+	template<typename T>
+	uint32 WriteBytes(const T& Data)
+	{
+		return WriteBytes(&Data, sizeof(T));
+	}
+
+//private:
+	FMemoryImageSection* Section;
+};
+
+class FMemoryUnfreezeContent
+{
+public:
+	FMemoryUnfreezeContent() : PrevPointerTable(nullptr) {}
+
+	const FPointerTableBase* TryGetPrevPointerTable() const { return PrevPointerTable; }
+
+	void UnfreezeObject(const void* Object, const FTypeLayoutDesc& TypeDesc, void* OutDst) const
+	{
+		TypeDesc.UnfrozenCopyFunc(*this, Object, TypeDesc, OutDst);
+	}
+
+	template<typename T>
+	void UnfreezeObject(const T& Object, void* OutDst) const
+	{
+		const FTypeLayoutDesc& TypeDesc = GetTypeLayoutDesc(TryGetPrevPointerTable(), Object);
+		UnfreezeObject(&Object, TypeDesc, OutDst);
+	}
+
+	const FPointerTableBase* PrevPointerTable;
+};
diff --git a/Engine/Source/Runtime/Core/Public/Serialization/MemoryLayout.h b/Engine/Source/Runtime/Core/Public/Serialization/MemoryLayout.h
new file mode 100644
--- /dev/null
+++ b/Engine/Source/Runtime/Core/Public/Serialization/MemoryLayout.h
@@ -0,0 +1,801 @@
+// Copyright 1998-2019 Epic Games, Inc. All Rights Reserved.
+
+#pragma once
+
+#include "Templates/IsAbstract.h"
+#include "Templates/IsPolymorphic.h"
+#include "Templates/IsTriviallyDestructible.h"
+#include "Templates/EnableIf.h"
+#include "Misc/DelayedAutoRegister.h"
+
+class FHashedName;
+class FSHA1;
+class FMemoryImageWriter;
+class FMemoryUnfreezeContent;
+class FPointerTableBase;
+struct FTypeLayoutDesc;
+
+// Duplicated from RHIDefinitions.h [AND MUST MATCH THE Freezing_bWithRayTracing in DataDrivenPlatformInfo.ini]
+#ifndef WITH_RAYTRACING
+#if (PLATFORM_WINDOWS && PLATFORM_64BITS)
+#define WITH_RAYTRACING 1
+#else
+#define WITH_RAYTRACING 0
+#endif
+#endif
+
+/**
+ * If this is set, TMemoryImagePtr, TIndexedPtr, and other wrapped pointer types used for memory images will be forced to 64bits, even when building 32bit targets
+ * This is intended to facilitate sharing packs between 32/64bit builds
+ * Android requires sharing data between 32/64bit exe
+ */
+#if PLATFORM_ANDROID
+#define UE_FORCE_64BIT_MEMORY_IMAGE_POINTERS 1
+#else
+#define UE_FORCE_64BIT_MEMORY_IMAGE_POINTERS 0
+#endif
+
+#if UE_FORCE_64BIT_MEMORY_IMAGE_POINTERS
+using FMemoryImagePtrInt = int64;
+using FMemoryImageUPtrInt = uint64;
+#else
+using FMemoryImagePtrInt = PTRINT;
+using FMemoryImageUPtrInt = UPTRINT;
+#endif
+
+/*#define UE_STATIC_ONLY(T) \
+	T() = delete; \
+	~T() = delete; \
+	UE_NONCOPYABLE(T)*/
+#define UE_STATIC_ONLY(T)
+
+/**
+ * Certain compilers don't accept full explicit template specializations inside class declarations.
+ * Partial specializations are accepted however.  So for these platforms, we force partial specialization via an extra dummy parameter that's otherwise ignored.
+ * This actually breaks MSVC, as MSVC doesn't allow complex statements as template values during partial specialization, so need use full specialization there.
+ */
+#if PLATFORM_ANDROID
+#define UE_DECLARE_INTERNAL_LINK_BASE(T) template<int Counter, typename Dummy=void> struct T
+#define UE_DECLARE_INTERNAL_LINK_SPECIALIZATION(T, Counter) template<typename Dummy> struct T<Counter, Dummy>
+#else
+#define UE_DECLARE_INTERNAL_LINK_BASE(T) template<int Counter> struct T
+#define UE_DECLARE_INTERNAL_LINK_SPECIALIZATION(T, Counter) template<> struct T<Counter>
+#endif
+
+template<typename T> const FTypeLayoutDesc& StaticGetTypeLayoutDesc();
+
+struct FPlatformTypeLayoutParameters
+{
+	uint32 MaxFieldAlignment = 0xffffffff;
+	bool b32Bit = false;
+	bool bForce64BitMemoryImagePointers = true;
+	bool bAlignBases = false;
+	bool bWithEditorOnly = false;
+	bool bWithRayTracing = false;
+	bool bIsCurrentPlatform = false;
+
+	inline bool Has32BitMemoryImagePointers() const { return b32Bit && !bForce64BitMemoryImagePointers; }
+	inline bool Has64BitMemoryImagePointers() const { return !Has32BitMemoryImagePointers(); }
+
+	inline uint32 GetRawPointerSize() const { return b32Bit ? sizeof(uint32) : sizeof(uint64); }
+	inline uint32 GetMemoryImagePointerSize() const { return Has32BitMemoryImagePointers() ? sizeof(uint32) : sizeof(uint64); }
+
+	CORE_API void InitializeForArchive(FArchive& Ar);
+	CORE_API void InitializeForPlatform(const FString& PlatformName, bool bHasEditorOnlyData);
+	CORE_API void InitializeForCurrent();
+	CORE_API void InitializeForMSVC();
+	CORE_API void InitializeForClang();
+
+	/**
+	 * Allow the layout parameters to modify the given DDC key string
+	 * Normally this should not be used, and instead hashed layouts of any dependent types should be preferred
+	 * This is mostly needed when fixing bugs related to writing memory images with certain parameter combinations
+	 */
+	CORE_API void AppendKeyString(FString& KeyString) const;
+};
+
+namespace ETypeLayoutInterface
+{
+	enum Type : uint8
+	{
+		NonVirtual,
+		Virtual,
+		Abstract,
+	};
+
+	inline bool HasVTable(Type InType) { return InType != NonVirtual; }
+}
+
+namespace EFieldLayoutFlags
+{
+	enum Type : uint8
+	{
+		None = 0u,
+		WithEditorOnly = (1u << 0),
+		WithRayTracing = (1u << 1),
+		Transient = (1u << 2),
+		UseInstanceWithNoProperty = (1u << 3),
+	};
+
+	FORCEINLINE Type MakeFlags(uint32 Flags = None) { return (Type)Flags; }
+	FORCEINLINE Type MakeFlagsEditorOnly(uint32 Flags = None) { return (Type)(WithEditorOnly | Flags); }
+	FORCEINLINE Type MakeFlagsRayTracing(uint32 Flags = None) { return (Type)(WithRayTracing | Flags); }
+}
+
+struct FFieldLayoutDesc
+{
+	typedef void (FWriteFrozenMemoryImageFunc)(FMemoryImageWriter& Writer, const void* Object, const void* FieldObject, const FTypeLayoutDesc& TypeDesc, const FTypeLayoutDesc& DerivedTypeDesc);
+
+	const TCHAR* Name;
+	const struct FTypeLayoutDesc* Type;
+	const FFieldLayoutDesc* Next;
+	FWriteFrozenMemoryImageFunc* WriteFrozenMemoryImageFunc;
+	uint32 Offset;
+	uint32 NumArray;
+	EFieldLayoutFlags::Type Flags;
+	uint8 BitFieldSize;
+	uint8 UFieldNameLength; // this is the number of characters in Name, omitting any _DEPRECATED suffix
+};
+
+struct FTypeLayoutDesc
+{
+	typedef void (FDestroyFunc)(void* Object, const FTypeLayoutDesc& TypeDesc);
+	typedef void (FWriteFrozenMemoryImageFunc)(FMemoryImageWriter& Writer, const void* Object, const FTypeLayoutDesc& TypeDesc, const FTypeLayoutDesc& DerivedTypeDesc);
+	typedef void (FUnfrozenCopyFunc)(const FMemoryUnfreezeContent& Context, const void* Object, const FTypeLayoutDesc& TypeDesc, void* OutDst);
+	typedef uint32 (FAppendHashFunc)(const FTypeLayoutDesc& TypeDesc, const FPlatformTypeLayoutParameters& LayoutParams, FSHA1& Hasher);
+	typedef uint32 (FGetTargetAlignmentFunc)(const FTypeLayoutDesc& TypeDesc, const FPlatformTypeLayoutParameters& LayoutParams);
+	typedef const void* (FGetDefaultFunc)();
+
+	static CORE_API const FTypeLayoutDesc& GetInvalidTypeLayout();
+
+	static CORE_API void Initialize(FTypeLayoutDesc& TypeDesc);
+	static CORE_API void Register(FTypeLayoutDesc& TypeDesc);
+	static CORE_API const FTypeLayoutDesc* Find(uint64 NameHash);
+
+	uint32 CORE_API GetOffsetToBase(const FTypeLayoutDesc& BaseTypeDesc) const;
+
+	const FTypeLayoutDesc* HashNext;
+	const TCHAR* Name;
+	const FFieldLayoutDesc* Fields;
+	FDestroyFunc* DestroyFunc;
+	FWriteFrozenMemoryImageFunc* WriteFrozenMemoryImageFunc;
+	FUnfrozenCopyFunc* UnfrozenCopyFunc;
+	FAppendHashFunc* AppendHashFunc;
+	FGetTargetAlignmentFunc* GetTargetAlignmentFunc;
+	FGetDefaultFunc* GetDefaultObjectFunc;
+
+	uint64 NameHash; // from FHashedName(Name)
+	uint32 Size;
+	uint32 SizeFromFields;
+	uint32 Alignment;
+	ETypeLayoutInterface::Type Interface;
+	uint8 NumBases;
+	uint8 NumVirtualBases;
+	uint8 IsIntrinsic : 1;
+	uint8 IsInitialized : 1;
+};
+inline bool operator==(const FTypeLayoutDesc& Lhs, const FTypeLayoutDesc& Rhs) { return &Lhs == &Rhs; }
+inline bool operator!=(const FTypeLayoutDesc& Lhs, const FTypeLayoutDesc& Rhs) { return &Lhs != &Rhs; }
+
+struct FRegisterTypeLayoutDesc
+{
+	explicit CORE_API FRegisterTypeLayoutDesc(FTypeLayoutDesc& TypeDesc);
+	CORE_API FRegisterTypeLayoutDesc(const TCHAR* Name, FTypeLayoutDesc& TypeDesc);
+};
+
+template<typename T, typename Base>
+inline uint32 GetBaseOffset()
+{
+	alignas(T) char Dummy[sizeof(T)];
+	T* Derived = reinterpret_cast<T*>(Dummy);
+	return (uint32)((uint64)static_cast<Base*>(Derived) - (uint64)Derived);
+}
+
+template<typename T, bool HasNonTrivialDestructor>
+struct TAssignDestroyHelper
+{
+	UE_STATIC_ONLY(TAssignDestroyHelper);
+	static FORCEINLINE void Do(FTypeLayoutDesc& TypeDesc) {}
+};
+
+template<typename T>
+struct TAssignDestroyHelper<T, true>
+{
+	UE_STATIC_ONLY(TAssignDestroyHelper);
+
+	static void Destroy(void* Object, const FTypeLayoutDesc&)
+	{
+		static_cast<T*>(Object)->~T();
+	}
+
+	static FORCEINLINE void Do(FTypeLayoutDesc& TypeDesc) { TypeDesc.DestroyFunc = &Destroy; }
+};
+
+
+/**
+ * Access to a global default object is required in order to patch vtables
+ * Normally this can be provided by a default-constructed object.  For objects without default constructors, a default object must be provided through the global function GetDefault<T>()
+ * THasCustomDefaultObject<T>::Value must be specialized to true for these types.  This is done by default for any UObject-derived types, which already provide GetDefault<T>()
+ */
+struct CProvidesDefaultUObject
+{
+	template<typename T>
+	auto Requires(const T&) -> decltype(T::StaticClass()->GetDefaultObject());
+};
+
+template<typename T>
+const T* GetDefault();
+
+template<typename T>
+struct THasCustomDefaultObject
+{
+	static const bool Value = TModels<CProvidesDefaultUObject, T>::Value;
+};
+
+template<typename T>
+typename TEnableIf<THasCustomDefaultObject<T>::Value, const T*>::Type InternalGetDefaultObject()
+{
+	return GetDefault<T>();
+}
+
+template<typename T>
+typename TEnableIf<!THasCustomDefaultObject<T>::Value, const T*>::Type InternalGetDefaultObject()
+{
+	static const T Default;
+	return &Default;
+}
+
+template<typename T, ETypeLayoutInterface::Type InterfaceType>
+struct TGetDefaultObjectHelper
+{
+	UE_STATIC_ONLY(TGetDefaultObjectHelper);
+	static const void* Do() { return nullptr; }
+};
+
+template<typename T>
+struct TGetDefaultObjectHelper<T, ETypeLayoutInterface::Virtual>
+{
+	UE_STATIC_ONLY(TGetDefaultObjectHelper);
+	static const void* Do() { return reinterpret_cast<const void*>(InternalGetDefaultObject<T>()); }
+};
+
+
+template<typename T, ETypeLayoutInterface::Type InterfaceType>
+struct TValidateInterfaceHelper;
+
+template<typename T> struct TValidateInterfaceHelper<T, ETypeLayoutInterface::NonVirtual>
+{
+	UE_STATIC_ONLY(TValidateInterfaceHelper);
+	static const bool Value = !TIsPolymorphic<T>::Value;
+};
+
+template<typename T> struct TValidateInterfaceHelper<T, ETypeLayoutInterface::Virtual>
+{
+	UE_STATIC_ONLY(TValidateInterfaceHelper);
+	static const bool Value = !TIsAbstract<T>::Value;
+};
+
+template<typename T> struct TValidateInterfaceHelper<T, ETypeLayoutInterface::Abstract>
+{
+	UE_STATIC_ONLY(TValidateInterfaceHelper);
+	static const bool Value = true;
+};
+
+namespace Freeze
+{
+	CORE_API void DefaultWriteMemoryImageField(FMemoryImageWriter& Writer, const void* Object, const void* FieldObject, const FTypeLayoutDesc& TypeDesc, const FTypeLayoutDesc& DerivedTypeDesc);
+	CORE_API void DefaultWriteMemoryImage(FMemoryImageWriter& Writer, const void* Object, const FTypeLayoutDesc& TypeDesc, const FTypeLayoutDesc& DerivedTypeDesc);
+	CORE_API void DefaultUnfrozenCopy(const FMemoryUnfreezeContent& Context, const void* Object, const FTypeLayoutDesc& TypeDesc, void* OutDst);
+	CORE_API uint32 DefaultAppendHash(const FTypeLayoutDesc& TypeDesc, const FPlatformTypeLayoutParameters& LayoutParams, FSHA1& Hasher);
+	CORE_API uint32 DefaultGetTargetAlignment(const FTypeLayoutDesc& TypeDesc, const FPlatformTypeLayoutParameters& LayoutParams);
+	CORE_API uint32 AppendHashForNameAndSize(const TCHAR* Name, uint32 Size, FSHA1& Hasher);
+	CORE_API void IntrinsicWriteMemoryImage(FMemoryImageWriter& Writer, const void* Object, uint32 Size);
+	CORE_API void IntrinsicWriteMemoryImage(FMemoryImageWriter& Writer, void*, const FTypeLayoutDesc&);
+
+	template<typename T>
+	FORCEINLINE void IntrinsicWriteMemoryImage(FMemoryImageWriter& Writer, const T& Object, const FTypeLayoutDesc& TypeDesc)
+	{
+		IntrinsicWriteMemoryImage(Writer, &Object, sizeof(T));
+	}
+
+	template<typename T>
+	FORCEINLINE void IntrinsicUnfrozenCopy(const FMemoryUnfreezeContent& Context, const T& Object, void* OutDst)
+	{
+		new(OutDst) T(Object);
+	}
+
+	template<typename T>
+	FORCEINLINE uint32 IntrinsicAppendHash(const T* DummyObject, const FTypeLayoutDesc& TypeDesc, const FPlatformTypeLayoutParameters& LayoutParams, FSHA1& Hasher)
+	{
+		return DefaultAppendHash(TypeDesc, LayoutParams, Hasher);
+	}
+
+	template<typename T>
+	FORCEINLINE uint32 IntrinsicGetTargetAlignment(const T* DummyObject, const FTypeLayoutDesc& TypeDesc, const FPlatformTypeLayoutParameters& LayoutParams)
+	{
+		return TypeDesc.Alignment;
+	}
+
+	inline uint32 IntrinsicAppendHash(void* const* DummyObject, const FTypeLayoutDesc& TypeDesc, const FPlatformTypeLayoutParameters& LayoutParams, FSHA1& Hasher)
+	{
+		return AppendHashForNameAndSize(TypeDesc.Name, LayoutParams.GetRawPointerSize(), Hasher);
+	}
+
+	inline uint32 IntrinsicGetTargetAlignment(void* const* DummyObject, const FTypeLayoutDesc& TypeDesc, const FPlatformTypeLayoutParameters& LayoutParams)
+	{
+		return LayoutParams.GetRawPointerSize();
+	}
+
+	// helper functions
+	CORE_API void ExtractBitFieldValue(const void* Value, uint32 SrcBitOffset, uint32 DestBitOffset, uint32 NumBits, uint64& InOutValue);
+	CORE_API bool IncludeField(const FFieldLayoutDesc* FieldDesc, const FPlatformTypeLayoutParameters& LayoutParams);
+	CORE_API uint32 GetTargetAlignment(const FTypeLayoutDesc& TypeDesc, const FPlatformTypeLayoutParameters& LayoutParams);
+
+	CORE_API uint32 AppendHash(const FTypeLayoutDesc& TypeDesc, const FPlatformTypeLayoutParameters& LayoutParams, FSHA1& Hasher);
+	CORE_API uint32 AppendHashPair(const FTypeLayoutDesc& KeyTypeDesc, const FTypeLayoutDesc& ValueTypeDesc, const FPlatformTypeLayoutParameters& LayoutParams, FSHA1& Hasher);
+	CORE_API uint32 HashLayout(const FTypeLayoutDesc& TypeDesc, const FPlatformTypeLayoutParameters& LayoutParams, FSHAHash& OutHash);
+	CORE_API uint32 HashLayouts(const TArray<const FTypeLayoutDesc*>& TypeLayouts, const FPlatformTypeLayoutParameters& LayoutParams, FSHAHash& OutHash);
+	CORE_API FSHAHash HashLayout(const FTypeLayoutDesc& TypeDesc, const FPlatformTypeLayoutParameters& LayoutParams);
+}
+
+
+struct CProvidesStaticClass
+{
+	template<typename T>
+	auto Requires(const T&) -> decltype(T::StaticClass());
+};
+
+struct CProvidesStaticStruct
+{
+	template<typename T>
+	auto Requires(const T&) -> decltype(T::StaticStruct());
+};
+
+template<typename T>
+struct TUsePropertyFreezing
+{
+	static const bool Value = (TModels<CProvidesStaticClass, T>::Value || TModels<CProvidesStaticStruct, T>::Value);
+};
+
+template<typename T>
+struct TProvidesStaticStruct
+{
+	static const bool Value = TModels<CProvidesStaticStruct, T>::Value;
+};
+
+template <typename T, bool bUsePropertyFreezing=TUsePropertyFreezing<T>::Value>
+struct TGetFreezeImageHelper
+{
+	static constexpr FTypeLayoutDesc::FWriteFrozenMemoryImageFunc* Value = &Freeze::DefaultWriteMemoryImage;
+};
+
+template <typename T, bool bProvidesStaticStruct=TProvidesStaticStruct<T>::Value>
+struct TGetFreezeImageFieldHelper
+{
+	static constexpr FFieldLayoutDesc::FWriteFrozenMemoryImageFunc* Value = &Freeze::DefaultWriteMemoryImageField;
+};
+
+template<typename T, typename BaseType>
+static void InternalInitializeBaseHelper(FTypeLayoutDesc& TypeDesc)
+{
+	alignas(FFieldLayoutDesc) static uint8 FieldBuffer[sizeof(FFieldLayoutDesc)] = { 0 };
+	FFieldLayoutDesc& FieldDesc = *(FFieldLayoutDesc*)FieldBuffer;
+	FieldDesc.Name = TEXT("BASE");
+	FieldDesc.UFieldNameLength = 4;
+	FieldDesc.Type = &StaticGetTypeLayoutDesc<BaseType>();
+	FieldDesc.WriteFrozenMemoryImageFunc = TGetFreezeImageFieldHelper<PREPROCESSOR_REMOVE_OPTIONAL_PARENS(T)>::Value;
+	FieldDesc.Offset = GetBaseOffset<T, BaseType>();
+	FieldDesc.NumArray = 1u;
+	FieldDesc.Next = TypeDesc.Fields;
+	TypeDesc.Fields = &FieldDesc;
+	++TypeDesc.NumBases;
+	if (ETypeLayoutInterface::HasVTable(FieldDesc.Type->Interface)) ++TypeDesc.NumVirtualBases;
+}
+
+template<typename T>
+static FORCEINLINE void InternalInitializeBasesHelper(FTypeLayoutDesc& TypeDesc) {}
+
+template<typename T, typename BaseType, typename ...RemainingBaseTypes>
+static void InternalInitializeBasesHelper(FTypeLayoutDesc& TypeDesc)
+{
+	InternalInitializeBasesHelper<T, RemainingBaseTypes...>(TypeDesc);
+	InternalInitializeBaseHelper<T, BaseType>(TypeDesc);
+}
+
+template<typename T>
+struct TGetBaseTypeHelper
+{
+	UE_STATIC_ONLY(TGetBaseTypeHelper);
+	template<typename InternalType> static typename InternalType::DerivedType Test(const typename InternalType::DerivedType*);
+	template<typename InternalType> static void Test(...);
+
+	using Type = decltype(Test<T>(nullptr));
+};
+
+template<typename T, typename BaseType>
+struct TInitializeBaseHelper
+{
+	UE_STATIC_ONLY(TInitializeBaseHelper);
+	static void Initialize(FTypeLayoutDesc& TypeDesc) { InternalInitializeBaseHelper<T, BaseType>(TypeDesc); }
+};
+
+template<typename T>
+struct TInitializeBaseHelper<T, void>
+{
+	UE_STATIC_ONLY(TInitializeBaseHelper);
+	static FORCEINLINE void Initialize(FTypeLayoutDesc& TypeDesc) {}
+};
+
+namespace Freeze
+{
+	// Finds the length of the field name, omitting any _DEPRECATED suffix
+	CORE_API uint8 FindFieldNameLength(const TCHAR* Name);
+}
+
+#define INTERNAL_LAYOUT_FIELD(T, InName, InOffset, InFlags, InNumArray, InBitFieldSize, Counter) \
+	PRAGMA_DISABLE_DEPRECATION_WARNINGS \
+	UE_DECLARE_INTERNAL_LINK_SPECIALIZATION(InternalLinkType, Counter - CounterBase) { \
+		UE_STATIC_ONLY(InternalLinkType); \
+		static void Initialize(FTypeLayoutDesc& TypeDesc) { \
+			InternalLinkType<Counter - CounterBase + 1>::Initialize(TypeDesc); \
+			alignas(FFieldLayoutDesc) static uint8 FieldBuffer[sizeof(FFieldLayoutDesc)] = { 0 }; \
+			FFieldLayoutDesc& FieldDesc = *(FFieldLayoutDesc*)FieldBuffer; \
+			FieldDesc.Name = TEXT(#InName); \
+			FieldDesc.UFieldNameLength = Freeze::FindFieldNameLength(FieldDesc.Name); \
+			FieldDesc.Type = &StaticGetTypeLayoutDesc<PREPROCESSOR_REMOVE_OPTIONAL_PARENS(T)>(); \
+			FieldDesc.WriteFrozenMemoryImageFunc = TGetFreezeImageFieldHelper<PREPROCESSOR_REMOVE_OPTIONAL_PARENS(T)>::Value; \
+			FieldDesc.Offset = InOffset; \
+			FieldDesc.NumArray = InNumArray; \
+			FieldDesc.Flags = InFlags; \
+			FieldDesc.BitFieldSize = InBitFieldSize; \
+			FieldDesc.Next = TypeDesc.Fields; \
+			TypeDesc.Fields = &FieldDesc; \
+		} \
+	}; \
+	PRAGMA_ENABLE_DEPRECATION_WARNINGS
+
+#define INTERNAL_LAYOUT_FIELD_WITH_WRITER(T, InName, InFunc, Counter) \
+	PRAGMA_DISABLE_DEPRECATION_WARNINGS \
+	UE_DECLARE_INTERNAL_LINK_SPECIALIZATION(InternalLinkType, Counter - CounterBase) { \
+		UE_STATIC_ONLY(InternalLinkType); \
+		static void CallWriteFrozenField(FMemoryImageWriter& Writer, const void* Object, const void* FieldObject, const FTypeLayoutDesc& TypeDesc, const FTypeLayoutDesc& DerivedTypeDesc) { \
+			static_cast<const DerivedType*>(Object)->InFunc(Writer, *(PREPROCESSOR_REMOVE_OPTIONAL_PARENS(T)*)FieldObject); \
+		} \
+		static void Initialize(FTypeLayoutDesc& TypeDesc) { \
+			InternalLinkType<Counter - CounterBase + 1>::Initialize(TypeDesc); \
+			alignas(FFieldLayoutDesc) static uint8 FieldBuffer[sizeof(FFieldLayoutDesc)] = { 0 }; \
+			FFieldLayoutDesc& FieldDesc = *(FFieldLayoutDesc*)FieldBuffer; \
+			FieldDesc.Name = TEXT(#InName); \
+			FieldDesc.UFieldNameLength = Freeze::FindFieldNameLength(FieldDesc.Name); \
+			FieldDesc.Type = &StaticGetTypeLayoutDesc<PREPROCESSOR_REMOVE_OPTIONAL_PARENS(T)>(); \
+			FieldDesc.WriteFrozenMemoryImageFunc = &CallWriteFrozenField; \
+			FieldDesc.Offset = STRUCT_OFFSET(DerivedType, InName); \
+			FieldDesc.NumArray = 1u; \
+			FieldDesc.Flags = EFieldLayoutFlags::None; \
+			FieldDesc.Next = TypeDesc.Fields; \
+			TypeDesc.Fields = &FieldDesc; \
+		} \
+	}; \
+PRAGMA_ENABLE_DEPRECATION_WARNINGS
+
+#define INTERNAL_LAYOUT_WRITE_MEMORY_IMAGE(Func, Counter) \
+	UE_DECLARE_INTERNAL_LINK_SPECIALIZATION(InternalLinkType, Counter - CounterBase) { \
+		UE_STATIC_ONLY(InternalLinkType); \
+		static void CallWriteMemoryImage(FMemoryImageWriter& Writer, const void* Object, const FTypeLayoutDesc& TypeDesc, const FTypeLayoutDesc& DerivedTypeDesc) { \
+			static_cast<const DerivedType*>(Object)->Func(Writer, TypeDesc); \
+		} \
+		static void Initialize(FTypeLayoutDesc& TypeDesc) { \
+			InternalLinkType<Counter - CounterBase + 1>::Initialize(TypeDesc); \
+			TypeDesc.WriteFrozenMemoryImageFunc = &CallWriteMemoryImage; \
+		} \
+	};
+
+#define LAYOUT_FIELD(T, Name, ...) PREPROCESSOR_REMOVE_OPTIONAL_PARENS(T) Name; INTERNAL_LAYOUT_FIELD(T, Name, STRUCT_OFFSET(DerivedType, Name), EFieldLayoutFlags::MakeFlags(__VA_ARGS__), 1u, 0u, __COUNTER__)
+#define LAYOUT_MUTABLE_FIELD(T, Name, ...) mutable PREPROCESSOR_REMOVE_OPTIONAL_PARENS(T) Name; INTERNAL_LAYOUT_FIELD(T, Name, STRUCT_OFFSET(DerivedType, Name), EFieldLayoutFlags::MakeFlags(__VA_ARGS__), 1u, 0u, __COUNTER__)
+#define LAYOUT_FIELD_INITIALIZED(T, Name, Value, ...) PREPROCESSOR_REMOVE_OPTIONAL_PARENS(T) Name = Value; INTERNAL_LAYOUT_FIELD(T, Name, STRUCT_OFFSET(DerivedType, Name), EFieldLayoutFlags::MakeFlags(__VA_ARGS__), 1u, 0u, __COUNTER__)
+#define LAYOUT_MUTABLE_FIELD_INITIALIZED(T, Name, Value, ...) mutable PREPROCESSOR_REMOVE_OPTIONAL_PARENS(T) Name = Value; INTERNAL_LAYOUT_FIELD(T, Name, STRUCT_OFFSET(DerivedType, Name), EFieldLayoutFlags::MakeFlags(__VA_ARGS__), 1u, 0u, __COUNTER__)
+#define LAYOUT_ARRAY(T, Name, NumArray, ...) PREPROCESSOR_REMOVE_OPTIONAL_PARENS(T) Name[NumArray]; INTERNAL_LAYOUT_FIELD(T, Name, STRUCT_OFFSET(DerivedType, Name), EFieldLayoutFlags::MakeFlags(__VA_ARGS__), NumArray, 0u, __COUNTER__)
+#define LAYOUT_MUTABLE_BITFIELD(T, Name, BitFieldSize, ...) mutable PREPROCESSOR_REMOVE_OPTIONAL_PARENS(T) Name : BitFieldSize; INTERNAL_LAYOUT_FIELD(T, Name, ~0u, EFieldLayoutFlags::MakeFlags(__VA_ARGS__), 1u, BitFieldSize, __COUNTER__)
+#define LAYOUT_BITFIELD(T, Name, BitFieldSize, ...) PREPROCESSOR_REMOVE_OPTIONAL_PARENS(T) Name : BitFieldSize; INTERNAL_LAYOUT_FIELD(T, Name, ~0u, EFieldLayoutFlags::MakeFlags(__VA_ARGS__), 1u, BitFieldSize, __COUNTER__)
+#define LAYOUT_FIELD_WITH_WRITER(T, Name, Func) PREPROCESSOR_REMOVE_OPTIONAL_PARENS(T) Name; INTERNAL_LAYOUT_FIELD_WITH_WRITER(T, Name, Func, __COUNTER__)
+#define LAYOUT_MUTABLE_FIELD_WITH_WRITER(T, Name, Func) mutable PREPROCESSOR_REMOVE_OPTIONAL_PARENS(T) Name; INTERNAL_LAYOUT_FIELD_WITH_WRITER(T, Name, Func, __COUNTER__)
+#define LAYOUT_WRITE_MEMORY_IMAGE(Func) INTERNAL_LAYOUT_WRITE_MEMORY_IMAGE(Func, __COUNTER__)
+
+#if WITH_EDITORONLY_DATA
+#define LAYOUT_FIELD_EDITORONLY(T, Name, ...) PREPROCESSOR_REMOVE_OPTIONAL_PARENS(T) Name; INTERNAL_LAYOUT_FIELD(T, Name, STRUCT_OFFSET(DerivedType, Name), EFieldLayoutFlags::MakeFlagsEditorOnly(__VA_ARGS__), 1u, 0u, __COUNTER__)
+#define LAYOUT_ARRAY_EDITORONLY(T, Name, NumArray, ...) PREPROCESSOR_REMOVE_OPTIONAL_PARENS(T) Name[NumArray]; INTERNAL_LAYOUT_FIELD(T, Name, STRUCT_OFFSET(DerivedType, Name), EFieldLayoutFlags::MakeFlagsEditorOnly(__VA_ARGS__), NumArray, 0u, __COUNTER__)
+#define LAYOUT_BITFIELD_EDITORONLY(T, Name, BitFieldSize, ...) PREPROCESSOR_REMOVE_OPTIONAL_PARENS(T) Name : BitFieldSize; INTERNAL_LAYOUT_FIELD(T, Name, ~0u, EFieldLayoutFlags::MakeFlagsEditorOnly(__VA_ARGS__), 1u, BitFieldSize, __COUNTER__)
+#else
+#define LAYOUT_FIELD_EDITORONLY(T, Name, ...)
+#define LAYOUT_ARRAY_EDITORONLY(T, Name, NumArray, ...)
+#define LAYOUT_BITFIELD_EDITORONLY(T, Name, BitFieldSize, ...)
+#endif
+
+#if WITH_RAYTRACING
+#define LAYOUT_FIELD_RAYTRACING(T, Name, ...) PREPROCESSOR_REMOVE_OPTIONAL_PARENS(T) Name; INTERNAL_LAYOUT_FIELD(T, Name, STRUCT_OFFSET(DerivedType, Name), EFieldLayoutFlags::MakeFlagsRayTracing(__VA_ARGS__), 1u, 0u, __COUNTER__)
+#define LAYOUT_FIELD_INITIALIZED_RAYTRACING(T, Name, Value, ...) PREPROCESSOR_REMOVE_OPTIONAL_PARENS(T) Name = Value; INTERNAL_LAYOUT_FIELD(T, Name, STRUCT_OFFSET(DerivedType, Name), EFieldLayoutFlags::MakeFlagsRayTracing(__VA_ARGS__), 1u, 0u, __COUNTER__)
+#define LAYOUT_ARRAY_RAYTRACING(T, Name, NumArray, ...) PREPROCESSOR_REMOVE_OPTIONAL_PARENS(T) Name[NumArray]; INTERNAL_LAYOUT_FIELD(T, Name, STRUCT_OFFSET(DerivedType, Name), EFieldLayoutFlags::MakeFlagsRayTracing(__VA_ARGS__), NumArray, 0u, __COUNTER__)
+#else
+#define LAYOUT_FIELD_RAYTRACING(T, Name, ...)
+#define LAYOUT_FIELD_INITIALIZED_RAYTRACING(T, Name, Value, ...)
+#define LAYOUT_ARRAY_RAYTRACING(T, Name, NumArray, ...)
+#endif
+
+#define INTERNAL_LAYOUT_INTERFACE_PREFIX_NonVirtual(...) __VA_ARGS__
+#define INTERNAL_LAYOUT_INTERFACE_PREFIX_Virtual(...) __VA_ARGS__ virtual
+#define INTERNAL_LAYOUT_INTERFACE_PREFIX_Abstract(...) virtual
+#define INTERNAL_LAYOUT_INTERFACE_PREFIX(Type) PREPROCESSOR_JOIN(INTERNAL_LAYOUT_INTERFACE_PREFIX_, Type)
+
+#define INTERNAL_LAYOUT_INTERFACE_SUFFIX_NonVirtual ;
+#define INTERNAL_LAYOUT_INTERFACE_SUFFIX_Virtual ;
+#define INTERNAL_LAYOUT_INTERFACE_SUFFIX_Abstract { return FTypeLayoutDesc::GetInvalidTypeLayout(); }
+#define INTERNAL_LAYOUT_INTERFACE_SUFFIX(Type) PREPROCESSOR_JOIN(INTERNAL_LAYOUT_INTERFACE_SUFFIX_, Type)
+
+#define INTERNAL_LAYOUT_INTERFACE_INLINE_IMPL_NonVirtual { return StaticGetTypeLayout(); }
+#define INTERNAL_LAYOUT_INTERFACE_INLINE_IMPL_Virtual { return StaticGetTypeLayout(); }
+#define INTERNAL_LAYOUT_INTERFACE_INLINE_IMPL_Abstract { return FTypeLayoutDesc::GetInvalidTypeLayout(); }
+#define INTERNAL_LAYOUT_INTERFACE_INLINE_IMPL(Type) PREPROCESSOR_JOIN(INTERNAL_LAYOUT_INTERFACE_INLINE_IMPL_, Type)
+
+#define INTERNAL_DECLARE_TYPE_LAYOUT_COMMON(T, InInterface) \
+	static const int CounterBase = __COUNTER__; \
+	public: using DerivedType = PREPROCESSOR_REMOVE_OPTIONAL_PARENS(T); \
+	static const ETypeLayoutInterface::Type InterfaceType = ETypeLayoutInterface::InInterface; \
+	UE_DECLARE_INTERNAL_LINK_BASE(InternalLinkType) { UE_STATIC_ONLY(InternalLinkType); static FORCEINLINE void Initialize(FTypeLayoutDesc& TypeDesc) {} }
+
+#define INTERNAL_DECLARE_INLINE_TYPE_LAYOUT(T, InInterface) \
+	public: static FTypeLayoutDesc& StaticGetTypeLayout() { \
+		static_assert(TValidateInterfaceHelper<PREPROCESSOR_REMOVE_OPTIONAL_PARENS(T), ETypeLayoutInterface::InInterface>::Value, #InInterface " is invalid interface for " #T); \
+		alignas(FTypeLayoutDesc) static uint8 TypeBuffer[sizeof(FTypeLayoutDesc)] = { 0 }; \
+		FTypeLayoutDesc& TypeDesc = *(FTypeLayoutDesc*)TypeBuffer; \
+		if (!TypeDesc.IsInitialized) { \
+			TypeDesc.IsInitialized = true; \
+			TypeDesc.Name = TEXT(#T); \
+			TypeDesc.WriteFrozenMemoryImageFunc = TGetFreezeImageHelper<PREPROCESSOR_REMOVE_OPTIONAL_PARENS(T)>::Value; \
+			TypeDesc.UnfrozenCopyFunc = &Freeze::DefaultUnfrozenCopy; \
+			TypeDesc.AppendHashFunc = &Freeze::DefaultAppendHash; \
+			TypeDesc.GetTargetAlignmentFunc = &Freeze::DefaultGetTargetAlignment; \
+			TypeDesc.Size = sizeof(PREPROCESSOR_REMOVE_OPTIONAL_PARENS(T)); \
+			TypeDesc.Alignment = alignof(PREPROCESSOR_REMOVE_OPTIONAL_PARENS(T)); \
+			TypeDesc.Interface = ETypeLayoutInterface::InInterface; \
+			TypeDesc.SizeFromFields = ~0u; \
+			TAssignDestroyHelper<PREPROCESSOR_REMOVE_OPTIONAL_PARENS(T), TIsTriviallyDestructible<PREPROCESSOR_REMOVE_OPTIONAL_PARENS(T)>::Value>::Do(TypeDesc); \
+			TypeDesc.GetDefaultObjectFunc = &TGetDefaultObjectHelper<PREPROCESSOR_REMOVE_OPTIONAL_PARENS(T), ETypeLayoutInterface::InInterface>::Do; \
+			InternalLinkType<1>::Initialize(TypeDesc); \
+			InternalInitializeBases<PREPROCESSOR_REMOVE_OPTIONAL_PARENS(T)>(TypeDesc); \
+			FTypeLayoutDesc::Initialize(TypeDesc); \
+		} \
+		return TypeDesc; } \
+	public: INTERNAL_LAYOUT_INTERFACE_PREFIX(InInterface)(PREPROCESSOR_NOTHING) const FTypeLayoutDesc& GetTypeLayout() const INTERNAL_LAYOUT_INTERFACE_INLINE_IMPL(InInterface) \
+	INTERNAL_DECLARE_TYPE_LAYOUT_COMMON(T, InInterface)
+
+#define INTERNAL_DECLARE_TYPE_LAYOUT(T, InInterface, RequiredAPI) \
+	public: RequiredAPI static FTypeLayoutDesc& StaticGetTypeLayout(); \
+	public: INTERNAL_LAYOUT_INTERFACE_PREFIX(InInterface)(RequiredAPI) const FTypeLayoutDesc& GetTypeLayout() const INTERNAL_LAYOUT_INTERFACE_SUFFIX(InInterface) \
+	INTERNAL_DECLARE_TYPE_LAYOUT_COMMON(T, InInterface)
+
+#define INTERNAL_DECLARE_LAYOUT_BASE(T) \
+	private: using InternalBaseType = typename TGetBaseTypeHelper<PREPROCESSOR_REMOVE_OPTIONAL_PARENS(T)>::Type; \
+	template<typename InternalType> static void InternalInitializeBases(FTypeLayoutDesc& TypeDesc) { TInitializeBaseHelper<InternalType, InternalBaseType>::Initialize(TypeDesc); }
+
+#define INTERNAL_DECLARE_LAYOUT_EXPLICIT_BASES(T, ...) \
+	template<typename InternalType> static void InternalInitializeBases(FTypeLayoutDesc& TypeDesc) { InternalInitializeBasesHelper<InternalType, __VA_ARGS__>(TypeDesc); }
+
+#define DECLARE_TYPE_LAYOUT(T, Interface) INTERNAL_DECLARE_LAYOUT_BASE(T); INTERNAL_DECLARE_TYPE_LAYOUT(T, Interface, PREPROCESSOR_NOTHING)
+#define DECLARE_INLINE_TYPE_LAYOUT(T, Interface) INTERNAL_DECLARE_LAYOUT_BASE(T); INTERNAL_DECLARE_INLINE_TYPE_LAYOUT(T, Interface)
+#define DECLARE_EXPORTED_TYPE_LAYOUT(T, RequiredAPI, Interface) INTERNAL_DECLARE_LAYOUT_BASE(T); INTERNAL_DECLARE_TYPE_LAYOUT(T, Interface, RequiredAPI)
+
+#define DECLARE_TYPE_LAYOUT_EXPLICIT_BASES(T, Interface, ...) INTERNAL_DECLARE_LAYOUT_EXPLICIT_BASES(T, __VA_ARGS__); INTERNAL_DECLARE_TYPE_LAYOUT(T, Interface, PREPROCESSOR_NOTHING)
+#define DECLARE_INLINE_TYPE_LAYOUT_EXPLICIT_BASES(T, Interface, ...) INTERNAL_DECLARE_LAYOUT_EXPLICIT_BASES(T, __VA_ARGS__); INTERNAL_DECLARE_INLINE_TYPE_LAYOUT(T, Interface)
+#define DECLARE_EXPORTED_TYPE_LAYOUT_EXPLICIT_BASES(T, RequiredAPI, Interface, ...) INTERNAL_DECLARE_LAYOUT_EXPLICIT_BASES(T, __VA_ARGS__); INTERNAL_DECLARE_TYPE_LAYOUT(T, Interface, RequiredAPI)
+
+#define INTERNAL_IMPLEMENT_TYPE_LAYOUT_COMMON(TemplatePrefix, T) \
+	PREPROCESSOR_REMOVE_OPTIONAL_PARENS(TemplatePrefix) FTypeLayoutDesc& PREPROCESSOR_REMOVE_OPTIONAL_PARENS(T)::StaticGetTypeLayout() { \
+		static_assert(TValidateInterfaceHelper<PREPROCESSOR_REMOVE_OPTIONAL_PARENS(T), InterfaceType>::Value, "Invalid interface for " #T); \
+		alignas(FTypeLayoutDesc) static uint8 TypeBuffer[sizeof(FTypeLayoutDesc)] = { 0 }; \
+		FTypeLayoutDesc& TypeDesc = *(FTypeLayoutDesc*)TypeBuffer; \
+		if (!TypeDesc.IsInitialized) { \
+			TypeDesc.IsInitialized = true; \
+			TypeDesc.Name = TEXT(#T); \
+			TypeDesc.WriteFrozenMemoryImageFunc = TGetFreezeImageHelper<PREPROCESSOR_REMOVE_OPTIONAL_PARENS(T)>::Value; \
+			TypeDesc.UnfrozenCopyFunc = &Freeze::DefaultUnfrozenCopy; \
+			TypeDesc.AppendHashFunc = &Freeze::DefaultAppendHash; \
+			TypeDesc.GetTargetAlignmentFunc = &Freeze::DefaultGetTargetAlignment; \
+			TypeDesc.Size = sizeof(PREPROCESSOR_REMOVE_OPTIONAL_PARENS(T)); \
+			TypeDesc.Alignment = alignof(PREPROCESSOR_REMOVE_OPTIONAL_PARENS(T)); \
+			TypeDesc.Interface = InterfaceType; \
+			TypeDesc.SizeFromFields = ~0u; \
+			TAssignDestroyHelper<PREPROCESSOR_REMOVE_OPTIONAL_PARENS(T), TIsTriviallyDestructible<PREPROCESSOR_REMOVE_OPTIONAL_PARENS(T)>::Value>::Do(TypeDesc); \
+			TypeDesc.GetDefaultObjectFunc = &TGetDefaultObjectHelper<PREPROCESSOR_REMOVE_OPTIONAL_PARENS(T), InterfaceType>::Do; \
+			InternalLinkType<1>::Initialize(TypeDesc); \
+			InternalInitializeBases<PREPROCESSOR_REMOVE_OPTIONAL_PARENS(T)>(TypeDesc); \
+			FTypeLayoutDesc::Initialize(TypeDesc); \
+		} \
+		return TypeDesc; }
+
+//#define INTERNAL_REGISTER_TYPE_LAYOUT(T) static const FRegisterTypeLayoutDesc ANONYMOUS_VARIABLE(RegisterTypeLayoutDesc)(PREPROCESSOR_REMOVE_OPTIONAL_PARENS(T)::StaticGetTypeLayout())
+#define INTERNAL_REGISTER_TYPE_LAYOUT(T) static const FDelayedAutoRegisterHelper ANONYMOUS_VARIABLE(DelayedAutoRegisterHelper)(EDelayedRegisterRunPhase::ShaderTypesReady, []{ FTypeLayoutDesc::Register(T::StaticGetTypeLayout()); });
+
+#define IMPLEMENT_UNREGISTERED_TEMPLATE_TYPE_LAYOUT(TemplatePrefix, T) \
+	INTERNAL_IMPLEMENT_TYPE_LAYOUT_COMMON(TemplatePrefix, T); \
+	PREPROCESSOR_REMOVE_OPTIONAL_PARENS(TemplatePrefix) const FTypeLayoutDesc& PREPROCESSOR_REMOVE_OPTIONAL_PARENS(T)::GetTypeLayout() const { return StaticGetTypeLayout(); }
+
+#define IMPLEMENT_TEMPLATE_TYPE_LAYOUT(TemplatePrefix, T) \
+	IMPLEMENT_UNREGISTERED_TEMPLATE_TYPE_LAYOUT(TemplatePrefix, T); \
+	INTERNAL_REGISTER_TYPE_LAYOUT(T)
+
+#define IMPLEMENT_TYPE_LAYOUT(T) IMPLEMENT_TEMPLATE_TYPE_LAYOUT(, T)
+
+#define IMPLEMENT_ABSTRACT_TYPE_LAYOUT(T) \
+	INTERNAL_IMPLEMENT_TYPE_LAYOUT_COMMON(, T); \
+	INTERNAL_REGISTER_TYPE_LAYOUT(T)
+
+#define REGISTER_INLINE_TYPE_LAYOUT(T) \
+	static const FDelayedAutoRegisterHelper ANONYMOUS_VARIABLE(DelayedAutoRegisterHelper)(EDelayedRegisterRunPhase::ShaderTypesReady, [] \
+		{ \
+			PREPROCESSOR_REMOVE_OPTIONAL_PARENS(T)::StaticGetTypeLayout().Name = TEXT(#T); \
+			FTypeLayoutDesc::Register(PREPROCESSOR_REMOVE_OPTIONAL_PARENS(T)::StaticGetTypeLayout()); \
+		} \
+	);
+
+struct CTypeLayout
+{
+	template<typename T>
+	auto Requires(const T&) -> decltype(T::StaticGetTypeLayout());
+};
+
+template<typename T>
+struct THasTypeLayout
+{
+	UE_STATIC_ONLY(THasTypeLayout);
+	static const bool Value = TModels<CTypeLayout, T>::Value;
+};
+
+template<typename T>
+struct TStaticGetTypeLayoutHelper
+{
+	UE_STATIC_ONLY(TStaticGetTypeLayoutHelper);
+	static const FTypeLayoutDesc& Do() { return T::StaticGetTypeLayout(); }
+};
+
+template<typename T>
+struct TGetTypeLayoutHelper
+{
+	UE_STATIC_ONLY(TGetTypeLayoutHelper);
+	static const FTypeLayoutDesc& Do(const T& Object) { return Object.GetTypeLayout(); }
+};
+
+template<typename T>
+inline const FTypeLayoutDesc& StaticGetTypeLayoutDesc() { return TStaticGetTypeLayoutHelper<T>::Do(); }
+
+template<typename T>
+inline const FTypeLayoutDesc& GetTypeLayoutDesc(const FPointerTableBase*, const T& Object) { return TGetTypeLayoutHelper<T>::Do(Object); }
+
+CORE_API extern void InternalDeleteObjectFromLayout(void* Object, const FTypeLayoutDesc& TypeDesc, bool bIsFrozen);
+
+template<typename T>
+inline void DeleteObjectFromLayout(T* Object, FPointerTableBase* PtrTable = nullptr, bool bIsFrozen = false)
+{
+	check(Object);
+	const FTypeLayoutDesc& TypeDesc = GetTypeLayoutDesc(PtrTable, *Object);
+	InternalDeleteObjectFromLayout(Object, TypeDesc, bIsFrozen);
+}
+
+#define DECLARE_TEMPLATE_INTRINSIC_TYPE_LAYOUT(TemplatePrefix, T) \
+	PREPROCESSOR_REMOVE_OPTIONAL_PARENS(TemplatePrefix) struct THasTypeLayout<PREPROCESSOR_REMOVE_OPTIONAL_PARENS(T)> { UE_STATIC_ONLY(THasTypeLayout); static const bool Value = true; }; \
+	PREPROCESSOR_REMOVE_OPTIONAL_PARENS(TemplatePrefix) struct TStaticGetTypeLayoutHelper<PREPROCESSOR_REMOVE_OPTIONAL_PARENS(T)> { \
+		UE_STATIC_ONLY(TStaticGetTypeLayoutHelper); \
+		static void CallWriteMemoryImage(FMemoryImageWriter& Writer, const void* Object, const FTypeLayoutDesc& TypeDesc, const FTypeLayoutDesc& DerivedTypeDesc) { \
+			Freeze::IntrinsicWriteMemoryImage(Writer, *static_cast<PREPROCESSOR_REMOVE_OPTIONAL_PARENS(T) const*>(Object), TypeDesc); \
+		} \
+		static void CallUnfrozenCopy(const FMemoryUnfreezeContent& Context, const void* Object, const FTypeLayoutDesc& TypeDesc, void* OutDst) { \
+			Freeze::IntrinsicUnfrozenCopy(Context, *static_cast<PREPROCESSOR_REMOVE_OPTIONAL_PARENS(T) const*>(Object), OutDst); \
+		} \
+		static uint32 CallAppendHash(const FTypeLayoutDesc& TypeDesc, const FPlatformTypeLayoutParameters& LayoutParams, FSHA1& Hasher) { \
+			return Freeze::IntrinsicAppendHash(static_cast<PREPROCESSOR_REMOVE_OPTIONAL_PARENS(T) const*>(nullptr), TypeDesc, LayoutParams, Hasher); \
+		} \
+		static uint32 CallGetTargetAlignment(const FTypeLayoutDesc& TypeDesc, const FPlatformTypeLayoutParameters& LayoutParams) { \
+			return Freeze::IntrinsicGetTargetAlignment(static_cast<PREPROCESSOR_REMOVE_OPTIONAL_PARENS(T) const*>(nullptr), TypeDesc, LayoutParams); \
+		} \
+		static const FTypeLayoutDesc& Do() { \
+			alignas(FTypeLayoutDesc) static uint8 TypeBuffer[sizeof(FTypeLayoutDesc)] = { 0 }; \
+			FTypeLayoutDesc& TypeDesc = *(FTypeLayoutDesc*)TypeBuffer; \
+			if (!TypeDesc.IsInitialized) { \
+				TypeDesc.IsInitialized = true; \
+				TypeDesc.IsIntrinsic = true; \
+				TypeDesc.Name = TEXT(#T); \
+				TypeDesc.WriteFrozenMemoryImageFunc = &CallWriteMemoryImage; \
+				TypeDesc.UnfrozenCopyFunc = &CallUnfrozenCopy; \
+				TypeDesc.AppendHashFunc = &CallAppendHash; \
+				TypeDesc.GetTargetAlignmentFunc = &CallGetTargetAlignment; \
+				TypeDesc.Size = sizeof(PREPROCESSOR_REMOVE_OPTIONAL_PARENS(T)); \
+				TypeDesc.Alignment = alignof(PREPROCESSOR_REMOVE_OPTIONAL_PARENS(T)); \
+				TypeDesc.Interface = ETypeLayoutInterface::NonVirtual; \
+				TypeDesc.SizeFromFields = sizeof(PREPROCESSOR_REMOVE_OPTIONAL_PARENS(T)); \
+				TAssignDestroyHelper<PREPROCESSOR_REMOVE_OPTIONAL_PARENS(T), TIsTriviallyDestructible<PREPROCESSOR_REMOVE_OPTIONAL_PARENS(T)>::Value>::Do(TypeDesc); \
+			} \
+			return TypeDesc; } }; \
+	PREPROCESSOR_REMOVE_OPTIONAL_PARENS(TemplatePrefix) struct TGetTypeLayoutHelper<PREPROCESSOR_REMOVE_OPTIONAL_PARENS(T)> { \
+		UE_STATIC_ONLY(TGetTypeLayoutHelper); \
+		static const FTypeLayoutDesc& Do(const PREPROCESSOR_REMOVE_OPTIONAL_PARENS(T)&) { return TStaticGetTypeLayoutHelper<PREPROCESSOR_REMOVE_OPTIONAL_PARENS(T)>::Do(); }}
+
+#define DECLARE_EXPORTED_TEMPLATE_INTRINSIC_TYPE_LAYOUT(TemplatePrefix, T, RequiredAPI) \
+	PREPROCESSOR_REMOVE_OPTIONAL_PARENS(TemplatePrefix) struct THasTypeLayout<PREPROCESSOR_REMOVE_OPTIONAL_PARENS(T)> { UE_STATIC_ONLY(THasTypeLayout); static const bool Value = true; }; \
+	PREPROCESSOR_REMOVE_OPTIONAL_PARENS(TemplatePrefix) struct TStaticGetTypeLayoutHelper<PREPROCESSOR_REMOVE_OPTIONAL_PARENS(T)> { \
+		UE_STATIC_ONLY(TStaticGetTypeLayoutHelper); \
+		static void CallWriteMemoryImage(FMemoryImageWriter& Writer, const void* Object, const FTypeLayoutDesc& TypeDesc, const FTypeLayoutDesc& DerivedTypeDesc); \
+		static void CallUnfrozenCopy(const FMemoryUnfreezeContent& Context, const void* Object, const FTypeLayoutDesc& TypeDesc, void* OutDst); \
+		static uint32 CallAppendHash(const FTypeLayoutDesc& TypeDesc, const FPlatformTypeLayoutParameters& LayoutParams, FSHA1& Hasher); \
+		static uint32 CallGetTargetAlignment(const FTypeLayoutDesc& TypeDesc, const FPlatformTypeLayoutParameters& LayoutParams); \
+		RequiredAPI static const FTypeLayoutDesc& Do(); }; \
+	PREPROCESSOR_REMOVE_OPTIONAL_PARENS(TemplatePrefix) struct TGetTypeLayoutHelper<PREPROCESSOR_REMOVE_OPTIONAL_PARENS(T)> { \
+		UE_STATIC_ONLY(TGetTypeLayoutHelper); \
+		static const FTypeLayoutDesc& Do(const PREPROCESSOR_REMOVE_OPTIONAL_PARENS(T)&) { return TStaticGetTypeLayoutHelper<PREPROCESSOR_REMOVE_OPTIONAL_PARENS(T)>::Do(); }}
+
+#define IMPLEMENT_EXPORTED_INTRINSIC_TYPE_LAYOUT(T) \
+		void TStaticGetTypeLayoutHelper<PREPROCESSOR_REMOVE_OPTIONAL_PARENS(T)>::CallWriteMemoryImage(FMemoryImageWriter& Writer, const void* Object, const FTypeLayoutDesc& TypeDesc, const FTypeLayoutDesc& DerivedTypeDesc) { \
+			Freeze::IntrinsicWriteMemoryImage(Writer, *static_cast<PREPROCESSOR_REMOVE_OPTIONAL_PARENS(T) const*>(Object), TypeDesc); \
+		} \
+		void TStaticGetTypeLayoutHelper<PREPROCESSOR_REMOVE_OPTIONAL_PARENS(T)>::CallUnfrozenCopy(const FMemoryUnfreezeContent& Context, const void* Object, const FTypeLayoutDesc& TypeDesc, void* OutDst) { \
+			Freeze::IntrinsicUnfrozenCopy(Context, *static_cast<PREPROCESSOR_REMOVE_OPTIONAL_PARENS(T) const*>(Object), OutDst); \
+		} \
+		uint32 TStaticGetTypeLayoutHelper<PREPROCESSOR_REMOVE_OPTIONAL_PARENS(T)>::CallAppendHash(const FTypeLayoutDesc& TypeDesc, const FPlatformTypeLayoutParameters& LayoutParams, FSHA1& Hasher) { \
+			return Freeze::IntrinsicAppendHash(static_cast<PREPROCESSOR_REMOVE_OPTIONAL_PARENS(T) const*>(nullptr), TypeDesc, LayoutParams, Hasher); \
+		} \
+		uint32 TStaticGetTypeLayoutHelper<PREPROCESSOR_REMOVE_OPTIONAL_PARENS(T)>::CallGetTargetAlignment(const FTypeLayoutDesc& TypeDesc, const FPlatformTypeLayoutParameters& LayoutParams) { \
+			return Freeze::IntrinsicGetTargetAlignment(static_cast<PREPROCESSOR_REMOVE_OPTIONAL_PARENS(T) const*>(nullptr), TypeDesc, LayoutParams); \
+		} \
+		const FTypeLayoutDesc& TStaticGetTypeLayoutHelper<PREPROCESSOR_REMOVE_OPTIONAL_PARENS(T)>::Do() { \
+			alignas(FTypeLayoutDesc) static uint8 TypeBuffer[sizeof(FTypeLayoutDesc)] = { 0 }; \
+			FTypeLayoutDesc& TypeDesc = *(FTypeLayoutDesc*)TypeBuffer; \
+			if (!TypeDesc.IsInitialized) { \
+				TypeDesc.IsInitialized = true; \
+				TypeDesc.IsIntrinsic = true; \
+				TypeDesc.Name = TEXT(#T); \
+				TypeDesc.WriteFrozenMemoryImageFunc = &CallWriteMemoryImage; \
+				TypeDesc.UnfrozenCopyFunc = &CallUnfrozenCopy; \
+				TypeDesc.AppendHashFunc = &CallAppendHash; \
+				TypeDesc.GetTargetAlignmentFunc = &CallGetTargetAlignment; \
+				TypeDesc.Size = sizeof(PREPROCESSOR_REMOVE_OPTIONAL_PARENS(T)); \
+				TypeDesc.Alignment = alignof(PREPROCESSOR_REMOVE_OPTIONAL_PARENS(T)); \
+				TypeDesc.Interface = ETypeLayoutInterface::NonVirtual; \
+				TypeDesc.SizeFromFields = sizeof(PREPROCESSOR_REMOVE_OPTIONAL_PARENS(T)); \
+				TAssignDestroyHelper<PREPROCESSOR_REMOVE_OPTIONAL_PARENS(T), TIsTriviallyDestructible<PREPROCESSOR_REMOVE_OPTIONAL_PARENS(T)>::Value>::Do(TypeDesc); \
+			} \
+			return TypeDesc; }
+
+#define DECLARE_INTRINSIC_TYPE_LAYOUT(T) DECLARE_TEMPLATE_INTRINSIC_TYPE_LAYOUT(template<>, T)
+
+#define ALIAS_TEMPLATE_TYPE_LAYOUT(TemplatePrefix, T, Alias) \
+	PREPROCESSOR_REMOVE_OPTIONAL_PARENS(TemplatePrefix) struct TStaticGetTypeLayoutHelper<PREPROCESSOR_REMOVE_OPTIONAL_PARENS(T)> : public TStaticGetTypeLayoutHelper<PREPROCESSOR_REMOVE_OPTIONAL_PARENS(Alias)> { UE_STATIC_ONLY(TStaticGetTypeLayoutHelper); }; \
+	PREPROCESSOR_REMOVE_OPTIONAL_PARENS(TemplatePrefix) struct TGetTypeLayoutHelper<PREPROCESSOR_REMOVE_OPTIONAL_PARENS(T)> : public TGetTypeLayoutHelper<PREPROCESSOR_REMOVE_OPTIONAL_PARENS(Alias)> { UE_STATIC_ONLY(TGetTypeLayoutHelper); }
+
+#define ALIAS_TYPE_LAYOUT(Type, Alias) \
+	static_assert(sizeof(PREPROCESSOR_REMOVE_OPTIONAL_PARENS(Type)) == sizeof(PREPROCESSOR_REMOVE_OPTIONAL_PARENS(Alias)), "Using a type alias but the sizes don't match!"); \
+	ALIAS_TEMPLATE_TYPE_LAYOUT(template<>, Type, Alias)
+
+DECLARE_INTRINSIC_TYPE_LAYOUT(char);
+DECLARE_INTRINSIC_TYPE_LAYOUT(short);
+DECLARE_INTRINSIC_TYPE_LAYOUT(int);
+DECLARE_INTRINSIC_TYPE_LAYOUT(int8);
+DECLARE_INTRINSIC_TYPE_LAYOUT(long);
+DECLARE_INTRINSIC_TYPE_LAYOUT(long long);
+DECLARE_INTRINSIC_TYPE_LAYOUT(unsigned char);
+DECLARE_INTRINSIC_TYPE_LAYOUT(unsigned short);
+DECLARE_INTRINSIC_TYPE_LAYOUT(unsigned int);
+DECLARE_INTRINSIC_TYPE_LAYOUT(unsigned long);
+DECLARE_INTRINSIC_TYPE_LAYOUT(unsigned long long);
+DECLARE_INTRINSIC_TYPE_LAYOUT(bool);
+DECLARE_INTRINSIC_TYPE_LAYOUT(float);
+DECLARE_INTRINSIC_TYPE_LAYOUT(double);
+DECLARE_INTRINSIC_TYPE_LAYOUT(wchar_t);
+DECLARE_INTRINSIC_TYPE_LAYOUT(char16_t);
+DECLARE_INTRINSIC_TYPE_LAYOUT(void*);
+
+DECLARE_INTRINSIC_TYPE_LAYOUT(FThreadSafeCounter);
+DECLARE_TEMPLATE_INTRINSIC_TYPE_LAYOUT(template<typename T>, TEnumAsByte<T>);
+
+// Map 'const' types to non-const type
+ALIAS_TEMPLATE_TYPE_LAYOUT(template<typename T>, const T, T);
+
+// All raw pointer types map to void*, since they're all handled the same way
+ALIAS_TEMPLATE_TYPE_LAYOUT(template<typename T>, T*, void*);
diff --git a/Engine/Source/Runtime/Core/Public/Stats/Stats.h b/Engine/Source/Runtime/Core/Public/Stats/Stats.h
--- a/Engine/Source/Runtime/Core/Public/Stats/Stats.h
+++ b/Engine/Source/Runtime/Core/Public/Stats/Stats.h
@@ -269,6 +269,7 @@ public:
 #endif
 
 #define SCOPE_SECONDS_ACCUMULATOR(Stat)
+#define SCOPE_MS_ACCUMULATOR(Stat)
 #define DEFINE_STAT(Stat)
 #define QUICK_USE_CYCLE_STAT(StatId,GroupId) TStatId()
 #define DECLARE_CYCLE_STAT(CounterName,StatId,GroupId)
@@ -355,20 +356,20 @@ struct FDynamicStats
 	}
 
 	template< typename TStatGroup >
-	static TStatId CreateStatIdInt64(const FString& StatNameOrDescription)
+	static TStatId CreateStatIdInt64(const FString& StatNameOrDescription, bool bIsAccumulator = false)
 	{
 #if	STATS
-		return CreateStatIdInternal<TStatGroup>(FName(*StatNameOrDescription), EStatDataType::ST_int64, false);
+		return CreateStatIdInternal<TStatGroup>(FName(*StatNameOrDescription), EStatDataType::ST_int64, false, !bIsAccumulator);
 #endif // STATS
 
 		return TStatId();
 	}
 
 	template< typename TStatGroup >
-	static TStatId CreateStatIdDouble(const FString& StatNameOrDescription)
+	static TStatId CreateStatIdDouble(const FString& StatNameOrDescription, bool bIsAccumulator=false)
 	{
 #if	STATS
-		return CreateStatIdInternal<TStatGroup>(FName(*StatNameOrDescription), EStatDataType::ST_double, false);
+		return CreateStatIdInternal<TStatGroup>(FName(*StatNameOrDescription), EStatDataType::ST_double, false, !bIsAccumulator);
 #endif // STATS
 		return TStatId();
 	}
@@ -417,20 +418,20 @@ struct FDynamicStats
 #if	STATS
 private: // private since this can only be declared if STATS is defined, due to EStatDataType in signature
 	template< typename TStatGroup >
-	static TStatId CreateStatIdInternal(const FName StatNameOrDescription, EStatDataType::Type Type, bool IsTimer)
+	static TStatId CreateStatIdInternal(const FName StatNameOrDescription, EStatDataType::Type Type, bool IsTimer, bool bClearEveryFrame=true)
 	{
 
 		FStartupMessages::Get().AddMetadata(StatNameOrDescription, nullptr,
 			TStatGroup::GetGroupName(),
 			TStatGroup::GetGroupCategory(),
 			TStatGroup::GetDescription(),
-			true, Type, IsTimer, false);
+			bClearEveryFrame, Type, IsTimer, false);
 
 		TStatId StatID = IStatGroupEnableManager::Get().GetHighPerformanceEnableForStat(StatNameOrDescription,
 			TStatGroup::GetGroupName(),
 			TStatGroup::GetGroupCategory(),
 			TStatGroup::DefaultEnable,
-			true, Type, nullptr, IsTimer, false);
+			bClearEveryFrame, Type, nullptr, IsTimer, false);
 
 		return StatID;
 	}
diff --git a/Engine/Source/Runtime/Core/Public/Stats/Stats2.h b/Engine/Source/Runtime/Core/Public/Stats/Stats2.h
--- a/Engine/Source/Runtime/Core/Public/Stats/Stats2.h
+++ b/Engine/Source/Runtime/Core/Public/Stats/Stats2.h
@@ -1588,16 +1588,17 @@ class FSimpleScopeSecondsStat
 {
 public:
 
-	FSimpleScopeSecondsStat(TStatId InStatId)
+	FSimpleScopeSecondsStat(TStatId InStatId, double InScale=1.0)
 		: StartTime(FPlatformTime::Seconds())
 		, StatId(InStatId)
+		, Scale(InScale)
 	{
 
 	}
 
 	virtual ~FSimpleScopeSecondsStat()
 	{
-		double TotalTime = FPlatformTime::Seconds() - StartTime;
+		double TotalTime = (FPlatformTime::Seconds() - StartTime) * Scale;
 		FThreadStats::AddMessage(StatId.GetName(), EStatOperation::Add, TotalTime);
 	}
 
@@ -1605,6 +1606,7 @@ private:
 
 	double StartTime;
 	TStatId StatId;
+	double Scale;
 };
 
 /** Manages startup messages, usually to update the metadata. */
@@ -1916,6 +1918,9 @@ struct FStat_##StatName\
 #define SCOPE_SECONDS_ACCUMULATOR(Stat) \
 	FSimpleScopeSecondsStat SecondsAccum_##Stat(GET_STATID(Stat));
 
+#define SCOPE_MS_ACCUMULATOR(Stat) \
+	FSimpleScopeSecondsStat SecondsAccum_##Stat(GET_STATID(Stat), 1000.0);
+
 #define SET_CYCLE_COUNTER(Stat,Cycles) \
 {\
 	if (FThreadStats::IsCollectingData() || !GET_STATISEVERYFRAME(Stat)) \
@@ -2165,6 +2170,8 @@ DECLARE_STATS_GROUP(TEXT("Init Views"),STATGROUP_InitViews, STATCAT_Advanced);
 DECLARE_STATS_GROUP(TEXT("Landscape"),STATGROUP_Landscape, STATCAT_Advanced);
 DECLARE_STATS_GROUP(TEXT("Light Rendering"),STATGROUP_LightRendering, STATCAT_Advanced);
 DECLARE_STATS_GROUP(TEXT("LoadTime"), STATGROUP_LoadTime, STATCAT_Advanced);
+DECLARE_STATS_GROUP(TEXT("LoadTimeClass"), STATGROUP_LoadTimeClass, STATCAT_Advanced);
+DECLARE_STATS_GROUP(TEXT("LoadTimeClassCount"), STATGROUP_LoadTimeClassCount, STATCAT_Advanced);
 DECLARE_STATS_GROUP_VERBOSE(TEXT("LoadTimeVerbose"), STATGROUP_LoadTimeVerbose, STATCAT_Advanced);
 DECLARE_STATS_GROUP_VERBOSE(TEXT("MathVerbose"), STATGROUP_MathVerbose, STATCAT_Advanced);
 DECLARE_STATS_GROUP(TEXT("Media"),STATGROUP_Media, STATCAT_Advanced);
diff --git a/Engine/Source/Runtime/Core/Public/Templates/RefCounting.h b/Engine/Source/Runtime/Core/Public/Templates/RefCounting.h
--- a/Engine/Source/Runtime/Core/Public/Templates/RefCounting.h
+++ b/Engine/Source/Runtime/Core/Public/Templates/RefCounting.h
@@ -5,6 +5,7 @@
 #include "CoreTypes.h"
 #include "Misc/AssertionMacros.h"
 #include "Serialization/Archive.h"
+#include "Serialization/MemoryLayout.h"
 #include "HAL/PlatformAtomics.h"
 
 /** A virtual interface for ref counted objects to implement. */
@@ -96,6 +97,36 @@ private:
 	mutable int32 NumRefs;
 };
 
+/**
+ * Like FRefCountedObject, but internal ref count is thread safe
+ */
+class CORE_API FThreadSafeRefCountedObject
+{
+public:
+	FThreadSafeRefCountedObject() : NumRefs(0) {}
+	virtual ~FThreadSafeRefCountedObject() { check(NumRefs.GetValue() == 0); }
+	uint32 AddRef() const
+	{
+		return uint32(NumRefs.Increment());
+	}
+	uint32 Release() const
+	{
+		uint32 Refs = uint32(NumRefs.Decrement());
+		if (Refs == 0)
+		{
+			delete this;
+		}
+		return Refs;
+	}
+	uint32 GetRefCount() const
+	{
+		return uint32(NumRefs.GetValue());
+	}
+private:
+	mutable FThreadSafeCounter NumRefs;
+};
+
+
 
 /**
  * A smart pointer to an object which implements AddRef/Release.
@@ -129,6 +160,16 @@ public:
 		}
 	}
 
+	template<typename CopyReferencedType>
+	explicit TRefCountPtr(const TRefCountPtr<CopyReferencedType>& Copy)
+	{
+		Reference = static_cast<ReferencedType*>(Copy.GetReference());
+		if (Reference)
+		{
+			Reference->AddRef();
+		}
+	}
+
 	FORCEINLINE TRefCountPtr(TRefCountPtr&& Copy)
 	{
 		Reference = Copy.Reference;
@@ -164,6 +205,12 @@ public:
 		return *this = InPtr.Reference;
 	}
 
+	template<typename CopyReferencedType>
+	FORCEINLINE TRefCountPtr& operator=(const TRefCountPtr<CopyReferencedType>& InPtr)
+	{
+		return *this = InPtr.GetReference();
+	}
+
 	TRefCountPtr& operator=(TRefCountPtr&& InPtr)
 	{
 		if (this != &InPtr)
@@ -254,6 +301,8 @@ private:
 	ReferencedType* Reference;
 };
 
+ALIAS_TEMPLATE_TYPE_LAYOUT(template<typename T>, TRefCountPtr<T>, void*);
+
 template<typename ReferencedType>
 FORCEINLINE bool operator==(const TRefCountPtr<ReferencedType>& A, const TRefCountPtr<ReferencedType>& B)
 {
diff --git a/Engine/Source/Runtime/Core/Public/Templates/SharedPointer.h b/Engine/Source/Runtime/Core/Public/Templates/SharedPointer.h
--- a/Engine/Source/Runtime/Core/Public/Templates/SharedPointer.h
+++ b/Engine/Source/Runtime/Core/Public/Templates/SharedPointer.h
@@ -1006,6 +1006,20 @@ private:
 	SharedPointerInternals::FSharedReferencer< Mode > SharedReferenceCount;
 };
 
+//template< class ObjectType, ESPMode Mode >
+//class TSharedPtr
+
+namespace Freeze
+{
+	template<class ObjectType, ESPMode Mode>
+	void IntrinsicWriteMemoryImage(FMemoryImageWriter& Writer, const TSharedPtr<ObjectType, Mode>& Object, const FTypeLayoutDesc&)
+	{
+		// we never want to freeze pointers, so write an empty one
+		Writer.WriteBytes(TSharedPtr<ObjectType, Mode>());
+	}
+}
+
+DECLARE_TEMPLATE_INTRINSIC_TYPE_LAYOUT((template<class ObjectType, ESPMode Mode>), (TSharedPtr<ObjectType, Mode>));
 
 template<class ObjectType, ESPMode Mode> struct TIsZeroConstructType<TSharedPtr<ObjectType, Mode>> { enum { Value = true }; };
 
diff --git a/Engine/Source/Runtime/Core/Public/Templates/Tuple.h b/Engine/Source/Runtime/Core/Public/Templates/Tuple.h
--- a/Engine/Source/Runtime/Core/Public/Templates/Tuple.h
+++ b/Engine/Source/Runtime/Core/Public/Templates/Tuple.h
@@ -9,6 +9,7 @@
 #include "Delegates/IntegerSequence.h"
 #include "Templates/Invoke.h"
 #include "Serialization/StructuredArchive.h"
+#include "Serialization/MemoryLayout.h"
 
 class FArchive;
 
@@ -418,6 +419,39 @@ public:
 	TTuple& operator=(const TTuple&) = default;
 };
 
+namespace Freeze
+{
+	template<typename KeyType, typename ValueType>
+	void IntrinsicWriteMemoryImage(FMemoryImageWriter& Writer, const TTuple<KeyType, ValueType>& Object, const FTypeLayoutDesc&)
+	{
+		Writer.WriteObject(Object.Key);
+		Writer.WriteObject(Object.Value);
+	}
+
+	template<typename KeyType, typename ValueType>
+	void IntrinsicUnfrozenCopy(const FMemoryUnfreezeContent& Context, const TTuple<KeyType, ValueType>& Object, void* OutDst)
+	{
+		TTuple<KeyType, ValueType>* DstObject = (TTuple<KeyType, ValueType>*)OutDst;
+		Context.UnfreezeObject(Object.Key, &DstObject->Key);
+		Context.UnfreezeObject(Object.Value, &DstObject->Value);
+	}
+
+	template<typename KeyType, typename ValueType>
+	uint32 IntrinsicAppendHash(const TTuple<KeyType, ValueType>* DummyObject, const FTypeLayoutDesc& TypeDesc, const FPlatformTypeLayoutParameters& LayoutParams, FSHA1& Hasher)
+	{
+		return Freeze::AppendHashPair(StaticGetTypeLayoutDesc<KeyType>(), StaticGetTypeLayoutDesc<ValueType>(), LayoutParams, Hasher);
+	}
+
+	template<typename KeyType, typename ValueType>
+	uint32 IntrinsicGetTargetAlignment(const TTuple<KeyType, ValueType>* DummyObject, const FTypeLayoutDesc& TypeDesc, const FPlatformTypeLayoutParameters& LayoutParams)
+	{
+		const uint32 KeyAlignment = GetTargetAlignment(StaticGetTypeLayoutDesc<KeyType>(), LayoutParams);
+		const uint32 ValueAlignment = GetTargetAlignment(StaticGetTypeLayoutDesc<ValueType>(), LayoutParams);
+		return FMath::Min(FMath::Max(KeyAlignment, ValueAlignment), LayoutParams.MaxFieldAlignment);
+	}
+}
+
+DECLARE_TEMPLATE_INTRINSIC_TYPE_LAYOUT((template <typename KeyType, typename ValueType>), (TTuple<KeyType, ValueType>));
 
 /**
  * Traits class which calculates the number of elements in a tuple.
diff --git a/Engine/Source/Runtime/Core/Public/Templates/UniquePtr.h b/Engine/Source/Runtime/Core/Public/Templates/UniquePtr.h
--- a/Engine/Source/Runtime/Core/Public/Templates/UniquePtr.h
+++ b/Engine/Source/Runtime/Core/Public/Templates/UniquePtr.h
@@ -6,6 +6,7 @@
 #include "Templates/UnrealTemplate.h"
 #include "Templates/IsArray.h"
 #include "Templates/RemoveExtent.h"
+#include "Serialization/MemoryLayout.h"
 
 // Single-ownership smart pointer in the vein of std::unique_ptr.
 // Use this when you need an object's lifetime to be strictly bound to the lifetime of a single smart pointer.
@@ -19,6 +20,8 @@
 template <typename T>
 struct TDefaultDelete
 {
+	DECLARE_INLINE_TYPE_LAYOUT(TDefaultDelete, NonVirtual);
+
 	TDefaultDelete() = default;
 	TDefaultDelete(const TDefaultDelete&) = default;
 	TDefaultDelete& operator=(const TDefaultDelete&) = default;
@@ -83,8 +86,10 @@ struct TDefaultDelete<T[]>
 };
 
 template <typename T, typename Deleter = TDefaultDelete<T>>
-class TUniquePtr : private Deleter
+class TUniquePtr : public /*private*/ Deleter // @todo loadtime: can we go back to private? I get this: error C2243: 'static_cast': conversion from 'T *' to 'Base *' exists, but is inaccessible
 {
+	DECLARE_INLINE_TYPE_LAYOUT_EXPLICIT_BASES(TUniquePtr, NonVirtual, Deleter);
+
 	template <typename OtherT, typename OtherDeleter>
 	friend class TUniquePtr;
 
@@ -318,7 +323,8 @@ private:
 	TUniquePtr(const TUniquePtr&);
 	TUniquePtr& operator=(const TUniquePtr&);
 
-	T* Ptr;
+	using PtrType = T*;
+	LAYOUT_FIELD(PtrType, Ptr);
 };
 
 template <typename T, typename Deleter>
diff --git a/Engine/Source/Runtime/Core/Public/Templates/UnrealTypeTraits.h b/Engine/Source/Runtime/Core/Public/Templates/UnrealTypeTraits.h
--- a/Engine/Source/Runtime/Core/Public/Templates/UnrealTypeTraits.h
+++ b/Engine/Source/Runtime/Core/Public/Templates/UnrealTypeTraits.h
@@ -17,6 +17,7 @@
 #include "Templates/IsEnum.h"
 #include "Templates/RemoveCV.h"
 #include "Traits/IsVoidType.h"
+#include "Templates/Models.h"
 
 #include "Templates/IsPODType.h"
 #include "Templates/IsTriviallyCopyConstructible.h"
diff --git a/Engine/Source/Runtime/Core/Public/UObject/LoadTimesObjectVersion.h b/Engine/Source/Runtime/Core/Public/UObject/LoadTimesObjectVersion.h
--- a/Engine/Source/Runtime/Core/Public/UObject/LoadTimesObjectVersion.h
+++ b/Engine/Source/Runtime/Core/Public/UObject/LoadTimesObjectVersion.h
@@ -1,7 +1,7 @@
 // Copyright Epic Games, Inc. All Rights Reserved.
 #pragma once
 
-// Custom serialization version for changes made in Dev-Rendering stream
+// Custom serialization version for changes made in LoadTimes stream
 struct CORE_API FLoadTimesObjectVersion
 {
 	enum Type
diff --git a/Engine/Source/Runtime/Core/Public/UObject/NameTypes.h b/Engine/Source/Runtime/Core/Public/UObject/NameTypes.h
--- a/Engine/Source/Runtime/Core/Public/UObject/NameTypes.h
+++ b/Engine/Source/Runtime/Core/Public/UObject/NameTypes.h
@@ -12,6 +12,7 @@
 #include "Containers/StringConv.h"
 #include "UObject/UnrealNames.h"
 #include "Templates/Atomic.h"
+#include "Serialization/MemoryLayout.h"
 
 /*----------------------------------------------------------------------------
 	Definitions.
@@ -924,6 +925,12 @@ private:
 template<> struct TIsZeroConstructType<class FName> { enum { Value = true }; };
 Expose_TNameOf(FName)
 
+namespace Freeze
+{
+	CORE_API void IntrinsicWriteMemoryImage(FMemoryImageWriter& Writer, const FName& Object, const FTypeLayoutDesc&);
+}
+
+DECLARE_INTRINSIC_TYPE_LAYOUT(FName);
 
 FORCEINLINE uint32 GetTypeHash(FName Name)
 {
diff --git a/Engine/Source/Runtime/Core/Public/UObject/UnrealNames.inl b/Engine/Source/Runtime/Core/Public/UObject/UnrealNames.inl
--- a/Engine/Source/Runtime/Core/Public/UObject/UnrealNames.inl
+++ b/Engine/Source/Runtime/Core/Public/UObject/UnrealNames.inl
@@ -137,6 +137,7 @@ REGISTER_NAME(255,Control)
 REGISTER_NAME(256,Voice)
 REGISTER_NAME(257, Zlib)
 REGISTER_NAME(258, Gzip)
+REGISTER_NAME(259, LZ4)
 
 // Online
 REGISTER_NAME(280,DGram)
diff --git a/Engine/Source/Runtime/Core/Public/Windows/WIndowsPlatform.h b/Engine/Source/Runtime/Core/Public/Windows/WIndowsPlatform.h
--- a/Engine/Source/Runtime/Core/Public/Windows/WIndowsPlatform.h
+++ b/Engine/Source/Runtime/Core/Public/Windows/WIndowsPlatform.h
@@ -118,6 +118,10 @@ typedef FWindowsPlatformTypes FPlatformTypes;
 	#define PRAGMA_ENABLE_OPTIMIZATION_ACTUAL  __pragma(clang optimize on)
 #endif
 
+#if !defined(__clang__)
+	#define PLATFORM_EMPTY_BASES __declspec(empty_bases)
+#endif
+
 // Tells the compiler to put the decorated function in a certain section (aka. segment) of the executable.
 #define PLATFORM_CODE_SECTION(Name) __declspec(code_seg(Name)) 
 
diff --git a/Engine/Source/Runtime/CoreUObject/Private/Misc/PackageName.cpp b/Engine/Source/Runtime/CoreUObject/Private/Misc/PackageName.cpp
--- a/Engine/Source/Runtime/CoreUObject/Private/Misc/PackageName.cpp
+++ b/Engine/Source/Runtime/CoreUObject/Private/Misc/PackageName.cpp
@@ -25,6 +25,7 @@
 #include "Misc/ScopeLock.h"
 #include "Misc/StringBuilder.h"
 #include "IO/IoDispatcher.h"
+#include "ProfilingDebugging/LoadTimeTracker.h"
 
 DEFINE_LOG_CATEGORY_STATIC(LogPackageName, Log, All);
 
@@ -912,6 +913,8 @@ bool FPackageName::FixPackageNameCase(FString& LongPackageName, const FStringVie
 
 bool FPackageName::DoesPackageExist(const FString& LongPackageName, const FGuid* Guid, FString* OutFilename, bool InAllowTextFormats)
 {
+	SCOPED_LOADTIMER(FPackageName_DoesPackageExist);
+
 	bool bFoundFile = false;
 
 	// Make sure passing filename as LongPackageName is supported.
