- Hopeful fix for CIS errors, I was unable to repro them

diff --git a/Engine/Source/Runtime/Core/Public/Containers/EnumAsByte.h b/Engine/Source/Runtime/Core/Public/Containers/EnumAsByte.h
--- a/Engine/Source/Runtime/Core/Public/Containers/EnumAsByte.h
+++ b/Engine/Source/Runtime/Core/Public/Containers/EnumAsByte.h
@@ -101,7 +101,8 @@ private:
 
 	FORCEINLINE friend uint32 GetTypeHash(const TEnumAsByte& Enum)
 	{
-		return GetTypeHash(Enum.Value);
+		// #include order makwes including TypeHash.h hard, but all GetTypeHash(uint8) does is return the uint8
+		return Enum.Value;
 	}
 };
 
