CIS fix - Move the typle hashing code to it's own header file so that including TypeHash.h in EnumAsByte.h no longer creates a cyclic dependency.

#rnx

diff --git a/Engine/Source/Runtime/Core/Public/Containers/EnumAsByte.h b/Engine/Source/Runtime/Core/Public/Containers/EnumAsByte.h
--- a/Engine/Source/Runtime/Core/Public/Containers/EnumAsByte.h
+++ b/Engine/Source/Runtime/Core/Public/Containers/EnumAsByte.h
@@ -5,6 +5,7 @@
 #include "CoreTypes.h"
 #include "Templates/IsPODType.h"
 #include "Templates/IsEnumClass.h"
+#include "Templates/TypeHash.h"
 
 template <bool> struct TEnumAsByte_EnumClass;
 template <> struct UE_DEPRECATED(4.15, "TEnumAsByte is not intended for use with enum classes - please derive your enum class from uint8 instead.") TEnumAsByte_EnumClass<true> {};
@@ -101,8 +102,7 @@ private:
 
 	FORCEINLINE friend uint32 GetTypeHash(const TEnumAsByte& Enum)
 	{
-		// #include order makwes including TypeHash.h hard, but all GetTypeHash(uint8) does is return the uint8
-		return Enum.Value;
+		return GetTypeHash(Enum.Value);
 	}
 };
 
diff --git a/Engine/Source/Runtime/Core/Public/Templates/Tuple.h b/Engine/Source/Runtime/Core/Public/Templates/Tuple.h
--- a/Engine/Source/Runtime/Core/Public/Templates/Tuple.h
+++ b/Engine/Source/Runtime/Core/Public/Templates/Tuple.h
@@ -10,6 +10,7 @@
 #include "Templates/Invoke.h"
 #include "Serialization/StructuredArchive.h"
 #include "Serialization/MemoryLayout.h"
+#include "Templates/TypeHash.h"
 
 class FArchive;
 
@@ -385,6 +386,26 @@ namespace UE4Tuple_Private
 	{
 		enum { Value = sizeof...(Types) };
 	};
+
+	template <uint32 ArgToCombine, uint32 ArgCount>
+	struct TGetTupleHashHelper
+	{
+		template <typename TupleType>
+		FORCEINLINE static uint32 Do(uint32 Hash, const TupleType& Tuple)
+		{
+			return TGetTupleHashHelper<ArgToCombine + 1, ArgCount>::Do(HashCombine(Hash, GetTypeHash(Tuple.template Get<ArgToCombine>())), Tuple);
+		}
+	};
+
+	template <uint32 ArgIndex>
+	struct TGetTupleHashHelper<ArgIndex, ArgIndex>
+	{
+		template <typename TupleType>
+		FORCEINLINE static uint32 Do(uint32 Hash, const TupleType& Tuple)
+		{
+			return Hash;
+		}
+	};
 }
 
 template <typename... Types>
@@ -419,6 +440,17 @@ public:
 	TTuple& operator=(const TTuple&) = default;
 };
 
+template <typename... Types>
+FORCEINLINE uint32 GetTypeHash(const TTuple<Types...>& Tuple)
+{
+	return UE4Tuple_Private::TGetTupleHashHelper<1u, sizeof...(Types)>::Do(GetTypeHash(Tuple.template Get<0>()), Tuple);
+}
+
+FORCEINLINE uint32 GetTypeHash(const TTuple<>& Tuple)
+{
+	return 0;
+}
+
 namespace Freeze
 {
 	template<typename KeyType, typename ValueType>
diff --git a/Engine/Source/Runtime/Core/Public/Templates/TypeHash.h b/Engine/Source/Runtime/Core/Public/Templates/TypeHash.h
--- a/Engine/Source/Runtime/Core/Public/Templates/TypeHash.h
+++ b/Engine/Source/Runtime/Core/Public/Templates/TypeHash.h
@@ -4,7 +4,6 @@
 
 #include "CoreTypes.h"
 #include "Templates/IsEnum.h"
-#include "Templates/Tuple.h"
 #include "Misc/Crc.h"
 
 /**
@@ -128,37 +127,3 @@ FORCEINLINE  typename TEnableIf<TIsEnum<EnumType>::Value, uint32>::Type GetTypeH
 {
 	return GetTypeHash((__underlying_type(EnumType))E);
 }
-
-namespace UE4TypeHash_Private
-{
-	template <uint32 ArgToCombine, uint32 ArgCount>
-	struct TGetTupleHashHelper
-	{
-		template <typename TupleType>
-		FORCEINLINE static uint32 Do(uint32 Hash, const TupleType& Tuple)
-		{
-			return TGetTupleHashHelper<ArgToCombine + 1, ArgCount>::Do(HashCombine(Hash, GetTypeHash(Tuple.template Get<ArgToCombine>())), Tuple);
-		}
-	};
-
-	template <uint32 ArgIndex>
-	struct TGetTupleHashHelper<ArgIndex, ArgIndex>
-	{
-		template <typename TupleType>
-		FORCEINLINE static uint32 Do(uint32 Hash, const TupleType& Tuple)
-		{
-			return Hash;
-		}
-	};
-}
-
-template <typename... Types>
-FORCEINLINE uint32 GetTypeHash(const TTuple<Types...>& Tuple)
-{
-	return UE4TypeHash_Private::TGetTupleHashHelper<1u, sizeof...(Types)>::Do(GetTypeHash(Tuple.template Get<0>()), Tuple);
-}
-
-FORCEINLINE uint32 GetTypeHash(const TTuple<>& Tuple)
-{
-	return 0;
-}
