Copying //UE4/Dev-Documentation to Samples-Main (//UE4/Samples-Main)

diff --git a/Engine/Config/BaseEditorSettings.ini b/Engine/Config/BaseEditorSettings.ini
--- a/Engine/Config/BaseEditorSettings.ini
+++ b/Engine/Config/BaseEditorSettings.ini
@@ -49,7 +49,6 @@ IncludeCollectionNames=True
 +Categories=(Identifier="Blueprints",Title=NSLOCTEXT("TutorialCategories","BlueprintsTitle","Blueprints"),Description=NSLOCTEXT("TutorialCategories","BlueprintsDescription","Tutorials covering the creation and usage of Blueprints."),Icon="FullBlueprintEditor.EditGlobalOptions",Texture=/Engine/Tutorial/BlueprintTutorials/TutorialAssets/Blueprint_64x.Blueprint_64x,SortOrder=200)
 +Categories=(Identifier="Editors",Title=NSLOCTEXT("TutorialCategories","EditorsTitle","Editors"),Description=NSLOCTEXT("TutorialCategories","EditorsDescription","Tutorials covering the various sub editors of Unreal Engine 4."),Icon=,Texture=/Engine/Tutorial/SubEditors/TutorialAssets/icon_Editor_Preferences_General_40x.icon_Editor_Preferences_General_40x,SortOrder=225)
 +Categories=(Identifier="Code",Title=NSLOCTEXT("TutorialCategories","CodeTitle","Code"),Description=NSLOCTEXT("TutorialCategories","CodeDescription","Write C++ code for use in Unreal Engine 4."),Icon="PlayWorld.RepeatLastPlay",Texture=None,SortOrder=250)
-+Categories=(Identifier="Paper2D",Title=NSLOCTEXT("TutorialCategories","Paper2DTitle","Paper2D"),Description=NSLOCTEXT("TutorialCategories","Paper2DDescription","Tutorials covering the usage of the Unreal Engine 4 2D game tool: Paper2D."),Icon="ClassThumbnail.PaperSprite",Texture=/Engine/Tutorial/Paper2D/TutorialAssets/Paper2DSprite_TutorialIcon.Paper2DSprite_TutorialIcon,SortOrder=500)
 +Categories=(Identifier="Animation",Title=NSLOCTEXT("TutorialCategories","AnimationTitle","Animation"),Description=NSLOCTEXT("TutorialCategories","AnimationDescription","Tutorials covering the animation system in Unreal Engine 4."),Icon=,Texture=/Engine/Tutorial/SubEditors/TutorialAssets/icon_ShowSkeletalMeshes_40x.icon_ShowSkeletalMeshes_40x,SortOrder=600)
 +Categories=(Identifier="Landscape",Title=NSLOCTEXT("TutorialCategories","LandscapeTitle","Landscape"),Description=NSLOCTEXT("TutorialCategories","LandscapeDescription","Tutorials covering the Unreal Editor 4 terrain editor: Landscape."),Icon="LevelEditor.LandscapeMode",Texture=/Engine/Tutorial/Landscape/TutorialAssets/Landscape.Landscape,SortOrder=700)
 +Categories=(Identifier="Foliage",Title=NSLOCTEXT("TutorialCategories","FoliageTitle","Foliage"),Description=NSLOCTEXT("TutorialCategories","FoliageDescription","Tutorials covering the Unreal Engine 4 Foliage tool."),Icon="LevelEditor.FoliageMode",Texture=/Engine/Tutorial/Foliage/TutorialAssets/Foliage.Foliage,SortOrder=800)
@@ -64,8 +63,6 @@ IncludeCollectionNames=True
 +TutorialContexts=(Context="MaterialEditor",BrowserFilter=,AttractTutorial=None,LaunchTutorial=/Engine/Tutorial/SubEditors/MaterialEditorTutorial.MaterialEditorTutorial_C)
 +TutorialContexts=(Context="Cascade",BrowserFilter=,AttractTutorial=None,LaunchTutorial=/Engine/Tutorial/SubEditors/ParticleSystemEditorTutorial.ParticleSystemEditorTutorial_C)
 +TutorialContexts=(Context="Persona",BrowserFilter=,AttractTutorial=None,LaunchTutorial=/Engine/Tutorial/SubEditors/PersonaAnimEditorWalkThrough.PersonaAnimEditorWalkthrough_C)
-+TutorialContexts=(Context="SpriteEditor",BrowserFilter=,AttractTutorial=None,LaunchTutorial=/Engine/Tutorial/Paper2D/SpriteEditorTutorial.SpriteEditorTutorial_C)
-+TutorialContexts=(Context="FlipbookEditor",BrowserFilter=,AttractTutorial=None,LaunchTutorial=/Engine/Tutorial/Paper2D/FlipbookEditorTutorial.FlipbookEditorTutorial_C)
 +TutorialContexts=(Context="LandscapeMode",BrowserFilter=,AttractTutorial=None,LaunchTutorial=/Engine/Tutorial/Landscape/Landscape_Manage_Mode.Landscape_Manage_Mode_C)
 +TutorialContexts=(Context="FoliageMode",BrowserFilter=,AttractTutorial=None,LaunchTutorial=/Engine/Tutorial/Foliage/Foliage_Intro_Tutorial.Foliage_Intro_Tutorial_C)
 
diff --git a/Engine/Plugins/2D/Paper2D/Source/Paper2DEditor/Paper2DEditor.Build.cs b/Engine/Plugins/2D/Paper2D/Source/Paper2DEditor/Paper2DEditor.Build.cs
--- a/Engine/Plugins/2D/Paper2D/Source/Paper2DEditor/Paper2DEditor.Build.cs
+++ b/Engine/Plugins/2D/Paper2D/Source/Paper2DEditor/Paper2DEditor.Build.cs
@@ -32,12 +32,12 @@ public class Paper2DEditor : ModuleRules
 				"Projects",
 				"NavigationSystem",
 				"ToolMenus",
-			});
+                "IntroTutorials"
+            });
 
 		PrivateIncludePathModuleNames.AddRange(
 			new string[] {
 				"Settings",
-				"IntroTutorials",
 				"AssetTools",
 				"LevelEditor"
 			});
diff --git a/Engine/Plugins/2D/Paper2D/Source/Paper2DEditor/Private/Paper2DEditorModule.cpp b/Engine/Plugins/2D/Paper2D/Source/Paper2DEditor/Private/Paper2DEditorModule.cpp
--- a/Engine/Plugins/2D/Paper2D/Source/Paper2DEditor/Private/Paper2DEditorModule.cpp
+++ b/Engine/Plugins/2D/Paper2D/Source/Paper2DEditor/Private/Paper2DEditorModule.cpp
@@ -71,6 +71,8 @@
 #include "ISettingsModule.h"
 
 // Intro tutorials
+#include "EditorTutorial.h"
+#include "IIntroTutorials.h"
 
 // Mesh paint adapters
 #include "MeshPaintModule.h"
@@ -218,6 +220,18 @@ public:
 			MeshPaintModule->RegisterGeometryAdapterFactory(SpriteMeshPaintAdapterFactory.ToSharedRef());
 		}
 
+		// Register Paper2D tutorial category with the tutorial browser
+		FTutorialCategory Paper2DCategory = FTutorialCategory();
+		Paper2DCategory.Identifier = "Paper2D";
+		Paper2DCategory.Title = NSLOCTEXT("TutorialCategories", "Paper2DTitle", "Paper2D");
+		Paper2DCategory.Description = NSLOCTEXT("TutorialCategories", "Paper2DDescription", "Tutorials covering the usage of the Unreal Engine 4 2D game tool: Paper2D.");
+		Paper2DCategory.Icon = "ClassThumbnail.PaperSprite";
+		Paper2DCategory.Texture = FSoftObjectPath("/Paper2D/Tutorial/Paper2D/TutorialAssets/Paper2DSprite_TutorialIcon.Paper2DSprite_TutorialIcon");
+		Paper2DCategory.SortOrder = 500;
+
+		IIntroTutorials& IntroTutorials = FModuleManager::GetModuleChecked<IIntroTutorials>(TEXT("IntroTutorials"));
+		IntroTutorials.RegisterCategory(Paper2DCategory);
+
 		//
 		RegisterSettings();
 	}
diff --git a/Engine/Source/Editor/IntroTutorials/Private/IntroTutorials.cpp b/Engine/Source/Editor/IntroTutorials/Private/IntroTutorials.cpp
--- a/Engine/Source/Editor/IntroTutorials/Private/IntroTutorials.cpp
+++ b/Engine/Source/Editor/IntroTutorials/Private/IntroTutorials.cpp
@@ -532,6 +532,7 @@ TSharedRef<SDockTab> FIntroTutorials::SpawnTutorialsBrowserTab(const FSpawnTabAr
 		.ToolTip(IDocumentation::Get()->CreateToolTip(Label, nullptr, "Shared/TutorialsBrowser", "Tab"));	
 
 	TSharedRef<STutorialsBrowser> TutorialsBrowser = SNew(STutorialsBrowser)
+		.ExternalCategories(ExternalCategories)
 		.OnLaunchTutorial(FOnLaunchTutorial::CreateRaw(this, &FIntroTutorials::LaunchTutorial));
 
 	NewTab->SetContent(TutorialsBrowser);
@@ -539,4 +540,11 @@ TSharedRef<SDockTab> FIntroTutorials::SpawnTutorialsBrowserTab(const FSpawnTabAr
 	return NewTab;
 }
 
+void FIntroTutorials::RegisterCategory(FTutorialCategory NewCategory)
+{
+	ExternalCategories.Add(NewCategory);
+}
+
+
+
 #undef LOCTEXT_NAMESPACE
diff --git a/Engine/Source/Editor/IntroTutorials/Private/IntroTutorials.h b/Engine/Source/Editor/IntroTutorials/Private/IntroTutorials.h
--- a/Engine/Source/Editor/IntroTutorials/Private/IntroTutorials.h
+++ b/Engine/Source/Editor/IntroTutorials/Private/IntroTutorials.h
@@ -123,6 +123,7 @@ public:
 	virtual void CloseAllTutorialContent() override;
 	virtual TSharedRef<SWidget> CreateTutorialsWidget(FName InContext, TWeakPtr<SWindow> InContextWindow = nullptr) const override;
 	virtual TSharedPtr<SWidget> CreateTutorialsLoadingWidget(TWeakPtr<SWindow> InContextWindow = nullptr) const override;
+	virtual void RegisterCategory(FTutorialCategory NewCategory) override;
 	// End of IIntroTutorials interface
 private:
 	/** The tab id of the tutorial tab */
@@ -157,4 +158,7 @@ private:
 
 	/** The collection of registered class type actions. */
 	TArray<TSharedRef<IClassTypeActions>> RegisteredClassTypeActions;
+
+	/** Extra tutorial categories that have been defined externally */
+	TArray<FTutorialCategory> ExternalCategories;
 };
diff --git a/Engine/Source/Editor/IntroTutorials/Private/STutorialsBrowser.cpp b/Engine/Source/Editor/IntroTutorials/Private/STutorialsBrowser.cpp
--- a/Engine/Source/Editor/IntroTutorials/Private/STutorialsBrowser.cpp
+++ b/Engine/Source/Editor/IntroTutorials/Private/STutorialsBrowser.cpp
@@ -543,6 +543,7 @@ void STutorialsBrowser::Construct(const FArguments& InArgs)
 	OnClosed = InArgs._OnClosed;
 	OnLaunchTutorial = InArgs._OnLaunchTutorial;
 	ParentWindow = InArgs._ParentWindow;
+	ExternalCategories = InArgs._ExternalCategories;
 
 	FAssetRegistryModule& AssetRegistryModule = FModuleManager::LoadModuleChecked<FAssetRegistryModule>(TEXT("AssetRegistry"));
 	AssetRegistryModule.Get().OnAssetAdded().AddSP(this, &STutorialsBrowser::HandleAssetAdded);
@@ -660,6 +661,12 @@ TSharedPtr<FTutorialListEntry_Category> STutorialsBrowser::RebuildCategories()
 	TSharedPtr<FTutorialListEntry_Category> RootCategory = MakeShareable(new FTutorialListEntry_Category(FOnCategorySelected::CreateSP(this, &STutorialsBrowser::OnCategorySelected)));
 	Categories.Add(RootCategory);
 
+	// add external categories
+	for (const auto& TutorialCategory : ExternalCategories)
+	{
+		Categories.Add(MakeShareable(new FTutorialListEntry_Category(TutorialCategory, FOnCategorySelected::CreateSP(this, &STutorialsBrowser::OnCategorySelected), TAttribute<FText>::Create(TAttribute<FText>::FGetter::CreateSP(this, &STutorialsBrowser::GetSearchText)))));
+	}
+
 	// rebuild categories
 	for(const auto& TutorialCategory : GetDefault<UTutorialSettings>()->Categories)
 	{
diff --git a/Engine/Source/Editor/IntroTutorials/Private/STutorialsBrowser.h b/Engine/Source/Editor/IntroTutorials/Private/STutorialsBrowser.h
--- a/Engine/Source/Editor/IntroTutorials/Private/STutorialsBrowser.h
+++ b/Engine/Source/Editor/IntroTutorials/Private/STutorialsBrowser.h
@@ -56,6 +56,7 @@ class STutorialsBrowser : public SCompoundWidget
 	SLATE_ARGUMENT(FSimpleDelegate, OnClosed)
 	SLATE_ARGUMENT(FOnLaunchTutorial, OnLaunchTutorial)
 	SLATE_ARGUMENT(TWeakPtr<SWindow>, ParentWindow)
+	SLATE_ARGUMENT(TArray<FTutorialCategory>, ExternalCategories)
 
 	SLATE_END_ARGS()
 
@@ -156,4 +157,7 @@ private:
 
 	/** Prevent us from refreshing too often */
 	float RefreshTimer;
+
+	/** Categories that were registered externally, i.e. by a plugin */
+	TArray<FTutorialCategory> ExternalCategories;
 };
diff --git a/Engine/Source/Editor/IntroTutorials/Private/TutorialStateSettings.cpp b/Engine/Source/Editor/IntroTutorials/Private/TutorialStateSettings.cpp
--- a/Engine/Source/Editor/IntroTutorials/Private/TutorialStateSettings.cpp
+++ b/Engine/Source/Editor/IntroTutorials/Private/TutorialStateSettings.cpp
@@ -3,6 +3,7 @@
 #include "TutorialStateSettings.h"
 #include "Templates/SubclassOf.h"
 #include "EditorTutorial.h"
+#include "Misc/PackageName.h"
 
 UTutorialStateSettings::UTutorialStateSettings(const FObjectInitializer& ObjectInitializer)
 	: Super(ObjectInitializer)
@@ -16,11 +17,14 @@ void UTutorialStateSettings::PostInitProperties()
 
 	for(const auto& Progress : TutorialsProgress)
 	{
-		TSubclassOf<UEditorTutorial> TutorialClass = LoadClass<UEditorTutorial>(NULL, *Progress.Tutorial.ToString(), NULL, LOAD_None, NULL);
-		if(TutorialClass != nullptr)
+		if (FPackageName::IsValidObjectPath(*Progress.Tutorial.ToString()))
 		{
-			UEditorTutorial* Tutorial = TutorialClass->GetDefaultObject<UEditorTutorial>();
-			ProgressMap.Add(Tutorial, Progress);
+			TSubclassOf<UEditorTutorial> TutorialClass = LoadClass<UEditorTutorial>(NULL, *Progress.Tutorial.ToString(), NULL, LOAD_None, NULL);
+			if (TutorialClass != nullptr)
+			{
+				UEditorTutorial* Tutorial = TutorialClass->GetDefaultObject<UEditorTutorial>();
+				ProgressMap.Add(Tutorial, Progress);
+			}
 		}
 	}
 }
diff --git a/Engine/Source/Editor/IntroTutorials/Public/IIntroTutorials.h b/Engine/Source/Editor/IntroTutorials/Public/IIntroTutorials.h
--- a/Engine/Source/Editor/IntroTutorials/Public/IIntroTutorials.h
+++ b/Engine/Source/Editor/IntroTutorials/Public/IIntroTutorials.h
@@ -9,6 +9,7 @@
 
 class SWindow;
 class UEditorTutorial;
+struct FTutorialCategory;
 
 /**
  * The public interface to this module.  In most cases, this interface is only public to sibling modules 
@@ -88,5 +89,12 @@ public:
 	* @return a widget used to access context-sensitive tutorials
 	*/
 	virtual TSharedPtr<SWidget> CreateTutorialsLoadingWidget(TWeakPtr<SWindow> InContextWindow = nullptr) const = 0;
+
+	/**
+	* Register a new category with the tutorial browser. Categories are automatically created when tutorial blueprints are discovered, 
+	* but this will allow you to apply some customization for tutorial categories defined by plugins
+	* @param	InTutorialCategory	The new category
+	*/
+	virtual void RegisterCategory(FTutorialCategory NewCategory) = 0;
 };
 
diff --git a/Engine/Source/Runtime/Core/Public/GenericPlatform/GenericPlatform.h b/Engine/Source/Runtime/Core/Public/GenericPlatform/GenericPlatform.h
--- a/Engine/Source/Runtime/Core/Public/GenericPlatform/GenericPlatform.h
+++ b/Engine/Source/Runtime/Core/Public/GenericPlatform/GenericPlatform.h
@@ -15,13 +15,15 @@ struct SelectIntPointerType
 template<typename T32BITS, typename T64BITS>
 struct SelectIntPointerType<T32BITS, T64BITS, 8>
 {
-	typedef T64BITS TIntPointer; // select the 64 bit type
+	// Select the 64 bit type.
+	typedef T64BITS TIntPointer;
 };
 
 template<typename T32BITS, typename T64BITS>
 struct SelectIntPointerType<T32BITS, T64BITS, 4>
 {
-	typedef T32BITS TIntPointer; // select the 32 bit type
+	// Select the 32 bit type.
+	typedef T32BITS TIntPointer;
 };
 
 /**
@@ -29,30 +31,65 @@ struct SelectIntPointerType<T32BITS, T64BITS, 4>
 **/
 struct FGenericPlatformTypes
 {
-	// Unsigned base types.
-	typedef unsigned char 		uint8;		// 8-bit  unsigned.
-	typedef unsigned short int	uint16;		// 16-bit unsigned.
-	typedef unsigned int		uint32;		// 32-bit unsigned.
-	typedef unsigned long long	uint64;		// 64-bit unsigned.
-
-	// Signed base types.
-	typedef	signed char			int8;		// 8-bit  signed.
-	typedef signed short int	int16;		// 16-bit signed.
-	typedef signed int	 		int32;		// 32-bit signed.
-	typedef signed long long	int64;		// 64-bit signed.
-
-	// Character types.
-	typedef char				ANSICHAR;	// An ANSI character       -                  8-bit fixed-width representation of 7-bit characters.
-	typedef wchar_t				WIDECHAR;	// A wide character        - In-memory only.  ?-bit fixed-width representation of the platform's natural wide character set.  Could be different sizes on different platforms.
-	typedef uint8				CHAR8;		// An 8-bit character type - In-memory only.  8-bit representation.  Should really be char8_t but making this the generic option is easier for compilers which don't fully support C++11 yet (i.e. MSVC).
-	typedef uint16				CHAR16;		// A 16-bit character type - In-memory only.  16-bit representation.  Should really be char16_t but making this the generic option is easier for compilers which don't fully support C++11 yet (i.e. MSVC).
-	typedef uint32				CHAR32;		// A 32-bit character type - In-memory only.  32-bit representation.  Should really be char32_t but making this the generic option is easier for compilers which don't fully support C++11 yet (i.e. MSVC).
-	typedef WIDECHAR			TCHAR;		// A switchable character  - In-memory only.  Either ANSICHAR or WIDECHAR, depending on a licensee's requirements.
-
-	typedef SelectIntPointerType<uint32, uint64, sizeof(void*)>::TIntPointer UPTRINT;	// unsigned int the same size as a pointer
-	typedef SelectIntPointerType<int32, int64, sizeof(void*)>::TIntPointer PTRINT;		// signed int the same size as a pointer
-	typedef UPTRINT SIZE_T;																// unsigned int the same size as a pointer
-	typedef PTRINT SSIZE_T;																// signed int the same size as a pointer
+	//~ Unsigned base types
+	
+	// 8-bit unsigned integer
+	typedef unsigned char 		uint8;
+	
+	// 16-bit unsigned integer
+	typedef unsigned short int	uint16;
+	
+	// 32-bit unsigned integer
+	typedef unsigned int		uint32;
+	
+	// 64-bit unsigned integer
+	typedef unsigned long long	uint64;
+
+	//~ Signed base types.
+	
+	// 8-bit signed integer
+	typedef	signed char			int8;
+	
+	// 16-bit signed integer
+	typedef signed short int	int16;
+	
+	// 32-bit signed integer
+	typedef signed int	 		int32;
+	
+	// 64-bit signed integer
+	typedef signed long long	int64;
+
+	//~ Character types
+	
+	// An ANSI character. 8-bit fixed-width representation of 7-bit characters.
+	typedef char				ANSICHAR;
+	
+	// A wide character. In-memory only. ?-bit fixed-width representation of the platform's natural wide character set. Could be different sizes on different platforms.
+	typedef wchar_t				WIDECHAR;
+	
+	// An 8-bit character type. In-memory only. 8-bit representation. Should really be char8_t but making this the generic option is easier for compilers which don't fully support C++11 yet (i.e. MSVC).
+	typedef uint8				CHAR8;
+	
+	// A 16-bit character type. In-memory only.  16-bit representation. Should really be char16_t but making this the generic option is easier for compilers which don't fully support C++11 yet (i.e. MSVC).
+	typedef uint16				CHAR16;		
+	
+	// A 32-bit character type. In-memory only. 32-bit representation. Should really be char32_t but making this the generic option is easier for compilers which don't fully support C++11 yet (i.e. MSVC).
+	typedef uint32				CHAR32;
+	
+	// A switchable character. In-memory only. Either ANSICHAR or WIDECHAR, depending on a licensee's requirements.
+	typedef WIDECHAR			TCHAR;
+
+	// Unsigned int. The same size as a pointer.
+	typedef SelectIntPointerType<uint32, uint64, sizeof(void*)>::TIntPointer UPTRINT;
+	
+	// Signed int. The same size as a pointer.
+	typedef SelectIntPointerType<int32, int64, sizeof(void*)>::TIntPointer PTRINT;
+	
+	// Unsigned int. The same size as a pointer.
+	typedef UPTRINT SIZE_T;
+	
+	// Signed int. The same size as a pointer.
+	typedef PTRINT SSIZE_T;
 
 	typedef int32					TYPE_OF_NULL;
 	typedef decltype(nullptr)		TYPE_OF_NULLPTR;
