Undo changelist 11279739.  This change was causing issues with ini serialization.

diff --git a/Engine/Source/Runtime/Core/Private/Misc/FileHelper.cpp b/Engine/Source/Runtime/Core/Private/Misc/FileHelper.cpp
--- a/Engine/Source/Runtime/Core/Private/Misc/FileHelper.cpp
+++ b/Engine/Source/Runtime/Core/Private/Misc/FileHelper.cpp
@@ -118,10 +118,10 @@ void FFileHelper::BufferToString( FString& Result, const uint8* Buffer, int32 Si
 			Size   -= 3;
 		}
 
-		int32 Length = FUTF8ToTCHAR_Convert::ConvertedLength(reinterpret_cast<const ANSICHAR*>(Buffer), Size);
-		ResultArray.AddUninitialized(Length + 1); // +1 for the null terminator
-		FUTF8ToTCHAR_Convert::Convert(ResultArray.GetData(), ResultArray.Num(), reinterpret_cast<const ANSICHAR*>(Buffer), Size);
-		ResultArray[Length] = TEXT('\0');
+		FUTF8ToTCHAR Conv((const ANSICHAR*)Buffer, Size);
+		int32 Length = Conv.Length();
+		ResultArray.AddUninitialized(Length + 1); // For the null terminator
+		CopyAssignItems(ResultArray.GetData(), Conv.Get(), Length);
 	}
 
 	if (ResultArray.Num() == 1)
