- Fixed UE_LOG and friends to work with SA more better

diff --git a/Engine/Source/Runtime/Core/Public/Logging/LogMacros.h b/Engine/Source/Runtime/Core/Public/Logging/LogMacros.h
--- a/Engine/Source/Runtime/Core/Public/Logging/LogMacros.h
+++ b/Engine/Source/Runtime/Core/Public/Logging/LogMacros.h
@@ -205,7 +205,6 @@ private:
 							FMsg::Logf_Internal(UE_LOG_SOURCE_FILE(__FILE__), __LINE__, LCategoryName.GetCategoryName(), ELogVerbosity::Verbosity, LFormat, UE_LOG_Args...); \
 							_DebugBreakAndPromptForRemote(); \
 							FDebug::ProcessFatalError(); \
-							CA_ASSUME(false); \
 						}, \
 						{ \
 							FMsg::Logf_Internal(nullptr, 0, LCategoryName.GetCategoryName(), ELogVerbosity::Verbosity, LFormat, UE_LOG_Args...); \
@@ -213,6 +212,7 @@ private:
 					) \
 				}; \
 				UE_LOG_noinline_lambda(CategoryName, Format, ##__VA_ARGS__); \
+				UE_LOG_EXPAND_IS_FATAL(Verbosity, CA_ASSUME(false);, PREPROCESSOR_NOTHING) \
 			} \
 		} \
 	}
@@ -284,7 +284,6 @@ private:
 								FMsg::Logf_Internal(UE_LOG_SOURCE_FILE(__FILE__), __LINE__, LCategoryName.GetCategoryName(), ELogVerbosity::Verbosity, LFormat, UE_LOG_Args...); \
 								_DebugBreakAndPromptForRemote(); \
 								FDebug::ProcessFatalError(); \
-								CA_ASSUME(false); \
 							}, \
 							{ \
 								FMsg::Logf_Internal(nullptr, 0, LCategoryName.GetCategoryName(), ELogVerbosity::Verbosity, LFormat, UE_LOG_Args...); \
@@ -293,6 +292,7 @@ private:
 						CA_ASSUME(true); \
 					}; \
 					UE_LOG_noinline_lambda(CategoryName, Format, ##__VA_ARGS__); \
+					UE_LOG_EXPAND_IS_FATAL(Verbosity, CA_ASSUME(false);, PREPROCESSOR_NOTHING) \
 				} \
 			} \
 		} \
