Rework some templates to hopefully make clang windows happy

diff --git a/Engine/Source/Runtime/Core/Public/Serialization/MemoryLayout.h b/Engine/Source/Runtime/Core/Public/Serialization/MemoryLayout.h
--- a/Engine/Source/Runtime/Core/Public/Serialization/MemoryLayout.h
+++ b/Engine/Source/Runtime/Core/Public/Serialization/MemoryLayout.h
@@ -367,13 +367,13 @@ struct TProvidesStaticStruct
 template <typename T, bool bUsePropertyFreezing=TUsePropertyFreezing<T>::Value>
 struct TGetFreezeImageHelper
 {
-	static constexpr FTypeLayoutDesc::FWriteFrozenMemoryImageFunc* Value = &Freeze::DefaultWriteMemoryImage;
+	static FORCEINLINE FTypeLayoutDesc::FWriteFrozenMemoryImageFunc* Do() { return &Freeze::DefaultWriteMemoryImage; }
 };
 
 template <typename T, bool bProvidesStaticStruct=TProvidesStaticStruct<T>::Value>
 struct TGetFreezeImageFieldHelper
 {
-	static constexpr FFieldLayoutDesc::FWriteFrozenMemoryImageFunc* Value = &Freeze::DefaultWriteMemoryImageField;
+	static FORCEINLINE FFieldLayoutDesc::FWriteFrozenMemoryImageFunc* Do() { return &Freeze::DefaultWriteMemoryImageField; }
 };
 
 template<typename T, typename BaseType>
@@ -384,7 +384,7 @@ static void InternalInitializeBaseHelper(FTypeLayoutDesc& TypeDesc)
 	FieldDesc.Name = TEXT("BASE");
 	FieldDesc.UFieldNameLength = 4;
 	FieldDesc.Type = &StaticGetTypeLayoutDesc<BaseType>();
-	FieldDesc.WriteFrozenMemoryImageFunc = TGetFreezeImageFieldHelper<PREPROCESSOR_REMOVE_OPTIONAL_PARENS(T)>::Value;
+	FieldDesc.WriteFrozenMemoryImageFunc = TGetFreezeImageFieldHelper<PREPROCESSOR_REMOVE_OPTIONAL_PARENS(T)>::Do();
 	FieldDesc.Offset = GetBaseOffset<T, BaseType>();
 	FieldDesc.NumArray = 1u;
 	FieldDesc.Next = TypeDesc.Fields;
@@ -444,7 +444,7 @@ namespace Freeze
 			FieldDesc.Name = TEXT(#InName); \
 			FieldDesc.UFieldNameLength = Freeze::FindFieldNameLength(FieldDesc.Name); \
 			FieldDesc.Type = &StaticGetTypeLayoutDesc<PREPROCESSOR_REMOVE_OPTIONAL_PARENS(T)>(); \
-			FieldDesc.WriteFrozenMemoryImageFunc = TGetFreezeImageFieldHelper<PREPROCESSOR_REMOVE_OPTIONAL_PARENS(T)>::Value; \
+			FieldDesc.WriteFrozenMemoryImageFunc = TGetFreezeImageFieldHelper<PREPROCESSOR_REMOVE_OPTIONAL_PARENS(T)>::Do(); \
 			FieldDesc.Offset = InOffset; \
 			FieldDesc.NumArray = InNumArray; \
 			FieldDesc.Flags = InFlags; \
@@ -551,7 +551,7 @@ PRAGMA_ENABLE_DEPRECATION_WARNINGS
 		if (!TypeDesc.IsInitialized) { \
 			TypeDesc.IsInitialized = true; \
 			TypeDesc.Name = TEXT(#T); \
-			TypeDesc.WriteFrozenMemoryImageFunc = TGetFreezeImageHelper<PREPROCESSOR_REMOVE_OPTIONAL_PARENS(T)>::Value; \
+			TypeDesc.WriteFrozenMemoryImageFunc = TGetFreezeImageHelper<PREPROCESSOR_REMOVE_OPTIONAL_PARENS(T)>::Do(); \
 			TypeDesc.UnfrozenCopyFunc = &Freeze::DefaultUnfrozenCopy; \
 			TypeDesc.AppendHashFunc = &Freeze::DefaultAppendHash; \
 			TypeDesc.GetTargetAlignmentFunc = &Freeze::DefaultGetTargetAlignment; \
@@ -597,7 +597,7 @@ PRAGMA_ENABLE_DEPRECATION_WARNINGS
 		if (!TypeDesc.IsInitialized) { \
 			TypeDesc.IsInitialized = true; \
 			TypeDesc.Name = TEXT(#T); \
-			TypeDesc.WriteFrozenMemoryImageFunc = TGetFreezeImageHelper<PREPROCESSOR_REMOVE_OPTIONAL_PARENS(T)>::Value; \
+			TypeDesc.WriteFrozenMemoryImageFunc = TGetFreezeImageHelper<PREPROCESSOR_REMOVE_OPTIONAL_PARENS(T)>::Do(); \
 			TypeDesc.UnfrozenCopyFunc = &Freeze::DefaultUnfrozenCopy; \
 			TypeDesc.AppendHashFunc = &Freeze::DefaultAppendHash; \
 			TypeDesc.GetTargetAlignmentFunc = &Freeze::DefaultGetTargetAlignment; \
