- Possible Linux IncludeTool analysis fix

diff --git a/Engine/Source/Runtime/Core/Private/Misc/Compression.cpp b/Engine/Source/Runtime/Core/Private/Misc/Compression.cpp
--- a/Engine/Source/Runtime/Core/Private/Misc/Compression.cpp
+++ b/Engine/Source/Runtime/Core/Private/Misc/Compression.cpp
@@ -11,7 +11,6 @@
 #include "Misc/ConfigCacheIni.h"
 #include "Misc/CompressedGrowableBuffer.h"
 #include "Misc/ICompressionFormat.h"
-#include "Compression/lz4hc.h"
 
 #include "Misc/MemoryReadStream.h"
 // #include "TargetPlatformBase.h"
@@ -19,6 +18,11 @@ THIRD_PARTY_INCLUDES_START
 #include "ThirdParty/zlib/zlib-1.2.5/Inc/zlib.h"
 THIRD_PARTY_INCLUDES_END
 
+THIRD_PARTY_INCLUDES_START
+#define LZ4_HC_STATIC_LINKING_ONLY
+#include "Compression/lz4hc.h"
+THIRD_PARTY_INCLUDES_END
+
 DECLARE_LOG_CATEGORY_EXTERN(LogCompression, Log, All);
 DEFINE_LOG_CATEGORY(LogCompression);
 
