Fix missing include
#rnx

diff --git a/Engine/Source/Runtime/Core/Public/Misc/DateTime.h b/Engine/Source/Runtime/Core/Public/Misc/DateTime.h
--- a/Engine/Source/Runtime/Core/Public/Misc/DateTime.h
+++ b/Engine/Source/Runtime/Core/Public/Misc/DateTime.h
@@ -5,6 +5,7 @@
 #include "CoreTypes.h"
 #include "Containers/UnrealString.h"
 #include "Misc/Timespan.h"
+#include "Serialization/StructuredArchive.h"
 #include "Templates/TypeHash.h"
 
 class FArchive;
