Add option CONTAINER_INITIAL_ALLOC_ZERO_SLACK to disable array slack, for easier toggling if there are issues, or per-project toggling. Check in disabled for now

diff --git a/Engine/Source/Runtime/Core/Public/Containers/ContainerAllocationPolicies.h b/Engine/Source/Runtime/Core/Public/Containers/ContainerAllocationPolicies.h
--- a/Engine/Source/Runtime/Core/Public/Containers/ContainerAllocationPolicies.h
+++ b/Engine/Source/Runtime/Core/Public/Containers/ContainerAllocationPolicies.h
@@ -11,6 +11,13 @@
 #include "Math/NumericLimits.h"
 #include "Templates/IsPolymorphic.h"
 
+// This option disables array slack for initial allocations, e.g where TArray::SetNum 
+// is called. This tends to save a lot of memory with almost no measured performance cost.
+// NOTE: This can cause latent memory corruption issues to become more prominent
+#ifndef CONTAINER_INITIAL_ALLOC_ZERO_SLACK
+#define CONTAINER_INITIAL_ALLOC_ZERO_SLACK 0
+#endif
+
 class FDefaultBitArrayAllocator;
 
 template<int IndexSize> class TSizedDefaultAllocator;
@@ -76,11 +83,24 @@ FORCEINLINE SizeType DefaultCalculateSlackGrow(SizeType NumElements, SizeType Nu
 	checkSlow(NumElements > NumAllocatedElements && NumElements > 0);
 
 	SIZE_T Grow = FirstGrow; // this is the amount for the first alloc
+
+#if CONTAINER_INITIAL_ALLOC_ZERO_SLACK
+	if (NumAllocatedElements)
+	{
+		// Allocate slack for the array proportional to its size.
+		Grow = SIZE_T(NumElements) + 3 * SIZE_T(NumElements) / 8 + ConstantGrow;
+	}
+	else if (SIZE_T(NumElements) > Grow)
+	{
+		Grow = SIZE_T(NumElements);
+	}
+#else
 	if (NumAllocatedElements || SIZE_T(NumElements) > Grow)
 	{
 		// Allocate slack for the array proportional to its size.
 		Grow = SIZE_T(NumElements) + 3 * SIZE_T(NumElements) / 8 + ConstantGrow;
 	}
+#endif
 	if (bAllowQuantize)
 	{
 		Retval = (SizeType)(FMemory::QuantizeSize(Grow * BytesPerElement, Alignment) / BytesPerElement);
