- Undid #pragma once, it didn't compile

diff --git a/Engine/Source/Runtime/Core/Public/Compression/lz4.h b/Engine/Source/Runtime/Core/Public/Compression/lz4.h
--- a/Engine/Source/Runtime/Core/Public/Compression/lz4.h
+++ b/Engine/Source/Runtime/Core/Public/Compression/lz4.h
@@ -33,8 +33,6 @@
     - LZ4 source repository : [REDACTED_PROFILE_URL]/lz4
 */
 
-#pragma once
-
 #include "CoreMinimal.h"
 
 #if defined (__cplusplus)
