Fix for accessing engine configs in SCWs.

diff --git a/Engine/Source/Runtime/Core/Private/Misc/ConfigCacheIni.cpp b/Engine/Source/Runtime/Core/Private/Misc/ConfigCacheIni.cpp
--- a/Engine/Source/Runtime/Core/Private/Misc/ConfigCacheIni.cpp
+++ b/Engine/Source/Runtime/Core/Private/Misc/ConfigCacheIni.cpp
@@ -3801,6 +3801,7 @@ void FConfigCacheIni::SerializeStateForBootstrap_Impl(FArchive& Ar)
 	Ar << GGameIni;
 	Ar << GGameUserSettingsIni;
 	Ar << GRuntimeOptionsIni;
+	Ar << GEngineIni;
 }
 
 void FConfigCacheIni::InitializeConfigSystem()
