Disable warning after upgrading to PVS-Studio 7.0.4:

V597: The compiler could delete the 'Memset' function call, which is used to flush 'x' buffer. The RtlSecureZeroMemory() function should be used to erase the private data.

#rnx

diff --git a/Engine/Source/Runtime/Core/Private/Misc/SecureHash.cpp b/Engine/Source/Runtime/Core/Private/Misc/SecureHash.cpp
--- a/Engine/Source/Runtime/Core/Private/Misc/SecureHash.cpp
+++ b/Engine/Source/Runtime/Core/Private/Misc/SecureHash.cpp
@@ -238,7 +238,7 @@ void FMD5::Transform( uint32* state, const uint8* block )
 	state[3] += d;
 
 	// Zeroize sensitive information.
-	FMemory::Memset( x, 0, sizeof(x) );
+	FMemory::Memset( x, 0, sizeof(x) ); //-V597
 }
 
 void FMD5::Encode( uint8* output, const uint32* input, int32 len )
