Fix for warning after upgrading to PVS-Studio 7.0.4:

V598: The 'Memcpy' function is used to copy the fields of 'TDebugHelper' class. Virtual table pointer will be damaged by this.

#rnx

diff --git a/Engine/Source/Runtime/Core/Public/Templates/Function.h b/Engine/Source/Runtime/Core/Public/Templates/Function.h
--- a/Engine/Source/Runtime/Core/Public/Templates/Function.h
+++ b/Engine/Source/Runtime/Core/Public/Templates/Function.h
@@ -575,7 +575,7 @@ namespace UE4Function_Private
 				#if ENABLE_TFUNCTIONREF_VISUALIZATION
 					// Use Memcpy to copy the other DebugPtrStorage, including vptr (because we don't know the bound type
 					// here), and then reseat the underlying pointer.  Possibly even more evil than the Set code.
-					FMemory::Memcpy(&DebugPtrStorage, &Other.DebugPtrStorage, sizeof(DebugPtrStorage));
+					FMemory::Memcpy(&DebugPtrStorage, &Other.DebugPtrStorage, sizeof(DebugPtrStorage)); //-V598
 					DebugPtrStorage.Ptr = Storage.GetPtr();
 				#endif
 
@@ -593,7 +593,7 @@ namespace UE4Function_Private
 				#if ENABLE_TFUNCTIONREF_VISUALIZATION
 					// Use Memcpy to copy the other DebugPtrStorage, including vptr (because we don't know the bound type
 					// here), and then reseat the underlying pointer.  Possibly even more evil than the Set code.
-					FMemory::Memcpy(&DebugPtrStorage, &Other.DebugPtrStorage, sizeof(DebugPtrStorage));
+					FMemory::Memcpy(&DebugPtrStorage, &Other.DebugPtrStorage, sizeof(DebugPtrStorage)); //-V598
 					DebugPtrStorage.Ptr = Storage.GetPtr();
 				#endif
 
@@ -612,7 +612,7 @@ namespace UE4Function_Private
 				#if ENABLE_TFUNCTIONREF_VISUALIZATION
 					// Use Memcpy to copy the other DebugPtrStorage, including vptr (because we don't know the bound type
 					// here), and then reseat the underlying pointer.  Possibly even more evil than the Set code.
-					FMemory::Memcpy(&DebugPtrStorage, &Other.DebugPtrStorage, sizeof(DebugPtrStorage));
+					FMemory::Memcpy(&DebugPtrStorage, &Other.DebugPtrStorage, sizeof(DebugPtrStorage)); //-V598
 					DebugPtrStorage.Ptr = NewPtr;
 				#endif
 			}
@@ -628,7 +628,7 @@ namespace UE4Function_Private
 				#if ENABLE_TFUNCTIONREF_VISUALIZATION
 					// Use Memcpy to copy the other DebugPtrStorage, including vptr (because we don't know the bound type
 					// here), and then reseat the underlying pointer.  Possibly even more evil than the Set code.
-					FMemory::Memcpy(&DebugPtrStorage, &Other.DebugPtrStorage, sizeof(DebugPtrStorage));
+					FMemory::Memcpy(&DebugPtrStorage, &Other.DebugPtrStorage, sizeof(DebugPtrStorage)); //-V598
 					DebugPtrStorage.Ptr = NewPtr;
 				#endif
 			}
