Use template specialization work-around on Mac as well as Android

diff --git a/Engine/Source/Runtime/Core/Public/Serialization/MemoryLayout.h b/Engine/Source/Runtime/Core/Public/Serialization/MemoryLayout.h
--- a/Engine/Source/Runtime/Core/Public/Serialization/MemoryLayout.h
+++ b/Engine/Source/Runtime/Core/Public/Serialization/MemoryLayout.h
@@ -55,7 +55,7 @@ using FMemoryImageUPtrInt = UPTRINT;
  * Partial specializations are accepted however.  So for these platforms, we force partial specialization via an extra dummy parameter that's otherwise ignored.
  * This actually breaks MSVC, as MSVC doesn't allow complex statements as template values during partial specialization, so need use full specialization there.
  */
-#if PLATFORM_ANDROID
+#if PLATFORM_ANDROID || PLATFORM_MAC
 #define UE_DECLARE_INTERNAL_LINK_BASE(T) template<int Counter, typename Dummy=void> struct T
 #define UE_DECLARE_INTERNAL_LINK_SPECIALIZATION(T, Counter) template<typename Dummy> struct T<Counter, Dummy>
 #else
