- Made sure that lz4.h works if CoreMinimal.h isn't included in any Unity files, but using the basic types (unsigned int) instead of types like uint32

diff --git a/Engine/Source/Runtime/Core/Public/Compression/lz4.h b/Engine/Source/Runtime/Core/Public/Compression/lz4.h
--- a/Engine/Source/Runtime/Core/Public/Compression/lz4.h
+++ b/Engine/Source/Runtime/Core/Public/Compression/lz4.h
@@ -557,7 +557,7 @@ LZ4LIB_STATIC_API void LZ4_attach_dictionary(LZ4_stream_t* workingStream, const
 #define LZ4_HASHTABLESIZE (1 << LZ4_MEMORY_USAGE)
 #define LZ4_HASH_SIZE_U32 (1 << LZ4_HASHLOG)       /* required as macro for static allocation */
 
-#if defined(__cplusplus) || (defined (__STDC_VERSION__) && (__STDC_VERSION__ >= 199901L) /* C99 */)
+#if 0// defined(__cplusplus) || (defined (__STDC_VERSION__) && (__STDC_VERSION__ >= 199901L) /* C99 */)
 
 typedef struct LZ4_stream_t_internal LZ4_stream_t_internal;
 struct LZ4_stream_t_internal {
