- Fixed NonUnity error
#rnx

diff --git a/Engine/Source/Runtime/Core/Public/Containers/SparseArray.h b/Engine/Source/Runtime/Core/Public/Containers/SparseArray.h
--- a/Engine/Source/Runtime/Core/Public/Containers/SparseArray.h
+++ b/Engine/Source/Runtime/Core/Public/Containers/SparseArray.h
@@ -17,6 +17,7 @@
 #include "Containers/BitArray.h"
 #include "Serialization/StructuredArchive.h"
 #include "Serialization/MemoryImageWriter.h"
+#include "Containers/UnrealString.h"
 
 
 #if UE_BUILD_SHIPPING || UE_BUILD_TEST
