Static analysis fix.

#rnx

diff --git a/Engine/Source/Runtime/Core/Public/Serialization/MemoryImage.h b/Engine/Source/Runtime/Core/Public/Serialization/MemoryImage.h
--- a/Engine/Source/Runtime/Core/Public/Serialization/MemoryImage.h
+++ b/Engine/Source/Runtime/Core/Public/Serialization/MemoryImage.h
@@ -867,7 +867,7 @@ namespace Freeze
 		{
 			const uint32 Index = Writer.GetPointerTable().AddIndexedPointer(TypeDesc, RawPtr);
 			check(Index != (uint32)INDEX_NONE);
-			const uint64 FrozenPackedIndex = (Index << 1llu) | 1u;
+			const uint64 FrozenPackedIndex = ((uint64)Index << 1u) | 1u;
 			Writer.WriteMemoryImagePointerSizedBytes(FrozenPackedIndex);
 		}
 		else
