[UE4-linux] UE4_DO_ROOT_PRIVILEGE_CHECK can be externally defined by build targets
+  This #define was previously internally set in source cpp in all cases.  This change allows build contexts to explicitly define the desired value to be used during preprocessing.

diff --git a/Engine/Source/Runtime/Core/Private/Unix/UnixPlatformMemory.cpp b/Engine/Source/Runtime/Core/Private/Unix/UnixPlatformMemory.cpp
--- a/Engine/Source/Runtime/Core/Private/Unix/UnixPlatformMemory.cpp
+++ b/Engine/Source/Runtime/Core/Private/Unix/UnixPlatformMemory.cpp
@@ -38,11 +38,13 @@
 #define UE4_PLATFORM_REDUCE_NUMBER_OF_MAPS					(0 && PLATFORM_UNIX && PLATFORM_64BITS)
 
 // do not do a root privilege check on non-x86-64 platforms (assume an embedded device)
-#if defined(_M_X64) || defined(__x86_64__) || defined (__amd64__) 
-	#define UE4_DO_ROOT_PRIVILEGE_CHECK	 1
-#else
-	#define UE4_DO_ROOT_PRIVILEGE_CHECK	 0
-#endif // defined(_M_X64) || defined(__x86_64__) || defined (__amd64__) 
+#ifndef UE4_DO_ROOT_PRIVILEGE_CHECK
+  #if defined(_M_X64) || defined(__x86_64__) || defined (__amd64__)
+    #define UE4_DO_ROOT_PRIVILEGE_CHECK	 1
+  #else
+    #define UE4_DO_ROOT_PRIVILEGE_CHECK	 0
+  #endif // defined(_M_X64) || defined(__x86_64__) || defined (__amd64__) 
+#endif // ifndef UE4_DO_ROOT_PRIVILEGE_CHECK
 
 // Set rather to use BinnedMalloc2 for binned malloc, can be overridden below
 #define USE_MALLOC_BINNED2 (1)
