Fixed PVS warnings.

#rnx

diff --git a/Engine/Source/Runtime/Core/Private/Compression/lz4.cpp b/Engine/Source/Runtime/Core/Private/Compression/lz4.cpp
--- a/Engine/Source/Runtime/Core/Private/Compression/lz4.cpp
+++ b/Engine/Source/Runtime/Core/Private/Compression/lz4.cpp
@@ -165,10 +165,10 @@
 #endif
 
 #ifndef likely
-#define likely(expr)     expect((expr) != 0, 1)
+#define likely(expr)     expect(int(expr) != 0, 1)
 #endif
 #ifndef unlikely
-#define unlikely(expr)   expect((expr) != 0, 0)
+#define unlikely(expr)   expect(int(expr) != 0, 0)
 #endif
 
 
