Slate: remove #if PLATFORM_ usage from PlatformTextField

diff --git a/Engine/Source/Runtime/Core/Public/HAL/PreprocessorHelpers.h b/Engine/Source/Runtime/Core/Public/HAL/PreprocessorHelpers.h
--- a/Engine/Source/Runtime/Core/Public/HAL/PreprocessorHelpers.h
+++ b/Engine/Source/Runtime/Core/Public/HAL/PreprocessorHelpers.h
@@ -43,3 +43,11 @@
 // Creates a string that can be used to include a header in the form "Platform/PlatformHeader.h", like "Windows/WindowsPlatformFile.h"
 #define COMPILED_PLATFORM_HEADER(Suffix) PREPROCESSOR_TO_STRING(PREPROCESSOR_JOIN(PLATFORM_HEADER_NAME/PLATFORM_HEADER_NAME, Suffix))
 #endif
+
+#if PLATFORM_IS_EXTENSION
+// Creates a string that can be used to include a header with the platform in its name, like "Pre/Fix/PlatformNameSuffix.h"
+#define COMPILED_PLATFORM_HEADER_WITH_PREFIX(Prefix, Suffix) PREPROCESSOR_TO_STRING(Prefix/PREPROCESSOR_JOIN(PLATFORM_HEADER_NAME, Suffix))
+#else
+// Creates a string that can be used to include a header with the platform in its name, like "Pre/Fix/PlatformName/PlatformNameSuffix.h"
+#define COMPILED_PLATFORM_HEADER_WITH_PREFIX(Prefix, Suffix) PREPROCESSOR_TO_STRING(Prefix/PLATFORM_HEADER_NAME/PREPROCESSOR_JOIN(PLATFORM_HEADER_NAME, Suffix))
+#endif
diff --git a/Engine/Source/Runtime/Slate/Public/Framework/Text/Android/AndroidPlatformTextField.h b/Engine/Source/Runtime/Slate/Public/Framework/Text/Android/AndroidPlatformTextField.h
--- a/Engine/Source/Runtime/Slate/Public/Framework/Text/Android/AndroidPlatformTextField.h
+++ b/Engine/Source/Runtime/Slate/Public/Framework/Text/Android/AndroidPlatformTextField.h
@@ -2,6 +2,8 @@
 
 #pragma once
 
+#if USE_ANDROID_JNI
+
 #include "CoreMinimal.h"
 #include "Framework/Application/IPlatformTextField.h"
 
@@ -20,3 +22,8 @@ private:
 
 typedef FAndroidPlatformTextField FPlatformTextField;
 
+#else
+
+#include "Framework/Text/GenericPlatformTextField.h"
+
+#endif
diff --git a/Engine/Source/Runtime/Slate/Public/Framework/Text/Linux/LinuxPlatformTextField.h b/Engine/Source/Runtime/Slate/Public/Framework/Text/Linux/LinuxPlatformTextField.h
new file mode 100644
--- /dev/null
+++ b/Engine/Source/Runtime/Slate/Public/Framework/Text/Linux/LinuxPlatformTextField.h
@@ -0,0 +1,5 @@
+// Copyright 1998-2019 Epic Games, Inc. All Rights Reserved.
+
+#pragma once
+
+#include "Framework/Text/GenericPlatformTextField.h"
diff --git a/Engine/Source/Runtime/Slate/Public/Framework/Text/LinuxAArch64/LinuxAArch64PlatformTextField.h b/Engine/Source/Runtime/Slate/Public/Framework/Text/LinuxAArch64/LinuxAArch64PlatformTextField.h
new file mode 100644
--- /dev/null
+++ b/Engine/Source/Runtime/Slate/Public/Framework/Text/LinuxAArch64/LinuxAArch64PlatformTextField.h
@@ -0,0 +1,5 @@
+// Copyright 1998-2019 Epic Games, Inc. All Rights Reserved.
+
+#pragma once
+
+#include "Framework/Text/GenericPlatformTextField.h"
diff --git a/Engine/Source/Runtime/Slate/Public/Framework/Text/Mac/MacPlatformTextField.h b/Engine/Source/Runtime/Slate/Public/Framework/Text/Mac/MacPlatformTextField.h
new file mode 100644
--- /dev/null
+++ b/Engine/Source/Runtime/Slate/Public/Framework/Text/Mac/MacPlatformTextField.h
@@ -0,0 +1,5 @@
+// Copyright 1998-2019 Epic Games, Inc. All Rights Reserved.
+
+#pragma once
+
+#include "Framework/Text/GenericPlatformTextField.h"
diff --git a/Engine/Source/Runtime/Slate/Public/Framework/Text/PlatformTextField.h b/Engine/Source/Runtime/Slate/Public/Framework/Text/PlatformTextField.h
--- a/Engine/Source/Runtime/Slate/Public/Framework/Text/PlatformTextField.h
+++ b/Engine/Source/Runtime/Slate/Public/Framework/Text/PlatformTextField.h
@@ -4,18 +4,4 @@
 
 #include "CoreMinimal.h"
 
-#if PLATFORM_IOS
-	#include "IOS/IOSPlatformTextField.h"
-#elif PLATFORM_ANDROID && USE_ANDROID_JNI
-	#include "Android/AndroidPlatformTextField.h"
-#elif PLATFORM_PS4
-	#include "Framework/Text/PS4PlatformTextField.h"
-#elif PLATFORM_XBOXONE
-	#include "XboxOne/XboxOnePlatformTextField.h"
-#elif PLATFORM_SWITCH
-	#include "Switch/SwitchPlatformTextField.h"
-#elif PLATFORM_QUAIL
-	#include "Quail/QuailPlatformTextField.h"
-#else
-	#include "Framework/Text/GenericPlatformTextField.h"
-#endif
+#include COMPILED_PLATFORM_HEADER_WITH_PREFIX(Framework/Text, PlatformTextField.h)
diff --git a/Engine/Source/Runtime/Slate/Public/Framework/Text/Windows/WindowsPlatformTextField.h b/Engine/Source/Runtime/Slate/Public/Framework/Text/Windows/WindowsPlatformTextField.h
new file mode 100644
--- /dev/null
+++ b/Engine/Source/Runtime/Slate/Public/Framework/Text/Windows/WindowsPlatformTextField.h
@@ -0,0 +1,5 @@
+// Copyright 1998-2019 Epic Games, Inc. All Rights Reserved.
+
+#pragma once
+
+#include "Framework/Text/GenericPlatformTextField.h"
