UE-80921 HoloLens: "error C2078: too many initializers" when packaging
-UnrealMathNeon.h initialization fix for HoloLens, whick builds for arm64 with msvc

diff --git a/Engine/Source/Runtime/Core/Public/Math/UnrealMathNeon.h b/Engine/Source/Runtime/Core/Public/Math/UnrealMathNeon.h
--- a/Engine/Source/Runtime/Core/Public/Math/UnrealMathNeon.h
+++ b/Engine/Source/Runtime/Core/Public/Math/UnrealMathNeon.h
@@ -1343,10 +1343,6 @@ FORCEINLINE void VectorSinCos(  VectorRegister* VSinAngles, VectorRegister* VCos
 // Returns true if the vector contains a component that is either NAN or +/-infinite.
 inline bool VectorContainsNaNOrInfinite(const VectorRegister& Vec)
 {
-#if PLATFORM_HOLOLENS
-	//TODO : The implementation below does not compile on Hololens Arm64.Temporarily working around this problem.
-	return false;
-#else
 	// https://en.wikipedia.org/wiki/IEEE_754-1985
 	// Infinity is represented with all exponent bits set, with the correct sign bit.
 	// NaN is represented with all exponent bits set, plus at least one fraction/significant bit set.
@@ -1359,12 +1355,33 @@ inline bool VectorContainsNaNOrInfinite(const VectorRegister& Vec)
 
 	// Mask off Exponent
 	VectorRegister ExpTest = VectorBitwiseAnd(Vec, FloatInfinity);
+
 	// Compare to full exponent & combine resulting flags into lane 0
-	static const int8x16_t Table = {0,4,8,12, 0,0,0,0, 0,0,0,0, 0,0,0,0};
+
+// This is the supported neon implementation, but it generates a lot more assembly instructions than the following compiler-specific implementations
+//#if PLATFORM_LITTLE_ENDIAN
+//	static const int8x8_t High = vcreate_s8(0x000000000C080400ULL);
+//#else
+//	static const int8x8_t High = vcreate_s8(0x0004080C00000000ULL);
+//#endif
+//	static const int8x8_t Low = vcreate_s8(0x0ULL);
+//	static const int8x16_t Table = vcombine_s8(High, Low);
+
+#ifdef _MSC_VER
+	// msvc can only initialize using the first union type which is: n128_u64[2];
+#  if PLATFORM_LITTLE_ENDIAN
+	static const int8x16_t Table = { 0x000000000C080400ULL, 0ULL };
+#  else
+	static const int8x16_t Table = { 0x0004080C00000000ULL, 0ULL };
+#  endif
+#else
+	// clang can initialize with this syntax, but not the msvc one
+	static const int8x16_t Table = { 0,4,8,12, 0,0,0,0, 0,0,0,0, 0,0,0,0 };
+#endif
+
 	uint8x16_t res = (uint8x16_t)VectorCompareEQ(ExpTest, FloatInfinity);
 	// If we have all zeros, all elements are finite
 	return vgetq_lane_u32((uint32x4_t)vqtbx1q_u8(res, res, Table), 0) != 0;
-#endif
 }
 
 //TODO: Vectorize
