Merging //UE4/Main@9566668 to Dev-RenderPlat-Staging
#rnx

diff --git a/Engine/Build/Build.version b/Engine/Build/Build.version
--- a/Engine/Build/Build.version
+++ b/Engine/Build/Build.version
@@ -1,6 +1,6 @@
 {
 	"MajorVersion": 4,
-	"MinorVersion": 24,
+	"MinorVersion": 25,
 	"PatchVersion": 0,
 	"Changelist": 0,
 	"CompatibleChangelist": 0,
diff --git a/Engine/Build/InstalledEngineBuild.xml b/Engine/Build/InstalledEngineBuild.xml
--- a/Engine/Build/InstalledEngineBuild.xml
+++ b/Engine/Build/InstalledEngineBuild.xml
@@ -35,7 +35,6 @@
 	<Option Name="WithTVOS" Restrict="true|false" DefaultValue="$(DefaultWithIOS)" Description="Include the tvOS target platform"/>
 	<Option Name="WithLinux" Restrict="true|false" DefaultValue="$(DefaultWithLinux)" Description="Include the Linux target platform"/>
 	<Option Name="WithLinuxAArch64" Restrict="true|false" DefaultValue="false" Description="Include the Linux AArch64 target platform"/>
-	<Option Name="WithHTML5" Restrict="true|false" DefaultValue="$(DefaultWithPlatform)" Description="Include the HTML5 target platform"/>
 	<Option Name="WithPS4" Restrict="true|false" DefaultValue="false" Description="Include the PS4 target platform"/>
 	<Option Name="WithXboxOne" Restrict="true|false" DefaultValue="false" Description="Include the XboxOne target platform"/>
 	<Option Name="WithSwitch" Restrict="true|false" DefaultValue="false" Description="Include the Switch target platform"/>
@@ -325,16 +324,6 @@
 
 		</Node>
 	</Agent>
-	<Agent Name="Target HTML5" Type="CompileWin64;Win64">
-		<Node Name="Compile UE4Game HTML5" Requires="Compile UnrealHeaderTool $(HostPlatform)" Produces="#UE4Game HTML5">
-
-			<ForEach Name="Configuration" Values="$(GameConfigurations)">
-				<Compile Target="UE4Game" Platform="HTML5" Configuration="$(Configuration)" Tag="#UE4Game HTML5" Arguments="-precompile -allmodules -nolink $(TargetDebugInfoArg)"/>
-				<Compile Target="UE4Game" Platform="HTML5" Configuration="$(Configuration)" Tag="#UE4Game HTML5" Arguments="-precompile $(TargetDebugInfoArg)" Clean="false"/>
-			</ForEach>
-
-		</Node>
-	</Agent>
 	<Agent Name="Target PS4" Type="CompileWin64;Win64">
 		<Node Name="Compile UE4Game PS4" Requires="Compile UnrealHeaderTool Win64" Produces="#UE4Game PS4">
 
@@ -529,7 +518,6 @@
 			<CsCompile Project="$(CsToolsDir)/Engine/Source/Programs/IOS/MobileDeviceInterface/MobileDeviceInterface.csproj" Configuration="Development" Platform="AnyCPU" Arguments="/verbosity:minimal /target:Rebuild" Tag="#Build Tools CS Binaries;#Sign Binaries" TagReferences="#Build Tools CS Binaries"/>
 			<CsCompile Project="$(CsToolsDir)/Engine/Source/Programs/IOS/DeploymentInterface/DeploymentInterface.csproj" Configuration="Development" Platform="AnyCPU" Arguments="/verbosity:minimal /target:Rebuild" Tag="#Build Tools CS Binaries;#Sign Binaries" TagReferences="#Build Tools CS Binaries"/>
 			<CsCompile Project="$(CsToolsDir)/Engine/Source/Programs/IOS/DeploymentServerLauncher/DeploymentServerLauncher.csproj" Configuration="Development" Platform="AnyCPU" Arguments="/verbosity:minimal /target:Rebuild" Tag="#Build Tools CS Binaries;#Sign Binaries" TagReferences="#Build Tools CS Binaries"/>
-			<CsCompile Project="$(RootDir)/Engine/Platforms/HTML5/Source/Programs/HTML5/HTML5LaunchHelper/HTML5LaunchHelper.csproj" Configuration="Development" Platform="AnyCPU" Arguments="/verbosity:minimal /target:Rebuild" Tag="#Build Tools CS Binaries;#Sign Binaries" TagReferences="#Build Tools CS Binaries"/>
 
 			<!-- Sign the binaries -->
 			<Sign Files="#Sign Binaries" If="'$(SignExecutables)' == true"/>
@@ -589,7 +577,6 @@
 		<Property Name="DDCPlatformsWin64" Value="$(DDCPlatformsWin64)+TVOS" If="'$(WithTVOS)' == true"/>
 		<Property Name="DDCPlatformsWin64" Value="$(DDCPlatformsWin64)+LinuxNoEditor" If="'$(WithLinux)' == true"/>
 		<Property Name="DDCPlatformsWin64" Value="$(DDCPlatformsWin64)+LinuxAArch64NoEditor" If="'$(WithLinuxAArch64)' == true"/>
-		<Property Name="DDCPlatformsWin64" Value="$(DDCPlatformsWin64)+HTML5" If="'$(WithHTML5)' == true"/>
 		<Property Name="DDCPlatformsWin64" Value="$(DDCPlatformsWin64)+PS4" If="'$(WithPS4)' == true"/>
 		<Property Name="DDCPlatformsWin64" Value="$(DDCPlatformsWin64)+XboxOne" If="'$(WithXboxOne)' == true"/>
 		<Property Name="DDCPlatformsWin64" Value="$(DDCPlatformsWin64)+Switch" If="'$(WithSwitch)' == true"/>
@@ -641,7 +628,6 @@
 		<Property Name="DDCPlatformsMac" Value="$(DDCPlatformsMac)+Android_ATC" If="'$(WithAndroid)' == true"/>
 		<Property Name="DDCPlatformsMac" Value="$(DDCPlatformsMac)+IOS" If="'$(WithIOS)' == true"/>
 		<Property Name="DDCPlatformsMac" Value="$(DDCPlatformsMac)+TVOS" If="'$(WithTVOS)' == true"/>
-		<Property Name="DDCPlatformsMac" Value="$(DDCPlatformsMac)+HTML5" If="'$(WithHTML5)' == true"/>
 		<Property Name="DDCPlatformsMac" Value="$(DDCPlatformsMac)+Lumin" If="'$(WithLumin)' == true"/>
 	</Do>
 
@@ -747,7 +733,6 @@
 		<Property Name="InstalledRequirements" Value="$(InstalledRequirements);#UE4Game TVOS;#UE4Game TVOS Unstripped;#UE4Game TVOS Stripped" If="'$(WithTVOS)' == true"/>
 		<Property Name="InstalledRequirements" Value="$(InstalledRequirements);#UE4Game Linux;#UE4Game Linux Unstripped;#UE4Game Linux Stripped;#Build Tools Linux" If="'$(WithLinux)' == true"/>
 		<Property Name="InstalledRequirements" Value="$(InstalledRequirements);#UE4Game LinuxAArch64;#UE4Game LinuxAArch64 Unstripped;#UE4Game LinuxAArch64 Stripped" If="'$(WithLinuxAArch64)' == true"/>
-		<Property Name="InstalledRequirements" Value="$(InstalledRequirements);#UE4Game HTML5" If="'$(WithHTML5)' == true"/>
 		<Property Name="InstalledRequirements" Value="$(InstalledRequirements);#UE4Game PS4" If="'$(WithPS4)' == true"/>
 		<Property Name="InstalledRequirements" Value="$(InstalledRequirements);#UE4Game XboxOne" If="'$(WithXboxOne)' == true"/>
 		<Property Name="InstalledRequirements" Value="$(InstalledRequirements);#UE4Game Switch" If="'$(WithSwitch)' == true"/>
@@ -816,11 +801,6 @@
 				<Property Name="CopyInstalledFilter" Value="$(CopyInstalledFilter);$(CopyLinuxFilter)"/>
 				<Property Name="CopyInstalledExceptions" Value="$(CopyInstalledExceptions);$(CopyLinuxExceptions)"/>
 			</Do>
-			<Do If="'$(WithHTML5)' == true">
-				<Tag Files="#UE4Game HTML5" With="#Installed Win64"/>
-				<Property Name="CopyInstalledFilter" Value="$(CopyInstalledFilter);$(CopyHTML5Filter)"/>
-				<Property Name="CopyInstalledExceptions" Value="$(CopyInstalledExceptions);$(CopyHTML5Exceptions)"/>
-			</Do>
 			<Do If="'$(WithPS4)' == true">
 				<Tag Files="#UE4Game PS4" With="#Installed Win64"/>
 				<Property Name="CopyInstalledFilter" Value="$(CopyInstalledFilter);$(CopyPS4Filter)"/>
@@ -919,7 +899,6 @@
 		<Property Name="FilterRequirements" Value="$(FilterRequirements);#UE4Game Android;#UE4Game Android Unstripped;#UE4Game Android Stripped" If="'$(WithAndroid)' == true"/>
 		<Property Name="FilterRequirements" Value="$(FilterRequirements);#UE4Game IOS;#UE4Game IOS Unstripped;#UE4Game IOS Stripped" If="'$(WithIOS)' == true"/>
 		<Property Name="FilterRequirements" Value="$(FilterRequirements);#UE4Game TVOS;#UE4Game TVOS Unstripped;#UE4Game TVOS Stripped" If="'$(WithTVOS)' == true"/>
-		<Property Name="FilterRequirements" Value="$(FilterRequirements);#UE4Game HTML5" If="'$(WithHTML5)' == true"/>
 		<Property Name="FilterRequirements" Value="$(FilterRequirements);#UE4Game LuminMac;#UE4Game LuminMac Unstripped;#UE4Game LuminMac Stripped" If="'$(WithLumin)' == true"/> 
 		<Property Name="FilterRequirements" Value="$(FilterRequirements);#CompressedDDCMac" If="'$(WithDDC)' == true"/>
 
@@ -966,11 +945,6 @@
 				<Property Name="CopyInstalledFilter" Value="$(CopyInstalledFilter);$(CopyTVOSFilter)"/>
 				<Property Name="CopyInstalledExceptions" Value="$(CopyInstalledExceptions);$(CopyTVOSExceptions)"/>
 			</Do>
-			<Do If="'$(WithHTML5)' == true">
-				<Tag Files="#UE4Game HTML5" With="#Installed Mac"/>
-				<Property Name="CopyInstalledFilter" Value="$(CopyInstalledFilter);$(CopyHTML5Filter)"/>
-				<Property Name="CopyInstalledExceptions" Value="$(CopyInstalledExceptions);$(CopyHTML5Exceptions)"/>
-			</Do>
 			<Do If="'$(WithLumin)' == true">
 				<Tag Files="#UE4Game LuminMac" Except="#UE4Game LuminMac Unstripped" With="#Installed Mac"/>
 				<Tag Files="#UE4Game LuminMac Stripped" With="#Saved Output"/>
diff --git a/Engine/Build/InstalledEngineFilters.xml b/Engine/Build/InstalledEngineFilters.xml
--- a/Engine/Build/InstalledEngineFilters.xml
+++ b/Engine/Build/InstalledEngineFilters.xml
@@ -470,22 +470,6 @@
 <Property Name="CopyLinuxFilter" Value="$(CopyLinuxFilter);Engine/Extras/ThirdPartyNotUE/putty/..."/>	
 
 <!-- Win64 specific -->
-	<!-- HTML5 -->
-	<Property Name="CopyHTML5Filter">
-		<!-- Build Files -->
-		Engine/Build/HTML5/...
-
-		<!-- SDK -->
-		Engine/Extras/ThirdPartyNotUE/emsdk/emscripten/...
-	</Property>
-	<Property Name="CopyHTML5Filter" Value="$(CopyHTML5Filter)" If="'$(HostPlatform)' == 'Win64'">
-		Engine/Extras/ThirdPartyNotUE/emsdk/Win64/...
-	</Property>
-	<Property Name="CopyHTML5Filter" Value="$(CopyHTML5Filter)" If="'$(HostPlatform)' == 'Mac'">
-		Engine/Extras/ThirdPartyNotUE/emsdk/Mac/...
-	</Property>
-		<Property Name="CopyHTML5Exceptions">
-	</Property>
 
 	<!-- PS4 -->
 	<Property Name="CopyPS4Filter">
diff --git a/Engine/Documentation/Source/Shared/Editor/MapErrors/MapErrors.INT.udn b/Engine/Documentation/Source/Shared/Editor/MapErrors/MapErrors.INT.udn
--- a/Engine/Documentation/Source/Shared/Editor/MapErrors/MapErrors.INT.udn
+++ b/Engine/Documentation/Source/Shared/Editor/MapErrors/MapErrors.INT.udn
@@ -2,7 +2,11 @@ Availability:Public
 Title:Map Check Errors
 Crumbs:%ROOT%, Engine, Engine/UI/LevelEditor
 Description:Descriptions of errors generated during map builds or running a map check in Unreal Editor.
-Version: 4.9
+Version: 4.22
+Type:reference
+Parent:Shared/Editor
+Tags:map
+Tags:error
 
 [TOC(start:2 end:3)]
 
diff --git a/Engine/Plugins/Developer/Concert/ConcertMain/Source/Concert/Private/ConcertVersion.cpp b/Engine/Plugins/Developer/Concert/ConcertMain/Source/Concert/Private/ConcertVersion.cpp
--- a/Engine/Plugins/Developer/Concert/ConcertMain/Source/Concert/Private/ConcertVersion.cpp
+++ b/Engine/Plugins/Developer/Concert/ConcertMain/Source/Concert/Private/ConcertVersion.cpp
@@ -117,7 +117,7 @@ void FConcertSessionVersionInfo::Initialize()
 		}
 	}
 
-	for (const FCustomVersion& EngineCustomVersion : FCustomVersionContainer::GetRegistered().GetAllVersions())
+	for (const FCustomVersion& EngineCustomVersion : FCurrentCustomVersions::GetAll().GetAllVersions())
 	{
 		FConcertCustomVersionInfo& CustomVersion = CustomVersions.AddDefaulted_GetRef();
 		CustomVersion.Initialize(EngineCustomVersion);
diff --git a/Engine/Plugins/FX/Niagara/Source/NiagaraEditor/Private/ViewModels/NiagaraSystemViewModel.cpp b/Engine/Plugins/FX/Niagara/Source/NiagaraEditor/Private/ViewModels/NiagaraSystemViewModel.cpp
--- a/Engine/Plugins/FX/Niagara/Source/NiagaraEditor/Private/ViewModels/NiagaraSystemViewModel.cpp
+++ b/Engine/Plugins/FX/Niagara/Source/NiagaraEditor/Private/ViewModels/NiagaraSystemViewModel.cpp
@@ -999,7 +999,19 @@ void FNiagaraSystemViewModel::ResetSystem(ETimeResetMode TimeResetMode, EMultiRe
 		}
 	}
 
-	TArray<UNiagaraComponent*> ReferencingComponents = FNiagaraEditorUtilities::GetComponentsThatReferenceSystem(GetSystem());
+	TArray<UNiagaraComponent*> ReferencingComponents;
+	if (MultiResetMode == EMultiResetMode::ResetThisInstance)
+	{
+		if (PreviewComponent != nullptr)
+		{
+			ReferencingComponents.Add(PreviewComponent);
+		}
+	}
+	else
+	{
+		ReferencingComponents = FNiagaraEditorUtilities::GetComponentsThatReferenceSystem(GetSystem());
+	}
+
 	for (auto Component : ReferencingComponents)
 	{
 		if (ReinitMode == EReinitMode::ResetSystem)
diff --git a/Engine/Plugins/FX/Niagara/Source/NiagaraEditor/Private/ViewModels/Stack/NiagaraStackRenderItemGroup.cpp b/Engine/Plugins/FX/Niagara/Source/NiagaraEditor/Private/ViewModels/Stack/NiagaraStackRenderItemGroup.cpp
--- a/Engine/Plugins/FX/Niagara/Source/NiagaraEditor/Private/ViewModels/Stack/NiagaraStackRenderItemGroup.cpp
+++ b/Engine/Plugins/FX/Niagara/Source/NiagaraEditor/Private/ViewModels/Stack/NiagaraStackRenderItemGroup.cpp
@@ -101,6 +101,8 @@ public:
 			}
 		}
 
+		FNiagaraSystemUpdateContext SystemUpdate(Emitter, true);
+
 		if (bVarsAdded)
 		{
 			FNotificationInfo Info(LOCTEXT("AddedVariables", "One or more variables have been added to the Spawn script to support the added renderer."));
diff --git a/Engine/Plugins/Online/Android/OnlineSubsystemGameCircle/Source/OnlineSubsystemGameCircle/Private/OnlineFriendsInterfaceGameCircle.cpp b/Engine/Plugins/Online/Android/OnlineSubsystemGameCircle/Source/OnlineSubsystemGameCircle/Private/OnlineFriendsInterfaceGameCircle.cpp
--- a/Engine/Plugins/Online/Android/OnlineSubsystemGameCircle/Source/OnlineSubsystemGameCircle/Private/OnlineFriendsInterfaceGameCircle.cpp
+++ b/Engine/Plugins/Online/Android/OnlineSubsystemGameCircle/Source/OnlineSubsystemGameCircle/Private/OnlineFriendsInterfaceGameCircle.cpp
@@ -65,6 +65,15 @@ void FOnlineFriendsInterfaceGameCircle::SetFriendAlias(int32 LocalUserNum, const
 	});
 }
 
+void FOnlineFriendsInterfaceGameCircle::DeleteFriendAlias(int32 LocalUserNum, const FUniqueNetId& FriendId, const FString& ListName, const FOnDeleteFriendAliasComplete& Delegate)
+{
+	TSharedRef<const FUniqueNetId> FriendIdRef = FriendId.AsShared();
+	GameCircleSubsystem->ExecuteNextTick([LocalUserNum, FriendIdRef, ListName, Delegate]()
+	{
+		Delegate.ExecuteIfBound(LocalUserNum, *FriendIdRef, ListName, FOnlineError(EOnlineErrorResult::NotImplemented));
+	});
+}
+
 bool FOnlineFriendsInterfaceGameCircle::DeleteFriend(int32 LocalUserNum, const FUniqueNetId& FriendId, const FString& ListName)
 {
 	return false;
diff --git a/Engine/Plugins/Online/Android/OnlineSubsystemGameCircle/Source/OnlineSubsystemGameCircle/Private/OnlineFriendsInterfaceGameCircle.h b/Engine/Plugins/Online/Android/OnlineSubsystemGameCircle/Source/OnlineSubsystemGameCircle/Private/OnlineFriendsInterfaceGameCircle.h
--- a/Engine/Plugins/Online/Android/OnlineSubsystemGameCircle/Source/OnlineSubsystemGameCircle/Private/OnlineFriendsInterfaceGameCircle.h
+++ b/Engine/Plugins/Online/Android/OnlineSubsystemGameCircle/Source/OnlineSubsystemGameCircle/Private/OnlineFriendsInterfaceGameCircle.h
@@ -20,6 +20,7 @@ public:
 	virtual bool AcceptInvite(int32 LocalUserNum, const FUniqueNetId& FriendId, const FString& ListName, const FOnAcceptInviteComplete& Delegate = FOnAcceptInviteComplete()) override;
  	virtual bool RejectInvite(int32 LocalUserNum, const FUniqueNetId& FriendId, const FString& ListName) override;
 	virtual void SetFriendAlias(int32 LocalUserNum, const FUniqueNetId& FriendId, const FString& ListName, const FString& Alias, const FOnSetFriendAliasComplete& Delegate = FOnSetFriendAliasComplete()) override;
+	virtual void DeleteFriendAlias(int32 LocalUserNum, const FUniqueNetId& FriendId, const FString& ListName, const FOnDeleteFriendAliasComplete& Delegate = FOnDeleteFriendAliasComplete()) override;
 	virtual bool DeleteFriend(int32 LocalUserNum, const FUniqueNetId& FriendId, const FString& ListName) override;
 	virtual bool GetFriendsList(int32 LocalUserNum, const FString& ListName, TArray< TSharedRef<FOnlineFriend> >& OutFriends) override;
 	virtual TSharedPtr<FOnlineFriend> GetFriend(int32 LocalUserNum, const FUniqueNetId& FriendId, const FString& ListName) override;
diff --git a/Engine/Plugins/Online/IOS/OnlineSubsystemIOS/Source/Private/OnlineFriendsInterfaceIOS.cpp b/Engine/Plugins/Online/IOS/OnlineSubsystemIOS/Source/Private/OnlineFriendsInterfaceIOS.cpp
--- a/Engine/Plugins/Online/IOS/OnlineSubsystemIOS/Source/Private/OnlineFriendsInterfaceIOS.cpp
+++ b/Engine/Plugins/Online/IOS/OnlineSubsystemIOS/Source/Private/OnlineFriendsInterfaceIOS.cpp
@@ -241,6 +241,16 @@ void FOnlineFriendsIOS::SetFriendAlias(int32 LocalUserNum, const FUniqueNetId& F
 	});
 }
 
+void FOnlineFriendsIOS::DeleteFriendAlias(int32 LocalUserNum, const FUniqueNetId& FriendId, const FString& ListName, const FOnDeleteFriendAliasComplete& Delegate)
+{
+	TSharedRef<const FUniqueNetId> FriendIdRef = FriendId.AsShared();
+	IOSSubsystem->ExecuteNextTick([LocalUserNum, FriendIdRef, ListName, Delegate]()
+	{
+		UE_LOG_ONLINE_FRIEND(Warning, TEXT("FOnlineFriendsIOS::DeleteFriendAlias is currently not supported"));
+		Delegate.ExecuteIfBound(LocalUserNum, *FriendIdRef, ListName, FOnlineError(EOnlineErrorResult::NotImplemented));
+	});
+}
+
 bool FOnlineFriendsIOS::DeleteFriend(int32 LocalUserNum, const FUniqueNetId& FriendId, const FString& ListName)
 {
 	TriggerOnDeleteFriendCompleteDelegates(LocalUserNum, false, FriendId, ListName, FString(TEXT("DeleteFriend() is not supported")));
diff --git a/Engine/Plugins/Online/IOS/OnlineSubsystemIOS/Source/Private/OnlineFriendsInterfaceIOS.h b/Engine/Plugins/Online/IOS/OnlineSubsystemIOS/Source/Private/OnlineFriendsInterfaceIOS.h
--- a/Engine/Plugins/Online/IOS/OnlineSubsystemIOS/Source/Private/OnlineFriendsInterfaceIOS.h
+++ b/Engine/Plugins/Online/IOS/OnlineSubsystemIOS/Source/Private/OnlineFriendsInterfaceIOS.h
@@ -97,7 +97,8 @@ public:
 	virtual bool AcceptInvite(int32 LocalUserNum, const FUniqueNetId& FriendId, const FString& ListName, const FOnAcceptInviteComplete& Delegate = FOnAcceptInviteComplete()) override;
  	virtual bool RejectInvite(int32 LocalUserNum, const FUniqueNetId& FriendId, const FString& ListName) override;
 	virtual void SetFriendAlias(int32 LocalUserNum, const FUniqueNetId& FriendId, const FString& ListName, const FString& Alias, const FOnSetFriendAliasComplete& Delegate = FOnSetFriendAliasComplete()) override;
- 	virtual bool DeleteFriend(int32 LocalUserNum, const FUniqueNetId& FriendId, const FString& ListName) override;
+	virtual void DeleteFriendAlias(int32 LocalUserNum, const FUniqueNetId& FriendId, const FString& ListName, const FOnDeleteFriendAliasComplete& Delegate = FOnDeleteFriendAliasComplete()) override; 
+	virtual bool DeleteFriend(int32 LocalUserNum, const FUniqueNetId& FriendId, const FString& ListName) override;
 	virtual bool GetFriendsList(int32 LocalUserNum, const FString& ListName, TArray< TSharedRef<FOnlineFriend> >& OutFriends) override;
 	virtual TSharedPtr<FOnlineFriend> GetFriend(int32 LocalUserNum, const FUniqueNetId& FriendId, const FString& ListName) override;
 	virtual bool IsFriend(int32 LocalUserNum, const FUniqueNetId& FriendId, const FString& ListName) override;
diff --git a/Engine/Plugins/Online/OnlineSubsystem/Source/Public/Interfaces/OnlineFriendsInterface.h b/Engine/Plugins/Online/OnlineSubsystem/Source/Public/Interfaces/OnlineFriendsInterface.h
--- a/Engine/Plugins/Online/OnlineSubsystem/Source/Public/Interfaces/OnlineFriendsInterface.h
+++ b/Engine/Plugins/Online/OnlineSubsystem/Source/Public/Interfaces/OnlineFriendsInterface.h
@@ -175,6 +175,16 @@ typedef FOnRejectInviteComplete::FDelegate FOnRejectInviteCompleteDelegate;
  */
 DECLARE_DELEGATE_FourParams(FOnSetFriendAliasComplete, int32 /*LocalUserNum*/, const FUniqueNetId& /*FriendId*/, const FString& /*ListName*/, const FOnlineError& /*Error*/);
 
+/**
+ * Delegate used when a delete friend alias request has completed
+ *
+ * @param LocalUserNum the controller number of the associated user that made the request
+ * @param FriendId player that we delete alias from
+ * @param ListName name of the friends list that was operated on
+ * @param Error error information on failure
+ */
+DECLARE_DELEGATE_FourParams(FOnDeleteFriendAliasComplete, int32 /*LocalUserNum*/, const FUniqueNetId& /*FriendId*/, const FString& /*ListName*/, const FOnlineError& /*Error*/);
+
 /**
  * Delegate used when an friend delete request has completed
  *
@@ -457,7 +467,17 @@ public:
 	 *
 	 */
 	virtual void SetFriendAlias(int32 LocalUserNum, const FUniqueNetId& FriendId, const FString& ListName, const FString& Alias, const FOnSetFriendAliasComplete& Delegate = FOnSetFriendAliasComplete()) = 0;
-	
+
+	/**
+	 * Starts an async task that delete an alias from a friend.
+	 *
+	 * @param LocalUserNum the user that is setting the alias
+	 * @param FriendId player that will have the alias deleted from
+	 * @param ListName name of the friends list to operate on
+	 *
+	 */
+	virtual void DeleteFriendAlias(int32 LocalUserNum, const FUniqueNetId& FriendId, const FString& ListName, const FOnDeleteFriendAliasComplete& Delegate = FOnDeleteFriendAliasComplete()) = 0;
+
 	/**
 	 * Delegate used when an invite reject request has completed
 	 *
diff --git a/Engine/Plugins/Online/OnlineSubsystemFacebook/Source/Private/OnlineFriendsFacebookCommon.cpp b/Engine/Plugins/Online/OnlineSubsystemFacebook/Source/Private/OnlineFriendsFacebookCommon.cpp
--- a/Engine/Plugins/Online/OnlineSubsystemFacebook/Source/Private/OnlineFriendsFacebookCommon.cpp
+++ b/Engine/Plugins/Online/OnlineSubsystemFacebook/Source/Private/OnlineFriendsFacebookCommon.cpp
@@ -199,6 +199,16 @@ void FOnlineFriendsFacebookCommon::SetFriendAlias(int32 LocalUserNum, const FUni
 	});
 }
 
+void FOnlineFriendsFacebookCommon::DeleteFriendAlias(int32 LocalUserNum, const FUniqueNetId& FriendId, const FString& ListName, const FOnDeleteFriendAliasComplete& Delegate)
+{
+	TSharedRef<const FUniqueNetId> FriendIdRef = FriendId.AsShared();
+	FacebookSubsystem->ExecuteNextTick([LocalUserNum, FriendIdRef, ListName, Delegate]()
+	{
+		UE_LOG_ONLINE_FRIEND(Warning, TEXT("FOnlineFriendsFacebookCommon::DeleteFriendAlias is not supported"));
+		Delegate.ExecuteIfBound(LocalUserNum, *FriendIdRef, ListName, FOnlineError(EOnlineErrorResult::NotImplemented));
+	});
+}
+
 bool FOnlineFriendsFacebookCommon::DeleteFriend(int32 LocalUserNum, const FUniqueNetId& FriendId, const FString& ListName)
 {
 	TriggerOnDeleteFriendCompleteDelegates(LocalUserNum, false, FriendId, ListName, FString(TEXT("DeleteFriend() is not supported")));
diff --git a/Engine/Plugins/Online/OnlineSubsystemFacebook/Source/Private/OnlineFriendsFacebookCommon.h b/Engine/Plugins/Online/OnlineSubsystemFacebook/Source/Private/OnlineFriendsFacebookCommon.h
--- a/Engine/Plugins/Online/OnlineSubsystemFacebook/Source/Private/OnlineFriendsFacebookCommon.h
+++ b/Engine/Plugins/Online/OnlineSubsystemFacebook/Source/Private/OnlineFriendsFacebookCommon.h
@@ -124,7 +124,8 @@ public:
 	virtual bool AcceptInvite(int32 LocalUserNum, const FUniqueNetId& FriendId, const FString& ListName, const FOnAcceptInviteComplete& Delegate = FOnAcceptInviteComplete()) override;
  	virtual bool RejectInvite(int32 LocalUserNum, const FUniqueNetId& FriendId, const FString& ListName) override;
 	virtual void SetFriendAlias(int32 LocalUserNum, const FUniqueNetId& FriendId, const FString& ListName, const FString& Alias, const FOnSetFriendAliasComplete& Delegate = FOnSetFriendAliasComplete()) override;
- 	virtual bool DeleteFriend(int32 LocalUserNum, const FUniqueNetId& FriendId, const FString& ListName) override;
+	virtual void DeleteFriendAlias(int32 LocalUserNum, const FUniqueNetId& FriendId, const FString& ListName, const FOnDeleteFriendAliasComplete& Delegate = FOnDeleteFriendAliasComplete()) override;
+	virtual bool DeleteFriend(int32 LocalUserNum, const FUniqueNetId& FriendId, const FString& ListName) override;
 	virtual bool GetFriendsList(int32 LocalUserNum, const FString& ListName, TArray< TSharedRef<FOnlineFriend> >& OutFriends) override;
 	virtual TSharedPtr<FOnlineFriend> GetFriend(int32 LocalUserNum, const FUniqueNetId& FriendId, const FString& ListName) override;
 	virtual bool IsFriend(int32 LocalUserNum, const FUniqueNetId& FriendId, const FString& ListName) override;
diff --git a/Engine/Plugins/Online/OnlineSubsystemOculus/Source/Private/OnlineFriendsInterfaceOculus.cpp b/Engine/Plugins/Online/OnlineSubsystemOculus/Source/Private/OnlineFriendsInterfaceOculus.cpp
--- a/Engine/Plugins/Online/OnlineSubsystemOculus/Source/Private/OnlineFriendsInterfaceOculus.cpp
+++ b/Engine/Plugins/Online/OnlineSubsystemOculus/Source/Private/OnlineFriendsInterfaceOculus.cpp
@@ -128,6 +128,15 @@ void FOnlineFriendsOculus::SetFriendAlias(int32 LocalUserNum, const FUniqueNetId
 	});
 }
 
+void FOnlineFriendsOculus::DeleteFriendAlias(int32 LocalUserNum, const FUniqueNetId& FriendId, const FString& ListName, const FOnDeleteFriendAliasComplete& Delegate)
+{
+	TSharedRef<const FUniqueNetId> FriendIdRef = FriendId.AsShared();
+	OculusSubsystem.ExecuteNextTick([LocalUserNum, FriendIdRef, ListName, Delegate]()
+	{
+		UE_LOG_ONLINE_FRIEND(Warning, TEXT("FOnlineFriendsOculus::DeleteFriendAlias is not implemented"));
+		Delegate.ExecuteIfBound(LocalUserNum, *FriendIdRef, ListName, FOnlineError(EOnlineErrorResult::NotImplemented));
+	});
+}
 bool FOnlineFriendsOculus::DeleteFriend(int32 LocalUserNum, const FUniqueNetId& FriendId, const FString& ListName)
 {
 	/** Does not exist in LibOVRPlatform */
diff --git a/Engine/Plugins/Online/OnlineSubsystemOculus/Source/Private/OnlineFriendsInterfaceOculus.h b/Engine/Plugins/Online/OnlineSubsystemOculus/Source/Private/OnlineFriendsInterfaceOculus.h
--- a/Engine/Plugins/Online/OnlineSubsystemOculus/Source/Private/OnlineFriendsInterfaceOculus.h
+++ b/Engine/Plugins/Online/OnlineSubsystemOculus/Source/Private/OnlineFriendsInterfaceOculus.h
@@ -112,6 +112,7 @@ public:
 	virtual bool AcceptInvite(int32 LocalUserNum, const FUniqueNetId& FriendId, const FString& ListName, const FOnAcceptInviteComplete& Delegate = FOnAcceptInviteComplete()) override;
 	virtual bool RejectInvite(int32 LocalUserNum, const FUniqueNetId& FriendId, const FString& ListName) override;
 	virtual void SetFriendAlias(int32 LocalUserNum, const FUniqueNetId& FriendId, const FString& ListName, const FString& Alias, const FOnSetFriendAliasComplete& Delegate = FOnSetFriendAliasComplete()) override;
+	virtual void DeleteFriendAlias(int32 LocalUserNum, const FUniqueNetId& FriendId, const FString& ListName, const FOnDeleteFriendAliasComplete& Delegate = FOnDeleteFriendAliasComplete()) override;
 	virtual bool DeleteFriend(int32 LocalUserNum, const FUniqueNetId& FriendId, const FString& ListName) override;
 	virtual bool GetFriendsList(int32 LocalUserNum, const FString& ListName, TArray< TSharedRef<FOnlineFriend> >& OutFriends) override;
 	virtual TSharedPtr<FOnlineFriend> GetFriend(int32 LocalUserNum, const FUniqueNetId& FriendId, const FString& ListName) override;
diff --git a/Engine/Plugins/Online/OnlineSubsystemSteam/Source/Private/OnlineFriendsInterfaceSteam.cpp b/Engine/Plugins/Online/OnlineSubsystemSteam/Source/Private/OnlineFriendsInterfaceSteam.cpp
--- a/Engine/Plugins/Online/OnlineSubsystemSteam/Source/Private/OnlineFriendsInterfaceSteam.cpp
+++ b/Engine/Plugins/Online/OnlineSubsystemSteam/Source/Private/OnlineFriendsInterfaceSteam.cpp
@@ -113,6 +113,16 @@ void FOnlineFriendsSteam::SetFriendAlias(int32 LocalUserNum, const FUniqueNetId&
 	});
 }
 
+void FOnlineFriendsSteam::DeleteFriendAlias(int32 LocalUserNum, const FUniqueNetId& FriendId, const FString& ListName, const FOnDeleteFriendAliasComplete& Delegate)
+{
+	TSharedRef<const FUniqueNetId> FriendIdRef = FriendId.AsShared();
+	SteamSubsystem->ExecuteNextTick([LocalUserNum, FriendIdRef, ListName, Delegate]()
+	{
+		UE_LOG_ONLINE_FRIEND(Warning, TEXT("FOnlineFriendsSteam::DeleteFriendAlias is not supported"));
+		Delegate.ExecuteIfBound(LocalUserNum, *FriendIdRef, ListName, FOnlineError(EOnlineErrorResult::NotImplemented));
+	});
+}
+
 bool FOnlineFriendsSteam::DeleteFriend(int32 LocalUserNum, const FUniqueNetId& FriendId, const FString& ListName)
 {
 	TriggerOnDeleteFriendCompleteDelegates(LocalUserNum, false, FriendId, ListName, FString(TEXT("DeleteFriend() is not supported")));
diff --git a/Engine/Plugins/Online/OnlineSubsystemSteam/Source/Private/OnlineFriendsInterfaceSteam.h b/Engine/Plugins/Online/OnlineSubsystemSteam/Source/Private/OnlineFriendsInterfaceSteam.h
--- a/Engine/Plugins/Online/OnlineSubsystemSteam/Source/Private/OnlineFriendsInterfaceSteam.h
+++ b/Engine/Plugins/Online/OnlineSubsystemSteam/Source/Private/OnlineFriendsInterfaceSteam.h
@@ -115,6 +115,7 @@ public:
 	virtual bool AcceptInvite(int32 LocalUserNum, const FUniqueNetId& FriendId, const FString& ListName, const FOnAcceptInviteComplete& Delegate = FOnAcceptInviteComplete()) override;
  	virtual bool RejectInvite(int32 LocalUserNum, const FUniqueNetId& FriendId, const FString& ListName) override;
 	virtual void SetFriendAlias(int32 LocalUserNum, const FUniqueNetId& FriendId, const FString& ListName, const FString& Alias, const FOnSetFriendAliasComplete& Delegate = FOnSetFriendAliasComplete()) override;
+	virtual void DeleteFriendAlias(int32 LocalUserNum, const FUniqueNetId& FriendId, const FString& ListName, const FOnDeleteFriendAliasComplete& Delegate = FOnDeleteFriendAliasComplete()) override;
 	virtual bool DeleteFriend(int32 LocalUserNum, const FUniqueNetId& FriendId, const FString& ListName) override;
 	virtual bool GetFriendsList(int32 LocalUserNum, const FString& ListName, TArray< TSharedRef<FOnlineFriend> >& OutFriends) override;
 	virtual TSharedPtr<FOnlineFriend> GetFriend(int32 LocalUserNum, const FUniqueNetId& FriendId, const FString& ListName) override;
diff --git a/Engine/Plugins/Online/OnlineSubsystemTencent/Source/Private/OnlineFriendsTencent.cpp b/Engine/Plugins/Online/OnlineSubsystemTencent/Source/Private/OnlineFriendsTencent.cpp
--- a/Engine/Plugins/Online/OnlineSubsystemTencent/Source/Private/OnlineFriendsTencent.cpp
+++ b/Engine/Plugins/Online/OnlineSubsystemTencent/Source/Private/OnlineFriendsTencent.cpp
@@ -479,6 +479,16 @@ void FOnlineFriendsTencent::SetFriendAlias(int32 LocalUserNum, const FUniqueNetI
 	});
 }
 
+void FOnlineFriendsTencent::DeleteFriendAlias(int32 LocalUserNum, const FUniqueNetId& FriendId, const FString& ListName, const FOnDeleteFriendAliasComplete& Delegate)
+{
+	TSharedRef<const FUniqueNetId> FriendIdRef = FriendId.AsShared();
+	TencentSubsystem->ExecuteNextTick([LocalUserNum, FriendIdRef, ListName, Delegate]()
+	{
+		UE_LOG_ONLINE_FRIEND(Warning, TEXT("FOnlineFriendsTencent::DeleteFriendAlias is not implemented"));
+		Delegate.ExecuteIfBound(LocalUserNum, *FriendIdRef, ListName, FOnlineError(EOnlineErrorResult::NotImplemented));
+	});
+}
+
 bool FOnlineFriendsTencent::DeleteFriend(int32 LocalUserNum, const FUniqueNetId& FriendId, const FString& ListName)
 {
 	/** NYI */
diff --git a/Engine/Plugins/Online/OnlineSubsystemTencent/Source/Private/OnlineFriendsTencent.h b/Engine/Plugins/Online/OnlineSubsystemTencent/Source/Private/OnlineFriendsTencent.h
--- a/Engine/Plugins/Online/OnlineSubsystemTencent/Source/Private/OnlineFriendsTencent.h
+++ b/Engine/Plugins/Online/OnlineSubsystemTencent/Source/Private/OnlineFriendsTencent.h
@@ -188,6 +188,7 @@ public:
 	virtual bool AcceptInvite(int32 LocalUserNum, const FUniqueNetId& FriendId, const FString& ListName, const FOnAcceptInviteComplete& Delegate = FOnAcceptInviteComplete()) override;
 	virtual bool RejectInvite(int32 LocalUserNum, const FUniqueNetId& FriendId, const FString& ListName) override;
 	virtual void SetFriendAlias(int32 LocalUserNum, const FUniqueNetId& FriendId, const FString& ListName, const FString& Alias, const FOnSetFriendAliasComplete& Delegate = FOnSetFriendAliasComplete()) override;
+	virtual void DeleteFriendAlias(int32 LocalUserNum, const FUniqueNetId& FriendId, const FString& ListName, const FOnDeleteFriendAliasComplete& Delegate = FOnDeleteFriendAliasComplete()) override;
 	virtual bool DeleteFriend(int32 LocalUserNum, const FUniqueNetId& FriendId, const FString& ListName) override;
 	virtual bool GetFriendsList(int32 LocalUserNum, const FString& ListName, TArray< TSharedRef<FOnlineFriend> >& OutFriends) override;
 	virtual TSharedPtr<FOnlineFriend> GetFriend(int32 LocalUserNum, const FUniqueNetId& FriendId, const FString& ListName) override;
diff --git a/Engine/Plugins/Runtime/OpenXR/Source/OpenXRHMD/Private/OpenXRHMD.cpp b/Engine/Plugins/Runtime/OpenXR/Source/OpenXRHMD/Private/OpenXRHMD.cpp
--- a/Engine/Plugins/Runtime/OpenXR/Source/OpenXRHMD/Private/OpenXRHMD.cpp
+++ b/Engine/Plugins/Runtime/OpenXR/Source/OpenXRHMD/Private/OpenXRHMD.cpp
@@ -678,11 +678,15 @@ void FOpenXRHMD::AdjustViewRect(EStereoscopicPass StereoPass, int32& X, int32& Y
 
 	const XrViewConfigurationView& Config = Configs[ViewIndex];
 
+	FIntPoint ViewRectMin(EForceInit::ForceInitToZero);
 	for (uint32 i = 0; i < ViewIndex; ++i)
 	{
-		X += Configs[i].recommendedImageRectWidth;
+		ViewRectMin.X += Configs[i].recommendedImageRectWidth;
 	}
+	QuantizeSceneBufferSize(ViewRectMin, ViewRectMin);
 
+	X = ViewRectMin.X;
+	Y = ViewRectMin.Y;
 	SizeX = Config.recommendedImageRectWidth;
 	SizeY = Config.recommendedImageRectHeight;
 }
@@ -818,6 +822,7 @@ FOpenXRHMD::FOpenXRHMD(const FAutoRegister& AutoRegister, XrInstance InInstance,
 	, bIsReady(false)
 	, bIsRendering(false)
 	, bRunRequested(false)
+	, bNeedReAllocatedDepth(false)
 	, CurrentSessionState(XR_SESSION_STATE_UNKNOWN)
 	, Instance(InInstance)
 	, System(InSystem)
@@ -838,7 +843,7 @@ FOpenXRHMD::FOpenXRHMD(const FAutoRegister& AutoRegister, XrInstance InInstance,
 	ViewState.next = nullptr;
 	ViewState.viewStateFlags = 0;
 
-	bNeedReAllocatedDepth = bDepthExtensionSupported = Extensions.Contains(XR_KHR_COMPOSITION_LAYER_DEPTH_EXTENSION_NAME);
+	bDepthExtensionSupported = Extensions.Contains(XR_KHR_COMPOSITION_LAYER_DEPTH_EXTENSION_NAME);
 	bHiddenAreaMaskSupported = Extensions.Contains(XR_KHR_VISIBILITY_MASK_EXTENSION_NAME);
 
 	{
@@ -1203,27 +1208,27 @@ bool FOpenXRHMD::AllocateRenderTargetTexture(uint32 Index, uint32 SizeX, uint32
 	// Grab the presentation texture out of the swapchain.
 	OutTargetableTexture = OutShaderResourceTexture = (FTexture2DRHIRef&)Swapchain->GetTextureRef();
 
-	if (bDepthExtensionSupported)
-	{
-		// Allocate the depth buffer swapchain while we're here.
-		DepthSwapchain = RenderBridge->CreateSwapchain(Session, PF_DepthStencil, SizeX, SizeY, NumMips, NumSamples, 0, TexCreate_DepthStencilTargetable);
-		if (!DepthSwapchain)
-		{
-			return false;
-		}
-		bNeedReAllocatedDepth = false;
-	}
+	bNeedReAllocatedDepth = bDepthExtensionSupported;
 
 	return true;
 }
 
 bool FOpenXRHMD::AllocateDepthTexture(uint32 Index, uint32 SizeX, uint32 SizeY, uint8 Format, uint32 NumMips, uint32 InTexFlags, uint32 TargetableTextureFlags, FTexture2DRHIRef& OutTargetableTexture, FTexture2DRHIRef& OutShaderResourceTexture, uint32 NumSamples)
 {
-	if (!DepthSwapchain.IsValid())
+	// FIXME: UE4 constantly calls this function even when there is no reason to reallocate the depth texture
+	if (!bDepthExtensionSupported || !bNeedReAllocatedDepth)
+	{
+		return false;
+	}
+
+	DepthSwapchain = RenderBridge->CreateSwapchain(Session, PF_DepthStencil, SizeX, SizeY, FMath::Max(NumMips, 1u), NumSamples, 0, TexCreate_DepthStencilTargetable);
+	if (!DepthSwapchain)
 	{
 		return false;
 	}
 
+	bNeedReAllocatedDepth = false;
+
 	OutTargetableTexture = OutShaderResourceTexture = (FTexture2DRHIRef&)DepthSwapchain->GetTextureRef();
 
 	return true;
@@ -1242,7 +1247,7 @@ void FOpenXRHMD::OnBeginRendering_RenderThread(FRHICommandListImmediate& RHICmdL
 			bIsRendering = true;
 
 			Swapchain->IncrementSwapChainIndex_RHIThread(FrameStateRHI.predictedDisplayPeriod);
-			if (bDepthExtensionSupported)
+			if (bDepthExtensionSupported && !bNeedReAllocatedDepth)
 			{
 				ensure(DepthSwapchain != nullptr);
 				DepthSwapchain->IncrementSwapChainIndex_RHIThread(FrameStateRHI.predictedDisplayPeriod);
@@ -1268,35 +1273,34 @@ void FOpenXRHMD::OnBeginRendering_RenderThread(FRHICommandListImmediate& RHICmdL
 	ViewsRHI.SetNum(Views.Num());
 	DepthLayersRHI.SetNum(Views.Num());
 
-	int32 OffsetX = 0;
-
-	const float WorldScale = GetWorldToMetersScale() * (1.0 / 100.0f); // physical scale is 100 UUs/meter
-	float NearZ = GNearClippingPlane * WorldScale;
+	FIntPoint ViewRectMin(EForceInit::ForceInitToZero);
+	float NearZ = GNearClippingPlane / GetWorldToMetersScale();
 
 	for (int32 ViewIndex = 0; ViewIndex < Views.Num(); ViewIndex++)
 	{
 		const XrView& View = Views[ViewIndex];
 		const XrViewConfigurationView& Config = Configs[ViewIndex];
+
+		FIntRect ViewRect;
+		QuantizeSceneBufferSize(ViewRectMin, ViewRect.Min);
+		ViewRect.Max = ViewRect.Min + FIntPoint(Config.recommendedImageRectWidth, Config.recommendedImageRectHeight);
 		FTransform ViewTransform = ToFTransform(View.pose, GetWorldToMetersScale());
 
 		XrCompositionLayerProjectionView& Projection = ViewsRHI[ViewIndex];
 		XrCompositionLayerDepthInfoKHR& DepthLayer = DepthLayersRHI[ViewIndex];
 
 		Projection.type = XR_TYPE_COMPOSITION_LAYER_PROJECTION_VIEW;
-		Projection.next = bDepthExtensionSupported ? &DepthLayer : nullptr;
+		Projection.next = nullptr;
 		Projection.fov = View.fov;
 		Projection.pose = ToXrPose(ViewTransform * BaseTransform, GetWorldToMetersScale());
 		Projection.subImage.swapchain = static_cast<FOpenXRSwapchain*>(GetSwapchain())->GetHandle();
 		Projection.subImage.imageArrayIndex = 0;
 		Projection.subImage.imageRect = {
-			{ OffsetX, 0 },
-			{
-				(int32)Config.recommendedImageRectWidth,
-				(int32)Config.recommendedImageRectHeight
-			}
+			{ ViewRect.Min.X, ViewRect.Min.Y },
+			{ ViewRect.Width(), ViewRect.Height() }
 		};
 
-		if (bDepthExtensionSupported)
+		if (bDepthExtensionSupported && !bNeedReAllocatedDepth)
 		{
 			DepthLayer.type = XR_TYPE_COMPOSITION_LAYER_DEPTH_INFO_KHR;
 			DepthLayer.next = nullptr;
@@ -1307,9 +1311,11 @@ void FOpenXRHMD::OnBeginRendering_RenderThread(FRHICommandListImmediate& RHICmdL
 			DepthLayer.maxDepth = 1.0f;
 			DepthLayer.nearZ = FLT_MAX;
 			DepthLayer.farZ = NearZ;
+
+			Projection.next = &DepthLayer;
 		}
 
-		OffsetX += Config.recommendedImageRectWidth;
+		ViewRectMin.X += Config.recommendedImageRectWidth;
 	}
 
 	// Give the RHI thread its own copy of the frame state and tracking space
@@ -1503,11 +1509,11 @@ FIntPoint FOpenXRHMD::GetIdealRenderTargetSize() const
 	{
 		Size.X += (int)Config.recommendedImageRectWidth;
 		Size.Y = FMath::Max(Size.Y, (int)Config.recommendedImageRectHeight);
-	}
 
-	// We always prefer the nearest multiple of 4 for our buffer sizes. Make sure we round up here,
-	// so we're consistent with the rest of the engine in creating our buffers.
-	QuantizeSceneBufferSize(Size, Size);
+		// We always prefer the nearest multiple of 4 for our buffer sizes. Make sure we round up here,
+		// so we're consistent with the rest of the engine in creating our buffers.
+		QuantizeSceneBufferSize(Size, Size);
+	}
 
 	return Size;
 }
diff --git a/Engine/Plugins/Runtime/WebSocketNetworking/Source/WebSocketNetworking/WebSocketNetworking.Build.cs b/Engine/Plugins/Runtime/WebSocketNetworking/Source/WebSocketNetworking/WebSocketNetworking.Build.cs
--- a/Engine/Plugins/Runtime/WebSocketNetworking/Source/WebSocketNetworking/WebSocketNetworking.Build.cs
+++ b/Engine/Plugins/Runtime/WebSocketNetworking/Source/WebSocketNetworking/WebSocketNetworking.Build.cs
@@ -10,6 +10,7 @@ namespace UnrealBuildTool.Rules
 				new string[] {
 					"Core",
 					"CoreUObject",
+					"NetCore",
 					"Engine",
 					"EngineSettings",
 					"ImageCore",
diff --git a/Engine/Source/Developer/Apple/MetalShaderFormat/Private/MetalShaderCompiler.cpp b/Engine/Source/Developer/Apple/MetalShaderFormat/Private/MetalShaderCompiler.cpp
--- a/Engine/Source/Developer/Apple/MetalShaderFormat/Private/MetalShaderCompiler.cpp
+++ b/Engine/Source/Developer/Apple/MetalShaderFormat/Private/MetalShaderCompiler.cpp
@@ -2835,8 +2835,19 @@ bool FinalizeLibrary_Metal(FName const& Format, FString const& WorkingDir, FStri
 					FGuid Guid = FGuid::NewGuid();
 					RemoteLibPath = OriginalRemoteLibPath + FString::Printf(TEXT(".%x%x%x%x"), Guid.A, Guid.B, Guid.C, Guid.D);
 				}
-			
-				Params = FString::Printf(TEXT("-o=\"%s\" \"%s\""), *RemoteLibPath, *ArchivePath);
+
+				uint64 XcodeBuildVers = 0;
+				uint16 XcodeVers = GetXcodeVersion(XcodeBuildVers);
+				uint16 XcodeMajorVers = ((XcodeVers >> 8) & 0xff);
+
+				if (XcodeMajorVers >= 11)
+				{
+					Params = FString::Printf(TEXT("-o \"%s\" \"%s\""), *RemoteLibPath, *ArchivePath);
+				}
+				else
+				{
+					Params = FString::Printf(TEXT("-o=\"%s\" \"%s\""), *RemoteLibPath, *ArchivePath);
+				}
 				ReturnCode = 0;
 				Results = TEXT("");
 				Errors = TEXT("");
diff --git a/Engine/Source/Editor/FoliageEdit/Private/SFoliageEdit.cpp b/Engine/Source/Editor/FoliageEdit/Private/SFoliageEdit.cpp
--- a/Engine/Source/Editor/FoliageEdit/Private/SFoliageEdit.cpp
+++ b/Engine/Source/Editor/FoliageEdit/Private/SFoliageEdit.cpp
@@ -295,7 +295,7 @@ void SFoliageEdit::Construct(const FArguments& InArgs)
 									.Visibility(this, &SFoliageEdit::GetVisibility_SpawnInCurrentLevelMode)
 									.OnCheckStateChanged(this, &SFoliageEdit::OnCheckStateChanged_SpawnInCurrentLevelMode)
 									.IsChecked(this, &SFoliageEdit::GetCheckState_SpawnInCurrentLevelMode)
-									.ToolTipText(LOCTEXT("SpawnInCurrentLevelModeTooltips", "Wether to place foliage meshes in the current level or in the level containing the mesh being painted on."))
+									.ToolTipText(LOCTEXT("SpawnInCurrentLevelModeTooltips", "Whether to place foliage meshes in the current level or in the level containing the mesh being painted on."))
 									[
 										SNew(STextBlock)
 										.Text(LOCTEXT("SpawnInCurrentLevelMode", "Place in Current Level"))
diff --git a/Engine/Source/Editor/Kismet/Private/BlueprintCompilationManager.cpp b/Engine/Source/Editor/Kismet/Private/BlueprintCompilationManager.cpp
--- a/Engine/Source/Editor/Kismet/Private/BlueprintCompilationManager.cpp
+++ b/Engine/Source/Editor/Kismet/Private/BlueprintCompilationManager.cpp
@@ -422,6 +422,7 @@ struct FCompilerData
 	}
 
 	bool IsSkeletonOnly() const { return JobType == ECompilationManagerJobType::SkeletonOnly; }
+	bool IsCppCompileType() const { return InternalOptions.CompileType == EKismetCompileType::Cpp; }
 	bool ShouldSetTemporaryBlueprintFlags() const { return JobType != ECompilationManagerJobType::RelinkOnly; }
 	bool ShouldResetErrorState() const { return JobType == ECompilationManagerJobType::Normal && InternalOptions.CompileType != EKismetCompileType::BytecodeOnly; }
 	bool ShouldValidateVariableNames() const { return JobType == ECompilationManagerJobType::Normal; }
@@ -1083,10 +1084,17 @@ void FBlueprintCompilationManagerImpl::FlushCompilationQueueImpl(bool bSuppressB
 					OldCDOs.Add(BP, BP->GeneratedClass->ClassDefaultObject);
 				}
 
+				EBlueprintCompileReinstancerFlags CompileReinstancerFlags =
+					EBlueprintCompileReinstancerFlags::AutoInferSaveOnCompile;
+				if (!CompilerData.IsCppCompileType())
+				{
+					CompileReinstancerFlags |= EBlueprintCompileReinstancerFlags::AvoidCDODuplication;
+				}
+
 				CompilerData.Reinstancer = TSharedPtr<FBlueprintCompileReinstancer>(
 					new FBlueprintCompileReinstancer(
 						BP->GeneratedClass,
-						EBlueprintCompileReinstancerFlags::AutoInferSaveOnCompile | EBlueprintCompileReinstancerFlags::AvoidCDODuplication
+						CompileReinstancerFlags
 					)
 				);
 
@@ -1127,7 +1135,7 @@ void FBlueprintCompilationManagerImpl::FlushCompilationQueueImpl(bool bSuppressB
 				// default value propagation occurs in ReinstaneBatch, CDO will be created via CompileFunctions call:
 				if(BP->ParentClass)
 				{
-					if(BP->GeneratedClass)
+					if(BP->GeneratedClass && !CompilerData.IsCppCompileType())
 					{
 						BP->GeneratedClass->ClassDefaultObject = nullptr;
 					}
diff --git a/Engine/Source/Editor/Kismet/Private/FindInBlueprintManager.cpp b/Engine/Source/Editor/Kismet/Private/FindInBlueprintManager.cpp
--- a/Engine/Source/Editor/Kismet/Private/FindInBlueprintManager.cpp
+++ b/Engine/Source/Editor/Kismet/Private/FindInBlueprintManager.cpp
@@ -115,7 +115,9 @@ inline UObject* GetAssetObject(UObject* InObject)
 // FSearchDataVersionInfo
 FSearchDataVersionInfo FSearchDataVersionInfo::Current =
 {
+	/** FiBDataVersion */
 	EFiBVersion::FIB_VER_LATEST,
+	/** EditorObjectVersion */
 	FEditorObjectVersion::LatestVersion
 };
 
@@ -366,6 +368,56 @@ namespace FiBSerializationHelpers
 		SizeOfDataAr << ReturnValue;
 		return ReturnValue;
 	}
+
+	/** Helper function to validate search data version info */
+	bool ValidateSearchDataVersionInfo(const FSearchDataVersionInfo& InVersionInfo)
+	{
+		return InVersionInfo.FiBDataVersion != EFiBVersion::FIB_VER_NONE && InVersionInfo.EditorObjectVersion >= 0;
+	}
+
+	/** Helper function to validate and/or deserialize version info if necessary */
+	bool ValidateSearchDataVersionInfo(const FString& InAssetPath, const FString& InFiBData, FSearchDataVersionInfo& InOutVersionInfo)
+	{
+		if (InOutVersionInfo.FiBDataVersion == EFiBVersion::FIB_VER_NONE)
+		{
+			// Deserialize the FiB data version
+			const int32 FiBDataLength = InFiBData.Len();
+			if (ensureMsgf(FiBDataLength > 0, TEXT("Versioned search data was zero length!")))
+			{
+				FBufferReader ReaderStream((void*)*InFiBData, FiBDataLength * sizeof(TCHAR), false);
+				InOutVersionInfo.FiBDataVersion = FiBSerializationHelpers::Deserialize<int32>(ReaderStream);
+			}
+		}
+
+		if (InOutVersionInfo.EditorObjectVersion < 0)
+		{
+			// Determine the editor object version that the asset package was last serialized with
+			FString PackageFilename;
+			const FString PackageName = FPackageName::ObjectPathToPackageName(InAssetPath);
+			if (ensureMsgf(FPackageName::DoesPackageExist(PackageName, nullptr, &PackageFilename), TEXT("FiB: Failed to map package to filename.")))
+			{
+				// Open a new file archive for reading
+				FArchive* PackageFile = IFileManager::Get().CreateFileReader(*PackageFilename);
+				if (ensureMsgf(PackageFile != nullptr, TEXT("FiB: Unable to open package to read file summary.")))
+				{
+					// Read the package file summary
+					FPackageFileSummary PackageFileSummary;
+					*PackageFile << PackageFileSummary;
+
+					// Close the file
+					delete PackageFile;
+
+					// If an editor object version exists in the package file summary, record it
+					if (const FCustomVersion* const EditorObjectVersion = PackageFileSummary.GetCustomVersionContainer().GetVersion(FEditorObjectVersion::GUID))
+					{
+						InOutVersionInfo.EditorObjectVersion = EditorObjectVersion->Version;
+					}
+				}
+			}
+		}
+
+		return ValidateSearchDataVersionInfo(InOutVersionInfo);
+	}
 }
 
 namespace BlueprintSearchMetaDataHelpers
@@ -1565,7 +1617,7 @@ void FFindInBlueprintSearchManager::OnAssetAdded(const FAssetData& InAssetData)
 			}
 			else
 			{
-				ExtractUnloadedFiBData(InAssetData, FiBVersionedSearchData, true);
+				ExtractUnloadedFiBData(InAssetData, FiBVersionedSearchData, EFiBVersion::FIB_VER_NONE);
 			}
 		}
 		else
@@ -1574,7 +1626,7 @@ void FFindInBlueprintSearchManager::OnAssetAdded(const FAssetData& InAssetData)
 			FAssetDataTagMapSharedView::FFindTagResult ResultLegacy = InAssetData.TagsAndValues.FindTag("FiB");
 			if (ResultLegacy.IsSet())
 			{
-				ExtractUnloadedFiBData(InAssetData, ResultLegacy.GetValue(), false);
+				ExtractUnloadedFiBData(InAssetData, ResultLegacy.GetValue(), EFiBVersion::FIB_VER_BASE);
 			}
 			// The asset has no FiB data, keep track of it so we can inform the user
 			else
@@ -1586,7 +1638,7 @@ void FFindInBlueprintSearchManager::OnAssetAdded(const FAssetData& InAssetData)
 	}
 }
 
-void FFindInBlueprintSearchManager::ExtractUnloadedFiBData(const FAssetData& InAssetData, const FString& InFiBData, bool bIsVersioned)
+void FFindInBlueprintSearchManager::ExtractUnloadedFiBData(const FAssetData& InAssetData, const FString& InFiBData, EFiBVersion InFiBDataVersion)
 {
 	if (SearchMap.Contains(InAssetData.ObjectPath))
 	{
@@ -1647,36 +1699,8 @@ void FFindInBlueprintSearchManager::ExtractUnloadedFiBData(const FAssetData& InA
 	NewSearchData.bMarkedForDeletion = false;
 	NewSearchData.Value = *InFiBData;
 
-	// Deserialize the version if available
-	if (bIsVersioned)
-	{
-		checkf(NewSearchData.Value.Len(), TEXT("Versioned search data was zero length!"));
-		FBufferReader ReaderStream((void*)*NewSearchData.Value, NewSearchData.Value.Len() * sizeof(TCHAR), false);
-		NewSearchData.VersionInfo.FiBDataVersion = FiBSerializationHelpers::Deserialize<int32>(ReaderStream);
-	}
-
-	// Determine the editor object version that the asset package was last serialized with
-	FString PackageFilename;
-	if (ensureMsgf(FPackageName::DoesPackageExist(InAssetData.PackageName.ToString(), nullptr, &PackageFilename), TEXT("FiB: Failed to map package to filename.")))
-	{
-		// Open a new file archive for reading
-		FArchive* PackageFile = IFileManager::Get().CreateFileReader(*PackageFilename);
-		if (ensureMsgf(PackageFile != nullptr, TEXT("FiB: Unable to open package to read file summary.")))
-		{
-			// Read the package file summary
-			FPackageFileSummary PackageFileSummary;
-			*PackageFile << PackageFileSummary;
-
-			// Close the file
-			delete PackageFile;
-
-			// If an editor object version exists in the package file summary, record it
-			if (const FCustomVersion* const EditorObjectVersion = PackageFileSummary.GetCustomVersionContainer().GetVersion(FEditorObjectVersion::GUID))
-			{
-				NewSearchData.VersionInfo.EditorObjectVersion = EditorObjectVersion->Version;
-			}
-		}
-	}
+	// This will be set to 'None' if the data is versioned. Deserialization of the actual version from the tag value is deferred until later.
+	NewSearchData.VersionInfo.FiBDataVersion = InFiBDataVersion;
 
 	// Since the asset was not loaded, pull out the searchable data stored in the asset
 	AddSearchDataToDatabase(MoveTemp(NewSearchData));
@@ -1961,8 +1985,12 @@ bool FFindInBlueprintSearchManager::ContinueSearchQuery(const FStreamSearch* InS
 				// If there is FiB data, parse it into an ImaginaryBlueprint
 				if (SearchArray[SearchIdx].Value.Len() > 0)
 				{
-					SearchArray[SearchIdx].ImaginaryBlueprint = MakeShareable(new FImaginaryBlueprint(FPaths::GetBaseFilename(SearchArray[SearchIdx].AssetPath.ToString()), SearchArray[SearchIdx].AssetPath.ToString(), SearchArray[SearchIdx].ParentClass, SearchArray[SearchIdx].Interfaces, SearchArray[SearchIdx].Value, SearchArray[SearchIdx].VersionInfo));
-					SearchArray[SearchIdx].Value.Empty();
+					const FString AssetPath = SearchArray[SearchIdx].AssetPath.ToString();
+					if (FiBSerializationHelpers::ValidateSearchDataVersionInfo(AssetPath, SearchArray[SearchIdx].Value, SearchArray[SearchIdx].VersionInfo))
+					{
+						SearchArray[SearchIdx].ImaginaryBlueprint = MakeShareable(new FImaginaryBlueprint(FPaths::GetBaseFilename(AssetPath), AssetPath, SearchArray[SearchIdx].ParentClass, SearchArray[SearchIdx].Interfaces, SearchArray[SearchIdx].Value, SearchArray[SearchIdx].VersionInfo));
+						SearchArray[SearchIdx].Value.Empty();
+					}
 				}
 
  				OutSearchData = SearchArray[SearchIdx++];
@@ -2487,6 +2515,9 @@ TSharedPtr< FJsonObject > FFindInBlueprintSearchManager::ConvertJsonStringToObje
 	int32 SizeOfData;
 	FBufferReader ReaderStream((void*)*InJsonString, InJsonString.Len() * sizeof(TCHAR), false);
 
+	// It's unexpected to not have valid search data versioning info at this point
+	checkf(FiBSerializationHelpers::ValidateSearchDataVersionInfo(InVersionInfo), TEXT("FiB: Invalid JSON stream data version."));
+
 	// If the stream is versioned, read past the version info
 	if (InVersionInfo.FiBDataVersion > EFiBVersion::FIB_VER_BASE)
 	{
diff --git a/Engine/Source/Editor/Kismet/Public/FindInBlueprintManager.h b/Engine/Source/Editor/Kismet/Public/FindInBlueprintManager.h
--- a/Engine/Source/Editor/Kismet/Public/FindInBlueprintManager.h
+++ b/Engine/Source/Editor/Kismet/Public/FindInBlueprintManager.h
@@ -105,6 +105,8 @@ struct KISMET_API FFindInBlueprintSearchTags
 /** FiB data versioning */
 enum EFiBVersion
 {
+	FIB_VER_NONE = -1, // Unknown version (not set)
+
 	FIB_VER_BASE = 0, // All Blueprints prior to versioning will automatically be assumed to be at 0 if they have FiB data collected
 	FIB_VER_VARIABLE_REFERENCE, // Variable references (FMemberReference) is collected in FiB
 	FIB_VER_INTERFACE_GRAPHS, // Implemented Interface Graphs is collected in FiB
@@ -118,7 +120,7 @@ enum EFiBVersion
 struct FSearchDataVersionInfo
 {
 	/** FiB asset registry tag data version */
-	int32 FiBDataVersion = EFiBVersion::FIB_VER_BASE;
+	int32 FiBDataVersion = EFiBVersion::FIB_VER_NONE;
 
 	/** Editor object version used to serialize values in the JSON string lookup table */
 	int32 EditorObjectVersion = -1;
@@ -593,7 +595,7 @@ private:
 	void RemoveBlueprintByPath(FName InPath);
 
 	/** Begins the process of extracting unloaded FiB data */
-	void ExtractUnloadedFiBData(const FAssetData& InAssetData, const FString& InFiBData, bool bIsVersioned);
+	void ExtractUnloadedFiBData(const FAssetData& InAssetData, const FString& InFiBData, EFiBVersion InFiBDataVersion);
 
 	/** Determines the global find results tab label */
 	FText GetGlobalFindResultsTabLabel(int32 TabIdx);
diff --git a/Engine/Source/Editor/KismetCompiler/Private/KismetCompiler.cpp b/Engine/Source/Editor/KismetCompiler/Private/KismetCompiler.cpp
--- a/Engine/Source/Editor/KismetCompiler/Private/KismetCompiler.cpp
+++ b/Engine/Source/Editor/KismetCompiler/Private/KismetCompiler.cpp
@@ -4094,13 +4094,26 @@ void FKismetCompilerContext::CompileFunctions(EInternalCompilerFlags InternalFla
 	// This is phase two, so we want to generated locals if PostponeLocalsGenerationUntilPhaseTwo is set:
 	const bool bGenerateLocals = !!(InternalFlags & EInternalCompilerFlags::PostponeLocalsGenerationUntilPhaseTwo);
 	// Don't propagate values to CDO if we're going to do that in reinstancing:
-	const bool bPropagateValuesToCDO = !(InternalFlags & EInternalCompilerFlags::PostponeDefaultObjectAssignmentUntilReinstancing);
+	bool bPropagateValuesToCDO = !(InternalFlags & EInternalCompilerFlags::PostponeDefaultObjectAssignmentUntilReinstancing);
 	// Don't RefreshExternalBlueprintDependencyNodes if the calling code has done so already:
 	const bool bSkipRefreshExternalBlueprintDependencyNodes = !!(InternalFlags & EInternalCompilerFlags::SkipRefreshExternalBlueprintDependencyNodes);
 	FKismetCompilerVMBackend Backend_VM(Blueprint, Schema, *this);
 
-	// Determine whether or not to skip generated class validation. This requires CDO value propagation to occur first.
-	bool bSkipGeneratedClassValidation = !bPropagateValuesToCDO || CompileOptions.CompileType == EKismetCompileType::Cpp;
+	// Determine whether or not to skip generated class validation.
+	bool bSkipGeneratedClassValidation;
+	if (CompileOptions.DoesRequireCppCodeGeneration())
+	{
+		// CPP codegen requires default value assignment to occur as part of the compilation phase, so we override it here.
+		bPropagateValuesToCDO = true;
+
+		// Also skip generated class validation since it may result in errors and we don't really need to keep the generated class.
+		bSkipGeneratedClassValidation = true;
+	}
+	else
+	{
+		// In all other cases, validation requires CDO value propagation to occur first.
+		bSkipGeneratedClassValidation = !bPropagateValuesToCDO;
+	}
 
 	if( bGenerateLocals )
 	{
diff --git a/Engine/Source/Editor/UnrealEd/Private/CookOnTheFlyServer.cpp b/Engine/Source/Editor/UnrealEd/Private/CookOnTheFlyServer.cpp
--- a/Engine/Source/Editor/UnrealEd/Private/CookOnTheFlyServer.cpp
+++ b/Engine/Source/Editor/UnrealEd/Private/CookOnTheFlyServer.cpp
@@ -4498,8 +4498,7 @@ void GetAdditionalCurrentIniVersionStrings( const ITargetPlatform* TargetPlatfor
 	}
 
 
-	static const FCustomVersionContainer& CustomVersionContainer = FCustomVersionContainer::GetRegistered();
-	for (const auto& CustomVersion : CustomVersionContainer.GetAllVersions())
+	for (const FCustomVersion& CustomVersion : FCurrentCustomVersions::GetAll().GetAllVersions())
 	{
 		FString CustomVersionString = FString::Printf(TEXT("%s:%s"), *CustomVersion.GetFriendlyName().ToString(), *CustomVersion.Key.ToString());
 		FString CustomVersionValue = FString::Printf(TEXT("%d"), CustomVersion.Version);
diff --git a/Engine/Source/Programs/AutomationTool/BuildGraph/Tasks/AgeStoreTask.cs b/Engine/Source/Programs/AutomationTool/BuildGraph/Tasks/AgeStoreTask.cs
--- a/Engine/Source/Programs/AutomationTool/BuildGraph/Tasks/AgeStoreTask.cs
+++ b/Engine/Source/Programs/AutomationTool/BuildGraph/Tasks/AgeStoreTask.cs
@@ -35,7 +35,7 @@ namespace Win.Automation
 		public int Days;
 
 		/// <summary>
-		/// The root of the build dir to check for existing buildversion named directories
+		/// The root of the build directory to check for existing buildversion named directories.
 		/// </summary>
 		[TaskParameter(Optional = true)]
 		public string BuildDir;
@@ -43,14 +43,14 @@ namespace Win.Automation
 		/// <summary>
 		/// A substring to match in directory file names before deleting symbols. This allows the "age store" task
 		/// to avoid deleting symbols from other builds in the case where multiple builds share the same symbol server.
-		/// Specific use of the filter value is determined by the symbol server structure defined by the platform tool chain.
+		/// Specific use of the filter value is determined by the symbol server structure defined by the platform toolchain.
 		/// </summary>
 		[TaskParameter(Optional = true)]
 		public string Filter;
 	}
 
 	/// <summary>
-	/// Task which strips symbols from a set of files. This task is named after the AGESTORE utility that comes with the Microsoft debugger tools SDK, but is actually a separate implementation. The main
+	/// Task that strips symbols from a set of files. This task is named after the AGESTORE utility that comes with the Microsoft debugger tools SDK, but is actually a separate implementation. The main
 	/// difference is that it uses the last modified time rather than last access time to determine which files to delete.
 	/// </summary>
 	[TaskElement("AgeStore", typeof(AgeStoreTaskParameters))]
diff --git a/Engine/Source/Programs/AutomationTool/BuildGraph/Tasks/CommandTask.cs b/Engine/Source/Programs/AutomationTool/BuildGraph/Tasks/CommandTask.cs
--- a/Engine/Source/Programs/AutomationTool/BuildGraph/Tasks/CommandTask.cs
+++ b/Engine/Source/Programs/AutomationTool/BuildGraph/Tasks/CommandTask.cs
@@ -22,18 +22,18 @@ namespace BuildGraph.Tasks
 	}
 
 	/// <summary>
-	/// Parameters for a task which calls another UAT command
+	/// Parameters for a task that calls another UAT command
 	/// </summary>
 	public class CommandTaskParameters
 	{
 		/// <summary>
-		/// The command name to execute
+		/// The command name to execute.
 		/// </summary>
 		[TaskParameter]
 		public string Name;
 
 		/// <summary>
-		/// Arguments to be passed to the command
+		/// Arguments to be passed to the command.
 		/// </summary>
 		[TaskParameter(Optional = true)]
 		public string Arguments;
diff --git a/Engine/Source/Programs/AutomationTool/BuildGraph/Tasks/CommandletTask.cs b/Engine/Source/Programs/AutomationTool/BuildGraph/Tasks/CommandletTask.cs
--- a/Engine/Source/Programs/AutomationTool/BuildGraph/Tasks/CommandletTask.cs
+++ b/Engine/Source/Programs/AutomationTool/BuildGraph/Tasks/CommandletTask.cs
@@ -18,19 +18,19 @@ namespace AutomationTool.Tasks
 	public class CommandletTaskParameters
 	{
 		/// <summary>
-		/// The commandlet name to execute
+		/// The commandlet name to execute.
 		/// </summary>
 		[TaskParameter]
 		public string Name;
 
 		/// <summary>
-		/// The project to run the editor with
+		/// The project to run the editor with.
 		/// </summary>
 		[TaskParameter(Optional = true, ValidationType = TaskParameterValidationType.FileSpec)]
 		public string Project;
 
 		/// <summary>
-		/// Arguments to be passed to the commandlet
+		/// Arguments to be passed to the commandlet.
 		/// </summary>
 		[TaskParameter(Optional = true)]
 		public string Arguments;
diff --git a/Engine/Source/Programs/AutomationTool/BuildGraph/Tasks/CompileTask.cs b/Engine/Source/Programs/AutomationTool/BuildGraph/Tasks/CompileTask.cs
--- a/Engine/Source/Programs/AutomationTool/BuildGraph/Tasks/CompileTask.cs
+++ b/Engine/Source/Programs/AutomationTool/BuildGraph/Tasks/CompileTask.cs
@@ -18,55 +18,55 @@ namespace AutomationTool
 	public class CompileTaskParameters
 	{
 		/// <summary>
-		/// The target to compile
+		/// The target to compile.
 		/// </summary>
 		[TaskParameter]
 		public string Target;
 
 		/// <summary>
-		/// The configuration to compile
+		/// The configuration to compile.
 		/// </summary>
 		[TaskParameter]
 		public UnrealTargetConfiguration Configuration;
 
 		/// <summary>
-		/// The platform to compile for
+		/// The platform to compile for.
 		/// </summary>
 		[TaskParameter]
 		public UnrealTargetPlatform Platform;
 
 		/// <summary>
-		/// The project to compile with
+		/// The project to compile with.
 		/// </summary>
 		[TaskParameter(Optional = true)]
 		public string Project;
 
 		/// <summary>
-		/// Additional arguments for UnrealBuildTool
+		/// Additional arguments for UnrealBuildTool.
 		/// </summary>
 		[TaskParameter(Optional = true)]
 		public string Arguments;
 
 		/// <summary>
-		/// Whether to allow using XGE for compilation
+		/// Whether to allow using XGE for compilation.
 		/// </summary>
 		[TaskParameter(Optional = true)]
 		public bool AllowXGE = true;
 
 		/// <summary>
-		/// Whether to allow using the parallel executor for this compile
+		/// Whether to allow using the parallel executor for this compile.
 		/// </summary>
 		[TaskParameter(Optional = true)]
 		public bool AllowParallelExecutor = true;
 
 		/// <summary>
-		/// Whether to allow cleaning this target. If unspecified, targets are cleaned if the -Clean argument passed on the command line.
+		/// Whether to allow cleaning this target. If unspecified, targets are cleaned if the -Clean argument is passed on the command line.
 		/// </summary>
 		[TaskParameter(Optional = true)]
 		public bool? Clean = null;
 
 		/// <summary>
-		/// Tag to be applied to build products of this task
+		/// Tag to be applied to build products of this task.
 		/// </summary>
 		[TaskParameter(Optional = true, ValidationType = TaskParameterValidationType.TagList)]
 		public string Tag;
diff --git a/Engine/Source/Programs/AutomationTool/BuildGraph/Tasks/CookTask.cs b/Engine/Source/Programs/AutomationTool/BuildGraph/Tasks/CookTask.cs
--- a/Engine/Source/Programs/AutomationTool/BuildGraph/Tasks/CookTask.cs
+++ b/Engine/Source/Programs/AutomationTool/BuildGraph/Tasks/CookTask.cs
@@ -1,4 +1,4 @@
-﻿// Copyright 1998-2019 Epic Games, Inc. All Rights Reserved.
+// Copyright 1998-2019 Epic Games, Inc. All Rights Reserved.
 
 using AutomationTool;
 using System;
@@ -18,37 +18,37 @@ namespace BuildGraph.Tasks
 	public class CookTaskParameters
 	{
 		/// <summary>
-		/// Project file to be cooked
+		/// Project file to be cooked.
 		/// </summary>
 		[TaskParameter]
 		public string Project;
 
 		/// <summary>
-		/// The cook platform to target (eg. WindowsNoEditor)
+		/// The cook platform to target (for example, WindowsNoEditor).
 		/// </summary>
 		[TaskParameter]
 		public string Platform;
 
 		/// <summary>
-		/// List of maps to be cooked, separated by '+' characters
+		/// List of maps to be cooked, separated by '+' characters.
 		/// </summary>
 		[TaskParameter(Optional = true)]
 		public string Maps;
 
 		/// <summary>
-		/// Additional arguments to be passed to the cooker
+		/// Additional arguments to be passed to the cooker.
 		/// </summary>
 		[TaskParameter(Optional = true)]
 		public bool Versioned = false;
 	
 		/// <summary>
-		/// Additional arguments to be passed to the cooker
+		/// Additional arguments to be passed to the cooker.
 		/// </summary>
 		[TaskParameter(Optional = true)]
 		public string Arguments = "";
 
 		/// <summary>
-		/// Tag to be applied to build products of this task
+		/// Tag to be applied to build products of this task.
 		/// </summary>
 		[TaskParameter(Optional = true, ValidationType = TaskParameterValidationType.TagList)]
 		public string Tag;
diff --git a/Engine/Source/Programs/AutomationTool/BuildGraph/Tasks/CopyTask.cs b/Engine/Source/Programs/AutomationTool/BuildGraph/Tasks/CopyTask.cs
--- a/Engine/Source/Programs/AutomationTool/BuildGraph/Tasks/CopyTask.cs
+++ b/Engine/Source/Programs/AutomationTool/BuildGraph/Tasks/CopyTask.cs
@@ -19,31 +19,31 @@ namespace BuildGraph.Tasks
 	public class CopyTaskParameters
 	{
 		/// <summary>
-		/// Filter to be applied to the list of input files. Optional.
+		/// Optional filter to be applied to the list of input files.
 		/// </summary>
 		[TaskParameter(Optional = true, ValidationType = TaskParameterValidationType.FileSpec)]
 		public string Files;
 
 		/// <summary>
-		/// The pattern(s) to copy from (eg. Engine/*.txt)
+		/// The pattern(s) to copy from (for example, Engine/*.txt).
 		/// </summary>
 		[TaskParameter(ValidationType = TaskParameterValidationType.FileSpec)]
 		public string From;
 
 		/// <summary>
-		/// The directory or to copy to
+		/// The directory to copy to.
 		/// </summary>
 		[TaskParameter(ValidationType = TaskParameterValidationType.FileSpec)]
 		public string To;
 
 		/// <summary>
-		/// Whether or not to overwrite existing files
+		/// Whether or not to overwrite existing files.
 		/// </summary>
 		[TaskParameter(Optional = true)]
 		public bool Overwrite = true;
 
 		/// <summary>
-		/// Tag to be applied to build products of this task
+		/// Tag to be applied to build products of this task.
 		/// </summary>
 		[TaskParameter(Optional = true, ValidationType = TaskParameterValidationType.TagList)]
 		public string Tag;
diff --git a/Engine/Source/Programs/AutomationTool/BuildGraph/Tasks/CsCompileTask.cs b/Engine/Source/Programs/AutomationTool/BuildGraph/Tasks/CsCompileTask.cs
--- a/Engine/Source/Programs/AutomationTool/BuildGraph/Tasks/CsCompileTask.cs
+++ b/Engine/Source/Programs/AutomationTool/BuildGraph/Tasks/CsCompileTask.cs
@@ -17,50 +17,50 @@ namespace AutomationTool.Tasks
 	public class CsCompileTaskParameters
 	{
 		/// <summary>
-		/// The C# project file to be compile. More than one project file can be specified by separating with semicolons.
+		/// The C# project file to compile. Using semicolons, more than one project file can be specified.
 		/// </summary>
 		[TaskParameter]
 		public string Project;
 
 		/// <summary>
-		/// The configuration to compile
+		/// The configuration to compile.
 		/// </summary>
 		[TaskParameter(Optional = true)]
 		public string Configuration;
 
 		/// <summary>
-		/// The platform to compile
+		/// The platform to compile.
 		/// </summary>
 		[TaskParameter(Optional = true)]
 		public string Platform;
 
 		/// <summary>
-		/// The target to build
+		/// The target to build.
 		/// </summary>
 		[TaskParameter(Optional = true)]
 		public string Target;
 
 		/// <summary>
-		/// Additional options to pass to the compiler
+		/// Additional options to pass to the compiler.
 		/// </summary>
 		[TaskParameter(Optional = true)]
 		public string Arguments;
 
 		/// <summary>
-		/// Only enumerate build products; do not actually compile the projects.
+		/// Only enumerate build products -- do not actually compile the projects.
 		/// </summary>
 		[TaskParameter(Optional = true)]
 		public bool EnumerateOnly;
 
 		/// <summary>
-		/// Tag to be applied to build products of this task
+		/// Tag to be applied to build products of this task.
 		/// </summary>
 		[TaskParameter(Optional = true, ValidationType = TaskParameterValidationType.TagList)]
 		public string Tag;
 
 		/// <summary>
-		/// Tag to be applied to any non-private references the projects have
-		/// (i.e. those that are external and not copied into the output dir)
+		/// Tag to be applied to any non-private references the projects have.
+		/// (for example, those that are external and not copied into the output directory).
 		/// </summary>
 		[TaskParameter(Optional = true, ValidationType = TaskParameterValidationType.TagList)]
 		public string TagReferences;
diff --git a/Engine/Source/Programs/AutomationTool/BuildGraph/Tasks/DeleteTask.cs b/Engine/Source/Programs/AutomationTool/BuildGraph/Tasks/DeleteTask.cs
--- a/Engine/Source/Programs/AutomationTool/BuildGraph/Tasks/DeleteTask.cs
+++ b/Engine/Source/Programs/AutomationTool/BuildGraph/Tasks/DeleteTask.cs
@@ -1,4 +1,4 @@
-﻿// Copyright 1998-2019 Epic Games, Inc. All Rights Reserved.
+// Copyright 1998-2019 Epic Games, Inc. All Rights Reserved.
 
 using AutomationTool;
 using System;
@@ -19,7 +19,7 @@ namespace BuildGraph.Tasks
 	public class DeleteTaskParameters
 	{
 		/// <summary>
-		/// List of file specifications separated by semicolons (eg. *.cpp;Engine/.../*.bat), or the name of a tag set
+		/// List of file specifications separated by semicolons (for example, *.cpp;Engine/.../*.bat), or the name of a tag set
 		/// </summary>
 		[TaskParameter(ValidationType = TaskParameterValidationType.FileSpec)]
 		public string Files;
diff --git a/Engine/Source/Programs/AutomationTool/BuildGraph/Tasks/LogTask.cs b/Engine/Source/Programs/AutomationTool/BuildGraph/Tasks/LogTask.cs
--- a/Engine/Source/Programs/AutomationTool/BuildGraph/Tasks/LogTask.cs
+++ b/Engine/Source/Programs/AutomationTool/BuildGraph/Tasks/LogTask.cs
@@ -1,4 +1,4 @@
-﻿// Copyright 1998-2019 Epic Games, Inc. All Rights Reserved.
+// Copyright 1998-2019 Epic Games, Inc. All Rights Reserved.
 
 using AutomationTool;
 using System;
@@ -19,7 +19,7 @@ namespace BuildGraph.Tasks
 	public class LogTaskParameters
 	{
 		/// <summary>
-		/// Message to print out
+		/// Message to print out.
 		/// </summary>
 		[TaskParameter(Optional = true)]
 		public string Message;
@@ -38,7 +38,7 @@ namespace BuildGraph.Tasks
 	}
 
 	/// <summary>
-	/// Print a message (and other optional diagnostic information) to the output log
+	/// Print a message (and other optional diagnostic information) to the output log.
 	/// </summary>
 	[TaskElement("Log", typeof(LogTaskParameters))]
 	public class LogTask : CustomTask
diff --git a/Engine/Source/Programs/AutomationTool/BuildGraph/Tasks/MoveTask.cs b/Engine/Source/Programs/AutomationTool/BuildGraph/Tasks/MoveTask.cs
--- a/Engine/Source/Programs/AutomationTool/BuildGraph/Tasks/MoveTask.cs
+++ b/Engine/Source/Programs/AutomationTool/BuildGraph/Tasks/MoveTask.cs
@@ -18,25 +18,25 @@ namespace BuildGraph.Tasks
 	public class MoveTaskParameters
 	{
 		/// <summary>
-		/// Filter to be applied to the list of input files. Optional.
+		/// Optional filter to be applied to the list of input files.
 		/// </summary>
 		[TaskParameter(Optional = true, ValidationType = TaskParameterValidationType.FileSpec)]
 		public string Files;
 
 		/// <summary>
-		/// The pattern(s) to copy from (eg. Engine/*.txt)
+		/// The pattern(s) to copy from (for example, Engine/*.txt).
 		/// </summary>
 		[TaskParameter(ValidationType = TaskParameterValidationType.FileSpec)]
 		public string From;
 
 		/// <summary>
-		/// The directory or to copy to
+		/// The directory to copy to.
 		/// </summary>
 		[TaskParameter(ValidationType = TaskParameterValidationType.FileSpec)]
 		public string To;
 
 		/// <summary>
-		/// Tag to be applied to build products of this task
+		/// Tag to be applied to build products of this task.
 		/// </summary>
 		[TaskParameter(Optional = true, ValidationType = TaskParameterValidationType.TagList)]
 		public string Tag;
diff --git a/Engine/Source/Programs/AutomationTool/BuildGraph/Tasks/MsBuildTask.cs b/Engine/Source/Programs/AutomationTool/BuildGraph/Tasks/MsBuildTask.cs
--- a/Engine/Source/Programs/AutomationTool/BuildGraph/Tasks/MsBuildTask.cs
+++ b/Engine/Source/Programs/AutomationTool/BuildGraph/Tasks/MsBuildTask.cs
@@ -1,4 +1,4 @@
-﻿// Copyright 1998-2019 Epic Games, Inc. All Rights Reserved.
+// Copyright 1998-2019 Epic Games, Inc. All Rights Reserved.
 
 using System;
 using System.Collections.Generic;
@@ -15,32 +15,32 @@ namespace AutomationTool.Tasks
 	/// </summary>
 	public class MsBuildTaskParameters
 	{
-		/// <summary>
-		/// The C# project file to be compile. More than one project file can be specified by separating with semicolons.
-		/// </summary>
-		[TaskParameter]
+        /// <summary>
+        /// The C# project file to compile. Using semicolons, more than one project file can be specified.
+        /// </summary>
+        [TaskParameter]
 		public string Project;
 
 		/// <summary>
-		/// The configuration to compile
+		/// The configuration to compile.
 		/// </summary>
 		[TaskParameter(Optional = true)]
 		public string Configuration;
 
 		/// <summary>
-		/// The platform to compile
+		/// The platform to compile.
 		/// </summary>
 		[TaskParameter(Optional = true)]
 		public string Platform;
 
 		/// <summary>
-		/// Additional options to pass to the compiler
+		/// Additional options to pass to the compiler.
 		/// </summary>
 		[TaskParameter(Optional = true)]
 		public string Arguments;
 
 		/// <summary>
-		/// The MSBuild output verbosity
+		/// The MSBuild output verbosity.
 		/// </summary>
 		[TaskParameter(Optional = true)]
 		public string Verbosity = "minimal";
diff --git a/Engine/Source/Programs/AutomationTool/BuildGraph/Tasks/PakFileTask.cs b/Engine/Source/Programs/AutomationTool/BuildGraph/Tasks/PakFileTask.cs
--- a/Engine/Source/Programs/AutomationTool/BuildGraph/Tasks/PakFileTask.cs
+++ b/Engine/Source/Programs/AutomationTool/BuildGraph/Tasks/PakFileTask.cs
@@ -19,19 +19,19 @@ namespace BuildGraph.Tasks
 	public class PakFileTaskParameters
 	{
 		/// <summary>
-		/// List of files, wildcards and tag sets to add to the pak file, separated by ';' characters.
+		/// List of files, wildcards, and tag sets to add to the pak file, separated by ';' characters.
 		/// </summary>
 		[TaskParameter(ValidationType = TaskParameterValidationType.FileSpec)]
 		public string Files;
 
 		/// <summary>
-		/// PAK file to output
+		/// PAK file to output.
 		/// </summary>
 		[TaskParameter]
 		public FileReference Output;
 
 		/// <summary>
-		/// Path to a Response File that contains a list of files to add to the pak file, instead of specifying them individually
+		/// Path to a Response File that contains a list of files to add to the pak file -- instead of specifying them individually.
 		/// </summary>
 		[TaskParameter(Optional = true)]
 		public FileReference ResponseFile;
@@ -43,31 +43,31 @@ namespace BuildGraph.Tasks
 		public HashSet<DirectoryReference> RebaseDir;
 
 		/// <summary>
-		/// Script which gives the order of files
+		/// Script that gives the order of files.
 		/// </summary>
 		[TaskParameter(Optional = true)]
 		public FileReference Order;
 
 		/// <summary>
-		/// Encryption keys for this pak file
+		/// Encryption keys for this pak file.
 		/// </summary>
 		[TaskParameter(Optional = true)]
 		public FileReference Sign;
 
 		/// <summary>
-		/// Whether to compress files
+		/// Whether to compress files.
 		/// </summary>
 		[TaskParameter(Optional = true)]
 		public bool Compress = true;
 
 		/// <summary>
-		/// Additional arguments to be passed to UnrealPak
+		/// Additional arguments to pass to UnrealPak.
 		/// </summary>
 		[TaskParameter(Optional = true)]
 		public string Arguments = "";
 
 		/// <summary>
-		/// Tag to be applied to build products of this task
+		/// Tag to be applied to build products of this task.
 		/// </summary>
 		[TaskParameter(Optional = true, ValidationType = TaskParameterValidationType.TagList)]
 		public string Tag;
diff --git a/Engine/Source/Programs/AutomationTool/BuildGraph/Tasks/RenameTask.cs b/Engine/Source/Programs/AutomationTool/BuildGraph/Tasks/RenameTask.cs
--- a/Engine/Source/Programs/AutomationTool/BuildGraph/Tasks/RenameTask.cs
+++ b/Engine/Source/Programs/AutomationTool/BuildGraph/Tasks/RenameTask.cs
@@ -1,4 +1,4 @@
-﻿// Copyright 1998-2019 Epic Games, Inc. All Rights Reserved.
+// Copyright 1998-2019 Epic Games, Inc. All Rights Reserved.
 
 using AutomationTool;
 using System;
@@ -20,13 +20,13 @@ namespace BuildGraph.Tasks
 	public class RenameTaskParameters
 	{
 		/// <summary>
-		/// The file or files to rename
+		/// The file or files to rename.
 		/// </summary>
 		[TaskParameter(ValidationType = TaskParameterValidationType.FileSpec)]
 		public string Files;
 
 		/// <summary>
-		/// The current file name, or pattern to match (eg. *.txt). Should not include any path separators.
+		/// The current file name, or pattern to match (for example, *.txt). Should not include any path separators.
 		/// </summary>
 		[TaskParameter(Optional = true)]
 		public string From;
@@ -38,7 +38,7 @@ namespace BuildGraph.Tasks
 		public string To;
 
 		/// <summary>
-		/// Tag to be applied to the renamed files
+		/// Tag to be applied to the renamed files.
 		/// </summary>
 		[TaskParameter(Optional = true, ValidationType = TaskParameterValidationType.TagList)]
 		public string Tag;
diff --git a/Engine/Source/Programs/AutomationTool/BuildGraph/Tasks/SanitizeReceiptTask.cs b/Engine/Source/Programs/AutomationTool/BuildGraph/Tasks/SanitizeReceiptTask.cs
--- a/Engine/Source/Programs/AutomationTool/BuildGraph/Tasks/SanitizeReceiptTask.cs
+++ b/Engine/Source/Programs/AutomationTool/BuildGraph/Tasks/SanitizeReceiptTask.cs
@@ -30,7 +30,7 @@ namespace AutomationTool.Tasks
 	}
 
 	/// <summary>
-	/// Task which tags build products and/or runtime dependencies by reading from *.target files.
+	/// Task that tags build products and/or runtime dependencies by reading from *.target files.
 	/// </summary>
 	[TaskElement("SanitizeReceipt", typeof(SanitizeReceiptTaskParameters))]
 	class SanitizeReceiptTask : CustomTask
diff --git a/Engine/Source/Programs/AutomationTool/BuildGraph/Tasks/SetVersionTask.cs b/Engine/Source/Programs/AutomationTool/BuildGraph/Tasks/SetVersionTask.cs
--- a/Engine/Source/Programs/AutomationTool/BuildGraph/Tasks/SetVersionTask.cs
+++ b/Engine/Source/Programs/AutomationTool/BuildGraph/Tasks/SetVersionTask.cs
@@ -1,4 +1,4 @@
-﻿// Copyright 1998-2019 Epic Games, Inc. All Rights Reserved.
+// Copyright 1998-2019 Epic Games, Inc. All Rights Reserved.
 
 using System;
 using System.Collections.Generic;
@@ -17,49 +17,49 @@ namespace AutomationTool.Tasks
 	public class SetVersionTaskParameters
 	{
 		/// <summary>
-		/// The changelist to set in the version files
+		/// The changelist to set in the version files.
 		/// </summary>
 		[TaskParameter]
 		public int Change;
 
 		/// <summary>
-		/// The engine compatible changelist to set in the version files
+		/// The engine compatible changelist to set in the version files.
 		/// </summary>
 		[TaskParameter(Optional = true)]
 		public int CompatibleChange;
 
 		/// <summary>
-		/// The branch string
+		/// The branch string.
 		/// </summary>
 		[TaskParameter]
 		public string Branch;
 
 		/// <summary>
-		/// The build version string
+		/// The build version string.
 		/// </summary>
 		[TaskParameter(Optional = true)]
 		public string Build;
 
 		/// <summary>
-		/// Whether to set the IS_LICENSEE_VERSION flag to true
+		/// Whether to set the IS_LICENSEE_VERSION flag to true.
 		/// </summary>
 		[TaskParameter(Optional = true)]
 		public bool Licensee;
 
 		/// <summary>
-		/// Whether to set the ENGINE_IS_PROMOTED_BUILD flag to true
+		/// Whether to set the ENGINE_IS_PROMOTED_BUILD flag to true.
 		/// </summary>
 		[TaskParameter(Optional = true)]
 		public bool Promoted = true;
 
 		/// <summary>
-		/// If set, don't actually write to the files - just return the version files that would be updated. Useful for local builds.
+		/// If set, do not write to the files -- just return the version files that would be updated. Useful for local builds.
 		/// </summary>
 		[TaskParameter(Optional = true)]
 		public bool SkipWrite;
 
 		/// <summary>
-		/// Tag to be applied to build products of this task
+		/// Tag to be applied to build products of this task.
 		/// </summary>
 		[TaskParameter(Optional = true, ValidationType = TaskParameterValidationType.TagList)]
 		public string Tag;
diff --git a/Engine/Source/Programs/AutomationTool/BuildGraph/Tasks/SignTask.cs b/Engine/Source/Programs/AutomationTool/BuildGraph/Tasks/SignTask.cs
--- a/Engine/Source/Programs/AutomationTool/BuildGraph/Tasks/SignTask.cs
+++ b/Engine/Source/Programs/AutomationTool/BuildGraph/Tasks/SignTask.cs
@@ -1,4 +1,4 @@
-﻿// Copyright 1998-2019 Epic Games, Inc. All Rights Reserved.
+// Copyright 1998-2019 Epic Games, Inc. All Rights Reserved.
 
 using AutomationTool;
 using System;
@@ -18,13 +18,13 @@ namespace BuildGraph.Tasks
 	public class SignTaskParameters
 	{
 		/// <summary>
-		/// List of file specifications separated by semicolons (eg. *.cpp;Engine/.../*.bat), or the name of a tag set
+		/// List of file specifications separated by semicolons (for example, *.cpp;Engine/.../*.bat), or the name of a tag set.
 		/// </summary>
 		[TaskParameter(ValidationType = TaskParameterValidationType.FileSpec)]
 		public string Files;
 
 		/// <summary>
-		/// Tag to be applied to build products of this task
+		/// Tag to be applied to build products of this task.
 		/// </summary>
 		[TaskParameter(Optional = true, ValidationType = TaskParameterValidationType.TagList)]
 		public string Tag;
diff --git a/Engine/Source/Programs/AutomationTool/BuildGraph/Tasks/SpawnTask.cs b/Engine/Source/Programs/AutomationTool/BuildGraph/Tasks/SpawnTask.cs
--- a/Engine/Source/Programs/AutomationTool/BuildGraph/Tasks/SpawnTask.cs
+++ b/Engine/Source/Programs/AutomationTool/BuildGraph/Tasks/SpawnTask.cs
@@ -19,19 +19,19 @@ namespace BuildGraph.Tasks
 	public class SpawnTaskParameters
 	{
 		/// <summary>
-		/// Executable to spawn
+		/// Executable to spawn.
 		/// </summary>
 		[TaskParameter]
 		public string Exe;
 
 		/// <summary>
-		/// Arguments for the newly created process
+		/// Arguments for the newly created process.
 		/// </summary>
 		[TaskParameter(Optional = true)]
 		public string Arguments;
 
 		/// <summary>
-		/// The minimum exit code which is treated as an error.
+		/// The minimum exit code, which is treated as an error.
 		/// </summary>
 		[TaskParameter(Optional = true)]
 		public int ErrorLevel = 1;
diff --git a/Engine/Source/Programs/AutomationTool/BuildGraph/Tasks/StageTask.cs b/Engine/Source/Programs/AutomationTool/BuildGraph/Tasks/StageTask.cs
--- a/Engine/Source/Programs/AutomationTool/BuildGraph/Tasks/StageTask.cs
+++ b/Engine/Source/Programs/AutomationTool/BuildGraph/Tasks/StageTask.cs
@@ -18,49 +18,49 @@ namespace BuildGraph.Tasks
 	public class StageTaskParameters
 	{
 		/// <summary>
-		/// The project that this target belongs to
+		/// The project that this target belongs to.
 		/// </summary>
 		[TaskParameter(Optional = true)]
 		public FileReference Project;
 
 		/// <summary>
-		/// Name of the target to stage
+		/// Name of the target to stage.
 		/// </summary>
 		[TaskParameter]
 		public string Target;
 
 		/// <summary>
-		/// Platform to stage
+		/// Platform to stage.
 		/// </summary>
 		[TaskParameter]
 		public UnrealTargetPlatform Platform;
 
 		/// <summary>
-		/// Configuration to be staged
+		/// Configuration to be staged.
 		/// </summary>
 		[TaskParameter]
 		public UnrealTargetConfiguration Configuration;
 
 		/// <summary>
-		/// Architecture to be staged
+		/// Architecture to be staged.
 		/// </summary>
 		[TaskParameter(Optional = true)]
 		public string Architecture;
 
 		/// <summary>
-		/// Directory the receipt files should be staged to
+		/// Directory that the receipt files should be staged to.
 		/// </summary>
 		[TaskParameter]
 		public DirectoryReference ToDir;
 
 		/// <summary>
-		/// Whether to overwrite existing files
+		/// Whether to overwrite existing files.
 		/// </summary>
 		[TaskParameter(Optional = true)]
 		public bool Overwrite;
 
 		/// <summary>
-		/// Tag to be applied to build products of this task
+		/// Tag to be applied to build products of this task.
 		/// </summary>
 		[TaskParameter(Optional = true, ValidationType = TaskParameterValidationType.TagList)]
 		public string Tag;
diff --git a/Engine/Source/Programs/AutomationTool/BuildGraph/Tasks/StripTask.cs b/Engine/Source/Programs/AutomationTool/BuildGraph/Tasks/StripTask.cs
--- a/Engine/Source/Programs/AutomationTool/BuildGraph/Tasks/StripTask.cs
+++ b/Engine/Source/Programs/AutomationTool/BuildGraph/Tasks/StripTask.cs
@@ -18,31 +18,31 @@ namespace BuildGraph.Tasks
 	public class StripTaskParameters
 	{
 		/// <summary>
-		/// The platform toolchain to strip binaries
+		/// The platform toolchain to strip binaries.
 		/// </summary>
 		[TaskParameter]
 		public UnrealTargetPlatform Platform;
 
 		/// <summary>
-		/// The directory to find files in
+		/// The directory to find files in.
 		/// </summary>
 		[TaskParameter(Optional = true)]
 		public DirectoryReference BaseDir;
 
 		/// <summary>
-		/// List of file specifications separated by semicolons (eg. Engine/.../*.pdb), or the name of a tag set
+		/// List of file specifications separated by semicolons (for example, Engine/.../*.pdb), or the name of a tag set.
 		/// </summary>
 		[TaskParameter(ValidationType = TaskParameterValidationType.FileSpec)]
 		public string Files;
 
 		/// <summary>
-		/// Output directory for the stripped files. Defaults to the input path (overwriting the input files).
+		/// Output directory for the stripped files. Defaults to the input path, overwriting the input files.
 		/// </summary>
 		[TaskParameter(Optional = true)]
 		public DirectoryReference OutputDir;
 
 		/// <summary>
-		/// Tag to be applied to build products of this task
+		/// Tag to be applied to build products of this task.
 		/// </summary>
 		[TaskParameter(Optional = true, ValidationType = TaskParameterValidationType.TagList)]
 		public string Tag;
diff --git a/Engine/Source/Programs/AutomationTool/BuildGraph/Tasks/SubmitTask.cs b/Engine/Source/Programs/AutomationTool/BuildGraph/Tasks/SubmitTask.cs
--- a/Engine/Source/Programs/AutomationTool/BuildGraph/Tasks/SubmitTask.cs
+++ b/Engine/Source/Programs/AutomationTool/BuildGraph/Tasks/SubmitTask.cs
@@ -17,19 +17,19 @@ namespace AutomationTool.Tasks
 	public class SubmitTaskParameters
 	{
 		/// <summary>
-		/// The description for the submitted changelist
+		/// The description for the submitted changelist.
 		/// </summary>
 		[TaskParameter]
 		public string Description;
 
 		/// <summary>
-		/// The files to submit
+		/// The files to submit.
 		/// </summary>
 		[TaskParameter(ValidationType = TaskParameterValidationType.FileSpec)]
 		public string Files;
 
 		/// <summary>
-		/// The Perforce file type for the submitted files (eg. binary+FS32)
+		/// The Perforce file type for the submitted files (for example, binary+FS32).
 		/// </summary>
 		[TaskParameter(Optional = true)]
 		public string FileType;
@@ -41,7 +41,7 @@ namespace AutomationTool.Tasks
 		public string Workspace;
 
 		/// <summary>
-		/// The stream for the workspace; defaults to the current stream. Ignored unless the Workspace attribute is also specified.
+		/// The stream for the workspace -- defaults to the current stream. Ignored unless the Workspace attribute is also specified.
 		/// </summary>
 		[TaskParameter(Optional = true)]
 		public string Stream;
@@ -53,13 +53,13 @@ namespace AutomationTool.Tasks
 		public DirectoryReference RootDir;
 
 		/// <summary>
-		/// Whether to revert unchanged files before attempting to submit
+		/// Whether to revert unchanged files before attempting to submit.
 		/// </summary>
 		[TaskParameter(Optional = true)]
 		public bool RevertUnchanged;
 
 		/// <summary>
-		/// Force the submit to happen, even if a resolve is needed (always accept current version)
+		/// Force the submit to happen -- even if a resolve is needed (always accept current version).
 		/// </summary>
 		[TaskParameter(Optional = true)]
 		public bool Force;
diff --git a/Engine/Source/Programs/AutomationTool/BuildGraph/Tasks/SymStoreTask.cs b/Engine/Source/Programs/AutomationTool/BuildGraph/Tasks/SymStoreTask.cs
--- a/Engine/Source/Programs/AutomationTool/BuildGraph/Tasks/SymStoreTask.cs
+++ b/Engine/Source/Programs/AutomationTool/BuildGraph/Tasks/SymStoreTask.cs
@@ -1,4 +1,4 @@
-﻿// Copyright 1998-2019 Epic Games, Inc. All Rights Reserved.
+// Copyright 1998-2019 Epic Games, Inc. All Rights Reserved.
 
 using System;
 using System.Collections.Generic;
@@ -42,14 +42,14 @@ namespace AutomationTool
         public string Product;
 
 		/// <summary>
-		/// BuildVersion associated with these symbols. Used for cleanup in AgeStore by matching this version against a directory name in a build share
+		/// BuildVersion associated with these symbols. Used for clean-up in AgeStore by matching this version against a directory name in a build share.
 		/// </summary>
 		[TaskParameter(Optional = true)]
 		public string BuildVersion;
 	}
 
     /// <summary>
-    /// Task which strips symbols from a set of files
+    /// Task that strips symbols from a set of files.
     /// </summary>
     [TaskElement("SymStore", typeof(SymStoreTaskParameters))]
     public class SymStoreTask : CustomTask
diff --git a/Engine/Source/Programs/AutomationTool/BuildGraph/Tasks/TagReceiptTask.cs b/Engine/Source/Programs/AutomationTool/BuildGraph/Tasks/TagReceiptTask.cs
--- a/Engine/Source/Programs/AutomationTool/BuildGraph/Tasks/TagReceiptTask.cs
+++ b/Engine/Source/Programs/AutomationTool/BuildGraph/Tasks/TagReceiptTask.cs
@@ -30,44 +30,44 @@ namespace AutomationTool.Tasks
 		public DirectoryReference EngineDir;
 
 		/// <summary>
-		/// Path to the project folder, used to expand $(ProjectDir) properties in receipt files. Defaults to the Engine directory for the current workspace. DEPRECATED.
+		/// Path to the project folder, used to expand $(ProjectDir) properties in receipt files. Defaults to the Engine directory for the current workspace -- DEPRECATED.
 		/// </summary>
 		[TaskParameter(Optional = true)]
 		public DirectoryReference ProjectDir;
 
 		/// <summary>
-		/// Whether to tag the Build Products listed in receipts
+		/// Whether to tag the Build Products listed in receipts.
 		/// </summary>
 		[TaskParameter(Optional = true)]
 		public bool BuildProducts;
 
 		/// <summary>
-		/// Which type of Build Products to tag (See TargetReceipt.cs - UnrealBuildTool.BuildProductType for valid values)
+		/// Which type of Build Products to tag (see TargetReceipt.cs - UnrealBuildTool.BuildProductType for valid values).
 		/// </summary>
 		[TaskParameter(Optional = true)]
 		public string BuildProductType;
 
 		/// <summary>
-		/// Whether to tag the Runtime Dependencies listed in receipts
+		/// Whether to tag the Runtime Dependencies listed in receipts.
 		/// </summary>
 		[TaskParameter(Optional = true)]
 		public bool RuntimeDependencies;
 
 		/// <summary>
-		/// Which type of Runtime Dependencies to tag (See TargetReceipt.cs - UnrealBuildTool.StagedFileType for valid values)
+		/// Which type of Runtime Dependencies to tag (see TargetReceipt.cs - UnrealBuildTool.StagedFileType for valid values).
 		/// </summary>
 		[TaskParameter(Optional = true)]
 		public string StagedFileType;
 
 		/// <summary>
-		/// Name of the tag to apply
+		/// Name of the tag to apply.
 		/// </summary>
 		[TaskParameter(ValidationType = TaskParameterValidationType.TagList)]
 		public string With;
 	}
 
 	/// <summary>
-	/// Task which tags build products and/or runtime dependencies by reading from *.target files.
+	/// Task that tags build products and/or runtime dependencies by reading from *.target files.
 	/// </summary>
 	[TaskElement("TagReceipt", typeof(TagReceiptTaskParameters))]
 	class TagReceiptTask : CustomTask
diff --git a/Engine/Source/Programs/AutomationTool/BuildGraph/Tasks/TagTask.cs b/Engine/Source/Programs/AutomationTool/BuildGraph/Tasks/TagTask.cs
--- a/Engine/Source/Programs/AutomationTool/BuildGraph/Tasks/TagTask.cs
+++ b/Engine/Source/Programs/AutomationTool/BuildGraph/Tasks/TagTask.cs
@@ -18,25 +18,25 @@ namespace AutomationTool.Tasks
 	public class TagTaskParameters
 	{
 		/// <summary>
-		/// Set the base directory to resolve relative paths and patterns against. If set, any absolute patterns (eg. /Engine/Build/...) are taken to be relative to this path. If not, they are taken to be truly absolute.
+		/// Set the base directory to resolve relative paths and patterns against. If set, any absolute patterns (for example, /Engine/Build/...) are taken to be relative to this path. If not, they are taken to be truly absolute.
 		/// </summary>
 		[TaskParameter(Optional = true)]
 		public DirectoryReference BaseDir;
 
 		/// <summary>
-		/// Set of files to work from, including wildcards and tag names, separated by semicolons. Resolved relative to BaseDir if set, otherwise to the branch root directory.
+		/// Set of files to work from, including wildcards and tag names, separated by semicolons. If set, resolved relative to BaseDir, otherwise resolved to the branch root directory.
 		/// </summary>
 		[TaskParameter(Optional = true, ValidationType = TaskParameterValidationType.FileSpec)]
 		public string Files;
 
 		/// <summary>
-		/// Set of text files to add additional files from. Each file list should have on file per line.
+		/// Set of text files to add additional files from. Each file list should have one file per line.
 		/// </summary>
 		[TaskParameter(Optional = true, ValidationType = TaskParameterValidationType.FileSpec)]
 		public string FileLists;
 
 		/// <summary>
-		/// Patterns to filter the list of files by, including tag names or wildcards. May include patterns that apply to the base directory if set. Defaults to all files if not specified.
+		/// Patterns to filter the list of files by, including tag names or wildcards. If set, may include patterns that apply to the base directory. If not specified, defaults to all files.
 		/// </summary>
 		[TaskParameter(Optional = true, ValidationType = TaskParameterValidationType.FileSpec)]
 		public string Filter;
@@ -48,7 +48,7 @@ namespace AutomationTool.Tasks
 		public string Except;
 
 		/// <summary>
-		/// Name of the tag to apply
+		/// Name of the tag to apply.
 		/// </summary>
 		[TaskParameter(ValidationType = TaskParameterValidationType.TagList)]
 		public string With;
diff --git a/Engine/Source/Programs/AutomationTool/BuildGraph/Tasks/UnzipTask.cs b/Engine/Source/Programs/AutomationTool/BuildGraph/Tasks/UnzipTask.cs
--- a/Engine/Source/Programs/AutomationTool/BuildGraph/Tasks/UnzipTask.cs
+++ b/Engine/Source/Programs/AutomationTool/BuildGraph/Tasks/UnzipTask.cs
@@ -18,19 +18,19 @@ namespace BuildGraph.Tasks
 	public class UnzipTaskParameters
 	{
 		/// <summary>
-		/// Path to the zip file to extract
+		/// Path to the zip file to extract.
 		/// </summary>
 		[TaskParameter(ValidationType = TaskParameterValidationType.FileSpec)]
 		public string ZipFile;
 
 		/// <summary>
-		/// Output directory for the extracted files
+		/// Output directory for the extracted files.
 		/// </summary>
 		[TaskParameter]
 		public DirectoryReference ToDir;
 
 		/// <summary>
-		/// Tag to be applied to the extracted files
+		/// Tag to be applied to the extracted files.
 		/// </summary>
 		[TaskParameter(Optional = true, ValidationType = TaskParameterValidationType.TagList)]
 		public string Tag;
diff --git a/Engine/Source/Programs/AutomationTool/BuildGraph/Tasks/ZipTask.cs b/Engine/Source/Programs/AutomationTool/BuildGraph/Tasks/ZipTask.cs
--- a/Engine/Source/Programs/AutomationTool/BuildGraph/Tasks/ZipTask.cs
+++ b/Engine/Source/Programs/AutomationTool/BuildGraph/Tasks/ZipTask.cs
@@ -18,25 +18,25 @@ namespace BuildGraph.Tasks
 	public class ZipTaskParameters
 	{
 		/// <summary>
-		/// The directory to read compressed files from
+		/// The directory to read compressed files from.
 		/// </summary>
 		[TaskParameter]
 		public DirectoryReference FromDir;
 
 		/// <summary>
-		/// List of file specifications separated by semicolons (eg. *.cpp;Engine/.../*.bat), or the name of a tag set. Relative paths are taken from FromDir.
+		/// List of file specifications separated by semicolons (for example, *.cpp;Engine/.../*.bat), or the name of a tag set. Relative paths are taken from FromDir.
 		/// </summary>
 		[TaskParameter(Optional = true, ValidationType = TaskParameterValidationType.FileSpec)]
 		public string Files;
 
 		/// <summary>
-		/// The zip file to create
+		/// The zip file to create.
 		/// </summary>
 		[TaskParameter]
 		public FileReference ZipFile;
 
 		/// <summary>
-		/// Tag to be applied to the created zip file
+		/// Tag to be applied to the created zip file.
 		/// </summary>
 		[TaskParameter(Optional = true, ValidationType = TaskParameterValidationType.TagList)]
 		public string Tag;
diff --git a/Engine/Source/Programs/AutomationTool/HoloLens/HoloLens.Automation.csproj b/Engine/Source/Programs/AutomationTool/HoloLens/HoloLens.Automation.csproj
--- a/Engine/Source/Programs/AutomationTool/HoloLens/HoloLens.Automation.csproj
+++ b/Engine/Source/Programs/AutomationTool/HoloLens/HoloLens.Automation.csproj
@@ -72,12 +72,15 @@
 		<Reference Include="Windows">
 		  <!-- This is a dummy entry to throw off the parser in DotNETUtilities's CsProjectInfo which doesn't understand MSBuild functions -->
 		  <HintPath />
+		  <!-- Use the version from the current Windows SDK -->
 		  <HintPath>$(Windows10SdkPath)\UnionMetadata\$(Windows10SdkLatestVersion)\Windows.winmd</HintPath>
-		  <!-- 16299 SDK (possibly others?) has ProductVersion == 10.0.16299 in the registry, but pathing needs the final .0 -->
+          <!-- Try using the version from AutoSDK -->
+          <HintPath Condition="'$(UE_SDKS_ROOT)' != '' And Exists('$(UE_SDKS_ROOT)\HostWin64\Win64\Windows Kits\10\UnionMetadata\10.0.18362.0\Windows.winmd')">$(UE_SDKS_ROOT)\HostWin64\Win64\Windows Kits\10\UnionMetadata\10.0.18362.0\Windows.winmd</HintPath>
+          <!-- 16299 SDK (possibly others?) has ProductVersion == 10.0.16299 in the registry, but pathing needs the final .0 -->
 		  <HintPath Condition="'$([System.IO.File]::Exists(%(HintPath)))'=='false'">$(Windows10SdkPath)\UnionMetadata\$(Windows10SdkLatestVersion).0\Windows.winmd</HintPath>
 		  <!-- Fall back to old, unversioned location -->
 		  <HintPath Condition="'$([System.IO.File]::Exists(%(HintPath)))'=='false'">$(Windows10SdkPath)\UnionMetadata\Windows.winmd</HintPath>
-		  <Private>false</Private>
+          <Private>false</Private>
 		</Reference>
 	  </ItemGroup>
 	</When>
diff --git a/Engine/Source/Programs/AutomationTool/Scripts/CopyBuildToStagingDirectory.Automation.cs b/Engine/Source/Programs/AutomationTool/Scripts/CopyBuildToStagingDirectory.Automation.cs
--- a/Engine/Source/Programs/AutomationTool/Scripts/CopyBuildToStagingDirectory.Automation.cs
+++ b/Engine/Source/Programs/AutomationTool/Scripts/CopyBuildToStagingDirectory.Automation.cs
@@ -553,7 +553,7 @@ public partial class Project : CommandUtils
 			{
 				// Work out which ICU data version we use for this platform
 				var ICUDataVersion = "icudt64l";
-				if (SC.StageTargetPlatform.PlatformType == UnrealTargetPlatform.HoloLens || SC.StageTargetPlatform.PlatformType == UnrealTargetPlatform.TVOS || SC.StageTargetPlatform.PlatformType == UnrealTargetPlatform.HTML5)
+				if (SC.StageTargetPlatform.PlatformType == UnrealTargetPlatform.HoloLens || SC.StageTargetPlatform.PlatformType == UnrealTargetPlatform.TVOS || SC.StageTargetPlatform.PlatformType.ToString() == "HTML5")
 				{
 					ICUDataVersion = "icudt53l";
 				}
diff --git a/Engine/Source/Programs/Enterprise/Datasmith/DatasmithRevitExporter/Private/DatasmithRevitExportContext.cs b/Engine/Source/Programs/Enterprise/Datasmith/DatasmithRevitExporter/Private/DatasmithRevitExportContext.cs
--- a/Engine/Source/Programs/Enterprise/Datasmith/DatasmithRevitExporter/Private/DatasmithRevitExportContext.cs
+++ b/Engine/Source/Programs/Enterprise/Datasmith/DatasmithRevitExporter/Private/DatasmithRevitExportContext.cs
@@ -719,10 +719,6 @@ namespace DatasmithRevitExporter
 				InstanceData.InstanceActor.SetIsComponent(true);
 			}
 
-                // The Datasmith instance actor is a component in the hierarchy.
-                //InstanceData.InstanceActor.SetIsComponent(true);
-            }
-
             public FDatasmithFacadeMesh PopInstance()
 			{
 				FInstanceData InstanceData = InstanceDataStack.Pop();
diff --git a/Engine/Source/Programs/UnrealBuildTool/Configuration/BuildConfiguration.cs b/Engine/Source/Programs/UnrealBuildTool/Configuration/BuildConfiguration.cs
--- a/Engine/Source/Programs/UnrealBuildTool/Configuration/BuildConfiguration.cs
+++ b/Engine/Source/Programs/UnrealBuildTool/Configuration/BuildConfiguration.cs
@@ -14,13 +14,13 @@ namespace UnrealBuildTool
 	/// </summary>
 	class BuildConfiguration
 	{
-		/// <summary>
-		/// Whether to ignore import library files that are out of date when building targets.  Set this to true to improve iteration time.
-		/// By default we don't bother relinking targets if only a dependent .lib has changed, as chances are that
-		/// the import library wasn't actually different unless a dependent header file of this target was also changed,
-		/// in which case the target would be rebuilt automatically.
-		/// </summary>
-		[XmlConfigFile]
+        /// <summary>
+        /// Whether to ignore import library files that are out of date when building targets. Set this to true to improve iteration time.
+        /// By default, we do not bother re-linking targets if only a dependent .lib has changed, as chances are that
+        /// the import library was not actually different unless a dependent header file of this target was also changed,
+        /// in which case the target would automatically be rebuilt.
+        /// </summary>
+        [XmlConfigFile]
 		public bool bIgnoreOutdatedImportLibraries = true;
 
 		/// <summary>
@@ -42,7 +42,7 @@ namespace UnrealBuildTool
 		public bool bLogDetailedActionStats = false;
 
 		/// <summary>
-		/// Whether the hybrid executor will be used (a remote executor and local executor)
+		/// Whether the hybrid executor will be used (a remote executor and local executor).
 		/// </summary>
 		[XmlConfigFile]
 		public bool bAllowHybridExecutor = false;
@@ -61,25 +61,25 @@ namespace UnrealBuildTool
 		public bool bAllowSNDBS = true;
 
 		/// <summary>
-		/// Enables support for very fast iterative builds by caching target data.  Turning this on causes Unreal Build Tool to emit
-		/// 'UBT Makefiles' for targets when they're built the first time.  Subsequent builds will load these Makefiles and begin
-		/// outdatedness checking and build invocation very quickly.  The caveat is that if source files are added or removed to
-		/// the project, UBT will need to gather information about those in order for your build to complete successfully.  Currently,
+		/// Enables support for very fast iterative builds by caching target data. Turning this on causes Unreal Build Tool to emit
+		/// 'UBT Makefiles' for targets when they are built the first time. Subsequent builds will load these Makefiles and begin
+		/// outdatedness checking and build invocation very quickly. The caveat is that if source files are added or removed to
+		/// the project, UBT will need to gather information about those in order for your build to complete successfully. Currently,
 		/// you must run the project file generator after adding/removing source files to tell UBT to re-gather this information.
 		/// 
 		/// Events that can invalidate the 'UBT Makefile':  
 		///		- Adding/removing .cpp files
 		///		- Adding/removing .h files with UObjects
-		///		- Adding new UObject types to a file that didn't previously have any
-		///		- Changing global build settings (most settings in this file qualify.)
+		///		- Adding new UObject types to a file that did not previously have any
+		///		- Changing global build settings (most settings in this file qualify)
 		///		- Changed code that affects how Unreal Header Tool works
 		///	
-		///	You can force regeneration of the 'UBT Makefile' by passing the '-gather' argument, or simply regenerating project files
+		///	You can force regeneration of the 'UBT Makefile' by passing the '-gather' argument, or simply regenerating project files.
 		///
 		///	This also enables the fast include file dependency scanning and caching system that allows Unreal Build Tool to detect out 
-		/// of date dependencies very quickly.  When enabled, a deep C++ include graph does not have to be generated, and instead
-		/// we only scan and cache indirect includes for after a dependent build product was already found to be out of date.  During the
-		/// next build, we'll load those cached indirect includes and check for outdatedness.
+		/// of date dependencies very quickly. When enabled, a deep C++ include graph does not have to be generated, and instead,
+		/// we only scan and cache indirect includes for after a dependent build product was already found to be out of date. During the
+		/// next build, we will load those cached indirect includes and check for outdatedness.
 		/// </summary>
 		[XmlConfigFile]
 		[CommandLine("-NoUBTMakefiles", Value = "false")]
@@ -87,7 +87,7 @@ namespace UnrealBuildTool
 
 		/// <summary>
 		/// Whether DMUCS/Distcc may be used.
-		/// Distcc requires some setup - so by default disable it so we don't break local or remote building
+		/// Distcc requires some setup -- so by default, disable it so that we do not break local or remote building.
 		/// </summary>
 		[XmlConfigFile]
 		public bool bAllowDistcc = false;
@@ -99,14 +99,14 @@ namespace UnrealBuildTool
 		public bool bAllowParallelExecutor = true;
 
 		/// <summary>
-		/// If true, force header regeneration. Intended for the build machine
+		/// If true, force header regeneration. Intended for the build machine.
 		/// </summary>
 		[CommandLine("-ForceHeaderGeneration")]
 		[XmlConfigFile(Category = "UEBuildConfiguration")]
 		public bool bForceHeaderGeneration = false;
 
 		/// <summary>
-		/// If true, do not build UHT, assume it is already built
+		/// If true, do not build UHT, assume it is already built.
 		/// </summary>
 		[CommandLine("-NoBuildUHT")]
 		[XmlConfigFile(Category = "UEBuildConfiguration")]
@@ -120,14 +120,14 @@ namespace UnrealBuildTool
 		public bool bFailIfGeneratedCodeChanges = false;
 
 		/// <summary>
-		/// True if hot-reload from IDE is allowed
+		/// True if hot-reload from IDE is allowed.
 		/// </summary>
 		[CommandLine("-NoHotReloadFromIDE", Value="false")]
 		[XmlConfigFile(Category = "UEBuildConfiguration")]
 		public bool bAllowHotReloadFromIDE = true;
 
 		/// <summary>
-		/// If true, the Debug version of UnrealHeaderTool will be build and run instead of the Development version.
+		/// If true, the Debug version of UnrealHeaderTool will be built and run instead of the Development version.
 		/// </summary>
 		[XmlConfigFile(Category = "UEBuildConfiguration")]
 		public bool bForceDebugUnrealHeaderTool = false;
diff --git a/Engine/Source/Programs/UnrealBuildTool/Executors/Distcc.cs b/Engine/Source/Programs/UnrealBuildTool/Executors/Distcc.cs
--- a/Engine/Source/Programs/UnrealBuildTool/Executors/Distcc.cs
+++ b/Engine/Source/Programs/UnrealBuildTool/Executors/Distcc.cs
@@ -15,7 +15,7 @@ namespace UnrealBuildTool
 	class Distcc : ActionExecutor
 	{
 		/// <summary>
-		/// When enabled allows DMUCS/Distcc to fallback to local compilation when remote compiling fails. Defaults to true as separation of pre-process and compile stages can introduce non-fatal errors.
+		/// When enabled, allows DMUCS/Distcc to fallback to local compilation when remote compiling fails. Defaults to true as separation of pre-process and compile stages can introduce non-fatal errors.
 		/// </summary>
 		[XmlConfigFile(Category = "BuildConfiguration")]
 		public bool bAllowDistccLocalFallback = true;
@@ -33,7 +33,7 @@ namespace UnrealBuildTool
 		public string DistccExecutablesPath = "/usr/local/bin";
 
 		/// <summary>
-		/// DMUCS coordinator hostname or IP address.
+		/// The DMUCS coordinator hostname or IP address.
 		/// </summary>
 		[XmlConfigFile(Category = "BuildConfiguration")]
 		public string DMUCSCoordinator = "localhost";
diff --git a/Engine/Source/Programs/UnrealBuildTool/Executors/HybridExecutor.cs b/Engine/Source/Programs/UnrealBuildTool/Executors/HybridExecutor.cs
--- a/Engine/Source/Programs/UnrealBuildTool/Executors/HybridExecutor.cs
+++ b/Engine/Source/Programs/UnrealBuildTool/Executors/HybridExecutor.cs
@@ -15,7 +15,7 @@ namespace UnrealBuildTool
 	class HybridExecutor : ActionExecutor
 	{
 		/// <summary>
-		/// Maximum number of actions to execute locally
+		/// Maximum number of actions to execute locally.
 		/// </summary>
 		[XmlConfigFile(Category = "HybridExecutor")]
 		int MaxLocalActions;
diff --git a/Engine/Source/Programs/UnrealBuildTool/Modes/BuildMode.cs b/Engine/Source/Programs/UnrealBuildTool/Modes/BuildMode.cs
--- a/Engine/Source/Programs/UnrealBuildTool/Modes/BuildMode.cs
+++ b/Engine/Source/Programs/UnrealBuildTool/Modes/BuildMode.cs
@@ -41,13 +41,13 @@ namespace UnrealBuildTool
 	class BuildMode : ToolMode
 	{
 		/// <summary>
-		/// Specifies the file to use for logging
+		/// Specifies the file to use for logging.
 		/// </summary>
 		[XmlConfigFile(Category = "BuildConfiguration")]
 		public string BaseLogFileName = "../Programs/UnrealBuildTool/Log.txt";
 
 		/// <summary>
-		/// Whether to skip checking for files identified by the junk manifest
+		/// Whether to skip checking for files identified by the junk manifest.
 		/// </summary>
 		[XmlConfigFile]
 		[CommandLine("-IgnoreJunk")]
diff --git a/Engine/Source/Programs/UnrealBuildTool/Modes/GenerateProjectFilesMode.cs b/Engine/Source/Programs/UnrealBuildTool/Modes/GenerateProjectFilesMode.cs
--- a/Engine/Source/Programs/UnrealBuildTool/Modes/GenerateProjectFilesMode.cs
+++ b/Engine/Source/Programs/UnrealBuildTool/Modes/GenerateProjectFilesMode.cs
@@ -40,7 +40,7 @@ namespace UnrealBuildTool
 		HashSet<ProjectFileFormat> ProjectFileFormats = new HashSet<ProjectFileFormat>();
 
 		/// <summary>
-		/// Platforms to disable native project file generators for. Platforms with native project file generators typically require IDE extensions to be installed.
+		/// Disable native project file generators for platforms. Platforms with native project file generators typically require IDE extensions to be installed.
 		/// </summary>
 		[XmlConfigFile(Category = "ProjectFileGenerator")]
 		string[] DisablePlatformProjectGenerators = null;
diff --git a/Engine/Source/Programs/UnrealBuildTool/Platform/IOS/IOSToolChain.cs b/Engine/Source/Programs/UnrealBuildTool/Platform/IOS/IOSToolChain.cs
--- a/Engine/Source/Programs/UnrealBuildTool/Platform/IOS/IOSToolChain.cs
+++ b/Engine/Source/Programs/UnrealBuildTool/Platform/IOS/IOSToolChain.cs
@@ -21,14 +21,14 @@ namespace UnrealBuildTool
 	class IOSToolChainSettings : AppleToolChainSettings
 	{
 		/// <summary>
-		/// Which version of the iOS SDK to target at build time
+		/// The version of the iOS SDK to target at build time.
 		/// </summary>
 		[XmlConfigFile(Category = "IOSToolChain")]
 		public string IOSSDKVersion = "latest";
 		public readonly float IOSSDKVersionFloat = 0.0f;
 
 		/// <summary>
-		/// Which version of the iOS to allow at build time
+		/// The version of the iOS to allow at build time.
 		/// </summary>
 		[XmlConfigFile(Category = "IOSToolChain")]
 		public string BuildIOSVersion = "7.0";
@@ -99,7 +99,7 @@ namespace UnrealBuildTool
 		// If you are looking for where to change the remote compile server name, look in RemoteToolChain.cs
 
 		/// <summary>
-		/// If this is set, then we don't do any post-compile steps except moving the executable into the proper spot on the Mac
+		/// If this is set, then we do not do any post-compile steps -- except moving the executable into the proper spot on Mac.
 		/// </summary>
 		[XmlConfigFile]
 		public static bool bUseDangerouslyFastMode = false;
diff --git a/Engine/Source/Programs/UnrealBuildTool/Platform/IOS/UEBuildIOS.cs b/Engine/Source/Programs/UnrealBuildTool/Platform/IOS/UEBuildIOS.cs
--- a/Engine/Source/Programs/UnrealBuildTool/Platform/IOS/UEBuildIOS.cs
+++ b/Engine/Source/Programs/UnrealBuildTool/Platform/IOS/UEBuildIOS.cs
@@ -18,7 +18,7 @@ namespace UnrealBuildTool
 	public partial class IOSTargetRules
 	{
 		/// <summary>
-		/// Whether to strip iOS symbols or not (implied by Shipping config)
+		/// Whether to strip iOS symbols or not (implied by Shipping config).
 		/// </summary>
 		[XmlConfigFile(Category = "BuildConfiguration")]
 		[CommandLine("-stripsymbols", Value = "true")]
diff --git a/Engine/Source/Programs/UnrealBuildTool/Platform/Mac/UEBuildMac.cs b/Engine/Source/Programs/UnrealBuildTool/Platform/Mac/UEBuildMac.cs
--- a/Engine/Source/Programs/UnrealBuildTool/Platform/Mac/UEBuildMac.cs
+++ b/Engine/Source/Programs/UnrealBuildTool/Platform/Mac/UEBuildMac.cs
@@ -16,28 +16,28 @@ namespace UnrealBuildTool
 	public class MacTargetRules
 	{
 		/// <summary>
-		/// Whether to generate dSYM files
-		/// Lists Architectures that you want to build
+		/// Whether to generate dSYM files.
+		/// Lists Architectures that you want to build.
 		/// </summary>
 		[XmlConfigFile(Category = "BuildConfiguration", Name = "bGeneratedSYMFile")]
 		public bool bGenerateDsymFile = true;
 
 		/// <summary>
-		/// Enables address sanitizer (ASan)
+		/// Enables address sanitizer (ASan).
 		/// </summary>
 		[CommandLine("-EnableASan")]
 		[XmlConfigFile(Category = "BuildConfiguration", Name = "bEnableAddressSanitizer")]
 		public bool bEnableAddressSanitizer = false;
 
 		/// <summary>
-		/// Enables thread sanitizer (TSan)
+		/// Enables thread sanitizer (TSan).
 		/// </summary>
 		[CommandLine("-EnableTSan")]
 		[XmlConfigFile(Category = "BuildConfiguration", Name = "bEnableThreadSanitizer")]
 		public bool bEnableThreadSanitizer = false;
 
 		/// <summary>
-		/// Enables undefined behavior sanitizer (UBSan)
+		/// Enables undefined behavior sanitizer (UBSan).
 		/// </summary>
 		[CommandLine("-EnableUBSan")]
 		[XmlConfigFile(Category = "BuildConfiguration", Name = "bEnableUndefinedBehaviorSanitizer")]
diff --git a/Engine/Source/Programs/UnrealBuildTool/Platform/Windows/UEBuildWindows.cs b/Engine/Source/Programs/UnrealBuildTool/Platform/Windows/UEBuildWindows.cs
--- a/Engine/Source/Programs/UnrealBuildTool/Platform/Windows/UEBuildWindows.cs
+++ b/Engine/Source/Programs/UnrealBuildTool/Platform/Windows/UEBuildWindows.cs
@@ -110,7 +110,7 @@ namespace UnrealBuildTool
 		TargetRules Target;
 
 		/// <summary>
-		/// Version of the compiler toolchain to use on Windows platform. A value of "default" will be changed to a specific version at UBT startup.
+		/// Version of the compiler toolchain to use on Windows platform. A value of "default" will be changed to a specific version at UBT start up.
 		/// </summary>
 		[ConfigFile(ConfigHierarchyType.Engine, "/Script/WindowsTargetPlatform.WindowsTargetSettings", "CompilerVersion")]
 		[XmlConfigFile(Category = "WindowsPlatform")]
@@ -125,15 +125,15 @@ namespace UnrealBuildTool
 		public WindowsArchitecture Architecture = WindowsArchitecture.x64;
 
 		/// <summary>
-		/// The specific toolchain version to use. This may be a specific version number (eg. "14.13.26128") or the string "Latest" to select the newest available version. By default, we use the
-		/// toolchain version indicated by WindowsPlatform.DefaultToolChainVersion if it is available, or the latest version otherwise.
+		/// The specific toolchain version to use. This may be a specific version number (for example, "14.13.26128"), or the string "Latest", to select the newest available version. By default, and if it is available, we use the
+		/// toolchain version indicated by WindowsPlatform.DefaultToolChainVersion (otherwise, we use the latest version).
 		/// </summary>
 		[XmlConfigFile(Category = "WindowsPlatform")]
 		public string CompilerVersion = null;
 
 		/// <summary>
-		/// The specific Windows SDK version to use. This may be a specific version number (eg. "8.1", "10.0", or "10.0.10150.0") or the string "Latest" to select the newest available version.
-		/// By default, we use the Windows SDK version indicated by WindowsPlatform.DefaultWindowsSdkVersion if it is available, or the latest version otherwise.
+		/// The specific Windows SDK version to use. This may be a specific version number (for example, "8.1", "10.0" or "10.0.10150.0"), or the string "Latest", to select the newest available version.
+		/// By default, and if it is available, we use the Windows SDK version indicated by WindowsPlatform.DefaultWindowsSdkVersion (otherwise, we use the latest version).
 		/// </summary>
 		[XmlConfigFile(Category = "WindowsPlatform")]
 		public string WindowsSdkVersion = null;
@@ -174,14 +174,14 @@ namespace UnrealBuildTool
 		public string ProductName;
 
 		/// <summary>
-		/// The static analyzer to use
+		/// The static analyzer to use.
 		/// </summary>
 		[XmlConfigFile(Category = "WindowsPlatform")]
 		[CommandLine("-StaticAnalyzer")]
 		public WindowsStaticAnalyzer StaticAnalyzer = WindowsStaticAnalyzer.None;
 
 		/// <summary>
-		/// Whether we should export a file containing .obj->source file mappings.
+		/// Whether we should export a file containing .obj to source file mappings.
 		/// </summary>
 		[XmlConfigFile]
 		[CommandLine("-ObjSrcMap")]
@@ -268,7 +268,7 @@ namespace UnrealBuildTool
 		public bool bOptimizeGlobalData = true;
 
 		/// <summary>
-		/// (Experimental) Appends the -ftime-trace argument to the command line for Clang, to output a JSON file containing a timeline for the compile. 
+		/// (Experimental) Appends the -ftime-trace argument to the command line for Clang to output a JSON file containing a timeline for the compile. 
 		/// See http://aras-p.info/blog/2019/01/16/time-trace-timeline-flame-chart-profiler-for-Clang/ for more info.
 		/// </summary>
 		[XmlConfigFile(Category = "WindowsPlatform")]
diff --git a/Engine/Source/Programs/UnrealBuildTool/ProjectFiles/ProjectFileGenerator.cs b/Engine/Source/Programs/UnrealBuildTool/ProjectFiles/ProjectFileGenerator.cs
--- a/Engine/Source/Programs/UnrealBuildTool/ProjectFiles/ProjectFileGenerator.cs
+++ b/Engine/Source/Programs/UnrealBuildTool/ProjectFiles/ProjectFileGenerator.cs
@@ -148,7 +148,7 @@ namespace UnrealBuildTool
 	static class ProjectFileGeneratorSettings
 	{
 		/// <summary>
-		/// Default list of project file formats to generate
+		/// Default list of project file formats to generate.
 		/// </summary>
 		[XmlConfigFile(Category = "ProjectFileGenerator", Name = "Format")]
 		public static string Format = null;
@@ -207,100 +207,100 @@ namespace UnrealBuildTool
 		protected string GameProjectName = null;
 
 		/// <summary>
-		/// Whether we should include configurations for "Test" and "Shipping" in generated projects (pass "-NoShippingConfigs" to disable this)
+		/// Whether we should include configurations for "Test" and "Shipping" in generated projects. Pass "-NoShippingConfigs" to disable this.
 		/// </summary>
 		[XmlConfigFile]
 		public static bool bIncludeTestAndShippingConfigs = true;
 
 		/// <summary>
-		/// True if intellisense data should be generated (takes a while longer)
+		/// True if intellisense data should be generated (takes a while longer).
 		/// </summary>
 		[XmlConfigFile]
 		bool bGenerateIntelliSenseData = true;
 
 		/// <summary>
-		/// True if we should include documentation in the generated projects
+		/// True if we should include documentation in the generated projects.
 		/// </summary>
 		[XmlConfigFile]
 		protected bool bIncludeDocumentation = false;
 
 		/// <summary>
-		/// True if all documentation languages should be included in generated projects, otherwise only "INT" will be included
+		/// True if all documentation languages should be included in generated projects, otherwise only INT files will be included.
 		/// </summary>
 		[XmlConfigFile]
 		bool bAllDocumentationLanguages = false;
 
 		/// <summary>
-		/// True if build targets should pass the -useprecompiled argument
+		/// True if build targets should pass the -useprecompiled argument.
 		/// </summary>
 		[XmlConfigFile]
 		public bool bUsePrecompiled = false;
 
 		/// <summary>
-		/// True if we should include engine source in the generated solution
+		/// True if we should include engine source in the generated solution.
 		/// </summary>
 		[XmlConfigFile]
 		protected bool bIncludeEngineSource = true;
 
 		/// <summary>
-		/// Whether to include enterprise source in the generated solution
+		/// Whether to include enterprise source in the generated solution.
 		/// </summary>
 		[XmlConfigFile]
 		protected bool bIncludeEnterpriseSource = true;
 
 		/// <summary>
-		/// True if shader source files should be included in the generated projects
+		/// True if shader source files should be included in generated projects.
 		/// </summary>
 		[XmlConfigFile]
 		protected bool bIncludeShaderSource = true;
 
 		/// <summary>
-		/// True if build system files should be included
+		/// True if build system files should be included.
 		/// </summary>
 		[XmlConfigFile]
 		bool bIncludeBuildSystemFiles = true;
 
 		/// <summary>
-		/// True if we should include config files (.ini files) in the generated project
+		/// True if we should include config (.ini) files in the generated project.
 		/// </summary>
 		[XmlConfigFile]
 		protected bool bIncludeConfigFiles = true;
 
 		/// <summary>
-		/// True if we should include localization files (.int/.kor/etc files) in the generated project
+		/// True if we should include localization files in the generated project.
 		/// </summary>
 		[XmlConfigFile]
 		bool bIncludeLocalizationFiles = false;
 
 		/// <summary>
-		/// True if we should include template files (.template files) in the generated project
+		/// True if we should include template files in the generated project.
 		/// </summary>
 		[XmlConfigFile]
 		protected bool bIncludeTemplateFiles = true;
 
 		/// <summary>
-		/// True if we should include program projects in the generated solution
+		/// True if we should include program projects in the generated solution.
 		/// </summary>
 		[XmlConfigFile]
 		protected bool IncludeEnginePrograms = true;
 
 		/// <summary>
-		/// Whether to include temporary targets generated by UAT to support content only projects with non-default settings
+		/// Whether to include temporary targets generated by UAT to support content only projects with non-default settings.
 		/// </summary>
 		[XmlConfigFile]
 		bool bIncludeTempTargets = false;
 
 		/// <summary>
-		/// True if we should include .NET Core projects in the generated solution
+		/// True if we should include .NET Core projects in the generated solution.
 		/// </summary>
 		[XmlConfigFile]
 		bool bIncludeDotNETCoreProjects = false;
 
-		/// <summary>
-		/// True if we should reflect "Source" sub-directories on disk in the master project as master project directories.
-		/// This arguably adds some visual clutter to the master project, but is truer to the on-disk file organization.
-		/// </summary>
-		[XmlConfigFile]
+        /// <summary>
+        /// True if we should reflect "Source" sub-directories on disk in the master project as master project directories.
+        /// This (arguably) adds some visual clutter to the master project but it is truer to the on-disk file organization.
+        /// </summary>
+        [XmlConfigFile]
 		bool bKeepSourceSubDirectories = true;
 
 		/// <summary>
@@ -353,13 +353,13 @@ namespace UnrealBuildTool
 		bool AllowDotNetCoreProjects = false;
 
 		/// <summary>
-		/// Name of the master project file (e.g. base file name for the solution file for Visual Studio, or the Xcode project file on Mac)
+		/// Name of the master project file -- for example, the base file name for the Visual Studio solution file, or the Xcode project file on Mac.
 		/// </summary>
 		[XmlConfigFile]
 		protected string MasterProjectName = "UE4";
 
 		/// <summary>
-		/// If true, sets the master project name according to the name of the folder it's in
+		/// If true, sets the master project name according to the name of the folder it is in.
 		/// </summary>
 		[XmlConfigFile]
 		protected bool bMasterProjectNameFromFolder = false;
diff --git a/Engine/Source/Programs/UnrealBuildTool/ProjectFiles/VisualStudio/VCProjectFileGenerator.cs b/Engine/Source/Programs/UnrealBuildTool/ProjectFiles/VisualStudio/VCProjectFileGenerator.cs
--- a/Engine/Source/Programs/UnrealBuildTool/ProjectFiles/VisualStudio/VCProjectFileGenerator.cs
+++ b/Engine/Source/Programs/UnrealBuildTool/ProjectFiles/VisualStudio/VCProjectFileGenerator.cs
@@ -58,19 +58,19 @@ namespace UnrealBuildTool
 		protected VCProjectFileFormat ProjectFileFormat = VCProjectFileFormat.Default;
 
 		/// <summary>
-		/// Whether to write a solution option (suo) file for the sln
+		/// Whether to write a solution option (suo) file for the sln.
 		/// </summary>
 		[XmlConfigFile(Category = "BuildConfiguration")]
 		protected bool bWriteSolutionOptionFile = true;
 
 		/// <summary>
-		/// Whether to add the -FastPDB option to build command lines by default
+		/// Whether to add the -FastPDB option to build command lines by default.
 		/// </summary>
 		[XmlConfigFile(Category = "BuildConfiguration")]
 		bool bAddFastPDBToProjects = false;
 
 		/// <summary>
-		/// Whether to generate per-file intellisense data
+		/// Whether to generate per-file intellisense data.
 		/// </summary>
 		[XmlConfigFile(Category = "BuildConfiguration")]
 		bool bUsePerFileIntellisense = false;
diff --git a/Engine/Source/Programs/UnrealBuildTool/ProjectFiles/VisualStudioCode/VSCodeProjectFileGenerator.cs b/Engine/Source/Programs/UnrealBuildTool/ProjectFiles/VisualStudioCode/VSCodeProjectFileGenerator.cs
--- a/Engine/Source/Programs/UnrealBuildTool/ProjectFiles/VisualStudioCode/VSCodeProjectFileGenerator.cs
+++ b/Engine/Source/Programs/UnrealBuildTool/ProjectFiles/VisualStudioCode/VSCodeProjectFileGenerator.cs
@@ -40,6 +40,9 @@ namespace UnrealBuildTool
 		private string FrameworkExecutableExtension;
 		private string FrameworkLibraryExtension = ".dll";
 
+		/// <summary>
+		/// Includes all files in the generated workspace.
+		/// </summary>
 		[XmlConfigFile(Name = "IncludeAllFiles")]
 		private bool IncludeAllFiles = false;
 
diff --git a/Engine/Source/Programs/UnrealBuildTool/System/SourceFileWorkingSet.cs b/Engine/Source/Programs/UnrealBuildTool/System/SourceFileWorkingSet.cs
--- a/Engine/Source/Programs/UnrealBuildTool/System/SourceFileWorkingSet.cs
+++ b/Engine/Source/Programs/UnrealBuildTool/System/SourceFileWorkingSet.cs
@@ -291,19 +291,19 @@ namespace UnrealBuildTool
 		};
 
 		/// <summary>
-		/// Sets the provider to use for determining the working set
+		/// Sets the provider to use for determining the working set.
 		/// </summary>
 		[XmlConfigFile]
 		static ProviderType Provider = ProviderType.Default;
 
 		/// <summary>
-		/// Sets the path to use for the repository. Interpreted relative to the UE root directory (ie. folder above the Engine folder) if relative.
+		/// Sets the path to use for the repository. Interpreted relative to the Unreal Engine root directory (the folder above the Engine folder) -- if relative.
 		/// </summary>
 		[XmlConfigFile]
 		public static string RepositoryPath = null;
 
 		/// <summary>
-		/// Sets the path to use for the Git executable. Defaults to "git" (assuming it's in the PATH).
+		/// Sets the path to use for the Git executable. Defaults to "git" (assuming it is in the PATH).
 		/// </summary>
 		[XmlConfigFile]
 		public static string GitPath = "git";
diff --git a/Engine/Source/Programs/UnrealBuildTool/ToolChain/RemoteMac.cs b/Engine/Source/Programs/UnrealBuildTool/ToolChain/RemoteMac.cs
--- a/Engine/Source/Programs/UnrealBuildTool/ToolChain/RemoteMac.cs
+++ b/Engine/Source/Programs/UnrealBuildTool/ToolChain/RemoteMac.cs
@@ -37,31 +37,31 @@ namespace UnrealBuildTool
 	class RemoteMac
 	{
 		/// <summary>
-		/// These two variables will be loaded from XML config file in XmlConfigLoader.Init()
+		/// These two variables will be loaded from the XML config file in XmlConfigLoader.Init().
 		/// </summary>
 		[XmlConfigFile]
 		private readonly string ServerName;
 
 		/// <summary>
-		/// The remote username
+		/// The remote username.
 		/// </summary>
 		[XmlConfigFile]
 		private readonly string UserName;
 
 		/// <summary>
-		/// Instead of looking for RemoteToolChainPrivate.key in the usual places (Documents/Unreal Engine/UnrealBuildTool/SSHKeys, Engine/Build/SSHKeys), this private key will be used if set
+		/// If set, instead of looking for RemoteToolChainPrivate.key in the usual places (Documents/Unreal, Engine/UnrealBuildTool/SSHKeys or Engine/Build/SSHKeys), this private key will be used.
 		/// </summary>
 		[XmlConfigFile]
 		private FileReference SshPrivateKey;
 
 		/// <summary>
-		/// The authentication used for Rsync (for the -e rsync flag)
+		/// The authentication used for Rsync (for the -e rsync flag).
 		/// </summary>
 		[XmlConfigFile]
 		private string RsyncAuthentication = "ssh -i '${CYGWIN_SSH_PRIVATE_KEY}'";
 
 		/// <summary>
-		/// The authentication used for SSH (probably similar to RsyncAuthentication)
+		/// The authentication used for SSH (probably similar to RsyncAuthentication).
 		/// </summary>
 		[XmlConfigFile]
 		private string SshAuthentication = "-i '${CYGWIN_SSH_PRIVATE_KEY}'";
diff --git a/Engine/Source/Programs/UnrealLightmass/Private/Lighting/AdaptiveVolumetricLightmap.cpp b/Engine/Source/Programs/UnrealLightmass/Private/Lighting/AdaptiveVolumetricLightmap.cpp
--- a/Engine/Source/Programs/UnrealLightmass/Private/Lighting/AdaptiveVolumetricLightmap.cpp
+++ b/Engine/Source/Programs/UnrealLightmass/Private/Lighting/AdaptiveVolumetricLightmap.cpp
@@ -150,15 +150,6 @@ bool FStaticLightingSystem::DoesVoxelIntersectSceneGeometry(const FBox& CellBoun
 
 			if (MeshInstance->StaticMesh->VoxelizationMesh != nullptr)
 			{
-				auto TransformBox = [](const FBox& Box, const FMatrix& Transform) -> const FBox {
-					FVector TransformedCorners[8];
-					for (int32 Corner = 0; Corner < 8; Corner++)
-					{
-						TransformedCorners[Corner] = Transform.TransformPosition(Box.GetExtrema(Corner));
-					}
-					return FBox(TransformedCorners, UE_ARRAY_COUNT(TransformedCorners));
-				};
-
 				if (MeshInstance->Mapping->GetVolumeMapping() == nullptr)
 				{
 					if (MeshInstance->StaticMesh->VoxelizationMesh->IntersectBox(ExpandedCellBoundsSurfaceGeometry.TransformBy(MeshInstance->WorldToLocal)) != nullptr)
diff --git a/Engine/Source/Runtime/AssetRegistry/Private/PackageReader.cpp b/Engine/Source/Runtime/AssetRegistry/Private/PackageReader.cpp
--- a/Engine/Source/Runtime/AssetRegistry/Private/PackageReader.cpp
+++ b/Engine/Source/Runtime/AssetRegistry/Private/PackageReader.cpp
@@ -83,17 +83,15 @@ bool FPackageReader::OpenPackageFile(EOpenPackageResult* OutErrorCode)
 	}
 
 	// Check serialized custom versions against latest custom versions.
-	const FCustomVersionContainer& LatestCustomVersions  = FCustomVersionContainer::GetRegistered();
-	const FCustomVersionArray&  PackageCustomVersions = PackageFileSummary.GetCustomVersionContainer().GetAllVersions();
-	for (const FCustomVersion& SerializedCustomVersion : PackageCustomVersions)
+	TArray<FCustomVersionDifference> Diffs = FCurrentCustomVersions::Compare(PackageFileSummary.GetCustomVersionContainer().GetAllVersions());
+	for (FCustomVersionDifference Diff : Diffs)
 	{
-		auto* LatestVersion = LatestCustomVersions.GetVersion(SerializedCustomVersion.Key);
-		if (!LatestVersion)
+		if (Diff.Type == ECustomVersionDifference::Missing)
 		{
 			SetPackageErrorCode(EOpenPackageResult::CustomVersionMissing);
 			return false;
 		}
-		else if (SerializedCustomVersion.Version > LatestVersion->Version)
+		else if (Diff.Type == ECustomVersionDifference::Newer)
 		{
 			SetPackageErrorCode(EOpenPackageResult::VersionTooNew);
 			return false;
diff --git a/Engine/Source/Runtime/AudioMixerCore/Public/AudioMixer.h b/Engine/Source/Runtime/AudioMixerCore/Public/AudioMixer.h
--- a/Engine/Source/Runtime/AudioMixerCore/Public/AudioMixer.h
+++ b/Engine/Source/Runtime/AudioMixerCore/Public/AudioMixer.h
@@ -7,6 +7,7 @@
 #include "AudioMixerTypes.h"
 #include "HAL/Runnable.h"
 #include "HAL/ThreadSafeBool.h"
+#include "Misc/ScopeLock.h"
 #include "Misc/SingleThreadRunnable.h"
 #include "AudioMixerNullDevice.h"
 #include "DSP/ParamInterpolator.h"
@@ -494,8 +495,19 @@ namespace Audio
 		/** Is called when an error is generated. */
 		inline void OnAudioMixerPlatformError(const FString& ErrorDetails, const FString& FileName, int32 LineNumber)
 		{
+#if !NO_LOGGING
+			// Log once on these errors to avoid Spam.
+			static FCriticalSection Cs;
+			static TSet<uint32> LogHistory;
+			FScopeLock Lock(&Cs);
 			LastError = FString::Printf(TEXT("Audio Platform Device Error: %s (File %s, Line %d)"), *ErrorDetails, *FileName, LineNumber);
-			UE_LOG(LogAudioMixer, Error, TEXT("%s"), *LastError);
+			uint32 Hash = GetTypeHash(LastError);
+			if (!LogHistory.Contains(Hash))
+			{
+				UE_LOG(LogAudioMixer, Error, TEXT("%s"), *LastError);
+				LogHistory.Add(Hash);
+			}
+#endif //!NO_LOGGING
 		}
 
 		/** Start generating audio from our mixer. */
diff --git a/Engine/Source/Runtime/Core/Private/Serialization/Archive.cpp b/Engine/Source/Runtime/Core/Private/Serialization/Archive.cpp
--- a/Engine/Source/Runtime/Core/Private/Serialization/Archive.cpp
+++ b/Engine/Source/Runtime/Core/Private/Serialization/Archive.cpp
@@ -441,7 +441,7 @@ const FCustomVersionContainer& FArchive::GetCustomVersions() const
 		// serialization code to use UsingCustomVersion to populate the container.
 		if (this->IsLoading())
 		{
-			*CustomVersionContainer = FCustomVersionContainer::GetRegistered();
+			*CustomVersionContainer = FCurrentCustomVersions::GetAll();
 		}
 		else
 		{
@@ -478,13 +478,8 @@ void FArchive::UsingCustomVersion(const FGuid& Key)
 		return;
 	}
 
-	auto* RegisteredVersion = FCustomVersionContainer::GetRegistered().GetVersion(Key);
-
-	// Ensure that the version has already been registered.
-	// If this fails, you probably don't have an FCustomVersionRegistration variable defined for this GUID.
-	check(RegisteredVersion);
-
-	const_cast<FCustomVersionContainer&>(GetCustomVersions()).SetVersion(Key, RegisteredVersion->Version, RegisteredVersion->GetFriendlyName());
+	FCustomVersion RegisteredVersion = FCurrentCustomVersions::Get(Key).GetValue();
+	const_cast<FCustomVersionContainer&>(GetCustomVersions()).SetVersion(Key, RegisteredVersion.Version, RegisteredVersion.GetFriendlyName());
 }
 
 int32 FArchive::CustomVer(const FGuid& Key) const
diff --git a/Engine/Source/Runtime/Core/Private/Serialization/CustomVersion.cpp b/Engine/Source/Runtime/Core/Private/Serialization/CustomVersion.cpp
--- a/Engine/Source/Runtime/Core/Private/Serialization/CustomVersion.cpp
+++ b/Engine/Source/Runtime/Core/Private/Serialization/CustomVersion.cpp
@@ -8,6 +8,8 @@
 #include "Serialization/StructuredArchive.h"
 #include "Algo/Sort.h"
 #include "Containers/Map.h"
+#include "Misc/ScopeRwLock.h"
+#include "CoreGlobals.h"
 
 namespace
 {
@@ -74,29 +76,8 @@ namespace
 }
 
 /** Defer FName creation and allocations from static FCustomVersionRegistrations that may never be needed */
-class FStaticCustomVersionRegistry
+struct FStaticCustomVersionRegistry
 {
-public:
-	static const FCustomVersionContainer& Get()
-	{
-		FStaticCustomVersionRegistry& StaticVersions = GetInstance();
-		StaticVersions.RegisterQueue();
-		return StaticVersions.Registered;
-	}
-
-	static void Register(FGuid Key, int32 Version, const TCHAR* Name)
-	{
-		RegistrationQueue& Queue = GetInstance().Queue;
-		check(Queue.Find(Key) == nullptr);
-		Queue.Add(Key, { Version, Name });
-	}
-
-	static void Unregister(FGuid Key)
-	{
-		 GetInstance().UnregisterImpl(Key);
-	}
-
-private:
 	struct FPendingRegistration
 	{
 		int32 Version;
@@ -104,14 +85,30 @@ private:
 	};
 	typedef TMap<FGuid, FPendingRegistration, TInlineSetAllocator<64>> RegistrationQueue;
 
+	FRWLock Lock;
 	FCustomVersionContainer Registered;
 	RegistrationQueue Queue;
 
-	static FStaticCustomVersionRegistry& GetInstance()
+	static FStaticCustomVersionRegistry& Get()
 	{
 		static FStaticCustomVersionRegistry Singleton;
 		return Singleton;
-	};
+	}
+
+	TOptional<FCustomVersion> Find(const FGuid& Guid) const
+	{
+		if (const FCustomVersion* RegisteredVersion = Registered.GetVersion(Guid))
+		{
+			return *RegisteredVersion;
+		}
+
+		if (const FPendingRegistration* Pending = Queue.Find(Guid))
+		{
+			return FCustomVersion(Guid, Pending->Version, Pending->FriendlyName);
+		}
+
+		return TOptional<FCustomVersion>();
+	}
 
 	void RegisterQueue()
 	{
@@ -148,7 +145,7 @@ private:
 		}
 	}
 
-	void UnregisterImpl(FGuid Key)
+	void Unregister(FGuid Key)
 	{
 		if (Queue.Remove(Key) == 0)
 		{
@@ -168,18 +165,99 @@ private:
 	}
 };
 
+FCustomVersionContainer FCurrentCustomVersions::GetAll()
+{
+	FStaticCustomVersionRegistry& Registry = FStaticCustomVersionRegistry::Get();
+
+	{
+		FReadScopeLock Scope(Registry.Lock);
+
+		if (Registry.Queue.Num() == 0)
+		{
+			return Registry.Registered;
+		}
+	}
+
+	FWriteScopeLock Scope(Registry.Lock);
+	Registry.RegisterQueue();
+	return Registry.Registered;
+}
+
+TOptional<FCustomVersion> FCurrentCustomVersions::Get(const FGuid& Guid)
+{
+	FStaticCustomVersionRegistry& Registry = FStaticCustomVersionRegistry::Get();
+
+	FReadScopeLock Scope(Registry.Lock);
+	return Registry.Find(Guid);
+}
+
+TArray<FCustomVersionDifference> FCurrentCustomVersions::Compare(const FCustomVersionArray& CompareVersions)
+{
+	TArray<FCustomVersionDifference> Result;
+
+	if (CompareVersions.Num())
+	{
+		FStaticCustomVersionRegistry& Registry = FStaticCustomVersionRegistry::Get();
+
+		FReadScopeLock Scope(Registry.Lock);
+
+		for (const FCustomVersion& CompareVersion : CompareVersions)
+		{
+			if (TOptional<FCustomVersion> CurrentVersion = Registry.Find(CompareVersion.Key))
+			{
+				if (int Delta = CurrentVersion.GetValue().Version - CompareVersion.Version)
+				{
+					Result.Add({ Delta < 0	? ECustomVersionDifference::Newer 
+											: ECustomVersionDifference::Older, &CompareVersion });
+				}
+			}
+			else
+			{
+				Result.Add({ ECustomVersionDifference::Missing, &CompareVersion });
+			}			
+		}
+	}
+
+	return Result;
+}
+
+void FCurrentCustomVersions::Register(const FGuid& Key, int32 Version, const TCHAR* Name)
+{
+	FStaticCustomVersionRegistry& Registry = FStaticCustomVersionRegistry::Get();
+
+	FWriteScopeLock Scope(Registry.Lock);
+	check(Registry.Queue.Find(Key) == nullptr);
+	Registry.Queue.Add(Key, { Version, Name });
+}
+
+void FCurrentCustomVersions::Unregister(const FGuid& Key)
+{
+	FStaticCustomVersionRegistry& Registry = FStaticCustomVersionRegistry::Get();
+
+	FWriteScopeLock Scope(Registry.Lock);
+	Registry.Unregister(Key);
+}
+
 const FName FCustomVersion::GetFriendlyName() const
 {
 	if (FriendlyName == NAME_None)
 	{
-		FriendlyName = FCustomVersionContainer::GetRegistered().GetFriendlyName(Key);
+		if (TOptional<FCustomVersion> CurrentVersion = FCurrentCustomVersions::Get(Key))
+		{
+			FriendlyName = CurrentVersion.GetValue().GetFriendlyName();
+		}
 	}
 	return FriendlyName;
 }
 
 const FCustomVersionContainer& FCustomVersionContainer::GetRegistered()
 {
-	return FStaticCustomVersionRegistry::Get();
+	FStaticCustomVersionRegistry& Registry = FStaticCustomVersionRegistry::Get();
+
+	// Even though returning a reference isn't thread-safe, we can still synchronize access to the queue
+	FWriteScopeLock Scope(Registry.Lock);
+	Registry.RegisterQueue();
+	return Registry.Registered;
 }
 
 void FCustomVersionContainer::Empty()
@@ -307,13 +385,3 @@ void FCustomVersionContainer::SetVersion(FGuid CustomKey, int32 Version, FName F
 		Versions.Add(FCustomVersion(CustomKey, Version, FriendlyName));
 	}
 }
-
-void FCustomVersionRegistration::QueueRegistration(FGuid Key, int32 Version, const TCHAR* Name)
-{
-	FStaticCustomVersionRegistry::Register(Key, Version, Name);
-}
-
-FCustomVersionRegistration::~FCustomVersionRegistration()
-{
-	FStaticCustomVersionRegistry::Unregister(Key);
-}
diff --git a/Engine/Source/Runtime/Core/Private/UObject/DevObjectVersion.cpp b/Engine/Source/Runtime/Core/Private/UObject/DevObjectVersion.cpp
--- a/Engine/Source/Runtime/Core/Private/UObject/DevObjectVersion.cpp
+++ b/Engine/Source/Runtime/Core/Private/UObject/DevObjectVersion.cpp
@@ -45,9 +45,8 @@ void FDevVersionRegistration::DumpVersionsToLog()
 	UE_LOG(LogDevObjectVersion, Log, TEXT("Number of dev versions registered: %d"), GDevVersions.Num());
 	for (FGuid& Guid : GDevVersions)
 	{
-		const FCustomVersion* Version = FCustomVersionContainer::GetRegistered().GetVersion(Guid);
-		check(Version);
-		UE_LOG(LogDevObjectVersion, Log, TEXT("  %s (%s): %d"), *Version->GetFriendlyName().ToString(), *Version->Key.ToString(EGuidFormats::DigitsWithHyphens), Version->Version);
+		FCustomVersion Version = FCurrentCustomVersions::Get(Guid).GetValue();
+		UE_LOG(LogDevObjectVersion, Log, TEXT("  %s (%s): %d"), *Version.GetFriendlyName().ToString(), *Version.Key.ToString(EGuidFormats::DigitsWithHyphens), Version.Version);
 	}
 #endif
 }
diff --git a/Engine/Source/Runtime/Core/Public/Serialization/CustomVersion.h b/Engine/Source/Runtime/Core/Public/Serialization/CustomVersion.h
--- a/Engine/Source/Runtime/Core/Public/Serialization/CustomVersion.h
+++ b/Engine/Source/Runtime/Core/Public/Serialization/CustomVersion.h
@@ -94,7 +94,7 @@ class CORE_API FCustomVersionRegistration;
  */
 class CORE_API FCustomVersionContainer
 {
-	friend class FStaticCustomVersionRegistry;
+	friend struct FStaticCustomVersionRegistry;
 
 public:
 	/** Gets available custom versions in this container. */
@@ -137,6 +137,7 @@ public:
 	 *
 	 * @return The registered version container.
 	 */
+	UE_DEPRECATED(4.24, "Use one of the thread-safe FCurrentCustomVersions methods instead")
 	static const FCustomVersionContainer& GetRegistered();
 
 	/**
@@ -159,6 +160,34 @@ private:
 
 };
 
+enum class ECustomVersionDifference { Missing, Newer, Older };
+
+struct FCustomVersionDifference
+{
+	ECustomVersionDifference Type;
+	const FCustomVersion* Version;
+};
+
+/** Provides access to code-defined custom versions registered via FCustomVersionRegistration. */
+class CORE_API FCurrentCustomVersions
+{
+public:
+	/** Get a copy of all versions that has been statically registered so far in the module loading process. */
+	static FCustomVersionContainer GetAll();
+
+	/** Get a copy of a single statically registered version if it exists. */
+	static TOptional<FCustomVersion> Get(const FGuid& Guid);
+
+	/** Compare a number of versions to current ones and return potential differences. */
+	static TArray<FCustomVersionDifference> Compare(const FCustomVersionArray& CompareVersions);
+
+private:
+	friend class FCustomVersionRegistration;
+
+	static void Register(const FGuid& Key, int32 Version, const TCHAR* FriendlyName);
+	static void Unregister(const FGuid& Key);
+};
+
 
 /**
  * This class will cause the registration of a custom version number and key with the global
@@ -173,14 +202,14 @@ public:
 	FCustomVersionRegistration(FGuid InKey, int32 Version, const TCHAR(&InFriendlyName)[N])
 	: Key(InKey)
 	{
-		QueueRegistration(InKey, Version, InFriendlyName);
+		FCurrentCustomVersions::Register(InKey, Version, InFriendlyName);
 	}
 
-	~FCustomVersionRegistration();
+	~FCustomVersionRegistration()
+	{
+		FCurrentCustomVersions::Unregister(Key);
+	}
 
 private:
 	FGuid Key;
-
-	/** Put registrations in an intermediate queue to avoid allocations and FName creation during static init */
-	static void QueueRegistration(FGuid Key, int32 Version, const TCHAR* FriendlyName);
 };
diff --git a/Engine/Source/Runtime/Core/Public/Stats/Stats.h b/Engine/Source/Runtime/Core/Public/Stats/Stats.h
--- a/Engine/Source/Runtime/Core/Public/Stats/Stats.h
+++ b/Engine/Source/Runtime/Core/Public/Stats/Stats.h
@@ -13,283 +13,7 @@
 #include "Stats/Stats2.h"
 
 /** 
- *	Unreal Engine Stats system
- *	
- *	This is a preliminary version of the documentation, any comments are welcome :)
- *	
- *	This system allows you to collect various performance data and then the data can be used to optimize your game.
- *	There are a few methods how to achieve this. This quick tutorial will describe all of them.
- *	For stats commands check out method PrintStatsHelpToOutputDevice();
- *	
- *	Stats system in the UE4 supports following stats types:
- *		Cycle Counter - a generic cycle counter used to counting the number of cycles during the lifetime of the object
- *		Float/Dword Counter - a counter that is cleared every frame
- *		Float/Dword Accumulator - a counter that is not cleared every frame, persistent stat, but it can be reset
- *		Memory - a special type of counter that is optimized for memory tracking
- *	
- *	Each stat needs to be grouped, this usually corresponds with displaying the specified stat group i.e. 'stat statsystem' which displays stats' related data.
- *	
- *	To define a stat group you need to use one of the following methods:
- *				DECLARE_STATS_GROUP(GroupDesc,GroupId,GroupCat) - declares a stats group which is enabled by default 
- *				DECLARE_STATS_GROUP_VERBOSE(GroupDesc,GroupId,GroupCat) - declares a stats group which is disabled by default
- *				DECLARE_STATS_GROUP_MAYBE_COMPILED_OUT(GroupDesc,GroupId,GroupCat) - declares a stats group which is disabled by default and may be stripped by the compiler
- *	
- *	where
- *		GroupDesc is a text description of the group
- *		GroupId is an UNIQUE id of the group
- *		GroupCat is reserved for future use
- *		CompileIn if set to true, the compiler may strip it out
- *		
- *	It can be done in the source or header file depending the usage scope.
- *	
- *	Examples:
- *				DECLARE_STATS_GROUP(TEXT("Threading"), STATGROUP_Threading, STATCAT_Advanced);
- *				DECLARE_STATS_GROUP_VERBOSE(TEXT("Linker Load"), STATGROUP_LinkerLoad, STATCAT_Advanced);
- *	
- *	Now, you can declare/define a stat.
- *	A stat can be used only in one cpp file, in the function scope, in the module scope or can be used in the whole project.
- *	
- *	For one file scope you need to use one of the following methods depending on the stat type.
- *				DECLARE_CYCLE_STAT(CounterName,StatId,GroupId) - declares a cycle counter stat
- *		
- *				DECLARE_SCOPE_CYCLE_COUNTER(CounterName,StatId,GroupId) - declares a cycle counter stat and uses it at the same time, it is limited to one function scope
- *				QUICK_SCOPE_CYCLE_COUNTER(StatId) - declares a cycle counter stat that will belong to stat group called 'Quick'
- *				RETURN_QUICK_DECLARE_CYCLE_STAT(StatId,GroupId) - returns a cycle counter, used by a few specialized classes, more information later
- *		
- *				DECLARE_FLOAT_COUNTER_STAT(CounterName,StatId,GroupId) - declares a float counter, technically speaking it's based on the double type, 8 bytes
- *				DECLARE_DWORD_COUNTER_STAT(CounterName,StatId,GroupId) - declared a dword counter, technically speaking it's based on the qword type, 8 bytes
- *				DECLARE_FLOAT_ACCUMULATOR_STAT(CounterName,StatId,GroupId) - declares a float accumulator
- *				DECLARE_DWORD_ACCUMULATOR_STAT(CounterName,StatId,GroupId) - declares a dword accumulator
- *				DECLARE_MEMORY_STAT(CounterName,StatId,GroupId) - declares a memory counter, same as the dword accumulator, but will be displayed with memory specific units
- *				DECLARE_MEMORY_STAT_POOL(CounterName,StatId,GroupId,Pool) - declares a memory counter with a pool
- *	
- *		If you want to have these stats accessible in the whole project/or wider range of files you need to use extern version.
- *		These methods are the same as the previously mentioned but with _EXTERN and the end of the name, here is the list:
- *				DECLARE_CYCLE_STAT_EXTERN(CounterName,StatId,GroupId, API)
- *				DECLARE_FLOAT_COUNTER_STAT_EXTERN(CounterName,StatId,GroupId, API)
- *				DECLARE_DWORD_COUNTER_STAT_EXTERN(CounterName,StatId,GroupId, API)
- *				DECLARE_FLOAT_ACCUMULATOR_STAT_EXTERN(CounterName,StatId,GroupId, API)
- *				DECLARE_DWORD_ACCUMULATOR_STAT_EXTERN(CounterName,StatId,GroupId, API)
- *				DECLARE_MEMORY_STAT_EXTERN(CounterName,StatId,GroupId, API)
- *				DECLARE_MEMORY_STAT_POOL_EXTERN(CounterName,StatId,GroupId,Pool, API)
- *				
- *		Then in the source file you need to define those stats.
- *				DEFINE_STAT(CounterName) - defines stats declared with _EXTERN
- *	
- *	where
- *		CounterName is a text description of the stat
- *		StatId is an UNIQUE id of the stat
- *		GroupId is an id of the group that the stat will belong to, the GroupId from DECLARE_STATS_GROUP*
- *		Pool is a platform specific memory pool, more details later
- *		API is the *_API of module, can be empty if the stat will be used only in that module
- *
- *	Examples:
- *		Custom memory stats with pools
- *			First you need to add a new pool to enum EMemoryCounterRegion, it can be global or platform specific.
- *			
- *				enum EMemoryCounterRegion
- *				{
- *					MCR_Invalid,	// not memory
- *					MCR_Physical,	// main system memory
- *					MCR_GPU,		// memory directly a GPU (graphics card, etc)
- *					MCR_GPUSystem,	// system memory directly accessible by a GPU
- *					MCR_TexturePool,// presized texture pools
- *					MCR_MAX
- *				};
- *			
- *			This is an example that will allow using the pools every where, see CORE_API. 
- *			THE NAME OF THE POOL MUST START WITH MCR_
- *			Header file.	
- *				DECLARE_MEMORY_STAT_POOL_EXTERN(TEXT("Physical Memory Pool [Physical]"),	MCR_Physical,		STATGROUP_Memory,  FPlatformMemory::MCR_Physical,	CORE_API);
- *				DECLARE_MEMORY_STAT_POOL_EXTERN(TEXT("GPU Memory Pool [GPU]"),				MCR_GPU,			STATGROUP_Memory,  FPlatformMemory::MCR_GPU,		CORE_API);
- *				DECLARE_MEMORY_STAT_POOL_EXTERN(TEXT("Texture Memory Pool [Texture]"),		MCR_TexturePool,	STATGROUP_Memory,  FPlatformMemory::MCR_TexturePool,CORE_API);
- *				
- *			Source file.
- *				DEFINE_STAT(MCR_Physical);
- *				DEFINE_STAT(MCR_GPU);
- *				DEFINE_STAT(MCR_TexturePool);
- *			
- *			This is a pool, so it needs to be initialized. Usually in the F*PlatformMemory::Init()
- *				SET_MEMORY_STAT(MCR_Physical, PhysicalPoolLimit);
- *				SET_MEMORY_STAT(MCR_GPU, GPUPoolLimit);
- *				SET_MEMORY_STAT(MCR_TexturePool, TexturePoolLimit);
- *			
- *			Now we have pools, so we can setup memory stats for those pools.
- *			Accessible everywhere.
- *				DECLARE_MEMORY_STAT_POOL_EXTERN(TEXT("Index buffer memory"),		STAT_IndexBufferMemory,		STATGROUP_RHI, FPlatformMemory::MCR_GPU, RHI_API);
- *				DECLARE_MEMORY_STAT_POOL_EXTERN(TEXT("Vertex buffer memory"),		STAT_VertexBufferMemory,	STATGROUP_RHI, FPlatformMemory::MCR_GPU, RHI_API);
- *				DECLARE_MEMORY_STAT_POOL_EXTERN(TEXT("Structured buffer memory"),	STAT_StructuredBufferMemory,STATGROUP_RHI, FPlatformMemory::MCR_GPU, RHI_API);
- *				DECLARE_MEMORY_STAT_POOL_EXTERN(TEXT("Pixel buffer memory"),		STAT_PixelBufferMemory,		STATGROUP_RHI, FPlatformMemory::MCR_GPU, RHI_API);
- *				
- *			Accessible only in the module where defined.	
- *				DECLARE_MEMORY_STAT_POOL_EXTERN(TEXT("Pool Memory Size"), STAT_TexturePoolSize,				STATGROUP_Streaming, FPlatformMemory::MCR_TexturePool, );
- *				DECLARE_MEMORY_STAT_POOL_EXTERN(TEXT("Pool Memory Used"), STAT_TexturePoolAllocatedSize,	STATGROUP_Streaming, FPlatformMemory::MCR_TexturePool, );
- *				
- *			And the last thing, updating the memory stats.
- *				INC_MEMORY_STAT_BY(STAT_PixelBufferMemory,NumBytes) - increases a memory stat by the specified value
- *				DEC_MEMORY_STAT_BY(STAT_PixelBufferMemory,NumBytes) - decreases a memory stat by the specified value
- *				SET_MEMORY_STAT(STAT_PixelBufferMemory,NumBytes) - sets a memory stat to the specified value
- *			
- *		Regular memory stats, without pools	
- *				DECLARE_MEMORY_STAT(TEXT("Total Physical"),		STAT_TotalPhysical,		STATGROUP_MemoryPlatform);
- *				DECLARE_MEMORY_STAT(TEXT("Total Virtual"),		STAT_TotalVirtual,		STATGROUP_MemoryPlatform);
- *				DECLARE_MEMORY_STAT(TEXT("Page Size"),			STAT_PageSize,			STATGROUP_MemoryPlatform);
- *				DECLARE_MEMORY_STAT(TEXT("Total Physical GB"),	STAT_TotalPhysicalGB,	STATGROUP_MemoryPlatform);
- *		
- *		Or DECLARE_MEMORY_STAT_EXTERN in the header file and then DEFINE_STAT in the source file.
- *		Updating the memory stats is done the same way as in the version with pools.
- *		
- *		
- *		Performance data using the cycle counters.
- *			First you need to add cycle counters.
- *				DECLARE_CYCLE_STAT(TEXT("Broadcast"),	STAT_StatsBroadcast,STATGROUP_StatSystem);
- *				DECLARE_CYCLE_STAT(TEXT("Condense"),	STAT_StatsCondense,	STATGROUP_StatSystem);
- *				
- *			Or DECLARE_CYCLE_STAT_EXTERN in the header file and then DEFINE_STAT in the source file.
- *			
- *			Now you can grab the performance data.
- *			
- *				Stats::Broadcast()
- *				{
- *					SCOPE_CYCLE_COUNTER(STAT_StatsBroadcast);
- *					...
- *					// a piece of code
- *					...
- *				}
- *
- *			and that's all.
- *			Sometimes you don't want to grab the stats every time the function is called, so you can use conditional cycle counter.
- *			It's not very common, but may be useful.
- *			
- *				Stats::Broadcast(bool bSomeCondition)
- *				{
- *					CONDITIONAL_SCOPE_CYCLE_COUNTER(STAT_StatsBroadcast,bSomeCondition);
- *					...
- *					// a piece of code
- *					...
- *				}
- *
- *			If you want to grab the performance data from one function you can use following construction.
- *			
- *				Stats::Broadcast(bool bSomeCondition)
- *				{
- *					DECLARE_SCOPE_CYCLE_COUNTER(TEXT("Broadcast"), STAT_StatsBroadcast, STATGROUP_StatSystem);
- *					...
- *					// a piece of code
- *					...
- *				}
- *				
- *			or
- *			
- *				Stats::Broadcast(bool bSomeCondition)
- *				{
- *					QUICK_SCOPE_CYCLE_COUNTER(TEXT("Stats::Broadcast"));
- *					...
- *					// a piece of code
- *					...
- *				}
- *				
- *			Mostly used for temporary stats.
- *			
- *			Those all cycle counters are used to generate the hierarchy. So you can get more detailed information about performance data.
- *			There is also an option to set flat cycle counter.
- *			
- *				Stats::Broadcast(bool bSomeCondition)
- *				{
- *					const uint32 BroadcastBeginTime = FPlatformTime::Cycles();
- *					...
- *					// a piece of code
- *					...
- *					const uint32 BroadcastEndTime = FPlatformTime::Cycles();
- *					SET_CYCLE_COUNTER(STAT_StatsBroadcast, BroadcastEndTime-BroadcastBeginTime);
- *				}
- *				
- *		
- *		A few tasks implemented in the UE4 use a different approach in terms of getting the performance data.
- *		They implement method GetStatId(). If there is no GetStatId(), the code will not compile.
- *		Here is an example.
- *		
- *				class FParallelAnimationCompletionTask
- *				{
- *					// ...
- *					// a piece of code
- *					FORCEINLINE TStatId GetStatId() const
- *					{
- *						RETURN_QUICK_DECLARE_CYCLE_STAT(FParallelAnimationCompletionTask, STATGROUP_TaskGraphTasks);
- *					}
- *					// a piece of code
- *					// ...
- *				};
- *				
- *				
- *		Generic data using the float or dword counters.		
- *			First you need to add a few counters.
- *				DECLARE_FLOAT_COUNTER_STAT_EXTERN(STAT_FloatCounter,StatId,STATGROUP_TestGroup, CORE_API)
- *				DECLARE_DWORD_COUNTER_STAT_EXTERN(STAT_DwordCounter,StatId,STATGROUP_TestGroup, CORE_API)
- *				DECLARE_FLOAT_ACCUMULATOR_STAT_EXTERN(STAT_FloatAccumulator,StatId,STATGROUP_TestGroup, CORE_API)
- *				DECLARE_DWORD_ACCUMULATOR_STAT_EXTERN(STAT_DwordAccumulator,StatId,STATGROUP_TestGroup, CORE_API)
- *				 
- *			Updating counters.
- *				INC_DWORD_STAT(StatId) - increases a dword stat by 1
- *				DEC_DWORD_STAT(StatId) - decreases a dword stat by 1
- *				INC_DWORD_STAT_BY(StatId,Amount) - increases a dword stat by the specified value
- *				DEC_DWORD_STAT_BY(StatId,Amount) - decreases a dword stat by the specified value
- *				SET_DWORD_STAT(StatId,Value) - sets a dword stat to the specified value
-
- *				INC_FLOAT_STAT_BY(StatId,Amount) - increases a float stat by the specified value
- *				DEC_FLOAT_STAT_BY(StatId,Amount) - decreases a float stat by the specified value
- *				SET_FLOAT_STAT(StatId,Value) - sets a float stat to the specified value
- *			
- *			
- *		A few helper methods
- *				 GET_STATID(StatId) - returns an instance of the TStatId of the stat, ADVANCED
- *				 GET_STATDESCRIPTION(StatId) - returns a description of the stat
- *
- *		
- *		If you don't want to use the stats system and just log some performance data, there is functionality for this.
- *		
- *				SCOPE_SECONDS_COUNTER(double&Seconds) - captures time passed in seconds, adding delta time to passed in variable
- *		
- *				Stats::Broadcast()
- *				{
- *					double ThisTime = 0;
- *					{
- *						 SCOPE_SECONDS_COUNTER(ThisTime);
-*						...
- *						// a piece of code
- *						...
- *					}
- *					UE_LOG(LogTemp, Log, TEXT("Stats::Broadcast %.2f"), ThisTime );	
- *				}
- *		
- *				FScopeLogTime - utility class to log time passed in seconds, adding cumulative stats to passed in variable, print the performance data to the log in the destructor
- *			
- *				SCOPE_LOG_TIME(Name,CumulativePtr) - using the given name prints the performance data and gathers cumulative stats
- *				SCOPE_LOG_TIME_IN_SECONDS(Name,CumulativePtr) - the same as above, but prints in seconds
- *			
- *				SCOPE_LOG_TIME_FUNC() - using the funcion name prints the performance data, cannot be nested
- *				SCOPE_LOG_TIME_FUNC_WITH_GLOBAL(CumulativePtr), same as above, but gather cumulative stats
- *			
- *			A few examples.
- *			
- *				double GMyBroadcastTime = 0.0;
- *				Stats::Broadcast()
- *				{
- *					SCOPE_LOG_TIME("Stats::Broadcast", &GMyBroadcastTime );
- *					SCOPE_LOG_TIME_IN_SECONDS("Stats::Broadcast (sec)", &GMyBroadcastTime );
- *					...
- *					// a piece of code
- *					...
- *				}
- *				
- *				Stats::Condense()
- *				{
- *					SCOPE_LOG_TIME_FUNC(); // The name should be "Stats::Condense()", may differ across compilers
- *					SCOPE_LOG_TIME_FUNC_WITH_GLOBAL(&GMyBroadcastTime);
- *					...
- *					// a piece of code
- *					...
- *				}
+ *	Learn about the Stats System at docs.unrealengine.com
  */
 
 struct TStatId;
diff --git a/Engine/Source/Runtime/Core/Public/UObject/ExternalPhysicsCustomObjectVersion.h b/Engine/Source/Runtime/Core/Public/UObject/ExternalPhysicsCustomObjectVersion.h
--- a/Engine/Source/Runtime/Core/Public/UObject/ExternalPhysicsCustomObjectVersion.h
+++ b/Engine/Source/Runtime/Core/Public/UObject/ExternalPhysicsCustomObjectVersion.h
@@ -36,6 +36,9 @@ struct CORE_API FExternalPhysicsCustomObjectVersion
 		//Save out multiple acceleration structures
 		SerializeMultiStructures,
 
+		// Add kinematic targets to TKinematicGeometryParticles
+		KinematicTargets,
+
 		// -----<new versions can be added above this line>-------------------------------------------------
 		VersionPlusOne,
 		LatestVersion = VersionPlusOne - 1
diff --git a/Engine/Source/Runtime/CoreUObject/Private/UObject/LinkerLoad.cpp b/Engine/Source/Runtime/CoreUObject/Private/UObject/LinkerLoad.cpp
--- a/Engine/Source/Runtime/CoreUObject/Private/UObject/LinkerLoad.cpp
+++ b/Engine/Source/Runtime/CoreUObject/Private/UObject/LinkerLoad.cpp
@@ -1292,39 +1292,27 @@ FLinkerLoad::ELinkerStatus FLinkerLoad::SerializePackageFileSummary()
 FLinkerLoad::ELinkerStatus FLinkerLoad::UpdateFromPackageFileSummary()
 {
 	// When unversioned, pretend we are the latest version
-	bool bCustomVersionIsLatest = Summary.bUnversioned;
-	if (!bCustomVersionIsLatest)
+	bool bCustomVersionIsLatest = true;
+	if (!Summary.bUnversioned)
 	{
-		// Check custom versions.
-		const FCustomVersionContainer& LatestCustomVersions  = FCustomVersionContainer::GetRegistered();
-		bool bAllSavedVersionsMatch = true;
-		const FCustomVersionArray&  PackageCustomVersions = Summary.GetCustomVersionContainer().GetAllVersions();
-		for (auto It = PackageCustomVersions.CreateConstIterator(); It; ++It)
+		TArray<FCustomVersionDifference> Diffs = FCurrentCustomVersions::Compare(Summary.GetCustomVersionContainer().GetAllVersions());
+		for (FCustomVersionDifference Diff : Diffs)
 		{
-			const FCustomVersion& SerializedCustomVersion = *It;
-
-			const FCustomVersion* LatestVersion = LatestCustomVersions.GetVersion(SerializedCustomVersion.Key);
-			if (!LatestVersion)
+			bCustomVersionIsLatest = false;
+			if (Diff.Type == ECustomVersionDifference::Missing)
 			{
 				// Loading a package with custom integration that we don't know about!
 				// Temporarily just warn and continue. @todo: this needs to be fixed properly
-				UE_LOG(LogLinker, Warning, TEXT("Package %s was saved with a custom integration that is not present. Tag %s  Version %d"), *Filename, *SerializedCustomVersion.Key.ToString(), SerializedCustomVersion.Version);
-				bAllSavedVersionsMatch = false;
+				UE_LOG(LogLinker, Warning, TEXT("Package %s was saved with a custom integration that is not present. Tag %s  Version %d"), *Filename, *Diff.Version->Key.ToString(), Diff.Version->Version);
 			}
-			else if (SerializedCustomVersion.Version > LatestVersion->Version)
+			else if (Diff.Type == ECustomVersionDifference::Newer)
 			{
+				FCustomVersion LatestVersion = FCurrentCustomVersions::Get(Diff.Version->Key).GetValue();
 				// Loading a package with a newer custom version than the current one.
-				UE_LOG(LogLinker, Error, TEXT("Package %s was saved with a newer custom version than the current. Tag %s Name '%s' PackageVersion %d  MaxExpected %d"), *Filename, *SerializedCustomVersion.Key.ToString(), *LatestVersion->GetFriendlyName().ToString(), SerializedCustomVersion.Version, LatestVersion->Version);
+				UE_LOG(LogLinker, Error, TEXT("Package %s was saved with a newer custom version than the current. Tag %s Name '%s' PackageVersion %d  MaxExpected %d"), *Filename, *Diff.Version->Key.ToString(), *LatestVersion.GetFriendlyName().ToString(), Diff.Version->Version, LatestVersion.Version);
 				return LINKER_Failed;
 			}
-			else if (SerializedCustomVersion.Version != LatestVersion->Version)
-			{
-				bAllSavedVersionsMatch = false;
-			}
 		}
-
-		const bool bSameNumberOfVersions = (PackageCustomVersions.Num() == LatestCustomVersions.GetAllVersions().Num());
-		bCustomVersionIsLatest = bSameNumberOfVersions && bAllSavedVersionsMatch;
 	}
 
 	const FCustomVersionContainer& SummaryVersions = Summary.GetCustomVersionContainer();
diff --git a/Engine/Source/Runtime/CoreUObject/Private/UObject/LinkerSave.cpp b/Engine/Source/Runtime/CoreUObject/Private/UObject/LinkerSave.cpp
--- a/Engine/Source/Runtime/CoreUObject/Private/UObject/LinkerSave.cpp
+++ b/Engine/Source/Runtime/CoreUObject/Private/UObject/LinkerSave.cpp
@@ -320,11 +320,10 @@ void FLinkerSave::UsingCustomVersion(const struct FGuid& Guid)
 	// Here we're going to try and dump the callstack that added a new custom version after package summary has been serialized
 	if (Summary.GetCustomVersionContainer().GetVersion(Guid) == nullptr)
 	{
-		const FCustomVersion* RegisteredVersion = FCustomVersionContainer::GetRegistered().GetVersion(Guid);
-		check(RegisteredVersion);
+		FCustomVersion RegisteredVersion = FCurrentCustomVersions::Get(Guid).GetValue();
 
 		FString CustomVersionWarning = FString::Printf(TEXT("Unexpected custom version \"%s\" used after package %s summary has been serialized. Callstack:\n"),
-			*RegisteredVersion->GetFriendlyName().ToString(), *LinkerRoot->GetName());
+			*RegisteredVersion.GetFriendlyName().ToString(), *LinkerRoot->GetName());
 
 		const int32 MaxStackFrames = 100;
 		uint64 StackFrames[MaxStackFrames];
diff --git a/Engine/Source/Runtime/CoreUObject/Private/UObject/ObjectBaseUtility.cpp b/Engine/Source/Runtime/CoreUObject/Private/UObject/ObjectBaseUtility.cpp
--- a/Engine/Source/Runtime/CoreUObject/Private/UObject/ObjectBaseUtility.cpp
+++ b/Engine/Source/Runtime/CoreUObject/Private/UObject/ObjectBaseUtility.cpp
@@ -67,20 +67,18 @@ int32 UObjectBaseUtility::GetLinkerCustomVersion(FGuid CustomVersionKey) const
 	if ( Loader != NULL )
 	{
 		// We have a linker so we can return its version.
-		auto* CustomVersion = Loader->Summary.GetCustomVersionContainer().GetVersion(CustomVersionKey);
+		const FCustomVersion* CustomVersion = Loader->Summary.GetCustomVersionContainer().GetVersion(CustomVersionKey);
 		return CustomVersion ? CustomVersion->Version : -1;
 	}
 	else if (GetOutermost() && GetOutermost()->LinkerCustomVersion.GetAllVersions().Num())
 	{
 		// Get the linker version associated with the package this object lives in
-		auto* CustomVersion = GetOutermost()->LinkerCustomVersion.GetVersion(CustomVersionKey);
+		const FCustomVersion* CustomVersion = GetOutermost()->LinkerCustomVersion.GetVersion(CustomVersionKey);
 		return CustomVersion ? CustomVersion->Version : -1;
 	}
 	// We don't have a linker associated as we e.g. might have been saved or had loaders reset, ...
 	// We must have a current version for this tag.
-	auto* CustomVersion = FCustomVersionContainer::GetRegistered().GetVersion(CustomVersionKey);
-	check(CustomVersion);
-	return CustomVersion->Version;
+	return FCurrentCustomVersions::Get(CustomVersionKey).GetValue().Version;
 }
 
 /**
diff --git a/Engine/Source/Runtime/CoreUObject/Private/UObject/PackageFileSummary.cpp b/Engine/Source/Runtime/CoreUObject/Private/UObject/PackageFileSummary.cpp
--- a/Engine/Source/Runtime/CoreUObject/Private/UObject/PackageFileSummary.cpp
+++ b/Engine/Source/Runtime/CoreUObject/Private/UObject/PackageFileSummary.cpp
@@ -128,7 +128,7 @@ void operator<<(FStructuredArchive::FSlot Slot, FPackageFileSummary& Sum)
 					Sum.bUnversioned = true;
 					Sum.FileVersionUE4 = GPackageFileUE4Version;
 					Sum.FileVersionLicenseeUE4 = GPackageFileLicenseeUE4Version;
-					Sum.CustomVersionContainer = FCustomVersionContainer::GetRegistered();
+					Sum.CustomVersionContainer = FCurrentCustomVersions::GetAll();
 				}
 			}
 			else
diff --git a/Engine/Source/Runtime/CoreUObject/Private/UObject/UObjectGlobals.cpp b/Engine/Source/Runtime/CoreUObject/Private/UObject/UObjectGlobals.cpp
--- a/Engine/Source/Runtime/CoreUObject/Private/UObject/UObjectGlobals.cpp
+++ b/Engine/Source/Runtime/CoreUObject/Private/UObject/UObjectGlobals.cpp
@@ -2509,7 +2509,7 @@ UObject::UObject()
 	EnsureNotRetrievingVTablePtr();
 
 	FObjectInitializer* ObjectInitializerPtr = FUObjectThreadContext::Get().TopInitializer();
-	UE_CLOG(!ObjectInitializerPtr, LogUObjectGlobals, Fatal, TEXT("%s is not being constructed with either NewObject, NewNamedObject or ConstructObject."), *GetName());
+	UE_CLOG(!ObjectInitializerPtr, LogUObjectGlobals, Fatal, TEXT("%s is not being constructed with NewObject."), *GetName());
 	FObjectInitializer& ObjectInitializer = *ObjectInitializerPtr;
 	UE_CLOG(ObjectInitializer.Obj != nullptr && ObjectInitializer.Obj != this, LogUObjectGlobals, Fatal, TEXT("UObject() constructor called but it's not the object that's currently being constructed with NewObject. Maybe you are trying to construct it on the stack, which is not supported."));
 	const_cast<FObjectInitializer&>(ObjectInitializer).Obj = this;
diff --git a/Engine/Source/Runtime/CoreUObject/Public/UObject/UObjectThreadContext.h b/Engine/Source/Runtime/CoreUObject/Public/UObject/UObjectThreadContext.h
--- a/Engine/Source/Runtime/CoreUObject/Public/UObject/UObjectThreadContext.h
+++ b/Engine/Source/Runtime/CoreUObject/Public/UObject/UObjectThreadContext.h
@@ -62,7 +62,7 @@ public:
 	FObjectInitializer& TopInitializerChecked()
 	{
 		FObjectInitializer* ObjectInitializerPtr = TopInitializer();
-		UE_CLOG(!ObjectInitializerPtr, LogUObjectThreadContext, Fatal, TEXT("Tried to get the current ObjectInitializer, but none is set. Please use NewObject or NewNamedObject to construct new UObject-derived classes."));
+		UE_CLOG(!ObjectInitializerPtr, LogUObjectThreadContext, Fatal, TEXT("Tried to get the current ObjectInitializer, but none is set. Please use NewObject to construct new UObject-derived classes."));
 		return *ObjectInitializerPtr;
 	}
 
diff --git a/Engine/Source/Runtime/Engine/Classes/Kismet/GameplayStatics.h b/Engine/Source/Runtime/Engine/Classes/Kismet/GameplayStatics.h
--- a/Engine/Source/Runtime/Engine/Classes/Kismet/GameplayStatics.h
+++ b/Engine/Source/Runtime/Engine/Classes/Kismet/GameplayStatics.h
@@ -36,10 +36,10 @@ class FMemoryReader;
 struct FDialogueContext;
 
 
-/** Delegate called from AsyncLoadGameFromSlot. First two parameters are passed in SlotName and UserIndex, third parameter is the returned SaveGame, or null if it failed to load */
+/** Delegate called from AsyncLoadGameFromSlot. First two parameters are passed in SlotName and UserIndex, third parameter is a bool indicating success (true) or failure (false). */
 DECLARE_DELEGATE_ThreeParams(FAsyncSaveGameToSlotDelegate, const FString&, const int32, bool);
 
-/** Delegate called from AsyncLoadGameFromSlot. First two parameters are passed in SlotName and UserIndex, third parameter is the returned SaveGame, or null if it failed to load */
+/** Delegate called from AsyncLoadGameFromSlot. First two parameters are passed in SlotName and UserIndex, third parameter is the returned SaveGame, or null if it failed to load. */
 DECLARE_DELEGATE_ThreeParams(FAsyncLoadGameFromSlotDelegate, const FString&, const int32, USaveGame*);
 
 /** Static class with useful gameplay utility functions that can be called from both Blueprint and C++ */
diff --git a/Engine/Source/Runtime/Engine/Private/GameplayStatics.cpp b/Engine/Source/Runtime/Engine/Private/GameplayStatics.cpp
--- a/Engine/Source/Runtime/Engine/Private/GameplayStatics.cpp
+++ b/Engine/Source/Runtime/Engine/Private/GameplayStatics.cpp
@@ -105,7 +105,7 @@ FSaveGameHeader::FSaveGameHeader(TSubclassOf<USaveGame> ObjectType)
 	, PackageFileUE4Version(GPackageFileUE4Version)
 	, SavedEngineVersion(FEngineVersion::Current())
 	, CustomVersionFormat(static_cast<int32>(ECustomVersionSerializationFormat::Latest))
-	, CustomVersions(FCustomVersionContainer::GetRegistered())
+	, CustomVersions(FCurrentCustomVersions::GetAll())
 	, SaveGameClassName(ObjectType->GetPathName())
 {}
 
diff --git a/Engine/Source/Runtime/Experimental/Chaos/Public/Chaos/KinematicGeometryParticles.h b/Engine/Source/Runtime/Experimental/Chaos/Public/Chaos/KinematicGeometryParticles.h
--- a/Engine/Source/Runtime/Experimental/Chaos/Public/Chaos/KinematicGeometryParticles.h
+++ b/Engine/Source/Runtime/Experimental/Chaos/Public/Chaos/KinematicGeometryParticles.h
@@ -106,7 +106,13 @@ class TKinematicGeometryParticlesImp : public TGeometryParticlesImp<T, d, SimTyp
 	CHAOS_API virtual void Serialize(FChaosArchive& Ar) override
 	{
 		TGeometryParticlesImp<T, d, SimType>::Serialize(Ar);
-		Ar << MV << MW << KinematicTargets;
+		Ar << MV << MW;
+		
+		Ar.UsingCustomVersion(FExternalPhysicsCustomObjectVersion::GUID);
+		if (Ar.CustomVer(FExternalPhysicsCustomObjectVersion::GUID) >= FExternalPhysicsCustomObjectVersion::KinematicTargets)
+		{
+			Ar << KinematicTargets;
+		}
 	}
 
   private:
diff --git a/Engine/Source/Runtime/Launch/Resources/Version.h b/Engine/Source/Runtime/Launch/Resources/Version.h
--- a/Engine/Source/Runtime/Launch/Resources/Version.h
+++ b/Engine/Source/Runtime/Launch/Resources/Version.h
@@ -55,7 +55,7 @@
 // These numbers define the banner UE4 version, and are the most significant numbers when ordering two engine versions (that is, a 4.12.* version is always 
 // newer than a 4.11.* version, regardless of the changelist that it was built with)
 #define ENGINE_MAJOR_VERSION	4
-#define ENGINE_MINOR_VERSION	24
+#define ENGINE_MINOR_VERSION	25
 #define ENGINE_PATCH_VERSION	0
 
 // Macros for encoding strings
diff --git a/Engine/Source/Runtime/Windows/AudioMixerXAudio2/Private/AudioMixerPlatformXAudio2.cpp b/Engine/Source/Runtime/Windows/AudioMixerXAudio2/Private/AudioMixerPlatformXAudio2.cpp
--- a/Engine/Source/Runtime/Windows/AudioMixerXAudio2/Private/AudioMixerPlatformXAudio2.cpp
+++ b/Engine/Source/Runtime/Windows/AudioMixerXAudio2/Private/AudioMixerPlatformXAudio2.cpp
@@ -33,21 +33,21 @@
 #include "Logging/LogMacros.h"
 
 // Macro to check result code for XAudio2 failure, get the string version, log, and goto a cleanup
-#define XAUDIO2_CLEANUP_ON_FAIL(Result)																\
-	if (FAILED(Result))																				\
-	{																								\
-		FString ErrorString = FString::Printf(TEXT("0x%X: %s"), Result, GetErrorString(Result));	\
-		AUDIO_PLATFORM_ERROR(*ErrorString);															\
-		goto Cleanup;																				\
+#define XAUDIO2_CLEANUP_ON_FAIL(Result)																					\
+	if (FAILED(Result))																									\
+	{																													\
+		FString ErrorString = FString::Printf(TEXT("%s -> 0x%X: %s"), TEXT( #Result ), Result, *GetErrorString(Result));\
+		AUDIO_PLATFORM_ERROR(*ErrorString);																				\
+		goto Cleanup;																									\
 	}
 
 // Macro to check result for XAudio2 failure, get string version, log, and return false
-#define XAUDIO2_RETURN_ON_FAIL(Result)																\
-	if (FAILED(Result))																				\
-	{																								\
-		FString ErrorString = FString::Printf(TEXT("0x%X: %s"), Result, GetErrorString(Result));	\
-		AUDIO_PLATFORM_ERROR(*ErrorString);															\
-		return false;																				\
+#define XAUDIO2_RETURN_ON_FAIL(Result)																					\
+	if (FAILED(Result))																									\
+	{																													\
+		FString ErrorString = FString::Printf(TEXT("%s -> 0x%X: %s"), TEXT( #Result ), Result, *GetErrorString(Result));\
+		AUDIO_PLATFORM_ERROR(*ErrorString);																				\
+		return false;																									\
 	}
 
 
@@ -104,8 +104,8 @@ namespace Audio
 	{
 	}
 
-	const TCHAR* FMixerPlatformXAudio2::GetErrorString(HRESULT Result)
-	{
+	FString FMixerPlatformXAudio2::GetErrorString(HRESULT Result)
+	{		
 		switch (Result)
 		{
 			case HRESULT(XAUDIO2_E_INVALID_CALL):			return TEXT("XAUDIO2_E_INVALID_CALL");
@@ -120,7 +120,18 @@ namespace Audio
 			case E_INVALIDARG:								return TEXT("E_INVALIDARG");
 			case E_OUTOFMEMORY:								return TEXT("E_OUTOFMEMORY");
 #endif
-			default:										return TEXT("UNKNOWN");
+			case 0xe000020b:								return TEXT("ERROR_NO_SUCH_DEVINST");
+			default:										
+			{
+				// We don't know this error, ask this system if it does.
+				TCHAR Buffer[1024] = { 0 }; 
+				FString Msg(FPlatformMisc::GetSystemErrorMessage(Buffer, UE_ARRAY_COUNT(Buffer), Result));
+				
+				// Anything to return? Otherwise "UNKNOWN"
+				return !Msg.IsEmpty() ? 
+					   Msg : 
+					   TEXT("UNKNOWN");
+			}
 		}
 	}
 
diff --git a/Engine/Source/Runtime/Windows/AudioMixerXAudio2/Private/AudioMixerPlatformXAudio2.h b/Engine/Source/Runtime/Windows/AudioMixerXAudio2/Private/AudioMixerPlatformXAudio2.h
--- a/Engine/Source/Runtime/Windows/AudioMixerXAudio2/Private/AudioMixerPlatformXAudio2.h
+++ b/Engine/Source/Runtime/Windows/AudioMixerXAudio2/Private/AudioMixerPlatformXAudio2.h
@@ -84,7 +84,7 @@ namespace Audio
 
 	private:
 
-		const TCHAR* GetErrorString(HRESULT Result);
+		FString GetErrorString(HRESULT Result);
 		bool AllowDeviceSwap();
 
 		// Used to teardown and reinitialize XAudio2.
