Add PLATFORM_LINUXAARCH64 to Platform.h

From CL 10861913

diff --git a/Engine/Source/Runtime/Core/Public/HAL/Platform.h b/Engine/Source/Runtime/Core/Public/HAL/Platform.h
--- a/Engine/Source/Runtime/Core/Public/HAL/Platform.h
+++ b/Engine/Source/Runtime/Core/Public/HAL/Platform.h
@@ -63,6 +63,9 @@
 #if !defined(PLATFORM_LINUX)
 	#define PLATFORM_LINUX 0
 #endif
+#if !defined(PLATFORM_LINUXAARCH64)
+	#define PLATFORM_LINUXAARCH64 0
+#endif
 #if !defined(PLATFORM_SWITCH)
 	#define PLATFORM_SWITCH 0
 #endif
