Save INI files as UTF-8 by default

The previous behavior is to save as ASCII if only ASCII code points are used, and to save as UTF-16 otherwise. This can cause problems with source control when the encoding changes between two revisions of a file.

diff --git a/Engine/Source/Runtime/Core/Private/Misc/ConfigCacheIni.cpp b/Engine/Source/Runtime/Core/Private/Misc/ConfigCacheIni.cpp
--- a/Engine/Source/Runtime/Core/Private/Misc/ConfigCacheIni.cpp
+++ b/Engine/Source/Runtime/Core/Private/Misc/ConfigCacheIni.cpp
@@ -491,7 +491,7 @@ static bool SaveConfigFileWrapper(const TCHAR* IniFile, const FString& Contents)
 	FCoreDelegates::PreSaveConfigFileDelegate.Broadcast(IniFile, Contents, SavedCount);
 
 	// save it even if a delegate did as well
-	bool bLocalWriteSucceeded = FFileHelper::SaveStringToFile(Contents, IniFile);
+	bool bLocalWriteSucceeded = FFileHelper::SaveStringToFile(Contents, IniFile, FFileHelper::EEncodingOptions::ForceUTF8WithoutBOM);
 
 	// success is based on a delegate or file write working (or both)
 	return SavedCount > 0 || bLocalWriteSucceeded;
