Add Strnlen to FPlatformString and TCString

diff --git a/Engine/Source/Runtime/Core/Private/Tests/HAL/PlatformStringTest.cpp b/Engine/Source/Runtime/Core/Private/Tests/HAL/PlatformStringTest.cpp
new file mode 100644
--- /dev/null
+++ b/Engine/Source/Runtime/Core/Private/Tests/HAL/PlatformStringTest.cpp
@@ -0,0 +1,33 @@
+// Copyright Epic Games, Inc. All Rights Reserved.
+
+#include "HAL/PlatformString.h"
+
+#include "Misc/AutomationTest.h"
+
+#if WITH_DEV_AUTOMATION_TESTS
+
+IMPLEMENT_SIMPLE_AUTOMATION_TEST(FPlatformStringTestStrnlen, "System.Core.HAL.PlatformString.Strnlen", EAutomationTestFlags::ApplicationContextMask | EAutomationTestFlags::EngineFilter)
+bool FPlatformStringTestStrnlen::RunTest(const FString& Parameters)
+{
+	TestEqual(TEXT("Strnlen(nullptr, 0)"), FPlatformString::Strnlen((const ANSICHAR*)nullptr, 0), 0);
+	TestEqual(TEXT("Strnlen(\"\", 0)"), FPlatformString::Strnlen("", 0), 0);
+	TestEqual(TEXT("Strnlen(\"1\", 0)"), FPlatformString::Strnlen("1", 0), 0);
+	TestEqual(TEXT("Strnlen(\"1\", 1)"), FPlatformString::Strnlen("1", 1), 1);
+	TestEqual(TEXT("Strnlen(\"1\", 2)"), FPlatformString::Strnlen("1", 2), 1);
+	TestEqual(TEXT("Strnlen(\"123\", 2)"), FPlatformString::Strnlen("123", 2), 2);
+	ANSICHAR AnsiBuffer[128] = "123456789";
+	TestEqual(TEXT("Strnlen(PaddedBuffer)"), FPlatformString::Strnlen(AnsiBuffer, UE_ARRAY_COUNT(AnsiBuffer)), 9);
+
+	TestEqual(TEXT("Strnlen(nullptr, 0)"), FPlatformString::Strnlen((const TCHAR*)nullptr, 0), 0);
+	TestEqual(TEXT("Strnlen(\"\", 0)"), FPlatformString::Strnlen(TEXT(""), 0), 0);
+	TestEqual(TEXT("Strnlen(\"1\", 0)"), FPlatformString::Strnlen(TEXT("1"), 0), 0);
+	TestEqual(TEXT("Strnlen(\"1\", 1)"), FPlatformString::Strnlen(TEXT("1"), 1), 1);
+	TestEqual(TEXT("Strnlen(\"1\", 2)"), FPlatformString::Strnlen(TEXT("1"), 2), 1);
+	TestEqual(TEXT("Strnlen(\"123\", 2)"), FPlatformString::Strnlen(TEXT("123"), 2), 2);
+	TCHAR Buffer[128] = TEXT("123456789");
+	TestEqual(TEXT("Strnlen(PaddedBuffer)"), FPlatformString::Strnlen(Buffer, UE_ARRAY_COUNT(Buffer)), 9);
+
+	return true;
+}
+
+#endif //WITH_DEV_AUTOMATION_TESTS
diff --git a/Engine/Source/Runtime/Core/Public/Android/AndroidPlatformString.h b/Engine/Source/Runtime/Core/Public/Android/AndroidPlatformString.h
--- a/Engine/Source/Runtime/Core/Public/Android/AndroidPlatformString.h
+++ b/Engine/Source/Runtime/Core/Public/Android/AndroidPlatformString.h
@@ -146,6 +146,16 @@ struct FAndroidPlatformString : public FGenericPlatformString
 		return Len;
 	}
 
+	static FORCEINLINE int32 Strnlen( const WIDECHAR* String, SIZE_T StringSize )
+	{
+		int Len = 0;
+		while (StringSize-- > 0 && *String++)
+		{
+			++Len;
+		}
+		return Len;
+	}
+
 	static FORCEINLINE void CopyWideToAnsi(ANSICHAR* Dest, const WIDECHAR* Src)
 	{
 		if (!Src || !Dest)
@@ -499,6 +509,11 @@ struct FAndroidPlatformString : public FGenericPlatformString
 		return strlen( String ); 
 	}
 
+	static FORCEINLINE int32 Strnlen( const ANSICHAR* String, SIZE_T StringSize )
+	{
+		return strnlen_s( String, StringSize );
+	}
+
 	static FORCEINLINE const ANSICHAR* Strstr( const ANSICHAR* String, const ANSICHAR* Find)
 	{
 		return strstr(String, Find);
@@ -581,6 +596,17 @@ struct FAndroidPlatformString : public FGenericPlatformString
 
 		return Result;
 	}
+
+	static FORCEINLINE int32 Strnlen( const UCS2CHAR* String, SIZE_T StringSize )
+	{
+		int32 Result = 0;
+		while (StringSize-- > 0 && *String++)
+		{
+			++Result;
+		}
+
+		return Result;
+	}
 };
 #endif
 
diff --git a/Engine/Source/Runtime/Core/Public/GenericPlatform/GenericWidePlatformString.h b/Engine/Source/Runtime/Core/Public/GenericPlatform/GenericWidePlatformString.h
--- a/Engine/Source/Runtime/Core/Public/GenericPlatform/GenericWidePlatformString.h
+++ b/Engine/Source/Runtime/Core/Public/GenericPlatform/GenericWidePlatformString.h
@@ -65,28 +65,53 @@ public:
 	CORE_API static int32 Strlen( const WIDECHAR* String )
 	{
 		int32 Length = -1;
-		
+
 		do
 		{
 			Length++;
 		}
 		while (*String++);
-		
+
 		return Length;
 	}
 
+	CORE_API static int32 Strnlen( const WIDECHAR* String, SIZE_T StringSize )
+	{
+		int32 Length = -1;
+
+		do
+		{
+			Length++;
+		}
+		while (StringSize-- > 0 && *String++);
+
+		return Length;
+	}
 
 #if PLATFORM_TCHAR_IS_CHAR16
 	static int32 Strlen( const wchar_t* String )
 	{
 		int32 Length = -1;
-		
+
 		do
 		{
 			Length++;
 		}
 		while (*String++);
-		
+
+		return Length;
+	}
+
+	static int32 Strnlen( const wchar_t* String, SIZE_T StringSize )
+	{
+		int32 Length = -1;
+
+		do
+		{
+			Length++;
+		}
+		while (StringSize-- > 0 && *String++);
+
 		return Length;
 	}
 #endif
@@ -212,6 +237,11 @@ public:
 		return strlen( String );
 	}
 
+	CORE_API static FORCEINLINE int32 Strnlen( const ANSICHAR* String, SIZE_T StringSize )
+	{
+		return strnlen( String, StringSize );
+	}
+
 	CORE_API static FORCEINLINE const ANSICHAR* Strstr( const ANSICHAR* String, const ANSICHAR* Find)
 	{
 		return strstr(String, Find);
@@ -294,6 +324,17 @@ public:
 
 		return Result;
 	}
+
+	CORE_API static FORCEINLINE int32 Strnlen( const UCS2CHAR* String, SIZE_T StringSize )
+	{
+		int32 Result = 0;
+		while (StringSize-- > 0 && *String++)
+		{
+			++Result;
+		}
+
+		return Result;
+	}
 };
 
 #endif
diff --git a/Engine/Source/Runtime/Core/Public/GenericPlatform/MicrosoftPlatformString.h b/Engine/Source/Runtime/Core/Public/GenericPlatform/MicrosoftPlatformString.h
--- a/Engine/Source/Runtime/Core/Public/GenericPlatform/MicrosoftPlatformString.h
+++ b/Engine/Source/Runtime/Core/Public/GenericPlatform/MicrosoftPlatformString.h
@@ -5,6 +5,11 @@
 #include "Misc/Char.h"
 #include "GenericPlatform/GenericPlatformString.h"
 #include "GenericPlatform/GenericPlatformStricmp.h"
+#include <stdarg.h>
+#include <stdio.h>
+#include <stdlib.h>
+#include <string.h>
+#include <tchar.h>
 
 /**
 * Microsoft specific implementation 
@@ -56,6 +61,11 @@ struct FMicrosoftPlatformString : public FGenericPlatformString
 		return _tcslen( String );
 	}
 
+	static FORCEINLINE int32 Strnlen( const WIDECHAR* String, SIZE_T StringSize )
+	{
+		return _tcsnlen( String, StringSize );
+	}
+
 	static FORCEINLINE const WIDECHAR* Strstr( const WIDECHAR* String, const WIDECHAR* Find)
 	{
 		return _tcsstr( String, Find );
@@ -163,6 +173,11 @@ struct FMicrosoftPlatformString : public FGenericPlatformString
 		return strlen( String ); 
 	}
 
+	static FORCEINLINE int32 Strnlen( const ANSICHAR* String, SIZE_T StringSize )
+	{
+		return strnlen( String, StringSize );
+	}
+
 	static FORCEINLINE const ANSICHAR* Strstr( const ANSICHAR* String, const ANSICHAR* Find)
 	{
 		return strstr(String, Find);
@@ -239,6 +254,11 @@ struct FMicrosoftPlatformString : public FGenericPlatformString
 	{
 		return _tcslen( (const WIDECHAR*)String );
 	}
+
+	static FORCEINLINE int32 Strnlen( const UCS2CHAR* String, SIZE_T StringSize )
+	{
+		return _tcsnlen( (const WIDECHAR*)String, StringSize );
+	}
 };
 
 #pragma warning(pop) // 'function' was was declared deprecated  (needed for the secure string functions)
diff --git a/Engine/Source/Runtime/Core/Public/GenericPlatform/StandardPlatformString.h b/Engine/Source/Runtime/Core/Public/GenericPlatform/StandardPlatformString.h
--- a/Engine/Source/Runtime/Core/Public/GenericPlatform/StandardPlatformString.h
+++ b/Engine/Source/Runtime/Core/Public/GenericPlatform/StandardPlatformString.h
@@ -65,6 +65,11 @@ public:
 		return (int32)wcslen( String );
 	}
 
+	static FORCEINLINE int32 Strnlen( const WIDECHAR* String, SIZE_T StringSize )
+	{
+		return (int32)wcsnlen_s( String, StringSize );
+	}
+
 	static FORCEINLINE const WIDECHAR* Strstr( const WIDECHAR* String, const WIDECHAR* Find)
 	{
 		return wcsstr( String, Find );
@@ -226,6 +231,11 @@ public:
 		return (int32)strlen( String );
 	}
 
+	static FORCEINLINE int32 Strnlen( const ANSICHAR* String, SIZE_T StringSize )
+	{
+		return (int32)strnlen_s( String, StringSize );
+	}
+
 	static FORCEINLINE const ANSICHAR* Strstr( const ANSICHAR* String, const ANSICHAR* Find)
 	{
 		return strstr(String, Find);
@@ -308,6 +318,17 @@ public:
 
 		return Result;
 	}
+
+	static FORCEINLINE int32 Strnlen( const UCS2CHAR* String, SIZE_T StringSize )
+	{
+		int32 Result = 0;
+		while (StringSize-- > 0 && *String++)
+		{
+			++Result;
+		}
+
+		return Result;
+	}
 };
 
 #endif
diff --git a/Engine/Source/Runtime/Core/Public/Misc/CString.h b/Engine/Source/Runtime/Core/Public/Misc/CString.h
--- a/Engine/Source/Runtime/Core/Public/Misc/CString.h
+++ b/Engine/Source/Runtime/Core/Public/Misc/CString.h
@@ -271,7 +271,15 @@ struct TCString
 	 * strlen wrapper
 	 */
 	static FORCEINLINE int32 Strlen( const CharType* String );
-	
+
+	/**
+	 * Calculate the length of the string up to the given size.
+	 * @param String A possibly-null-terminated string in a character array with a size of at least StringSize.
+	 * @param StringSize The maximum number of characters to read from String.
+	 * @return Length The smaller of StringSize and the number of characters in String before a null character.
+	 */
+	static FORCEINLINE int32 Strnlen( const CharType* String, SIZE_T StringSize );
+
 	/**
 	 * strstr wrapper
 	 */
@@ -699,6 +707,12 @@ int32 TCString<T>::Strlen( const CharType* String )
 	return FPlatformString::Strlen(String);
 }
 
+template <typename T> FORCEINLINE
+int32 TCString<T>::Strnlen( const CharType* String, SIZE_T StringSize ) 
+{
+	return FPlatformString::Strnlen(String, StringSize);
+}
+
 template <typename T> FORCEINLINE
 const typename TCString<T>::CharType* TCString<T>::Strstr( const CharType* String, const CharType* Find )
 {
