Fix handling of %.*s in FGenericWidePlatformString::GetVarArgs

diff --git a/Engine/Source/Runtime/Core/Private/GenericPlatform/GenericWidePlatformString.cpp b/Engine/Source/Runtime/Core/Private/GenericPlatform/GenericWidePlatformString.cpp
--- a/Engine/Source/Runtime/Core/Private/GenericPlatform/GenericWidePlatformString.cpp
+++ b/Engine/Source/Runtime/Core/Private/GenericPlatform/GenericWidePlatformString.cpp
@@ -424,12 +424,20 @@ int32 FGenericWidePlatformString::GetVarArgs( WIDECHAR* Dest, SIZE_T DestSize, c
 		if (*Src == '.')
 		{
 			const TCHAR *Cur = Src + 1;
-			while ((*Cur >= '0') && (*Cur <= '9'))
+			if (*Cur == '*')
 			{
+				PrecisionLen = va_arg(ArgPtr, int32);
 				Cur++;
 			}
+			else
+			{
+				while ((*Cur >= '0') && (*Cur <= '9'))
+				{
+					Cur++;
+				}
 
-			PrecisionLen = Atoi(Src + 1);
+				PrecisionLen = Atoi(Src + 1);
+			}
 			Src = Cur;
 		}
 
@@ -763,7 +771,7 @@ int32 FGenericWidePlatformString::GetVarArgs( WIDECHAR* Dest, SIZE_T DestSize, c
 					Val = Null;
 				}
 
-				int RetCnt = Strlen(Val);
+				int RetCnt = PrecisionLen < 0 ? Strlen(Val) : Strnlen(Val, PrecisionLen);
 				int Spaces = FPlatformMath::Max(FPlatformMath::Abs(FieldLen) - RetCnt, 0);
 				if (Spaces > 0 && FieldLen > 0)
 				{
@@ -795,7 +803,7 @@ int32 FGenericWidePlatformString::GetVarArgs( WIDECHAR* Dest, SIZE_T DestSize, c
 					Val = Null;
 				}
 
-				int RetCnt = Strlen(Val);
+				int RetCnt = PrecisionLen < 0 ? Strlen(Val) : Strnlen(Val, PrecisionLen);
 				int Spaces = FPlatformMath::Max(FPlatformMath::Abs(FieldLen) - RetCnt, 0);
 				if (Spaces > 0 && FieldLen > 0)
 				{
diff --git a/Engine/Source/Runtime/Core/Private/Tests/HAL/PlatformStringTest.cpp b/Engine/Source/Runtime/Core/Private/Tests/HAL/PlatformStringTest.cpp
--- a/Engine/Source/Runtime/Core/Private/Tests/HAL/PlatformStringTest.cpp
+++ b/Engine/Source/Runtime/Core/Private/Tests/HAL/PlatformStringTest.cpp
@@ -6,6 +6,25 @@
 
 #if WITH_DEV_AUTOMATION_TESTS
 
+template <typename CharType, SIZE_T Size>
+static void InvokePlatformStringGetVarArgs(CharType (&Dest)[Size], const CharType* Fmt, ...)
+{
+	va_list ap;
+	va_start(ap, Fmt);
+	FPlatformString::GetVarArgs(Dest, Size, Fmt, ap);
+	va_end(ap);
+}
+
+IMPLEMENT_SIMPLE_AUTOMATION_TEST(FPlatformStringTestGetVarArgs, "System.Core.HAL.PlatformString.GetVarArgs", EAutomationTestFlags::ApplicationContextMask | EAutomationTestFlags::EngineFilter)
+bool FPlatformStringTestGetVarArgs::RunTest(const FString& Parameters)
+{
+	TCHAR Buffer[128];
+	InvokePlatformStringGetVarArgs(Buffer, TEXT("A%.*sZ"), 4, TEXT(" to B"));
+	TestEqual(TEXT("GetVarArgs(%.*s)"), Buffer, TEXT("A to Z"));
+
+	return true;
+}
+
 IMPLEMENT_SIMPLE_AUTOMATION_TEST(FPlatformStringTestStrnlen, "System.Core.HAL.PlatformString.Strnlen", EAutomationTestFlags::ApplicationContextMask | EAutomationTestFlags::EngineFilter)
 bool FPlatformStringTestStrnlen::RunTest(const FString& Parameters)
 {
