Edigrated 11250215

  - Added missing default value in BaseEditorSettings.ini for a bSendUsageData

Details:

- When the out-of-process CrashReportClientEditor (CRC) is used, the Editor (GenericPlatformCrashContext.cpp) caches analytics information in memory and sync the data to a files consumed by CRC. The value for bSendUsageData is read from the config .ini file.
- Since the default value for "bSendUsageData" was not in BaseEditorSettings.ini, the default value used was the one from the cached value in memory (which was zeroed, so false)
- If the user toggled the options in the Editor preferences, a .ini override was created under the Saved folder and from that point, the correct "bSendUsageData" value could be read according to the one set in the Editor preferences.

diff --git a/Engine/Config/BaseEditorSettings.ini b/Engine/Config/BaseEditorSettings.ini
--- a/Engine/Config/BaseEditorSettings.ini
+++ b/Engine/Config/BaseEditorSettings.ini
@@ -9,6 +9,9 @@ bCopyStarterContentPreference = true
 [/Script/UnrealEd.CrashReportsPrivacySettings]
 bSendUnattendedBugReports=True
 
+[/Script/UnrealEd.AnalyticsPrivacySettings]
+bSendUsageData=True
+
 [/Script/NewsFeed.NewsFeedSettings]
 CdnSourceUrl="http://cdn.unrealengine.com/newsfeed/[[Culture]]/index.json"
 LocalSourcePath="Editor/NewsFeed/[[Culture]]/"
diff --git a/Engine/Source/Runtime/Core/Private/GenericPlatform/GenericPlatformCrashContext.cpp b/Engine/Source/Runtime/Core/Private/GenericPlatform/GenericPlatformCrashContext.cpp
--- a/Engine/Source/Runtime/Core/Private/GenericPlatform/GenericPlatformCrashContext.cpp
+++ b/Engine/Source/Runtime/Core/Private/GenericPlatform/GenericPlatformCrashContext.cpp
@@ -372,7 +372,7 @@ void FGenericCrashContext::InitializeFromConfig()
 	// Read the initial un-localized crash context text
 	UpdateLocalizedStrings();
 
-	// Set privacy settings
+	// Set privacy settings -> WARNING: Ensure those setting have a default values in Engine/Config/BaseEditorSettings.ini file, otherwise, they will not be found.
 	GConfig->GetBool(TEXT("/Script/UnrealEd.CrashReportsPrivacySettings"), TEXT("bSendUnattendedBugReports"), NCached::UserSettings.bSendUnattendedBugReports, GEditorSettingsIni);
 	GConfig->GetBool(TEXT("/Script/UnrealEd.AnalyticsPrivacySettings"), TEXT("bSendUsageData"), NCached::UserSettings.bSendUsageData, GEditorSettingsIni);
 
