Sequencer: Fix minor tick divisions regression

#rnx

diff --git a/Engine/Source/Runtime/Core/Private/Misc/FrameRate.cpp b/Engine/Source/Runtime/Core/Private/Misc/FrameRate.cpp
--- a/Engine/Source/Runtime/Core/Private/Misc/FrameRate.cpp
+++ b/Engine/Source/Runtime/Core/Private/Misc/FrameRate.cpp
@@ -257,7 +257,7 @@ bool FFrameRate::ComputeGridSpacing(float PixelsPerSecond, double& OutMajorInter
 		OutMajorInterval  = MajorIntervalFrames * AsInterval();
 
 		// Find the lowest number of divisions we can show that's larger than the minimum tick size
-		OutMinorDivisions = MajorIntervalFrames;
+		OutMinorDivisions = 0;
 		for (int32 DivIndex = 0; DivIndex < BaseIndex; ++DivIndex)
 		{
 			if (Base % CommonBases[DivIndex] == 0)
