Added check() to squelch static analyser warning

diff --git a/Engine/Source/Runtime/Core/Private/Tests/Misc/CharTest.cpp b/Engine/Source/Runtime/Core/Private/Tests/Misc/CharTest.cpp
--- a/Engine/Source/Runtime/Core/Private/Tests/Misc/CharTest.cpp
+++ b/Engine/Source/Runtime/Core/Private/Tests/Misc/CharTest.cpp
@@ -33,6 +33,8 @@ void RunCharTests(FAutomationTestBase& Test, uint32 MaxChar)
 bool TCharTest::RunTest(const FString& Parameters)
 {
 	const char* CurrentLocale = setlocale(LC_CTYPE, nullptr);
+	check(CurrentLocale);
+
 	if (strcmp("C", CurrentLocale))
 	{
 		AddError(FString::Printf(TEXT("Locale is \"%s\" but should be \"C\". Did something call setlocale()?"), ANSI_TO_TCHAR(CurrentLocale)));
