Edigrate 11208431 from Dev-Console to Release-4.25

[Switch] MemPro Platform Extension fixes
- fixing some issues that arose from the latest stream copy-up

#rnx

diff --git a/Engine/Source/Runtime/Core/Public/MemPro/MemPro.h b/Engine/Source/Runtime/Core/Public/MemPro/MemPro.h
--- a/Engine/Source/Runtime/Core/Public/MemPro/MemPro.h
+++ b/Engine/Source/Runtime/Core/Public/MemPro/MemPro.h
@@ -317,7 +317,7 @@ namespace
 //@EPIC BEGIN: other platforms
 #elif defined(MEMPRO_PLATFORM_SWITCH)
 
-	#include "Switch/MemProSwitch.h"
+	#include "MemPro/MemProSwitch.h"
 
 //@EPIC END
 
