When two threads are trying to create a folder path avoid short cutting and exiting before the full path has been created

diff --git a/Engine/Source/Runtime/Core/Private/Unix/UnixPlatformFile.cpp b/Engine/Source/Runtime/Core/Private/Unix/UnixPlatformFile.cpp
--- a/Engine/Source/Runtime/Core/Private/Unix/UnixPlatformFile.cpp
+++ b/Engine/Source/Runtime/Core/Private/Unix/UnixPlatformFile.cpp
@@ -1158,28 +1158,22 @@ bool FUnixPlatformFile::CreateDirectoriesFromPath(const TCHAR* Path)
 		{
 			SubPath[i+1] = 0;
 
-			// directory exists?
-			struct stat SubPathFileInfo;
-			if (stat(SubPath, &SubPathFileInfo) == -1)
+			if (mkdir(SubPath, 0755) == -1)
 			{
-				// nope. create it.
-				if (mkdir(SubPath, 0755) == -1)
-				{
-					int ErrNo = errno;
+				int ErrNo = errno;
 
-					// We happen to have gotten into a position after stat that the folder was created so dont treat this as an error
-					bool bErrnoDirExist = ErrNo == EEXIST;
+				// Folder already exists, continue and make sure the rest of the path is created
+				if (ErrNo == EEXIST)
+				{
+					continue;
+				}
 
-					if (!bErrnoDirExist)
-					{
-						UE_LOG_UNIX_FILE(Warning, TEXT( "create dir('%s') failed: errno=%d (%s)" ), UTF8_TO_TCHAR(DirPath), ErrNo, UTF8_TO_TCHAR(strerror(ErrNo)));
-					}
+				UE_LOG_UNIX_FILE(Warning, TEXT( "create dir('%s') failed: errno=%d (%s)" ), UTF8_TO_TCHAR(DirPath), ErrNo, UTF8_TO_TCHAR(strerror(ErrNo)));
 
-					FMemory::Free(DirPath);
-					FMemory::Free(SubPath);
+				FMemory::Free(DirPath);
+				FMemory::Free(SubPath);
 
-					return bErrnoDirExist;
-				}
+				return false;
 			}
 		}
 	}
