Fix issues with plugins stack traces not resolving correctly.

Disabling on demand symbol loading due to issue where modules that are not in the main binary directory would not have symbols loaded correctly. This manifested as "UnknownFunction" entries in the logged stack trace. Crash reports were not affected since they are done externally.

diff --git a/Engine/Source/Runtime/Core/Private/Windows/WindowsPlatformStackWalk.cpp b/Engine/Source/Runtime/Core/Private/Windows/WindowsPlatformStackWalk.cpp
--- a/Engine/Source/Runtime/Core/Private/Windows/WindowsPlatformStackWalk.cpp
+++ b/Engine/Source/Runtime/Core/Private/Windows/WindowsPlatformStackWalk.cpp
@@ -55,7 +55,8 @@ static const TCHAR* CrashReporterSettings = TEXT("/Script/UnrealEd.CrashReporter
 #endif
 
 // Optimization that only loads symbols on demand rather than always loading all symbols for loaded modules.
-#if 1
+// Warning: Enabling this currently has a bug where symbols are not loaded for modules outside of the main binary directory (i.e. plugins).
+#if 0
 #define ON_DEMAND_SYMBOL_LOADING 1
 #else
 #define ON_DEMAND_SYMBOL_LOADING 0
@@ -1027,6 +1028,7 @@ bool FWindowsPlatformStackWalk::InitStackWalkingInternal(void* Process)
 #if !ON_DEMAND_SYMBOL_LOADING
 		if (!FPlatformProperties::IsMonolithicBuild() && FPlatformStackWalk::WantsDetailedCallstacksInNonMonolithicBuilds())
 		{
+			const FString RemoteStorage = GetRemoteStorage(GetDownstreamStorage());
 			LoadSymbolsForProcessModules(RemoteStorage);
 		}
 #endif
