#UE4 Workaround for MSVC compiler bug in 16.4 where CanEverRender would return true when launching with -nullrhi. This will be fixed in 16.5, but for now we are working around it with FORCENOINLINE.

diff --git a/Engine/Source/Runtime/Core/Public/Misc/App.h b/Engine/Source/Runtime/Core/Public/Misc/App.h
--- a/Engine/Source/Runtime/Core/Public/Misc/App.h
+++ b/Engine/Source/Runtime/Core/Public/Misc/App.h
@@ -348,6 +348,13 @@ public:
 
 public:
 
+	// MSVC 16.4 has a bug that does not properly handle the local static bool in CanEverRender. This will be fixed in 16.5. Working around by using FORCENOINLINE.
+#if !UE_SERVER && defined(_MSC_VER) && _MSC_VER >= 1924
+#define INLINE_CANEVERRENDER FORCENOINLINE
+#else
+#define INLINE_CANEVERRENDER FORCEINLINE
+#endif
+
 	/**
 	 * Checks whether this application can render anything.
 	 * Certain application types never render, while for others this behavior may be controlled by switching to NullRHI.
@@ -355,7 +362,7 @@ public:
 	 *
 	 * @return true if the application can render, false otherwise.
 	 */
-	FORCEINLINE static bool CanEverRender()
+	INLINE_CANEVERRENDER static bool CanEverRender()
 	{
 #if UE_SERVER
 		return false;
