Removed named event overhead from STAT_AddSubjectPrimitive.  Also, removed now empty verbose groups from CL 10988572.

diff --git a/Engine/Source/Runtime/Core/Public/Stats/Stats2.h b/Engine/Source/Runtime/Core/Public/Stats/Stats2.h
--- a/Engine/Source/Runtime/Core/Public/Stats/Stats2.h
+++ b/Engine/Source/Runtime/Core/Public/Stats/Stats2.h
@@ -2210,11 +2210,9 @@ DECLARE_STATS_GROUP(TEXT("MapBuildData"),STATGROUP_MapBuildData, STATCAT_Advance
 DECLARE_STATS_GROUP(TEXT("Shader Compiling"),STATGROUP_ShaderCompiling, STATCAT_Advanced);
 DECLARE_STATS_GROUP(TEXT("Shader Compression"),STATGROUP_Shaders, STATCAT_Advanced);
 DECLARE_STATS_GROUP(TEXT("Shadow Rendering"),STATGROUP_ShadowRendering, STATCAT_Advanced);
-DECLARE_STATS_GROUP_VERBOSE(TEXT("Shadow Rendering Verbose"), STATGROUP_ShadowRenderingVerbose, STATCAT_Advanced);
 DECLARE_STATS_GROUP(TEXT("Stat System"),STATGROUP_StatSystem, STATCAT_Advanced);
 DECLARE_STATS_GROUP_SORTBYNAME(TEXT("Streaming Overview"),STATGROUP_StreamingOverview, STATCAT_Advanced);
 DECLARE_STATS_GROUP_SORTBYNAME(TEXT("Streaming Details"),STATGROUP_StreamingDetails, STATCAT_Advanced);
-DECLARE_STATS_GROUP_VERBOSE(TEXT("Streaming Details Verbose"), STATGROUP_StreamingDetailsVerbose, STATCAT_Advanced);
 DECLARE_STATS_GROUP_SORTBYNAME(TEXT("Streaming"),STATGROUP_Streaming, STATCAT_Advanced);
 DECLARE_STATS_GROUP(TEXT("Target Platform"),STATGROUP_TargetPlatform, STATCAT_Advanced);
 DECLARE_STATS_GROUP(TEXT("Text"),STATGROUP_Text, STATCAT_Advanced);
diff --git a/Engine/Source/Runtime/RenderCore/Private/RenderCore.cpp b/Engine/Source/Runtime/RenderCore/Private/RenderCore.cpp
--- a/Engine/Source/Runtime/RenderCore/Private/RenderCore.cpp
+++ b/Engine/Source/Runtime/RenderCore/Private/RenderCore.cpp
@@ -147,7 +147,6 @@ DEFINE_STAT(STAT_PerObjectShadows);
 DEFINE_STAT(STAT_PreShadows);
 DEFINE_STAT(STAT_CachedPreShadows);
 DEFINE_STAT(STAT_ShadowDynamicPathDrawCalls);
-DEFINE_STAT(STAT_AddSubjectPrimitive);
 
 DEFINE_STAT(STAT_TranslucentInjectTime);
 DEFINE_STAT(STAT_DirectLightRenderingTime);
diff --git a/Engine/Source/Runtime/RenderCore/Public/RenderCore.h b/Engine/Source/Runtime/RenderCore/Public/RenderCore.h
--- a/Engine/Source/Runtime/RenderCore/Public/RenderCore.h
+++ b/Engine/Source/Runtime/RenderCore/Public/RenderCore.h
@@ -90,7 +90,6 @@ DECLARE_CYCLE_STAT_EXTERN(TEXT("WholeScene Shadow Projections"),STAT_RenderWhole
 DECLARE_CYCLE_STAT_EXTERN(TEXT("WholeScene Shadow Depths"),STAT_RenderWholeSceneShadowDepthsTime,STATGROUP_ShadowRendering, RENDERCORE_API);
 DECLARE_CYCLE_STAT_EXTERN(TEXT("PerObject Shadow Projections"),STAT_RenderPerObjectShadowProjectionsTime,STATGROUP_ShadowRendering, RENDERCORE_API);
 DECLARE_CYCLE_STAT_EXTERN(TEXT("PerObject Shadow Depths"),STAT_RenderPerObjectShadowDepthsTime,STATGROUP_ShadowRendering, RENDERCORE_API);
-DECLARE_CYCLE_STAT_EXTERN(TEXT("Add Subject Primitive"), STAT_AddSubjectPrimitive, STATGROUP_ShadowRenderingVerbose, RENDERCORE_API);
 DECLARE_MEMORY_STAT_EXTERN(TEXT("Per-Frame Shadowmap Atlases"),STAT_ShadowmapAtlasMemory,STATGROUP_ShadowRendering, RENDERCORE_API); 
 DECLARE_MEMORY_STAT_EXTERN(TEXT("Cached Shadowmaps"),STAT_CachedShadowmapMemory,STATGROUP_ShadowRendering, RENDERCORE_API); 
 DECLARE_DWORD_COUNTER_STAT_EXTERN(TEXT("Whole Scene shadows"),STAT_WholeSceneShadows,STATGROUP_ShadowRendering, RENDERCORE_API);
diff --git a/Engine/Source/Runtime/Renderer/Private/ShadowSetup.cpp b/Engine/Source/Runtime/Renderer/Private/ShadowSetup.cpp
--- a/Engine/Source/Runtime/Renderer/Private/ShadowSetup.cpp
+++ b/Engine/Source/Runtime/Renderer/Private/ShadowSetup.cpp
@@ -1019,8 +1019,6 @@ bool FProjectedShadowInfo::ShouldDrawStaticMeshes(FViewInfo& InCurrentView, bool
 
 void FProjectedShadowInfo::AddSubjectPrimitive(FPrimitiveSceneInfo* PrimitiveSceneInfo, TArray<FViewInfo>* ViewArray, ERHIFeatureLevel::Type FeatureLevel, bool bRecordShadowSubjectsForMobileShading)
 {
-	SCOPE_CYCLE_COUNTER(STAT_AddSubjectPrimitive);
-
 	// Ray traced shadows use the GPU managed distance field object buffers, no CPU culling should be used
 	check(!bRayTracedDistanceField);
 
