Optimize TSet::Empty to avoid the pathological case where the hash size is quite large but there's not so many elements in the hash, by clearing only the elements buckets instead of clearing the whole set of buckets.
This is mainly used when marking/unmarking objects during packages saving and other archives tagging objects, but also in several cases all over the engine.

diff --git a/Engine/Source/Runtime/Core/Public/Containers/Set.h b/Engine/Source/Runtime/Core/Public/Containers/Set.h
--- a/Engine/Source/Runtime/Core/Public/Containers/Set.h
+++ b/Engine/Source/Runtime/Core/Public/Containers/Set.h
@@ -423,16 +423,22 @@ public:
 	void Empty(int32 ExpectedNumElements = 0)
 	{
 		// Empty the elements array, and reallocate it for the expected number of elements.
+		const int32 DesiredHashSize = Allocator::GetNumberOfHashBuckets(ExpectedNumElements);
+		const bool ShouldDoRehash = ShouldRehash(ExpectedNumElements,DesiredHashSize,true);
+
+		if (!ShouldDoRehash)
+		{
+			// If the hash was already the desired size, clear the references to the elements that have now been removed.
+			UnhashElements();
+		}
+
 		Elements.Empty(ExpectedNumElements);
 
 		// Resize the hash to the desired size for the expected number of elements.
-		if(!ConditionalRehash(ExpectedNumElements,true))
+		if (ShouldDoRehash)
 		{
-			// If the hash was already the desired size, clear the references to the elements that have now been removed.
-			for (int32 HashIndex = 0, LocalHashSize = HashSize; HashIndex < LocalHashSize; ++HashIndex)
-			{
-				GetTypedHash(HashIndex) = FSetElementId();
-			}
+			HashSize = DesiredHashSize;
+			Rehash();
 		}
 	}
 
@@ -445,10 +451,8 @@ public:
 		}
     	
 		// Reset the elements array.
+		UnhashElements();
 		Elements.Reset();
-
-		// Clear the references to the elements that have now been removed.
-		FSetElementId::ResetRange(Hash.GetAllocation(), HashSize);
     }
 
 	/** Shrinks the set's element storage to avoid slack. */
@@ -1301,6 +1305,45 @@ private:
 		LinkElement(ElementId, Element, KeyFuncs::GetKeyHash(KeyFuncs::GetSetKey(Element.Value)));
 	}
 
+	/** Returns if it should be faster to clear the hash by going through elements instead of reseting the whole bucket lists*/
+	FORCEINLINE bool ShouldClearByElements()
+	{
+		return Num() < (HashSize / 4);
+	}
+
+	/** Reset elements buckets of FSetElementIds to invalid */
+	void UnhashElements()
+	{
+		if (ShouldClearByElements())
+		{
+			// Faster path: only reset hash buckets to FSetElementId for elements in the hash
+			for (const SetElementType& Element: Elements)
+			{
+				Hash.GetAllocation()[Element.HashIndex] = FSetElementId();
+			}
+		}
+		else
+		{
+			FSetElementId::ResetRange(Hash.GetAllocation(), HashSize);
+		}
+	}
+
+	/**
+	 * Checks if the hash has an appropriate number of buckets, and if it should be resized.
+	 * @param NumHashedElements - The number of elements to size the hash for.
+	 * @param DesiredHashSize - Desired size if we should rehash.
+	 * @param bAllowShrinking - true if the hash is allowed to shrink.
+	 * @return true if the set should berehashed.
+	 */
+	FORCEINLINE bool ShouldRehash(int32 NumHashedElements,int32 DesiredHashSize,bool bAllowShrinking = false) const
+	{
+		// If the hash hasn't been created yet, or is smaller than the desired hash size, rehash.
+		return (NumHashedElements > 0 &&
+				(!HashSize ||
+				HashSize < DesiredHashSize ||
+				(HashSize > DesiredHashSize && bAllowShrinking)));
+	}
+
 	/**
 	 * Checks if the hash has an appropriate number of buckets, and if not resizes it.
 	 * @param NumHashedElements - The number of elements to size the hash for.
@@ -1312,20 +1355,14 @@ private:
 		// Calculate the desired hash size for the specified number of elements.
 		const int32 DesiredHashSize = Allocator::GetNumberOfHashBuckets(NumHashedElements);
 
-		// If the hash hasn't been created yet, or is smaller than the desired hash size, rehash.
-		if(NumHashedElements > 0 &&
-			(!HashSize ||
-			HashSize < DesiredHashSize ||
-			(HashSize > DesiredHashSize && bAllowShrinking)))
+		if (ShouldRehash(NumHashedElements, DesiredHashSize, bAllowShrinking))
 		{
 			HashSize = DesiredHashSize;
 			Rehash();
 			return true;
 		}
-		else
-		{
-			return false;
-		}
+
+		return false;
 	}
 
 	/** Resizes the hash. */
