Attempted fix for strange COM failure on build machine.

diff --git a/Engine/Source/Runtime/Core/Private/Windows/WindowsPlatformProcess.cpp b/Engine/Source/Runtime/Core/Private/Windows/WindowsPlatformProcess.cpp
--- a/Engine/Source/Runtime/Core/Private/Windows/WindowsPlatformProcess.cpp
+++ b/Engine/Source/Runtime/Core/Private/Windows/WindowsPlatformProcess.cpp
@@ -1628,6 +1628,16 @@ bool FWindowsPlatformProcess::Daemonize()
 	return true;
 }
 
+void FWindowsPlatformProcess::SetupAudioThread()
+{
+	ensure(FPlatformMisc::CoInitialize());
+}
+
+void FWindowsPlatformProcess::TeardownAudioThread()
+{
+	FPlatformMisc::CoUninitialize();
+}
+
 /**
  * Maps a relative virtual address (RVA) to an address in memory.
  *
diff --git a/Engine/Source/Runtime/Core/Public/GenericPlatform/GenericPlatformProcess.h b/Engine/Source/Runtime/Core/Public/GenericPlatform/GenericPlatformProcess.h
--- a/Engine/Source/Runtime/Core/Public/GenericPlatform/GenericPlatformProcess.h
+++ b/Engine/Source/Runtime/Core/Public/GenericPlatform/GenericPlatformProcess.h
@@ -222,6 +222,9 @@ struct CORE_API FGenericPlatformProcess
 	/** Allow the platform to do anything it needs for audio thread */
 	static void SetupAudioThread() { }
 
+	/** Allow the platform to tear down the audio thread */
+	static void TeardownAudioThread() { }
+	
 	/** Content saved to the game or engine directories should be rerouted to user directories instead **/
 	static bool ShouldSaveToUserDir();
 
diff --git a/Engine/Source/Runtime/Core/Public/Windows/WindowsPlatformProcess.h b/Engine/Source/Runtime/Core/Public/Windows/WindowsPlatformProcess.h
--- a/Engine/Source/Runtime/Core/Public/Windows/WindowsPlatformProcess.h
+++ b/Engine/Source/Runtime/Core/Public/Windows/WindowsPlatformProcess.h
@@ -186,6 +186,8 @@ public:
 	static FSemaphore* NewInterprocessSynchObject(const TCHAR* Name, bool bCreate, uint32 MaxLocks = 1);
 	static bool DeleteInterprocessSynchObject(FSemaphore * Object);
 	static bool Daemonize();
+	static void SetupAudioThread();
+	static void TeardownAudioThread();
 protected:
 
 	/**
diff --git a/Engine/Source/Runtime/Engine/Private/AudioThread.cpp b/Engine/Source/Runtime/Engine/Private/AudioThread.cpp
--- a/Engine/Source/Runtime/Engine/Private/AudioThread.cpp
+++ b/Engine/Source/Runtime/Engine/Private/AudioThread.cpp
@@ -192,6 +192,7 @@ bool FAudioThread::Init()
 void FAudioThread::Exit()
 {
 	GAudioThreadId = 0;
+	FPlatformProcess::TeardownAudioThread();
 }
 
 uint32 FAudioThread::Run()
