Add SIGABRT to fatal signals list.
#rnx

diff --git a/Engine/Source/Runtime/Core/Private/Android/AndroidPlatformMisc.cpp b/Engine/Source/Runtime/Core/Private/Android/AndroidPlatformMisc.cpp
--- a/Engine/Source/Runtime/Core/Private/Android/AndroidPlatformMisc.cpp
+++ b/Engine/Source/Runtime/Core/Private/Android/AndroidPlatformMisc.cpp
@@ -893,7 +893,8 @@ static constexpr int32 TargetSignals[] =
 	SIGFPE,
 	SIGBUS,
 	SIGSEGV,
-	SIGSYS
+	SIGSYS,
+	SIGABRT
 };
 static constexpr int32 NumTargetSignals = UE_ARRAY_COUNT(TargetSignals);
 
@@ -913,6 +914,8 @@ static const char* SignalToString(int32 Signal)
 			return "SIGSEGV";
 		case SIGSYS:
 			return "SIGSYS";
+		case SIGABRT:
+			return "SIGABRT";
 		default:
 			return FAndroidCrashContext::ItoANSI(Signal,16, 16);
 	}
