Flush logs when triggering a non-fatal report
#rnx

diff --git a/Engine/Source/Runtime/Core/Private/Android/AndroidPlatformMisc.cpp b/Engine/Source/Runtime/Core/Private/Android/AndroidPlatformMisc.cpp
--- a/Engine/Source/Runtime/Core/Private/Android/AndroidPlatformMisc.cpp
+++ b/Engine/Source/Runtime/Core/Private/Android/AndroidPlatformMisc.cpp
@@ -48,6 +48,10 @@
 #include "Android/AndroidPlatformStackWalk.h"
 #include "Android/AndroidSignals.h"
 
+#include "Misc/OutputDevice.h"
+#include "Logging/LogMacros.h"
+#include "Misc/OutputDeviceError.h"
+
 #if USE_ANDROID_JNI
 extern AAssetManager * AndroidThunkCpp_GetAssetManager();
 extern int32 GAndroidPackageVersion;
@@ -1218,6 +1222,20 @@ void FAndroidMisc::TriggerNonFatalCrashHandler(ECrashContextType InType, const F
 {
 	check(InType == ECrashContextType::Ensure);
 
+	if (GLog)
+	{
+		GLog->PanicFlushThreadedLogs();
+		GLog->Flush();
+	}
+	if (GWarn)
+	{
+		GWarn->Flush();
+	}
+	if (GError)
+	{
+		GError->Flush();
+	}
+
 	FAndroidCrashContext CrashContext(InType, *Message);
 
 	CrashContext.CaptureCrashInfo();
