Constructors to allow the passing of a deleter with state.
Fix to the assignment operator of an array-based TUniquePtr with a different deleter.
Unification of constructor syntax.

diff --git a/Engine/Source/Runtime/Core/Public/Templates/UniquePtr.h b/Engine/Source/Runtime/Core/Public/Templates/UniquePtr.h
--- a/Engine/Source/Runtime/Core/Public/Templates/UniquePtr.h
+++ b/Engine/Source/Runtime/Core/Public/Templates/UniquePtr.h
@@ -16,6 +16,10 @@
 // TUniquePtr<MyClass> Ptr1(new MyClass);    // The MyClass object is owned by Ptr1.
 // TUniquePtr<MyClass> Ptr2(Ptr1);           // Error - TUniquePtr is not copyable
 // TUniquePtr<MyClass> Ptr3(MoveTemp(Ptr1)); // Ptr3 now owns the MyClass object - Ptr1 is now nullptr.
+//
+// If you provide a custom deleter, it is up to your deleter to handle null pointers.  This is a departure
+// from std::unique_ptr which will not invoke the deleter if the owned pointer is null:
+// https://en.cppreference.com/w/cpp/memory/unique_ptr/~unique_ptr
 
 template <typename T>
 struct TDefaultDelete
@@ -100,7 +104,23 @@ public:
 	 * Default constructor - initializes the TUniquePtr to null.
 	 */
 	FORCEINLINE TUniquePtr()
-		: Ptr(nullptr)
+		: Deleter()
+		, Ptr    (nullptr)
+	{
+	}
+
+	/**
+	 * Pointer constructor - takes ownership of the pointed-to object
+	 *
+	 * @param InPtr The pointed-to object to take ownership of.
+	 */
+	template <
+		typename U,
+		typename = decltype(ImplicitConv<T*>((U*)nullptr))
+	>
+	explicit FORCEINLINE TUniquePtr(U* InPtr)
+		: Deleter()
+		, Ptr    (InPtr)
 	{
 	}
 
@@ -109,8 +129,28 @@ public:
 	 *
 	 * @param InPtr The pointed-to object to take ownership of.
 	 */
-	explicit FORCEINLINE TUniquePtr(T* InPtr)
-		: Ptr(InPtr)
+	template <
+		typename U,
+		typename = decltype(ImplicitConv<T*>((U*)nullptr))
+	>
+	explicit FORCEINLINE TUniquePtr(U* InPtr, Deleter&& InDeleter)
+		: Deleter(MoveTemp(InDeleter))
+		, Ptr    (InPtr)
+	{
+	}
+
+	/**
+	 * Pointer constructor - takes ownership of the pointed-to object
+	 *
+	 * @param InPtr The pointed-to object to take ownership of.
+	 */
+	template <
+		typename U,
+		typename = decltype(ImplicitConv<T*>((U*)nullptr))
+	>
+	explicit FORCEINLINE TUniquePtr(U* InPtr, const Deleter& InDeleter)
+		: Deleter(InDeleter)
+		, Ptr    (InPtr)
 	{
 	}
 
@@ -118,7 +158,8 @@ public:
 	 * nullptr constructor - initializes the TUniquePtr to null.
 	 */
 	FORCEINLINE TUniquePtr(TYPE_OF_NULLPTR)
-		: Ptr(nullptr)
+		: Deleter()
+		, Ptr    (nullptr)
 	{
 	}
 
@@ -340,7 +381,8 @@ public:
 	 * Default constructor - initializes the TUniquePtr to null.
 	 */
 	FORCEINLINE TUniquePtr()
-		: Ptr(nullptr)
+		: Deleter()
+		, Ptr    (nullptr)
 	{
 	}
 
@@ -354,7 +396,38 @@ public:
 		typename = decltype(ImplicitConv<T(*)[]>((U(*)[])nullptr))
 	>
 	explicit FORCEINLINE TUniquePtr(U* InPtr)
-		: Ptr(InPtr)
+		: Deleter()
+		, Ptr    (InPtr)
+	{
+	}
+
+	/**
+	 * Pointer constructor - takes ownership of the pointed-to array
+	 *
+	 * @param InPtr The pointed-to array to take ownership of.
+	 */
+	template <
+		typename U,
+		typename = decltype(ImplicitConv<T(*)[]>((U(*)[])nullptr))
+	>
+	explicit FORCEINLINE TUniquePtr(U* InPtr, Deleter&& InDeleter)
+		: Deleter(MoveTemp(InDeleter))
+		, Ptr    (InPtr)
+	{
+	}
+
+	/**
+	 * Pointer constructor - takes ownership of the pointed-to array
+	 *
+	 * @param InPtr The pointed-to array to take ownership of.
+	 */
+	template <
+		typename U,
+		typename = decltype(ImplicitConv<T(*)[]>((U(*)[])nullptr))
+	>
+	explicit FORCEINLINE TUniquePtr(U* InPtr, const Deleter& InDeleter)
+		: Deleter(InDeleter)
+		, Ptr    (InPtr)
 	{
 	}
 
@@ -362,7 +435,8 @@ public:
 	 * nullptr constructor - initializes the TUniquePtr to null.
 	 */
 	FORCEINLINE TUniquePtr(TYPE_OF_NULLPTR)
-		: Ptr(nullptr)
+		: Deleter()
+		, Ptr    (nullptr)
 	{
 	}
 
@@ -418,7 +492,7 @@ public:
 		typename OtherDeleter,
 		typename = decltype(ImplicitConv<T(*)[]>((OtherT(*)[])nullptr))
 	>
-	FORCEINLINE TUniquePtr& operator=(TUniquePtr<OtherT>&& Other)
+	FORCEINLINE TUniquePtr& operator=(TUniquePtr<OtherT, OtherDeleter>&& Other)
 	{
 		// We delete last, because we don't want odd side effects if the destructor of T relies on the state of this or Other
 		T* OldPtr = Ptr;
