Add copy constructor to FHashTable.

#rnx

diff --git a/Engine/Source/Runtime/Core/Public/Containers/HashTable.h b/Engine/Source/Runtime/Core/Public/Containers/HashTable.h
--- a/Engine/Source/Runtime/Core/Public/Containers/HashTable.h
+++ b/Engine/Source/Runtime/Core/Public/Containers/HashTable.h
@@ -181,6 +181,7 @@ class FHashTable
 {
 public:
 					FHashTable( uint32 InHashSize = 1024, uint32 InIndexSize = 0 );
+					FHashTable( const FHashTable& Other );
 					~FHashTable();
 
 	void			Initialize(uint32 InHashSize = 1024, uint32 InIndexSize = 0);
@@ -227,6 +228,26 @@ FORCEINLINE FHashTable::FHashTable( uint32 InHashSize, uint32 InIndexSize )
 	}
 }
 
+FORCEINLINE FHashTable::FHashTable( const FHashTable& Other )
+	: HashSize( 0 )
+	, HashMask( 0 )
+	, IndexSize( 0 )
+	, Hash( EmptyHash )
+	, NextIndex( NULL )
+{
+	if (Other.HashSize > 0u)
+	{
+		Initialize(Other.HashSize, Other.IndexSize);
+
+		check(HashSize == Other.HashSize);
+		check(HashMask == Other.HashMask);
+		check(IndexSize == Other.IndexSize);
+
+		FMemory::Memcpy(Hash, Other.Hash, HashSize * 4);
+		FMemory::Memcpy(NextIndex, Other.NextIndex, IndexSize * 4);
+	}
+}
+
 FORCEINLINE void FHashTable::Initialize(uint32 InHashSize, uint32 InIndexSize)
 {
 	check(HashSize == 0u);
