Fix Blacklist Names IsBlacklistAll()
#rnx

diff --git a/Engine/Source/Runtime/Core/Public/Misc/BlacklistNames.h b/Engine/Source/Runtime/Core/Public/Misc/BlacklistNames.h
--- a/Engine/Source/Runtime/Core/Public/Misc/BlacklistNames.h
+++ b/Engine/Source/Runtime/Core/Public/Misc/BlacklistNames.h
@@ -43,7 +43,7 @@ public:
 	const TMap<FName, FBlacklistOwners>& GetWhitelist() const { return Whitelist; }
 
 	/** Are all items set to be filtered out */
-	bool IsBlacklistAll() const { return Blacklist.Num() > 0; }
+	bool IsBlacklistAll() const { return BlacklistAll.Num() > 0; }
 
 	/** Triggered when filter changes */
 	FSimpleMulticastDelegate& OnFilterChanged() { return OnFilterChangedDelegate; }
@@ -109,7 +109,7 @@ public:
 	const TMap<FString, FBlacklistOwners>& GetWhitelist() const { return Whitelist; }
 
 	/** Are all items set to be filtered out */
-	bool IsBlacklistAll() const { return Blacklist.Num() > 0; }
+	bool IsBlacklistAll() const { return BlacklistAll.Num() > 0; }
 	
 	/** Triggered when filter changes */
 	FSimpleMulticastDelegate& OnFilterChanged() { return OnFilterChangedDelegate; }
