TBitArray fixes and optimizations

Initialize unused TBitArray bits, switch to faster unsigned operator% and don't shrink past initial allocation size

operator<, == and GetTypedHash() would access unused bits in the last word

diff --git a/Engine/Source/Runtime/Core/Public/Containers/BitArray.h b/Engine/Source/Runtime/Core/Public/Containers/BitArray.h
--- a/Engine/Source/Runtime/Core/Public/Containers/BitArray.h
+++ b/Engine/Source/Runtime/Core/Public/Containers/BitArray.h
@@ -25,6 +25,15 @@ struct FBitSet
 		Mask ^= LowestBitMask;
 		return BitIndex;
 	}
+
+	// Clang generates 7 instructions for int32 DivideAndRoundUp but only 2 for uint32
+	static constexpr uint32 BitsPerWord = NumBitsPerDWORD;
+
+	FORCEINLINE static uint32 CalculateNumWords(int32 NumBits)
+	{
+		checkSlow(NumBits >= 0);
+		return FMath::DivideAndRoundUp(static_cast<uint32>(NumBits), BitsPerWord);
+	}
 };
 
 
@@ -193,42 +202,19 @@ public:
 	TBitArray()
 	:	NumBits(0)
 	,	MaxBits(AllocatorInstance.GetInitialCapacity() * NumBitsPerDWORD)
-	{}
+	{
+		SetWords(GetData(), AllocatorInstance.GetInitialCapacity(), false);
+	}
 
 	/**
 	 * Minimal initialization constructor.
 	 * @param Value - The value to initial the bits to.
 	 * @param InNumBits - The initial number of bits in the array.
 	 */
-	explicit TBitArray(bool bValue, int32 InNumBits)
-	:	NumBits(InNumBits)
-	,	MaxBits(AllocatorInstance.GetInitialCapacity() * NumBitsPerDWORD)
+	FORCEINLINE explicit TBitArray(bool bValue, int32 InNumBits)
+	:	MaxBits(AllocatorInstance.GetInitialCapacity() * NumBitsPerDWORD)
 	{
-		if (NumBits > 0)
-		{
-			const int32 Words = FMath::DivideAndRoundUp(NumBits, NumBitsPerDWORD);
-
-			if (Words > FMath::DivideAndRoundUp(MaxBits, NumBitsPerDWORD))
-			{
-				AllocatorInstance.ResizeAllocation(0, Words, sizeof(uint32));
-				MaxBits = Words * NumBitsPerDWORD;
-			}
-			
-			if (Words > 8)
-			{
-				FMemory::Memset(GetData(), bValue ? 0xff : 0, Words * sizeof(uint32));
-			}
-			else
-			{
-				uint32 Word = bValue ? ~0u : 0;
-				int32 WordIdx = 0;
-				do
-				{
-					GetData()[WordIdx] = Word;
-				}
-				while (++WordIdx < Words);
-			}
-		}
+		Init(bValue, InNumBits);
 	}
 
 	/**
@@ -275,10 +261,9 @@ public:
 
 		Empty(Copy.Num());
 		NumBits = Copy.NumBits;
-		if(NumBits)
+		if (NumBits)
 		{
-			const int32 NumDWORDs = FMath::DivideAndRoundUp(MaxBits, NumBitsPerDWORD);
-			FMemory::Memcpy(GetData(),Copy.GetData(),NumDWORDs * sizeof(uint32));
+			FMemory::Memcpy(GetData(),Copy.GetData(), GetNumWords() * sizeof(uint32));
 		}
 		return *this;
 	}
@@ -290,8 +275,7 @@ public:
 			return false;
 		}
 
-		int NumBytes = FMath::DivideAndRoundUp(NumBits, NumBitsPerDWORD) * sizeof(uint32);
-		return FMemory::Memcmp(GetData(), Other.GetData(), NumBytes) == 0;
+		return FMemory::Memcmp(GetData(), Other.GetData(), GetNumWords() * sizeof(uint32)) == 0;
 	}
 
 	FORCEINLINE bool operator<(const TBitArray<Allocator>& Other) const
@@ -302,7 +286,7 @@ public:
 			return Num() < Other.Num();
 		}
 
-		uint32 NumWords = FMath::DivideAndRoundUp(Num(), NumBitsPerDWORD);
+		uint32 NumWords = GetNumWords();
 		const uint32* Data0 = GetData();
 		const uint32* Data1 = Other.GetData();
 
@@ -323,6 +307,38 @@ public:
 	}
 
 private:
+	FORCEINLINE uint32 GetNumWords() const
+	{
+		return FBitSet::CalculateNumWords(NumBits);
+	}
+
+	FORCEINLINE uint32 GetMaxWords() const
+	{
+		return FBitSet::CalculateNumWords(MaxBits);
+	}
+
+	FORCEINLINE uint32 GetLastWordMask() const
+	{
+		const uint32 UnusedBits = (FBitSet::BitsPerWord - static_cast<uint32>(NumBits) % FBitSet::BitsPerWord) % FBitSet::BitsPerWord;
+		return ~0u >> UnusedBits;
+	}
+
+	FORCEINLINE static void SetWords(uint32* Words, int32 NumWords, bool bValue)
+	{
+		if (NumWords > 8)
+		{
+			FMemory::Memset(Words, bValue ? 0xff : 0, NumWords * sizeof(uint32));
+		}
+		else
+		{
+			uint32 Word = bValue ? ~0u : 0u;
+			for (int32 Idx = 0; Idx < NumWords; ++Idx)
+			{
+				Words[Idx] = Word;
+			}
+		}
+	}
+
 	template <typename BitArrayType>
 	static FORCEINLINE typename TEnableIf<TContainerTraits<BitArrayType>::MoveWillEmptyContainer>::Type MoveOrCopy(BitArrayType& ToArray, BitArrayType& FromArray)
 	{
@@ -358,11 +374,8 @@ public:
 			BitArray.Realloc(0);
 		}
 
-		// calc the number of dwords for all the bits
-		const int32 NumDWORDs = FMath::DivideAndRoundUp(BitArray.NumBits, NumBitsPerDWORD);
-
 		// serialize the data as one big chunk
-		Ar.Serialize(BitArray.GetData(), NumDWORDs * sizeof(uint32));
+		Ar.Serialize(BitArray.GetData(), BitArray.GetNumWords() * sizeof(uint32));
 
 		return Ar;
 	}
@@ -409,15 +422,22 @@ public:
 	 */
 	void Empty(int32 ExpectedNumBits = 0)
 	{
-		NumBits = 0;
+		ExpectedNumBits = static_cast<int32>(FBitSet::CalculateNumWords(ExpectedNumBits)) * NumBitsPerDWORD;
+		const int32 InitialMaxBits = AllocatorInstance.GetInitialCapacity() * NumBitsPerDWORD;
 
-		ExpectedNumBits = FMath::DivideAndRoundUp(ExpectedNumBits, NumBitsPerDWORD) * NumBitsPerDWORD;
-		// If the expected number of bits doesn't match the allocated number of bits, reallocate.
-		if(MaxBits != ExpectedNumBits)
+		// If we need more bits or can shrink our allocation, do so
+		if (ExpectedNumBits > MaxBits || MaxBits > InitialMaxBits)
 		{
-			MaxBits = ExpectedNumBits;
+			MaxBits = FMath::Max(ExpectedNumBits, InitialMaxBits);
 			Realloc(0);
 		}
+		else
+		{
+			// Reuse current initial allocation but clear used words
+			SetWords(GetData(), GetNumWords(), false);
+		}
+
+		NumBits = 0;
 	}
 
 	/**
@@ -430,8 +450,8 @@ public:
 		if (Number > MaxBits)
 		{
 			const uint32 MaxDWORDs = AllocatorInstance.CalculateSlackGrow(
-				FMath::DivideAndRoundUp(Number,  NumBitsPerDWORD),
-				FMath::DivideAndRoundUp(MaxBits, NumBitsPerDWORD),
+				FBitSet::CalculateNumWords(Number),
+				GetMaxWords(),
 				sizeof(uint32)
 				);
 			MaxBits = MaxDWORDs * NumBitsPerDWORD;
@@ -445,23 +465,43 @@ public:
 	void Reset()
 	{
 		// We need this because iterators often use whole DWORDs when masking, which includes off-the-end elements
-		FMemory::Memset(GetData(), 0, FMath::DivideAndRoundUp(NumBits, NumBitsPerDWORD) * sizeof(uint32));
+		SetWords(GetData(), GetNumWords(), false);
 
 		NumBits = 0;
 	}
 
 	/**
-	 * Resets the array's contents.
+	 * Resets the array's contents. Use TBitArray(bool bValue, int32 InNumBits) instead of default constructor and Init().
+	 *
 	 * @param Value - The value to initial the bits to.
 	 * @param NumBits - The number of bits in the array.
 	 */
-	void Init(bool Value,int32 InNumBits)
+	FORCEINLINE void Init(bool bValue, int32 InNumBits)
 	{
-		Empty(InNumBits);
-		if(InNumBits)
+		NumBits = InNumBits;
+		
+		const uint32 NumWords = GetNumWords();
+		const uint32 MaxWords = GetMaxWords();
+
+		if (NumWords > MaxWords)
+		{
+			AllocatorInstance.ResizeAllocation(0, NumWords, sizeof(uint32));
+			MaxBits = NumWords * NumBitsPerDWORD;
+
+			uint32* Words = GetData();
+			SetWords(Words, NumWords, bValue);
+			Words[NumWords - 1] &= GetLastWordMask();
+		}
+		else if (bValue & (NumWords > 0))
+		{
+			uint32* Words = GetData();
+			SetWords(Words, NumWords - 1, true);
+			Words[NumWords - 1] = GetLastWordMask();
+			SetWords(Words + NumWords, MaxWords - NumWords, false);
+		}
+		else
 		{
-			NumBits = InNumBits;
-			FMemory::Memset(GetData(),Value ? 0xff : 0, FMath::DivideAndRoundUp(NumBits, NumBitsPerDWORD) * sizeof(uint32));
+			SetWords(GetData(), MaxWords, false);
 		}
 	}
 
@@ -473,9 +513,9 @@ public:
 
 		if (InNumBits > MaxBits)
 		{
-			const int32 PreviousNumDWORDs = FMath::DivideAndRoundUp(PreviousNumBits, NumBitsPerDWORD);
+			const int32 PreviousNumDWORDs = FBitSet::CalculateNumWords(PreviousNumBits);
 			const uint32 MaxDWORDs = AllocatorInstance.CalculateSlackReserve(
-				FMath::DivideAndRoundUp(InNumBits, NumBitsPerDWORD), sizeof(uint32));
+				FBitSet::CalculateNumWords(InNumBits), sizeof(uint32));
 			
 			AllocatorInstance.ResizeAllocation(PreviousNumDWORDs, MaxDWORDs, sizeof(uint32));	
 
@@ -614,15 +654,15 @@ public:
 	 */
 	uint32 GetAllocatedSize( void ) const
 	{
-		return FMath::DivideAndRoundUp(MaxBits, NumBitsPerDWORD) * sizeof(uint32);
+		return FBitSet::CalculateNumWords(MaxBits) * sizeof(uint32);
 	}
 
 	/** Tracks the container's memory use through an archive. */
 	void CountBytes(FArchive& Ar) const
 	{
 		Ar.CountBytes(
-			FMath::DivideAndRoundUp(NumBits, NumBitsPerDWORD) * sizeof(uint32),
-			FMath::DivideAndRoundUp(MaxBits, NumBitsPerDWORD) * sizeof(uint32)
+			GetNumWords() * sizeof(uint32),
+			GetMaxWords() * sizeof(uint32)
 		);
 	}
 
@@ -637,7 +677,7 @@ public:
 
 		const uint32* RESTRICT DwordArray = GetData();
 		const int32 LocalNumBits = NumBits;
-		const int32 DwordCount = FMath::DivideAndRoundUp(LocalNumBits, NumBitsPerDWORD);
+		const int32 DwordCount = FBitSet::CalculateNumWords(LocalNumBits);
 		int32 DwordIndex = 0;
 		while (DwordIndex < DwordCount && DwordArray[DwordIndex] == Test)
 		{
@@ -672,7 +712,7 @@ public:
 		uint32 Mask = ~0u >> (NumBitsPerDWORD - SlackIndex);
 
 		// Iterate over the array until we see a word with a zero bit.
-		uint32 DwordIndex = FMath::DivideAndRoundUp(LocalNumBits, NumBitsPerDWORD);
+		uint32 DwordIndex = FBitSet::CalculateNumWords(LocalNumBits);
 		const uint32* RESTRICT DwordArray = GetData();
 		const uint32 Test = bValue ? 0u : ~0u;
 		for (;;)
@@ -713,7 +753,7 @@ public:
 		// Iterate over the array until we see a word with a zero bit.
 		uint32* RESTRICT DwordArray = GetData();
 		const int32 LocalNumBits = NumBits;
-		const int32 DwordCount = FMath::DivideAndRoundUp(LocalNumBits, NumBitsPerDWORD);
+		const int32 DwordCount = FBitSet::CalculateNumWords(LocalNumBits);
 		int32 DwordIndex = FMath::DivideAndRoundDown(ConservativeStartIndex, NumBitsPerDWORD);
 		while (DwordIndex < DwordCount && DwordArray[DwordIndex] == (uint32)-1)
 		{
@@ -750,7 +790,7 @@ public:
 		uint32 Mask = ~0u >> (NumBitsPerDWORD - SlackIndex);
 
 		// Iterate over the array until we see a word with a zero bit.
-		uint32 DwordIndex = FMath::DivideAndRoundUp(LocalNumBits, NumBitsPerDWORD);
+		uint32 DwordIndex = FBitSet::CalculateNumWords(LocalNumBits);
 		uint32* RESTRICT DwordArray = GetData();
 		for (;;)
 		{
@@ -960,8 +1000,8 @@ private:
 
 	FORCENOINLINE void Realloc(int32 PreviousNumBits)
 	{
-		const int32 PreviousNumDWORDs = FMath::DivideAndRoundUp(PreviousNumBits, NumBitsPerDWORD);
-		const int32 MaxDWORDs = FMath::DivideAndRoundUp(MaxBits, NumBitsPerDWORD);
+		const uint32 PreviousNumDWORDs = FBitSet::CalculateNumWords(PreviousNumBits);
+		const uint32 MaxDWORDs = FBitSet::CalculateNumWords(MaxBits);
 
 		AllocatorInstance.ResizeAllocation(PreviousNumDWORDs,MaxDWORDs,sizeof(uint32));
 
@@ -1013,7 +1053,7 @@ DECLARE_TEMPLATE_INTRINSIC_TYPE_LAYOUT(template<typename Allocator>, TBitArray<A
 template<typename Allocator>
 FORCEINLINE uint32 GetTypeHash(const TBitArray<Allocator>& BitArray)
 {
-	uint32 NumWords = FMath::DivideAndRoundUp(BitArray.Num(), NumBitsPerDWORD);
+	uint32 NumWords = FBitSet::CalculateNumWords(BitArray.Num());
 	uint32 Hash = NumWords;
 	const uint32* Data = BitArray.GetData();
 	for (uint32 i = 0; i < NumWords; i++)
@@ -1301,7 +1341,7 @@ public:
 	{
 		NumBits = 0;
 
-		Slack = FMath::DivideAndRoundUp(Slack, NumBitsPerDWORD) * NumBitsPerDWORD;
+		Slack = FBitSet::CalculateNumWords(Slack) * NumBitsPerDWORD;
 		// If the expected number of bits doesn't match the allocated number of bits, reallocate.
 		if (MaxBits != Slack)
 		{
@@ -1363,11 +1403,11 @@ private:
 	FORCENOINLINE void Realloc(int32 PreviousNumBits)
 	{
 		const uint32 MaxDWORDs = AllocatorInstance.CalculateSlackReserve(
-			FMath::DivideAndRoundUp(MaxBits, NumBitsPerDWORD),
+			FBitSet::CalculateNumWords(MaxBits),
 			sizeof(uint32)
 			);
 		MaxBits = MaxDWORDs * NumBitsPerDWORD;
-		const uint32 PreviousNumDWORDs = FMath::DivideAndRoundUp(PreviousNumBits, NumBitsPerDWORD);
+		const uint32 PreviousNumDWORDs = FBitSet::CalculateNumWords(PreviousNumBits);
 
 		AllocatorInstance.ResizeAllocation(PreviousNumDWORDs, MaxDWORDs, sizeof(uint32));
 
@@ -1381,12 +1421,12 @@ private:
 	{
 		// Allocate memory for the new bits.
 		const uint32 MaxDWORDs = AllocatorInstance.CalculateSlackGrow(
-			FMath::DivideAndRoundUp(NumBits, NumBitsPerDWORD),
-			FMath::DivideAndRoundUp(MaxBits, NumBitsPerDWORD),
+			FBitSet::CalculateNumWords(NumBits),
+			FBitSet::CalculateNumWords(MaxBits),
 			sizeof(uint32)
 			);
 		MaxBits = MaxDWORDs * NumBitsPerDWORD;
-		const uint32 PreviousNumDWORDs = FMath::DivideAndRoundUp(PreviousNumBits, NumBitsPerDWORD);
+		const uint32 PreviousNumDWORDs = FBitSet::CalculateNumWords(PreviousNumBits);
 		AllocatorInstance.ResizeAllocation(PreviousNumDWORDs, MaxDWORDs, sizeof(uint32));
 		if (MaxDWORDs && MaxDWORDs > PreviousNumDWORDs)
 		{
