Enable CONTAINER_INITIAL_ALLOC_ZERO_SLACK (cross platform memory saving, saves up to 96MB on consoles)

diff --git a/Engine/Source/Runtime/Core/Public/Containers/ContainerAllocationPolicies.h b/Engine/Source/Runtime/Core/Public/Containers/ContainerAllocationPolicies.h
--- a/Engine/Source/Runtime/Core/Public/Containers/ContainerAllocationPolicies.h
+++ b/Engine/Source/Runtime/Core/Public/Containers/ContainerAllocationPolicies.h
@@ -15,7 +15,7 @@
 // is called. This tends to save a lot of memory with almost no measured performance cost.
 // NOTE: This can cause latent memory corruption issues to become more prominent
 #ifndef CONTAINER_INITIAL_ALLOC_ZERO_SLACK
-#define CONTAINER_INITIAL_ALLOC_ZERO_SLACK 0
+#define CONTAINER_INITIAL_ALLOC_ZERO_SLACK 1 // ON
 #endif
 
 class FDefaultBitArrayAllocator;
