Fix for fixed buffer overrun and static buffer reuse in crash context.

In some scenarios when copying the shared crash context the engine would overrun the fixed size string buffer. This issue appeared after two ensures fired, since the dynamic string buffers lifetime was equal to the crash reporting thread, and only reset on initialization.

diff --git a/Engine/Source/Runtime/Core/Private/GenericPlatform/GenericPlatformCrashContext.cpp b/Engine/Source/Runtime/Core/Private/GenericPlatform/GenericPlatformCrashContext.cpp
--- a/Engine/Source/Runtime/Core/Private/GenericPlatform/GenericPlatformCrashContext.cpp
+++ b/Engine/Source/Runtime/Core/Private/GenericPlatform/GenericPlatformCrashContext.cpp
@@ -299,16 +299,19 @@ void FGenericCrashContext::CopySharedCrashContext(FSharedCrashContext& Dst)
 	//Copy the session
 	FMemory::Memcpy(Dst.SessionContext, NCached::Session);
 	FMemory::Memcpy(Dst.UserSettings, NCached::UserSettings);
+	FMemory::Memset(Dst.DynamicData, 0);
 
 	TCHAR* DynamicDataStart = &Dst.DynamicData[0];
 	TCHAR* DynamicDataPtr = DynamicDataStart;
 
+	#define CR_DYNAMIC_BUFFER_REMAIN uint32((CR_MAX_DYNAMIC_BUFFER_CHARS) - (DynamicDataPtr-DynamicDataStart))
+
 	Dst.EnabledPluginsOffset = (uint32)(DynamicDataPtr - DynamicDataStart);
 	Dst.EnabledPluginsNum = NCached::EnabledPluginsList.Num();
 	for (const FString& Plugin : NCached::EnabledPluginsList)
 	{
-		FCString::Strcat(DynamicDataPtr, Plugin.Len(), *Plugin);
-		FCString::Strcat(DynamicDataPtr, 1, CR_PAIR_DELIM);
+		FCString::Strncat(DynamicDataPtr, *Plugin, CR_DYNAMIC_BUFFER_REMAIN);
+		FCString::Strncat(DynamicDataPtr, CR_PAIR_DELIM, CR_DYNAMIC_BUFFER_REMAIN);
 	}
 	DynamicDataPtr += FCString::Strlen(DynamicDataPtr) + 1;
 
@@ -316,10 +319,10 @@ void FGenericCrashContext::CopySharedCrashContext(FSharedCrashContext& Dst)
 	Dst.EngineDataNum = NCached::EngineData.Num();
 	for (const TPair<FString, FString>& Pair : NCached::EngineData)
 	{
-		FCString::Strcat(DynamicDataPtr, Pair.Key.Len(), *Pair.Key);
-		FCString::Strcat(DynamicDataPtr, 1, CR_PAIR_EQ);
-		FCString::Strcat(DynamicDataPtr, Pair.Value.Len(), *Pair.Value);
-		FCString::Strcat(DynamicDataPtr, 1, CR_PAIR_DELIM);
+		FCString::Strncat(DynamicDataPtr, *Pair.Key, CR_DYNAMIC_BUFFER_REMAIN);
+		FCString::Strncat(DynamicDataPtr, CR_PAIR_EQ, CR_DYNAMIC_BUFFER_REMAIN);
+		FCString::Strncat(DynamicDataPtr, *Pair.Value, CR_DYNAMIC_BUFFER_REMAIN);
+		FCString::Strncat(DynamicDataPtr, CR_PAIR_DELIM, CR_DYNAMIC_BUFFER_REMAIN);
 	}
 	DynamicDataPtr += FCString::Strlen(DynamicDataPtr) + 1;
 
@@ -327,12 +330,14 @@ void FGenericCrashContext::CopySharedCrashContext(FSharedCrashContext& Dst)
 	Dst.GameDataNum = NCached::GameData.Num();
 	for (const TPair<FString, FString>& Pair : NCached::GameData)
 	{
-		FCString::Strcat(DynamicDataPtr, Pair.Key.Len(), *Pair.Key);
-		FCString::Strcat(DynamicDataPtr, 1, CR_PAIR_EQ);
-		FCString::Strcat(DynamicDataPtr, Pair.Value.Len(), *Pair.Value);
-		FCString::Strcat(DynamicDataPtr, 1, CR_PAIR_DELIM);
+		FCString::Strncat(DynamicDataPtr, *Pair.Key, CR_DYNAMIC_BUFFER_REMAIN);
+		FCString::Strncat(DynamicDataPtr, CR_PAIR_EQ, CR_DYNAMIC_BUFFER_REMAIN);
+		FCString::Strncat(DynamicDataPtr, *Pair.Value, CR_DYNAMIC_BUFFER_REMAIN);
+		FCString::Strncat(DynamicDataPtr, CR_PAIR_DELIM, CR_DYNAMIC_BUFFER_REMAIN);
 	}
 	DynamicDataPtr += FCString::Strlen(DynamicDataPtr) + 1;
+
+	#undef CR_DYNAMIC_BUFFER_REMAIN
 }
 
 void FGenericCrashContext::SetMemoryStats(const FPlatformMemoryStats& InMemoryStats)
