Fix signed/unsigned comparison errors building debug configuration on Win32 (size_t is 32-bit)

#rnx

diff --git a/Engine/Source/Runtime/Core/Public/Misc/MemStack.h b/Engine/Source/Runtime/Core/Public/Misc/MemStack.h
--- a/Engine/Source/Runtime/Core/Public/Misc/MemStack.h
+++ b/Engine/Source/Runtime/Core/Public/Misc/MemStack.h
@@ -230,14 +230,14 @@ inline void* operator new(size_t Size, FMemStackBase& Mem, int32 Count = 1, int3
 {
 	// Get uninitialized memory.
 	const size_t SizeInBytes = Size * Count;
-	checkSlow(SizeInBytes <= TNumericLimits<int32>::Max());
+	checkSlow(SizeInBytes <= (size_t)TNumericLimits<int32>::Max());
 	return Mem.PushBytes( (int32)SizeInBytes, Align );
 }
 inline void* operator new(size_t Size, FMemStackBase& Mem, EMemZeroed Tag, int32 Count = 1, int32 Align = DEFAULT_ALIGNMENT)
 {
 	// Get zero-filled memory.
 	const size_t SizeInBytes = Size * Count;
-	checkSlow(SizeInBytes <= TNumericLimits<int32>::Max());
+	checkSlow(SizeInBytes <= (size_t)TNumericLimits<int32>::Max());
 	uint8* Result = Mem.PushBytes( (int32)SizeInBytes, Align );
 	FMemory::Memzero( Result, SizeInBytes );
 	return Result;
@@ -246,7 +246,7 @@ inline void* operator new(size_t Size, FMemStackBase& Mem, EMemOned Tag, int32 C
 {
 	// Get one-filled memory.
 	const size_t SizeInBytes = Size * Count;
-	checkSlow(SizeInBytes <= TNumericLimits<int32>::Max());
+	checkSlow(SizeInBytes <= (size_t)TNumericLimits<int32>::Max());
 	uint8* Result = Mem.PushBytes( (int32)SizeInBytes, Align );
 	FMemory::Memset( Result, 0xff, SizeInBytes );
 	return Result;
@@ -255,14 +255,14 @@ inline void* operator new[](size_t Size, FMemStackBase& Mem, int32 Count = 1, in
 {
 	// Get uninitialized memory.
 	const size_t SizeInBytes = Size * Count;
-	checkSlow(SizeInBytes <= TNumericLimits<int32>::Max());
+	checkSlow(SizeInBytes <= (size_t)TNumericLimits<int32>::Max());
 	return Mem.PushBytes( (int32)SizeInBytes, Align );
 }
 inline void* operator new[](size_t Size, FMemStackBase& Mem, EMemZeroed Tag, int32 Count = 1, int32 Align = DEFAULT_ALIGNMENT)
 {
 	// Get zero-filled memory.
 	const size_t SizeInBytes = Size * Count;
-	checkSlow(SizeInBytes <= TNumericLimits<int32>::Max());
+	checkSlow(SizeInBytes <= (size_t)TNumericLimits<int32>::Max());
 	uint8* Result = Mem.PushBytes( (int32)SizeInBytes, Align );
 	FMemory::Memzero( Result, SizeInBytes );
 	return Result;
@@ -271,7 +271,7 @@ inline void* operator new[](size_t Size, FMemStackBase& Mem, EMemOned Tag, int32
 {
 	// Get one-filled memory.
 	const size_t SizeInBytes = Size * Count;
-	checkSlow(SizeInBytes <= TNumericLimits<int32>::Max());
+	checkSlow(SizeInBytes <= (size_t)TNumericLimits<int32>::Max());
 	uint8* Result = Mem.PushBytes( (int32)SizeInBytes, Align );
 	FMemory::Memset( Result, 0xff, SizeInBytes );
 	return Result;
