MN - Disable Staic Analysis false positive.

diff --git a/Engine/Source/Runtime/Core/Private/Compression/lz4.cpp b/Engine/Source/Runtime/Core/Private/Compression/lz4.cpp
--- a/Engine/Source/Runtime/Core/Private/Compression/lz4.cpp
+++ b/Engine/Source/Runtime/Core/Private/Compression/lz4.cpp
@@ -462,6 +462,7 @@ LZ4_memcpy_using_offset(BYTE* dstPtr, const BYTE* srcPtr, BYTE* dstEnd, const si
 static unsigned LZ4_NbCommonBytes (reg_t val)
 {
     if (LZ4_isLittleEndian()) {
+        CA_SUPPRESS(6239);
         if (sizeof(val)==8) {
 #       if defined(_MSC_VER) && defined(_WIN64) && !defined(LZ4_FORCE_SW_BITCOUNT)
             unsigned long r = 0;
@@ -496,6 +497,7 @@ static unsigned LZ4_NbCommonBytes (reg_t val)
 #       endif
         }
     } else   /* Big Endian CPU */ {
+        CA_SUPPRESS(6239);
         if (sizeof(val)==8) {   /* 64-bits */
 #       if defined(_MSC_VER) && defined(_WIN64) && !defined(LZ4_FORCE_SW_BITCOUNT)
             unsigned long r = 0;
@@ -551,6 +553,7 @@ unsigned LZ4_count(const BYTE* pIn, const BYTE* pMatch, const BYTE* pInLimit)
         return (unsigned)(pIn - pStart);
     }
 
+    CA_SUPPRESS(6239);
     if ((STEPSIZE==8) && (pIn<(pInLimit-3)) && (LZ4_read32(pMatch) == LZ4_read32(pIn))) { pIn+=4; pMatch+=4; }
     if ((pIn<(pInLimit-1)) && (LZ4_read16(pMatch) == LZ4_read16(pIn))) { pIn+=2; pMatch+=2; }
     if ((pIn<pInLimit) && (*pMatch == *pIn)) pIn++;
@@ -649,6 +652,7 @@ static U32 LZ4_hash5(U64 sequence, tableType_t const tableType)
 
 LZ4_FORCE_INLINE U32 LZ4_hashPosition(const void* const p, tableType_t const tableType)
 {
+    CA_SUPPRESS(6239);
     if ((sizeof(reg_t)==8) && (tableType != byU16)) return LZ4_hash5(LZ4_read_ARCH(p), tableType);
     return LZ4_hash4(LZ4_read32(p), tableType);
 }
