Android Crash Reporter
 - Dump logcat on non fatal reports.

diff --git a/Engine/Source/Runtime/Core/Private/Android/AndroidPlatformCrashContext.cpp b/Engine/Source/Runtime/Core/Private/Android/AndroidPlatformCrashContext.cpp
--- a/Engine/Source/Runtime/Core/Private/Android/AndroidPlatformCrashContext.cpp
+++ b/Engine/Source/Runtime/Core/Private/Android/AndroidPlatformCrashContext.cpp
@@ -65,12 +65,12 @@ struct FAndroidCrashInfo
 	bool bInitialized = false;
 } GAndroidCrashInfo;
 
-const FString FAndroidCrashContext::GetCrashDirectoryName()
+const FString FAndroidCrashContext::GetGlobalCrashDirectoryPath()
 {
 	return FString(GAndroidCrashInfo.TargetDirectory);
 }
 
-void FAndroidCrashContext::GetCrashDirectoryName(char(&DirectoryNameOUT)[CrashReportMaxPathSize])
+void FAndroidCrashContext::GetGlobalCrashDirectoryPath(char(&DirectoryNameOUT)[CrashReportMaxPathSize])
 {
 	FCStringAnsi::Strncpy(DirectoryNameOUT, GAndroidCrashInfo.TargetDirectory, CrashReportMaxPathSize);
 }
@@ -272,7 +272,7 @@ FAndroidCrashContext::FAndroidCrashContext(ECrashContextType InType, const TCHAR
 	}
 	else
 	{
-		GetCrashDirectoryName(ReportDirectory);
+		GetGlobalCrashDirectoryPath(ReportDirectory);
 	}
 }
 
diff --git a/Engine/Source/Runtime/Core/Public/Android/AndroidPlatformCrashContext.h b/Engine/Source/Runtime/Core/Public/Android/AndroidPlatformCrashContext.h
--- a/Engine/Source/Runtime/Core/Public/Android/AndroidPlatformCrashContext.h
+++ b/Engine/Source/Runtime/Core/Public/Android/AndroidPlatformCrashContext.h
@@ -45,8 +45,13 @@ struct CORE_API FAndroidCrashContext : public FGenericCrashContext
 	static const int32 CrashReportMaxPathSize = 512;
 
 	static void Initialize();
-	static const FString GetCrashDirectoryName();
-	static void GetCrashDirectoryName(char(&DirectoryNameOUT)[CrashReportMaxPathSize]);
+	// Returns the report directory used for this crash context.
+	FString GetCurrentReportDirectoryPath() const {	return FString(ReportDirectory); }
+
+	// Returns the main crash directory for this session. This will not be correct for non-fatal reports.
+	static const FString GetGlobalCrashDirectoryPath();
+	// Fills DirectoryNameOUT with the global crash directory for a fatal crash this session. This will not be correct for non-fatal reports.
+	static void GetGlobalCrashDirectoryPath(char(&DirectoryNameOUT)[CrashReportMaxPathSize]);
 
 	void AddAndroidCrashProperty(const FString& Key, const FString& Value);
 
