Add support for passing "abscrashreportclientlog" and "nullrhi" command line arguments to crash reporter monitor.

This allows enabling crash reporter client to run in monitor mode for game runtimes again.

diff --git a/Engine/Source/Runtime/Core/Private/Windows/WindowsPlatformCrashContext.cpp b/Engine/Source/Runtime/Core/Private/Windows/WindowsPlatformCrashContext.cpp
--- a/Engine/Source/Runtime/Core/Private/Windows/WindowsPlatformCrashContext.cpp
+++ b/Engine/Source/Runtime/Core/Private/Windows/WindowsPlatformCrashContext.cpp
@@ -31,12 +31,13 @@
 #include <Shlwapi.h>
 #include <psapi.h>
 #include <tlhelp32.h>
+#include <shellapi.h>
 
 #ifndef UE_LOG_CRASH_CALLSTACK
 	#define UE_LOG_CRASH_CALLSTACK 1
 #endif
 
-#if WITH_EDITOR
+#if !IS_PROGRAM
 	#define USE_CRASH_REPORTER_MONITOR 1
 #else 
 	#define USE_CRASH_REPORTER_MONITOR 0
@@ -414,6 +415,35 @@ FProcHandle LaunchCrashReportClient(void** OutWritePipe, void** OutReadPipe)
 		FCString::Strncat(CrashReporterClientArgs, PidStr, CR_CLIENT_MAX_ARGS_LEN);
 	}
 
+	// Parse commandline arguments relevant to pass to the client. Note that since we run this from static initialization
+	// FCommandline has not yet been initialized, instead we need to use the OS provided methods.
+	{
+		LPWSTR* ArgList;
+		int ArgCount;
+		ArgList = CommandLineToArgvW(GetCommandLineW(), &ArgCount);
+		if (ArgList != nullptr)
+		{
+			for (int It = 0; It < ArgCount; ++It)
+			{
+				TCHAR Path[MAX_PATH];
+				if (FParse::Value(ArgList[It], TEXT("abscrashreportclientlog="), Path, UE_ARRAY_COUNT(Path)))
+				{
+					FCString::Strncat(CrashReporterClientArgs, TEXT(" -abslog="), CR_CLIENT_MAX_ARGS_LEN);
+					FCString::Strncat(CrashReporterClientArgs, Path, CR_CLIENT_MAX_ARGS_LEN);
+				}
+				
+			#if !USE_NULL_RHI
+				const bool bHasNullRHIOnCommandline = FParse::Param(ArgList[It], TEXT("nullrhi"));
+				if (bHasNullRHIOnCommandline)
+			#endif
+				{
+					FCString::Strncat(CrashReporterClientArgs, TEXT(" -nullrhi"), CR_CLIENT_MAX_ARGS_LEN);
+				}
+			}
+		}
+	}
+
+
 #if WITH_EDITOR // Disaster recovery is only enabled for the Editor. Start the server even if in -game, -server, commandlet, the client-side will not connect (its too soon here to query this executable config).
 	{
 		// Disaster recovery service command line.
