Provide more control over OBB selection

#ue4
#android

diff --git a/Engine/Build/Android/Java/src/com/epicgames/ue4/GameActivity.java.template b/Engine/Build/Android/Java/src/com/epicgames/ue4/GameActivity.java.template
--- a/Engine/Build/Android/Java/src/com/epicgames/ue4/GameActivity.java.template
+++ b/Engine/Build/Android/Java/src/com/epicgames/ue4/GameActivity.java.template
@@ -2699,9 +2699,11 @@ public class GameActivity extends $${gameActivitySuperClass}$$ implements Surfac
 			return;
 		}
 
-		Log.debug("APK path: " + getPackageResourcePath());
+		String APKPath = getPackageResourcePath();
+//$${gameActivityOverrideAPKOBBPackaging}$$
+		Log.debug("APK path: " + APKPath);
 		Log.debug("OBB in APK: " + (PackageDataInsideApkValue==1));
-		nativeSetGlobalActivity(UseExternalFilesDir, PublicLogFiles, getFilesDir().getPath(), getExternalFilesDir(null).getPath(), PackageDataInsideApkValue==1, getPackageResourcePath());
+		nativeSetGlobalActivity(UseExternalFilesDir, PublicLogFiles, getFilesDir().getPath(), getExternalFilesDir(null).getPath(), PackageDataInsideApkValue==1, APKPath);
 
 		// tell the engine if this is a portrait app
 		nativeSetWindowInfo(getResources().getConfiguration().orientation == Configuration.ORIENTATION_PORTRAIT, DepthBufferPreference);
@@ -6108,6 +6110,7 @@ public class GameActivity extends $${gameActivitySuperClass}$$ implements Surfac
 	public native boolean nativeIsShippingBuild();
 	public native void nativeSetAndroidStartupState(boolean bDebuggerAttached);
 	public native void nativeSetGlobalActivity(boolean bUseExternalFilesDir, boolean bPublicLogFiles, String internalFilePath, String externalFilePath, boolean bOBBInAPK, String APKPath);
+	public native void nativeSetObbFilePaths(String OBBMainFilePath, String OBBPatchFilePath);
 	public native void nativeSetWindowInfo(boolean bIsPortrait, int DepthBufferPreference);
 	public native void nativeSetObbInfo(String ProjectName, String PackageName, int Version, int PatchVersion, String AppType);
 	public native void nativeSetAndroidVersionInformation( String AndroidVersion, String PhoneMake, String PhoneModel, String PhoneBuildNumber, String OSLanguage );
diff --git a/Engine/Config/BaseEngine.ini b/Engine/Config/BaseEngine.ini
--- a/Engine/Config/BaseEngine.ini
+++ b/Engine/Config/BaseEngine.ini
@@ -2102,6 +2102,7 @@ bBundleLanguageSplit=True
 bBundleDensitySplit=True
 bFullScreen=True
 bForceLDLinker=False
+bForceSmallOBBFiles=False
 
 [/Script/AndroidPlatformEditor.AndroidSDKSettings]
 SDKAPILevel=latest
diff --git a/Engine/Source/Programs/AutomationTool/Android/AndroidPlatform.Automation.cs b/Engine/Source/Programs/AutomationTool/Android/AndroidPlatform.Automation.cs
--- a/Engine/Source/Programs/AutomationTool/Android/AndroidPlatform.Automation.cs
+++ b/Engine/Source/Programs/AutomationTool/Android/AndroidPlatform.Automation.cs
@@ -14,7 +14,8 @@ using Tools.DotNETCommon;
 
 public class AndroidPlatform : Platform
 {
-	// Maximum allowed OBB size (2 GiB or 4 GiB based on project setting)
+	// Maximum allowed OBB size (500 MiB, 2 GiB or 4 GiB based on project settings)
+	private const Int64 SmallOBBSizeAllowed = 524288000;
 	private const Int64 NormalOBBSizeAllowed = 2147483648;
 	private const Int64 MaxOBBSizeAllowed = 4294967296;
 
@@ -312,9 +313,11 @@ public class AndroidPlatform : Platform
 	private Int64 GetMaxOBBSizeAllowed(DeploymentContext SC)
 	{
 		ConfigHierarchy Ini = ConfigCache.ReadHierarchy(ConfigHierarchyType.Engine, DirectoryReference.FromFile(SC.RawProjectPath), SC.StageTargetPlatform.PlatformType);
+		bool bForceSmallOBBFiles = false;
 		bool bAllowLargeOBBFiles = false;
+		Ini.GetBool("/Script/AndroidRuntimeSettings.AndroidRuntimeSettings", "bForceSmallOBBFiles", out bForceSmallOBBFiles);
 		Ini.GetBool("/Script/AndroidRuntimeSettings.AndroidRuntimeSettings", "bAllowLargeOBBFiles", out bAllowLargeOBBFiles);
-		return bAllowLargeOBBFiles ? MaxOBBSizeAllowed : NormalOBBSizeAllowed;
+		return bForceSmallOBBFiles ? SmallOBBSizeAllowed : (bAllowLargeOBBFiles ? MaxOBBSizeAllowed : NormalOBBSizeAllowed);
 	}
 
 	private bool AllowPatchOBBFile(DeploymentContext SC)
@@ -449,7 +452,7 @@ public class AndroidPlatform : Platform
 			OBBNeedsUpdate = true;
 		}
 		Int64 OBBSizeAllowed = GetMaxOBBSizeAllowed(SC);
-		string LimitString = (OBBSizeAllowed < MaxOBBSizeAllowed) ? "2 GiB" : "4 GiB";
+		string LimitString = (OBBSizeAllowed < NormalOBBSizeAllowed) ? "500 MiB" : ((OBBSizeAllowed < MaxOBBSizeAllowed) ? "2 GiB" : "4 GiB");
 
 		if (!OBBNeedsUpdate)
 		{
diff --git a/Engine/Source/Programs/UnrealBuildTool/Platform/Android/UEDeployAndroid.cs b/Engine/Source/Programs/UnrealBuildTool/Platform/Android/UEDeployAndroid.cs
--- a/Engine/Source/Programs/UnrealBuildTool/Platform/Android/UEDeployAndroid.cs
+++ b/Engine/Source/Programs/UnrealBuildTool/Platform/Android/UEDeployAndroid.cs
@@ -4467,6 +4467,7 @@ namespace UnrealBuildTool
 				{ "//$${gameActivityReadMetadataAdditions}$$", UPL.ProcessPluginNode(NDKArch, "gameActivityReadMetadataAdditions", "")},
 				{ "//$${gameActivityOnCreateAdditions}$$", UPL.ProcessPluginNode(NDKArch, "gameActivityOnCreateAdditions", "")},
 				{ "//$${gameActivityOnCreateFinalAdditions}$$", UPL.ProcessPluginNode(NDKArch, "gameActivityOnCreateFinalAdditions", "")},
+				{ "//$${gameActivityOverrideAPKOBBPackaging}$$", UPL.ProcessPluginNode(NDKArch, "gameActivityOverrideAPKOBBPackaging", "")},
 				{ "//$${gameActivityOnDestroyAdditions}$$", UPL.ProcessPluginNode(NDKArch, "gameActivityOnDestroyAdditions", "")},
 				{ "//$${gameActivityOnStartAdditions}$$", UPL.ProcessPluginNode(NDKArch, "gameActivityOnStartAdditions", "")},
 				{ "//$${gameActivityOnStopAdditions}$$", UPL.ProcessPluginNode(NDKArch, "gameActivityOnStopAdditions", "")},
diff --git a/Engine/Source/Runtime/Android/AndroidRuntimeSettings/Classes/AndroidRuntimeSettings.h b/Engine/Source/Runtime/Android/AndroidRuntimeSettings/Classes/AndroidRuntimeSettings.h
--- a/Engine/Source/Runtime/Android/AndroidRuntimeSettings/Classes/AndroidRuntimeSettings.h
+++ b/Engine/Source/Runtime/Android/AndroidRuntimeSettings/Classes/AndroidRuntimeSettings.h
@@ -237,6 +237,10 @@ public:
 	UPROPERTY(GlobalConfig, EditAnywhere, Category = "APK Packaging", Meta = (DisplayName = "Disable verify OBB on first start/update."))
 	bool bDisableVerifyOBBOnStartUp;
 
+	// If checked, OBB is limited to 500 MiB.
+	UPROPERTY(GlobalConfig, EditAnywhere, Category = "APK Packaging", Meta = (DisplayName = "Force small OBB files."))
+	bool bForceSmallOBBFiles;
+
 	// If checked, OBB is not limited to 2 GiB allowed by Google Play Store (still limited to 4 GiB ZIP limit). 
 	UPROPERTY(GlobalConfig, EditAnywhere, Category = "APK Packaging", Meta = (DisplayName = "Allow large OBB files."))
 	bool bAllowLargeOBBFiles;
diff --git a/Engine/Source/Runtime/Core/Private/Android/AndroidPlatformFile.cpp b/Engine/Source/Runtime/Core/Private/Android/AndroidPlatformFile.cpp
--- a/Engine/Source/Runtime/Core/Private/Android/AndroidPlatformFile.cpp
+++ b/Engine/Source/Runtime/Core/Private/Android/AndroidPlatformFile.cpp
@@ -81,6 +81,10 @@ FString GAndroidAppType;
 FString GFilePathBase;
 // Obb File Path base - setup during load
 FString GOBBFilePathBase;
+// Obb Main filepath
+FString GOBBMainFilePath;
+// Obb Patch filepath
+FString GOBBPatchFilePath;
 // Internal File Direcory Path (for application) - setup during load
 FString GInternalFilePath;
 // External File Direcory Path (for application) - setup during load
@@ -988,7 +992,17 @@ public:
 				FFileHandleAndroid* OBBFile = static_cast<FFileHandleAndroid*>(new FFileHandleAndroid(*OBBEntry.File, 0, OBBEntry.File->Size()));
 				check(nullptr != OBBFile);
 				ZipResource.AddPatchFile(MakeShareable(OBBFile));
-				FPlatformMisc::LowLevelOutputDebugStringf(TEXT("Mounted OBB in APK: %s"), *GAPKFilename);
+				FPlatformMisc::LowLevelOutputDebugStringf(TEXT("Mounted main OBB in APK: %s"), *GAPKFilename);
+
+				// check for optional patch obb in APK
+				if (APKZip.HasEntry("assets/patch.obb.png"))
+				{
+					auto patchOBBEntry = APKZip.GetEntry("assets/patch.obb.png");
+					FFileHandleAndroid* patchOBBFile = static_cast<FFileHandleAndroid*>(new FFileHandleAndroid(*patchOBBEntry.File, 0, patchOBBEntry.File->Size()));
+					check(nullptr != patchOBBFile);
+					ZipResource.AddPatchFile(MakeShareable(patchOBBFile));
+					FPlatformMisc::LowLevelOutputDebugStringf(TEXT("Mounted patch OBB in APK: %s"), *GAPKFilename);
+				}
 			}
 			else
 			{
@@ -1001,25 +1015,42 @@ public:
 			// For external OBBs we mount the specific OBB files,
 			// main and patch, only. As required by Android specs.
 			// See <http://developer.android.com/google/play/expansion-files.html>
+			// but first checks for overrides of expected OBB file paths if provided
 			FString OBBDir1 = GOBBFilePathBase + FString(TEXT("/Android/obb/") + GPackageName);
 			FString OBBDir2 = GOBBFilePathBase + FString(TEXT("/obb/") + GPackageName);
 			FString MainOBBName = FString::Printf(TEXT("main.%d.%s.obb"), GAndroidPackageVersion, *GPackageName);
 			FString PatchOBBName = FString::Printf(TEXT("patch.%d.%s.obb"), GAndroidPackageVersion, *GPackageName);
-			if (FileExists(*(OBBDir1 / MainOBBName), true))
+
+			if (!GOBBMainFilePath.IsEmpty() && FileExists(*GOBBMainFilePath, true))
+			{
+				MountOBB(*GOBBMainFilePath);
+				FPlatformMisc::LowLevelOutputDebugStringf(TEXT("Mounted main OBB: %s"), *GOBBMainFilePath);
+			}
+			else if (FileExists(*(OBBDir1 / MainOBBName), true))
 			{
 				MountOBB(*(OBBDir1 / MainOBBName));
+				FPlatformMisc::LowLevelOutputDebugStringf(TEXT("Mounted main OBB: %s"), *(OBBDir1 / MainOBBName));
 			}
 			else if (FileExists(*(OBBDir2 / MainOBBName), true))
 			{
 				MountOBB(*(OBBDir2 / MainOBBName));
+				FPlatformMisc::LowLevelOutputDebugStringf(TEXT("Mounted main OBB: %s"), *(OBBDir2 / MainOBBName));
+			}
+
+			if (!GOBBPatchFilePath.IsEmpty() && FileExists(*GOBBPatchFilePath, true))
+			{
+				MountOBB(*GOBBPatchFilePath);
+				FPlatformMisc::LowLevelOutputDebugStringf(TEXT("Mounted patch OBB: %s"), *GOBBPatchFilePath);
 			}
-			if (FileExists(*(OBBDir1 / PatchOBBName), true))
+			else if (FileExists(*(OBBDir1 / PatchOBBName), true))
 			{
 				MountOBB(*(OBBDir1 / PatchOBBName));
+				FPlatformMisc::LowLevelOutputDebugStringf(TEXT("Mounted patch OBB: %s"), *(OBBDir1 / PatchOBBName));
 			}
 			else if (FileExists(*(OBBDir2 / PatchOBBName), true))
 			{
 				MountOBB(*(OBBDir2 / PatchOBBName));
+				FPlatformMisc::LowLevelOutputDebugStringf(TEXT("Mounted patch OBB: %s"), *(OBBDir2 / PatchOBBName));
 			}
 		}
 
diff --git a/Engine/Source/Runtime/Launch/Private/Android/AndroidJNI.cpp b/Engine/Source/Runtime/Launch/Private/Android/AndroidJNI.cpp
--- a/Engine/Source/Runtime/Launch/Private/Android/AndroidJNI.cpp
+++ b/Engine/Source/Runtime/Launch/Private/Android/AndroidJNI.cpp
@@ -43,6 +43,8 @@ extern FString GFontPathBase;
 extern bool GOBBinAPK;
 extern bool GOverrideAndroidLogDir;
 extern FString GOBBFilePathBase;
+extern FString GOBBMainFilePath;
+extern FString GOBBPatchFilePath;
 extern FString GAPKFilename;
 
 FOnActivityResult FJavaWrapper::OnActivityResultDelegate;
@@ -1568,6 +1570,8 @@ JNIEXPORT jint JNI_OnLoad(JavaVM* InJavaVM, void* InReserved)
 	// Copy that somewhere safe
 	GFilePathBase = FJavaHelper::FStringFromLocalRef(Env, (jstring)Env->CallObjectMethod(*externalStoragePath, getFilePath, nullptr));
 	GOBBFilePathBase = GFilePathBase;
+	GOBBMainFilePath = TEXT("");
+	GOBBPatchFilePath = TEXT("");
 
 	// then release...
 	FPlatformMisc::LowLevelOutputDebugStringf(TEXT("Path found as '%s'\n"), *GFilePathBase);
@@ -1591,6 +1595,13 @@ JNIEXPORT jint JNI_OnLoad(JavaVM* InJavaVM, void* InReserved)
 
 //Native-defined functions
 
+//This function is declared in the Java-defined class, GameActivity.java: "public native void naativeSetObbFilePaths();"
+JNI_METHOD void Java_com_epicgames_ue4_GameActivity_nativeSetObbFilePaths(JNIEnv* jenv, jobject thiz, jstring OBBMainFilePath, jstring OBBPatchFilePath)
+{
+	GOBBMainFilePath = FJavaHelper::FStringFromParam(jenv, OBBMainFilePath);
+	GOBBPatchFilePath = FJavaHelper::FStringFromParam(jenv, OBBPatchFilePath);
+}
+
 //This function is declared in the Java-defined class, GameActivity.java: "public native void nativeSetGlobalActivity();"
 JNI_METHOD void Java_com_epicgames_ue4_GameActivity_nativeSetGlobalActivity(JNIEnv* jenv, jobject thiz, jboolean bUseExternalFilesDir, jboolean bPublicLogFiles, jstring internalFilePath, jstring externalFilePath, jboolean bOBBinAPK, jstring APKFilename /*, jobject googleServices*/)
 {
