fix for crash while attempting to convert the minidump to human readable form

diff --git a/Engine/Source/Runtime/Core/Private/IOS/IOSPlatformMisc.cpp b/Engine/Source/Runtime/Core/Private/IOS/IOSPlatformMisc.cpp
--- a/Engine/Source/Runtime/Core/Private/IOS/IOSPlatformMisc.cpp
+++ b/Engine/Source/Runtime/Core/Private/IOS/IOSPlatformMisc.cpp
@@ -1768,27 +1768,6 @@ FIOSCrashContext::FIOSCrashContext(ECrashContextType InType, const TCHAR* InErro
 void FIOSCrashContext::CopyMinidump(char const* OutputPath, char const* InputPath) const
 {
 #if !PLATFORM_TVOS
-    NSError* Error = nil;
-    NSString* Path = FString(ANSI_TO_TCHAR(InputPath)).GetNSString();
-    NSData* CrashData = [NSData dataWithContentsOfFile: Path options: NSMappedRead error: &Error];
-    if (CrashData && !Error)
-    {
-        PLCrashReport* CrashLog = [[PLCrashReport alloc] initWithData: CrashData error: &Error];
-        if (CrashLog && !Error)
-        {
-            NSString* Report = [PLCrashReportTextFormatter stringValueForCrashReport: CrashLog withTextFormat: PLCrashReportTextFormatiOS];
-            FString CrashDump = FString(Report);
-            [Report writeToFile: Path atomically: YES encoding: NSUTF8StringEncoding error: &Error];
-        }
-        else
-        {
-            NSLog(@"****UE4 %@", [Error localizedDescription]);
-        }
-    }
-    else
-    {
-        NSLog(@"****UE4 %@", [Error localizedDescription]);
-    }
     int ReportFile = open(OutputPath, O_CREAT|O_WRONLY, 0766);
     int DumpFile = open(InputPath, O_RDONLY, 0766);
     if (ReportFile != -1 && DumpFile != -1)
@@ -1809,6 +1788,51 @@ void FIOSCrashContext::CopyMinidump(char const* OutputPath, char const* InputPat
 #endif
 }
 
+void FIOSCrashContext::ConvertMinidump(char const* OutputPath, char const* InputPath)
+{
+#if !PLATFORM_TVOS
+	NSError* Error = nil;
+	NSString* Path = FString(ANSI_TO_TCHAR(InputPath)).GetNSString();
+	NSData* CrashData = [NSData dataWithContentsOfFile: Path options: NSMappedRead error: &Error];
+	if (CrashData && !Error)
+	{
+		PLCrashReport* CrashLog = [[PLCrashReport alloc] initWithData: CrashData error: &Error];
+		if (CrashLog && !Error)
+		{
+			NSString* Report = [PLCrashReportTextFormatter stringValueForCrashReport: CrashLog withTextFormat: PLCrashReportTextFormatiOS];
+			FString CrashDump = FString(Report);
+			[Report writeToFile: Path atomically: YES encoding: NSUTF8StringEncoding error: &Error];
+		}
+		else
+		{
+			NSLog(@"****UE4 %@", [Error localizedDescription]);
+		}
+	}
+	else
+	{
+		NSLog(@"****UE4 %@", [Error localizedDescription]);
+	}
+	int ReportFile = open(OutputPath, O_CREAT|O_WRONLY, 0766);
+	int DumpFile = open(InputPath, O_RDONLY, 0766);
+	if (ReportFile != -1 && DumpFile != -1)
+	{
+		char Data[PATH_MAX];
+		
+		int Bytes = 0;
+		while((Bytes = read(DumpFile, Data, PATH_MAX)) > 0)
+		{
+			write(ReportFile, Data, Bytes);
+		}
+		
+		close(DumpFile);
+		close(ReportFile);
+		
+		unlink(InputPath);
+	}
+#endif
+}
+
+
 void FIOSCrashContext::GenerateInfoInFolder(char const* const InfoFolder, bool bIsEnsure) const
 {
     // create a crash-specific directory
diff --git a/Engine/Source/Runtime/Core/Public/IOS/IOSPlatformCrashContext.h b/Engine/Source/Runtime/Core/Public/IOS/IOSPlatformCrashContext.h
--- a/Engine/Source/Runtime/Core/Public/IOS/IOSPlatformCrashContext.h
+++ b/Engine/Source/Runtime/Core/Public/IOS/IOSPlatformCrashContext.h
@@ -23,6 +23,9 @@ struct CORE_API FIOSCrashContext : public FApplePlatformCrashContext
 	
 	/** Creates the crash folder */
 	FString CreateCrashFolder() const;
+
+	/** Converts the PLCrashReporter minidump */
+	static void ConvertMinidump(char const* OutputPath, char const* InputPath);
 };
 
 typedef FIOSCrashContext FPlatformCrashContext;
