Add a new large page allocator and allow the xbox to take advantage of 4MB pages.
This saves ~2ms on the game thread, ~1.6ms on the render thread and ~1ms on the RHI thread.
After 20 demo loops it uses 20MB more than just the normal cache (only as far as single large allocations is concerned).
Also adds a timer for virtual allocations so we can track how much time we spend doing it (checked in disabled).

diff --git a/Engine/Source/Runtime/Core/Private/HAL/Allocators/CachedOSVeryLargePageAllocator.cpp b/Engine/Source/Runtime/Core/Private/HAL/Allocators/CachedOSVeryLargePageAllocator.cpp
new file mode 100644
--- /dev/null
+++ b/Engine/Source/Runtime/Core/Private/HAL/Allocators/CachedOSVeryLargePageAllocator.cpp
@@ -0,0 +1,147 @@
+// Copyright Epic Games, Inc. All Rights Reserved.
+
+#include "HAL/Allocators/CachedOSVeryLargePageAllocator.h"
+#include "HAL/UnrealMemory.h"
+#include "Logging/LogMacros.h"
+#include "CoreGlobals.h"
+#include "HAL/LowLevelMemTracker.h"
+
+
+void FCachedOSVeryLargePageAllocator::Init()
+{
+	Block = FPlatformMemory::FPlatformVirtualMemoryBlock::AllocateVirtual(AddressSpaceToReserve);
+	AddressSpaceReserved = (uintptr_t)Block.GetVirtualPointer();
+
+
+	FreeLargePagesHead = nullptr;
+	for (int i = 0; i < AddressSpaceToReserve / SizeOfLargePage; i++)
+	{
+		LargePagesArray[i].Init((void*)((uintptr_t)AddressSpaceReserved + (i*SizeOfLargePage)));
+		LargePagesArray[i].LinkHead(FreeLargePagesHead);
+	}
+
+	UsedLargePagesHead = nullptr;
+
+	UsedLargePagesWithSpaceHead = nullptr;
+}
+
+void* FCachedOSVeryLargePageAllocator::Allocate(SIZE_T Size)
+{
+	Size = Align(Size, 4096);
+
+	void* ret = nullptr;
+	if (Size == SizeOfSubPage)
+	{
+
+		if (UsedLargePagesWithSpaceHead == nullptr)
+		{
+			FLargePage* LargePage = FreeLargePagesHead;
+			if (LargePage)
+			{
+				Block.Commit(LargePage->BaseAddress - AddressSpaceReserved, SizeOfLargePage);
+				LargePage->Unlink();
+				LargePage->LinkHead(UsedLargePagesWithSpaceHead);
+			}
+		}
+		if (UsedLargePagesWithSpaceHead != nullptr)
+		{
+			FLargePage* LargePage = UsedLargePagesWithSpaceHead;
+			ret = LargePage->Allocate();
+			if (ret && LargePage->NumberOfFreeSubPages == 0)
+			{
+				LargePage->Unlink();
+				LargePage->LinkHead(UsedLargePagesHead);
+			}
+		}
+	}
+
+	if( ret == nullptr )
+	{
+		ret = CachedOSPageAllocator.Allocate(Size);
+	}
+	return ret;
+}
+
+#define LARGEPAGEALLOCATOR_SORT_OnAddress 0
+
+void FCachedOSVeryLargePageAllocator::Free(void* Ptr, SIZE_T Size)
+{
+	uint64 Index = ((uintptr_t)Ptr - (uintptr_t)AddressSpaceReserved) / SizeOfLargePage;
+	if (Index < (AddressSpaceToReserve / SizeOfLargePage))
+	{
+		FLargePage* LargePage = &LargePagesArray[Index];
+
+		LargePage->Free(Ptr);
+
+		if (LargePage->NumberOfFreeSubPages == (SizeOfLargePage / SizeOfSubPage))
+		{
+			// totally free, need to move which list we are in and remove the backing store
+			LargePage->Unlink();
+			LargePage->LinkHead(FreeLargePagesHead);
+			Block.Decommit(LargePage->BaseAddress - AddressSpaceReserved, SizeOfLargePage);
+		}
+		else if (LargePage->NumberOfFreeSubPages == 1)
+		{
+			LargePage->Unlink();
+#if LARGEPAGEALLOCATOR_SORT_OnAddress
+			FLargePage* InsertPoint = UsedLargePagesWithSpaceHead;
+			while (InsertPoint != nullptr)
+			{
+				if (LargePage->BaseAddress < InsertPoint->BaseAddress)	// sort on address
+				{
+					break;
+				}
+				InsertPoint = InsertPoint->Next();
+			}
+			if (InsertPoint == nullptr || InsertPoint == UsedLargePagesWithSpaceHead)
+			{
+				LargePage->LinkHead(UsedLargePagesWithSpaceHead);
+			}
+			else
+			{
+				LargePage->LinkBefore(InsertPoint);
+			}
+#else
+			LargePage->LinkHead(UsedLargePagesWithSpaceHead);
+#endif
+		}
+		else
+		{
+#if !LARGEPAGEALLOCATOR_SORT_OnAddress
+			FLargePage* InsertPoint = LargePage->Next();
+			FLargePage* LastInsertPoint = nullptr;
+
+			if ((InsertPoint != nullptr) && LargePage->NumberOfFreeSubPages > InsertPoint->NumberOfFreeSubPages)
+			{
+				LastInsertPoint = InsertPoint;
+				LargePage->Unlink();
+				while (InsertPoint != nullptr)
+				{
+					if (LargePage->NumberOfFreeSubPages <= InsertPoint->NumberOfFreeSubPages)	// sort on number of free sub pages
+					{
+						break;
+					}
+					InsertPoint = InsertPoint->Next();
+				}
+				if (InsertPoint != nullptr)
+				{
+					LargePage->LinkBefore(InsertPoint);
+				}
+				else
+				{
+					LargePage->LinkAfter(LastInsertPoint);
+				}
+			}
+#endif
+		}
+	}
+	else
+	{
+		CachedOSPageAllocator.Free(Ptr, Size);
+	}
+}
+
+void FCachedOSVeryLargePageAllocator::FreeAll()
+{
+	CachedOSPageAllocator.FreeAll();
+}
diff --git a/Engine/Source/Runtime/Core/Private/HAL/MallocBinned2.cpp b/Engine/Source/Runtime/Core/Private/HAL/MallocBinned2.cpp
--- a/Engine/Source/Runtime/Core/Private/HAL/MallocBinned2.cpp
+++ b/Engine/Source/Runtime/Core/Private/HAL/MallocBinned2.cpp
@@ -10,6 +10,7 @@
 #include "HAL/MemoryMisc.h"
 #include "HAL/PlatformMisc.h"
 #include "Misc/App.h"
+#include "HAL/MallocTimer.h"
 
 PRAGMA_DISABLE_UNSAFE_TYPECAST_WARNINGS
 
@@ -1308,4 +1309,10 @@ void FMallocBinned2::DumpAllocatorStats(class FOutputDevice& Ar)
 	#endif
 #endif
 
+void FMallocBinned2::UpdateStats()
+{
+	FScopedVirtualMallocTimer::UpdateStats();
+}
+
+
 PRAGMA_ENABLE_UNSAFE_TYPECAST_WARNINGS
diff --git a/Engine/Source/Runtime/Core/Private/HAL/MallocTimer.cpp b/Engine/Source/Runtime/Core/Private/HAL/MallocTimer.cpp
new file mode 100644
--- /dev/null
+++ b/Engine/Source/Runtime/Core/Private/HAL/MallocTimer.cpp
@@ -0,0 +1,45 @@
+// Copyright Epic Games, Inc. All Rights Reserved.
+
+#include "HAL/MallocTimer.h"
+
+#if UE_TIME_VIRTUALMALLOC
+#include "ProfilingDebugging/CsvProfiler.h"
+#if CSV_PROFILER
+CSV_DEFINE_CATEGORY_MODULE(CORE_API, VirtualMemory, true);
+#endif
+uint64 FScopedVirtualMallocTimer::GTotalCycles[FScopedVirtualMallocTimer::IndexType::Max] = { 0 };
+
+void FScopedVirtualMallocTimer::UpdateStats()
+{
+	static uint64 GLastTotalCycles[IndexType::Max] = { 0 };
+	static uint64 GLastFrame = 0;
+
+	uint64 Frames = GFrameCounter - GLastFrame;
+	if (Frames)
+	{
+		GLastFrame = GFrameCounter;
+		// not atomic; we assume the error is minor
+		uint64 TotalCycles[IndexType::Max] = { 0 };
+		float TotalSeconds = 0.0f;
+		for (int32 Comp = 0; Comp < IndexType::Max; Comp++)
+		{
+			TotalCycles[Comp] = GTotalCycles[Comp] - GLastTotalCycles[Comp];
+			//TotalCycles[Comp] = GTotalCycles[Comp];
+			GLastTotalCycles[Comp] = GTotalCycles[Comp];
+			TotalSeconds += 1000.0f * float(FPlatformTime::GetSecondsPerCycle64()) * float(TotalCycles[Comp]);
+
+		}
+#if CSV_PROFILER
+#if 0	// extra detail
+		CSV_CUSTOM_STAT(VirtualMemory, Reserve, 1000.0f * float(FPlatformTime::GetSecondsPerCycle64()) * float(TotalCycles[0]), ECsvCustomStatOp::Set);
+		CSV_CUSTOM_STAT(VirtualMemory, Commit, 1000.0f * float(FPlatformTime::GetSecondsPerCycle64()) * float(TotalCycles[1]), ECsvCustomStatOp::Set);
+		CSV_CUSTOM_STAT(VirtualMemory, Combined, 1000.0f * float(FPlatformTime::GetSecondsPerCycle64()) * float(TotalCycles[2]), ECsvCustomStatOp::Set);
+		CSV_CUSTOM_STAT(VirtualMemory, DeCommit, 1000.0f * float(FPlatformTime::GetSecondsPerCycle64()) * float(TotalCycles[3]), ECsvCustomStatOp::Set);
+		CSV_CUSTOM_STAT(VirtualMemory, Free, 1000.0f * float(FPlatformTime::GetSecondsPerCycle64()) * float(TotalCycles[4]), ECsvCustomStatOp::Set);
+#else
+		CSV_CUSTOM_STAT(VirtualMemory, TotalInSeconds, TotalSeconds, ECsvCustomStatOp::Set);
+#endif
+#endif	// CSV_PROFILER
+	}
+}
+#endif
diff --git a/Engine/Source/Runtime/Core/Public/HAL/Allocators/CachedOSVeryLargePageAllocator.h b/Engine/Source/Runtime/Core/Public/HAL/Allocators/CachedOSVeryLargePageAllocator.h
new file mode 100644
--- /dev/null
+++ b/Engine/Source/Runtime/Core/Public/HAL/Allocators/CachedOSVeryLargePageAllocator.h
@@ -0,0 +1,99 @@
+// Copyright Epic Games, Inc. All Rights Reserved.
+
+#pragma once
+#include "CoreTypes.h"
+#include "HAL/CriticalSection.h"
+#include "Containers/List.h"
+#include "HAL/Allocators/CachedOSPageAllocator.h"
+#include "HAL/PlatformTime.h"
+
+#if PLATFORM_64BITS
+#define CACHEDOSVERYLARGEPAGEALLOCATOR_BYTE_LIMIT (64*1024*1024)
+#else
+#define CACHEDOSVERYLARGEPAGEALLOCATOR_BYTE_LIMIT (16*1024*1024)
+#endif
+//#define CACHEDOSVERYLARGEPAGEALLOCATOR_MAX_CACHED_OS_FREES (CACHEDOSVERYLARGEPAGEALLOCATOR_BYTE_LIMIT/(1024*64))
+#define CACHEDOSVERYLARGEPAGEALLOCATOR_MAX_CACHED_OS_FREES (128)
+
+class FCachedOSVeryLargePageAllocator
+{
+	static const uint64 AddressSpaceToReserve = ((1024 * 1024) * (1024 + 512));
+	static const uint64 SizeOfLargePage = (1024 * 1024 * 4);
+	static const uint64 SizeOfSubPage = (1024 * 64);
+public:
+
+	FCachedOSVeryLargePageAllocator()
+	{
+		Init();
+	}
+
+	~FCachedOSVeryLargePageAllocator()
+	{
+		// this leaks everything!
+	}
+
+	void* Allocate(SIZE_T Size);
+
+	void Free(void* Ptr, SIZE_T Size);
+
+	void FreeAll();
+
+	uint64 GetCachedFreeTotal()
+	{
+		return CachedOSPageAllocator.GetCachedFreeTotal();
+	}
+
+private:
+
+	void Init();
+
+	uintptr_t	AddressSpaceReserved;
+
+	FPlatformMemory::FPlatformVirtualMemoryBlock Block;
+
+	struct FLargePage : public TIntrusiveLinkedList<FLargePage>
+	{
+		uintptr_t	FreeSubPages[SizeOfLargePage / SizeOfSubPage];
+		int32		NumberOfFreeSubPages;
+
+		uintptr_t	BaseAddress;
+
+		void Init(void* InBaseAddress)
+		{
+			BaseAddress = (uintptr_t)InBaseAddress;
+			NumberOfFreeSubPages = (SizeOfLargePage / SizeOfSubPage);
+			uintptr_t Ptr = BaseAddress;
+			for (int i = 0; i < NumberOfFreeSubPages; i++)
+			{
+				FreeSubPages[i] = Ptr;
+				Ptr += SizeOfSubPage;
+			}
+		}
+
+		void Free(void* Ptr)
+		{
+			FreeSubPages[NumberOfFreeSubPages++] = (uintptr_t)Ptr;
+		}
+
+		void* Allocate()
+		{
+			void* ret = nullptr;
+			if (NumberOfFreeSubPages)
+			{
+				ret = (void*)FreeSubPages[--NumberOfFreeSubPages];
+			}
+			return ret;
+		}
+	};
+
+	FLargePage*	FreeLargePagesHead;				// no backing store
+
+	FLargePage*	UsedLargePagesHead;				// has backing store and is full
+
+	FLargePage*	UsedLargePagesWithSpaceHead;	// has backing store and still has room
+
+	FLargePage	LargePagesArray[AddressSpaceToReserve / SizeOfLargePage];
+
+	TCachedOSPageAllocator<CACHEDOSVERYLARGEPAGEALLOCATOR_MAX_CACHED_OS_FREES, CACHEDOSVERYLARGEPAGEALLOCATOR_BYTE_LIMIT> CachedOSPageAllocator;
+
+};
diff --git a/Engine/Source/Runtime/Core/Public/HAL/MallocBinned2.h b/Engine/Source/Runtime/Core/Public/HAL/MallocBinned2.h
--- a/Engine/Source/Runtime/Core/Public/HAL/MallocBinned2.h
+++ b/Engine/Source/Runtime/Core/Public/HAL/MallocBinned2.h
@@ -11,11 +11,16 @@
 #include "HAL/CriticalSection.h"
 #include "HAL/PlatformTLS.h"
 #include "HAL/Allocators/CachedOSPageAllocator.h"
+#include "HAL/Allocators/CachedOSVeryLargePageAllocator.h"
 #include "HAL/Allocators/PooledVirtualMemoryAllocator.h"
 #include "HAL/PlatformMath.h"
 #include "HAL/LowLevelMemTracker.h"
 #include "Templates/Atomic.h"
 
+#ifndef UE_USE_VERYLARGEPAGEALLOCATOR
+#define UE_USE_VERYLARGEPAGEALLOCATOR 0
+#endif
+
 #define BINNED2_MAX_CACHED_OS_FREES (64)
 #if PLATFORM_64BITS
 	#define BINNED2_MAX_CACHED_OS_FREES_BYTE_LIMIT (64*1024*1024)
@@ -249,7 +254,11 @@ class CORE_API FMallocBinned2 final : public FMalloc
 	uint64 NumPoolsPerPage;
 
 #if !PLATFORM_UNIX
+#if UE_USE_VERYLARGEPAGEALLOCATOR
+	FCachedOSVeryLargePageAllocator CachedOSPageAllocator;
+#else
 	TCachedOSPageAllocator<BINNED2_MAX_CACHED_OS_FREES, BINNED2_MAX_CACHED_OS_FREES_BYTE_LIMIT> CachedOSPageAllocator;
+#endif
 #else
 	FPooledVirtualMemoryAllocator CachedOSPageAllocator;
 #endif
@@ -675,6 +684,7 @@ public:
 	virtual void SetupTLSCachesOnCurrentThread() override;
 	virtual void ClearAndDisableTLSCachesOnCurrentThread() override;
 	virtual const TCHAR* GetDescriptiveName() override;
+	virtual void UpdateStats() override;
 	// End FMalloc interface.
 
 	void FlushCurrentThreadCache();
@@ -721,3 +731,4 @@ public:
 		#include "FMemory.inl"
 	#endif
 #endif
+
diff --git a/Engine/Source/Runtime/Core/Public/HAL/MallocTimer.h b/Engine/Source/Runtime/Core/Public/HAL/MallocTimer.h
new file mode 100644
--- /dev/null
+++ b/Engine/Source/Runtime/Core/Public/HAL/MallocTimer.h
@@ -0,0 +1,67 @@
+// Copyright Epic Games, Inc. All Rights Reserved.
+
+#pragma once
+#include "CoreTypes.h"
+#include "HAL/CriticalSection.h"
+#include "HAL/PlatformTime.h"
+
+#ifndef UE_TIME_VIRTUALMALLOC
+#define UE_TIME_VIRTUALMALLOC 0
+#endif
+
+#if UE_TIME_VIRTUALMALLOC
+struct CORE_API FScopedVirtualMallocTimer
+{
+	enum IndexType
+	{
+		Reserve,
+		Commit,
+		Combined,
+		DeCommit,
+		Free,
+
+		Max
+	};
+	static uint64 GTotalCycles[IndexType::Max];
+
+	int32 Index;
+	uint64 Cycles;
+
+	FORCEINLINE FScopedVirtualMallocTimer(int32 InIndex = 0)
+		: Index(InIndex)
+		, Cycles(FPlatformTime::Cycles64())
+	{
+	}
+	FORCEINLINE ~FScopedVirtualMallocTimer()
+	{
+		uint64 Add = uint64(FPlatformTime::Cycles64() - Cycles);
+		FPlatformAtomics::InterlockedAdd((volatile int64*)&GTotalCycles[Index], Add);
+	}
+
+	static void UpdateStats();
+};
+#else	//UE_TIME_VIRTUALMALLOC
+struct CORE_API FScopedVirtualMallocTimer
+{
+	enum IndexType
+	{
+		Reserve,
+		Commit,
+		Combined,
+		DeCommit,
+		Free,
+
+		Max
+	};
+	FORCEINLINE FScopedVirtualMallocTimer(int32 InIndex = 0)
+	{
+	}
+	FORCEINLINE ~FScopedVirtualMallocTimer()
+	{
+	}
+
+	static FORCEINLINE void UpdateStats()
+	{
+	}
+};
+#endif //UE_TIME_VIRTUALMALLOC
