Back out changelist 11633817 to fix build

diff --git a/Engine/Source/Runtime/Core/Private/HAL/Allocators/CachedOSVeryLargePageAllocator.cpp b/Engine/Source/Runtime/Core/Private/HAL/Allocators/CachedOSVeryLargePageAllocator.cpp
deleted file mode 100644
diff --git a/Engine/Source/Runtime/Core/Private/HAL/MallocBinned2.cpp b/Engine/Source/Runtime/Core/Private/HAL/MallocBinned2.cpp
--- a/Engine/Source/Runtime/Core/Private/HAL/MallocBinned2.cpp
+++ b/Engine/Source/Runtime/Core/Private/HAL/MallocBinned2.cpp
@@ -10,7 +10,6 @@
 #include "HAL/MemoryMisc.h"
 #include "HAL/PlatformMisc.h"
 #include "Misc/App.h"
-#include "HAL/MallocTimer.h"
 
 PRAGMA_DISABLE_UNSAFE_TYPECAST_WARNINGS
 
@@ -1309,10 +1308,4 @@ void FMallocBinned2::DumpAllocatorStats(class FOutputDevice& Ar)
 	#endif
 #endif
 
-void FMallocBinned2::UpdateStats()
-{
-	FScopedVirtualMallocTimer::UpdateStats();
-}
-
-
 PRAGMA_ENABLE_UNSAFE_TYPECAST_WARNINGS
diff --git a/Engine/Source/Runtime/Core/Private/HAL/MallocTimer.cpp b/Engine/Source/Runtime/Core/Private/HAL/MallocTimer.cpp
deleted file mode 100644
diff --git a/Engine/Source/Runtime/Core/Public/HAL/Allocators/CachedOSVeryLargePageAllocator.h b/Engine/Source/Runtime/Core/Public/HAL/Allocators/CachedOSVeryLargePageAllocator.h
deleted file mode 100644
diff --git a/Engine/Source/Runtime/Core/Public/HAL/MallocBinned2.h b/Engine/Source/Runtime/Core/Public/HAL/MallocBinned2.h
--- a/Engine/Source/Runtime/Core/Public/HAL/MallocBinned2.h
+++ b/Engine/Source/Runtime/Core/Public/HAL/MallocBinned2.h
@@ -11,16 +11,11 @@
 #include "HAL/CriticalSection.h"
 #include "HAL/PlatformTLS.h"
 #include "HAL/Allocators/CachedOSPageAllocator.h"
-#include "HAL/Allocators/CachedOSVeryLargePageAllocator.h"
 #include "HAL/Allocators/PooledVirtualMemoryAllocator.h"
 #include "HAL/PlatformMath.h"
 #include "HAL/LowLevelMemTracker.h"
 #include "Templates/Atomic.h"
 
-#ifndef UE_USE_VERYLARGEPAGEALLOCATOR
-#define UE_USE_VERYLARGEPAGEALLOCATOR 0
-#endif
-
 #define BINNED2_MAX_CACHED_OS_FREES (64)
 #if PLATFORM_64BITS
 	#define BINNED2_MAX_CACHED_OS_FREES_BYTE_LIMIT (64*1024*1024)
@@ -254,11 +249,7 @@ class CORE_API FMallocBinned2 final : public FMalloc
 	uint64 NumPoolsPerPage;
 
 #if !PLATFORM_UNIX
-#if UE_USE_VERYLARGEPAGEALLOCATOR
-	FCachedOSVeryLargePageAllocator CachedOSPageAllocator;
-#else
 	TCachedOSPageAllocator<BINNED2_MAX_CACHED_OS_FREES, BINNED2_MAX_CACHED_OS_FREES_BYTE_LIMIT> CachedOSPageAllocator;
-#endif
 #else
 	FPooledVirtualMemoryAllocator CachedOSPageAllocator;
 #endif
@@ -684,7 +675,6 @@ public:
 	virtual void SetupTLSCachesOnCurrentThread() override;
 	virtual void ClearAndDisableTLSCachesOnCurrentThread() override;
 	virtual const TCHAR* GetDescriptiveName() override;
-	virtual void UpdateStats() override;
 	// End FMalloc interface.
 
 	void FlushCurrentThreadCache();
@@ -731,4 +721,3 @@ public:
 		#include "FMemory.inl"
 	#endif
 #endif
-
diff --git a/Engine/Source/Runtime/Core/Public/HAL/MallocTimer.h b/Engine/Source/Runtime/Core/Public/HAL/MallocTimer.h
deleted file mode 100644
