Add a commandline option so the large page cache can be disabled, fix a build break

diff --git a/Engine/Source/Runtime/Core/Private/HAL/Allocators/CachedOSVeryLargePageAllocator.cpp b/Engine/Source/Runtime/Core/Private/HAL/Allocators/CachedOSVeryLargePageAllocator.cpp
--- a/Engine/Source/Runtime/Core/Private/HAL/Allocators/CachedOSVeryLargePageAllocator.cpp
+++ b/Engine/Source/Runtime/Core/Private/HAL/Allocators/CachedOSVeryLargePageAllocator.cpp
@@ -7,6 +7,9 @@
 #include "HAL/LowLevelMemTracker.h"
 
 #if UE_USE_VERYLARGEPAGEALLOCATOR
+
+CORE_API bool GEnableVeryLargePageAllocator = true;
+
 void FCachedOSVeryLargePageAllocator::Init()
 {
 	Block = FPlatformMemory::FPlatformVirtualMemoryBlock::AllocateVirtual(AddressSpaceToReserve);
@@ -21,8 +24,12 @@ void FCachedOSVeryLargePageAllocator::Init()
 	}
 
 	UsedLargePagesHead = nullptr;
-
 	UsedLargePagesWithSpaceHead = nullptr;
+
+	if (!GEnableVeryLargePageAllocator)
+	{
+		bEnabled = false;
+	}
 }
 
 void* FCachedOSVeryLargePageAllocator::Allocate(SIZE_T Size)
@@ -30,7 +37,8 @@ void* FCachedOSVeryLargePageAllocator::Allocate(SIZE_T Size)
 	Size = Align(Size, 4096);
 
 	void* ret = nullptr;
-	if (Size == SizeOfSubPage)
+
+	if (bEnabled && Size == SizeOfSubPage)
 	{
 
 		if (UsedLargePagesWithSpaceHead == nullptr)
diff --git a/Engine/Source/Runtime/Core/Public/HAL/Allocators/CachedOSVeryLargePageAllocator.h b/Engine/Source/Runtime/Core/Public/HAL/Allocators/CachedOSVeryLargePageAllocator.h
--- a/Engine/Source/Runtime/Core/Public/HAL/Allocators/CachedOSVeryLargePageAllocator.h
+++ b/Engine/Source/Runtime/Core/Public/HAL/Allocators/CachedOSVeryLargePageAllocator.h
@@ -3,6 +3,7 @@
 #pragma once
 #include "CoreTypes.h"
 #include "HAL/CriticalSection.h"
+#include "HAL/UnrealMemory.h"
 #include "Containers/List.h"
 #include "HAL/Allocators/CachedOSPageAllocator.h"
 #include "HAL/PlatformMemory.h"
@@ -25,6 +26,7 @@ class FCachedOSVeryLargePageAllocator
 public:
 
 	FCachedOSVeryLargePageAllocator()
+		: bEnabled(true)
 	{
 		Init();
 	}
@@ -49,6 +51,7 @@ private:
 
 	void Init();
 
+	bool bEnabled;
 	uintptr_t	AddressSpaceReserved;
 
 	FPlatformMemory::FPlatformVirtualMemoryBlock Block;
@@ -99,5 +102,5 @@ private:
 	TCachedOSPageAllocator<CACHEDOSVERYLARGEPAGEALLOCATOR_MAX_CACHED_OS_FREES, CACHEDOSVERYLARGEPAGEALLOCATOR_BYTE_LIMIT> CachedOSPageAllocator;
 
 };
-
+CORE_API extern bool GEnableVeryLargePageAllocator;
 #endif // UE_USE_VERYLARGEPAGEALLOCATOR
\ No newline at end of file
