addition of IOS 13 API to track free memory left.

diff --git a/Engine/Source/Runtime/Core/Private/Apple/ApplePlatformMemory.cpp b/Engine/Source/Runtime/Core/Private/Apple/ApplePlatformMemory.cpp
--- a/Engine/Source/Runtime/Core/Private/Apple/ApplePlatformMemory.cpp
+++ b/Engine/Source/Runtime/Core/Private/Apple/ApplePlatformMemory.cpp
@@ -24,6 +24,9 @@
 #include <sys/mman.h>
 #include <sys/mount.h>
 #include <objc/runtime.h>
+#if PLATFORM_IOS && defined(__IPHONE_13_0)
+#include <os/proc.h>
+#endif
 #include <CoreFoundation/CFBase.h>
 #include "HAL/LowLevelMemTracker.h"
 #include "Apple/AppleLLM.h"
@@ -260,10 +263,14 @@ FPlatformMemoryStats FApplePlatformMemory::GetStats()
 	static FPlatformMemoryStats MemoryStats;
 	
 	// Gather platform memory stats.
+#if PLATFORM_IOS && defined(__IPHONE_13_0)
+	uint64_t FreeMem = os_proc_available_memory();
+#else
 	vm_statistics Stats;
 	mach_msg_type_number_t StatsSize = sizeof(Stats);
 	host_statistics(mach_host_self(), HOST_VM_INFO, (host_info_t)&Stats, &StatsSize);
 	uint64_t FreeMem = (Stats.free_count + Stats.inactive_count) * MemoryConstants.PageSize;
+#endif
 	MemoryStats.AvailablePhysical = FreeMem;
 	
 	// Just get memory information for the process and report the working set instead
diff --git a/Engine/Source/Runtime/Core/Private/HAL/LowLevelMemTracker.cpp b/Engine/Source/Runtime/Core/Private/HAL/LowLevelMemTracker.cpp
--- a/Engine/Source/Runtime/Core/Private/HAL/LowLevelMemTracker.cpp
+++ b/Engine/Source/Runtime/Core/Private/HAL/LowLevelMemTracker.cpp
@@ -41,6 +41,7 @@ DECLARE_LLM_MEMORY_STAT(TEXT("Tracked Total"), STAT_PlatformTrackedTotalLLM, STA
 DECLARE_LLM_MEMORY_STAT(TEXT("Untagged"), STAT_PlatformUntaggedTotalLLM, STATGROUP_LLMPlatform);
 DECLARE_LLM_MEMORY_STAT(TEXT("Untracked"), STAT_PlatformUntrackedLLM, STATGROUP_LLMPlatform);
 DECLARE_LLM_MEMORY_STAT(TEXT("Overhead"), STAT_PlatformOverheadLLM, STATGROUP_LLMPlatform);
+DECLARE_LLM_MEMORY_STAT(TEXT("OS Available"), STAT_PlatformOSAvailableLLM, STATGROUP_LLMPlatform);
 DECLARE_LLM_MEMORY_STAT(TEXT("FMalloc"), STAT_FMallocLLM, STATGROUP_LLMPlatform);
 DECLARE_LLM_MEMORY_STAT(TEXT("FMalloc Unused"), STAT_FMallocUnusedLLM, STATGROUP_LLMFULL);
 DECLARE_LLM_MEMORY_STAT(TEXT("ThreadStack"), STAT_ThreadStackLLM, STATGROUP_LLMFULL);
@@ -662,6 +663,7 @@ void FLowLevelMemTracker::UpdateStatsPerFrame(const TCHAR* LogName)
 	GetTracker(ELLMTracker::Platform)->SetTagAmount(ELLMTag::PlatformTotal, PlatformProcessMemory, false);
 	GetTracker(ELLMTracker::Platform)->SetTagAmount(ELLMTag::PlatformUntracked, PlatformTotalUntracked, false);
 	GetTracker(ELLMTracker::Platform)->SetTagAmount(ELLMTag::PlatformOverhead, Overhead, true);
+	GetTracker(ELLMTracker::Platform)->SetTagAmount(ELLMTag::PlatformOSAvailable, PlatformStats.AvailablePhysical, false);
 
 	int64 TrackedTotal = GetTracker(ELLMTracker::Default)->GetTagAmount(ELLMTag::TrackedTotal);
 	// remove the Overhead from the default LLM as it's not something anyone needs to investigate when finding what to reduce
diff --git a/Engine/Source/Runtime/Core/Public/HAL/LowLevelMemTracker.h b/Engine/Source/Runtime/Core/Public/HAL/LowLevelMemTracker.h
--- a/Engine/Source/Runtime/Core/Public/HAL/LowLevelMemTracker.h
+++ b/Engine/Source/Runtime/Core/Public/HAL/LowLevelMemTracker.h
@@ -142,6 +142,7 @@ enum class ELLMTagSet : uint8
 	macro(PlatformUntaggedTotal,				"Untagged",						GET_STATFNAME(STAT_PlatformUntaggedTotalLLM),				NAME_None,										-1)\
 	macro(PlatformUntracked,					"Untracked",					GET_STATFNAME(STAT_PlatformUntrackedLLM),					NAME_None,										-1)\
 	macro(PlatformOverhead,						"LLMOverhead",					GET_STATFNAME(STAT_PlatformOverheadLLM),					NAME_None,										-1)\
+	macro(PlatformOSAvailable,					"OSAvailable",					GET_STATFNAME(STAT_PlatformOSAvailableLLM),					NAME_None,										-1)\
 	macro(FMalloc,								"FMalloc",						GET_STATFNAME(STAT_FMallocLLM),								NAME_None,										-1)\
 	macro(FMallocUnused,						"FMallocUnused",				GET_STATFNAME(STAT_FMallocUnusedLLM),						GET_STATFNAME(STAT_EngineSummaryLLM),			-1)\
 	macro(ThreadStack,							"ThreadStack",					GET_STATFNAME(STAT_ThreadStackLLM),							GET_STATFNAME(STAT_EngineSummaryLLM),			-1)\
