fix for new IOS memory API usage crash on IOS 12 or less.

diff --git a/Engine/Source/Runtime/Core/Private/Apple/ApplePlatformMemory.cpp b/Engine/Source/Runtime/Core/Private/Apple/ApplePlatformMemory.cpp
--- a/Engine/Source/Runtime/Core/Private/Apple/ApplePlatformMemory.cpp
+++ b/Engine/Source/Runtime/Core/Private/Apple/ApplePlatformMemory.cpp
@@ -263,14 +263,22 @@ FPlatformMemoryStats FApplePlatformMemory::GetStats()
 	static FPlatformMemoryStats MemoryStats;
 	
 	// Gather platform memory stats.
-#if PLATFORM_IOS && defined(__IPHONE_13_0)
-	uint64_t FreeMem = os_proc_available_memory();
-#else
-	vm_statistics Stats;
-	mach_msg_type_number_t StatsSize = sizeof(Stats);
-	host_statistics(mach_host_self(), HOST_VM_INFO, (host_info_t)&Stats, &StatsSize);
-	uint64_t FreeMem = (Stats.free_count + Stats.inactive_count) * MemoryConstants.PageSize;
+	uint64_t FreeMem = 0;
+#if PLATFORM_IOS
+#if defined(__IPHONE_13_0)
+	if (@available(iOS 13.0,*))
+	{
+		FreeMem = os_proc_available_memory();
+	}
+	else
+#endif
 #endif
+	{
+		vm_statistics Stats;
+		mach_msg_type_number_t StatsSize = sizeof(Stats);
+		host_statistics(mach_host_self(), HOST_VM_INFO, (host_info_t)&Stats, &StatsSize);
+		FreeMem = (Stats.free_count + Stats.inactive_count) * MemoryConstants.PageSize;
+	}
 	MemoryStats.AvailablePhysical = FreeMem;
 	
 	// Just get memory information for the process and report the working set instead
