Added csv exclusive stat to Logf_Internal to help measure perf issues due to costly logs

diff --git a/Engine/Source/Runtime/Core/Private/Logging/LogMacros.cpp b/Engine/Source/Runtime/Core/Private/Logging/LogMacros.cpp
--- a/Engine/Source/Runtime/Core/Private/Logging/LogMacros.cpp
+++ b/Engine/Source/Runtime/Core/Private/Logging/LogMacros.cpp
@@ -72,6 +72,7 @@ void FMsg::Logf_InternalImpl(const ANSICHAR* File, int32 Line, const FLogCategor
 {
 #if !NO_LOGGING
 	QUICK_SCOPE_CYCLE_COUNTER(STAT_FMsgLogf);
+	CSV_SCOPED_TIMING_STAT_EXCLUSIVE(Log)
 	CSV_CUSTOM_STAT(FMsgLogf, FMsgLogfCount, 1, ECsvCustomStatOp::Accumulate);
 
 	if (Verbosity != ELogVerbosity::Fatal)
