Add command line argument to modify engine save directory name

diff --git a/Engine/Source/Runtime/Core/Private/Misc/Paths.cpp b/Engine/Source/Runtime/Core/Private/Misc/Paths.cpp
--- a/Engine/Source/Runtime/Core/Private/Misc/Paths.cpp
+++ b/Engine/Source/Runtime/Core/Private/Misc/Paths.cpp
@@ -31,6 +31,7 @@ struct FPaths::FStaticData
 
 	FString         UserDirArg;
 	FString         GameSavedDir;
+	FString         EngineSavedDir;
 	FString         ShaderDir;
 	FString         UserFolder;
 	TArray<FString> EngineLocalizationPaths;
@@ -44,6 +45,7 @@ struct FPaths::FStaticData
 
 	bool bUserDirArgInitialized                    = false;
 	bool bGameSavedDirInitialized                  = false;
+	bool bEngineSavedDirInitialized                = false;
 	bool bShaderDirInitialized                     = false;
 	bool bUserFolderInitialized                    = false;
 	bool bEngineLocalizationPathsInitialized       = false;
@@ -65,12 +67,12 @@ namespace UE4Paths_Private
 	auto IsSlashOrBackslash    = [](TCHAR C) { return C == TEXT('/') || C == TEXT('\\'); };
 	auto IsNotSlashOrBackslash = [](TCHAR C) { return C != TEXT('/') && C != TEXT('\\'); };
 
-	FString GameSavedDir()
+	FString GetSavedDirSuffix(const FString& BaseDir, const TCHAR* CommandLineArgument)
 	{
-		FString Result = FPaths::ProjectUserDir();
+		FString Result = BaseDir + TEXT("Saved");
 
 		FString NonDefaultSavedDirSuffix;
-		if (FParse::Value(FCommandLine::Get(), TEXT("-saveddirsuffix="), NonDefaultSavedDirSuffix))
+		if (FParse::Value(FCommandLine::Get(), CommandLineArgument, NonDefaultSavedDirSuffix))
 		{
 			for (int32 CharIdx = 0; CharIdx < NonDefaultSavedDirSuffix.Len(); ++CharIdx)
 			{
@@ -80,20 +82,28 @@ namespace UE4Paths_Private
 					--CharIdx;
 				}
 			}
-
-			if (!NonDefaultSavedDirSuffix.IsEmpty())
-			{
-				Result += TEXT("Saved_") + NonDefaultSavedDirSuffix + TEXT("/");
-			}
 		}
-		else
+
+		if (!NonDefaultSavedDirSuffix.IsEmpty())
 		{
-			Result += TEXT("Saved/");
+			Result += TEXT("_") + NonDefaultSavedDirSuffix;
 		}
 
+		Result += TEXT("/");
+
 		return Result;
 	}
 
+	FString GameSavedDir()
+	{
+		return GetSavedDirSuffix(FPaths::ProjectUserDir(), TEXT("-saveddirsuffix="));
+	}
+
+	FString EngineSavedDir()
+	{
+		return GetSavedDirSuffix(FPaths::EngineUserDir(), TEXT("-enginesaveddirsuffix="));
+	}
+
 	FString ConvertRelativePathToFullInternal(FString&& BasePath, FString&& InPath)
 	{
 		FString FullyPathed;
@@ -189,7 +199,13 @@ FString FPaths::EngineIntermediateDir()
 
 FString FPaths::EngineSavedDir()
 {
-	return EngineUserDir() + TEXT("Saved/");
+	FStaticData& StaticData = TLazySingleton<FStaticData>::Get();
+	if (!StaticData.bEngineSavedDirInitialized)
+	{
+		StaticData.EngineSavedDir = UE4Paths_Private::EngineSavedDir();
+		StaticData.bEngineSavedDirInitialized = true;
+	}
+	return StaticData.EngineSavedDir;
 }
 
 FString FPaths::EnginePluginsDir()
