Slight bug fix! missing some allocations that could be reused because the calling malloc doesn't know the true size of allocations it gets back!

diff --git a/Engine/Source/Runtime/Core/Private/HAL/Allocators/CachedOSVeryLargePageAllocator.cpp b/Engine/Source/Runtime/Core/Private/HAL/Allocators/CachedOSVeryLargePageAllocator.cpp
--- a/Engine/Source/Runtime/Core/Private/HAL/Allocators/CachedOSVeryLargePageAllocator.cpp
+++ b/Engine/Source/Runtime/Core/Private/HAL/Allocators/CachedOSVeryLargePageAllocator.cpp
@@ -74,6 +74,7 @@ void* FCachedOSVeryLargePageAllocator::Allocate(SIZE_T Size)
 
 void FCachedOSVeryLargePageAllocator::Free(void* Ptr, SIZE_T Size)
 {
+	Size = Align(Size, 4096);
 	uint64 Index = ((uintptr_t)Ptr - (uintptr_t)AddressSpaceReserved) / SizeOfLargePage;
 	if (Index < (AddressSpaceToReserve / SizeOfLargePage))
 	{
