Add option for android crash context to not write log when storing crashinfo.
Add gameActivityOnRestartApplicationAdditions to allow plugins to get notice of app restart requests.
#rnx

diff --git a/Engine/Build/Android/Java/src/com/epicgames/ue4/GameActivity.java.template b/Engine/Build/Android/Java/src/com/epicgames/ue4/GameActivity.java.template
--- a/Engine/Build/Android/Java/src/com/epicgames/ue4/GameActivity.java.template
+++ b/Engine/Build/Android/Java/src/com/epicgames/ue4/GameActivity.java.template
@@ -5992,6 +5992,7 @@ public class GameActivity extends $${gameActivitySuperClass}$$ implements Surfac
 		context.startActivity(intent);
 
 		// kill the current application instance. Undefined behaviour after this!
+//$${gameActivityOnRestartApplicationAdditions}$$
 		Runtime.getRuntime().exit(0); 
 	}
 	
diff --git a/Engine/Source/Programs/UnrealBuildTool/Platform/Android/UEDeployAndroid.cs b/Engine/Source/Programs/UnrealBuildTool/Platform/Android/UEDeployAndroid.cs
--- a/Engine/Source/Programs/UnrealBuildTool/Platform/Android/UEDeployAndroid.cs
+++ b/Engine/Source/Programs/UnrealBuildTool/Platform/Android/UEDeployAndroid.cs
@@ -4453,6 +4453,7 @@ namespace UnrealBuildTool
 				{ "//$${gameActivityGetLoginIdAdditions}$$", UPL.ProcessPluginNode(NDKArch, "gameActivityGetLoginIdAdditions", "")},
 				{ "//$${gameActivityGetFunnelIdAdditions}$$", UPL.ProcessPluginNode(NDKArch, "gameActivityGetFunnelIdAdditions", "")},
 				{ "//$${gameActivityAllowedRemoteNotificationsAdditions}$$", UPL.ProcessPluginNode(NDKArch, "gameActivityAllowedRemoteNotificationsAdditions", "")},
+				{ "//$${gameActivityOnRestartApplicationAdditions}$$", UPL.ProcessPluginNode(NDKArch, "gameActivityOnRestartApplicationAdditions", "")},				
 				{ "//$${soLoadLibrary}$$", UPL.ProcessPluginNode(NDKArch, "soLoadLibrary", LoadLibraryDefaults)},
 				{ "$${gameActivitySuperClass}$$", SuperClassDefault},
 			};
diff --git a/Engine/Source/Runtime/Core/Private/Android/AndroidPlatformCrashContext.cpp b/Engine/Source/Runtime/Core/Private/Android/AndroidPlatformCrashContext.cpp
--- a/Engine/Source/Runtime/Core/Private/Android/AndroidPlatformCrashContext.cpp
+++ b/Engine/Source/Runtime/Core/Private/Android/AndroidPlatformCrashContext.cpp
@@ -137,7 +137,7 @@ static void CrashReportFileCopy(const char* DestPath, const char* SourcePath)
 	close(SourceHandle);
 }
 
-void FAndroidCrashContext::StoreCrashInfo() const
+void FAndroidCrashContext::StoreCrashInfo(bool bWriteLog) const
 {
 	char FilePath[CrashReportMaxPathSize] = { 0 };
 	FCStringAnsi::Strcpy(FilePath, ReportDirectory);
@@ -145,12 +145,15 @@ void FAndroidCrashContext::StoreCrashInfo() const
 	FCStringAnsi::Strcat(FilePath, FGenericCrashContext::CrashContextRuntimeXMLNameA);
 	SerializeAsXML(*FString(FilePath)); // CreateFileWriter will also create destination directory.
 
-	// copy log:
-	FCStringAnsi::Strcpy(FilePath, ReportDirectory);
-	FCStringAnsi::Strcat(FilePath, "/");
-	FCStringAnsi::Strcat(FilePath, FCStringAnsi::Strlen(GAndroidCrashInfo.AppName) ? GAndroidCrashInfo.AppName : "UE4");
-	FCStringAnsi::Strcat(FilePath, ".log");
-	CrashReportFileCopy(FilePath, GAndroidCrashInfo.AppLogPath);
+	if(bWriteLog)
+	{
+		// copy log:
+		FCStringAnsi::Strcpy(FilePath, ReportDirectory);
+		FCStringAnsi::Strcat(FilePath, "/");
+		FCStringAnsi::Strcat(FilePath, FCStringAnsi::Strlen(GAndroidCrashInfo.AppName) ? GAndroidCrashInfo.AppName : "UE4");
+		FCStringAnsi::Strcat(FilePath, ".log");
+		CrashReportFileCopy(FilePath, GAndroidCrashInfo.AppLogPath);
+	}
 }
 
 
@@ -265,14 +268,16 @@ FAndroidCrashContext::FAndroidCrashContext(ECrashContextType InType, const TCHAR
 , Info(NULL)
 , Context(NULL)
 {
-	if (GetType() == ECrashContextType::Ensure)
+	switch(GetType())
 	{
-		// create a new report folder.
-		GenerateReportDirectoryName(ReportDirectory);
-	}
-	else
-	{
-		GetGlobalCrashDirectoryPath(ReportDirectory);
+		case ECrashContextType::AbnormalShutdown:
+		case ECrashContextType::Ensure:
+			// create a new report folder.
+			GenerateReportDirectoryName(ReportDirectory);
+			break;
+		default:
+			GetGlobalCrashDirectoryPath(ReportDirectory);
+			break;
 	}
 }
 
diff --git a/Engine/Source/Runtime/Core/Public/Android/AndroidPlatformCrashContext.h b/Engine/Source/Runtime/Core/Public/Android/AndroidPlatformCrashContext.h
--- a/Engine/Source/Runtime/Core/Public/Android/AndroidPlatformCrashContext.h
+++ b/Engine/Source/Runtime/Core/Public/Android/AndroidPlatformCrashContext.h
@@ -40,7 +40,7 @@ struct CORE_API FAndroidCrashContext : public FGenericCrashContext
 	virtual void AddPlatformSpecificProperties() const override;
 
 	void CaptureCrashInfo();
-	void StoreCrashInfo() const;
+	void StoreCrashInfo(bool bWriteLog) const;
 
 	static const int32 CrashReportMaxPathSize = 512;
 
