Do not use crash report client in monitor mode for dedicated servers (Windows only).

diff --git a/Engine/Source/Runtime/Core/Private/Windows/WindowsPlatformCrashContext.cpp b/Engine/Source/Runtime/Core/Private/Windows/WindowsPlatformCrashContext.cpp
--- a/Engine/Source/Runtime/Core/Private/Windows/WindowsPlatformCrashContext.cpp
+++ b/Engine/Source/Runtime/Core/Private/Windows/WindowsPlatformCrashContext.cpp
@@ -1008,8 +1008,11 @@ public:
 #endif
 
 #if USE_CRASH_REPORTER_MONITOR
-		CrashClientHandle = LaunchCrashReportClient(&CrashMonitorWritePipe, &CrashMonitorReadPipe);
-		FMemory::Memzero(SharedContext);
+		if (!FPlatformProperties::IsServerOnly())
+		{
+			CrashClientHandle = LaunchCrashReportClient(&CrashMonitorWritePipe, &CrashMonitorReadPipe);
+			FMemory::Memzero(SharedContext);
+		}
 #endif
 
 		// Create a background thread that will process the crash and generate crash reports
