Add LLM Scope for Localization so that is not falling in Untagged category. This was inflating the Untagged catogory.

#rnx

diff --git a/Engine/Source/Runtime/Core/Private/Internationalization/TextLocalizationManager.cpp b/Engine/Source/Runtime/Core/Private/Internationalization/TextLocalizationManager.cpp
--- a/Engine/Source/Runtime/Core/Private/Internationalization/TextLocalizationManager.cpp
+++ b/Engine/Source/Runtime/Core/Private/Internationalization/TextLocalizationManager.cpp
@@ -851,6 +851,8 @@ void FTextLocalizationManager::RefreshResources()
 
 void FTextLocalizationManager::OnPakFileMounted(const IPakFile& PakFile)
 {
+	LLM_SCOPE(ELLMTag::Localization);
+
 	int32 ChunkId = PakFile.PakGetPakchunkIndex();
 	if (ChunkId == INDEX_NONE || ChunkId == 0)
 	{
