Sets ArIsError and ArIsCriticalError to private

diff --git a/Engine/Plugins/Enterprise/DatasmithCADImporter/Source/DatasmithDispatcher/Private/DatasmithCommands.cpp b/Engine/Plugins/Enterprise/DatasmithCADImporter/Source/DatasmithDispatcher/Private/DatasmithCommands.cpp
--- a/Engine/Plugins/Enterprise/DatasmithCADImporter/Source/DatasmithDispatcher/Private/DatasmithCommands.cpp
+++ b/Engine/Plugins/Enterprise/DatasmithCADImporter/Source/DatasmithDispatcher/Private/DatasmithCommands.cpp
@@ -39,7 +39,7 @@ namespace DatasmithDispatcher
 		if (TSharedPtr<ICommand> Command = CreateCommand(static_cast<ECommandId>(type)))
 		{
 			ArReader << *Command;
-			return ArReader.ArIsError ? nullptr : Command;
+			return ArReader.IsError() ? nullptr : Command;
 		}
 		return nullptr;
 	}
diff --git a/Engine/Plugins/Runtime/AR/Google/GoogleARCore/Source/GoogleARCoreBase/Private/GoogleARCoreAugmentedImageDatabase.cpp b/Engine/Plugins/Runtime/AR/Google/GoogleARCore/Source/GoogleARCoreBase/Private/GoogleARCoreAugmentedImageDatabase.cpp
--- a/Engine/Plugins/Runtime/AR/Google/GoogleARCore/Source/GoogleARCoreBase/Private/GoogleARCoreAugmentedImageDatabase.cpp
+++ b/Engine/Plugins/Runtime/AR/Google/GoogleARCore/Source/GoogleARCoreBase/Private/GoogleARCoreAugmentedImageDatabase.cpp
@@ -209,9 +209,7 @@ void UGoogleARCoreAugmentedImageDatabase::Serialize(FArchive& Ar)
 
 			if (OutReturnCode)
 			{
-				Ar.SetError();
-				Ar.ArIsError = 1;
-				Ar.ArIsCriticalError = 1;
+				Ar.SetCriticalError();
 				UE_LOG(LogGoogleARCoreAPI, Error, TEXT("Failed to build augmented image database: %s"), *OutStderr);
 			}
 			else
diff --git a/Engine/Plugins/Runtime/AR/Google/GoogleARCore/Source/GoogleARCoreBase/Private/GoogleARCoreCookSupport.cpp b/Engine/Plugins/Runtime/AR/Google/GoogleARCore/Source/GoogleARCoreBase/Private/GoogleARCoreCookSupport.cpp
--- a/Engine/Plugins/Runtime/AR/Google/GoogleARCore/Source/GoogleARCoreBase/Private/GoogleARCoreCookSupport.cpp
+++ b/Engine/Plugins/Runtime/AR/Google/GoogleARCore/Source/GoogleARCoreBase/Private/GoogleARCoreCookSupport.cpp
@@ -287,9 +287,7 @@ void FGoogleARCoreSessionConfigCookSupport::OnSerializeSessionConfig(UARSessionC
 
 		if (OutReturnCode)
 		{
-			Ar.SetError();
-			Ar.ArIsError = 1;
-			Ar.ArIsCriticalError = 1;
+			Ar.SetCriticalError();
 			UE_LOG(LogGoogleARCore, Error, TEXT("Failed to build augmented image database: %s"), *OutStderr);
 		}
 		else
diff --git a/Engine/Source/Runtime/Core/Private/Containers/String.cpp b/Engine/Source/Runtime/Core/Private/Containers/String.cpp
--- a/Engine/Source/Runtime/Core/Private/Containers/String.cpp
+++ b/Engine/Source/Runtime/Core/Private/Containers/String.cpp
@@ -1385,8 +1385,7 @@ FArchive& operator<<( FArchive& Ar, FString& A )
 			// If SaveNum cannot be negated due to integer overflow, Ar is corrupted.
 			if (SaveNum == MIN_int32)
 			{
-				Ar.ArIsError = 1;
-				Ar.ArIsCriticalError = 1;
+				Ar.SetCriticalError();
 				UE_LOG(LogCore, Error, TEXT("Archive is corrupted"));
 				return Ar;
 			}
@@ -1398,8 +1397,7 @@ FArchive& operator<<( FArchive& Ar, FString& A )
 		// Protect against network packets allocating too much memory
 		if ((MaxSerializeSize > 0) && (SaveNum > MaxSerializeSize))
 		{
-			Ar.ArIsError         = 1;
-			Ar.ArIsCriticalError = 1;
+			Ar.SetCriticalError();
 			UE_LOG(LogCore, Error, TEXT("String is too large (Size: %i, Max: %i)"), SaveNum, MaxSerializeSize);
 			return Ar;
 		}
diff --git a/Engine/Source/Runtime/Core/Private/HAL/FileManagerGeneric.cpp b/Engine/Source/Runtime/Core/Private/HAL/FileManagerGeneric.cpp
--- a/Engine/Source/Runtime/Core/Private/HAL/FileManagerGeneric.cpp
+++ b/Engine/Source/Runtime/Core/Private/HAL/FileManagerGeneric.cpp
@@ -675,7 +675,7 @@ void FArchiveFileReaderGeneric::Seek( int64 InPos )
 	if (!SeekLowLevel(InPos))
 	{
 		TCHAR ErrorBuffer[1024];
-		ArIsError = true;
+		SetError();
 		UE_LOG(LogFileManager, Error, TEXT("SetFilePointer on %s Failed %lld/%lld: %lld %s"), *Filename, InPos, Size, Pos, FPlatformMisc::GetSystemErrorMessage(ErrorBuffer, 1024, 0));
 	}
 
@@ -714,7 +714,7 @@ void FArchiveFileReaderGeneric::CloseLowLevel()
 bool FArchiveFileReaderGeneric::Close()
 {
 	CloseLowLevel();
-	return !ArIsError;
+	return !IsError();
 }
 
 bool FArchiveFileReaderGeneric::InternalPrecache( int64 PrecacheOffset, int64 PrecacheSize )
@@ -746,7 +746,7 @@ bool FArchiveFileReaderGeneric::InternalPrecache( int64 PrecacheOffset, int64 Pr
 		if( Count!=BufferCount )
 		{
 			TCHAR ErrorBuffer[1024];
-			ArIsError = true;
+			SetError();
 			UE_LOG( LogFileManager, Warning, TEXT( "ReadFile failed: Count=%lld BufferCount=%lld Error=%s" ), Count, BufferCount, FPlatformMisc::GetSystemErrorMessage( ErrorBuffer, 1024, 0 ) );
 		}
 	}
@@ -757,7 +757,7 @@ void FArchiveFileReaderGeneric::Serialize( void* V, int64 Length )
 {
 	if (Pos + Length > Size)
 	{
-		ArIsError = true;
+		SetError();
 		UE_LOG(LogFileManager, Error, TEXT("Requested read of %d bytes when %d bytes remain (file=%s, size=%d)"), Length, Size-Pos, *Filename, Size);
 		return;
 	}
@@ -776,7 +776,7 @@ void FArchiveFileReaderGeneric::Serialize( void* V, int64 Length )
 				if( Count!=Length )
 				{
 					TCHAR ErrorBuffer[1024];
-					ArIsError = true;
+					SetError();
 					UE_LOG( LogFileManager, Warning, TEXT( "ReadFile failed: Count=%lld Length=%lld Error=%s for file %s" ), 
 						Count, Length, FPlatformMisc::GetSystemErrorMessage( ErrorBuffer, 1024, 0 ), *Filename );
 				}
@@ -785,18 +785,18 @@ void FArchiveFileReaderGeneric::Serialize( void* V, int64 Length )
 			}
 			if (!InternalPrecache(Pos, MAX_int32))
 			{
-				ArIsError = true;
+				SetError();
 				UE_LOG( LogFileManager, Warning, TEXT( "ReadFile failed during precaching for file %s" ),*Filename );
 				return;
 			}
 			Copy = FMath::Min( Length, BufferBase+BufferArray.Num()-Pos );
 			if( Copy<=0 )
 			{
-				ArIsError = true;
+				SetError();
 				UE_LOG( LogFileManager, Error, TEXT( "ReadFile beyond EOF %lld+%lld/%lld for file %s" ), 
 					Pos, Length, Size, *Filename );
 			}
-			if( ArIsError )
+			if( IsError() )
 			{
 				return;
 			}
@@ -854,7 +854,7 @@ void FArchiveFileWriterGeneric::Seek( int64 InPos )
 	FlushBuffer();
 	if( !SeekLowLevel( InPos ) )
 	{
-		ArIsError = true;
+		SetError();
 		LogWriteError(TEXT("Error seeking file"));
 	}
 	Pos = InPos;
@@ -865,10 +865,10 @@ bool FArchiveFileWriterGeneric::Close()
 	FlushBuffer();
 	if( !CloseLowLevel() )
 	{
-		ArIsError = true;
+		SetError();
 		LogWriteError(TEXT("Error closing file"));
 	}
-	return !ArIsError;
+	return !IsError();
 }
 
 void FArchiveFileWriterGeneric::Serialize( void* V, int64 Length )
@@ -879,7 +879,7 @@ void FArchiveFileWriterGeneric::Serialize( void* V, int64 Length )
 		FlushBuffer();
 		if( !WriteLowLevel( (uint8*)V, Length ) )
 		{
-			ArIsError = true;
+			SetError();
 			LogWriteError(TEXT("Error writing to file"));
 		}
 	}
@@ -916,7 +916,7 @@ bool FArchiveFileWriterGeneric::FlushBuffer()
 		bDidWriteData = WriteLowLevel(BufferArray.GetData(), BufferNum);
 		if (!bDidWriteData)
 		{
-			ArIsError = true;
+			SetError();
 			LogWriteError(TEXT("Error flushing file"));
 		}
 		BufferArray.Reset();
diff --git a/Engine/Source/Runtime/Core/Private/Internationalization/TextKey.cpp b/Engine/Source/Runtime/Core/Private/Internationalization/TextKey.cpp
--- a/Engine/Source/Runtime/Core/Private/Internationalization/TextKey.cpp
+++ b/Engine/Source/Runtime/Core/Private/Internationalization/TextKey.cpp
@@ -223,8 +223,7 @@ bool LoadKeyString(FArchive& Ar, FInlineStringBuffer& OutStrBuffer)
 	// If SaveNum is still less than 0, they must have passed in MIN_INT. Archive is corrupted.
 	if (SaveNum < 0)
 	{
-		Ar.ArIsError = 1;
-		Ar.ArIsCriticalError = 1;
+		Ar.SetCriticalError();
 		return false;
 	}
 
@@ -232,8 +231,7 @@ bool LoadKeyString(FArchive& Ar, FInlineStringBuffer& OutStrBuffer)
 	const int64 MaxSerializeSize = Ar.GetMaxSerializeSize();
 	if ((MaxSerializeSize > 0) && (SaveNum > MaxSerializeSize))
 	{
-		Ar.ArIsError = 1;
-		Ar.ArIsCriticalError = 1;
+		Ar.SetCriticalError();
 		return false;
 	}
 
diff --git a/Engine/Source/Runtime/Core/Private/Serialization/Archive.cpp b/Engine/Source/Runtime/Core/Private/Serialization/Archive.cpp
--- a/Engine/Source/Runtime/Core/Private/Serialization/Archive.cpp
+++ b/Engine/Source/Runtime/Core/Private/Serialization/Archive.cpp
@@ -1,6 +1,7 @@
 // Copyright Epic Games, Inc. All Rights Reserved.
 
 #include "Serialization/Archive.h"
+#include "Serialization/ArchiveProxy.h"
 #include "Math/UnrealMathUtility.h"
 #include "HAL/UnrealMemory.h"
 #include "Containers/Array.h"
@@ -9,7 +10,6 @@
 #include "Logging/LogMacros.h"
 #include "Misc/Parse.h"
 #include "UObject/ObjectVersion.h"
-#include "Serialization/ArchiveProxy.h"
 #include "Serialization/NameAsStringProxyArchive.h"
 #include "Misc/CommandLine.h"
 #include "Internationalization/Text.h"
@@ -261,6 +261,22 @@ FString FArchiveState::GetArchiveName() const
 	return TEXT("FArchive");
 }
 
+void FArchiveState::SetError()
+{
+	ArIsError = true;
+}
+
+void FArchiveState::ClearError()
+{
+	ArIsError = false;
+}
+
+void FArchiveState::SetCriticalError()
+{
+	SetError();
+	ArIsCriticalError = true;
+}
+
 void FArchiveState::GetSerializedPropertyChain(TArray<class FProperty*>& OutProperties) const
 {
 	if (SerializedPropertyChain)
@@ -446,7 +462,7 @@ void FArchive::SerializeBool( bool& D )
 	{
 		UE_LOG(LogSerialization, Error, TEXT("Invalid boolean encountered while reading archive %s - stream is most likely corrupted."), *GetArchiveName());
 
-		this->ArIsError = true;
+		this->SetError();
 	}
 	D = !!OldUBoolValue;
 }
diff --git a/Engine/Source/Runtime/Core/Private/Serialization/BitReader.cpp b/Engine/Source/Runtime/Core/Private/Serialization/BitReader.cpp
--- a/Engine/Source/Runtime/Core/Private/Serialization/BitReader.cpp
+++ b/Engine/Source/Runtime/Core/Private/Serialization/BitReader.cpp
@@ -144,7 +144,7 @@ void FBitReader::SetData( uint8* Src, int64 CountBits )
 {
 	Num			= CountBits;
 	Pos			= 0;
-	ArIsError	= 0;
+	ClearError();
 
 	Buffer.Empty();
 	Buffer.AddUninitialized( (Num+7)>>3 );
@@ -164,7 +164,7 @@ void FBitReader::SetData( TArray<uint8>&& Src, int64 CountBits )
 {
 	Num			= CountBits;
 	Pos			= 0;
-	ArIsError	= 0;
+	ClearError();
 
 	Buffer = MoveTemp(Src);
 
@@ -178,7 +178,7 @@ void FBitReader::SetData( FBitReader& Src, int64 CountBits )
 {
 	Num			= CountBits;
 	Pos			= 0;
-	ArIsError	= 0;
+	ClearError();
 
 	// Setup network version
 	this->SetEngineNetVer(Src.EngineNetVer());
@@ -231,7 +231,7 @@ void FBitReader::SetOverflowed(int64 LengthBits)
 	UE_LOG(LogNetSerialization, Error, TEXT("FBitReader::SetOverflowed() called! (ReadLen: %i, Remaining: %i, Max: %i)"),
 			LengthBits, (Num - Pos), Num);
 
-	ArIsError = 1;
+	SetError();
 }
 
 void FBitReader::SerializeBitsWithOffset( void* Dest, int32 DestBit, int64 LengthBits )
diff --git a/Engine/Source/Runtime/Core/Private/Serialization/BitWriter.cpp b/Engine/Source/Runtime/Core/Private/Serialization/BitWriter.cpp
--- a/Engine/Source/Runtime/Core/Private/Serialization/BitWriter.cpp
+++ b/Engine/Source/Runtime/Core/Private/Serialization/BitWriter.cpp
@@ -203,7 +203,7 @@ void FBitWriter::SetOverflowed(int32 LengthBits)
 				LengthBits, (Max-Num), Max);
 	}
 
-	ArIsError = 1;
+	SetError();
 }
 
 void FBitWriter::CountMemory(FArchive& Ar) const
@@ -280,7 +280,15 @@ void FBitWriterMark::Pop( FBitWriter& Writer )
 		checkSlow(End<=Writer.Buffer.Num());
 		FMemory::Memzero( &Writer.Buffer[Start], End-Start );
 	}
-	Writer.ArIsError = Overflowed;
+
+	if (Overflowed)
+	{
+		Writer.SetError();
+	}
+	else
+	{
+		Writer.ClearError();
+	}
 	Writer.Num       = Num;
 }
 
diff --git a/Engine/Source/Runtime/Core/Private/Serialization/LargeMemoryReader.cpp b/Engine/Source/Runtime/Core/Private/Serialization/LargeMemoryReader.cpp
--- a/Engine/Source/Runtime/Core/Private/Serialization/LargeMemoryReader.cpp
+++ b/Engine/Source/Runtime/Core/Private/Serialization/LargeMemoryReader.cpp
@@ -22,7 +22,7 @@ FLargeMemoryReader::FLargeMemoryReader(const uint8* InData, const int64 Num, ELa
 
 void FLargeMemoryReader::Serialize(void* OutData, int64 Num)
 {
-	if (Num && !ArIsError)
+	if (Num && !IsError())
 	{
 		// Only serialize if we have the requested amount of data
 		if (Offset + Num <= NumBytes)
@@ -32,7 +32,7 @@ void FLargeMemoryReader::Serialize(void* OutData, int64 Num)
 		}
 		else
 		{
-			ArIsError = true;
+			SetError();
 		}
 	}
 }
diff --git a/Engine/Source/Runtime/Core/Private/UObject/UnrealNames.cpp b/Engine/Source/Runtime/Core/Private/UObject/UnrealNames.cpp
--- a/Engine/Source/Runtime/Core/Private/UObject/UnrealNames.cpp
+++ b/Engine/Source/Runtime/Core/Private/UObject/UnrealNames.cpp
@@ -2478,8 +2478,7 @@ FArchive& operator<<(FArchive& Ar, FNameEntrySerialized& E)
 			// If StringLen cannot be negated due to integer overflow, Ar is corrupted.
 			if (StringLen == MIN_int32)
 			{
-				Ar.ArIsError = 1;
-				Ar.ArIsCriticalError = 1;
+				Ar.SetCriticalError();
 				UE_LOG(LogUnrealNames, Error, TEXT("Archive is corrupted"));
 				return Ar;
 			}
@@ -2490,8 +2489,7 @@ FArchive& operator<<(FArchive& Ar, FNameEntrySerialized& E)
 			// Protect against network packets allocating too much memory
 			if ((MaxSerializeSize > 0) && (StringLen > MaxSerializeSize))
 			{
-				Ar.ArIsError = 1;
-				Ar.ArIsCriticalError = 1;
+				Ar.SetCriticalError();
 				UE_LOG(LogUnrealNames, Error, TEXT("String is too large"));
 				return Ar;
 			}
@@ -2518,8 +2516,7 @@ FArchive& operator<<(FArchive& Ar, FNameEntrySerialized& E)
 			// Protect against network packets allocating too much memory
 			if ((MaxSerializeSize > 0) && (StringLen > MaxSerializeSize))
 			{
-				Ar.ArIsError = 1;
-				Ar.ArIsCriticalError = 1;
+				Ar.SetCriticalError();
 				UE_LOG(LogUnrealNames, Error, TEXT("String is too large"));
 				return Ar;
 			}
diff --git a/Engine/Source/Runtime/Core/Public/Misc/SecureHash.h b/Engine/Source/Runtime/Core/Public/Misc/SecureHash.h
--- a/Engine/Source/Runtime/Core/Public/Misc/SecureHash.h
+++ b/Engine/Source/Runtime/Core/Public/Misc/SecureHash.h
@@ -477,7 +477,7 @@ public:
 		}
 		
 		// note that we don't allow the base class CLose to happen, as the FAsyncSHAVerify will free the buffer if needed
-		return !ArIsError;
+		return !IsError();
 	}
 	/**
   	 * Returns the name of the Archive.  Useful for getting the name of the package a struct or object
diff --git a/Engine/Source/Runtime/Core/Public/Serialization/Archive.h b/Engine/Source/Runtime/Core/Public/Serialization/Archive.h
--- a/Engine/Source/Runtime/Core/Public/Serialization/Archive.h
+++ b/Engine/Source/Runtime/Core/Public/Serialization/Archive.h
@@ -110,10 +110,21 @@ public:
 		return ArIsError;
 	}
 
-	void SetError()
-	{
-		ArIsError = true;
-	}
+	/**
+	 * Sets ArIsError to true. Also sets error in the proxy archiver if one is wrapping this.
+	 */
+	void SetError();
+
+	/**
+	 * Sets ArIsError to false, this does not clear any CriticalErrors
+	 */
+	void ClearError();
+
+	/**
+	 * Sets the archiver IsCriticalError and IsError to true. Also sets CriticalError in the proxy archiver if one is wrapping this.
+	 */
+	void SetCriticalError();
+
 
 	FORCEINLINE bool IsByteSwapping()
 	{
@@ -621,13 +632,14 @@ protected:
 	/** Whether this archive saves to persistent storage. */
 	uint8 ArIsPersistent : 1;
 
-public:
+private:
 	/** Whether this archive contains errors. */
 	uint8 ArIsError : 1;
 
 	/** Whether this archive contains critical errors. */
 	uint8 ArIsCriticalError : 1;
 
+public:
 	/** Quickly tell if an archive contains script code. */
 	uint8 ArContainsCode : 1;
 
@@ -1499,11 +1511,13 @@ public:
 
 	virtual bool Close()
 	{
-		return !ArIsError;
+		return !IsError();
 	}
 
 	using FArchiveState::GetError;
 	using FArchiveState::SetError;
+	using FArchiveState::ClearError;
+	using FArchiveState::SetCriticalError;
 
 	/**
 	 * Serializes and compresses/ uncompresses data. This is a shared helper function for compression
diff --git a/Engine/Source/Runtime/Core/Public/Serialization/ArrayReader.h b/Engine/Source/Runtime/Core/Public/Serialization/ArrayReader.h
--- a/Engine/Source/Runtime/Core/Public/Serialization/ArrayReader.h
+++ b/Engine/Source/Runtime/Core/Public/Serialization/ArrayReader.h
@@ -32,7 +32,7 @@ public:
 
 	void Serialize(void* Data, int64 Count)
 	{
-		if (Count && !ArIsError)
+		if (Count && !IsError())
 		{
 			// Only serialize if we have the requested amount of data
 			if (Offset + Count <= Num())
@@ -42,7 +42,7 @@ public:
 			}
 			else
 			{
-				ArIsError = true;
+				SetError();
 			}
 		}
 	}
diff --git a/Engine/Source/Runtime/Core/Public/Serialization/BitReader.h b/Engine/Source/Runtime/Core/Public/Serialization/BitReader.h
--- a/Engine/Source/Runtime/Core/Public/Serialization/BitReader.h
+++ b/Engine/Source/Runtime/Core/Public/Serialization/BitReader.h
@@ -167,7 +167,7 @@ public:
 	}
 	FORCEINLINE_DEBUGGABLE bool AtEnd()
 	{
-		return ArIsError || Pos>=Num;
+		return IsError() || Pos>=Num;
 	}
 	FORCEINLINE_DEBUGGABLE int64 GetNumBytes() const
 	{
diff --git a/Engine/Source/Runtime/Core/Public/Serialization/BitWriter.h b/Engine/Source/Runtime/Core/Public/Serialization/BitWriter.h
--- a/Engine/Source/Runtime/Core/Public/Serialization/BitWriter.h
+++ b/Engine/Source/Runtime/Core/Public/Serialization/BitWriter.h
@@ -212,7 +212,7 @@ public:
 	FORCEINLINE_DEBUGGABLE void Init( FBitWriter& Writer)
 	{
 		Num = Writer.Num;
-		Overflowed = Writer.ArIsError;
+		Overflowed = Writer.IsError();
 	}
 
 	void Reset()
diff --git a/Engine/Source/Runtime/Core/Public/Serialization/BufferReader.h b/Engine/Source/Runtime/Core/Public/Serialization/BufferReader.h
--- a/Engine/Source/Runtime/Core/Public/Serialization/BufferReader.h
+++ b/Engine/Source/Runtime/Core/Public/Serialization/BufferReader.h
@@ -45,7 +45,7 @@ public:
 			FMemory::Free(ReaderData);
 			ReaderData = nullptr;
 		}
-		return !ArIsError;
+		return !IsError();
 	}
 	void Serialize(void* Data, int64 Num) final
 	{
diff --git a/Engine/Source/Runtime/Core/Public/Serialization/BufferWriter.h b/Engine/Source/Runtime/Core/Public/Serialization/BufferWriter.h
--- a/Engine/Source/Runtime/Core/Public/Serialization/BufferWriter.h
+++ b/Engine/Source/Runtime/Core/Public/Serialization/BufferWriter.h
@@ -56,7 +56,7 @@ public:
 			FMemory::Free(WriterData);
 			WriterData = nullptr;
 		}
-		return !ArIsError;
+		return !IsError();
 	}
 	void Serialize(void* Data, int64 Num)
 	{
diff --git a/Engine/Source/Runtime/Core/Public/Serialization/MemoryReader.h b/Engine/Source/Runtime/Core/Public/Serialization/MemoryReader.h
--- a/Engine/Source/Runtime/Core/Public/Serialization/MemoryReader.h
+++ b/Engine/Source/Runtime/Core/Public/Serialization/MemoryReader.h
@@ -33,7 +33,7 @@ public:
 
 	void Serialize( void* Data, int64 Num )
 	{
-		if (Num && !ArIsError)
+		if (Num && !IsError())
 		{
 			// Only serialize if we have the requested amount of data
 			if (Offset + Num <= TotalSize())
@@ -43,7 +43,7 @@ public:
 			}
 			else
 			{
-				ArIsError = true;
+				SetError();
 			}
 		}
 	}
@@ -91,7 +91,7 @@ public:
 
 	void Serialize( void* Data, int64 Num )
 	{
-		if (Num && !ArIsError)
+		if (Num && !IsError())
 		{
 			// Only serialize if we have the requested amount of data
 			if (Offset + Num <= TotalSize())
@@ -101,7 +101,7 @@ public:
 			}
 			else
 			{
-				ArIsError = true;
+				SetError();
 			}
 		}
 	}
diff --git a/Engine/Source/Runtime/Core/Public/Serialization/StaticMemoryReader.h b/Engine/Source/Runtime/Core/Public/Serialization/StaticMemoryReader.h
--- a/Engine/Source/Runtime/Core/Public/Serialization/StaticMemoryReader.h
+++ b/Engine/Source/Runtime/Core/Public/Serialization/StaticMemoryReader.h
@@ -32,7 +32,7 @@ public:
 
 	void Serialize( void* Data, int64 Num )
 	{
-		if (Num && !ArIsError)
+		if (Num && !IsError())
 		{
 			// Only serialize if we have the requested amount of data
 			if (Offset + Num <= TotalSize())
@@ -42,7 +42,7 @@ public:
 			}
 			else
 			{
-				ArIsError = true;
+				SetError();
 			}
 		}
 	}
diff --git a/Engine/Source/Runtime/CoreUObject/Private/Serialization/ArchiveStackTrace.cpp b/Engine/Source/Runtime/CoreUObject/Private/Serialization/ArchiveStackTrace.cpp
--- a/Engine/Source/Runtime/CoreUObject/Private/Serialization/ArchiveStackTrace.cpp
+++ b/Engine/Source/Runtime/CoreUObject/Private/Serialization/ArchiveStackTrace.cpp
@@ -616,13 +616,13 @@ void FArchiveStackTrace::CompareWithInternal(const FPackageData& SourcePackage,
 						SourceVal.Serialize(SourceReader);
 						DestVal  .Serialize(DestReader);
 
-						if (!SourceReader.ArIsError && !DestReader.ArIsError)
-						{
-							SourceVal.ExportText(BeforePropertyVal);
-							DestVal  .ExportText(AfterPropertyVal);
-						}
-					}
-				}
+									if (!SourceReader.IsError() && !DestReader.IsError())
+									{
+										SourceVal.ExportText(BeforePropertyVal);
+										DestVal  .ExportText(AfterPropertyVal);
+									}
+								}
+							}
 
 				FString DiffValues;
 				if (BeforePropertyVal != AfterPropertyVal)
diff --git a/Engine/Source/Runtime/CoreUObject/Private/Serialization/AsyncLoading.cpp b/Engine/Source/Runtime/CoreUObject/Private/Serialization/AsyncLoading.cpp
--- a/Engine/Source/Runtime/CoreUObject/Private/Serialization/AsyncLoading.cpp
+++ b/Engine/Source/Runtime/CoreUObject/Private/Serialization/AsyncLoading.cpp
@@ -7179,9 +7179,9 @@ FAsyncArchive::~FAsyncArchive()
 
 void FAsyncArchive::ReadCallback(bool bWasCancelled, IAsyncReadRequest* Request)
 {
-	if (bWasCancelled || ArIsError)
+	if (bWasCancelled || IsError())
 	{
-		ArIsError = true;
+		SetError();
 		return; // we don't do much with this, the code on the other thread knows how to deal with my request
 	}
 	if (LoadPhase == ELoadPhase::WaitingForSize)
@@ -7190,7 +7190,7 @@ void FAsyncArchive::ReadCallback(bool bWasCancelled, IAsyncReadRequest* Request)
 		FileSize = Request->GetSizeResults();
 		if (FileSize < 32)
 		{
-			ArIsError = true;
+			SetError();
 		}
 		else
 		{
@@ -7227,7 +7227,7 @@ void FAsyncArchive::ReadCallback(bool bWasCancelled, IAsyncReadRequest* Request)
 		uint8* Mem = Request->GetReadResults();
 		if (!Mem)
 		{
-			ArIsError = true;
+			SetError();
 			FPlatformMisc::MemoryBarrier();
 			LoadPhase = ELoadPhase::WaitingForHeader;
 		}
@@ -7238,7 +7238,7 @@ void FAsyncArchive::ReadCallback(bool bWasCancelled, IAsyncReadRequest* Request)
 			Ar << Sum;
 			if (Ar.IsError() || Sum.TotalHeaderSize > FileSize || Sum.GetFileVersionUE4() < VER_UE4_OLDEST_LOADABLE_PACKAGE)
 			{
-				ArIsError = true;
+				SetError();
 			}
 			else
 			{
@@ -7342,7 +7342,7 @@ bool FAsyncArchive::Close()
 	// Invalidate any precached data and free memory.
 	FlushCache();
 	// Return true if there were NO errors, false otherwise.
-	return !ArIsError;
+	return !IsError();
 }
 
 bool FAsyncArchive::SetCompressionMap(TArray<FCompressedChunk>* InCompressedChunks, ECompressionFlags InCompressionFlags)
@@ -7435,12 +7435,12 @@ void FAsyncArchive::CompleteRead()
 	{
 		FlushPrecacheBlock();
 	}
-	if (!ArIsError)
+	if (!IsError())
 	{
 		uint8* Mem = ReadRequestPtr->GetReadResults();
 		if (!Mem)
 		{
-			ArIsError = true;
+			SetError();
 		}
 		else
 		{
@@ -7627,7 +7627,7 @@ bool FAsyncArchive::PrecacheInternal(int64 RequestOffset, int64 RequestSize, boo
 	}
 	if (ReadRequestSize <= 0)
 	{
-		ArIsError = true;
+		SetError();
 		return true;
 	}
 	double StartTime = FPlatformTime::Seconds();
@@ -7812,7 +7812,7 @@ void FAsyncArchive::StartReadingHeader()
 {
 	//LogItem(TEXT("Start Header"));
 	WaitForIntialPhases();
-	if (!ArIsError)
+	if (!IsError())
 	{
 		if (int32(LoadPhase) < int32(ELoadPhase::WaitingForHeader))
 		{
@@ -7867,7 +7867,7 @@ void CallSerializeHook();
 
 void FAsyncArchive::Serialize(void* Data, int64 Count)
 {
-	if (!Count || ArIsError)
+	if (!Count || IsError())
 	{
 		return;
 	}
@@ -7959,7 +7959,7 @@ void FAsyncArchive::Serialize(void* Data, int64 Count)
 		{
 		WaitRead();
 		}
-		if (ArIsError)
+		if (IsError())
 		{
 			return;
 		}
@@ -7993,7 +7993,7 @@ void FAsyncArchive::Serialize(void* Data, int64 Count)
 		{
 			verify(WaitRead());
 			void* OldRead2 = ReadRequestPtr;
-			if (!ArIsError)
+			if (!IsError())
 			{
 				checkf(AfterBlockOffset >= PrecacheStartPos && AfterBlockOffset + AfterBlockSize <= PrecacheEndPos, 
 					TEXT("Sync After Block Wait ????  %lld %lld     %lld %lld <-  %lld %lld     %lld %lld <-  %lld %lld    %p <- %p <- %p    %lld %lld <-  %lld %lld"), 
@@ -8005,7 +8005,7 @@ void FAsyncArchive::Serialize(void* Data, int64 Count)
 				);
 			}
 		}
-		if (ArIsError)
+		if (IsError())
 		{
 			return;
 		}
diff --git a/Engine/Source/Runtime/CoreUObject/Private/Serialization/AsyncLoading2.cpp b/Engine/Source/Runtime/CoreUObject/Private/Serialization/AsyncLoading2.cpp
--- a/Engine/Source/Runtime/CoreUObject/Private/Serialization/AsyncLoading2.cpp
+++ b/Engine/Source/Runtime/CoreUObject/Private/Serialization/AsyncLoading2.cpp
@@ -147,7 +147,7 @@ public:
 	void Serialize(void* Data, int64 Length) override
 	{
 #if DEVIRTUALIZE_FLinkerLoad_Serialize
-		if (!Length || ArIsError)
+		if (!Length || IsError())
 		{
 			return;
 		}
@@ -1024,8 +1024,7 @@ public:
 			TEXT("Index: %d/%d"), NameIndex, GlobalNameMap->Num());
 
 		Name = FName();
-		ArIsError = true;
-		ArIsCriticalError = true;
+		SetCriticalError();
 	}
 
 	inline virtual FArchive& operator<<(FName& Name) override
diff --git a/Engine/Source/Runtime/CoreUObject/Public/Serialization/ObjectReader.h b/Engine/Source/Runtime/CoreUObject/Public/Serialization/ObjectReader.h
--- a/Engine/Source/Runtime/CoreUObject/Public/Serialization/ObjectReader.h
+++ b/Engine/Source/Runtime/CoreUObject/Public/Serialization/ObjectReader.h
@@ -46,7 +46,7 @@ public:
 
 	void Serialize(void* Data, int64 Num)
 	{
-		if (Num && !ArIsError)
+		if (Num && !IsError())
 		{
 			// Only serialize if we have the requested amount of data
 			if (Offset + Num <= TotalSize())
@@ -56,7 +56,7 @@ public:
 			}
 			else
 			{
-				ArIsError = true;
+				SetError();
 			}
 		}
 	}
diff --git a/Engine/Source/Runtime/CoreUObject/Public/UObject/LinkerLoad.h b/Engine/Source/Runtime/CoreUObject/Public/UObject/LinkerLoad.h
--- a/Engine/Source/Runtime/CoreUObject/Public/UObject/LinkerLoad.h
+++ b/Engine/Source/Runtime/CoreUObject/Public/UObject/LinkerLoad.h
@@ -838,8 +838,7 @@ private:
 		{
 			Name = FName();
 			BadNameIndexError(NameIndex);
-			ArIsError = true;
-			ArIsCriticalError = true;
+			SetCriticalError();
 		}
 
 		return *this;
diff --git a/Engine/Source/Runtime/Engine/Public/Net/DataBunch.h b/Engine/Source/Runtime/Engine/Public/Net/DataBunch.h
--- a/Engine/Source/Runtime/Engine/Public/Net/DataBunch.h
+++ b/Engine/Source/Runtime/Engine/Public/Net/DataBunch.h
@@ -183,13 +183,13 @@ struct FControlChannelOutBunch : public FOutBunch
 	FArchive& operator<<(FName& Name)
 	{
 		UE_LOG(LogNet, Fatal,TEXT("Cannot send Names on the control channel"));
-		ArIsError = true;
+		SetError();
 		return *this;
 	}
 	FArchive& operator<<(UObject*& Object)
 	{
 		UE_LOG(LogNet, Fatal,TEXT("Cannot send Objects on the control channel"));
-		ArIsError = true;
+		SetError();
 		return *this;
 	}
 };
diff --git a/Engine/Source/Runtime/MovieSceneTracks/Private/Sections/MovieSceneEventSection.cpp b/Engine/Source/Runtime/MovieSceneTracks/Private/Sections/MovieSceneEventSection.cpp
--- a/Engine/Source/Runtime/MovieSceneTracks/Private/Sections/MovieSceneEventSection.cpp
+++ b/Engine/Source/Runtime/MovieSceneTracks/Private/Sections/MovieSceneEventSection.cpp
@@ -61,9 +61,9 @@ public:
 	}
 
 	// Unsupported serialization
-	virtual FArchive& operator<<(UObject*& Res) override 					{ ArIsError = true; return *this; }
-	virtual FArchive& operator<<(FLazyObjectPtr& LazyObjectPtr) override 	{ ArIsError = true; return *this; }
-	virtual FArchive& operator<<(FWeakObjectPtr& Value) override 			{ ArIsError = true; return *this; }
+	virtual FArchive& operator<<(UObject*& Res) override 					{ SetError(); return *this; }
+	virtual FArchive& operator<<(FLazyObjectPtr& LazyObjectPtr) override 	{ SetError(); return *this; }
+	virtual FArchive& operator<<(FWeakObjectPtr& Value) override 			{ SetError(); return *this; }
 };
 
 /** Custom archive used for writing event parameter struct payloads */
@@ -203,7 +203,7 @@ public:
 
 	virtual void Serialize(void* Data, int64 Num) override
 	{
-		if (Num && !ArIsError)
+		if (Num && !IsError())
 		{
 			// Only serialize if we have the requested amount of data
 			if (Offset + Num <= Bytes.Num())
@@ -213,7 +213,7 @@ public:
 			}
 			else
 			{
-				ArIsError = true;
+				SetError();
 			}
 		}
 	}
diff --git a/Engine/Source/Runtime/Net/Core/Private/Net/Core/Misc/PacketAudit.cpp b/Engine/Source/Runtime/Net/Core/Private/Net/Core/Misc/PacketAudit.cpp
--- a/Engine/Source/Runtime/Net/Core/Private/Net/Core/Misc/PacketAudit.cpp
+++ b/Engine/Source/Runtime/Net/Core/Private/Net/Core/Misc/PacketAudit.cpp
@@ -157,7 +157,7 @@ public:
 
 	virtual void Serialize(void* Data, int64 Num) override
 	{
-		if (Num > 0 && !ArIsError)
+		if (Num > 0 && !IsError())
 		{
 			if (Offset + Num <= TotalSize())
 			{
@@ -166,7 +166,7 @@ public:
 			}
 			else
 			{
-				ArIsError = true;
+				SetError();
 			}
 		}
 	}
diff --git a/Engine/Source/Runtime/NetworkReplayStreaming/HttpNetworkReplayStreaming/Private/HttpNetworkReplayStreaming.cpp b/Engine/Source/Runtime/NetworkReplayStreaming/HttpNetworkReplayStreaming/Private/HttpNetworkReplayStreaming.cpp
--- a/Engine/Source/Runtime/NetworkReplayStreaming/HttpNetworkReplayStreaming/Private/HttpNetworkReplayStreaming.cpp
+++ b/Engine/Source/Runtime/NetworkReplayStreaming/HttpNetworkReplayStreaming/Private/HttpNetworkReplayStreaming.cpp
@@ -156,7 +156,7 @@ void FHttpStreamFArchive::Serialize( void* V, int64 Length )
 	{
 		if ( Pos + Length > Buffer.Num() )
 		{
-			ArIsError = true;
+			SetError();
 			return;
 		}
 			
diff --git a/Engine/Source/Runtime/NetworkReplayStreaming/InMemoryNetworkReplayStreaming/Private/InMemoryNetworkReplayStreaming.cpp b/Engine/Source/Runtime/NetworkReplayStreaming/InMemoryNetworkReplayStreaming/Private/InMemoryNetworkReplayStreaming.cpp
--- a/Engine/Source/Runtime/NetworkReplayStreaming/InMemoryNetworkReplayStreaming/Private/InMemoryNetworkReplayStreaming.cpp
+++ b/Engine/Source/Runtime/NetworkReplayStreaming/InMemoryNetworkReplayStreaming/Private/InMemoryNetworkReplayStreaming.cpp
@@ -566,7 +566,7 @@ void FInMemoryReplayStreamArchive::Serialize(void* V, int64 Length)
 	{
 		if (Pos + Length > TotalSize())
 		{
-			ArIsError = true;
+			SetError();
 			return;
 		}
 		
@@ -574,7 +574,7 @@ void FInMemoryReplayStreamArchive::Serialize(void* V, int64 Length)
 
 		if (CurrentChunk == nullptr)
 		{
-			ArIsError = true;
+			SetError();
 			return;
 		}
 
@@ -591,7 +591,7 @@ void FInMemoryReplayStreamArchive::Serialize(void* V, int64 Length)
 
 		if (CurrentChunk == nullptr)
 		{
-			ArIsError = true;
+			SetError();
 			return;
 		}
 
diff --git a/Engine/Source/Runtime/Online/BuildPatchServices/Private/Data/ManifestUObject.cpp b/Engine/Source/Runtime/Online/BuildPatchServices/Private/Data/ManifestUObject.cpp
--- a/Engine/Source/Runtime/Online/BuildPatchServices/Private/Data/ManifestUObject.cpp
+++ b/Engine/Source/Runtime/Online/BuildPatchServices/Private/Data/ManifestUObject.cpp
@@ -219,7 +219,7 @@ public:
 
 	virtual void Serialize(void* Data, int64 Num) override
 	{
-		if (Num && !ArIsError)
+		if (Num && !IsError())
 		{
 			const int64 NumBytesToAdd = Offset + Num - Bytes.Num();
 			if (NumBytesToAdd > 0)
@@ -227,7 +227,7 @@ public:
 				const int64 NewArrayCount = Bytes.Num() + NumBytesToAdd;
 				if (NewArrayCount >= MAX_int32)
 				{
-					ArIsError = true;
+					SetError();
 					return;
 				}
 				Bytes.AddUninitialized((int32)NumBytesToAdd);
@@ -306,7 +306,7 @@ public:
 		}
 		else
 		{
-			ArIsError = true;
+			SetError();
 		}
 		Offset = NameTableReader.Tell();
 	}
@@ -329,7 +329,7 @@ public:
 
 	virtual FArchive& operator<<(FName& N) override
 	{
-		if (ArIsError)
+		if (IsError())
 		{
 			N = NAME_None;
 		}
@@ -345,7 +345,7 @@ public:
 			else
 			{
 				N = NAME_None;
-				ArIsError = true;
+				SetError();
 			}
 		}
 		return *this;
@@ -353,7 +353,7 @@ public:
 
 	virtual void Serialize(void* Data, int64 Num) override
 	{
-		if (Num && !ArIsError)
+		if (Num && !IsError())
 		{
 			// Only serialize if we have the requested amount of data
 			if (Offset + Num <= Bytes.Num())
@@ -363,7 +363,7 @@ public:
 			}
 			else
 			{
-				ArIsError = true;
+				SetError();
 			}
 		}
 	}
diff --git a/Engine/Source/Runtime/Online/BuildPatchServices/Private/Tests/Fake/FileSystem.fake.h b/Engine/Source/Runtime/Online/BuildPatchServices/Private/Tests/Fake/FileSystem.fake.h
--- a/Engine/Source/Runtime/Online/BuildPatchServices/Private/Tests/Fake/FileSystem.fake.h
+++ b/Engine/Source/Runtime/Online/BuildPatchServices/Private/Tests/Fake/FileSystem.fake.h
@@ -36,7 +36,7 @@ namespace BuildPatchServices
 
 		void Serialize(void* Data, int64 Num)
 		{
-			if (Num && !ArIsError)
+			if (Num && !IsError())
 			{
 				if (Offset + Num <= TotalSize())
 				{
@@ -45,7 +45,7 @@ namespace BuildPatchServices
 				}
 				else
 				{
-					ArIsError = true;
+					SetError();
 				}
 			}
 		}
