Override GetPortableCallstack to avoid 1MB allocations in our signal handler

diff --git a/Engine/Source/Runtime/Core/Private/Unix/UnixPlatformCrashContext.cpp b/Engine/Source/Runtime/Core/Private/Unix/UnixPlatformCrashContext.cpp
--- a/Engine/Source/Runtime/Core/Private/Unix/UnixPlatformCrashContext.cpp
+++ b/Engine/Source/Runtime/Core/Private/Unix/UnixPlatformCrashContext.cpp
@@ -290,8 +290,8 @@ void FUnixCrashContext::CaptureStackTrace()
 	// Only do work the first time this function is called - this is mainly a carry over from Windows where it can be called multiple times, left intact for extra safety.
 	if (!bCapturedBacktrace)
 	{
-		const SIZE_T StackTraceSize = 65535;
-		ANSICHAR* StackTrace = (ANSICHAR*) FMemory::Malloc( StackTraceSize );
+		static const SIZE_T StackTraceSize = 65535;
+		static ANSICHAR StackTrace[StackTraceSize];
 		StackTrace[0] = 0;
 
 		int32 IgnoreCount = NumMinidumpFramesToIgnore;
@@ -307,11 +307,44 @@ void FUnixCrashContext::CaptureStackTrace()
 		FCString::Strncat( GErrorHist, UTF8_TO_TCHAR(StackTrace), UE_ARRAY_COUNT(GErrorHist) - 1 );
 		CreateExceptionInfoString(Signal, Info, Context);
 
-		FMemory::Free( StackTrace );
 		bCapturedBacktrace = true;
 	}
 }
 
+void FUnixCrashContext::GetPortableCallStack(const uint64* StackFrames, int32 NumStackFrames, TArray<FCrashStackFrame>& OutCallStack) const
+{
+	// Update the callstack with offsets from each module
+	OutCallStack.Reset(NumStackFrames);
+	for (int32 Idx = 0; Idx < NumStackFrames; Idx++)
+	{
+		const uint64 StackFrame = StackFrames[Idx];
+
+		// Try to find the module containing this stack frame
+		const FStackWalkModuleInfo* FoundModule = nullptr;
+
+		Dl_info DylibInfo;
+		int32 Result = dladdr((const void*)StackFrame, &DylibInfo);
+		if (Result != 0)
+		{
+			ANSICHAR* DylibPath = (ANSICHAR*)DylibInfo.dli_fname;
+			ANSICHAR* DylibName = FCStringAnsi::Strrchr(DylibPath, '/');
+			if (DylibName)
+			{
+				DylibName += 1;
+			}
+			else
+			{
+				DylibName = DylibPath;
+			}
+			OutCallStack.Add(FCrashStackFrame(FPaths::GetBaseFilename(DylibName), (uint64)DylibInfo.dli_fbase, StackFrame - (uint64)DylibInfo.dli_fbase));
+		}
+		else
+		{
+			OutCallStack.Add(FCrashStackFrame(TEXT("Unknown"), 0, StackFrame));
+		}
+	}
+}
+
 namespace UnixCrashReporterTracker
 {
 	FProcHandle CurrentlyRunningCrashReporter;
diff --git a/Engine/Source/Runtime/Core/Public/Unix/UnixPlatformCrashContext.h b/Engine/Source/Runtime/Core/Public/Unix/UnixPlatformCrashContext.h
--- a/Engine/Source/Runtime/Core/Public/Unix/UnixPlatformCrashContext.h
+++ b/Engine/Source/Runtime/Core/Public/Unix/UnixPlatformCrashContext.h
@@ -107,6 +107,8 @@ struct CORE_API FUnixCrashContext : public FGenericCrashContext
 	 */
 	void SetFirstCrashHandlerFrame(uint64* ProgramCounter);
 
+	virtual void GetPortableCallStack(const uint64* StackFrames, int32 NumStackFrames, TArray<FCrashStackFrame>& OutCallStack) const override;
+
 	void AddPlatformSpecificProperties() const;
 
 protected:
