Edit for crash investigation and slight optimization in SlowTask.

diff --git a/Engine/Source/Runtime/Core/Private/Misc/SlowTask.cpp b/Engine/Source/Runtime/Core/Private/Misc/SlowTask.cpp
--- a/Engine/Source/Runtime/Core/Private/Misc/SlowTask.cpp
+++ b/Engine/Source/Runtime/Core/Private/Misc/SlowTask.cpp
@@ -98,11 +98,19 @@ void FSlowTask::MakeDialogDelayed(float Threshold, bool bShowCancelButton, bool
 	bDelayedDialogAllowInPIE = bAllowInPIE;
 }
 
-void FSlowTask::EnterProgressFrame(float ExpectedWorkThisFrame, FText Text)
+void FSlowTask::EnterProgressFrame(float ExpectedWorkThisFrame, const FText& Text)
 {
 	check(!bEnabled || IsInGameThread());
 
-	FrameMessage = Text;
+	// Should actually be FrameMessage = Text; but this code is to investigate crashes in FSlowTask::GetCurrentMessage()
+	if (!Text.IsEmpty())
+	{
+		FrameMessage = Text;
+	}
+	else
+	{
+		FrameMessage = FText::GetEmpty();
+	}
 	CompletedWork += CurrentFrameScope;
 
 	// Make sure OS events are getting through while the task is being processed
@@ -124,7 +132,7 @@ void FSlowTask::EnterProgressFrame(float ExpectedWorkThisFrame, FText Text)
 	}
 }
 
-FText FSlowTask::GetCurrentMessage() const
+const FText& FSlowTask::GetCurrentMessage() const
 {
 	return FrameMessage.IsEmpty() ? DefaultMessage : FrameMessage;
 }
diff --git a/Engine/Source/Runtime/Core/Public/Misc/SlowTask.h b/Engine/Source/Runtime/Core/Public/Misc/SlowTask.h
--- a/Engine/Source/Runtime/Core/Public/Misc/SlowTask.h
+++ b/Engine/Source/Runtime/Core/Public/Misc/SlowTask.h
@@ -111,12 +111,12 @@ public:
 	 * @param		ExpectedWorkThisFrame	The amount of work that will happen between now and the next frame, as a numerator of TotalAmountOfWork.
 	 * @param		Text					Optional text to describe this frame's purpose.
 	 */
-	void EnterProgressFrame(float ExpectedWorkThisFrame = 1.f, FText Text = FText());
+	void EnterProgressFrame(float ExpectedWorkThisFrame = 1.f, const FText& Text = FText());
 
 	/**
 	 * Get the frame message or default message if empty
 	 */
-	FText GetCurrentMessage() const;
+	const FText& GetCurrentMessage() const;
 
 	/**
 	 * True if the user has requested that the slow task be canceled
