Checks in FVector::operator[] are now checkSlow.
#Were the #1 costliest checks due to the op[] usage inside Raycast code.
#Was measured to take up to 4.4ms in a single frame on Win64 test server.

diff --git a/Engine/Source/Runtime/Core/Public/Math/Vector.h b/Engine/Source/Runtime/Core/Public/Math/Vector.h
--- a/Engine/Source/Runtime/Core/Public/Math/Vector.h
+++ b/Engine/Source/Runtime/Core/Public/Math/Vector.h
@@ -1425,13 +1425,13 @@ FORCEINLINE FVector FVector::operator/=(const FVector& V)
 
 FORCEINLINE float& FVector::operator[](int32 Index)
 {
-	check(Index >= 0 && Index < 3);
+	checkSlow(Index >= 0 && Index < 3);
 	return (&X)[Index];
 }
 
 FORCEINLINE float FVector::operator[](int32 Index)const
 {
-	check(Index >= 0 && Index < 3);
+	checkSlow(Index >= 0 && Index < 3);
 	return (&X)[Index];
 }
 
