Implement FPlatformMisc::RestartApplication for Android.

diff --git a/Engine/Source/Runtime/Core/Private/Android/AndroidPlatformMisc.cpp b/Engine/Source/Runtime/Core/Private/Android/AndroidPlatformMisc.cpp
--- a/Engine/Source/Runtime/Core/Private/Android/AndroidPlatformMisc.cpp
+++ b/Engine/Source/Runtime/Core/Private/Android/AndroidPlatformMisc.cpp
@@ -142,6 +142,18 @@ void FAndroidMisc::RequestExit( bool Force )
 	}
 }
 
+extern void AndroidThunkCpp_RestartApplication(const FString& IntentString);
+
+bool FAndroidMisc::RestartApplication()
+{
+#if USE_ANDROID_JNI
+	AndroidThunkCpp_RestartApplication(TEXT(""));
+	return true;
+#else
+	return FGenericPlatformMisc::RestartApplication();
+#endif
+}
+
 void FAndroidMisc::LocalPrint(const TCHAR *Message)
 {
 	// Builds for distribution should not have logging in them:
diff --git a/Engine/Source/Runtime/Core/Public/Android/AndroidPlatformMisc.h b/Engine/Source/Runtime/Core/Public/Android/AndroidPlatformMisc.h
--- a/Engine/Source/Runtime/Core/Public/Android/AndroidPlatformMisc.h
+++ b/Engine/Source/Runtime/Core/Public/Android/AndroidPlatformMisc.h
@@ -32,6 +32,7 @@ class TFunction;
 struct CORE_API FAndroidMisc : public FGenericPlatformMisc
 {
 	static void RequestExit( bool Force );
+	static bool RestartApplication();
 	static void LocalPrint(const TCHAR *Message);
 	static bool IsLocalPrintThreadSafe() { return true; }
 	static void PlatformPreInit();
