Implement FPlatformMisc::RequestExit and FPlatformMisc::RequestExitWithStatus for IOS

diff --git a/Engine/Source/Runtime/Core/Private/IOS/IOSPlatformMisc.cpp b/Engine/Source/Runtime/Core/Private/IOS/IOSPlatformMisc.cpp
--- a/Engine/Source/Runtime/Core/Private/IOS/IOSPlatformMisc.cpp
+++ b/Engine/Source/Runtime/Core/Private/IOS/IOSPlatformMisc.cpp
@@ -1746,6 +1746,35 @@ bool FIOSPlatformMisc::HasSeparateChannelForDebugOutput()
 #endif
 }
 
+void FIOSPlatformMisc::RequestExit(bool Force)
+{
+	if (Force)
+	{
+		FApplePlatformMisc::RequestExit(Force);
+	}
+	else
+	{
+		// ForceExit is sort of a misnomer here.  This will exit the engine loop before calling _Exit() from the app delegate
+		[[IOSAppDelegate GetDelegate] ForceExit];
+	}
+}
+
+void FIOSPlatformMisc::RequestExitWithStatus(bool Force, uint8 ReturnCode)
+{
+	if (Force)
+	{
+		FApplePlatformMisc::RequestExit(Force);
+	}
+	else
+	{
+		// Implementation will ignore the return code - this may be important, so warn.
+		UE_LOG(LogIOS, Warning, TEXT("FIOSPlatformMisc::RequestExitWithStatus(%i, %d) - return code will be ignored by the generic implementation."), Force, ReturnCode);
+
+		// ForceExit is sort of a misnomer here.  This will exit the engine loop before calling _Exit() from the app delegate
+		[[IOSAppDelegate GetDelegate] ForceExit];
+	}
+}
+
 void FIOSPlatformMisc::GPUAssert()
 {
     // make this a fatal error that ends here not in the log
diff --git a/Engine/Source/Runtime/Core/Public/IOS/IOSPlatformMisc.h b/Engine/Source/Runtime/Core/Public/IOS/IOSPlatformMisc.h
--- a/Engine/Source/Runtime/Core/Public/IOS/IOSPlatformMisc.h
+++ b/Engine/Source/Runtime/Core/Public/IOS/IOSPlatformMisc.h
@@ -45,6 +45,9 @@ struct CORE_API FIOSPlatformMisc : public FApplePlatformMisc
 	static const TCHAR* GamePersistentDownloadDir();
     static bool HasSeparateChannelForDebugOutput();
 
+	static void RequestExit(bool Force);
+	static void RequestExitWithStatus(bool Force, uint8 ReturnCode);
+
 	UE_DEPRECATED(4.21, "Use GetDeviceVolume, it is now callable on all platforms.")
 	static int GetAudioVolume();
 
