Fixed FCachedFileHandle issues with append and truncate

FCachedFileHandle wasn't taking the current Tell value of the underlying handle on construction, so for a file opened in append mode it would be out-of-sync. It also wasn't updating its cached values on Truncate.

This also fixes some platform file impls not updating their cached file size on Truncate.

diff --git a/Engine/Source/Runtime/Core/Private/Windows/WindowsPlatformFile.cpp b/Engine/Source/Runtime/Core/Private/Windows/WindowsPlatformFile.cpp
--- a/Engine/Source/Runtime/Core/Private/Windows/WindowsPlatformFile.cpp
+++ b/Engine/Source/Runtime/Core/Private/Windows/WindowsPlatformFile.cpp
@@ -661,7 +661,12 @@ public:
 	{
 		// SetEndOfFile isn't an overlapped operation, so we need to call UpdatedNonOverlappedPos after seeking to ensure that the file pointer is in the correct place
 		check(IsValid());
-		return Seek(NewSize) && UpdatedNonOverlappedPos() && SetEndOfFile(FileHandle) != 0;
+		if (Seek(NewSize) && UpdatedNonOverlappedPos() && SetEndOfFile(FileHandle) != 0)
+		{
+			UpdateFileSize();
+			return true;
+		}
+		return false;
 	}
 };
 
diff --git a/Engine/Source/Runtime/Core/Public/HAL/IPlatformFileCachedWrapper.h b/Engine/Source/Runtime/Core/Public/HAL/IPlatformFileCachedWrapper.h
--- a/Engine/Source/Runtime/Core/Public/HAL/IPlatformFileCachedWrapper.h
+++ b/Engine/Source/Runtime/Core/Public/HAL/IPlatformFileCachedWrapper.h
@@ -23,8 +23,8 @@ class CORE_API FCachedFileHandle : public IFileHandle
 public:
 	FCachedFileHandle(IFileHandle* InFileHandle, bool bInReadable, bool bInWritable)
 		: FileHandle(InFileHandle)
-		, FilePos(0)
-		, TellPos(0)
+		, FilePos(InFileHandle->Tell())
+		, TellPos(FilePos)
 		, FileSize(InFileHandle->Size())
 		, bWritable(bInWritable)
 		, bReadable(bInReadable)
@@ -177,7 +177,13 @@ public:
 	{
 		if (bWritable)
 		{
-			return FileHandle->Truncate(NewSize);
+			if (FileHandle->Truncate(NewSize))
+			{
+				FlushCache();
+				FilePos = TellPos = FileHandle->Tell();
+				FileSize = FileHandle->Size();
+				return true;
+			}
 		}
 		return false;
 	}
