Allow Truncate to retry if interrupted

Ensure we're using the 64-bit version of Truncate so it can deal with large files correctly.

diff --git a/Engine/Source/Runtime/Core/Private/Android/AndroidPlatformFile.cpp b/Engine/Source/Runtime/Core/Private/Android/AndroidPlatformFile.cpp
--- a/Engine/Source/Runtime/Core/Private/Android/AndroidPlatformFile.cpp
+++ b/Engine/Source/Runtime/Core/Private/Android/AndroidPlatformFile.cpp
@@ -35,10 +35,12 @@ DEFINE_LOG_CATEGORY_STATIC(LogAndroidFile, Log, All);
 	#define __lseek(_fd, _offset, _whence)			lseek64(_fd, _offset, _whence)
 	#define __pread(_fd, _buf, _count, _offset)		pread64(_fd, _buf, _count, _offset)
 	#define __pwrite(_fd, _buf, _count, _offset)	pwrite64(_fd, _buf, _count, _offset)
+	#define __ftruncate(_fd, _length)				ftruncate64(_fd, _length)
 #else
 	#define __lseek(_fd, _offset, _whence)			lseek(_fd, _offset, _whence)
 	#define __pread(_fd, _buf, _count, _offset)		pread(_fd, _buf, _count, _offset)
 	#define __pwrite(_fd, _buf, _count, _offset)	pwrite(_fd, _buf, _count, _offset)
+	#define __ftruncate(_fd, _length)				ftruncate(_fd, _length)
 #endif
 
 // make an FTimeSpan object that represents the "epoch" for time_t (from a stat struct)
@@ -389,7 +391,9 @@ public:
 			return false;
 		}
 
-		return ftruncate(File->Handle, NewSize) == 0;
+		int Result = 0;
+		do { Result = __ftruncate(File->Handle, NewSize); } while (Result < 0 && errno == EINTR);
+		return Result == 0;
 	}
 
 	virtual int64 Size() override
diff --git a/Engine/Source/Runtime/Core/Private/Apple/ApplePlatformFile.cpp b/Engine/Source/Runtime/Core/Private/Apple/ApplePlatformFile.cpp
--- a/Engine/Source/Runtime/Core/Private/Apple/ApplePlatformFile.cpp
+++ b/Engine/Source/Runtime/Core/Private/Apple/ApplePlatformFile.cpp
@@ -221,7 +221,9 @@ public:
 			return false;
 		}
 #endif
-		return ftruncate(FileHandle, NewSize) == 0;
+		int Result = 0;
+		do { Result = ftruncate(FileHandle, NewSize); } while (Result < 0 && errno == EINTR);
+		return Result == 0;
 	}
 	virtual int64 Size() override
 	{
diff --git a/Engine/Source/Runtime/Core/Private/IOS/IOSPlatformFile.cpp b/Engine/Source/Runtime/Core/Private/IOS/IOSPlatformFile.cpp
--- a/Engine/Source/Runtime/Core/Private/IOS/IOSPlatformFile.cpp
+++ b/Engine/Source/Runtime/Core/Private/IOS/IOSPlatformFile.cpp
@@ -246,7 +246,9 @@ public:
 			return false;
 		}
 #endif
-		return ftruncate(FileHandle, NewSize) == 0;
+		int Result = 0;
+		do { Result = ftruncate(FileHandle, NewSize); } while (Result < 0 && errno == EINTR);
+		return Result == 0;
 	}
 
 	virtual int64 Size( ) override
diff --git a/Engine/Source/Runtime/Core/Private/Lumin/LuminPlatformFile.cpp b/Engine/Source/Runtime/Core/Private/Lumin/LuminPlatformFile.cpp
--- a/Engine/Source/Runtime/Core/Private/Lumin/LuminPlatformFile.cpp
+++ b/Engine/Source/Runtime/Core/Private/Lumin/LuminPlatformFile.cpp
@@ -294,7 +294,9 @@ public:
 			return false;
 		}
 #endif
-		return ftruncate(FileHandle, NewSize) == 0;
+		int Result = 0;
+		do { Result = ftruncate(FileHandle, NewSize); } while (Result < 0 && errno == EINTR);
+		return Result == 0;
 	}
 
 	virtual int64 Size() override
diff --git a/Engine/Source/Runtime/Core/Private/Unix/UnixPlatformFile.cpp b/Engine/Source/Runtime/Core/Private/Unix/UnixPlatformFile.cpp
--- a/Engine/Source/Runtime/Core/Private/Unix/UnixPlatformFile.cpp
+++ b/Engine/Source/Runtime/Core/Private/Unix/UnixPlatformFile.cpp
@@ -231,7 +231,10 @@ public:
 	virtual bool Truncate(int64 NewSize) override
 	{
 		check(IsValid());
-		return ftruncate(FileHandle, NewSize) == 0;
+		
+		int Result = 0;
+		do { Result = ftruncate(FileHandle, NewSize); } while (Result < 0 && errno == EINTR);
+		return Result == 0;
 	}
 
 	virtual int64 Size() override
