Re-enabled a test case for file appending, and added one for file truncation

#rnx

diff --git a/Engine/Source/Runtime/Core/Private/Tests/Misc/FileTest.cpp b/Engine/Source/Runtime/Core/Private/Tests/Misc/FileTest.cpp
--- a/Engine/Source/Runtime/Core/Private/Tests/Misc/FileTest.cpp
+++ b/Engine/Source/Runtime/Core/Private/Tests/Misc/FileTest.cpp
@@ -4,89 +4,223 @@
 #include "Misc/AssertionMacros.h"
 #include "Containers/Array.h"
 #include "Containers/UnrealString.h"
+#include "HAL/PlatformFile.h"
+#include "HAL/PlatformFilemanager.h"
 #include "HAL/FileManager.h"
 #include "Misc/FileHelper.h"
 #include "Misc/Paths.h"
 #include "Misc/Guid.h"
+#include "Misc/ScopeExit.h"
 #include "Misc/AutomationTest.h"
 
 #if WITH_DEV_AUTOMATION_TESTS
 
-IMPLEMENT_SIMPLE_AUTOMATION_TEST(FFileTests, "System.Core.Misc.File", EAutomationTestFlags::EditorContext | EAutomationTestFlags::SmokeFilter)
-
 // These file tests are designed to ensure expected file writing behavior, as well as cross-platform consistency
 
-bool FFileTests::RunTest( const FString& Parameters )
+IMPLEMENT_SIMPLE_AUTOMATION_TEST(FFileTruncateTest, "System.Core.Misc.FileTruncate", EAutomationTestFlags::EditorContext | EAutomationTestFlags::CriticalPriority | EAutomationTestFlags::EngineFilter)
+bool FFileTruncateTest::RunTest(const FString& Parameters)
 {
-	// Disabled for now, pending changes to make all platforms consistent.
-	return true;
+	const FString TempFilename = FPaths::CreateTempFilename(*FPaths::ProjectIntermediateDir());
+	IPlatformFile& PlatformFile = FPlatformFileManager::Get().GetPlatformFile();
+	ON_SCOPE_EXIT
+	{
+		// Delete temp file
+		PlatformFile.DeleteFile(*TempFilename);
+	};
+
+	// Open a test file
+	if (TUniquePtr<IFileHandle> TestFile = TUniquePtr<IFileHandle>(PlatformFile.OpenWrite(*TempFilename, /*bAppend*/false, /*bAllowRead*/true)))
+	{
+		// Append 4 int32 values of incrementing value to this file
+		int32 Val = 1;
+		TestFile->Write((const uint8*)&Val, sizeof(Val));
+
+		++Val;
+		TestFile->Write((const uint8*)&Val, sizeof(Val));
+
+		// Tell here, so we can move back and truncate after writing
+		const int64 ExpectedTruncatePos = TestFile->Tell();
+		++Val;
+		TestFile->Write((const uint8*)&Val, sizeof(Val));
+
+		// Tell here, so we can attempt to read here after truncation
+		const int64 TestReadPos = TestFile->Tell();
+		++Val;
+		TestFile->Write((const uint8*)&Val, sizeof(Val));
+
+		// Validate that the Tell position is at the end of the file, and that the size is reported correctly
+		{
+			const int64 ActualEOFPos = TestFile->Tell();
+			const int64 ExpectedEOFPos = (sizeof(int32) * 4);
+			if (ActualEOFPos != ExpectedEOFPos)
+			{
+				AddError(FString::Printf(TEXT("File was not the expected size (got %d, expected %d): %s"), ActualEOFPos, ExpectedEOFPos, *TempFilename));
+				return false;
+			}
+
+			const int64 ActualFileSize = TestFile->Size();
+			if (ActualFileSize != ExpectedEOFPos)
+			{
+				AddError(FString::Printf(TEXT("File was not the expected size (got %d, expected %d): %s"), ActualFileSize, ExpectedEOFPos, *TempFilename));
+				return false;
+			}
+		}
 
-	/*
-		const FString TempFilename = FPaths::EngineSavedDir() / FGuid::NewGuid().ToString();
-		TArray<uint8> TestData;
-		TArray<uint8> ReadData;
-		uint8 One = 1;
-		FArchive* TestFile = nullptr;
+		// Truncate the file at our test pos
+		if (!TestFile->Truncate(ExpectedTruncatePos))
+		{
+			AddError(FString::Printf(TEXT("File truncation request failed: %s"), *TempFilename));
+			return false;
+		}
+
+		// Validate that the size is reported correctly
+		{
+			const int64 ActualFileSize = TestFile->Size();
+			if (ActualFileSize != ExpectedTruncatePos)
+			{
+				AddError(FString::Printf(TEXT("File was not the expected size after truncation (got %d, expected %d): %s"), ActualFileSize, ExpectedTruncatePos, *TempFilename));
+				return false;
+			}
+		}
 
-		// We will test the platform abstraction class, IFileManager
-		IFileManager& FileManager = IFileManager::Get();
+		// Validate that we can't read past the truncation point
+		{
+			int32 Dummy = 0;
+			if (TestFile->Seek(TestReadPos) && TestFile->Read((uint8*)&Dummy, sizeof(Dummy)))
+			{
+				AddError(FString::Printf(TEXT("File read seek outside the truncated range: %s"), *TempFilename));
+				return false;
+			}
+		}
+	}
+	else
+	{
+		AddError(FString::Printf(TEXT("File failed to open: %s"), *TempFilename));
+		return false;
+	}
 
-		// Ensure always starting off without test file existing
-		FileManager.Delete(*TempFilename);
-		check(FileManager.FileExists(*TempFilename) == false);
+	return true;
+}
 
-		// Check a new file can be created
+IMPLEMENT_SIMPLE_AUTOMATION_TEST(FFileAppendTest, "System.Core.Misc.FileAppend", EAutomationTestFlags::EditorContext | EAutomationTestFlags::CriticalPriority | EAutomationTestFlags::EngineFilter)
+bool FFileAppendTest::RunTest( const FString& Parameters )
+{
+	const FString TempFilename = FPaths::CreateTempFilename(*FPaths::ProjectIntermediateDir());
+	IPlatformFile& PlatformFile = FPlatformFileManager::Get().GetPlatformFile();
+	ON_SCOPE_EXIT
+	{
+		// Delete temp file
+		PlatformFile.DeleteFile(*TempFilename);
+	};
+
+	// Scratch data for testing
+	uint8 One = 1;
+	TArray<uint8> TestData;
+
+	// Check a new file can be created
+	if (TUniquePtr<IFileHandle> TestFile = TUniquePtr<IFileHandle>(PlatformFile.OpenWrite(*TempFilename, /*bAppend*/false, /*bAllowRead*/true)))
+	{
 		TestData.AddZeroed(64);
-		TestFile = FileManager.CreateFileWriter(*TempFilename, 0);
-		check(TestFile != nullptr);
-		TestFile->Serialize(TestData.GetData(), TestData.Num());
-		delete TestFile;
-
-		// Confirm same data
-		check(FFileHelper::LoadFileToArray(ReadData, *TempFilename));
- 		if(ReadData != TestData)
+
+		TestFile->Write(TestData.GetData(), TestData.Num());
+	}
+	else
+	{
+		AddError(FString::Printf(TEXT("File failed to open when new: %s"), *TempFilename));
+		return false;
+	}
+
+	// Confirm same data
+	{
+		TArray<uint8> ReadData;
+		if (!FFileHelper::LoadFileToArray(ReadData, *TempFilename))
+		{
+			AddError(FString::Printf(TEXT("File failed to load after writing: %s"), *TempFilename));
+			return false;
+		}
+
+		if (ReadData != TestData)
 		{
+			AddError(FString::Printf(TEXT("File data was incorrect after writing: %s"), *TempFilename));
 			return false;
 		}
+	}
+
+	// Using append flag should open the file, and writing data immediately should append to the end.
+	// We should also be capable of seeking writing.
+	if (TUniquePtr<IFileHandle> TestFile = TUniquePtr<IFileHandle>(PlatformFile.OpenWrite(*TempFilename, /*bAppend*/true, /*bAllowRead*/true)))
+	{
+		// Validate the file actually opened in append mode correctly
+		{
+			const int64 ActualEOFPos = TestFile->Tell();
+			const int64 ExpectedEOFPos = TestFile->Size();
+			if (ActualEOFPos != ExpectedEOFPos)
+			{
+				AddError(FString::Printf(TEXT("File did not seek to the end when opening (got %d, expected %d): %s"), ActualEOFPos, ExpectedEOFPos, *TempFilename));
+				return false;
+			}
+		}
 
-		// Using append flag should open the file, and writing data immediatly should append to the end.
-		// We should also be capable of seeking writing.
 		TestData.Add(One);
 		TestData[10] = One;
-		TestFile = FileManager.CreateFileWriter(*TempFilename, EFileWrite::FILEWRITE_Append);
-		check(TestFile != nullptr);
-		TestFile->Serialize(&One, 1);
+
+		TestFile->Write(&One, 1);
 		TestFile->Seek(10);
-		TestFile->Serialize(&One, 1);
-		delete TestFile;
+		TestFile->Write(&One, 1);
+	}
+	else
+	{
+		AddError(FString::Printf(TEXT("File failed to open when appending: %s"), *TempFilename));
+		return false;
+	}
+
+	// Confirm same data
+	{
+		TArray<uint8> ReadData;
+		if (!FFileHelper::LoadFileToArray(ReadData, *TempFilename))
+		{
+			AddError(FString::Printf(TEXT("File failed to load after appending: %s"), *TempFilename));
+			return false;
+		}
 
-		// Confirm same data
-		check(FFileHelper::LoadFileToArray(ReadData, *TempFilename));
-		if(ReadData != TestData)
+		if (ReadData != TestData)
 		{
+			AddError(FString::Printf(TEXT("File data was incorrect after appending: %s"), *TempFilename));
 			return false;
 		}
+	}
 
-		// No flags should clobber existing file
-		TestData.Empty();
+	// No append should clobber existing file
+	if (TUniquePtr<IFileHandle> TestFile = TUniquePtr<IFileHandle>(PlatformFile.OpenWrite(*TempFilename, /*bAppend*/false, /*bAllowRead*/true)))
+	{
+		TestData.Reset();
 		TestData.Add(One);
-		TestFile = FileManager.CreateFileWriter(*TempFilename, 0);
-		check(TestFile != nullptr);
-		TestFile->Serialize(&One, 1);
-		delete TestFile;
-
-		// Confirm same data
-		check(FFileHelper::LoadFileToArray(ReadData, *TempFilename));
-		if(ReadData != TestData)
+
+		TestFile->Write(&One, 1);
+	}
+	else
+	{
+		AddError(FString::Printf(TEXT("File failed to open when clobbering: %s"), *TempFilename));
+		return false;
+	}
+
+	// Confirm same data
+	{
+		TArray<uint8> ReadData;
+		if (!FFileHelper::LoadFileToArray(ReadData, *TempFilename))
 		{
+			AddError(FString::Printf(TEXT("File failed to load after clobbering: %s"), *TempFilename));
 			return false;
 		}
 
-		// Delete temp file
-		FileManager.Delete(*TempFilename);
+		if (ReadData != TestData)
+		{
+			AddError(FString::Printf(TEXT("File data was incorrect after clobbering: %s"), *TempFilename));
+			return false;
+		}
+	}
 
-		return true;
-	*/
+	return true;
 }
 
 #endif //WITH_DEV_AUTOMATION_TESTS
