- Fixed value FGenericPlatformMemoryConstants::TotalPhysical not set on Windows (was initialized to zero, but never set)

diff --git a/Engine/Source/Runtime/Core/Private/Windows/WindowsPlatformMemory.cpp b/Engine/Source/Runtime/Core/Private/Windows/WindowsPlatformMemory.cpp
--- a/Engine/Source/Runtime/Core/Private/Windows/WindowsPlatformMemory.cpp
+++ b/Engine/Source/Runtime/Core/Private/Windows/WindowsPlatformMemory.cpp
@@ -224,6 +224,7 @@ FPlatformMemoryStats FWindowsPlatformMemory::GetStats()
 	FPlatformMemory::Memzero( &ProcessMemoryCounters, sizeof( ProcessMemoryCounters ) );
 	::GetProcessMemoryInfo( ::GetCurrentProcess(), &ProcessMemoryCounters, sizeof(ProcessMemoryCounters) );
 
+	MemoryStats.TotalPhysical = MemoryStatusEx.ullTotalPhys;
 	MemoryStats.AvailablePhysical = MemoryStatusEx.ullAvailPhys;
 	MemoryStats.AvailableVirtual = MemoryStatusEx.ullAvailVirtual;
 	
