Added Algo::Compare, to compare contiguous ranges using operator== on each element

diff --git a/Engine/Source/Runtime/Core/Public/Algo/Compare.h b/Engine/Source/Runtime/Core/Public/Algo/Compare.h
--- a/Engine/Source/Runtime/Core/Public/Algo/Compare.h
+++ b/Engine/Source/Runtime/Core/Public/Algo/Compare.h
@@ -7,6 +7,39 @@
 
 namespace Algo
 {
+	/**
+	* Compares two contiguous containers using operator== to compare pairs of elements.
+	*
+	* @param  InputA     Container of elements that are used as the first argument to operator==.
+	* @param  InputB     Container of elements that are used as the second argument to operator==.
+	*
+	* @return Whether the containers are the same size and operator== returned true for every pair of elements.
+	*/
+	template <typename InAT, typename InBT>
+	constexpr bool Compare(InAT&& InputA, InBT&& InputB)
+	{
+		const SIZE_T SizeA = GetNum(InputA);
+		const SIZE_T SizeB = GetNum(InputB);
+
+		if (SizeA != SizeB)
+		{
+			return false;
+		}
+
+		auto* A = GetData(InputA);
+		auto* B = GetData(InputB);
+
+		for (SIZE_T Count = SizeA; Count; --Count)
+		{
+			if (!(*A++ == *B++))
+			{
+				return false;
+			}
+		}
+
+		return true;
+	}
+
 	/**
 	* Compares two contiguous containers using a predicate to compare pairs of elements.
 	*
