Fixed memory stomping occuring when a long command line is used.
Modified FGenericCrashContext to perform safe string copies

diff --git a/Engine/Source/Runtime/Core/Private/GenericPlatform/GenericPlatformCrashContext.cpp b/Engine/Source/Runtime/Core/Private/GenericPlatform/GenericPlatformCrashContext.cpp
--- a/Engine/Source/Runtime/Core/Private/GenericPlatform/GenericPlatformCrashContext.cpp
+++ b/Engine/Source/Runtime/Core/Private/GenericPlatform/GenericPlatformCrashContext.cpp
@@ -81,6 +81,12 @@ namespace NCached
 	static TArray<FString> EnabledPluginsList;
 	static TMap<FString, FString> EngineData;
 	static TMap<FString, FString> GameData;
+
+	template <size_t CharCount>
+	void Set(TCHAR(&Dest)[CharCount], const TCHAR* pSrc)
+	{
+		FCString::Strncpy(Dest, pSrc, CharCount);
+	}
 }
 
 void FGenericCrashContext::Initialize()
@@ -91,38 +97,38 @@ void FGenericCrashContext::Initialize()
 	NCached::Session.bIsSourceDistribution = FEngineBuildSettings::IsSourceDistribution();
 	NCached::Session.ProcessId = FPlatformProcess::GetCurrentProcessId();
 
-	FCString::Strcpy(NCached::Session.GameName, *FString::Printf( TEXT("UE4-%s"), FApp::GetProjectName() ));
-	FCString::Strcpy(NCached::Session.GameSessionID, TEXT("")); // Updated by callback
-	FCString::Strcpy(NCached::Session.GameStateName, TEXT("")); // Updated by callback
-	FCString::Strcpy(NCached::Session.UserActivityHint, TEXT("")); // Updated by callback
-	FCString::Strcpy(NCached::Session.BuildConfigurationName, LexToString(FApp::GetBuildConfiguration()));
-	FCString::Strcpy(NCached::Session.ExecutableName, FPlatformProcess::ExecutableName());
-	FCString::Strcpy(NCached::Session.BaseDir, FPlatformProcess::BaseDir());
-	FCString::Strcpy(NCached::Session.RootDir, FPlatformMisc::RootDir());
-	FCString::Strcpy(NCached::Session.EpicAccountId, *FPlatformMisc::GetEpicAccountId());
-	FCString::Strcpy(NCached::Session.LoginIdStr, *FPlatformMisc::GetLoginId());
+	NCached::Set(NCached::Session.GameName, *FString::Printf(TEXT("UE4-%s"), FApp::GetProjectName()));
+	NCached::Set(NCached::Session.GameSessionID, TEXT("")); // Updated by callback
+	NCached::Set(NCached::Session.GameStateName, TEXT("")); // Updated by callback
+	NCached::Set(NCached::Session.UserActivityHint, TEXT("")); // Updated by callback
+	NCached::Set(NCached::Session.BuildConfigurationName, LexToString(FApp::GetBuildConfiguration()));
+	NCached::Set(NCached::Session.ExecutableName, FPlatformProcess::ExecutableName());
+	NCached::Set(NCached::Session.BaseDir, FPlatformProcess::BaseDir());
+	NCached::Set(NCached::Session.RootDir, FPlatformMisc::RootDir());
+	NCached::Set(NCached::Session.EpicAccountId, *FPlatformMisc::GetEpicAccountId());
+	NCached::Set(NCached::Session.LoginIdStr, *FPlatformMisc::GetLoginId());
 
 	FString OsVersion, OsSubVersion;
 	FPlatformMisc::GetOSVersions(OsVersion, OsSubVersion);
-	FCString::Strcpy(NCached::Session.OsVersion, *OsVersion);
-	FCString::Strcpy(NCached::Session.OsSubVersion, *OsSubVersion);
+	NCached::Set(NCached::Session.OsVersion, *OsVersion);
+	NCached::Set(NCached::Session.OsSubVersion, *OsSubVersion);
 
 	NCached::Session.NumberOfCores = FPlatformMisc::NumberOfCores();
 	NCached::Session.NumberOfCoresIncludingHyperthreads = FPlatformMisc::NumberOfCoresIncludingHyperthreads();
 
-	FCString::Strcpy(NCached::Session.CPUVendor, *FPlatformMisc::GetCPUVendor());
-	FCString::Strcpy(NCached::Session.CPUBrand, *FPlatformMisc::GetCPUBrand());
-	FCString::Strcpy(NCached::Session.PrimaryGPUBrand, *FPlatformMisc::GetPrimaryGPUBrand());
-	FCString::Strcpy(NCached::Session.UserName, FPlatformProcess::UserName());
-	FCString::Strcpy(NCached::Session.DefaultLocale, *FPlatformMisc::GetDefaultLocale());
+	NCached::Set(NCached::Session.CPUVendor, *FPlatformMisc::GetCPUVendor());
+	NCached::Set(NCached::Session.CPUBrand, *FPlatformMisc::GetCPUBrand());
+	NCached::Set(NCached::Session.PrimaryGPUBrand, *FPlatformMisc::GetPrimaryGPUBrand());
+	NCached::Set(NCached::Session.UserName, FPlatformProcess::UserName());
+	NCached::Set(NCached::Session.DefaultLocale, *FPlatformMisc::GetDefaultLocale());
 
 	// Information that cannot be gathered if command line is not initialized (e.g. crash during static init)
 	if (FCommandLine::IsInitialized())
 	{
 		NCached::Session.bIsUE4Release = FApp::IsEngineInstalled();
-		FCString::Strcpy(NCached::Session.CommandLine, (FCommandLine::IsInitialized() ? FCommandLine::GetOriginalForLogging() : TEXT("")));
-		FCString::Strcpy(NCached::Session.EngineMode, FGenericPlatformMisc::GetEngineMode());
-		FCString::Strcpy(NCached::Session.EngineModeEx, FGenericCrashContext::EngineModeExUnknown); // Updated from callback
+		NCached::Set(NCached::Session.CommandLine, (FCommandLine::IsInitialized() ? FCommandLine::GetOriginalForLogging() : TEXT("")));
+		NCached::Set(NCached::Session.EngineMode, FGenericPlatformMisc::GetEngineMode());
+		NCached::Set(NCached::Session.EngineModeEx, FGenericCrashContext::EngineModeExUnknown); // Updated from callback
 
 		// Use -epicapp value from the commandline to start. This will also be set by the game
 		FParse::Value(FCommandLine::Get(), TEXT("EPICAPP="), NCached::Session.DeploymentName, CR_MAX_GENERIC_FIELD_CHARS, true);
@@ -146,7 +152,7 @@ void FGenericCrashContext::Initialize()
 	// Create a unique base guid for bug report ids
 	const FGuid Guid = FGuid::NewGuid();
 	const FString IniPlatformName(FPlatformProperties::IniPlatformName());
-	FCString::Strcpy(NCached::Session.CrashGUIDRoot, *FString::Printf(TEXT("%s%s-%s"), CrashGUIDRootPrefix, *IniPlatformName, *Guid.ToString(EGuidFormats::Digits)));
+	NCached::Set(NCached::Session.CrashGUIDRoot, *FString::Printf(TEXT("%s%s-%s"), CrashGUIDRootPrefix, *IniPlatformName, *Guid.ToString(EGuidFormats::Digits)));
 
 	if (GIsRunning)
 	{
@@ -181,17 +187,17 @@ void FGenericCrashContext::Initialize()
 
 	FCoreDelegates::UserActivityStringChanged.AddLambda([](const FString& InUserActivity)
 	{
-		FCString::Strcpy(NCached::Session.UserActivityHint, *InUserActivity);
+		NCached::Set(NCached::Session.UserActivityHint, *InUserActivity);
 	});
 
 	FCoreDelegates::GameSessionIDChanged.AddLambda([](const FString& InGameSessionID)
 	{
-		FCString::Strcpy(NCached::Session.GameSessionID, *InGameSessionID);
+		NCached::Set(NCached::Session.GameSessionID, *InGameSessionID);
 	});
 
 	FCoreDelegates::GameStateClassChanged.AddLambda([](const FString& InGameStateName)
 	{
-		FCString::Strcpy(NCached::Session.GameStateName, *InGameStateName);
+		NCached::Set(NCached::Session.GameStateName, *InGameStateName);
 	});
 
 PRAGMA_DISABLE_DEPRECATION_WARNINGS
@@ -199,11 +205,11 @@ PRAGMA_DISABLE_DEPRECATION_WARNINGS
 	{
 		if (InParams.bSetCrashReportClientMessageText)
 		{
-			FCString::Strcpy(NCached::Session.CrashReportClientRichText, *InParams.CrashReportClientMessageText);
+			NCached::Set(NCached::Session.CrashReportClientRichText, *InParams.CrashReportClientMessageText);
 		}
 		if (InParams.bSetGameNameSuffix)
 		{
-			FCString::Strcpy(NCached::Session.GameName, *(FString(TEXT("UE4-")) + FApp::GetProjectName() + InParams.GameNameSuffix));
+			NCached::Set(NCached::Session.GameName, *(FString(TEXT("UE4-")) + FApp::GetProjectName() + InParams.GameNameSuffix));
 		}
 		if (InParams.SendUnattendedBugReports.IsSet())
 		{
@@ -220,17 +226,17 @@ PRAGMA_ENABLE_DEPRECATION_WARNINGS
 	FCoreDelegates::OnPostEngineInit.AddLambda([] 
 	{
 		// IsEditor global have been properly initialized here.
-		FCString::Strcpy(NCached::Session.EngineMode, FGenericPlatformMisc::GetEngineMode());
+		NCached::Set(NCached::Session.EngineMode, FGenericPlatformMisc::GetEngineMode());
 	});
 
 	FCoreDelegates::IsVanillaProductChanged.AddLambda([](bool bIsVanilla)
 	{
-		FCString::Strcpy(NCached::Session.EngineModeEx, bIsVanilla ? FGenericCrashContext::EngineModeExVanilla : FGenericCrashContext::EngineModeExDirty);
+		NCached::Set(NCached::Session.EngineModeEx, bIsVanilla ? FGenericCrashContext::EngineModeExVanilla : FGenericCrashContext::EngineModeExDirty);
 	});
 
 	FCoreDelegates::ConfigReadyForUse.AddStatic(FGenericCrashContext::InitializeFromConfig);
 
-	FCString::Strcpy(NCached::UserSettings.LogFilePath, *FPlatformOutputDevices::GetAbsoluteLogFilename());
+	NCached::Set(NCached::UserSettings.LogFilePath, *FPlatformOutputDevices::GetAbsoluteLogFilename());
 
 	NCached::UserSettings.bNoDialog = FApp::IsUnattended() || IsRunningDedicatedServer();
 
@@ -396,7 +402,7 @@ void FGenericCrashContext::UpdateLocalizedStrings()
 	FText CrashReportClientRichText;
 	if (GConfig->GetText(TEXT("CrashContextProperties"), TEXT("CrashReportClientRichText"), CrashReportClientRichText, GEngineIni))
 	{
-		FCString::Strcpy(NCached::Session.CrashReportClientRichText, *CrashReportClientRichText.ToString());
+		NCached::Set(NCached::Session.CrashReportClientRichText, *CrashReportClientRichText.ToString());
 	}
 #endif
 }
@@ -659,7 +665,7 @@ void FGenericCrashContext::SetNumMinidumpFramesToIgnore(int InNumMinidumpFramesT
 
 void FGenericCrashContext::SetDeploymentName(const FString& EpicApp)
 {
-	FCString::Strcpy(NCached::Session.DeploymentName, *EpicApp);
+	NCached::Set(NCached::Session.DeploymentName, *EpicApp);
 }
 
 void FGenericCrashContext::SetCrashTrigger(ECrashTrigger Type)
@@ -882,7 +888,7 @@ const TCHAR* FGenericCrashContext::GetCrashConfigFilePath()
 	if (FCString::Strlen(NCached::Session.CrashConfigFilePath) == 0)
 	{
 		FString CrashConfigFilePath = FPaths::Combine(GetCrashConfigFolder(), NCached::Session.CrashGUIDRoot, FGenericCrashContext::CrashConfigFileNameW);
-		FCString::Strcpy(NCached::Session.CrashConfigFilePath, *CrashConfigFilePath);
+		NCached::Set(NCached::Session.CrashConfigFilePath, *CrashConfigFilePath);
 	}
 	return NCached::Session.CrashConfigFilePath;
 }
