- Made internals of LoadFileToStringArray[WithPredicate] to allow for 64-bit sized files (which still returns a 32-bit TArray of 32-bit FString lines)

diff --git a/Engine/Source/Editor/UnrealEd/Private/Commandlets/ShaderPipelineCacheToolsCommandlet.cpp b/Engine/Source/Editor/UnrealEd/Private/Commandlets/ShaderPipelineCacheToolsCommandlet.cpp
--- a/Engine/Source/Editor/UnrealEd/Private/Commandlets/ShaderPipelineCacheToolsCommandlet.cpp
+++ b/Engine/Source/Editor/UnrealEd/Private/Commandlets/ShaderPipelineCacheToolsCommandlet.cpp
@@ -64,12 +64,10 @@ static void LoadStableSCL(TArray<FStableShaderKeyAndValue>& StableArray, const F
 {
 	UE_LOG(LogShaderPipelineCacheTools, Display, TEXT("Loading %.*s..."), FileName.Len(), FileName.GetData());
 
-	FString SourceFileContents;
-	TArray<FStringView> SourceFileLines;
-	if (FFileHelper::LoadFileToString(SourceFileContents, *FString(FileName)))
+	TArray<FString> SourceFileLines;
+	if (FFileHelper::LoadFileToStringArrayWithPredicate(SourceFileLines, *FString(FileName), [](const FString & Line) { return !Line.IsEmpty(); }))
 	{
-		UE_LOG(LogShaderPipelineCacheTools, Display, TEXT("Loaded %.*s, length = %d..."), FileName.Len(), FileName.GetData(), SourceFileContents.Len());
-		UE::String::ParseLines(SourceFileContents, [&SourceFileLines](FStringView Line) { if (!Line.IsEmpty()) { SourceFileLines.Add(Line); } });
+		UE_LOG(LogShaderPipelineCacheTools, Display, TEXT("Loaded %.*s, %d lines..."), FileName.Len(), FileName.GetData(), SourceFileLines.Num());
 	}
 
 	if (SourceFileLines.Num() < 1)
diff --git a/Engine/Source/Runtime/Core/Private/Misc/FileHelper.cpp b/Engine/Source/Runtime/Core/Private/Misc/FileHelper.cpp
--- a/Engine/Source/Runtime/Core/Private/Misc/FileHelper.cpp
+++ b/Engine/Source/Runtime/Core/Private/Misc/FileHelper.cpp
@@ -221,25 +221,34 @@ bool FFileHelper::LoadFileToString(FString& Result, IPlatformFile* PlatformFile,
 	return LoadFileToString(Result, *Reader.Get(), VerifyFlags);
 }
 
-bool FFileHelper::LoadFileToStringArray( TArray<FString>& Result, const TCHAR* Filename, EHashOptions VerifyFlags )
+bool FFileHelper::LoadFileToStringArray( TArray<FString>& Result, const TCHAR* Filename )
 {
 	Result.Empty();
 
-	FString Buffer;
-	if(!LoadFileToString(Buffer, Filename, VerifyFlags))
+	TArray64<uint8> RawBuffer;
+	if(!LoadFileToArray(RawBuffer, Filename))
 	{
 		return false;
 	}
 
-	for(const TCHAR* Pos = *Buffer; *Pos != 0; )
+	for(const uint8* Pos = (uint8*)RawBuffer.GetData(); *Pos != 0; )
 	{
-		const TCHAR* LineStart = Pos;
+		const uint8* LineStart = Pos;
 		while(*Pos != 0 && *Pos != '\r' && *Pos != '\n')
 		{
 			Pos++;
 		}
 
-		Result.Emplace((int32)(Pos - LineStart), LineStart);
+		if (Pos - LineStart > MAX_int32)
+		{
+			UE_LOG(LogStreaming, Warning, TEXT("Single line too long found in LoadFileToStringArrayWithPredicate, File: %s"), Filename);
+			return false;
+		}
+
+		FString Line;
+		BufferToString(Line, LineStart, UE_PTRDIFF_TO_INT32(Pos - LineStart));
+
+		Result.Add(MoveTemp(Line));
 
 		if(*Pos == '\r')
 		{
@@ -254,25 +263,39 @@ bool FFileHelper::LoadFileToStringArray( TArray<FString>& Result, const TCHAR* F
 	return true;
 }
 
-bool FFileHelper::LoadFileToStringArrayWithPredicate(TArray<FString>& Result, const TCHAR* Filename, TFunctionRef<bool(const FString&)> Predicate, EHashOptions VerifyFlags)
+// DEPRECATED
+bool FFileHelper::LoadFileToStringArray(TArray<FString>& Result, const TCHAR* Filename, EHashOptions VerifyFlags)
+{
+	return LoadFileToStringArray(Result, Filename);
+}
+
+bool FFileHelper::LoadFileToStringArrayWithPredicate(TArray<FString>& Result, const TCHAR* Filename, TFunctionRef<bool(const FString&)> Predicate)
 {
 	Result.Empty();
 
-	FString Buffer;
-	if (!LoadFileToString(Buffer, Filename, VerifyFlags))
+	TArray64<uint8> RawBuffer;
+	if (!LoadFileToArray(RawBuffer, Filename))
 	{
 		return false;
 	}
 
-	for (const TCHAR* Pos = *Buffer; *Pos != 0; )
+	for (const uint8* Pos = (uint8*)RawBuffer.GetData(); *Pos != 0; )
 	{
-		const TCHAR* LineStart = Pos;
+		const uint8* LineStart = Pos;
 		while (*Pos != 0 && *Pos != '\r' && *Pos != '\n')
 		{
 			Pos++;
 		}
 
-		FString Line(UE_PTRDIFF_TO_INT32(Pos - LineStart), LineStart);
+		if (Pos - LineStart > MAX_int32)
+		{
+			UE_LOG(LogStreaming, Warning, TEXT("Single line too long found in LoadFileToStringArrayWithPredicate, File: %s"), Filename);
+			return false;
+		}
+
+		FString Line;
+		BufferToString(Line, LineStart, UE_PTRDIFF_TO_INT32(Pos - LineStart));
+		
 		if (Invoke(Predicate, Line))
 		{
 			Result.Add(MoveTemp(Line));
@@ -291,6 +314,12 @@ bool FFileHelper::LoadFileToStringArrayWithPredicate(TArray<FString>& Result, co
 	return true;
 }
 
+// DEPRECATED
+bool FFileHelper::LoadFileToStringArrayWithPredicate(TArray<FString>& Result, const TCHAR* Filename, TFunctionRef<bool(const FString&)> Predicate, EHashOptions VerifyFlags)
+{
+	return LoadFileToStringArrayWithPredicate(Result, Filename, Predicate);
+}
+
 /**
  * Save a binary array to a file.
  */
diff --git a/Engine/Source/Runtime/Core/Public/Misc/FileHelper.h b/Engine/Source/Runtime/Core/Public/Misc/FileHelper.h
--- a/Engine/Source/Runtime/Core/Public/Misc/FileHelper.h
+++ b/Engine/Source/Runtime/Core/Public/Misc/FileHelper.h
@@ -98,9 +98,11 @@ struct CORE_API FFileHelper
 	 *
 	 * @param Result       String representation of the loaded file
 	 * @param Filename     Name of the file to load
-	 * @param VerifyFlags  Flags controlling the hash verification behavior ( see EHashOptions )
 	 */
-	static bool LoadFileToStringArray( TArray<FString>& Result, const TCHAR* Filename, EHashOptions VerifyFlags = EHashOptions::None );
+	static bool LoadFileToStringArray( TArray<FString>& Result, const TCHAR* Filename );
+
+	UE_DEPRECATED(4.26, "LoadFileToStringArray no longer supports VerifyFlags. You can use UE::String::ParseLines to split up a string loaded with LoadFileToString")
+	static bool LoadFileToStringArray(TArray<FString>& Result, const TCHAR* Filename, EHashOptions VerifyFlags);
 
 	/**
 	 * Load a text file to an array of strings, filtered by a user-defined predicate. Supports all combination of ANSI/Unicode files and platforms.
@@ -108,9 +110,11 @@ struct CORE_API FFileHelper
 	 * @param Result       String representation of the loaded file
 	 * @param Filename     Name of the file to load
 	 * @param Predicate    Condition for whether or not to add the line to the array
-	 * @param VerifyFlags  Flags controlling the hash verification behavior ( see EHashOptions )
 	 */
-	static bool LoadFileToStringArrayWithPredicate(TArray<FString>& Result, const TCHAR* Filename, TFunctionRef<bool(const FString&)> Predicate, EHashOptions VerifyFlags = EHashOptions::None);
+	static bool LoadFileToStringArrayWithPredicate(TArray<FString>& Result, const TCHAR* Filename, TFunctionRef<bool(const FString&)> Predicate);
+
+	UE_DEPRECATED(4.26, "LoadFileToStringArrayWithPredicate no longer supports VerifyFlags. You can use UE::String::ParseLines to split up a string loaded with LoadFileToString")
+	static bool LoadFileToStringArrayWithPredicate(TArray<FString>& Result, const TCHAR* Filename, TFunctionRef<bool(const FString&)> Predicate, EHashOptions VerifyFlags);
 
 	/**
 	 * Save a binary array to a file.
