#UE4 Made the time threshold before warning about FlushCurrentThreadCache configurable

diff --git a/Engine/Source/Runtime/Core/Private/HAL/MallocBinned2.cpp b/Engine/Source/Runtime/Core/Private/HAL/MallocBinned2.cpp
--- a/Engine/Source/Runtime/Core/Private/HAL/MallocBinned2.cpp
+++ b/Engine/Source/Runtime/Core/Private/HAL/MallocBinned2.cpp
@@ -51,6 +51,13 @@ static FAutoConsoleVariableRef GMallocBinned2AllocExtraCVar(
 	TEXT("When we do acquire the lock, how many blocks cached in TLS caches. In no case will we grab more than a page.")
 	);
 
+double GMallocBinned2FlushThreadCacheMaxWaitTime = 0.02f;
+static FAutoConsoleVariableRef GMallocBinned2FlushThreadCacheMaxWaitTimeCVar(
+	TEXT("MallocBinned2.FlushThreadCacheMaxWaitTime"),
+	GMallocBinned2FlushThreadCacheMaxWaitTime,
+	TEXT("The threshold of time before warning about FlushCurrentThreadCache taking too long (seconds).")
+);
+
 #endif
 
 #if BINNED2_ALLOCATOR_STATS
@@ -1071,11 +1078,16 @@ void FMallocBinned2::FlushCurrentThreadCache()
 	}
 
 	// These logs must happen outside the above mutex to avoid deadlocks
-	if (WaitForMutexTime > 0.02f)
+#if BINNED2_ALLOW_RUNTIME_TWEAKING
+	double MaxWaitTimeThreshold = GMallocBinned2FlushThreadCacheMaxWaitTime;
+#else
+	double MaxWaitTimeThreshold = 0.02f;
+#endif
+	if (WaitForMutexTime > MaxWaitTimeThreshold)
 	{
 		UE_LOG(LogMemory, Warning, TEXT("FMallocBinned2 took %6.2fms to wait for mutex for trim."), WaitForMutexTime * 1000.0f);
 	}
-	if (WaitForMutexAndTrimTime > 0.02f)
+	if (WaitForMutexAndTrimTime > MaxWaitTimeThreshold)
 	{
 		UE_LOG(LogMemory, Warning, TEXT("FMallocBinned2 took %6.2fms to wait for mutex AND trim."), WaitForMutexAndTrimTime * 1000.0f);
 	}
diff --git a/Engine/Source/Runtime/Core/Private/HAL/MallocBinned3.cpp b/Engine/Source/Runtime/Core/Private/HAL/MallocBinned3.cpp
--- a/Engine/Source/Runtime/Core/Private/HAL/MallocBinned3.cpp
+++ b/Engine/Source/Runtime/Core/Private/HAL/MallocBinned3.cpp
@@ -66,6 +66,13 @@ static FAutoConsoleVariableRef GMallocBinned3AllocExtraCVar(
 	TEXT("When we do acquire the lock, how many blocks cached in TLS caches. In no case will we grab more than a page.")
 	);
 
+double GMallocBinned3FlushThreadCacheMaxWaitTime = 0.02f;
+static FAutoConsoleVariableRef GMallocBinned3FlushThreadCacheMaxWaitTimeCVar(
+	TEXT("MallocBinned3.FlushThreadCacheMaxWaitTime"),
+	GMallocBinned3FlushThreadCacheMaxWaitTime,
+	TEXT("The threshold of time before warning about FlushCurrentThreadCache taking too long (seconds).")
+);
+
 #endif
 
 #if BINNED3_ALLOCATOR_STATS
@@ -1308,11 +1315,16 @@ void FMallocBinned3::FlushCurrentThreadCache()
 	}
 
 	// These logs must happen outside the above mutex to avoid deadlocks
-	if (WaitForMutexTime > 0.02f)
+#if BINNED3_ALLOW_RUNTIME_TWEAKING
+	double MaxWaitTimeThreshold = GMallocBinned3FlushThreadCacheMaxWaitTime;
+#else
+	double MaxWaitTimeThreshold = 0.02f;
+#endif
+	if (WaitForMutexTime > MaxWaitTimeThreshold)
 	{
 		UE_LOG(LogMemory, Warning, TEXT("FMallocBinned3 took %6.2fms to wait for mutex for trim."), WaitForMutexTime * 1000.0f);
 	}
-	if (WaitForMutexAndTrimTime > 0.02f)
+	if (WaitForMutexAndTrimTime > MaxWaitTimeThreshold)
 	{
 		UE_LOG(LogMemory, Warning, TEXT("FMallocBinned3 took %6.2fms to wait for mutex AND trim."), WaitForMutexAndTrimTime * 1000.0f);
 	}
