Various engine init optimizations and new string conversion paths

* Added some core string conversion paths
* Added FString::Append() and += overloads to avoid pointless temporaries when appending c strings and various other string types
* Added UTF32 <-> TCHAR string conversion macros
* Optimized QuotedString parsing
* Optimized property parsing
* PackageName optimizations
* SoftObjectPath optimizations

diff --git a/Engine/Plugins/Runtime/GameplayAbilities/Source/GameplayAbilities/Private/GameplayCueManager.cpp b/Engine/Plugins/Runtime/GameplayAbilities/Source/GameplayAbilities/Private/GameplayCueManager.cpp
--- a/Engine/Plugins/Runtime/GameplayAbilities/Source/GameplayAbilities/Private/GameplayCueManager.cpp
+++ b/Engine/Plugins/Runtime/GameplayAbilities/Source/GameplayAbilities/Private/GameplayCueManager.cpp
@@ -892,7 +892,7 @@ void UGameplayCueManager::BuildCuesToAddToGlobalSet(const TArray<FAssetData>& As
 			{
 				// Add a new NotifyData entry to our flat list for this one
 				FSoftObjectPath StringRef;
-				StringRef.SetPath(FPackageName::ExportTextPathToObjectPath(*GeneratedClassTag));
+				StringRef.SetPath(FPackageName::ExportTextPathToObjectPath(GeneratedClassTag));
 
 				OutCuesToAdd.Add(FGameplayCueReferencePair(GameplayCueTag, StringRef));
 
diff --git a/Engine/Source/Editor/AIGraph/Private/AIGraphTypes.cpp b/Engine/Source/Editor/AIGraph/Private/AIGraphTypes.cpp
--- a/Engine/Source/Editor/AIGraph/Private/AIGraphTypes.cpp
+++ b/Engine/Source/Editor/AIGraph/Private/AIGraphTypes.cpp
@@ -305,7 +305,7 @@ void FGraphNodeClassHelper::OnAssetRemoved(const struct FAssetData& AssetData)
 	if (AssetData.GetTagValue(FBlueprintTags::GeneratedClassPath, AssetClassName))
 	{
 		ConstructorHelpers::StripObjectClass(AssetClassName);
-		AssetClassName = FPackageName::ObjectPathToObjectName(*AssetClassName);
+		AssetClassName = FPackageName::ObjectPathToObjectName(AssetClassName);
 
 		TSharedPtr<FGraphNodeClassNode> Node = FindBaseClassNode(RootNode, AssetClassName);
 		if (Node.IsValid() && Node->ParentNode.IsValid())
diff --git a/Engine/Source/Editor/UnrealEd/Private/Factories/ActorFactory.cpp b/Engine/Source/Editor/UnrealEd/Private/Factories/ActorFactory.cpp
--- a/Engine/Source/Editor/UnrealEd/Private/Factories/ActorFactory.cpp
+++ b/Engine/Source/Editor/UnrealEd/Private/Factories/ActorFactory.cpp
@@ -1414,8 +1414,8 @@ bool UActorFactoryBlueprint::CanCreateActorFrom( const FAssetData& AssetData, FT
 		FAssetRegistryModule& AssetRegistryModule = FModuleManager::LoadModuleChecked<FAssetRegistryModule>("AssetRegistry");
 		IAssetRegistry& AssetRegistry = AssetRegistryModule.Get();
 
-		const FString ObjectPath = FPackageName::ExportTextPathToObjectPath(*ParentClassPath);
-		const FName ParentClassPathFName = FName( *FPackageName::ObjectPathToObjectName(ObjectPath) );
+		const FStringView ObjectPath = FPackageName::ExportTextPathToObjectPath(FStringView(ParentClassPath));
+		const FName ParentClassPathFName = FName( FPackageName::ObjectPathToObjectName(ObjectPath) );
 		TArray<FName> AncestorClassNames;
 		AssetRegistry.GetAncestorClassNames(ParentClassPathFName, AncestorClassNames);
 
diff --git a/Engine/Source/Runtime/Core/Private/Containers/String.cpp b/Engine/Source/Runtime/Core/Private/Containers/String.cpp
--- a/Engine/Source/Runtime/Core/Private/Containers/String.cpp
+++ b/Engine/Source/Runtime/Core/Private/Containers/String.cpp
@@ -152,11 +152,49 @@ FString& FString::operator=(const FStringView& Other)
 	return *this;
 }
 
-FString& FString::operator+=(const FStringView& Str)
+
+template<typename CharType>
+void AppendCharacters(TArray<TCHAR>& Out, const CharType* Str, int32 Count)
+{
+	check(Count >= 0);
+
+	if (!Count)
+	{
+		return;
+	}
+
+	checkSlow(Str);
+
+	const int32 OldNum = Out.Num();
+
+	// Reserve enough space - including an extra gap for a null terminator if we don't already have a string allocated
+	Out.AddUninitialized(Count + (OldNum ? 0 : 1));
+
+	TCHAR* Dest = Out.GetData() + OldNum - (OldNum ? 1 : 0);
+
+	// Copy characters to end of string, overwriting null terminator if we already have one
+	FPlatformString::Convert(Dest, Count, Str, Count);
+
+	// (Re-)establish the null terminator
+	Dest[Count] = '\0';
+}
+
+void FString::AppendChars(const ANSICHAR* Str, int32 Count)
 {
 	CheckInvariants();
-	AppendChars(Str.GetData(), Str.Len());
-	return *this;
+	AppendCharacters(Data, Str, Count);
+}
+
+void FString::AppendChars(const WIDECHAR* Str, int32 Count)
+{
+	CheckInvariants();
+	AppendCharacters(Data, Str, Count);
+}
+
+void FString::AppendChars(const UCS2CHAR* Str, int32 Count)
+{
+	CheckInvariants();
+	AppendCharacters(Data, Str, Count);
 }
 
 FString::operator FStringView() const
diff --git a/Engine/Source/Runtime/Core/Private/Misc/Parse.cpp b/Engine/Source/Runtime/Core/Private/Misc/Parse.cpp
--- a/Engine/Source/Runtime/Core/Private/Misc/Parse.cpp
+++ b/Engine/Source/Runtime/Core/Private/Misc/Parse.cpp
@@ -9,8 +9,10 @@
 #include "Misc/Paths.h"
 #include "Containers/Set.h"
 #include "Internationalization/Text.h"
+#include "Misc/AsciiSet.h"
 #include "Misc/Guid.h"
 #include "Misc/OutputDeviceNull.h"
+#include "Misc/StringBuilder.h"
 #include "HAL/IConsoleManager.h"
 #include "Containers/LazyPrintf.h"
 
@@ -342,10 +344,8 @@ bool FParse::Value( const TCHAR* Stream, const TCHAR* Match, FString& Value, boo
 	return false;
 }
 
-// 
-// Parse a quoted string.
-//
-bool FParse::QuotedString( const TCHAR* Buffer, FString& Value, int32* OutNumCharsRead )
+template<class T>
+bool ParseQuotedString( const TCHAR* Buffer, T& Value, int32* OutNumCharsRead)
 {
 	if (OutNumCharsRead)
 	{
@@ -362,7 +362,8 @@ bool FParse::QuotedString( const TCHAR* Buffer, FString& Value, int32* OutNumCha
 
 	auto ShouldParse = [](const TCHAR Ch)
 	{
-		return Ch != 0 && Ch != TCHAR('"') && Ch != TCHAR('\n') && Ch != TCHAR('\r');
+		constexpr FAsciiSet ParsableCharacters = ~(FAsciiSet("\"\n\r") + '\0');
+		return ParsableCharacters.Test(Ch);
 	};
 
 	while (ShouldParse(*Buffer))
@@ -403,7 +404,7 @@ bool FParse::QuotedString( const TCHAR* Buffer, FString& Value, int32* OutNumCha
 		}
 		else if (FChar::IsOctDigit(*Buffer)) // octal sequence (\012)
 		{
-			FString OctSequence;
+			TStringBuilder<16> OctSequence;
 			while (ShouldParse(*Buffer) && FChar::IsOctDigit(*Buffer) && OctSequence.Len() < 3) // Octal sequences can only be up-to 3 digits long
 			{
 				OctSequence += *Buffer++;
@@ -415,7 +416,7 @@ bool FParse::QuotedString( const TCHAR* Buffer, FString& Value, int32* OutNumCha
 		{
 			++Buffer;
 
-			FString HexSequence;
+			TStringBuilder<16> HexSequence;
 			while (ShouldParse(*Buffer) && FChar::IsHexDigit(*Buffer))
 			{
 				HexSequence += *Buffer++;
@@ -427,37 +428,27 @@ bool FParse::QuotedString( const TCHAR* Buffer, FString& Value, int32* OutNumCha
 		{
 			++Buffer;
 
-			FString UnicodeSequence;
+			TStringBuilder<16> UnicodeSequence;
 			while (ShouldParse(*Buffer) && FChar::IsHexDigit(*Buffer) && UnicodeSequence.Len() < 4) // UTF-16 sequences can only be up-to 4 digits long
 			{
 				UnicodeSequence += *Buffer++;
 			}
 
-			const uint32 UnicodeCodepoint = (uint32)FCString::Strtoi(*UnicodeSequence, nullptr, 16);
-
-			FString UnicodeString;
-			if (FUnicodeChar::CodepointToString(UnicodeCodepoint, UnicodeString))
-			{
-				Value += MoveTemp(UnicodeString);
-			}
+			const UTF16CHAR Utf16Char = static_cast<UTF16CHAR>(FCString::Strtoi(*UnicodeSequence, nullptr, 16));
+			Value += FStringView(FUTF16ToTCHAR(&Utf16Char, /* Len */ 1).Get());
 		}
 		else if (*Buffer == TCHAR('U') && FChar::IsHexDigit(*(Buffer + 1))) // UTF-32 sequence (\U12345678)
 		{
 			++Buffer;
 
-			FString UnicodeSequence;
+			TStringBuilder<16> UnicodeSequence;
 			while (ShouldParse(*Buffer) && FChar::IsHexDigit(*Buffer) && UnicodeSequence.Len() < 8) // UTF-32 sequences can only be up-to 8 digits long
 			{
 				UnicodeSequence += *Buffer++;
 			}
 
-			const uint32 UnicodeCodepoint = (uint32)FCString::Strtoi(*UnicodeSequence, nullptr, 16);
-
-			FString UnicodeString;
-			if (FUnicodeChar::CodepointToString(UnicodeCodepoint, UnicodeString))
-			{
-				Value += MoveTemp(UnicodeString);
-			}
+			const UTF32CHAR Utf32Str = static_cast<UTF32CHAR>(FCString::Strtoi(*UnicodeSequence, nullptr, 16));
+			Value += FStringView(FUTF32ToTCHAR(&Utf32Str, /* Len */ 1).Get());
 		}
 		else // unhandled escape sequence
 		{
@@ -480,6 +471,16 @@ bool FParse::QuotedString( const TCHAR* Buffer, FString& Value, int32* OutNumCha
 	return true;
 }
 
+bool FParse::QuotedString( const TCHAR* Buffer, FString& Value, int32* OutNumCharsRead )
+{
+	return ParseQuotedString(Buffer, Value, OutNumCharsRead);
+}
+
+bool FParse::QuotedString( const TCHAR* Buffer, FStringBuilderBase& Value, int32* OutNumCharsRead )
+{
+	return ParseQuotedString(Buffer, Value, OutNumCharsRead);
+}
+
 // 
 // Parse an Text token
 // This is expected to in the form NSLOCTEXT("Namespace","Key","SourceString") or LOCTEXT("Key","SourceString")
diff --git a/Engine/Source/Runtime/Core/Private/Misc/StringBuilder.cpp b/Engine/Source/Runtime/Core/Private/Misc/StringBuilder.cpp
--- a/Engine/Source/Runtime/Core/Private/Misc/StringBuilder.cpp
+++ b/Engine/Source/Runtime/Core/Private/Misc/StringBuilder.cpp
@@ -98,6 +98,11 @@ FStringBuilderBase&	operator<<(FStringBuilderBase& Builder, const FString& Str)
 	return Builder.Append(*Str, Str.Len());
 }
 
+FStringBuilderBase&	operator+=(FStringBuilderBase& Builder, const FString& Str)
+{
+	return Builder.Append(*Str, Str.Len());
+}
+
 // Instantiate templates once
 
 template class TStringBuilderImpl<ANSICHAR>;
diff --git a/Engine/Source/Runtime/Core/Private/UObject/UnrealNames.cpp b/Engine/Source/Runtime/Core/Private/UObject/UnrealNames.cpp
--- a/Engine/Source/Runtime/Core/Private/UObject/UnrealNames.cpp
+++ b/Engine/Source/Runtime/Core/Private/UObject/UnrealNames.cpp
@@ -2199,6 +2199,12 @@ bool FName::IsValidXName(const FString& InName, const FString& InInvalidChars, F
 	return true;
 }
 
+FStringBuilderBase& operator<<(FStringBuilderBase& Builder, FNameEntryId Id)
+{
+	FName::GetEntry(Id)->AppendNameToString(Builder);
+	return Builder;
+}
+
 template <typename CharType, int N>
 void CheckLazyName(const CharType(&Literal)[N])
 {
diff --git a/Engine/Source/Runtime/Core/Public/Containers/StringConv.h b/Engine/Source/Runtime/Core/Public/Containers/StringConv.h
--- a/Engine/Source/Runtime/Core/Public/Containers/StringConv.h
+++ b/Engine/Source/Runtime/Core/Public/Containers/StringConv.h
@@ -754,7 +754,7 @@ public:
 		else
 		{
 			// Normal codepoint
-			*(OutputIterator++) = Codepoint;
+			*(OutputIterator++) = (ToType)Codepoint;
 		}
 
 		return UE_PTRDIFF_TO_INT32(OutputIterator - OutputIteratorStartPosition);
@@ -1159,16 +1159,31 @@ typedef TStringConversion<FUTF8ToTCHAR_Convert> FUTF8ToTCHAR;
 
 // special handling for platforms still using a 32-bit TCHAR
 #if PLATFORM_TCHAR_IS_4_BYTES
+
 typedef TStringConversion<TUTF32ToUTF16_Convert<TCHAR, UTF16CHAR>> FTCHARToUTF16;
 typedef TStringConversion<TUTF16ToUTF32_Convert<UTF16CHAR, TCHAR>> FUTF16ToTCHAR;
 #define TCHAR_TO_UTF16(str) (UTF16CHAR*)FTCHARToUTF16((const TCHAR*)str).Get()
 #define UTF16_TO_TCHAR(str) (TCHAR*)FUTF16ToTCHAR((const UTF16CHAR*)str).Get()
+
+static_assert(sizeof(TCHAR) == sizeof(UTF32CHAR), "TCHAR and UTF32CHAR are expected to be the same size for inline conversion! PLATFORM_TCHAR_IS_4_BYTES is not configured correctly for this platform.");
+typedef TStringPointer<TCHAR, UTF32CHAR> FTCHARToUTF32;
+typedef TStringPointer<UTF32CHAR, TCHAR> FUTF32ToTCHAR;
+#define TCHAR_TO_UTF32(str) (UTF32CHAR*)(str)
+#define UTF32_TO_TCHAR(str) (TCHAR*)(str)
+
 #else
+
 static_assert(sizeof(TCHAR) == sizeof(UTF16CHAR), "TCHAR and UTF16CHAR are expected to be the same size for inline conversion! PLATFORM_TCHAR_IS_4_BYTES is not configured correctly for this platform.");
 typedef TStringPointer<TCHAR, UTF16CHAR> FTCHARToUTF16;
 typedef TStringPointer<UTF16CHAR, TCHAR> FUTF16ToTCHAR;
 #define TCHAR_TO_UTF16(str) (UTF16CHAR*)(str)
 #define UTF16_TO_TCHAR(str) (TCHAR*)(str)
+
+typedef TStringConversion<TUTF16ToUTF32_Convert<TCHAR, UTF32CHAR>> FTCHARToUTF32;
+typedef TStringConversion<TUTF32ToUTF16_Convert<UTF32CHAR, TCHAR>> FUTF32ToTCHAR;
+#define TCHAR_TO_UTF32(str) (UTF32CHAR*)FTCHARToUTF32((const TCHAR*)str).Get()
+#define UTF32_TO_TCHAR(str) (TCHAR*)FUTF32ToTCHAR((const UTF32CHAR*)str).Get()
+
 #endif
 
 // special handling for going from char16_t to wchar_t for third party libraries that need wchar_t
diff --git a/Engine/Source/Runtime/Core/Public/Containers/StringView.h b/Engine/Source/Runtime/Core/Public/Containers/StringView.h
--- a/Engine/Source/Runtime/Core/Public/Containers/StringView.h
+++ b/Engine/Source/Runtime/Core/Public/Containers/StringView.h
@@ -138,6 +138,8 @@ public:
 	inline void		RemovePrefix(SizeType CharCount)	{ DataPtr += CharCount; Size -= CharCount; }
 	/** Modifies the view to remove the given number of characters from the end. */
 	inline void		RemoveSuffix(SizeType CharCount)	{ Size -= CharCount; }
+	/** Resets to an empty view */
+	inline void		Reset()								{ DataPtr = nullptr; Size = 0; }
 
 	// Operations
 
diff --git a/Engine/Source/Runtime/Core/Public/Containers/UnrealString.h b/Engine/Source/Runtime/Core/Public/Containers/UnrealString.h
--- a/Engine/Source/Runtime/Core/Public/Containers/UnrealString.h
+++ b/Engine/Source/Runtime/Core/Public/Containers/UnrealString.h
@@ -11,6 +11,7 @@
 #include "Misc/AssertionMacros.h"
 #include "HAL/UnrealMemory.h"
 #include "Templates/IsArithmetic.h"
+#include "Templates/IsArray.h"
 #include "Templates/UnrealTypeTraits.h"
 #include "Templates/UnrealTemplate.h"
 #include "Math/NumericLimits.h"
@@ -359,67 +360,49 @@ public:
 	}
 #endif
 
-	/**
-	 * Appends an array of characters to the string.  
+	/** 
+	 * Appends a character range without null-terminators in it
 	 *
-	 * @param Array A pointer to the start of an array of characters to append.  This array need not be null-terminated, and null characters are not treated specially.
-	 * @param Count The number of characters to copy from Array.
+	 * @param Str can be null if Count is 0. Can be unterminated, Str[Count] isn't read.
 	 */
-	FORCEINLINE void AppendChars(const TCHAR* Array, int32 Count)
-	{
-		check(Count >= 0);
-
-		if (!Count)
-			return;
-
-		checkSlow(Array);
-
-		int32 Index = Data.Num();
-
-		// Reserve enough space - including an extra gap for a null terminator if we don't already have a string allocated
-		Data.AddUninitialized(Count + (Index ? 0 : 1));
-
-		TCHAR* EndPtr = Data.GetData() + Index - (Index ? 1 : 0);
-
-		// Copy characters to end of string, overwriting null terminator if we already have one
-		CopyAssignItems(EndPtr, Array, Count);
+	void AppendChars(const ANSICHAR* Str, int32 Count);
+	void AppendChars(const WIDECHAR* Str, int32 Count);
+	void AppendChars(const UCS2CHAR* Str, int32 Count);
 
-		// (Re-)establish the null terminator
-		*(EndPtr + Count) = 0;
+	/** Append a string and return a reference to this */
+	template<class CharType>
+	FORCEINLINE FString& Append(const CharType* Str, int32 Count)
+	{
+		AppendChars(Str, Count);
+		return *this;
 	}
-
+	
 	/**
-	 * Concatenate this with the given string
+	 * Append a valid null-terminated string and return a reference to this
 	 *
-	 * @param Str string view to be concatenated onto the end of this
-	 * @return reference to this
-	 */
-	FString& operator+=(const FStringView& Str);
-
-	/**
-	 * Concatenate this with given string
-	 * 
-	 * @param Str array of TCHAR to be concatenated onto the end of this
-	 * @return reference to this
+	 * CharType is not const to use this overload for mutable char arrays and call
+	 * Strlen() instead of getting the static length N from GetNum((&T)[N]).
+	 * Oddly MSVC ranks a const T* overload over T&& for T[N] while clang picks T&&.
 	 */
-	FORCEINLINE FString& operator+=( const TCHAR* Str )
+	template<class CharType>
+	FORCEINLINE FString& Append(/* no const! */ CharType* Str)
 	{
 		checkSlow(Str);
-		CheckInvariants();
-
-		AppendChars(Str, FCString::Strlen(Str));
-
+		AppendChars(Str, TCString<typename TRemoveConst<CharType>::Type>::Strlen(Str));
+		return *this;
+	}
+	
+	/** Append a string and return a reference to this */
+	template <typename StrType, typename = decltype(GetData(DeclVal<StrType>()))>
+	FORCEINLINE FString& Append(StrType&& Str)
+	{
+		static_assert(!TIsArray<StrType>::Value, "Char arrays shouldn't use this overload");
+		AppendChars(GetData(Str), GetNum(Str));
 		return *this;
 	}
 
-	/**
-	 * Concatenate this with given char
-	 * 
-	 * @param inChar other Char to be concatenated onto the end of this string
-	 * @return reference to this
-	 */
-	template <typename CharType>
-	FORCEINLINE typename TEnableIf<TIsCharType<CharType>::Value, FString&>::Type operator+=(CharType InChar)
+	/** Append a single character and return a reference to this */
+	FORCEINLINE FString& AppendChar(TCHAR InChar)
 	{
 		CheckInvariants();
 
@@ -440,49 +423,14 @@ public:
 		return *this;
 	}
 
-	/**
-	 * Concatenate this with given char
-	 * 
-	 * @param InChar other Char to be concatenated onto the end of this string
-	 * @return reference to this
-	 */
-	FORCEINLINE FString& AppendChar(const TCHAR InChar)
-	{
-		*this += InChar;
-		return *this;
-	}
-
-	FORCEINLINE FString& Append(const FString& Text)
-	{
-		*this += Text;
-		return *this;
-	}
-
-	FString& Append(const TCHAR* Text, int32 Count)
-	{
-		CheckInvariants();
-
-		if (Count != 0)
-		{
-			// position to insert the character.  
-			// At the end of the string if we have existing characters, otherwise at the 0 position
-			int32 InsertIndex = (Data.Num() > 0) ? Data.Num() - 1 : 0;
-
-			// number of characters to add.  If we don't have any existing characters, 
-			// we'll need to append the terminating zero as well.
-			int32 FinalCount = (Data.Num() > 0) ? Count : Count + 1;
-
-			Data.AddUninitialized(FinalCount);
+	/** Append a string and return a reference to this */
+	template <typename StrType>
+	FORCEINLINE FString& operator+=(StrType&& Str)	{ return Append(Forward<StrType>(Str)); }
 
-			for (int32 Index = 0; Index < Count; Index++)
-			{
-				Data[InsertIndex + Index] = Text[Index];
-			}
-
-			Data[Data.Num() - 1] = 0;
-		}
-		return *this;
-	}
+	/** Append a single character and return a reference to this */
+	FORCEINLINE FString& operator+=(ANSICHAR Char)	{ return AppendChar(Char); }
+	FORCEINLINE FString& operator+=(WIDECHAR Char)	{ return AppendChar(Char); }
+	FORCEINLINE FString& operator+=(UCS2CHAR Char)	{ return AppendChar(Char); }
 
 	/**
 	 * Removes characters within the string.
@@ -565,22 +513,6 @@ public:
 	 */
 	void PathAppend(const TCHAR* Str, int32 StrLength);
 
-	/**
-	 * Concatenate this with given string
-	 * 
-	 * @param Str other string to be concatenated onto the end of this
-	 * @return reference to this
-	 */
-	FORCEINLINE FString& operator+=( const FString& Str )
-	{
-		CheckInvariants();
-		Str.CheckInvariants();
-
-		AppendChars(Str.Data.GetData(), Str.Len());
-
-		return *this;
-	}
-
 	/**
 	 * Concatenates an FString with a TCHAR.
 	 * 
diff --git a/Engine/Source/Runtime/Core/Public/Misc/Parse.h b/Engine/Source/Runtime/Core/Public/Misc/Parse.h
--- a/Engine/Source/Runtime/Core/Public/Misc/Parse.h
+++ b/Engine/Source/Runtime/Core/Public/Misc/Parse.h
@@ -5,6 +5,8 @@
 #include "CoreTypes.h"
 #include "Containers/UnrealString.h"
 
+class FStringBuilderBase;
+
 /*-----------------------------------------------------------------------------
 	Parsing functions.
 -----------------------------------------------------------------------------*/
@@ -70,6 +72,7 @@ struct CORE_API FParse
 	static bool Text( const TCHAR* Stream, FText& Value, const TCHAR* Namespace = nullptr );
 	/** Parse a quoted string token. */
 	static bool QuotedString( const TCHAR* Stream, FString& Value, int32* OutNumCharsRead = nullptr );
+	static bool QuotedString( const TCHAR* Stream, FStringBuilderBase& Value, int32* OutNumCharsRead = nullptr );
 
 	//
 	// Parse a hex digit.
diff --git a/Engine/Source/Runtime/Core/Public/Misc/StringBuilder.h b/Engine/Source/Runtime/Core/Public/Misc/StringBuilder.h
--- a/Engine/Source/Runtime/Core/Public/Misc/StringBuilder.h
+++ b/Engine/Source/Runtime/Core/Public/Misc/StringBuilder.h
@@ -468,7 +468,8 @@ struct TIsContiguousContainer<TStringBuilder<N>> { enum { Value = true }; };
 
 // Append operator implementations
 
-inline FStringBuilderBase&					operator<<(FStringBuilderBase& Builder, const TCHAR Char)						{ return Builder.Append(Char); }
+inline FStringBuilderBase&					operator<<(FStringBuilderBase& Builder, ANSICHAR Char)							{ return Builder.Append(Char); }
+inline FStringBuilderBase&					operator<<(FStringBuilderBase& Builder, WIDECHAR Char)							{ return Builder.Append(Char); }
 
 #if USE_STRING_LITERAL_PATH
 
@@ -504,6 +505,18 @@ inline FStringBuilderBase&					operator<<(FStringBuilderBase& Builder, uint8 Val
 inline FStringBuilderBase&					operator<<(FStringBuilderBase& Builder, int16 Value)							{ return Builder << int32(Value); }
 inline FStringBuilderBase&					operator<<(FStringBuilderBase& Builder, uint16 Value)							{ return Builder << uint32(Value); }
 
+/** Prefer using << instead of +=. operator+= is only intended for mechanical FString -> FStringView replacement */
+
+inline FStringBuilderBase&					operator+=(FStringBuilderBase& Builder, FWideStringView Str)					{ return Builder.Append(Str); }
+inline FStringBuilderBase&					operator+=(FStringBuilderBase& Builder, FAnsiStringView Str)					{ return Builder.AppendAnsi(Str); }
+inline FStringBuilderBase&					operator+=(FStringBuilderBase& Builder, const WIDECHAR* Str)					{ return Builder.Append(Str); }
+inline FStringBuilderBase&					operator+=(FStringBuilderBase& Builder, const ANSICHAR* Str)					{ return Builder.AppendAnsi(Str); }
+
+CORE_API FStringBuilderBase&				operator+=(FStringBuilderBase& Builder, const FString& Str);
+inline FStringBuilderBase&					operator+=(FStringBuilderBase& Builder, ANSICHAR Char)							{ return Builder.Append(Char); }
+inline FStringBuilderBase&					operator+=(FStringBuilderBase& Builder, WIDECHAR Char)							{ return Builder.Append(Char); }
+
+
 // Append operator implementations for FAnsiStringBuilderBase
 
 inline FAnsiStringBuilderBase&				operator<<(FAnsiStringBuilderBase& Builder, const ANSICHAR Char)					{ return Builder.Append(Char); }
diff --git a/Engine/Source/Runtime/Core/Public/UObject/NameTypes.h b/Engine/Source/Runtime/Core/Public/UObject/NameTypes.h
--- a/Engine/Source/Runtime/Core/Public/UObject/NameTypes.h
+++ b/Engine/Source/Runtime/Core/Public/UObject/NameTypes.h
@@ -973,6 +973,8 @@ inline FStringBuilderBase& operator<<(FStringBuilderBase& Builder, const FName&
 	return Builder;
 }
 
+CORE_API FStringBuilderBase& operator<<(FStringBuilderBase& Builder, FNameEntryId Id);
+
 /**
  * Equality operator with CharType* on left hand side and FName on right hand side
  * 
diff --git a/Engine/Source/Runtime/CoreUObject/Private/Misc/PackageName.cpp b/Engine/Source/Runtime/CoreUObject/Private/Misc/PackageName.cpp
--- a/Engine/Source/Runtime/CoreUObject/Private/Misc/PackageName.cpp
+++ b/Engine/Source/Runtime/CoreUObject/Private/Misc/PackageName.cpp
@@ -50,16 +50,31 @@ namespace PackageNameConstants
 	const int32 MinPackageNameLength = 4;
 }
 
-bool FPackageName::IsShortPackageName(const FString& PossiblyLongName)
+bool FPackageName::IsShortPackageName(FStringView PossiblyLongName)
 {
 	// Long names usually have / as first character so check from the front
-	int32 SlashIndex = INDEX_NONE;
-	return !PossiblyLongName.FindChar('/', SlashIndex);
+	for (TCHAR Char : PossiblyLongName)
+	{
+		if (Char == '/')
+		{
+			return false;
+		}
+	}
+
+	return true;
+}
+
+bool FPackageName::IsShortPackageName(const FString& PossiblyLongName)
+{
+	return IsShortPackageName(FStringView(PossiblyLongName));
 }
 
-bool FPackageName::IsShortPackageName(const FName PossiblyLongName)
+bool FPackageName::IsShortPackageName(FName PossiblyLongName)
 {
-	return IsShortPackageName(PossiblyLongName.ToString());
+	// Only get "plain" part of the name. The number suffix, e.g. "_123", can't contain slashes.
+	TCHAR Buffer[NAME_SIZE];
+	uint32 Len = PossiblyLongName.GetPlainNameString(Buffer);
+	return IsShortPackageName(FStringView(Buffer, static_cast<int32>(Len)));
 }
 
 FString FPackageName::GetShortName(const FString& LongName)
@@ -339,7 +354,7 @@ private:
 	}
 };
 
-void FPackageName::InternalFilenameToLongPackageName(const FStringView& InFilename, FStringBuilderBase& OutPackageName)
+void FPackageName::InternalFilenameToLongPackageName(FStringView InFilename, FStringBuilderBase& OutPackageName)
 {
 	const FLongPackagePathsSingleton& Paths = FLongPackagePathsSingleton::Get();
 	FString Filename(InFilename);
@@ -873,7 +888,7 @@ bool FPackageName::FindPackageFileWithoutExtension(const FString& InPackageFilen
 	return false;
 }
 
-bool FPackageName::FixPackageNameCase(FString& LongPackageName, const FStringView& Extension)
+bool FPackageName::FixPackageNameCase(FString& LongPackageName, FStringView Extension)
 {
 	// Find the matching long package root
 	const FLongPackagePathsSingleton& Paths = FLongPackagePathsSingleton::Get();
@@ -1439,6 +1454,45 @@ bool FPackageName::ParseExportTextPath(const FString& InExportTextPath, FString*
 	return false;
 }
 
+bool FPackageName::ParseExportTextPath(FStringView InExportTextPath, FStringView* OutClassName, FStringView* OutObjectPath)
+{
+	int32 Index;
+	if (InExportTextPath.FindChar('\'', Index))
+	{
+		if (OutClassName)
+		{
+			*OutClassName = InExportTextPath.Left(Index);
+		}
+
+		if (OutObjectPath)
+		{
+			*OutObjectPath = InExportTextPath.Right(Index + 1);
+			OutObjectPath->RemoveSuffix(InExportTextPath.EndsWith('\''));
+		}
+
+		return true;
+	}
+	
+	return false;
+}
+
+bool FPackageName::ParseExportTextPath(const TCHAR* InExportTextPath, FStringView* OutClassName, FStringView* OutObjectPath)
+{
+	return ParseExportTextPath(FStringView(InExportTextPath), OutClassName, OutObjectPath);
+}
+
+FStringView FPackageName::ExportTextPathToObjectPath(FStringView InExportTextPath)
+{
+	FStringView ObjectPath;
+	if (FPackageName::ParseExportTextPath(InExportTextPath, nullptr, &ObjectPath))
+	{
+		return ObjectPath;
+	}
+	
+	// Could not parse the export text path. Could already be an object path, just return it back.
+	return InExportTextPath;
+}
+
 FString FPackageName::ExportTextPathToObjectPath(const FString& InExportTextPath)
 {
 	FString ObjectPath;
@@ -1451,6 +1505,11 @@ FString FPackageName::ExportTextPathToObjectPath(const FString& InExportTextPath
 	return InExportTextPath;
 }
 
+const TCHAR* FPackageName::ExportTextPathToObjectPath(const TCHAR* InExportTextPath)
+{
+	return ExportTextPathToObjectPath(FStringView(InExportTextPath)).GetData();
+}
+
 FString FPackageName::ObjectPathToPackageName(const FString& InObjectPath)
 {
 	// Check for package delimiter
@@ -1464,7 +1523,8 @@ FString FPackageName::ObjectPathToPackageName(const FString& InObjectPath)
 	return InObjectPath;
 }
 
-FString FPackageName::ObjectPathToObjectName(const FString& InObjectPath)
+template<class T>
+T ObjectPathToObjectNameImpl(const T& InObjectPath)
 {
 	// Check for a subobject
 	int32 SubObjectDelimiterIdx;
@@ -1484,6 +1544,16 @@ FString FPackageName::ObjectPathToObjectName(const FString& InObjectPath)
 	return InObjectPath;
 }
 
+FString FPackageName::ObjectPathToObjectName(const FString& InObjectPath)
+{
+	return ObjectPathToObjectNameImpl(InObjectPath);
+}
+
+FStringView FPackageName::ObjectPathToObjectName(FStringView InObjectPath)
+{
+	return ObjectPathToObjectNameImpl(InObjectPath);
+}
+
 bool FPackageName::IsScriptPackage(const FString& InPackageName)
 {
 	return InPackageName.StartsWith(FLongPackagePathsSingleton::Get().ScriptRootPath);
diff --git a/Engine/Source/Runtime/CoreUObject/Private/UObject/Property.cpp b/Engine/Source/Runtime/CoreUObject/Private/UObject/Property.cpp
--- a/Engine/Source/Runtime/CoreUObject/Private/UObject/Property.cpp
+++ b/Engine/Source/Runtime/CoreUObject/Private/UObject/Property.cpp
@@ -369,10 +369,41 @@ IMPLEMENT_STRUCT(FallbackStruct);
 	Helpers.
 -----------------------------------------------------------------------------*/
 
+constexpr FAsciiSet AlphaNumericChars = "abcdefghijklmnopqrstuvwxyzABCDEFGHIJKLMNOPQRSTUVWXYZ0123456789";
+
+FORCEINLINE constexpr bool IsValidTokenStart(TCHAR FirstChar, bool bDottedNames)
+{
+	return AlphaNumericChars.Test(FirstChar) || (bDottedNames && FirstChar == '/') || FirstChar > 255;
+}
+
+FORCEINLINE constexpr FStringView ParsePropertyToken(const TCHAR* Str, bool DottedNames)
+{
+	constexpr FAsciiSet RegularTokenChars = AlphaNumericChars  + '_' + '-' + '"';
+	constexpr FAsciiSet DottedTokenChars = RegularTokenChars + '.' + '/' + SUBOBJECT_DELIMITER_CHAR;
+	FAsciiSet CurrentTokenChars = DottedNames ? DottedTokenChars : RegularTokenChars;
+
+	const TCHAR* It = Str;
+	while (true)
+	{
+		// Include allowed ASCII characters
+		It = FAsciiSet::Skip(It, CurrentTokenChars);
+
+		if (*It <= 255)
+		{
+			return FStringView(Str, It - Str);
+		}
+		else
+		{
+			// Include wide characters
+			for (++It; *It > 255; ++It);
+		}
+	} 
+}
+
 //
 // Parse a token.
 //
-const TCHAR* FPropertyHelpers::ReadToken( const TCHAR* Buffer, FString& String, bool DottedNames )
+const TCHAR* FPropertyHelpers::ReadToken( const TCHAR* Buffer, FString& String, bool bDottedNames)
 {
 	if( *Buffer == TCHAR('"') )
 	{
@@ -384,30 +415,51 @@ const TCHAR* FPropertyHelpers::ReadToken( const TCHAR* Buffer, FString& String,
 		}
 		Buffer += NumCharsRead;
 	}
-	else if( FChar::IsAlnum( *Buffer ) || (DottedNames && (*Buffer==TCHAR('/'))) || (*Buffer > 255) )
+	else if (IsValidTokenStart(*Buffer, bDottedNames))
 	{
-		const TCHAR* StartIdentifier = Buffer;
-		// Get identifier.
-		while ((FChar::IsAlnum(*Buffer) || (*Buffer > 255) || *Buffer == TCHAR('_') || *Buffer == TCHAR('-') || *Buffer == TCHAR('+') || (DottedNames && (*Buffer == TCHAR('.') || *Buffer == TCHAR('/') || *Buffer == SUBOBJECT_DELIMITER_CHAR))))
-		{
-			Buffer++;
-		}
+		FStringView Token = ParsePropertyToken(Buffer, bDottedNames);
+		String += Token;
+		Buffer += Token.Len();
+	}
+	else
+	{
+		// Get just one.
+		String += *Buffer;
+	}
+	return Buffer;
+}
 
-		if (StartIdentifier != Buffer)
+const TCHAR* FPropertyHelpers::ReadToken( const TCHAR* Buffer, FStringBuilderBase& Out, bool bDottedNames)
+{
+	if( *Buffer == TCHAR('"') )
+	{
+		int32 NumCharsRead = 0;
+		if (!FParse::QuotedString(Buffer, Out, &NumCharsRead))
 		{
-			// Efficiently copy/allocate
-			String.AppendChars(StartIdentifier, Buffer - StartIdentifier);
+			UE_LOG(LogProperty, Warning, TEXT("ReadToken: Bad quoted string: %s"), Buffer );
+			return nullptr;
 		}
+		Buffer += NumCharsRead;
+
+		// TODO special handling of null-terminator here?
+	}
+	else if (IsValidTokenStart(*Buffer, bDottedNames))
+	{
+		FStringView Token = ParsePropertyToken(Buffer, bDottedNames);
+		Out << Token;
+		Buffer += Token.Len();
 	}
 	else
 	{
 		// Get just one.
-		String += *Buffer;
+		if (*Buffer)
+		{
+			Out << *Buffer;
+		}
 	}
 	return Buffer;
 }
 
-
 /*-----------------------------------------------------------------------------
 	FProperty implementation.
 -----------------------------------------------------------------------------*/
diff --git a/Engine/Source/Runtime/CoreUObject/Private/UObject/PropertyBaseObject.cpp b/Engine/Source/Runtime/CoreUObject/Private/UObject/PropertyBaseObject.cpp
--- a/Engine/Source/Runtime/CoreUObject/Private/UObject/PropertyBaseObject.cpp
+++ b/Engine/Source/Runtime/CoreUObject/Private/UObject/PropertyBaseObject.cpp
@@ -264,14 +264,14 @@ bool FObjectPropertyBase::ParseObjectPropertyValue(const FProperty* Property, UO
 
  	const TCHAR* InBuffer = Buffer;
 
-	FString Temp;
-	Buffer = FPropertyHelpers::ReadToken(Buffer, Temp, true);
+	TStringBuilder<256> Temp;
+	Buffer = FPropertyHelpers::ReadToken(Buffer, /* out */ Temp, true);
 	if ( Buffer == nullptr )
 	{
 		return false;
 	}
 
-	if ( Temp == TEXT("None") )
+	if ( Temp == TEXT("None"_SV) )
 	{
 		out_ResolvedValue = nullptr;
 	}
@@ -285,8 +285,8 @@ bool FObjectPropertyBase::ParseObjectPropertyValue(const FProperty* Property, UO
 
 		if( *Buffer == TCHAR('\'') )
 		{
-			FString ObjectText;
-			Buffer = FPropertyHelpers::ReadToken( ++Buffer, ObjectText, true );
+			Temp.Reset();
+			Buffer = FPropertyHelpers::ReadToken( ++Buffer, /* out */ Temp, true);
 			if( Buffer == nullptr )
 			{
 				return false;
@@ -299,12 +299,12 @@ bool FObjectPropertyBase::ParseObjectPropertyValue(const FProperty* Property, UO
 
 			// ignore the object class, it isn't fully qualified, and searching ANY_PACKAGE might get the wrong one!
 			// Try the find the object.
-			out_ResolvedValue = FObjectPropertyBase::FindImportedObject(Property, OwnerObject, ObjectClass, RequiredMetaClass, *ObjectText, PortFlags, InSerializeContext, bAllowAnyPackage);
+			out_ResolvedValue = FObjectPropertyBase::FindImportedObject(Property, OwnerObject, ObjectClass, RequiredMetaClass, Temp.ToString(), PortFlags, InSerializeContext, bAllowAnyPackage);
 		}
 		else
 		{
 			// Try the find the object.
-			out_ResolvedValue = FObjectPropertyBase::FindImportedObject(Property, OwnerObject, ObjectClass, RequiredMetaClass, *Temp, PortFlags, InSerializeContext, bAllowAnyPackage);
+			out_ResolvedValue = FObjectPropertyBase::FindImportedObject(Property, OwnerObject, ObjectClass, RequiredMetaClass, Temp.ToString(), PortFlags, InSerializeContext, bAllowAnyPackage);
 		}
 
 		if ( out_ResolvedValue != nullptr && !out_ResolvedValue->GetClass()->IsChildOf(RequiredMetaClass) )
@@ -403,7 +403,7 @@ UObject* FObjectPropertyBase::FindImportedObject( const FProperty* Property, UOb
 		if (Result == nullptr && (PortFlags & PPF_SerializedAsImportText))
 		{
 			// Check string asset redirectors
-			FSoftObjectPath Path = FString(Text);
+			FSoftObjectPath Path(Text);
 			if (Path.PreSavePath())
 			{
 				Result = StaticFindObjectSafe(ObjectClass, nullptr, *Path.ToString());
diff --git a/Engine/Source/Runtime/CoreUObject/Private/UObject/PropertyName.cpp b/Engine/Source/Runtime/CoreUObject/Private/UObject/PropertyName.cpp
--- a/Engine/Source/Runtime/CoreUObject/Private/UObject/PropertyName.cpp
+++ b/Engine/Source/Runtime/CoreUObject/Private/UObject/PropertyName.cpp
@@ -36,19 +36,19 @@ const TCHAR* FNameProperty::ImportText_Internal( const TCHAR* Buffer, void* Data
 {
 	if (!(PortFlags & PPF_Delimited))
 	{
-		*(FName*)Data = FName(Buffer, FNAME_Add);
+		*(FName*)Data = FName(Buffer);
 
 		// in order to indicate that the value was successfully imported, advance the buffer past the last character that was imported
 		Buffer += FCString::Strlen(Buffer);
 	}
 	else
 	{
-		FString Temp;
-		Buffer = FPropertyHelpers::ReadToken(Buffer, Temp, true);
+		TStringBuilder<256> Token;
+		Buffer = FPropertyHelpers::ReadToken(Buffer, /* out */ Token, true);
 		if (!Buffer)
 			return NULL;
 
-		*(FName*)Data = FName(*Temp, FNAME_Add);
+		*(FName*)Data = FName(Token);
 	}
 	return Buffer;
 }
diff --git a/Engine/Source/Runtime/CoreUObject/Private/UObject/PropertyStr.cpp b/Engine/Source/Runtime/CoreUObject/Private/UObject/PropertyStr.cpp
--- a/Engine/Source/Runtime/CoreUObject/Private/UObject/PropertyStr.cpp
+++ b/Engine/Source/Runtime/CoreUObject/Private/UObject/PropertyStr.cpp
@@ -126,7 +126,7 @@ const TCHAR* FStrProperty::ImportText_Internal( const TCHAR* Buffer, void* Data,
 			ErrorText->Logf(TEXT("Missing terminating '\"' in string property value: %s"), Start);
 			return NULL;
 		}
-		*(FString*)Data = Temp;
+		*(FString*)Data = MoveTemp(Temp);
 	}
 	return Buffer;
 }
diff --git a/Engine/Source/Runtime/CoreUObject/Private/UObject/SoftObjectPath.cpp b/Engine/Source/Runtime/CoreUObject/Private/UObject/SoftObjectPath.cpp
--- a/Engine/Source/Runtime/CoreUObject/Private/UObject/SoftObjectPath.cpp
+++ b/Engine/Source/Runtime/CoreUObject/Private/UObject/SoftObjectPath.cpp
@@ -22,18 +22,25 @@ FSoftObjectPath::FSoftObjectPath(const UObject* InObject)
 FString FSoftObjectPath::ToString() const
 {
 	// Most of the time there is no sub path so we can do a single string allocation
-	FString AssetPathString = GetAssetPathString();
 	if (SubPathString.IsEmpty())
 	{
-		return AssetPathString;
+		return GetAssetPathString();
 	}
+
+	TCHAR Buffer[FName::StringBufferSize];
+	FStringView AssetPathString;
+	if (AssetPathName.IsValid())
+	{
+		AssetPathString = FStringView(Buffer, AssetPathName.ToString(Buffer));
+	}
+
 	FString FullPathString;
 
 	// Preallocate to correct size and then append strings
-	FullPathString.Empty(AssetPathString.Len() + SubPathString.Len() + 1);
-	FullPathString.Append(AssetPathString);
-	FullPathString.AppendChar(':');
-	FullPathString.Append(SubPathString);
+	FullPathString.Reserve(AssetPathString.Len() + SubPathString.Len() + 1);
+	FullPathString += AssetPathString;
+	FullPathString += ':';
+	FullPathString += SubPathString;
 	return FullPathString;
 }
 
@@ -50,14 +57,14 @@ void FSoftObjectPath::ToString(FStringBuilderBase& Builder) const
 	}
 }
 
-void FSoftObjectPath::SetPath(FString Path)
+void FSoftObjectPath::SetPath(FWideStringView Path)
 {
 	if (Path.IsEmpty() || Path.Equals(TEXT("None"), ESearchCase::CaseSensitive))
 	{
 		// Empty path, just empty the pathname.
 		Reset();
 	}
-	else if (ensureMsgf(!FPackageName::IsShortPackageName(Path), TEXT("Cannot create SoftObjectPath with short package name '%s'! You must pass in fully qualified package names"), *Path))
+	else if (ensureMsgf(!FPackageName::IsShortPackageName(Path), TEXT("Cannot create SoftObjectPath with short package name '%.*s'! You must pass in fully qualified package names"), Path.Len(), Path.GetData()))
 	{
 		if (Path[0] != '/')
 		{
@@ -65,24 +72,65 @@ void FSoftObjectPath::SetPath(FString Path)
 			Path = FPackageName::ExportTextPathToObjectPath(Path);
 		}
 
-		int32 ColonIndex = INDEX_NONE;
-
+		int32 ColonIndex;
 		if (Path.FindChar(':', ColonIndex))
 		{
 			// Has a subobject, split on that then create a name from the temporary path
+			AssetPathName = Path.Left(ColonIndex);
 			SubPathString = Path.Mid(ColonIndex + 1);
-			Path.RemoveAt(ColonIndex, Path.Len() - ColonIndex);
-			AssetPathName = *Path;
 		}
 		else
 		{
 			// No Subobject
-			AssetPathName = *Path;
+			AssetPathName = Path;
 			SubPathString.Empty();
 		}
 	}
 }
 
+void FSoftObjectPath::SetPath(FAnsiStringView Path)
+{
+	TStringBuilder<256> Wide;
+	Wide << Path;
+	SetPath(Wide);
+}
+
+void FSoftObjectPath::SetPath(FName PathName)
+{
+	if (PathName.IsNone())
+	{
+		Reset();
+	}
+	else
+	{
+		TCHAR Buffer[FName::StringBufferSize];
+		FStringView Path(Buffer, PathName.ToString(Buffer));
+
+		if (ensureMsgf(!FPackageName::IsShortPackageName(Path), TEXT("Cannot create SoftObjectPath with short package name '%s'! You must pass in fully qualified package names"), Path.GetData()))
+		{
+			if (Path[0] != '/')
+			{
+				// Possibly an ExportText path. Trim the ClassName.
+				Path = FPackageName::ExportTextPathToObjectPath(Path);
+			}
+
+			int32 ColonIndex;
+			if (Path.FindChar(':', ColonIndex))
+			{
+				// Has a subobject, split on that then create a name from the temporary path
+				AssetPathName = Path.Left(ColonIndex);
+				SubPathString = Path.Mid(ColonIndex + 1);
+			}
+			else
+			{
+				// No Subobject
+				AssetPathName = Path.GetData() == Buffer ? PathName : FName(Path);
+				SubPathString.Empty();
+			}
+		}
+	}
+}
+
 #if WITH_EDITOR
 	extern bool* GReportSoftObjectPathRedirects;
 #endif
@@ -217,13 +265,6 @@ bool FSoftObjectPath::operator==(FSoftObjectPath const& Other) const
 	return AssetPathName == Other.AssetPathName && SubPathString == Other.SubPathString;
 }
 
-FSoftObjectPath& FSoftObjectPath::operator=(FSoftObjectPath Other)
-{
-	AssetPathName = Other.AssetPathName;
-	SubPathString = MoveTemp(Other.SubPathString);
-	return *this;
-}
-
 bool FSoftObjectPath::ExportTextItem(FString& ValueStr, FSoftObjectPath const& DefaultValue, UObject* Parent, int32 PortFlags, UObject* ExportRootScope) const
 {
 	if (0 != (PortFlags & EPropertyPortFlags::PPF_ExportCpp))
@@ -248,16 +289,16 @@ bool FSoftObjectPath::ExportTextItem(FString& ValueStr, FSoftObjectPath const& D
 
 bool FSoftObjectPath::ImportTextItem(const TCHAR*& Buffer, int32 PortFlags, UObject* Parent, FOutputDevice* ErrorText, FArchive* InSerializingArchive)
 {
-	FString ImportedPath;
-	const TCHAR* NewBuffer = FPropertyHelpers::ReadToken(Buffer, ImportedPath, 1);
+	TStringBuilder<256> ImportedPath;
+	const TCHAR* NewBuffer = FPropertyHelpers::ReadToken(Buffer, /* out */ ImportedPath, /* dotted names */ true);
 	if (!NewBuffer)
 	{
 		return false;
 	}
 	Buffer = NewBuffer;
-	if (ImportedPath == TEXT("None"))
+	if (ImportedPath == TEXT("None"_SV))
 	{
-		ImportedPath = TEXT("");
+		Reset();
 	}
 	else
 	{
@@ -267,7 +308,7 @@ bool FSoftObjectPath::ImportTextItem(const TCHAR*& Buffer, int32 PortFlags, UObj
 			// We have to skip over the first ' as FPropertyHelpers::ReadToken doesn't read single-quoted strings correctly, but does read a path correctly
 			Buffer++; // Skip the leading '
 			ImportedPath.Reset();
-			NewBuffer = FPropertyHelpers::ReadToken(Buffer, ImportedPath, 1);
+			NewBuffer = FPropertyHelpers::ReadToken(Buffer, /* out */ ImportedPath, /* dotted names */ true);
 			if (!NewBuffer)
 			{
 				return false;
@@ -278,9 +319,9 @@ bool FSoftObjectPath::ImportTextItem(const TCHAR*& Buffer, int32 PortFlags, UObj
 				return false;
 			}
 		}
-	}
 
-	SetPath(MoveTemp(ImportedPath));
+		SetPath(ImportedPath);
+	}
 
 #if WITH_EDITOR
 	if (Parent && IsEditorOnlyObject(Parent))
diff --git a/Engine/Source/Runtime/CoreUObject/Public/Misc/PackageName.h b/Engine/Source/Runtime/CoreUObject/Public/Misc/PackageName.h
--- a/Engine/Source/Runtime/CoreUObject/Public/Misc/PackageName.h
+++ b/Engine/Source/Runtime/CoreUObject/Public/Misc/PackageName.h
@@ -145,6 +145,7 @@ public:
 	 */
 	static bool IsShortPackageName(const FString& PossiblyLongName);
 	static bool IsShortPackageName(const FName PossiblyLongName);
+	static bool IsShortPackageName(FStringView PossiblyLongName);
 
 	/**
 	 * Converts package name to short name.
@@ -432,7 +433,10 @@ public:
 	 * @param OutObjectPath The path to the object.
 	 * @return True if the supplied export text path could be parsed
 	 */
-	static bool ParseExportTextPath(const FString& InExportTextPath, FString* OutClassName, FString* OutObjectPath);
+	static bool ParseExportTextPath(FStringView InExportTextPath, FStringView* OutClassName, FStringView* OutObjectPath);
+	static bool ParseExportTextPath(const FString& InExportTextPath, FString* OutClassName, FString* OutObjectPath);	
+	static bool ParseExportTextPath(const TCHAR* InExportTextPath, FStringView* OutClassName, FStringView* OutObjectPath);
+
 
 	/** 
 	 * Returns the path to the object referred to by the supplied export text path, excluding the class name.
@@ -440,7 +444,9 @@ public:
 	 * @param InExportTextPath The export text path for an object. Takes on the form: ClassName'ObjectPath'
 	 * @return The path to the object referred to by the supplied export path.
 	 */
-	static FString ExportTextPathToObjectPath(const FString& InExportTextPath);
+	static FStringView	ExportTextPathToObjectPath(FStringView InExportTextPath);
+	static FString		ExportTextPathToObjectPath(const FString& InExportTextPath);
+	static const TCHAR*	ExportTextPathToObjectPath(const TCHAR* InExportTextPath);
 
 	/** 
 	 * Returns the name of the package referred to by the specified object path
@@ -450,6 +456,7 @@ public:
 	/** 
 	 * Returns the name of the object referred to by the specified object path
 	 */
+	static FStringView ObjectPathToObjectName(FStringView InObjectPath);
 	static FString ObjectPathToObjectName(const FString& InObjectPath);
 
 	/**
@@ -500,7 +507,7 @@ public:
 	 * @param Extension The extension for this package
 	 * @return True if the long package name was fixed up, false otherwise
 	 */
-	static bool FixPackageNameCase(FString& LongPackageName, const FStringView& Extension);
+	static bool FixPackageNameCase(FString& LongPackageName, FStringView Extension);
 
 	UE_DEPRECATED(4.17, "Deprecated. Call TryConvertLongPackageNameToFilename instead, which also works on nested paths")
 	static bool ConvertRootPathToContentPath(const FString& RootPath, FString& OutContentPath);
@@ -528,7 +535,7 @@ private:
 	 * @param InFilename
 	 * @param OutPackageName Long package name.
 	 */
-	static void InternalFilenameToLongPackageName(const FStringView& InFilename, FStringBuilderBase& OutPackageName);
+	static void InternalFilenameToLongPackageName(FStringView InFilename, FStringBuilderBase& OutPackageName);
 
 	/** Event that is triggered when a new content path is mounted */
 	static FOnContentPathMountedEvent OnContentPathMountedEvent;
diff --git a/Engine/Source/Runtime/CoreUObject/Public/UObject/SoftObjectPath.h b/Engine/Source/Runtime/CoreUObject/Public/UObject/SoftObjectPath.h
--- a/Engine/Source/Runtime/CoreUObject/Public/UObject/SoftObjectPath.h
+++ b/Engine/Source/Runtime/CoreUObject/Public/UObject/SoftObjectPath.h
@@ -31,12 +31,15 @@ struct COREUOBJECT_API FSoftObjectPath
 	{
 	}
 
-	/** Construct from a path string */
-	FSoftObjectPath(FString PathString)
-	{
-		SetPath(MoveTemp(PathString));
-	}
-
+	/** Construct from a path string. Non-explicit for backwards compatibility. */
+	FSoftObjectPath(const FString& Path)						{ SetPath(FStringView(Path)); }
+	explicit FSoftObjectPath(FWideStringView Path)				{ SetPath(Path); }
+	explicit FSoftObjectPath(FAnsiStringView Path)				{ SetPath(Path); }
+	explicit FSoftObjectPath(FName Path)						{ SetPath(Path); }
+	explicit FSoftObjectPath(const WIDECHAR* Path)				{ SetPath(FWideStringView(Path)); }
+	explicit FSoftObjectPath(const ANSICHAR* Path)				{ SetPath(FAnsiStringView(Path)); }
+	explicit FSoftObjectPath(TYPE_OF_NULLPTR)					{}
+	
 	/** Construct from an asset FName and subobject pair */
 	FSoftObjectPath(FName InAssetPathName, FString InSubPathString)
 		: AssetPathName(InAssetPathName)
@@ -47,6 +50,16 @@ struct COREUOBJECT_API FSoftObjectPath
 	FSoftObjectPath(const UObject* InObject);
 
 	~FSoftObjectPath() {}
+	
+	FSoftObjectPath& operator=(const FSoftObjectPath& Path)	= default;
+	FSoftObjectPath& operator=(FSoftObjectPath&& Path) = default;
+	FSoftObjectPath& operator=(const FString& Path)						{ SetPath(FStringView(Path)); return *this; }
+	FSoftObjectPath& operator=(FWideStringView Path)					{ SetPath(Path); return *this; }
+	FSoftObjectPath& operator=(FAnsiStringView Path)					{ SetPath(Path); return *this; }
+	FSoftObjectPath& operator=(FName Path)								{ SetPath(Path); return *this; }
+	FSoftObjectPath& operator=(const WIDECHAR* Path)					{ SetPath(FWideStringView(Path)); return *this; }
+	FSoftObjectPath& operator=(const ANSICHAR* Path)					{ SetPath(FAnsiStringView(Path)); return *this; }
+	FSoftObjectPath& operator=(TYPE_OF_NULLPTR)							{ Reset(); return *this; }
 
 	/** Returns string representation of reference, in form /package/path.assetname[:subpath] */
 	FString ToString() const;
@@ -94,7 +107,12 @@ struct COREUOBJECT_API FSoftObjectPath
 	}
 
 	/** Sets asset path of this reference based on a string path */
-	void SetPath(FString Path);
+	void SetPath(FWideStringView Path);
+	void SetPath(FAnsiStringView Path);
+	void SetPath(FName Path);
+	void SetPath(const WIDECHAR* Path)			{ SetPath(FWideStringView(Path)); }
+	void SetPath(const ANSICHAR* Path)			{ SetPath(FAnsiStringView(Path)); }
+	void SetPath(const FString& Path)			{ SetPath(FStringView(Path)); }
 
 	/**
 	 * Attempts to load the asset, this will call LoadObject which can be very slow
@@ -149,7 +167,7 @@ struct COREUOBJECT_API FSoftObjectPath
 	{
 		return !(*this == Other);
 	}
-	FSoftObjectPath& operator=(FSoftObjectPath Other);
+
 	bool ExportTextItem(FString& ValueStr, FSoftObjectPath const& DefaultValue, UObject* Parent, int32 PortFlags, UObject* ExportRootScope) const;
 	bool ImportTextItem( const TCHAR*& Buffer, int32 PortFlags, UObject* Parent, FOutputDevice* ErrorText, FArchive* InSerializingArchive = nullptr );
 	bool SerializeFromMismatchedTag(struct FPropertyTag const& Tag, FStructuredArchive::FSlot Slot);
diff --git a/Engine/Source/Runtime/CoreUObject/Public/UObject/UnrealType.h b/Engine/Source/Runtime/CoreUObject/Public/UObject/UnrealType.h
--- a/Engine/Source/Runtime/CoreUObject/Public/UObject/UnrealType.h
+++ b/Engine/Source/Runtime/CoreUObject/Public/UObject/UnrealType.h
@@ -863,9 +863,10 @@ public:
 class COREUOBJECT_API FPropertyHelpers
 {
 public:
+	static const TCHAR* ReadToken( const TCHAR* Buffer, FString& Out, bool DottedNames = false);
 
-	static const TCHAR* ReadToken( const TCHAR* Buffer, FString& String, bool DottedNames = 0 );
-
+	// @param Out Appended to
+	static const TCHAR* ReadToken( const TCHAR* Buffer, FStringBuilderBase& Out, bool DottedNames = false);
 };
 
 
diff --git a/Engine/Source/Runtime/Engine/Private/AssetManager.cpp b/Engine/Source/Runtime/Engine/Private/AssetManager.cpp
--- a/Engine/Source/Runtime/Engine/Private/AssetManager.cpp
+++ b/Engine/Source/Runtime/Engine/Private/AssetManager.cpp
@@ -2338,17 +2338,35 @@ bool UAssetManager::GetAssetDataForPath(const FSoftObjectPath& ObjectPath, FAsse
 	return AssetData.IsValid();
 }
 
-FSoftObjectPath UAssetManager::GetAssetPathForData(const FAssetData& AssetData) const
+static bool EndsWithBlueprint(FName Name)
 {
-	FString AssetPath = AssetData.IsValid() ? AssetData.ObjectPath.ToString() : FString();
-
-	// All blueprint types end with blueprint as the class, there is no better test unfortunately
-	if (AssetData.AssetClass.ToString().EndsWith(TEXT("Blueprint")))
+	// Numbered names can't end with Blueprint
+	if (Name.IsNone() || Name.GetNumber() != FName().GetNumber())
 	{
-		AssetPath += TEXT("_C");
+		return false;
 	}
 
-	return FSoftObjectPath(AssetPath);
+	TCHAR Buffer[NAME_SIZE];
+	FStringView PlainName(Buffer, Name.GetPlainNameString(Buffer));
+	return PlainName.EndsWith(TEXT("Blueprint"));
+}
+
+FSoftObjectPath UAssetManager::GetAssetPathForData(const FAssetData& AssetData) const
+{
+	if (!AssetData.IsValid())
+	{
+		return FSoftObjectPath();
+	}
+	else if (EndsWithBlueprint(AssetData.AssetClass))
+	{
+		TStringBuilder<256> AssetPath;
+		AssetPath << AssetData.ObjectPath << TEXT("_C");
+		return FSoftObjectPath(FStringView(AssetPath));
+	}
+	else
+	{
+		return FSoftObjectPath(AssetData.ObjectPath);	
+	}
 }
 
 void UAssetManager::GetAssetDataForPathInternal(IAssetRegistry& AssetRegistry, const FString& AssetPath, OUT FAssetData& OutAssetData) const
diff --git a/Engine/Source/Runtime/Engine/Private/ObjectLibrary.cpp b/Engine/Source/Runtime/Engine/Private/ObjectLibrary.cpp
--- a/Engine/Source/Runtime/Engine/Private/ObjectLibrary.cpp
+++ b/Engine/Source/Runtime/Engine/Private/ObjectLibrary.cpp
@@ -498,7 +498,7 @@ int32 UObjectLibrary::LoadAssetsFromAssetData()
 		for (int32 AssetIdx = 0; AssetIdx < AssetDataList.Num(); AssetIdx++)
 		{
 			FAssetData& Data = AssetDataList[AssetIdx];
-			AssetsToStream.AddUnique(Data.PackageName.ToString());
+			AssetsToStream.AddUnique(FSoftObjectPath(Data.PackageName));
 		}
 
 		if (AssetsToStream.Num())
diff --git a/Engine/Source/Runtime/Engine/Private/StreamableManager.cpp b/Engine/Source/Runtime/Engine/Private/StreamableManager.cpp
--- a/Engine/Source/Runtime/Engine/Private/StreamableManager.cpp
+++ b/Engine/Source/Runtime/Engine/Private/StreamableManager.cpp
@@ -1164,7 +1164,7 @@ TSharedPtr<FStreamableHandle> FStreamableManager::RequestAsyncLoad(TArray<FSoftO
 			--NumValidRequests;
 			continue;
 		}
-		else if (FPackageName::IsShortPackageName(TargetName.ToString()))
+		else if (FPackageName::IsShortPackageName(TargetName.GetAssetPathName()))
 		{
 			UE_LOG(LogStreamableManager, Error, TEXT("RequestAsyncLoad called with invalid package name %s"), *TargetName.ToString());
 			NewRequest->CancelHandle();
