- Fixed up LoadFileToStringArray[WithPredicate] to use old (32-bit) method for widechar files
- Fixed loop condition for 64-bit method
- To reduce code duplication, LoadFileToStringArray now calls LoadFileToStringArrayWithPredicate

diff --git a/Engine/Source/Runtime/Core/Private/Misc/FileHelper.cpp b/Engine/Source/Runtime/Core/Private/Misc/FileHelper.cpp
--- a/Engine/Source/Runtime/Core/Private/Misc/FileHelper.cpp
+++ b/Engine/Source/Runtime/Core/Private/Misc/FileHelper.cpp
@@ -222,74 +222,83 @@ bool FFileHelper::LoadFileToString(FString& Result, IPlatformFile* PlatformFile,
 }
 
 bool FFileHelper::LoadFileToStringArray( TArray<FString>& Result, const TCHAR* Filename )
+{
+	return LoadFileToStringArrayWithPredicate(Result, Filename, [](const FString&) { return true;  });
+}
+
+// DEPRECATED
+bool FFileHelper::LoadFileToStringArray(TArray<FString>& Result, const TCHAR* Filename, EHashOptions VerifyFlags)
+{
+	return LoadFileToStringArray(Result, Filename);
+}
+
+bool FFileHelper::LoadFileToStringArrayWithPredicate(TArray<FString>& Result, const TCHAR* Filename, TFunctionRef<bool(const FString&)> Predicate)
 {
 	Result.Empty();
 
 	TArray64<uint8> RawBuffer;
-	if(!LoadFileToArray(RawBuffer, Filename))
+	// can be silent here, since returning false is enough
+	if (!LoadFileToArray(RawBuffer, Filename, FILEREAD_Silent))
 	{
 		return false;
 	}
 
-	for(const uint8* Pos = (uint8*)RawBuffer.GetData(); *Pos != 0; )
+	// we only support the 64-bit enabled "per-line conversion" functionality for UTF-8/ANSI strings, because the \r checks against a byte may fail
+	// so we have to use the old "full string conversion" method, which doesn't work with 64-bits worth of data
+	if (RawBuffer.Num() >= 2 && !(RawBuffer.Num() & 1) && 
+		((RawBuffer[0] == 0xFF && RawBuffer[1] == 0xFE) || (RawBuffer[0] == 0xFE && RawBuffer[1] == 0xFF)))
 	{
-		const uint8* LineStart = Pos;
-		while(*Pos != 0 && *Pos != '\r' && *Pos != '\n')
+		// make sure we can use 32-bit algorithm
+		if (RawBuffer.Num() > MAX_int32)
 		{
-			Pos++;
-		}
-
-		if (Pos - LineStart > MAX_int32)
-		{
-			UE_LOG(LogStreaming, Warning, TEXT("Single line too long found in LoadFileToStringArrayWithPredicate, File: %s"), Filename);
+			UE_LOG(LogStreaming, Error, TEXT("A widechar format file used in LoadFileToStringArray[WithPredicate], but it's too large to be processed. File: %s"), Filename);
 			return false;
 		}
 
-		FString Line;
-		BufferToString(Line, LineStart, UE_PTRDIFF_TO_INT32(Pos - LineStart));
-
-		Result.Add(MoveTemp(Line));
+		FString Buffer;
+		BufferToString(Buffer, RawBuffer.GetData(), RawBuffer.Num());
 
-		if(*Pos == '\r')
+		for (const TCHAR* Pos = *Buffer; *Pos != 0; )
 		{
-			Pos++;
-		}
-		if(*Pos == '\n')
-		{
-			Pos++;
-		}
-	}
-
-	return true;
-}
+			const TCHAR* LineStart = Pos;
+			while (*Pos != 0 && *Pos != '\r' && *Pos != '\n')
+			{
+				Pos++;
+			}
 
-// DEPRECATED
-bool FFileHelper::LoadFileToStringArray(TArray<FString>& Result, const TCHAR* Filename, EHashOptions VerifyFlags)
-{
-	return LoadFileToStringArray(Result, Filename);
-}
+			FString Line(UE_PTRDIFF_TO_INT32(Pos - LineStart), LineStart);
+			if (Invoke(Predicate, Line))
+			{
+				Result.Add(MoveTemp(Line));
+			}
 
-bool FFileHelper::LoadFileToStringArrayWithPredicate(TArray<FString>& Result, const TCHAR* Filename, TFunctionRef<bool(const FString&)> Predicate)
-{
-	Result.Empty();
+			if (*Pos == '\r')
+			{
+				Pos++;
+			}
+			if (*Pos == '\n')
+			{
+				Pos++;
+			}
+		}
 
-	TArray64<uint8> RawBuffer;
-	if (!LoadFileToArray(RawBuffer, Filename))
-	{
-		return false;
+		return true;
 	}
 
-	for (const uint8* Pos = (uint8*)RawBuffer.GetData(); *Pos != 0; )
+
+	int64 Length = RawBuffer.Num();
+	for (const uint8* Pos = (uint8*)RawBuffer.GetData(); Length > 0; Length--)
 	{
 		const uint8* LineStart = Pos;
-		while (*Pos != 0 && *Pos != '\r' && *Pos != '\n')
+		while (Length > 0 && *Pos != '\r' && *Pos != '\n')
 		{
 			Pos++;
+			Length--;
 		}
 
 		if (Pos - LineStart > MAX_int32)
 		{
-			UE_LOG(LogStreaming, Warning, TEXT("Single line too long found in LoadFileToStringArrayWithPredicate, File: %s"), Filename);
+			UE_LOG(LogStreaming, Error, TEXT("Single line too long found in LoadFileToStringArrayWithPredicate, File: %s"), Filename);
 			return false;
 		}
 
@@ -304,10 +313,12 @@ bool FFileHelper::LoadFileToStringArrayWithPredicate(TArray<FString>& Result, co
 		if (*Pos == '\r')
 		{
 			Pos++;
+			Length--;
 		}
 		if (*Pos == '\n')
 		{
 			Pos++;
+			Length--;
 		}
 	}
 
