#UE4 Made FlushThreadCacheMaxWaitTime usable from ini/commandline in normal builds. This isn't really a runtime tweak, just a way to control the threshold for a warning

diff --git a/Engine/Source/Runtime/Core/Private/HAL/MallocBinned2.cpp b/Engine/Source/Runtime/Core/Private/HAL/MallocBinned2.cpp
--- a/Engine/Source/Runtime/Core/Private/HAL/MallocBinned2.cpp
+++ b/Engine/Source/Runtime/Core/Private/HAL/MallocBinned2.cpp
@@ -51,15 +51,16 @@ static FAutoConsoleVariableRef GMallocBinned2AllocExtraCVar(
 	TEXT("When we do acquire the lock, how many blocks cached in TLS caches. In no case will we grab more than a page.")
 	);
 
+#endif
+
 double GMallocBinned2FlushThreadCacheMaxWaitTime = 0.02f;
 static FAutoConsoleVariableRef GMallocBinned2FlushThreadCacheMaxWaitTimeCVar(
 	TEXT("MallocBinned2.FlushThreadCacheMaxWaitTime"),
 	GMallocBinned2FlushThreadCacheMaxWaitTime,
-	TEXT("The threshold of time before warning about FlushCurrentThreadCache taking too long (seconds).")
+	TEXT("The threshold of time before warning about FlushCurrentThreadCache taking too long (seconds)."),
+	ECVF_ReadOnly
 );
 
-#endif
-
 #if BINNED2_ALLOCATOR_STATS
 TAtomic<int64> AllocatedSmallPoolMemory(0); // memory that's requested to be allocated by the game
 TAtomic<int64> AllocatedOSSmallPoolMemory(0);
@@ -1078,16 +1079,11 @@ void FMallocBinned2::FlushCurrentThreadCache()
 	}
 
 	// These logs must happen outside the above mutex to avoid deadlocks
-#if BINNED2_ALLOW_RUNTIME_TWEAKING
-	double MaxWaitTimeThreshold = GMallocBinned2FlushThreadCacheMaxWaitTime;
-#else
-	double MaxWaitTimeThreshold = 0.02f;
-#endif
-	if (WaitForMutexTime > MaxWaitTimeThreshold)
+	if (WaitForMutexTime > GMallocBinned2FlushThreadCacheMaxWaitTime)
 	{
 		UE_LOG(LogMemory, Warning, TEXT("FMallocBinned2 took %6.2fms to wait for mutex for trim."), WaitForMutexTime * 1000.0f);
 	}
-	if (WaitForMutexAndTrimTime > MaxWaitTimeThreshold)
+	if (WaitForMutexAndTrimTime > GMallocBinned2FlushThreadCacheMaxWaitTime)
 	{
 		UE_LOG(LogMemory, Warning, TEXT("FMallocBinned2 took %6.2fms to wait for mutex AND trim."), WaitForMutexAndTrimTime * 1000.0f);
 	}
diff --git a/Engine/Source/Runtime/Core/Private/HAL/MallocBinned3.cpp b/Engine/Source/Runtime/Core/Private/HAL/MallocBinned3.cpp
--- a/Engine/Source/Runtime/Core/Private/HAL/MallocBinned3.cpp
+++ b/Engine/Source/Runtime/Core/Private/HAL/MallocBinned3.cpp
@@ -66,15 +66,16 @@ static FAutoConsoleVariableRef GMallocBinned3AllocExtraCVar(
 	TEXT("When we do acquire the lock, how many blocks cached in TLS caches. In no case will we grab more than a page.")
 	);
 
+#endif
+
 double GMallocBinned3FlushThreadCacheMaxWaitTime = 0.02f;
 static FAutoConsoleVariableRef GMallocBinned3FlushThreadCacheMaxWaitTimeCVar(
 	TEXT("MallocBinned3.FlushThreadCacheMaxWaitTime"),
 	GMallocBinned3FlushThreadCacheMaxWaitTime,
-	TEXT("The threshold of time before warning about FlushCurrentThreadCache taking too long (seconds).")
+	TEXT("The threshold of time before warning about FlushCurrentThreadCache taking too long (seconds)."),
+	ECVF_ReadOnly
 );
 
-#endif
-
 #if BINNED3_ALLOCATOR_STATS
 int64 Binned3AllocatedSmallPoolMemory = 0; // memory that's requested to be allocated by the game
 int64 Binned3AllocatedOSSmallPoolMemory = 0;
@@ -1315,16 +1316,11 @@ void FMallocBinned3::FlushCurrentThreadCache()
 	}
 
 	// These logs must happen outside the above mutex to avoid deadlocks
-#if BINNED3_ALLOW_RUNTIME_TWEAKING
-	double MaxWaitTimeThreshold = GMallocBinned3FlushThreadCacheMaxWaitTime;
-#else
-	double MaxWaitTimeThreshold = 0.02f;
-#endif
-	if (WaitForMutexTime > MaxWaitTimeThreshold)
+	if (WaitForMutexTime > GMallocBinned3FlushThreadCacheMaxWaitTime)
 	{
 		UE_LOG(LogMemory, Warning, TEXT("FMallocBinned3 took %6.2fms to wait for mutex for trim."), WaitForMutexTime * 1000.0f);
 	}
-	if (WaitForMutexAndTrimTime > MaxWaitTimeThreshold)
+	if (WaitForMutexAndTrimTime > GMallocBinned3FlushThreadCacheMaxWaitTime)
 	{
 		UE_LOG(LogMemory, Warning, TEXT("FMallocBinned3 took %6.2fms to wait for mutex AND trim."), WaitForMutexAndTrimTime * 1000.0f);
 	}
