#UE4 CIS fix

diff --git a/Engine/Source/Runtime/Core/Private/HAL/MallocBinned2.cpp b/Engine/Source/Runtime/Core/Private/HAL/MallocBinned2.cpp
--- a/Engine/Source/Runtime/Core/Private/HAL/MallocBinned2.cpp
+++ b/Engine/Source/Runtime/Core/Private/HAL/MallocBinned2.cpp
@@ -53,7 +53,7 @@ static FAutoConsoleVariableRef GMallocBinned2AllocExtraCVar(
 
 #endif
 
-double GMallocBinned2FlushThreadCacheMaxWaitTime = 0.02f;
+float GMallocBinned2FlushThreadCacheMaxWaitTime = 0.02f;
 static FAutoConsoleVariableRef GMallocBinned2FlushThreadCacheMaxWaitTimeCVar(
 	TEXT("MallocBinned2.FlushThreadCacheMaxWaitTime"),
 	GMallocBinned2FlushThreadCacheMaxWaitTime,
diff --git a/Engine/Source/Runtime/Core/Private/HAL/MallocBinned3.cpp b/Engine/Source/Runtime/Core/Private/HAL/MallocBinned3.cpp
--- a/Engine/Source/Runtime/Core/Private/HAL/MallocBinned3.cpp
+++ b/Engine/Source/Runtime/Core/Private/HAL/MallocBinned3.cpp
@@ -68,7 +68,7 @@ static FAutoConsoleVariableRef GMallocBinned3AllocExtraCVar(
 
 #endif
 
-double GMallocBinned3FlushThreadCacheMaxWaitTime = 0.02f;
+float GMallocBinned3FlushThreadCacheMaxWaitTime = 0.02f;
 static FAutoConsoleVariableRef GMallocBinned3FlushThreadCacheMaxWaitTimeCVar(
 	TEXT("MallocBinned3.FlushThreadCacheMaxWaitTime"),
 	GMallocBinned3FlushThreadCacheMaxWaitTime,
