Handle decompression of GZIP archives in memory. Functionality was previously missing, and just returned immediately from FCompression::UncompressMemory().

diff --git a/Engine/Source/Runtime/Core/Private/Misc/Compression.cpp b/Engine/Source/Runtime/Core/Private/Misc/Compression.cpp
--- a/Engine/Source/Runtime/Core/Private/Misc/Compression.cpp
+++ b/Engine/Source/Runtime/Core/Private/Misc/Compression.cpp
@@ -163,6 +163,67 @@ static bool appCompressMemoryGZIP(void* CompressedBuffer, int32& CompressedSize,
 	return bOperationSucceeded;
 }
 
+/**
+ * Thread-safe abstract compression routine. Compresses memory from uncompressed buffer and writes it to compressed
+ * buffer. Updates CompressedSize with size of compressed data.
+ *
+ * @param	UncompressedBuffer			Buffer containing uncompressed data
+ * @param	UncompressedSize			Size of uncompressed data in bytes
+ * @param	CompressedBuffer			Buffer compressed data is going to be read from
+ * @param	CompressedSize				Size of CompressedBuffer data in bytes
+ * @return true if compression succeeds, false if it fails because CompressedBuffer was too small or other reasons
+ */
+bool appUncompressMemoryGZIP(void* UncompressedBuffer, int32 UncompressedSize, const void* CompressedBuffer, int32 CompressedSize)
+{
+	DECLARE_SCOPE_CYCLE_COUNTER(TEXT("Uncompress Memory GZIP"), STAT_appUncompressMemoryGZIP, STATGROUP_Compression);
+
+	// Zlib wants to use unsigned long.
+	unsigned long ZCompressedSize = CompressedSize;
+	unsigned long ZUncompressedSize = UncompressedSize;
+
+	z_stream stream;
+	stream.zalloc = &zalloc;
+	stream.zfree = &zfree;
+	stream.opaque = Z_NULL;
+	stream.next_in = (uint8*)CompressedBuffer;
+	stream.avail_in = ZCompressedSize;
+	stream.next_out = (uint8*)UncompressedBuffer;
+	stream.avail_out = ZUncompressedSize;
+
+	int32 Result = inflateInit2(&stream, 16 + MAX_WBITS);
+
+	if (Result != Z_OK)
+		return false;
+
+	// Uncompress data.
+	Result = inflate(&stream, Z_FINISH);
+	if (Result == Z_STREAM_END)
+	{
+		ZUncompressedSize = stream.total_out;
+	}
+
+	int32 EndResult = inflateEnd(&stream);
+	if (Result >= Z_OK)
+	{
+		Result = EndResult;
+	}
+
+	// These warnings will be compiled out in shipping.
+	UE_CLOG(Result == Z_MEM_ERROR, LogCompression, Warning, TEXT("appUncompressMemoryGZIP failed: Error: Z_MEM_ERROR, not enough memory!"));
+	UE_CLOG(Result == Z_BUF_ERROR, LogCompression, Warning, TEXT("appUncompressMemoryGZIP failed: Error: Z_BUF_ERROR, not enough room in the output buffer!"));
+	UE_CLOG(Result == Z_DATA_ERROR, LogCompression, Warning, TEXT("appUncompressMemoryGZIP failed: Error: Z_DATA_ERROR, input data was corrupted or incomplete!"));
+
+	bool bOperationSucceeded = (Result == Z_OK);
+
+	// Sanity check to make sure we uncompressed as much data as we expected to.
+	if (UncompressedSize != ZUncompressedSize)
+	{
+		UE_LOG(LogCompression, Warning, TEXT("appUncompressMemoryGZIP failed: Mismatched uncompressed size. Expected: %d, Got:%d. Result: %d"), UncompressedSize, ZUncompressedSize, Result);
+		bOperationSucceeded = false;
+	}
+	return bOperationSucceeded;
+}
+
 /**
  * Thread-safe abstract compression routine. Compresses memory from uncompressed buffer and writes it to compressed
  * buffer. Updates CompressedSize with size of compressed data.
@@ -530,7 +591,8 @@ bool FCompression::UncompressMemory(FName FormatName, void* UncompressedBuffer,
 	}
 	else if (FormatName == NAME_Gzip)
 	{
-		// @todo buh?
+		// hardcoded gzip
+		bUncompressSucceeded = appUncompressMemoryGZIP(UncompressedBuffer, UncompressedSize, CompressedBuffer, CompressedSize);
 	}
 	else if (FormatName == NAME_LZ4)
 	{
