- Missed a warning fix
#rnx

diff --git a/Engine/Source/Runtime/Core/Private/Misc/FileHelper.cpp b/Engine/Source/Runtime/Core/Private/Misc/FileHelper.cpp
--- a/Engine/Source/Runtime/Core/Private/Misc/FileHelper.cpp
+++ b/Engine/Source/Runtime/Core/Private/Misc/FileHelper.cpp
@@ -256,7 +256,7 @@ bool FFileHelper::LoadFileToStringArrayWithPredicate(TArray<FString>& Result, co
 		}
 
 		FString Buffer;
-		BufferToString(Buffer, RawBuffer.GetData(), RawBuffer.Num());
+		BufferToString(Buffer, RawBuffer.GetData(), (int32)RawBuffer.Num());
 
 		for (const TCHAR* Pos = *Buffer; *Pos != 0; )
 		{
