Fix a typo in TStringView::LeftInline

Reported by [REDACTED_PERSON]

#rnx

diff --git a/Engine/Source/Runtime/Core/Public/Containers/StringView.h b/Engine/Source/Runtime/Core/Public/Containers/StringView.h
--- a/Engine/Source/Runtime/Core/Public/Containers/StringView.h
+++ b/Engine/Source/Runtime/Core/Public/Containers/StringView.h
@@ -175,7 +175,7 @@ public:
 	CORE_API ViewType TrimEnd() const;
 
 	/** Modifies the view to be the given number of characters from the left. */
-	inline void LeftInlineInline(SizeType CharCount) { *this = Left(CharCount); }
+	inline void LeftInline(SizeType CharCount) { *this = Left(CharCount); }
 	/** Modifies the view by chopping the given number of characters from the right. */
 	inline void LeftChopInline(SizeType CharCount) { *this = LeftChop(CharCount); }
 	/** Modifies the view to be the given number of characters from the right. */
