Fix bUseNEONForArmV7=true compile error for android where vqtbx1q_u8 is an A64 instruction.
Fallback to use FPU version for VectorContainsNanOrInfinite.

diff --git a/Engine/Source/Runtime/Core/Public/Math/UnrealMathNeon.h b/Engine/Source/Runtime/Core/Public/Math/UnrealMathNeon.h
--- a/Engine/Source/Runtime/Core/Public/Math/UnrealMathNeon.h
+++ b/Engine/Source/Runtime/Core/Public/Math/UnrealMathNeon.h
@@ -1343,6 +1343,12 @@ FORCEINLINE void VectorSinCos(  VectorRegister* VSinAngles, VectorRegister* VCos
 // Returns true if the vector contains a component that is either NAN or +/-infinite.
 inline bool VectorContainsNaNOrInfinite(const VectorRegister& Vec)
 {
+// ArmV7 doesn't have vqtbx1q_u8, so fallback in this case.
+#if PLATFORM_ANDROID_ARM
+	float VecComponents[4];
+	vst1q_f32(VecComponents, Vec);
+	return  !FMath::IsFinite(VecComponents[0]) || !FMath::IsFinite(VecComponents[1]) || !FMath::IsFinite(VecComponents[2]) || !FMath::IsFinite(VecComponents[3]);
+#else
 	// https://en.wikipedia.org/wiki/IEEE_754-1985
 	// Infinity is represented with all exponent bits set, with the correct sign bit.
 	// NaN is represented with all exponent bits set, plus at least one fraction/significant bit set.
@@ -1382,6 +1388,7 @@ inline bool VectorContainsNaNOrInfinite(const VectorRegister& Vec)
 	uint8x16_t res = (uint8x16_t)VectorCompareEQ(ExpTest, FloatInfinity);
 	// If we have all zeros, all elements are finite
 	return vgetq_lane_u32((uint32x4_t)vqtbx1q_u8(res, res, Table), 0) != 0;
+#endif
 }
 
 //TODO: Vectorize
