Fix FParse::QuotedString() for wide characters

diff --git a/Engine/Source/Runtime/Core/Private/Misc/Parse.cpp b/Engine/Source/Runtime/Core/Private/Misc/Parse.cpp
--- a/Engine/Source/Runtime/Core/Private/Misc/Parse.cpp
+++ b/Engine/Source/Runtime/Core/Private/Misc/Parse.cpp
@@ -362,8 +362,8 @@ bool ParseQuotedString( const TCHAR* Buffer, T& Value, int32* OutNumCharsRead)
 
 	auto ShouldParse = [](const TCHAR Ch)
 	{
-		constexpr FAsciiSet ParsableCharacters = ~(FAsciiSet("\"\n\r") + '\0');
-		return ParsableCharacters.Test(Ch);
+		constexpr FAsciiSet StopCharacters = FAsciiSet("\"\n\r") + '\0';
+		return StopCharacters.Test(Ch) == 0;
 	};
 
 	while (ShouldParse(*Buffer))
