- Fixed a bug with LoadFileToStringArray with dropping lines
#rnx

diff --git a/Engine/Source/Runtime/Core/Private/Misc/FileHelper.cpp b/Engine/Source/Runtime/Core/Private/Misc/FileHelper.cpp
--- a/Engine/Source/Runtime/Core/Private/Misc/FileHelper.cpp
+++ b/Engine/Source/Runtime/Core/Private/Misc/FileHelper.cpp
@@ -287,7 +287,7 @@ bool FFileHelper::LoadFileToStringArrayWithPredicate(TArray<FString>& Result, co
 
 
 	int64 Length = RawBuffer.Num();
-	for (const uint8* Pos = (uint8*)RawBuffer.GetData(); Length > 0; Length--)
+	for (const uint8* Pos = (uint8*)RawBuffer.GetData(); Length > 0; )
 	{
 		const uint8* LineStart = Pos;
 		while (Length > 0 && *Pos != '\r' && *Pos != '\n')
