Avoid FName -> FString allocation in FName::IsValidXName

#rnx

diff --git a/Engine/Source/Runtime/Core/Private/UObject/UnrealNames.cpp b/Engine/Source/Runtime/Core/Private/UObject/UnrealNames.cpp
--- a/Engine/Source/Runtime/Core/Private/UObject/UnrealNames.cpp
+++ b/Engine/Source/Runtime/Core/Private/UObject/UnrealNames.cpp
@@ -2165,7 +2165,24 @@ FString FName::SafeString(FNameEntryId InDisplayIndex, int32 InstanceNumber)
 	return FName(InDisplayIndex, InDisplayIndex, InstanceNumber).ToString();
 }
 
+bool FName::IsValidXName(const FName InName, const FString& InInvalidChars, FText* OutReason, const FText* InErrorCtx)
+{
+	TStringBuilder<FName::StringBufferSize> NameStr;
+	InName.ToString(NameStr);
+	return IsValidXName(FStringView(NameStr), InInvalidChars, OutReason, InErrorCtx);
+}
+
+bool FName::IsValidXName(const TCHAR* InName, const FString& InInvalidChars, FText* OutReason, const FText* InErrorCtx)
+{
+	return IsValidXName(FStringView(InName), InInvalidChars, OutReason, InErrorCtx);
+}
+
 bool FName::IsValidXName(const FString& InName, const FString& InInvalidChars, FText* OutReason, const FText* InErrorCtx)
+{
+	return IsValidXName(FStringView(InName), InInvalidChars, OutReason, InErrorCtx);
+}
+
+bool FName::IsValidXName(const FStringView& InName, const FString& InInvalidChars, FText* OutReason, const FText* InErrorCtx)
 {
 	if (InName.IsEmpty() || InInvalidChars.IsEmpty())
 	{
@@ -2177,7 +2194,8 @@ bool FName::IsValidXName(const FString& InName, const FString& InInvalidChars, F
 	TSet<TCHAR> AlreadyMatchedInvalidChars;
 	for (const TCHAR InvalidChar : InInvalidChars)
 	{
-		if (!AlreadyMatchedInvalidChars.Contains(InvalidChar) && InName.GetCharArray().Contains(InvalidChar))
+		int32 InvalidCharIndex = INDEX_NONE;
+		if (!AlreadyMatchedInvalidChars.Contains(InvalidChar) && InName.FindChar(InvalidChar, InvalidCharIndex))
 		{
 			MatchedInvalidChars.AppendChar(InvalidChar);
 			AlreadyMatchedInvalidChars.Add(InvalidChar);
diff --git a/Engine/Source/Runtime/Core/Public/UObject/NameTypes.h b/Engine/Source/Runtime/Core/Public/UObject/NameTypes.h
--- a/Engine/Source/Runtime/Core/Public/UObject/NameTypes.h
+++ b/Engine/Source/Runtime/Core/Public/UObject/NameTypes.h
@@ -596,7 +596,10 @@ public:
 	 *
 	 * @return	true if the name is valid
 	 */
+	static bool IsValidXName( const FName InName, const FString& InInvalidChars, FText* OutReason = nullptr, const FText* InErrorCtx = nullptr );
+	static bool IsValidXName( const TCHAR* InName, const FString& InInvalidChars, FText* OutReason = nullptr, const FText* InErrorCtx = nullptr );
 	static bool IsValidXName( const FString& InName, const FString& InInvalidChars, FText* OutReason = nullptr, const FText* InErrorCtx = nullptr );
+	static bool IsValidXName( const FStringView& InName, const FString& InInvalidChars, FText* OutReason = nullptr, const FText* InErrorCtx = nullptr );
 
 	/**
 	 * Checks to see that a FName follows the rules that Unreal requires.
@@ -609,7 +612,7 @@ public:
 	 */
 	bool IsValidXName( const FString& InInvalidChars = INVALID_NAME_CHARACTERS, FText* OutReason = nullptr, const FText* InErrorCtx = nullptr ) const
 	{
-		return IsValidXName(ToString(), InInvalidChars, OutReason, InErrorCtx);
+		return IsValidXName(*this, InInvalidChars, OutReason, InErrorCtx);
 	}
 
 	/**
@@ -622,7 +625,7 @@ public:
 	 */
 	bool IsValidXName( FText& OutReason, const FString& InInvalidChars = INVALID_NAME_CHARACTERS ) const
 	{
-		return IsValidXName(ToString(), InInvalidChars, &OutReason);
+		return IsValidXName(*this, InInvalidChars, &OutReason);
 	}
 
 	/**
@@ -634,7 +637,7 @@ public:
 	 */
 	bool IsValidObjectName( FText& OutReason ) const
 	{
-		return IsValidXName(ToString(), INVALID_OBJECTNAME_CHARACTERS, &OutReason);
+		return IsValidXName(*this, INVALID_OBJECTNAME_CHARACTERS, &OutReason);
 	}
 
 	/**
@@ -647,7 +650,7 @@ public:
 	 */
 	bool IsValidGroupName( FText& OutReason, bool bIsGroupName=false ) const
 	{
-		return IsValidXName(ToString(), INVALID_LONGPACKAGE_CHARACTERS, &OutReason);
+		return IsValidXName(*this, INVALID_LONGPACKAGE_CHARACTERS, &OutReason);
 	}
 
 	/**
