Update FFileHelper::SaveStringToFile to take a FStringView

diff --git a/Engine/Source/Runtime/Core/Private/Misc/FileHelper.cpp b/Engine/Source/Runtime/Core/Private/Misc/FileHelper.cpp
--- a/Engine/Source/Runtime/Core/Private/Misc/FileHelper.cpp
+++ b/Engine/Source/Runtime/Core/Private/Misc/FileHelper.cpp
@@ -367,7 +367,7 @@ bool FFileHelper::SaveArrayToFile(const TArray64<uint8>& Array, const TCHAR* Fil
  * Write the FString to a file.
  * Supports all combination of ANSI/Unicode files and platforms.
  */
-bool FFileHelper::SaveStringToFile( const FString& String, const TCHAR* Filename,  EEncodingOptions EncodingOptions, IFileManager* FileManager /*= &IFileManager::Get()*/, uint32 WriteFlags )
+bool FFileHelper::SaveStringToFile( FStringView String, const TCHAR* Filename,  EEncodingOptions EncodingOptions, IFileManager* FileManager /*= &IFileManager::Get()*/, uint32 WriteFlags )
 {
 	// max size of the string is a UCS2CHAR for each character and some UNICODE magic 
 	TUniquePtr<FArchive> Ar = TUniquePtr<FArchive>( FileManager->CreateFileWriter( Filename, WriteFlags ) );
@@ -377,18 +377,18 @@ bool FFileHelper::SaveStringToFile( const FString& String, const TCHAR* Filename
 	if( String.IsEmpty() )
 		return true;
 
-	bool SaveAsUnicode = EncodingOptions == EEncodingOptions::ForceUnicode || ( EncodingOptions == EEncodingOptions::AutoDetect && !FCString::IsPureAnsi(*String) );
+	bool SaveAsUnicode = EncodingOptions == EEncodingOptions::ForceUnicode || ( EncodingOptions == EEncodingOptions::AutoDetect && !FCString::IsPureAnsi(String.GetData(), String.Len()) );
 	if( EncodingOptions == EEncodingOptions::ForceUTF8 )
 	{
 		UTF8CHAR UTF8BOM[] = { 0xEF, 0xBB, 0xBF };
 		Ar->Serialize( &UTF8BOM, UE_ARRAY_COUNT(UTF8BOM) * sizeof(UTF8CHAR) );
 
-		FTCHARToUTF8 UTF8String(*String, String.Len());
+		FTCHARToUTF8 UTF8String(String.GetData(), String.Len());
 		Ar->Serialize( (UTF8CHAR*)UTF8String.Get(), UTF8String.Length() * sizeof(UTF8CHAR) );
 	}
 	else if ( EncodingOptions == EEncodingOptions::ForceUTF8WithoutBOM )
 	{
-		FTCHARToUTF8 UTF8String(*String, String.Len());
+		FTCHARToUTF8 UTF8String(String.GetData(), String.Len());
 		Ar->Serialize((UTF8CHAR*)UTF8String.Get(), UTF8String.Length() * sizeof(UTF8CHAR));
 	}
 	else if (SaveAsUnicode)
@@ -397,12 +397,12 @@ bool FFileHelper::SaveStringToFile( const FString& String, const TCHAR* Filename
 		Ar->Serialize( &BOM, sizeof(UTF16CHAR) );
 
 		// Note: This is a no-op on platforms that are using a 16-bit TCHAR
-		FTCHARToUTF16 UTF16String(*String, String.Len());
+		FTCHARToUTF16 UTF16String(String.GetData(), String.Len());
 		Ar->Serialize((UTF16CHAR*)UTF16String.Get(), UTF16String.Length() * sizeof(UTF16CHAR));
 	}
 	else
 	{
-		auto Src = StringCast<ANSICHAR>(*String, String.Len());
+		auto Src = StringCast<ANSICHAR>(String.GetData(), String.Len());
 		Ar->Serialize( (ANSICHAR*)Src.Get(), Src.Length() * sizeof(ANSICHAR) );
 	}
 
diff --git a/Engine/Source/Runtime/Core/Public/Misc/CString.h b/Engine/Source/Runtime/Core/Public/Misc/CString.h
--- a/Engine/Source/Runtime/Core/Public/Misc/CString.h
+++ b/Engine/Source/Runtime/Core/Public/Misc/CString.h
@@ -63,6 +63,7 @@ struct TCString
 	 * @param Str - string that will be checked
 	 **/
 	static FORCEINLINE bool IsPureAnsi(const CharType* Str);
+	static FORCEINLINE bool IsPureAnsi(const CharType* Str, const SIZE_T StrLen);
 
 	/**
 	 * Returns whether this string contains only numeric characters 
@@ -880,7 +881,7 @@ int32 TCString<T>::GetVarArgs( CharType* Dest, SIZE_T DestSize, const CharType*&
 	TCString<WIDECHAR> specializations
 -----------------------------------------------------------------------------*/
 template <> FORCEINLINE
-bool TCString<WIDECHAR>::IsPureAnsi(const WIDECHAR* Str)
+bool TCString<WIDECHAR>::IsPureAnsi(const CharType* Str)
 {
 	for( ; *Str; Str++ )
 	{
@@ -893,6 +894,20 @@ bool TCString<WIDECHAR>::IsPureAnsi(const WIDECHAR* Str)
 }
 
 
+template <> FORCEINLINE
+bool TCString<WIDECHAR>::IsPureAnsi(const CharType* Str, const SIZE_T StrLen)
+{
+	for (SIZE_T Idx = 0; Idx < StrLen; Idx++, Str++)
+	{
+		if (*Str > 0x7f)
+		{
+			return false;
+		}
+	}
+	return true;
+}
+
+
 template <>
 inline int32 TCString<WIDECHAR>::SprintfImpl(CharType* Dest, const CharType* Fmt, ...)
 {
@@ -918,7 +933,14 @@ FORCEINLINE bool TCString<TCHAR>::ToBool(const WIDECHAR* Str)
 /*-----------------------------------------------------------------------------
 	TCString<ANSICHAR> specializations
 -----------------------------------------------------------------------------*/
-template <> FORCEINLINE bool TCString<ANSICHAR>::IsPureAnsi(const CharType* Str)
+template <> FORCEINLINE
+bool TCString<ANSICHAR>::IsPureAnsi(const CharType* Str)
+{
+	return true;
+}
+
+template <> FORCEINLINE
+bool TCString<ANSICHAR>::IsPureAnsi(const CharType* Str, const SIZE_T StrLen)
 {
 	return true;
 }
diff --git a/Engine/Source/Runtime/Core/Public/Misc/FileHelper.h b/Engine/Source/Runtime/Core/Public/Misc/FileHelper.h
--- a/Engine/Source/Runtime/Core/Public/Misc/FileHelper.h
+++ b/Engine/Source/Runtime/Core/Public/Misc/FileHelper.h
@@ -3,6 +3,7 @@
 #pragma once
 
 #include "CoreTypes.h"
+#include "Containers/StringView.h"
 #include "Containers/UnrealString.h"
 #include "HAL/FileManager.h"
 #include "Containers/ArrayView.h"
@@ -130,7 +131,7 @@ struct CORE_API FFileHelper
 	 * Write the FString to a file.
 	 * Supports all combination of ANSI/Unicode files and platforms.
 	 */
-	static bool SaveStringToFile( const FString& String, const TCHAR* Filename, EEncodingOptions EncodingOptions = EEncodingOptions::AutoDetect, IFileManager* FileManager = &IFileManager::Get(), uint32 WriteFlags = 0 );
+	static bool SaveStringToFile( FStringView String, const TCHAR* Filename, EEncodingOptions EncodingOptions = EEncodingOptions::AutoDetect, IFileManager* FileManager = &IFileManager::Get(), uint32 WriteFlags = 0 );
 
 	/**
 	 * Write the FString to a file.
