Fix for ensure messages disappearing in editor log

Calls to PanicFlushLog should only happen on crashes, and not ensures. On ensures on any thread other than the main thread the editor log, which is not thread safe, will loose these messages. On crashes we don't care about the non thread-safe output devices since we will not recover, we are only interested in the log file.

diff --git a/Engine/Source/Runtime/Core/Private/Windows/WindowsPlatformCrashContext.cpp b/Engine/Source/Runtime/Core/Private/Windows/WindowsPlatformCrashContext.cpp
--- a/Engine/Source/Runtime/Core/Private/Windows/WindowsPlatformCrashContext.cpp
+++ b/Engine/Source/Runtime/Core/Private/Windows/WindowsPlatformCrashContext.cpp
@@ -569,11 +569,6 @@ int32 ReportCrashForMonitor(
 	SharedContext->SessionContext.bIsExitRequested = IsEngineExitRequested();
 	FCString::Strcpy(SharedContext->ErrorMessage, CR_MAX_ERROR_MESSAGE_CHARS-1, ErrorMessage);
 
-	if (GLog)
-	{
-		GLog->PanicFlushThreadedLogs();
-	}
-
 	// Setup all the thread ids and names using snapshot dbghelp. Since it's not possible to 
 	// query thread names from an external process.
 	uint32 ThreadIdx = 0;
@@ -746,7 +741,6 @@ int32 ReportCrashUsingCrashReportClient(FWindowsPlatformCrashContext& InContext,
 			InContext.CopyPlatformSpecificFiles(*CrashFolderAbsolute, (void*) ExceptionInfo);
 
 			// Copy the log file to output
-			GLog->PanicFlushThreadedLogs();
 			FGenericCrashContext::DumpLog(CrashFolderAbsolute);
 
 			// Build machines do not upload these automatically since it is not okay to have lingering processes after the build completes.
@@ -1133,13 +1127,13 @@ private:
 		// Stop the heartbeat thread so that it doesn't interfere with crashreporting
 		FThreadHeartBeat::Get().Stop();
 
-		GLog->PanicFlushThreadedLogs();
-
 		// Then try run time crash processing and broadcast information about a crash.
 		FCoreDelegates::OnHandleSystemError.Broadcast();
 
 		if (GLog)
 		{
+			//Panic flush the logs to make sure there are no entries queued. This is
+			//not thread safe so it will skip for example editor log.
 			GLog->PanicFlushThreadedLogs();
 		}
 		
