- Fixed up PreLoad file re-attempt to read
#rnx

diff --git a/Engine/Source/Runtime/Core/Private/Misc/PreLoadFile.cpp b/Engine/Source/Runtime/Core/Private/Misc/PreLoadFile.cpp
--- a/Engine/Source/Runtime/Core/Private/Misc/PreLoadFile.cpp
+++ b/Engine/Source/Runtime/Core/Private/Misc/PreLoadFile.cpp
@@ -51,9 +51,12 @@ void FPreLoadFile::KickOffRead()
 		Path = Path.Replace(TEXT("{PROJECT}"), *FPaths::ProjectDir());
 	}
 
-	check(CompletionEvent == nullptr);
-	CompletionEvent = FPlatformProcess::GetSynchEventFromPool();
-	check(CompletionEvent != nullptr);
+	// if we are reading again because we had failed earlier, the CompletionEvent may be valid, no need to recreate it
+	if (CompletionEvent == nullptr)
+	{
+		CompletionEvent = FPlatformProcess::GetSynchEventFromPool();
+		check(CompletionEvent != nullptr);
+	}
 
 #if PLATFORM_CAN_ASYNC_PRELOAD_FILES
 	FAsyncFileCallBack SizeCallbackFunction = [this](bool bWasCancelled, IAsyncReadRequest* SizeRequest)
@@ -90,21 +93,26 @@ void FPreLoadFile::KickOffRead()
 		Data = FMemory::Malloc(FileSize);
 		Reader->Serialize(Data, FileSize);
 		delete Reader;
+		CompletionEvent->Trigger();
 	}
 	else
 	{
+		// if we are failing a second time, then just mark the event complete, we have exhausted attempts
+		if (bFailedToOpenInKickOff)
+		{
+			CompletionEvent->Trigger();
+		}
 		// it's possible pak files with the file weren't mounted yet, so note that we didn't find it, and try again in TakeOwnership time
 		bFailedToOpenInKickOff = true;
 
 	}
-	CompletionEvent->Trigger();
 #endif
 }	
 
 
 void* FPreLoadFile::TakeOwnershipOfLoadedData(int64* OutFileSize)
 {
-	// may need to attempt to read again, if (re-)requesting data after the initial boot sequence
+	// may need to attempt to read again, if (re-)requesting data after the initial boot sequence, or if a pak file wasn't mounted in time
 	if (!CompletionEvent || bFailedToOpenInKickOff)
 	{
 		KickOffRead();
@@ -125,8 +133,6 @@ void* FPreLoadFile::TakeOwnershipOfLoadedData(int64* OutFileSize)
 		*OutFileSize = FileSize;
 	}
 
-	checkf(Data != nullptr, TEXT("Failed to find/open/read PreLoaded file %s"), *Path);
-
 	void* ReturnData = Data;
 	Data = nullptr;
 
