- Added a null check to fix the analysis.

#rnx

diff --git a/Engine/Source/Runtime/Core/Private/Windows/WindowsPlatformCrashContext.cpp b/Engine/Source/Runtime/Core/Private/Windows/WindowsPlatformCrashContext.cpp
--- a/Engine/Source/Runtime/Core/Private/Windows/WindowsPlatformCrashContext.cpp
+++ b/Engine/Source/Runtime/Core/Private/Windows/WindowsPlatformCrashContext.cpp
@@ -591,8 +591,11 @@ int32 ReportCrashForMonitor(
 					if (ThreadEntry.th32ThreadID != CurrentThreadId)
 					{
 						HANDLE ThreadHandle = OpenThread(THREAD_SUSPEND_RESUME, FALSE, ThreadEntry.th32ThreadID);
-						SuspendThread(ThreadHandle);
-						ThreadHandles.Push(ThreadHandle);
+						if (ThreadHandle != NULL)
+						{
+							SuspendThread(ThreadHandle);
+							ThreadHandles.Push(ThreadHandle);
+						}
 					}
 
 					SharedContext->ThreadIds[ThreadIdx] = ThreadEntry.th32ThreadID;
