Added and used StringFwd.h for string builder and string view declarations

Including StringFwd.h in place of forward declarations offers more flexibility to change how these types are declared.

[RN] minor Core

diff --git a/Engine/Source/Runtime/Core/Public/Containers/StringFwd.h b/Engine/Source/Runtime/Core/Public/Containers/StringFwd.h
new file mode 100644
--- /dev/null
+++ b/Engine/Source/Runtime/Core/Public/Containers/StringFwd.h
@@ -0,0 +1,19 @@
+// Copyright Epic Games, Inc. All Rights Reserved.
+
+#pragma once
+
+// String
+class FString;
+
+// String Builder
+class FStringBuilderBase;
+class FAnsiStringBuilderBase;
+template <int32 N>
+class TStringBuilder;
+template <int32 N>
+class TAnsiStringBuilder;
+
+// String View
+class FStringView;
+class FAnsiStringView;
+using FWideStringView = FStringView;
diff --git a/Engine/Source/Runtime/Core/Public/Containers/StringView.h b/Engine/Source/Runtime/Core/Public/Containers/StringView.h
--- a/Engine/Source/Runtime/Core/Public/Containers/StringView.h
+++ b/Engine/Source/Runtime/Core/Public/Containers/StringView.h
@@ -2,6 +2,7 @@
 
 #pragma once
 
+#include "Containers/StringFwd.h"
 #include "CoreTypes.h"
 #include "Math/NumericLimits.h"
 #include "Math/UnrealMathUtility.h"
@@ -12,9 +13,6 @@
 #include "Templates/UnrealTemplate.h"
 #include "Traits/IsContiguousContainer.h"
 
-class FAnsiStringView;
-class FStringView;
-
 namespace StringViewPrivate
 {
 	/**
diff --git a/Engine/Source/Runtime/Core/Public/Containers/UnrealString.h b/Engine/Source/Runtime/Core/Public/Containers/UnrealString.h
--- a/Engine/Source/Runtime/Core/Public/Containers/UnrealString.h
+++ b/Engine/Source/Runtime/Core/Public/Containers/UnrealString.h
@@ -16,6 +16,7 @@
 #include "Templates/UnrealTemplate.h"
 #include "Math/NumericLimits.h"
 #include "Containers/Array.h"
+#include "Containers/StringFwd.h"
 #include "Misc/CString.h"
 #include "Misc/Crc.h"
 #include "Math/UnrealMathUtility.h"
@@ -26,8 +27,6 @@
 #include "Templates/TypeHash.h"
 #include "Templates/IsFloatingPoint.h"
 
-class FString;
-class FStringView;
 struct FStringFormatArg;
 template<typename KeyType,typename ValueType,typename SetAllocator ,typename KeyFuncs > class TMap;
 
diff --git a/Engine/Source/Runtime/Core/Public/Misc/ConfigCacheIni.h b/Engine/Source/Runtime/Core/Public/Misc/ConfigCacheIni.h
--- a/Engine/Source/Runtime/Core/Public/Misc/ConfigCacheIni.h
+++ b/Engine/Source/Runtime/Core/Public/Misc/ConfigCacheIni.h
@@ -10,6 +10,7 @@
 #include "Containers/Array.h"
 #include "Containers/UnrealString.h"
 #include "Containers/Map.h"
+#include "Containers/StringFwd.h"
 #include "Math/Color.h"
 #include "UObject/NameTypes.h"
 #include "Logging/LogMacros.h"
@@ -20,8 +21,6 @@
 #include "Misc/Paths.h"
 #include "Serialization/StructuredArchive.h"
 
-class FStringBuilderBase;
-
 CORE_API DECLARE_LOG_CATEGORY_EXTERN(LogConfig, Log, All);
 
 // Server builds should be tweakable even in Shipping
diff --git a/Engine/Source/Runtime/Core/Public/Misc/Parse.h b/Engine/Source/Runtime/Core/Public/Misc/Parse.h
--- a/Engine/Source/Runtime/Core/Public/Misc/Parse.h
+++ b/Engine/Source/Runtime/Core/Public/Misc/Parse.h
@@ -3,10 +3,9 @@
 #pragma once
 
 #include "CoreTypes.h"
+#include "Containers/StringFwd.h"
 #include "Containers/UnrealString.h"
 
-class FStringBuilderBase;
-
 /*-----------------------------------------------------------------------------
 	Parsing functions.
 -----------------------------------------------------------------------------*/
diff --git a/Engine/Source/Runtime/Core/Public/Misc/PathViews.h b/Engine/Source/Runtime/Core/Public/Misc/PathViews.h
--- a/Engine/Source/Runtime/Core/Public/Misc/PathViews.h
+++ b/Engine/Source/Runtime/Core/Public/Misc/PathViews.h
@@ -3,9 +3,7 @@
 #pragma once
 
 #include "CoreTypes.h"
-
-class FStringBuilderBase;
-class FStringView;
+#include "Containers/StringFwd.h"
 
 class CORE_API FPathViews
 {
diff --git a/Engine/Source/Runtime/Core/Public/Misc/SecureHash.h b/Engine/Source/Runtime/Core/Public/Misc/SecureHash.h
--- a/Engine/Source/Runtime/Core/Public/Misc/SecureHash.h
+++ b/Engine/Source/Runtime/Core/Public/Misc/SecureHash.h
@@ -7,6 +7,7 @@
 #include "Containers/UnrealString.h"
 #include "Containers/Map.h"
 #include "Containers/StringConv.h"
+#include "Containers/StringFwd.h"
 #include "Containers/StringView.h"
 #include "Stats/Stats.h"
 #include "Async/AsyncWork.h"
@@ -15,9 +16,6 @@
 #include "String/HexToBytes.h"
 #include "Serialization/MemoryLayout.h"
 
-class FAnsiStringBuilderBase;
-class FStringBuilderBase;
-class FStringView;
 struct FMD5Hash;
 
 /*-----------------------------------------------------------------------------
diff --git a/Engine/Source/Runtime/Core/Public/Misc/StringBuilder.h b/Engine/Source/Runtime/Core/Public/Misc/StringBuilder.h
--- a/Engine/Source/Runtime/Core/Public/Misc/StringBuilder.h
+++ b/Engine/Source/Runtime/Core/Public/Misc/StringBuilder.h
@@ -2,8 +2,9 @@
 
 #pragma once
 
-#include "Containers/StringView.h"
 #include "CoreTypes.h"
+#include "Containers/StringFwd.h"
+#include "Containers/StringView.h"
 #include "Misc/CString.h"
 #include "Templates/AndOrNot.h"
 #include "Templates/EnableIf.h"
@@ -14,9 +15,6 @@
 
 #define USE_STRING_LITERAL_PATH 1
 
-class FAnsiStringBuilderBase;
-class FStringBuilderBase;
-
 namespace StringBuilderPrivate
 {
 	template <typename CharType> struct TStringBuilderBaseType;
diff --git a/Engine/Source/Runtime/Core/Public/String/BytesToHex.h b/Engine/Source/Runtime/Core/Public/String/BytesToHex.h
--- a/Engine/Source/Runtime/Core/Public/String/BytesToHex.h
+++ b/Engine/Source/Runtime/Core/Public/String/BytesToHex.h
@@ -4,9 +4,7 @@
 
 #include "CoreTypes.h"
 #include "Containers/ArrayView.h"
-
-class FStringBuilderBase;
-class FAnsiStringBuilderBase;
+#include "Containers/StringFwd.h"
 
 namespace UE
 {
diff --git a/Engine/Source/Runtime/Core/Public/String/HexToBytes.h b/Engine/Source/Runtime/Core/Public/String/HexToBytes.h
--- a/Engine/Source/Runtime/Core/Public/String/HexToBytes.h
+++ b/Engine/Source/Runtime/Core/Public/String/HexToBytes.h
@@ -3,9 +3,7 @@
 #pragma once
 
 #include "CoreTypes.h"
-
-class FStringView;
-class FAnsiStringView;
+#include "Containers/StringFwd.h"
 
 namespace UE
 {
diff --git a/Engine/Source/Runtime/Core/Public/String/LexFromString.h b/Engine/Source/Runtime/Core/Public/String/LexFromString.h
--- a/Engine/Source/Runtime/Core/Public/String/LexFromString.h
+++ b/Engine/Source/Runtime/Core/Public/String/LexFromString.h
@@ -3,8 +3,7 @@
 #pragma once
 
 #include "CoreTypes.h"
-
-class FStringView;
+#include "Containers/StringFwd.h"
 
 CORE_API void LexFromString(int8& OutValue,   const FStringView& InString);
 CORE_API void LexFromString(int16& OutValue,  const FStringView& InString);
diff --git a/Engine/Source/Runtime/Core/Public/String/ParseLines.h b/Engine/Source/Runtime/Core/Public/String/ParseLines.h
--- a/Engine/Source/Runtime/Core/Public/String/ParseLines.h
+++ b/Engine/Source/Runtime/Core/Public/String/ParseLines.h
@@ -3,10 +3,9 @@
 #pragma once
 
 #include "CoreTypes.h"
+#include "Containers/StringFwd.h"
 #include "Templates/Function.h"
 
-class FStringView;
-
 namespace UE
 {
 namespace String
diff --git a/Engine/Source/Runtime/Core/Public/String/ParseTokens.h b/Engine/Source/Runtime/Core/Public/String/ParseTokens.h
--- a/Engine/Source/Runtime/Core/Public/String/ParseTokens.h
+++ b/Engine/Source/Runtime/Core/Public/String/ParseTokens.h
@@ -4,10 +4,9 @@
 
 #include "CoreTypes.h"
 #include "Containers/ArrayView.h"
+#include "Containers/StringFwd.h"
 #include "Templates/Function.h"
 
-class FStringView;
-
 namespace UE
 {
 namespace String
diff --git a/Engine/Source/Runtime/Core/Public/UObject/NameTypes.h b/Engine/Source/Runtime/Core/Public/UObject/NameTypes.h
--- a/Engine/Source/Runtime/Core/Public/UObject/NameTypes.h
+++ b/Engine/Source/Runtime/Core/Public/UObject/NameTypes.h
@@ -10,6 +10,7 @@
 #include "Containers/UnrealString.h"
 #include "HAL/CriticalSection.h"
 #include "Containers/StringConv.h"
+#include "Containers/StringFwd.h"
 #include "UObject/UnrealNames.h"
 #include "Templates/Atomic.h"
 #include "Serialization/MemoryLayout.h"
@@ -28,10 +29,6 @@
 	#define WITH_CASE_PRESERVING_NAME WITH_EDITORONLY_DATA
 #endif
 
-class FAnsiStringBuilderBase;
-class FAnsiStringView;
-class FStringBuilderBase;
-class FStringView;
 class FText;
 
 /** Maximum size of name. */
diff --git a/Engine/Source/Runtime/CoreUObject/Public/Misc/PackageName.h b/Engine/Source/Runtime/CoreUObject/Public/Misc/PackageName.h
--- a/Engine/Source/Runtime/CoreUObject/Public/Misc/PackageName.h
+++ b/Engine/Source/Runtime/CoreUObject/Public/Misc/PackageName.h
@@ -7,9 +7,8 @@
 #pragma once
 
 #include "CoreMinimal.h"
+#include "Containers/StringFwd.h"
 
-class FStringBuilderBase;
-class FStringView;
 struct FFileStatData;
 
 class COREUOBJECT_API FPackageName
diff --git a/Engine/Source/Runtime/RHI/Public/RHI.h b/Engine/Source/Runtime/RHI/Public/RHI.h
--- a/Engine/Source/Runtime/RHI/Public/RHI.h
+++ b/Engine/Source/Runtime/RHI/Public/RHI.h
@@ -10,6 +10,7 @@
 #include "Stats/Stats.h"
 #include "RHIDefinitions.h"
 #include "Containers/StaticArray.h"
+#include "Containers/StringFwd.h"
 
 #ifndef RHI_COMMAND_LIST_DEBUG_TRACES
 #define RHI_COMMAND_LIST_DEBUG_TRACES 0
@@ -19,7 +20,6 @@ template <typename T> class TArrayView;
 
 class FResourceArrayInterface;
 class FResourceBulkDataInterface;
-class FStringView;
 
 /** Alignment of the shader parameters struct is required to be 16-byte boundaries. */
 #define SHADER_PARAMETER_STRUCT_ALIGNMENT 16
diff --git a/Engine/Source/Runtime/RenderCore/Public/ShaderCodeLibrary.h b/Engine/Source/Runtime/RenderCore/Public/ShaderCodeLibrary.h
--- a/Engine/Source/Runtime/RenderCore/Public/ShaderCodeLibrary.h
+++ b/Engine/Source/Runtime/RenderCore/Public/ShaderCodeLibrary.h
@@ -8,13 +8,12 @@
 
 #include "CoreMinimal.h"
 #include "Containers/ArrayView.h"
+#include "Containers/StringFwd.h"
 #include "RHI.h"
 
 DECLARE_LOG_CATEGORY_EXTERN(LogShaderLibrary, Log, All);
 
-class FAnsiStringBuilderBase;
 class FShaderPipeline;
-class FStringView;
 class FShaderMapResource;
 class FShaderMapResourceCode;
 
