Re enable the very large page cache
Add some stats tracking for how much slack is held in the large page cache
Add a way for some allocations to pass hint flags down to underlying allocators/caches

diff --git a/Engine/Source/Runtime/Core/Private/HAL/Allocators/CachedOSVeryLargePageAllocator.cpp b/Engine/Source/Runtime/Core/Private/HAL/Allocators/CachedOSVeryLargePageAllocator.cpp
--- a/Engine/Source/Runtime/Core/Private/HAL/Allocators/CachedOSVeryLargePageAllocator.cpp
+++ b/Engine/Source/Runtime/Core/Private/HAL/Allocators/CachedOSVeryLargePageAllocator.cpp
@@ -17,83 +17,93 @@ void FCachedOSVeryLargePageAllocator::Init()
 
 
 	FreeLargePagesHead = nullptr;
-	for (int i = 0; i < AddressSpaceToReserve / SizeOfLargePage; i++)
+	for (int i = 0; i < NumberOfLargePages; i++)
 	{
 		LargePagesArray[i].Init((void*)((uintptr_t)AddressSpaceReserved + (i*SizeOfLargePage)));
 		LargePagesArray[i].LinkHead(FreeLargePagesHead);
 	}
 
 	UsedLargePagesHead = nullptr;
-	UsedLargePagesWithSpaceHead = nullptr;
-
+	for (int i = 0; i < FMemory::AllocationHints::Max; i++)
+	{
+		UsedLargePagesWithSpaceHead[i] = nullptr;
+	}
 	if (!GEnableVeryLargePageAllocator)
 	{
 		bEnabled = false;
 	}
 }
 
-void* FCachedOSVeryLargePageAllocator::Allocate(SIZE_T Size)
+void* FCachedOSVeryLargePageAllocator::Allocate(SIZE_T Size, uint32 AllocationHint)
 {
 	Size = Align(Size, 4096);
 
 	void* ret = nullptr;
 
-	if (bEnabled && Size == SizeOfSubPage)
+	if (bEnabled && Size == SizeOfSubPage && AllocationHint == FMemory::AllocationHints::SmallPool)
 	{
 
-		if (UsedLargePagesWithSpaceHead == nullptr)
+		if (UsedLargePagesWithSpaceHead[AllocationHint] == nullptr)
 		{
 			FLargePage* LargePage = FreeLargePagesHead;
 			if (LargePage)
 			{
 				Block.Commit(LargePage->BaseAddress - AddressSpaceReserved, SizeOfLargePage);
+				LargePage->AllocationHint = AllocationHint;
 				LargePage->Unlink();
-				LargePage->LinkHead(UsedLargePagesWithSpaceHead);
+				LargePage->LinkHead(UsedLargePagesWithSpaceHead[AllocationHint]);
+				CachedFree += SizeOfLargePage;
 			}
 		}
-		if (UsedLargePagesWithSpaceHead != nullptr)
+		FLargePage* LargePage = UsedLargePagesWithSpaceHead[AllocationHint];
+		if (LargePage != nullptr)
 		{
-			FLargePage* LargePage = UsedLargePagesWithSpaceHead;
 			ret = LargePage->Allocate();
-			if (ret && LargePage->NumberOfFreeSubPages == 0)
+			if (ret)
 			{
-				LargePage->Unlink();
-				LargePage->LinkHead(UsedLargePagesHead);
+				if (LargePage->NumberOfFreeSubPages == 0)
+				{
+					LargePage->Unlink();
+					LargePage->LinkHead(UsedLargePagesHead);
+				}
+				CachedFree -= SizeOfSubPage;
 			}
 		}
 	}
 
-	if( ret == nullptr )
+	if (ret == nullptr)
 	{
 		ret = CachedOSPageAllocator.Allocate(Size);
 	}
 	return ret;
 }
 
-#define LARGEPAGEALLOCATOR_SORT_OnAddress 0
+#define LARGEPAGEALLOCATOR_SORT_OnAddress 1
 
 void FCachedOSVeryLargePageAllocator::Free(void* Ptr, SIZE_T Size)
 {
 	Size = Align(Size, 4096);
 	uint64 Index = ((uintptr_t)Ptr - (uintptr_t)AddressSpaceReserved) / SizeOfLargePage;
-	if (Index < (AddressSpaceToReserve / SizeOfLargePage))
+	if (Index < (NumberOfLargePages))
 	{
 		FLargePage* LargePage = &LargePagesArray[Index];
 
 		LargePage->Free(Ptr);
+		CachedFree += SizeOfSubPage;
 
-		if (LargePage->NumberOfFreeSubPages == (SizeOfLargePage / SizeOfSubPage))
+		if (LargePage->NumberOfFreeSubPages == NumberOfSubPagesPerLargePage)
 		{
 			// totally free, need to move which list we are in and remove the backing store
 			LargePage->Unlink();
 			LargePage->LinkHead(FreeLargePagesHead);
 			Block.Decommit(LargePage->BaseAddress - AddressSpaceReserved, SizeOfLargePage);
+			CachedFree -= SizeOfLargePage;
 		}
 		else if (LargePage->NumberOfFreeSubPages == 1)
 		{
 			LargePage->Unlink();
 #if LARGEPAGEALLOCATOR_SORT_OnAddress
-			FLargePage* InsertPoint = UsedLargePagesWithSpaceHead;
+			FLargePage* InsertPoint = UsedLargePagesWithSpaceHead[LargePage->AllocationHint];
 			while (InsertPoint != nullptr)
 			{
 				if (LargePage->BaseAddress < InsertPoint->BaseAddress)	// sort on address
@@ -102,16 +112,16 @@ void FCachedOSVeryLargePageAllocator::Free(void* Ptr, SIZE_T Size)
 				}
 				InsertPoint = InsertPoint->Next();
 			}
-			if (InsertPoint == nullptr || InsertPoint == UsedLargePagesWithSpaceHead)
+			if (InsertPoint == nullptr || InsertPoint == UsedLargePagesWithSpaceHead[LargePage->AllocationHint])
 			{
-				LargePage->LinkHead(UsedLargePagesWithSpaceHead);
+				LargePage->LinkHead(UsedLargePagesWithSpaceHead[LargePage->AllocationHint]);
 			}
 			else
 			{
 				LargePage->LinkBefore(InsertPoint);
 			}
 #else
-			LargePage->LinkHead(UsedLargePagesWithSpaceHead);
+			LargePage->LinkHead(UsedLargePagesWithSpaceHead[LargePage->AllocationHint]);
 #endif
 		}
 		else
diff --git a/Engine/Source/Runtime/Core/Public/HAL/Allocators/CachedOSPageAllocator.h b/Engine/Source/Runtime/Core/Public/HAL/Allocators/CachedOSPageAllocator.h
--- a/Engine/Source/Runtime/Core/Public/HAL/Allocators/CachedOSPageAllocator.h
+++ b/Engine/Source/Runtime/Core/Public/HAL/Allocators/CachedOSPageAllocator.h
@@ -33,7 +33,7 @@ struct TCachedOSPageAllocator : private FCachedOSPageAllocator
 	{
 	}
 
-	FORCEINLINE void* Allocate(SIZE_T Size)
+	FORCEINLINE void* Allocate(SIZE_T Size, uint32 AllocationHint = 0)
 	{
 		return AllocateImpl(Size, CachedByteLimit, FreedPageBlocks, FreedPageBlocks + FreedPageBlocksNum, FreedPageBlocksNum, CachedTotal);
 	}
diff --git a/Engine/Source/Runtime/Core/Public/HAL/Allocators/CachedOSVeryLargePageAllocator.h b/Engine/Source/Runtime/Core/Public/HAL/Allocators/CachedOSVeryLargePageAllocator.h
--- a/Engine/Source/Runtime/Core/Public/HAL/Allocators/CachedOSVeryLargePageAllocator.h
+++ b/Engine/Source/Runtime/Core/Public/HAL/Allocators/CachedOSVeryLargePageAllocator.h
@@ -8,6 +8,7 @@
 #include "HAL/Allocators/CachedOSPageAllocator.h"
 #include "HAL/PlatformMemory.h"
 
+
 #if UE_USE_VERYLARGEPAGEALLOCATOR
 
 #if PLATFORM_64BITS
@@ -23,10 +24,13 @@ class FCachedOSVeryLargePageAllocator
 	static const uint64 AddressSpaceToReserve = ((1024 * 1024) * (1024 + 512));
 	static const uint64 SizeOfLargePage = (1024 * 1024 * 4);
 	static const uint64 SizeOfSubPage = (1024 * 64);
+	static const uint64 NumberOfLargePages = (AddressSpaceToReserve / SizeOfLargePage);
+	static const uint64 NumberOfSubPagesPerLargePage = (SizeOfLargePage / SizeOfSubPage);
 public:
 
 	FCachedOSVeryLargePageAllocator()
 		: bEnabled(true)
+		, CachedFree(0)
 	{
 		Init();
 	}
@@ -36,7 +40,7 @@ public:
 		// this leaks everything!
 	}
 
-	void* Allocate(SIZE_T Size);
+	void* Allocate(SIZE_T Size, uint32 AllocationHint = 0);
 
 	void Free(void* Ptr, SIZE_T Size);
 
@@ -44,7 +48,7 @@ public:
 
 	uint64 GetCachedFreeTotal()
 	{
-		return CachedOSPageAllocator.GetCachedFreeTotal();
+		return CachedFree + CachedOSPageAllocator.GetCachedFreeTotal();
 	}
 
 private:
@@ -53,20 +57,22 @@ private:
 
 	bool bEnabled;
 	uintptr_t	AddressSpaceReserved;
+	uint64		CachedFree;
 
 	FPlatformMemory::FPlatformVirtualMemoryBlock Block;
 
 	struct FLargePage : public TIntrusiveLinkedList<FLargePage>
 	{
-		uintptr_t	FreeSubPages[SizeOfLargePage / SizeOfSubPage];
+		uintptr_t	FreeSubPages[NumberOfSubPagesPerLargePage];
 		int32		NumberOfFreeSubPages;
+		uint32		AllocationHint;
 
 		uintptr_t	BaseAddress;
 
 		void Init(void* InBaseAddress)
 		{
 			BaseAddress = (uintptr_t)InBaseAddress;
-			NumberOfFreeSubPages = (SizeOfLargePage / SizeOfSubPage);
+			NumberOfFreeSubPages = NumberOfSubPagesPerLargePage;
 			uintptr_t Ptr = BaseAddress;
 			for (int i = 0; i < NumberOfFreeSubPages; i++)
 			{
@@ -95,9 +101,9 @@ private:
 
 	FLargePage*	UsedLargePagesHead;				// has backing store and is full
 
-	FLargePage*	UsedLargePagesWithSpaceHead;	// has backing store and still has room
+	FLargePage*	UsedLargePagesWithSpaceHead[FMemory::AllocationHints::Max];	// has backing store and still has room
 
-	FLargePage	LargePagesArray[AddressSpaceToReserve / SizeOfLargePage];
+	FLargePage	LargePagesArray[NumberOfLargePages];
 
 	TCachedOSPageAllocator<CACHEDOSVERYLARGEPAGEALLOCATOR_MAX_CACHED_OS_FREES, CACHEDOSVERYLARGEPAGEALLOCATOR_BYTE_LIMIT> CachedOSPageAllocator;
 
diff --git a/Engine/Source/Runtime/Core/Public/HAL/UnrealMemory.h b/Engine/Source/Runtime/Core/Public/HAL/UnrealMemory.h
--- a/Engine/Source/Runtime/Core/Public/HAL/UnrealMemory.h
+++ b/Engine/Source/Runtime/Core/Public/HAL/UnrealMemory.h
@@ -80,14 +80,24 @@ struct FScopedMallocTimer
 #endif
 
 
-
-
 /*-----------------------------------------------------------------------------
 	FMemory.
 -----------------------------------------------------------------------------*/
 
 struct CORE_API FMemory
 {
+	/** Some allocators can be given hints to treat allocations differently depending on how the memory is used, it's lifetime etc. */
+	enum AllocationHints
+	{
+		None = -1,
+		Default,
+		Temporary,
+		SmallPool,
+
+		Max
+	};
+
+
 	/** @name Memory functions (wrapper for FPlatformMemory) */
 
 	static FORCEINLINE void* Memmove( void* Dest, const void* Src, SIZE_T Count )
