Fix handling of UNC paths on Windows

- Add the \\?\UNC\ prefix when necessary to support long paths.
- Remove the \\?\UNC\ prefix from UNC paths to avoid paths like //host/path/file.ext being translated to UNC/host/path/file.ext.

[RN] fix Platforms/Windows

diff --git a/Engine/Source/Runtime/Core/Private/Windows/WindowsPlatformFile.cpp b/Engine/Source/Runtime/Core/Private/Windows/WindowsPlatformFile.cpp
--- a/Engine/Source/Runtime/Core/Private/Windows/WindowsPlatformFile.cpp
+++ b/Engine/Source/Runtime/Core/Private/Windows/WindowsPlatformFile.cpp
@@ -818,12 +818,16 @@ private:
 		Result.ReplaceCharInline(TEXT('/'), TEXT('\\'), ESearchCase::CaseSensitive);
 
 		// Handle Windows Path length over MAX_PATH
-		if (!bIsUNCPath && Result.Len() > MAX_PATH)
+		if (Result.Len() > MAX_PATH)
 		{
-			// Pathnames exceeding MAX_PATH are supported by WindowsAPI functions if you add the \\?\ prefix,
-			// This does not work on \\UNC paths in all WindowsAPI functions however, so only do this for D:\ paths.
-			// Our engine does not currently handle UNC paths greater than MAXPATH; we will need to talk to Microsoft if we want to handle that robustly.
-			Result = TEXT("\\\\?\\") + Result;
+			if (bIsUNCPath)
+			{
+				Result = TEXT("\\\\?\\UNC") + Result.RightChop(1);
+			}
+			else
+			{
+				Result = TEXT("\\\\?\\") + Result;
+			}
 		}
 
 		return Result;
@@ -955,8 +959,18 @@ public:
 			CloseHandle(hFile);
 			TRACE_PLATFORMFILE_END_CLOSE();
 		}
-		// Convert the result back into an UnrealPath(\\ -> / )
-		NormalizedFileName.RemoveFromStart(TEXT("\\\\?\\"), ESearchCase::CaseSensitive);
+
+		// Remove the Windows device path prefix.
+		if (NormalizedFileName.StartsWith(TEXT("\\\\?\\UNC\\"), ESearchCase::CaseSensitive))
+		{
+			NormalizedFileName.RemoveAt(2, 6); // remove ?\UNC\ to convert \\?\UNC\Path\... to \\Path\...
+		}
+		else
+		{
+			NormalizedFileName.RemoveFromStart(TEXT("\\\\?\\"), ESearchCase::CaseSensitive);
+		}
+
+		// Convert the result back into an UnrealPath (\\ -> /)
 		NormalizedFileName.ReplaceCharInline(TEXT('\\'), TEXT('/'), ESearchCase::CaseSensitive);
 
 		return NormalizedFileName;
