Fix static analysis warning from StringFwd.h

#rnx

diff --git a/Engine/Source/Runtime/Core/Public/Containers/StringFwd.h b/Engine/Source/Runtime/Core/Public/Containers/StringFwd.h
--- a/Engine/Source/Runtime/Core/Public/Containers/StringFwd.h
+++ b/Engine/Source/Runtime/Core/Public/Containers/StringFwd.h
@@ -4,9 +4,6 @@
 
 #include "CoreTypes.h"
 
-// String
-class FString;
-
 // String Builder
 class FStringBuilderBase;
 class FAnsiStringBuilderBase;
