- Implement missing DirectX Math functions to improve VectorVM performance

#rnx

diff --git a/Engine/Source/Runtime/Core/Public/Math/UnrealMathDirectX.h b/Engine/Source/Runtime/Core/Public/Math/UnrealMathDirectX.h
--- a/Engine/Source/Runtime/Core/Public/Math/UnrealMathDirectX.h
+++ b/Engine/Source/Runtime/Core/Public/Math/UnrealMathDirectX.h
@@ -970,73 +970,59 @@ FORCEINLINE bool VectorContainsNaNOrInfinite(const VectorRegister& Vec)
 	return !IsFinite;
 }
 
-//TODO: Vectorize
 FORCEINLINE VectorRegister VectorExp(const VectorRegister& X)
 {
-	return MakeVectorRegister(FMath::Exp(VectorGetComponent(X, 0)), FMath::Exp(VectorGetComponent(X, 1)), FMath::Exp(VectorGetComponent(X, 2)), FMath::Exp(VectorGetComponent(X, 3)));
+	return DirectX::XMVectorExp(X);
 }
 
-//TODO: Vectorize
 FORCEINLINE VectorRegister VectorExp2(const VectorRegister& X)
 {
 	return DirectX::XMVectorExp2(X);
 }
 
-//TODO: Vectorize
 FORCEINLINE VectorRegister VectorLog(const VectorRegister& X)
 {
-	return MakeVectorRegister(FMath::Loge(VectorGetComponent(X, 0)), FMath::Loge(VectorGetComponent(X, 1)), FMath::Loge(VectorGetComponent(X, 2)), FMath::Loge(VectorGetComponent(X, 3)));
+	return DirectX::XMVectorLog(X);
 }
 
-//TODO: Vectorize
 FORCEINLINE VectorRegister VectorLog2(const VectorRegister& X)
 {
 	return DirectX::XMVectorLog2(X);
 }
 
-//TODO: Vectorize
 FORCEINLINE VectorRegister VectorSin(const VectorRegister& X)
 {
-	return MakeVectorRegister(FMath::Sin(VectorGetComponent(X, 0)), FMath::Sin(VectorGetComponent(X, 1)), FMath::Sin(VectorGetComponent(X, 2)), FMath::Sin(VectorGetComponent(X, 3)));
+	return DirectX::XMVectorSin(X);
 }
 
-//TODO: Vectorize
 FORCEINLINE VectorRegister VectorCos(const VectorRegister& X)
 {
-	return MakeVectorRegister(FMath::Cos(VectorGetComponent(X, 0)), FMath::Cos(VectorGetComponent(X, 1)), FMath::Cos(VectorGetComponent(X, 2)), FMath::Cos(VectorGetComponent(X, 3)));
+	return DirectX::XMVectorCos(X);
 }
 
-//TODO: Vectorize
 FORCEINLINE VectorRegister VectorTan(const VectorRegister& X)
 {
-	return MakeVectorRegister(FMath::Tan(VectorGetComponent(X, 0)), FMath::Tan(VectorGetComponent(X, 1)), FMath::Tan(VectorGetComponent(X, 2)), FMath::Tan(VectorGetComponent(X, 3)));
+	return DirectX::XMVectorTan(X);
 }
 
-//TODO: Vectorize
 FORCEINLINE VectorRegister VectorASin(const VectorRegister& X)
 {
-	return MakeVectorRegister(FMath::Asin(VectorGetComponent(X, 0)), FMath::Asin(VectorGetComponent(X, 1)), FMath::Asin(VectorGetComponent(X, 2)), FMath::Asin(VectorGetComponent(X, 3)));
+	return DirectX::XMVectorASin(X);
 }
 
-//TODO: Vectorize
 FORCEINLINE VectorRegister VectorACos(const VectorRegister& X)
 {
-	return MakeVectorRegister(FMath::Acos(VectorGetComponent(X, 0)), FMath::Acos(VectorGetComponent(X, 1)), FMath::Acos(VectorGetComponent(X, 2)), FMath::Acos(VectorGetComponent(X, 3)));
+	return DirectX::XMVectorACos(X);
 }
 
-//TODO: Vectorize
 FORCEINLINE VectorRegister VectorATan(const VectorRegister& X)
 {
-	return MakeVectorRegister(FMath::Atan(VectorGetComponent(X, 0)), FMath::Atan(VectorGetComponent(X, 1)), FMath::Atan(VectorGetComponent(X, 2)), FMath::Atan(VectorGetComponent(X, 3)));
+	return DirectX::XMVectorATan(X);
 }
 
-//TODO: Vectorize
 FORCEINLINE VectorRegister VectorATan2(const VectorRegister& X, const VectorRegister& Y)
 {
-	return MakeVectorRegister(FMath::Atan2(VectorGetComponent(X, 0), VectorGetComponent(Y, 0)),
-		FMath::Atan2(VectorGetComponent(X, 1), VectorGetComponent(Y, 1)),
-		FMath::Atan2(VectorGetComponent(X, 2), VectorGetComponent(Y, 2)),
-		FMath::Atan2(VectorGetComponent(X, 3), VectorGetComponent(Y, 3)));
+	return DirectX::XMVectorATan2(X, Y);
 }
 
 FORCEINLINE VectorRegister VectorCeil(const VectorRegister& X)
