Read Command line from documents folder in Shipping TestFlight builds

This should allow external QA to set command line options in TestFlight builds without letting players in live do the same.

Test plan:
Ran a Test build, confirmed it pulled my ue4commandline.txt from the documents folder

diff --git a/Engine/Source/Runtime/Core/Public/IOS/IOSCommandLineHelper.h b/Engine/Source/Runtime/Core/Public/IOS/IOSCommandLineHelper.h
--- a/Engine/Source/Runtime/Core/Public/IOS/IOSCommandLineHelper.h
+++ b/Engine/Source/Runtime/Core/Public/IOS/IOSCommandLineHelper.h
@@ -92,11 +92,17 @@ class FIOSCommandLineHelper
 			// command line text file included in the bundle
 			FString BundleCommandLineFilePath = FString([[NSBundle mainBundle] bundlePath]) + TEXT("/ue4commandline.txt");
 
-#if !UE_BUILD_SHIPPING
 			// command line text file pushed to the documents folder
 			FString DocumentsCommandLineFilePath = FString([NSSearchPathForDirectoriesInDomains(NSDocumentDirectory, NSUserDomainMask, YES) objectAtIndex:0]) + TEXT("/ue4commandline.txt");
-			if (!TryReadCommandLineFile(DocumentsCommandLineFilePath))
+			
+#if UE_BUILD_SHIPPING
+			bool bIsTestFlightOrInternal = [[[[NSBundle mainBundle] appStoreReceiptURL] lastPathComponent] isEqualToString: @"sandboxReceipt"];
+#else
+			bool bIsTestFlightOrInternal = YES;
 #endif
+			
+			bool bFoundDocumentsCommandLine = bIsTestFlightOrInternal && TryReadCommandLineFile(DocumentsCommandLineFilePath);
+			if (!bFoundDocumentsCommandLine)
 			{
 				TryReadCommandLineFile(BundleCommandLineFilePath);
 			}
