Support VirtualTextureStreaming for Linux

diff --git a/Engine/Source/Runtime/Core/Public/Linux/LinuxPlatformProperties.h b/Engine/Source/Runtime/Core/Public/Linux/LinuxPlatformProperties.h
--- a/Engine/Source/Runtime/Core/Public/Linux/LinuxPlatformProperties.h
+++ b/Engine/Source/Runtime/Core/Public/Linux/LinuxPlatformProperties.h
@@ -159,6 +159,12 @@ struct FLinuxPlatformProperties
 	{
 		return !IsServerOnly();
 	}
+
+	static FORCEINLINE bool SupportsVirtualTextureStreaming()
+	{
+		return true;
+	}
+
 };
 
 #ifdef PROPERTY_HEADER_SHOULD_DEFINE_TYPE
