Fix for crashreportclient.ini not included in cooked builds

diff --git a/Engine/Source/Runtime/Core/Private/GenericPlatform/GenericPlatformCrashContext.cpp b/Engine/Source/Runtime/Core/Private/GenericPlatform/GenericPlatformCrashContext.cpp
--- a/Engine/Source/Runtime/Core/Private/GenericPlatform/GenericPlatformCrashContext.cpp
+++ b/Engine/Source/Runtime/Core/Private/GenericPlatform/GenericPlatformCrashContext.cpp
@@ -888,6 +888,7 @@ const TCHAR* FGenericCrashContext::GetCrashConfigFilePath()
 	if (FCString::Strlen(NCached::Session.CrashConfigFilePath) == 0)
 	{
 		FString CrashConfigFilePath = FPaths::Combine(GetCrashConfigFolder(), NCached::Session.CrashGUIDRoot, FGenericCrashContext::CrashConfigFileNameW);
+		CrashConfigFilePath = IFileManager::Get().ConvertToAbsolutePathForExternalAppForRead(*CrashConfigFilePath);
 		NCached::Set(NCached::Session.CrashConfigFilePath, *CrashConfigFilePath);
 	}
 	return NCached::Session.CrashConfigFilePath;
