Implement GetOnDemandChunksForPakchunkIndices()

diff --git a/Engine/Source/Runtime/Core/Private/GenericPlatform/GenericPlatformMisc.cpp b/Engine/Source/Runtime/Core/Private/GenericPlatform/GenericPlatformMisc.cpp
--- a/Engine/Source/Runtime/Core/Private/GenericPlatform/GenericPlatformMisc.cpp
+++ b/Engine/Source/Runtime/Core/Private/GenericPlatform/GenericPlatformMisc.cpp
@@ -1415,6 +1415,11 @@ bool FGenericPlatformMisc::RequestDeviceCheckToken(TFunction<void(const TArray<u
 	return false;
 }
 
+TArray<FCustomChunk> FGenericPlatformMisc::GetOnDemandChunksForPakchunkIndices(const TArray<int32>& PakchunkIndices)
+{
+	return TArray<FCustomChunk>();
+}
+
 TArray<FCustomChunk> FGenericPlatformMisc::GetAllOnDemandChunks()
 {
 	return TArray<FCustomChunk>();
diff --git a/Engine/Source/Runtime/Core/Public/GenericPlatform/GenericPlatformMisc.h b/Engine/Source/Runtime/Core/Public/GenericPlatform/GenericPlatformMisc.h
--- a/Engine/Source/Runtime/Core/Public/GenericPlatform/GenericPlatformMisc.h
+++ b/Engine/Source/Runtime/Core/Public/GenericPlatform/GenericPlatformMisc.h
@@ -1382,6 +1382,7 @@ public:
 
 	static bool RequestDeviceCheckToken(TFunction<void(const TArray<uint8>&)> QuerySucceededFunc, TFunction<void(const FString&, const FString&)> QueryFailedFunc);
 
+	static TArray<FCustomChunk> GetOnDemandChunksForPakchunkIndices(const TArray<int32>& PakchunkIndices);
 	static TArray<FCustomChunk> GetAllOnDemandChunks();
 	static TArray<FCustomChunk> GetAllLanguageChunks();
 	static TArray<FCustomChunk> GetCustomChunksByType(ECustomChunkType DesiredChunkType);
