Added support for performing text substitution for fragments within a culture display name

#rnx

diff --git a/Engine/Source/Runtime/Core/Private/Internationalization/Culture.cpp b/Engine/Source/Runtime/Core/Private/Internationalization/Culture.cpp
--- a/Engine/Source/Runtime/Core/Private/Internationalization/Culture.cpp
+++ b/Engine/Source/Runtime/Core/Private/Internationalization/Culture.cpp
@@ -8,6 +8,59 @@
 #include "Internationalization/LegacyCulture.h"
 #endif
 
+void ApplyCultureDisplayNameSubstitutes(FString& InOutDisplayName)
+{
+	static TArray<TTuple<FString, FString>> CultureDisplayNameSubstitutes;
+
+	// Conditionally load the required config data
+	{
+		static bool bHasInitializedCultureDisplayNameSubstitutes = false;
+		if (!bHasInitializedCultureDisplayNameSubstitutes && GConfig && GConfig->IsReadyForUse())
+		{
+			bHasInitializedCultureDisplayNameSubstitutes = true;
+
+			TArray<FString> CultureDisplayNameSubstitutesStrArray;
+			{
+				const bool ShouldLoadEditor = GIsEditor;
+				const bool ShouldLoadGame = FApp::IsGame();
+
+				GConfig->GetArray(TEXT("Internationalization"), TEXT("CultureDisplayNameSubstitutes"), CultureDisplayNameSubstitutesStrArray, GEngineIni);
+
+				if (ShouldLoadEditor)
+				{
+					TArray<FString> EditorArray;
+					GConfig->GetArray(TEXT("Internationalization"), TEXT("CultureDisplayNameSubstitutes"), EditorArray, GEditorIni);
+					CultureDisplayNameSubstitutesStrArray.Append(MoveTemp(EditorArray));
+				}
+
+				if (ShouldLoadGame)
+				{
+					TArray<FString> GameArray;
+					GConfig->GetArray(TEXT("Internationalization"), TEXT("CultureDisplayNameSubstitutes"), GameArray, GGameIni);
+					CultureDisplayNameSubstitutesStrArray.Append(MoveTemp(GameArray));
+				}
+			}
+
+			// Each substitute should be a semi-colon separated pair of data: Old;New
+			CultureDisplayNameSubstitutes.Reserve(CultureDisplayNameSubstitutesStrArray.Num());
+			for (const FString& CultureDisplayNameSubstituteStr : CultureDisplayNameSubstitutesStrArray)
+			{
+				FString OldDisplayFragment;
+				FString NewDisplayFragment;
+				if (CultureDisplayNameSubstituteStr.Split(TEXT(";"), &OldDisplayFragment, &NewDisplayFragment, ESearchCase::CaseSensitive))
+				{
+					CultureDisplayNameSubstitutes.Add(MakeTuple(MoveTemp(OldDisplayFragment), MoveTemp(NewDisplayFragment)));
+				}
+			}
+		}
+	}
+
+	for (const auto& CultureDisplayNameSubstitutePair : CultureDisplayNameSubstitutes)
+	{
+		InOutDisplayName.ReplaceInline(*CultureDisplayNameSubstitutePair.Key, *CultureDisplayNameSubstitutePair.Value, ESearchCase::CaseSensitive);
+	}
+}
+
 FCultureRef FCulture::Create(TUniquePtr<FCultureImplementation>&& InImplementation)
 {
 	check(InImplementation.IsValid());
@@ -16,20 +69,16 @@ FCultureRef FCulture::Create(TUniquePtr<FCultureImplementation>&& InImplementati
 
 FCulture::FCulture(TUniquePtr<FCultureImplementation>&& InImplementation)
 	: Implementation(MoveTemp(InImplementation))
-	, CachedDisplayName(Implementation->GetDisplayName())
-	, CachedEnglishName(Implementation->GetEnglishName())
 	, CachedName(Implementation->GetName())
-	, CachedNativeName(Implementation->GetNativeName())
 	, CachedUnrealLegacyThreeLetterISOLanguageName(Implementation->GetUnrealLegacyThreeLetterISOLanguageName())
 	, CachedThreeLetterISOLanguageName(Implementation->GetThreeLetterISOLanguageName())
 	, CachedTwoLetterISOLanguageName(Implementation->GetTwoLetterISOLanguageName())
-	, CachedNativeLanguage(Implementation->GetNativeLanguage())
 	, CachedRegion(Implementation->GetRegion())
-	, CachedNativeRegion(Implementation->GetNativeRegion())
 	, CachedScript(Implementation->GetScript())
 	, CachedVariant(Implementation->GetVariant())
 	, CachedIsRightToLeft(Implementation->IsRightToLeft())
-{ 
+{
+	RefreshCultureDisplayNames();
 }
 
 FCulture::~FCulture()
@@ -239,8 +288,23 @@ const TArray<ETextPluralForm>& FCulture::GetValidPluralForms(const ETextPluralTy
 	return Implementation->GetValidPluralForms(PluralType);
 }
 
-void FCulture::HandleCultureChanged()
+void FCulture::RefreshCultureDisplayNames(const bool bFullRefresh)
 {
-	// Re-cache the DisplayName, as this may change when the active culture is changed
 	CachedDisplayName = Implementation->GetDisplayName();
+	ApplyCultureDisplayNameSubstitutes(CachedDisplayName);
+
+	if (bFullRefresh)
+	{
+		CachedEnglishName = Implementation->GetEnglishName();
+		ApplyCultureDisplayNameSubstitutes(CachedEnglishName);
+
+		CachedNativeName = Implementation->GetNativeName();
+		ApplyCultureDisplayNameSubstitutes(CachedNativeName);
+
+		CachedNativeLanguage = Implementation->GetNativeLanguage();
+		ApplyCultureDisplayNameSubstitutes(CachedNativeLanguage);
+
+		CachedNativeRegion = Implementation->GetNativeRegion();
+		ApplyCultureDisplayNameSubstitutes(CachedNativeRegion);
+	}
 }
diff --git a/Engine/Source/Runtime/Core/Private/Internationalization/ICUInternationalization.cpp b/Engine/Source/Runtime/Core/Private/Internationalization/ICUInternationalization.cpp
--- a/Engine/Source/Runtime/Core/Private/Internationalization/ICUInternationalization.cpp
+++ b/Engine/Source/Runtime/Core/Private/Internationalization/ICUInternationalization.cpp
@@ -549,6 +549,16 @@ bool FICUInternationalization::IsCultureAllowed(const FString& Name)
 	return (EnabledCultures.Num() == 0 || EnabledCultures.Contains(Name)) && !DisabledCultures.Contains(Name);
 }
 
+void FICUInternationalization::RefreshCultureDisplayNames()
+{
+	// Update the cached display names in any existing cultures
+	FScopeLock Lock(&CachedCulturesCS);
+	for (const auto& CachedCulturePair : CachedCultures)
+	{
+		CachedCulturePair.Value->RefreshCultureDisplayNames();
+	}
+}
+
 void FICUInternationalization::HandleLanguageChanged(const FString& Name)
 {
 	UErrorCode ICUStatus = U_ZERO_ERROR;
@@ -558,7 +568,7 @@ void FICUInternationalization::HandleLanguageChanged(const FString& Name)
 	FScopeLock Lock(&CachedCulturesCS);
 	for (const auto& CachedCulturePair : CachedCultures)
 	{
-		CachedCulturePair.Value->HandleCultureChanged();
+		CachedCulturePair.Value->RefreshCultureDisplayNames(/*bFullRefresh*/false);
 	}
 }
 
diff --git a/Engine/Source/Runtime/Core/Private/Internationalization/ICUInternationalization.h b/Engine/Source/Runtime/Core/Private/Internationalization/ICUInternationalization.h
--- a/Engine/Source/Runtime/Core/Private/Internationalization/ICUInternationalization.h
+++ b/Engine/Source/Runtime/Core/Private/Internationalization/ICUInternationalization.h
@@ -38,6 +38,8 @@ public:
 	bool IsCultureRemapped(const FString& Name, FString* OutMappedCulture);
 	bool IsCultureAllowed(const FString& Name);
 
+	void RefreshCultureDisplayNames();
+
 	void HandleLanguageChanged(const FString& Name);
 	void GetCultureNames(TArray<FString>& CultureNames) const;
 	TArray<FString> GetPrioritizedCultureNames(const FString& Name);
diff --git a/Engine/Source/Runtime/Core/Private/Internationalization/Internationalization.cpp b/Engine/Source/Runtime/Core/Private/Internationalization/Internationalization.cpp
--- a/Engine/Source/Runtime/Core/Private/Internationalization/Internationalization.cpp
+++ b/Engine/Source/Runtime/Core/Private/Internationalization/Internationalization.cpp
@@ -419,6 +419,11 @@ bool FInternationalization::IsCultureAllowed(const FString& Name)
 	return Implementation->IsCultureAllowed(Name);
 }
 
+void FInternationalization::RefreshCultureDisplayNames()
+{
+	Implementation->RefreshCultureDisplayNames();
+}
+
 void FInternationalization::GetCultureNames(TArray<FString>& CultureNames) const
 {
 	Implementation->GetCultureNames(CultureNames);
diff --git a/Engine/Source/Runtime/Core/Private/Internationalization/LegacyInternationalization.cpp b/Engine/Source/Runtime/Core/Private/Internationalization/LegacyInternationalization.cpp
--- a/Engine/Source/Runtime/Core/Private/Internationalization/LegacyInternationalization.cpp
+++ b/Engine/Source/Runtime/Core/Private/Internationalization/LegacyInternationalization.cpp
@@ -46,6 +46,10 @@ bool FLegacyInternationalization::IsCultureAllowed(const FString& Name)
 	return true;
 }
 
+void FLegacyInternationalization::RefreshCultureDisplayNames()
+{
+}
+
 void FLegacyInternationalization::HandleLanguageChanged(const FString& Name)
 {
 }
diff --git a/Engine/Source/Runtime/Core/Private/Internationalization/LegacyInternationalization.h b/Engine/Source/Runtime/Core/Private/Internationalization/LegacyInternationalization.h
--- a/Engine/Source/Runtime/Core/Private/Internationalization/LegacyInternationalization.h
+++ b/Engine/Source/Runtime/Core/Private/Internationalization/LegacyInternationalization.h
@@ -23,6 +23,8 @@ public:
 	bool IsCultureRemapped(const FString& Name, FString* OutMappedCulture);
 	bool IsCultureAllowed(const FString& Name);
 
+	void RefreshCultureDisplayNames();
+
 	void HandleLanguageChanged(const FString& Name);
 	void GetCultureNames(TArray<FString>& CultureNames) const;
 	TArray<FString> GetPrioritizedCultureNames(const FString& Name);
diff --git a/Engine/Source/Runtime/Core/Private/Internationalization/TextLocalizationManager.cpp b/Engine/Source/Runtime/Core/Private/Internationalization/TextLocalizationManager.cpp
--- a/Engine/Source/Runtime/Core/Private/Internationalization/TextLocalizationManager.cpp
+++ b/Engine/Source/Runtime/Core/Private/Internationalization/TextLocalizationManager.cpp
@@ -327,6 +327,10 @@ void InitEngineTextLocalization()
 	FStringTableRegistry::Get();
 	FStringTableRedirects::InitStringTableRedirects();
 
+	// Run this now that the config system is definitely initialized
+	// to refresh anything that was cached before it was ready
+	FInternationalization::Get().RefreshCultureDisplayNames();
+
 	ELocalizationLoadFlags LocLoadFlags = ELocalizationLoadFlags::None;
 	LocLoadFlags |= (WITH_EDITOR ? ELocalizationLoadFlags::Editor : ELocalizationLoadFlags::None);
 	LocLoadFlags |= ELocalizationLoadFlags::Engine;
diff --git a/Engine/Source/Runtime/Core/Public/Internationalization/Culture.h b/Engine/Source/Runtime/Core/Public/Internationalization/Culture.h
--- a/Engine/Source/Runtime/Core/Public/Internationalization/Culture.h
+++ b/Engine/Source/Runtime/Core/Public/Internationalization/Culture.h
@@ -97,8 +97,8 @@ public:
 	 * @param PluralType The type of plural form to get (cardinal or ordinal)
 	 */
 	const TArray<ETextPluralForm>& GetValidPluralForms(const ETextPluralType PluralType) const;
-
-	void HandleCultureChanged();
+	
+	void RefreshCultureDisplayNames(const bool bFullRefresh = true);
 
 private:
 	explicit FCulture(TUniquePtr<FCultureImplementation>&& InImplementation);
diff --git a/Engine/Source/Runtime/Core/Public/Internationalization/Internationalization.h b/Engine/Source/Runtime/Core/Public/Internationalization/Internationalization.h
--- a/Engine/Source/Runtime/Core/Public/Internationalization/Internationalization.h
+++ b/Engine/Source/Runtime/Core/Public/Internationalization/Internationalization.h
@@ -174,6 +174,9 @@ public:
 	/** Is the given culture enabled or disabled in this build? */
 	CORE_API bool IsCultureAllowed(const FString& Name);
 
+	/** Refresh the display names of the cached cultures */
+	CORE_API void RefreshCultureDisplayNames();
+
 #if ENABLE_LOC_TESTING
 	static CORE_API FString& Leetify(FString& SourceString);
 #endif
