Allow FStringView to be used with TMap and TSet

#rnx

diff --git a/Engine/Source/Runtime/Core/Public/Containers/StringView.h b/Engine/Source/Runtime/Core/Public/Containers/StringView.h
--- a/Engine/Source/Runtime/Core/Public/Containers/StringView.h
+++ b/Engine/Source/Runtime/Core/Public/Containers/StringView.h
@@ -6,6 +6,7 @@
 #include "Containers/StringFwd.h"
 #include "Math/NumericLimits.h"
 #include "Math/UnrealMathUtility.h"
+#include "Misc/Crc.h"
 #include "Misc/CString.h"
 #include "Templates/AndOrNot.h"
 #include "Templates/EnableIf.h"
@@ -296,4 +297,12 @@ constexpr inline FWideStringView operator "" _WSV(const WIDECHAR* String, size_t
 
 //////////////////////////////////////////////////////////////////////////
 
+/** Case insensitive string hash function. */
+template <typename CharType>
+FORCEINLINE uint32 GetTypeHash(const TStringViewImpl<CharType>& S)
+{
+	// This must match the GetTypeHash behavior of FString
+	return FCrc::Strihash_DEPRECATED(S.Len(), S.GetData());
+}
+
 #include "StringView.inl"
diff --git a/Engine/Source/Runtime/Core/Public/Containers/UnrealString.h b/Engine/Source/Runtime/Core/Public/Containers/UnrealString.h
--- a/Engine/Source/Runtime/Core/Public/Containers/UnrealString.h
+++ b/Engine/Source/Runtime/Core/Public/Containers/UnrealString.h
@@ -1936,7 +1936,8 @@ inline int32 GetNum(const FString& String)
 /** Case insensitive string hash function. */
 FORCEINLINE uint32 GetTypeHash(const FString& S)
 {
-	return FCrc::Strihash_DEPRECATED(*S);
+	// This must match the GetTypeHash behavior of FStringView
+	return FCrc::Strihash_DEPRECATED(S.Len(), *S);
 }
 
 /** 
diff --git a/Engine/Source/Runtime/Core/Public/Misc/Crc.h b/Engine/Source/Runtime/Core/Public/Misc/Crc.h
--- a/Engine/Source/Runtime/Core/Public/Misc/Crc.h
+++ b/Engine/Source/Runtime/Core/Public/Misc/Crc.h
@@ -86,11 +86,29 @@ struct CORE_API FCrc
 		// make sure table is initialized
 		check(CRCTable_DEPRECATED[1] != 0);
 
-		int32 Length = TCString<CharType>::Strlen( Data );
 		uint32 CRC = 0xFFFFFFFF;
-		for( int32 i=0; i<Length; i++ )
+		while (*Data)
 		{
-			CharType C = Data[i];
+			CharType C = *Data++;
+			int32 CL = (C&255);
+			CRC = (CRC << 8) ^ CRCTable_DEPRECATED[(CRC >> 24) ^ CL];
+			int32 CH = (C>>8)&255;
+			CRC = (CRC << 8) ^ CRCTable_DEPRECATED[(CRC >> 24) ^ CH];
+		}
+		return ~CRC;
+	}
+
+	/** String CRC. */
+	template <typename CharType>
+	static inline uint32 StrCrc_DEPRECATED(const int32 DataLen, const CharType* Data)
+	{
+		// make sure table is initialized
+		check(CRCTable_DEPRECATED[1] != 0);
+
+		uint32 CRC = 0xFFFFFFFF;
+		for (int32 i = 0; i < DataLen; i++)
+		{
+			CharType C = *Data++;
 			int32 CL = (C&255);
 			CRC = (CRC << 8) ^ CRCTable_DEPRECATED[(CRC >> 24) ^ CL];
 			int32 CH = (C>>8)&255;
@@ -101,6 +119,7 @@ struct CORE_API FCrc
 
 	/** Case insensitive string hash function. */
 	template <typename CharType> static inline uint32 Strihash_DEPRECATED( const CharType* Data );
+	template <typename CharType> static inline uint32 Strihash_DEPRECATED( const int32 DataLen, const CharType* Data );
 
 	/** generates CRC hash of the memory area */
 	static uint32 MemCrc_DEPRECATED( const void* Data, int32 Length, uint32 CRC=0 );
@@ -122,6 +141,22 @@ inline uint32 FCrc::Strihash_DEPRECATED(const ANSICHAR* Data)
 	return Hash;
 }
 
+template <>
+inline uint32 FCrc::Strihash_DEPRECATED(const int32 DataLen, const ANSICHAR* Data)
+{
+	// make sure table is initialized
+	check(CRCTable_DEPRECATED[1] != 0);
+
+	uint32 Hash=0;
+	for (int32 Idx = 0; Idx < DataLen; ++Idx)
+	{
+		ANSICHAR Ch = TChar<ANSICHAR>::ToUpper(*Data++);
+		uint8 B  = Ch;
+		Hash = ((Hash >> 8) & 0x00FFFFFF) ^ CRCTable_DEPRECATED[(Hash ^ B) & 0x000000FF];
+	}
+	return Hash;
+}
+
 template <>
 inline uint32 FCrc::Strihash_DEPRECATED(const WIDECHAR* Data)
 {
@@ -139,3 +174,21 @@ inline uint32 FCrc::Strihash_DEPRECATED(const WIDECHAR* Data)
 	}
 	return Hash;
 }
+
+template <>
+inline uint32 FCrc::Strihash_DEPRECATED(const int32 DataLen, const WIDECHAR* Data)
+{
+	// make sure table is initialized
+	check(CRCTable_DEPRECATED[1] != 0);
+
+	uint32 Hash=0;
+	for (int32 Idx = 0; Idx < DataLen; ++Idx)
+	{
+		WIDECHAR Ch = TChar<WIDECHAR>::ToUpper(*Data++);
+		uint16  B  = Ch;
+		Hash     = ((Hash >> 8) & 0x00FFFFFF) ^ CRCTable_DEPRECATED[(Hash ^ B) & 0x000000FF];
+		B        = Ch>>8;
+		Hash     = ((Hash >> 8) & 0x00FFFFFF) ^ CRCTable_DEPRECATED[(Hash ^ B) & 0x000000FF];
+	}
+	return Hash;
+}
