Missed file.

diff --git a/Engine/Source/Runtime/Core/Public/UObject/ReleaseObjectVersion.h b/Engine/Source/Runtime/Core/Public/UObject/ReleaseObjectVersion.h
--- a/Engine/Source/Runtime/Core/Public/UObject/ReleaseObjectVersion.h
+++ b/Engine/Source/Runtime/Core/Public/UObject/ReleaseObjectVersion.h
@@ -95,6 +95,10 @@ struct CORE_API FReleaseObjectVersion
 
 		// Upgrade filmback
 		DeprecateFilmbackSettings,
+		
+		// custom collision type
+		CustomImplicitCollisionType,
+
 
 		// -----<new versions can be added above this line>-------------------------------------------------
 		VersionPlusOne,
