Allow FStringView to be used with the ByHash TMap functions for FString TMap key queries (by [REDACTED_PERSON])

diff --git a/Engine/Source/Runtime/Core/Public/Containers/Map.h b/Engine/Source/Runtime/Core/Public/Containers/Map.h
--- a/Engine/Source/Runtime/Core/Public/Containers/Map.h
+++ b/Engine/Source/Runtime/Core/Public/Containers/Map.h
@@ -81,14 +81,28 @@ struct TDefaultMapKeyFuncs : BaseKeyFuncs<TPair<KeyType, ValueType>, KeyType, bI
 	{
 		return Element.Key;
 	}
+
 	static FORCEINLINE bool Matches(KeyInitType A, KeyInitType B)
 	{
 		return A == B;
 	}
+
+	template<typename ComparableKey>
+	static FORCEINLINE bool Matches(KeyInitType A, ComparableKey B)
+	{
+		return A == B;
+	}
+
 	static FORCEINLINE uint32 GetKeyHash(KeyInitType Key)
 	{
 		return GetTypeHash(Key);
 	}
+
+	template<typename ComparableKey>
+	static FORCEINLINE uint32 GetKeyHash(ComparableKey Key)
+	{
+		return GetTypeHash(Key);
+	}
 };
 
 template<typename KeyType, typename ValueType, bool bInAllowDuplicateKeys>
diff --git a/Engine/Source/Runtime/Core/Public/Containers/Set.h b/Engine/Source/Runtime/Core/Public/Containers/Set.h
--- a/Engine/Source/Runtime/Core/Public/Containers/Set.h
+++ b/Engine/Source/Runtime/Core/Public/Containers/Set.h
@@ -53,6 +53,14 @@ struct DefaultKeyFuncs : BaseKeyFuncs<ElementType,ElementType,bInAllowDuplicateK
 		return Element;
 	}
 
+	/**
+	 * @return True if the keys match.
+	 */
+	static FORCEINLINE bool Matches(KeyInitType A, KeyInitType B)
+	{
+		return A == B;
+	}
+
 	/**
 	 * @return True if the keys match.
 	 */
@@ -67,6 +75,13 @@ struct DefaultKeyFuncs : BaseKeyFuncs<ElementType,ElementType,bInAllowDuplicateK
 	{
 		return GetTypeHash(Key);
 	}
+
+	/** Calculates a hash index for a key. */
+	template<typename ComparableKey>
+	static FORCEINLINE uint32 GetKeyHash(ComparableKey Key)
+	{
+		return GetTypeHash(Key);
+	}
 };
 
 /** This is used to provide type specific behavior for a move which will destroy B. */
