Workaround for [[nodiscard]] error

diff --git a/Engine/Source/Runtime/Core/Private/IO/IoStore.cpp b/Engine/Source/Runtime/Core/Private/IO/IoStore.cpp
--- a/Engine/Source/Runtime/Core/Private/IO/IoStore.cpp
+++ b/Engine/Source/Runtime/Core/Private/IO/IoStore.cpp
@@ -214,7 +214,7 @@ FIoStoreWriter::FIoStoreWriter(FIoStoreEnvironment& InEnvironment)
 
 FIoStoreWriter::~FIoStoreWriter()
 {
-	Impl->FlushMetadata();
+	(void)Impl->FlushMetadata();
 }
 
 FIoStatus FIoStoreWriter::Initialize()
