Unshelved from pending changelist '11553754':

Add OnEnginePreExit() delegate to core delegates. This delegate is fired off before closing out the AssetTools, WorldBrowser and AssetRegistry modules, for handling shutdown tasks that would require these.

diff --git a/Engine/Source/Runtime/Core/Private/Misc/CoreDelegates.cpp b/Engine/Source/Runtime/Core/Private/Misc/CoreDelegates.cpp
--- a/Engine/Source/Runtime/Core/Private/Misc/CoreDelegates.cpp
+++ b/Engine/Source/Runtime/Core/Private/Misc/CoreDelegates.cpp
@@ -81,6 +81,7 @@ FSimpleMulticastDelegate FCoreDelegates::OnPostEngineInit;
 FSimpleMulticastDelegate FCoreDelegates::OnFEngineLoopInitComplete;
 FSimpleMulticastDelegate FCoreDelegates::OnExit;
 FSimpleMulticastDelegate FCoreDelegates::OnPreExit;
+FSimpleMulticastDelegate FCoreDelegates::OnEnginePreExit;
 FCoreDelegates::FGatherAdditionalLocResPathsDelegate FCoreDelegates::GatherAdditionalLocResPathsCallback;
 FSimpleMulticastDelegate FCoreDelegates::ColorPickerChanged;
 FSimpleMulticastDelegate FCoreDelegates::OnBeginFrame;
diff --git a/Engine/Source/Runtime/Core/Public/Misc/CoreDelegates.h b/Engine/Source/Runtime/Core/Public/Misc/CoreDelegates.h
--- a/Engine/Source/Runtime/Core/Public/Misc/CoreDelegates.h
+++ b/Engine/Source/Runtime/Core/Public/Misc/CoreDelegates.h
@@ -234,6 +234,9 @@ public:
 	// Called when before the application is exiting.
 	static FSimpleMulticastDelegate OnPreExit;
 
+	// Called before the engine exits. Separate from OnPreExit as OnEnginePreExit occurs before shutting down any core modules.
+	static FSimpleMulticastDelegate OnEnginePreExit;
+
 	/** Delegate for gathering up additional localization paths that are unknown to the UE4 core (such as plugins) */
 	DECLARE_MULTICAST_DELEGATE_OneParam(FGatherAdditionalLocResPathsDelegate, TArray<FString>&);
 	static FGatherAdditionalLocResPathsDelegate GatherAdditionalLocResPathsCallback;
diff --git a/Engine/Source/Runtime/Engine/Private/UnrealEngine.cpp b/Engine/Source/Runtime/Engine/Private/UnrealEngine.cpp
--- a/Engine/Source/Runtime/Engine/Private/UnrealEngine.cpp
+++ b/Engine/Source/Runtime/Engine/Private/UnrealEngine.cpp
@@ -1782,6 +1782,8 @@ void UEngine::ShutdownAudioDeviceManager()
 
 void UEngine::PreExit()
 {
+	FCoreDelegates::OnEnginePreExit.Broadcast();
+
 #if UE_NET_TRACE_ENABLED
 	FNetTrace::SetTraceVerbosity(0);
 #endif
