FArchive error propagation to all proxies

diff --git a/Engine/Plugins/Runtime/PacketHandlers/CompressionComponents/Oodle/Source/OodleHandlerComponent/Private/OodleHandlerComponent.cpp b/Engine/Plugins/Runtime/PacketHandlers/CompressionComponents/Oodle/Source/OodleHandlerComponent/Private/OodleHandlerComponent.cpp
--- a/Engine/Plugins/Runtime/PacketHandlers/CompressionComponents/Oodle/Source/OodleHandlerComponent/Private/OodleHandlerComponent.cpp
+++ b/Engine/Plugins/Runtime/PacketHandlers/CompressionComponents/Oodle/Source/OodleHandlerComponent/Private/OodleHandlerComponent.cpp
@@ -545,7 +545,7 @@ void OodleHandlerComponent::InitializeDictionary(FString FilePath, TSharedPtr<FO
 	// Load the dictionary, if it's not yet loaded
 	if (DictionaryRef == nullptr)
 	{
-		FArchive* ReadArc = IFileManager::Get().CreateFileReader(*FilePath);
+		TUniquePtr<FArchive> ReadArc(IFileManager::Get().CreateFileReader(*FilePath));
 
 		if (ReadArc != nullptr)
 		{
@@ -607,9 +607,6 @@ void OodleHandlerComponent::InitializeDictionary(FString FilePath, TSharedPtr<FO
 
 
 			ReadArc->Close();
-
-			delete ReadArc;
-			ReadArc = nullptr;
 		}
 		else
 		{
diff --git a/Engine/Source/Runtime/Core/Private/Serialization/Archive.cpp b/Engine/Source/Runtime/Core/Private/Serialization/Archive.cpp
--- a/Engine/Source/Runtime/Core/Private/Serialization/Archive.cpp
+++ b/Engine/Source/Runtime/Core/Private/Serialization/Archive.cpp
@@ -133,6 +133,8 @@ FArchiveState& FArchiveState::operator=(const FArchiveState& ArchiveToCopy)
 
 FArchiveState::~FArchiveState()
 {
+	checkf(NextProxy == nullptr, TEXT("Archive destroyed before its proxies"));
+
 	delete CustomVersionContainer;
 
 	delete SerializedPropertyChain;
@@ -250,31 +252,67 @@ void FArchiveState::CopyTrivialFArchiveStatusMembers(const FArchiveState& Archiv
 #endif // USE_STABLE_LOCALIZATION_KEYS
 }
 
-/**
- * Returns the name of the Archive.  Useful for getting the name of the package a struct or object
- * is in when a loading error occurs.
- *
- * This is overridden for the specific Archive Types
- **/
-FString FArchiveState::GetArchiveName() const
+
+void FArchiveState::LinkProxy(FArchiveState& Inner, FArchiveState& Proxy)
 {
-	return TEXT("FArchive");
+	Proxy.NextProxy = Inner.NextProxy;
+	Inner.NextProxy = &Proxy;
+}
+
+void FArchiveState::UnlinkProxy(FArchiveState& Inner, FArchiveState& Proxy)
+{
+	FArchiveState* Prev = &Inner;
+	while (Prev->NextProxy != &Proxy)
+	{
+		Prev = Prev->NextProxy;
+		checkf(Prev, TEXT("Proxy link not found - likely  lifetime violation"));
+	}
+
+	Prev->NextProxy = Proxy.NextProxy;
+	Proxy.NextProxy = nullptr;
+}
+
+template<typename T>
+FORCEINLINE void FArchiveState::ForEachState(T Func)
+{
+	FArchiveState& RootState = GetInnermostState();
+	Func(RootState);
+
+	for (FArchiveState* Proxy = RootState.NextProxy; Proxy; Proxy = Proxy->NextProxy)
+	{
+		Func(*Proxy);
+	}
+}
+
+void FArchiveState::SetArchiveState(const FArchiveState& InState)
+{
+	ForEachState([&InState](FArchiveState& State) { State = InState; });
 }
 
 void FArchiveState::SetError()
 {
-	ArIsError = true;
+	ForEachState([](FArchiveState& State) { State.ArIsError = true; });
+}
+
+void FArchiveState::SetCriticalError()
+{
+	ForEachState([](FArchiveState& State) { State.ArIsError = State.ArIsCriticalError = true; });
 }
 
 void FArchiveState::ClearError()
 {
-	ArIsError = false;
+	ForEachState([](FArchiveState& State) { State.ArIsError = false; });
 }
 
-void FArchiveState::SetCriticalError()
+/**
+ * Returns the name of the Archive.  Useful for getting the name of the package a struct or object
+ * is in when a loading error occurs.
+ *
+ * This is overridden for the specific Archive Types
+ **/
+FString FArchiveState::GetArchiveName() const
 {
-	SetError();
-	ArIsCriticalError = true;
+	return TEXT("FArchive");
 }
 
 void FArchiveState::GetSerializedPropertyChain(TArray<class FProperty*>& OutProperties) const
@@ -1222,9 +1260,6 @@ void FArchiveState::SetIsPersistent(bool bInIsPersistent)
 	ArIsPersistent = bInIsPersistent;
 }
 
-void FArchive::SetArchiveState(const FArchiveState& InState)
-{
-	ImplicitConv<FArchiveState&>(*this) = InState;
-}
+static_assert(sizeof(FArchive) == sizeof(FArchiveState), "New FArchive members should be added to FArchiveState instead");
 
 PRAGMA_ENABLE_UNSAFE_TYPECAST_WARNINGS
diff --git a/Engine/Source/Runtime/Core/Private/Serialization/ArchiveProxy.cpp b/Engine/Source/Runtime/Core/Private/Serialization/ArchiveProxy.cpp
--- a/Engine/Source/Runtime/Core/Private/Serialization/ArchiveProxy.cpp
+++ b/Engine/Source/Runtime/Core/Private/Serialization/ArchiveProxy.cpp
@@ -10,5 +10,10 @@ FArchiveProxy::FArchiveProxy(FArchive& InInnerArchive)
 : FArchive    (InInnerArchive)
 , InnerArchive(InInnerArchive)
 {
+	LinkProxy(InnerArchive.GetArchiveState(), GetArchiveState());
 }
 
+FArchiveProxy::~FArchiveProxy()
+{
+	UnlinkProxy(InnerArchive.GetArchiveState(), GetArchiveState());
+}
diff --git a/Engine/Source/Runtime/Core/Public/Serialization/Archive.h b/Engine/Source/Runtime/Core/Public/Serialization/Archive.h
--- a/Engine/Source/Runtime/Core/Public/Serialization/Archive.h
+++ b/Engine/Source/Runtime/Core/Public/Serialization/Archive.h
@@ -67,7 +67,33 @@ private:
 
 	virtual ~FArchiveState() = 0;
 
+protected:
+	static void LinkProxy(FArchiveState& Inner, FArchiveState& Proxy);
+	static void UnlinkProxy(FArchiveState& Inner, FArchiveState& Proxy);
+
 public:
+	virtual FArchiveState& GetInnermostState()
+	{
+		return *this;
+	}
+
+	void SetArchiveState(const FArchiveState& InState);
+
+	/**
+	 * Sets ArIsError to true. Also sets error in the proxy archiver if one is wrapping this.
+	 */
+	void SetError();
+
+	/**
+	 * Sets ArIsError to false, this does not clear any CriticalErrors
+	 */
+	void ClearError();
+
+	/**
+	 * Sets the archiver IsCriticalError and IsError to true. Also sets CriticalError in the proxy archiver if one is wrapping this.
+	 */
+	void SetCriticalError();
+
 	virtual void CountBytes(SIZE_T InNum, SIZE_T InMax) { }
 
 	/**
@@ -105,27 +131,11 @@ public:
 		return ((Pos != INDEX_NONE) && (Pos >= TotalSize()));
 	}
 
-	virtual bool GetError()
+	FORCEINLINE bool GetError() const
 	{
 		return ArIsError;
 	}
 
-	/**
-	 * Sets ArIsError to true. Also sets error in the proxy archiver if one is wrapping this.
-	 */
-	void SetError();
-
-	/**
-	 * Sets ArIsError to false, this does not clear any CriticalErrors
-	 */
-	void ClearError();
-
-	/**
-	 * Sets the archiver IsCriticalError and IsError to true. Also sets CriticalError in the proxy archiver if one is wrapping this.
-	 */
-	void SetCriticalError();
-
-
 	FORCEINLINE bool IsByteSwapping()
 	{
 	#if PLATFORM_LITTLE_ENDIAN
@@ -858,6 +868,11 @@ protected:
 	 */
 	mutable bool bCustomVersionsAreReset;
 
+private:
+	/** Linked list to all proxies */
+	FArchiveState* NextProxy = nullptr;
+
+	template<typename T> void ForEachState(T Func);
 };
 
 /**
@@ -957,6 +972,10 @@ public:
 	FArchive& operator=(const FArchive& ArchiveToCopy) = default;
 	~FArchive() = default;
 
+protected:
+	using FArchiveState::LinkProxy;
+	using FArchiveState::UnlinkProxy;
+
 public:
 
 	/**
@@ -1428,6 +1447,21 @@ public:
 
 	virtual void Preload(UObject* Object) { }
 
+	FArchiveState& GetArchiveState()
+	{
+		return ImplicitConv<FArchiveState&>(*this);
+	}
+
+	const FArchiveState& GetArchiveState() const
+	{
+		return ImplicitConv<const FArchiveState&>(*this);
+	}
+
+	using FArchiveState::SetArchiveState;
+	using FArchiveState::SetError;
+	using FArchiveState::ClearError;
+	using FArchiveState::SetCriticalError;
+	using FArchiveState::GetInnermostState;
 	using FArchiveState::CountBytes;
 	using FArchiveState::GetArchiveName;
 	using FArchiveState::GetLinker;
@@ -1515,9 +1549,6 @@ public:
 	}
 
 	using FArchiveState::GetError;
-	using FArchiveState::SetError;
-	using FArchiveState::ClearError;
-	using FArchiveState::SetCriticalError;
 
 	/**
 	 * Serializes and compresses/ uncompresses data. This is a shared helper function for compression
@@ -1761,13 +1792,6 @@ public:
 	}
 #endif
 
-	const FArchiveState& GetArchiveState() const
-	{
-		return ImplicitConv<const FArchiveState&>(*this);
-	}
-
-	virtual void SetArchiveState(const FArchiveState& InState);
-
 private:
 	// Used internally only to control the amount of generated code/type under control.
 	template<typename T>
@@ -1950,7 +1974,7 @@ private:
 	using FArchiveState::bCustomVersionsAreReset;
 };
 
-static_assert(sizeof(FArchive) == sizeof(FArchiveState), "New FArchive members should be added to FArchiveState instead");
+
 
 
 /**
diff --git a/Engine/Source/Runtime/Core/Public/Serialization/ArchiveProxy.h b/Engine/Source/Runtime/Core/Public/Serialization/ArchiveProxy.h
--- a/Engine/Source/Runtime/Core/Public/Serialization/ArchiveProxy.h
+++ b/Engine/Source/Runtime/Core/Public/Serialization/ArchiveProxy.h
@@ -20,6 +20,13 @@ public:
 	 * @param InInnerArchive The inner archive to proxy.
 	 */
 	CORE_API FArchiveProxy(FArchive& InInnerArchive);
+	CORE_API ~FArchiveProxy();
+
+	// Non-copyable
+	FArchiveProxy(FArchiveProxy&&) = delete;
+	FArchiveProxy(const FArchiveProxy&) = delete;
+	FArchiveProxy& operator=(FArchiveProxy&&) = delete;
+	FArchiveProxy& operator=(const FArchiveProxy&) = delete;
 
 	virtual FArchive& operator<<(FName& Value) override
 	{
@@ -161,11 +168,6 @@ public:
 		return InnerArchive.Close();
 	}
 
-	virtual bool GetError() override
-	{
-		return InnerArchive.GetError();
-	}
-
 	virtual void MarkScriptSerializationStart(const UObject* Obj) override
 	{
 		InnerArchive.MarkScriptSerializationStart(Obj);
@@ -284,9 +286,9 @@ public:
 		return InnerArchive.GetCacheableArchive();
 	}
 
-	virtual void SetArchiveState(const ::FArchiveState& InState) override
+	virtual ::FArchiveState& GetInnermostState() override
 	{
-		return InnerArchive.SetArchiveState(InState);
+		return InnerArchive.GetInnermostState();
 	}
 
 protected:
diff --git a/Engine/Source/Runtime/PakFile/Private/IPlatformFilePak.cpp b/Engine/Source/Runtime/PakFile/Private/IPlatformFilePak.cpp
--- a/Engine/Source/Runtime/PakFile/Private/IPlatformFilePak.cpp
+++ b/Engine/Source/Runtime/PakFile/Private/IPlatformFilePak.cpp
@@ -6186,27 +6186,18 @@ bool FPakFile::Check()
 /**
 * FThreadCheckingArchiveProxy - checks that inner archive is only used from the specified thread ID
 */
-class FThreadCheckingArchiveProxy : public FArchiveProxy
+class FThreadCheckingArchiveProxy : TUniquePtr<FArchive>, public FArchiveProxy
 {
 public:
 
 	const uint32 ThreadId;
-	FArchive* InnerArchivePtr;
 
 	FThreadCheckingArchiveProxy(FArchive* InReader, uint32 InThreadId)
-		: FArchiveProxy(*InReader)
+		: TUniquePtr(InReader) // Make sure proxy is destroyed before InReader
+		, FArchiveProxy(*InReader)
 		, ThreadId(InThreadId)
-		, InnerArchivePtr(InReader)
 	{}
 
-	virtual ~FThreadCheckingArchiveProxy()
-	{
-		if (InnerArchivePtr)
-		{
-			delete InnerArchivePtr;
-		}
-	}
-
 	//~ Begin FArchiveProxy Interface
 	virtual void Serialize(void* Data, int64 Length) override
 	{
