Physics for ocean flood :
* Unified all water-related computations into a single function on the base water body class (overridable) : QueryWaterInfoClosestToWorldLocation (water surface/plane location/normal, velocity, depth, immersion, with/without simple/complex waves)
* Handled ocean flood offset
* Made ocean collision component transient and adjusted it dynamically with flood height

diff --git a/Engine/Source/Runtime/Core/Public/UObject/FortniteMainBranchObjectVersion.h b/Engine/Source/Runtime/Core/Public/UObject/FortniteMainBranchObjectVersion.h
--- a/Engine/Source/Runtime/Core/Public/UObject/FortniteMainBranchObjectVersion.h
+++ b/Engine/Source/Runtime/Core/Public/UObject/FortniteMainBranchObjectVersion.h
@@ -102,6 +102,9 @@ struct CORE_API FFortniteMainBranchObjectVersion
 		// Anim layer node names are now conformed by Guid
 		AnimLayerGuidConformation,
 
+		// Ocean collision component has become dynamic
+		MakeOceanCollisionTransient,
+
 		// -----<new versions can be added above this line>-------------------------------------------------
 		VersionPlusOne,
 		LatestVersion = VersionPlusOne - 1
