Remove obsolete workaround and comment (secure_getenv).

diff --git a/Engine/Source/Runtime/Core/Private/Unix/UnixPlatformMisc.cpp b/Engine/Source/Runtime/Core/Private/Unix/UnixPlatformMisc.cpp
--- a/Engine/Source/Runtime/Core/Private/Unix/UnixPlatformMisc.cpp
+++ b/Engine/Source/Runtime/Core/Private/Unix/UnixPlatformMisc.cpp
@@ -42,9 +42,6 @@
 #include "FramePro/FrameProProfiler.h"
 #include "BuildSettings.h"
 
-// define for glibc 2.12.2 and lower (which is shipped with CentOS 6.x and which we target by default)
-#define __secure_getenv getenv
-
 extern bool GInitializedSDL;
 
 static int SysGetRandomSupported = -1;
