Add support for multithreading in forked processes

* Use -PostForkThreading when SupportsMultithreading() is false to allow multithreading in forked processes
* Forked processes do not inherit threads from the master, so this keeps the master process in singlethread mode but transforms the RunnableThreads into real threads on the forked version.
* Use FForkProcessHelper::CreateForkableThread to use this feature
   * ForkableThreads will behave correctly according to the environment they are called in (singlethread, multithread, forkable)
* SyncEvents created in -PostForkThreading environments are real events on the master process (Win, Unix platform only)

diff --git a/Engine/Source/Runtime/Core/Private/GenericPlatform/GenericPlatformProcess.cpp b/Engine/Source/Runtime/Core/Private/GenericPlatform/GenericPlatformProcess.cpp
--- a/Engine/Source/Runtime/Core/Private/GenericPlatform/GenericPlatformProcess.cpp
+++ b/Engine/Source/Runtime/Core/Private/GenericPlatform/GenericPlatformProcess.cpp
@@ -17,6 +17,7 @@
 #include "Misc/EventPool.h"
 #include "Misc/EngineVersion.h"
 #include "Misc/LazySingleton.h"
+#include "Misc/Fork.h"
 #include "ProfilingDebugging/CsvProfiler.h"
 #include "Async/TaskGraphInterfaces.h"
 
@@ -445,7 +446,13 @@ FEvent* FGenericPlatformProcess::CreateSynchEvent(bool bIsManualReset)
 {
 #if PLATFORM_USE_PTHREADS
 	FEvent* Event = NULL;
-	if (FPlatformProcess::SupportsMultithreading())
+	
+	// Create fake singlethread events in environments that don't support multithreading
+	// For processes that intend to fork: create real events even in the master process since the allocated mutex might get reused by the
+	//forked child process when he gets to run in multithread mode.
+	const bool bIsMultithread = FPlatformProcess::SupportsMultithreading() || FForkProcessHelper::SupportsMultithreadingPostFork();
+
+	if (bIsMultithread)
 	{
 		// Allocate the new object
 		Event = new FPThreadEvent();
@@ -476,6 +483,11 @@ FEvent* FGenericPlatformProcess::GetSynchEventFromPool(bool bIsManualReset)
 		: TLazySingleton<FEventPool<EEventPoolTypes::AutoReset>>::Get().GetEventFromPool();
 }
 
+void FGenericPlatformProcess::FlushPoolSyncEvents()
+{
+	TLazySingleton<FEventPool<EEventPoolTypes::ManualReset>>::Get().EmptyPool();
+	TLazySingleton<FEventPool<EEventPoolTypes::AutoReset>>::Get().EmptyPool();
+}
 
 void FGenericPlatformProcess::ReturnSynchEventToPool(FEvent* Event)
 {
diff --git a/Engine/Source/Runtime/Core/Private/HAL/ThreadingBase.cpp b/Engine/Source/Runtime/Core/Private/HAL/ThreadingBase.cpp
--- a/Engine/Source/Runtime/Core/Private/HAL/ThreadingBase.cpp
+++ b/Engine/Source/Runtime/Core/Private/HAL/ThreadingBase.cpp
@@ -7,6 +7,7 @@
 #include "Misc/CoreStats.h"
 #include "Misc/EventPool.h"
 #include "Misc/LazySingleton.h"
+#include "Misc/Fork.h"
 #include "Templates/Atomic.h"
 #include "HAL/IConsoleManager.h"
 #include "HAL/PlatformStackWalk.h"
@@ -102,18 +103,20 @@ class FFakeThread : public FRunnableThread
 	/** Thread Id pool */
 	static uint32 ThreadIdCounter;
 
+protected:
+
 	/** Thread is suspended. */
 	bool bIsSuspended;
 
 	/** Runnable object associated with this thread. */
-	FSingleThreadRunnable* Runnable;
+	FSingleThreadRunnable* SingleThreadRunnable;
 
 public:
 
 	/** Constructor. */
 	FFakeThread()
 		: bIsSuspended(false)
-		, Runnable(nullptr)
+		, SingleThreadRunnable(nullptr)
 	{
 		ThreadID = ThreadIdCounter++;
 		// Auto register with single thread manager.
@@ -130,9 +133,9 @@ public:
 	/** Tick one time per frame. */
 	virtual void Tick() override
 	{
-		if (Runnable && !bIsSuspended)
+		if (SingleThreadRunnable && !bIsSuspended)
 		{
-			Runnable->Tick();
+			SingleThreadRunnable->Tick();
 		}
 	}
 
@@ -167,12 +170,24 @@ public:
 		EThreadCreateFlags InCreateFlags = EThreadCreateFlags::None) override
 
 	{
-		Runnable = InRunnable->GetSingleThreadInterface();
-		if (Runnable)
+		ThreadName = InThreadName;
+		ThreadAffinityMask = InThreadAffinityMask;
+
+		SingleThreadRunnable = InRunnable->GetSingleThreadInterface();
+		if (SingleThreadRunnable)
 		{
 			InRunnable->Init();
-		}		
-		return Runnable != nullptr;
+
+			Runnable = InRunnable;
+		}
+		return SingleThreadRunnable != nullptr;
+	}
+
+protected:
+
+	virtual FRunnableThread::ThreadType GetThreadType() const override
+	{
+		return ThreadType::Fake;
 	}
 };
 uint32 FFakeThread::ThreadIdCounter = 0xffff;
@@ -200,16 +215,21 @@ void FThreadManager::RemoveThread(FRunnableThread* Thread)
 
 void FThreadManager::Tick()
 {	
-	if (!FPlatformProcess::SupportsMultithreading())
+	const bool bIsSingleThreadEnvironment = FPlatformProcess::SupportsMultithreading() == false;
+	if (bIsSingleThreadEnvironment)
 	{
 		QUICK_SCOPE_CYCLE_COUNTER(STAT_FSingleThreadManager_Tick);
 
 		FScopeLock ThreadsLock(&ThreadsCritical);
 
-		// Tick all registered threads.
+		// Tick all registered fake threads.
 		for (TPair<uint32, FRunnableThread*>& ThreadPair : Threads)
 		{
-			ThreadPair.Value->Tick();
+			// Only fake and forkable threads are ticked by the ThreadManager
+			if( ThreadPair.Value->GetThreadType() != FRunnableThread::ThreadType::Real )
+			{
+				ThreadPair.Value->Tick();
+			}
 		}
 	}
 }
@@ -284,6 +304,21 @@ FThreadManager& FThreadManager::Get()
 	return Singleton;
 }
 
+TArray<FRunnableThread*> FThreadManager::GetForkableThreads()
+{
+	TArray<FRunnableThread*> ForkableThreads;
+	FScopeLock Lock(&ThreadsCritical);
+	for (const TPair<uint32, FRunnableThread*>& Pair : Threads)
+	{
+		if (Pair.Value->GetThreadType() == FRunnableThread::ThreadType::Forkable)
+		{
+			ForkableThreads.Add(Pair.Value);
+		}
+	}
+
+	return ForkableThreads;
+}
+
 bool FThreadManager::bIsInitialized = false;
 
 
@@ -406,47 +441,55 @@ FRunnableThread* FRunnableThread::Create(
 	uint64 InThreadAffinityMask,
 	EThreadCreateFlags InCreateFlags)
 {
+	bool bCreateRealThread = FPlatformProcess::SupportsMultithreading();
+
 	FRunnableThread* NewThread = nullptr;
-	if (FPlatformProcess::SupportsMultithreading())
+
+	if (bCreateRealThread)
 	{
 		check(InRunnable);
 		// Create a new thread object
 		NewThread = FPlatformProcess::CreateRunnableThread();
-		if (NewThread)
-		{
-			// Call the thread's create method
-			if (NewThread->CreateInternal(InRunnable,ThreadName,InStackSize,InThreadPri,InThreadAffinityMask,InCreateFlags) == false)
-			{
-				// We failed to start the thread correctly so clean up
-				delete NewThread;
-				NewThread = nullptr;
-			}
-		}
 	}
 	else if (InRunnable->GetSingleThreadInterface())
 	{
 		// Create a fake thread when multithreading is disabled.
 		NewThread = new FFakeThread();
-		if (NewThread->CreateInternal(InRunnable,ThreadName,InStackSize,InThreadPri) == false)
-		{
-			// We failed to start the thread correctly so clean up
-			delete NewThread;
-			NewThread = nullptr;
-		}
 	}
 
 	if (NewThread)
 	{
-		TRACE_CREATE_THREAD(NewThread->GetThreadID(), *NewThread->GetThreadName(), InThreadPri);
+		SetupCreatedThread(NewThread, InRunnable, ThreadName, InStackSize, InThreadPri, InThreadAffinityMask, InCreateFlags);
 	}
-#if	STATS
-	if( NewThread )
+
+	return NewThread;
+}
+
+void FRunnableThread::SetupCreatedThread(FRunnableThread*& NewThread, class FRunnable* InRunnable,  const TCHAR* ThreadName, uint32 InStackSize, EThreadPriority InThreadPri, uint64 InThreadAffinityMask, EThreadCreateFlags InCreateFlags)
+{
+	// Call the thread's create method
+	bool bIsValid = NewThread->CreateInternal(InRunnable, ThreadName, InStackSize, InThreadPri, InThreadAffinityMask, InCreateFlags);
+
+	if( bIsValid )
 	{
-		FStartupMessages::Get().AddThreadMetadata( FName( *NewThread->GetThreadName() ), NewThread->GetThreadID() );
+		check(NewThread->Runnable);
+		NewThread->PostCreate(InThreadPri);
 	}
-#endif // STATS
+	else
+	{
+		// We failed to start the thread correctly so clean up
+		delete NewThread;
+		NewThread = nullptr;
+	}
+}
 
-	return NewThread;
+void FRunnableThread::PostCreate(EThreadPriority InThreadPriority)
+{
+	TRACE_CREATE_THREAD(GetThreadID(), *GetThreadName(), InThreadPriority);
+
+#if	STATS
+	FStartupMessages::Get().AddThreadMetadata( FName( *GetThreadName() ), GetThreadID() );
+#endif // STATS
 }
 
 void FRunnableThread::SetTls()
@@ -903,3 +946,214 @@ void FTlsAutoCleanup::Register()
 		RunnableThread->TlsInstances.Add( this );
 	}
 }
+
+//-------------------------------------------------------------------------------
+// FForkableThread
+//-------------------------------------------------------------------------------
+
+/**
+ * This thread starts as a fake thread and gets ticked like it was in a single-threaded environment.
+ * Once it receives the OnPostFork event it creates and holds a real thread that
+ * will cause the RunnableObject to be executed in it's own thread.
+ */
+class FForkableThread : public FFakeThread
+{
+	typedef FFakeThread Super;
+
+private:
+
+	/** Real thread that gets created right after forking */
+	FRunnableThread* RealThread = nullptr;
+
+	/** Cached values to use when the real thread is created post-fork */
+	EThreadPriority CachedPriority = TPri_Normal;
+	uint32 CachedStackSize = 0;
+
+public:
+
+	virtual void Tick() override
+	{
+		// Tick in single-thread mode when the real thread isn't created yet
+		if(RealThread == nullptr)
+		{
+			Super::Tick();
+		}
+	}
+
+	virtual void SetThreadPriority(EThreadPriority NewPriority) override
+	{
+		CachedPriority = NewPriority;
+		
+		if (RealThread)
+		{
+			RealThread->SetThreadPriority(NewPriority);
+		}
+	}
+
+	virtual void Suspend(bool bShouldPause) override
+	{
+		Super::Suspend(bShouldPause);
+
+		if (RealThread)
+		{
+			RealThread->Suspend(bShouldPause);
+		}
+	}
+
+	virtual bool Kill(bool bShouldWait) override
+	{
+		bool bExitedCorrectly = true;
+
+		if (RealThread)
+		{
+			bExitedCorrectly = RealThread->Kill(bShouldWait);
+		}
+
+		Super::Kill(bShouldWait);
+
+		return bExitedCorrectly;
+	}
+
+	virtual void WaitForCompletion() override
+	{
+		if (RealThread)
+		{
+			RealThread->WaitForCompletion();
+		}
+
+		Super::WaitForCompletion();
+	}
+
+	virtual bool CreateInternal(FRunnable* InRunnable, const TCHAR* InThreadName, uint32 InStackSize, EThreadPriority InThreadPri, uint64 InThreadAffinityMask, EThreadCreateFlags InCreateFlags) override
+	{
+		checkf(FForkProcessHelper::SupportsMultithreadingPostFork(), TEXT("ForkableThreads should only be created when -PostForkThreading is enabled"));
+		checkf(FForkProcessHelper::IsForkedMultithreadInstance() == false, TEXT("Once forked we create a real runnable thread instead of a ForkableThread"));
+
+		// Call the fake thread creator
+		bool bCreated = Super::CreateInternal(InRunnable, InThreadName, InStackSize, InThreadPri, InThreadAffinityMask, InCreateFlags);
+
+		// Cache the target values until we create the real thread
+		CachedStackSize = InStackSize;
+		CachedPriority = InThreadPri;
+
+		return bCreated;
+	}
+
+protected:
+
+	virtual void OnPostFork() override
+	{
+		check(FForkProcessHelper::IsForkedMultithreadInstance());
+
+		check(RealThread == nullptr);
+		RealThread = FPlatformProcess::CreateRunnableThread();
+		bool bCreated = RealThread->CreateInternal(Runnable, *GetThreadName(), CachedStackSize, CachedPriority, ThreadAffinityMask, EThreadCreateFlags::None);
+
+		if (bCreated)
+		{
+			RealThread->PostCreate(CachedPriority);
+
+			// Suspend the thread if the fake thread was suspended too
+			//TODO: this lets the thread run for a few cycles before hitting the suspend call...
+			if (bIsSuspended)
+			{
+				RealThread->Suspend(bIsSuspended);
+			}
+		}
+		else
+		{
+			delete RealThread;
+			RealThread = nullptr;
+		}
+	}
+
+	virtual FRunnableThread::ThreadType GetThreadType() const override
+	{
+		return ThreadType::Forkable;
+	}
+};
+
+//-------------------------------------------------------------------------------
+// ForkableThreadHelper
+//-------------------------------------------------------------------------------
+
+bool FForkProcessHelper::bIsForkedMultithreadInstance = false;
+
+void FForkProcessHelper::OnForkingOccured()
+{
+	if( SupportsMultithreadingPostFork() )
+	{
+		bIsForkedMultithreadInstance = true;
+
+		// Use a local list of forkable threads so we don't keep a lock on the global list during thread creation
+		auto ForkableThreads = FThreadManager::Get().GetForkableThreads();
+		for (FRunnableThread* ForkableThread : ForkableThreads)
+		{
+			ForkableThread->OnPostFork();
+		}
+	}
+}
+
+bool FForkProcessHelper::IsForkedMultithreadInstance()
+{
+	return bIsForkedMultithreadInstance;
+}
+
+bool FForkProcessHelper::SupportsMultithreadingPostFork()
+{
+	static bool bSupportsMultithreadingPostFork = FCommandLine::IsInitialized() ? FParse::Param(FCommandLine::Get(), TEXT("PostForkThreading")) : false;
+	return bSupportsMultithreadingPostFork;
+}
+
+FRunnableThread* FForkProcessHelper::CreateForkableThread(class FRunnable* InRunnable, const TCHAR* InThreadName, uint32 InStackSize, EThreadPriority InThreadPri, uint64 InThreadAffinityMask, EThreadCreateFlags InCreateFlags)
+{
+	bool bCreateRealThread = FPlatformProcess::SupportsMultithreading();
+	bool bCreateForkableThread(false);
+
+	// Look for conditions allowing real threads in a non-multithread environment
+	if (bCreateRealThread == false)
+	{
+		if( SupportsMultithreadingPostFork() )
+		{
+			if( IsForkedMultithreadInstance() )
+			{
+				// Already forked, create a real thread immediately
+				bCreateRealThread = true;
+			}
+			else
+			{
+				// We have yet to fork the process, create a forkable thread to handle the fork event
+				bCreateForkableThread = true;
+			}
+		
+		}
+	}
+
+	FRunnableThread* NewThread(nullptr);
+	if (bCreateRealThread)
+	{
+		check(InRunnable);
+		NewThread = FPlatformProcess::CreateRunnableThread();
+	}
+	else if (bCreateForkableThread)
+	{
+		if( InRunnable->GetSingleThreadInterface() )
+		{
+			NewThread = new FForkableThread();
+		}
+	}
+	else
+	{
+		if (InRunnable->GetSingleThreadInterface())
+		{
+			NewThread = new FFakeThread();
+		}
+	}
+
+	if (NewThread)
+	{
+		FRunnableThread::SetupCreatedThread(NewThread, InRunnable, InThreadName, InStackSize, InThreadPri, InThreadAffinityMask, InCreateFlags);
+	}
+
+	return NewThread;
+}
diff --git a/Engine/Source/Runtime/Core/Private/Misc/EventPool.h b/Engine/Source/Runtime/Core/Private/Misc/EventPool.h
--- a/Engine/Source/Runtime/Core/Private/Misc/EventPool.h
+++ b/Engine/Source/Runtime/Core/Private/Misc/EventPool.h
@@ -84,10 +84,7 @@ public:
 #if USE_EVENT_POOLING
 	~FEventPool()
 	{
-		while (FEvent* Event = Pool.Pop())
-		{
-			delete Event;
-		}
+		EmptyPool();
 	}
 #endif
 
@@ -139,6 +136,16 @@ public:
 #endif
 	}
 
+	void EmptyPool()
+	{
+#if USE_EVENT_POOLING
+		while (FEvent* Event = Pool.Pop())
+		{
+			delete Event;
+		}
+#endif
+	}
+
 private:
 #if USE_EVENT_POOLING
 	/** Holds the collection of recycled events. */
diff --git a/Engine/Source/Runtime/Core/Private/Windows/WindowsPlatformProcess.cpp b/Engine/Source/Runtime/Core/Private/Windows/WindowsPlatformProcess.cpp
--- a/Engine/Source/Runtime/Core/Private/Windows/WindowsPlatformProcess.cpp
+++ b/Engine/Source/Runtime/Core/Private/Windows/WindowsPlatformProcess.cpp
@@ -23,6 +23,7 @@
 #include "Windows/WindowsHWrapper.h"
 #include "Misc/ConfigCacheIni.h"
 #include "Misc/CoreDelegates.h"
+#include "Misc/Fork.h"
 #include "ProfilingDebugging/CsvProfiler.h"
 
 #include "Windows/AllowWindowsPlatformTypes.h"
@@ -1301,9 +1302,12 @@ void FWindowsPlatformProcess::SleepInfinite()
 
 FEvent* FWindowsPlatformProcess::CreateSynchEvent(bool bIsManualReset)
 {
+	// While windows does not support forking we can still simulate the forking codeflow and test the singlethread to multithread switch on Win targets
+	const bool bIsMultithread = FPlatformProcess::SupportsMultithreading() || FForkProcessHelper::SupportsMultithreadingPostFork();
+
 	// Allocate the new object
 	FEvent* Event = NULL;	
-	if (FPlatformProcess::SupportsMultithreading())
+	if (bIsMultithread)
 	{
 		Event = new FEventWin();
 	}
diff --git a/Engine/Source/Runtime/Core/Public/GenericPlatform/GenericPlatformProcess.h b/Engine/Source/Runtime/Core/Public/GenericPlatform/GenericPlatformProcess.h
--- a/Engine/Source/Runtime/Core/Public/GenericPlatform/GenericPlatformProcess.h
+++ b/Engine/Source/Runtime/Core/Public/GenericPlatform/GenericPlatformProcess.h
@@ -512,6 +512,11 @@ struct CORE_API FGenericPlatformProcess
 	 */
 	static class FEvent* GetSynchEventFromPool(bool bIsManualReset = false);
 
+	/**
+	 * Deletes all the recycled sync events contained by the pools
+	 */
+	static void FlushPoolSyncEvents();
+
 	/**
 	 * Returns an event to the pool.
 	 *
diff --git a/Engine/Source/Runtime/Core/Public/HAL/RunnableThread.h b/Engine/Source/Runtime/Core/Public/HAL/RunnableThread.h
--- a/Engine/Source/Runtime/Core/Public/HAL/RunnableThread.h
+++ b/Engine/Source/Runtime/Core/Public/HAL/RunnableThread.h
@@ -21,6 +21,8 @@ class CORE_API FRunnableThread
 	friend class FThreadSingletonInitializer;
 	friend class FTlsAutoCleanup;
 	friend class FThreadManager;
+	friend class FForkableThread;
+	friend class FForkProcessHelper;
 
 	/** Index of TLS slot for FRunnableThread pointer. */
 	static uint32 RunnableTlsSlot;
@@ -157,8 +159,43 @@ protected:
 	/** ID set during thread creation. */
 	uint32 ThreadID;
 
+protected:
+
+	/** List of unique thread types we can create */
+	enum class ThreadType
+	{
+		// Regular thread that executes the runnable object in it's own context
+		Real,
+		// Fake threads are created for a single threaded environment and are always executed from the main tick
+		Fake,
+		// Forkable threads will behave like fake threads for the master process, but will become real threads on forked processes
+		Forkable,
+	};
+
 private:
 
+	/** Called to setup a newly created RunnableThread */
+	static void SetupCreatedThread(FRunnableThread*& NewThread, class FRunnable* InRunnable, const TCHAR* ThreadName, uint32 InStackSize, EThreadPriority InThreadPri, uint64 InThreadAffinityMask, EThreadCreateFlags InCreateFlags);
+
 	/** Used by the thread manager to tick threads in single-threaded mode */
 	virtual void Tick() {}
+
+	/** Returns the type of thread this is */
+	virtual FRunnableThread::ThreadType GetThreadType() const
+	{
+		return ThreadType::Real;
+	}
+
+	/**
+	 * Called on the forked process when the forkable thread can create a real thread
+	 */
+	virtual void OnPostFork()
+	{
+		checkf(false, TEXT("Only forkable threads should receive OnPostFork."));
+	}
+
+	/**
+	 * Called after the internal thread is created so it can register debug information
+	 */
+	void PostCreate(EThreadPriority ThreadPriority);
 };
diff --git a/Engine/Source/Runtime/Core/Public/HAL/ThreadManager.h b/Engine/Source/Runtime/Core/Public/HAL/ThreadManager.h
--- a/Engine/Source/Runtime/Core/Public/HAL/ThreadManager.h
+++ b/Engine/Source/Runtime/Core/Public/HAL/ThreadManager.h
@@ -72,4 +72,11 @@ public:
 	 * @return Thread manager object.
 	 */
 	static FThreadManager& Get();
+
+private:
+
+	friend class FForkProcessHelper;
+
+	/** Returns a list of registered forkable threads  */
+	TArray<FRunnableThread*> GetForkableThreads();
 };
diff --git a/Engine/Source/Runtime/Core/Public/Misc/Fork.h b/Engine/Source/Runtime/Core/Public/Misc/Fork.h
new file mode 100644
--- /dev/null
+++ b/Engine/Source/Runtime/Core/Public/Misc/Fork.h
@@ -0,0 +1,64 @@
+// Copyright Epic Games, Inc. All Rights Reserved.
+
+#pragma once
+
+#include "CoreMinimal.h"
+#include "HAL/PlatformAffinity.h"
+
+/**
+ * Helper functions for processes that fork in order to share memory pages.
+ *
+ * About multithreading:
+ * When a process gets forked, any existing threads will not exist on the new forked process.
+ * To solve this we use forkable threads that are notified when the fork occurs and will automatically convert themselves into real runnable threads.
+ * On the master process, these forkable threads will be fake threads that are executed on the main thread and will block the critical path.
+ *
+ * Currently the game code is responsible for calling Fork on itself than calling FForkProcessHelper::OnForkingOccured to transform the forkable threads.
+ * Ideally the fork point is done right after the game has loaded all the assets it wants to share so it can maximize the shared memory pool.
+ * From the fork point any memory page that gets written into by a forked process will be transferred into a unique page for this process.
+ * 
+ */
+class CORE_API FForkProcessHelper 
+{
+public:
+
+	/**
+	 * Are we a forked process that supports multithreading
+	 */
+	static bool IsForkedMultithreadInstance();
+
+	/**
+	 * Event triggered when the fork occurred on the child process
+	 */
+	static void OnForkingOccured();
+
+	/**
+	 * Tells if we allow multithreading on forked processes.
+	 * Enabled via the commandline option: -PostForkThreading
+	 */
+	static bool SupportsMultithreadingPostFork();
+
+	/**
+	 * Creates a thread according to the environment it's in:
+	 *	In environments with SupportsMultithreading: create a real thread that will tick the runnable object itself
+	 *	In environments without multithreading: create a fake thread that is ticked by the main thread.
+	 *  In environments without multithreading but that allows multithreading post-fork: 
+	 *		If called on the original master process: will create a forkable thread that is ticked in the main thread pre-fork but becomes a real thread post-fork
+	 *      If called on a forked child process: will create a real thread immediately
+	 */
+	static FRunnableThread* CreateForkableThread(
+		class FRunnable* InRunnable,
+		const TCHAR* InThreadName,
+		uint32 InStackSize = 0,
+		EThreadPriority InThreadPri = TPri_Normal,
+		uint64 InThreadAffinityMask = FPlatformAffinity::GetNoAffinityMask(),
+		EThreadCreateFlags InCreateFlags = EThreadCreateFlags::None
+	);
+
+private:
+
+	static bool bIsForkedMultithreadInstance;
+};
+
+
+
