Fix unity build due to missing  include

diff --git a/Engine/Source/Runtime/Core/Private/HAL/ThreadingBase.cpp b/Engine/Source/Runtime/Core/Private/HAL/ThreadingBase.cpp
--- a/Engine/Source/Runtime/Core/Private/HAL/ThreadingBase.cpp
+++ b/Engine/Source/Runtime/Core/Private/HAL/ThreadingBase.cpp
@@ -4,6 +4,7 @@
 #include "HAL/ThreadingBase.h"
 #include "UObject/NameTypes.h"
 #include "Stats/Stats.h"
+#include "Misc/CommandLine.h"
 #include "Misc/CoreStats.h"
 #include "Misc/EventPool.h"
 #include "Misc/LazySingleton.h"
