Deprecated fixed string builders for 4.25

Fixed string builders are difficult to use correctly in many cases and, when used correctly, offer no performance or memory improvement over an extendable string builder of the same size.

#rn minor Core

diff --git a/Engine/Source/Runtime/Core/Public/Misc/StringBuilder.h b/Engine/Source/Runtime/Core/Public/Misc/StringBuilder.h
--- a/Engine/Source/Runtime/Core/Public/Misc/StringBuilder.h
+++ b/Engine/Source/Runtime/Core/Public/Misc/StringBuilder.h
@@ -375,7 +375,7 @@ template <>
 struct TIsContiguousContainer<FAnsiStringBuilderBase> { enum { Value = true }; };
 
 template<int32 N>
-class TFixedAnsiStringBuilder : public FAnsiStringBuilderBase
+class UE_DEPRECATED(4.25, "'TFixedAnsiStringBuilder' is deprecated. Please use 'TAnsiStringBuilder' instead!") TFixedAnsiStringBuilder : public FAnsiStringBuilderBase
 {
 public:
 	inline TFixedAnsiStringBuilder()
@@ -389,8 +389,10 @@ private:
 	ANSICHAR	StringBuffer[N];
 };
 
+PRAGMA_DISABLE_DEPRECATION_WARNINGS
 template <int32 N>
 struct TIsContiguousContainer<TFixedAnsiStringBuilder<N>> { enum { Value = true }; };
+PRAGMA_ENABLE_DEPRECATION_WARNINGS
 
 template<int32 N>
 class TAnsiStringBuilder : public FAnsiStringBuilderBase
@@ -428,7 +430,7 @@ template <>
 struct TIsContiguousContainer<FStringBuilderBase> { enum { Value = true }; };
 
 template<int32 N>
-class TFixedStringBuilder : public FStringBuilderBase
+class UE_DEPRECATED(4.25, "'TFixedStringBuilder' is deprecated. Please use 'TStringBuilder' instead!") TFixedStringBuilder : public FStringBuilderBase
 {
 public:
 	inline TFixedStringBuilder()
@@ -442,8 +444,10 @@ private:
 	TCHAR	StringBuffer[N];
 };
 
+PRAGMA_DISABLE_DEPRECATION_WARNINGS
 template <int32 N>
 struct TIsContiguousContainer<TFixedStringBuilder<N>> { enum { Value = true }; };
+PRAGMA_ENABLE_DEPRECATION_WARNINGS
 
 template<int32 N>
 class TStringBuilder : public FStringBuilderBase
