- Fixed an issue from the LoadTimes stream where serializing 0 items would not empty out an array when undoing an array add
#rnx

diff --git a/Engine/Source/Runtime/Core/Public/Containers/Array.h b/Engine/Source/Runtime/Core/Public/Containers/Array.h
--- a/Engine/Source/Runtime/Core/Public/Containers/Array.h
+++ b/Engine/Source/Runtime/Core/Public/Containers/Array.h
@@ -1141,6 +1141,11 @@ public:
 
 		if (SerializeNum == 0)
 		{
+			// if we are loading, then we have to reset the size to 0, in case it isn't currently 0
+			if (Ar.IsLoading())
+			{
+				A.Empty();
+			}
 			return Ar;
 		}
 
