Fixed the UE_DEPRECATED sample usage for aliases

#rnx

diff --git a/Engine/Source/Runtime/Core/Public/Misc/CoreMiscDefines.h b/Engine/Source/Runtime/Core/Public/Misc/CoreMiscDefines.h
--- a/Engine/Source/Runtime/Core/Public/Misc/CoreMiscDefines.h
+++ b/Engine/Source/Runtime/Core/Public/Misc/CoreMiscDefines.h
@@ -160,8 +160,7 @@ const FPlatformUserId PLATFORMUSERID_NONE = INDEX_NONE;
  *		UE_DEPRECATED(4.xx, "Message")
  *		typedef FThing MyType;
  *
- *		UE_DEPRECATED(4.xx, "Message")
- *		using MyAlias = FThing;
+ *		using MyAlias UE_DEPRECATED(4.xx, "Message") = FThing;
  *
  *		UE_DEPRECATED(4.xx, "Message")
  *		int32 MyVariable;
