Fix mis-match malloc/delete, which was causing a check to fail

diff --git a/Engine/Source/Runtime/Core/Private/FileCache/FileCache.cpp b/Engine/Source/Runtime/Core/Private/FileCache/FileCache.cpp
--- a/Engine/Source/Runtime/Core/Private/FileCache/FileCache.cpp
+++ b/Engine/Source/Runtime/Core/Private/FileCache/FileCache.cpp
@@ -523,6 +523,11 @@ public:
 		}
 	}
 
+	void operator delete(void* InMem)
+	{
+		FMemory::Free(InMem);
+	}
+
 	int64 InitialSlotOffset;
 	int64 Size;
 	int32 NumCacheSlots;
