Made the string builder types more generic

The string builder type hierarchy has been simplified to TStringBuilderBase<CharType> and TStringBuilderWithBuffer<CharType, BufferSize> that derives from it. All other builder types are now aliases of these two, and there are now aliases for wide builders.

The string literal path has been removed because it could cause unexpected behavior with const char arrays and it broke the use case (Builder << (Condition ? "A" : "B")). In practice, appending a string literal will typically calculate the length during compilation and compile to a call to Append(Str, Len).

The append operators have been simplified to append anything convertible to a string view.

[RN] minor Core

diff --git a/Engine/Source/Runtime/Core/Private/Misc/StringBuilder.cpp b/Engine/Source/Runtime/Core/Private/Misc/StringBuilder.cpp
--- a/Engine/Source/Runtime/Core/Private/Misc/StringBuilder.cpp
+++ b/Engine/Source/Runtime/Core/Private/Misc/StringBuilder.cpp
@@ -1,13 +1,9 @@
 // Copyright Epic Games, Inc. All Rights Reserved.
 
 #include "Misc/StringBuilder.h"
-#include "Containers/UnrealString.h"
-#include "Misc/AssertionMacros.h"
 #include "HAL/PlatformMath.h"
 #include "HAL/UnrealMemory.h"
 
-//#include "doctest.h"
-
 static inline uint64_t NextPowerOfTwo(uint64_t x)
 {
 	x -= 1;
@@ -20,28 +16,25 @@ static inline uint64_t NextPowerOfTwo(uint64_t x)
 	return x + 1;
 }
 
-template<typename C>
-TStringBuilderImpl<C>::~TStringBuilderImpl()
+template <typename C>
+TStringBuilderBase<C>::~TStringBuilderBase()
 {
 	if (bIsDynamic)
 	{
-		FreeBuffer(Base, UE_PTRDIFF_TO_INT32(End - Base));
+		FreeBuffer(Base, static_cast<SIZE_T>(End - Base));
 	}
 }
 
-template<typename C>
-void
-TStringBuilderImpl<C>::Extend(SIZE_T ExtraCapacity)
+template <typename C>
+void TStringBuilderBase<C>::Extend(SIZE_T ExtraCapacity)
 {
-	check(bIsExtendable);
-
 	const SIZE_T OldCapacity = End - Base;
 	const SIZE_T NewCapacity = NextPowerOfTwo(OldCapacity + ExtraCapacity);
 
 	C* NewBase = (C*)AllocBuffer(NewCapacity);
 
 	SIZE_T Pos = CurPos - Base;
-	memcpy(NewBase, Base, Pos * sizeof(C));
+	FMemory::Memcpy(NewBase, Base, Pos * sizeof(C));
 
 	if (bIsDynamic)
 	{
@@ -54,23 +47,20 @@ TStringBuilderImpl<C>::Extend(SIZE_T ExtraCapacity)
 	bIsDynamic	= true;
 }
 
-template<typename C>
-void*
-TStringBuilderImpl<C>::AllocBuffer(SIZE_T ByteCount)
+template <typename C>
+void* TStringBuilderBase<C>::AllocBuffer(SIZE_T CharCount)
 {
-	return FMemory::Malloc(ByteCount * sizeof(C));
+	return FMemory::Malloc(CharCount * sizeof(C));
 }
 
-template<typename C>
-void
-TStringBuilderImpl<C>::FreeBuffer(void* Buffer, SIZE_T ByteCount)
+template <typename C>
+void TStringBuilderBase<C>::FreeBuffer(void* Buffer, SIZE_T CharCount)
 {
 	FMemory::Free(Buffer);
 }
 
-template<typename C>
-typename TStringBuilderImpl<C>::BuilderType&
-TStringBuilderImpl<C>::AppendfImpl(BuilderType& Self, const C* Fmt, ...)
+template <typename C>
+TStringBuilderBase<C>& TStringBuilderBase<C>::AppendfImpl(BuilderType& Self, const C* Fmt, ...)
 {
 	for (;;)
 	{
@@ -93,69 +83,7 @@ TStringBuilderImpl<C>::AppendfImpl(BuilderType& Self, const C* Fmt, ...)
 	}
 }
 
-FStringBuilderBase&	operator<<(FStringBuilderBase& Builder, const FString& Str)
-{
-	return Builder.Append(*Str, Str.Len());
-}
-
-FStringBuilderBase&	operator+=(FStringBuilderBase& Builder, const FString& Str)
-{
-	return Builder.Append(*Str, Str.Len());
-}
-
 // Instantiate templates once
 
-template class TStringBuilderImpl<ANSICHAR>;
-template class TStringBuilderImpl<TCHAR>;
-
-#if 0
-#if !(UE_BUILD_SHIPPING || UE_BUILD_TEST)
-
-TEST_CASE("Core.StringView")
-{
-	{
-		TAnsiStringBuilder<32> boo;
-		CHECK(boo.Len() == 0);
-		CHECK(strcmp(boo.ToString(), "") == 0);
-
-		{
-			FAnsiStringView view = boo;
-			CHECK(boo.Len() == view.Len());
-		}
-
-		boo.Append("bababa");
-
-		CHECK(strcmp(boo.ToString(), "bababa") == 0);
-
-		{
-			FAnsiStringView view = boo;
-			CHECK(boo.Len() == view.Len());
-		}
-	}
-
-	{
-		TWideStringBuilder<32> booboo;
-		booboo.Append(TEXT("babababa"));
-
-		{
-			FWideStringView view = booboo;
-			CHECK(booboo.Len() == view.Len());
-		}
-
-		//booboo << TEXT("bobo");
-		//booboo << "baba";	// ANSI
-		//booboo << 'b';
-
-		char foo[] = "aaaaa";
-		char noo[] = "oOoOo";
-
-		booboo << noo;
-		booboo << foo;
-
-		FString ouch(TEXT("ouch"));
-		booboo << ouch;
-	}
-}
-
-#endif
-#endif
+template class TStringBuilderBase<ANSICHAR>;
+template class TStringBuilderBase<WIDECHAR>;
diff --git a/Engine/Source/Runtime/Core/Private/Tests/Misc/StringBuilderTest.cpp b/Engine/Source/Runtime/Core/Private/Tests/Misc/StringBuilderTest.cpp
new file mode 100644
--- /dev/null
+++ b/Engine/Source/Runtime/Core/Private/Tests/Misc/StringBuilderTest.cpp
@@ -0,0 +1,103 @@
+// Copyright Epic Games, Inc. All Rights Reserved.
+
+#include "Misc/StringBuilder.h"
+
+#include "Containers/UnrealString.h"
+#include "Misc/AutomationTest.h"
+
+static_assert(TIsSame<typename FStringBuilderBase::ElementType, TCHAR>::Value, "FStringBuilderBase must use TCHAR.");
+static_assert(TIsSame<typename FAnsiStringBuilderBase::ElementType, ANSICHAR>::Value, "FAnsiStringBuilderBase must use ANSICHAR.");
+static_assert(TIsSame<typename FWideStringBuilderBase::ElementType, WIDECHAR>::Value, "FWideStringBuilderBase must use WIDECHAR.");
+
+static_assert(TIsSame<FStringBuilderBase, TStringBuilderBase<TCHAR>>::Value, "FStringBuilderBase must be the same as TStringBuilderBase<TCHAR>.");
+static_assert(TIsSame<FAnsiStringBuilderBase, TStringBuilderBase<ANSICHAR>>::Value, "FAnsiStringBuilderBase must be the same as TStringBuilderBase<ANSICHAR>.");
+static_assert(TIsSame<FWideStringBuilderBase, TStringBuilderBase<WIDECHAR>>::Value, "FWideStringBuilderBase must be the same as TStringBuilderBase<WIDECHAR>.");
+
+static_assert(TIsContiguousContainer<FStringBuilderBase>::Value, "FStringBuilderBase must be a contiguous container.");
+static_assert(TIsContiguousContainer<FAnsiStringBuilderBase>::Value, "FAnsiStringBuilderBase must be a contiguous container.");
+static_assert(TIsContiguousContainer<FWideStringBuilderBase>::Value, "FWideStringBuilderBase must be a contiguous container.");
+
+static_assert(TIsContiguousContainer<TStringBuilder<128>>::Value, "TStringBuilder<N> must be a contiguous container.");
+static_assert(TIsContiguousContainer<TAnsiStringBuilder<128>>::Value, "TAnsiStringBuilder<N> must be a contiguous container.");
+static_assert(TIsContiguousContainer<TWideStringBuilder<128>>::Value, "TWideStringBuilder<N> must be a contiguous container.");
+
+#if WITH_DEV_AUTOMATION_TESTS
+
+IMPLEMENT_SIMPLE_AUTOMATION_TEST(FStringBuilderTestAppendString, "System.Core.StringBuilder.AppendString", EAutomationTestFlags::ApplicationContextMask | EAutomationTestFlags::EngineFilter)
+bool FStringBuilderTestAppendString::RunTest(const FString& Parameters)
+{
+	// Append Char
+	{
+		TStringBuilder<7> Builder;
+		Builder << TEXT('A') << TEXT('B') << TEXT('C');
+		Builder << 'D' << 'E' << 'F';
+		TestEqual(TEXT("Append Char"), FStringView(Builder), TEXT("ABCDEF"_SV));
+
+		TAnsiStringBuilder<4> AnsiBuilder;
+		AnsiBuilder << 'A' << 'B' << 'C';
+		TestEqual(TEXT("Append AnsiChar"), FAnsiStringView(AnsiBuilder), "ABC"_ASV);
+	}
+
+	// Append C String
+	{
+		TStringBuilder<7> Builder;
+		Builder << TEXT("ABC");
+		Builder << "DEF";
+		TestEqual(TEXT("Append C String"), FStringView(Builder), TEXT("ABCDEF"_SV));
+
+		TAnsiStringBuilder<4> AnsiBuilder;
+		AnsiBuilder << "ABC";
+		TestEqual(TEXT("Append Ansi C String"), FAnsiStringView(AnsiBuilder), "ABC"_ASV);
+	}
+
+	// Append FStringView
+	{
+		TStringBuilder<7> Builder;
+		Builder << TEXT("ABC"_SV);
+		Builder << "DEF"_ASV;
+		TestEqual(TEXT("Append FStringView"), FStringView(Builder), TEXT("ABCDEF"_SV));
+
+		TAnsiStringBuilder<4> AnsiBuilder;
+		AnsiBuilder << "ABC"_ASV;
+		TestEqual(TEXT("Append FAnsiStringView"), FAnsiStringView(AnsiBuilder), "ABC"_ASV);
+	}
+
+	// Append FStringBuilderBase
+	{
+		TStringBuilder<4> Builder;
+		Builder << TEXT("ABC");
+		TStringBuilder<4> BuilderCopy;
+		BuilderCopy << Builder;
+		TestEqual(TEXT("Append FStringBuilderBase"), FStringView(BuilderCopy), TEXT("ABC"_SV));
+
+		TAnsiStringBuilder<4> AnsiBuilder;
+		AnsiBuilder << "ABC";
+		TAnsiStringBuilder<4> AnsiBuilderCopy;
+		AnsiBuilderCopy << AnsiBuilder;
+		TestEqual(TEXT("Append FAnsiStringBuilderBase"), FAnsiStringView(AnsiBuilderCopy), "ABC"_ASV);
+	}
+
+	// Append FString
+	{
+		TStringBuilder<4> Builder;
+		Builder << FString(TEXT("ABC"));
+		TestEqual(TEXT("Append FString"), FStringView(Builder), TEXT("ABC"_SV));
+	}
+
+	// Append Char Array
+	{
+		const TCHAR String[16] = TEXT("ABC");
+		TStringBuilder<4> Builder;
+		Builder << String;
+		TestEqual(TEXT("Append Char Array"), FStringView(Builder), TEXT("ABC"_SV));
+
+		const ANSICHAR AnsiString[16] = "ABC";
+		TAnsiStringBuilder<4> AnsiBuilder;
+		AnsiBuilder << AnsiString;
+		TestEqual(TEXT("Append Char Array"), FAnsiStringView(AnsiBuilder), "ABC"_ASV);
+	}
+
+	return true;
+}
+
+#endif //WITH_DEV_AUTOMATION_TESTS
diff --git a/Engine/Source/Runtime/Core/Public/Containers/StringFwd.h b/Engine/Source/Runtime/Core/Public/Containers/StringFwd.h
--- a/Engine/Source/Runtime/Core/Public/Containers/StringFwd.h
+++ b/Engine/Source/Runtime/Core/Public/Containers/StringFwd.h
@@ -5,14 +5,31 @@
 #include "CoreTypes.h"
 
 // String Builder
-class FStringBuilderBase;
-class FAnsiStringBuilderBase;
-template <int32 N>
-class TStringBuilder;
-template <int32 N>
-class TAnsiStringBuilder;
+
+template <typename CharType> class TStringBuilderBase;
+template <typename CharType, int32 BufferSize> class TStringBuilderWithBuffer;
+
+/** The base string builder type for TCHAR. */
+using FStringBuilderBase = TStringBuilderBase<TCHAR>;
+/** The base string builder type for ANSICHAR. */
+using FAnsiStringBuilderBase = TStringBuilderBase<ANSICHAR>;
+/** The base string builder type for WIDECHAR. */
+using FWideStringBuilderBase = TStringBuilderBase<WIDECHAR>;
+
+/** An extendable string builder for TCHAR. */
+template <int32 BufferSize> using TStringBuilder = TStringBuilderWithBuffer<TCHAR, BufferSize>;
+/** An extendable string builder for ANSICHAR. */
+template <int32 BufferSize> using TAnsiStringBuilder = TStringBuilderWithBuffer<ANSICHAR, BufferSize>;
+/** An extendable string builder for WIDECHAR. */
+template <int32 BufferSize> using TWideStringBuilder = TStringBuilderWithBuffer<WIDECHAR, BufferSize>;
+
+/** A fixed-size string builder for TCHAR. */
+template <int32 BufferSize> using TFixedStringBuilder UE_DEPRECATED(4.25, "'TFixedStringBuilder' is deprecated. Please use 'TStringBuilder' instead!") = TStringBuilderWithBuffer<TCHAR, BufferSize>;
+/** A fixed-size string builder for ANSICHAR. */
+template <int32 BufferSize> using TFixedAnsiStringBuilder UE_DEPRECATED(4.25, "'TFixedAnsiStringBuilder' is deprecated. Please use 'TAnsiStringBuilder' instead!") = TStringBuilderWithBuffer<ANSICHAR, BufferSize>;
 
 // String View
+
 class FStringView;
 class FAnsiStringView;
 using FWideStringView = FStringView;
diff --git a/Engine/Source/Runtime/Core/Public/Misc/StringBuilder.h b/Engine/Source/Runtime/Core/Public/Misc/StringBuilder.h
--- a/Engine/Source/Runtime/Core/Public/Misc/StringBuilder.h
+++ b/Engine/Source/Runtime/Core/Public/Misc/StringBuilder.h
@@ -13,59 +13,45 @@
 #include "Templates/UnrealTemplate.h"
 #include "Traits/IsContiguousContainer.h"
 
-#define USE_STRING_LITERAL_PATH 1
-
-namespace StringBuilderPrivate
-{
-	template <typename CharType> struct TStringBuilderBaseType;
-	template <> struct TStringBuilderBaseType<ANSICHAR> { using Type = FAnsiStringBuilderBase; };
-	template <> struct TStringBuilderBaseType<WIDECHAR> { using Type = FStringBuilderBase; };
-}
-
-/** The string builder base type to be used by append operators and function output parameters for a given character type. */
+/**
+ * String Builder
+ *
+ * This class helps with the common task of constructing new strings.
+ *
+ * It does this by allocating buffer space which is used to hold the
+ * constructed string. The intent is that the builder is allocated on
+ * the stack as a function local variable to avoid heap allocations.
+ *
+ * The buffer is always contiguous and the class is not intended to be
+ * used to construct extremely large strings.
+ *
+ * This is not intended to be used as a mechanism for holding on to
+ * strings for a long time. The use case is explicitly to aid in
+ * *constructing* strings on the stack and subsequently passing the
+ * string into a function call or a more permanent string storage
+ * mechanism like FString et al.
+ *
+ * The amount of buffer space to allocate is specified via a template
+ * parameter and if the constructed string should overflow this initial
+ * buffer, a new buffer will be allocated using regular dynamic memory
+ * allocations. For instances where you absolutely must not allocate
+ * any memory, you should use the fixed variants named
+ * TFixed*StringBuilder.
+ *
+ * Overflow allocation should be the exceptional case however -- always
+ * try to size the buffer so that it can hold the vast majority of
+ * strings you expect to construct.
+ *
+ * Be mindful that stack is a limited resource, so if you are writing a
+ * highly recursive function you may want to use some other mechanism
+ * to build your strings.
+ */
 template <typename CharType>
-using TStringBuilderBase = typename StringBuilderPrivate::TStringBuilderBaseType<CharType>::Type;
-
-/** String Builder implementation
-
-	This class helps with the common task of constructing new strings. 
-	
-	It does this by allocating buffer space which is used to hold the 
-	constructed string. The intent is that the builder is allocated on
-	the stack as a function local variable to avoid heap allocations.
-
-	The buffer is always contiguous and the class is not intended to be
-	used to construct extremely large strings. 
-	
-	This is not intended to be used as a mechanism for holding on to 
-	strings for a long time - the use case is explicitly to aid in 
-	*constructing* strings on the stack	and subsequently passing the 
-	string into a function call or a more permanent string storage 
-	mechanism like FString et al
-
-	The amount of buffer space to allocate is specified via a template
-	parameter and if the constructed string should overflow this initial
-	buffer, a new buffer will be allocated using regular dynamic memory
-	allocations. For instances where you absolutely must not allocate
-	any memory, you should use the fixed variants named
-	TFixed*StringBuilder
-
-	Overflow allocation should be the exceptional case however -- always
-	try to size the buffer so that it can hold the vast majority of 
-	strings you expect to construct.
-
-	Be mindful that stack is a limited resource, so if you are writing a
-	highly recursive function you may want to use some other mechanism
-	to build your strings.
-
-  */
-
-template<typename C>
-class TStringBuilderImpl
+class TStringBuilderBase
 {
 public:
 	/** The character type that this builder operates on. */
-	using ElementType = C;
+	using ElementType = CharType;
 	/** The string builder base type to be used by append operators and function output parameters. */
 	using BuilderType = TStringBuilderBase<ElementType>;
 	/** The string view type that this builder is compatible with. */
@@ -79,24 +65,29 @@ public:
 	template <typename RangeType>
 	using TCanAppendRange = TAnd<TIsContiguousContainer<RangeType>, TCanAppend<decltype(*::GetData(DeclVal<RangeType>()))>>;
 
-				TStringBuilderImpl() = default;
-	CORE_API	~TStringBuilderImpl();
+				TStringBuilderBase() = default;
+	CORE_API	~TStringBuilderBase();
+
+				TStringBuilderBase(const TStringBuilderBase&) = delete;
+				TStringBuilderBase(TStringBuilderBase&&) = delete;
 
-				TStringBuilderImpl(const TStringBuilderImpl&) = delete;
-				TStringBuilderImpl(TStringBuilderImpl&&) = delete;
+	TStringBuilderBase& operator=(const TStringBuilderBase&) = delete;
+	TStringBuilderBase& operator=(TStringBuilderBase&&) = delete;
 
-	const TStringBuilderImpl& operator=(const TStringBuilderImpl&) = delete;
-	const TStringBuilderImpl& operator=(TStringBuilderImpl&&) = delete;
+	inline TStringBuilderBase(CharType* BufferPointer, int32 BufferCapacity)
+	{
+		Initialize(BufferPointer, BufferCapacity);
+	}
 
-	inline int32	Len() const			{ return int32(CurPos - Base); }
-	inline C* GetData()					{ return Base; }
-	inline const C* GetData() const		{ return Base; }
-	inline const C* ToString() const	{ EnsureNulTerminated(); return Base; }
-	inline const C* operator*() const	{ EnsureNulTerminated(); return Base; }
+	inline int32 Len() const					{ return int32(CurPos - Base); }
+	inline CharType* GetData()					{ return Base; }
+	inline const CharType* GetData() const		{ return Base; }
+	inline const CharType* ToString() const		{ EnsureNulTerminated(); return Base; }
+	inline const CharType* operator*() const	{ EnsureNulTerminated(); return Base; }
 
-	inline operator ViewType() const	{ return ViewType(Base, Len()); }
+	inline operator ViewType() const			{ return ViewType(Base, Len()); }
 
-	inline const C	LastChar() const	{ return *(CurPos - 1); }
+	inline const CharType	LastChar() const	{ return *(CurPos - 1); }
 
 	/**
 	 * Empties the string builder, but doesn't change memory allocation.
@@ -130,24 +121,20 @@ public:
 		CurPos -= InCount;
 	}
 
-	inline BuilderType& Append(C Char)
+	inline BuilderType& Append(CharType Char)
 	{
 		EnsureCapacity(1);
 
 		*CurPos++ = Char;
 
-		return BaseBuilder();
+		return *this;
 	}
 
-#if USE_STRING_LITERAL_PATH
-	inline BuilderType& AppendAnsi(const ANSICHAR*& NulTerminatedString)
-#else
 	inline BuilderType& AppendAnsi(const ANSICHAR* NulTerminatedString)
-#endif
 	{
 		if (!NulTerminatedString)
 		{
-			return BaseBuilder();
+			return *this;
 		}
 
 		return AppendAnsi(NulTerminatedString, TCString<ANSICHAR>::Strlen(NulTerminatedString));
@@ -162,7 +149,7 @@ public:
 	{
 		EnsureCapacity(Length);
 
-		C* RESTRICT Dest = CurPos;
+		CharType* RESTRICT Dest = CurPos;
 		CurPos += Length;
 
 		for (int32 i = 0; i < Length; ++i)
@@ -170,21 +157,17 @@ public:
 			Dest[i] = String[i];
 		}
 
-		return BaseBuilder();
+		return *this;
 	}
 
-#if USE_STRING_LITERAL_PATH
-	inline BuilderType& Append(const C*& NulTerminatedString)
-#else
-	inline BuilderType& Append(const C* NulTerminatedString)
-#endif
+	inline BuilderType& Append(const CharType* NulTerminatedString)
 	{
 		if (!NulTerminatedString)
 		{
-			return BaseBuilder();
+			return *this;
 		}
 
-		return Append(NulTerminatedString, TCString<C>::Strlen(NulTerminatedString));
+		return Append(NulTerminatedString, TCString<CharType>::Strlen(NulTerminatedString));
 	}
 
 	inline BuilderType& Append(const ViewType& StringView)
@@ -192,38 +175,16 @@ public:
 		return Append(StringView.GetData(), StringView.Len());
 	}
 
-	inline BuilderType& Append(const C* String, int32 Length)
+	inline BuilderType& Append(const CharType* String, int32 Length)
 	{
 		EnsureCapacity(Length);
-		C* RESTRICT Dest = CurPos;
+		CharType* RESTRICT Dest = CurPos;
 		CurPos += Length;
 
-		FMemory::Memcpy(Dest, String, Length * sizeof(C));
-
-		return BaseBuilder();
-	}
-
-#if USE_STRING_LITERAL_PATH
-	template<int32 ArrayLength>
-	inline BuilderType& Append(C(&CharArray)[ArrayLength])
-	{
-		return Append(&CharArray[0]);
-	}
-
-	template<int32 LiteralLength>
-	inline BuilderType& Append(const C (&StringLiteral)[LiteralLength])
-	{
-		constexpr int32 Length = LiteralLength - 1;
-		return Append(StringLiteral, Length);
-	}
+		FMemory::Memcpy(Dest, String, Length * sizeof(CharType));
 
-	template<int32 LiteralLength>
-	inline BuilderType& AppendAnsi(const ANSICHAR (&StringLiteral)[LiteralLength])
-	{
-		constexpr int32 Length = LiteralLength - 1;
-		return AppendAnsi(StringLiteral, Length);
+		return *this;
 	}
-#endif
 
 	/**
 	 * Append every element of the range to the builder, separating the elements by the delimiter.
@@ -240,7 +201,6 @@ public:
 		typename = typename TEnableIf<TAnd<TCanAppendRange<RangeType&&>, TCanAppend<DelimiterType&&>>::Value>::Type>
 	inline BuilderType& Join(RangeType&& InRange, DelimiterType&& InDelimiter)
 	{
-		BuilderType& Self = BaseBuilder();
 		bool bFirst = true;
 		for (auto&& Elem : Forward<RangeType>(InRange))
 		{
@@ -250,11 +210,11 @@ public:
 			}
 			else
 			{
-				Self << InDelimiter;
+				*this << InDelimiter;
 			}
-			Self << Elem;
+			*this << Elem;
 		}
-		return Self;
+		return *this;
 	}
 
 	/**
@@ -274,7 +234,6 @@ public:
 		typename = typename TEnableIf<TAnd<TCanAppendRange<RangeType>, TCanAppend<DelimiterType&&>, TCanAppend<QuoteType&&>>::Value>::Type>
 	inline BuilderType& JoinQuoted(RangeType&& InRange, DelimiterType&& InDelimiter, QuoteType&& InQuote)
 	{
-		BuilderType& Self = BaseBuilder();
 		bool bFirst = true;
 		for (auto&& Elem : Forward<RangeType>(InRange))
 		{
@@ -284,11 +243,11 @@ public:
 			}
 			else
 			{
-				Self << InDelimiter;
+				*this << InDelimiter;
 			}
-			Self << InQuote << Elem << InQuote;
+			*this << InQuote << Elem << InQuote;
 		}
-		return Self;
+		return *this;
 	}
 
 	/**
@@ -297,20 +256,18 @@ public:
 	 * @param Format A format string that specifies how to format the additional arguments. Refer to standard printf format.
 	 */
 	template <typename FmtType, typename... Types>
-	typename TEnableIf<TIsArrayOrRefOfType<FmtType, C>::Value, BuilderType&>::Type Appendf(const FmtType& Fmt, Types... Args)
+	typename TEnableIf<TIsArrayOrRefOfType<FmtType, CharType>::Value, BuilderType&>::Type Appendf(const FmtType& Fmt, Types... Args)
 	{
-		static_assert(TIsArrayOrRefOfType<FmtType, C>::Value, "Formatting string must be a character array.");
+		static_assert(TIsArrayOrRefOfType<FmtType, CharType>::Value, "Formatting string must be a character array.");
 		static_assert(TAnd<TIsValidVariadicFunctionArg<Types>...>::Value, "Invalid argument(s) passed to Appendf.");
-		return AppendfImpl(BaseBuilder(), Fmt, Forward<Types>(Args)...);
+		return AppendfImpl(*this, Fmt, Forward<Types>(Args)...);
 	}
 
 private:
-	CORE_API static BuilderType& VARARGS AppendfImpl(BuilderType& Self, const C* Fmt, ...);
-
-	inline BuilderType& BaseBuilder() { return static_cast<BuilderType&>(*this); }
+	CORE_API static BuilderType& VARARGS AppendfImpl(BuilderType& Self, const CharType* Fmt, ...);
 
 protected:
-	inline void Initialize(C* InBase, int32 InCapacity)
+	inline void Initialize(CharType* InBase, int32 InCapacity)
 	{
 		Base	= InBase;
 		CurPos	= InBase;
@@ -338,215 +295,98 @@ protected:
 		Extend(RequiredAdditionalCapacity);
 	}
 
-	CORE_API void	Extend(SIZE_T extraCapacity);
-	CORE_API void*	AllocBuffer(SIZE_T byteCount);
-	CORE_API void	FreeBuffer(void* buffer, SIZE_T byteCount);
+	CORE_API void	Extend(SIZE_T ExtraCapacity);
+	CORE_API void*	AllocBuffer(SIZE_T CharCount);
+	CORE_API void	FreeBuffer(void* Buffer, SIZE_T CharCount);
 
-	C*			Base;
-	C* 			CurPos;
-	C*			End;
+	CharType*	Base;
+	CharType*	CurPos;
+	CharType*	End;
 	bool		bIsDynamic = false;
-	bool		bIsExtendable = false;
 };
 
 template <typename CharType>
-struct TIsContiguousContainer<TStringBuilderImpl<CharType>> { enum { Value = true }; };
+struct TIsContiguousContainer<TStringBuilderBase<CharType>> { enum { Value = true }; };
 
 template <typename CharType>
-constexpr inline SIZE_T GetNum(const TStringBuilderImpl<CharType>& Builder)
+constexpr inline SIZE_T GetNum(const TStringBuilderBase<CharType>& Builder)
 {
 	return Builder.Len();
 }
 
 //////////////////////////////////////////////////////////////////////////
 
-class FAnsiStringBuilderBase : public TStringBuilderImpl<ANSICHAR>
-{
-protected:
-	inline FAnsiStringBuilderBase(ANSICHAR* BufferPointer, int32 BufferCapacity)
-	{
-		Initialize(BufferPointer, BufferCapacity);
-	}
-
-	~FAnsiStringBuilderBase() = default;
-};
-
-template <>
-struct TIsContiguousContainer<FAnsiStringBuilderBase> { enum { Value = true }; };
-
-template<int32 N>
-class UE_DEPRECATED(4.25, "'TFixedAnsiStringBuilder' is deprecated. Please use 'TAnsiStringBuilder' instead!") TFixedAnsiStringBuilder : public FAnsiStringBuilderBase
+/**
+ * A string builder with inline storage.
+ *
+ * Avoid using this type directly. Prefer the aliases in StringFwd.h like TStringBuilder<N>.
+ */
+template <typename CharType, int32 BufferSize>
+class TStringBuilderWithBuffer : public TStringBuilderBase<CharType>
 {
 public:
-	inline TFixedAnsiStringBuilder()
-	: FAnsiStringBuilderBase(StringBuffer, N)
+	inline TStringBuilderWithBuffer()
+		: TStringBuilderBase<CharType>(StringBuffer, BufferSize)
 	{
 	}
 
-	~TFixedAnsiStringBuilder() = default;
-
 private:
-	ANSICHAR	StringBuffer[N];
+	CharType StringBuffer[BufferSize];
 };
 
-PRAGMA_DISABLE_DEPRECATION_WARNINGS
-template <int32 N>
-struct TIsContiguousContainer<TFixedAnsiStringBuilder<N>> { enum { Value = true }; };
-PRAGMA_ENABLE_DEPRECATION_WARNINGS
-
-template<int32 N>
-class TAnsiStringBuilder : public FAnsiStringBuilderBase
-{
-public:
-	inline TAnsiStringBuilder()
-	: FAnsiStringBuilderBase(StringBuffer, N)
-	{
-		bIsExtendable = true;
-	}
-
-	~TAnsiStringBuilder() = default;
-
-private:
-	ANSICHAR	StringBuffer[N];
-};
-
-template <int32 N>
-struct TIsContiguousContainer<TAnsiStringBuilder<N>> { enum { Value = true }; };
+template <typename CharType, int32 BufferSize>
+struct TIsContiguousContainer<TStringBuilderWithBuffer<CharType, BufferSize>> { enum { Value = true }; };
 
 //////////////////////////////////////////////////////////////////////////
 
-class FStringBuilderBase : public TStringBuilderImpl<TCHAR>
-{
-protected:
-	FStringBuilderBase(TCHAR* BufferPointer, int32 BufferCapacity)
-	{
-		Initialize(BufferPointer, BufferCapacity);
-	}
-
-	~FStringBuilderBase() = default;
-};
+// String Append Operators
 
-template <>
-struct TIsContiguousContainer<FStringBuilderBase> { enum { Value = true }; };
+inline FAnsiStringBuilderBase&		operator<<(FAnsiStringBuilderBase& Builder, ANSICHAR Char)							{ return Builder.Append(Char); }
+inline FWideStringBuilderBase&		operator<<(FWideStringBuilderBase& Builder, ANSICHAR Char)							{ return Builder.Append(Char); }
+inline FWideStringBuilderBase&		operator<<(FWideStringBuilderBase& Builder, WIDECHAR Char)							{ return Builder.Append(Char); }
 
-template<int32 N>
-class UE_DEPRECATED(4.25, "'TFixedStringBuilder' is deprecated. Please use 'TStringBuilder' instead!") TFixedStringBuilder : public FStringBuilderBase
+template <typename T>
+inline auto operator<<(FAnsiStringBuilderBase& Builder, T&& Str) -> decltype(Builder.Append(ImplicitConv<FAnsiStringView>(Forward<T>(Str))))
 {
-public:
-	inline TFixedStringBuilder()
-	: FStringBuilderBase(StringBuffer, N)
-	{
-	}
-
-	~TFixedStringBuilder() = default;
-
-private:
-	TCHAR	StringBuffer[N];
-};
-
-PRAGMA_DISABLE_DEPRECATION_WARNINGS
-template <int32 N>
-struct TIsContiguousContainer<TFixedStringBuilder<N>> { enum { Value = true }; };
-PRAGMA_ENABLE_DEPRECATION_WARNINGS
+	return Builder.Append(ImplicitConv<FAnsiStringView>(Forward<T>(Str)));
+}
 
-template<int32 N>
-class TStringBuilder : public FStringBuilderBase
+template <typename T>
+inline auto operator<<(FWideStringBuilderBase& Builder, T&& Str) -> decltype(Builder.AppendAnsi(ImplicitConv<FAnsiStringView>(Forward<T>(Str))))
 {
-public:
-	inline TStringBuilder()
-	: FStringBuilderBase(StringBuffer, N)
-	{
-		bIsExtendable = true;
-	}
-
-	~TStringBuilder() = default;
-
-private:
-	TCHAR	StringBuffer[N];
-};
-
-template <int32 N>
-struct TIsContiguousContainer<TStringBuilder<N>> { enum { Value = true }; };
-
-// Append operator implementations
-
-inline FStringBuilderBase&					operator<<(FStringBuilderBase& Builder, ANSICHAR Char)							{ return Builder.Append(Char); }
-inline FStringBuilderBase&					operator<<(FStringBuilderBase& Builder, WIDECHAR Char)							{ return Builder.Append(Char); }
-
-#if USE_STRING_LITERAL_PATH
-
-inline FStringBuilderBase&					operator<<(FStringBuilderBase& Builder, const TCHAR*& NulTerminatedString)		{ return Builder.Append(NulTerminatedString); }
-inline FStringBuilderBase&					operator<<(FStringBuilderBase& Builder, const ANSICHAR*& NulTerminatedString)	{ return Builder.AppendAnsi(NulTerminatedString); }
-
-template<int32 N> FStringBuilderBase&		operator<<(FStringBuilderBase& Builder, const TCHAR (&StringLiteral)[N])		{ return Builder.Append(StringLiteral); }
-template<int32 N> FStringBuilderBase&		operator<<(FStringBuilderBase& Builder, const ANSICHAR (&StringLiteral)[N])		{ return Builder.AppendAnsi(StringLiteral); }
-
-template<int32 N> FStringBuilderBase&		operator<<(FStringBuilderBase& Builder, TCHAR (&NulTerminatedString)[N])		{ return Builder.Append(&NulTerminatedString[0], TCString<TCHAR>::Strlen(NulTerminatedString)); }
-template<int32 N> FStringBuilderBase&		operator<<(FStringBuilderBase& Builder, ANSICHAR (&NulTerminatedString)[N])		{ return Builder.AppendAnsi(&NulTerminatedString[0], TCString<ANSICHAR>::Strlen(NulTerminatedString)); }
-
-#else
-
-inline FStringBuilderBase&					operator<<(FStringBuilderBase& Builder, const TCHAR* NulTerminatedString)		{ return Builder.Append(NulTerminatedString); }
-inline FStringBuilderBase&					operator<<(FStringBuilderBase& Builder, const ANSICHAR* NulTerminatedString)	{ return Builder.AppendAnsi(NulTerminatedString); }
-
-#endif
-
-CORE_API FStringBuilderBase&				operator<<(FStringBuilderBase& Builder, const FString& Str);
-
-inline FStringBuilderBase&					operator<<(FStringBuilderBase& Builder, const FStringView& Str)					{ return Builder.Append(Str); }
-inline FStringBuilderBase&					operator<<(FStringBuilderBase& Builder, const FAnsiStringView& Str)				{ return Builder.AppendAnsi(Str); }
-
-inline FStringBuilderBase&					operator<<(FStringBuilderBase& Builder, int32 Value)							{ return Builder.Appendf(TEXT("%d"), Value); }
-inline FStringBuilderBase&					operator<<(FStringBuilderBase& Builder, uint32 Value)							{ return Builder.Appendf(TEXT("%u"), Value); }
-
-inline FStringBuilderBase&					operator<<(FStringBuilderBase& Builder, int64 Value)							{ return Builder.Appendf(TEXT(INT64_FMT), Value); }
-inline FStringBuilderBase&					operator<<(FStringBuilderBase& Builder, uint64 Value)							{ return Builder.Appendf(TEXT(UINT64_FMT), Value); }
-
-inline FStringBuilderBase&					operator<<(FStringBuilderBase& Builder, int8 Value)								{ return Builder << int32(Value); }
-inline FStringBuilderBase&					operator<<(FStringBuilderBase& Builder, uint8 Value)							{ return Builder << uint32(Value); }
-
-inline FStringBuilderBase&					operator<<(FStringBuilderBase& Builder, int16 Value)							{ return Builder << int32(Value); }
-inline FStringBuilderBase&					operator<<(FStringBuilderBase& Builder, uint16 Value)							{ return Builder << uint32(Value); }
-
-/** Prefer using << instead of +=. operator+= is only intended for mechanical FString -> FStringView replacement */
-
-inline FStringBuilderBase&					operator+=(FStringBuilderBase& Builder, FWideStringView Str)					{ return Builder.Append(Str); }
-inline FStringBuilderBase&					operator+=(FStringBuilderBase& Builder, FAnsiStringView Str)					{ return Builder.AppendAnsi(Str); }
-inline FStringBuilderBase&					operator+=(FStringBuilderBase& Builder, const WIDECHAR* Str)					{ return Builder.Append(Str); }
-inline FStringBuilderBase&					operator+=(FStringBuilderBase& Builder, const ANSICHAR* Str)					{ return Builder.AppendAnsi(Str); }
-
-CORE_API FStringBuilderBase&				operator+=(FStringBuilderBase& Builder, const FString& Str);
-inline FStringBuilderBase&					operator+=(FStringBuilderBase& Builder, ANSICHAR Char)							{ return Builder.Append(Char); }
-inline FStringBuilderBase&					operator+=(FStringBuilderBase& Builder, WIDECHAR Char)							{ return Builder.Append(Char); }
-
-
-// Append operator implementations for FAnsiStringBuilderBase
-
-inline FAnsiStringBuilderBase&				operator<<(FAnsiStringBuilderBase& Builder, const ANSICHAR Char)					{ return Builder.Append(Char); }
-
-#if USE_STRING_LITERAL_PATH
-
-inline FAnsiStringBuilderBase&				operator<<(FAnsiStringBuilderBase& Builder, const ANSICHAR*& NulTerminatedString)	{ return Builder.Append(NulTerminatedString); }
-template<int32 N> FAnsiStringBuilderBase&	operator<<(FAnsiStringBuilderBase& Builder, const ANSICHAR (&StringLiteral)[N])		{ return Builder.Append(StringLiteral); }
-template<int32 N> FAnsiStringBuilderBase&	operator<<(FAnsiStringBuilderBase& Builder, ANSICHAR (&NulTerminatedString)[N])		{ return Builder.Append(&NulTerminatedString[0], TCString<ANSICHAR>::Strlen(NulTerminatedString)); }
-
-#else
-
-inline FAnsiStringBuilderBase&				operator<<(FAnsiStringBuilderBase& Builder, const ANSICHAR* NulTerminatedString)	{ return Builder.Append(NulTerminatedString); }
-
-#endif
-
-template <typename T, typename = decltype(ImplicitConv<FAnsiStringView>(DeclVal<T>()))>
-inline FAnsiStringBuilderBase&				operator<<(FAnsiStringBuilderBase& Builder, T&& Str)								{ return Builder.Append(ImplicitConv<FAnsiStringView>(Forward<T>(Str))); }
-
-inline FAnsiStringBuilderBase&				operator<<(FAnsiStringBuilderBase& Builder, int32 Value)							{ return Builder.Appendf("%d", Value); }
-inline FAnsiStringBuilderBase&				operator<<(FAnsiStringBuilderBase& Builder, uint32 Value)							{ return Builder.Appendf("%u", Value); }
-
-inline FAnsiStringBuilderBase&				operator<<(FAnsiStringBuilderBase& Builder, int64 Value)							{ return Builder.Appendf(INT64_FMT, Value); }
-inline FAnsiStringBuilderBase&				operator<<(FAnsiStringBuilderBase& Builder, uint64 Value)							{ return Builder.Appendf(UINT64_FMT, Value); }
+	return Builder.AppendAnsi(ImplicitConv<FAnsiStringView>(Forward<T>(Str)));
+}
 
-inline FAnsiStringBuilderBase&				operator<<(FAnsiStringBuilderBase& Builder, int8 Value)								{ return Builder << int32(Value); }
-inline FAnsiStringBuilderBase&				operator<<(FAnsiStringBuilderBase& Builder, uint8 Value)							{ return Builder << uint32(Value); }
+template <typename T>
+inline auto operator<<(FWideStringBuilderBase& Builder, T&& Str) -> decltype(Builder.Append(ImplicitConv<FWideStringView>(Forward<T>(Str))))
+{
+	return Builder.Append(ImplicitConv<FWideStringView>(Forward<T>(Str)));
+}
 
-inline FAnsiStringBuilderBase&				operator<<(FAnsiStringBuilderBase& Builder, int16 Value)							{ return Builder << int32(Value); }
-inline FAnsiStringBuilderBase&				operator<<(FAnsiStringBuilderBase& Builder, uint16 Value)							{ return Builder << uint32(Value); }
+// Prefer using << instead of += as operator+= is only intended for mechanical FString -> FStringView replacement.
+inline FStringBuilderBase&			operator+=(FStringBuilderBase& Builder, ANSICHAR Char)								{ return Builder.Append(Char); }
+inline FStringBuilderBase&			operator+=(FStringBuilderBase& Builder, WIDECHAR Char)								{ return Builder.Append(Char); }
+inline FStringBuilderBase&			operator+=(FStringBuilderBase& Builder, FAnsiStringView Str)						{ return Builder.AppendAnsi(Str); }
+inline FStringBuilderBase&			operator+=(FStringBuilderBase& Builder, FWideStringView Str)						{ return Builder.Append(Str); }
+
+// Integer Append Operators
+
+inline FAnsiStringBuilderBase&		operator<<(FAnsiStringBuilderBase& Builder, int32 Value)							{ return Builder.Appendf("%d", Value); }
+inline FAnsiStringBuilderBase&		operator<<(FAnsiStringBuilderBase& Builder, uint32 Value)							{ return Builder.Appendf("%u", Value); }
+inline FWideStringBuilderBase&		operator<<(FWideStringBuilderBase& Builder, int32 Value)							{ return Builder.Appendf(TEXT("%d"), Value); }
+inline FWideStringBuilderBase&		operator<<(FWideStringBuilderBase& Builder, uint32 Value)							{ return Builder.Appendf(TEXT("%u"), Value); }
+
+inline FAnsiStringBuilderBase&		operator<<(FAnsiStringBuilderBase& Builder, int64 Value)							{ return Builder.Appendf(INT64_FMT, Value); }
+inline FAnsiStringBuilderBase&		operator<<(FAnsiStringBuilderBase& Builder, uint64 Value)							{ return Builder.Appendf(UINT64_FMT, Value); }
+inline FWideStringBuilderBase&		operator<<(FWideStringBuilderBase& Builder, int64 Value)							{ return Builder.Appendf(TEXT(INT64_FMT), Value); }
+inline FWideStringBuilderBase&		operator<<(FWideStringBuilderBase& Builder, uint64 Value)							{ return Builder.Appendf(TEXT(UINT64_FMT), Value); }
+
+inline FAnsiStringBuilderBase&		operator<<(FAnsiStringBuilderBase& Builder, int8 Value)								{ return Builder << int32(Value); }
+inline FAnsiStringBuilderBase&		operator<<(FAnsiStringBuilderBase& Builder, uint8 Value)							{ return Builder << uint32(Value); }
+inline FWideStringBuilderBase&		operator<<(FWideStringBuilderBase& Builder, int8 Value)								{ return Builder << int32(Value); }
+inline FWideStringBuilderBase&		operator<<(FWideStringBuilderBase& Builder, uint8 Value)							{ return Builder << uint32(Value); }
+
+inline FAnsiStringBuilderBase&		operator<<(FAnsiStringBuilderBase& Builder, int16 Value)							{ return Builder << int32(Value); }
+inline FAnsiStringBuilderBase&		operator<<(FAnsiStringBuilderBase& Builder, uint16 Value)							{ return Builder << uint32(Value); }
+inline FWideStringBuilderBase&		operator<<(FWideStringBuilderBase& Builder, int16 Value)							{ return Builder << int32(Value); }
+inline FWideStringBuilderBase&		operator<<(FWideStringBuilderBase& Builder, uint16 Value)							{ return Builder << uint32(Value); }
