Append a basic OS version check to iOS crash and MTBF metrics so we can filter out expected stability issues on unsupported IOS versions. Version is defined by the project cs as a series of defines for Major, Minor and Patch versions.

diff --git a/Engine/Source/Runtime/Core/Private/IOS/IOSPlatformMisc.cpp b/Engine/Source/Runtime/Core/Private/IOS/IOSPlatformMisc.cpp
--- a/Engine/Source/Runtime/Core/Private/IOS/IOSPlatformMisc.cpp
+++ b/Engine/Source/Runtime/Core/Private/IOS/IOSPlatformMisc.cpp
@@ -1323,6 +1323,15 @@ FString FIOSPlatformMisc::GetProjectVersion()
 	return localVersionString;
 }
 
+bool FIOSPlatformMisc::IsSupportedIOSVersion()
+{
+#if defined(IOS_MIN_SUPPORTED_VERSION_MAJOR) && defined(IOS_MIN_SUPPORTED_VERSION_MINOR) && defined(IOS_MIN_SUPPORTED_VERSION_PATCH)
+	return IOSVersionCompare(IOS_MIN_SUPPORTED_VERSION_MAJOR, IOS_MIN_SUPPORTED_VERSION_MINOR, IOS_MIN_SUPPORTED_VERSION_PATCH) >= 0;
+#else
+	return true;
+#endif
+}
+
 bool FIOSPlatformMisc::RequestDeviceCheckToken(TFunction<void(const TArray<uint8>&)> QuerySucceededFunc, TFunction<void(const FString&, const FString&)> QueryFailedFunc)
 {
 	DCDevice* DeviceCheckDevice = [DCDevice currentDevice];
diff --git a/Engine/Source/Runtime/Core/Public/IOS/IOSPlatformMisc.h b/Engine/Source/Runtime/Core/Public/IOS/IOSPlatformMisc.h
--- a/Engine/Source/Runtime/Core/Public/IOS/IOSPlatformMisc.h
+++ b/Engine/Source/Runtime/Core/Public/IOS/IOSPlatformMisc.h
@@ -105,6 +105,8 @@ struct CORE_API FIOSPlatformMisc : public FApplePlatformMisc
 	static void RequestStoreReview();
 
 	static bool IsUpdateAvailable();
+
+	static bool IsSupportedIOSVersion();
 	
 	// Possible iOS devices
 	enum EIOSDevice
