Log memory stats on a master process before he forks a new child process

diff --git a/Engine/Source/Runtime/Core/Private/Unix/UnixPlatformProcess.cpp b/Engine/Source/Runtime/Core/Private/Unix/UnixPlatformProcess.cpp
--- a/Engine/Source/Runtime/Core/Private/Unix/UnixPlatformProcess.cpp
+++ b/Engine/Source/Runtime/Core/Private/Unix/UnixPlatformProcess.cpp
@@ -1360,6 +1360,8 @@ FGenericPlatformProcess::EWaitAndForkResult FUnixPlatformProcess::WaitAndFork()
 	UE_LOG(LogHAL, Log, TEXT("   *** WaitAndFork awaiting signal %d to create child processes... ***"), WAIT_AND_FORK_QUEUE_SIGNAL);
 	GLog->Flush();
 
+	FPlatformMemoryStats PreviousMasterMemStats = FPlatformMemory::GetStats();
+
 	EWaitAndForkResult RetVal = EWaitAndForkResult::Parent;
 	struct FPidAndSignal
 	{
@@ -1378,6 +1380,15 @@ FGenericPlatformProcess::EWaitAndForkResult FUnixPlatformProcess::WaitAndFork()
 		{
 			// Sleep for a short while to avoid spamming new processes to the OS all at once
 			FPlatformProcess::Sleep(WAIT_AND_FORK_CHILD_SPAWN_DELAY);
+
+			FPlatformMemoryStats CurrentMasterMemStats = FPlatformMemory::GetStats();
+			UE_LOG(LogHal, Log, TEXT("MemoryStats PreFork: AvailablePhysical: %.02fMiB (%+.02fMiB), PeakPhysical: %.02fMiB (%+.02fMiB), PeakVirtual: %.02fMiB (%+.02fMiB)"),
+				CurrentMasterMemStats.AvailablePhysical / (1024.f*1024.f), (CurrentMasterMemStats.AvailablePhysical - PreviousMasterMemStats.AvailablePhysical) / (1024.f*1024.f),
+				CurrentMasterMemStats.PeakUsedPhysical / (1024.f*1024.f), (CurrentMasterMemStats.PeakUsedPhysical - PreviousMasterMemStats.PeakUsedPhysical) / (1024.f*1024.f),
+				CurrentMasterMemStats.PeakUsedVirtual / (1024.f*1024.f), (CurrentMasterMemStats.PeakUsedVirtual - PreviousMasterMemStats.PeakUsedVirtual) / (1024.f*1024.f)
+				
+			);
+			PreviousMasterMemStats = CurrentMasterMemStats;
 			
 			// Make sure there are no pending messages in the log.
 			GLog->Flush();
