Fix linux compil error due to typo in Log

diff --git a/Engine/Source/Runtime/Core/Private/Unix/UnixPlatformProcess.cpp b/Engine/Source/Runtime/Core/Private/Unix/UnixPlatformProcess.cpp
--- a/Engine/Source/Runtime/Core/Private/Unix/UnixPlatformProcess.cpp
+++ b/Engine/Source/Runtime/Core/Private/Unix/UnixPlatformProcess.cpp
@@ -1382,7 +1382,7 @@ FGenericPlatformProcess::EWaitAndForkResult FUnixPlatformProcess::WaitAndFork()
 			FPlatformProcess::Sleep(WAIT_AND_FORK_CHILD_SPAWN_DELAY);
 
 			FPlatformMemoryStats CurrentMasterMemStats = FPlatformMemory::GetStats();
-			UE_LOG(LogHal, Log, TEXT("MemoryStats PreFork: AvailablePhysical: %.02fMiB (%+.02fMiB), PeakPhysical: %.02fMiB (%+.02fMiB), PeakVirtual: %.02fMiB (%+.02fMiB)"),
+			UE_LOG(LogHAL, Log, TEXT("MemoryStats PreFork: AvailablePhysical: %.02fMiB (%+.02fMiB), PeakPhysical: %.02fMiB (%+.02fMiB), PeakVirtual: %.02fMiB (%+.02fMiB)"),
 				CurrentMasterMemStats.AvailablePhysical / (1024.f*1024.f), (CurrentMasterMemStats.AvailablePhysical - PreviousMasterMemStats.AvailablePhysical) / (1024.f*1024.f),
 				CurrentMasterMemStats.PeakUsedPhysical / (1024.f*1024.f), (CurrentMasterMemStats.PeakUsedPhysical - PreviousMasterMemStats.PeakUsedPhysical) / (1024.f*1024.f),
 				CurrentMasterMemStats.PeakUsedVirtual / (1024.f*1024.f), (CurrentMasterMemStats.PeakUsedVirtual - PreviousMasterMemStats.PeakUsedVirtual) / (1024.f*1024.f)
