Fixed an unterminated string being appended when processing Unicode escape sequences in FParse::QuotedString

#rnx

diff --git a/Engine/Source/Runtime/Core/Private/Misc/Parse.cpp b/Engine/Source/Runtime/Core/Private/Misc/Parse.cpp
--- a/Engine/Source/Runtime/Core/Private/Misc/Parse.cpp
+++ b/Engine/Source/Runtime/Core/Private/Misc/Parse.cpp
@@ -428,27 +428,29 @@ bool ParseQuotedString( const TCHAR* Buffer, T& Value, int32* OutNumCharsRead)
 		{
 			++Buffer;
 
-			TStringBuilder<16> UnicodeSequence;
+			TStringBuilder<4> UnicodeSequence;
 			while (ShouldParse(*Buffer) && FChar::IsHexDigit(*Buffer) && UnicodeSequence.Len() < 4) // UTF-16 sequences can only be up-to 4 digits long
 			{
 				UnicodeSequence += *Buffer++;
 			}
 
 			const UTF16CHAR Utf16Char = static_cast<UTF16CHAR>(FCString::Strtoi(*UnicodeSequence, nullptr, 16));
-			Value += FStringView(FUTF16ToTCHAR(&Utf16Char, /* Len */ 1).Get());
+			const FUTF16ToTCHAR Utf16Str(&Utf16Char, /* Len */ 1);
+			Value += FStringView(Utf16Str.Get(), Utf16Str.Length());
 		}
 		else if (*Buffer == TCHAR('U') && FChar::IsHexDigit(*(Buffer + 1))) // UTF-32 sequence (\U12345678)
 		{
 			++Buffer;
 
-			TStringBuilder<16> UnicodeSequence;
+			TStringBuilder<8> UnicodeSequence;
 			while (ShouldParse(*Buffer) && FChar::IsHexDigit(*Buffer) && UnicodeSequence.Len() < 8) // UTF-32 sequences can only be up-to 8 digits long
 			{
 				UnicodeSequence += *Buffer++;
 			}
 
-			const UTF32CHAR Utf32Str = static_cast<UTF32CHAR>(FCString::Strtoi(*UnicodeSequence, nullptr, 16));
-			Value += FStringView(FUTF32ToTCHAR(&Utf32Str, /* Len */ 1).Get());
+			const UTF32CHAR Utf32Char = static_cast<UTF32CHAR>(FCString::Strtoi(*UnicodeSequence, nullptr, 16));
+			const FUTF32ToTCHAR Utf32Str(&Utf32Char, /* Len */ 1);
+			Value += FStringView(Utf32Str.Get(), Utf32Str.Length());
 		}
 		else // unhandled escape sequence
 		{
