TaskGraph can now multithreaded on a forked server with -PostForkThreading
The Taskgpraph needs to be shutdown and recreated entirely by the forked process for the switch to occur.
The functionnality can be disabled via "TaskGraph.EnableForkedMultithreading 0" if needed.

diff --git a/Engine/Source/Runtime/Core/Private/Async/TaskGraph.cpp b/Engine/Source/Runtime/Core/Private/Async/TaskGraph.cpp
--- a/Engine/Source/Runtime/Core/Private/Async/TaskGraph.cpp
+++ b/Engine/Source/Runtime/Core/Private/Async/TaskGraph.cpp
@@ -23,6 +23,7 @@
 #include "Math/RandomStream.h"
 #include "HAL/IConsoleManager.h"
 #include "Misc/App.h"
+#include "Misc/Fork.h"
 #include "Containers/LockFreeFixedSizeAllocator.h"
 #include "Async/TaskGraphInterfaces.h"
 #include "HAL/LowLevelMemTracker.h"
@@ -81,6 +82,13 @@ static FAutoConsoleVariableRef CVarEnablePowerSavingThreadPriorityReduction(
 	TEXT("If 1, then high pri thread tasks which are marked EPowerSavingEligibility::Eligible can be dropped to normal priority.")
 );
 
+CORE_API bool GAllowTaskGraphForkMultithreading = true;
+static FAutoConsoleVariableRef CVarEnableForkedMultithreading(
+	TEXT("TaskGraph.EnableForkedMultithreading"),
+	GAllowTaskGraphForkMultithreading,
+	TEXT("When false will prevent the task graph from running multithreaded on forked processes.")
+);
+
 #if CREATE_HIPRI_TASK_THREADS || CREATE_BACKGROUND_TASK_THREADS
 	static void ThreadSwitchForABTest(const TArray<FString>& Args)
 	{
@@ -581,10 +589,12 @@ public:
 
 		Queue(QueueIndex).QuitForReturn = false;
 		verify(++Queue(QueueIndex).RecursionGuard == 1);
+		const bool bIsMultiThread = FTaskGraphInterface::IsMultithread();
 		do
 		{
-			ProcessTasksNamedThread(QueueIndex, FPlatformProcess::SupportsMultithreading());
-		} while (!Queue(QueueIndex).QuitForReturn && !Queue(QueueIndex).QuitForShutdown && FPlatformProcess::SupportsMultithreading()); // @Hack - quit now when running with only one thread.
+			const bool bAllowStall = bIsMultiThread;
+			ProcessTasksNamedThread(QueueIndex, bAllowStall);
+		} while (!Queue(QueueIndex).QuitForReturn && !Queue(QueueIndex).QuitForShutdown && bIsMultiThread); // @Hack - quit now when running with only one thread.
 		verify(!--Queue(QueueIndex).RecursionGuard);
 	}
 
@@ -857,15 +867,16 @@ public:
 			FMemory::SetupTLSCachesOnCurrentThread();
 		}
 		check(!QueueIndex);
+		const bool bIsMultiThread = FTaskGraphInterface::IsMultithread();
 		do
 		{
 			ProcessTasks();			
-		} while (!Queue.QuitForShutdown && FPlatformProcess::SupportsMultithreading()); // @Hack - quit now when running with only one thread.
+		} while (!Queue.QuitForShutdown && bIsMultiThread); // @Hack - quit now when running with only one thread.
 	}
 
 	virtual uint64 ProcessTasksUntilIdle(int32 QueueIndex) override
 	{
-		if (!FPlatformProcess::SupportsMultithreading())
+		if (FTaskGraphInterface::IsMultithread() == false)
 		{
 			return ProcessTasks();
 		}
@@ -1006,13 +1017,14 @@ private:
 #endif
 
 				TestRandomizedThreads();
-				if (FPlatformProcess::SupportsMultithreading())
+				const bool bIsMultithread = FTaskGraphInterface::IsMultithread();
+				if (bIsMultithread)
 				{
 					FScopeCycleCounter Scope(StallStatId);
 					Queue.StallRestartEvent->Wait(MAX_uint32, bCountAsStall);
 					bDidStall = true;
 				}
-				if (Queue.QuitForShutdown || !FPlatformProcess::SupportsMultithreading())
+				if (Queue.QuitForShutdown || !bIsMultithread)
 				{
 					break;
 				}
@@ -1162,7 +1174,7 @@ public:
 		int32 NumTaskThreads = FPlatformMisc::NumberOfWorkerThreadsToSpawn();
 
 		// if we don't want any performance-based threads, then force the task graph to not create any worker threads, and run in game thread
-		if (!FPlatformProcess::SupportsMultithreading())
+		if (!FTaskGraphInterface::IsMultithread())
 		{
 			// this is the logic that used to be spread over a couple of places, that will make the rest of this function disable a worker thread
 			// @todo: it could probably be made simpler/clearer
@@ -1225,6 +1237,7 @@ public:
 			FString Name;
 			const ANSICHAR* GroupName = "TaskGraphNormal";
 			int32 Priority = ThreadIndexToPriorityIndex(ThreadIndex);
+            // These are below normal threads so that they sleep when the named threads are active
 			EThreadPriority ThreadPri;
 			uint64 Affinity = FPlatformAffinity::GetTaskGraphThreadMask();
 			if (Priority == 1)
@@ -1262,7 +1275,16 @@ public:
 #else
 			uint32 StackSize = 512 * 1024;
 #endif
-			WorkerThreads[ThreadIndex].RunnableThread = FRunnableThread::Create(&Thread(ThreadIndex), *Name, StackSize, ThreadPri, Affinity); // these are below normal threads so that they sleep when the named threads are active
+            // We only create forkable threads on the Forked instance since the TaskGraph needs to be shutdown and recreated to properly make the switch from singlethread to multithread.
+			if (FForkProcessHelper::IsForkedMultithreadInstance() && GAllowTaskGraphForkMultithreading)
+			{
+				WorkerThreads[ThreadIndex].RunnableThread = FForkProcessHelper::CreateForkableThread(&Thread(ThreadIndex), *Name, StackSize, ThreadPri, Affinity);
+			}
+			else
+			{
+				WorkerThreads[ThreadIndex].RunnableThread = FRunnableThread::Create(&Thread(ThreadIndex), *Name, StackSize, ThreadPri, Affinity); 
+			}
+			
 			WorkerThreads[ThreadIndex].bAttached = true;
 			if (WorkerThreads[ThreadIndex].RunnableThread)
 			{
@@ -1315,7 +1337,7 @@ public:
 		if (ENamedThreads::GetThreadIndex(ThreadToExecuteOn) == ENamedThreads::AnyThread)
 		{
 			TASKGRAPH_SCOPE_CYCLE_COUNTER(3, STAT_TaskGraph_QueueTask_AnyThread);
-			if (FPlatformProcess::SupportsMultithreading())
+			if (FTaskGraphInterface::IsMultithread())
 			{
 				uint32 TaskPriority = ENamedThreads::GetTaskPriority(Task->ThreadToExecuteOn);
 				int32 Priority = ENamedThreads::GetThreadPriorityIndex(Task->ThreadToExecuteOn);
@@ -1478,7 +1500,7 @@ public:
 		}
 		else
 		{
-			if (!FPlatformProcess::SupportsMultithreading())
+			if (!FTaskGraphInterface::IsMultithread())
 			{
 				bool bAnyPending = false;
 				for (int32 Index = 0; Index < Tasks.Num(); Index++)
@@ -1709,6 +1731,11 @@ FTaskGraphInterface& FTaskGraphInterface::Get()
 	return *TaskGraphImplementationSingleton;
 }
 
+bool FTaskGraphInterface::IsMultithread()
+{
+	return FPlatformProcess::SupportsMultithreading() || (FForkProcessHelper::IsForkedMultithreadInstance() && GAllowTaskGraphForkMultithreading);
+}
+
 
 // Statics and some implementations from FBaseGraphTask and FGraphEvent
 
@@ -2216,7 +2243,7 @@ static void TaskGraphBenchmark(const TArray<FString>& Args)
 	FThreadSafeCounter Counter;
 	FThreadSafeCounter Cycles;
 
-	if (!FPlatformProcess::SupportsMultithreading())
+	if (!FTaskGraphInterface::IsMultithread())
 	{
 		UE_LOG(LogConsoleResponse, Display, TEXT("WARNING: TaskGraphBenchmark disabled for non multi-threading platforms"));
 		return;
@@ -2518,7 +2545,7 @@ static void TestLockFree(int32 OuterIters = 3)
 	TGuardValue<int32> ReentrantGuard(GPrintBroadcastWarnings, 0);
 
 
-	if (!FPlatformProcess::SupportsMultithreading())
+	if (!FTaskGraphInterface::IsMultithread())
 	{
 		UE_LOG(LogConsoleResponse, Display, TEXT("WARNING: TestLockFree disabled for non multi-threading platforms"));
 		return;
diff --git a/Engine/Source/Runtime/Core/Public/Async/ParallelFor.h b/Engine/Source/Runtime/Core/Public/Async/ParallelFor.h
--- a/Engine/Source/Runtime/Core/Public/Async/ParallelFor.h
+++ b/Engine/Source/Runtime/Core/Public/Async/ParallelFor.h
@@ -15,6 +15,7 @@
 #include "Stats/Stats.h"
 #include "Async/TaskGraphInterfaces.h"
 #include "Misc/App.h"
+#include "Misc/Fork.h"
 
 // Flags controlling the ParallelFor's behavior.
 enum class EParallelForFlags
@@ -193,7 +194,8 @@ namespace ParallelForImpl
 		check(Num >= 0);
 
 		int32 AnyThreadTasks = 0;
-		if (Num > 1 && (Flags & EParallelForFlags::ForceSingleThread) == EParallelForFlags::None && FApp::ShouldUseThreadingForPerformance())
+		const bool bIsMultithread = FApp::ShouldUseThreadingForPerformance() || FForkProcessHelper::IsForkedMultithreadInstance();
+		if (Num > 1 && (Flags & EParallelForFlags::ForceSingleThread) == EParallelForFlags::None && bIsMultithread)
 		{
 			AnyThreadTasks = FMath::Min<int32>(FTaskGraphInterface::Get().GetNumWorkerThreads(), Num - 1);
 		}
@@ -251,7 +253,8 @@ namespace ParallelForImpl
 		SCOPE_CYCLE_COUNTER(STAT_ParallelFor);
 
 		int32 AnyThreadTasks = 0;
-		if ((Flags & EParallelForFlags::ForceSingleThread) == EParallelForFlags::None && FApp::ShouldUseThreadingForPerformance())
+		const bool bIsMultithread = FApp::ShouldUseThreadingForPerformance() || FForkProcessHelper::IsForkedMultithreadInstance();
+		if ((Flags & EParallelForFlags::ForceSingleThread) == EParallelForFlags::None && bIsMultithread)
 		{
 			AnyThreadTasks = FMath::Min<int32>(FTaskGraphInterface::Get().GetNumWorkerThreads(), Num);
 		}
diff --git a/Engine/Source/Runtime/Core/Public/Async/TaskGraphInterfaces.h b/Engine/Source/Runtime/Core/Public/Async/TaskGraphInterfaces.h
--- a/Engine/Source/Runtime/Core/Public/Async/TaskGraphInterfaces.h
+++ b/Engine/Source/Runtime/Core/Public/Async/TaskGraphInterfaces.h
@@ -318,6 +318,12 @@ public:
 	**/
 	static CORE_API FTaskGraphInterface& Get();
 
+	/**
+	* The task graph is always multi-threaded for platforms that support it.
+	* For forked processes, the taskgraph will be singlethread for the master process but becomes multithread in the forked process
+	*/
+	static bool IsMultithread();
+
 	/** Return the current thread type, if known. **/
 	virtual ENamedThreads::Type GetCurrentThreadIfKnown(bool bLocalQueue = false) = 0;
 
diff --git a/Engine/Source/Runtime/Engine/Private/TickTaskManager.cpp b/Engine/Source/Runtime/Engine/Private/TickTaskManager.cpp
--- a/Engine/Source/Runtime/Engine/Private/TickTaskManager.cpp
+++ b/Engine/Source/Runtime/Engine/Private/TickTaskManager.cpp
@@ -405,6 +405,11 @@ public:
 	**/
 	FORCEINLINE static bool SingleThreadedMode()
 	{
+		if (FForkProcessHelper::IsForkedMultithreadInstance())
+		{
+			return false;
+		}
+
 		if (!FApp::ShouldUseThreadingForPerformance() || IsRunningDedicatedServer() || FPlatformMisc::NumberOfCores() < 3 || !FPlatformProcess::SupportsMultithreading())
 		{
 			return true;
