Prevent thread id clashes between real and fake threads
* Use 32bit MSB as mask to assign to fake thread ids exclusively
* Added check to detect if a kernel thread id potentially clashes with a fake thread id.

diff --git a/Engine/Source/Runtime/Core/Private/HAL/ThreadingBase.cpp b/Engine/Source/Runtime/Core/Private/HAL/ThreadingBase.cpp
--- a/Engine/Source/Runtime/Core/Private/HAL/ThreadingBase.cpp
+++ b/Engine/Source/Runtime/Core/Private/HAL/ThreadingBase.cpp
@@ -114,12 +114,16 @@ protected:
 
 public:
 
+	/** Use the MSB as a mask to prevent clashes between kernel assigned thread ids and fake thread ids */
+	static constexpr uint32 FakeIdReservedBit = 1 << 31;
+
 	/** Constructor. */
 	FFakeThread()
 		: bIsSuspended(false)
 		, SingleThreadRunnable(nullptr)
 	{
 		ThreadID = ThreadIdCounter++;
+		ThreadID |= FakeIdReservedBit;
 		// Auto register with single thread manager.
 		FThreadManager::Get().AddThread(ThreadID, this);
 	}
@@ -196,7 +200,15 @@ uint32 FFakeThread::ThreadIdCounter = 0xffff;
 
 void FThreadManager::AddThread(uint32 ThreadId, FRunnableThread* Thread)
 {
-	FScopeLock ThreadsLock(&ThreadsCritical);
+	const bool bIsSingleThreadEnvironment = FPlatformProcess::SupportsMultithreading() == false;
+
+	if (bIsSingleThreadEnvironment && Thread->GetThreadType() == FRunnableThread::ThreadType::Real)
+	{
+		checkf((ThreadId & FFakeThread::FakeIdReservedBit) == 0, TEXT("The thread ID  assigned by the kernel clashes with the bit reserved for identifying fake threads. Need to revisit the fake ID assignment algo."));
+	}
+
+	FScopeLock ThreadsLock(&ThreadsCritical);	
+
 	// Some platforms do not support TLS
 	if (!Threads.Contains(ThreadId))
 	{
