Re-enable inlining for FApp::CanEverRender for MSVC 1926 onwards now that a fix is confirmed in16.6 Preview 2.

#rnx

diff --git a/Engine/Source/Runtime/Core/Public/Misc/App.h b/Engine/Source/Runtime/Core/Public/Misc/App.h
--- a/Engine/Source/Runtime/Core/Public/Misc/App.h
+++ b/Engine/Source/Runtime/Core/Public/Misc/App.h
@@ -348,8 +348,9 @@ public:
 
 public:
 
-	// MSVC 16.4 has a bug that does not properly handle the local static bool in CanEverRender. This will be fixed in 16.5. Working around by using FORCENOINLINE.
-#if !UE_SERVER && defined(_MSC_VER) && _MSC_VER >= 1924
+// MSVC 16.4 (1924) has a bug that does not properly handle the local static bool in CanEverRender. This will be fixed in 16.5. Working around by using FORCENOINLINE.
+// MSVC 16.6 (1926) fixes the problem so inlning can be enabled again from that point onwards.
+#if !UE_SERVER && defined(_MSC_VER) && _MSC_VER >= 1924 && _MSC_VER <= 1925
 #define INLINE_CANEVERRENDER FORCENOINLINE
 #else
 #define INLINE_CANEVERRENDER FORCEINLINE
