Extended FCookStats to support new use cases

FCookStats::FScopedStatsCounter is now movable and copying is blocked to avoid incorrect behavior. This supports usage patterns such as: COOK_STAT(auto Timer = Condition ? UsageStats.TimeSyncWork() : UsageStats.TimeAsyncWork());

FCookStats::FDDCResourceUsage has been extracted into a standalone type, FDerivedDataUsageStats. The name is more in line with other DDC type names, and can now be forward-declared when necessary.

[RN] minor Core/Cooker

diff --git a/Engine/Source/Runtime/Core/Public/ProfilingDebugging/CookStats.h b/Engine/Source/Runtime/Core/Public/ProfilingDebugging/CookStats.h
--- a/Engine/Source/Runtime/Core/Public/ProfilingDebugging/CookStats.h
+++ b/Engine/Source/Runtime/Core/Public/ProfilingDebugging/CookStats.h
@@ -145,6 +145,8 @@ private:
 
 };
 
+struct FDerivedDataUsageStats;
+
 /** 
  * Shared code to provide some common cook telemetry functionality.
  * Exposes a CallStats class that exposes counters that track call count, cycles, and "bytes processed" for a call, bucketing by "cache hits/misses" and "game/other thread".
@@ -269,6 +271,21 @@ public:
 			, bIsInGameThread(IsInGameThread())
 		{
 		}
+
+		FScopedStatsCounter(FScopedStatsCounter&& InOther)
+			: Stats(InOther.Stats)
+			, StartTime(InOther.StartTime)
+			, BytesProcessed(InOther.BytesProcessed)
+			, bIsInGameThread(InOther.bIsInGameThread)
+			, bCanceled(InOther.bCanceled)
+			, bTrackCyclesOnly(InOther.bTrackCyclesOnly)
+			, HitOrMiss(InOther.HitOrMiss)
+		{
+			InOther.Cancel();
+		}
+
+		FScopedStatsCounter(const FScopedStatsCounter&) = delete;
+
 		/** Stop the timer and flushes the stats to the underlying stats object. */
 		~FScopedStatsCounter()
 		{
@@ -329,55 +346,57 @@ public:
 		CallStats::EHitOrMiss HitOrMiss = CallStats::EHitOrMiss::Miss;
 	};
 
-	/** 
-	 * Stats for a resource class that uses the DDC to load the resouce.
-	 * This supports timing a call for synchronous work and asynchronous waits.
-	 * 
-	 * SyncWork timing should be any code that is explicitly synchronous, be it a DDC fetch or building the DDC resource.
-	 * For SyncWork, a "hit" means we fetched from the DDC, and a "miss" means we had to build the resource. Either way,
-	 * the bytes processed should be set appropriately via calls to AddHitOrMiss().
-	 * 
-	 * AsyncWait timing should be any code that is waiting on any explicitly asynchronous work to complete.
-	 * This should only time the wait portion.
-	 * A "hit" or "miss" should be interpreted the same way as in SyncWork.
-	 * The Cycles counter should be measuring WAIT time, NOT the time the async task took to complete.
-	 * 
-	 * Many systems build assets in a very different spot in the code than they are retrieved. When that happens
-	 * TrackCyclesOnly() should be called on the timer, and the Miss should be timed where the resource
-	 * is actually built so the size can cycles used to build the resource can be tracked properly.
-	 * In such cases, the time to TRY to get from the DDC will be timed along with the actual time to build the resource
-	 * But there will only be a single CallCount recorded. This is desired so that every Miss() counted will properly equate
-	 * to one resource built. Look for calls to TrackCyclesOnly() to see how this works in practice.
-	 *
-	 * Async tasks can sometimes actually be executed synchronously if the code ends up waiting on the results
-	 * before the task can even get started. Either way, this should be treated as "async" wait time. Be sure
-	 * not to double count this time. Basically, do your counting at the level that decides whether to execute
-	 * the work synchronously or not, so you can be sure to track the time properly as SyncWork or AsyncWait.
-	 */
-	struct FDDCResourceUsageStats
+	using FDDCResourceUsageStats = FDerivedDataUsageStats;
+};
+
+/**
+ * Stats for a resource class that uses the DDC to load the resource.
+ * This supports timing a call for synchronous work and asynchronous waits.
+ *
+ * SyncWork timing should be any code that is explicitly synchronous, be it a DDC fetch or building the DDC resource.
+ * For SyncWork, a "hit" means we fetched from the DDC, and a "miss" means we had to build the resource. Either way,
+ * the bytes processed should be set appropriately via calls to AddHitOrMiss().
+ *
+ * AsyncWait timing should be any code that is waiting on any explicitly asynchronous work to complete.
+ * This should only time the wait portion.
+ * A "hit" or "miss" should be interpreted the same way as in SyncWork.
+ * The Cycles counter should be measuring WAIT time, NOT the time the async task took to complete.
+ *
+ * Many systems build assets in a very different spot in the code than they are retrieved. When that happens
+ * TrackCyclesOnly() should be called on the timer, and the Miss should be timed where the resource
+ * is actually built so the size can cycles used to build the resource can be tracked properly.
+ * In such cases, the time to TRY to get from the DDC will be timed along with the actual time to build the resource
+ * But there will only be a single CallCount recorded. This is desired so that every Miss() counted will properly equate
+ * to one resource built. Look for calls to TrackCyclesOnly() to see how this works in practice.
+ *
+ * Async tasks can sometimes actually be executed synchronously if the code ends up waiting on the results
+ * before the task can even get started. Either way, this should be treated as "async" wait time. Be sure
+ * not to double count this time. Basically, do your counting at the level that decides whether to execute
+ * the work synchronously or not, so you can be sure to track the time properly as SyncWork or AsyncWait.
+ */
+struct FDerivedDataUsageStats
+{
+	/** Call this where the code is building the resource synchronously */
+	FCookStats::FScopedStatsCounter TimeSyncWork()
 	{
-		/** Call this where the code is building the resource synchronously */
-		FCookStats::FScopedStatsCounter TimeSyncWork()
-		{
-			return FCookStats::FScopedStatsCounter(SyncWorkStats);
-		}
+		return FCookStats::FScopedStatsCounter(SyncWorkStats);
+	}
 
-		/** Call this where the code is waiting on async work to build the resource */
-		FCookStats::FScopedStatsCounter TimeAsyncWait()
-		{
-			return FCookStats::FScopedStatsCounter(AsyncWaitStats);
-		}
+	/** Call this where the code is waiting on async work to build the resource */
+	FCookStats::FScopedStatsCounter TimeAsyncWait()
+	{
+		return FCookStats::FScopedStatsCounter(AsyncWaitStats);
+	}
 
-		// expose these publicly for low level access. These should really never be accessed directly except when finished accumulating them.
-		FCookStats::CallStats SyncWorkStats;
-		FCookStats::CallStats AsyncWaitStats;
+	// expose these publicly for low level access. These should really never be accessed directly except when finished accumulating them.
+	FCookStats::CallStats SyncWorkStats;
+	FCookStats::CallStats AsyncWaitStats;
 
-		void LogStats(FCookStatsManager::AddStatFuncRef AddStat, const FString& StatName, const FString& NodeName) const
-		{
-			SyncWorkStats.LogStats(AddStat, StatName, NodeName, TEXT("SyncWork"));
-			AsyncWaitStats.LogStats(AddStat, StatName, NodeName, TEXT("AsyncWait"));
-		}
-	};
+	void LogStats(FCookStatsManager::AddStatFuncRef AddStat, const FString& StatName, const FString& NodeName) const
+	{
+		SyncWorkStats.LogStats(AddStat, StatName, NodeName, TEXT("SyncWork"));
+		AsyncWaitStats.LogStats(AddStat, StatName, NodeName, TEXT("AsyncWait"));
+	}
 };
 
 /** any expression inside this block will be compiled out if ENABLE_COOK_STATS is not true. Shorthand for the multi-line #if ENABLE_COOK_STATS guard. */
