Added ensure to validate the BaseAllocator is thread safe when we flip from singlethread to multithread on forked processes

diff --git a/Engine/Source/Runtime/Core/Private/HAL/ThreadingBase.cpp b/Engine/Source/Runtime/Core/Private/HAL/ThreadingBase.cpp
--- a/Engine/Source/Runtime/Core/Private/HAL/ThreadingBase.cpp
+++ b/Engine/Source/Runtime/Core/Private/HAL/ThreadingBase.cpp
@@ -1096,6 +1096,8 @@ void FForkProcessHelper::OnForkingOccured()
 {
 	if( SupportsMultithreadingPostFork() )
 	{
+		ensureMsgf(GMalloc->IsInternallyThreadSafe(), TEXT("The BaseAllocator %s is not threadsafe. Switch to a multithread allocator or ensure the FMallocThreadSafeProxy wraps it."), GMalloc->GetDescriptiveName());
+
 		bIsForkedMultithreadInstance = true;
 
 		// Use a local list of forkable threads so we don't keep a lock on the global list during thread creation
@@ -1114,7 +1116,8 @@ bool FForkProcessHelper::IsForkedMultithreadInstance()
 
 bool FForkProcessHelper::SupportsMultithreadingPostFork()
 {
-	static bool bSupportsMultithreadingPostFork = FCommandLine::IsInitialized() ? FParse::Param(FCommandLine::Get(), TEXT("PostForkThreading")) : false;
+	check(FCommandLine::IsInitialized());
+	static bool bSupportsMultithreadingPostFork = FParse::Param(FCommandLine::Get(), TEXT("PostForkThreading"));
 	return bSupportsMultithreadingPostFork;
 }
 
