Taskgraph
* Added TaskGraph.ForkedProcessMaxWorkerThreads cvar to directly control number of worker threads spawned on forked processes that enable multithreading

diff --git a/Engine/Source/Runtime/Core/Private/Async/TaskGraph.cpp b/Engine/Source/Runtime/Core/Private/Async/TaskGraph.cpp
--- a/Engine/Source/Runtime/Core/Private/Async/TaskGraph.cpp
+++ b/Engine/Source/Runtime/Core/Private/Async/TaskGraph.cpp
@@ -89,6 +89,13 @@ static FAutoConsoleVariableRef CVarEnableForkedMultithreading(
 	TEXT("When false will prevent the task graph from running multithreaded on forked processes.")
 );
 
+static int32 CVar_ForkedProcess_MaxWorkerThreads = 2;
+static FAutoConsoleVariableRef CVarForkedProcessMaxWorkerThreads(
+	TEXT("TaskGraph.ForkedProcessMaxWorkerThreads"),
+	CVar_ForkedProcess_MaxWorkerThreads,
+	TEXT("Configures the number of worker threads a forked process should spawn if it allows multithreading.")
+);
+
 #if CREATE_HIPRI_TASK_THREADS || CREATE_BACKGROUND_TASK_THREADS
 	static void ThreadSwitchForABTest(const TArray<FString>& Args)
 	{
@@ -1190,6 +1197,11 @@ public:
 		else
 		{
 			LastExternalThread = ENamedThreads::ActualRenderingThread;
+
+			if (FForkProcessHelper::IsForkedMultithreadInstance())
+			{
+				NumTaskThreads = CVar_ForkedProcess_MaxWorkerThreads;
+			}
 		}
 		
 		NumNamedThreads = LastExternalThread + 1;
