Fix overflow in fork memory logs when mem delta is negative.

diff --git a/Engine/Source/Runtime/Core/Private/Unix/UnixPlatformProcess.cpp b/Engine/Source/Runtime/Core/Private/Unix/UnixPlatformProcess.cpp
--- a/Engine/Source/Runtime/Core/Private/Unix/UnixPlatformProcess.cpp
+++ b/Engine/Source/Runtime/Core/Private/Unix/UnixPlatformProcess.cpp
@@ -1360,7 +1360,22 @@ FGenericPlatformProcess::EWaitAndForkResult FUnixPlatformProcess::WaitAndFork()
 	UE_LOG(LogHAL, Log, TEXT("   *** WaitAndFork awaiting signal %d to create child processes... ***"), WAIT_AND_FORK_QUEUE_SIGNAL);
 	GLog->Flush();
 
-	FPlatformMemoryStats PreviousMasterMemStats = FPlatformMemory::GetStats();
+	struct FMemoryStatsHolder
+	{
+		float AvailablePhysical;
+		float PeakUsedPhysical;
+		float PeakUsedVirtual;
+
+		constexpr float ByteToMiB(uint64 InBytes) { return InBytes / (1024.f * 1024.f); }
+
+		FMemoryStatsHolder(const FPlatformMemoryStats& PlatformStats)
+			: AvailablePhysical(ByteToMiB(PlatformStats.AvailablePhysical))
+			, PeakUsedPhysical(ByteToMiB(PlatformStats.PeakUsedPhysical))
+			, PeakUsedVirtual(ByteToMiB(PlatformStats.PeakUsedVirtual))
+		{ }
+	};
+
+	FMemoryStatsHolder PreviousMasterMemStats(FPlatformMemory::GetStats());
 
 	EWaitAndForkResult RetVal = EWaitAndForkResult::Parent;
 	struct FPidAndSignal
@@ -1381,11 +1396,11 @@ FGenericPlatformProcess::EWaitAndForkResult FUnixPlatformProcess::WaitAndFork()
 			// Sleep for a short while to avoid spamming new processes to the OS all at once
 			FPlatformProcess::Sleep(WAIT_AND_FORK_CHILD_SPAWN_DELAY);
 
-			FPlatformMemoryStats CurrentMasterMemStats = FPlatformMemory::GetStats();
+			FMemoryStatsHolder CurrentMasterMemStats(FPlatformMemory::GetStats());
 			UE_LOG(LogHAL, Log, TEXT("MemoryStats PreFork: AvailablePhysical: %.02fMiB (%+.02fMiB), PeakPhysical: %.02fMiB (%+.02fMiB), PeakVirtual: %.02fMiB (%+.02fMiB)"),
-				CurrentMasterMemStats.AvailablePhysical / (1024.f*1024.f), (CurrentMasterMemStats.AvailablePhysical - PreviousMasterMemStats.AvailablePhysical) / (1024.f*1024.f),
-				CurrentMasterMemStats.PeakUsedPhysical / (1024.f*1024.f), (CurrentMasterMemStats.PeakUsedPhysical - PreviousMasterMemStats.PeakUsedPhysical) / (1024.f*1024.f),
-				CurrentMasterMemStats.PeakUsedVirtual / (1024.f*1024.f), (CurrentMasterMemStats.PeakUsedVirtual - PreviousMasterMemStats.PeakUsedVirtual) / (1024.f*1024.f)
+				CurrentMasterMemStats.AvailablePhysical, (CurrentMasterMemStats.AvailablePhysical - PreviousMasterMemStats.AvailablePhysical),
+				CurrentMasterMemStats.PeakUsedPhysical, (CurrentMasterMemStats.PeakUsedPhysical - PreviousMasterMemStats.PeakUsedPhysical),
+				CurrentMasterMemStats.PeakUsedVirtual, (CurrentMasterMemStats.PeakUsedVirtual - PreviousMasterMemStats.PeakUsedVirtual)
 				
 			);
 			PreviousMasterMemStats = CurrentMasterMemStats;
