Updates to FileCache to work better with shader preloading
- FileCache can dynamically allocate/release memory as needed
- Preloaded lines are locked in the cache.  One the preloaded memory has been read, the line is unlocked.  This prevents preloaded lines from being evicted before they're needed.
- Cache line size changed from 64k to 16k, reduce fragmentation due to locked preloaded lines

diff --git a/Engine/Source/Runtime/Core/Private/FileCache/FileCache.cpp b/Engine/Source/Runtime/Core/Private/FileCache/FileCache.cpp
--- a/Engine/Source/Runtime/Core/Private/FileCache/FileCache.cpp
+++ b/Engine/Source/Runtime/Core/Private/FileCache/FileCache.cpp
@@ -4,6 +4,7 @@
 #include "Containers/BinaryHeap.h"
 #include "Containers/Queue.h"
 #include "Containers/LockFreeList.h"
+#include "Containers/Ticker.h"
 #include "Templates/TypeHash.h"
 #include "Misc/ScopeLock.h"
 #include "Async/AsyncFileHandle.h"
@@ -22,9 +23,17 @@ DECLARE_CYCLE_STAT(TEXT("EvictAll"), STAT_SFC_EvictAll, STATGROUP_SFC);
 // These below are pretty high throughput and probably should be removed once the system gets more mature
 DECLARE_CYCLE_STAT(TEXT("Find Eviction Candidate"), STAT_SFC_FindEvictionCandidate, STATGROUP_SFC);
 
+DECLARE_MEMORY_STAT(TEXT("Allocate Size"), STAT_SFC_AllocatedSize, STATGROUP_SFC);
+
+DECLARE_MEMORY_STAT(TEXT("Locked Size"), STAT_SFC_LockedSize, STATGROUP_SFC);
+DECLARE_MEMORY_STAT(TEXT("Preloaded Size"), STAT_SFC_PreloadedSize, STATGROUP_SFC);
+DECLARE_MEMORY_STAT(TEXT("Fine Preloaded Size"), STAT_SFC_FinePreloadedSize, STATGROUP_SFC);
+
 DEFINE_LOG_CATEGORY_STATIC(LogStreamingFileCache, Log, All);
 
-static const int CacheLineSize = 64 * 1024;
+static const int CacheSlotCapacity = 64 * 1024;
+static const int CacheLineSize = 16 * 1024;
+static const int PreloadBlockSize = CacheLineSize / 64;
 
 static int32 GNumFileCacheBlocks = 256;
 static FAutoConsoleVariableRef CVarNumFileCacheBlocks(
@@ -34,6 +43,15 @@ static FAutoConsoleVariableRef CVarNumFileCacheBlocks(
 	ECVF_RenderThreadSafe
 );
 
+static int32 GLineReleaseFrameThreshold = 300;
+static FAutoConsoleVariableRef CVarLineReleaseFrameThreshold(
+	TEXT("fc.LineReleaseFrameThreshold"),
+	GLineReleaseFrameThreshold,
+	TEXT(""),
+	ECVF_RenderThreadSafe
+);
+
+
 // 
 // Strongly typed ids to avoid confusion in the code
 // 
@@ -86,39 +104,45 @@ class FFileCacheHandle;
 class FFileCache
 {
 public:
-	explicit FFileCache(int32 NumSlots);
-
-	~FFileCache()
+	enum ESlotListType
 	{
-		FMemory::Free(Memory);
-	}
+		SLOTLIST_Free,
+		SLOTLIST_UnlockedAllocated,
+		SLOTLIST_Num,
+	};
 
-	uint8* GetSlotMemory(CacheSlotID SlotID)
+	struct FSlotInfo
 	{
-		check(SlotID.Get() < SlotInfo.Num() - 1);
-		check(IsSlotLocked(SlotID)); // slot must be locked in order to access memory
-		return Memory + SlotID.Get() * CacheSlotID::BlockSize;
-	}
+		uint64 PreloadMask = 0u;
+		FFileCacheHandle* Handle = nullptr;
+		CacheLineID LineID;
+		int32 NextSlotIndex = 0;
+		int32 PrevSlotIndex = 0;
+		uint32 LastUsedFrameNumber = 0u;
+		int16 LockCount = 0;
+	};
+
+	explicit FFileCache(int32 NumSlots);
+	~FFileCache();
+
+	bool OnTick(float DeltaTime);
+
+	uint8* GetSlotMemory(CacheSlotID SlotID);
 
 	CacheSlotID AcquireAndLockSlot(FFileCacheHandle* InHandle, CacheLineID InLineID);
 	bool IsSlotLocked(CacheSlotID InSlotID) const;
 	void LockSlot(CacheSlotID InSlotID);
 	void UnlockSlot(CacheSlotID InSlotID);
+	void MarkSlotPreloadedRegion(CacheSlotID InSlotID, int64 InOffset, int64 InSize);
+	void ClearSlotPreloadedRegion(CacheSlotID InSlotID, int64 InOffset, int64 InSize);
+
+	void ReleaseMemory(int32 NumSlotsToRelease = 1);
 
 	// if InFile is null, will evict all slots
 	bool EvictAll(FFileCacheHandle* InFile = nullptr);
 
 	void FlushCompletedRequests();
 
-	struct FSlotInfo
-	{
-		FFileCacheHandle* Handle;
-		CacheLineID LineID;
-		int32 NextSlotIndex;
-		int32 PrevSlotIndex;
-		int32 LockCount;
-	};
-
 	void EvictFileCacheFromConsole()
 	{
 		EvictAll();
@@ -145,37 +169,37 @@ public:
 
 	inline void UnlinkSlot(int32 SlotIndex)
 	{
-		check(SlotIndex != 0);
+		check(SlotIndex >= SLOTLIST_Num);
 		FSlotInfo& Info = SlotInfo[SlotIndex];
 		SlotInfo[Info.PrevSlotIndex].NextSlotIndex = Info.NextSlotIndex;
 		SlotInfo[Info.NextSlotIndex].PrevSlotIndex = Info.PrevSlotIndex;
 		Info.NextSlotIndex = Info.PrevSlotIndex = SlotIndex;
 	}
 
-	inline void LinkSlotTail(int32 SlotIndex)
+	inline void LinkSlotTail(ESlotListType List, int32 SlotIndex)
 	{
-		check(SlotIndex != 0);
-		FSlotInfo& HeadInfo = SlotInfo[0];
+		check(SlotIndex >= SLOTLIST_Num);
+		FSlotInfo& HeadInfo = SlotInfo[List];
 		FSlotInfo& Info = SlotInfo[SlotIndex];
 		check(Info.NextSlotIndex == SlotIndex);
 		check(Info.PrevSlotIndex == SlotIndex);
 
-		Info.NextSlotIndex = 0;
+		Info.NextSlotIndex = List;
 		Info.PrevSlotIndex = HeadInfo.PrevSlotIndex;
 		SlotInfo[HeadInfo.PrevSlotIndex].NextSlotIndex = SlotIndex;
 		HeadInfo.PrevSlotIndex = SlotIndex;
 	}
 
-	inline void LinkSlotHead(int32 SlotIndex)
+	inline void LinkSlotHead(ESlotListType List, int32 SlotIndex)
 	{
-		check(SlotIndex != 0);
-		FSlotInfo& HeadInfo = SlotInfo[0];
+		check(SlotIndex >= SLOTLIST_Num);
+		FSlotInfo& HeadInfo = SlotInfo[List];
 		FSlotInfo& Info = SlotInfo[SlotIndex];
 		check(Info.NextSlotIndex == SlotIndex);
 		check(Info.PrevSlotIndex == SlotIndex);
 
 		Info.NextSlotIndex = HeadInfo.NextSlotIndex;
-		Info.PrevSlotIndex = 0;
+		Info.PrevSlotIndex = List;
 		SlotInfo[HeadInfo.NextSlotIndex].PrevSlotIndex = SlotIndex;
 		HeadInfo.NextSlotIndex = SlotIndex;
 	}
@@ -184,14 +208,17 @@ public:
 
 	FAutoConsoleCommand EvictFileCacheCommand;
 
+	FDelegateHandle TickHandle;
+
 	TLockFreePointerListUnordered<IAsyncReadRequest, PLATFORM_CACHE_LINE_SIZE> CompletedRequests;
 	FThreadSafeCounter CompletedRequestsCounter;
 
 	// allocated with an extra dummy entry at index0 for linked list head
 	TArray<FSlotInfo> SlotInfo;
-	uint8* Memory;
+	uint8* SlotMemory[CacheSlotCapacity];
 	int32 SizeInBytes;
-	int32 NumFreeSlots;
+	int32 NumSlots;
+	int32 NumAllocatedSlots;
 };
 
 static FFileCache &GetCache()
@@ -240,12 +267,13 @@ PRAGMA_ENABLE_UNSAFE_TYPECAST_WARNINGS
 	// Returns the size within the first cache line covering the byte range to read
 	template<typename BlockIDType> inline int64 GetBlockSize(int64 Offset, int64 Size)
 	{
-		int64 OffsetInBlock = GetBlockOffset<BlockIDType>(Offset);
-		return FMath::Min((int64)(BlockIDType::BlockSize - OffsetInBlock), Size - Offset);
+		int64 OffsetInSlot = GetBlockOffset<BlockIDType>(Offset);
+		return FMath::Min((int64)(BlockIDType::BlockSize - OffsetInSlot), Size - Offset);
 	}
 
 	virtual IMemoryReadStreamRef ReadData(FGraphEventArray& OutCompletionEvents, int64 Offset, int64 BytesToRead, EAsyncIOPriorityAndFlags Priority) override;
-	virtual FGraphEventRef PreloadData(const FFileCachePreloadEntry* PreloadEntries, int32 NumEntries, EAsyncIOPriorityAndFlags Priority) override;
+	virtual FGraphEventRef PreloadData(const FFileCachePreloadEntry* PreloadEntries, int32 NumEntries, int64 InOffset, EAsyncIOPriorityAndFlags Priority) override;
+	virtual void ReleasePreloadedData(const FFileCachePreloadEntry* PreloadEntries, int32 NumEntries, int64 InOffset) override;
 
 	IMemoryReadStreamRef ReadDataUncached(FGraphEventArray& OutCompletionEvents, int64 Offset, int64 BytesToRead, EAsyncIOPriorityAndFlags Priority);
 
@@ -275,87 +303,255 @@ private:
 
 ///////////////
 
-FFileCache::FFileCache(int32 NumSlots)
+FFileCache::FFileCache(int32 InNumSlots)
 	: EvictFileCacheCommand(TEXT("r.VT.EvictFileCache"), TEXT("Evict all the file caches in the VT system."),
 		FConsoleCommandDelegate::CreateRaw(this, &FFileCache::EvictFileCacheFromConsole))
-	, SizeInBytes(NumSlots * CacheSlotID::BlockSize)
-	, NumFreeSlots(NumSlots)
+	, SizeInBytes(InNumSlots * CacheSlotID::BlockSize)
+	, NumSlots(InNumSlots)
+	, NumAllocatedSlots(0)
 {
-	Memory = (uint8*)FMemory::Malloc(SizeInBytes);
+	FMemory::Memzero(SlotMemory);
 
-	SlotInfo.AddUninitialized(NumSlots + 1);
-	for (int i = 0; i <= NumSlots; ++i)
+	SlotInfo.AddDefaulted(InNumSlots + SLOTLIST_Num);
+	for (int32 i = 0; i < SLOTLIST_Num; ++i)
+	{
+		FSlotInfo& Info = SlotInfo[i];
+		Info.NextSlotIndex = i;
+		Info.PrevSlotIndex = i;
+	}
+
+	// All slots begin in free list
+	SlotInfo[SLOTLIST_Free].PrevSlotIndex = InNumSlots;
+	SlotInfo[InNumSlots].NextSlotIndex = SLOTLIST_Free;
+
+	for (int32 i = SLOTLIST_Num; i < SlotInfo.Num(); ++i)
 	{
 		FSlotInfo& Info = SlotInfo[i];
-		Info.Handle = nullptr;
-		Info.LineID = CacheLineID();
-		Info.LockCount = 0;
 		Info.NextSlotIndex = i + 1;
 		Info.PrevSlotIndex = i - 1;
 	}
 
-	// list is circular
-	SlotInfo[0].PrevSlotIndex = NumSlots;
-	SlotInfo[NumSlots].NextSlotIndex = 0;
+	TickHandle = FTicker::GetCoreTicker().AddTicker(FTickerDelegate::CreateRaw(this, &FFileCache::OnTick), 0.1f);
+}
+
+FFileCache::~FFileCache()
+{
+	SET_MEMORY_STAT(STAT_SFC_AllocatedSize, 0);
+	for (int32 i = 0; i < CacheSlotCapacity; ++i)
+	{
+		if (SlotMemory[i])
+		{
+			FMemory::Free(SlotMemory[i]);
+		}
+	}
+
+	FTicker::GetCoreTicker().RemoveTicker(TickHandle);
+}
+
+bool FFileCache::OnTick(float DeltaTime)
+{
+	FScopeLock Lock(&CriticalSection);
+	ReleaseMemory(30);
+	return true;
 }
 
 CacheSlotID FFileCache::AcquireAndLockSlot(FFileCacheHandle* InHandle, CacheLineID InLineID)
 {
-	check(NumFreeSlots > 0);
-	--NumFreeSlots;
+	int32 SlotIndex = SlotInfo[SLOTLIST_UnlockedAllocated].NextSlotIndex;
+	if (SlotIndex == SLOTLIST_UnlockedAllocated)
+	{
+		SlotIndex = SlotInfo[SLOTLIST_Free].NextSlotIndex;
+		if (SlotIndex == SLOTLIST_Free)
+		{
+			SlotIndex = SlotInfo.AddDefaulted();
+			check(SlotIndex < CacheSlotCapacity);
+			SlotInfo[SlotIndex].NextSlotIndex = SlotInfo[SlotIndex].PrevSlotIndex = SlotIndex;
+		}
+		else
+		{
+			UnlinkSlot(SlotIndex);
+		}
 
-	const int32 SlotIndex = SlotInfo[0].NextSlotIndex;
-	check(SlotIndex != 0);
+		FSlotInfo& Info = SlotInfo[SlotIndex];
+		check(!SlotMemory[SlotIndex]);
+		check(!Info.Handle);
 
-	FSlotInfo& Info = SlotInfo[SlotIndex];
-	check(Info.LockCount == 0); // slot should not be in free list if it's locked
-	if (Info.Handle)
+		INC_MEMORY_STAT_BY(STAT_SFC_AllocatedSize, CacheSlotID::BlockSize);
+		SlotMemory[SlotIndex] = (uint8*)FMemory::Malloc(CacheSlotID::BlockSize);
+		++NumAllocatedSlots;
+	}
+	else
 	{
-		Info.Handle->Evict(Info.LineID);
+		UnlinkSlot(SlotIndex);
+		FSlotInfo& Info = SlotInfo[SlotIndex];
+		if (Info.Handle)
+		{
+			Info.Handle->Evict(Info.LineID);
+			Info.Handle = nullptr;
+		}
 	}
 
+	INC_MEMORY_STAT_BY(STAT_SFC_LockedSize, CacheSlotID::BlockSize);
+	FSlotInfo& Info = SlotInfo[SlotIndex];
+	check(Info.LockCount == 0); // slot should not be in free list if it's locked
+	check(Info.PreloadMask == 0u);
 	Info.LockCount = 1;
 	Info.Handle = InHandle;
 	Info.LineID = InLineID;
-	UnlinkSlot(SlotIndex);
 
-	return CacheSlotID(SlotIndex - 1);
+	check(SlotMemory[SlotIndex]);
+	return CacheSlotID(SlotIndex - SLOTLIST_Num);
+}
+
+uint8* FFileCache::GetSlotMemory(CacheSlotID InSlotID)
+
+{
+	const int32 SlotIndex = InSlotID.Get() + SLOTLIST_Num;
+	// May be called with no slot, SlotInfo array not safe to access here
+	uint8* Memory = SlotMemory[SlotIndex];
+	check(Memory);
+	return Memory;
 }
 
 bool FFileCache::IsSlotLocked(CacheSlotID InSlotID) const
 {
-	const int32 SlotIndex = InSlotID.Get() + 1;
+	const int32 SlotIndex = InSlotID.Get() + SLOTLIST_Num;
 	const FSlotInfo& Info = SlotInfo[SlotIndex];
 	return Info.LockCount > 0;
 }
 
 void FFileCache::LockSlot(CacheSlotID InSlotID)
 {
-	const int32 SlotIndex = InSlotID.Get() + 1;
+	const int32 SlotIndex = InSlotID.Get() + SLOTLIST_Num;
 	FSlotInfo& Info = SlotInfo[SlotIndex];
-	const int32 PrevLockCount = Info.LockCount;
+	check(SlotMemory[SlotIndex]);
+	const int16 PrevLockCount = Info.LockCount;
+	check(PrevLockCount < 0x7fff);
 	if (PrevLockCount == 0)
 	{
-		check(NumFreeSlots > 0);
-		--NumFreeSlots;
-		UnlinkSlot(SlotIndex);
+		INC_MEMORY_STAT_BY(STAT_SFC_LockedSize, CacheSlotID::BlockSize);
+		if (Info.PreloadMask == 0u)
+		{
+			UnlinkSlot(SlotIndex);
+		}
 	}
 	Info.LockCount = PrevLockCount + 1;
 }
 
 void FFileCache::UnlockSlot(CacheSlotID InSlotID)
 {
-	const int32 SlotIndex = InSlotID.Get() + 1;
-	const int32 PrevLockCount = SlotInfo[SlotIndex].LockCount;
+	const int32 SlotIndex = InSlotID.Get() + SLOTLIST_Num;
+	FSlotInfo& Info = SlotInfo[SlotIndex];
+	const int16 PrevLockCount = Info.LockCount;
+	check(SlotMemory[SlotIndex]);
 	check(PrevLockCount > 0);
+
 	if (PrevLockCount == 1)
 	{
+		DEC_MEMORY_STAT_BY(STAT_SFC_LockedSize, CacheSlotID::BlockSize);
+		Info.LastUsedFrameNumber = GFrameNumber;
+
 		// move slot back to the free list when it's unlocked
-		LinkSlotTail(SlotIndex);
-		++NumFreeSlots;
-		check(NumFreeSlots < SlotInfo.Num());
+		if (Info.PreloadMask == 0u)
+		{
+			LinkSlotTail(SLOTLIST_UnlockedAllocated, SlotIndex);
+		}
+	}
+	Info.LockCount = PrevLockCount - 1;
+}
+
+static uint64 MakePreloladMask(int64 InOffset, int64 InSize)
+{
+	checkSlow(InOffset >= 0 && InOffset < CacheLineSize);
+	checkSlow(InSize > 0 && InSize <= CacheLineSize);
+	const uint32 StartBlock = FMath::DivideAndRoundDown<uint32>((uint32)InOffset, PreloadBlockSize);
+	const uint32 EndBlock = FMath::DivideAndRoundUp<uint32>((uint32)(InOffset + InSize), PreloadBlockSize);
+	const uint32 NumBlocks = EndBlock - StartBlock;
+	checkf(NumBlocks > 0u && NumBlocks <= 64u, TEXT("Invalid NumBlocks %d"), NumBlocks);
+	if (NumBlocks < 64u)
+	{
+		const uint64 Mask = ((uint64)1u << NumBlocks) - 1u;
+		return Mask << StartBlock;
+	}
+	check(StartBlock == 0u);
+	return ~(uint64)0u;
+}
+
+void FFileCache::MarkSlotPreloadedRegion(CacheSlotID InSlotID, int64 InOffset, int64 InSize)
+{
+	const int32 SlotIndex = InSlotID.Get() + SLOTLIST_Num;
+	FSlotInfo& Info = SlotInfo[SlotIndex];
+
+	if (Info.PreloadMask == 0u)
+	{
+		INC_MEMORY_STAT_BY(STAT_SFC_PreloadedSize, CacheSlotID::BlockSize);
+		if (Info.LockCount == 0)
+		{
+			UnlinkSlot(SlotIndex);
+		}
+	}
+
+	const uint64 Mask = MakePreloladMask(InOffset, InSize);
+	const uint64 BitsSet = Mask & ~Info.PreloadMask;
+	INC_MEMORY_STAT_BY(STAT_SFC_FinePreloadedSize, FMath::CountBits(BitsSet) * PreloadBlockSize);
+	Info.PreloadMask |= Mask;
+}
+
+void FFileCache::ClearSlotPreloadedRegion(CacheSlotID InSlotID, int64 InOffset, int64 InSize)
+{
+	const int32 SlotIndex = InSlotID.Get() + SLOTLIST_Num;
+	FSlotInfo& Info = SlotInfo[SlotIndex];
+
+	if (Info.PreloadMask != 0u)
+	{
+		const uint64 Mask = MakePreloladMask(InOffset, InSize);
+		const uint64 BitsClear = Mask & Info.PreloadMask;
+		DEC_MEMORY_STAT_BY(STAT_SFC_FinePreloadedSize, FMath::CountBits(BitsClear) * PreloadBlockSize);
+		Info.PreloadMask &= ~Mask;
+		if (Info.PreloadMask == 0u)
+		{
+			DEC_MEMORY_STAT_BY(STAT_SFC_PreloadedSize, CacheSlotID::BlockSize);
+			if (Info.LockCount == 0)
+			{
+				LinkSlotTail(SLOTLIST_UnlockedAllocated, SlotIndex);
+				Info.LastUsedFrameNumber = 0u;
+			}
+		}
+	}
+}
+
+void FFileCache::ReleaseMemory(int32 InNumSlotsToRelease)
+{
+	const uint32 CurrentFrameNumber = GFrameNumber;
+	int32 NumSlotsToRelease = FMath::Min(InNumSlotsToRelease, NumAllocatedSlots - NumSlots);
+	while (NumSlotsToRelease > 0)
+	{
+		const int32 SlotIndex = SlotInfo[SLOTLIST_UnlockedAllocated].NextSlotIndex;
+		FSlotInfo& Info = SlotInfo[SlotIndex];
+		if (SlotIndex == SLOTLIST_UnlockedAllocated || Info.LastUsedFrameNumber + GLineReleaseFrameThreshold >= CurrentFrameNumber)
+		{
+			break;
+		}
+
+		check(SlotMemory[SlotIndex]);
+		check(Info.LockCount == 0); // slot should not be in free list if it's locked
+		check(Info.PreloadMask == 0u);
+		if (Info.Handle)
+		{
+			Info.Handle->Evict(Info.LineID);
+			Info.Handle = nullptr;
+		}
+
+		DEC_MEMORY_STAT_BY(STAT_SFC_AllocatedSize, CacheSlotID::BlockSize);
+		FMemory::Free(SlotMemory[SlotIndex]);
+		SlotMemory[SlotIndex] = nullptr;
+		Info.LineID = CacheLineID();
+		--NumAllocatedSlots;
+		--NumSlotsToRelease;
+
+		UnlinkSlot(SlotIndex);
+		LinkSlotTail(SLOTLIST_Free, SlotIndex);
 	}
-	SlotInfo[SlotIndex].LockCount = PrevLockCount - 1;
 }
 
 bool FFileCache::EvictAll(FFileCacheHandle* InFile)
@@ -375,10 +571,16 @@ bool FFileCache::EvictAll(FFileCacheHandle* InFile)
 				Info.Handle->Evict(Info.LineID);
 				Info.Handle = nullptr;
 				Info.LineID = CacheLineID();
+				if (Info.PreloadMask > 0u)
+				{
+					DEC_MEMORY_STAT_BY(STAT_SFC_PreloadedSize, CacheSlotID::BlockSize);
+					DEC_MEMORY_STAT_BY(STAT_SFC_FinePreloadedSize, FMath::CountBits(Info.PreloadMask) * PreloadBlockSize);
+					Info.PreloadMask = 0u;
+				}
 	
 				// move evicted slots to the front of list so they'll be re-used more quickly
 				UnlinkSlot(SlotIndex);
-				LinkSlotHead(SlotIndex);
+				LinkSlotHead(SLOTLIST_UnlockedAllocated, SlotIndex);
 			}
 			else
 			{
@@ -500,10 +702,10 @@ public:
 		const int32 SlotIndex = (int32)FMath::DivideAndRoundDown(Offset, (int64)CacheSlotID::BlockSize);
 		const int32 OffsetInSlot = (int32)(Offset - SlotIndex * CacheSlotID::BlockSize);
 		checkSlow(SlotIndex >= 0 && SlotIndex < NumCacheSlots);
-		const void* SlotMemory = Cache.GetSlotMemory(CacheSlots[SlotIndex]);
+		const uint8* SlotMemory = Cache.GetSlotMemory(CacheSlots[SlotIndex]);
 
 		OutSize = FMath::Min(InSize, (int64)CacheSlotID::BlockSize - OffsetInSlot);
-		return (uint8*)SlotMemory + OffsetInSlot;
+		return SlotMemory + OffsetInSlot;
 	}
 
 	virtual int64 GetSize() override
@@ -515,12 +717,24 @@ public:
 	{
 		FFileCache& Cache = GetCache();
 		FScopeLock CacheLock(&Cache.CriticalSection);
+
+		int64 SizeRemaining = Size;
+		int64 OffsetInSlot = InitialSlotOffset;
+		int64 SizeInSlot = FMath::Min<int64>(Size, CacheLineSize - OffsetInSlot);
 		for (int i = 0; i < NumCacheSlots; ++i)
 		{
 			const CacheSlotID& SlotID = CacheSlots[i];
 			check(SlotID.IsValid());
+			Cache.ClearSlotPreloadedRegion(SlotID, OffsetInSlot, SizeInSlot);
 			Cache.UnlockSlot(SlotID);
+
+			SizeRemaining -= SizeInSlot;
+			check(SizeRemaining >= 0);
+			OffsetInSlot = 0;
+			SizeInSlot = FMath::Min<int64>(SizeRemaining, CacheLineSize);
 		}
+
+		check(SizeRemaining == 0);
 	}
 
 	void operator delete(void* InMem)
@@ -621,7 +835,7 @@ IMemoryReadStreamRef FFileCacheHandle::ReadData(FGraphEventArray& OutCompletionE
 	FFileCache& Cache = GetCache();
 
 	FScopeLock CacheLock(&Cache.CriticalSection);
-	if (NumSlotsNeeded > Cache.NumFreeSlots)
+	/*if (NumSlotsNeeded > Cache.NumFreeSlots)
 	{
 		// not enough free slots in the cache to service this request
 		CacheLock.Unlock();
@@ -631,7 +845,7 @@ IMemoryReadStreamRef FFileCacheHandle::ReadData(FGraphEventArray& OutCompletionE
 			return ReadDataUncached(OutCompletionEvents, Offset, BytesToRead, Priority);
 		}
 		return nullptr;
-	}
+	}*/
 
 	if (EndLine.Get() >= NumSlots)
 	{
@@ -710,7 +924,7 @@ struct FFileCachePreloadTask
 	FORCEINLINE TStatId GetStatId() const { return TStatId(); }
 };
 
-FGraphEventRef FFileCacheHandle::PreloadData(const FFileCachePreloadEntry* PreloadEntries, int32 NumEntries, EAsyncIOPriorityAndFlags Priority)
+FGraphEventRef FFileCacheHandle::PreloadData(const FFileCachePreloadEntry* PreloadEntries, int32 NumEntries, int64 InOffset, EAsyncIOPriorityAndFlags Priority)
 {
 	SCOPED_LOADTIMER(FFileCacheHandle_PreloadData);
 
@@ -724,13 +938,13 @@ FGraphEventRef FFileCacheHandle::PreloadData(const FFileCachePreloadEntry* Prelo
 
 	{
 		const FFileCachePreloadEntry& LastEntry = PreloadEntries[NumEntries - 1];
-		const CacheLineID EndLine = GetBlock<CacheLineID>(LastEntry.Offset + LastEntry.Size - 1);
-		if (EndLine.Get() >= NumSlots)
+		const CacheLineID LastEndLine = GetBlock<CacheLineID>(LastEntry.Offset + LastEntry.Size - 1);
+		if (LastEndLine.Get() >= NumSlots)
 		{
 			// If we're still waiting on SizeRequest, may need to lazily allocate some slots to service this request
 			// If this happens after SizeRequest has completed, that means something must have gone wrong
 			check(SizeRequestEvent);
-			NumSlots = EndLine.Get() + 1;
+			NumSlots = LastEndLine.Get() + 1;
 			// TArray is max signed int
 			check(NumSlots < MAX_int32);
 			LineToSlot.SetNum((int32)NumSlots, false);
@@ -744,18 +958,40 @@ FGraphEventRef FFileCacheHandle::PreloadData(const FFileCachePreloadEntry* Prelo
 
 	CacheLineID CurrentLine(0);
 	int64 PrevOffset = -1;
-	for (int32 EntryIndex = 0; EntryIndex < NumEntries && Cache.NumFreeSlots > 0; ++EntryIndex)
+	uint32 NumSlotsLoaded = 0u;
+	for (int32 EntryIndex = 0; EntryIndex < NumEntries; ++EntryIndex)
 	{
 		const FFileCachePreloadEntry& Entry = PreloadEntries[EntryIndex];
-		const CacheLineID StartLine = GetBlock<CacheLineID>(Entry.Offset);
-		const CacheLineID EndLine = GetBlock<CacheLineID>(Entry.Offset + Entry.Size - 1);
+		int64 EntryOffset = InOffset + Entry.Offset;
+		const int64 EndOffset = EntryOffset + Entry.Size;
+		const CacheLineID StartLine = GetBlock<CacheLineID>(EntryOffset);
+		const CacheLineID EndLine = GetBlock<CacheLineID>(EndOffset - 1);
 
 		checkf(Entry.Offset > PrevOffset, TEXT("Preload entries must be sorted by Offset [%lld, %lld), %lld"),
 			Entry.Offset, Entry.Offset + Entry.Size, PrevOffset);
 		PrevOffset = Entry.Offset;
 
-		CurrentLine = CacheLineID(FMath::Max(CurrentLine.Get(), StartLine.Get()));
-		while (CurrentLine.Get() <= EndLine.Get() && Cache.NumFreeSlots > 0)
+		int64 OffsetInSlot = EntryOffset - StartLine.Get() * CacheLineSize;
+		int64 SizeInSlot = FMath::Min<int64>(Entry.Size, CacheLineSize - OffsetInSlot);
+		if (CurrentLine.Get() > StartLine.Get())
+		{
+			// Will hit this case if the last line of the previous entry is the same as the first line of this entry
+			// In this case, we'll already have a slot allocated for the line, just need to mark additional preload region
+			check(EntryIndex > 0);
+			check(CurrentLine.Get() == StartLine.Get() + 1);
+			const CacheSlotID SlotID = LineToSlot[StartLine.Get()];
+			check(SlotID.IsValid());
+			Cache.MarkSlotPreloadedRegion(SlotID, OffsetInSlot, SizeInSlot);
+			EntryOffset += SizeInSlot;
+			OffsetInSlot = 0;
+			SizeInSlot = FMath::Min<int64>(EndOffset - EntryOffset, CacheLineSize);
+		}
+		else
+		{
+			CurrentLine = StartLine;
+		}
+
+		while (CurrentLine.Get() <= EndLine.Get())
 		{
 			CacheSlotID& SlotID = LineToSlot[CurrentLine.Get()];
 			if (!SlotID.IsValid())
@@ -763,7 +999,10 @@ FGraphEventRef FFileCacheHandle::PreloadData(const FFileCachePreloadEntry* Prelo
 				// no valid slot for this line, grab a new slot from cache and start a read request
 				SlotID = AcquireSlotAndReadLine(Cache, CurrentLine, Priority);
 				LockedSlots.Add(SlotID);
+				++NumSlotsLoaded;
 			}
+		
+			Cache.MarkSlotPreloadedRegion(SlotID, OffsetInSlot, SizeInSlot);
 
 			FPendingRequest& PendingRequest = LineToRequest[CurrentLine.Get()];
 			if (PendingRequest.Event && !PendingRequest.Event->IsComplete())
@@ -778,6 +1017,9 @@ FGraphEventRef FFileCacheHandle::PreloadData(const FFileCachePreloadEntry* Prelo
 			}
 
 			++CurrentLine;
+			EntryOffset += SizeInSlot;
+			OffsetInSlot = 0;
+			SizeInSlot = FMath::Min<int64>(EndOffset - EntryOffset, CacheLineSize);
 		}
 	}
 
@@ -798,6 +1040,50 @@ FGraphEventRef FFileCacheHandle::PreloadData(const FFileCachePreloadEntry* Prelo
 	return CompletionEvent;
 }
 
+void FFileCacheHandle::ReleasePreloadedData(const FFileCachePreloadEntry* PreloadEntries, int32 NumEntries, int64 InOffset)
+{
+	check(NumEntries > 0);
+
+	FFileCache& Cache = GetCache();
+
+	FScopeLock CacheLock(&Cache.CriticalSection);
+
+	int64 PrevOffset = -1;
+	uint32 NumSlotsUnloaded = 0u;
+	for (int32 EntryIndex = 0; EntryIndex < NumEntries; ++EntryIndex)
+	{
+		const FFileCachePreloadEntry& Entry = PreloadEntries[EntryIndex];
+		int64 EntryOffset = InOffset + Entry.Offset;
+		const int64 EndOffset = EntryOffset + Entry.Size;
+		const CacheLineID StartLine = GetBlock<CacheLineID>(EntryOffset);
+		const CacheLineID EndLine = GetBlock<CacheLineID>(EndOffset - 1);
+
+		checkf(Entry.Offset > PrevOffset, TEXT("Preload entries must be sorted by Offset [%lld, %lld), %lld"),
+			Entry.Offset, Entry.Offset + Entry.Size, PrevOffset);
+		PrevOffset = Entry.Offset;
+
+		int64 OffsetInSlot = EntryOffset - StartLine.Get() * CacheLineSize;
+		int64 SizeInSlot = FMath::Min<int64>(Entry.Size, CacheLineSize - OffsetInSlot);
+		CacheLineID CurrentLine = StartLine;
+		while (CurrentLine.Get() <= EndLine.Get())
+		{
+			CacheSlotID& SlotID = LineToSlot[CurrentLine.Get()];
+			if (SlotID.IsValid())
+			{
+				Cache.ClearSlotPreloadedRegion(SlotID, OffsetInSlot, SizeInSlot);
+				++NumSlotsUnloaded;
+			}
+
+			++CurrentLine;
+			EntryOffset += SizeInSlot;
+			OffsetInSlot = 0;
+			SizeInSlot = FMath::Min<int64>(EndOffset - EntryOffset, CacheLineSize);
+		}
+	}
+
+	Cache.ReleaseMemory(NumSlotsUnloaded);
+}
+
 void FFileCacheHandle::Evict(CacheLineID LineID)
 {
 	LineToSlot[LineID.Get()] = CacheSlotID();
diff --git a/Engine/Source/Runtime/Core/Public/FileCache/FileCache.h b/Engine/Source/Runtime/Core/Public/FileCache/FileCache.h
--- a/Engine/Source/Runtime/Core/Public/FileCache/FileCache.h
+++ b/Engine/Source/Runtime/Core/Public/FileCache/FileCache.h
@@ -7,7 +7,6 @@
 #include "Async/TaskGraphInterfaces.h"
 #include "Misc/MemoryReadStream.h"
 
-//
 struct FFileCachePreloadEntry
 {
 	FFileCachePreloadEntry() : Offset(0), Size(0) {}
@@ -60,7 +59,9 @@ public:
 	 */
 	virtual IMemoryReadStreamRef ReadData(FGraphEventArray& OutCompletionEvents, int64 Offset, int64 BytesToRead, EAsyncIOPriorityAndFlags Priority) = 0;
 
-	virtual FGraphEventRef PreloadData(const FFileCachePreloadEntry* PreloadEntries, int32 NumEntries, EAsyncIOPriorityAndFlags Priority) = 0;
+	virtual FGraphEventRef PreloadData(const FFileCachePreloadEntry* PreloadEntries, int32 NumEntries, int64 InOffset, EAsyncIOPriorityAndFlags Priority) = 0;
+
+	virtual void ReleasePreloadedData(const FFileCachePreloadEntry* PreloadEntries, int32 NumEntries, int64 InOffset) = 0;
 
 	/**
 	 * Wait until all outstanding read requests complete. 
diff --git a/Engine/Source/Runtime/RHI/Public/RHIResources.h b/Engine/Source/Runtime/RHI/Public/RHIResources.h
--- a/Engine/Source/Runtime/RHI/Public/RHIResources.h
+++ b/Engine/Source/Runtime/RHI/Public/RHIResources.h
@@ -2282,6 +2282,7 @@ public:
 	virtual int32 FindShaderIndex(const FSHAHash& Hash) = 0;
 	virtual FGraphEventRef PreloadShader(int32 ShaderIndex) { return FGraphEventRef(); }
 	virtual FGraphEventRef PreloadShaderMap(int32 ShaderMapIndex) { return FGraphEventRef(); }
+	virtual void ReleasePreloadedShaderMap(int32 ShaderMapIndex) {}
 
 	virtual TRefCountPtr<FRHIShader> CreateShader(int32 ShaderIndex) { return nullptr; }
 	virtual void Teardown() {};
diff --git a/Engine/Source/Runtime/RenderCore/Private/ShaderCodeArchive.cpp b/Engine/Source/Runtime/RenderCore/Private/ShaderCodeArchive.cpp
--- a/Engine/Source/Runtime/RenderCore/Private/ShaderCodeArchive.cpp
+++ b/Engine/Source/Runtime/RenderCore/Private/ShaderCodeArchive.cpp
@@ -196,6 +196,14 @@ FShaderCodeArchive* FShaderCodeArchive::Create(EShaderPlatform InPlatform, FArch
 	Ar << Library->SerializedShaders;
 	Library->LibraryCodeOffset = Ar.Tell();
 
+#if TRACK_SHADER_PRELOADS
+	Library->ShaderFramePreloaded.SetNumUninitialized(Library->SerializedShaders.GetNumShaders());
+	for (uint32& Frame : Library->ShaderFramePreloaded)
+	{
+		Frame = ~0u;
+	}
+#endif // TRACK_SHADER_PRELOADS
+
 	// Open library for async reads
 	Library->FileCacheHandle = IFileCacheHandle::CreateFileCacheHandle(*InDestFilePath);
 
@@ -229,16 +237,22 @@ void FShaderCodeArchive::Teardown()
 	}
 }
 
-IMemoryReadStreamRef FShaderCodeArchive::ReadShaderCode(int32 Index)
+IMemoryReadStreamRef FShaderCodeArchive::ReadShaderCode(int32 ShaderIndex)
 {
 	SCOPED_LOADTIMER(FShaderCodeArchive_ReadShaderCode);
 
-	const FShaderCodeEntry& Entry = SerializedShaders.ShaderEntries[Index];
+	const FShaderCodeEntry& Entry = SerializedShaders.ShaderEntries[ShaderIndex];
 
 	FGraphEventArray ReadCompleteEvents;
 	IMemoryReadStreamRef LoadedCode = FileCacheHandle->ReadData(ReadCompleteEvents, LibraryCodeOffset + Entry.Offset, Entry.Size, AIOP_CriticalPath);
 	if (ReadCompleteEvents.Num() > 0)
 	{
+#if TRACK_SHADER_PRELOADS
+		if (ShaderFramePreloaded[ShaderIndex] < 0xffffffff)
+		{
+			UE_LOG(LogShaderLibrary, Warning, TEXT("** ShaderCode was preloaded on frame %d, unloaded by frame %d"), ShaderFramePreloaded[ShaderIndex], GFrameNumber);
+		}
+#endif // TRACK_SHADER_PRELOADS
 		FTaskGraphInterface::Get().WaitUntilTasksComplete(ReadCompleteEvents);
 	}
 
@@ -248,18 +262,40 @@ IMemoryReadStreamRef FShaderCodeArchive::ReadShaderCode(int32 Index)
 FGraphEventRef FShaderCodeArchive::PreloadShader(int32 ShaderIndex)
 {
 	const FShaderCodeEntry& ShaderEntry = SerializedShaders.ShaderEntries[ShaderIndex];
-
+#if TRACK_SHADER_PRELOADS
+	ShaderFramePreloaded[ShaderIndex] = FMath::Min(ShaderFramePreloaded[ShaderIndex], GFrameNumber);
+#endif // TRACK_SHADER_PRELOADS
 	const EAsyncIOPriorityAndFlags IOPriority = (EAsyncIOPriorityAndFlags)GShaderCodeLibraryAsyncLoadingPriority;
 	const FFileCachePreloadEntry PreloadEntry(ShaderEntry.Offset, ShaderEntry.Size);
-	return FileCacheHandle->PreloadData(&PreloadEntry, 1, IOPriority);
+	return FileCacheHandle->PreloadData(&PreloadEntry, 1, LibraryCodeOffset, IOPriority);
 }
 
 FGraphEventRef FShaderCodeArchive::PreloadShaderMap(int32 ShaderMapIndex)
 {
 	const FShaderMapEntry& ShaderMapEntry = SerializedShaders.ShaderMapEntries[ShaderMapIndex];
-
+#if TRACK_SHADER_PRELOADS
+	const uint32 FrameNumber = GFrameNumber;
+	for (uint32 i = 0u; i < ShaderMapEntry.NumShaders; ++i)
+	{
+		const uint32 ShaderIndex = SerializedShaders.ShaderIndices[ShaderMapEntry.ShaderIndicesOffset + i];
+		ShaderFramePreloaded[ShaderIndex] = FMath::Min(ShaderFramePreloaded[ShaderIndex], FrameNumber);
+	}
+#endif // TRACK_SHADER_PRELOADS
 	const EAsyncIOPriorityAndFlags IOPriority = (EAsyncIOPriorityAndFlags)GShaderCodeLibraryAsyncLoadingPriority;
-	return FileCacheHandle->PreloadData(&SerializedShaders.PreloadEntries[ShaderMapEntry.FirstPreloadIndex], ShaderMapEntry.NumPreloadEntries, IOPriority);
+	return FileCacheHandle->PreloadData(&SerializedShaders.PreloadEntries[ShaderMapEntry.FirstPreloadIndex], ShaderMapEntry.NumPreloadEntries, LibraryCodeOffset, IOPriority);
+}
+
+void FShaderCodeArchive::ReleasePreloadedShaderMap(int32 ShaderMapIndex)
+{
+	const FShaderMapEntry& ShaderMapEntry = SerializedShaders.ShaderMapEntries[ShaderMapIndex];
+#if TRACK_SHADER_PRELOADS
+	for (uint32 i = 0u; i < ShaderMapEntry.NumShaders; ++i)
+	{
+		const uint32 ShaderIndex = SerializedShaders.ShaderIndices[ShaderMapEntry.ShaderIndicesOffset + i];
+		ShaderFramePreloaded[ShaderIndex] = ~0u;
+	}
+#endif // TRACK_SHADER_PRELOADS
+	FileCacheHandle->ReleasePreloadedData(&SerializedShaders.PreloadEntries[ShaderMapEntry.FirstPreloadIndex], ShaderMapEntry.NumPreloadEntries, LibraryCodeOffset);
 }
 
 TRefCountPtr<FRHIShader> FShaderCodeArchive::CreateShader(int32 Index)
diff --git a/Engine/Source/Runtime/RenderCore/Private/ShaderCodeLibrary.cpp b/Engine/Source/Runtime/RenderCore/Private/ShaderCodeLibrary.cpp
--- a/Engine/Source/Runtime/RenderCore/Private/ShaderCodeLibrary.cpp
+++ b/Engine/Source/Runtime/RenderCore/Private/ShaderCodeLibrary.cpp
@@ -788,6 +788,8 @@ void FShaderMapResource_SharedCode::ReleaseRHI()
 		}
 	}
 
+	LibraryInstance->Library->ReleasePreloadedShaderMap(ShaderMapIndex);
+
 	FShaderMapResource::ReleaseRHI();
 }
 
diff --git a/Engine/Source/Runtime/RenderCore/Public/ShaderCodeArchive.h b/Engine/Source/Runtime/RenderCore/Public/ShaderCodeArchive.h
--- a/Engine/Source/Runtime/RenderCore/Public/ShaderCodeArchive.h
+++ b/Engine/Source/Runtime/RenderCore/Public/ShaderCodeArchive.h
@@ -110,6 +110,8 @@ public:
 	}
 };
 
+#define TRACK_SHADER_PRELOADS STATS
+
 class FShaderCodeArchive : public FRHIShaderLibrary
 {
 public:
@@ -122,7 +124,6 @@ public:
 	uint32 GetSizeBytes() const
 	{
 		return sizeof(*this) +
-			SerializedShaders.GetAllocatedSize() +
 			SerializedShaders.GetAllocatedSize();
 	}
 
@@ -150,6 +151,8 @@ public:
 
 	virtual FGraphEventRef PreloadShaderMap(int32 ShaderMapIndex) override;
 
+	virtual void ReleasePreloadedShaderMap(int32 ShaderMapIndex) override;
+
 	virtual TRefCountPtr<FRHIShader> CreateShader(int32 Index) override;
 	virtual void Teardown() override;
 
@@ -173,4 +176,8 @@ protected:
 
 	// The shader code present in the library
 	FSerializedShaderArchive SerializedShaders;
+
+#if TRACK_SHADER_PRELOADS
+	TArray<uint32> ShaderFramePreloaded;
+#endif
 };
