Allow text substitution for fragments within a culture display name to be culture specific

#rnx

diff --git a/Engine/Source/Runtime/Core/Private/Internationalization/Culture.cpp b/Engine/Source/Runtime/Core/Private/Internationalization/Culture.cpp
--- a/Engine/Source/Runtime/Core/Private/Internationalization/Culture.cpp
+++ b/Engine/Source/Runtime/Core/Private/Internationalization/Culture.cpp
@@ -1,6 +1,7 @@
 // Copyright Epic Games, Inc. All Rights Reserved.
 
 #include "Internationalization/Culture.h"
+#include "Containers/ArrayView.h"
 #include "Misc/ConfigCacheIni.h"
 #include "Misc/App.h"
 
@@ -10,9 +11,16 @@
 #include "Internationalization/LegacyCulture.h"
 #endif
 
-void ApplyCultureDisplayNameSubstitutes(FString& InOutDisplayName)
+void ApplyCultureDisplayNameSubstitutes(TArrayView<const FString> InPrioritizedCultureNames, FString& InOutDisplayName)
 {
-	static TArray<TTuple<FString, FString>> CultureDisplayNameSubstitutes;
+	struct FDisplayNameSubstitute
+	{
+		FString Culture;
+		FString OldString;
+		FString NewString;
+	};
+
+	static TArray<FDisplayNameSubstitute> CultureDisplayNameSubstitutes;
 
 	// Conditionally load the required config data
 	{
@@ -43,23 +51,38 @@ void ApplyCultureDisplayNameSubstitutes(FString& InOutDisplayName)
 				}
 			}
 
-			// Each substitute should be a semi-colon separated pair of data: Old;New
+			// Each substitute should be a semi-colon separated set of data: [Culture;]Old;New
 			CultureDisplayNameSubstitutes.Reserve(CultureDisplayNameSubstitutesStrArray.Num());
 			for (const FString& CultureDisplayNameSubstituteStr : CultureDisplayNameSubstitutesStrArray)
 			{
-				FString OldDisplayFragment;
-				FString NewDisplayFragment;
-				if (CultureDisplayNameSubstituteStr.Split(TEXT(";"), &OldDisplayFragment, &NewDisplayFragment, ESearchCase::CaseSensitive))
+				TArray<FString> DisplayFragments;
+				CultureDisplayNameSubstituteStr.ParseIntoArray(DisplayFragments, TEXT(";"));
+
+				if (DisplayFragments.Num() == 2)
 				{
-					CultureDisplayNameSubstitutes.Add(MakeTuple(MoveTemp(OldDisplayFragment), MoveTemp(NewDisplayFragment)));
+					CultureDisplayNameSubstitutes.Add(FDisplayNameSubstitute{ FString(), MoveTemp(DisplayFragments[0]), MoveTemp(DisplayFragments[1]) });
+				}
+				else if (DisplayFragments.Num() == 3)
+				{
+					CultureDisplayNameSubstitutes.Add(FDisplayNameSubstitute{ MoveTemp(DisplayFragments[0]), MoveTemp(DisplayFragments[1]), MoveTemp(DisplayFragments[2]) });
 				}
 			}
+
+			// Sort by culture name length, so that more specific cultures get first refusal at a replacement
+			CultureDisplayNameSubstitutes.StableSort([](const FDisplayNameSubstitute& InOne, const FDisplayNameSubstitute& InTwo)
+			{
+				return InOne.Culture.Len() > InTwo.Culture.Len();
+			});
 		}
 	}
 
-	for (const auto& CultureDisplayNameSubstitutePair : CultureDisplayNameSubstitutes)
+	for (const FDisplayNameSubstitute& CultureDisplayNameSubstitute : CultureDisplayNameSubstitutes)
 	{
-		InOutDisplayName.ReplaceInline(*CultureDisplayNameSubstitutePair.Key, *CultureDisplayNameSubstitutePair.Value, ESearchCase::CaseSensitive);
+		const bool bValidForCulture = CultureDisplayNameSubstitute.Culture.IsEmpty() || InPrioritizedCultureNames.Contains(CultureDisplayNameSubstitute.Culture);
+		if (bValidForCulture)
+		{
+			InOutDisplayName.ReplaceInline(*CultureDisplayNameSubstitute.OldString, *CultureDisplayNameSubstitute.NewString, ESearchCase::CaseSensitive);
+		}
 	}
 }
 
@@ -80,7 +103,7 @@ FCulture::FCulture(TUniquePtr<FCultureImplementation>&& InImplementation)
 	, CachedVariant(Implementation->GetVariant())
 	, CachedIsRightToLeft(Implementation->IsRightToLeft())
 {
-	RefreshCultureDisplayNames();
+	RefreshCultureDisplayNames(TArray<FString>()); // The display name for the current language will be updated by a post-construct call to RefreshCultureDisplayNames with the correct language data
 }
 
 FCulture::~FCulture()
@@ -290,23 +313,30 @@ const TArray<ETextPluralForm>& FCulture::GetValidPluralForms(const ETextPluralTy
 	return Implementation->GetValidPluralForms(PluralType);
 }
 
-void FCulture::RefreshCultureDisplayNames(const bool bFullRefresh)
+void FCulture::RefreshCultureDisplayNames(const TArray<FString>& InPrioritizedDisplayCultureNames, const bool bFullRefresh)
 {
 	CachedDisplayName = Implementation->GetDisplayName();
-	ApplyCultureDisplayNameSubstitutes(CachedDisplayName);
+	ApplyCultureDisplayNameSubstitutes(InPrioritizedDisplayCultureNames, CachedDisplayName);
 
 	if (bFullRefresh)
 	{
-		CachedEnglishName = Implementation->GetEnglishName();
-		ApplyCultureDisplayNameSubstitutes(CachedEnglishName);
+		{
+			static const FString EnglishCultureName = TEXT("en");
+			CachedEnglishName = Implementation->GetEnglishName();
+			ApplyCultureDisplayNameSubstitutes(MakeArrayView(&EnglishCultureName, 1), CachedEnglishName);
+		}
 
-		CachedNativeName = Implementation->GetNativeName();
-		ApplyCultureDisplayNameSubstitutes(CachedNativeName);
+		{
+			const TArray<FString> PrioritizedNativeCultureNames = GetPrioritizedParentCultureNames();
+
+			CachedNativeName = Implementation->GetNativeName();
+			ApplyCultureDisplayNameSubstitutes(PrioritizedNativeCultureNames, CachedNativeName);
 
-		CachedNativeLanguage = Implementation->GetNativeLanguage();
-		ApplyCultureDisplayNameSubstitutes(CachedNativeLanguage);
+			CachedNativeLanguage = Implementation->GetNativeLanguage();
+			ApplyCultureDisplayNameSubstitutes(PrioritizedNativeCultureNames, CachedNativeLanguage);
 
-		CachedNativeRegion = Implementation->GetNativeRegion();
-		ApplyCultureDisplayNameSubstitutes(CachedNativeRegion);
+			CachedNativeRegion = Implementation->GetNativeRegion();
+			ApplyCultureDisplayNameSubstitutes(PrioritizedNativeCultureNames, CachedNativeRegion);
+		}
 	}
 }
diff --git a/Engine/Source/Runtime/Core/Private/Internationalization/ICUInternationalization.cpp b/Engine/Source/Runtime/Core/Private/Internationalization/ICUInternationalization.cpp
--- a/Engine/Source/Runtime/Core/Private/Internationalization/ICUInternationalization.cpp
+++ b/Engine/Source/Runtime/Core/Private/Internationalization/ICUInternationalization.cpp
@@ -179,6 +179,8 @@ bool FICUInternationalization::Initialize()
 	I18N->CurrentLanguage = I18N->DefaultLanguage;
 	I18N->CurrentLocale = I18N->DefaultLocale;
 
+	HandleLanguageChanged(I18N->CurrentLanguage.ToSharedRef());
+
 #if ENABLE_LOC_TESTING
 	I18N->AddCustomCulture(MakeShared<FLeetCulture>(I18N->InvariantCulture.ToSharedRef()));
 #endif
@@ -549,26 +551,28 @@ bool FICUInternationalization::IsCultureAllowed(const FString& Name)
 	return (EnabledCultures.Num() == 0 || EnabledCultures.Contains(Name)) && !DisabledCultures.Contains(Name);
 }
 
-void FICUInternationalization::RefreshCultureDisplayNames()
+void FICUInternationalization::RefreshCultureDisplayNames(const TArray<FString>& InPrioritizedDisplayCultureNames)
 {
 	// Update the cached display names in any existing cultures
 	FScopeLock Lock(&CachedCulturesCS);
 	for (const auto& CachedCulturePair : CachedCultures)
 	{
-		CachedCulturePair.Value->RefreshCultureDisplayNames();
+		CachedCulturePair.Value->RefreshCultureDisplayNames(InPrioritizedDisplayCultureNames);
 	}
 }
 
-void FICUInternationalization::HandleLanguageChanged(const FString& Name)
+void FICUInternationalization::HandleLanguageChanged(const FCultureRef InNewLanguage)
 {
 	UErrorCode ICUStatus = U_ZERO_ERROR;
-	uloc_setDefault(StringCast<char>(*Name).Get(), &ICUStatus);
+	uloc_setDefault(StringCast<char>(*InNewLanguage->GetName()).Get(), &ICUStatus);
+
+	CachedPrioritizedDisplayCultureNames = InNewLanguage->GetPrioritizedParentCultureNames();
 
 	// Update the cached display names in any existing cultures
 	FScopeLock Lock(&CachedCulturesCS);
 	for (const auto& CachedCulturePair : CachedCultures)
 	{
-		CachedCulturePair.Value->RefreshCultureDisplayNames(/*bFullRefresh*/false);
+		CachedCulturePair.Value->RefreshCultureDisplayNames(CachedPrioritizedDisplayCultureNames, /*bFullRefresh*/false);
 	}
 }
 
@@ -741,8 +745,13 @@ FCulturePtr FICUInternationalization::FindOrMakeCanonizedCulture(const FString&
 
 	if (NewCulture.IsValid())
 	{
-		FScopeLock Lock(&CachedCulturesCS);
-		CachedCultures.Add(Name, NewCulture.ToSharedRef());
+		// Ensure the display name is up-to-date
+		NewCulture->RefreshCultureDisplayNames(CachedPrioritizedDisplayCultureNames, /*bFullRefresh*/false);
+
+		{
+			FScopeLock Lock(&CachedCulturesCS);
+			CachedCultures.Add(Name, NewCulture.ToSharedRef());
+		}
 	}
 
 	return NewCulture;
diff --git a/Engine/Source/Runtime/Core/Private/Internationalization/ICUInternationalization.h b/Engine/Source/Runtime/Core/Private/Internationalization/ICUInternationalization.h
--- a/Engine/Source/Runtime/Core/Private/Internationalization/ICUInternationalization.h
+++ b/Engine/Source/Runtime/Core/Private/Internationalization/ICUInternationalization.h
@@ -38,9 +38,9 @@ public:
 	bool IsCultureRemapped(const FString& Name, FString* OutMappedCulture);
 	bool IsCultureAllowed(const FString& Name);
 
-	void RefreshCultureDisplayNames();
+	void RefreshCultureDisplayNames(const TArray<FString>& InPrioritizedDisplayCultureNames);
 
-	void HandleLanguageChanged(const FString& Name);
+	void HandleLanguageChanged(const FCultureRef InNewLanguage);
 	void GetCultureNames(TArray<FString>& CultureNames) const;
 	TArray<FString> GetPrioritizedCultureNames(const FString& Name);
 	FCulturePtr GetCulture(const FString& Name);
@@ -106,6 +106,8 @@ private:
 	TUniquePtr<icu::GregorianCalendar> InvariantGregorianCalendar;
 	FCriticalSection InvariantGregorianCalendarCS;
 
+	TArray<FString> CachedPrioritizedDisplayCultureNames;
+
 	static UBool OpenDataFile(const void* InContext, void** OutFileContext, void** OutContents, const char* InPath);
 	static void CloseDataFile(const void* InContext, void* const InFileContext, void* const InContents);
 
diff --git a/Engine/Source/Runtime/Core/Private/Internationalization/Internationalization.cpp b/Engine/Source/Runtime/Core/Private/Internationalization/Internationalization.cpp
--- a/Engine/Source/Runtime/Core/Private/Internationalization/Internationalization.cpp
+++ b/Engine/Source/Runtime/Core/Private/Internationalization/Internationalization.cpp
@@ -54,7 +54,7 @@ bool FInternationalization::SetCurrentCulture(const FString& InCultureName)
 			CurrentLocale = NewCulture;
 			CurrentAssetGroupCultures.Reset();
 	
-			Implementation->HandleLanguageChanged(InCultureName);
+			Implementation->HandleLanguageChanged(CurrentLanguage.ToSharedRef());
 	
 			BroadcastCultureChanged();
 		}
@@ -73,7 +73,7 @@ bool FInternationalization::SetCurrentLanguage(const FString& InCultureName)
 		{
 			CurrentLanguage = NewCulture;
 
-			Implementation->HandleLanguageChanged(InCultureName);
+			Implementation->HandleLanguageChanged(CurrentLanguage.ToSharedRef());
 
 			BroadcastCultureChanged();
 		}
@@ -110,7 +110,7 @@ bool FInternationalization::SetCurrentLanguageAndLocale(const FString& InCulture
 			CurrentLanguage = NewCulture;
 			CurrentLocale = NewCulture;
 
-			Implementation->HandleLanguageChanged(InCultureName);
+			Implementation->HandleLanguageChanged(CurrentLanguage.ToSharedRef());
 
 			BroadcastCultureChanged();
 		}
@@ -205,7 +205,7 @@ void FInternationalization::RestoreCultureState(const FCultureStateSnapshot& InS
 
 				CurrentLanguage = NewCulture;
 
-				Implementation->HandleLanguageChanged(InSnapshot.Language);
+				Implementation->HandleLanguageChanged(CurrentLanguage.ToSharedRef());
 			}
 		}
 	}
@@ -419,9 +419,9 @@ bool FInternationalization::IsCultureAllowed(const FString& Name)
 	return Implementation->IsCultureAllowed(Name);
 }
 
-void FInternationalization::RefreshCultureDisplayNames()
+void FInternationalization::RefreshCultureDisplayNames(const TArray<FString>& InPrioritizedDisplayCultureNames)
 {
-	Implementation->RefreshCultureDisplayNames();
+	Implementation->RefreshCultureDisplayNames(InPrioritizedDisplayCultureNames);
 }
 
 void FInternationalization::GetCultureNames(TArray<FString>& CultureNames) const
diff --git a/Engine/Source/Runtime/Core/Private/Internationalization/LegacyInternationalization.cpp b/Engine/Source/Runtime/Core/Private/Internationalization/LegacyInternationalization.cpp
--- a/Engine/Source/Runtime/Core/Private/Internationalization/LegacyInternationalization.cpp
+++ b/Engine/Source/Runtime/Core/Private/Internationalization/LegacyInternationalization.cpp
@@ -46,11 +46,11 @@ bool FLegacyInternationalization::IsCultureAllowed(const FString& Name)
 	return true;
 }
 
-void FLegacyInternationalization::RefreshCultureDisplayNames()
+void FLegacyInternationalization::RefreshCultureDisplayNames(const TArray<FString>& InPrioritizedDisplayCultureNames)
 {
 }
 
-void FLegacyInternationalization::HandleLanguageChanged(const FString& Name)
+void FLegacyInternationalization::HandleLanguageChanged(const FCultureRef InNewLanguage)
 {
 }
 
diff --git a/Engine/Source/Runtime/Core/Private/Internationalization/LegacyInternationalization.h b/Engine/Source/Runtime/Core/Private/Internationalization/LegacyInternationalization.h
--- a/Engine/Source/Runtime/Core/Private/Internationalization/LegacyInternationalization.h
+++ b/Engine/Source/Runtime/Core/Private/Internationalization/LegacyInternationalization.h
@@ -23,9 +23,9 @@ public:
 	bool IsCultureRemapped(const FString& Name, FString* OutMappedCulture);
 	bool IsCultureAllowed(const FString& Name);
 
-	void RefreshCultureDisplayNames();
+	void RefreshCultureDisplayNames(const TArray<FString>& InPrioritizedDisplayCultureNames);
 
-	void HandleLanguageChanged(const FString& Name);
+	void HandleLanguageChanged(const FCultureRef InNewLanguage);
 	void GetCultureNames(TArray<FString>& CultureNames) const;
 	TArray<FString> GetPrioritizedCultureNames(const FString& Name);
 	FCulturePtr GetCulture(const FString& Name);
diff --git a/Engine/Source/Runtime/Core/Public/Internationalization/Culture.h b/Engine/Source/Runtime/Core/Public/Internationalization/Culture.h
--- a/Engine/Source/Runtime/Core/Public/Internationalization/Culture.h
+++ b/Engine/Source/Runtime/Core/Public/Internationalization/Culture.h
@@ -98,7 +98,7 @@ public:
 	 */
 	const TArray<ETextPluralForm>& GetValidPluralForms(const ETextPluralType PluralType) const;
 	
-	void RefreshCultureDisplayNames(const bool bFullRefresh = true);
+	void RefreshCultureDisplayNames(const TArray<FString>& InPrioritizedDisplayCultureNames, const bool bFullRefresh = true);
 
 private:
 	explicit FCulture(TUniquePtr<FCultureImplementation>&& InImplementation);
diff --git a/Engine/Source/Runtime/Core/Public/Internationalization/Internationalization.h b/Engine/Source/Runtime/Core/Public/Internationalization/Internationalization.h
--- a/Engine/Source/Runtime/Core/Public/Internationalization/Internationalization.h
+++ b/Engine/Source/Runtime/Core/Public/Internationalization/Internationalization.h
@@ -175,7 +175,7 @@ public:
 	CORE_API bool IsCultureAllowed(const FString& Name);
 
 	/** Refresh the display names of the cached cultures */
-	CORE_API void RefreshCultureDisplayNames();
+	CORE_API void RefreshCultureDisplayNames(const TArray<FString>& InPrioritizedDisplayCultureNames);
 
 #if ENABLE_LOC_TESTING
 	static CORE_API FString& Leetify(FString& SourceString);
