Missing file from CL# 12491125

#rnx

diff --git a/Engine/Source/Runtime/Core/Private/Internationalization/TextLocalizationManager.cpp b/Engine/Source/Runtime/Core/Private/Internationalization/TextLocalizationManager.cpp
--- a/Engine/Source/Runtime/Core/Private/Internationalization/TextLocalizationManager.cpp
+++ b/Engine/Source/Runtime/Core/Private/Internationalization/TextLocalizationManager.cpp
@@ -329,7 +329,7 @@ void InitEngineTextLocalization()
 
 	// Run this now that the config system is definitely initialized
 	// to refresh anything that was cached before it was ready
-	FInternationalization::Get().RefreshCultureDisplayNames();
+	FInternationalization::Get().RefreshCultureDisplayNames(FInternationalization::Get().GetCurrentLanguage()->GetPrioritizedParentCultureNames());
 
 	ELocalizationLoadFlags LocLoadFlags = ELocalizationLoadFlags::None;
 	LocLoadFlags |= (WITH_EDITOR ? ELocalizationLoadFlags::Editor : ELocalizationLoadFlags::None);
