Removed delta mem comparisons for master process memory between each fork since was confirmed to always be 0.

diff --git a/Engine/Source/Runtime/Core/Private/Unix/UnixPlatformProcess.cpp b/Engine/Source/Runtime/Core/Private/Unix/UnixPlatformProcess.cpp
--- a/Engine/Source/Runtime/Core/Private/Unix/UnixPlatformProcess.cpp
+++ b/Engine/Source/Runtime/Core/Private/Unix/UnixPlatformProcess.cpp
@@ -1397,11 +1397,10 @@ FGenericPlatformProcess::EWaitAndForkResult FUnixPlatformProcess::WaitAndFork()
 			FPlatformProcess::Sleep(WAIT_AND_FORK_CHILD_SPAWN_DELAY);
 
 			FMemoryStatsHolder CurrentMasterMemStats(FPlatformMemory::GetStats());
-			UE_LOG(LogHAL, Log, TEXT("MemoryStats PreFork: AvailablePhysical: %.02fMiB (%+.02fMiB), PeakPhysical: %.02fMiB (%+.02fMiB), PeakVirtual: %.02fMiB (%+.02fMiB)"),
+			UE_LOG(LogHAL, Log, TEXT("MemoryStats PreFork: AvailablePhysical: %.02fMiB (%+.02fMiB), PeakPhysical: %.02fMiB, PeakVirtual: %.02fMiB"),
 				CurrentMasterMemStats.AvailablePhysical, (CurrentMasterMemStats.AvailablePhysical - PreviousMasterMemStats.AvailablePhysical),
-				CurrentMasterMemStats.PeakUsedPhysical, (CurrentMasterMemStats.PeakUsedPhysical - PreviousMasterMemStats.PeakUsedPhysical),
-				CurrentMasterMemStats.PeakUsedVirtual, (CurrentMasterMemStats.PeakUsedVirtual - PreviousMasterMemStats.PeakUsedVirtual)
-				
+				CurrentMasterMemStats.PeakUsedPhysical,
+				CurrentMasterMemStats.PeakUsedVirtual				
 			);
 			PreviousMasterMemStats = CurrentMasterMemStats;
 			
