Enable the actual optimization after it was accidently moved to the wrong CL

diff --git a/Engine/Source/Runtime/Core/Private/HAL/MallocBinned2.cpp b/Engine/Source/Runtime/Core/Private/HAL/MallocBinned2.cpp
--- a/Engine/Source/Runtime/Core/Private/HAL/MallocBinned2.cpp
+++ b/Engine/Source/Runtime/Core/Private/HAL/MallocBinned2.cpp
@@ -11,6 +11,10 @@
 #include "HAL/PlatformMisc.h"
 #include "Misc/App.h"
 #include "HAL/MallocTimer.h"
+#include "ProfilingDebugging/CsvProfiler.h"
+#if CSV_PROFILER
+CSV_DEFINE_CATEGORY_MODULE(CORE_API, FMemory, true);
+#endif
 
 PRAGMA_DISABLE_UNSAFE_TYPECAST_WARNINGS
 
@@ -619,7 +623,7 @@ FMallocBinned2::FPoolInfo& FMallocBinned2::FPoolList::PushNewPoolToFront(FMalloc
 	const uint32 LocalPageSize = Allocator.PageSize;
 
 	// Allocate memory.
-	void* FreePtr = Allocator.CachedOSPageAllocator.Allocate(LocalPageSize);
+	void* FreePtr = Allocator.CachedOSPageAllocator.Allocate(LocalPageSize, FMemory::AllocationHints::SmallPool);
 	if (!FreePtr)
 	{
 		Private::OutOfMemory(LocalPageSize);
@@ -1319,6 +1323,10 @@ void FMallocBinned2::DumpAllocatorStats(class FOutputDevice& Ar)
 
 void FMallocBinned2::UpdateStats()
 {
+#if CSV_PROFILER
+	CSV_CUSTOM_STAT(FMemory, AllocatorCachedSlackMB, (int32)(CachedOSPageAllocator.GetCachedFreeTotal()/(1024*1024)), ECsvCustomStatOp::Set);
+#endif
+
 	FScopedVirtualMallocTimer::UpdateStats();
 }
 
