Buildfix, adding default param for allocation hints to PooledVirtualMemoryAllocator

diff --git a/Engine/Source/Runtime/Core/Private/HAL/Allocators/PooledVirtualMemoryAllocator.cpp b/Engine/Source/Runtime/Core/Private/HAL/Allocators/PooledVirtualMemoryAllocator.cpp
--- a/Engine/Source/Runtime/Core/Private/HAL/Allocators/PooledVirtualMemoryAllocator.cpp
+++ b/Engine/Source/Runtime/Core/Private/HAL/Allocators/PooledVirtualMemoryAllocator.cpp
@@ -41,7 +41,7 @@ FPooledVirtualMemoryAllocator::FPooledVirtualMemoryAllocator()
 	}
 }
 
-void* FPooledVirtualMemoryAllocator::Allocate(SIZE_T Size)
+void* FPooledVirtualMemoryAllocator::Allocate(SIZE_T Size, uint32 /*AllocationHint = 0*/)
 {
 	if (Size > Limits::MaxAllocationSizeToPool)
 	{
diff --git a/Engine/Source/Runtime/Core/Public/HAL/Allocators/PooledVirtualMemoryAllocator.h b/Engine/Source/Runtime/Core/Public/HAL/Allocators/PooledVirtualMemoryAllocator.h
--- a/Engine/Source/Runtime/Core/Public/HAL/Allocators/PooledVirtualMemoryAllocator.h
+++ b/Engine/Source/Runtime/Core/Public/HAL/Allocators/PooledVirtualMemoryAllocator.h
@@ -61,7 +61,7 @@ struct FPooledVirtualMemoryAllocator
 {
 	FPooledVirtualMemoryAllocator();
 
-	void* Allocate(SIZE_T Size);
+	void* Allocate(SIZE_T Size, uint32 AllocationHint = 0);
 	void Free(void* Ptr, SIZE_T Size);
 	void FreeAll();
 
