TRobinHoodHashMap - Change costly check to a checkSlow()

diff --git a/Engine/Source/Runtime/Core/Public/Experimental/Containers/RobinHoodHashTable.h b/Engine/Source/Runtime/Core/Public/Experimental/Containers/RobinHoodHashTable.h
--- a/Engine/Source/Runtime/Core/Public/Experimental/Containers/RobinHoodHashTable.h
+++ b/Engine/Source/Runtime/Core/Public/Experimental/Containers/RobinHoodHashTable.h
@@ -237,10 +237,10 @@ namespace RobinHoodHashTable_Private
 			{
 				CHECK_CONCURRENT_ACCESS(ConcurrentWriters == 0);
 				CHECK_CONCURRENT_ACCESS(FPlatformAtomics::InterlockedIncrement(&ConcurrentReaders) >= 1);
-#if DO_CHECK
+#if DO_GUARD_SLOW
 				for (IndexType Idx : FreeList)
 				{
-					check(Idx != Index);
+					checkSlow(Idx != Index);
 				}
 #endif
 				CHECK_CONCURRENT_ACCESS(FPlatformAtomics::InterlockedDecrement(&ConcurrentReaders) >= 0);
@@ -252,10 +252,10 @@ namespace RobinHoodHashTable_Private
 			{
 				CHECK_CONCURRENT_ACCESS(ConcurrentWriters == 0);
 				CHECK_CONCURRENT_ACCESS(FPlatformAtomics::InterlockedIncrement(&ConcurrentReaders) >= 1);
-#if DO_CHECK
+#if DO_GUARD_SLOW
 				for (IndexType Idx : FreeList)
 				{
-					check(Idx != Index);
+					checkSlow(Idx != Index);
 				}
 #endif
 				CHECK_CONCURRENT_ACCESS(FPlatformAtomics::InterlockedDecrement(&ConcurrentReaders) >= 0);
