Fixes for overlapping LLM tags on iOS
- Reserve space in AppleLLM tag for MetalLLM and use those spaces.
- Remove NanoMallocGrabberLLM tags. Still have corresponding stats, but no overlapping LLM tags.
#tests iOS Client

diff --git a/Engine/Source/Runtime/Apple/MetalRHI/Private/MetalLLM.h b/Engine/Source/Runtime/Apple/MetalRHI/Private/MetalLLM.h
--- a/Engine/Source/Runtime/Apple/MetalRHI/Private/MetalLLM.h
+++ b/Engine/Source/Runtime/Apple/MetalRHI/Private/MetalLLM.h
@@ -4,6 +4,7 @@
 
 #include "MetalRHIPrivate.h"
 #include "HAL/LowLevelMemTracker.h"
+#include "Apple/AppleLLM.h"
 
 @interface FMetalDeallocHandler : FApplePlatformObject<NSObject>
 {
@@ -20,7 +21,7 @@
 
 enum class ELLMTagMetal : LLM_TAG_TYPE
 {
-	Buffers = (LLM_TAG_TYPE)ELLMTag::PlatformTagStart,
+	Buffers = (LLM_TAG_TYPE)ELLMTagApple::AppleMetalTagsStart,
 	Textures,
 	Heaps,
 	RenderTargets,
@@ -28,7 +29,7 @@ enum class ELLMTagMetal : LLM_TAG_TYPE
 	Count
 };
 
-static_assert((int32)ELLMTagMetal::Count <= (int32)ELLMTag::PlatformTagEnd, "too many ELLMTagMetal tags");
+static_assert((int32)ELLMTagMetal::Count <= (int32)ELLMTagApple::AppleMetalTagsEnd, "too many ELLMTagMetal tags. Need to increase LLM_TAG_APPLE_NUM_METAL_TAGS_RESERVED");
 
 namespace MetalLLM
 {
diff --git a/Engine/Source/Runtime/Core/Private/HAL/MallocBinned.cpp b/Engine/Source/Runtime/Core/Private/HAL/MallocBinned.cpp
--- a/Engine/Source/Runtime/Core/Private/HAL/MallocBinned.cpp
+++ b/Engine/Source/Runtime/Core/Private/HAL/MallocBinned.cpp
@@ -1125,10 +1125,10 @@ void FMallocBinned::InitializeStatsMetadata()
 	GET_STATFNAME(STAT_Binned_SlackCurrent);
 	
 #if USE_OS_SMALL_BLOCK_GRAB_MEMORY_FROM_OS && ENABLE_LOW_LEVEL_MEM_TRACKER
-	FLowLevelMemTracker::Get().RegisterPlatformTag((int32)ELLMTagNanoMallocGrabber::NanoMallocPagesCurrent, TEXT("Nano Malloc Pages Current"), GET_STATFNAME(STAT_Binned_NanoMallocPages_Current), GET_STATFNAME(STAT_EngineSummaryLLM));
-	FLowLevelMemTracker::Get().RegisterPlatformTag((int32)ELLMTagNanoMallocGrabber::NanoMallocPagesPeak, TEXT("Nano Malloc Pages Peak"), GET_STATFNAME(STAT_Binned_NanoMallocPages_Peak), GET_STATFNAME(STAT_EngineSummaryLLM));
-	FLowLevelMemTracker::Get().RegisterPlatformTag((int32)ELLMTagNanoMallocGrabber::NanoMallocPagesWasteCurrent, TEXT("Nano Malloc Pages Waste Current"), GET_STATFNAME(STAT_Binned_NanoMallocPages_WasteCurrent), GET_STATFNAME(STAT_EngineSummaryLLM));
-	FLowLevelMemTracker::Get().RegisterPlatformTag((int32)ELLMTagNanoMallocGrabber::NanoMallocPagesWastePeak, TEXT("Nano Malloc Pages Waste Peak"), GET_STATFNAME(STAT_Binned_NanoMallocPages_WastePeak), GET_STATFNAME(STAT_EngineSummaryLLM));
+	GET_STATFNAME(STAT_Binned_NanoMallocPages_Current);
+	GET_STATFNAME(STAT_Binned_NanoMallocPages_Peak);
+	GET_STATFNAME(STAT_Binned_NanoMallocPages_WasteCurrent);
+	GET_STATFNAME(STAT_Binned_NanoMallocPages_WastePeak);
 #endif //USE_OS_SMALL_BLOCK_GRAB_MEMORY_FROM_OS
 }
 
diff --git a/Engine/Source/Runtime/Core/Public/Apple/AppleLLM.h b/Engine/Source/Runtime/Core/Public/Apple/AppleLLM.h
--- a/Engine/Source/Runtime/Core/Public/Apple/AppleLLM.h
+++ b/Engine/Source/Runtime/Core/Public/Apple/AppleLLM.h
@@ -9,13 +9,18 @@
 #define LLM_SCOPE_APPLE(Tag) LLM_SCOPE((ELLMTag)Tag)
 #define LLM_PLATFORM_SCOPE_APPLE(Tag) LLM_PLATFORM_SCOPE((ELLMTag)Tag)
 
+#define LLM_TAG_APPLE_NUM_METAL_TAGS_RESERVED 5
+
 enum class ELLMTagApple : LLM_TAG_TYPE
 {
 	ObjectiveC = (LLM_TAG_TYPE)ELLMTag::PlatformTagStart, // Use Instruments for detailed breakdown!
 
+	/* Reserved for Metal tags to be utlized in MetalLLM.h */
+	AppleMetalTagsStart,
+	AppleMetalTagsEnd = (LLM_TAG_TYPE)((LLM_TAG_TYPE)AppleMetalTagsStart + LLM_TAG_APPLE_NUM_METAL_TAGS_RESERVED),
+	
 	Count
 };
-
 static_assert((int32)ELLMTagApple::Count <= (int32)ELLMTag::PlatformTagEnd, "too many ELLMTagApple tags");
 
 namespace AppleLLM
diff --git a/Engine/Source/Runtime/Core/Public/HAL/MallocBinned.h b/Engine/Source/Runtime/Core/Public/HAL/MallocBinned.h
--- a/Engine/Source/Runtime/Core/Public/HAL/MallocBinned.h
+++ b/Engine/Source/Runtime/Core/Public/HAL/MallocBinned.h
@@ -106,17 +106,6 @@ DECLARE_LLM_MEMORY_STAT_EXTERN(TEXT("Nano Malloc Pages Waste Current"), STAT_Bin
 DECLARE_LLM_MEMORY_STAT_EXTERN(TEXT("Nano Malloc Pages Waste Peak"),STAT_Binned_NanoMallocPages_WastePeak,STATGROUP_LLMPlatform, CORE_API);
 #endif //USE_OS_SMALL_BLOCK_GRAB_MEMORY_FROM_OS
 
-#if USE_OS_SMALL_BLOCK_GRAB_MEMORY_FROM_OS && ENABLE_LOW_LEVEL_MEM_TRACKER
-enum class ELLMTagNanoMallocGrabber : LLM_TAG_TYPE
-{
-	NanoMallocPagesCurrent = (LLM_TAG_TYPE)ELLMTag::PlatformTagStart, // Use Instruments for detailed breakdown!
-	NanoMallocPagesPeak,
-	NanoMallocPagesWasteCurrent,
-	NanoMallocPagesWastePeak,
-	Count
-};
-#endif //#if USE_OS_SMALL_BLOCK_GRAB_MEMORY_FROM_OS
-
 
 //
 // Optimized virtual memory allocator.
