Added code to support stompmalloc and d3ddebug in engine test runs to flush out issues

diff --git a/Engine/Source/Runtime/Core/Private/Misc/App.cpp b/Engine/Source/Runtime/Core/Private/Misc/App.cpp
--- a/Engine/Source/Runtime/Core/Private/Misc/App.cpp
+++ b/Engine/Source/Runtime/Core/Private/Misc/App.cpp
@@ -353,6 +353,7 @@ void FApp::PrintStartupLogMessages()
 	UE_LOG(LogInit, Log, TEXT("%sCommand Line: %s"), *FilteredString, FCommandLine::GetForLogging());
 	UE_LOG(LogInit, Log, TEXT("Base Directory: %s"), FPlatformProcess::BaseDir());
 	//UE_LOG(LogInit, Log, TEXT("Character set: %s"), sizeof(TCHAR)==1 ? TEXT("ANSI") : TEXT("Unicode") );
+	UE_LOG(LogInit, Log, TEXT("Allocator: %s"), GMalloc->GetDescriptiveName());
 	UE_LOG(LogInit, Log, TEXT("Installed Engine Build: %d"), FApp::IsEngineInstalled() ? 1 : 0);
 
 	FDevVersionRegistration::DumpVersionsToLog();
