Copying //UE4/Release-4.25Plus-Dev to Release-4.25Plus (//UE4/Release-4.25Plus)

diff --git a/Engine/Source/Programs/AutomationTool/AutomationUtils/Platform.cs b/Engine/Source/Programs/AutomationTool/AutomationUtils/Platform.cs
--- a/Engine/Source/Programs/AutomationTool/AutomationUtils/Platform.cs
+++ b/Engine/Source/Programs/AutomationTool/AutomationUtils/Platform.cs
@@ -481,6 +481,14 @@ namespace AutomationTool
 			return "";
 		}
 
+		/// <summary>
+		/// Returns platform specific command line options for the IoStore cmdlet
+		/// </summary>
+		public virtual string GetPlatformIoStoreCommandLine(ProjectParams Params, DeploymentContext SC)
+		{
+			return "";
+		}
+
 		/// <summary>
 		/// True if this platform is supported.
 		/// </summary>
@@ -621,6 +629,11 @@ namespace AutomationTool
 			return new UnrealTargetPlatform[] { PlatformType };
 		}
 
+		public virtual DirectoryReference GetProjectRootForStage(DirectoryReference RuntimeRoot, StagedDirectoryReference RelativeProjectRootForStage)
+		{
+			return DirectoryReference.Combine(RuntimeRoot, RelativeProjectRootForStage.Name);
+		}
+
 		// let the platform set the exe extension if it chooses (otherwise, use
 		// the switch statement in GetExeExtension below)
 		protected virtual string GetPlatformExeExtension()
diff --git a/Engine/Source/Programs/AutomationTool/IOS/IOSPlatform.Automation.cs b/Engine/Source/Programs/AutomationTool/IOS/IOSPlatform.Automation.cs
--- a/Engine/Source/Programs/AutomationTool/IOS/IOSPlatform.Automation.cs
+++ b/Engine/Source/Programs/AutomationTool/IOS/IOSPlatform.Automation.cs
@@ -1799,4 +1799,9 @@ public class IOSPlatform : Platform
 	}
 
 	#endregion
+
+	public override DirectoryReference GetProjectRootForStage(DirectoryReference RuntimeRoot, StagedDirectoryReference RelativeProjectRootForStage)
+	{
+		return DirectoryReference.Combine(RuntimeRoot, "cookeddata/" + RelativeProjectRootForStage.Name);
+	}
 }
diff --git a/Engine/Source/Programs/AutomationTool/Scripts/CopyBuildToStagingDirectory.Automation.cs b/Engine/Source/Programs/AutomationTool/Scripts/CopyBuildToStagingDirectory.Automation.cs
--- a/Engine/Source/Programs/AutomationTool/Scripts/CopyBuildToStagingDirectory.Automation.cs
+++ b/Engine/Source/Programs/AutomationTool/Scripts/CopyBuildToStagingDirectory.Automation.cs
@@ -163,19 +163,29 @@ public partial class Project : CommandUtils
 		return CmdLine.ToString();
 	}
 
-	static private string GetIoStoreCommand(Dictionary<string, string> UnrealPakResponseFile, FileReference OutputLocation, FileReference PakOrderFileLocation, bool Compressed, EncryptionAndSigning.CryptoSettings CryptoSettings, string EncryptionKeyGuid, FileReference SecondaryPakOrderFileLocation=null)
+	static private string GetIoStoreCommandArguments(
+		Dictionary<string, string> UnrealPakResponseFile,
+		StagedFileReference ContainerRelativePath,
+		FileReference PakOrderFileLocation,
+		string PlatformOptions,
+		bool Compressed,
+		EncryptionAndSigning.CryptoSettings CryptoSettings,
+		string EncryptionKeyGuid,
+		FileReference SecondaryPakOrderFileLocation=null)
 	{
 		StringBuilder CmdLine = new StringBuilder();
-		CmdLine.AppendFormat("-Output={0}", MakePathSafeToUseWithCommandLine(Path.ChangeExtension(OutputLocation.FullName, null)));
+		CmdLine.AppendFormat("-Output={0}", MakePathSafeToUseWithCommandLine(ContainerRelativePath.Name));
 
 		// Force encryption of ALL files if we're using specific encryption key. This should be made an option per encryption key in the settings, but for our initial
 		// implementation we will just assume that we require maximum security for this data.
 		bool bForceEncryption = !string.IsNullOrEmpty(EncryptionKeyGuid);
-		string PakName = Path.GetFileNameWithoutExtension(OutputLocation.FullName);
-		string UnrealPakResponseFileName = CombinePaths(CmdEnv.LogFolder, "PakListIoStore_" + PakName + ".txt");
+		string ContainerName = Path.GetFileNameWithoutExtension(ContainerRelativePath.Name);
+		string UnrealPakResponseFileName = CombinePaths(CmdEnv.LogFolder, "PakListIoStore_" + ContainerName + ".txt");
 		WritePakResponseFile(UnrealPakResponseFileName, UnrealPakResponseFile, Compressed, CryptoSettings, bForceEncryption);
 		CmdLine.AppendFormat(" -ResponseFile={0}", CommandUtils.MakePathSafeToUseWithCommandLine(UnrealPakResponseFileName));
 
+		CmdLine.Append(PlatformOptions);
+
 		return CmdLine.ToString();
 	}
 
@@ -2090,6 +2100,50 @@ public partial class Project : CommandUtils
 			}
 		}
 
+		string BulkOption = "";
+		{
+			ConfigHierarchy PlatformEngineConfig;
+			if (Params.EngineConfigs.TryGetValue(SC.StageTargetPlatform.PlatformType, out PlatformEngineConfig))
+			{
+				bool bMasterEnable = false;
+				PlatformEngineConfig.GetBool("MemoryMappedFiles", "MasterEnable", out bMasterEnable);
+				if (bMasterEnable)
+				{
+					int Value = 0;
+					PlatformEngineConfig.GetInt32("MemoryMappedFiles", "Alignment", out Value);
+					if (Value > 0)
+					{
+						BulkOption = String.Format(" -AlignForMemoryMapping={0}", Value);
+					}
+				}
+			}
+		}
+
+		string PatchOptions = String.Empty;
+		if (bShouldGeneratePatch)
+		{
+			string PatchSeekOptMode = String.Empty;
+			string PatchSeekOptMaxGapSize = String.Empty;
+			string PatchSeekOptMaxInflationPercent = String.Empty;
+			string PatchSeekOptMaxAdjacentOrderDiff = String.Empty;
+			if (PlatformGameConfig.GetString("/Script/UnrealEd.ProjectPackagingSettings", "PatchSeekOptMode", out PatchSeekOptMode))
+			{
+				PatchOptions += String.Format(" -PatchSeekOptMode={0}", PatchSeekOptMode);
+			}
+			if (PlatformGameConfig.GetString("/Script/UnrealEd.ProjectPackagingSettings", "PatchSeekOptMaxGapSize", out PatchSeekOptMaxGapSize))
+			{
+				PatchOptions += String.Format(" -PatchSeekOptMaxGapSize={0}", PatchSeekOptMaxGapSize);
+			}
+			if (PlatformGameConfig.GetString("/Script/UnrealEd.ProjectPackagingSettings", "PatchSeekOptMaxInflationPercent", out PatchSeekOptMaxInflationPercent))
+			{
+				PatchOptions += String.Format(" -PatchSeekOptMaxInflationPercent={0}", PatchSeekOptMaxInflationPercent);
+			}
+			if (PlatformGameConfig.GetString("/Script/UnrealEd.ProjectPackagingSettings", "PatchSeekOptMaxAdjacentOrderDiff", out PatchSeekOptMaxAdjacentOrderDiff))
+			{
+				PatchOptions += String.Format(" -PatchSeekOptMaxAdjacentOrderDiff={0}", PatchSeekOptMaxAdjacentOrderDiff);
+			}
+		}
+
 		List<Tuple<FileReference, StagedFileReference, string>> Outputs = new List<Tuple<FileReference, StagedFileReference, string>>();
 		foreach (CreatePakParams PakParams in PakParamsList)
 		{
@@ -2232,46 +2286,10 @@ public partial class Project : CommandUtils
 						SecondaryOrderFile = PakOrderFileLocations[1];
 					}
 
-					string BulkOption = "";
-                    ConfigHierarchy PlatformEngineConfig = null;
-                    if (Params.EngineConfigs.TryGetValue(SC.StageTargetPlatform.PlatformType, out PlatformEngineConfig))
-                    {
-                        bool bMasterEnable = false;
-						PlatformEngineConfig.GetBool("MemoryMappedFiles", "MasterEnable", out bMasterEnable);
-                        if (bMasterEnable)
-                        {
-                            int Value = 0;
-							PlatformEngineConfig.GetInt32("MemoryMappedFiles", "Alignment", out Value);
-                            if (Value > 0)
-                            {
-                                BulkOption = String.Format(" -AlignForMemoryMapping={0}", Value);
-                            }
-                        }
-                    }
-
-					string AdditionalArgs = String.Empty;
-					if (bShouldGeneratePatch && !ShouldSkipGeneratingPatch(PlatformGameConfig, PakParams.PakName))
-					{
-						string PatchSeekOptMode = String.Empty;
-						string PatchSeekOptMaxGapSize = String.Empty;
-						string PatchSeekOptMaxInflationPercent = String.Empty;
-						string PatchSeekOptMaxAdjacentOrderDiff = String.Empty;
-						if (PlatformGameConfig.GetString("/Script/UnrealEd.ProjectPackagingSettings", "PatchSeekOptMode", out PatchSeekOptMode))
-						{
-							AdditionalArgs += String.Format(" -PatchSeekOptMode={0}", PatchSeekOptMode);
-						}
-						if (PlatformGameConfig.GetString("/Script/UnrealEd.ProjectPackagingSettings", "PatchSeekOptMaxGapSize", out PatchSeekOptMaxGapSize))
-						{
-							AdditionalArgs += String.Format(" -PatchSeekOptMaxGapSize={0}", PatchSeekOptMaxGapSize);
-						}
-						if (PlatformGameConfig.GetString("/Script/UnrealEd.ProjectPackagingSettings", "PatchSeekOptMaxInflationPercent", out PatchSeekOptMaxInflationPercent))
-						{
-							AdditionalArgs += String.Format(" -PatchSeekOptMaxInflationPercent={0}", PatchSeekOptMaxInflationPercent);
-						}
-						if (PlatformGameConfig.GetString("/Script/UnrealEd.ProjectPackagingSettings", "PatchSeekOptMaxAdjacentOrderDiff", out PatchSeekOptMaxAdjacentOrderDiff))
-						{
-							AdditionalArgs += String.Format(" -PatchSeekOptMaxAdjacentOrderDiff={0}", PatchSeekOptMaxAdjacentOrderDiff);
-						}
+					string AdditionalArgs = PatchOptions;
+					if (bShouldGeneratePatch && ShouldSkipGeneratingPatch(PlatformGameConfig, PakParams.PakName))
+					{
+						AdditionalArgs = " -SkipPatch";
 					}
 
 					bool bPakFallbackOrderForNonUassetFiles = false;
@@ -2281,7 +2299,8 @@ public partial class Project : CommandUtils
 						AdditionalArgs += " -fallbackOrderForNonUassetFiles";
 					}
 
-					if (PlatformEngineConfig != null)
+					ConfigHierarchy PlatformEngineConfig;
+                    if (Params.EngineConfigs.TryGetValue(SC.StageTargetPlatform.PlatformType, out PlatformEngineConfig))					
 					{
 						// if the runtime will want to reduce memory usage, we have to disable the pak index freezing
 						bool bUnloadPakEntries = false;
@@ -2297,7 +2316,7 @@ public partial class Project : CommandUtils
 					// pass the targetplatform so the index may be able to be frozen
 					AdditionalArgs += " -platform=" + ConfigHierarchy.GetIniPlatformName(SC.StageTargetPlatform.IniPlatformType);
 
-					Dictionary<string, string> UnrealPakResponseFile;
+					Dictionary<string, string> UnrealPakResponseFile = PakParams.UnrealPakResponseFile;
 					if (ShouldCreateIoStoreContainerFiles(Params, SC))
 					{
 						UnrealPakResponseFile = new Dictionary<string, string>();
@@ -2316,13 +2335,47 @@ public partial class Project : CommandUtils
 								UnrealPakResponseFile.Add(Entry.Key, Entry.Value);
 							}
 						}
-						IoStoreCommands.Add(GetIoStoreCommand(IoStoreResponseFile, OutputLocation, PrimaryOrderFile, PakParams.bCompressed, CryptoSettings, PakParams.EncryptionKeyGuid, SecondaryOrderFile));
-					}
-					else
-					{
-						UnrealPakResponseFile = PakParams.UnrealPakResponseFile;
-					}
-					Commands.Add(GetUnrealPakArguments(Params.RawProjectPath, UnrealPakResponseFile, OutputLocation, PrimaryOrderFile, SC.StageTargetPlatform.GetPlatformPakCommandLine(Params, SC) + AdditionalArgs + BulkOption + CompressionFormats + " " + Params.AdditionalPakOptions, PakParams.bCompressed, CryptoSettings, CryptoKeysCacheFilename, PatchSourceContentPath, PakParams.EncryptionKeyGuid, SecondaryOrderFile));
+
+						string ContainerName = PakParams.PakName + "-" + SC.FinalCookPlatform;
+						StagedFileReference ContainerRelativeLocation;
+						if (Params.HasDLCName)
+						{
+							ContainerRelativeLocation = StagedFileReference.Combine(Params.DLCFile.Directory.MakeRelativeTo(SC.ProjectRoot), "Content", "Paks", SC.FinalCookPlatform, Params.DLCFile.GetFileNameWithoutExtension() + ContainerName);
+						}
+						else
+						{
+							ContainerRelativeLocation = StagedFileReference.Combine("Content", "Paks", ContainerName);
+						}
+						if (SC.StageTargetPlatform.DeployLowerCaseFilenames())
+						{
+							ContainerRelativeLocation = ContainerRelativeLocation.ToLowerInvariant();
+						}
+						ContainerRelativeLocation = SC.StageTargetPlatform.Remap(ContainerRelativeLocation);
+
+						IoStoreCommands.Add(GetIoStoreCommandArguments(
+							IoStoreResponseFile,
+							ContainerRelativeLocation,
+							PrimaryOrderFile,
+							AdditionalArgs,
+							PakParams.bCompressed,
+							CryptoSettings,
+							PakParams.EncryptionKeyGuid,
+							SecondaryOrderFile));
+					}
+
+					Commands.Add(GetUnrealPakArguments(
+						Params.RawProjectPath,
+						UnrealPakResponseFile,
+						OutputLocation,
+						PrimaryOrderFile,
+						SC.StageTargetPlatform.GetPlatformPakCommandLine(Params, SC) + AdditionalArgs + BulkOption + CompressionFormats + " " + Params.AdditionalPakOptions,
+						PakParams.bCompressed,
+						CryptoSettings,
+						CryptoKeysCacheFilename,
+						PatchSourceContentPath,
+						PakParams.EncryptionKeyGuid,
+						SecondaryOrderFile));
+
 					LogNames.Add(OutputLocation.GetFileNameWithoutExtension());
 				}
 			}
@@ -2347,7 +2400,7 @@ public partial class Project : CommandUtils
 				}
 			}
 
-			FileReference PackageOpenOrderFileLocation = null;
+			FileReference GameOpenOrderFileLocation = null;
 			FileReference CookerOpenOrderFileLocation = null;
 			// search CookPlaform (e.g. IOSClient and then regular platform (e.g. IOS).
 			string[] OrderLocations = new string[] { SC.FinalCookPlatform, SC.StageTargetPlatform.GetTargetPlatformDescriptor().Type.ToString() };
@@ -2355,10 +2408,10 @@ public partial class Project : CommandUtils
 			foreach (string OrderLocation in OrderLocations.Reverse())
 			{
 				DirectoryReference IoStoreOrderFileLocationBase = DirectoryReference.Combine(SC.ProjectRoot, "Build", OrderLocation, "FileOpenOrder");
-				FileReference FileLocation = FileReference.Combine(IoStoreOrderFileLocationBase, "PackageOpenOrder.log");
+				FileReference FileLocation = FileReference.Combine(IoStoreOrderFileLocationBase, "GameOpenOrder.log");
 				if (FileExists_NoExceptions(FileLocation.FullName))
 				{
-					PackageOpenOrderFileLocation = FileLocation;
+					GameOpenOrderFileLocation = FileLocation;
 				}
 				FileLocation = FileReference.Combine(IoStoreOrderFileLocationBase, "CookerOpenOrder.log");
 				if (FileExists_NoExceptions(FileLocation.FullName))
@@ -2366,7 +2419,15 @@ public partial class Project : CommandUtils
 					CookerOpenOrderFileLocation = FileLocation;
 				}
 			}
-			RunIoStore(Params, SC, IoStoreCommandsFileName, PackageOpenOrderFileLocation, CookerOpenOrderFileLocation);
+
+			string AdditionalArgs = " " 
+				+ SC.StageTargetPlatform.GetPlatformIoStoreCommandLine(Params, SC)
+				+ BulkOption
+				+ CompressionFormats
+				+ " " + AdditionalCompressionOptionsOnCommandLine
+				+ " " + Params.AdditionalPakOptions;
+
+			RunIoStore(Params, SC, IoStoreCommandsFileName, GameOpenOrderFileLocation, CookerOpenOrderFileLocation, AdditionalArgs);
 		}
 
 		// Do any additional processing on the command output
@@ -2510,18 +2571,33 @@ public partial class Project : CommandUtils
 		}
 	}
 
-	private static void RunIoStore(ProjectParams Params, DeploymentContext SC, string CommandsFileName, FileReference PackageOrderFileLocation, FileReference CookerOpenOrderFileLocation)
+	private static void RunIoStore(ProjectParams Params, DeploymentContext SC, string CommandsFileName, FileReference GameOpenOrderFileLocation, FileReference CookerOpenOrderFileLocation, string AdditionalArgs)
 	{
-		FileReference OutputLocation = FileReference.Combine(SC.RuntimeRootDir, SC.RelativeProjectRootForStage.Name);
+		DirectoryReference OutputLocation = SC.StageTargetPlatform.GetProjectRootForStage(SC.RuntimeRootDir, SC.RelativeProjectRootForStage);
+
 		string CommandletParams = String.Format("-OutputDirectory={0} -CookedDirectory={1} -Commands={2}", MakePathSafeToUseWithCommandLine(OutputLocation.FullName), MakePathSafeToUseWithCommandLine(SC.PlatformCookDir.ToString()), MakePathSafeToUseWithCommandLine(CommandsFileName));
-		if (PackageOrderFileLocation != null)
+		if (GameOpenOrderFileLocation != null)
 		{
-			CommandletParams += String.Format(" -PackageOrder={0}", MakePathSafeToUseWithCommandLine(PackageOrderFileLocation.FullName));
+			CommandletParams += String.Format(" -GameOrder={0}", MakePathSafeToUseWithCommandLine(GameOpenOrderFileLocation.FullName));
 		}
 		if (CookerOpenOrderFileLocation != null)
 		{
 			CommandletParams += String.Format(" -CookerOrder={0}", MakePathSafeToUseWithCommandLine(CookerOpenOrderFileLocation.FullName));
 		}
+		if (!string.IsNullOrWhiteSpace(AdditionalArgs))
+		{
+			CommandletParams += AdditionalArgs;
+		}
+
+		if (Params.HasCreateReleaseVersion)
+		{
+			CommandletParams += String.Format(" -CreateReleaseVersionDirectory={0}", MakePathSafeToUseWithCommandLine(Params.GetCreateReleaseVersionPath(SC, Params.Client)));
+		}
+		if (Params.HasBasedOnReleaseVersion)
+		{
+			CommandletParams += String.Format(" -BasedOnReleaseVersionDirectory={0}", MakePathSafeToUseWithCommandLine(Params.GetBasedOnReleaseVersionPath(SC, Params.Client)));
+		}
+
 		LogInformation("Running IoStore commandlet with arguments: {0}", CommandletParams);
 		RunCommandlet(SC.RawProjectPath, Params.UE4Exe, "IoStore", CommandletParams);
 	}
diff --git a/Engine/Source/Runtime/Core/Private/Containers/String.cpp b/Engine/Source/Runtime/Core/Private/Containers/String.cpp
--- a/Engine/Source/Runtime/Core/Private/Containers/String.cpp
+++ b/Engine/Source/Runtime/Core/Private/Containers/String.cpp
@@ -125,7 +125,16 @@ FString::FString(const FStringView& Other)
 	{
 		Reserve(OtherLen);
 		Append(Other.GetData(), OtherLen);
-		AppendChar(TEXT('\0'));
+	}
+}
+
+FString::FString(const FStringView& Other, int32 ExtraSlack)
+{
+	Reserve(Other.Len() + ExtraSlack);
+
+	if (!Other.IsEmpty())
+	{
+		Append(Other.GetData(), Other.Len());
 	}
 }
 
diff --git a/Engine/Source/Runtime/Core/Private/GenericPlatform/GenericPlatformIoDispatcher.cpp b/Engine/Source/Runtime/Core/Private/GenericPlatform/GenericPlatformIoDispatcher.cpp
--- a/Engine/Source/Runtime/Core/Private/GenericPlatform/GenericPlatformIoDispatcher.cpp
+++ b/Engine/Source/Runtime/Core/Private/GenericPlatform/GenericPlatformIoDispatcher.cpp
@@ -7,6 +7,7 @@
 #include "HAL/PlatformFilemanager.h"
 #include "ProfilingDebugging/CountersTrace.h"
 #include "HAL/RunnableThread.h"
+#include "HAL/IConsoleManager.h"
 #include "Misc/ScopeLock.h"
 
 //PRAGMA_DISABLE_OPTIMIZATION
@@ -15,7 +16,14 @@ TRACE_DECLARE_INT_COUNTER(IoDispatcherSequentialReads, TEXT("IoDispatcher/Sequen
 TRACE_DECLARE_INT_COUNTER(IoDispatcherForwardSeeks, TEXT("IoDispatcher/ForwardSeeks"));
 TRACE_DECLARE_INT_COUNTER(IoDispatcherBackwardSeeks, TEXT("IoDispatcher/BackwardSeeks"));
 TRACE_DECLARE_MEMORY_COUNTER(IoDispatcherTotalSeekDistance, TEXT("IoDispatcher/TotalSeekDistance"));
-TRACE_DECLARE_INT_COUNTER(IoDispatcherPendingBlocksCount, TEXT("IoDispatcher/PendingBlocksCount"));
+TRACE_DECLARE_INT_COUNTER(IoDispatcherCacheHits, TEXT("IoDispatcher/CacheHits"));
+
+int32 GIoDispatcherCacheSizeMB = 0;
+static FAutoConsoleVariableRef CVar_IoDispatcherCacheSizeMB(
+	TEXT("s.IoDispatcherCacheSizeMB"),
+	GIoDispatcherCacheSizeMB,
+	TEXT("IoDispatcher cache memory size (in megabytes).")
+);
 
 FGenericIoDispatcherEventQueue::FGenericIoDispatcherEventQueue()
 	: Event(FPlatformProcess::GetSynchEventFromPool())
@@ -37,17 +45,46 @@ void FGenericIoDispatcherEventQueue::Wait()
 	Event->Wait();
 }
 
-FGenericFileIoStoreImpl::FGenericFileIoStoreImpl(FGenericIoDispatcherEventQueue& InEventQueue)
+FGenericFileIoStoreImpl::FGenericFileIoStoreImpl(FGenericIoDispatcherEventQueue& InEventQueue, uint64 InReadBufferSize)
 	: EventQueue(InEventQueue)
-	, PendingBlockEvent(FPlatformProcess::GetSynchEventFromPool())
+	, ReadBufferSize(InReadBufferSize)
 {
-	Thread = FRunnableThread::Create(this, TEXT("IoService"), 0, TPri_AboveNormal);
+	uint64 CacheMemorySize = GIoDispatcherCacheSizeMB > 0 ? uint64(GIoDispatcherCacheSizeMB) << 20 : 0;
+	uint64 CacheBlockCount = CacheMemorySize / InReadBufferSize;
+	if (CacheBlockCount)
+	{
+		CacheMemorySize = CacheBlockCount * InReadBufferSize;
+		CacheMemory = reinterpret_cast<uint8*>(FMemory::Malloc(CacheMemorySize));
+		FCachedBlock* Prev = &CacheLruHead;
+		for (uint64 CacheBlockIndex = 0; CacheBlockIndex < CacheBlockCount; ++CacheBlockIndex)
+		{
+			FCachedBlock* CachedBlock = new FCachedBlock();
+			CachedBlock->Key = uint64(-1);
+			CachedBlock->Buffer = CacheMemory + CacheBlockIndex * InReadBufferSize;
+			Prev->LruNext = CachedBlock;
+			CachedBlock->LruPrev = Prev;
+			Prev = CachedBlock;
+		}
+		Prev->LruNext = &CacheLruTail;
+		CacheLruTail.LruPrev = Prev;
+	}
+	else
+	{
+		CacheLruHead.LruNext = &CacheLruTail;
+		CacheLruTail.LruPrev = &CacheLruHead;
+	}
 }
 
 FGenericFileIoStoreImpl::~FGenericFileIoStoreImpl()
 {
-	delete Thread;
-	FPlatformProcess::ReturnSynchEventToPool(PendingBlockEvent);
+	FCachedBlock* CachedBlock = CacheLruHead.LruNext;
+	while (CachedBlock != &CacheLruTail)
+	{
+		FCachedBlock* Next = CachedBlock->LruNext;
+		delete CachedBlock;
+		CachedBlock = Next;
+	}
+	FMemory::Free(CacheMemory);
 }
 
 bool FGenericFileIoStoreImpl::OpenContainer(const TCHAR* ContainerFilePath, uint64& ContainerFileHandle, uint64& ContainerFileSize)
@@ -63,150 +100,90 @@ bool FGenericFileIoStoreImpl::OpenContainer(const TCHAR* ContainerFilePath, uint
 	return true;
 }
 
-void FGenericFileIoStoreImpl::BeginReadsForRequest(FFileIoStoreResolvedRequest& ResolvedRequest)
+void FGenericFileIoStoreImpl::ReadBlockFromFile(uint8* Target, uint64 InFileHandle, FFileIoStoreRawBlock* Block)
 {
-	if (!ResolvedRequest.Request->IoBuffer.DataSize())
+	FCachedBlock* CachedBlock = nullptr;
+	bool bIsCacheableBlock = CacheMemory != nullptr && !Block->DirectToRequest && Block->Size == ReadBufferSize;
+	if (bIsCacheableBlock)
 	{
-		ResolvedRequest.Request->IoBuffer = FIoBuffer(ResolvedRequest.ResolvedSize);
+		CachedBlock = CachedBlocks.FindRef(Block->Key.Hash);
 	}
-}
 
-void FGenericFileIoStoreImpl::ReadBlockFromFile(FFileIoStoreReadBlock* Block)
-{
+	if (CachedBlock)
+	{
+		FMemory::Memcpy(Target, CachedBlock->Buffer, ReadBufferSize);
+		CachedBlock->LruPrev->LruNext = CachedBlock->LruNext;
+		CachedBlock->LruNext->LruPrev = CachedBlock->LruPrev;
+		
+		CachedBlock->LruPrev = &CacheLruHead;
+		CachedBlock->LruNext = CacheLruHead.LruNext;
+		
+		CachedBlock->LruPrev->LruNext = CachedBlock;
+		CachedBlock->LruNext->LruPrev = CachedBlock;
+
+		TRACE_COUNTER_INCREMENT(IoDispatcherCacheHits);
+	}
+	else
 	{
-		FScopeLock Lock(&PendingBlocksCritical);
-		if (!PendingBlocksHead)
+		IFileHandle* FileHandle = reinterpret_cast<IFileHandle*>(static_cast<UPTRINT>(InFileHandle));
+		if (FileHandle->Tell() != Block->Offset)
 		{
-			PendingBlocksHead = PendingBlocksTail = Block;
+			if (uint64(FileHandle->Tell()) > Block->Offset)
+			{
+				TRACE_COUNTER_INCREMENT(IoDispatcherBackwardSeeks);
+			}
+			else
+			{
+				TRACE_COUNTER_INCREMENT(IoDispatcherForwardSeeks);
+			}
+			TRACE_COUNTER_ADD(IoDispatcherTotalSeekDistance, FMath::Abs(FileHandle->Tell() - int64(Block->Offset)));
+			FileHandle->Seek(Block->Offset);
 		}
 		else
 		{
-			PendingBlocksTail->Next = Block;
-			PendingBlocksTail = Block;
+			TRACE_COUNTER_INCREMENT(IoDispatcherSequentialReads);
 		}
-		Block->Next = nullptr;
-		TRACE_COUNTER_INCREMENT(IoDispatcherPendingBlocksCount);
-	}
-	PendingBlockEvent->Trigger();
-}
+		FileHandle->Read(Target, Block->Size);
 
-void FGenericFileIoStoreImpl::EndReadsForRequest()
-{
+		if (bIsCacheableBlock)
+		{
+			FCachedBlock* BlockToReplace = CacheLruTail.LruPrev;
+			CachedBlocks.Remove(BlockToReplace->Key);
+			BlockToReplace->Key = Block->Key.Hash;
+			CachedBlocks.Add(BlockToReplace->Key, BlockToReplace);
+			
+			BlockToReplace->LruPrev->LruNext = BlockToReplace->LruNext;
+			BlockToReplace->LruNext->LruPrev = BlockToReplace->LruPrev;
 
-}
+			BlockToReplace->LruPrev = &CacheLruHead;
+			BlockToReplace->LruNext = CacheLruHead.LruNext;
 
-FFileIoStoreReadBlock* FGenericFileIoStoreImpl::GetNextCompletedBlock()
-{
-	FScopeLock _(&CompletedBlocksCritical);
-	FFileIoStoreReadBlock* CompletedBlock = CompletedBlocksHead;
-	if (!CompletedBlocksHead)
-	{
-		return nullptr;
+			BlockToReplace->LruPrev->LruNext = BlockToReplace;
+			BlockToReplace->LruNext->LruPrev = BlockToReplace;
+			
+			FMemory::Memcpy(BlockToReplace->Buffer, Target, ReadBufferSize);
+		}
 	}
-	CompletedBlocksHead = CompletedBlocksHead->Next;
-	if (!CompletedBlocksHead)
 	{
-		CompletedBlocksTail = nullptr;
+		FScopeLock _(&CompletedBlocksCritical);
+		if (!CompletedBlocksHead)
+		{
+			CompletedBlocksHead = CompletedBlocksTail = Block;
+		}
+		else
+		{
+			CompletedBlocksTail->Next = Block;
+			CompletedBlocksTail = Block;
+		}
+		Block->Next = nullptr;
 	}
-	return CompletedBlock;
-}
-
-bool FGenericFileIoStoreImpl::Init()
-{
-	return true;
-}
-
-void FGenericFileIoStoreImpl::Stop()
-{
-	bStopRequested = true;
-	PendingBlockEvent->Trigger();
+	EventQueue.Notify();
 }
 
-uint32 FGenericFileIoStoreImpl::Run()
+FFileIoStoreRawBlock* FGenericFileIoStoreImpl::GetCompletedBlocks()
 {
-	FFileIoStoreReadBlock* ScheduledBlocksHead = nullptr;
-	FFileIoStoreReadBlock* ScheduledBlocksTail = nullptr;
-
-	while (!bStopRequested)
-	{
-		while (!bStopRequested)
-		{
-			{
-				FScopeLock PendingBlocksLock(&PendingBlocksCritical);
-				if (PendingBlocksHead)
-				{
-					if (!ScheduledBlocksTail)
-					{
-						ScheduledBlocksHead = PendingBlocksHead;
-						ScheduledBlocksTail = PendingBlocksTail;
-					}
-					else
-					{
-						ScheduledBlocksTail->Next = PendingBlocksHead;
-						ScheduledBlocksTail = PendingBlocksTail;
-					}
-					PendingBlocksHead = PendingBlocksTail = nullptr;
-				}
-			}
-
-			if (!ScheduledBlocksHead)
-			{
-				break;
-			}
-
-			FFileIoStoreReadBlock* BlockToRead = ScheduledBlocksHead;
-			ScheduledBlocksHead = ScheduledBlocksHead->Next;
-			if (!ScheduledBlocksHead)
-			{
-				ScheduledBlocksTail = nullptr;
-			}
-			
-			if (!BlockToRead->Buffer.DataSize())
-			{
-				BlockToRead->Buffer = FIoBuffer(BlockToRead->Size);
-			}
-
-			/*FString ScopeNameString;
-			ScopeNameString.Appendf(TEXT("ReadBlockFromFile: %lld (%lld)"), BlockToRead->Key.BlockIndex, BlockToRead->Size);
-			TRACE_CPUPROFILER_EVENT_SCOPE_TEXT(*ScopeNameString);
-			*/
-			TRACE_CPUPROFILER_EVENT_SCOPE(ReadBlockFromFile);
-			IFileHandle* FileHandle = reinterpret_cast<IFileHandle*>(static_cast<UPTRINT>(BlockToRead->Key.FileHandle));
-			if (FileHandle->Tell() != BlockToRead->Offset)
-			{
-				if (uint64(FileHandle->Tell()) > BlockToRead->Offset)
-				{
-					TRACE_COUNTER_INCREMENT(IoDispatcherBackwardSeeks);
-				}
-				else
-				{
-					TRACE_COUNTER_INCREMENT(IoDispatcherForwardSeeks);
-				}
-				TRACE_COUNTER_ADD(IoDispatcherTotalSeekDistance, FMath::Abs(FileHandle->Tell() - int64(BlockToRead->Offset)));
-				FileHandle->Seek(BlockToRead->Offset);
-			}
-			else
-			{
-				TRACE_COUNTER_INCREMENT(IoDispatcherSequentialReads);
-			}
-			FileHandle->Read(BlockToRead->Buffer.Data(), BlockToRead->Size);
-			TRACE_COUNTER_DECREMENT(IoDispatcherPendingBlocksCount);
-			{
-				FScopeLock _(&CompletedBlocksCritical);
-				if (!CompletedBlocksHead)
-				{
-					CompletedBlocksHead = CompletedBlocksTail = BlockToRead;
-				}
-				else
-				{
-					CompletedBlocksTail->Next = BlockToRead;
-					CompletedBlocksTail = BlockToRead;
-				}
-				BlockToRead->Next = nullptr;
-			}
-			EventQueue.Notify();
-		}
-		PendingBlockEvent->Wait();
-	}
-	return 0;
+	FScopeLock _(&CompletedBlocksCritical);
+	FFileIoStoreRawBlock* Result = CompletedBlocksHead;
+	CompletedBlocksHead = CompletedBlocksTail = nullptr;
+	return Result;
 }
diff --git a/Engine/Source/Runtime/Core/Private/GenericPlatform/GenericPlatformIoDispatcher.h b/Engine/Source/Runtime/Core/Private/GenericPlatform/GenericPlatformIoDispatcher.h
--- a/Engine/Source/Runtime/Core/Private/GenericPlatform/GenericPlatformIoDispatcher.h
+++ b/Engine/Source/Runtime/Core/Private/GenericPlatform/GenericPlatformIoDispatcher.h
@@ -6,8 +6,9 @@
 #include "HAL/CriticalSection.h"
 #include "Templates/Tuple.h"
 #include "HAL/Runnable.h"
+#include "Containers/Map.h"
 
-struct FFileIoStoreReadBlock;
+struct FFileIoStoreRawBlock;
 struct FFileIoStoreResolvedRequest;
 class IAsyncReadRequest;
 
@@ -18,38 +19,44 @@ public:
 	~FGenericIoDispatcherEventQueue();
 	void Notify();
 	void Wait();
-	void Poll() {};
+	void WaitForIo()
+	{
+		Wait();
+	}
 
 private:
 	FEvent* Event;
 };
 
 class FGenericFileIoStoreImpl
-	: public FRunnable
 {
 public:
-	FGenericFileIoStoreImpl(FGenericIoDispatcherEventQueue& InEventQueue);
+	FGenericFileIoStoreImpl(FGenericIoDispatcherEventQueue& InEventQueue, uint64 InReadBufferSize);
 	~FGenericFileIoStoreImpl();
 	bool OpenContainer(const TCHAR* ContainerFilePath, uint64& ContainerFileHandle, uint64& ContainerFileSize);
-	void BeginReadsForRequest(FFileIoStoreResolvedRequest& ResolvedRequest);
-	void ReadBlockFromFile(FFileIoStoreReadBlock* Block);
-	void EndReadsForRequest();
-	FFileIoStoreReadBlock* GetNextCompletedBlock();
-	virtual bool Init() override;
-	virtual uint32 Run() override;
-	virtual void Stop() override;
+	void ReadBlockFromFile(uint8* Target, uint64 FileHandle, FFileIoStoreRawBlock* Block);
+	void FlushReads() {};
+	FFileIoStoreRawBlock* GetCompletedBlocks();
 
 private:
+	struct FCachedBlock
+	{
+		FCachedBlock* LruPrev = nullptr;
+		FCachedBlock* LruNext = nullptr;
+		uint64 Key;
+		uint8* Buffer;
+	};
+
 	FGenericIoDispatcherEventQueue& EventQueue;
+	const uint64 ReadBufferSize;
 
-	FCriticalSection PendingBlocksCritical;
-	FFileIoStoreReadBlock* PendingBlocksHead = nullptr;
-	FFileIoStoreReadBlock* PendingBlocksTail = nullptr;
 	FCriticalSection CompletedBlocksCritical;
-	FFileIoStoreReadBlock* CompletedBlocksHead = nullptr;
-	FFileIoStoreReadBlock* CompletedBlocksTail = nullptr;
-	FEvent* PendingBlockEvent;
-	FRunnableThread* Thread;
-	TAtomic<bool> bStopRequested{ false };
+	FFileIoStoreRawBlock* CompletedBlocksHead = nullptr;
+	FFileIoStoreRawBlock* CompletedBlocksTail = nullptr;
+	uint8* CacheMemory = nullptr;
+	TMap<uint64, FCachedBlock*> CachedBlocks;
+	FCachedBlock CacheLruHead;
+	FCachedBlock CacheLruTail;
+
 };
 
diff --git a/Engine/Source/Runtime/Core/Private/IO/IoBuffer.cpp b/Engine/Source/Runtime/Core/Private/IO/IoBuffer.cpp
--- a/Engine/Source/Runtime/Core/Private/IO/IoBuffer.cpp
+++ b/Engine/Source/Runtime/Core/Private/IO/IoBuffer.cpp
@@ -84,6 +84,22 @@ FIoBuffer::BufCore::MakeOwned()
 	SetIsOwned(true);
 }
 
+TIoStatusOr<uint8*>	FIoBuffer::BufCore::ReleaseMemory()
+{
+	if (IsMemoryOwned())
+	{
+		uint8* BufferPtr = Data();
+		SetDataAndSize(nullptr, 0);
+		ClearFlags();
+
+		return BufferPtr;
+	}
+	else
+	{
+		return FIoStatus(EIoErrorCode::InvalidParameter, TEXT("Cannot call release on a FIoBuffer unless it owns it's memory"));
+	}
+}
+
 //////////////////////////////////////////////////////////////////////////
 
 FIoBuffer::FIoBuffer()
@@ -121,3 +137,9 @@ FIoBuffer::MakeOwned() const
 {
 	CorePtr->MakeOwned();
 }
+ 
+TIoStatusOr<uint8*>
+FIoBuffer::Release()
+{
+	return CorePtr->ReleaseMemory();
+}
diff --git a/Engine/Source/Runtime/Core/Private/IO/IoDispatcher.cpp b/Engine/Source/Runtime/Core/Private/IO/IoDispatcher.cpp
--- a/Engine/Source/Runtime/Core/Private/IO/IoDispatcher.cpp
+++ b/Engine/Source/Runtime/Core/Private/IO/IoDispatcher.cpp
@@ -14,6 +14,7 @@
 #include "Serialization/LargeMemoryReader.h"
 #include "GenericPlatform/GenericPlatformChunkInstall.h"
 #include "HAL/Event.h"
+#include "Async/MappedFileHandle.h"
 
 DEFINE_LOG_CATEGORY(LogIoDispatcher);
 
@@ -132,8 +133,9 @@ class FIoDispatcherImpl
 	: public FRunnable
 {
 public:
-	FIoDispatcherImpl()
-		: FileIoStore(EventQueue)
+	FIoDispatcherImpl(bool bInIsMultithreaded)
+		: bIsMultithreaded(bInIsMultithreaded)
+		, FileIoStore(EventQueue, bInIsMultithreaded)
 	{
 		FCoreDelegates::GetMemoryTrimDelegate().AddLambda([this]()
 		{
@@ -168,9 +170,21 @@ public:
 	FIoRequestImpl* AllocRequest(FIoBatchImpl* Batch, const FIoChunkId& ChunkId, FIoReadOptions Options)
 	{
 		FIoRequestImpl* Request = AllocRequest(ChunkId, Options);
+
 		Request->Batch = Batch;
-		Request->BatchNextRequest = Batch->FirstRequest;
-		Batch->FirstRequest = Request;
+
+		if (Batch->HeadRequest == nullptr)
+		{
+			Batch->HeadRequest = Request;
+			Batch->TailRequest = Request;
+		}
+		else
+		{
+			Batch->TailRequest->BatchNextRequest = Request;
+			Batch->TailRequest = Request;
+		}
+
+		check(Batch->TailRequest->BatchNextRequest == nullptr);
 
 		return Request;
 	}
@@ -189,21 +203,24 @@ public:
 
 	void FreeBatch(FIoBatchImpl* Batch)
 	{
-		FIoRequestImpl* Request = Batch->FirstRequest;
-
-		while (Request)
+		if (Batch != nullptr)
 		{
-			FIoRequestImpl* Tmp = Request;
-			Request = Request->BatchNextRequest;
-			FreeRequest(Tmp);
-		}
+			FIoRequestImpl* Request = Batch->HeadRequest;
 
-		BatchAllocator.Destroy(Batch);
+			while (Request)
+			{
+				FIoRequestImpl* Tmp = Request;
+				Request = Request->BatchNextRequest;
+				FreeRequest(Tmp);
+			}
+
+			BatchAllocator.Destroy(Batch);
+		}
 	}
 
-	void ReadWithCallback(const FIoChunkId& Chunk, const FIoReadOptions& Options, TFunction<void(TIoStatusOr<FIoBuffer>)>&& Callback)
+	void ReadWithCallback(const FIoChunkId& ChunkId, const FIoReadOptions& Options, FIoReadCallback&& Callback)
 	{
-		FIoRequestImpl* Request = AllocRequest(Chunk, Options);
+		FIoRequestImpl* Request = AllocRequest(ChunkId, Options);
 		Request->Callback = MoveTemp(Callback);
 		{
 			FScopeLock _(&WaitingLock);
@@ -219,6 +236,15 @@ public:
 			}
 		}
 		EventQueue.Notify();
+		if (!bIsMultithreaded)
+		{
+			ProcessRequests();
+		}
+	}
+
+	TIoStatusOr<FIoMappedRegion> OpenMapped(const FIoChunkId& ChunkId, const FIoReadOptions& Options)
+	{
+		return FileIoStore.OpenMapped(ChunkId, Options);
 	}
 
 	FIoStatus Mount(const FIoStoreEnvironment& Environment)
@@ -239,7 +265,7 @@ public:
 	template<typename Func>
 	void IterateBatch(const FIoBatchImpl* Batch, Func&& InCallbackFunction)
 	{
-		FIoRequestImpl* Request = Batch->FirstRequest;
+		FIoRequestImpl* Request = Batch->HeadRequest;
 
 		while (Request)
 		{
@@ -255,13 +281,14 @@ public:
 			FScopeLock _(&WaitingLock);
 			if (!WaitingRequestsHead)
 			{
-				WaitingRequestsHead = Batch->FirstRequest;
+				WaitingRequestsHead = Batch->HeadRequest;
 			}
 			else
 			{
-				WaitingRequestsTail->NextRequest = Batch->FirstRequest;
+				WaitingRequestsTail->NextRequest = Batch->HeadRequest;
 			}
-			WaitingRequestsTail = Batch->FirstRequest;
+
+			WaitingRequestsTail = Batch->HeadRequest;
 			while (WaitingRequestsTail->BatchNextRequest)
 			{
 				WaitingRequestsTail->NextRequest = WaitingRequestsTail->BatchNextRequest;
@@ -270,6 +297,50 @@ public:
 			WaitingRequestsTail->NextRequest = nullptr;
 		}
 		EventQueue.Notify();
+		if (!bIsMultithreaded)
+		{
+			ProcessRequests();
+		}
+	}
+
+	FIoStatus SetupBatchForContiguousRead(FIoBatchImpl* Batch, void* InTargetVa, FIoReadCallback&& InCallback)
+	{
+		// Create the buffer
+		uint64 TotalSize = 0;
+		for (FIoRequestImpl* Request = Batch->HeadRequest; Request; Request = Request->BatchNextRequest)
+		{	
+			TotalSize += FMath::Min(GetSizeForChunk(Request->ChunkId).ConsumeValueOrDie(), Request->Options.GetSize());
+		}
+
+		// Set up memory buffers
+		if (InTargetVa != nullptr)
+		{
+			Batch->IoBuffer = FIoBuffer(FIoBuffer::Wrap, InTargetVa, TotalSize);
+		}
+		else
+		{
+			Batch->IoBuffer = FIoBuffer(FIoBuffer::AssumeOwnership, FMemory::Malloc(TotalSize), TotalSize);
+		}
+
+		uint8* DstBuffer = Batch->IoBuffer.Data();
+
+		// Now assign to each request
+		uint8* Ptr = DstBuffer;
+		for (FIoRequestImpl* Request = Batch->HeadRequest; Request; Request = Request->BatchNextRequest)
+		{
+			if (Request->Options.GetTargetVa() != nullptr)
+			{
+				return FIoStatus(EIoErrorCode::InvalidParameter, TEXT("A FIoBatch reading to a contiguous buffer cannot contain FIoRequests that have a TargetVa"));
+			}
+
+			Request->Options.SetTargetVa(Ptr);
+			Ptr += FMath::Min(GetSizeForChunk(Request->ChunkId).ConsumeValueOrDie(), Request->Options.GetSize());
+		}
+
+		// Set up callback
+		Batch->Callback = MoveTemp(InCallback);
+
+		return FIoStatus(EIoErrorCode::Ok);
 	}
 
 private:
@@ -277,18 +348,15 @@ private:
 
 	void ProcessCompletedBlocks()
 	{
-		EventQueue.Poll();
-		while (FileIoStore.ProcessCompletedBlock())
-		{
-			ProcessCompletedRequests();
-		}
+		FileIoStore.ProcessCompletedBlocks();
+		ProcessCompletedRequests();
 	}
 
 	void ProcessCompletedRequests()
 	{
+		//TRACE_CPUPROFILER_EVENT_SCOPE(ProcessCompletedRequests);
 		while (SubmittedRequestsHead && SubmittedRequestsHead->UnfinishedReadsCount == 0)
 		{
-			//TRACE_CPUPROFILER_EVENT_SCOPE(CompleteRequest);
 			FIoRequestImpl* NextRequest = SubmittedRequestsHead->NextRequest;
 			CompleteRequest(SubmittedRequestsHead);
 			SubmittedRequestsHead = NextRequest;
@@ -301,6 +369,8 @@ private:
 
 	void CompleteRequest(FIoRequestImpl* Request)
 	{
+		//TRACE_CPUPROFILER_EVENT_SCOPE(CompleteRequest);
+
 		if (!Request->Status.IsCompleted())
 		{
 			Request->Status = EIoErrorCode::Ok;
@@ -316,12 +386,35 @@ private:
 				Request->Callback(Request->Status);
 			}
 		}
-		if (!Request->Batch)
+
+		if (Request->Batch)
+		{
+			InvokeCallbackIfBatchCompleted(Request->Batch);
+		}
+		else
 		{
 			FreeRequest(Request);
 		}
 	}
 
+	void InvokeCallbackIfBatchCompleted(FIoBatchImpl* Batch)
+	{	
+		if (!Batch->Callback)
+		{
+			// No point checking if the batch does not have a callback
+			return;
+		}
+
+		// If there is no valid tail request then it should not have been possible to call this method
+		check(Batch->TailRequest != nullptr); 
+
+		// Since the requests will be processed in order we can just check the tail request
+		if (Batch->TailRequest->Status.IsCompleted())
+		{
+			Batch->Callback(Batch->IoBuffer);
+		}		
+	}
+
 	void ProcessIncomingRequests()
 	{
 		FIoRequestImpl* RequestsToSubmitHead = nullptr;
@@ -358,28 +451,39 @@ private:
 				RequestsToSubmitTail = nullptr;
 			}
 
-			//TRACE_CPUPROFILER_EVENT_SCOPE(ResolveRequest);
-
-			EIoStoreResolveResult Result = FileIoStore.Resolve(Request);
-			if (Result == IoStoreResolveResult_NotFound)
-			{
-				Request->Status = FIoStatus(EIoErrorCode::NotFound);
-			}
-			if (!SubmittedRequestsTail)
 			{
-				SubmittedRequestsHead = SubmittedRequestsTail = Request;
-			}
-			else
-			{
-				SubmittedRequestsTail->NextRequest = Request;
-				SubmittedRequestsTail = Request;
+				TRACE_CPUPROFILER_EVENT_SCOPE(ResolveRequest);
+
+				EIoStoreResolveResult Result = FileIoStore.Resolve(Request);
+				if (Result == IoStoreResolveResult_NotFound)
+				{
+					Request->Status = FIoStatus(EIoErrorCode::NotFound);
+				}
+				if (!SubmittedRequestsTail)
+				{
+					SubmittedRequestsHead = SubmittedRequestsTail = Request;
+				}
+				else
+				{
+					SubmittedRequestsTail->NextRequest = Request;
+					SubmittedRequestsTail = Request;
+				}
+				Request->NextRequest = nullptr;
 			}
-			Request->NextRequest = nullptr;
 
 			ProcessCompletedBlocks();
 		}
 	}
 
+	void ProcessRequests()
+	{
+		TRACE_CPUPROFILER_EVENT_SCOPE(ProcessEventQueue);
+
+		ProcessIncomingRequests();
+		ProcessCompletedBlocks();
+		ProcessCompletedRequests();
+	}
+
 	virtual bool Init()
 	{
 		return true;
@@ -387,14 +491,19 @@ private:
 
 	virtual uint32 Run()
 	{
+		FMemory::SetupTLSCachesOnCurrentThread();
 		while (!bStopRequested)
 		{
-			EventQueue.Wait();
-			
-			TRACE_CPUPROFILER_EVENT_SCOPE(ProcessEventQueue);
-			ProcessIncomingRequests();
-			ProcessCompletedBlocks();
-			ProcessCompletedRequests();
+			if (SubmittedRequestsHead)
+			{
+				TRACE_CPUPROFILER_EVENT_SCOPE(IoDispatcherWaitForIo);
+				EventQueue.WaitForIo();
+			}
+			else
+			{
+				EventQueue.Wait();
+			}
+			ProcessRequests();
 		}
 		return 0;
 	}
@@ -408,6 +517,7 @@ private:
 	using FRequestAllocator = TBlockAllocator<FIoRequestImpl, 4096>;
 	using FBatchAllocator = TBlockAllocator<FIoBatchImpl, 4096>;
 
+	bool bIsMultithreaded;
 	FIoDispatcherEventQueue EventQueue;
 
 	FFileIoStore FileIoStore;
@@ -423,7 +533,7 @@ private:
 };
 
 FIoDispatcher::FIoDispatcher()
-	: Impl(new FIoDispatcherImpl())
+	: Impl(new FIoDispatcherImpl(FGenericPlatformProcess::SupportsMultithreading()))
 {
 }
 
@@ -444,15 +554,22 @@ FIoDispatcher::NewBatch()
 }
 
 void
-FIoDispatcher::FreeBatch(FIoBatch Batch)
+FIoDispatcher::FreeBatch(FIoBatch& Batch)
 {
 	Impl->FreeBatch(Batch.Impl);
+	Batch.Impl = nullptr;
 }
 
 void
-FIoDispatcher::ReadWithCallback(const FIoChunkId& Chunk, const FIoReadOptions& Options, TFunction<void(TIoStatusOr<FIoBuffer>)>&& Callback)
+FIoDispatcher::ReadWithCallback(const FIoChunkId& ChunkId, const FIoReadOptions& Options, FIoReadCallback&& Callback)
+{
+	Impl->ReadWithCallback(ChunkId, Options, MoveTemp(Callback));
+}
+
+TIoStatusOr<FIoMappedRegion>
+FIoDispatcher::OpenMapped(const FIoChunkId& ChunkId, const FIoReadOptions& Options)
 {
-	Impl->ReadWithCallback(Chunk, Options, MoveTemp(Callback));
+	return Impl->OpenMapped(ChunkId, Options);
 }
 
 // Polling methods
@@ -500,13 +617,20 @@ FIoDispatcher::Get()
 	return *GIoDispatcher;
 }
 
+//////////////////////////////////////////////////////////////////////////
+
 FIoBatch::FIoBatch(FIoDispatcherImpl* InDispatcher, FIoBatchImpl* InImpl)
 	: Dispatcher(InDispatcher)
 	, Impl(InImpl)
-	, CompletionEvent()
 {
 }
 
+bool
+FIoBatch::IsValid() const
+{
+	return Impl != nullptr;
+}
+
 FIoRequest
 FIoBatch::Read(const FIoChunkId& ChunkId, FIoReadOptions Options)
 {
@@ -528,10 +652,23 @@ FIoBatch::Issue()
 	Dispatcher->IssueBatch(Impl);
 }
 
+FIoStatus
+FIoBatch::IssueWithCallback(FIoBatchReadOptions Options, FIoReadCallback&& Callback)
+{
+	FIoStatus Status = Dispatcher->SetupBatchForContiguousRead(Impl, Options.GetTargetVa(), MoveTemp(Callback));
+
+	if (Status.IsOk())
+	{
+		Dispatcher->IssueBatch(Impl);
+	}
+	
+	return Status;
+}
+
 void
 FIoBatch::Wait()
 {
-	for (FIoRequestImpl* Request = Impl->FirstRequest; Request; Request = Request->BatchNextRequest)
+	for (FIoRequestImpl* Request = Impl->HeadRequest; Request; Request = Request->BatchNextRequest)
 	{
 		while (!Request->Status.IsCompleted())
 		{
@@ -578,3 +715,4 @@ FIoRequest::GetResult() const
 		return Impl->Status;
 	}
 }
+
diff --git a/Engine/Source/Runtime/Core/Private/IO/IoDispatcherFileBackend.cpp b/Engine/Source/Runtime/Core/Private/IO/IoDispatcherFileBackend.cpp
--- a/Engine/Source/Runtime/Core/Private/IO/IoDispatcherFileBackend.cpp
+++ b/Engine/Source/Runtime/Core/Private/IO/IoDispatcherFileBackend.cpp
@@ -7,29 +7,57 @@
 #include "HAL/PlatformFilemanager.h"
 #include "GenericPlatform/GenericPlatformFile.h"
 #include "HAL/IConsoleManager.h"
+#include "Async/AsyncWork.h"
+#include "Async/MappedFileHandle.h"
+#include "HAL/RunnableThread.h"
+#include "Misc/ScopeLock.h"
 
 TRACE_DECLARE_MEMORY_COUNTER(IoDispatcherTotalBytesRead, TEXT("IoDispatcher/TotalBytesRead"));
 TRACE_DECLARE_MEMORY_COUNTER(IoDispatcherTotalBytesScattered, TEXT("IoDispatcher/TotalBytesScattered"));
-TRACE_DECLARE_INT_COUNTER(IoDispatcherCacheHitsCold, TEXT("IoDispatcher/CacheHitsCold"));
-TRACE_DECLARE_INT_COUNTER(IoDispatcherCacheHitsHot, TEXT("IoDispatcher/CacheHitsHot"));
-TRACE_DECLARE_INT_COUNTER(IoDispatcherCacheMisses, TEXT("IoDispatcher/CacheMisses"));
 
 //PRAGMA_DISABLE_OPTIMIZATION
 
-int32 GIoDispatcherBlockSizeKB = 256;
-/*static FAutoConsoleVariableRef CVar_IoDispatcherBlockSizeKB(
-	TEXT("s.IoDispatcherBlockSizeKB"),
-	GIoDispatcherBlockSizeKB,
-	TEXT("IoDispatcher read block size (in kilobytes).")
-);*/
-
-int32 GIoDispatcherCacheSizeMB = 32;
-static FAutoConsoleVariableRef CVar_IoDispatcherCacheSizeMB(
-	TEXT("s.IoDispatcherCacheSizeMB"),
-	GIoDispatcherCacheSizeMB,
-	TEXT("IoDispatcher cache size (in megabytes).")
+int32 GIoDispatcherBufferSizeKB = 256;
+static FAutoConsoleVariableRef CVar_IoDispatcherBufferSizeKB(
+	TEXT("s.IoDispatcherBufferSizeKB"),
+	GIoDispatcherBufferSizeKB,
+	TEXT("IoDispatcher read buffer size (in kilobytes).")
 );
 
+int32 GIoDispatcherBufferMemoryMB = 8;
+static FAutoConsoleVariableRef CVar_IoDispatcherBufferMemoryMB(
+	TEXT("s.IoDispatcherBufferMemoryMB"),
+	GIoDispatcherBufferMemoryMB,
+	TEXT("IoDispatcher buffer memory size (in megabytes).")
+);
+
+int32 GIoDispatcherDecompressionWorkerCount = 4;
+static FAutoConsoleVariableRef CVar_IoDispatcherDecompressionWorkerCount(
+	TEXT("s.IoDispatcherDecompressionWorkerCount"),
+	GIoDispatcherDecompressionWorkerCount,
+	TEXT("IoDispatcher decompression worker count.")
+);
+
+class FMappedFileProxy final : public IMappedFileHandle
+{
+public:
+	FMappedFileProxy(IMappedFileHandle* InSharedMappedFileHandle, uint64 InSize)
+		: IMappedFileHandle(InSize)
+		, SharedMappedFileHandle(InSharedMappedFileHandle)
+	{
+		check(InSharedMappedFileHandle != nullptr);
+	}
+
+	virtual ~FMappedFileProxy() { }
+
+	virtual IMappedFileRegion* MapRegion(int64 Offset = 0, int64 BytesToMap = MAX_int64, bool bPreloadHint = false) override
+	{
+		return SharedMappedFileHandle->MapRegion(Offset, BytesToMap, bPreloadHint);
+	}
+private:
+	IMappedFileHandle* SharedMappedFileHandle;
+};
+
 FFileIoStoreReader::FFileIoStoreReader(FFileIoStoreImpl& InPlatformImpl)
 	: PlatformImpl(InPlatformImpl)
 {
@@ -48,11 +76,13 @@ FIoStatus FFileIoStoreReader::Initialize(const FIoStoreEnvironment& Environment)
 	ContainerFilePath.Append(TEXT(".ucas"));
 	TocFilePath.Append(TEXT(".utoc"));
 
-	if (!PlatformImpl.OpenContainer(*ContainerFilePath, ContainerFileHandle, ContainerFileSize))
+	if (!PlatformImpl.OpenContainer(*ContainerFilePath, ContainerFile.FileHandle, ContainerFile.FileSize))
 	{
 		return FIoStatusBuilder(EIoErrorCode::FileOpenFailed) << TEXT("Failed to open IoStore container file '") << *ContainerFilePath << TEXT("'");
 	}
 
+	ContainerFile.FilePath = ContainerFilePath;
+
 	TUniquePtr<uint8[]> TocBuffer;
 	bool bTocReadOk = false;
 
@@ -91,20 +121,54 @@ FIoStatus FFileIoStoreReader::Initialize(const FIoStoreEnvironment& Environment)
 		return FIoStatusBuilder(EIoErrorCode::CorruptToc) << TEXT("TOC entry size mismatch while reading '") << *TocFilePath << TEXT("'");
 	}
 
-	const FIoStoreTocEntry* Entry = reinterpret_cast<const FIoStoreTocEntry*>(TocBuffer.Get() + sizeof(FIoStoreTocHeader));
-	uint32 EntryCount = Header->TocEntryCount;
+	const FIoStoreTocEntry* TocEntry = reinterpret_cast<const FIoStoreTocEntry*>(TocBuffer.Get() + sizeof(FIoStoreTocHeader));
+	uint32 TocEntryCount = Header->TocEntryCount;
 
-	Toc.Reserve(EntryCount);
-	while (EntryCount--)
+	uint32 CompressedBlockEntryCount = Header->CompressionBlockCount;
+	const FIoStoreCompressedBlockEntry* CompressedBlockEntry = CompressedBlockEntryCount > 0
+		? reinterpret_cast<const FIoStoreCompressedBlockEntry*>(TocEntry + TocEntryCount)
+		: nullptr;
+
+	const ANSICHAR* CompressionMethodNames = reinterpret_cast<const ANSICHAR*>(CompressedBlockEntry + CompressedBlockEntryCount);
+
+	Toc.Reserve(TocEntryCount);
+	while (TocEntryCount--)
 	{
-		if ((Entry->GetOffset() + Entry->GetLength()) > ContainerFileSize)
+		if (!CompressedBlockEntryCount && (TocEntry->GetOffset() + TocEntry->GetLength()) > ContainerFile.FileSize)
 		{
-			// TODO: add details
-			return FIoStatusBuilder(EIoErrorCode::CorruptToc) << TEXT("TOC entry out of container bounds while reading '") << *TocFilePath << TEXT("'");
+			return FIoStatusBuilder(EIoErrorCode::CorruptToc) << TEXT("TOC TocEntry out of container bounds while reading '") << *TocFilePath << TEXT("'");
 		}
 
-		Toc.Add(Entry->ChunkId, Entry->OffsetAndLength);
-		++Entry;
+		Toc.Add(TocEntry->ChunkId, TocEntry->OffsetAndLength);
+		++TocEntry;
+	}
+
+	if (CompressedBlockEntryCount > 0)
+	{
+		UE_LOG(LogIoDispatcher, Display, TEXT("Loading compressed toc: %s"), *TocFilePath);
+		ContainerFile.CompressionBlockSize = Header->CompressionBlockSize;
+		ContainerFile.CompressionBlocks.Reserve(Header->CompressionBlockCount);
+
+		while (CompressedBlockEntryCount--)
+		{
+			if (CompressedBlockEntry->OffsetAndLength.GetOffset() + CompressedBlockEntry->OffsetAndLength.GetLength() > ContainerFile.FileSize)
+			{
+				return FIoStatusBuilder(EIoErrorCode::CorruptToc) << TEXT("TOC TocCompressedBlockEntry out of container bounds while reading '") << *TocFilePath << TEXT("'");
+			}
+
+			ContainerFile.CompressionBlocks.Emplace(*CompressedBlockEntry);
+			++CompressedBlockEntry;
+		}
+	}
+	else
+	{
+		UE_LOG(LogIoDispatcher, Display, TEXT("Loading uncompressed toc: %s"), *TocFilePath);
+	}
+
+	for (uint32 CompressonNameIndex = 0; CompressonNameIndex < Header->CompressionNameCount; CompressonNameIndex++)
+	{
+		const ANSICHAR* CompressionMethodName = CompressionMethodNames + CompressonNameIndex * FIoStoreCompressionInfo::CompressionMethodNameLen;
+		ContainerFile.CompressionMethods.Add(FName(CompressionMethodName));
 	}
 
 	return FIoStatus::Ok;
@@ -129,36 +193,57 @@ TIoStatusOr<uint64> FFileIoStoreReader::GetSizeForChunk(const FIoChunkId& ChunkI
 	}
 }
 
-bool FFileIoStoreReader::Resolve(FFileIoStoreResolvedRequest& ResolvedRequest)
+const FIoOffsetAndLength* FFileIoStoreReader::Resolve(const FIoChunkId& ChunkId) const
 {
-	const FIoOffsetAndLength* OffsetAndLength = Toc.Find(ResolvedRequest.Request->ChunkId);
+	return Toc.Find(ChunkId);
+}
 
-	if (!OffsetAndLength)
+IMappedFileHandle* FFileIoStoreReader::GetMappedContainerFileHandle()
+{
+	if (!ContainerFile.MappedFileHandle)
 	{
-		return false;
+		IPlatformFile& Ipf = FPlatformFileManager::Get().GetPlatformFile();
+		ContainerFile.MappedFileHandle.Reset(Ipf.OpenMapped(*ContainerFile.FilePath));
 	}
 
-	ResolvedRequest.ResolvedFileHandle = ContainerFileHandle;
-	ResolvedRequest.ResolvedFileSize = ContainerFileSize;
-	uint64 RequestedOffset = ResolvedRequest.Request->Options.GetOffset();
-	ResolvedRequest.ResolvedOffset = OffsetAndLength->GetOffset() + RequestedOffset;
-	if (RequestedOffset > OffsetAndLength->GetLength())
+	check(ContainerFile.FileSize > 0);
+	return new FMappedFileProxy(ContainerFile.MappedFileHandle.Get(), ContainerFile.FileSize);
+}
+
+FFileIoStore::FFileIoStore(FIoDispatcherEventQueue& InEventQueue, bool bInIsMultithreaded)
+	: ReadBufferSize(GIoDispatcherBufferSizeKB > 0 ? uint64(GIoDispatcherBufferSizeKB) << 10 : 256 << 10)
+	, EventQueue(InEventQueue)
+	, bIsMultithreaded(bInIsMultithreaded)
+	, PlatformImpl(InEventQueue, ReadBufferSize)
+	, BufferAvailableEvent(FPlatformProcess::GetSynchEventFromPool())
+	, PendingBlockEvent(FPlatformProcess::GetSynchEventFromPool())
+{
+	uint64 BufferCount = uint64(GIoDispatcherBufferMemoryMB > 0 ? uint64(GIoDispatcherBufferMemoryMB) << 20 : 32ull << 20) / ReadBufferSize;
+	uint64 MemorySize = BufferCount * ReadBufferSize;
+	BufferMemory = reinterpret_cast<uint8*>(FMemory::Malloc(MemorySize));
+	for (uint64 BufferIndex = 0; BufferIndex < BufferCount; ++BufferIndex)
 	{
-		ResolvedRequest.ResolvedSize = 0;
+		FFileIoStoreBuffer* Buffer = new FFileIoStoreBuffer();
+		Buffer->Memory = BufferMemory + BufferIndex * ReadBufferSize;
+		Buffer->Next = FirstFreeBuffer;
+		FirstFreeBuffer = Buffer;
 	}
-	else
+
+	uint64 DecompressionContextCount = uint64(GIoDispatcherDecompressionWorkerCount > 0 ? GIoDispatcherDecompressionWorkerCount : 4);
+	for (uint64 ContextIndex = 0; ContextIndex < DecompressionContextCount; ++ContextIndex)
 	{
-		ResolvedRequest.ResolvedSize = FMath::Min(ResolvedRequest.Request->Options.GetSize(), OffsetAndLength->GetLength() - RequestedOffset);
+		FFileIoStoreCompressionContext* Context = new FFileIoStoreCompressionContext();
+		Context->Next = FirstFreeCompressionContext;
+		FirstFreeCompressionContext = Context;
 	}
-	return true;
+
+	Thread = FRunnableThread::Create(this, TEXT("IoService"), 0, TPri_AboveNormal);
 }
 
-FFileIoStore::FFileIoStore(FIoDispatcherEventQueue& InEventQueue)
-	: PlatformImpl(InEventQueue)
-	, CacheBlockSize(GIoDispatcherBlockSizeKB > 0 ? uint64(GIoDispatcherBlockSizeKB) << 10 : 256 << 10)
+FFileIoStore::~FFileIoStore()
 {
-	LruHead.LruNext = &LruTail;
-	LruTail.LruPrev = &LruHead;
+	delete Thread;
+	FPlatformProcess::ReturnSynchEventToPool(PendingBlockEvent);
 }
 
 FIoStatus FFileIoStore::Mount(const FIoStoreEnvironment& Environment)
@@ -178,10 +263,22 @@ EIoStoreResolveResult FFileIoStore::Resolve(FIoRequestImpl* Request)
 	FReadScopeLock _(IoStoreReadersLock);
 	FFileIoStoreResolvedRequest ResolvedRequest;
 	ResolvedRequest.Request = Request;
+	int32 ReaderIndex = 0;
 	for (FFileIoStoreReader* Reader : IoStoreReaders)
 	{
-		if (Reader->Resolve(ResolvedRequest))
+		if (const FIoOffsetAndLength* OffsetAndLength = Reader->Resolve(ResolvedRequest.Request->ChunkId))
 		{
+			uint64 RequestedOffset = ResolvedRequest.Request->Options.GetOffset();
+			ResolvedRequest.ResolvedOffset = OffsetAndLength->GetOffset() + RequestedOffset;
+			if (RequestedOffset > OffsetAndLength->GetLength())
+			{
+				ResolvedRequest.ResolvedSize = 0;
+			}
+			else
+			{
+				ResolvedRequest.ResolvedSize = FMath::Min(ResolvedRequest.Request->Options.GetSize(), OffsetAndLength->GetLength() - RequestedOffset);
+			}
+
 			Request->UnfinishedReadsCount = 0;
 			if (ResolvedRequest.ResolvedSize > 0)
 			{
@@ -189,25 +286,59 @@ EIoStoreResolveResult FFileIoStore::Resolve(FIoRequestImpl* Request)
 				{
 					ResolvedRequest.Request->IoBuffer = FIoBuffer(FIoBuffer::Wrap, TargetVa, ResolvedRequest.ResolvedSize);
 				}
-				PlatformImpl.BeginReadsForRequest(ResolvedRequest);
-				const uint32 RequestBeginBlockIndex = (uint32)(ResolvedRequest.ResolvedOffset / CacheBlockSize);
-				const uint32 RequestEndBlockIndex = (uint32)((ResolvedRequest.ResolvedOffset + ResolvedRequest.ResolvedSize - 1) / CacheBlockSize + 1);
-				const uint32 BlockCount = RequestEndBlockIndex - RequestBeginBlockIndex;
-				check(BlockCount > 0);
-				ReadBlockCached(RequestBeginBlockIndex, ResolvedRequest);
-				if (BlockCount > 1)
+				else
+				{
+					ResolvedRequest.Request->IoBuffer.SetSize(ResolvedRequest.ResolvedSize);
+				}
+
+				const FFileIoStoreContainerFile& ContainerFile = Reader->GetContainerFile();
+				const bool bIsCompressed = ContainerFile.CompressionBlockSize > 0;
+				const uint64 BlockSize = bIsCompressed ? ContainerFile.CompressionBlockSize : ReadBufferSize;
+				const uint64 RequestEndOffset = ResolvedRequest.ResolvedOffset + ResolvedRequest.ResolvedSize;
+				int32 RequestBeginBlockIndex = int32(ResolvedRequest.ResolvedOffset / BlockSize);
+				int32 RequestEndBlockIndex = int32((RequestEndOffset - 1) / BlockSize + 1);
+
+				if (bIsCompressed)
 				{
-					if (BlockCount > 2)
+					for (int32 BlockIndex = RequestBeginBlockIndex; BlockIndex < RequestEndBlockIndex; ++BlockIndex)
 					{
-						ReadBlocksUncached(RequestBeginBlockIndex + 1, BlockCount - 2, ResolvedRequest);
+						ReadBlockAndScatter(ReaderIndex, BlockIndex, ResolvedRequest);
+					}
+				}
+				else
+				{
+					int32 NumBlocks = RequestEndBlockIndex - RequestBeginBlockIndex;
+
+					const bool bFirstBlockIsPartial = ResolvedRequest.ResolvedOffset % BlockSize != 0;
+					if (bFirstBlockIsPartial)
+					{
+						ReadBlockAndScatter(ReaderIndex, RequestBeginBlockIndex, ResolvedRequest);
+						++RequestBeginBlockIndex;
+						--NumBlocks;
+					}
+
+					if (NumBlocks > 0)
+					{
+						const bool bLastBlockIsPartial = RequestEndOffset % BlockSize != 0;
+						const int32 NumMergedBlocks = bLastBlockIsPartial ? NumBlocks - 1  : NumBlocks;
+						if (NumMergedBlocks > 0)
+						{
+							const uint64 MergedReadOffset = RequestBeginBlockIndex * BlockSize;
+							const uint64 MergedReadSize = NumMergedBlocks * BlockSize;
+							ReadNoScatter(ReaderIndex, MergedReadOffset, MergedReadSize, ResolvedRequest);
+						}
+
+						if (bLastBlockIsPartial)
+						{
+							ReadBlockAndScatter(ReaderIndex, RequestEndBlockIndex - 1, ResolvedRequest);
+						}
 					}
-					ReadBlockCached(RequestEndBlockIndex - 1, ResolvedRequest);
 				}
-				PlatformImpl.EndReadsForRequest();
 			}
 
 			return IoStoreResolveResult_OK;
 		}
+		++ReaderIndex;
 	}
 
 	return IoStoreResolveResult_NotFound;
@@ -248,141 +379,574 @@ bool FFileIoStore::IsValidEnvironment(const FIoStoreEnvironment& Environment)
 	return FPlatformFileManager::Get().GetPlatformFile().FileExists(*TocFilePath);
 }
 
-bool FFileIoStore::ProcessCompletedBlock()
+FAutoConsoleTaskPriority CPrio_IoDispatcherTaskPriority(
+	TEXT("TaskGraph.TaskPriorities.IoDispatcherAsyncTasks"),
+	TEXT("Task and thread priority for IoDispatcher decompression."),
+	ENamedThreads::BackgroundThreadPriority, // if we have background priority task threads, then use them...
+	ENamedThreads::NormalTaskPriority, // .. at normal task priority
+	ENamedThreads::NormalTaskPriority // if we don't have background threads, then use normal priority threads at normal task priority instead
+);
+
+ENamedThreads::Type FFileIoStore::FDecompressAsyncTask::GetDesiredThread()
+{
+	return CPrio_IoDispatcherTaskPriority.Get();
+}
+
+void FFileIoStore::ScatterBlock(FFileIoStoreCompressedBlock* CompressedBlock, bool bIsAsync)
 {
-	const uint64 CacheMemorySize = GIoDispatcherCacheSizeMB > 0 ? uint64(GIoDispatcherCacheSizeMB) << 20 : 0;
-	FFileIoStoreReadBlock* CompletedBlock = PlatformImpl.GetNextCompletedBlock();
-	if (!CompletedBlock)
+	TRACE_CPUPROFILER_EVENT_SCOPE(IoDispatcherScatter);
+	
+	FFileIoStoreCompressionContext* CompressionContext = CompressedBlock->CompressionContext;
+	check(CompressionContext);
+	uint8* CompressedBuffer;
+	if (CompressedBlock->RawBlocksCount > 1)
+	{
+		check(CompressedBlock->CompressedDataBuffer);
+		CompressedBuffer = CompressedBlock->CompressedDataBuffer;
+	}
+	else
+	{
+		FFileIoStoreRawBlock* RawBlock = CompressedBlock->SingleRawBlock;
+		check(CompressedBlock->RawOffset >= RawBlock->Offset);
+		uint64 OffsetInBuffer = CompressedBlock->RawOffset - RawBlock->Offset;
+		CompressedBuffer = RawBlock->Buffer->Memory + OffsetInBuffer;
+	}
+	uint8* UncompressedBuffer;
+	if (CompressedBlock->CompressionMethod.IsNone())
 	{
-		return false;
+		UncompressedBuffer = CompressedBuffer;
 	}
-	check(!CompletedBlock->bIsReady);
-	CompletedBlock->bIsReady = true;
-	TRACE_COUNTER_ADD(IoDispatcherTotalBytesRead, CompletedBlock->Size);
-	for (FFileIoStoreReadBlockScatter& Scatter : CompletedBlock->ScatterList)
+	else
 	{
-		if (Scatter.DstOffset != MAX_uint64)
+		if (CompressionContext->UncompressedBufferSize < CompressedBlock->Size)
 		{
-			FMemory::Memcpy(Scatter.Request->IoBuffer.Data() + Scatter.DstOffset, CompletedBlock->Buffer.Data() + Scatter.SrcOffset, Scatter.Size);
+			FMemory::Free(CompressionContext->UncompressedBuffer);
+			CompressionContext->UncompressedBuffer = reinterpret_cast<uint8*>(FMemory::Malloc(CompressedBlock->Size));
+			CompressionContext->UncompressedBufferSize = CompressedBlock->Size;
 		}
-		TRACE_COUNTER_ADD(IoDispatcherTotalBytesScattered, Scatter.Size);
+		UncompressedBuffer = CompressionContext->UncompressedBuffer;
+
+		bool bFailed = !FCompression::UncompressMemory(CompressedBlock->CompressionMethod, UncompressedBuffer, int32(CompressedBlock->Size), CompressedBuffer, int32(CompressedBlock->RawSize));
+		check(!bFailed);
+	}
+
+	for (FFileIoStoreBlockScatter& Scatter : CompressedBlock->ScatterList)
+	{
+		FMemory::Memcpy(Scatter.Request->IoBuffer.Data() + Scatter.DstOffset, UncompressedBuffer + Scatter.SrcOffset, Scatter.Size);
 		check(Scatter.Request->UnfinishedReadsCount > 0);
 		--Scatter.Request->UnfinishedReadsCount;
 	}
-	//CompletedBlock->ScatterList.Empty();
 
-	if (CompletedBlock->LruPrev)
+	if (bIsAsync)
 	{
-		CurrentCacheUsage += CompletedBlock->Size;
+		FScopeLock Lock(&DecompressedBlocksCritical);
+		CompressedBlock->Next = FirstDecompressedBlock;
+		FirstDecompressedBlock = CompressedBlock;
+
+		EventQueue.Notify();
+	}
+}
 
-		FFileIoStoreReadBlock* EvictionCandidate = LruTail.LruPrev;
-		while (CurrentCacheUsage > CacheMemorySize && EvictionCandidate != &LruHead && EvictionCandidate->bIsReady)
+void FFileIoStore::AllocMemoryForRequest(FIoRequestImpl* Request)
+{
+	if (!Request->IoBuffer.Data())
+	{
+		TRACE_CPUPROFILER_EVENT_SCOPE(AllocMemoryForRequest);
+		Request->IoBuffer = FIoBuffer(Request->IoBuffer.DataSize());
+	}
+}
+
+void FFileIoStore::FinalizeCompressedBlock(FFileIoStoreCompressedBlock* CompressedBlock)
+{
+	if (CompressedBlock->RawBlocksCount > 1)
+	{
+		check(CompressedBlock->CompressedDataBuffer);
+		FMemory::Free(CompressedBlock->CompressedDataBuffer);
+	}
+	else
+	{
+		FFileIoStoreRawBlock* RawBlock = CompressedBlock->SingleRawBlock;
+		check(RawBlock->RefCount > 0);
+		if (--RawBlock->RefCount == 0)
 		{
-			FFileIoStoreReadBlock* NextEvictionCandidate = EvictionCandidate->LruPrev;
-			EvictionCandidate->LruNext->LruPrev = EvictionCandidate->LruPrev;
-			EvictionCandidate->LruPrev->LruNext = EvictionCandidate->LruNext;
-			CurrentCacheUsage -= EvictionCandidate->Size;
-			CachedBlocksMap.Remove(EvictionCandidate->Key);
-			delete EvictionCandidate;
-			EvictionCandidate = NextEvictionCandidate;
+			FreeBuffer(RawBlock->Buffer);
+			delete RawBlock;
 		}
 	}
-	else
+	check(CompressedBlock->CompressionContext);
+	FreeCompressionContext(CompressedBlock->CompressionContext);
+	for (FFileIoStoreBlockScatter& Scatter : CompressedBlock->ScatterList)
 	{
-		CachedBlocksMap.Remove(CompletedBlock->Key);
-		delete CompletedBlock;
+		TRACE_COUNTER_ADD(IoDispatcherTotalBytesScattered, Scatter.Size);
 	}
-	return true;
+	delete CompressedBlock;
 }
 
-void FFileIoStore::ReadBlockCached(uint32 BlockIndex, const FFileIoStoreResolvedRequest& ResolvedRequest)
+void FFileIoStore::ProcessCompletedBlocks()
 {
-	FFileIoStoreCacheBlockKey Key;
-	Key.FileHandle = ResolvedRequest.ResolvedFileHandle;
-	Key.BlockIndex = BlockIndex;
-	uint64 BlockOffset = uint64(BlockIndex) * uint64(CacheBlockSize);
-	FFileIoStoreReadBlock* CachedBlock = CachedBlocksMap.FindRef(Key);
-	if (!CachedBlock)
+	//TRACE_CPUPROFILER_EVENT_SCOPE(ProcessCompletedBlocks);
+	
+	FFileIoStoreRawBlock* CompletedBlock = PlatformImpl.GetCompletedBlocks();
+	while (CompletedBlock)
 	{
-		uint64 ReadSize = FMath::Min(ResolvedRequest.ResolvedFileSize, BlockOffset + CacheBlockSize) - BlockOffset;
+		TRACE_COUNTER_ADD(IoDispatcherTotalBytesRead, CompletedBlock->Size);
+		FFileIoStoreRawBlock* NextBlock = CompletedBlock->Next;
+
+		if (CompletedBlock->DirectToRequest)
+		{
+			check(CompletedBlock->DirectToRequest->UnfinishedReadsCount > 0);
+			--CompletedBlock->DirectToRequest->UnfinishedReadsCount;
+			TRACE_COUNTER_ADD(IoDispatcherTotalBytesScattered, CompletedBlock->Size);
+		}
+		else
+		{
+			RawBlocksMap.Remove(CompletedBlock->Key);
+
+			//TRACE_CPUPROFILER_EVENT_SCOPE(ProcessCompletedBlock);
+			for (FFileIoStoreCompressedBlock* CompressedBlock : CompletedBlock->CompressedBlocks)
+			{
+				if (CompressedBlock->RawBlocksCount > 1)
+				{
+					//TRACE_CPUPROFILER_EVENT_SCOPE(HandleComplexBlock);
+					if (!CompressedBlock->CompressedDataBuffer)
+					{
+						CompressedBlock->CompressedDataBuffer = reinterpret_cast<uint8*>(FMemory::Malloc(CompressedBlock->RawSize));
+					}
 
-		CachedBlock = new FFileIoStoreReadBlock();
-		CachedBlock->Key = Key;
-		CachedBlock->bIsReady = false;
-		CachedBlock->Offset = BlockOffset;
-		CachedBlock->Size = ReadSize;
-		CachedBlocksMap.Add(CachedBlock->Key, CachedBlock);
+					uint8* Src = CompletedBlock->Buffer->Memory;
+					uint8* Dst = CompressedBlock->CompressedDataBuffer;
+					uint64 CopySize = CompletedBlock->Size;
+					int64 CompletedBlockOffsetInBuffer = int64(CompletedBlock->Offset) - int64(CompressedBlock->RawOffset);
+					if (CompletedBlockOffsetInBuffer < 0)
+					{
+						Src -= CompletedBlockOffsetInBuffer;
+						CopySize += CompletedBlockOffsetInBuffer;
+					}
+					else
+					{
+						Dst += CompletedBlockOffsetInBuffer;
+					}
+					uint64 CompressedBlockRawEndOffset = CompressedBlock->RawOffset + CompressedBlock->RawSize;
+					uint64 CompletedBlockEndOffset = CompletedBlock->Offset + CompletedBlock->Size;
+					if (CompletedBlockEndOffset > CompressedBlockRawEndOffset)
+					{
+						CopySize -= CompletedBlockEndOffset - CompressedBlockRawEndOffset;
+					}
+					FMemory::Memcpy(Dst, Src, CopySize);
+					check(CompletedBlock->RefCount > 0);
+					--CompletedBlock->RefCount;
+				}
+				
+				check(CompressedBlock->UnfinishedRawBlocksCount > 0);
+				if (--CompressedBlock->UnfinishedRawBlocksCount == 0)
+				{
+					CompressedBlocksMap.Remove(CompressedBlock->Key);
+					if (!ReadyForDecompressionTail)
+					{
+						ReadyForDecompressionHead = ReadyForDecompressionTail = CompressedBlock;
+					}
+					else
+					{
+						ReadyForDecompressionTail->Next = CompressedBlock;
+						ReadyForDecompressionTail = CompressedBlock;
+					}
+					CompressedBlock->Next = nullptr;
+				}
+			}
+			if (CompletedBlock->RefCount == 0)
+			{
+				FreeBuffer(CompletedBlock->Buffer);
+				delete CompletedBlock;
+			}
+		}
+		CompletedBlock = NextBlock;
+	}
+	
+	FFileIoStoreCompressedBlock* BlockToReap;
+	{
+		FScopeLock Lock(&DecompressedBlocksCritical);
+		BlockToReap = FirstDecompressedBlock;
+		FirstDecompressedBlock = nullptr;
+	}
 
-		PlatformImpl.ReadBlockFromFile(CachedBlock);
-		TRACE_COUNTER_INCREMENT(IoDispatcherCacheMisses);
+	while (BlockToReap)
+	{
+		FFileIoStoreCompressedBlock* Next = BlockToReap->Next;
+		FinalizeCompressedBlock(BlockToReap);
+		BlockToReap = Next;
 	}
-	else
+
+	FFileIoStoreCompressedBlock* BlockToDecompress = ReadyForDecompressionHead;
+	while (BlockToDecompress)
 	{
-		if (CachedBlock->bIsReady)
+		FFileIoStoreCompressedBlock* Next = BlockToDecompress->Next;
+		BlockToDecompress->CompressionContext = AllocCompressionContext();
+		if (!BlockToDecompress->CompressionContext)
+		{
+			break;
+		}
+		for (FFileIoStoreBlockScatter& Scatter : BlockToDecompress->ScatterList)
+		{
+			AllocMemoryForRequest(Scatter.Request);
+		}
+		if (BlockToDecompress->CompressionMethod.IsNone())
 		{
-			TRACE_COUNTER_INCREMENT(IoDispatcherCacheHitsHot);
+			ScatterBlock(BlockToDecompress, false);
+			FinalizeCompressedBlock(BlockToDecompress);
 		}
 		else
 		{
-			TRACE_COUNTER_INCREMENT(IoDispatcherCacheHitsCold);
+			TGraphTask<FDecompressAsyncTask>::CreateTask().ConstructAndDispatchWhenReady(*this, BlockToDecompress);
 		}
+		BlockToDecompress = Next;
+	}
+	ReadyForDecompressionHead = BlockToDecompress;
+	if (!ReadyForDecompressionHead)
+	{
+		ReadyForDecompressionTail = nullptr;
 	}
+}
 
-	if (CachedBlock->LruPrev)
+TIoStatusOr<FIoMappedRegion> FFileIoStore::OpenMapped(const FIoChunkId& ChunkId, const FIoReadOptions& Options)
+{
+	if (!FPlatformProperties::SupportsMemoryMappedFiles())
 	{
-		check(CachedBlock->LruNext);
-		CachedBlock->LruPrev->LruNext = CachedBlock->LruNext;
-		CachedBlock->LruNext->LruPrev = CachedBlock->LruPrev;
+		return FIoStatus(EIoErrorCode::Unknown, TEXT("Platform does not support memory mapped files"));
 	}
 
-	CachedBlock->LruNext = LruHead.LruNext;
-	CachedBlock->LruPrev = &LruHead;
-	LruHead.LruNext->LruPrev = CachedBlock;
-	LruHead.LruNext = CachedBlock;
+	if (Options.GetTargetVa() != nullptr)
+	{
+		return FIoStatus(EIoErrorCode::InvalidParameter, TEXT("Invalid read options"));
+	}
 
+	IPlatformFile& Ipf = FPlatformFileManager::Get().GetPlatformFile();
+
+	for (FFileIoStoreReader* Reader : IoStoreReaders)
+	{
+		if (const FIoOffsetAndLength* OffsetAndLength = Reader->Resolve(ChunkId))
+		{
+			uint64 ResolvedOffset = OffsetAndLength->GetOffset();
+			uint64 ResolvedSize = FMath::Min(Options.GetSize(), OffsetAndLength->GetLength());
+			
+			const FFileIoStoreContainerFile& ContainerFile = Reader->GetContainerFile();
+			const bool bIsCompressed = ContainerFile.CompressionBlockSize > 0;
+
+			IMappedFileHandle* MappedFileHandle = Reader->GetMappedContainerFileHandle();
+			IMappedFileRegion* MappedFileRegion = nullptr;
+
+			if (bIsCompressed)
+			{
+				int32 BlockIndex = int32(ResolvedOffset / ContainerFile.CompressionBlockSize);
+				const FIoStoreCompressedBlockEntry& CompressionBlockEntry = ContainerFile.CompressionBlocks[BlockIndex];
+				const int64 BlockOffset = (int64)CompressionBlockEntry.OffsetAndLength.GetOffset();
+				check(BlockOffset > 0 && IsAligned(BlockOffset, FPlatformProperties::GetMemoryMappingAlignment()));
+
+				MappedFileRegion = MappedFileHandle->MapRegion(BlockOffset + Options.GetOffset(), ResolvedSize);
+				check(IsAligned(MappedFileRegion->GetMappedPtr(), FPlatformProperties::GetMemoryMappingAlignment()));
+			}
+			else
+			{
+				MappedFileRegion = MappedFileHandle->MapRegion(ResolvedOffset + Options.GetOffset(), ResolvedSize);
+			}
+
+			return FIoMappedRegion { MappedFileHandle, MappedFileRegion };
+		}
+	}
+
+	return FIoStatus::Invalid;
+}
+
+void FFileIoStore::ReadBlockAndScatter(uint32 ReaderIndex, uint32 BlockIndex, const FFileIoStoreResolvedRequest& ResolvedRequest)
+{
+	/*TStringBuilder<256> ScopeName;
+	ScopeName.Appendf(TEXT("ReadBlock %d"), BlockIndex);
+	TRACE_CPUPROFILER_EVENT_SCOPE_TEXT(*ScopeName);*/
+
+	const FFileIoStoreContainerFile& ContainerFile = IoStoreReaders[ReaderIndex]->GetContainerFile();
+
+	FFileIoStoreRawBlock* NewBlocksHead = nullptr;
+	FFileIoStoreRawBlock* NewBlocksTail = nullptr;
+
+	FFileIoStoreBlockKey CompressedBlockKey;
+	CompressedBlockKey.FileIndex = ReaderIndex;
+	CompressedBlockKey.BlockIndex = BlockIndex;
+	FFileIoStoreCompressedBlock* CompressedBlock = CompressedBlocksMap.FindRef(CompressedBlockKey);
+	if (!CompressedBlock)
+	{
+		CompressedBlock = new FFileIoStoreCompressedBlock();
+		CompressedBlock->Key = CompressedBlockKey;
+		CompressedBlocksMap.Add(CompressedBlockKey, CompressedBlock);
+
+		uint64 RawOffset;
+		uint64 RawSize;
+		if (ContainerFile.CompressionBlockSize > 0)
+		{
+			const FIoStoreCompressedBlockEntry& CompressionBlockEntry = ContainerFile.CompressionBlocks[BlockIndex];
+			CompressedBlock->Size = ContainerFile.CompressionBlockSize;
+			CompressedBlock->CompressionMethod = 
+				(CompressionBlockEntry.CompressionMethodIndex == FIoStoreCompressionInfo::InvalidCompressionIndex) ?
+				NAME_None :
+				ContainerFile.CompressionMethods[CompressionBlockEntry.CompressionMethodIndex];
+			RawOffset = CompressionBlockEntry.OffsetAndLength.GetOffset();
+			RawSize = CompressionBlockEntry.OffsetAndLength.GetLength();
+		}
+		else
+		{
+			CompressedBlock->Size = ReadBufferSize;
+			RawOffset = BlockIndex * CompressedBlock->Size;
+			RawSize = ReadBufferSize;
+		}
+		CompressedBlock->RawOffset = RawOffset;
+		CompressedBlock->RawSize = RawSize;
+		const uint32 RawBeginBlockIndex = uint32(RawOffset / ReadBufferSize);
+		const uint32 RawEndBlockIndex = uint32((RawOffset + RawSize - 1) / ReadBufferSize + 1);
+		const uint32 RawBlockCount = RawEndBlockIndex - RawBeginBlockIndex;
+		CompressedBlock->RawBlocksCount = RawBlockCount;
+		check(RawBlockCount > 0);
+		for (uint32 RawBlockIndex = RawBeginBlockIndex; RawBlockIndex < RawEndBlockIndex; ++RawBlockIndex)
+		{
+			FFileIoStoreBlockKey RawBlockKey;
+			RawBlockKey.BlockIndex = RawBlockIndex;
+			RawBlockKey.FileIndex = ReaderIndex;
+
+			FFileIoStoreRawBlock* RawBlock = RawBlocksMap.FindRef(RawBlockKey);
+			if (!RawBlock)
+			{
+				RawBlock = new FFileIoStoreRawBlock();
+				RawBlock->Key = RawBlockKey;
+				RawBlocksMap.Add(RawBlockKey, RawBlock);
+
+				RawBlock->Offset = RawBlockIndex * ReadBufferSize;
+				uint64 ReadSize = FMath::Min(ContainerFile.FileSize, RawBlock->Offset + ReadBufferSize) - RawBlock->Offset;
+				RawBlock->Size = ReadSize;
+				if (!NewBlocksTail)
+				{
+					NewBlocksHead = NewBlocksTail = RawBlock;
+				}
+				else
+				{
+					NewBlocksTail->Next = RawBlock;
+					NewBlocksTail = RawBlock;
+				}
+			}
+			if (RawBlockCount == 1)
+			{
+				CompressedBlock->SingleRawBlock = RawBlock;
+			}
+			RawBlock->CompressedBlocks.Add(CompressedBlock);
+			++RawBlock->RefCount;
+			++CompressedBlock->UnfinishedRawBlocksCount;
+		}
+	}
+
+	uint64 BlockOffset = BlockIndex * CompressedBlock->Size;
 	uint64 RequestStartOffsetInBlock = FMath::Max<int64>(0, int64(ResolvedRequest.ResolvedOffset) - BlockOffset);
-	uint64 RequestEndOffsetInBlock = FMath::Min<uint64>(CacheBlockSize, ResolvedRequest.ResolvedOffset + ResolvedRequest.ResolvedSize - BlockOffset);
+	uint64 RequestEndOffsetInBlock = FMath::Min<uint64>(CompressedBlock->Size, ResolvedRequest.ResolvedOffset + ResolvedRequest.ResolvedSize - BlockOffset);
 	uint64 BlockOffsetInRequest = FMath::Max<int64>(0, int64(BlockOffset) - ResolvedRequest.ResolvedOffset);
 	uint64 RequestSizeInBlock = RequestEndOffsetInBlock - RequestStartOffsetInBlock;
 	check(RequestSizeInBlock <= ResolvedRequest.Request->IoBuffer.DataSize());
-	check(RequestStartOffsetInBlock + RequestSizeInBlock <= CacheBlockSize);
-	check(ResolvedRequest.Request->IoBuffer.Data() + BlockOffsetInRequest + RequestSizeInBlock <= ResolvedRequest.Request->IoBuffer.Data() + ResolvedRequest.Request->IoBuffer.DataSize());
-	if (CachedBlock->bIsReady)
+	check(RequestStartOffsetInBlock + RequestSizeInBlock <= CompressedBlock->Size);
+	check(BlockOffsetInRequest + RequestSizeInBlock <= ResolvedRequest.Request->IoBuffer.DataSize());
+
+	++ResolvedRequest.Request->UnfinishedReadsCount;
+	FFileIoStoreBlockScatter& Scatter = CompressedBlock->ScatterList.AddDefaulted_GetRef();
+	Scatter.Request = ResolvedRequest.Request;
+	Scatter.DstOffset = BlockOffsetInRequest;
+	Scatter.SrcOffset = RequestStartOffsetInBlock;
+	Scatter.Size = RequestSizeInBlock;
+
+	if (NewBlocksHead)
 	{
-		FMemory::Memcpy(ResolvedRequest.Request->IoBuffer.Data() + BlockOffsetInRequest, CachedBlock->Buffer.Data() + RequestStartOffsetInBlock, RequestSizeInBlock);
-		TRACE_COUNTER_ADD(IoDispatcherTotalBytesScattered, RequestSizeInBlock);
+		{
+			FScopeLock Lock(&PendingBlocksCritical);
+			if (!PendingBlocksTail)
+			{
+				PendingBlocksHead = NewBlocksHead;
+			}
+			else
+			{
+				PendingBlocksTail->Next = NewBlocksHead;
+			}
+			PendingBlocksTail = NewBlocksTail;
+		}
+		PendingBlockEvent->Trigger();
+		if (!bIsMultithreaded)
+		{
+			ReadPendingBlocks();
+		}
 	}
-	else
+}
+
+void FFileIoStore::ReadNoScatter(uint32 ReaderIndex, uint64 Offset, uint64 Size, const FFileIoStoreResolvedRequest& ResolvedRequest)
+{
+	/*TStringBuilder<256> ScopeName;
+	ScopeName.Appendf(TEXT("ReadNoScatter %lld %lld"), Offset, Size);
+	TRACE_CPUPROFILER_EVENT_SCOPE_TEXT(*ScopeName);*/
+
+	const FFileIoStoreContainerFile& ContainerFile = IoStoreReaders[ReaderIndex]->GetContainerFile();
+
+	FFileIoStoreRawBlock* NewBlocksHead = nullptr;
+	FFileIoStoreRawBlock* NewBlocksTail = nullptr;
+
+	FFileIoStoreRawBlock* RawBlock = new FFileIoStoreRawBlock();
+	RawBlock->Key.BlockIndex = -1;
+	RawBlock->Key.FileIndex = ReaderIndex;
+	RawBlock->Offset = Offset;
+	uint64 ReadSize = FMath::Min(ContainerFile.FileSize, RawBlock->Offset + Size) - RawBlock->Offset;
+	RawBlock->Size = ReadSize;
+	RawBlock->DirectToRequest = ResolvedRequest.Request;
+	check(Offset >= ResolvedRequest.ResolvedOffset);
+	RawBlock->DirectToRequestOffset = Offset - ResolvedRequest.ResolvedOffset;
+	++ResolvedRequest.Request->UnfinishedReadsCount;
+
+	AllocMemoryForRequest(ResolvedRequest.Request);
+
 	{
-		++ResolvedRequest.Request->UnfinishedReadsCount;
-		FFileIoStoreReadBlockScatter& Scatter = CachedBlock->ScatterList.AddDefaulted_GetRef();
-		Scatter.Request = ResolvedRequest.Request;
-		Scatter.DstOffset = BlockOffsetInRequest;
-		Scatter.SrcOffset = RequestStartOffsetInBlock;
-		check(RequestSizeInBlock <= TNumericLimits<uint32>::Max());
-		Scatter.Size = (uint32)RequestSizeInBlock;
+		FScopeLock Lock(&PendingBlocksCritical);
+		if (!PendingBlocksTail)
+		{
+			PendingBlocksHead = RawBlock;
+		}
+		else
+		{
+			PendingBlocksTail->Next = RawBlock;
+		}
+		PendingBlocksTail = RawBlock;
+	}
+	PendingBlockEvent->Trigger();
+	if (!bIsMultithreaded)
+	{
+		ReadPendingBlocks();
 	}
 }
 
-void FFileIoStore::ReadBlocksUncached(uint32 BeginBlockIndex, uint32 BlockCount, FFileIoStoreResolvedRequest& ResolvedRequest)
+FFileIoStoreBuffer* FFileIoStore::AllocBuffer()
 {
-	uint64 BlockOffset = uint64(BeginBlockIndex) * uint64(CacheBlockSize);
-	uint64 BlockOffsetInRequest = BlockOffset - ResolvedRequest.ResolvedOffset;
-	uint64 ReadSize = BlockCount * CacheBlockSize;
-	check(ResolvedRequest.Request->IoBuffer.Data() + BlockOffsetInRequest + ReadSize <= ResolvedRequest.Request->IoBuffer.Data() + ResolvedRequest.Request->IoBuffer.DataSize());
-
-	FFileIoStoreReadBlock* UncachedBlock = new FFileIoStoreReadBlock();
-	UncachedBlock->Size = ReadSize;
-	UncachedBlock->Offset = BlockOffset;
-	UncachedBlock->Key.FileHandle = ResolvedRequest.ResolvedFileHandle;
-	UncachedBlock->Key.BlockIndex = BlockOffset / CacheBlockSize;
-	UncachedBlock->Buffer = FIoBuffer(FIoBuffer::Wrap, ResolvedRequest.Request->IoBuffer.Data() + BlockOffsetInRequest, ReadSize);
+	for (;;)
+	{
+		{
+			FScopeLock Lock(&BuffersCritical);
+			FFileIoStoreBuffer* Buffer = FirstFreeBuffer;
+			if (Buffer)
+			{
+				FirstFreeBuffer = Buffer->Next;
+				return Buffer;
+			}
+		}
+		TRACE_CPUPROFILER_EVENT_SCOPE(WaitForBuffer);
+		PlatformImpl.FlushReads();
+		BufferAvailableEvent->Wait();
+	}
 
-	++ResolvedRequest.Request->UnfinishedReadsCount;
-	FFileIoStoreReadBlockScatter& Scatter = UncachedBlock->ScatterList.AddDefaulted_GetRef();
-	Scatter.Request = ResolvedRequest.Request;
-	Scatter.DstOffset = MAX_uint64;
-	Scatter.SrcOffset = MAX_uint64;
-	Scatter.Size = ReadSize;
-	PlatformImpl.ReadBlockFromFile(UncachedBlock);
+	return nullptr;
+}
+
+void FFileIoStore::FreeBuffer(FFileIoStoreBuffer* Buffer)
+{
+	{
+		FScopeLock Lock(&BuffersCritical);
+		Buffer->Next = FirstFreeBuffer;
+		FirstFreeBuffer = Buffer;
+	}
+	BufferAvailableEvent->Trigger();
+}
+
+FFileIoStoreCompressionContext* FFileIoStore::AllocCompressionContext()
+{
+	FFileIoStoreCompressionContext* Result = FirstFreeCompressionContext;
+	if (Result)
+	{
+		FirstFreeCompressionContext = FirstFreeCompressionContext->Next;
+	}
+	return Result;
+}
+
+void FFileIoStore::FreeCompressionContext(FFileIoStoreCompressionContext* CompressionContext)
+{
+	CompressionContext->Next = FirstFreeCompressionContext;
+	FirstFreeCompressionContext = CompressionContext;
+}
+
+void FFileIoStore::ReadPendingBlocks()
+{
+	while (!bStopRequested)
+	{
+		{
+			FScopeLock PendingBlocksLock(&PendingBlocksCritical);
+			if (PendingBlocksHead)
+			{
+				if (!ScheduledBlocksTail)
+				{
+					ScheduledBlocksHead = PendingBlocksHead;
+					ScheduledBlocksTail = PendingBlocksTail;
+				}
+				else
+				{
+					ScheduledBlocksTail->Next = PendingBlocksHead;
+					ScheduledBlocksTail = PendingBlocksTail;
+				}
+				PendingBlocksHead = PendingBlocksTail = nullptr;
+			}
+		}
+
+		if (!ScheduledBlocksHead)
+		{
+			break;
+		}
+
+		FFileIoStoreRawBlock* BlockToRead = ScheduledBlocksHead;
+
+		/*TStringBuilder<256> ScopeName;
+		if (BlockToRead->DirectToRequest)
+		{
+			ScopeName.Appendf(TEXT("ReadNoScatter %lld %lld"), BlockToRead->Offset, BlockToRead->Size);
+		}
+		else
+		{
+			ScopeName.Appendf(TEXT("ReadBlock %d"), BlockToRead->Key.BlockIndex);
+		}
+		TRACE_CPUPROFILER_EVENT_SCOPE_TEXT(*ScopeName);*/
+
+		uint8* Target;
+		if (BlockToRead->DirectToRequest)
+		{
+			Target = BlockToRead->DirectToRequest->IoBuffer.Data() + BlockToRead->DirectToRequestOffset;
+		}
+		else
+		{
+			BlockToRead->Buffer = AllocBuffer();
+			Target = BlockToRead->Buffer->Memory;
+		}
+		
+		ScheduledBlocksHead = ScheduledBlocksHead->Next;
+		if (!ScheduledBlocksHead)
+		{
+			ScheduledBlocksTail = nullptr;
+		}
+		
+		PlatformImpl.ReadBlockFromFile(Target, IoStoreReaders[BlockToRead->Key.FileIndex]->GetContainerFile().FileHandle, BlockToRead);
+	}
+	PlatformImpl.FlushReads();
+}
+
+bool FFileIoStore::Init()
+{
+	return true;
+}
+
+void FFileIoStore::Stop()
+{
+	bStopRequested = true;
+	PendingBlockEvent->Trigger();
+	BufferAvailableEvent->Trigger();
+}
+
+uint32 FFileIoStore::Run()
+{
+	while (!bStopRequested)
+	{
+		ReadPendingBlocks();
+		PendingBlockEvent->Wait();
+	}
+	return 0;
 }
diff --git a/Engine/Source/Runtime/Core/Private/IO/IoDispatcherFileBackend.h b/Engine/Source/Runtime/Core/Private/IO/IoDispatcherFileBackend.h
--- a/Engine/Source/Runtime/Core/Private/IO/IoDispatcherFileBackend.h
+++ b/Engine/Source/Runtime/Core/Private/IO/IoDispatcherFileBackend.h
@@ -2,28 +2,64 @@
 
 #pragma once
 
-#include"IO/IoDispatcherPrivate.h"
+#include "IO/IoDispatcherPrivate.h"
 #include "IO/IoStore.h"
 #include "Containers/Array.h"
 #include "Containers/Map.h"
+#include "Stats/Stats.h"
+#include "Async/TaskGraphInterfaces.h"
+#include "HAL/Runnable.h"
 
-struct FFileIoStoreCacheBlockKey
+class IMappedFileHandle;
+
+struct FFileIoStoreContainerFile
+{
+	uint64 FileHandle = 0;
+	uint64 FileSize = 0;
+	uint64 CompressionBlockSize = 0;
+	TArray<FName> CompressionMethods;
+	TArray<FIoStoreCompressedBlockEntry> CompressionBlocks;
+	FString FilePath;
+	TUniquePtr<IMappedFileHandle> MappedFileHandle;
+};
+
+struct FFileIoStoreBuffer
 {
-	uint64 FileHandle;
-	uint64 BlockIndex;
+	FFileIoStoreBuffer* Next = nullptr;
+	uint8* Memory = nullptr;
+};
 
-	friend bool operator==(const FFileIoStoreCacheBlockKey& A, const FFileIoStoreCacheBlockKey& B)
+class FFileIoStoreBufferAllocator
+{
+public:
+	FFileIoStoreBuffer* AllocBuffer();
+};
+
+struct FFileIoStoreBlockKey
+{
+	union
 	{
-		return A.BlockIndex == B.BlockIndex && A.FileHandle == B.FileHandle;
+		struct
+		{
+			uint32 FileIndex;
+			uint32 BlockIndex;
+		};
+		uint64 Hash;
+	};
+		
+
+	friend bool operator==(const FFileIoStoreBlockKey& A, const FFileIoStoreBlockKey& B)
+	{
+		return A.Hash == B.Hash;
 	}
 
-	friend uint32 GetTypeHash(const FFileIoStoreCacheBlockKey& Key)
+	friend uint32 GetTypeHash(const FFileIoStoreBlockKey& Key)
 	{
-		return HashCombine(GetTypeHash(Key.FileHandle), GetTypeHash(Key.BlockIndex));
+		return GetTypeHash(Key.Hash);
 	}
 };
 
-struct FFileIoStoreReadBlockScatter
+struct FFileIoStoreBlockScatter
 {
 	FIoRequestImpl* Request = nullptr;
 	uint64 DstOffset;
@@ -31,26 +67,47 @@ struct FFileIoStoreReadBlockScatter
 	uint64 Size;
 };
 
-struct FFileIoStoreReadBlock
+struct FFileIoStoreCompressionContext
 {
-	FFileIoStoreReadBlock* Next = nullptr;
-	FFileIoStoreReadBlock* LruPrev = nullptr;
-	FFileIoStoreReadBlock* LruNext = nullptr;
-	FFileIoStoreCacheBlockKey Key;
-	FIoBuffer Buffer;
-	uint64 Size = 0;
-	uint64 Offset = 0;
-	TArray<FFileIoStoreReadBlockScatter> ScatterList;
-	bool bIsReady = false;
+	FFileIoStoreCompressionContext* Next = nullptr;
+	uint64 UncompressedBufferSize = 0;
+	uint8* UncompressedBuffer = nullptr;
+};
+
+struct FFileIoStoreCompressedBlock
+{
+	FFileIoStoreCompressedBlock* Next = nullptr;
+	FFileIoStoreBlockKey Key;
+	FName CompressionMethod;
+	uint64 Size;
+	uint64 RawOffset;
+	uint64 RawSize;
+	uint32 RawBlocksCount = 0;
+	uint32 UnfinishedRawBlocksCount = 0;
+	struct FFileIoStoreRawBlock* SingleRawBlock;
+	TArray<FFileIoStoreBlockScatter, TInlineAllocator<16>> ScatterList;
+	FFileIoStoreCompressionContext* CompressionContext = nullptr;
+	uint8* CompressedDataBuffer = nullptr;
+};
+
+struct FFileIoStoreRawBlock
+{
+	FFileIoStoreRawBlock* Next = nullptr;
+	FFileIoStoreBlockKey Key;
+	uint64 Offset;
+	uint64 Size;
+	FFileIoStoreBuffer* Buffer = nullptr;
+	TArray<FFileIoStoreCompressedBlock*, TInlineAllocator<4>> CompressedBlocks;
+	uint32 RefCount = 0;
+	FIoRequestImpl* DirectToRequest = nullptr;
+	uint64 DirectToRequestOffset = 0;
 };
 
 struct FFileIoStoreResolvedRequest
 {
 	FIoRequestImpl* Request;
-	uint64 ResolvedFileHandle;
 	uint64 ResolvedOffset;
 	uint64 ResolvedSize;
-	uint64 ResolvedFileSize;
 };
 
 class FFileIoStoreReader
@@ -60,40 +117,105 @@ public:
 	FIoStatus Initialize(const FIoStoreEnvironment& Environment);
 	bool DoesChunkExist(const FIoChunkId& ChunkId) const;
 	TIoStatusOr<uint64> GetSizeForChunk(const FIoChunkId& ChunkId) const;
-	bool Resolve(FFileIoStoreResolvedRequest& ResolvedRequest);
+	const FIoOffsetAndLength* Resolve(const FIoChunkId& ChunkId) const;
+	const FFileIoStoreContainerFile& GetContainerFile() { return ContainerFile; }
+	IMappedFileHandle* GetMappedContainerFileHandle();
 
 private:
 	FFileIoStoreImpl& PlatformImpl;
 
 	TMap<FIoChunkId, FIoOffsetAndLength> Toc;
-	uint64 ContainerFileHandle;
-	uint64 ContainerFileSize;
+	FFileIoStoreContainerFile ContainerFile;
 };
 
 class FFileIoStore
+	: public FRunnable
 {
 public:
-	FFileIoStore(FIoDispatcherEventQueue& InEventQueue);
+	FFileIoStore(FIoDispatcherEventQueue& InEventQueue, bool bInIsMultithreaded);
+	~FFileIoStore();
 	FIoStatus Mount(const FIoStoreEnvironment& Environment);
 	EIoStoreResolveResult Resolve(FIoRequestImpl* Request);
 	bool DoesChunkExist(const FIoChunkId& ChunkId) const;
 	TIoStatusOr<uint64> GetSizeForChunk(const FIoChunkId& ChunkId) const;
-	bool ProcessCompletedBlock();
+	void ProcessCompletedBlocks();
+	TIoStatusOr<FIoMappedRegion> OpenMapped(const FIoChunkId& ChunkId, const FIoReadOptions& Options);
 
 	static bool IsValidEnvironment(const FIoStoreEnvironment& Environment);
 
+	virtual bool Init() override;
+	virtual uint32 Run() override;
+	virtual void Stop() override;
+
 private:
-	void InitCache();
-	void ReadBlockCached(uint32 BlockIndex, const FFileIoStoreResolvedRequest& ResolvedRequest);
-	void ReadBlocksUncached(uint32 BeginBlockIndex, uint32 BlockCount, FFileIoStoreResolvedRequest& ResolvedRequest);
+	class FDecompressAsyncTask
+	{
+	public:
+		FDecompressAsyncTask(FFileIoStore& InOuter, FFileIoStoreCompressedBlock* InCompressedBlock)
+			: Outer(InOuter)
+			, CompressedBlock(InCompressedBlock)
+		{
 
+		}
+
+		static FORCEINLINE TStatId GetStatId()
+		{
+			RETURN_QUICK_DECLARE_CYCLE_STAT(FIoStoreDecompressTask, STATGROUP_TaskGraphTasks);
+		}
+
+		static ENamedThreads::Type GetDesiredThread();
+
+		FORCEINLINE static ESubsequentsMode::Type GetSubsequentsMode()
+		{
+			return ESubsequentsMode::FireAndForget;
+		}
+
+		void DoTask(ENamedThreads::Type CurrentThread, const FGraphEventRef& MyCompletionGraphEvent)
+		{
+			Outer.ScatterBlock(CompressedBlock, true);
+		}
+
+	private:
+		FFileIoStore& Outer;
+		FFileIoStoreCompressedBlock* CompressedBlock;
+	};
+
+	void InitCache();
+	void ReadBlockAndScatter(uint32 ReaderIndex, uint32 BlockIndex, const FFileIoStoreResolvedRequest& ResolvedRequest);
+	void ReadNoScatter(uint32 ReaderIndex, uint64 Offset, uint64 Size, const FFileIoStoreResolvedRequest& ResolvedRequest);
+	FFileIoStoreBuffer* AllocBuffer();
+	void FreeBuffer(FFileIoStoreBuffer* Buffer);
+	FFileIoStoreCompressionContext* AllocCompressionContext();
+	void FreeCompressionContext(FFileIoStoreCompressionContext* CompressionContext);
+	void ScatterBlock(FFileIoStoreCompressedBlock* CompressedBlock, bool bIsAsync);
+	void AllocMemoryForRequest(FIoRequestImpl* Request);
+	void FinalizeCompressedBlock(FFileIoStoreCompressedBlock* CompressedBlock);
+	void ReadPendingBlocks();
+
+	const uint64 ReadBufferSize;
+	FIoDispatcherEventQueue& EventQueue;
+	bool bIsMultithreaded;
 	FFileIoStoreImpl PlatformImpl;
 
+	FRunnableThread* Thread;
+	TAtomic<bool> bStopRequested{ false };
 	mutable FRWLock IoStoreReadersLock;
 	TArray<FFileIoStoreReader*> IoStoreReaders;
-	TMap<FFileIoStoreCacheBlockKey, FFileIoStoreReadBlock*> CachedBlocksMap;
-	FFileIoStoreReadBlock LruHead;
-	FFileIoStoreReadBlock LruTail;
-	const uint64 CacheBlockSize;
-	uint64 CurrentCacheUsage = 0;
+	TMap<FFileIoStoreBlockKey, FFileIoStoreCompressedBlock*> CompressedBlocksMap;
+	TMap<FFileIoStoreBlockKey, FFileIoStoreRawBlock*> RawBlocksMap;
+	uint8* BufferMemory = nullptr;
+	FEvent* BufferAvailableEvent;
+	FCriticalSection BuffersCritical;
+	FFileIoStoreBuffer* FirstFreeBuffer = nullptr;
+	FFileIoStoreCompressionContext* FirstFreeCompressionContext = nullptr;
+	FEvent* PendingBlockEvent;
+	FCriticalSection PendingBlocksCritical;
+	FFileIoStoreRawBlock* PendingBlocksHead = nullptr;
+	FFileIoStoreRawBlock* PendingBlocksTail = nullptr;
+	FFileIoStoreCompressedBlock* ReadyForDecompressionHead = nullptr;
+	FFileIoStoreCompressedBlock* ReadyForDecompressionTail = nullptr;
+	FFileIoStoreRawBlock* ScheduledBlocksHead = nullptr;
+	FFileIoStoreRawBlock* ScheduledBlocksTail = nullptr;
+	FCriticalSection DecompressedBlocksCritical;
+	FFileIoStoreCompressedBlock* FirstDecompressedBlock = nullptr;
 };
diff --git a/Engine/Source/Runtime/Core/Private/IO/IoDispatcherPrivate.h b/Engine/Source/Runtime/Core/Private/IO/IoDispatcherPrivate.h
--- a/Engine/Source/Runtime/Core/Private/IO/IoDispatcherPrivate.h
+++ b/Engine/Source/Runtime/Core/Private/IO/IoDispatcherPrivate.h
@@ -27,7 +27,12 @@ enum EIoStoreResolveResult
 class FIoBatchImpl
 {
 public:
-	FIoRequestImpl* FirstRequest = nullptr;
+	FIoRequestImpl* HeadRequest = nullptr;
+	FIoRequestImpl* TailRequest = nullptr;
+
+	// Used for contiguous reads
+	FIoBuffer IoBuffer;
+	FIoReadCallback Callback;
 };
 
 class FIoRequestImpl
@@ -40,7 +45,7 @@ public:
 	FIoChunkId ChunkId;
 	FIoReadOptions Options;
 	FIoBuffer IoBuffer;
-	uint32 UnfinishedReadsCount;
-	TFunction<void(TIoStatusOr<FIoBuffer>)> Callback;
+	TAtomic<uint32> UnfinishedReadsCount{ 0 };
+	FIoReadCallback Callback;
 };
 
