Copying //UE4/Release-4.25Plus-Dev to Release-4.25Plus (//UE4/Release-4.25Plus)

diff --git a/Engine/Source/Runtime/Core/Private/IO/IoStore.cpp b/Engine/Source/Runtime/Core/Private/IO/IoStore.cpp
--- a/Engine/Source/Runtime/Core/Private/IO/IoStore.cpp
+++ b/Engine/Source/Runtime/Core/Private/IO/IoStore.cpp
@@ -5,11 +5,15 @@
 #include "HAL/FileManager.h"
 #include "Templates/UniquePtr.h"
 #include "Misc/Paths.h"
+#include "Misc/Compression.h"
+#include "Serialization/BufferWriter.h"
+#include "Async/Async.h"
+#include "Async/Future.h"
+#include "Async/ParallelFor.h"
 
 //////////////////////////////////////////////////////////////////////////
 
 constexpr char FIoStoreTocHeader::TocMagicImg[];
-constexpr uint32 IoChunkAlignment = 16;
 
 //////////////////////////////////////////////////////////////////////////
 
@@ -28,6 +32,432 @@ void FIoStoreEnvironment::InitializeFileEnvironment(FStringView InPath)
 
 //////////////////////////////////////////////////////////////////////////
 
+struct FIoStoreWriterBlock
+{
+	FIoStoreWriterBlock* Next = nullptr;
+	FName CompressionMethod = NAME_None;
+	uint64 CompressedSize = 0;
+	uint64 Alignment = 0;
+	FIoBuffer UncompressedData;
+	FIoBuffer CompressedData;
+	FGraphEventRef CompressionTask;
+	bool bForceUncompressed = false;
+};
+
+class FIoStoreWriterContextImpl
+{
+public:
+	FIoStoreWriterContextImpl()
+	{
+
+	}
+
+	~FIoStoreWriterContextImpl()
+	{
+		if (FreeBlockEvent)
+		{
+			FPlatformProcess::ReturnSynchEventToPool(FreeBlockEvent);
+		}
+	}
+
+	UE_NODISCARD FIoStatus Initialize(const FIoStoreWriterSettings& InWriterSettings)
+	{
+		check(InWriterSettings.CompressionBlockSize > 0);
+		WriterSettings = InWriterSettings;
+		FreeBlockEvent = FPlatformProcess::GetSynchEventFromPool(false);
+		int32 BlockCount = int32((2ull << 30) / WriterSettings.CompressionBlockSize);
+		Blocks.SetNum(BlockCount);
+		for (int32 BlockIndex = 0; BlockIndex < BlockCount; ++BlockIndex)
+		{
+			FIoStoreWriterBlock& Block = Blocks[BlockIndex];
+			Block.UncompressedData = FIoBuffer(WriterSettings.CompressionBlockSize);
+			int32 MaxCompressedSize = FCompression::CompressMemoryBound(WriterSettings.CompressionMethod, int32(WriterSettings.CompressionBlockSize));
+			Block.CompressedData = FIoBuffer(MaxCompressedSize);
+			Block.Next = FirstFreeBlock;
+			FirstFreeBlock = &Block;
+		}
+
+		return FIoStatus::Ok;
+	}
+
+	const FIoStoreWriterSettings& Settings() const
+	{
+		return WriterSettings;
+	}
+
+	FIoStoreWriterBlock* AllocBlock()
+	{
+		for (;;)
+		{
+			{
+				FScopeLock Lock(&FreeBlocksCritical);
+				if (FirstFreeBlock)
+				{
+					FIoStoreWriterBlock* Result = FirstFreeBlock;
+					FirstFreeBlock = FirstFreeBlock->Next;
+					Result->bForceUncompressed = false;
+					Result->Alignment = WriterSettings.CompressionBlockAlignment;
+					return Result;
+				}
+			}
+			{
+				TRACE_CPUPROFILER_EVENT_SCOPE(WaitForBlock);
+				FreeBlockEvent->Wait();
+			}
+		}
+	}
+
+	void FreeBlock(FIoStoreWriterBlock* Block)
+	{
+		{
+			FScopeLock Lock(&FreeBlocksCritical);
+			Block->Next = FirstFreeBlock;
+			FirstFreeBlock = Block;
+		}
+		FreeBlockEvent->Trigger();
+	}
+
+private:
+	FIoStoreWriterSettings WriterSettings;
+	TArray<FIoStoreWriterBlock> Blocks;
+	FCriticalSection FreeBlocksCritical;
+	FIoStoreWriterBlock* FirstFreeBlock = nullptr;
+	FEvent* FreeBlockEvent = nullptr;
+};
+
+FIoStoreWriterContext::FIoStoreWriterContext()
+	: Impl(new FIoStoreWriterContextImpl())
+{
+
+}
+
+FIoStoreWriterContext::~FIoStoreWriterContext()
+{
+	delete Impl;
+}
+
+UE_NODISCARD FIoStatus FIoStoreWriterContext::Initialize(const FIoStoreWriterSettings& InWriterSettings)
+{
+	return Impl->Initialize(InWriterSettings);
+}
+
+static uint64 GetPadding(const uint64 Offset, const uint64 Alignment)
+{
+	return (Alignment - (Offset % Alignment)) % Alignment;
+}
+
+class FChunkWriter
+{
+public:
+	FChunkWriter() { }
+
+	virtual ~FChunkWriter()
+	{
+		if (WriteQueueEvent)
+		{
+			FPlatformProcess::ReturnSynchEventToPool(WriteQueueEvent);
+		}
+	}
+
+	FIoStatus Initialize(FIoStoreWriterContextImpl& InContext, const TCHAR* InFileName, bool bInIsContainerCompressed)
+	{
+		TRACE_CPUPROFILER_EVENT_SCOPE(InitializeChunkWriter);
+		WriterContext = &InContext;
+		bIsContainerCompressed = bInIsContainerCompressed;
+		check(WriterContext->Settings().CompressionBlockSize > 0);
+		CompressionInfo.BlockSize = WriterContext->Settings().CompressionBlockSize;
+
+		IPlatformFile& Ipf = IPlatformFile::GetPlatformPhysical();
+		FileHandle.Reset(Ipf.OpenWrite(InFileName, /* append */ false, /* allowread */ true));
+		if (!FileHandle)
+		{
+			return FIoStatus(EIoErrorCode::FileNotOpen, TEXT("Failed to open container file handle"));
+		}
+
+		WriteQueueEvent = FPlatformProcess::GetSynchEventFromPool(false);
+		WriterTask = Async(EAsyncExecution::Thread, [this]()
+		{
+			TRACE_CPUPROFILER_EVENT_SCOPE(WriteContainerThread);
+			for (;;)
+			{
+				FIoStoreWriterBlock* NextWriteQueueItem = nullptr;
+				{
+					FScopeLock Lock(&WriteQueueCritical);
+					NextWriteQueueItem = WriteQueueHead;
+					WriteQueueHead = WriteQueueTail = nullptr;
+				}
+				if (!NextWriteQueueItem && bAllScheduled)
+				{
+					return;
+				}
+				while (NextWriteQueueItem)
+				{
+					FIoStoreWriterBlock* CurrentWriteQueueItem = NextWriteQueueItem;
+					NextWriteQueueItem = NextWriteQueueItem->Next;
+
+					if (bIsContainerCompressed && !CurrentWriteQueueItem->bForceUncompressed)
+					{
+						FTaskGraphInterface::Get().WaitUntilTaskCompletes(CurrentWriteQueueItem->CompressionTask);
+					}
+					if (CurrentWriteQueueItem->Alignment > 0)
+					{
+						const uint64 Padding = GetPadding(FileHandle->Tell(), CurrentWriteQueueItem->Alignment);
+						if (Padding > 0)
+						{
+							AlignmentPaddingBuffer.SetNumZeroed(int32(Padding), false);
+							FileHandle->Write(AlignmentPaddingBuffer.GetData(), Padding);
+							PaddingSize += Padding;
+						}
+						check(FileHandle->Tell() % CurrentWriteQueueItem->Alignment == 0);
+					}
+					const int64 CompressedFileOffset = FileHandle->Tell();
+					FIoStoreCompressedBlockEntry CompressedBlockEntry;
+					CompressedBlockEntry.OffsetAndLength.SetOffset(CompressedFileOffset);
+					const uint8* SourceData;
+					uint64 SourceDataSize;
+					if (CurrentWriteQueueItem->CompressionMethod == NAME_None)
+					{
+						SourceData = CurrentWriteQueueItem->UncompressedData.Data();
+						SourceDataSize = CurrentWriteQueueItem->UncompressedData.DataSize();
+					}
+					else
+					{
+						SourceData = CurrentWriteQueueItem->CompressedData.Data();
+						SourceDataSize = CurrentWriteQueueItem->CompressedSize;
+					}
+
+					FIoBuffer& SourceBuffer = CurrentWriteQueueItem->CompressionMethod == NAME_None ? CurrentWriteQueueItem->UncompressedData : CurrentWriteQueueItem->CompressedData;
+					CompressedBlockEntry.OffsetAndLength.SetLength(SourceDataSize);
+					CompressedBlockEntry.CompressionMethodIndex = CompressionInfo.GetCompressionMethodIndex(CurrentWriteQueueItem->CompressionMethod);
+					CompressionInfo.BlockEntries.Emplace(CompressedBlockEntry);
+					{
+						TRACE_CPUPROFILER_EVENT_SCOPE(WriteBlockToFile);
+						FileHandle->Write(SourceData, SourceDataSize);
+					}
+					WriterContext->FreeBlock(CurrentWriteQueueItem);
+				}
+				if (!bAllScheduled)
+				{
+					WriteQueueEvent->Wait();
+				}
+			}
+		});
+		return FIoStatus::Ok;
+	}
+
+	TIoStatusOr<FIoStoreTocEntry> Write(FIoChunkId ChunkId, FIoBuffer Chunk, FIoWriteOptions WriteOptions)
+	{
+		if (WriteOptions.Alignment > 0)
+		{
+			if (CurrentBlockOffset > 0)
+			{
+				FIoStatus IoStatus = FIoStatusBuilder(EIoErrorCode::InvalidParameter) << TEXT("Can only force alignment at the start of a block");
+				return IoStatus;
+			}
+			if (!CurrentBlock)
+			{
+				CurrentBlock = WriterContext->AllocBlock();
+			}
+			CurrentBlock->Alignment = WriteOptions.Alignment;
+		}
+
+		if (bIsContainerCompressed && WriteOptions.bForceUncompressed)
+		{
+			if (CurrentBlockOffset > 0 && CurrentBlock && !CurrentBlock->bForceUncompressed)
+			{
+				FIoStatus IoStatus = FIoStatusBuilder(EIoErrorCode::InvalidParameter) << TEXT("Can only change compression mode at the start of a block");
+				return IoStatus;
+			}
+		}
+
+		FIoStoreTocEntry TocEntry;
+		TocEntry.SetOffset(UncompressedOffset);
+		TocEntry.SetLength(Chunk.DataSize());
+		TocEntry.ChunkId = ChunkId;
+
+		UncompressedOffset += Chunk.DataSize();
+
+		uint64 RemainingBytesInChunk = Chunk.DataSize();
+		uint64 OffsetInChunk = 0;
+		while (RemainingBytesInChunk > 0)
+		{
+			if (!CurrentBlock)
+			{
+				CurrentBlock = WriterContext->AllocBlock();
+				if (RemainingBytesInChunk >= CurrentBlock->UncompressedData.DataSize())
+				{
+					// Entire block belongs to this chunk, skip patch alignment
+					CurrentBlock->Alignment = 0;
+				}
+			}
+			check(CurrentBlock->UncompressedData.DataSize() > CurrentBlockOffset);
+			CurrentBlock->bForceUncompressed = WriteOptions.bForceUncompressed;
+			uint64 BytesToWrite = FMath::Min(RemainingBytesInChunk, CurrentBlock->UncompressedData.DataSize() - CurrentBlockOffset);
+			FMemory::Memcpy(CurrentBlock->UncompressedData.Data() + CurrentBlockOffset, Chunk.Data() + OffsetInChunk, BytesToWrite);
+			OffsetInChunk += BytesToWrite;
+			check(RemainingBytesInChunk >= BytesToWrite);
+			RemainingBytesInChunk -= BytesToWrite;
+			CurrentBlockOffset += BytesToWrite;
+			if (CurrentBlockOffset == CurrentBlock->UncompressedData.DataSize())
+			{
+				FlushCurrentBlock();
+			}
+		}
+		
+		return TocEntry;
+	}
+
+	FIoStatus WritePadding(uint64 Count)
+	{
+		while (Count > 0)
+		{
+			if (!CurrentBlock)
+			{
+				CurrentBlock = WriterContext->AllocBlock();
+			}
+			check(CurrentBlock->UncompressedData.DataSize() > CurrentBlockOffset);
+			uint64 BytesToWrite = FMath::Min(Count, CurrentBlock->UncompressedData.DataSize() - CurrentBlockOffset);
+			FMemory::Memzero(CurrentBlock->UncompressedData.Data() + CurrentBlockOffset, BytesToWrite);
+			check(Count >= BytesToWrite);
+			Count -= BytesToWrite;
+			CurrentBlockOffset += BytesToWrite;
+			UncompressedOffset += BytesToWrite;
+			if (CurrentBlockOffset == CurrentBlock->UncompressedData.DataSize())
+			{
+				FlushCurrentBlock();
+			}
+		}
+		return FIoStatus::Ok;
+	}
+
+	FIoStatus Flush()
+	{
+		TRACE_CPUPROFILER_EVENT_SCOPE(EndWriteContainer);
+		FlushCurrentBlock();
+		bAllScheduled = true;
+		WriteQueueEvent->Trigger();
+		WriterTask.Wait();
+		
+		check(!WriteQueueHead);
+		check(!WriteQueueTail);
+		bool bIsUncompressed = true;
+		for (const FIoStoreCompressedBlockEntry& CompressedBlock : CompressionInfo.BlockEntries)
+		{
+			if (CompressedBlock.CompressionMethodIndex != FIoStoreCompressionInfo::InvalidCompressionIndex)
+			{
+				bIsUncompressed = false;
+				break;
+			}
+		}
+		if (bIsUncompressed)
+		{
+			CompressionInfo.BlockEntries.Empty();
+			CompressionInfo.BlockSize = 0;
+		}
+		CompressionInfo.UncompressedContainerSize = UncompressedOffset;
+		CompressionInfo.CompressedContainerSize = FileHandle.IsValid() ? FileHandle->Tell() : 0;
+
+		return FIoStatus::Ok;
+	}
+
+	const FIoStoreCompressionInfo& GetCompressionInfo()
+	{
+		return CompressionInfo;
+	}
+
+	int64 GetPaddingSize()
+	{
+		return PaddingSize;
+	}
+
+private:
+	void CompressBlock(FIoStoreWriterBlock* Block)
+	{
+		TRACE_CPUPROFILER_EVENT_SCOPE(CompressBlock);
+		int32 CompressedSize = int32(Block->CompressedData.DataSize());
+		const bool bCompressed = FCompression::CompressMemory(
+			WriterContext->Settings().CompressionMethod,
+			Block->CompressedData.Data(),
+			CompressedSize,
+			Block->UncompressedData.Data(),
+			int32(Block->UncompressedData.DataSize()));
+		check(bCompressed);
+		check(CompressedSize > 0);
+
+		if (CompressedSize > Block->UncompressedData.DataSize())
+		{
+			Block->CompressionMethod = NAME_None;
+			Block->CompressedSize = Block->UncompressedData.DataSize();
+		}
+		else
+		{
+			Block->CompressionMethod = WriterContext->Settings().CompressionMethod;
+			Block->CompressedSize = CompressedSize;
+		}
+	}
+
+	void FlushCurrentBlock()
+	{
+		if (!CurrentBlock)
+		{
+			return;
+		}
+		if (CurrentBlockOffset < CurrentBlock->UncompressedData.DataSize())
+		{
+			uint64 PadCount = CurrentBlock->UncompressedData.DataSize() - CurrentBlockOffset;
+			FMemory::Memzero(CurrentBlock->UncompressedData.Data() + CurrentBlockOffset, PadCount);
+			UncompressedOffset += PadCount;
+		}
+		if (bIsContainerCompressed && !CurrentBlock->bForceUncompressed)
+		{
+			FIoStoreWriterBlock* ReadyToCompressBlock = CurrentBlock;
+			CurrentBlock->CompressionTask = FFunctionGraphTask::CreateAndDispatchWhenReady([this, ReadyToCompressBlock]()
+			{
+				CompressBlock(ReadyToCompressBlock);
+			}, TStatId(), nullptr, ENamedThreads::AnyHiPriThreadHiPriTask);
+		}
+		else
+		{
+			CurrentBlock->CompressionMethod = NAME_None;
+			CurrentBlock->CompressedSize = CurrentBlock->UncompressedData.DataSize();
+		}
+		{
+			FScopeLock Lock(&WriteQueueCritical);
+			if (!WriteQueueTail)
+			{
+				WriteQueueHead = WriteQueueTail = CurrentBlock;
+			}
+			else
+			{
+				WriteQueueTail->Next = CurrentBlock;
+				WriteQueueTail = CurrentBlock;
+			}
+			CurrentBlock->Next = nullptr;
+		}
+		WriteQueueEvent->Trigger();
+		CurrentBlock = nullptr;
+		CurrentBlockOffset = 0;
+	}
+
+	FIoStoreWriterContextImpl* WriterContext = nullptr;
+	FIoStoreCompressionInfo CompressionInfo;
+	int64 PaddingSize = 0;
+	TUniquePtr<IFileHandle> FileHandle;
+	uint64 UncompressedOffset = 0;
+	FIoStoreWriterBlock* CurrentBlock = nullptr;
+	uint64 CurrentBlockOffset = 0;
+	TFuture<void> WriterTask;
+	FCriticalSection WriteQueueCritical;
+	FIoStoreWriterBlock* WriteQueueHead = nullptr;
+	FIoStoreWriterBlock* WriteQueueTail = nullptr;
+	FEvent* WriteQueueEvent = nullptr;
+	TAtomic<bool> bAllScheduled{ false };
+	TArray<uint8> AlignmentPaddingBuffer;
+	bool bIsContainerCompressed = false;
+};
+
+//////////////////////////////////////////////////////////////////////////
+
 class FIoStoreWriterImpl
 {
 public:
@@ -36,19 +466,25 @@ public:
 	{
 	}
 
-	FIoStatus Initialize()
+	UE_NODISCARD FIoStatus Initialize(FIoStoreWriterContextImpl& InContext, bool bInIsContainerCompressed)
 	{
 		IPlatformFile& Ipf = IPlatformFile::GetPlatformPhysical();
 
 		FString TocFilePath = Environment.GetPath() + TEXT(".utoc");
 		FString ContainerFilePath = Environment.GetPath() + TEXT(".ucas");
 
+		Result.ContainerName = FPaths::GetBaseFilename(Environment.GetPath());
+		Result.CompressionMethod = InContext.Settings().CompressionMethod;
+
 		Ipf.CreateDirectoryTree(*FPaths::GetPath(ContainerFilePath));
 
-		ContainerFileHandle.Reset(Ipf.OpenWrite(*ContainerFilePath, /* append */ false, /* allowread */ true));
+		ChunkWriter = MakeUnique<FChunkWriter>();
+		FIoStatus Status = ChunkWriter->Initialize(InContext, *ContainerFilePath, bInIsContainerCompressed);
 
-		if (!ContainerFileHandle)
+		if (!Status.IsOk())
 		{ 
+			ChunkWriter.Reset();
+
 			return FIoStatusBuilder(EIoErrorCode::FileOpenFailed) << TEXT("Failed to open IoStore container file '") << *ContainerFilePath << TEXT("'");
 		}
 
@@ -59,7 +495,12 @@ public:
 			return FIoStatusBuilder(EIoErrorCode::FileOpenFailed) << TEXT("Failed to open IoStore TOC file '") << *TocFilePath << TEXT("'");
 		}
 
-		return FIoStatus::Ok;
+		if (InContext.Settings().bEnableCsvOutput)
+		{
+			Status = EnableCsvOutput();
+		}
+
+		return Status;
 	}
 
 	FIoStatus EnableCsvOutput()
@@ -76,9 +517,9 @@ public:
 		return FIoStatus::Ok;
 	}
 
-	UE_NODISCARD FIoStatus Append(FIoChunkId ChunkId, FIoBuffer Chunk, const TCHAR* Name)
+	UE_NODISCARD FIoStatus Append(FIoChunkId ChunkId, FIoBuffer Chunk, FIoWriteOptions WriteOptions)
 	{
-		if (!ContainerFileHandle)
+		if (!ChunkWriter)
 		{
 			return FIoStatus(EIoErrorCode::FileNotOpen, TEXT("No container file to append to"));
 		}
@@ -93,33 +534,19 @@ public:
 			return FIoStatus(EIoErrorCode::InvalidParameter, TEXT("ChunkId is already mapped"));
 		}
 
-		FIoStoreTocEntry TocEntry;
-
-		check(ContainerFileHandle->Tell() % IoChunkAlignment == 0);
-
-		TocEntry.SetOffset(ContainerFileHandle->Tell());
-		TocEntry.SetLength(Chunk.DataSize());
-		TocEntry.ChunkId = ChunkId;
-
 		IsMetadataDirty = true;
 
-		bool Success = ContainerFileHandle->Write(Chunk.Data(), Chunk.DataSize());
-
-		if (uint32 UnpaddedBytes = Chunk.DataSize() % IoChunkAlignment)
-		{
-			static constexpr uint8 Zeroes[IoChunkAlignment] = {};
-			uint32 Padding = (IoChunkAlignment - UnpaddedBytes) % IoChunkAlignment;
-			Success &= ContainerFileHandle->Write(Zeroes, Padding);
-		}
+		TIoStatusOr<FIoStoreTocEntry> TocEntryStatus = ChunkWriter->Write(ChunkId, Chunk, WriteOptions);
 
-		if (Success)
+		if (TocEntryStatus.IsOk())
 		{
+			FIoStoreTocEntry TocEntry = TocEntryStatus.ConsumeValueOrDie();
 			Toc.Add(ChunkId, TocEntry);
 
 			if (CsvArchive)
 			{
 				ANSICHAR Line[MAX_SPRINTF];
-				FCStringAnsi::Sprintf(Line, "%s,%lld,%lld\n", TCHAR_TO_ANSI(Name), TocEntry.GetOffset(), TocEntry.GetLength());
+				FCStringAnsi::Sprintf(Line, "%s,%lld,%lld\n", (WriteOptions.DebugName ? TCHAR_TO_ANSI(WriteOptions.DebugName) : ""), TocEntry.GetOffset(), TocEntry.GetLength());
 				CsvArchive->Serialize(Line, FCStringAnsi::Strlen(Line));
 			}
 
@@ -131,6 +558,11 @@ public:
 		}
 	}
 
+	UE_NODISCARD FIoStatus AppendPadding(uint64 Count)
+	{
+		return ChunkWriter->WritePadding(Count);
+	}
+
 	UE_NODISCARD FIoStatus MapPartialRange(FIoChunkId OriginalChunkId, uint64 Offset, uint64 Length, FIoChunkId ChunkIdPartialRange)
 	{
 		//TODO: Does RelativeOffset + Length overflow?
@@ -169,41 +601,86 @@ public:
 		return FIoStatus::Ok;
 	}
 
-	UE_NODISCARD FIoStatus FlushMetadata()
+	UE_NODISCARD TIoStatusOr<FIoStoreWriterResult> Flush()
 	{
-		TocFileHandle->Seek(0);
+		if (!IsMetadataDirty)
+		{
+			return Result;
+		}
+
+		IsMetadataDirty = false;
+
+		ChunkWriter->Flush();
+		const FIoStoreCompressionInfo& CompressionInfo = ChunkWriter->GetCompressionInfo();
 
 		FIoStoreTocHeader TocHeader;
-		FMemory::Memset(TocHeader, 0);
+		FMemory::Memzero(&TocHeader, sizeof(TocHeader));
 
 		TocHeader.MakeMagic();
-		TocHeader.TocHeaderSize = sizeof TocHeader;
+		TocHeader.TocHeaderSize = sizeof(TocHeader);
 		TocHeader.TocEntryCount = Toc.Num();
 		TocHeader.TocEntrySize = sizeof(FIoStoreTocEntry);
+		TocHeader.CompressionBlockCount = CompressionInfo.BlockEntries.Num();
+		TocHeader.CompressionBlockSize = uint32(CompressionInfo.BlockSize);
+		TocHeader.CompressionNameCount = CompressionInfo.CompressionMethods.Num();
 
-		const bool Success = TocFileHandle->Write(reinterpret_cast<const uint8*>(&TocHeader), sizeof TocHeader);
-
-		if (!Success)
+		TocFileHandle->Seek(0);
+		if (!TocFileHandle->Write(reinterpret_cast<const uint8*>(&TocHeader), sizeof(TocHeader)))
 		{
-			return FIoStatus(EIoErrorCode::WriteError, TEXT("TOC write failed"));
+			return FIoStatus(EIoErrorCode::WriteError, TEXT("Failed to write TOC header"));
 		}
 
+		// Chunk entries
 		for (auto& _: Toc)
 		{
 			FIoStoreTocEntry& TocEntry = _.Value;
 			
-			TocFileHandle->Write(reinterpret_cast<const uint8*>(&TocEntry), sizeof TocEntry);
+			if (!TocFileHandle->Write(reinterpret_cast<const uint8*>(&TocEntry), sizeof(TocEntry)))
+			{
+				return FIoStatus(EIoErrorCode::WriteError, TEXT("Failed to write TOC entry"));
+			}
 		}
 
-		return FIoStatus::Ok;
+		// Compression blocks
+		for (const FIoStoreCompressedBlockEntry& CompressedBlockEntry : CompressionInfo.BlockEntries)
+		{
+			if (!TocFileHandle->Write(reinterpret_cast<const uint8*>(&CompressedBlockEntry), sizeof(CompressedBlockEntry)))
+			{
+				return FIoStatus(EIoErrorCode::WriteError, TEXT("Failed to write compression block TOC entry"));
+			}
+		}
+
+		// Compression methods
+		ANSICHAR AnsiMethodName[FIoStoreCompressionInfo::CompressionMethodNameLen];
+
+		for (FName MethodName : CompressionInfo.CompressionMethods)
+		{
+			FMemory::Memzero(AnsiMethodName, FIoStoreCompressionInfo::CompressionMethodNameLen);
+			FCStringAnsi::Strcpy(AnsiMethodName, FIoStoreCompressionInfo::CompressionMethodNameLen, TCHAR_TO_ANSI(*MethodName.ToString()));
+
+			if (!TocFileHandle->Write(reinterpret_cast<const uint8*>(AnsiMethodName), FIoStoreCompressionInfo::CompressionMethodNameLen))
+			{
+				return FIoStatus(EIoErrorCode::WriteError, TEXT("Failed to write compression method TOC entry"));
+			}
+		}
+
+		Result.TocSize = TocFileHandle->Tell();
+		Result.TocEntryCount = TocHeader.TocEntryCount;
+		Result.PaddingSize = ChunkWriter->GetPaddingSize();
+		Result.UncompressedContainerSize = CompressionInfo.UncompressedContainerSize;
+		Result.CompressedContainerSize = CompressionInfo.CompressedContainerSize;
+
+		return Result;
 	}
 
 private:
+	friend class FIoStoreWriter;
 	FIoStoreEnvironment&				Environment;
 	TMap<FIoChunkId, FIoStoreTocEntry>	Toc;
-	TUniquePtr<IFileHandle>				ContainerFileHandle;
+	TUniquePtr<FChunkWriter>			ChunkWriter;
 	TUniquePtr<IFileHandle>				TocFileHandle;
 	TUniquePtr<FArchive>				CsvArchive;
+	FIoStoreWriterResult				Result;
 	bool								IsMetadataDirty = true;
 };
 
@@ -214,22 +691,22 @@ FIoStoreWriter::FIoStoreWriter(FIoStoreEnvironment& InEnvironment)
 
 FIoStoreWriter::~FIoStoreWriter()
 {
-	(void)Impl->FlushMetadata();
+	(void)Impl->Flush();
 }
 
-FIoStatus FIoStoreWriter::Initialize()
+FIoStatus FIoStoreWriter::Initialize(const FIoStoreWriterContext& Context, bool bIsContainerCompressed)
 {
-	return Impl->Initialize();
+	return Impl->Initialize(*Context.Impl, bIsContainerCompressed);
 }
 
-FIoStatus FIoStoreWriter::EnableCsvOutput()
+FIoStatus FIoStoreWriter::Append(FIoChunkId ChunkId, FIoBuffer Chunk, FIoWriteOptions WriteOptions)
 {
-	return Impl->EnableCsvOutput();
+	return Impl->Append(ChunkId, Chunk, WriteOptions);
 }
 
-FIoStatus FIoStoreWriter::Append(FIoChunkId ChunkId, FIoBuffer Chunk, const TCHAR* Name)
+FIoStatus FIoStoreWriter::AppendPadding(uint64 Count)
 {
-	return Impl->Append(ChunkId, Chunk, Name);
+	return Impl->AppendPadding(Count);
 }
 
 FIoStatus FIoStoreWriter::MapPartialRange(FIoChunkId OriginalChunkId, uint64 Offset, uint64 Length, FIoChunkId ChunkIdPartialRange)
@@ -237,7 +714,7 @@ FIoStatus FIoStoreWriter::MapPartialRange(FIoChunkId OriginalChunkId, uint64 Off
 	return Impl->MapPartialRange(OriginalChunkId, Offset, Length, ChunkIdPartialRange);
 }
 
-FIoStatus FIoStoreWriter::FlushMetadata()
+TIoStatusOr<FIoStoreWriterResult> FIoStoreWriter::Flush()
 {
-	return Impl->FlushMetadata();
+	return Impl->Flush();
 }
diff --git a/Engine/Source/Runtime/Core/Private/IO/IoStore.h b/Engine/Source/Runtime/Core/Private/IO/IoStore.h
--- a/Engine/Source/Runtime/Core/Private/IO/IoStore.h
+++ b/Engine/Source/Runtime/Core/Private/IO/IoStore.h
@@ -15,7 +15,10 @@ struct FIoStoreTocHeader
 	uint32	TocHeaderSize;
 	uint32	TocEntryCount;
 	uint32	TocEntrySize;	// For sanity checking
-	uint32	TocPad[25];
+	uint32	CompressionBlockCount;
+	uint32	CompressionBlockSize;
+	uint32	CompressionNameCount;
+	uint32	TocPad[22];
 
 	void MakeMagic()
 	{
@@ -109,3 +112,54 @@ struct FIoStoreTocEntry
 		OffsetAndLength.SetLength(Length);
 	}
 };
+/**
+ * Compression block entry.
+ */
+struct FIoStoreCompressedBlockEntry
+{
+	/** Offset and size of the compressed block. */
+	FIoOffsetAndLength OffsetAndLength;
+	/** Index into the compression methods array. */
+	uint8 CompressionMethodIndex;
+};
+
+/**
+ * I/O store compresssion info.
+ */
+struct FIoStoreCompressionInfo
+{
+	enum
+	{
+		/** Compression method name max length. */
+		CompressionMethodNameLen = 32,
+		/** No compression. */
+		InvalidCompressionIndex = 255,
+	};
+
+	uint8 GetCompressionMethodIndex(FName CompressionMethod)
+	{
+		if (CompressionMethod == NAME_None)
+		{
+			return InvalidCompressionIndex;
+		}
+
+		for (int32 Idx = 0; Idx < CompressionMethods.Num(); ++Idx)
+		{
+			if (CompressionMethods[Idx] == CompressionMethod)
+			{
+				return uint8(Idx);
+			}
+		}
+
+		const uint8 Idx = uint8(CompressionMethods.Num());
+		CompressionMethods.Add(CompressionMethod);
+
+		return Idx;
+	}
+
+	TArray<FIoStoreCompressedBlockEntry> BlockEntries;
+	TArray<FName> CompressionMethods;
+	int64 BlockSize = 0;
+	int64 UncompressedContainerSize = 0;
+	int64 CompressedContainerSize = 0;
+};
diff --git a/Engine/Source/Runtime/Core/Private/Misc/PathViews.cpp b/Engine/Source/Runtime/Core/Private/Misc/PathViews.cpp
--- a/Engine/Source/Runtime/Core/Private/Misc/PathViews.cpp
+++ b/Engine/Source/Runtime/Core/Private/Misc/PathViews.cpp
@@ -3,6 +3,7 @@
 #include "Misc/PathViews.h"
 
 #include "Algo/FindLast.h"
+#include "Containers/UnrealString.h"
 #include "Containers/StringView.h"
 #include "Misc/StringBuilder.h"
 
@@ -10,6 +11,8 @@ namespace UE4PathViews_Private
 {
 	static bool IsSlashOrBackslash(TCHAR C) { return C == TEXT('/') || C == TEXT('\\'); }
 	static bool IsNotSlashOrBackslash(TCHAR C) { return C != TEXT('/') && C != TEXT('\\'); }
+
+	static bool IsSlashOrBackslashOrPeriod(TCHAR C) { return C == TEXT('/') || C == TEXT('\\') || C == TEXT('.'); }
 }
 
 FStringView FPathViews::GetCleanFilename(const FStringView& InPath)
@@ -84,3 +87,35 @@ void FPathViews::Append(FStringBuilderBase& Builder, const FStringView& Suffix)
 	}
 	Builder.Append(Suffix);
 }
+
+FString FPathViews::ChangeExtension(const FStringView& InPath, const FStringView& InNewExtension)
+{
+	// Make sure the period we found was actually for a file extension and not part of the file path.
+	const TCHAR* PathEndPos = Algo::FindLastByPredicate(InPath, UE4PathViews_Private::IsSlashOrBackslashOrPeriod);
+	if (PathEndPos != nullptr && *PathEndPos == TEXT('.'))
+	{
+		const FStringView::SizeType Pos = FStringView::SizeType(PathEndPos - InPath.GetData());
+		const FStringView FileWithoutExtension = InPath.Left(Pos);
+
+		if (!InNewExtension.IsEmpty() && !InNewExtension.StartsWith('.'))
+		{
+			// The new extension lacks a period so we need to add it ourselves.
+			FString Result(FileWithoutExtension, InNewExtension.Len() + 1);
+			Result += '.';
+			Result += InNewExtension;
+
+			return Result;
+		}
+		else
+		{
+			FString Result(FileWithoutExtension, InNewExtension.Len());
+			Result += InNewExtension;
+
+			return Result;
+		}
+	}
+	else
+	{
+		return FString(InPath);
+	}
+}
diff --git a/Engine/Source/Runtime/Core/Public/Containers/UnrealString.h b/Engine/Source/Runtime/Core/Public/Containers/UnrealString.h
--- a/Engine/Source/Runtime/Core/Public/Containers/UnrealString.h
+++ b/Engine/Source/Runtime/Core/Public/Containers/UnrealString.h
@@ -135,8 +135,49 @@ public:
 		}
 	}
 
+	/**
+	 * Constructor to create FString with specified number of characters from another string with additional character zero
+	 *
+	 * @param Other the other string to create a new copy from
+	 * @param ExtraSlack number of extra characters to add to the end of the other string in this string
+	 */
+	template <
+		typename CharType,
+		typename = typename TEnableIf<TIsCharType<CharType>::Value>::Type // This TEnableIf is to ensure we don't instantiate this constructor for non-char types, like id* in Obj-C
+	>
+		FORCEINLINE FString(const CharType* Src, int32 ExtraSlack)
+	{
+		if (Src && *Src)
+		{
+			int32 SrcLen = TCString<CharType>::Strlen(Src) + 1;
+			int32 DestLen = FPlatformString::ConvertedLength<TCHAR>(Src, SrcLen);
+			Data.Reserve(DestLen + ExtraSlack);
+			Data.AddUninitialized(DestLen);
+
+			FPlatformString::Convert(Data.GetData(), DestLen, Src, SrcLen);
+		}
+		else
+		{
+			Reserve(ExtraSlack); 
+		}
+	}
+
+	/**
+	 * Create an FString from a FStringView
+	 *
+	 * @param Other The string view to create a new copy from
+	 */
 	explicit FString(const FStringView& Other);
 
+	/**
+	 * Create an FString from a FStringView with extra space for characters at the end of the string
+	 *
+	 * @param Other The string view to create a new copy from
+	 * @param ExtraSlack number of extra characters to add to the end of the other string in this string
+	 */
+	explicit FString(const FStringView& Other, int32 ExtraSlack);
+	
+
 #ifdef __OBJC__
 	/** Convert Objective-C NSString* to FString */
 	FORCEINLINE FString(const NSString* In)
@@ -1758,7 +1799,10 @@ public:
 	FORCEINLINE void Reserve(int32 CharacterCount)
 	{
 		checkSlow(CharacterCount >= 0 && CharacterCount < MAX_int32);
-		Data.Reserve(CharacterCount + 1);
+		if (CharacterCount > 0)
+		{
+			Data.Reserve(CharacterCount + 1);
+		}	
 	}
 
 	/**
diff --git a/Engine/Source/Runtime/Core/Public/IO/IoDispatcher.h b/Engine/Source/Runtime/Core/Public/IO/IoDispatcher.h
--- a/Engine/Source/Runtime/Core/Public/IO/IoDispatcher.h
+++ b/Engine/Source/Runtime/Core/Public/IO/IoDispatcher.h
@@ -27,7 +27,10 @@ class FIoStoreEnvironment;
 class FIoRequestImpl;
 class FIoBatchImpl;
 class FIoDispatcherImpl;
+class FIoStoreWriterContextImpl;
 class FIoStoreWriterImpl;
+class IMappedFileHandle;
+class IMappedFileRegion;
 
 CORE_API DECLARE_LOG_CATEGORY_EXTERN(LogIoDispatcher, Log, All);
 
@@ -400,12 +403,24 @@ public:
 
 	inline void			SetSize(uint64 InSize)	{ return CorePtr->SetSize(InSize); }
 
-	inline bool			IsAvailable() const;
 	inline bool			IsMemoryOwned() const	{ return CorePtr->IsMemoryOwned(); }
 
 	inline void			EnsureOwned() const		{ if (!CorePtr->IsMemoryOwned()) { MakeOwned(); } }
 
 	CORE_API void		MakeOwned() const;
+	
+	/**
+	 * Relinquishes control of the internal buffer to the caller and removes it from the FIoBuffer.
+	 * This allows the caller to assume ownership of the internal data and prevent it from being deleted along with 
+	 * the FIoBuffer.
+	 *
+	 * NOTE: It is only valid to call this if the FIoBuffer currently owns the internal memory allocation, as the 
+	 * point of the call is to take ownership of it. If the FIoBuffer is only wrapping the allocation then it will
+	 * return a failed FIoStatus instead.
+	 *
+	 * @return A status wrapper around the memory pointer. Even if the status is valid the pointer might still be null.
+	 */
+	UE_NODISCARD CORE_API TIoStatusOr<uint8*> Release();
 
 private:
 	/** Core buffer object. For internal use only, used by FIoBuffer
@@ -436,6 +451,8 @@ private:
 
 		void	MakeOwned();
 
+		TIoStatusOr<uint8*> ReleaseMemory();
+
 		inline void SetIsOwned(bool InOwnsMemory)
 		{
 			if (InOwnsMemory)
@@ -473,7 +490,7 @@ private:
 			return uint32(NumRefs);
 		}
 
-		bool			IsMemoryOwned() const	{ return Flags & OwnsMemory; }
+		bool IsMemoryOwned() const	{ return Flags & OwnsMemory; }
 
 	private:
 		CORE_API void				CheckRefCount() const;
@@ -501,6 +518,11 @@ private:
 		};
 
 		void EnsureDataIsResident() {}
+
+		void ClearFlags()
+		{
+			Flags = 0;
+		}
 	};
 
 	// Reference-counted "core"
@@ -579,6 +601,7 @@ enum class EIoChunkType : uint8
 	ExportBundleData,
 	BulkData,
 	OptionalBulkData,
+	MemoryMappedBulkData,
 	LoaderGlobalMeta,
 	LoaderInitialLoadMeta,
 	LoaderGlobalNames,
@@ -686,6 +709,27 @@ private:
 
 //////////////////////////////////////////////////////////////////////////
 
+class FIoBatchReadOptions
+{
+public:
+	FIoBatchReadOptions() = default;
+
+	void SetTargetVa(void* InTargetVa)
+	{
+		TargetVa = InTargetVa;
+	}
+
+	void* GetTargetVa() const
+	{
+		return TargetVa;
+	}
+
+private:
+	void* TargetVa = nullptr;
+};
+
+//////////////////////////////////////////////////////////////////////////
+
 /**
   */
 class FIoRequest
@@ -714,6 +758,8 @@ private:
 	friend class FIoBatch;
 };
 
+using FIoReadCallback = TFunction<void(TIoStatusOr<FIoBuffer>)>;
+
 /** I/O batch
 
 	This is a primitive used to group I/O requests for synchronization
@@ -726,18 +772,49 @@ class FIoBatch
 	FIoBatch(FIoDispatcherImpl* InDispatcher, FIoBatchImpl* InImpl);
 
 public:
+	FIoBatch() = default;
+
+	CORE_API bool IsValid() const;
+
 	CORE_API FIoRequest Read(const FIoChunkId& Chunk, FIoReadOptions Options);
 
 	CORE_API void ForEachRequest(TFunction<bool(FIoRequest&)>&& Callback);
 
+	/**
+	 * Initiates the loading of the batch as individual requests.
+	 */
 	CORE_API void Issue();
+
+	/**
+	 * Initiates the loading of the batch to a single contiguous output buffer. The requests will be in the
+	 * same order that they were added to the FIoBatch.
+	 * NOTE: It is not valid to call this on a batch containing requests that have been given a TargetVa to 
+	 * read into as the requests are supposed to read into the batch's output buffer, doing so will cause the
+	 * method to return an error 'InvalidParameter'.
+	 *
+	 * @param Options A set of options allowing customization on how the load will work.
+	 * @param Callback An optional callback that will be triggered once the batch has finished loading. 
+	 * The batch's output buffer will be provided as the parameter of the callback.
+	 *
+	 * @return This methods had the capacity to fail so the return value should be checked.
+	 */
+	UE_NODISCARD CORE_API FIoStatus IssueWithCallback(FIoBatchReadOptions Options, FIoReadCallback&& Callback);
+	
 	CORE_API void Wait();
 	CORE_API void Cancel();
 
 private:
 	FIoDispatcherImpl*	Dispatcher		= nullptr;
 	FIoBatchImpl*		Impl			= nullptr;
-	FEvent*				CompletionEvent = nullptr;
+};
+
+/**
+ * Mapped region.
+ */
+struct FIoMappedRegion
+{
+	IMappedFileHandle* MappedFileHandle = nullptr;
+	IMappedFileRegion* MappedFileRegion = nullptr;
 };
 
 /** I/O dispatcher
@@ -751,9 +828,11 @@ public:
 	CORE_API FIoStatus				Mount(const FIoStoreEnvironment& Environment);
 
 	CORE_API FIoBatch				NewBatch();
-	CORE_API void					FreeBatch(FIoBatch Batch);
+	CORE_API void					FreeBatch(FIoBatch& Batch);
+
 
-	CORE_API void					ReadWithCallback(const FIoChunkId& Chunk, const FIoReadOptions& Options, TFunction<void(TIoStatusOr<FIoBuffer>)>&& Callback);
+	CORE_API void					ReadWithCallback(const FIoChunkId& ChunkId, const FIoReadOptions& Options, FIoReadCallback&& Callback);
+	CORE_API TIoStatusOr<FIoMappedRegion> OpenMapped(const FIoChunkId& ChunkId, const FIoReadOptions& Options);
 
 	// Polling methods
 	CORE_API bool					DoesChunkExist(const FIoChunkId& ChunkId) const;
@@ -796,6 +875,46 @@ private:
 
 //////////////////////////////////////////////////////////////////////////
 
+struct FIoStoreWriterSettings
+{
+	FName CompressionMethod = NAME_None;
+	int64 CompressionBlockSize = 0;
+	int64 CompressionBlockAlignment = 0;
+	bool bEnableCsvOutput = false;
+};
+
+struct FIoStoreWriterResult
+{
+	FString ContainerName;
+	int64 TocSize = 0;
+	int64 TocEntryCount = 0;
+	int64 PaddingSize = 0;
+	int64 UncompressedContainerSize = 0;
+	int64 CompressedContainerSize = 0;
+	FName CompressionMethod = NAME_None;
+};
+
+struct FIoWriteOptions
+{
+	const TCHAR* DebugName = nullptr;
+	int64 Alignment = 0;
+	bool bForceUncompressed = false;
+};
+
+class FIoStoreWriterContext
+{
+public:
+	CORE_API FIoStoreWriterContext();
+	CORE_API ~FIoStoreWriterContext();
+
+	UE_NODISCARD CORE_API FIoStatus Initialize(const FIoStoreWriterSettings& InWriterSettings);
+
+private:
+	friend class FIoStoreWriter;
+	
+	FIoStoreWriterContextImpl* Impl;
+};
+
 class FIoStoreWriter
 {
 public:
@@ -805,9 +924,9 @@ public:
 	FIoStoreWriter(const FIoStoreWriter&) = delete;
 	FIoStoreWriter& operator=(const FIoStoreWriter&) = delete;
 
-	CORE_API FIoStatus	Initialize();
-	CORE_API FIoStatus	EnableCsvOutput();
-	CORE_API FIoStatus	Append(FIoChunkId ChunkId, FIoBuffer Chunk, const TCHAR* Name);
+	UE_NODISCARD CORE_API FIoStatus	Initialize(const FIoStoreWriterContext& Context, bool bIsContainerCompressed);
+	UE_NODISCARD CORE_API FIoStatus	Append(FIoChunkId ChunkId, FIoBuffer Chunk, FIoWriteOptions WriteOptions);
+	UE_NODISCARD CORE_API FIoStatus	AppendPadding(uint64 Count);
 
 	/**
 	 * Creates an addressable range in an already mapped Chunk.
@@ -817,8 +936,8 @@ public:
 	 * @param Length The length of the range in bytes
 	 * @param ChunkIdPartialRange The FIoChunkId that will map to the range
 	 */
-	CORE_API FIoStatus	MapPartialRange(FIoChunkId OriginalChunkId, uint64 Offset, uint64 Length, FIoChunkId ChunkIdPartialRange);
-	CORE_API FIoStatus	FlushMetadata();
+	UE_NODISCARD CORE_API FIoStatus MapPartialRange(FIoChunkId OriginalChunkId, uint64 Offset, uint64 Length, FIoChunkId ChunkIdPartialRange);
+	UE_NODISCARD CORE_API TIoStatusOr<FIoStoreWriterResult> Flush();
 
 private:
 	FIoStoreWriterImpl*		Impl;
diff --git a/Engine/Source/Runtime/Core/Public/Misc/PathViews.h b/Engine/Source/Runtime/Core/Public/Misc/PathViews.h
--- a/Engine/Source/Runtime/Core/Public/Misc/PathViews.h
+++ b/Engine/Source/Runtime/Core/Public/Misc/PathViews.h
@@ -4,6 +4,7 @@
 
 #include "CoreTypes.h"
 #include "Containers/StringFwd.h"
+class FString;
 
 class CORE_API FPathViews
 {
@@ -133,4 +134,14 @@ public:
 	 * @param Suffix A possibly-empty suffix that does not start with a separator.
 	 */
 	static void Append(FStringBuilderBase& Builder, const FStringView& Suffix);
+
+	/** Replaces the pre-existing file extension of a filename. 
+	 *
+	 * @param InPath A valid file path with a pre-existing extension.
+	 * @param InNewExtension The new extension to use (prefixing with a '.' is optional)
+	 *
+	 * @return The new file path complete with the new extension unless InPath is not valid in which
+	 * case a copy of InPath will be returned instead.
+	 */
+	static FString ChangeExtension(const FStringView& InPath, const FStringView& InNewExtension);
 };
diff --git a/Engine/Source/Runtime/CoreUObject/Private/Serialization/AsyncLoading.cpp b/Engine/Source/Runtime/CoreUObject/Private/Serialization/AsyncLoading.cpp
--- a/Engine/Source/Runtime/CoreUObject/Private/Serialization/AsyncLoading.cpp
+++ b/Engine/Source/Runtime/CoreUObject/Private/Serialization/AsyncLoading.cpp
@@ -148,30 +148,6 @@ static FAutoConsoleVariableRef CFlushStreamingOnExit(
 	ECVF_Default
 );
 
-static int32 GWarnIfTimeLimitExceeded = 0;
-static FAutoConsoleVariableRef CVarWarnIfTimeLimitExceeded(
-	TEXT("s.WarnIfTimeLimitExceeded"),
-	GWarnIfTimeLimitExceeded,
-	TEXT("Enables log warning if time limit for time-sliced package streaming has been exceeded."),
-	ECVF_Default
-	);
-
-static float GTimeLimitExceededMultiplier = 1.5f;
-static FAutoConsoleVariableRef CVarTimeLimitExceededMultiplier(
-	TEXT("s.TimeLimitExceededMultiplier"),
-	GTimeLimitExceededMultiplier,
-	TEXT("Multiplier for time limit exceeded warning time threshold."),
-	ECVF_Default
-	);
-
-static float GTimeLimitExceededMinTime = 0.005f;
-static FAutoConsoleVariableRef CVarTimeLimitExceededMinTime(
-	TEXT("s.TimeLimitExceededMinTime"),
-	GTimeLimitExceededMinTime,
-	TEXT("Minimum time the time limit exceeded warning will be triggered by."),
-	ECVF_Default
-	);
-
 static int32 GEventDrivenLoaderEnabledInCookedBuilds = 0;
 static FAutoConsoleVariableRef CVarEventDrivenLoaderEnabled(
 	TEXT("s.EventDrivenLoaderEnabled"),
@@ -194,7 +170,6 @@ static FAutoConsoleVariableRef CVar_MaxPrecacheRequestsInFlight(
 	TEXT("Controls the maximum amount of precache requests to have in flight.")
 );
 
-
 int32 GMaxIncomingRequestsToStall = 100;
 static FAutoConsoleVariableRef CVar_MaxIncomingRequestsToStall(
 	TEXT("s.MaxIncomingRequestsToStall"),
@@ -270,28 +245,6 @@ uint32 FAsyncLoadingThread::AsyncLoadingThreadID = 0;
 FThreadSafeCounter FAsyncLoadingThread::AsyncLoadingTickCounter;
 int32 FAsyncLoadingThread::CurrentAsyncLoadingTickThreadIndex = INDEX_NONE;
 
-static void IsTimeLimitExceededPrint(double InTickStartTime, double CurrentTime, double LastTestTime, float InTimeLimit, const TCHAR* InLastTypeOfWorkPerformed = nullptr, UObject* InLastObjectWorkWasPerformedOn = nullptr)
-{
-	static double LastPrintStartTime = -1.0;
-	// Log single operations that take longer than time limit (but only in cooked builds)
-	if (LastPrintStartTime != InTickStartTime &&
-		(CurrentTime - InTickStartTime) > GTimeLimitExceededMinTime &&
-		(CurrentTime - InTickStartTime) > (GTimeLimitExceededMultiplier * InTimeLimit))
-	{
-		float EstimatedTimeForThisStep = (CurrentTime - InTickStartTime) * 1000;
-		if (LastTestTime > InTickStartTime)
-		{
-			EstimatedTimeForThisStep = (CurrentTime - LastTestTime) * 1000;
-		}
-		LastPrintStartTime = InTickStartTime;
-		UE_LOG(LogStreaming, Warning, TEXT("IsTimeLimitExceeded: %s %s Load Time %5.2fms   Last Step Time %5.2fms"),
-			InLastTypeOfWorkPerformed ? InLastTypeOfWorkPerformed : TEXT("unknown"),
-			InLastObjectWorkWasPerformedOn ? *InLastObjectWorkWasPerformedOn->GetFullName() : TEXT("nullptr"),
-			(CurrentTime - InTickStartTime) * 1000,
-			EstimatedTimeForThisStep);
-	}
-}
-
 static FORCEINLINE bool IsTimeLimitExceeded(double InTickStartTime, bool bUseTimeLimit, float InTimeLimit, const TCHAR* InLastTypeOfWorkPerformed = nullptr, UObject* InLastObjectWorkWasPerformedOn = nullptr)
 {
 	static double LastTestTime = -1.0;
@@ -2032,7 +1985,12 @@ EAsyncPackageState::Type FAsyncPackage::SetupImports_Event()
 				if (GIsInitialLoad && !ImportLinker && ImportPackage->HasAnyPackageFlags(PKG_CompiledIn) && !bDynamicImport)
 				{
 					// OuterMostNonPackageIndex is used here because if it is a CDO or subobject, etc, we wait for the outermost thing that is not a package
-					bFireIfNoArcsAdded = !EDLBootNotificationManager.AddWaitingPackage(this, OuterMostImport.ObjectName, Linker->Imp(OuterMostNonPackageIndex).ObjectName, FPackageIndex::FromImport(LocalImportIndex));
+					bFireIfNoArcsAdded = !EDLBootNotificationManager.AddWaitingPackage(
+						this,
+						OuterMostImport.ObjectName,
+						Linker->Imp(OuterMostNonPackageIndex).ObjectName,
+						FPackageIndex::FromImport(LocalImportIndex),
+						/*bIgnoreMissingPackage*/ false);
 				}
 #endif
 				if (bFireIfNoArcsAdded && // if bFireIfNoArcsAdded is false, then we know we are waiting for compiled in thing, so don't bother to look for the import now
