Fixed static analysis warnings
#rnx

diff --git a/Engine/Source/Runtime/Core/Private/IO/IoStore.cpp b/Engine/Source/Runtime/Core/Private/IO/IoStore.cpp
--- a/Engine/Source/Runtime/Core/Private/IO/IoStore.cpp
+++ b/Engine/Source/Runtime/Core/Private/IO/IoStore.cpp
@@ -211,7 +211,7 @@ public:
 						check(FileHandle->Tell() % CurrentWriteQueueItem->Alignment == 0);
 					}
 					const int64 CompressedFileOffset = FileHandle->Tell();
-					FIoStoreCompressedBlockEntry CompressedBlockEntry;
+					FIoStoreCompressedBlockEntry& CompressedBlockEntry = CompressionInfo.BlockEntries.AddDefaulted_GetRef();
 					CompressedBlockEntry.OffsetAndLength.SetOffset(CompressedFileOffset);
 					const uint8* SourceData;
 					uint64 SourceDataSize;
@@ -229,7 +229,6 @@ public:
 					FIoBuffer& SourceBuffer = CurrentWriteQueueItem->CompressionMethod == NAME_None ? CurrentWriteQueueItem->UncompressedData : CurrentWriteQueueItem->CompressedData;
 					CompressedBlockEntry.OffsetAndLength.SetLength(SourceDataSize);
 					CompressedBlockEntry.CompressionMethodIndex = CompressionInfo.GetCompressionMethodIndex(CurrentWriteQueueItem->CompressionMethod);
-					CompressionInfo.BlockEntries.Emplace(CompressedBlockEntry);
 					{
 						TRACE_CPUPROFILER_EVENT_SCOPE(WriteBlockToFile);
 						FileHandle->Write(SourceData, SourceDataSize);
@@ -284,6 +283,7 @@ public:
 			if (!CurrentBlock)
 			{
 				CurrentBlock = WriterContext->AllocBlock();
+				check(CurrentBlock);
 				if (RemainingBytesInChunk >= CurrentBlock->UncompressedData.DataSize())
 				{
 					// Entire block belongs to this chunk, skip patch alignment
@@ -315,7 +315,7 @@ public:
 			{
 				CurrentBlock = WriterContext->AllocBlock();
 			}
-			check(CurrentBlock->UncompressedData.DataSize() > CurrentBlockOffset);
+			check(CurrentBlock && CurrentBlock->UncompressedData.DataSize() > CurrentBlockOffset);
 			uint64 BytesToWrite = FMath::Min(Count, CurrentBlock->UncompressedData.DataSize() - CurrentBlockOffset);
 			FMemory::Memzero(CurrentBlock->UncompressedData.Data() + CurrentBlockOffset, BytesToWrite);
 			check(Count >= BytesToWrite);
diff --git a/Engine/Source/Runtime/CoreUObject/Public/Serialization/AsyncLoading2.h b/Engine/Source/Runtime/CoreUObject/Public/Serialization/AsyncLoading2.h
--- a/Engine/Source/Runtime/CoreUObject/Public/Serialization/AsyncLoading2.h
+++ b/Engine/Source/Runtime/CoreUObject/Public/Serialization/AsyncLoading2.h
@@ -80,7 +80,7 @@ public:
 
 	inline bool IsImport() const
 	{
-		return IsScriptImport() | IsPackageImport();
+		return IsScriptImport() || IsPackageImport();
 	}
 
 	inline bool IsScriptImport() const
