fix GetFreeMemoryMB to properly utilize the FPlatformMemory::GetStats method so that we get consistent results

diff --git a/Engine/Source/Runtime/Core/Private/IOS/IOSPlatformMisc.cpp b/Engine/Source/Runtime/Core/Private/IOS/IOSPlatformMisc.cpp
--- a/Engine/Source/Runtime/Core/Private/IOS/IOSPlatformMisc.cpp
+++ b/Engine/Source/Runtime/Core/Private/IOS/IOSPlatformMisc.cpp
@@ -74,15 +74,8 @@ float GOriginalBrightness = -1.0f;
 
 static int32 GetFreeMemoryMB()
 {
-	// get free memory
-	vm_size_t PageSize;
-	host_page_size(mach_host_self(), &PageSize);
-
-	// get memory stats
-	vm_statistics Stats;
-	mach_msg_type_number_t StatsSize = sizeof(Stats);
-	host_statistics(mach_host_self(), HOST_VM_INFO, (host_info_t)&Stats, &StatsSize);
-	return ((Stats.free_count + Stats.inactive_count) * PageSize) / 1024 / 1024;
+	FPlatformMemoryStats MemoryStats = FPlatformMemory::GetStats();
+	return MemoryStats.AvailablePhysical / 1024 / 1024;
 }
 
 void FIOSPlatformMisc::PlatformInit()
