Editgrate fixes from Dev-Rendering

- Add static cleanup function for platform specific files which are used by the CrashContext
- Move CrashVideo.avi to platform specific cleanup function
- ReportGPUCrash on Windows raised exception via Try/Catch because function can be called when the guarded entry is not set

diff --git a/Engine/Source/Runtime/Core/Private/Android/AndroidPlatformCrashContext.cpp b/Engine/Source/Runtime/Core/Private/Android/AndroidPlatformCrashContext.cpp
--- a/Engine/Source/Runtime/Core/Private/Android/AndroidPlatformCrashContext.cpp
+++ b/Engine/Source/Runtime/Core/Private/Android/AndroidPlatformCrashContext.cpp
@@ -27,6 +27,11 @@ extern FString AndroidThunkCpp_GetMetaDataString(const FString& Key);
 
 extern FString AndroidRelativeToAbsolutePath(bool bUseInternalBasePath, FString RelPath);
 
+/** Implement platform specific static cleanup function */
+void FGenericCrashContext::CleanupPlatformSpecificFiles()
+{
+}
+
 struct FAndroidCrashInfo
 {
 	FAndroidCrashInfo()
diff --git a/Engine/Source/Runtime/Core/Private/Apple/ApplePlatformCrashContext.cpp b/Engine/Source/Runtime/Core/Private/Apple/ApplePlatformCrashContext.cpp
--- a/Engine/Source/Runtime/Core/Private/Apple/ApplePlatformCrashContext.cpp
+++ b/Engine/Source/Runtime/Core/Private/Apple/ApplePlatformCrashContext.cpp
@@ -9,6 +9,14 @@
 #include "Containers/StringConv.h"
 #include "CoreGlobals.h"
 
+/** Implement platform specific static cleanup function */
+void FGenericCrashContext::CleanupPlatformSpecificFiles()
+{
+	// Is this the correct way to cleanup on Apple?
+	FString CrashVideoPath = FPaths::ProjectLogDir() + TEXT("CrashVideo.avi");
+	IFileManager::Get().Delete(*CrashVideoPath);
+}
+
 FApplePlatformCrashContext::FApplePlatformCrashContext(ECrashContextType InType, const TCHAR* InErrorMessage)
 :	FGenericCrashContext(InType, InErrorMessage)
 ,	Signal(0)
diff --git a/Engine/Source/Runtime/Core/Private/GenericPlatform/GenericPlatformCrashContext.cpp b/Engine/Source/Runtime/Core/Private/GenericPlatform/GenericPlatformCrashContext.cpp
--- a/Engine/Source/Runtime/Core/Private/GenericPlatform/GenericPlatformCrashContext.cpp
+++ b/Engine/Source/Runtime/Core/Private/GenericPlatform/GenericPlatformCrashContext.cpp
@@ -242,6 +242,8 @@ PRAGMA_ENABLE_DEPRECATION_WARNINGS
 
 	SerializeTempCrashContextToFile();
 
+	CleanupPlatformSpecificFiles();
+
 	bIsInitialized = true;
 #endif	// !NOINITCRASHREPORTER
 }
diff --git a/Engine/Source/Runtime/Core/Private/HoloLens/HoloLensPlatformCrashContext.cpp b/Engine/Source/Runtime/Core/Private/HoloLens/HoloLensPlatformCrashContext.cpp
--- a/Engine/Source/Runtime/Core/Private/HoloLens/HoloLensPlatformCrashContext.cpp
+++ b/Engine/Source/Runtime/Core/Private/HoloLens/HoloLensPlatformCrashContext.cpp
@@ -459,6 +459,13 @@ void ReportHang(const TCHAR* ErrorMessage, const uint64* StackFrames, int32 NumS
 
 #if WER_CUSTOM_REPORTS
 
+/** Implement platform specific static cleanup function */
+void FGenericCrashContext::CleanupPlatformSpecificFiles()
+{
+	FString CrashVideoPath = FPaths::GameLogDir() / TEXT("CrashVideo.avi");
+	IFileManager::Get().Delete(*CrashVideoPath);
+}
+
 static FCriticalSection EnsureLock;
 static bool bReentranceGuard = false;
 
diff --git a/Engine/Source/Runtime/Core/Private/Unix/UnixPlatformCrashContext.cpp b/Engine/Source/Runtime/Core/Private/Unix/UnixPlatformCrashContext.cpp
--- a/Engine/Source/Runtime/Core/Private/Unix/UnixPlatformCrashContext.cpp
+++ b/Engine/Source/Runtime/Core/Private/Unix/UnixPlatformCrashContext.cpp
@@ -84,6 +84,11 @@ FString DescribeSignal(int32 Signal, siginfo_t* Info, ucontext_t *Context)
 #undef HANDLE_CASE
 }
 
+/** Implement platform specific static cleanup function */
+void FGenericCrashContext::CleanupPlatformSpecificFiles()
+{
+}
+
 __thread siginfo_t FUnixCrashContext::FakeSiginfoForEnsures;
 
 FUnixCrashContext::~FUnixCrashContext()
diff --git a/Engine/Source/Runtime/Core/Private/Windows/WindowsPlatformCrashContext.cpp b/Engine/Source/Runtime/Core/Private/Windows/WindowsPlatformCrashContext.cpp
--- a/Engine/Source/Runtime/Core/Private/Windows/WindowsPlatformCrashContext.cpp
+++ b/Engine/Source/Runtime/Core/Private/Windows/WindowsPlatformCrashContext.cpp
@@ -137,6 +137,17 @@ struct FAssertInfo
 	}
 };
 
+/**
+* Implement platform specific static cleanup function
+*/
+void FGenericCrashContext::CleanupPlatformSpecificFiles()
+{
+	// Manually delete any potential leftover gpu dumps because the crash reporter will upload any leftover crash data from last session
+	const FString CrashVideoPath = FPaths::ProjectLogDir() + TEXT("CrashVideo.avi");
+	IFileManager::Get().Delete(*CrashVideoPath);
+}
+
+
 void FWindowsPlatformCrashContext::GetProcModuleHandles(const FProcHandle& ProcessHandle, FModuleHandleArray& OutHandles)
 {
 	// Get all the module handles for the current process. Each module handle is its base address.
@@ -1397,11 +1408,22 @@ FORCENOINLINE void ReportGPUCrash(const TCHAR* ErrorMessage, int NumStackFramesT
 {
 	/** This is the last place to gather memory stats before exception. */
 	FGenericCrashContext::SetMemoryStats(FPlatformMemory::GetStats());
+	
+	// GPUCrash can be called when the guarded entry is not set
+#if !PLATFORM_SEH_EXCEPTIONS_DISABLED
+	__try
+	{
+		FAssertInfo Info(ErrorMessage, NumStackFramesToIgnore + 2); // +2 for this function and RaiseException()
 
-	FAssertInfo Info(ErrorMessage, NumStackFramesToIgnore + 2); // +2 for this function and RaiseException()
+		ULONG_PTR Arguments[] = { (ULONG_PTR)&Info };
+		::RaiseException(GPUCrashExceptionCode, 0, UE_ARRAY_COUNT(Arguments), Arguments);
+	}
+	__except (ReportCrash(GetExceptionInformation()))
+	{
+		FPlatformMisc::RequestExit(false);
+	}
+#endif
 
-	ULONG_PTR Arguments[] = { (ULONG_PTR)&Info };
-	::RaiseException(GPUCrashExceptionCode, 0, UE_ARRAY_COUNT(Arguments), Arguments);
 }
 
 void ReportHang(const TCHAR* ErrorMessage, const uint64* StackFrames, int32 NumStackFrames, uint32 HungThreadId)
diff --git a/Engine/Source/Runtime/Core/Public/GenericPlatform/GenericPlatformCrashContext.h b/Engine/Source/Runtime/Core/Public/GenericPlatform/GenericPlatformCrashContext.h
--- a/Engine/Source/Runtime/Core/Public/GenericPlatform/GenericPlatformCrashContext.h
+++ b/Engine/Source/Runtime/Core/Public/GenericPlatform/GenericPlatformCrashContext.h
@@ -402,6 +402,9 @@ public:
 	/** Allows platform implementations to copy files to report directory. */
 	virtual void CopyPlatformSpecificFiles(const TCHAR* OutputDirectory, void* Context);
 
+	/** Cleanup platform specific files - called on startup, implemented per platform */
+	static void CleanupPlatformSpecificFiles();
+
 	/**
 	 * @return whether this crash is a non-crash event
 	 */
diff --git a/Engine/Source/Runtime/Engine/Private/UnrealEngine.cpp b/Engine/Source/Runtime/Engine/Private/UnrealEngine.cpp
--- a/Engine/Source/Runtime/Engine/Private/UnrealEngine.cpp
+++ b/Engine/Source/Runtime/Engine/Private/UnrealEngine.cpp
@@ -1631,12 +1631,7 @@ void UEngine::Init(IEngineLoop* InEngineLoop)
 	{
 		GetDerivedDataCacheRef().NotifyBootComplete();
 	}
-
-	// Manually delete any potential leftover crash videos in case we can't access the module
-	// because the crash reporter will upload any leftover crash video from last session
-	FString CrashVideoPath = FPaths::ProjectLogDir() + TEXT("CrashVideo.avi");
-	IFileManager::Get().Delete(*CrashVideoPath);
-
+	
 	// register the engine with the travel and network failure broadcasts
 	// games can override these to provide proper behavior in each error case
 	OnTravelFailure().AddUObject(this, &UEngine::HandleTravelFailure);
