Fix for compiler errors due to 3rd party PLCrashReporter libraries not having header guards and iOS now including them in multiple places.
#tests iOS Client
#swarm https://p4-swarm.epicgames.net/reviews/12508562/

diff --git a/Engine/Source/Runtime/Core/Private/IOS/IOSPlatformCrashContext.cpp b/Engine/Source/Runtime/Core/Private/IOS/IOSPlatformCrashContext.cpp
--- a/Engine/Source/Runtime/Core/Private/IOS/IOSPlatformCrashContext.cpp
+++ b/Engine/Source/Runtime/Core/Private/IOS/IOSPlatformCrashContext.cpp
@@ -5,18 +5,14 @@
 =============================================================================*/
 #include "IOS/IOSPlatformCrashContext.h"
 
+#include "IOS/IOSPlatformPLCrashReporterIncludes.h"
+
 #include "HAL/ExceptionHandling.h"
 #include "HAL/PlatformOutputDevices.h"
 
 #include "Misc/App.h"
 #include "Misc/EngineVersion.h"
 
-#if !PLATFORM_TVOS
-#include "PLCrashReporter.h"
-#include "PLCrashReport.h"
-#include "PLCrashReportTextFormatter.h"
-#endif
-
 
 FIOSApplicationInfo GIOSAppInfo;
 
diff --git a/Engine/Source/Runtime/Core/Private/IOS/IOSPlatformMisc.cpp b/Engine/Source/Runtime/Core/Private/IOS/IOSPlatformMisc.cpp
--- a/Engine/Source/Runtime/Core/Private/IOS/IOSPlatformMisc.cpp
+++ b/Engine/Source/Runtime/Core/Private/IOS/IOSPlatformMisc.cpp
@@ -21,11 +21,7 @@
 #include "Misc/ScopeExit.h"
 #include "Apple/ApplePlatformCrashContext.h"
 #include "IOS/IOSPlatformCrashContext.h"
-#if !PLATFORM_TVOS
-#include "PLCrashReporter.h"
-#include "PLCrashReport.h"
-#include "PLCrashReportTextFormatter.h"
-#endif
+#include "IOS/IOSPlatformPLCrashReporterIncludes.h"
 #include "HAL/FileManager.h"
 #include "HAL/PlatformOutputDevices.h"
 #include "Misc/OutputDeviceError.h"
diff --git a/Engine/Source/Runtime/Core/Public/IOS/IOSPlatformPLCrashReporterIncludes.h b/Engine/Source/Runtime/Core/Public/IOS/IOSPlatformPLCrashReporterIncludes.h
new file mode 100644
--- /dev/null
+++ b/Engine/Source/Runtime/Core/Public/IOS/IOSPlatformPLCrashReporterIncludes.h
@@ -0,0 +1,15 @@
+// Copyright Epic Games, Inc. All Rights Reserved.
+
+/*=============================================================================================
+	IOSPlatformPLCrashReporterIncludes.h: Wrapper for the third party PLCrashReporter includes for IOS
+==============================================================================================*/
+
+#pragma once
+
+THIRD_PARTY_INCLUDES_START
+#if !PLATFORM_TVOS
+#include "PLCrashReporter.h"
+#include "PLCrashReport.h"
+#include "PLCrashReportTextFormatter.h"
+#endif
+THIRD_PARTY_INCLUDES_END
