Fix for AppleLLM initializing using an invalid calculation for number of stats leading to us never registering ApplePlatformLLM stats.
#tests iOS Client

diff --git a/Engine/Source/Runtime/Core/Private/Apple/AppleLLM.cpp b/Engine/Source/Runtime/Core/Private/Apple/AppleLLM.cpp
--- a/Engine/Source/Runtime/Core/Private/Apple/AppleLLM.cpp
+++ b/Engine/Source/Runtime/Core/Private/Apple/AppleLLM.cpp
@@ -93,7 +93,7 @@ void NSDeallocateObjectInterposer(id Obj)
  */
 void AppleLLM::Initialise()
 {
-	int32 TagCount = sizeof(ELLMTagApple) / sizeof(FLLMTagInfoApple);
+	int32 TagCount = sizeof(ELLMTagNamesApple) / sizeof(FLLMTagInfoApple);
 
 	for (int32 Index = 0; Index < TagCount; ++Index)
 	{
