switch to OS values for determining total memory available.

diff --git a/Engine/Source/Runtime/Core/Private/Apple/ApplePlatformMemory.cpp b/Engine/Source/Runtime/Core/Private/Apple/ApplePlatformMemory.cpp
--- a/Engine/Source/Runtime/Core/Private/Apple/ApplePlatformMemory.cpp
+++ b/Engine/Source/Runtime/Core/Private/Apple/ApplePlatformMemory.cpp
@@ -285,7 +285,18 @@ FPlatformMemoryStats FApplePlatformMemory::GetStats()
 	mach_task_basic_info_data_t TaskInfo;
 	mach_msg_type_number_t TaskInfoCount = MACH_TASK_BASIC_INFO_COUNT;
 	task_info( mach_task_self(), MACH_TASK_BASIC_INFO, (task_info_t)&TaskInfo, &TaskInfoCount );
-	MemoryStats.UsedPhysical = TaskInfo.resident_size;
+#if PLATFORM_IOS
+#if defined(__IPHONE_13_0)
+	if (@available(iOS 13.0,*))
+	{
+		MemoryStats.UsedPhysical = MemoryConstants.TotalPhysical - FreeMem;
+	}
+	else
+#endif
+#endif
+	{
+		MemoryStats.UsedPhysical = TaskInfo.resident_size;
+	}
 	if(MemoryStats.UsedPhysical > MemoryStats.PeakUsedPhysical)
 	{
 		MemoryStats.PeakUsedPhysical = MemoryStats.UsedPhysical;
@@ -313,9 +324,43 @@ const FPlatformMemoryConstants& FApplePlatformMemory::GetConstants()
 		
 		// Get memory.
 		int64 AvailablePhysical = 0;
-		int Mib[] = {CTL_HW, HW_MEMSIZE};
-		size_t Length = sizeof(int64);
-		sysctl(Mib, 2, &AvailablePhysical, &Length, NULL, 0);
+#if PLATFORM_IOS
+#if defined(__IPHONE_13_0)
+		if (@available(iOS 13.0,*))
+		{
+			AvailablePhysical = os_proc_available_memory();
+			
+			// quantize to the known jetsam limits, we should be within 50MB of the correct one
+			uint64 JetsamLimits[] = { 1520435200, 1939865600, 2201170740, 2252710350, 3006477100 }; // { 2GB, gimped 3GB, gimped 4GB, 3GB, 4GB
+			if (AvailablePhysical < JetsamLimits[0])
+			{
+				AvailablePhysical = JetsamLimits[0];
+			}
+			else if (AvailablePhysical < JetsamLimits[1])
+			{
+				AvailablePhysical = JetsamLimits[1];
+			}
+			else if (AvailablePhysical < JetsamLimits[2])
+			{
+				AvailablePhysical = JetsamLimits[2];
+			}
+			else if (AvailablePhysical < JetsamLimits[3])
+			{
+				AvailablePhysical = JetsamLimits[3];
+			}
+			else if (AvailablePhysical < JetsamLimits[4])
+			{
+				AvailablePhysical = JetsamLimits[4];
+			}
+		}
+		else
+#endif
+#endif
+		{
+			int Mib[] = {CTL_HW, HW_MEMSIZE};
+			size_t Length = sizeof(int64);
+			sysctl(Mib, 2, &AvailablePhysical, &Length, NULL, 0);
+		}
 		
 		MemoryConstants.TotalPhysical = AvailablePhysical;
 		MemoryConstants.TotalVirtual = AvailablePhysical;
