UE4.25 - Hololens compile fix
#rnx

diff --git a/Engine/Source/Runtime/Core/Private/HoloLens/HoloLensPlatformCrashContext.cpp b/Engine/Source/Runtime/Core/Private/HoloLens/HoloLensPlatformCrashContext.cpp
--- a/Engine/Source/Runtime/Core/Private/HoloLens/HoloLensPlatformCrashContext.cpp
+++ b/Engine/Source/Runtime/Core/Private/HoloLens/HoloLensPlatformCrashContext.cpp
@@ -457,15 +457,15 @@ void ReportHang(const TCHAR* ErrorMessage, const uint64* StackFrames, int32 NumS
 }
 #endif
 
-#if WER_CUSTOM_REPORTS
-
 /** Implement platform specific static cleanup function */
 void FGenericCrashContext::CleanupPlatformSpecificFiles()
 {
-	FString CrashVideoPath = FPaths::GameLogDir() / TEXT("CrashVideo.avi");
+	FString CrashVideoPath = FPaths::ProjectLogDir() / TEXT("CrashVideo.avi");
 	IFileManager::Get().Delete(*CrashVideoPath);
 }
 
+#if WER_CUSTOM_REPORTS
+
 static FCriticalSection EnsureLock;
 static bool bReentranceGuard = false;
 
