Fix build for Apple platforms.

#rnx

diff --git a/Engine/Source/Runtime/Core/Private/Apple/ApplePlatformCrashContext.cpp b/Engine/Source/Runtime/Core/Private/Apple/ApplePlatformCrashContext.cpp
--- a/Engine/Source/Runtime/Core/Private/Apple/ApplePlatformCrashContext.cpp
+++ b/Engine/Source/Runtime/Core/Private/Apple/ApplePlatformCrashContext.cpp
@@ -7,6 +7,8 @@
 #include "Apple/ApplePlatformCrashContext.h"
 #include "Apple/ApplePlatformSymbolication.h"
 #include "Containers/StringConv.h"
+#include "HAL/FileManager.h"
+#include "Misc/Paths.h"
 #include "CoreGlobals.h"
 
 /** Implement platform specific static cleanup function */
