re-organize IsSupported API to pull game specific code out of the engine and clean up some of our platform #ifs

diff --git a/Engine/Source/Runtime/Core/Private/IOS/IOSPlatformMisc.cpp b/Engine/Source/Runtime/Core/Private/IOS/IOSPlatformMisc.cpp
--- a/Engine/Source/Runtime/Core/Private/IOS/IOSPlatformMisc.cpp
+++ b/Engine/Source/Runtime/Core/Private/IOS/IOSPlatformMisc.cpp
@@ -1312,15 +1312,6 @@ FString FIOSPlatformMisc::GetProjectVersion()
 	return localVersionString;
 }
 
-bool FIOSPlatformMisc::IsSupportedIOSVersion()
-{
-#if defined(IOS_MIN_SUPPORTED_VERSION_MAJOR) && defined(IOS_MIN_SUPPORTED_VERSION_MINOR) && defined(IOS_MIN_SUPPORTED_VERSION_PATCH)
-	return IOSVersionCompare(IOS_MIN_SUPPORTED_VERSION_MAJOR, IOS_MIN_SUPPORTED_VERSION_MINOR, IOS_MIN_SUPPORTED_VERSION_PATCH) >= 0;
-#else
-	return true;
-#endif
-}
-
 bool FIOSPlatformMisc::RequestDeviceCheckToken(TFunction<void(const TArray<uint8>&)> QuerySucceededFunc, TFunction<void(const FString&, const FString&)> QueryFailedFunc)
 {
 	DCDevice* DeviceCheckDevice = [DCDevice currentDevice];
diff --git a/Engine/Source/Runtime/Core/Public/IOS/IOSPlatformMisc.h b/Engine/Source/Runtime/Core/Public/IOS/IOSPlatformMisc.h
--- a/Engine/Source/Runtime/Core/Public/IOS/IOSPlatformMisc.h
+++ b/Engine/Source/Runtime/Core/Public/IOS/IOSPlatformMisc.h
@@ -105,8 +105,6 @@ struct CORE_API FIOSPlatformMisc : public FApplePlatformMisc
 	static void RequestStoreReview();
 
 	static bool IsUpdateAvailable();
-
-	static bool IsSupportedIOSVersion();
 	
 	// Possible iOS devices
 	enum EIOSDevice
