Changed render thread task graph execution to use a polling configuration vs task wakeup triggers.

The motivation here is to reduce the overhead of ENQUEUE_RENDER_COMMAND. A good chunk of the time was spent in FEvent::Trigger to wake up the render thread to process a small amount of work. The set polling rate is 1ms.

With this change, tests show a 0.2ms decrease in GT avg (tested on switch and ps4). RT avg increased slightly (50us) on ps4, and was unchanged on switch.

Further testing to monitor impact on all platforms will continue. Especially on high end pc configurations with respect to low latency.

To disable, set TaskGraph.DoRenderThreadWakeupTrigger=true

#rnx

diff --git a/Engine/Source/Runtime/Core/Private/Async/TaskGraph.cpp b/Engine/Source/Runtime/Core/Private/Async/TaskGraph.cpp
--- a/Engine/Source/Runtime/Core/Private/Async/TaskGraph.cpp
+++ b/Engine/Source/Runtime/Core/Private/Async/TaskGraph.cpp
@@ -61,6 +61,20 @@ namespace ENamedThreads
 	CORE_API int32 bHasHighPriorityThreads = CREATE_HIPRI_TASK_THREADS;
 }
 
+static bool GDoRenderThreadWakeupTrigger = false;
+static FAutoConsoleVariableRef CVarDoRenderThreadWakeupTrigger(
+	TEXT("TaskGraph.DoRenderThreadWakeupTrigger"),
+	GDoRenderThreadWakeupTrigger,
+	TEXT("If true, task graph tasks sent to the render thread trigger a wakeup. See TaskGraph.RenderThreadPollPeriodMs.")
+);
+
+static int32 GRenderThreadPollPeriodMs = 1;
+static FAutoConsoleVariableRef CVarRenderThreadPollPeriodMs(
+	TEXT("TaskGraph.RenderThreadPollPeriodMs"),
+	GRenderThreadPollPeriodMs,
+	TEXT("Render thread polling period in milliseconds.")
+);
+
 static int32 GIgnoreThreadToDoGatherOn = 0;
 static FAutoConsoleVariableRef CVarIgnoreThreadToDoGatherOn(
 	TEXT("TaskGraph.IgnoreThreadToDoGatherOn"),
@@ -683,7 +697,7 @@ public:
 				{
 					{
 						FScopeCycleCounter Scope(StallStatId);
-						Queue(QueueIndex).StallRestartEvent->Wait(MAX_uint32, bCountAsStall);
+						Queue(QueueIndex).StallRestartEvent->Wait(IsInRenderingThread() ? GRenderThreadPollPeriodMs : MAX_uint32, bCountAsStall);
 						if (Queue(QueueIndex).QuitForShutdown)
 						{
 							return ProcessedTasks;
@@ -762,10 +776,13 @@ public:
 
 		if (ThreadToStart >= 0)
 		{
-			QUICK_SCOPE_CYCLE_COUNTER(STAT_TaskGraph_EnqueueFromOtherThread_Trigger);
 			checkThreadGraph(ThreadToStart == 0);
-			TASKGRAPH_SCOPE_CYCLE_COUNTER(1, STAT_TaskGraph_EnqueueFromOtherThread_Trigger);
-			Queue(QueueIndex).StallRestartEvent->Trigger();
+			if ((ENamedThreads::GetThreadIndex(Task->ThreadToExecuteOn) != ENamedThreads::GetRenderThread()) || GDoRenderThreadWakeupTrigger)
+			{
+				QUICK_SCOPE_CYCLE_COUNTER(STAT_TaskGraph_EnqueueFromOtherThread_Trigger);
+				TASKGRAPH_SCOPE_CYCLE_COUNTER(1, STAT_TaskGraph_EnqueueFromOtherThread_Trigger);
+				Queue(QueueIndex).StallRestartEvent->Trigger();
+			}
 			return true;
 		}
 		return false;
