#Bringing shelf CL: 12582662 from Dev-Anim for ISPC math issue with negative scales + re-enabling ISPC for FillComponentSpaceTransform in skeletal mesh components
#rnx

diff --git a/Engine/Source/Runtime/Core/Public/Math/Matrix.isph b/Engine/Source/Runtime/Core/Public/Math/Matrix.isph
--- a/Engine/Source/Runtime/Core/Public/Math/Matrix.isph
+++ b/Engine/Source/Runtime/Core/Public/Math/Matrix.isph
@@ -108,41 +108,21 @@ inline FMatrix operator*(const FMatrix &A, const uniform FMatrix &B)
 inline void MatrixRemoveScaling(uniform FMatrix &M)
 {
 	// For each row, find magnitude, and if its non-zero re-scale so its unit length.
-	varying float S0, S1, S2, SquareSum = 0.0f;
-	*((uniform FVector *uniform)&S0) = *((uniform FVector *uniform)&M.M[0]);
-	*((uniform FVector *uniform)&S1) = *((uniform FVector *uniform)&M.M[4]);
-	*((uniform FVector *uniform)&S2) = *((uniform FVector *uniform)&M.M[8]);
-
-#if EXPLICIT_VECTOR == 1
-	// Using explicit vector aligns to 16 bytes. Save the 4th components when pulling from Matrix type
-	const uniform float SavedS0 = extract(S0, 3);
-	S0 = insert(S0, 3, 0.0f);
-	const uniform float SavedS1 = extract(S1, 3);
-	S1 = insert(S1, 3, 0.0f);
-	const uniform float SavedS2 = extract(S2, 3);
-	S2 = insert(S2, 3, 0.0f);
-#endif
-
-	const uniform FVector VSquareSum = SetVector(reduce_add(S0 * S0), reduce_add(S1 * S1), reduce_add(S2 * S2));
-	*((uniform FVector *uniform)&SquareSum) = *((uniform FVector *uniform)&VSquareSum);
-
-	const varying float Scale = select( SquareSum - SMALL_NUMBER >= 0.f, InvSqrt(SquareSum), 1.0f );
-	const uniform FVector VScale = *((uniform FVector *uniform)&Scale);
-
-	S0 = S0 * VScale.V[0];
-	S1 = S1 * VScale.V[1];
-	S2 = S2 * VScale.V[2];
-
-#if EXPLICIT_VECTOR == 1
-	// Reintroduce 4th components here
-	S0 = insert(S0, 3, SavedS0);
-	S1 = insert(S1, 3, SavedS1);
-	S2 = insert(S2, 3, SavedS2);
-#endif
-
-	*((uniform FVector *uniform)&M.M[0]) = *((uniform FVector *uniform)&S0);
-	*((uniform FVector *uniform)&M.M[4]) = *((uniform FVector *uniform)&S1);
-	*((uniform FVector *uniform)&M.M[8]) = *((uniform FVector *uniform)&S2);
+	const uniform float SquareSum0 = (M.M[0] * M.M[0]) + (M.M[1] * M.M[1]) + (M.M[2] * M.M[2]);
+	const uniform float SquareSum1 = (M.M[4] * M.M[4]) + (M.M[5] * M.M[5]) + (M.M[6] * M.M[6]);
+	const uniform float SquareSum2 = (M.M[8] * M.M[8]) + (M.M[9] * M.M[9]) + (M.M[10] * M.M[10]);
+	const uniform float Scale0 = select( SquareSum0 - SMALL_NUMBER >= 0.f, InvSqrt(SquareSum0), 1.0f );
+	const uniform float Scale1 = select( SquareSum1 - SMALL_NUMBER >= 0.f, InvSqrt(SquareSum1), 1.0f );
+	const uniform float Scale2 = select( SquareSum2 - SMALL_NUMBER >= 0.f, InvSqrt(SquareSum2), 1.0f );
+	M.M[0] *= Scale0;
+	M.M[1] *= Scale0;
+	M.M[2] *= Scale0;
+	M.M[4] *= Scale1;
+	M.M[5] *= Scale1;
+	M.M[6] *= Scale1;
+	M.M[8] *= Scale2;
+	M.M[9] *= Scale2;
+	M.M[10] *= Scale2;
 }
 
 inline uniform FVector MatrixGetOrigin(const uniform FMatrix &M)
@@ -183,7 +163,7 @@ inline uniform bool IsAnyMatrixScaledAxesNearlyZero(const uniform FMatrix &M, co
 	uniform FVector X, Y, Z;
 	MatrixGetScaledAxes(M, X, Y, Z);
 
-	if(VectorIsAnyNearlyZero(X, Tolerance) || VectorIsAnyNearlyZero(Y, Tolerance) || VectorIsAnyNearlyZero(Z, Tolerance))
+	if(VectorIsAllNearlyZero(X, Tolerance) || VectorIsAllNearlyZero(Y, Tolerance) || VectorIsAllNearlyZero(Z, Tolerance))
 	{
 		return true;
 	}
diff --git a/Engine/Source/Runtime/Core/Public/Math/Transform.isph b/Engine/Source/Runtime/Core/Public/Math/Transform.isph
--- a/Engine/Source/Runtime/Core/Public/Math/Transform.isph
+++ b/Engine/Source/Runtime/Core/Public/Math/Transform.isph
@@ -100,15 +100,12 @@ inline uniform FMatrix ToMatrixWithScale(const uniform FTransform &T)
 
 static inline uniform bool AnyHasNegativeScale(const uniform FVector4 &InScale3D, const uniform FVector4 &InOtherScale3D)
 {
-	varying bool b = false;
 	varying float S0 = 0.0f;
 	varying float S1 = 0.0f;
 	*((uniform FVector4 *uniform)&S0) = *((uniform FVector4 *uniform)&InScale3D);
 	*((uniform FVector4 *uniform)&S1) = *((uniform FVector4 *uniform)&InOtherScale3D);
 
-	b = intbits(min(S0, S1)) & SIGN_BIT ? true : false;
-
-	return any(b);
+	return any(S0 < 0 || S1 < 0);
 }
 
 static inline void ConstructTransformFromMatrixWithDesiredScale(const uniform FMatrix& AMatrix, const uniform FMatrix& BMatrix, const uniform FVector4& DesiredScale, uniform FTransform& OutTransform)
diff --git a/Engine/Source/Runtime/Core/Public/Math/Vector.isph b/Engine/Source/Runtime/Core/Public/Math/Vector.isph
--- a/Engine/Source/Runtime/Core/Public/Math/Vector.isph
+++ b/Engine/Source/Runtime/Core/Public/Math/Vector.isph
@@ -828,26 +828,20 @@ inline FIntVector SetIntVector(const FVector &A)
 
 inline uniform bool VectorIsAnyNearlyZero(const uniform FVector &V, const uniform float Tolerance)
 {
-	varying bool NearlyZero = false;
 	varying float S0 = FLT_MAX;
 
 	*((uniform FVector *uniform)&S0) = *((uniform FVector *uniform)&V);
 
-	NearlyZero = abs(S0) <= Tolerance ? true : false;
-
-	return any(NearlyZero);
+	return any(abs(S0) <= Tolerance);
 }
 
 inline uniform bool VectorIsAllNearlyZero(const uniform FVector &V, const uniform float Tolerance)
 {
-	varying bool NearlyZero = false;
 	varying float S0 = FLT_MIN;
 
 	*((uniform FVector *uniform)&S0) = *((uniform FVector *uniform)&V);
 
-	NearlyZero = abs(S0) <= Tolerance ? true : false;
-
-	return all(NearlyZero);
+	return all(abs(S0) <= Tolerance);
 }
 
 inline FVector VectorLerp(const FVector& A, const FVector& B, const float Alpha)
diff --git a/Engine/Source/Runtime/Engine/Private/Components/SkeletalMeshComponent.cpp b/Engine/Source/Runtime/Engine/Private/Components/SkeletalMeshComponent.cpp
--- a/Engine/Source/Runtime/Engine/Private/Components/SkeletalMeshComponent.cpp
+++ b/Engine/Source/Runtime/Engine/Private/Components/SkeletalMeshComponent.cpp
@@ -1395,7 +1395,6 @@ void USkeletalMeshComponent::FillComponentSpaceTransforms(const USkeletalMesh* I
 #endif
 	}
 
-#if 0 		 // temporarily disabling ISPC code due to negative scale issue
 	if (INTEL_ISPC)
 	{
 #if INTEL_ISPC
@@ -1410,7 +1409,6 @@ void USkeletalMeshComponent::FillComponentSpaceTransforms(const USkeletalMesh* I
 #endif
 	}
 	else
-#endif // 0
 	{
 		for (int32 i = 1; i < FillComponentSpaceTransformsRequiredBones.Num(); i++)
 		{
diff --git a/Engine/Source/Runtime/Engine/Private/Components/SkeletalMeshComponent.ispc b/Engine/Source/Runtime/Engine/Private/Components/SkeletalMeshComponent.ispc
--- a/Engine/Source/Runtime/Engine/Private/Components/SkeletalMeshComponent.ispc
+++ b/Engine/Source/Runtime/Engine/Private/Components/SkeletalMeshComponent.ispc
@@ -1,7 +1,6 @@
 // Copyright Epic Games, Inc. All Rights Reserved.
 
 #define EXPLICIT_VECTOR4 1
-#define EXPLICIT_VECTOR 1
 #define EXPLICIT_MATRIX 1
 
 #include "Math/Vector.isph"
