Make sure LLM can handle running with llmtagsets since they create so many categories.
Also throw an error if LLM overflows the preallocated space for category titles since that will corrupt the LLM CSV data.
#rnx

diff --git a/Engine/Source/Runtime/Core/Private/HAL/LowLevelMemTracker.cpp b/Engine/Source/Runtime/Core/Private/HAL/LowLevelMemTracker.cpp
--- a/Engine/Source/Runtime/Core/Private/HAL/LowLevelMemTracker.cpp
+++ b/Engine/Source/Runtime/Core/Private/HAL/LowLevelMemTracker.cpp
@@ -22,6 +22,16 @@ TAutoConsoleVariable<int32> CVarLLMWriteInterval(
 	TEXT("The number of seconds between each line in the LLM csv (zero to write every frame)")
 );
 
+TAutoConsoleVariable<int32> CVarLLMHeaderMaxSize(
+	TEXT("LLM.LLMHeaderMaxSize"),
+#if LLM_ALLOW_ASSETS_TAGS
+	500000, // When using asset tags, you will have MANY more LLM titles since so many are auto generated.
+#else
+	5000,
+#endif
+	TEXT("The maximum total number of characters allowed for all of the LLM titles")
+);
+
 DECLARE_LLM_MEMORY_STAT(TEXT("LLM Overhead"), STAT_LLMOverheadTotal, STATGROUP_LLMOverhead);
 
 DEFINE_STAT(STAT_EngineSummaryLLM);
@@ -2041,10 +2051,7 @@ void FLLMCsvWriter::WriteGraph(FLLMCustomTag* CustomTags, const int32* ParentTag
 		LLMCheck(Archive);
 
 		// create space for column titles that are filled in as we get them
-		for (int32 i = 0; i < 500; ++i)
-		{
-			Write(TEXT("          "));
-		}
+		Write(FString::ChrN(CVarLLMHeaderMaxSize.GetValueOnGameThread(), ' '));
 		Write(TEXT("\n"));
 	}
 
@@ -2068,6 +2075,12 @@ void FLLMCsvWriter::WriteGraph(FLLMCustomTag* CustomTags, const int32* ParentTag
 			Write(Text);
 		}
 
+		int64 ColumnTitleTotalSize = Archive->Tell();
+		if (ColumnTitleTotalSize >= CVarLLMHeaderMaxSize.GetValueOnGameThread())
+		{
+			UE_LOG(LogHAL, Error, TEXT("LLM column titles have overflowed, LLM CSM data will be corrupted. Increase CVarLLMHeaderMaxSize > %d"), ColumnTitleTotalSize);
+		}
+
 		Archive->Seek(OriginalOffset);
 
 		LastWriteStatValueCount = StatValueCountLocal;
