Disabled render thread polling optimization. Also reverted building actor tick modifications that went along with it.

We're investigating a few potential issues and will target a future release.

#rnx

diff --git a/Engine/Source/Runtime/Core/Private/Async/TaskGraph.cpp b/Engine/Source/Runtime/Core/Private/Async/TaskGraph.cpp
--- a/Engine/Source/Runtime/Core/Private/Async/TaskGraph.cpp
+++ b/Engine/Source/Runtime/Core/Private/Async/TaskGraph.cpp
@@ -61,14 +61,14 @@ namespace ENamedThreads
 	CORE_API int32 bHasHighPriorityThreads = CREATE_HIPRI_TASK_THREADS;
 }
 
-static bool GDoRenderThreadWakeupTrigger = false;
+static bool GDoRenderThreadWakeupTrigger = true;
 static FAutoConsoleVariableRef CVarDoRenderThreadWakeupTrigger(
 	TEXT("TaskGraph.DoRenderThreadWakeupTrigger"),
 	GDoRenderThreadWakeupTrigger,
 	TEXT("If true, task graph tasks sent to the render thread trigger a wakeup. See TaskGraph.RenderThreadPollPeriodMs.")
 );
 
-static int32 GRenderThreadPollPeriodMs = 1;
+static int32 GRenderThreadPollPeriodMs = (int32)MAX_uint32;
 static FAutoConsoleVariableRef CVarRenderThreadPollPeriodMs(
 	TEXT("TaskGraph.RenderThreadPollPeriodMs"),
 	GRenderThreadPollPeriodMs,
