Add missing include, fix non-unity builds

diff --git a/Engine/Source/Runtime/Core/Public/Containers/HashTable.h b/Engine/Source/Runtime/Core/Public/Containers/HashTable.h
--- a/Engine/Source/Runtime/Core/Public/Containers/HashTable.h
+++ b/Engine/Source/Runtime/Core/Public/Containers/HashTable.h
@@ -4,6 +4,7 @@
 
 #include "CoreTypes.h"
 #include "Serialization/MemoryLayout.h"
+#include "Serialization/MemoryImageWriter.h"
 #include "Misc/AssertionMacros.h"
 #include "HAL/UnrealMemory.h"
 #include "Math/UnrealMathUtility.h"
