RemoveAllSwap returns the number of elements removed to match RemoveAll.

diff --git a/Engine/Source/Runtime/Core/Public/Containers/Array.h b/Engine/Source/Runtime/Core/Public/Containers/Array.h
--- a/Engine/Source/Runtime/Core/Public/Containers/Array.h
+++ b/Engine/Source/Runtime/Core/Public/Containers/Array.h
@@ -2273,8 +2273,9 @@ public:
 	 * @see Remove, RemoveSingle, RemoveSingleSwap, RemoveSwap
 	 */
 	template <class PREDICATE_CLASS>
-	void RemoveAllSwap(const PREDICATE_CLASS& Predicate, bool bAllowShrinking = true)
+	SizeType RemoveAllSwap(const PREDICATE_CLASS& Predicate, bool bAllowShrinking = true)
 	{
+		const SizeType OriginalNum = ArrayNum;
 		for (SizeType ItemIndex = 0; ItemIndex < Num();)
 		{
 			if (Predicate((*this)[ItemIndex]))
@@ -2286,6 +2287,7 @@ public:
 				++ItemIndex;
 			}
 		}
+		return OriginalNum - ArrayNum;
 	}
 
 	/**
