Unix
* Fix FUnixPlatformMemory::GetStats() that always returned 0 when memory was over 10.24gigs

diff --git a/Engine/Source/Runtime/Core/Private/Unix/UnixPlatformMemory.cpp b/Engine/Source/Runtime/Core/Private/Unix/UnixPlatformMemory.cpp
--- a/Engine/Source/Runtime/Core/Private/Unix/UnixPlatformMemory.cpp
+++ b/Engine/Source/Runtime/Core/Private/Unix/UnixPlatformMemory.cpp
@@ -604,8 +604,10 @@ namespace UnixPlatformMemory
 			return 0;
 		}
 
+		const int NewLen = Len - kSuffixLength;
+
 		// let's check that this is indeed "kB"
-		char * Suffix = &Line[Len - kSuffixLength];
+		char * Suffix = &Line[NewLen];
 		if (strcmp(Suffix, " kB\n") != 0)
 		{
 			// Unix the kernel changed the format, huh?
@@ -615,12 +617,12 @@ namespace UnixPlatformMemory
 		// kill the kB
 		*Suffix = 0;
 
-		// find the beginning of the number
-		for (const char * NumberBegin = Suffix; NumberBegin >= Line; --NumberBegin)
+        // find the beginning of the number
+		for (const char* NumberBegin = Line; NumberBegin < Suffix; ++NumberBegin)
 		{
-			if (*NumberBegin == ' ')
+			if (isdigit(*NumberBegin))
 			{
-				return static_cast< uint64 >(atol(NumberBegin + 1)) * 1024ULL;
+				return atoll(NumberBegin) * 1024ULL;
 			}
 		}
 
