Fix settings for 32bit windows

diff --git a/Engine/Config/Windows/DataDrivenPlatformInfo.ini b/Engine/Config/Windows/DataDrivenPlatformInfo.ini
--- a/Engine/Config/Windows/DataDrivenPlatformInfo.ini
+++ b/Engine/Config/Windows/DataDrivenPlatformInfo.ini
@@ -3,8 +3,8 @@ bIsConfidential=false
 AudioCompressionSettingsIniSectionName=/Script/WindowsTargetPlatform.WindowsTargetSettings
 Freezing_b32Bit=ini:Engine:Windows:/Script/WindowsTargetPlatform.WindowsTargetSettings:bTarget32Bit
 Freezing_bWithRayTracing=!ini:Engine:Windows:/Script/WindowsTargetPlatform.WindowsTargetSettings:bTarget32Bit
-Freezing_MaxFieldAlignment32=8
-Freezing_MaxFieldAlignment64=16
+Freezing_MaxFieldAlignment32=4
+Freezing_MaxFieldAlignment64=8
 Freezing_bAlignBases=true
 
 [PlatformInfo WindowsNoEditor]
diff --git a/Engine/Source/Runtime/Core/Public/Misc/DataDrivenPlatformInfoRegistry.h b/Engine/Source/Runtime/Core/Public/Misc/DataDrivenPlatformInfoRegistry.h
--- a/Engine/Source/Runtime/Core/Public/Misc/DataDrivenPlatformInfoRegistry.h
+++ b/Engine/Source/Runtime/Core/Public/Misc/DataDrivenPlatformInfoRegistry.h
@@ -31,7 +31,7 @@ struct CORE_API FDataDrivenPlatformInfoRegistry
 		// MemoryFreezing information, matches FPlatformTypeLayoutParameters - defaults are clang, noneditor
 		uint32 Freezing_MaxFieldAlignment = 0xffffffff;
 		bool Freezing_b32Bit = false;
-		bool Freezing_bForce64BitMemoryImagePointers = true;
+		bool Freezing_bForce64BitMemoryImagePointers = false;
 		bool Freezing_bAlignBases = false;
 		bool Freezing_bWithRayTracing = false;
 
