Fix Compile error

diff --git a/Engine/Source/Runtime/Core/Public/Experimental/Containers/RobinHoodHashTable.h b/Engine/Source/Runtime/Core/Public/Experimental/Containers/RobinHoodHashTable.h
--- a/Engine/Source/Runtime/Core/Public/Experimental/Containers/RobinHoodHashTable.h
+++ b/Engine/Source/Runtime/Core/Public/Experimental/Containers/RobinHoodHashTable.h
@@ -713,7 +713,7 @@ namespace RobinHoodHashTable_Private
 			FHashElementId Id = FindId(Key);
 			if (Id.IsValid())
 			{
-				return KeyValueData.GetByElementId(Id).FindImpl();
+				return KeyValueData.Get(Id.GetIndex()).FindImpl();
 			}
 
 			return nullptr;
@@ -724,7 +724,7 @@ namespace RobinHoodHashTable_Private
 			FHashElementId Id = FindId(Key);
 			if (Id.IsValid())
 			{
-				return KeyValueData.GetByElementId(Id).FindImpl();
+				return KeyValueData.Get(Id.GetIndex()).FindImpl();
 			}
 
 			return nullptr;
