remove unnecessary code from COTF changes.

diff --git a/Engine/Source/Runtime/Core/Public/GenericPlatform/GenericPlatformFile.h b/Engine/Source/Runtime/Core/Public/GenericPlatform/GenericPlatformFile.h
--- a/Engine/Source/Runtime/Core/Public/GenericPlatform/GenericPlatformFile.h
+++ b/Engine/Source/Runtime/Core/Public/GenericPlatform/GenericPlatformFile.h
@@ -303,19 +303,6 @@ public:
 	/** For case insensitive filesystems, returns the full path of the file with the same case as in the filesystem */
 	virtual FString GetFilenameOnDisk(const TCHAR* Filename) = 0;
 
-	/** Return the size of the file, or -1 if it doesn't exist. **/
-	virtual int64		FileSize(const TCHAR* Filename, bool bUseWriteLocation)
-	{
-		return FileSize(Filename);
-	}
-
-	/** Check if a file exists in a location suitable for writing at runtime */
-	virtual bool		FileExists(const TCHAR* Filename, bool bUseWriteLocation)
-	{
-		return FileSize(Filename, bUseWriteLocation) >= 0;
-	}
-
-
 	/** Attempt to open a file for reading.
 	 *
 	 * @param Filename file to be opened
diff --git a/Engine/Source/Runtime/Core/Public/HAL/IPlatformFileCachedWrapper.h b/Engine/Source/Runtime/Core/Public/HAL/IPlatformFileCachedWrapper.h
--- a/Engine/Source/Runtime/Core/Public/HAL/IPlatformFileCachedWrapper.h
+++ b/Engine/Source/Runtime/Core/Public/HAL/IPlatformFileCachedWrapper.h
@@ -313,18 +313,10 @@ public:
 	{
 		return LowerLevel->FileExists(Filename);
 	}
-	virtual bool		FileExists(const TCHAR* Filename, bool bUseWriteLocation) override
-	{
-		return LowerLevel->FileExists(Filename, bUseWriteLocation);
-	}
 	virtual int64		FileSize(const TCHAR* Filename) override
 	{
 		return LowerLevel->FileSize(Filename);
 	}
-	virtual int64		FileSize(const TCHAR* Filename, bool bUseWriteLocation) override
-	{
-		return LowerLevel->FileSize(Filename, bUseWriteLocation);
-	}
 	virtual bool		DeleteFile(const TCHAR* Filename) override
 	{
 		return LowerLevel->DeleteFile(Filename);
diff --git a/Engine/Source/Runtime/NetworkFile/Private/NetworkPlatformFile.cpp b/Engine/Source/Runtime/NetworkFile/Private/NetworkPlatformFile.cpp
--- a/Engine/Source/Runtime/NetworkFile/Private/NetworkPlatformFile.cpp
+++ b/Engine/Source/Runtime/NetworkFile/Private/NetworkPlatformFile.cpp
@@ -252,9 +252,9 @@ void FNetworkPlatformFile::ProcessServerCachedFilesResponse(FArrayReader& Respon
 	// Check the stored cooked version
 	FString CookedVersionFile = FPaths::GeneratedConfigDir() / TEXT("CookedVersion.txt");
 
-	if (InnerPlatformFile->FileExists(*CookedVersionFile, true) == true)
+	if (InnerPlatformFile->FileExists(*CookedVersionFile) == true)
 	{
-		IFileHandle* FileHandle = InnerPlatformFile->OpenRead(*CookedVersionFile, true);
+		IFileHandle* FileHandle = InnerPlatformFile->OpenRead(*CookedVersionFile);
 		if (FileHandle != NULL)
 		{
 			int32 StoredPackageCookedVersion;
