#UE4 Move PersistentDownloadDir to the Saved folder on platforms where it is not otherwise specified

diff --git a/Engine/Source/Runtime/Core/Private/GenericPlatform/GenericPlatformMisc.cpp b/Engine/Source/Runtime/Core/Private/GenericPlatform/GenericPlatformMisc.cpp
--- a/Engine/Source/Runtime/Core/Private/GenericPlatform/GenericPlatformMisc.cpp
+++ b/Engine/Source/Runtime/Core/Private/GenericPlatform/GenericPlatformMisc.cpp
@@ -1065,12 +1065,7 @@ const TCHAR* FGenericPlatformMisc::GamePersistentDownloadDir()
 
 	if (GamePersistentDownloadDir.Len() == 0)
 	{
-		FString BaseProjectDir = ProjectDir();
-
-		if (BaseProjectDir.Len() > 0)
-		{
-			GamePersistentDownloadDir = BaseProjectDir / TEXT("PersistentDownloadDir");
-		}
+		GamePersistentDownloadDir = Paths::ProjectSavedDir() / TEXT("PersistentDownloadDir");
 	}
 	return *GamePersistentDownloadDir;
 }
