Merging CL#12677435

#UE4 CIS fix

diff --git a/Engine/Source/Runtime/Core/Private/GenericPlatform/GenericPlatformMisc.cpp b/Engine/Source/Runtime/Core/Private/GenericPlatform/GenericPlatformMisc.cpp
--- a/Engine/Source/Runtime/Core/Private/GenericPlatform/GenericPlatformMisc.cpp
+++ b/Engine/Source/Runtime/Core/Private/GenericPlatform/GenericPlatformMisc.cpp
@@ -1065,7 +1065,7 @@ const TCHAR* FGenericPlatformMisc::GamePersistentDownloadDir()
 
 	if (GamePersistentDownloadDir.Len() == 0)
 	{
-		GamePersistentDownloadDir = Paths::ProjectSavedDir() / TEXT("PersistentDownloadDir");
+		GamePersistentDownloadDir = FPaths::ProjectSavedDir() / TEXT("PersistentDownloadDir");
 	}
 	return *GamePersistentDownloadDir;
 }
