Move NoRedist files out of Core module.

diff --git a/Engine/Source/Runtime/Core/Core.Build.cs b/Engine/Source/Runtime/Core/Core.Build.cs
--- a/Engine/Source/Runtime/Core/Core.Build.cs
+++ b/Engine/Source/Runtime/Core/Core.Build.cs
@@ -182,8 +182,6 @@ public class Core : ModuleRules
 			PublicDefinitions.Add("WITH_VS_PERF_PROFILER=0");
 		}
 
-		WhitelistRestrictedFolders.Add("Private/NoRedist");
-
         if (Target.bWithDirectXMath && (Target.Platform == UnrealTargetPlatform.Win64 || Target.Platform == UnrealTargetPlatform.Win32))
         {
             PublicDefinitions.Add("WITH_DIRECTXMATH=1");
