On Windows, exit the application with exit code 777003 if the crash reporter crashes while reporting or code 777004 if the crash handler (the __except clause) crashes after reporting an error.

Details:

The 4.24.3 analytics shows many unexplained exit codes, 23 647 at the moment. Normally, the Editor will exit with code 0 if everything when well, 3 or 1 if it gracefully handled a crash, 255 it was aborted. But we also see may others like the following predominent cases below:

-1073741819 => STATUS_ACCESS_VIOLATION => 8081 cases
-1073740791 => STATUS_STACK_BUFFER_OVERRUN  => 7581 cases
-1073740771 => STATUS_FATAL_USER_CALLBACK_EXCEPTION => 5357 cases

On Windows, the crash reporting system should catch and report STATUS_ACCESS_VIOLATION and then exit with code 3 (as the error was handled). For example, if you add a null pointer dereference(STATUS_ACCESS_VIOLATION) in the code, the crash reporter handle it and the Editor exit with code 3. Just like if you enter 'debug crash' console command, the editor gracefully handle the error and exit with code 3.  But if you move the null pointer dereference in the crash handler thread itself, the error is not handled and the Editor exits with code STATUS_ACCESS_VIOLATION. This hints that our crash reporting thread is likely crashing in the wild. It would be useful to isolate those cases from the other cases and keep count of how many times this happens.

diff --git a/Engine/Source/Programs/CrashReportClient/Private/CrashReportClientApp.cpp b/Engine/Source/Programs/CrashReportClient/Private/CrashReportClientApp.cpp
--- a/Engine/Source/Programs/CrashReportClient/Private/CrashReportClientApp.cpp
+++ b/Engine/Source/Programs/CrashReportClient/Private/CrashReportClientApp.cpp
@@ -90,16 +90,6 @@ enum SubmitCrashReportResult {
 	SuccessContinue		// Succeeded sending report, continue running (if monitor mode).
 };
 
-/** Defines special exit codes for the Editor to use in the summary report when the real exit code could not be determined, to diagnose unexpected situations. */
-enum EditorSpecialExitCodes
-{
-	/** Editor is still running, but CRC was requested to exit. (ex. The user clicked the 'Close Without Sending' but the Editor is not exited yet) */
-	StillRunning = 777001,
-
-	/** Editor is not running anymore, but CRC could not read the Editor exit code. */
-	NotAvailable = 777002,
-};
-
 /**
  * Look for the report to upload, either in the command line or in the platform's report queue
  */
@@ -961,12 +951,12 @@ void RunCrashReportClient(const TCHAR* CommandLine)
 					else if (!bMonitoredProcessExited)
 					{
 						// Persist a custom exit code - Editor is still running, but CRC was requested to exit. (ex. The user clicked the 'Close Without Sending' but the Editor is not exited yet)
-						MonitoredSession.SaveExitCode(EditorSpecialExitCodes::StillRunning);
+						MonitoredSession.SaveExitCode(ECrashExitCodes::MonitoredApplicationStillRunning);
 					}
 					else
 					{
 						// Persist the custom exit code - Editor is not running anymore, but CRC could not read it.
-						MonitoredSession.SaveExitCode(EditorSpecialExitCodes::NotAvailable);
+						MonitoredSession.SaveExitCode(ECrashExitCodes::MonitoredApplicationExitCodeNotAvailable);
 					}
 				}
 				else
diff --git a/Engine/Source/Runtime/Core/Private/Windows/WindowsPlatformCrashContext.cpp b/Engine/Source/Runtime/Core/Private/Windows/WindowsPlatformCrashContext.cpp
--- a/Engine/Source/Runtime/Core/Private/Windows/WindowsPlatformCrashContext.cpp
+++ b/Engine/Source/Runtime/Core/Private/Windows/WindowsPlatformCrashContext.cpp
@@ -984,26 +984,40 @@ private:
 	/** Main loop that waits for a crash to trigger the report generation */
 	FORCENOINLINE uint32 Run()
 	{
-		// Removed vectored exception handler, we are now guaranteed to
-		// be able to catch unhandled exception in the main try/catch block.
-#if _WIN32_WINNT >= 0x0500
-		if (!FPlatformMisc::IsDebuggerPresent())
-		{
-			RemoveVectoredExceptionHandler(VectoredExceptionHandle);
-		}
+#if !PLATFORM_SEH_EXCEPTIONS_DISABLED
+		__try
 #endif
-		while (StopTaskCounter.GetValue() == 0)
 		{
-			if (WaitForSingleObject(CrashEvent, 500) == WAIT_OBJECT_0)
+			// Removed vectored exception handler, we are now guaranteed to
+			// be able to catch unhandled exception in the main try/catch block.
+#if _WIN32_WINNT >= 0x0500
+			if (!FPlatformMisc::IsDebuggerPresent())
 			{
-				ResetEvent(CrashHandledEvent);
-				HandleCrashInternal();
-				ResetEvent(CrashEvent);
-				// Let the thread that crashed know we're done.				
-				SetEvent(CrashHandledEvent);
-				break;
+				RemoveVectoredExceptionHandler(VectoredExceptionHandle);
 			}
+#endif
+			while (StopTaskCounter.GetValue() == 0)
+			{
+				if (WaitForSingleObject(CrashEvent, 500) == WAIT_OBJECT_0)
+				{
+					ResetEvent(CrashHandledEvent);
+					HandleCrashInternal();
+
+					ResetEvent(CrashEvent);
+					// Let the thread that crashed know we're done.
+					SetEvent(CrashHandledEvent);
+
+					break;
+				}
+			}
+		}
+#if !PLATFORM_SEH_EXCEPTIONS_DISABLED
+		__except(EXCEPTION_EXECUTE_HANDLER)
+		{
+			// The crash reporting thread crashed itself. Exit with a code that the out-of-process monitor will be able to pick up and report into analytics.
+			::exit(ECrashExitCodes::CrashReporterCrashed);
 		}
+#endif
 		return 0;
 	}
 
@@ -1343,12 +1357,20 @@ LONG WINAPI UnhandledStaticInitException(LPEXCEPTION_POINTERS ExceptionInfo)
  * Fallback for catching exceptions which aren't caught elsewhere. This allows catching exceptions on threads created outside the engine.
  * Note that Windows does not call this handler if a debugger is attached, separately to our internal logic around crash handling.
  */
-LONG WINAPI UnhandledException(EXCEPTION_POINTERS *ExceptionInfo)
+LONG WINAPI UnhandledException(EXCEPTION_POINTERS* ExceptionInfo)
 {
-	ReportCrash(ExceptionInfo);
-	GIsCriticalError = true;
-	FPlatformMisc::RequestExit(true);
-	return EXCEPTION_CONTINUE_SEARCH;
+	__try
+	{
+		ReportCrash(ExceptionInfo);
+		GIsCriticalError = true;
+		FPlatformMisc::RequestExit(true);
+		return EXCEPTION_CONTINUE_SEARCH;
+	}
+	__except(EXCEPTION_EXECUTE_HANDLER)
+	{
+		// The crash handler crashed itself, exit with a code that the out-of-process monitor will be able to pick up and report into analytics.
+		::exit(ECrashExitCodes::CrashHandlerCrashed);
+	}
 }
 
 // #CrashReport: 2015-05-28 This should be named EngineCrashHandler
diff --git a/Engine/Source/Runtime/Core/Private/Windows/WindowsRunnableThread.cpp b/Engine/Source/Runtime/Core/Private/Windows/WindowsRunnableThread.cpp
--- a/Engine/Source/Runtime/Core/Private/Windows/WindowsRunnableThread.cpp
+++ b/Engine/Source/Runtime/Core/Private/Windows/WindowsRunnableThread.cpp
@@ -5,6 +5,7 @@
 #include "Stats/Stats.h"
 #include "Misc/FeedbackContext.h"
 #include "HAL/ExceptionHandling.h"
+#include "GenericPlatform/GenericPlatformCrashContext.h"
 
 DEFINE_LOG_CATEGORY_STATIC(LogThreadingWindows, Log, All);
 
@@ -36,18 +37,26 @@ uint32 FRunnableThreadWin::GuardedRun()
 #if !PLATFORM_SEH_EXCEPTIONS_DISABLED
 		__except (ReportCrash( GetExceptionInformation() ))
 		{
-			// Make sure the information which thread crashed makes it into the log.
-			UE_LOG( LogThreadingWindows, Error, TEXT( "Runnable thread %s crashed." ), *ThreadName );
-			GWarn->Flush();
+			__try
+			{
+				// Make sure the information which thread crashed makes it into the log.
+				UE_LOG( LogThreadingWindows, Error, TEXT( "Runnable thread %s crashed." ), *ThreadName );
+				GWarn->Flush();
 
-			// Append the thread name at the end of the error report.
-			FCString::Strncat( GErrorHist, LINE_TERMINATOR TEXT( "Crash in runnable thread " ), UE_ARRAY_COUNT( GErrorHist ) );
-			FCString::Strncat( GErrorHist, *ThreadName, UE_ARRAY_COUNT( GErrorHist ) );
+				// Append the thread name at the end of the error report.
+				FCString::Strncat( GErrorHist, LINE_TERMINATOR TEXT( "Crash in runnable thread " ), UE_ARRAY_COUNT( GErrorHist ) );
+				FCString::Strncat( GErrorHist, *ThreadName, UE_ARRAY_COUNT( GErrorHist ) );
 
-			// Crashed.
-			ExitCode = 1;		
-			GError->HandleError();
-			FPlatformMisc::RequestExit( true );
+				// Crashed.
+				ExitCode = 1;
+				GError->HandleError();
+				FPlatformMisc::RequestExit( true );
+			}
+			__except(EXCEPTION_EXECUTE_HANDLER)
+			{
+				// The crash handler crashed itself, exit with a code that the out-of-process monitor will be able to pick up and report into analytics.
+				::exit(ECrashExitCodes::CrashHandlerCrashed);
+			}
 		}
 #endif // !PLATFORM_SEH_EXCEPTIONS_DISABLED
 	}
diff --git a/Engine/Source/Runtime/Core/Public/GenericPlatform/GenericPlatformCrashContext.h b/Engine/Source/Runtime/Core/Public/GenericPlatform/GenericPlatformCrashContext.h
--- a/Engine/Source/Runtime/Core/Public/GenericPlatform/GenericPlatformCrashContext.h
+++ b/Engine/Source/Runtime/Core/Public/GenericPlatform/GenericPlatformCrashContext.h
@@ -9,6 +9,23 @@
 
 struct FProgramCounterSymbolInfo;
 
+/** Defines special exit codes used to diagnose abnormal terminations. The code values are arbitrary, but easily recongnizable in decimal. They are meant to be
+    used with the out-of-process monitoring/analytics in order to figure out unexpected cases. */
+enum ECrashExitCodes : int32
+{
+	/** Used by out-of-process monitor in analytics report, the application is still running, but out-of-process monitor was requested to exit before the application exit code could be read. */
+	MonitoredApplicationStillRunning = 777001,
+
+	/** Used by out-of-process monitor in analytics report, the application is not running anymore, but the out-of-process monitor could not read the Editor exit code (either is is not supported by the OS or is not available). */
+	MonitoredApplicationExitCodeNotAvailable = 777002,
+
+	/** Used by the application when the crash reporter crashed itself while reporting a crash.*/
+	CrashReporterCrashed = 777003,
+
+	/** Used by the application when the crash handler crashed itself (crash in the __except() clause for example).*/
+	CrashHandlerCrashed = 777004,
+};
+
 /** 
  * Symbol information associated with a program counter. 
  * FString version.
