Add missing #if !PLATFORM_SEH_EXCEPTIONS_DISABLED guard around __try/__except.

diff --git a/Engine/Source/Runtime/Core/Private/Windows/WindowsPlatformCrashContext.cpp b/Engine/Source/Runtime/Core/Private/Windows/WindowsPlatformCrashContext.cpp
--- a/Engine/Source/Runtime/Core/Private/Windows/WindowsPlatformCrashContext.cpp
+++ b/Engine/Source/Runtime/Core/Private/Windows/WindowsPlatformCrashContext.cpp
@@ -1359,18 +1359,22 @@ LONG WINAPI UnhandledStaticInitException(LPEXCEPTION_POINTERS ExceptionInfo)
  */
 LONG WINAPI UnhandledException(EXCEPTION_POINTERS* ExceptionInfo)
 {
+#if !PLATFORM_SEH_EXCEPTIONS_DISABLED
 	__try
+#endif
 	{
 		ReportCrash(ExceptionInfo);
 		GIsCriticalError = true;
 		FPlatformMisc::RequestExit(true);
 		return EXCEPTION_CONTINUE_SEARCH;
 	}
+#if !PLATFORM_SEH_EXCEPTIONS_DISABLED
 	__except(EXCEPTION_EXECUTE_HANDLER)
 	{
 		// The crash handler crashed itself, exit with a code that the out-of-process monitor will be able to pick up and report into analytics.
 		::exit(ECrashExitCodes::CrashHandlerCrashed);
 	}
+#endif
 }
 
 // #CrashReport: 2015-05-28 This should be named EngineCrashHandler
