Malloc
* Remove SupportsMultithreading test from GCreateMalloc function
    * This was added during HTML5 support which is now removed.
    * The test never worked since the cmdline has yet to be initialized at that point
   * Its not clear what this fixed originally ( cl 3958226 )
* MallocThreadSafeProxy now returns that its thread safe

diff --git a/Engine/Source/Runtime/Core/Private/HAL/UnrealMemory.cpp b/Engine/Source/Runtime/Core/Private/HAL/UnrealMemory.cpp
--- a/Engine/Source/Runtime/Core/Private/HAL/UnrealMemory.cpp
+++ b/Engine/Source/Runtime/Core/Private/HAL/UnrealMemory.cpp
@@ -328,11 +328,6 @@ static int FMemory_GCreateMalloc_ThreadUnsafe()
 	// Setup malloc crash as soon as possible.
 	FPlatformMallocCrash::Get( GMalloc );
 
-	if ( ! FPlatformProcess::SupportsMultithreading() )
-	{
-		return 0;
-	}
-
 #if PLATFORM_USES_FIXED_GMalloc_CLASS
 #if USE_MALLOC_PROFILER || MALLOC_VERIFY || MALLOC_LEAKDETECTION || UE_USE_MALLOC_FILL_BYTES
 #error "Turn off PLATFORM_USES_FIXED_GMalloc_CLASS in order to use special allocator proxies"
diff --git a/Engine/Source/Runtime/Core/Public/HAL/MallocThreadSafeProxy.h b/Engine/Source/Runtime/Core/Public/HAL/MallocThreadSafeProxy.h
--- a/Engine/Source/Runtime/Core/Public/HAL/MallocThreadSafeProxy.h
+++ b/Engine/Source/Runtime/Core/Public/HAL/MallocThreadSafeProxy.h
@@ -119,4 +119,9 @@ public:
 		check(UsedMalloc);
 		UsedMalloc->Trim(bTrimThreadCaches);
 	}
+
+	virtual bool IsInternallyThreadSafe() const override
+	{ 
+		return true; 
+	}
 };
