Ensure that the project directory is updated after the first time that the project file path is set (if it is ever set).

Tested by using the "Always load last project on startup" feature.

diff --git a/Engine/Source/Runtime/Core/Private/GenericPlatform/GenericPlatformMisc.cpp b/Engine/Source/Runtime/Core/Private/GenericPlatform/GenericPlatformMisc.cpp
--- a/Engine/Source/Runtime/Core/Private/GenericPlatform/GenericPlatformMisc.cpp
+++ b/Engine/Source/Runtime/Core/Private/GenericPlatform/GenericPlatformMisc.cpp
@@ -954,18 +954,28 @@ const TCHAR* FGenericPlatformMisc::ProjectDir()
 	bool bIsIniReady = GConfig && GConfig->IsReadyForUse();
 	if (bWasIniReady != bIsIniReady)
 	{
-		ProjectDir = TEXT("");
+		ProjectDir.Reset();
 		bWasIniReady = bIsIniReady;
 	}
 
+	// track if last time we called this function the project file path was set
+	static bool bWasProjectFilePathReady = false;
+	if (!bWasProjectFilePathReady && FPaths::IsProjectFilePathSet())
+	{
+		ProjectDir.Reset();
+		bWasProjectFilePathReady = true;
+	}
+
 	// try using the override game dir if it exists, which will override all below logic
 	if (ProjectDir.Len() == 0)
 	{
+		ProjectDir.Reserve(FPlatformMisc::GetMaxPathLength());
 		ProjectDir = OverrideProjectDir;
 	}
 
 	if (ProjectDir.Len() == 0)
 	{
+		ProjectDir.Reserve(FPlatformMisc::GetMaxPathLength());
 		if (FPlatformProperties::IsProgram())
 		{
 			// monolithic, game-agnostic executables, the ini is in Engine/Config/Platform
