- Changes to support looking in Restricted locations for .ini files
- Contains a temporary compat for DedicatedServer.ini files - in the future they will be moved to a DedicatedServer directory (like platforms are), buyt for now, it will look in existing location.

diff --git a/Engine/Source/Programs/AutomationTool/Scripts/CryptoKeys.Automation.cs b/Engine/Source/Programs/AutomationTool/Scripts/CryptoKeys.Automation.cs
--- a/Engine/Source/Programs/AutomationTool/Scripts/CryptoKeys.Automation.cs
+++ b/Engine/Source/Programs/AutomationTool/Scripts/CryptoKeys.Automation.cs
@@ -40,7 +40,7 @@ namespace AutomationTool
 			}
 
 			FileReference OutputFile = FileReference.Combine(ProjectPath.Directory, "Config", "DefaultCrypto.ini");
-			FileReference NoRedistOutputFile = FileReference.Combine(ProjectPath.Directory, "Config", "NoRedist", "DefaultCrypto.ini");
+			FileReference NoRedistOutputFile = FileReference.Combine(ProjectPath.Directory, "Restricted", "NoRedist", "Config", "DefaultCrypto.ini");
 			FileReference DestinationFile = OutputFile;
 
 			// If the project has a DefaultCrypto.ini in a NoRedist folder, we want to copy the newly generated file into that location
diff --git a/Engine/Source/Programs/UnrealBuildTool/System/ConfigHierarchy.cs b/Engine/Source/Programs/UnrealBuildTool/System/ConfigHierarchy.cs
--- a/Engine/Source/Programs/UnrealBuildTool/System/ConfigHierarchy.cs
+++ b/Engine/Source/Programs/UnrealBuildTool/System/ConfigHierarchy.cs
@@ -729,79 +729,50 @@ namespace UnrealBuildTool
 		}
 
 
-		enum EConfigFlag
+		class ConfigLayerExpansion
 		{
-			None,
-			// Required,  // not needed in C# land
-			// AllowCommandLineOverride,  // not needed in C# land
-			// DedicatedServerOnly, // not needed in C# land
-			// GenerateCacheKey,  // not needed in C# land
+			// a set of replacements from the source file to possible other files
+			public string Before1 = null;
+			public string After1 = null;
+			public string Before2 = null;
+			public string After2 = null;
 		};
 
-		class ConfigLayer
-		{
-			// Used by the editor to display in the ini-editor
-			// string EditorName; // don't need editor name in C# land
-			// Path to the ini file (with variables)
-			public string Path;
-			// Special flag
-			// public EConfigFlag Flag = EConfigFlag.None;
-
-			public string ExtEnginePath = null;
-			public string ExtProjectPath = null;
-		}
-
-		struct ConfigLayerExpansion
-		{
-			// The subdirectory for this expansion (ie "NoRedist")
-			public string DirectoryPrefix;
-			// The filename prefix for this expansion (ie "Shippable")
-			public string FilePrefix;
-			// Optional flags 
-			// public EConfigFlag Flag;
-		};
-
-		static ConfigLayer[] ConfigLayers =
+		static string [] ConfigLayers =
 		{
 			// Engine/Base.ini
-			new ConfigLayer { Path = "{ENGINE}/Base.ini" }, //, Flag = EConfigFlag.Required },
+			"{ENGINE}/Config/Base.ini",
 			// Engine/Base*.ini
- 			new ConfigLayer { Path = "{ENGINE}/{ED}{EF}Base{TYPE}.ini" },
+ 			"{ENGINE}/Config/Base{TYPE}.ini",
 			// Engine/Platform/BasePlatform*.ini
-			new ConfigLayer { Path = "{ENGINE}/{ED}{PLATFORM}/{EF}Base{PLATFORM}{TYPE}.ini", ExtEnginePath = "{EXTENGINE}/{ED}{EF}Base{PLATFORM}{TYPE}.ini" },
+			"{ENGINE}/Config/{PLATFORM}/Base{PLATFORM}{TYPE}.ini",
 			// Project/Default*.ini
-			new ConfigLayer { Path = "{PROJECT}/{ED}{EF}Default{TYPE}.ini" }, //, Flag = EConfigFlag.AllowCommandLineOverride },
+			"{PROJECT}/Config/Default{TYPE}.ini",
 			// Engine/Platform/Platform*.ini
-			new ConfigLayer { Path = "{ENGINE}/{ED}{PLATFORM}/{EF}{PLATFORM}{TYPE}.ini", ExtEnginePath = "{EXTENGINE}/{ED}{EF}{PLATFORM}{TYPE}.ini" },
+			"{ENGINE}/Config/{PLATFORM}/{PLATFORM}{TYPE}.ini",
 			// Project/Platform/Platform*.ini
-			new ConfigLayer { Path = "{PROJECT}/{ED}{PLATFORM}/{EF}{PLATFORM}{TYPE}.ini", ExtProjectPath = "{EXTPROJECT}/{ED}{EF}{PLATFORM}{TYPE}.ini" },
+			"{PROJECT}/Config/{PLATFORM}/{PLATFORM}{TYPE}.ini",
 
 			// UserSettings/.../User*.ini
-			new ConfigLayer { Path = "{USERSETTINGS}/Unreal Engine/Engine/Config/User{TYPE}.ini" },
+			"{USERSETTINGS}/Unreal Engine/Engine/Config/User{TYPE}.ini",
 			// UserDir/.../User*.ini
-			new ConfigLayer { Path = "{USER}/Unreal Engine/Engine/Config/User{TYPE}.ini" },
+			"{USER}/Unreal Engine/Engine/Config/User{TYPE}.ini",
 			// Project/User*.ini
-			new ConfigLayer { Path = "{PROJECT}/User{TYPE}.ini" },
+			"{PROJECT}/User{TYPE}.ini",
 		};
 
 		static ConfigLayerExpansion[] ConfigLayerExpansions =
 		{
 			// The base expansion (ie, no expansion)
-			new ConfigLayerExpansion { DirectoryPrefix = "", FilePrefix = "" }, 
-
-			// When running a dedicated server, not used in UBT
-			// new ConfigLayerExpansion { DirectoryPrefix = "", FilePrefix = "DedicatedServer" }, //  Flag_DedicatedServerOnly },
-
-			// This file is remapped in UAT from inside NFL or NoRedist, because those directories are stripped while packaging
-			new ConfigLayerExpansion { DirectoryPrefix = "", FilePrefix = "Shippable" },
-			// Hidden directory from licensees
-			new ConfigLayerExpansion { DirectoryPrefix = "NotForLicensees/", FilePrefix = "" },
-			// Settings that need to be hidden from licensees, but are needed for shipping
-			new ConfigLayerExpansion { DirectoryPrefix = "NotForLicensees/", FilePrefix = "Shippable" },
-			// Hidden directory from non-Epic
-			new ConfigLayerExpansion { DirectoryPrefix = "NoRedist/", FilePrefix = "" },
-			// Settings that need to be hidden from non-Epic, but are needed for shipping
-			new ConfigLayerExpansion { DirectoryPrefix = "NoRedist/", FilePrefix = "Shippable" },
+			new ConfigLayerExpansion { }, 
+			// Restricted Locations
+			new ConfigLayerExpansion { Before1 = "{ENGINE}/", After1 = "{ENGINE}/Restricted/NotForLicensees/",	Before2 = "{PROJECT}/Config/", After2 = "{RESTRICTEDPROJECT_NFL}/Config/" },
+			new ConfigLayerExpansion { Before1 = "{ENGINE}/", After1 = "{ENGINE}/Restricted/NoRedist/",			Before2 = "{PROJECT}/Config/", After2 = "{RESTRICTEDPROJECT_NR}/Config/" },
+			// Platform Extensions
+			new ConfigLayerExpansion { Before1 = "{ENGINE}/Config/{PLATFORM}/", After1 = "{EXTENGINE}/Config/",	Before2 = "{PROJECT}/Config/{PLATFORM}/", After2 = "{EXTPROJECT}/Config/" },
+			// Platform Extensions in Restricted Locations
+			new ConfigLayerExpansion { Before1 = "{ENGINE}/Config/{PLATFORM}/", After1 = "{ENGINE}/Restricted/NotForLicensees/Platforms/{PLATFORM}/Config/",   Before2 = "{PROJECT}/Config/{PLATFORM}/", After2 = "{RESTRICTEDPROJECT_NFL}/Platforms/{PLATFORM}/Config/" },
+			new ConfigLayerExpansion { Before1 = "{ENGINE}/Config/{PLATFORM}/", After1 = "{ENGINE}/Restricted/NoRedist/Platforms/{PLATFORM}/Config/",          Before2 = "{PROJECT}/Config/{PLATFORM}/", After2 = "{RESTRICTEDPROJECT_NR}/Platforms/{PLATFORM}/Config/" },
 		};
 
 		// Match FPlatformProcess::UserDir()
@@ -822,66 +793,72 @@ namespace UnrealBuildTool
 			return PersonalConfigFolder;
 		}
 
-		private static string GetLayerPath(ConfigLayer Layer, string PlatformExtensionName, string IniPlatformName, DirectoryReference ProjectDir, string BaseIniName,
-			out bool bHasPlatformTag, out bool bHasProjectTag, out bool bHasExpansionTag)
+
+		private static string PerformBasicReplacements(string InString, string BaseIniName)
+		{
+			string OutString = InString.Replace("{TYPE}", BaseIniName);
+			OutString = OutString.Replace("{USERSETTINGS}", Utils.GetUserSettingDirectory().FullName);
+			OutString = OutString.Replace("{USER}", GetUserDir());
+
+			return OutString;
+		}
+
+
+		private static string PerformExpansionReplacements(ConfigLayerExpansion Expansion, string InString)
 		{
-			// cache some platform extension information that can be used inside the loops
-			string PlatformExtensionEngineConfigDir = DirectoryReference.Combine(UnrealBuildTool.EngineDirectory, "Platforms", PlatformExtensionName, "Config").FullName;
-			string PlatformExtensionProjectConfigDir = ProjectDir != null ? DirectoryReference.Combine(ProjectDir, "Platforms", PlatformExtensionName, "Config").FullName : null;
-			bool bHasPlatformExtensionEngineConfigDir = Directory.Exists(PlatformExtensionEngineConfigDir);
-			bool bHasPlatformExtensionProjectConfigDir = PlatformExtensionProjectConfigDir != null && Directory.Exists(PlatformExtensionProjectConfigDir);
-
-			bHasPlatformTag = Layer.Path.Contains("{PLATFORM}");
-			bHasProjectTag = Layer.Path.Contains("{PROJECT}");
-			bHasExpansionTag = Layer.Path.Contains("{ED}") || Layer.Path.Contains("{EF}");
-			bool bHasUserTag = Layer.Path.Contains("{USER}");
-
-			// skip platform layers if we are "platform-less", or user layers without a user dir
-			if ((bHasPlatformTag && IniPlatformName == "None") ||
-				(bHasProjectTag && ProjectDir == null) ||
-				(bHasUserTag && GetUserDir() == null))
+			// if there's replacement to do, the output is just the output
+			if (Expansion.Before1 == null)
+			{
+				return InString;
+			}
+
+			// if nothing to replace, then skip it entirely
+			if (!InString.Contains(Expansion.Before1) && (Expansion.Before2 == null || !InString.Contains(Expansion.Before2)))
 			{
 				return null;
 			}
 
-			// basic replacements
-			string LayerPath;
-			// you can only have PROJECT or ENGINE, not both
-			if (bHasProjectTag)
+			// replace the directory bits
+			string OutString = InString.Replace(Expansion.Before1, Expansion.After1);
+			if (Expansion.Before2 != null)
 			{
-				if (bHasPlatformTag && bHasPlatformExtensionProjectConfigDir)
-				{
-					LayerPath = Layer.ExtProjectPath.Replace("{EXTPROJECT}", PlatformExtensionProjectConfigDir);
-				}
-				else
-				{
-					LayerPath = Layer.Path.Replace("{PROJECT}", Path.Combine(ProjectDir.FullName, "Config"));
-				}
+				OutString = OutString.Replace(Expansion.Before2, Expansion.After2);
 			}
-			else
+			return OutString;
+		}
+
+		private static string PerformFinalExpansions(string InString, string PlatformName, DirectoryReference ProjectDir)
+		{
+			string PlatformExtensionEngineConfigDir = DirectoryReference.Combine(UnrealBuildTool.EngineDirectory, "Platforms", PlatformName).FullName;
+
+			string OutString = InString.Replace("{ENGINE}", UnrealBuildTool.EngineDirectory.FullName);
+			OutString = OutString.Replace("{EXTENGINE}", PlatformExtensionEngineConfigDir);
+			OutString = OutString.Replace("{PLATFORM}", PlatformName);
+
+			if (ProjectDir != null)
 			{
-				if (bHasPlatformTag && bHasPlatformExtensionEngineConfigDir)
+				DirectoryReference NFLDir;
+				DirectoryReference NRDir;
+				if (ProjectDir.IsUnderDirectory(UnrealBuildTool.EngineDirectory))
 				{
-					LayerPath = Layer.ExtEnginePath.Replace("{EXTENGINE}", PlatformExtensionEngineConfigDir);
+					string RelativeDir = ProjectDir.MakeRelativeTo(UnrealBuildTool.EngineDirectory);
+					NFLDir = DirectoryReference.Combine(UnrealBuildTool.EngineDirectory, "Restricted/NotForLicensees", RelativeDir);
+					NRDir = DirectoryReference.Combine(UnrealBuildTool.EngineDirectory, "Restricted/NoRedist", RelativeDir);
 				}
 				else
 				{
-					LayerPath = Layer.Path.Replace("{ENGINE}", Path.Combine(UnrealBuildTool.EngineDirectory.FullName, "Config"));
+					NFLDir = DirectoryReference.Combine(ProjectDir, "Restricted/NotForLicensees");
+					NRDir = DirectoryReference.Combine(ProjectDir, "Restricted/NoRedist");
 				}
-			}
-			LayerPath = LayerPath.Replace("{TYPE}", BaseIniName);
-			LayerPath = LayerPath.Replace("{USERSETTINGS}", Utils.GetUserSettingDirectory().FullName);
-			LayerPath = LayerPath.Replace("{USER}", GetUserDir());
+				string PlatformExtensionProjectConfigDir = DirectoryReference.Combine(ProjectDir, "Platforms", PlatformName).FullName;
 
-			return LayerPath;
-		}
-
-		private static string GetExpansionPath(ConfigLayerExpansion Expansion, string LayerPath)
-		{
-			string ExpansionPath = LayerPath.Replace("{ED}", Expansion.DirectoryPrefix);
-			ExpansionPath = ExpansionPath.Replace("{EF}", Expansion.FilePrefix);
+				OutString = OutString.Replace("{PROJECT}", ProjectDir.FullName);
+				OutString = OutString.Replace("{EXTPROJECT}", PlatformExtensionProjectConfigDir);
+				OutString = OutString.Replace("{RESTRICTEDPROJECT_NFL}", NFLDir.FullName);
+				OutString = OutString.Replace("{RESTRICTEDPROJECT_NR}", NRDir.FullName);
+			}
 
-			return ExpansionPath;
+			return OutString;
 		}
 
 		/// <summary>
@@ -892,24 +869,36 @@ namespace UnrealBuildTool
 			string BaseIniName = Enum.GetName(typeof(ConfigHierarchyType), Type);
 			string PlatformName = GetIniPlatformName(Platform);
 
-			foreach (ConfigLayer Layer in ConfigLayers)
+			Console.WriteLine("Ini hiererachy: {0} / {1} / {2}", PlatformName, ProjectDir != null ? ProjectDir.FullName : "NoProject", GetUserDir() != null ? GetUserDir() : "NoUserDir");
+			foreach (string Layer in ConfigLayers)
 			{
-				bool bHasPlatformTag, bHasProjectTag, bHasExpansionTag;
-				string LayerPath = GetLayerPath(Layer, Platform.ToString(), PlatformName, ProjectDir, BaseIniName, out bHasPlatformTag, out bHasProjectTag, out bHasExpansionTag);
+				bool bHasPlatformTag = Layer.Contains("{PLATFORM}");
+				bool bHasProjectTag = Layer.Contains("{PROJECT}");
+				bool bHasUserTag = Layer.Contains("{USER}");
 
-				// skip the layer if we aren't going to use it
-				if (LayerPath == null)
+				// skip certain layers if we are platform-less, project-less, or userdir-less
+				if ((bHasPlatformTag && PlatformName == "None") ||
+					(bHasProjectTag && ProjectDir == null) ||
+					(bHasUserTag && GetUserDir() == null))
 				{
 					continue;
 				}
 
-				// handle expansion (and platform - the C++ code will validate that only expansion layers have platforms)
-				if (bHasExpansionTag)
+				string LayerPath = PerformBasicReplacements(Layer, BaseIniName);
+
+				// we only expand engine/project inis
+				if (Layer.Contains("{ENGINE}") || Layer.Contains("{PROJECT}"))
 				{
 					foreach (ConfigLayerExpansion Expansion in ConfigLayerExpansions)
 					{
 						// expansion replacements
-						string ExpansionPath = GetExpansionPath(Expansion, LayerPath);
+						string ExpandedPath = PerformExpansionReplacements(Expansion, LayerPath);
+
+						// if nothing was replaced, then skip it, as it won't change anything
+						if (ExpandedPath == null)
+						{
+							continue;
+						}
 
 						// now go up the ini parent chain
 						if (bHasPlatformTag)
@@ -924,19 +913,18 @@ namespace UnrealBuildTool
 									// may not even exist as a UnrealTargetPlatform, and all we have is a string to look up, and it would just get the same
 									// string back, if we did look it up. This could become an issue if Win64 becomes a PlatformExtension, and wants to have 
 									// a parent Platform, of ... something. This is likely to never be an issue, but leaving this note here just in case.
-									string LocalLayerPath = GetLayerPath(Layer, ParentPlatform, ParentPlatform, ProjectDir, BaseIniName, out bHasPlatformTag, out bHasProjectTag, out bHasExpansionTag);
-									string LocalExpansionPath = GetExpansionPath(Expansion, LocalLayerPath);
-									yield return new FileReference(LocalExpansionPath.Replace("{PLATFORM}", ParentPlatform));
+									yield return new FileReference(PerformFinalExpansions(ExpandedPath, ParentPlatform, ProjectDir));
 								}
 							}
 							// always yield the active platform last 
-							yield return new FileReference(ExpansionPath.Replace("{PLATFORM}", PlatformName));
+							yield return new FileReference(PerformFinalExpansions(ExpandedPath, PlatformName, ProjectDir));
 						}
 						else
 						{
-							yield return new FileReference(ExpansionPath);
+							yield return new FileReference(PerformFinalExpansions(ExpandedPath, "", ProjectDir));
 						}
 					}
+					Console.WriteLine("");
 				}
 				else
 				{
diff --git a/Engine/Source/Programs/UnrealBuildTool/System/UnrealPluginLanguage.cs b/Engine/Source/Programs/UnrealBuildTool/System/UnrealPluginLanguage.cs
--- a/Engine/Source/Programs/UnrealBuildTool/System/UnrealPluginLanguage.cs
+++ b/Engine/Source/Programs/UnrealBuildTool/System/UnrealPluginLanguage.cs
@@ -3236,7 +3236,7 @@ namespace UnrealBuildTool
 			yield return FileReference.Combine(EngineDirectory, "Config", "Base" + BaseIniName + ".ini");
 
 			// Engine/Config/NotForLicensees/Base* ini
-			yield return FileReference.Combine(EngineDirectory, "Config", "NotForLicensees", "Base" + BaseIniName + ".ini");
+			yield return FileReference.Combine(EngineDirectory, "Restricted", "NotForLicensees", "Config", "Base" + BaseIniName + ".ini");
 		}
 
 
@@ -3254,10 +3254,10 @@ namespace UnrealBuildTool
 				yield return FileReference.Combine(EngineDirectory, "Config", "Base" + BaseIniName + ".ini");
 
 				// Engine/Config/NotForLicensees/Base* ini
-				yield return FileReference.Combine(EngineDirectory, "Config", "NotForLicensees", "Base" + BaseIniName + ".ini");
+				yield return FileReference.Combine(EngineDirectory, "Restricted", "NotForLicensees", "Config", "Base" + BaseIniName + ".ini");
 
 				// Engine/Config/NoRedist/Base* ini
-				yield return FileReference.Combine(EngineDirectory, "Config", "NoRedist", "Base" + BaseIniName + ".ini");
+				yield return FileReference.Combine(EngineDirectory, "Restricted", "NoRedist", "Config", "Base" + BaseIniName + ".ini");
 			}
 
 			if (ProjectDirectory != null)
@@ -3266,10 +3266,10 @@ namespace UnrealBuildTool
 				yield return FileReference.Combine(ProjectDirectory, "Config", "Default" + BaseIniName + ".ini");
 
 				// Game/Config/NotForLicensees/Default* ini
-				yield return FileReference.Combine(ProjectDirectory, "Config", "NotForLicensees", "Default" + BaseIniName + ".ini");
+				yield return FileReference.Combine(ProjectDirectory, "Restricted", "NotForLicensees", "Config", "Default" + BaseIniName + ".ini");
 
 				// Game/Config/NoRedist/Default* ini
-				yield return FileReference.Combine(ProjectDirectory, "Config", "NoRedist", "Default" + BaseIniName + ".ini");
+				yield return FileReference.Combine(ProjectDirectory, "Restricted", "NoRedist", "Config", "Default" + BaseIniName + ".ini");
 			}
 
 			string PlatformName = GetIniPlatformName(Platform);
diff --git a/Engine/Source/Runtime/Core/Private/Misc/ConfigCacheIni.cpp b/Engine/Source/Runtime/Core/Private/Misc/ConfigCacheIni.cpp
--- a/Engine/Source/Runtime/Core/Private/Misc/ConfigCacheIni.cpp
+++ b/Engine/Source/Runtime/Core/Private/Misc/ConfigCacheIni.cpp
@@ -500,10 +500,12 @@ static bool SaveConfigFileWrapper(const TCHAR* IniFile, const FString& Contents)
 
 enum class EConfigLayerFlags : int32
 {
+	None						= 0,
 	Required					= (1 << 0),
 	AllowCommandLineOverride	= (1 << 1),
 	DedicatedServerOnly			= (1 << 2), // replaces Default, Base, and (NOT {PLATFORM} yet) with an empty string
 	GenerateCacheKey			= (1 << 3),
+	NoExpand					= (1 << 4),
 };
 ENUM_CLASS_FLAGS(EConfigLayerFlags);
 
@@ -517,8 +519,6 @@ struct FConfigLayer
 	const TCHAR* EditorName;
 	// Path to the ini file (with variables)
 	const TCHAR* Path;
-	// Path to the platform extension version
-	const TCHAR* PlatformExtensionPath;
 	// Special flag
 	EConfigLayerFlags Flag;
 
@@ -531,73 +531,69 @@ struct FConfigLayer
 	**************************************************/
 
 	// Engine/Base.ini
-	{ TEXT("AbsoluteBase"),				TEXT("{ENGINE}Base.ini"), TEXT(""), EConfigLayerFlags::Required },
+	{ TEXT("AbsoluteBase"),				TEXT("{ENGINE}/Config/Base.ini"), EConfigLayerFlags::Required | EConfigLayerFlags::NoExpand},
 	
 	// Engine/Base*.ini
-	{ TEXT("Base"),						TEXT("{ENGINE}{ED}{EF}Base{TYPE}.ini") },
+	{ TEXT("Base"),						TEXT("{ENGINE}/Config/Base{TYPE}.ini") },
 	// Engine/Platform/BasePlatform*.ini
-	{ TEXT("BasePlatform"),				TEXT("{ENGINE}{ED}{PLATFORM}/{EF}Base{PLATFORM}{TYPE}.ini"), TEXT("{EXTENGINE}/{ED}{EF}Base{PLATFORM}{TYPE}.ini"),  },
+	{ TEXT("BasePlatform"),				TEXT("{ENGINE}/Config/{PLATFORM}/Base{PLATFORM}{TYPE}.ini")  },
 	// Project/Default*.ini
-	{ TEXT("ProjectDefault"),			TEXT("{PROJECT}{ED}{EF}Default{TYPE}.ini"), TEXT(""), EConfigLayerFlags::AllowCommandLineOverride | EConfigLayerFlags::GenerateCacheKey },
+	{ TEXT("ProjectDefault"),			TEXT("{PROJECT}/Config/Default{TYPE}.ini"), EConfigLayerFlags::AllowCommandLineOverride | EConfigLayerFlags::GenerateCacheKey },
 	// Engine/Platform/Platform*.ini
-	{ TEXT("EnginePlatform"),			TEXT("{ENGINE}{ED}{PLATFORM}/{EF}{PLATFORM}{TYPE}.ini"), TEXT("{EXTENGINE}/{ED}{EF}{PLATFORM}{TYPE}.ini") },
+	{ TEXT("EnginePlatform"),			TEXT("{ENGINE}/Config/{PLATFORM}/{PLATFORM}{TYPE}.ini") },
 	// Project/Platform/Platform*.ini
-	{ TEXT("ProjectPlatform"),			TEXT("{PROJECT}{ED}{PLATFORM}/{EF}{PLATFORM}{TYPE}.ini"), TEXT("{EXTPROJECT}/{ED}{EF}{PLATFORM}{TYPE}.ini") },
+	{ TEXT("ProjectPlatform"),			TEXT("{PROJECT}/Config/{PLATFORM}/{PLATFORM}{TYPE}.ini") },
 	
 	// UserSettings/.../User*.ini
-	{ TEXT("UserSettingsDir"),			TEXT("{USERSETTINGS}Unreal Engine/Engine/Config/User{TYPE}.ini") },
+	{ TEXT("UserSettingsDir"),			TEXT("{USERSETTINGS}Unreal Engine/Engine/Config/User{TYPE}.ini"), EConfigLayerFlags::NoExpand },
 	// UserDir/.../User*.ini
-	{ TEXT("UserDir"),					TEXT("{USER}Unreal Engine/Engine/Config/User{TYPE}.ini") },
+	{ TEXT("UserDir"),					TEXT("{USER}Unreal Engine/Engine/Config/User{TYPE}.ini"), EConfigLayerFlags::NoExpand },
 	// Project/User*.ini
-	{ TEXT("GameDirUser"),				TEXT("{PROJECT}User{TYPE}.ini"), TEXT(""), EConfigLayerFlags::GenerateCacheKey },
+	{ TEXT("GameDirUser"),				TEXT("{PROJECT}/Config/User{TYPE}.ini"), EConfigLayerFlags::GenerateCacheKey | EConfigLayerFlags::NoExpand },
 };
 
 
 /**
- * This describes extra files per layer, most so that Epic needs to be able to ship a project, but still share the project 
- * with licensees. These settings are the settings that should not be shared outside of Epic because they could cause 
- * problems if a licensee blindly copy and pasted the settings (they can't copy what they don't have!)
+ * This describes extra files per layer, to deal with restricted and NDA covered platform files that can't have the settings
+ * be in the Base/Default ini files.
+ * Note that we treat DedicatedServer as a "Platform" where it will have it's own directory of files, like a platform
  */
 struct FConfigLayerExpansion
 {
-	// The subdirectory for this expansion (ie "NoRedist")
-	const TCHAR* DirectoryPrefix;
-	// The filename prefix for this expansion (ie "Shippable")
-	const TCHAR* FilePrefix;
-	// Optional flags 
-	EConfigLayerFlags Flag;
-} GConfigLayerExpansions[] = 
-{
-	/**************************************************
-	**** CRITICAL NOTES
-	**** If you change this array, you need to also change EnumerateConfigFileLocations() in ConfigHierarchy.cs!!!
-	**************************************************/
-	
-	// The base expansion (ie, no expansion)
-	{ TEXT(""), TEXT("") },
-
-	// When running a dedicated server - does not support {PLATFORM} layers, so those are skipped
-	{ TEXT(""), TEXT("DedicatedServer"),  EConfigLayerFlags::DedicatedServerOnly },
-	// This file is remapped in UAT from inside NFL or NoRedist, because those directories are stripped while packaging
-	{ TEXT(""), TEXT("Shippable") },
-	// Hidden directory from licensees
-	{ TEXT("NotForLicensees/"), TEXT("") },
-	// Settings that need to be hidden from licensees, but are needed for shipping
-	{ TEXT("NotForLicensees/"), TEXT("Shippable") },
-	// Hidden directory from non-Epic
-	{ TEXT("NoRedist/"), TEXT("") },
-	// Settings that need to be hidden from non-Epic, but are needed for shipping
-	{ TEXT("NoRedist/"), TEXT("Shippable") },
+	// a set of replacements from the source file to possible other files
+	const TCHAR* Before1;
+	const TCHAR* After1;
+	const TCHAR* Before2;
+	const TCHAR* After2;
+};
+
+/**************************************************
+**** CRITICAL NOTES
+**** If you change these arrays, you need to also change EnumerateConfigFileLocations() in ConfigHierarchy.cs!!!
+**************************************************/
+static FConfigLayerExpansion GConfigExpansions[] =
+{
+
+	// No replacements
+	{},
+	// Restricted Locations
+	{ TEXT("{ENGINE}/"), TEXT("{ENGINE}/Restricted/NotForLicensees/"),	TEXT("{PROJECT}/Config/"), TEXT("{RESTRICTEDPROJECT_NFL}/Config/") },
+	{ TEXT("{ENGINE}/"), TEXT("{ENGINE}/Restricted/NoRedist/"),			TEXT("{PROJECT}/Config/"), TEXT("{RESTRICTEDPROJECT_NR}/Config/") },
+	// Platform Extensions
+	{ TEXT("{ENGINE}/Config/{PLATFORM}/"), TEXT("{EXTENGINE}/Config/"),	TEXT("{PROJECT}/Config/{PLATFORM}/"), TEXT("{EXTPROJECT}/Config/") },
+	// Platform Extensions in Restricted Locations
+	{ TEXT("{ENGINE}/Config/{PLATFORM}/"), TEXT("{ENGINE}/Restricted/NotForLicensees/Platforms/{PLATFORM}/Config/"),	TEXT("{PROJECT}/Config/{PLATFORM}/"), TEXT("{RESTRICTEDPROJECT_NFL}/Platforms/{PLATFORM}/Config/") },
+	{ TEXT("{ENGINE}/Config/{PLATFORM}/"), TEXT("{ENGINE}/Restricted/NoRedist/Platforms/{PLATFORM}/Config/"),			TEXT("{PROJECT}/Config/{PLATFORM}/"), TEXT("{RESTRICTEDPROJECT_NR}/Platforms/{PLATFORM}/Config/") },
 };
 
 constexpr int32 MaxPlatformIndex = 99;
 
-constexpr int32 GetStaticKey(int32 LayerIndex, int32 LayerExpansionIndex, int32 PlatformIndex)
+constexpr int32 GetStaticKey(int32 LayerIndex, int32 ReplacementIndex, int32 PlatformIndex)
 {
-	return LayerIndex * 10000 + LayerExpansionIndex * 100 + PlatformIndex;
+	return LayerIndex * 10000 + ReplacementIndex * 100 + PlatformIndex;
 }
 
-constexpr int32 MaxStaticHierarchyKey = GetStaticKey(UE_ARRAY_COUNT(GConfigLayers) - 1, UE_ARRAY_COUNT(GConfigLayerExpansions) - 1, MaxPlatformIndex);
+constexpr int32 MaxStaticHierarchyKey = GetStaticKey(UE_ARRAY_COUNT(GConfigLayers) - 1, UE_ARRAY_COUNT(GConfigExpansions) - 1, MaxPlatformIndex);
 
 /*-----------------------------------------------------------------------------
 	FConfigFileHierarchy
@@ -612,9 +608,9 @@ int32 FConfigFileHierarchy::GenerateDynamicKey()
 	return ++KeyGen;
 }
 
-int32 FConfigFileHierarchy::AddStaticLayer(FIniFilename Filename, int32 LayerIndex, int32 LayerExpansionIndex /*= 0*/, int32 PlatformIndex /*= 0*/)
+int32 FConfigFileHierarchy::AddStaticLayer(FIniFilename Filename, int32 LayerIndex, int32 ExpansionIndex /*= 0*/, int32 PlatformIndex /*= 0*/)
 {
-	int32 Key = GetStaticKey(LayerIndex, LayerExpansionIndex, PlatformIndex);
+	int32 Key = GetStaticKey(LayerIndex, ExpansionIndex, PlatformIndex);
 	Emplace(Key, MoveTemp(Filename));
 	return Key;
 }
@@ -3517,96 +3513,85 @@ static void ConditionalOverrideIniFilename(FString& IniFilename, const TCHAR* Ba
 #if !UE_BUILD_SHIPPING
 	// Figure out what to look for on the commandline for an override. Disabled in shipping builds for security reasons
 	const FString CommandLineSwitch = FString::Printf(TEXT("DEF%sINI="), BaseIniName);	
-	if (FParse::Value(FCommandLine::Get(), *CommandLineSwitch, IniFilename) == false)
-	{
-		// if it's not found on the commandline, then generate it	
-		FPaths::MakeStandardFilename(IniFilename);
-	}
+	FParse::Value(FCommandLine::Get(), *CommandLineSwitch, IniFilename);
 #endif
 }
 
 
-static FString GetLayerPath(const FConfigLayer& Layer, const FString& PlatformName, const TCHAR* EngineConfigDir, const TCHAR* SourceConfigDir, const TCHAR* BaseIniName, bool& bHasPlatformTag, bool& bHasProjectTag)
+static FString PerformBasicReplacements(const FString& InString, const TCHAR* BaseIniName)
 {
-	static FString LastPlatform;
-	static FString PlatformExtensionEngineConfigDir;
-	static FString PlatformExtensionProjectConfigDir;
-	static bool bHasPlatformExtensionEngineConfigDir;
-	static bool bHasPlatformExtensionProjectConfigDir;
+	FString OutString = InString.Replace(TEXT("{TYPE}"), BaseIniName, ESearchCase::CaseSensitive);
+	OutString = OutString.Replace(TEXT("{USERSETTINGS}"), FPlatformProcess::UserSettingsDir(), ESearchCase::CaseSensitive);
+	OutString = OutString.Replace(TEXT("{USER}"), FPlatformProcess::UserDir(), ESearchCase::CaseSensitive);
 
-	if (LastPlatform != PlatformName)
-	{
-		LastPlatform = PlatformName;
-		PlatformExtensionEngineConfigDir = FPaths::Combine(*FPaths::EnginePlatformExtensionsDir(), *PlatformName, TEXT("Config"));
-		PlatformExtensionProjectConfigDir = FPaths::Combine(*FPaths::ProjectPlatformExtensionsDir(), *PlatformName, TEXT("Config"));
-		bHasPlatformExtensionEngineConfigDir = FPaths::DirectoryExists(*PlatformExtensionEngineConfigDir);
-		bHasPlatformExtensionProjectConfigDir = FPaths::DirectoryExists(*PlatformExtensionProjectConfigDir);
-	}
+	return OutString;
+}
 
-	// cache some searches
-	bHasPlatformTag = FCString::Strstr(Layer.Path, TEXT("{PLATFORM}")) != nullptr;
-	bHasProjectTag = FCString::Strstr(Layer.Path, TEXT("{PROJECT}")) != nullptr;
 
-	// you can only have PROJECT or ENGINE, not both
-	FString LayerPath;
-	if (bHasProjectTag)
+static FString PerformExpansionReplacements(const FConfigLayerExpansion& Expansion, const FString& InString)
+{
+	// if there's replacement to do, the output is just the output
+	if (Expansion.Before1 == nullptr)
 	{
-		if (bHasPlatformTag && bHasPlatformExtensionProjectConfigDir)
-		{
-			LayerPath = FString(Layer.PlatformExtensionPath).Replace(TEXT("{EXTPROJECT}"), *PlatformExtensionProjectConfigDir);
-		}
-		else
-		{
-			LayerPath = FString(Layer.Path).Replace(TEXT("{PROJECT}"), SourceConfigDir);
-		}
+		return InString;
 	}
-	else
+
+	// if nothing to replace, then skip it entirely
+	if (!InString.Contains(Expansion.Before1) && (Expansion.Before2 == nullptr || !InString.Contains(Expansion.Before2)))
 	{
-		if (bHasPlatformTag && bHasPlatformExtensionEngineConfigDir)
-		{
-			LayerPath = FString(Layer.PlatformExtensionPath).Replace(TEXT("{EXTENGINE}"), *PlatformExtensionEngineConfigDir);
-		}
-		else
-		{
-			LayerPath = FString(Layer.Path).Replace(TEXT("{ENGINE}"), EngineConfigDir);
-		}
+		return TEXT("");
 	}
 
-	LayerPath = LayerPath.Replace(TEXT("{TYPE}"), BaseIniName, ESearchCase::CaseSensitive);
-	LayerPath = LayerPath.Replace(TEXT("{USERSETTINGS}"), FPlatformProcess::UserSettingsDir(), ESearchCase::CaseSensitive);
-	LayerPath = LayerPath.Replace(TEXT("{USER}"), FPlatformProcess::UserDir(), ESearchCase::CaseSensitive);
-
-	return LayerPath;
+	// replace the directory bits
+	FString OutString = InString.Replace(Expansion.Before1, Expansion.After1, ESearchCase::CaseSensitive);
+	if (Expansion.Before2 != nullptr)
+	{
+		OutString = OutString.Replace(Expansion.Before2, Expansion.After2, ESearchCase::CaseSensitive);
+	}
+	return OutString;
 }
 
-static FString GetExpansionPath(const FConfigLayerExpansion& Expansion, const FString& LayerPath, bool bHasPlatformTag)
+static FString PerformFinalExpansions(const FString InString, const FString& PlatformName, const TCHAR* EngineConfigDir, const TCHAR* SourceConfigDir)
 {
-	// replace the expansion tags
-	FString ExpansionPath = LayerPath.Replace(TEXT("{ED}"), Expansion.DirectoryPrefix, ESearchCase::CaseSensitive);
-	ExpansionPath = ExpansionPath.Replace(TEXT("{EF}"), Expansion.FilePrefix, ESearchCase::CaseSensitive);
+	static FString LastPlatform;
+	static FString PlatformExtensionEngineDir;
+	static FString PlatformExtensionProjectDir;
+	static FString ProjectNotForLicenseesDir;
+	static FString ProjectNoRedistDir;
 
-	// check for dedicated server expansion
-	if (EnumHasAnyFlags(Expansion.Flag, EConfigLayerFlags::DedicatedServerOnly))
+	if (LastPlatform != PlatformName)
 	{
-		// it's unclear how a platform DS ini would be named, so for now, not supported
-		if (bHasPlatformTag)
-		{
-			return TEXT("");
-		}
-		else if (IsRunningDedicatedServer())
+		LastPlatform = PlatformName;
+		PlatformExtensionEngineDir = FPaths::Combine(*FPaths::EnginePlatformExtensionsDir(), *PlatformName);
+		PlatformExtensionProjectDir = FPaths::Combine(*FPaths::ProjectPlatformExtensionsDir(), *PlatformName);
+	}
+
+	// cache some slow operations
+	if (ProjectNotForLicenseesDir.Len() == 0)
+	{
+		if (FPaths::IsUnderDirectory(FPaths::ProjectDir(), FPaths::EngineDir()))
 		{
-			ExpansionPath = ExpansionPath.Replace(TEXT("Base"), TEXT(""));
-			ExpansionPath = ExpansionPath.Replace(TEXT("Default"), TEXT(""));
-			// ExpansionPath = ExpansionPath.Replace(TEXT("{PLATFORM}"), TEXT(""));
+			FString RelativeDir = FPaths::ProjectDir();
+			FPaths::MakePathRelativeTo(RelativeDir, *FPaths::EngineDir());
+			ProjectNotForLicenseesDir = FPaths::Combine(FPaths::EngineDir(), TEXT("Restricted/NotForLicensees"), RelativeDir);
+			ProjectNoRedistDir = FPaths::Combine(FPaths::EngineDir(), TEXT("Restricted/NoRedist"), RelativeDir);
 		}
 		else
 		{
-			// skip this expansion if not DS
-			return TEXT("");
+			ProjectNotForLicenseesDir = FPaths::Combine(FPaths::ProjectDir(), TEXT("Restricted/NotForLicensees"));
+			ProjectNoRedistDir = FPaths::Combine(FPaths::ProjectDir(), TEXT("Restricted/NoRedist"));
 		}
 	}
 
-	return ExpansionPath;
+	FString OutString = InString.Replace(TEXT("{ENGINE}"), EngineConfigDir);
+	OutString = OutString.Replace(TEXT("{EXTENGINE}"), *PlatformExtensionEngineDir);
+	OutString = OutString.Replace(TEXT("{PROJECT}"), SourceConfigDir);
+	OutString = OutString.Replace(TEXT("{EXTPROJECT}"), *PlatformExtensionProjectDir);
+	OutString = OutString.Replace(TEXT("{PLATFORM}"), *PlatformName);
+	OutString = OutString.Replace(TEXT("{RESTRICTEDPROJECT_NFL}"), *ProjectNotForLicenseesDir);
+	OutString = OutString.Replace(TEXT("{RESTRICTEDPROJECT_NR}"), *ProjectNoRedistDir);
+
+	return OutString;
 }
 
 /**
@@ -3621,9 +3606,19 @@ void FConfigFile::AddStaticLayersToHierarchy(const TCHAR* InBaseIniName, const T
 	SourceEngineConfigDir = EngineConfigDir;
 	SourceProjectConfigDir = SourceConfigDir;
 
+	// for the replacement, we need to have a directory called Config - or we will have to do extra processing for these non-standard cases
+	check(SourceEngineConfigDir.EndsWith(TEXT("Config/")));
+	check(SourceProjectConfigDir.EndsWith(TEXT("Config/")));
+
+	FString UsedEngineDir = FPaths::GetPath(FPaths::GetPath(SourceEngineConfigDir));
+	FString UsedProjectDir = FPaths::GetPath(FPaths::GetPath(SourceProjectConfigDir));
+
 	// get the platform name
 	const FString PlatformName(InPlatformName ? InPlatformName : ANSI_TO_TCHAR(FPlatformProperties::IniPlatformName()));
 
+	// string that can have a reference to it, lower down
+	const FString DedicatedServerString = IsRunningDedicatedServer() ? TEXT("DedicatedServer") : TEXT("");
+
 	// cache some platform extension information that can be used inside the loops
 
 	// go over all the config layers
@@ -3632,8 +3627,8 @@ void FConfigFile::AddStaticLayersToHierarchy(const TCHAR* InBaseIniName, const T
 		const FConfigLayer& Layer = GConfigLayers[LayerIndex];
 
 		// start replacing basic variables
-		bool bHasPlatformTag, bHasProjectTag;
-		FString LayerPath = GetLayerPath(Layer, PlatformName, EngineConfigDir, SourceConfigDir, InBaseIniName, bHasPlatformTag, bHasProjectTag);
+		FString LayerPath = PerformBasicReplacements(Layer.Path, InBaseIniName);
+		bool bHasPlatformTag = LayerPath.Contains(TEXT("{PLATFORM}"));
 
 		// PROGRAMs don't require any ini files
 #if IS_PROGRAM
@@ -3643,7 +3638,7 @@ void FConfigFile::AddStaticLayersToHierarchy(const TCHAR* InBaseIniName, const T
 #endif
 
 		// expand if it it has {ED} or {EF} expansion tags
-		if (FCString::Strstr(Layer.Path, TEXT("{ED}")) || FCString::Strstr(Layer.Path, TEXT("{EF}")))
+		if (!EnumHasAnyFlags(Layer.Flag, EConfigLayerFlags::NoExpand))
 		{
 			// we assume none of the more special tags in expanded ones
 			checkfSlow(FCString::Strstr(Layer.Path, TEXT("{USERSETTINGS}")) == nullptr && FCString::Strstr(Layer.Path, TEXT("{USER}")) == nullptr, TEXT("Expanded config %s shouldn't have a {USER*} tags in it"), *Layer.Path);
@@ -3651,58 +3646,57 @@ void FConfigFile::AddStaticLayersToHierarchy(const TCHAR* InBaseIniName, const T
 			checkfSlow(LayerIndex < UE_ARRAY_COUNT(GConfigLayers) - 1, TEXT("Final layer %s shouldn't be an expansion layer, as it needs to generate the  hierarchy cache key"), *Layer.Path);
 
 			// loop over all the possible expansions
-			for (int32 ExpansionIndex = 0; ExpansionIndex < UE_ARRAY_COUNT(GConfigLayerExpansions); ExpansionIndex++)
+			for (int32 ExpansionIndex = 0; ExpansionIndex < UE_ARRAY_COUNT(GConfigExpansions); ExpansionIndex++)
 			{
-				const FConfigLayerExpansion& Expansion = GConfigLayerExpansions[ExpansionIndex];
+				FString ExpandedPath = PerformExpansionReplacements(GConfigExpansions[ExpansionIndex], LayerPath);
 
-				bool bSkipExpansion = false;
-				FString ExpansionPath = GetExpansionPath(Expansion, LayerPath, bHasPlatformTag);
-
-				// skip this expansion if we didn't want it
-				if (ExpansionPath.Len() == 0)
+				// if we didn't replace anything, skip it
+				if (ExpandedPath.Len() == 0)
 				{
 					continue;
 				}
-
+				
 				// allow for override, only on BASE EXPANSION!
 				if (EnumHasAnyFlags(Layer.Flag, EConfigLayerFlags::AllowCommandLineOverride) && ExpansionIndex == 0)
 				{
 					checkfSlow(!bHasPlatformTag, TEXT("EConfigLayerFlags::AllowCommandLineOverride config %s shouldn't have a PLATFORM in it"), *Layer.Path);
 
-					ConditionalOverrideIniFilename(ExpansionPath, InBaseIniName);
+					ConditionalOverrideIniFilename(ExpandedPath, InBaseIniName);
 				}
 
 				// check if we should be generating the cache key - only at the end of all expansions
-				bool bGenerateCacheKey = EnumHasAnyFlags(Layer.Flag, EConfigLayerFlags::GenerateCacheKey) && ExpansionIndex == UE_ARRAY_COUNT(GConfigLayerExpansions);
+				bool bGenerateCacheKey = EnumHasAnyFlags(Layer.Flag, EConfigLayerFlags::GenerateCacheKey) && ExpansionIndex == UE_ARRAY_COUNT(GConfigExpansions) - 1;
 				checkfSlow(!(bGenerateCacheKey && bHasPlatformTag), TEXT("EConfigLayerFlags::GenerateCacheKey shouldn't have a platform tag"));
 
 				const FDataDrivenPlatformInfoRegistry::FPlatformInfo& Info = FDataDrivenPlatformInfoRegistry::GetPlatformInfo(PlatformName);
 
 				// go over parents, and then this platform, unless there's no platform tag, then we simply want to run through the loop one time to add it to the
-				int NumPlatforms = bHasPlatformTag ? Info.IniParentChain.Num() + 1 : 1;
+				int32 NumPlatforms = bHasPlatformTag ? Info.IniParentChain.Num() + 1 : 1;
+				int32 CurrentPlatformIndex = NumPlatforms - 1;
+				int32 DedicatedServerIndex = -1;
+				
+				// make DedicatedServer another platform
+				if (bHasPlatformTag && IsRunningDedicatedServer())
+				{
+					NumPlatforms++;
+					DedicatedServerIndex = CurrentPlatformIndex + 1;
+				}
+
 				check(NumPlatforms < MaxPlatformIndex);
 				for (int PlatformIndex = 0; PlatformIndex < NumPlatforms; PlatformIndex++)
 				{
-					// the platform to work on (active platform is always last)
-					const FString& CurrentPlatform = (PlatformIndex == NumPlatforms - 1) ? PlatformName : Info.IniParentChain[PlatformIndex];
-					FString PlatformPath;
-
-					// ini parent may not be a PlatformExtension even if this was, so update based on the current platform after the first
-					if (PlatformIndex < NumPlatforms - 1)
-					{
-						// overriding the bHas*Tag fields is fine here since it will evaluate to the same
-						FString LocalLayerPath = GetLayerPath(Layer, CurrentPlatform, EngineConfigDir, SourceConfigDir, InBaseIniName, bHasPlatformTag, bHasProjectTag);
+					const FString& CurrentPlatform = 
+						(PlatformIndex == DedicatedServerIndex) ? DedicatedServerString :
+						(PlatformIndex == CurrentPlatformIndex) ? PlatformName : 
+						Info.IniParentChain[PlatformIndex];
 
-						// re-expand the new path
-						FString LocalExpansionPath = GetExpansionPath(Expansion, LocalLayerPath, bHasPlatformTag);
+					FString PlatformPath = PerformFinalExpansions(ExpandedPath, CurrentPlatform, *UsedEngineDir, *UsedProjectDir);
 
-						// modify get the final platform path from the parent ini's location
-						PlatformPath = LocalExpansionPath.Replace(TEXT("{PLATFORM}"), *CurrentPlatform, ESearchCase::CaseSensitive);
-					}
-					else
+					// @todo restricted - ideally, we would move DedicatedServer files into a directory, like platforms are, but for short term compat,
+					// convert the path back to the original (DedicatedServer/DedicatedServerEngine.ini -> DedicatedServerEngine.ini)
+					if (PlatformIndex == DedicatedServerIndex)
 					{
-						// replace
-						PlatformPath = ExpansionPath.Replace(TEXT("{PLATFORM}"), *CurrentPlatform, ESearchCase::CaseSensitive);
+						PlatformPath.ReplaceInline(TEXT("Config/DedicatedServer/"), TEXT("Config/"));
 					}
 
 					// add this to the list!
@@ -3720,10 +3714,12 @@ void FConfigFile::AddStaticLayersToHierarchy(const TCHAR* InBaseIniName, const T
 			checkfSlow(!bHasPlatformTag, TEXT("Non-expanded config %s shouldn't have a PLATFORM in it"), *Layer.Path);
 			checkfSlow(!EnumHasAnyFlags(Layer.Flag, EConfigLayerFlags::AllowCommandLineOverride), TEXT("Non-expanded config can't have a EConfigLayerFlags::AllowCommandLineOverride"));
 
+			FString FinalPath = PerformFinalExpansions(LayerPath, TEXT(""), *UsedEngineDir, *UsedProjectDir);
+
 			// final layer needs to generate a hierarchy cache
 			SourceIniHierarchy.AddStaticLayer(
-				FIniFilename(LayerPath, bIsRequired, EnumHasAnyFlags(Layer.Flag, EConfigLayerFlags::GenerateCacheKey) ?
-					GenerateHierarchyCacheKey(SourceIniHierarchy, LayerPath, InBaseIniName) :
+				FIniFilename(FinalPath, bIsRequired, EnumHasAnyFlags(Layer.Flag, EConfigLayerFlags::GenerateCacheKey) ?
+					GenerateHierarchyCacheKey(SourceIniHierarchy, FinalPath, InBaseIniName) :
 					FString(TEXT(""))),
 				LayerIndex);
 		}
diff --git a/Engine/Source/Runtime/Core/Public/Misc/ConfigCacheIni.h b/Engine/Source/Runtime/Core/Public/Misc/ConfigCacheIni.h
--- a/Engine/Source/Runtime/Core/Public/Misc/ConfigCacheIni.h
+++ b/Engine/Source/Runtime/Core/Public/Misc/ConfigCacheIni.h
@@ -313,7 +313,7 @@ public:
 private:
 	int32 GenerateDynamicKey();
 
-	int32 AddStaticLayer(FIniFilename Filename, int32 LayerIndex, int32 LayerExpansionIndex = 0, int32 PlatformIndex = 0);
+	int32 AddStaticLayer(FIniFilename Filename, int32 LayerIndex, int32 ExpansionIndex=0, int32 PlatformIndex=0);
 	int32 AddDynamicLayer(FIniFilename Filename);
 
 	friend class FConfigFile;
