FHashedName converts strings to upper-case before hashing, otherwise hash depends on what particular case of a given string is first seen

diff --git a/Engine/Plugins/Compositing/OpenColorIO/Source/OpenColorIO/Private/OpenColorIODerivedDataVersion.h b/Engine/Plugins/Compositing/OpenColorIO/Source/OpenColorIO/Private/OpenColorIODerivedDataVersion.h
--- a/Engine/Plugins/Compositing/OpenColorIO/Source/OpenColorIO/Private/OpenColorIODerivedDataVersion.h
+++ b/Engine/Plugins/Compositing/OpenColorIO/Source/OpenColorIO/Private/OpenColorIODerivedDataVersion.h
@@ -10,4 +10,4 @@ OpenColorIOShaderDerivedDataVersion.h: Derived data version for OpenColorIO shad
 
 // In case of merge conflicts with DDC versions, you MUST generate a new GUID and set this new
 // guid as version
-#define OPENCOLORIO_DERIVEDDATA_VER		TEXT("176398465B0E4D26A11F2DE8FA3E6916")
\ No newline at end of file
+#define OPENCOLORIO_DERIVEDDATA_VER		TEXT("FECF9DF3-80B5-4551-B341-7244B95FDEE7")
\ No newline at end of file
diff --git a/Engine/Plugins/FX/Niagara/Source/NiagaraShader/Private/NiagaraShaderDerivedDataVersion.h b/Engine/Plugins/FX/Niagara/Source/NiagaraShader/Private/NiagaraShaderDerivedDataVersion.h
--- a/Engine/Plugins/FX/Niagara/Source/NiagaraShader/Private/NiagaraShaderDerivedDataVersion.h
+++ b/Engine/Plugins/FX/Niagara/Source/NiagaraShader/Private/NiagaraShaderDerivedDataVersion.h
@@ -11,4 +11,4 @@ NiagaraShaderDerivedDataVersion.h: Shader derived data version for Niagara.
 // In case of merge conflicts with DDC versions, you MUST generate a new GUID and set this new
 // guid as version
 
-#define NIAGARASHADERMAP_DERIVEDDATA_VER		TEXT("40A7B0D0F38847349C93BD4445DAEDA5")
+#define NIAGARASHADERMAP_DERIVEDDATA_VER		TEXT("3136AAAD-A9DE-4368-9F86-E89840E7CC87")
diff --git a/Engine/Shaders/Public/ShaderVersion.ush b/Engine/Shaders/Public/ShaderVersion.ush
--- a/Engine/Shaders/Public/ShaderVersion.ush
+++ b/Engine/Shaders/Public/ShaderVersion.ush
@@ -3,4 +3,4 @@
 // in Platform.ush (which should be included in any shader) it allows to invalidate the shader DDC.
 // 
 // If you are merging streams and there is a conflict with this GUID you should make a new GUID rather than taking one or the other.
-// GUID = E54ED269-25A1-4064-884C-D3017370E4FA
+// GUID = F0FE2A62-52BB-48E7-8FEB-0E8766DA3D1D
diff --git a/Engine/Source/Runtime/Core/Private/Serialization/MemoryImage.cpp b/Engine/Source/Runtime/Core/Private/Serialization/MemoryImage.cpp
--- a/Engine/Source/Runtime/Core/Private/Serialization/MemoryImage.cpp
+++ b/Engine/Source/Runtime/Core/Private/Serialization/MemoryImage.cpp
@@ -740,6 +740,10 @@ FHashedName::FHashedName(const FName& InName)
 		{
 			WIDECHAR WideNameBuffer[NAME_SIZE];
 			Entry->GetWideName(WideNameBuffer);
+			for (int32 i = 0; i < Entry->GetNameLength(); ++i)
+			{
+				WideNameBuffer[i] = FCharWide::ToUpper(WideNameBuffer[i]);
+			}
 			const FTCHARToUTF8 NameUTF8(WCHAR_TO_TCHAR(WideNameBuffer));
 			Hash = CityHash64WithSeed(NameUTF8.Get(), NameUTF8.Length(), InternalNumber);
 		}
@@ -747,6 +751,10 @@ FHashedName::FHashedName(const FName& InName)
 		{
 			ANSICHAR AnsiNameBuffer[NAME_SIZE];
 			Entry->GetAnsiName(AnsiNameBuffer);
+			for (int32 i = 0; i < Entry->GetNameLength(); ++i)
+			{
+				AnsiNameBuffer[i] = FCharAnsi::ToUpper(AnsiNameBuffer[i]);
+			}
 			Hash = CityHash64WithSeed(AnsiNameBuffer, Entry->GetNameLength(), InternalNumber);
 		}
 	}
