#UE4 IWYU fixes

diff --git a/Engine/Source/Runtime/Core/Public/Containers/HashTable.h b/Engine/Source/Runtime/Core/Public/Containers/HashTable.h
--- a/Engine/Source/Runtime/Core/Public/Containers/HashTable.h
+++ b/Engine/Source/Runtime/Core/Public/Containers/HashTable.h
@@ -3,6 +3,8 @@
 #pragma once
 
 #include "CoreTypes.h"
+#include "Containers/ContainerAllocationPolicies.h"
+#include "Serialization/MemoryImageWriter.h"
 #include "Serialization/MemoryLayout.h"
 #include "Serialization/MemoryImageWriter.h"
 #include "Containers/ContainerAllocationPolicies.h"
