Fix iOS non-unity build errors
#rnx

diff --git a/Engine/Source/Runtime/Core/Private/IOS/IOSPlatformCrashContext.cpp b/Engine/Source/Runtime/Core/Private/IOS/IOSPlatformCrashContext.cpp
--- a/Engine/Source/Runtime/Core/Private/IOS/IOSPlatformCrashContext.cpp
+++ b/Engine/Source/Runtime/Core/Private/IOS/IOSPlatformCrashContext.cpp
@@ -8,10 +8,14 @@
 #include "IOS/IOSPlatformPLCrashReporterIncludes.h"
 
 #include "HAL/ExceptionHandling.h"
+#include "HAL/FileManager.h"
 #include "HAL/PlatformOutputDevices.h"
 
+#include "Internationalization/Culture.h"
+
 #include "Misc/App.h"
 #include "Misc/EngineVersion.h"
+#include "Miscs/Paths.h"
 
 
 FIOSApplicationInfo GIOSAppInfo;
