Undefine GetFreeSpace() windows Macro.

diff --git a/Engine/Source/Runtime/Core/Public/Windows/MinWindows.h b/Engine/Source/Runtime/Core/Public/Windows/MinWindows.h
--- a/Engine/Source/Runtime/Core/Public/Windows/MinWindows.h
+++ b/Engine/Source/Runtime/Core/Public/Windows/MinWindows.h
@@ -66,3 +66,5 @@ struct IUnknown;
 
 // Finally now we can include windows.h
 #include <Windows.h>
+
+#undef GetFreeSpace
