Optimize SaveConfig during editor startup
  - 757 ms -> 127ms

#rnx

diff --git a/Engine/Source/Runtime/Core/Private/Misc/ConfigCacheIni.cpp b/Engine/Source/Runtime/Core/Private/Misc/ConfigCacheIni.cpp
--- a/Engine/Source/Runtime/Core/Private/Misc/ConfigCacheIni.cpp
+++ b/Engine/Source/Runtime/Core/Private/Misc/ConfigCacheIni.cpp
@@ -1371,18 +1371,21 @@ bool DoesConfigPropertyValueMatch( FConfigFile* InConfigFile, const FString& InS
 
 		if( Section )
 		{
+			const bool bIsInputStringValidFloat = FDefaultValueHelper::IsStringValidFloat(InPropertyValue);
+
 			// Start Array check, if the property is in an array, we need to iterate over all properties.
 			for (FConfigSection::TConstKeyIterator It(*Section, InPropertyName); It && !bFoundAMatch; ++It)
 			{
 				const FString& PropertyValue = It.Value().GetSavedValue();
-				bFoundAMatch = PropertyValue == InPropertyValue;
-
+				bFoundAMatch = 
+					PropertyValue.Len() == InPropertyValue.Len() &&
+					PropertyValue == InPropertyValue;
+				
 				// if our properties don't match, run further checks
 				if( !bFoundAMatch )
 				{
 					// Check that the mismatch isn't just a string comparison issue with floats
-					if( FDefaultValueHelper::IsStringValidFloat( PropertyValue ) &&
-						FDefaultValueHelper::IsStringValidFloat( InPropertyValue ) )
+					if (bIsInputStringValidFloat && FDefaultValueHelper::IsStringValidFloat( PropertyValue ))
 					{
 						bFoundAMatch = FCString::Atof( *PropertyValue ) == FCString::Atof( *InPropertyValue );
 					}
