iOS Asan memory fixes.

Attempts to keep previous behavior. Will first try to allocate < 256 bytes from system malloc.
If the system returns a pointer outside of its nanozone we will fall back to using UE's allocators.

Determines the location of the nano malloc region via the mach vm. The system always creates a contiguous region tagged a specific way.
If the region is not present, falls back to UE allocators.

diff --git a/Engine/Source/Runtime/Core/Private/Apple/ApplePlatformMemory.cpp b/Engine/Source/Runtime/Core/Private/Apple/ApplePlatformMemory.cpp
--- a/Engine/Source/Runtime/Core/Private/Apple/ApplePlatformMemory.cpp
+++ b/Engine/Source/Runtime/Core/Private/Apple/ApplePlatformMemory.cpp
@@ -28,6 +28,9 @@
 #include <os/proc.h>
 #endif
 #include <CoreFoundation/CFBase.h>
+#include <mach/vm_map.h>
+#include <mach/vm_region.h>
+#include <mach/vm_statistics.h>
 #include "HAL/LowLevelMemTracker.h"
 #include "Apple/AppleLLM.h"
 
@@ -186,6 +189,79 @@ void FApplePlatformMemory::ConfigureDefaultCFAllocator(void)
 	CFAllocatorSetDefault(Alloc);
 }
 
+vm_address_t FApplePlatformMemory::NanoRegionStart = 0;
+vm_address_t FApplePlatformMemory::NanoRegionEnd = 0;
+
+void FApplePlatformMemory::NanoMallocInit()
+{
+	/*
+		iOS reserves 512MB of address space for 'nano' allocations (allocations <= 256 bytes)
+		Nano malloc has buckets for sizes 16, 32, 48....256
+		The number of buckets and their sizes are fixed and do not grow
+		We'll walk through the buckets and ask the VM about the backing regions
+		We may have to check several sizes because we can hit a case where all the buckets
+		for a specific size are full - which means malloc will put that allocation into
+		the MALLOC_TINY region instead.
+	 
+		The OS always tags the nano VM region with user_tag == VM_MEMORY_MALLOC_NANO (which is 11)
+	 
+		Being apple this is subject to change at any time and may be different in debug modes, etc.
+		We'll fall back to the UE allocators if we can't find the nano region.
+	 
+		We want to detect this as early as possible, before any of the memory system is initialized.
+	*/
+	
+	NanoRegionStart = 0;
+	NanoRegionEnd = 0;
+	
+	size_t MallocSize = 16;
+	while(true)
+	{
+		void* NanoMalloc = ::malloc(MallocSize);
+		FMemory::Memzero(NanoMalloc, MallocSize); // This will wire the memory. Shouldn't be necessary but better safe than sorry.
+	
+		kern_return_t kr = KERN_SUCCESS;
+		vm_address_t address = (vm_address_t)(NanoMalloc);
+		vm_size_t regionSizeInBytes = 0;
+		mach_port_t regionObjectOut;
+		vm_region_extended_info_data_t regionInfo;
+		mach_msg_type_number_t infoSize = sizeof(vm_region_extended_info_data_t);
+		kr = vm_region_64(mach_task_self(), &address, &regionSizeInBytes, VM_REGION_EXTENDED_INFO, (vm_region_info_t) &regionInfo, &infoSize, &regionObjectOut);
+		check(kr == KERN_SUCCESS);
+		
+		::free(NanoMalloc);
+		
+		if(regionInfo.user_tag == VM_MEMORY_MALLOC_NANO)
+		{
+			uint8_t* Start = (uint8_t*) address;
+			uint8_t* End = Start + regionSizeInBytes;
+			NanoRegionStart = address;
+			NanoRegionEnd = (vm_address_t) End;
+			break;
+		}
+		
+		MallocSize += 16;
+		
+		if(MallocSize > 256)
+		{
+			// Nano region wasn't found.
+			// We'll fall back to the UE allocator
+			// This can happen when using various tools
+			check(NanoRegionStart == 0 && NanoRegionEnd == 0);
+			break;
+		}
+	}
+	
+	if(NanoRegionStart == 0 && NanoRegionEnd == 0)
+	{
+		FPlatformMisc::LowLevelOutputDebugStringf(TEXT("WARNING: No nano malloc region found. We will always use UE allocators\n"));
+	}
+	else
+	{
+		FPlatformMisc::LowLevelOutputDebugStringf(TEXT("Detected nanozone %p - %p\n"), (void*) NanoRegionStart, (void*) NanoRegionEnd);
+	}
+}
+
 void FApplePlatformMemory::Init()
 {
 	FGenericPlatformMemory::Init();
@@ -198,6 +274,7 @@ void FApplePlatformMemory::Init()
 		   MemoryConstants.TotalPhysicalGB,
 		   float((MemoryConstants.TotalVirtual-MemoryConstants.TotalPhysical)/1024.0/1024.0/1024.0),
 		   float(MemoryConstants.TotalVirtual/1024.0/1024.0/1024.0) );
+	
 }
 
 // Set rather to use BinnedMalloc2 for binned malloc, can be overridden below
@@ -613,6 +690,21 @@ void FApplePlatformMemory::BinnedFreeToOS(void* Ptr, SIZE_T Size)
 #endif // USE_MALLOC_BINNED2
 }
 
+bool FApplePlatformMemory::PtrIsOSMalloc( void* Ptr)
+{
+	return malloc_zone_from_ptr(Ptr) != nullptr;
+}
+
+bool FApplePlatformMemory::IsNanoMallocAvailable()
+{
+	return (NanoRegionStart != 0) && (NanoRegionEnd != 0);
+}
+
+bool FApplePlatformMemory::PtrIsFromNanoMalloc( void* Ptr)
+{
+	return IsNanoMallocAvailable() && ((uintptr_t) Ptr >= NanoRegionStart && (uintptr_t) Ptr < NanoRegionEnd);
+}
+
 size_t FApplePlatformMemory::FPlatformVirtualMemoryBlock::GetVirtualSizeAlignment()
 {
 	static SIZE_T OSPageSize = FPlatformMemory::GetConstants().PageSize;
diff --git a/Engine/Source/Runtime/Core/Private/HAL/MallocBinned.cpp b/Engine/Source/Runtime/Core/Private/HAL/MallocBinned.cpp
--- a/Engine/Source/Runtime/Core/Private/HAL/MallocBinned.cpp
+++ b/Engine/Source/Runtime/Core/Private/HAL/MallocBinned.cpp
@@ -168,11 +168,8 @@ struct FMallocBinned::Private
 	static bool bHasInitializedStatsMetadata;
 	
 #if USE_OS_SMALL_BLOCK_ALLOC
-	static uint64 SmallBlockStartPtr;
-	static uint64 SmallBlockEndPtr;
-	enum { SMALL_BLOCK_MAX_TOTAL_POOL_SIZE = 0x20000000 };
-
 #if USE_OS_SMALL_BLOCK_GRAB_MEMORY_FROM_OS
+	enum { SMALL_BLOCK_MAX_TOTAL_POOL_SIZE = 0x20000000 };
 	enum { SMALL_BLOCK_GRAB_ALLOC_ALIGN = 16 };
 	enum { SMALL_BLOCK_GRAB_MAX_ALLOC_SIZE = 256 };
 	enum { SMALL_BLOCK_GRAB_MIN_ALLOC_SIZE = SMALL_BLOCK_GRAB_ALLOC_ALIGN };
@@ -708,23 +705,15 @@ struct FMallocBinned::Private
 		MEM_TIME(MemTime -= FPlatformTime::Seconds());
 		BINNED_DECREMENT_STATCOUNTER(Allocator.CurrentAllocs);
 
-		UPTRINT BasePtr;
-		FPoolInfo* Pool = FindPoolInfo(Allocator, (UPTRINT)Ptr, BasePtr);
 #if PLATFORM_IOS || PLATFORM_MAC
-        if (Pool == nullptr)
-        {
-#if USE_OS_SMALL_BLOCK_ALLOC && !USE_OS_SMALL_BLOCK_GRAB_MEMORY_FROM_OS
-			// check the small block allocation range
-			if ((uint64)Ptr >= SmallBlockStartPtr && (uint64)Ptr < SmallBlockEndPtr)
-			{
-				SmallOSFree(Allocator, Ptr, Private::SMALL_BLOCK_POOL_SIZE);
-				return;
-			}
-#endif
-            UE_LOG(LogMemory, Warning, TEXT("Attempting to free a pointer we didn't allocate!"));
-            return;
-        }
+		if(FPlatformMemory::PtrIsOSMalloc(Ptr))
+		{
+			SmallOSFree(Allocator, Ptr, Private::SMALL_BLOCK_POOL_SIZE);
+			return;
+		}
 #endif
+		UPTRINT BasePtr;
+		FPoolInfo* Pool = FindPoolInfo(Allocator, (UPTRINT)Ptr, BasePtr);
 		checkSlow(Pool);
 		checkSlow(Pool->GetBytes() != 0);
 		if (Pool->TableIndex < Allocator.BinnedOSTableIndex)
@@ -1019,11 +1008,6 @@ struct FMallocBinned::Private
 	}
 };
 
-#if USE_OS_SMALL_BLOCK_ALLOC
-uint64 FMallocBinned::Private::SmallBlockStartPtr = 0;
-uint64 FMallocBinned::Private::SmallBlockEndPtr = 0;
-#endif
-
 #if USE_OS_SMALL_BLOCK_GRAB_MEMORY_FROM_OS
 uint64 FMallocBinned::Private::SmallBlockGrab_FreeStartPointers[MAX_NUM_BLOCK_START_ADDRESSES];
 int FMallocBinned::Private::NumFreeSmallBlockGrabAllocations = 0;
@@ -1171,17 +1155,9 @@ FMallocBinned::FMallocBinned(uint32 InPageSize, uint64 AddressLimit)
 	check(AddressLimit > PageSize); // Check to catch 32 bit overflow in AddressLimit
 
 #if USE_OS_SMALL_BLOCK_ALLOC
-	//Do very early (very small) malloc so we can find where the initial nano malloc pool is
-	void* ptr = ::malloc(16);
-	
-	//Round back down to 00's address
-	Private::SmallBlockStartPtr = (uint64_t)ptr & ~(0x1fffffff);
 	
-	//Add pool size to start pointer. Magic number comes from found instruments pool size of 512MB
-	Private::SmallBlockEndPtr = Private::SmallBlockStartPtr + Private::SMALL_BLOCK_MAX_TOTAL_POOL_SIZE;
+	FPlatformMemory::NanoMallocInit();
 	
-	::free(ptr);
-
 #if USE_OS_SMALL_BLOCK_GRAB_MEMORY_FROM_OS
 
 	void* GrabberHelperLoc = ::malloc(sizeof(Private::FMemoryAllocationGrabberHelper));
@@ -1330,7 +1306,7 @@ void* FMallocBinned::Malloc(SIZE_T Size, uint32 Alignment)
 	FFreeMem* Free = nullptr;
 	bool bUsePools = true;
 #if USE_OS_SMALL_BLOCK_ALLOC && !USE_OS_SMALL_BLOCK_GRAB_MEMORY_FROM_OS
-	if (Size <= Private::SMALL_BLOCK_POOL_SIZE)
+	if (FPlatformMemory::IsNanoMallocAvailable() && Size <= Private::SMALL_BLOCK_POOL_SIZE)
 	{
 		//Make sure we have initialized our hash buckets even if we are using the NANO_MALLOC grabber, as otherwise we can end
 		//up making bad assumptions and trying to grab invalid data during a Realloc of this data.
@@ -1338,19 +1314,20 @@ void* FMallocBinned::Malloc(SIZE_T Size, uint32 Alignment)
 		{
 			HashBuckets = Private::CreateHashBuckets(*this);
 		}
-		
+
 		bUsePools = false;
 		UPTRINT AlignedSize = Align(Size, Alignment);
 		SIZE_T ActualPoolSize; //TODO: use this to reduce waste?
 		Free = (FFreeMem*)Private::SmallOSAlloc(*this, AlignedSize, ActualPoolSize);
+		check(FPlatformMemory::PtrIsOSMalloc(Free));
 		
-		if ((uint64)Free < Private::SmallBlockStartPtr || (uint64)Free >= Private::SmallBlockEndPtr)
+		if(!FPlatformMemory::PtrIsFromNanoMalloc(Free))
 		{
-			// if this happens then we need to fall in to the bins below
-//			FPlatformMisc::LowLevelOutputDebugStringf(TEXT("Overflowed!!! Aligned Size:%d"), AlignedSize);
+			// This means we've overflowed the nano zone's internal buckets, which are fixed
+			// So we need to fall back to UE's allocator
 			Private::SmallOSFree(*this, Free, AlignedSize);
-			Free = nullptr;
 			bUsePools = true;
+			Free = nullptr;
 		}
 	}
 #endif
@@ -1436,6 +1413,10 @@ void* FMallocBinned::Malloc(SIZE_T Size, uint32 Alignment)
 			BINNED_PEAK_STATCOUNTER(UsedPeak, BINNED_ADD_STATCOUNTER(UsedCurrent, Size));
 			BINNED_PEAK_STATCOUNTER(WastePeak, BINNED_ADD_STATCOUNTER(WasteCurrent, (int64)(AlignedSize - Size)));
 		}
+		
+#if USE_OS_SMALL_BLOCK_ALLOC
+		check(!FPlatformMemory::PtrIsOSMalloc(Free));
+#endif
 	}
 
 	MEM_TIME(MemTime += FPlatformTime::Seconds());
@@ -1462,29 +1443,25 @@ void* FMallocBinned::Realloc( void* Ptr, SIZE_T NewSize, uint32 Alignment )
 	void* NewPtr = Ptr;
 	if( Ptr && NewSize )
 	{
-		FPoolInfo* Pool = Private::FindPoolInfo(*this, (UPTRINT)Ptr, BasePtr);
-		
 #if USE_OS_SMALL_BLOCK_ALLOC && !USE_OS_SMALL_BLOCK_GRAB_MEMORY_FROM_OS
-		// special case for really small blocks
-		if ((Pool == nullptr)
-			&& ((uint64)Ptr >= Private::SmallBlockStartPtr)
-			&& ((uint64)Ptr < Private::SmallBlockEndPtr))
+		if(FPlatformMemory::PtrIsOSMalloc(Ptr))
 		{
 			NewPtr = ::realloc(Ptr, NewSize);
-
-			if ((uint64)NewPtr < Private::SmallBlockStartPtr || (uint64)NewPtr >= Private::SmallBlockEndPtr)
+			
+			if(!FPlatformMemory::PtrIsFromNanoMalloc(NewPtr))
 			{
-				// if this happens then we need to Malloc and copy
+				// We've overflowed the nano region
+				// Fall back to UE's allocator
 				Ptr = NewPtr;
 				NewPtr = Malloc(NewSizeUnmodified, Alignment);
 				FMemory::Memcpy(NewPtr, Ptr, NewSize);
-				::free( Ptr );
+				Private::SmallOSFree(*this, Ptr, NewSize);
 			}
 		}
 		else
 #endif
 		{
-			
+			FPoolInfo* Pool = Private::FindPoolInfo(*this, (UPTRINT)Ptr, BasePtr);
 
 			if( Pool->TableIndex < BinnedOSTableIndex )
 			{
@@ -1560,24 +1537,20 @@ bool FMallocBinned::GetAllocationSize(void *Original, SIZE_T &SizeOut)
 		return false;
 	}
 
-	UPTRINT BasePtr;
-	FPoolInfo* Pool = Private::FindPoolInfo(*this, (UPTRINT)Original, BasePtr);
+	
 
 #if PLATFORM_IOS || PLATFORM_MAC
-	if (Pool == nullptr)
-	{
 #if USE_OS_SMALL_BLOCK_ALLOC && !USE_OS_SMALL_BLOCK_GRAB_MEMORY_FROM_OS
-		if ((uint64)Original >= Private::SmallBlockStartPtr && (uint64)Original < Private::SmallBlockEndPtr)
-		{
-			SizeOut = Private::SMALL_BLOCK_POOL_SIZE;
-			return true;
-		}
-#endif
-		UE_LOG(LogMemory, Warning, TEXT("Attempting to access memory pool info for a pointer we didn't allocate!"));
-		return false;
+	if(FPlatformMemory::PtrIsOSMalloc(Original))
+	{
+		SizeOut = Private::SMALL_BLOCK_POOL_SIZE;
+		return true;
 	}
 #endif
-
+#endif
+	
+	UPTRINT BasePtr;
+	FPoolInfo* Pool = Private::FindPoolInfo(*this, (UPTRINT)Original, BasePtr);
 	PTRINT OffsetFromBase = (PTRINT)Original - (PTRINT)BasePtr;
 	check(OffsetFromBase >= 0);
 
diff --git a/Engine/Source/Runtime/Core/Public/Apple/ApplePlatformMemory.h b/Engine/Source/Runtime/Core/Public/Apple/ApplePlatformMemory.h
--- a/Engine/Source/Runtime/Core/Public/Apple/ApplePlatformMemory.h
+++ b/Engine/Source/Runtime/Core/Public/Apple/ApplePlatformMemory.h
@@ -10,6 +10,8 @@
 #include <libkern/OSAtomic.h>
 #include <Foundation/NSObject.h>
 
+#include <mach/mach.h>
+
 NS_ASSUME_NONNULL_BEGIN
 
 /**
@@ -69,6 +71,10 @@ struct CORE_API FApplePlatformMemory : public FGenericPlatformMemory
 	static bool PageProtect(void* const Ptr, const SIZE_T Size, const bool bCanRead, const bool bCanWrite);
 	static void* BinnedAllocFromOS(SIZE_T Size);
 	static void BinnedFreeToOS(void* Ptr, SIZE_T Size);
+	static bool PtrIsOSMalloc( void* Ptr);
+	static bool PtrIsFromNanoMalloc( void* Ptr);
+	static bool IsNanoMallocAvailable();
+	static void NanoMallocInit();
 
 	class FPlatformVirtualMemoryBlock : public FBasicVirtualMemoryBlock
 	{
@@ -124,6 +130,9 @@ struct CORE_API FApplePlatformMemory : public FGenericPlatformMemory
 	
 	/** Setup the current default CFAllocator to use our malloc functions. */
 	static void ConfigureDefaultCFAllocator(void);
+	
+	static vm_address_t NanoRegionStart;
+	static vm_address_t NanoRegionEnd;
 };
 
 NS_ASSUME_NONNULL_END
diff --git a/Engine/Source/Runtime/Core/Public/GenericPlatform/GenericPlatformMemory.h b/Engine/Source/Runtime/Core/Public/GenericPlatform/GenericPlatformMemory.h
--- a/Engine/Source/Runtime/Core/Public/GenericPlatform/GenericPlatformMemory.h
+++ b/Engine/Source/Runtime/Core/Public/GenericPlatform/GenericPlatformMemory.h
@@ -342,6 +342,54 @@ struct CORE_API FGenericPlatformMemory
 	 * @param Size size of the allocation previously passed to BinnedAllocFromOS
 	 */
 	static void BinnedFreeToOS( void* Ptr, SIZE_T Size );
+	
+	/**
+	 *	Performs initial setup for Nano malloc.
+	 *	This is a noop on non-apple platforms
+	 */
+	static void NanoMallocInit()
+	{
+		return;
+	}
+	
+	/**
+	 * Was this pointer allocated by the OS malloc?
+	 * Currently only Apple platforms implement this to detect small block allocations.
+	 *
+	 * @param The pointer to query
+	 * @return True if this pointer was allocated by the OS.
+	 *
+	 */
+	static bool PtrIsOSMalloc( void* Ptr)
+	{
+		return false;
+	}
+	
+	/**
+	 *
+	 * Nano Malloc is Apple's tiny block allocator.
+	 * Does the Nano malloc zone exist?
+	 *
+	 * @return True if Nano malloc is enabled and available.
+	*
+	*/
+	static bool IsNanoMallocAvailable()
+	{
+		return false;
+	}
+	
+	/**
+	* Was this pointer allocated by in the Nano Malloc Zone?
+	* Currently only Apple platforms implement this to detect small block allocations.
+	*
+	* @param The pointer to query
+	* @return True if this pointer is in the Nano Malloc Region
+	*
+	*/
+	static bool PtrIsFromNanoMalloc( void* Ptr)
+	{
+		return false;
+	}
 
 	class FBasicVirtualMemoryBlock
 	{
