Fixed FFrameProProfiler setting explicit values for GCycleStatsShouldEmitNamedEvents instead of inc/dec.

#rnx

diff --git a/Engine/Source/Runtime/Core/Private/FramePro/FrameProProfiler.cpp b/Engine/Source/Runtime/Core/Private/FramePro/FrameProProfiler.cpp
--- a/Engine/Source/Runtime/Core/Private/FramePro/FrameProProfiler.cpp
+++ b/Engine/Source/Runtime/Core/Private/FramePro/FrameProProfiler.cpp
@@ -374,7 +374,7 @@ FString FFrameProProfiler::StartFrameProRecording(const FString& FilenameRoot, i
 	GFrameProEnabled = true;
 
 	// Enable named events as well
-	GCycleStatsShouldEmitNamedEvents = true;
+	GCycleStatsShouldEmitNamedEvents += 1;
 
 	// Set recording flag
 	GFrameProIsRecording = true;
@@ -393,7 +393,7 @@ void FFrameProProfiler::StopFrameProRecording()
 	FramePro::StopRecording();
 
 	// Disable named events
-	GCycleStatsShouldEmitNamedEvents = false;
+	GCycleStatsShouldEmitNamedEvents -= 1;
 
 	// Clear recording flag
 	GFrameProIsRecording = false;
