fix for material crash on game startup on IOS

diff --git a/Engine/Source/Runtime/Core/Private/Misc/PreLoadFile.cpp b/Engine/Source/Runtime/Core/Private/Misc/PreLoadFile.cpp
--- a/Engine/Source/Runtime/Core/Private/Misc/PreLoadFile.cpp
+++ b/Engine/Source/Runtime/Core/Private/Misc/PreLoadFile.cpp
@@ -28,7 +28,12 @@ struct FPreloadFileRegistry
 #endif
 
 FPreLoadFile::FPreLoadFile(const TCHAR* InPath)
-	: FDelayedAutoRegisterHelper(STATS ? EDelayedRegisterRunPhase::StatSystemReady : EDelayedRegisterRunPhase::FileSystemReady, [this] { bSystemNoLongerTakingRequests = false; KickOffRead(); })
+	: FDelayedAutoRegisterHelper(STATS ? EDelayedRegisterRunPhase::StatSystemReady : EDelayedRegisterRunPhase::FileSystemReady, [this] {
+		bSystemNoLongerTakingRequests = false;
+#if PLATFORM_CAN_ASYNC_PRELOAD_FILES
+		KickOffRead();
+#endif
+	})
 	, bIsComplete(false)
 	, bFailedToOpenInKickOff(false)
 	, Data(nullptr)
