Fixed framepro.stoprec underflowing GCycleStatsShouldEmitNamedEvents if called when not recording a framepro capture.

#rnx

diff --git a/Engine/Source/Runtime/Core/Private/FramePro/FrameProProfiler.cpp b/Engine/Source/Runtime/Core/Private/FramePro/FrameProProfiler.cpp
--- a/Engine/Source/Runtime/Core/Private/FramePro/FrameProProfiler.cpp
+++ b/Engine/Source/Runtime/Core/Private/FramePro/FrameProProfiler.cpp
@@ -390,6 +390,11 @@ static FAutoConsoleCommand StartFrameProRecordCommand(
 
 void FFrameProProfiler::StopFrameProRecording()
 {
+	if (!GFrameProIsRecording)
+	{
+		return;
+	}
+
 	FramePro::StopRecording();
 
 	// Disable named events
