FNEM - More cache-friendly and faster Octree implementation.
Also changing the interface for better encapsulation by cutting the origial itterator concept and focusing on small subset of configurable itteration strategies instead.

diff --git a/Engine/Plugins/MeshPainting/Source/MeshPaintingToolset/Private/BaseMeshPaintComponentAdapter.cpp b/Engine/Plugins/MeshPainting/Source/MeshPaintingToolset/Private/BaseMeshPaintComponentAdapter.cpp
--- a/Engine/Plugins/MeshPainting/Source/MeshPaintingToolset/Private/BaseMeshPaintComponentAdapter.cpp
+++ b/Engine/Plugins/MeshPainting/Source/MeshPaintingToolset/Private/BaseMeshPaintComponentAdapter.cpp
@@ -85,16 +85,14 @@ TArray<uint32> FBaseMeshPaintComponentAdapter::SphereIntersectTriangles(const fl
 	// definitely don't want our brush to be cut off by a hard triangle edge
 	const float SquaredRadiusBias = ComponentSpaceSquaredBrushRadius * 0.025f;
 
-	for (FMeshPaintTriangleOctree::TConstElementBoxIterator<> TriIt(*MeshTriOctree, FBoxCenterAndExtent(ComponentSpaceBrushPosition, FVector(FMath::Sqrt(ComponentSpaceSquaredBrushRadius + SquaredRadiusBias)))); TriIt.HasPendingElements(); TriIt.Advance())
+	MeshTriOctree->IterateElementsWithBoundsTest(FBoxCenterAndExtent(ComponentSpaceBrushPosition, FVector(FMath::Sqrt(ComponentSpaceSquaredBrushRadius + SquaredRadiusBias))), [&OutTriangles, bOnlyFrontFacing, &ComponentSpaceCameraPosition](const FMeshPaintTriangle& CurrentTri)
 	{
-		// Check to see if the triangle is front facing
-		const FMeshPaintTriangle & CurrentTri = TriIt.GetCurrentElement();
 		const float SignedPlaneDist = FVector::PointPlaneDist(ComponentSpaceCameraPosition, CurrentTri.Vertices[0], CurrentTri.Normal);
 		if (!bOnlyFrontFacing || SignedPlaneDist < 0.0f)
 		{
 			OutTriangles.Add(CurrentTri.Index);
 		}
-	}
+	});
 
 	return OutTriangles;
 }
diff --git a/Engine/Plugins/MeshPainting/Source/MeshPaintingToolset/Private/MeshPaintHelpers.cpp b/Engine/Plugins/MeshPainting/Source/MeshPaintingToolset/Private/MeshPaintHelpers.cpp
--- a/Engine/Plugins/MeshPainting/Source/MeshPaintingToolset/Private/MeshPaintHelpers.cpp
+++ b/Engine/Plugins/MeshPainting/Source/MeshPaintingToolset/Private/MeshPaintHelpers.cpp
@@ -1322,7 +1322,6 @@ void UMeshPaintingToolset::ApplyVertexColorsToAllLODs(IMeshPaintComponentAdapter
 					for (uint32 VertexIndex = 0; VertexIndex < ApplyLOD.GetNumVertices(); ++VertexIndex)
 					{
 						TArray<FPaintedMeshVertex> PointsToConsider;
-						TVertexColorPropogationPosOctree::TConstIterator<> OctreeIter(VertPosOctree);
 						const FVector CurPosition = ApplyLOD.StaticVertexBuffers.PositionVertexBuffer.VertexPosition(VertexIndex);
 
 						FPackedNormal VertexTangentZ;
@@ -1331,42 +1330,11 @@ void UMeshPaintingToolset::ApplyVertexColorsToAllLODs(IMeshPaintComponentAdapter
 						FVector CurNormal = VertexTangentZ.ToFVector();
 
 						// Iterate through the octree attempting to find the vertices closest to the current new point
-						while (OctreeIter.HasPendingNodes())
+						VertPosOctree.IterateElementsWithBoundsTest(FBoxCenterAndExtent(CurPosition, FVector::ZeroVector), [&PointsToConsider](const FPaintedMeshVertex& Vertex)
 						{
-							const TVertexColorPropogationPosOctree::FNode& CurNode = OctreeIter.GetCurrentNode();
-							const FOctreeNodeContext& CurContext = OctreeIter.GetCurrentContext();
-
-							// Find the child of the current node, if any, that contains the current new point
-							FOctreeChildNodeRef ChildRef = CurContext.GetContainingChild(FBoxCenterAndExtent(CurPosition, FVector::ZeroVector));
-
-							if (!ChildRef.IsNULL())
-							{
-								const TVertexColorPropogationPosOctree::FNode* ChildNode = CurNode.GetChild(ChildRef);
-
-								// If the specified child node exists and contains any of the old vertices, push it to the iterator for future consideration
-								if (ChildNode && ChildNode->GetInclusiveElementCount() > 0)
-								{
-									OctreeIter.PushChild(ChildRef);
-								}
-								// If the child node doesn't have any of the old vertices in it, it's not worth pursuing any further. In an attempt to find
-								// anything to match vs. the new point, add all of the children of the current octree node that have old points in them to the
-								// iterator for future consideration.
-								else
-								{
-									FOREACH_OCTREE_CHILD_NODE(OctreeChildRef)
-									{
-										if (CurNode.HasChild(OctreeChildRef))
-										{
-											OctreeIter.PushChild(OctreeChildRef);
-										}
-									}
-								}
-							}
-
 							// Add all of the elements in the current node to the list of points to consider for closest point calculations
-							PointsToConsider.Append(CurNode.GetElements());
-							OctreeIter.Advance();
-						}
+							PointsToConsider.Add(Vertex);
+						});
 
 						// If any points to consider were found, iterate over each and find which one is the closest to the new point 
 						if (PointsToConsider.Num() > 0)
diff --git a/Engine/Plugins/Runtime/EditableMesh/Source/EditableMesh/EditableMesh.cpp b/Engine/Plugins/Runtime/EditableMesh/Source/EditableMesh/EditableMesh.cpp
--- a/Engine/Plugins/Runtime/EditableMesh/Source/EditableMesh/EditableMesh.cpp
+++ b/Engine/Plugins/Runtime/EditableMesh/Source/EditableMesh/EditableMesh.cpp
@@ -8099,38 +8099,15 @@ void UEditableMesh::SearchSpatialDatabaseWithPredicate( TFunctionRef< bool( cons
 	if( IsSpatialDatabaseAllowed() && ensure( Octree.IsValid() ) )
 	{
 		// @todo mesheditor perf: Do we need to use a custom stack allocator for iterating?  The default should probably be okay.
-		for( FEditableMeshOctree::TConstIterator<> OctreeIt( *Octree );
-			 OctreeIt.HasPendingNodes();
-			 OctreeIt.Advance() )
+		Octree->IterateElementsWithPredicate([&Predicate](const FBoxCenterAndExtent& NodeBounds)
+		{
+			return Predicate(NodeBounds.GetBox());
+		},
+		[&OutPolygons](const FEditableMeshOctreePolygon& OctreePolygon)
 		{
-			const FEditableMeshOctree::FNode& OctreeNode = OctreeIt.GetCurrentNode();
-			const FOctreeNodeContext& OctreeNodeContext = OctreeIt.GetCurrentContext();
-
-			// Leaf nodes have no children, so don't bother iterating
-			if( !OctreeNode.IsLeaf() )
-			{
-				FOREACH_OCTREE_CHILD_NODE( ChildRef )
-				{
-					if( OctreeNode.HasChild( ChildRef ) )
-					{
-						const FOctreeNodeContext ChildContext = OctreeNodeContext.GetChildContext( ChildRef );
-
-						if( Predicate( ChildContext.Bounds.GetBox() ) )
-						{
-							// Push it on the iterator's pending node stack.
-							OctreeIt.PushChild( ChildRef );
-						}
-					}
-				}
-			}
-
 			// All of the elements in this octree node are candidates.  Note this node may not be a leaf node, and that's OK.
-			for( FEditableMeshOctree::ElementConstIt OctreeElementIt( OctreeNode.GetElementIt() ); OctreeElementIt; ++OctreeElementIt )
-			{
-				const FEditableMeshOctreePolygon& OctreePolygon = *OctreeElementIt;
-				OutPolygons.Add( OctreePolygon.PolygonID );
-			}
-		}
+			OutPolygons.Add(OctreePolygon.PolygonID);
+		});
 	}
 }
 
@@ -8176,50 +8153,29 @@ void UEditableMesh::SearchSpatialDatabaseForPolygonsPotentiallyIntersectingPlane
 	if( IsSpatialDatabaseAllowed() && ensure( Octree.IsValid() ) )
 	{
 		// @todo mesheditor perf: Do we need to use a custom stack allocator for iterating?  The default should probably be okay.
-		for( FEditableMeshOctree::TConstIterator<> OctreeIt( *Octree );
-			 OctreeIt.HasPendingNodes();
-			 OctreeIt.Advance() )
+		Octree->IterateElementsWithPredicate([&InPlane](const FBoxCenterAndExtent& NodeBounds)
 		{
-			const FEditableMeshOctree::FNode& OctreeNode = OctreeIt.GetCurrentNode();
-			const FOctreeNodeContext& OctreeNodeContext = OctreeIt.GetCurrentContext();
+			const bool bIsOverlappingLineSegment =
+				FMath::PlaneAABBIntersection(
+					InPlane,
+					NodeBounds.GetBox());
 
-			// Leaf nodes have no children, so don't bother iterating
-			if( !OctreeNode.IsLeaf() )
+			if (bIsOverlappingLineSegment)
 			{
-				// Find children of this octree node that overlap our line segment
-				FOREACH_OCTREE_CHILD_NODE( ChildRef )
-				{
-					if( OctreeNode.HasChild( ChildRef ) )
-					{
-						const FOctreeNodeContext ChildContext = OctreeNodeContext.GetChildContext( ChildRef );
-
-                        const bool bIsOverlappingLineSegment =
-                            FMath::PlaneAABBIntersection(
-                                InPlane,
-                                ChildContext.Bounds.GetBox());
-
-						if( bIsOverlappingLineSegment )
-						{
-							// DrawDebugBox( GWorld, ChildContext.Bounds.Center, ChildContext.Bounds.Extent * 0.8f, FQuat::Identity, FColor::Green, false, 0.0f );		// @todo mesheditor debug: (also, wrong coordinate system!)
-
-							// Push it on the iterator's pending node stack.
-							OctreeIt.PushChild( ChildRef );
-						}
-						else
-						{
-							// DrawDebugBox( GWorld, ChildContext.Bounds.Center, ChildContext.Bounds.Extent, FQuat::Identity, FColor( 128, 128, 128 ), false, 0.0f );	// @todo mesheditor debug: (also, wrong coordinate system!)
-						}
-					}
-				}
+				// DrawDebugBox( GWorld, ChildContext.Bounds.Center, ChildContext.Bounds.Extent * 0.8f, FQuat::Identity, FColor::Green, false, 0.0f );		// @todo mesheditor debug: (also, wrong coordinate system!)
+				// Push it on the iterator's pending node stack.
+				return true;
 			}
-
-			// All of the elements in this octree node are candidates.  Note this node may not be a leaf node, and that's OK.
-			for( FEditableMeshOctree::ElementConstIt OctreeElementIt( OctreeNode.GetElementIt() ); OctreeElementIt; ++OctreeElementIt )
+			else
 			{
-				const FEditableMeshOctreePolygon& OctreePolygon = *OctreeElementIt;
-				OutPolygons.Add( OctreePolygon.PolygonID );
+				// DrawDebugBox( GWorld, ChildContext.Bounds.Center, ChildContext.Bounds.Extent, FQuat::Identity, FColor( 128, 128, 128 ), false, 0.0f );	// @todo mesheditor debug: (also, wrong coordinate system!)
+				return false;
 			}
-		}
+		},
+		[&OutPolygons](const FEditableMeshOctreePolygon& OctreePolygon)
+		{
+			OutPolygons.Add(OctreePolygon.PolygonID);
+		});
 	}
 }
 
diff --git a/Engine/Source/Developer/FunctionalTesting/Private/FunctionalAITest.cpp b/Engine/Source/Developer/FunctionalTesting/Private/FunctionalAITest.cpp
--- a/Engine/Source/Developer/FunctionalTesting/Private/FunctionalAITest.cpp
+++ b/Engine/Source/Developer/FunctionalTesting/Private/FunctionalAITest.cpp
@@ -118,9 +118,8 @@ void AFunctionalAITest::OnTimeout()
 		const FVector TransformedOrigin = GetTransform().TransformPosition(NavMeshDebugOrigin);
 		const FBox DebugBounds = FBox::BuildAABB(TransformedOrigin, NavMeshDebugExtent);
 
-		for (FNavigationOctree::TConstElementBoxIterator<FNavigationOctree::DefaultStackAllocator> It(*NavigationOctree, DebugBounds); It.HasPendingElements(); It.Advance())
+		NavigationOctree->IterateElementsWithBoundsTest(DebugBounds, [&AreaFilter, &Navmesh](const FNavigationOctreeElement& Element)
 		{
-			const FNavigationOctreeElement& Element = It.GetCurrentElement();
 			if (Element.IsMatchingFilter(AreaFilter))
 			{
 				const FCompositeNavModifier NavModifier = Element.GetModifierForAgent(&Navmesh->GetConfig());
@@ -135,7 +134,7 @@ void AFunctionalAITest::OnTimeout()
 
 				UE_LOG(LogFunctionalTest, Log, TEXT("> modifier, owner:%s areas:%s"), *GetNameSafe(Element.GetOwner()), *DebugAreaNames);
 			}
-		}
+		});
 	}
 
 	Super::OnTimeout();
diff --git a/Engine/Source/Developer/SkeletalMeshUtilitiesCommon/Private/LODUtilities.cpp b/Engine/Source/Developer/SkeletalMeshUtilitiesCommon/Private/LODUtilities.cpp
--- a/Engine/Source/Developer/SkeletalMeshUtilitiesCommon/Private/LODUtilities.cpp
+++ b/Engine/Source/Developer/SkeletalMeshUtilitiesCommon/Private/LODUtilities.cpp
@@ -1443,43 +1443,12 @@ void MatchVertexIndexUsingPosition(
 		FBox CurBox(PositionSrc - Extent, PositionSrc + Extent);
 		while (OcTreeTriangleResults.Num() <= 0)
 		{
-			TTriangleElementOctree::TConstIterator<> OctreeIter(OcTree);
-			while (OctreeIter.HasPendingNodes())
+			OcTree.IterateElementsWithBoundsTest(CurBox, [&OcTreeTriangleResults](const FTriangleElement& Element)
 			{
-				const TTriangleElementOctree::FNode& CurNode = OctreeIter.GetCurrentNode();
-				const FOctreeNodeContext& CurContext = OctreeIter.GetCurrentContext();
-
-				// Find the child of the current node, if any, that contains the current new point
-				FOctreeChildNodeRef ChildRef = CurContext.GetContainingChild(CurBox);
-
-				if (!ChildRef.IsNULL())
-				{
-					const TTriangleElementOctree::FNode* ChildNode = CurNode.GetChild(ChildRef);
-
-					// If the specified child node exists and contains any of the old vertices, push it to the iterator for future consideration
-					if (ChildNode && ChildNode->GetInclusiveElementCount() > 0)
-					{
-						OctreeIter.PushChild(ChildRef);
-					}
-					// If the child node doesn't have any of the old vertices in it, it's not worth pursuing any further. In an attempt to find
-					// anything to match vs. the new point, add all of the children of the current octree node that have old points in them to the
-					// iterator for future consideration.
-					else
-					{
-						FOREACH_OCTREE_CHILD_NODE(OctreeChildRef)
-						{
-							if (CurNode.HasChild(OctreeChildRef))
-							{
-								OctreeIter.PushChild(OctreeChildRef);
-							}
-						}
-					}
-				}
-
 				// Add all of the elements in the current node to the list of points to consider for closest point calculations
-				OcTreeTriangleResults.Append(CurNode.GetElements());
-				OctreeIter.Advance();
-			}
+				OcTreeTriangleResults.Add(Element);
+			});
+
 			//Increase the extend so we try to found in a larger area
 			Extent *= 2;
 			CurBox = FBox(PositionSrc - Extent, PositionSrc + Extent);
diff --git a/Engine/Source/Editor/MeshPaint/Private/BaseMeshPaintGeometryAdapter.cpp b/Engine/Source/Editor/MeshPaint/Private/BaseMeshPaintGeometryAdapter.cpp
--- a/Engine/Source/Editor/MeshPaint/Private/BaseMeshPaintGeometryAdapter.cpp
+++ b/Engine/Source/Editor/MeshPaint/Private/BaseMeshPaintGeometryAdapter.cpp
@@ -78,16 +78,14 @@ TArray<uint32> FBaseMeshPaintGeometryAdapter::SphereIntersectTriangles(const flo
 	// definitely don't want our brush to be cut off by a hard triangle edge
 	const float SquaredRadiusBias = ComponentSpaceSquaredBrushRadius * 0.025f;
 
-	for (FMeshPaintTriangleOctree::TConstElementBoxIterator<> TriIt(*MeshTriOctree, FBoxCenterAndExtent(ComponentSpaceBrushPosition, FVector(FMath::Sqrt(ComponentSpaceSquaredBrushRadius + SquaredRadiusBias)))); TriIt.HasPendingElements(); TriIt.Advance())
+	MeshTriOctree->IterateElementsWithBoundsTest(FBoxCenterAndExtent(ComponentSpaceBrushPosition, FVector(FMath::Sqrt(ComponentSpaceSquaredBrushRadius + SquaredRadiusBias))), [bOnlyFrontFacing , &OutTriangles, &ComponentSpaceCameraPosition](const FMeshPaintTriangle& CurrentTri)
 	{
-		// Check to see if the triangle is front facing
-		const FMeshPaintTriangle & CurrentTri = TriIt.GetCurrentElement();
 		const float SignedPlaneDist = FVector::PointPlaneDist(ComponentSpaceCameraPosition, CurrentTri.Vertices[0], CurrentTri.Normal);
 		if (!bOnlyFrontFacing || SignedPlaneDist < 0.0f)
 		{
 			OutTriangles.Add(CurrentTri.Index);
 		}
-	}
+	});
 
 	return OutTriangles;
 }
diff --git a/Engine/Source/Editor/MeshPaint/Private/MeshPaintHelpers.cpp b/Engine/Source/Editor/MeshPaint/Private/MeshPaintHelpers.cpp
--- a/Engine/Source/Editor/MeshPaint/Private/MeshPaintHelpers.cpp
+++ b/Engine/Source/Editor/MeshPaint/Private/MeshPaintHelpers.cpp
@@ -1835,7 +1835,6 @@ void MeshPaintHelpers::ApplyVertexColorsToAllLODs(IMeshPaintGeometryAdapter& Geo
 					for (uint32 VertexIndex = 0; VertexIndex < ApplyLOD.GetNumVertices(); ++VertexIndex)
 					{
 						TArray<FPaintedMeshVertex> PointsToConsider;
-						TVertexColorPropogationPosOctree::TConstIterator<> OctreeIter(VertPosOctree);
 						const FVector CurPosition = ApplyLOD.StaticVertexBuffers.PositionVertexBuffer.VertexPosition(VertexIndex);
 
 						FPackedNormal VertexTangentZ;
@@ -1844,42 +1843,10 @@ void MeshPaintHelpers::ApplyVertexColorsToAllLODs(IMeshPaintGeometryAdapter& Geo
 						FVector CurNormal = VertexTangentZ.ToFVector();
 
 						// Iterate through the octree attempting to find the vertices closest to the current new point
-						while (OctreeIter.HasPendingNodes())
+						VertPosOctree.IterateElementsWithBoundsTest(FBoxCenterAndExtent(CurPosition, FVector::ZeroVector), [&PointsToConsider](const FPaintedMeshVertex& Vertex)
 						{
-							const TVertexColorPropogationPosOctree::FNode& CurNode = OctreeIter.GetCurrentNode();
-							const FOctreeNodeContext& CurContext = OctreeIter.GetCurrentContext();
-
-							// Find the child of the current node, if any, that contains the current new point
-							FOctreeChildNodeRef ChildRef = CurContext.GetContainingChild(FBoxCenterAndExtent(CurPosition, FVector::ZeroVector));
-
-							if (!ChildRef.IsNULL())
-							{
-								const TVertexColorPropogationPosOctree::FNode* ChildNode = CurNode.GetChild(ChildRef);
-
-								// If the specified child node exists and contains any of the old vertices, push it to the iterator for future consideration
-								if (ChildNode && ChildNode->GetInclusiveElementCount() > 0)
-								{
-									OctreeIter.PushChild(ChildRef);
-								}
-								// If the child node doesn't have any of the old vertices in it, it's not worth pursuing any further. In an attempt to find
-								// anything to match vs. the new point, add all of the children of the current octree node that have old points in them to the
-								// iterator for future consideration.
-								else
-								{
-									FOREACH_OCTREE_CHILD_NODE(OctreeChildRef)
-									{
-										if (CurNode.HasChild(OctreeChildRef))
-										{
-											OctreeIter.PushChild(OctreeChildRef);
-										}
-									}
-								}
-							}
-
-							// Add all of the elements in the current node to the list of points to consider for closest point calculations
-							PointsToConsider.Append(CurNode.GetElements());
-							OctreeIter.Advance();
-						}
+							PointsToConsider.Add(Vertex);
+						});
 
 						// If any points to consider were found, iterate over each and find which one is the closest to the new point 
 						if (PointsToConsider.Num() > 0)
diff --git a/Engine/Source/Editor/UnrealEd/Private/Fbx/FbxSkeletalMeshImport.cpp b/Engine/Source/Editor/UnrealEd/Private/Fbx/FbxSkeletalMeshImport.cpp
--- a/Engine/Source/Editor/UnrealEd/Private/Fbx/FbxSkeletalMeshImport.cpp
+++ b/Engine/Source/Editor/UnrealEd/Private/Fbx/FbxSkeletalMeshImport.cpp
@@ -220,44 +220,10 @@ void RemapSkeletalMeshVertexColorToImportData(const USkeletalMesh* SkeletalMesh,
 		const FVector& Normal = WedgeIndexToNormal.FindChecked(WedgeIndex);
 
 		TArray<FSoftSkinVertex> PointsToConsider;
-		TSKCVertPosOctree::TConstIterator<> OctreeIter(VertPosOctree);
-		// Iterate through the octree attempting to find the vertices closest to the current new point
-		while (OctreeIter.HasPendingNodes())
+		VertPosOctree.IterateElementsWithBoundsTest(FBoxCenterAndExtent(Position, FVector::ZeroVector), [&PointsToConsider](const FSoftSkinVertex& Vertex)
 		{
-			const TSKCVertPosOctree::FNode& CurNode = OctreeIter.GetCurrentNode();
-			const FOctreeNodeContext& CurContext = OctreeIter.GetCurrentContext();
-
-			// Find the child of the current node, if any, that contains the current new point
-			FOctreeChildNodeRef ChildRef = CurContext.GetContainingChild(FBoxCenterAndExtent(Position, FVector::ZeroVector));
-
-			if (!ChildRef.IsNULL())
-			{
-				const TSKCVertPosOctree::FNode* ChildNode = CurNode.GetChild(ChildRef);
-
-				// If the specified child node exists and contains any of the old vertices, push it to the iterator for future consideration
-				if (ChildNode && ChildNode->GetInclusiveElementCount() > 0)
-				{
-					OctreeIter.PushChild(ChildRef);
-				}
-				// If the child node doesn't have any of the old vertices in it, it's not worth pursuing any further. In an attempt to find
-				// anything to match vs. the new point, add all of the children of the current octree node that have old points in them to the
-				// iterator for future consideration.
-				else
-				{
-					FOREACH_OCTREE_CHILD_NODE(OctreeChildRef)
-					{
-						if (CurNode.HasChild(OctreeChildRef))
-						{
-							OctreeIter.PushChild(OctreeChildRef);
-						}
-					}
-				}
-			}
-
-			// Add all of the elements in the current node to the list of points to consider for closest point calculations
-			PointsToConsider.Append(CurNode.GetElements());
-			OctreeIter.Advance();
-		}
+			PointsToConsider.Add(Vertex);
+		});
 
 		if (PointsToConsider.Num() > 0)
 		{
diff --git a/Engine/Source/Runtime/Core/Private/Math/GenericOctree.cpp b/Engine/Source/Runtime/Core/Private/Math/GenericOctree.cpp
--- a/Engine/Source/Runtime/Core/Private/Math/GenericOctree.cpp
+++ b/Engine/Source/Runtime/Core/Private/Math/GenericOctree.cpp
@@ -3,5 +3,3 @@
 #include "Math/GenericOctree.h"
 
 DEFINE_LOG_CATEGORY(LogGenericOctree);
-
-const float FOctreeNodeContext::NegativeOneOneTable[2] = {-1.0f,1.0f};
diff --git a/Engine/Source/Runtime/Core/Public/Math/GenericOctree.h b/Engine/Source/Runtime/Core/Public/Math/GenericOctree.h
--- a/Engine/Source/Runtime/Core/Public/Math/GenericOctree.h
+++ b/Engine/Source/Runtime/Core/Public/Math/GenericOctree.h
@@ -144,6 +144,7 @@ public:
 	FOctreeChildNodeRef(int32 InIndex = 0)
 	:	Index(InIndex)
 	{
+		checkSlow(Index < 8);
 		// some compilers do not allow multiple members of a union to be specified in the constructor init list
 		bNULL = false;
 	}
@@ -279,60 +280,47 @@ public:
 		ChildCenterOffset = Bounds.Extent.X - LooseChildExtent;
 	}
 
-	/** the float table {-1.0f,1.0f} **/
-	static CORE_API const float NegativeOneOneTable[2];
+	inline VectorRegister GetChildOffsetVec(int i) const
+	{
+		union { VectorRegister v; VectorRegisterInt i; } Mask;
+		Mask.v = MakeVectorRegister(1u, 2u, 4u, 8u);
+		VectorRegisterInt X = VectorIntLoad1(&i);
+		VectorRegisterInt A = VectorIntAnd(X, Mask.i);
+		Mask.i = VectorIntCompareEQ(Mask.i, A);
+		return VectorSelect(Mask.v, VectorSetFloat1(ChildCenterOffset), VectorSetFloat1(-ChildCenterOffset));
+	}
 
 	/** Child node initialization constructor. */
-	FORCEINLINE FOctreeNodeContext GetChildContext(FOctreeChildNodeRef ChildRef) const
-	{
-		return FOctreeNodeContext(FBoxCenterAndExtent(
-			FVector(
-			Bounds.Center.X + ChildCenterOffset * NegativeOneOneTable[ChildRef.X],
-			Bounds.Center.Y + ChildCenterOffset * NegativeOneOneTable[ChildRef.Y],
-			Bounds.Center.Z + ChildCenterOffset * NegativeOneOneTable[ChildRef.Z]
-			),
-			FVector(
-			ChildExtent,
-			ChildExtent,
-			ChildExtent
-			)
-			));
+	inline FOctreeNodeContext GetChildContext(FOctreeChildNodeRef ChildRef) const
+	{
+		FBoxCenterAndExtent LocalBounds;
+		VectorRegister ZeroW = MakeVectorRegister(1.0f, 1.0f, 1.0f, 0.0f);
+		VectorStoreAligned(VectorMultiply(ZeroW, VectorAdd(VectorLoadAligned(&Bounds.Center), GetChildOffsetVec(ChildRef.Index))), &LocalBounds.Center);
+		VectorStoreAligned(VectorMultiply(ZeroW, VectorSetFloat1(ChildExtent)), &LocalBounds.Extent);
+		return FOctreeNodeContext(LocalBounds);
 	}
 
 	/** Construct a child context given the child ref. Optimized to remove all LHS. */
-	FORCEINLINE void GetChildContext(FOctreeChildNodeRef ChildRef, FOctreeNodeContext * RESTRICT ChildContext) const
-	{
-		const FOctreeNodeContext* RESTRICT ParentContext = this;
-		ChildContext->Bounds.Center.X = ParentContext->Bounds.Center.X + ParentContext->ChildCenterOffset * NegativeOneOneTable[ChildRef.X];
-		ChildContext->Bounds.Center.Y = ParentContext->Bounds.Center.Y + ParentContext->ChildCenterOffset * NegativeOneOneTable[ChildRef.Y];
-		ChildContext->Bounds.Center.Z = ParentContext->Bounds.Center.Z + ParentContext->ChildCenterOffset * NegativeOneOneTable[ChildRef.Z];
-		ChildContext->Bounds.Center.W = 0.0f;
-		ChildContext->Bounds.Extent.X = ParentContext->ChildExtent;
-		ChildContext->Bounds.Extent.Y = ParentContext->ChildExtent;
-		ChildContext->Bounds.Extent.Z = ParentContext->ChildExtent;
-		ChildContext->Bounds.Extent.W = 0.0f;
-
-		const float TightChildExtent = ParentContext->ChildExtent * 0.5f;
+	inline void GetChildContext(FOctreeChildNodeRef ChildRef, FOctreeNodeContext* ChildContext) const
+	{
+		VectorRegister ZeroW = MakeVectorRegister(1.0f, 1.0f, 1.0f, 0.0f);
+		VectorStoreAligned(VectorMultiply(ZeroW, VectorAdd(VectorLoadAligned(&Bounds.Center), GetChildOffsetVec(ChildRef.Index))), &ChildContext->Bounds.Center);
+		VectorStoreAligned(VectorMultiply(ZeroW, VectorSetFloat1(ChildExtent)), &ChildContext->Bounds.Extent);
+
+		const float TightChildExtent = ChildExtent * 0.5f;
 		const float LooseChildExtent = TightChildExtent * (1.0f + 1.0f / (float)LoosenessDenominator);
 		ChildContext->ChildExtent = LooseChildExtent;
-		ChildContext->ChildCenterOffset = ParentContext->ChildExtent - LooseChildExtent;
+		ChildContext->ChildCenterOffset = ChildExtent - LooseChildExtent;
 	}
 	
 	/** Child node initialization constructor. */
-	FORCEINLINE FOctreeNodeContext GetChildContext(FOctreeChildNodeRef ChildRef, uint32 InInCullBits, uint32 InOutCullBits) const
-	{
-		return FOctreeNodeContext(FBoxCenterAndExtent(
-			FVector(
-			Bounds.Center.X + ChildCenterOffset * NegativeOneOneTable[ChildRef.X],
-			Bounds.Center.Y + ChildCenterOffset * NegativeOneOneTable[ChildRef.Y],
-			Bounds.Center.Z + ChildCenterOffset * NegativeOneOneTable[ChildRef.Z]
-			),
-			FVector(
-			ChildExtent,
-			ChildExtent,
-			ChildExtent
-			)
-			),InInCullBits,InOutCullBits);
+	inline FOctreeNodeContext GetChildContext(FOctreeChildNodeRef ChildRef, uint32 InInCullBits, uint32 InOutCullBits) const
+	{
+		FBoxCenterAndExtent LocalBounds;
+		VectorRegister ZeroW = MakeVectorRegister(1.0f, 1.0f, 1.0f, 0.0f);
+		VectorStoreAligned(VectorMultiply(ZeroW, VectorAdd(VectorLoadAligned(&Bounds.Center), GetChildOffsetVec(ChildRef.Index))), &LocalBounds.Center);
+		VectorStoreAligned(VectorMultiply(ZeroW, VectorSetFloat1(ChildExtent)), &LocalBounds.Extent);
+		return FOctreeNodeContext(LocalBounds, InInCullBits, InOutCullBits);
 	}
 	/**
 	 * Determines which of the octree node's children intersect with a bounding box.
@@ -350,367 +338,557 @@ public:
 	FOctreeChildNodeRef GetContainingChild(const FBoxCenterAndExtent& BoundingBox) const;
 };
 
+CORE_API DECLARE_LOG_CATEGORY_EXTERN(LogGenericOctree, Log, All);
+
 /** An octree. */
 template<typename ElementType,typename OctreeSemantics>
 class TOctree
 {
+	using ElementArrayType = TArray<ElementType, typename OctreeSemantics::ElementAllocator>;
 public:
+	using FNodeIndex = uint32;
 
-	typedef TArray<ElementType, typename OctreeSemantics::ElementAllocator> ElementArrayType;
-	typedef typename ElementArrayType::TConstIterator ElementConstIt;
+private:
+	struct FNode
+	{
+		FNodeIndex ChildNodes = INDEX_NONE;
+		uint32 InclusiveNumElements = 0;
 
-	/** A node in the octree. */
-	class FNode
+		bool IsLeaf() const
+		{
+			return ChildNodes == INDEX_NONE;
+		}
+	};
+
+	FOctreeNodeContext RootNodeContext;
+	TArray<FNode> TreeNodes;
+	TArray<FNodeIndex> ParentLinks;
+	TArray<ElementArrayType> TreeElements;
+
+	class FFreeList
 	{
-	public:
+		struct FSpan
+		{
+			FNodeIndex Start;
+			FNodeIndex End;
+		};
 
-		friend class TOctree;
+		TArray<FSpan> FreeList;
 
-		/** Initialization constructor. */
-		explicit FNode(const FNode* InParent)
-		:	Parent(InParent)
-		,	InclusiveNumElements(0)
-		,	bIsLeaf(true)
+	public:
+		FFreeList()
 		{
-			FOREACH_OCTREE_CHILD_NODE(ChildRef)
-			{
-				Children[ChildRef.Index] = NULL;
-			}
+			Reset();
 		}
 
-		/** Destructor. */
-		~FNode()
+		void Push(FNodeIndex NodeIndex)
 		{
-			FOREACH_OCTREE_CHILD_NODE(ChildRef)
+			//find the index that points to our right side node
+			int Index = 1; //exclude the dummy
+			int Size = FreeList.Num() - 1;
+
+			//start with binary search for larger lists
+			while (Size > 32)
 			{
-				delete Children[ChildRef.Index];
+				const int LeftoverSize = Size % 2;
+				Size = Size / 2;
+
+				const int CheckIndex = Index + Size;
+				const int IndexIfLess = CheckIndex + LeftoverSize;
+
+				Index = FreeList[CheckIndex].Start > NodeIndex ? IndexIfLess : Index;
 			}
-		}
 
-		// Accessors.
-		FORCEINLINE ElementConstIt GetElementIt() const { return ElementConstIt(Elements); }
-		FORCEINLINE bool IsLeaf() const { return bIsLeaf; }
-		FORCEINLINE bool HasChild(FOctreeChildNodeRef ChildRef) const
-		{
-			return Children[ChildRef.Index] != NULL && Children[ChildRef.Index]->InclusiveNumElements > 0;
-		}
-		FORCEINLINE FNode* GetChild(FOctreeChildNodeRef ChildRef) const
-		{
-			return Children[ChildRef.Index];
-		}
-		FORCEINLINE int32 GetElementCount() const
-		{
-			return Elements.Num();
-		}
-		FORCEINLINE int32 GetInclusiveElementCount() const
-		{
-			return InclusiveNumElements;
-		}
-		FORCEINLINE const ElementArrayType& GetElements() const
-		{
-			return Elements;
-		}
-		void ShrinkElements()
-		{
-			Elements.Shrink();
-			FOREACH_OCTREE_CHILD_NODE(ChildRef)
+			//small size array optimization
+			int ArrayEnd = Index + Size;
+			while (Index < ArrayEnd)
 			{
-				if (Children[ChildRef.Index])
+				if (FreeList[Index].Start < NodeIndex)
 				{
-					Children[ChildRef.Index]->ShrinkElements();
+					break;
 				}
+				Index++;
 			}
-		}
 
-		void ApplyOffset(const FVector& InOffset)
-		{
-			for (auto& Element : Elements)
-			{
-				OctreeSemantics::ApplyOffset(Element, InOffset);
-			}
-			
-			FOREACH_OCTREE_CHILD_NODE(ChildRef)
+			//can we merge with the right node
+			if (Index < FreeList.Num() && FreeList[Index].End + 1 == NodeIndex)
 			{
-				if (Children[ChildRef.Index])
+				FreeList[Index].End = NodeIndex;
+
+				//are we filling the gap between the left and right node
+				if (FreeList[Index - 1].Start - 1 == NodeIndex)
 				{
-					Children[ChildRef.Index]->ApplyOffset(InOffset);
+					FreeList[Index - 1].Start = FreeList[Index].Start;
+					FreeList.RemoveAt(Index);
 				}
+				return;
 			}
-		}
-		
-	private:
 
-		/** The elements in this node. */
-		mutable ElementArrayType Elements;
+			//can we merge with the left node
+			if (FreeList[Index - 1].Start - 1 == NodeIndex)
+			{
+				FreeList[Index - 1].Start = NodeIndex;
+				return;
+			}
 
-		/** The parent of this node. */
-		const FNode* Parent;
+			//we are node that could not be merged
+			FreeList.Insert(FSpan{ NodeIndex , NodeIndex }, Index);
+		}
 
-		/** The children of the node. */
-		mutable FNode* Children[8];
+		FNodeIndex Pop()
+		{
+			FSpan& Span = FreeList.Last();
+			FNodeIndex Index = Span.Start;
+			checkSlow(Index != INDEX_NONE);
+			if (Span.Start == Span.End)
+			{
+				FreeList.Pop();
+				return Index;
+			}
+			else
+			{
+				Span.Start++;
+				return Index;
+			}
+		}
 
-		/** The number of elements contained by the node and its child nodes. */
-		mutable uint32 InclusiveNumElements : 31;
+		void Reset()
+		{
+			FreeList.Reset(1);
+			//push a dummy
+			FreeList.Push(FSpan{ INDEX_NONE, INDEX_NONE });
+		}
 
-		/** true if the meshes should be added directly to the node, rather than subdividing when possible. */
-		mutable uint32 bIsLeaf : 1;
+		int Num() const
+		{
+			//includes one dummy
+			return FreeList.Num() - 1;
+		}
 	};
 
+	TArray<FNodeIndex> FreeList;
+	/** The extent of a leaf at the maximum allowed depth of the tree. */
+	float MinLeafExtent;
 
-
-	/** A reference to an octree node, its context, and a read lock. */
-	class FNodeReference
-	{
-	public:
-
-		const FNode* Node;
-		FOctreeNodeContext Context;
-
-		/** Default constructor. */
-		FNodeReference():
-			Node(NULL),
-			Context()
-		{}
-
-		/** Initialization constructor. */
-		FNodeReference(const FNode* InNode,const FOctreeNodeContext& InContext):
-			Node(InNode),
-			Context(InContext)
-		{}
-	};	
-
-	/** The default iterator allocator gives the stack enough inline space to contain a path and its siblings from root to leaf. */
-	typedef TInlineAllocator<7 * (14 - 1) + 8> DefaultStackAllocator;
-
-	/** An octree node iterator. */
-	template<typename StackAllocator = DefaultStackAllocator>
-	class TConstIterator
+	FNodeIndex AllocateEightNodes()
 	{
-	public:
-
-		/** Pushes a child of the current node onto the stack of nodes to visit. */
-		void PushChild(FOctreeChildNodeRef ChildRef)
+		FNodeIndex Index = INDEX_NONE;
+		if (FreeList.Num())
 		{
-			FNodeReference* NewNode = new (NodeStack) FNodeReference;
-			NewNode->Node = CurrentNode.Node->GetChild(ChildRef);
-			CurrentNode.Context.GetChildContext(ChildRef, &NewNode->Context);
+			Index = (FreeList.Pop() * 8) + 1;
 		}
-		/** Pushes a child of the current node onto the stack of nodes to visit. */
-		void PushChild(FOctreeChildNodeRef ChildRef,uint32 FullyInsideView,uint32 FullyOutsideView )
+		else
 		{
-			FNodeReference* NewNode = new (NodeStack) FNodeReference;
-			NewNode->Node = CurrentNode.Node->GetChild(ChildRef);
-			CurrentNode.Context.GetChildContext(ChildRef, &NewNode->Context);
-			NewNode->Context.InCullBits = FullyInsideView;
-			NewNode->Context.OutCullBits = FullyOutsideView;
+			Index = TreeNodes.AddDefaulted(8);
+			ParentLinks.AddDefaulted();
+			FNodeIndex ElementIndex = TreeElements.AddDefaulted(8);
+			checkSlow(Index == ElementIndex);
 		}
-		/** Pushes a child of the current node onto the stack of nodes to visit. */
-		void PushChild(FOctreeChildNodeRef ChildRef,const FOctreeNodeContext& Context )
+		return Index;
+	}
+
+	void FreeEightNodes(FNodeIndex Index)
+	{
+		checkSlow(Index != INDEX_NONE && Index != 0);
+		for (int i = 0; i < 8; i++)
 		{
-			new (NodeStack) FNodeReference(CurrentNode.Node->GetChild(ChildRef),Context);
+			TreeNodes[Index + i] = FNode();
+			checkSlow(TreeElements[Index + i].Num() == 0);
 		}
+		ParentLinks[(Index - 1) / 8] = INDEX_NONE;
+		//TODO shrink the TreeNodes as well as the TreeElements and ParentLinks to reduce high watermark memory footprint.
+		FreeList.Push((Index - 1) / 8);
+	}
 
-
-		/** Iterates to the next node. */
-		void Advance()
+	void AddElementInternal(FNodeIndex CurrentNodeIndex, const FOctreeNodeContext& NodeContext, const FBoxCenterAndExtent& ElementBounds, typename TCallTraits<ElementType>::ConstReference Element, ElementArrayType& TempElementStorage)
+	{
+		checkSlow(CurrentNodeIndex != INDEX_NONE);
+		TreeNodes[CurrentNodeIndex].InclusiveNumElements++;
+		if (TreeNodes[CurrentNodeIndex].IsLeaf())
 		{
-			if(NodeStack.Num())
+			if (TreeElements[CurrentNodeIndex].Num() + 1 > OctreeSemantics::MaxElementsPerLeaf && NodeContext.Bounds.Extent.X > MinLeafExtent)
 			{
-				CurrentNode = NodeStack[NodeStack.Num() - 1];
-				NodeStack.RemoveAt(NodeStack.Num() - 1);
+				TempElementStorage = MoveTemp(TreeElements[CurrentNodeIndex]);
+
+				FNodeIndex ChildStartIndex = AllocateEightNodes();
+				ParentLinks[(ChildStartIndex - 1) / 8] = CurrentNodeIndex;
+				TreeNodes[CurrentNodeIndex].ChildNodes = ChildStartIndex;
+				TreeNodes[CurrentNodeIndex].InclusiveNumElements = 0;
+
+				for (typename TCallTraits<ElementType>::ConstReference ChildElement : TempElementStorage)
+				{
+					const FBoxCenterAndExtent ChildElementBounds(OctreeSemantics::GetBoundingBox(ChildElement));
+					AddElementInternal(CurrentNodeIndex, NodeContext, ChildElementBounds, ChildElement, TempElementStorage);
+				}
+
+				TempElementStorage.Reset();
+				AddElementInternal(CurrentNodeIndex, NodeContext, ElementBounds, Element, TempElementStorage);
+				return;
 			}
 			else
 			{
-				CurrentNode = FNodeReference();
+				int ElementIndex = TreeElements[CurrentNodeIndex].Emplace(Element);
+
+				SetElementId(Element, FOctreeElementId(CurrentNodeIndex, ElementIndex));	
+				return;
 			}
 		}
-
-		/** Checks if there are any nodes left to iterate over. */
-		bool HasPendingNodes() const
+		else
 		{
-			return CurrentNode.Node != NULL;
+			const FOctreeChildNodeRef ChildRef = NodeContext.GetContainingChild(ElementBounds);
+			if (ChildRef.IsNULL())
+			{
+				int ElementIndex = TreeElements[CurrentNodeIndex].Emplace(Element);
+				SetElementId(Element, FOctreeElementId(CurrentNodeIndex, ElementIndex));
+				return;
+			}
+			else
+			{
+				FNodeIndex ChildNodeIndex = TreeNodes[CurrentNodeIndex].ChildNodes + ChildRef.Index;
+				FOctreeNodeContext ChildNodeContext = NodeContext.GetChildContext(ChildRef);
+				AddElementInternal(ChildNodeIndex, ChildNodeContext, ElementBounds, Element, TempElementStorage);
+				return;
+			}
 		}
 
-		/** Starts iterating at the root of an octree. */
-		TConstIterator(const TOctree& Tree)
-		:	CurrentNode(FNodeReference(&Tree.RootNode,Tree.RootNodeContext))
-		{}
+		UE_LOG(LogGenericOctree, Fatal,
+			TEXT("Failed to find an octree node for an element with bounds (%f,%f,%f) +/- (%f,%f,%f)!"),
+			ElementBounds.Center.X,
+			ElementBounds.Center.Y,
+			ElementBounds.Center.Z,
+			ElementBounds.Extent.X,
+			ElementBounds.Extent.Y,
+			ElementBounds.Extent.Z
+		);
+	}
 
-		/** Starts iterating at a particular node of an octree. */
-		TConstIterator(const FNode& Node,const FOctreeNodeContext& Context):
-			CurrentNode(FNodeReference(&Node,Context))
-		{}
+	void CollapseNodesInternal(FNodeIndex CurrentNodeIndex, ElementArrayType& CollapsedNodeElements)
+	{
+		CollapsedNodeElements.Append(MoveTemp(TreeElements[CurrentNodeIndex]));	
+		TreeElements[CurrentNodeIndex].Reset();
 
-		// Accessors.
-		const FNode& GetCurrentNode() const
-		{
-			return *CurrentNode.Node;
-		}
-		const FOctreeNodeContext& GetCurrentContext() const
+		if (!TreeNodes[CurrentNodeIndex].IsLeaf())
 		{
-			return CurrentNode.Context;
+			FNodeIndex ChildStartIndex = TreeNodes[CurrentNodeIndex].ChildNodes;
+			for (int i = 0; i < 8; i++)
+			{
+				CollapseNodesInternal(ChildStartIndex + i, CollapsedNodeElements);
+			}
+			FreeEightNodes(ChildStartIndex);
 		}
+	}
 
-	private:
+	template<typename PredicateFunc, typename IterateFunc>
+	void IterateNodesWithPredicateInternal(FNodeIndex CurrentNodeIndex, const FOctreeNodeContext& NodeContext, const PredicateFunc& Predicate, const IterateFunc& Func) const
+	{
+		if (TreeNodes[CurrentNodeIndex].InclusiveNumElements > 0)
+		{
+			if (Predicate(NodeContext.Bounds))
+			{
+				Func(CurrentNodeIndex);
 
-		/** The node that is currently being visited. */
-		FNodeReference CurrentNode;
-	
-		/** The nodes which are pending iteration. */
-		TArray<FNodeReference,StackAllocator> NodeStack;
-	};
+				if (!TreeNodes[CurrentNodeIndex].IsLeaf())
+				{
+					FNodeIndex ChildStartIndex = TreeNodes[CurrentNodeIndex].ChildNodes;
+					for (int i = 0; i < 8; i++)
+					{
+						IterateNodesWithPredicateInternal(ChildStartIndex + i, NodeContext.GetChildContext(FOctreeChildNodeRef(i)), Predicate, Func);
+					}
+				}
+			}
+		}
+	}
 
-	/** Iterates over the elements in the octree that intersect a bounding box. */
-	template<typename StackAllocator = DefaultStackAllocator>
-	class TConstElementBoxIterator
+	template<typename IterateFunc>
+	void IterateElementsWithBoundsTestInternal(FNodeIndex CurrentNodeIndex, const FOctreeNodeContext& NodeContext, const FBoxCenterAndExtent& BoxBounds, const IterateFunc& Func) const
 	{
-	public:
-
-		/** Iterates to the next element. */
-		void Advance()
+		if (TreeNodes[CurrentNodeIndex].InclusiveNumElements > 0)
 		{
-			++ElementIt;
-			AdvanceToNextIntersectingElement();
+			for (typename TCallTraits<ElementType>::ConstReference Element : TreeElements[CurrentNodeIndex])
+			{
+				if (Intersect(OctreeSemantics::GetBoundingBox(Element), BoxBounds))
+				{
+					Func(Element);
+				}
+			}
+
+			if (!TreeNodes[CurrentNodeIndex].IsLeaf())
+			{
+				const FOctreeChildNodeSubset IntersectingChildSubset = NodeContext.GetIntersectingChildren(BoxBounds);
+				FNodeIndex ChildStartIndex = TreeNodes[CurrentNodeIndex].ChildNodes;
+				for (int i = 0; i < 8; i++)
+				{
+					if(IntersectingChildSubset.Contains(FOctreeChildNodeRef(i)))
+					{
+						IterateElementsWithBoundsTestInternal(ChildStartIndex + i, NodeContext.GetChildContext(FOctreeChildNodeRef(i)), BoxBounds, Func);
+					}
+				}
+			}
 		}
+	}
 
-		/** Checks if there are any elements left to iterate over. */
-		bool HasPendingElements() const
+	template<typename IterateFunc>
+	void FindFirstElementWithBoundsTestInternal(FNodeIndex CurrentNodeIndex, const FOctreeNodeContext& NodeContext, const FBoxCenterAndExtent& BoxBounds, const IterateFunc& Func, bool& ContinueTraversal) const
+	{
+		if (ContinueTraversal && TreeNodes[CurrentNodeIndex].InclusiveNumElements > 0)
 		{
-			return NodeIt.HasPendingNodes();
+			for (int Index = 0; ContinueTraversal && Index < TreeElements[CurrentNodeIndex].Num(); Index++)
+			{
+				typename TCallTraits<ElementType>::ConstReference Element = TreeElements[CurrentNodeIndex][Index];
+				if (Intersect(OctreeSemantics::GetBoundingBox(Element), BoxBounds))
+				{
+					ContinueTraversal &= Func(Element);
+				}
+			}
+
+			if (ContinueTraversal && !TreeNodes[CurrentNodeIndex].IsLeaf())
+			{
+				const FOctreeChildNodeSubset IntersectingChildSubset = NodeContext.GetIntersectingChildren(BoxBounds);
+				FNodeIndex ChildStartIndex = TreeNodes[CurrentNodeIndex].ChildNodes;
+				for (int i = 0; i < 8; i++)
+				{
+					if (IntersectingChildSubset.Contains(FOctreeChildNodeRef(i)))
+					{
+						FindFirstElementWithBoundsTestInternal(ChildStartIndex + i, NodeContext.GetChildContext(FOctreeChildNodeRef(i)), BoxBounds, Func, ContinueTraversal);
+					}
+				}
+			}
 		}
+	}
+public:
 
-		/** Initialization constructor. */
-		TConstElementBoxIterator(const TOctree& Tree,const FBoxCenterAndExtent& InBoundingBox)
-		:	IteratorBounds(InBoundingBox)
-		,	NodeIt(Tree)
-		,	ElementIt(Tree.RootNode.GetElementIt())
+	/**
+	 * this function will call the passed in function for all elements in the Octree in node by node in no specified order.
+	 * @param Func - Function to call with each Element.
+	 */
+	template<typename IterateAllElementsFunc>
+	inline void IterateAllElements(const IterateAllElementsFunc& Func) const
+	{
+		for (const ElementArrayType& Elements : TreeElements)
 		{
-			ProcessChildren();
-			AdvanceToNextIntersectingElement();
+			for (typename TCallTraits<ElementType>::ConstReference Element : Elements)
+			{
+				Func(Element);
+			}
 		}
+	}
 
-		// Accessors.
-		const ElementType& GetCurrentElement() const
+	/**
+	 * this function will traverse the Octree starting from the root in depth first order and the predicate can be used to implement custom culling for each node.
+	 * @param Predicate - a Function when given the bounds of the currently traversed node that returns true if traversal should continue or false to skip that branch.
+	 * @param Func - Function that will receive the node ID which can be stored and later used to get the elements using GetElementsForNode for all nodes that passed the predicate.
+	 */
+	template<typename PredicateFunc, typename IterateFunc>
+	inline void IterateNodesWithPredicate(const PredicateFunc& Predicate, const IterateFunc& Func) const
+	{
+		IterateNodesWithPredicateInternal(0, RootNodeContext, Predicate, Func);
+	}
+
+	/**
+	 * this function will traverse the Octree starting from the root in depth first order and the predicate can be used to implement custom culling for each node.
+	 * @param Predicate - a Function when given the bounds of the currently traversed node that returns true if traversal should continue or false to skip that branch.
+	 * @param Func - Function to call with each Element for nodes that passed the predicate.
+	 */
+	template<typename PredicateFunc, typename IterateFunc>
+	inline void IterateElementsWithPredicate(const PredicateFunc& Predicate, const IterateFunc& Func) const
+	{
+		IterateNodesWithPredicateInternal(0, RootNodeContext, Predicate, [&Func, this](FNodeIndex NodeIndex)
 		{
-			return *ElementIt;
-		}
+			for (typename TCallTraits<ElementType>::ConstReference Element : TreeElements[NodeIndex])
+			{
+				Func(Element);
+			}
+		});
+	}
 
-	private:
+	/**
+	 * this function will traverse the Octree using a fast box-box interssection this should be the preferred way of traversing the tree.
+	 * @param BoxBounds - the bounds to test if a node is traversed or skipped.
+	 * @param Func - Function to call with each Element for nodes that passed bounds test.
+	 */
+	template<typename IterateBoundsFunc>
+	inline void IterateElementsWithBoundsTest(const FBoxCenterAndExtent& BoxBounds, const IterateBoundsFunc& Func) const
+	{
+		IterateElementsWithBoundsTestInternal(0, RootNodeContext, BoxBounds, Func);
+	}
 
-		/** The bounding box to check for intersection with. */
-		FBoxCenterAndExtent IteratorBounds;
+	/**
+	 * this function will traverse the Octree using a fast box-box intersection and aborts traversal as soon as the Element function returns false.
+	 * @param BoxBounds - the bounds to test if a node is traversed or skipped.
+	 * @param Func - Function to call with each Element for nodes that passed bounds test.
+	 */
+	template<typename IterateBoundsFunc>
+	inline void FindFirstElementWithBoundsTest(const FBoxCenterAndExtent& BoxBounds, const IterateBoundsFunc& Func) const
+	{
+		bool ContinueTraversal = true;
+		FindFirstElementWithBoundsTestInternal(0, RootNodeContext, BoxBounds, Func, ContinueTraversal);
+	}
+
+	/**
+	 * Adds an element to the octree.
+	 * @param Element - The element to add.
+	 */
+	inline void AddElement(typename TCallTraits<ElementType>::ConstReference Element)
+	{
+		ElementArrayType TempElementStorage;
+		const FBoxCenterAndExtent ElementBounds(OctreeSemantics::GetBoundingBox(Element));
+		AddElementInternal(0, RootNodeContext, ElementBounds, Element, TempElementStorage);
+	}
 
-		/** The octree node iterator. */
-		TConstIterator<StackAllocator> NodeIt;
+	/**
+	 * Removes an element from the octree.
+	 * @param ElementId - The element to remove from the octree.
+	 */
+	inline void RemoveElement(FOctreeElementId ElementId)
+	{
+		checkSlow(ElementId.IsValidId());
 
-		/** The element iterator for the current node. */
-		ElementConstIt ElementIt;
+		// Remove the element from the node's element list.
+		TreeElements[ElementId.NodeIndex].RemoveAtSwap(ElementId.ElementIndex);
 
-		/** Processes the children of the current node. */
-		void ProcessChildren()
+		if (ElementId.ElementIndex < TreeElements[ElementId.NodeIndex].Num())
 		{
-			// Add the child nodes that intersect the bounding box to the node iterator's stack.
-			const FNode& CurrentNode = NodeIt.GetCurrentNode();
-			const FOctreeNodeContext& Context = NodeIt.GetCurrentContext();
-			const FOctreeChildNodeSubset IntersectingChildSubset = Context.GetIntersectingChildren(IteratorBounds);
-			FOREACH_OCTREE_CHILD_NODE(ChildRef)
-			{
-				if(IntersectingChildSubset.Contains(ChildRef) && CurrentNode.HasChild(ChildRef))
-				{
-					NodeIt.PushChild(ChildRef);
-				}
-			}
+			// Update the external element id for the element that was swapped into the vacated element index.
+			SetElementId(TreeElements[ElementId.NodeIndex][ElementId.ElementIndex], ElementId);
 		}
 
-		/** Advances the iterator to the next intersecting primitive, starting at a primitive in the current node. */
-		void AdvanceToNextIntersectingElement()
+		FNodeIndex CollapseNodeIndex = INDEX_NONE;
 		{
-			check(NodeIt.HasPendingNodes()); // please don't call this after iteration has ended
-
-			while (1)
+			// Update the inclusive element counts for the nodes between the element and the root node,
+			// and find the largest node that is small enough to collapse.
+			FNodeIndex NodeIndex = ElementId.NodeIndex;
+			while (true)
 			{
-				ElementConstIt LocalElementIt(ElementIt);
-				if (LocalElementIt)
+				TreeNodes[NodeIndex].InclusiveNumElements--;
+				if (TreeNodes[NodeIndex].InclusiveNumElements < OctreeSemantics::MinInclusiveElementsPerNode)
 				{
-					FPlatformMisc::Prefetch( &(*LocalElementIt) );
-					FPlatformMisc::Prefetch( &(*LocalElementIt), PLATFORM_CACHE_LINE_SIZE );
+					CollapseNodeIndex = NodeIndex;
+				}
 
-					// this is redundantly pull out of the while loop to prevent LHS on the iterator
-					// Check if the current element intersects the bounding box.
-					if(Intersect(OctreeSemantics::GetBoundingBox(*LocalElementIt),IteratorBounds))
-					{
-						// If it intersects, break out of the advancement loop.
-						Move(ElementIt, LocalElementIt);
-						return;
-					}
+				if (NodeIndex == 0)
+				{
+					break;
+				}
 
-					// Check if we've advanced past the elements in the current node.
-					while(++LocalElementIt)
-					{
-						FPlatformMisc::Prefetch( &(*LocalElementIt), PLATFORM_CACHE_LINE_SIZE );
+				NodeIndex = ParentLinks[(NodeIndex - 1) / 8];			
+			}
+		}
 
-						// Check if the current element intersects the bounding box.
-						if(Intersect(OctreeSemantics::GetBoundingBox(*LocalElementIt),IteratorBounds))
+		// Collapse the largest node that was pushed below the threshold for collapse by the removal.
+		if (CollapseNodeIndex != INDEX_NONE && !TreeNodes[CollapseNodeIndex].IsLeaf())
+		{
+			if (TreeElements[CollapseNodeIndex].Num() < (int32)TreeNodes[CollapseNodeIndex].InclusiveNumElements)
+			{
+				ElementArrayType TempElementStorage;
+				TempElementStorage.Reserve(TreeNodes[CollapseNodeIndex].InclusiveNumElements);
+				// Gather the elements contained in this node and its children.
+				CollapseNodesInternal(CollapseNodeIndex, TempElementStorage);
+				TreeElements[CollapseNodeIndex] = MoveTemp(TempElementStorage);
 
-						{
-							// If it intersects, break out of the advancement loop.
-							Move(ElementIt, LocalElementIt);
-							return;
-						}
-					}
-				}
-				// Advance to the next node.
-				NodeIt.Advance();
-				if(!NodeIt.HasPendingNodes())
+				for (int ElementIndex = 0; ElementIndex < TreeElements[CollapseNodeIndex].Num(); ElementIndex++)
 				{
-					Move(ElementIt, LocalElementIt);
-					return;
+					// Update the external element id for the element that's being collapsed.
+					SetElementId(TreeElements[CollapseNodeIndex][ElementIndex], FOctreeElementId(CollapseNodeIndex, ElementIndex));
 				}
-				ProcessChildren();
-				// The element iterator can't be assigned to, but it can be replaced by Move.
-				Move(ElementIt,NodeIt.GetCurrentNode().GetElementIt());
+
+				// Mark the node as a leaf.
+				TreeNodes[CollapseNodeIndex].ChildNodes = INDEX_NONE;
 			}
 		}
-	};
-
-	/**
-	 * Adds an element to the octree.
-	 * @param Element - The element to add.
-	 * @return An identifier for the element in the octree.
-	 */
-	void AddElement(typename TTypeTraits<ElementType>::ConstInitType Element);
+	}
 
 	/**
-	 * Removes an element from the octree.
-	 * @param ElementId - The element to remove from the octree.
+	 * this function resets the octree to empty.
 	 */
-	void RemoveElement(FOctreeElementId ElementId);
-
 	void Destroy()
 	{
-		RootNode.~FNode();
-		new (&RootNode) FNode(NULL);
-
-		// this looks a bit @hacky, but FNode's destructor doesn't 
-		// update TotalSizeBytes so better to set it to 0 than
-		// not do nothing with it (would contain obviously false value)
-		SetOctreeMemoryUsage(this, 0);
+		TreeNodes.Reset(1);
+		TreeElements.Reset(1);
+		FreeList.Reset();
+		ParentLinks.Reset();
+		TreeNodes.AddDefaulted();
+		TreeElements.AddDefaulted();
 	}
 
 	/** Accesses an octree element by ID. */
-	ElementType& GetElementById(FOctreeElementId ElementId);
+	ElementType& GetElementById(FOctreeElementId ElementId)
+	{
+		return TreeElements[ElementId.NodeIndex][ElementId.ElementIndex];
+	}
 
 	/** Accesses an octree element by ID. */
-	const ElementType& GetElementById(FOctreeElementId ElementId) const;
+	const ElementType& GetElementById(FOctreeElementId ElementId) const
+	{
+		return TreeElements[ElementId.NodeIndex][ElementId.ElementIndex];
+	}
+
+	/**
+	 * check if a FOctreeElementId is valid.
+	 * @param ElementId - The ElementId to check.
+	 */
+	bool IsValidElementId(FOctreeElementId ElementId) const
+	{
+		return ElementId.IsValidId() && ElementId.ElementIndex < TreeElements[ElementId.NodeIndex].Num();
+	};
 
-	/** Checks if given ElementId represents a valid Octree element */
-	bool IsValidElementId(FOctreeElementId ElementId) const;
+	/**
+	 * return all elements for a given node.
+	 * @param NodeIndex - The the index of the node can be obtained using IterateNodesWithPredicate.
+	 */
+	TArrayView<const ElementType> GetElementsForNode(FNodeIndex NodeIndex) const
+	{
+		return TreeElements[NodeIndex];
+	}
 
 	/** Writes stats for the octree to the log. */
-	void DumpStats() const;
+	void DumpStats() const
+	{
+		int32 NumNodes = 0;
+		int32 NumLeaves = 0;
+		int32 NumElements = 0;
+		int32 MaxElementsPerNode = 0;
+		TArray<int32> NodeElementDistribution;
+
+		IterateNodesWithPredicateInternal(0, RootNodeContext, [](const FBoxCenterAndExtent&) { return true; }, [&, this](FNodeIndex NodeIndex)
+		{
+			const int32 CurrentNodeElementCount = GetElementsForNode(NodeIndex).Num();
+
+			NumNodes++;
+			if (TreeNodes[NodeIndex].IsLeaf())
+			{
+				NumLeaves++;
+			}
+
+			NumElements += CurrentNodeElementCount;
+			MaxElementsPerNode = FMath::Max(MaxElementsPerNode, CurrentNodeElementCount);
+
+			if (CurrentNodeElementCount >= NodeElementDistribution.Num())
+			{
+				NodeElementDistribution.AddZeroed(CurrentNodeElementCount - NodeElementDistribution.Num() + 1);
+			}
+			NodeElementDistribution[CurrentNodeElementCount]++;
+		});
+
+		UE_LOG(LogGenericOctree, Log, TEXT("Octree overview:"));
+		UE_LOG(LogGenericOctree, Log, TEXT("\t%i nodes"), NumNodes);
+		UE_LOG(LogGenericOctree, Log, TEXT("\t%i leaves"), NumLeaves);
+		UE_LOG(LogGenericOctree, Log, TEXT("\t%i elements"), NumElements);
+		UE_LOG(LogGenericOctree, Log, TEXT("\t%i >= elements per node"), MaxElementsPerNode);
+		UE_LOG(LogGenericOctree, Log, TEXT("Octree node element distribution:"));
+		for (int32 i = 0; i < NodeElementDistribution.Num(); i++)
+		{
+			if (NodeElementDistribution[i] > 0)
+			{
+				UE_LOG(LogGenericOctree, Log, TEXT("\tElements: %3i, Nodes: %3i"), i, NodeElementDistribution[i]);
+			}
+		}
+	}
 
 	SIZE_T GetSizeBytes() const
 	{
+		SIZE_T TotalSizeBytes = TreeNodes.GetAllocatedSize();
+		TotalSizeBytes += TreeElements.GetAllocatedSize();
+		TotalSizeBytes += TreeNodes[0].InclusiveNumElements * sizeof(ElementType);
 		return TotalSizeBytes;
 	}
 
@@ -727,7 +905,10 @@ public:
 
 	void ShrinkElements()
 	{
-		RootNode.ShrinkElements();
+		for (ElementArrayType& Elements : TreeElements)
+		{
+			Elements.Shrink();
+		}
 	}
 
 	/** 
@@ -736,38 +917,46 @@ public:
 	 * bGlobalOctree - hint that this octree is used as a boundless global volume, 
 	 *  so only content will be shifted but not origin of the octree
 	 */
-	void ApplyOffset(const FVector& InOffset, bool bGlobalOctree = false);
-
-	/** Initialization constructor. */
-	TOctree(const FVector& InOrigin,float InExtent);
-
-	/** DO NOT USE. This constructor is for internal usage only for hot-reload purposes. */
-	TOctree();
+	void ApplyOffset(const FVector& InOffset, bool bGlobalOctree = false)
+	{
+		ElementArrayType TempElementStorage;
+		TempElementStorage.Reserve(TreeNodes[0].InclusiveNumElements);
 
-private:
+		//collect all elements
+		CollapseNodesInternal(0, TempElementStorage);
+		checkSlow(TreeNodes[0].IsLeaf());
+		Destroy();
 
-	/** The octree's root node. */
-	FNode RootNode;
+		if (!bGlobalOctree)
+		{
+			RootNodeContext.Bounds.Center += FVector4(InOffset, 0.0f);
+		}
 
-	/** The octree's root node's context. */
-	FOctreeNodeContext RootNodeContext;
+		// Add all elements from a saved nodes to a new empty octree
+		for (const ElementType& Element : TempElementStorage)
+		{
+			AddElement(Element);
+		}
+	}
 
-	/** The extent of a leaf at the maximum allowed depth of the tree. */
-	float MinLeafExtent;
+	/** Initialization constructor. */
+	TOctree(const FVector& InOrigin,float InExtent)
+		: RootNodeContext(FBoxCenterAndExtent(InOrigin, FVector(InExtent, InExtent, InExtent)), 0, 0)
+		, MinLeafExtent(InExtent* FMath::Pow((1.0f + 1.0f / (float)FOctreeNodeContext::LoosenessDenominator) / 2.0f, OctreeSemantics::MaxNodeDepth))
+	{
+		TreeNodes.AddDefaulted();
+		TreeElements.AddDefaulted();
+	}
 
-	/** this function basically set TotalSizeBytes, but gives opportunity to 
-	 *	include this Octree in memory stats */
-	template<typename TElement,typename TSemantics>
-	friend void SetOctreeMemoryUsage(class TOctree<TElement,TSemantics>* Octree, int32 NewSize);
+	/** DO NOT USE. This constructor is for internal usage only for hot-reload purposes. */
+	TOctree() 
+	{
+		TreeNodes.AddDefaulted();
+		TreeElements.AddDefaulted();
+	}
 
-	SIZE_T TotalSizeBytes;
+private:
 
-	/** Adds an element to a node or its children. */
-	void AddElementToNode(
-		typename TTypeTraits<ElementType>::ConstInitType Element,
-		const FNode& InNode,
-		const FOctreeNodeContext& InContext
-		);
 
 	// Concept definition for the new octree semantics, which adds a new TOctree parameter
 	struct COctreeSemanticsV2
@@ -798,4 +987,3 @@ protected:
 };
 
 #include "GenericOctree.inl"
-
diff --git a/Engine/Source/Runtime/Core/Public/Math/GenericOctree.inl b/Engine/Source/Runtime/Core/Public/Math/GenericOctree.inl
--- a/Engine/Source/Runtime/Core/Public/Math/GenericOctree.inl
+++ b/Engine/Source/Runtime/Core/Public/Math/GenericOctree.inl
@@ -19,7 +19,6 @@ struct FMath;
 template<typename ElementType,typename OctreeSemantics> class TOctree;
 template<typename T> struct TTypeTraits;
 
- CORE_API DECLARE_LOG_CATEGORY_EXTERN(LogGenericOctree, Log, All);
 
 FORCEINLINE bool FOctreeChildNodeSubset::Contains(FOctreeChildNodeRef ChildRef) const
 {
@@ -51,12 +50,8 @@ FORCEINLINE FOctreeChildNodeSubset FOctreeNodeContext::GetIntersectingChildren(c
 		);
 
 	// Intersect the query bounds with the node's children's bounds.
-	Result.bPositiveX = VectorAnyGreaterThan(VectorReplicate(QueryBoundsMax,0),VectorReplicate(PositiveChildBoundsMin,0)) != false;
-	Result.bPositiveY = VectorAnyGreaterThan(VectorReplicate(QueryBoundsMax,1),VectorReplicate(PositiveChildBoundsMin,1)) != false;
-	Result.bPositiveZ = VectorAnyGreaterThan(VectorReplicate(QueryBoundsMax,2),VectorReplicate(PositiveChildBoundsMin,2)) != false;
-	Result.bNegativeX = VectorAnyGreaterThan(VectorReplicate(QueryBoundsMin,0),VectorReplicate(NegativeChildBoundsMax,0)) == false;
-	Result.bNegativeY = VectorAnyGreaterThan(VectorReplicate(QueryBoundsMin,1),VectorReplicate(NegativeChildBoundsMax,1)) == false;
-	Result.bNegativeZ = VectorAnyGreaterThan(VectorReplicate(QueryBoundsMin,2),VectorReplicate(NegativeChildBoundsMax,2)) == false;
+	Result.PositiveChildBits = VectorMaskBits(VectorMask_GT(QueryBoundsMax, PositiveChildBoundsMin)) & 0x7;
+	Result.NegativeChildBits = VectorMaskBits(VectorMask_LE(QueryBoundsMin, NegativeChildBoundsMax)) & 0x7;
 	return Result;
 }
 
@@ -83,336 +78,9 @@ FORCEINLINE FOctreeChildNodeRef FOctreeNodeContext::GetContainingChild(const FBo
 	else
 	{
 		// Return the child node that the query is closest to as the containing child.
-		Result.X = VectorAnyGreaterThan(VectorReplicate(QueryBoundsCenter,0),VectorReplicate(BoundsCenter,0)) != false;
-		Result.Y = VectorAnyGreaterThan(VectorReplicate(QueryBoundsCenter,1),VectorReplicate(BoundsCenter,1)) != false;
-		Result.Z = VectorAnyGreaterThan(VectorReplicate(QueryBoundsCenter,2),VectorReplicate(BoundsCenter,2)) != false;
+		Result.Index = VectorMaskBits(VectorMask_GT(QueryBoundsCenter, BoundsCenter)) & 0x7;
 	}
 
 	return Result;
 }
 
-template<typename ElementType,typename OctreeSemantics>
-void TOctree<ElementType,OctreeSemantics>::AddElement(typename TTypeTraits<ElementType>::ConstInitType Element)
-{
-	AddElementToNode(Element,RootNode,RootNodeContext);
-}
-															
-template<typename ElementType,typename OctreeSemantics>
-void TOctree<ElementType,OctreeSemantics>::AddElementToNode(
-	typename TTypeTraits<ElementType>::ConstInitType Element,
-	const FNode& InNode,
-	const FOctreeNodeContext& InContext
-	)
-{
-	const FBoxCenterAndExtent ElementBounds(OctreeSemantics::GetBoundingBox(Element));
-
-	for(TConstIterator<TInlineAllocator<1> > NodeIt(InNode,InContext);NodeIt.HasPendingNodes();NodeIt.Advance())
-	{
-		const FNode& Node = NodeIt.GetCurrentNode();
-		const FOctreeNodeContext& Context = NodeIt.GetCurrentContext();
-		const bool bIsLeaf = Node.IsLeaf();
-
-		bool bAddElementToThisNode = false;
-
-		// Increment the number of elements included in this node and its children.
-		Node.InclusiveNumElements++;
-
-		if(bIsLeaf)
-		{
-			// If this is a leaf, check if adding this element would turn it into a node by overflowing its element list.
-			if(Node.Elements.Num() + 1 > OctreeSemantics::MaxElementsPerLeaf && Context.Bounds.Extent.X > MinLeafExtent)
-			{
-				// Copy the leaf's elements, remove them from the leaf, and turn it into a node.
-				ElementArrayType ChildElements;
-				Exchange(ChildElements,Node.Elements);
-				SetOctreeMemoryUsage(this, TotalSizeBytes - ChildElements.Num() * sizeof(ElementType));
-				Node.InclusiveNumElements = 0;
-
-				// Allow elements to be added to children of this node.
-				Node.bIsLeaf = false;
-
-				// Re-add all of the node's child elements, potentially creating children of this node for them.
-				for(ElementConstIt ElementIt(ChildElements);ElementIt;++ElementIt)
-				{
-					AddElementToNode(*ElementIt,Node,Context);
-				}
-
-				// Add the element to this node.
-				AddElementToNode(Element,Node,Context);
-				return;
-			}
-			else
-			{
-				// If the leaf has room for the new element, simply add it to the list.
-				bAddElementToThisNode = true;
-			}
-		}
-		else
-		{
-			// If this isn't a leaf, find a child that entirely contains the element.
-			const FOctreeChildNodeRef ChildRef = Context.GetContainingChild(ElementBounds);	
-			if(ChildRef.IsNULL())
-			{
-				// If none of the children completely contain the element, add it to this node directly.
-				bAddElementToThisNode = true;
-			}
-			else
-			{
-				// Create the child node if it hasn't been created yet.
-				if(!Node.Children[ChildRef.Index])
-				{
-					Node.Children[ChildRef.Index] = new typename TOctree<ElementType,OctreeSemantics>::FNode(&Node);
-					SetOctreeMemoryUsage(this, TotalSizeBytes + sizeof(*Node.Children[ChildRef.Index]));
-				}
-
-				// Push the node onto the stack to visit.
-				NodeIt.PushChild(ChildRef);
-			}
-		}
-
-		if(bAddElementToThisNode)
-		{
-			// Add the element to this node.
-			new(Node.Elements) ElementType(Element);
-
-			SetOctreeMemoryUsage(this, TotalSizeBytes + sizeof(ElementType));
-			
-			// Set the element's ID.
-			SetElementId(Element,FOctreeElementId(&Node,Node.Elements.Num() - 1));
-			return;
-		}
-	}
-
-	UE_LOG(LogGenericOctree, Fatal,
-		TEXT("Failed to find an octree node for an element with bounds (%f,%f,%f) +/- (%f,%f,%f)!"),
-		ElementBounds.Center.X,
-		ElementBounds.Center.Y,
-		ElementBounds.Center.Z,
-		ElementBounds.Extent.X,
-		ElementBounds.Extent.Y,
-		ElementBounds.Extent.Z
-		);
-}
-
-template<typename ElementType,typename OctreeSemantics>
-void TOctree<ElementType,OctreeSemantics>::RemoveElement(FOctreeElementId ElementId)
-{
-	check(ElementId.IsValidId()); 
-
-	FNode* ElementIdNode = (FNode*)ElementId.Node;
-
-	// Remove the element from the node's element list.
-	ElementIdNode->Elements.RemoveAtSwap(ElementId.ElementIndex);
-
-	SetOctreeMemoryUsage(this, TotalSizeBytes - sizeof(ElementType));
-
-	if(ElementId.ElementIndex < ElementIdNode->Elements.Num())
-	{
-		// Update the external element id for the element that was swapped into the vacated element index.
-		SetElementId(ElementIdNode->Elements[ElementId.ElementIndex],ElementId);
-	}
-
-	// Update the inclusive element counts for the nodes between the element and the root node,
-	// and find the largest node that is small enough to collapse.
-	const FNode* CollapseNode = NULL;
-	for(const FNode* Node = ElementIdNode;Node;Node = Node->Parent)
-	{
-		--Node->InclusiveNumElements;
-		if(Node->InclusiveNumElements < OctreeSemantics::MinInclusiveElementsPerNode)
-		{
-			CollapseNode = Node;
-		}
-	}
-
-	// Collapse the largest node that was pushed below the threshold for collapse by the removal.
-	if (CollapseNode && !CollapseNode->bIsLeaf)
-	{
-		if (CollapseNode->Elements.Num() < (int32)CollapseNode->InclusiveNumElements)
-		{
-			CollapseNode->Elements.Reserve(CollapseNode->InclusiveNumElements);
-
-			// Gather the elements contained in this node and its children.
-			for (TConstIterator<> ChildNodeIt(*CollapseNode, RootNodeContext); ChildNodeIt.HasPendingNodes(); ChildNodeIt.Advance())
-			{
-				const FNode& ChildNode = ChildNodeIt.GetCurrentNode();
-
-				if (&ChildNode != CollapseNode)
-				{
-					// Child node will be collapsed so move the child's elements to the collapse node element list.
-					for (ElementType& Element : ChildNode.Elements)
-					{
-						const int32 NewElementIndex = CollapseNode->Elements.Add(MoveTemp(Element));
-
-						// Update the external element id for the element that's being collapsed.
-						SetElementId(CollapseNode->Elements[NewElementIndex], FOctreeElementId(CollapseNode, NewElementIndex));
-					}
-				}
-
-				// Recursively visit all child nodes.
-				FOREACH_OCTREE_CHILD_NODE(ChildRef)
-				{
-					if (ChildNode.HasChild(ChildRef))
-					{
-						ChildNodeIt.PushChild(ChildRef);
-					}
-				}
-			}
-
-			// Free the child nodes.
-			FOREACH_OCTREE_CHILD_NODE(ChildRef)
-			{
-				if (CollapseNode->Children[ChildRef.Index])
-				{
-					SetOctreeMemoryUsage(this, TotalSizeBytes - sizeof(*CollapseNode->Children[ChildRef.Index]));
-				}
-
-				delete CollapseNode->Children[ChildRef.Index];
-				CollapseNode->Children[ChildRef.Index] = NULL;
-			}
-		}
-
-		// Mark the node as a leaf.
-		CollapseNode->bIsLeaf = true;		
-	}
-}
-
-template<typename ElementType,typename OctreeSemantics>
-ElementType& TOctree<ElementType,OctreeSemantics>::GetElementById(FOctreeElementId ElementId)
-{
-	check(ElementId.IsValidId());
-	FNode* ElementIdNode = (FNode*)ElementId.Node;
-	return ElementIdNode->Elements[ElementId.ElementIndex];
-}
-
-template<typename ElementType, typename OctreeSemantics>
-const ElementType& TOctree<ElementType, OctreeSemantics>::GetElementById(FOctreeElementId ElementId) const
-{
-	check(ElementId.IsValidId());
-	FNode* ElementIdNode = (FNode*)ElementId.Node;
-	return ElementIdNode->Elements[ElementId.ElementIndex];
-}
-
-template<typename ElementType,typename OctreeSemantics> 
-bool TOctree<ElementType,OctreeSemantics>::IsValidElementId(FOctreeElementId ElementId) const
-{
-	return ElementId.IsValidId() 
-		&& ElementId.ElementIndex != INDEX_NONE
-		&& ElementId.ElementIndex < ((FNode*)ElementId.Node)->Elements.Num();
-};
-
-template<typename ElementType,typename OctreeSemantics>
-void TOctree<ElementType,OctreeSemantics>::DumpStats() const
-{
-	int32 NumNodes = 0;
-	int32 NumLeaves = 0;
-	int32 NumElements = 0;
-	int32 MaxElementsPerNode = 0;
-	TArray<int32> NodeElementDistribution;
-
-	for(TConstIterator<> NodeIt(*this);NodeIt.HasPendingNodes();NodeIt.Advance())
-	{
-		const FNode& CurrentNode = NodeIt.GetCurrentNode();
-		const int32 CurrentNodeElementCount = CurrentNode.GetElementCount();
-
-		NumNodes++;
-		if(CurrentNode.IsLeaf())
-		{
-			NumLeaves++;
-		}
-
-		NumElements += CurrentNodeElementCount;
-		MaxElementsPerNode = FMath::Max(MaxElementsPerNode,CurrentNodeElementCount);
-
-		if( CurrentNodeElementCount >= NodeElementDistribution.Num() )
-		{
-			NodeElementDistribution.AddZeroed( CurrentNodeElementCount - NodeElementDistribution.Num() + 1 );
-		}
-		NodeElementDistribution[CurrentNodeElementCount]++;
-
-		FOREACH_OCTREE_CHILD_NODE(ChildRef)
-		{
-			if(CurrentNode.HasChild(ChildRef))
-			{
-				NodeIt.PushChild(ChildRef);
-			}
-		}
-	}
-
-	UE_LOG(LogGenericOctree, Log, TEXT("Octree overview:"));
-	UE_LOG(LogGenericOctree, Log, TEXT("\t%i nodes"),NumNodes);
-	UE_LOG(LogGenericOctree, Log, TEXT("\t%i leaves"),NumLeaves);
-	UE_LOG(LogGenericOctree, Log, TEXT("\t%i elements"),NumElements);
-	UE_LOG(LogGenericOctree, Log, TEXT("\t%i >= elements per node"),MaxElementsPerNode);
-	UE_LOG(LogGenericOctree, Log, TEXT("Octree node element distribution:"));
-	for( int32 i=0; i<NodeElementDistribution.Num(); i++ )
-	{
-		if( NodeElementDistribution[i] > 0 )
-		{
-			UE_LOG(LogGenericOctree, Log, TEXT("\tElements: %3i, Nodes: %3i"),i,NodeElementDistribution[i]);
-		}
-	}
-}
-
-template<typename ElementType,typename OctreeSemantics>
-FORCEINLINE void SetOctreeMemoryUsage(TOctree<ElementType,OctreeSemantics>* Octree, int32 NewSize)
-{
-	Octree->TotalSizeBytes = NewSize;
-}
-
-template<typename ElementType,typename OctreeSemantics>
-TOctree<ElementType,OctreeSemantics>::TOctree(const FVector& InOrigin,float InExtent)
-:	RootNode(NULL)
-,	RootNodeContext(FBoxCenterAndExtent(InOrigin,FVector(InExtent,InExtent,InExtent)),0,0)
-,	MinLeafExtent(InExtent * FMath::Pow((1.0f + 1.0f / (float)FOctreeNodeContext::LoosenessDenominator) / 2.0f,OctreeSemantics::MaxNodeDepth))
-,	TotalSizeBytes(0)
-{
-}
-
-template<typename ElementType, typename OctreeSemantics>
-TOctree<ElementType, OctreeSemantics>::TOctree()
-	: RootNode(nullptr)
-{
-	EnsureRetrievingVTablePtrDuringCtor(TEXT("TOctree()"));
-}
-
-template<typename ElementType,typename OctreeSemantics>
-void TOctree<ElementType,OctreeSemantics>::ApplyOffset(const FVector& InOffset, bool bGlobalOctree)
-{
-	// Shift elements
-	RootNode.ApplyOffset(InOffset);
-	
-	// Make a local copy of all nodes
-	FNode OldRootNode = RootNode;
-	
-	// Assign to root node a NULL node, so Destroy procedure will not delete child nodes
-	RootNode = FNode(nullptr);
-	// Call destroy to clean up octree
-	Destroy();
-
-	if (!bGlobalOctree)
-	{
-		RootNodeContext.Bounds.Center += FVector4(InOffset, 0.0f);
-	}
-	
-	// Add all elements from a saved nodes to a new empty octree
-	for (TConstIterator<> NodeIt(OldRootNode, RootNodeContext); NodeIt.HasPendingNodes(); NodeIt.Advance())
-	{
-		const auto& CurrentNode = NodeIt.GetCurrentNode();
-
-		FOREACH_OCTREE_CHILD_NODE(ChildRef)
-		{
-			if (CurrentNode.HasChild(ChildRef))
-			{
-				NodeIt.PushChild(ChildRef);
-			}
-		}
-
-		for (auto ElementIt = CurrentNode.GetElementIt(); ElementIt; ++ElementIt)
-		{
-			AddElement(*ElementIt);
-		}
-	}
-
-	// Saved nodes will be deleted on scope exit
-}
-
diff --git a/Engine/Source/Runtime/Core/Public/Math/GenericOctreePublic.h b/Engine/Source/Runtime/Core/Public/Math/GenericOctreePublic.h
--- a/Engine/Source/Runtime/Core/Public/Math/GenericOctreePublic.h
+++ b/Engine/Source/Runtime/Core/Public/Math/GenericOctreePublic.h
@@ -20,27 +20,27 @@ public:
 
 	/** Default constructor. */
 	FOctreeElementId()
-		:	Node(NULL)
+		:	NodeIndex(INDEX_NONE)
 		,	ElementIndex(INDEX_NONE)
 	{}
 
 	/** @return a boolean value representing whether the id is NULL. */
 	bool IsValidId() const
 	{
-		return Node != NULL;
+		return NodeIndex != INDEX_NONE;
 	}
 
 private:
 
 	/** The node the element is in. */
-	const void* Node;
+	uint32 NodeIndex;
 
 	/** The index of the element in the node's element array. */
 	int32 ElementIndex;
 
 	/** Initialization constructor. */
-	FOctreeElementId(const void* InNode,int32 InElementIndex)
-		:	Node(InNode)
+	FOctreeElementId(uint32 InNodeIndex, int32 InElementIndex)
+		:	NodeIndex(InNodeIndex)
 		,	ElementIndex(InElementIndex)
 	{}
 
diff --git a/Engine/Source/Runtime/Core/Public/Math/UnrealMathSSE.h b/Engine/Source/Runtime/Core/Public/Math/UnrealMathSSE.h
--- a/Engine/Source/Runtime/Core/Public/Math/UnrealMathSSE.h
+++ b/Engine/Source/Runtime/Core/Public/Math/UnrealMathSSE.h
@@ -1489,6 +1489,6 @@ FORCEINLINE VectorRegisterInt VectorIntAbs(const VectorRegisterInt& A)
 * @param Ptr	Unaligned memory pointer to the 4 int32s
 * @return		VectorRegisterInt(*Ptr, *Ptr, *Ptr, *Ptr)
 */
-#define VectorIntLoad1( Ptr )	_mm_shuffle_epi32(_mm_loadu_si128((VectorRegisterInt*)(Ptr)),_MM_SHUFFLE(0,0,0,0))
+#define VectorIntLoad1( Ptr )	_mm_set1_epi32(*(Ptr))
 #endif
 
diff --git a/Engine/Source/Runtime/Engine/Private/ConvexVolume.cpp b/Engine/Source/Runtime/Engine/Private/ConvexVolume.cpp
--- a/Engine/Source/Runtime/Engine/Private/ConvexVolume.cpp
+++ b/Engine/Source/Runtime/Engine/Private/ConvexVolume.cpp
@@ -159,11 +159,9 @@ FOutcode FConvexVolume::GetBoxIntersectionOutcode(const FVector& Origin,const FV
 
 static FORCEINLINE bool IntersectBoxWithPermutedPlanes(
 	const FConvexVolume::FPermutedPlaneArray& PermutedPlanes,
-	const VectorRegister BoxOrigin,
-	const VectorRegister BoxExtent )
+	const VectorRegister& BoxOrigin,
+	const VectorRegister& BoxExtent )
 {
-	bool Result = true;
-
 	checkSlow(PermutedPlanes.Num() % 4 == 0);
 
 	// Splat origin into 3 vectors
@@ -206,11 +204,10 @@ static FORCEINLINE bool IntersectBoxWithPermutedPlanes(
 		// Check for completely outside
 		if ( VectorAnyGreaterThan(Distance,PushOut) )
 		{
-			Result = false;
-			break;
+			return false;
 		}
 	}
-	return Result;
+	return true;
 }
 
 bool FConvexVolume::IntersectBox(const FVector& Origin,const FVector& Extent) const
@@ -240,8 +237,6 @@ bool FConvexVolume::IntersectBox( const FVector& Origin,const FVector& Translati
 
 bool FConvexVolume::IntersectBox(const FVector& Origin,const FVector& Extent, bool& bOutFullyContained) const
 {
-	bool Result = true;
-
 	// Assume fully contained
 	bOutFullyContained = true;
 
@@ -291,9 +286,8 @@ bool FConvexVolume::IntersectBox(const FVector& Origin,const FVector& Extent, bo
 		// Check for completely outside
 		if (VectorAnyGreaterThan(Distance,PushOut))
 		{
-			Result = false;
 			bOutFullyContained = false;
-			break;
+			return false;
 		}
 
 		// Definitely inside frustums, but check to see if it's fully contained
@@ -302,7 +296,7 @@ bool FConvexVolume::IntersectBox(const FVector& Origin,const FVector& Extent, bo
 			bOutFullyContained = false;
 		}
 	}
-	return Result;
+	return true;
 }
 
 //
@@ -311,8 +305,6 @@ bool FConvexVolume::IntersectBox(const FVector& Origin,const FVector& Extent, bo
 
 bool FConvexVolume::IntersectSphere(const FVector& Origin,const float& Radius) const
 {
-	bool Result = true;
-
 	checkSlow(PermutedPlanes.Num() % 4 == 0);
 
 	// Load the origin & radius
@@ -345,11 +337,10 @@ bool FConvexVolume::IntersectSphere(const FVector& Origin,const float& Radius) c
 		// Check for completely outside
 		if (VectorAnyGreaterThan(Distance,VRadius))
 		{
-			Result = false;
-			break;
+			return false;
 		}
 	}
-	return Result;
+	return true;
 }
 
 
@@ -359,8 +350,6 @@ bool FConvexVolume::IntersectSphere(const FVector& Origin,const float& Radius) c
 
 bool FConvexVolume::IntersectSphere(const FVector& Origin,const float& Radius, bool& bOutFullyContained) const
 {
-	bool Result = true;
-
 	//Assume fully contained
 	bOutFullyContained = true;
 
@@ -398,9 +387,8 @@ bool FConvexVolume::IntersectSphere(const FVector& Origin,const float& Radius, b
 		int32 Mask = VectorAnyGreaterThan(Distance,VRadius);
 		if (Mask)
 		{
-			Result = false;
 			bOutFullyContained = false;
-			break;
+			return false;
 		}
 
 		//the sphere is definitely inside the frustums, but let's check if it's FULLY contained by checking the NEGATIVE radius (on the inside of each frustum plane)
@@ -410,7 +398,7 @@ bool FConvexVolume::IntersectSphere(const FVector& Origin,const float& Radius, b
 			bOutFullyContained = false;
 		}
 	}
-	return Result;
+	return true;
 }
 
 bool FConvexVolume::IntersectLineSegment(const FVector& InStart, const FVector& InEnd) const
diff --git a/Engine/Source/Runtime/Engine/Private/PrecomputedLightVolume.cpp b/Engine/Source/Runtime/Engine/Private/PrecomputedLightVolume.cpp
--- a/Engine/Source/Runtime/Engine/Private/PrecomputedLightVolume.cpp
+++ b/Engine/Source/Runtime/Engine/Private/PrecomputedLightVolume.cpp
@@ -221,24 +221,10 @@ FArchive& operator<<(FArchive& Ar,FPrecomputedLightVolumeData& Volume)
 			if (!Ar.IsCooking() || Ar.CookingTarget()->SupportsFeature(ETargetPlatformFeatures::HighQualityLightmaps))
 			{
 				// Gather an array of samples from the octree
-				for(FLightVolumeOctree::TConstIterator<> NodeIt(Volume.HighQualityLightmapOctree); NodeIt.HasPendingNodes(); NodeIt.Advance())
+				Volume.HighQualityLightmapOctree.IterateAllElements([&HighQualitySamples](const FVolumeLightingSample& Sample)
 				{
-					const FLightVolumeOctree::FNode& CurrentNode = NodeIt.GetCurrentNode();
-
-					FOREACH_OCTREE_CHILD_NODE(ChildRef)
-					{
-						if(CurrentNode.HasChild(ChildRef))
-						{
-							NodeIt.PushChild(ChildRef);
-						}
-					}
-
-					for (FLightVolumeOctree::ElementConstIt ElementIt(CurrentNode.GetElementIt()); ElementIt; ++ElementIt)
-					{
-						const FVolumeLightingSample& Sample = *ElementIt;
-						HighQualitySamples.Add(Sample);
-					}
-				}
+					HighQualitySamples.Add(Sample);
+				});
 			}
 			
 			Ar << HighQualitySamples;
@@ -248,24 +234,10 @@ FArchive& operator<<(FArchive& Ar,FPrecomputedLightVolumeData& Volume)
 			if (!Ar.IsCooking() || Ar.CookingTarget()->SupportsFeature(ETargetPlatformFeatures::LowQualityLightmaps))
 			{
 				// Gather an array of samples from the octree
-				for(FLightVolumeOctree::TConstIterator<> NodeIt(Volume.LowQualityLightmapOctree); NodeIt.HasPendingNodes(); NodeIt.Advance())
+				Volume.LowQualityLightmapOctree.IterateAllElements([&LowQualitySamples](const FVolumeLightingSample& Sample)
 				{
-					const FLightVolumeOctree::FNode& CurrentNode = NodeIt.GetCurrentNode();
-
-					FOREACH_OCTREE_CHILD_NODE(ChildRef)
-					{
-						if(CurrentNode.HasChild(ChildRef))
-						{
-							NodeIt.PushChild(ChildRef);
-						}
-					}
-
-					for (FLightVolumeOctree::ElementConstIt ElementIt(CurrentNode.GetElementIt()); ElementIt; ++ElementIt)
-					{
-						const FVolumeLightingSample& Sample = *ElementIt;
-						LowQualitySamples.Add(Sample);
-					}
-				}
+					LowQualitySamples.Add(Sample);
+				});
 			}
 
 			Ar << LowQualitySamples;
@@ -345,37 +317,8 @@ void FPrecomputedLightVolumeData::InvalidateLightingCache()
 SIZE_T FPrecomputedLightVolumeData::GetAllocatedBytes() const
 {
 	SIZE_T NodeBytes = 0;
-
-	for (FLightVolumeOctree::TConstIterator<> NodeIt(HighQualityLightmapOctree); NodeIt.HasPendingNodes(); NodeIt.Advance())
-	{
-		const FLightVolumeOctree::FNode& CurrentNode = NodeIt.GetCurrentNode();
-		NodeBytes += sizeof(FLightVolumeOctree::FNode);
-		NodeBytes += CurrentNode.GetElements().GetAllocatedSize();
-
-		FOREACH_OCTREE_CHILD_NODE(ChildRef)
-		{
-			if(CurrentNode.HasChild(ChildRef))
-			{
-				NodeIt.PushChild(ChildRef);
-			}
-		}
-	}
-
-	for (FLightVolumeOctree::TConstIterator<> NodeIt(LowQualityLightmapOctree); NodeIt.HasPendingNodes(); NodeIt.Advance())
-	{
-		const FLightVolumeOctree::FNode& CurrentNode = NodeIt.GetCurrentNode();
-		NodeBytes += sizeof(FLightVolumeOctree::FNode);
-		NodeBytes += CurrentNode.GetElements().GetAllocatedSize();
-
-		FOREACH_OCTREE_CHILD_NODE(ChildRef)
-		{
-			if(CurrentNode.HasChild(ChildRef))
-			{
-				NodeIt.PushChild(ChildRef);
-			}
-		}
-	}
-
+	NodeBytes += HighQualityLightmapOctree.GetSizeBytes();
+	NodeBytes += LowQualityLightmapOctree.GetSizeBytes();
 	return NodeBytes;
 }
 
@@ -454,11 +397,8 @@ void FPrecomputedLightVolume::InterpolateIncidentRadiancePoint(
 		FVector WorldPosition = InWorldPosition - WorldOriginOffset; // relocate from world to volume space
 		FBoxCenterAndExtent BoundingBox( WorldPosition, FVector::ZeroVector );
 		// Iterate over the octree nodes containing the query point.
-		for (FLightVolumeOctree::TConstElementBoxIterator<> OctreeIt(*OctreeForRendering, BoundingBox);
-			OctreeIt.HasPendingElements();
-			OctreeIt.Advance())
+		OctreeForRendering->IterateElementsWithBoundsTest(BoundingBox, [&AccumulatedIncidentRadiance, &SkyBentNormal, &AccumulatedDirectionalLightShadowing, &AccumulatedWeight, &WorldPosition](const FVolumeLightingSample& VolumeSample)
 		{
-			const FVolumeLightingSample& VolumeSample = OctreeIt.GetCurrentElement();
 			const float DistanceSquared = (VolumeSample.Position - WorldPosition).SizeSquared();
 			const float RadiusSquared = FMath::Square(VolumeSample.Radius);
 
@@ -475,7 +415,7 @@ void FPrecomputedLightVolume::InterpolateIncidentRadiancePoint(
 				AccumulatedDirectionalLightShadowing += VolumeSample.DirectionalLightShadowing * SampleWeight;
 				AccumulatedWeight += SampleWeight;
 			}
-		}
+		});
 	}
 }
 
@@ -501,12 +441,10 @@ void FPrecomputedLightVolume::InterpolateIncidentRadianceBlock(
 		PotentiallyIntersectingSamples.Reset(100);
 
 		// Iterate over the octree nodes containing the query point.
-		for (FLightVolumeOctree::TConstElementBoxIterator<> OctreeIt(*OctreeForRendering, BoundingBox);
-			OctreeIt.HasPendingElements();
-			OctreeIt.Advance())
+		OctreeForRendering->IterateElementsWithBoundsTest(BoundingBox, [&](const FVolumeLightingSample& VolumeSample)
 		{
-			PotentiallyIntersectingSamples.Add(&OctreeIt.GetCurrentElement());
-		}
+			PotentiallyIntersectingSamples.Add(&VolumeSample);
+		});
 		
 		const int32 LinearIndexBase = DestCellPosition.Z * DestCellDimensions.Y * DestCellDimensions.X
 			+ DestCellPosition.Y * DestCellDimensions.X
@@ -562,18 +500,15 @@ void FPrecomputedLightVolume::InterpolateIncidentRadianceBlock(
 
 void FPrecomputedLightVolume::DebugDrawSamples(FPrimitiveDrawInterface* PDI, bool bDrawDirectionalShadowing) const
 {
-	for (FLightVolumeOctree::TConstElementBoxIterator<> OctreeIt(*OctreeForRendering, OctreeForRendering->GetRootBounds());
-		OctreeIt.HasPendingElements();
-		OctreeIt.Advance())
+	OctreeForRendering->IterateAllElements([bDrawDirectionalShadowing, PDI, this](const FVolumeLightingSample& VolumeSample)
 	{
-		const FVolumeLightingSample& VolumeSample = OctreeIt.GetCurrentElement();
 		const FLinearColor AverageColor = bDrawDirectionalShadowing
 			? FLinearColor(VolumeSample.DirectionalLightShadowing, VolumeSample.DirectionalLightShadowing, VolumeSample.DirectionalLightShadowing)
 			: VolumeSample.Lighting.CalcIntegral() / (FSHVector2::ConstantBasisIntegral * PI);
 		
 		FVector SamplePosition = VolumeSample.Position + WorldOriginOffset; //relocate from volume to world space
 		PDI->DrawPoint(SamplePosition, AverageColor, 10, SDPG_World);
-	}
+	});
 }
 
 bool FPrecomputedLightVolume::IntersectBounds(const FBoxSphereBounds& InBounds) const
diff --git a/Engine/Source/Runtime/Engine/Private/Rendering/SkeletalMeshLODImporterData.cpp b/Engine/Source/Runtime/Engine/Private/Rendering/SkeletalMeshLODImporterData.cpp
--- a/Engine/Source/Runtime/Engine/Private/Rendering/SkeletalMeshLODImporterData.cpp
+++ b/Engine/Source/Runtime/Engine/Private/Rendering/SkeletalMeshLODImporterData.cpp
@@ -937,33 +937,15 @@ void FOctreeQueryHelper::FindNearestWedgeIndexes(const FVector& SearchPosition,
 	FVector Extend(2.0f);
 	for (int i = 0; i < 2; ++i)
 	{
-		TWedgeInfoPosOctree::TConstIterator<> OctreeIter((*WedgePosOctree));
 		// Iterate through the octree attempting to find the vertices closest to the current new point
-		while (OctreeIter.HasPendingNodes())
+		// The first shot is an intersection with a 1 CM cube box around the search position, this ensure we dont fall in the wrong neighbourg
+		WedgePosOctree->IterateElementsWithBoundsTest(FBoxCenterAndExtent(SearchPosition, Extend), [&OutNearestWedges, &MinSquaredDistance, &SearchPosition](const FWedgeInfo& WedgeInfo)
 		{
-			const TWedgeInfoPosOctree::FNode& CurNode = OctreeIter.GetCurrentNode();
-			const FOctreeNodeContext& CurContext = OctreeIter.GetCurrentContext();
-
-			// Find the child of the current node, if any, that contains the current new point
-
-			//The first shot is an intersection with a 1 CM cube box around the search position, this ensure we dont fall in the wrong neighbourg
-			FOctreeChildNodeSubset ChilNodesSubset = CurContext.GetIntersectingChildren(FBoxCenterAndExtent(SearchPosition, Extend));
-			FOREACH_OCTREE_CHILD_NODE(OctreeChildRef)
-			{
-				if (ChilNodesSubset.Contains(OctreeChildRef) && CurNode.HasChild(OctreeChildRef))
-				{
-					OctreeIter.PushChild(OctreeChildRef);
-				}
-			}
 			// Add all of the elements in the current node to the list of points to consider for closest point calculations
-			for (const FWedgeInfo& WedgeInfo : CurNode.GetElements())
-			{
-				float VectorDelta = FVector::DistSquared(SearchPosition, WedgeInfo.Position);
-				MinSquaredDistance = FMath::Min(VectorDelta, MinSquaredDistance);
-				OutNearestWedges.Add(WedgeInfo);
-			}
-			OctreeIter.Advance();
-		}
+			float VectorDelta = FVector::DistSquared(SearchPosition, WedgeInfo.Position);
+			MinSquaredDistance = FMath::Min(VectorDelta, MinSquaredDistance);
+			OutNearestWedges.Add(WedgeInfo);
+		});
 
 		if (i == 0)
 		{
diff --git a/Engine/Source/Runtime/Engine/Private/StaticMeshBuild.cpp b/Engine/Source/Runtime/Engine/Private/StaticMeshBuild.cpp
--- a/Engine/Source/Runtime/Engine/Private/StaticMeshBuild.cpp
+++ b/Engine/Source/Runtime/Engine/Private/StaticMeshBuild.cpp
@@ -562,7 +562,6 @@ void RemapPaintedVertexColors(const TArray<FPaintedVertex>& InPaintedVertices,
 	for ( uint32 NewVertIndex = 0; NewVertIndex < NewPositions.GetNumVertices(); ++NewVertIndex )
 	{
 		PointsToConsider.Reset();
-		TSMCVertPosOctree::TConstIterator<> OctreeIter( VertPosOctree );
 		const FVector& CurPosition = NewPositions.VertexPosition( NewVertIndex );
 		FVector CurNormal = FVector::ZeroVector;
 		if (OptionalVertexBuffer)
@@ -571,42 +570,10 @@ void RemapPaintedVertexColors(const TArray<FPaintedVertex>& InPaintedVertices,
 		}
 
 		// Iterate through the octree attempting to find the vertices closest to the current new point
-		while ( OctreeIter.HasPendingNodes() )
+		VertPosOctree.IterateElementsWithBoundsTest(FBoxCenterAndExtent(CurPosition, FVector::ZeroVector), [&PointsToConsider](const FPaintedVertex& Vertex)
 		{
-			const TSMCVertPosOctree::FNode& CurNode = OctreeIter.GetCurrentNode();
-			const FOctreeNodeContext& CurContext = OctreeIter.GetCurrentContext();
-
-			// Find the child of the current node, if any, that contains the current new point
-			FOctreeChildNodeRef ChildRef = CurContext.GetContainingChild( FBoxCenterAndExtent( CurPosition, FVector::ZeroVector ) );
-
-			if ( !ChildRef.IsNULL() )
-			{
-				const TSMCVertPosOctree::FNode* ChildNode = CurNode.GetChild( ChildRef );
-
-				// If the specified child node exists and contains any of the old vertices, push it to the iterator for future consideration
-				if ( ChildNode && ChildNode->GetInclusiveElementCount() > 0 )
-				{
-					OctreeIter.PushChild( ChildRef );
-				}
-				// If the child node doesn't have any of the old vertices in it, it's not worth pursuing any further. In an attempt to find
-				// anything to match vs. the new point, add all of the children of the current octree node that have old points in them to the
-				// iterator for future consideration.
-				else
-				{
-					FOREACH_OCTREE_CHILD_NODE( OctreeChildRef )
-					{
-						if( CurNode.HasChild( OctreeChildRef ) )
-						{
-							OctreeIter.PushChild( OctreeChildRef );
-						}
-					}
-				}
-			}
-
-			// Add all of the elements in the current node to the list of points to consider for closest point calculations
-			PointsToConsider.Append( CurNode.GetElements() );
-			OctreeIter.Advance();
-		}
+			PointsToConsider.Add(Vertex);
+		});
 
 		// If any points to consider were found, iterate over each and find which one is the closest to the new point 
 		if ( PointsToConsider.Num() > 0 )
diff --git a/Engine/Source/Runtime/Experimental/GeometryCollectionCore/Private/GeometryCollection/GeometryCollectionProximityUtility.cpp b/Engine/Source/Runtime/Experimental/GeometryCollectionCore/Private/GeometryCollection/GeometryCollectionProximityUtility.cpp
--- a/Engine/Source/Runtime/Experimental/GeometryCollectionCore/Private/GeometryCollection/GeometryCollectionProximityUtility.cpp
+++ b/Engine/Source/Runtime/Experimental/GeometryCollectionCore/Private/GeometryCollection/GeometryCollectionProximityUtility.cpp
@@ -273,36 +273,11 @@ void FGeometryCollectionProximityUtility::UpdateProximity(FGeometryCollection* G
 
 		// 	Query the Octree
 		OtherFaceTransformDataArray.Reserve(NumFaces);
-		for (FProximityTriangleOctree::TConstIterator<> OctreeIt(*MeshTriOctree); OctreeIt.HasPendingNodes(); OctreeIt.Advance())
-		{
-			const FProximityTriangleOctree::FNode& OctreeNode = OctreeIt.GetCurrentNode();
-			const FOctreeNodeContext& OctreeNodeContext = OctreeIt.GetCurrentContext();
-
-			// Leaf nodes have no children, so don't bother iterating
-			if (!OctreeNode.IsLeaf())
-			{
-				FOREACH_OCTREE_CHILD_NODE(ChildRef)
-				{
-					if (OctreeNode.HasChild(ChildRef))
-					{
-						const FOctreeNodeContext ChildContext = OctreeNodeContext.GetChildContext(ChildRef);
 
-						if (ThisFaceBounds.Intersect(ChildContext.Bounds.GetBox()))
-						{
-							// Push it on the iterator's pending node stack.
-							OctreeIt.PushChild(ChildRef);
-						}
-					}
-				}
-			}
-
-			// All of the elements in this octree node are candidates.  Note this node may not be a leaf node, and that's OK.
-			for (FProximityTriangleOctree::ElementConstIt OctreeElementIt(OctreeNode.GetElementIt()); OctreeElementIt; ++OctreeElementIt)
-			{
-				const FProximityTriangle& OctreePolygon = *OctreeElementIt;
-				OtherFaceTransformDataArray.Add(FaceTransformDataArray[OctreePolygon.ArrayIndex]);
-			}
-		}
+		MeshTriOctree->IterateElementsWithBoundsTest(ThisFaceBounds, [&OtherFaceTransformDataArray, &FaceTransformDataArray](const FProximityTriangle& OctreePolygon)
+		{
+			OtherFaceTransformDataArray.Add(FaceTransformDataArray[OctreePolygon.ArrayIndex]);
+		});
 
 		for (auto& OtherFaceTransformDataRef : OtherFaceTransformDataArray)
 		{
diff --git a/Engine/Source/Runtime/NavigationSystem/Private/NavMesh/NavMeshRenderingComponent.cpp b/Engine/Source/Runtime/NavigationSystem/Private/NavMesh/NavMeshRenderingComponent.cpp
--- a/Engine/Source/Runtime/NavigationSystem/Private/NavMesh/NavMeshRenderingComponent.cpp
+++ b/Engine/Source/Runtime/NavigationSystem/Private/NavMesh/NavMeshRenderingComponent.cpp
@@ -751,59 +751,49 @@ void FNavMeshSceneProxyData::GatherData(const ARecastNavMesh* NavMesh, int32 InN
 				TArray<FVector> CollidingVerts;
 				TArray<uint32> CollidingIndices;
 
-				for (FNavigationOctree::TConstIterator<> NodeIt(*NavOctree); NodeIt.HasPendingNodes(); NodeIt.Advance())
+				FBox CurrentNodeBoundsBox;
+				NavOctree->IterateElementsWithPredicate([bGatherOctree, bGatherPathCollidingGeometry, this, &CurrentNodeBoundsBox](const FBoxCenterAndExtent& NodeBounds)
 				{
-					const FNavigationOctree::FNode& CurrentNode = NodeIt.GetCurrentNode();
-					const FOctreeNodeContext& CorrentContext = NodeIt.GetCurrentContext();					
-
 					if (bGatherOctree)
 					{
-						OctreeBounds.Add(CorrentContext.Bounds);
+						OctreeBounds.Add(NodeBounds);
+					}
 
-						if (bGatherOctreeDetails)
-						{
-							for (FNavigationOctree::ElementConstIt ElementIt(CurrentNode.GetElementIt()); ElementIt; ElementIt++)
-							{
-								const FNavigationOctreeElement& Element = *ElementIt;
-								OctreeBounds.Add(FBoxCenterAndExtent(Element.Bounds));
-							}
-						}
+					if (bGatherPathCollidingGeometry)
+					{
+						CurrentNodeBoundsBox = NodeBounds.GetBox();
+					}
+					return true;
+				},
+				[bGatherOctree, bGatherOctreeDetails, bGatherPathCollidingGeometry, this, NavMesh, &CollidingVerts, &CollidingIndices, &CurrentNodeBoundsBox](const FNavigationOctreeElement& Element)
+				{
+					if (bGatherOctree && bGatherOctreeDetails)
+					{
+						OctreeBounds.Add(FBoxCenterAndExtent(Element.Bounds));
 					}
 
 					if (bGatherPathCollidingGeometry)
 					{
-						for (FNavigationOctree::ElementConstIt ElementIt(CurrentNode.GetElementIt()); ElementIt; ElementIt++)
+						if (Element.ShouldUseGeometry(NavMesh->GetConfig()) && Element.Data->CollisionData.Num())
 						{
-							const FNavigationOctreeElement& Element = *ElementIt;
-							if (Element.ShouldUseGeometry(NavMesh->GetConfig()) && Element.Data->CollisionData.Num())
+							const FRecastGeometryCache CachedGeometry(Element.Data->CollisionData.GetData());
+							TArray<FTransform> InstanceTransforms;
+							Element.Data->NavDataPerInstanceTransformDelegate.ExecuteIfBound(CurrentNodeBoundsBox, InstanceTransforms);
+
+							if (InstanceTransforms.Num() == 0)
 							{
-								const FRecastGeometryCache CachedGeometry(Element.Data->CollisionData.GetData());
-								TArray<FTransform> InstanceTransforms;
-								Element.Data->NavDataPerInstanceTransformDelegate.ExecuteIfBound(CorrentContext.Bounds.GetBox(), InstanceTransforms);
-								
-								if (InstanceTransforms.Num() == 0)
-								{
-									FNavMeshRenderingHelpers::AddRecastGeometry(CollidingVerts, CollidingIndices, CachedGeometry.Verts, CachedGeometry.Header.NumVerts, CachedGeometry.Indices, CachedGeometry.Header.NumFaces);
-								}
-								else
+								FNavMeshRenderingHelpers::AddRecastGeometry(CollidingVerts, CollidingIndices, CachedGeometry.Verts, CachedGeometry.Header.NumVerts, CachedGeometry.Indices, CachedGeometry.Header.NumFaces);
+							}
+							else
+							{
+								for (const auto& Transform : InstanceTransforms)
 								{
-									for (const auto& Transform : InstanceTransforms)
-									{
-										FNavMeshRenderingHelpers::AddRecastGeometry(CollidingVerts, CollidingIndices, CachedGeometry.Verts, CachedGeometry.Header.NumVerts, CachedGeometry.Indices, CachedGeometry.Header.NumFaces, Transform);
-									}
+									FNavMeshRenderingHelpers::AddRecastGeometry(CollidingVerts, CollidingIndices, CachedGeometry.Verts, CachedGeometry.Header.NumVerts, CachedGeometry.Indices, CachedGeometry.Header.NumFaces, Transform);
 								}
 							}
 						}
 					}
-					
-					FOREACH_OCTREE_CHILD_NODE(ChildRef)
-					{
-						if (CurrentNode.HasChild(ChildRef))
-						{
-							NodeIt.PushChild(ChildRef);
-						}
-					}
-				}
+				});
 			
 				if (CollidingVerts.Num())
 				{
@@ -1575,11 +1565,8 @@ FBoxSphereBounds UNavMeshRenderingComponent::CalcBounds(const FTransform& LocalT
 			const FNavigationOctree* NavOctree = NavSys ? NavSys->GetNavOctree() : nullptr;
 			if (NavOctree)
 			{
-				for (FNavigationOctree::TConstIterator<> NodeIt(*NavOctree); NodeIt.HasPendingNodes(); NodeIt.Advance())
-				{
-					const FOctreeNodeContext& CorrentContext = NodeIt.GetCurrentContext();
-					BoundingBox += CorrentContext.Bounds.GetBox();
-				}
+				//this is only iterating over the rootnode
+				BoundingBox += NavOctree->GetRootBounds().GetBox();
 			}
 		}
 	}
diff --git a/Engine/Source/Runtime/NavigationSystem/Private/NavMesh/RecastNavMeshGenerator.cpp b/Engine/Source/Runtime/NavigationSystem/Private/NavMesh/RecastNavMeshGenerator.cpp
--- a/Engine/Source/Runtime/NavigationSystem/Private/NavMesh/RecastNavMeshGenerator.cpp
+++ b/Engine/Source/Runtime/NavigationSystem/Private/NavMesh/RecastNavMeshGenerator.cpp
@@ -2225,17 +2225,14 @@ void FRecastTileGenerator::GatherGeometry(const FRecastNavMeshGenerator& ParentG
 	}
 	const FNavDataConfig& OwnerNavDataConfig = ParentGenerator.GetOwner()->GetConfig();
 
-	for (FNavigationOctree::TConstElementBoxIterator<FNavigationOctree::DefaultStackAllocator> It(*NavigationOctree, ParentGenerator.GrowBoundingBox(TileBB, /*bIncludeAgentHeight*/ false));
-		It.HasPendingElements();
-		It.Advance())
+	NavigationOctree->IterateElementsWithBoundsTest(ParentGenerator.GrowBoundingBox(TileBB, /*bIncludeAgentHeight*/ false), [&OwnerNavDataConfig, &NavigationOctree, this, bGeometryChanged](const FNavigationOctreeElement& Element)
 	{
-		const FNavigationOctreeElement& Element = It.GetCurrentElement();
 		const bool bShouldUse = Element.ShouldUseGeometry(OwnerNavDataConfig);
 		if (bShouldUse)
 		{
 			GatherNavigationDataGeometry(Element.Data, *NavSys, OwnerNavDataConfig, bGeometryChanged);
 		}
-	}
+	});
 }
 
 void FRecastTileGenerator::GatherNavigationDataGeometry(const TSharedRef<FNavigationRelevantData, ESPMode::ThreadSafe>& ElementData, UNavigationSystemV1& NavSys, const FNavDataConfig& OwnerNavDataConfig, const bool bGeometryChanged)
diff --git a/Engine/Source/Runtime/NavigationSystem/Private/NavigationDataHandler.cpp b/Engine/Source/Runtime/NavigationSystem/Private/NavigationDataHandler.cpp
--- a/Engine/Source/Runtime/NavigationSystem/Private/NavigationDataHandler.cpp
+++ b/Engine/Source/Runtime/NavigationSystem/Private/NavigationDataHandler.cpp
@@ -341,7 +341,7 @@ bool FNavigationDataHandler::UpdateNavOctreeElementBounds(UActorComponent& Comp,
 			ElementId = OctreeController.GetObjectsNavOctreeId(Comp);
 			if (ElementId != nullptr && ensure(OctreeController.IsValidElement(*ElementId)))
 			{
-				FNavigationOctreeElement& ElementData = OctreeController.NavOctree->GetElementById(*ElementId);
+				const FNavigationOctreeElement& ElementData = OctreeController.NavOctree->GetElementById(*ElementId);
 				DirtyAreasController.AddArea(DirtyArea, ElementData.Data->GetDirtyFlag(), [&Comp] { return &Comp; });
 			}
 		}
@@ -360,14 +360,13 @@ void FNavigationDataHandler::FindElementsInNavOctree(const FBox& QueryBox, const
 		return;
 	}
 
-	for (FNavigationOctree::TConstElementBoxIterator<> It(*OctreeController.NavOctree, QueryBox); It.HasPendingElements(); It.Advance())
+	OctreeController.NavOctree->IterateElementsWithBoundsTest(QueryBox, [&Elements, &Filter](const FNavigationOctreeElement& Element)
 	{
-		const FNavigationOctreeElement& Element = It.GetCurrentElement();
 		if (Element.IsMatchingFilter(Filter))
 		{
 			Elements.Add(Element);
 		}
-	}
+	});
 }
 
 bool FNavigationDataHandler::ReplaceAreaInOctreeData(const UObject& Object, TSubclassOf<UNavArea> OldArea, TSubclassOf<UNavArea> NewArea, bool bReplaceChildClasses)
@@ -453,7 +452,7 @@ void FNavigationDataHandler::RemoveLevelCollisionFromOctree(ULevel& Level)
 			if (ensure(OctreeController.IsValidElement(*ElementId)))
 			{
 				// mark area occupied by given actor as dirty
-				FNavigationOctreeElement& ElementData = OctreeController.NavOctree->GetElementById(*ElementId);
+				const FNavigationOctreeElement& ElementData = OctreeController.NavOctree->GetElementById(*ElementId);
 				DirtyAreasController.AddArea(ElementData.Bounds.GetBox(), ENavigationDirtyFlag::All, [&Level] { return &Level; });
 			}
 
diff --git a/Engine/Source/Runtime/NavigationSystem/Private/NavigationOctreeController.cpp b/Engine/Source/Runtime/NavigationSystem/Private/NavigationOctreeController.cpp
--- a/Engine/Source/Runtime/NavigationSystem/Private/NavigationOctreeController.cpp
+++ b/Engine/Source/Runtime/NavigationSystem/Private/NavigationOctreeController.cpp
@@ -30,7 +30,7 @@ bool FNavigationOctreeController::GetNavOctreeElementData(const UObject& NodeOwn
 	if (ElementId != nullptr && IsValidElement(*ElementId))
 	{
 		// mark area occupied by given actor as dirty
-		FNavigationOctreeElement& ElementData = NavOctree->GetElementById(*ElementId);
+		const FNavigationOctreeElement& ElementData = NavOctree->GetElementById(*ElementId);
 		DirtyFlags = ElementData.Data->GetDirtyFlag();
 		DirtyBounds = ElementData.Bounds.GetBox();
 		return true;
diff --git a/Engine/Source/Runtime/NavigationSystem/Public/NavigationOctree.h b/Engine/Source/Runtime/NavigationSystem/Public/NavigationOctree.h
--- a/Engine/Source/Runtime/NavigationSystem/Public/NavigationOctree.h
+++ b/Engine/Source/Runtime/NavigationSystem/Public/NavigationOctree.h
@@ -126,6 +126,34 @@ public:
 		StoreNavGeometry,
 	};
 
+	/**
+ * Adds an element to the octree.
+ * @param Element - The element to add.
+ */
+	inline void AddElement(const FNavigationOctreeElement& Element)
+	{
+		DEC_MEMORY_STAT_BY(STAT_NavigationMemory, OctreeSizeBytes);
+		DEC_MEMORY_STAT_BY(STAT_Navigation_CollisionTreeMemory, OctreeSizeBytes);
+		TOctree<FNavigationOctreeElement, FNavigationOctreeSemantics>::AddElement(Element);
+		OctreeSizeBytes = GetSizeBytes();
+		INC_MEMORY_STAT_BY(STAT_NavigationMemory, OctreeSizeBytes);
+		INC_MEMORY_STAT_BY(STAT_Navigation_CollisionTreeMemory, OctreeSizeBytes);
+	}
+
+	/**
+	 * Removes an element from the octree.
+	 * @param ElementId - The element to remove from the octree.
+	 */
+	inline void RemoveElement(FOctreeElementId ElementId)
+	{
+		DEC_MEMORY_STAT_BY(STAT_NavigationMemory, OctreeSizeBytes);
+		DEC_MEMORY_STAT_BY(STAT_Navigation_CollisionTreeMemory, OctreeSizeBytes);
+		TOctree<FNavigationOctreeElement, FNavigationOctreeSemantics>::RemoveElement(ElementId);
+		OctreeSizeBytes = GetSizeBytes();
+		INC_MEMORY_STAT_BY(STAT_NavigationMemory, OctreeSizeBytes);
+		INC_MEMORY_STAT_BY(STAT_Navigation_CollisionTreeMemory, OctreeSizeBytes);
+	}
+
 	FNavigationOctree(const FVector& Origin, float Radius);
 	virtual ~FNavigationOctree();
 
@@ -172,18 +200,6 @@ protected:
 	ENavDataGatheringMode DefaultGeometryGatheringMode;
 	uint32 bGatherGeometry : 1;
 	uint32 NodesMemory;
+private:
+	SIZE_T OctreeSizeBytes = 0;
 };
-
-template<>
-FORCEINLINE_DEBUGGABLE void SetOctreeMemoryUsage(TOctree<FNavigationOctreeElement, FNavigationOctreeSemantics>* Octree, int32 NewSize)
-{
-	{
-		DEC_DWORD_STAT_BY( STAT_NavigationMemory, Octree->TotalSizeBytes );
-		DEC_DWORD_STAT_BY(STAT_Navigation_CollisionTreeMemory, Octree->TotalSizeBytes);
-	}
-	Octree->TotalSizeBytes = NewSize;
-	{
-		INC_DWORD_STAT_BY( STAT_NavigationMemory, NewSize );
-		INC_DWORD_STAT_BY(STAT_Navigation_CollisionTreeMemory, NewSize);
-	}
-}
diff --git a/Engine/Source/Runtime/Renderer/Private/LightSceneInfo.cpp b/Engine/Source/Runtime/Renderer/Private/LightSceneInfo.cpp
--- a/Engine/Source/Runtime/Renderer/Private/LightSceneInfo.cpp
+++ b/Engine/Source/Runtime/Renderer/Private/LightSceneInfo.cpp
@@ -116,10 +116,11 @@ void FLightSceneInfo::AddToScene()
 			Scene->LocalShadowCastingLightOctree.AddElement(LightSceneInfoCompact);
 			// Find primitives that the light affects in the primitive octree.
 			FMemMark MemStackMark(FMemStack::Get());
-			for (FScenePrimitiveOctree::TConstElementBoxIterator<SceneRenderingAllocator> PrimitiveIt(Scene->PrimitiveOctree, GetBoundingBox()); PrimitiveIt.HasPendingElements(); PrimitiveIt.Advance())
+
+			Scene->PrimitiveOctree.IterateElementsWithBoundsTest(GetBoundingBox(), [&LightSceneInfoCompact, this](const FPrimitiveSceneInfoCompact& PrimitiveSceneInfoCompact)
 			{
-				CreateLightPrimitiveInteraction(LightSceneInfoCompact, PrimitiveIt.GetCurrentElement());
-			}
+				CreateLightPrimitiveInteraction(LightSceneInfoCompact, PrimitiveSceneInfoCompact);
+			});
 		}
 	}
 }
diff --git a/Engine/Source/Runtime/Renderer/Private/RendererScene.cpp b/Engine/Source/Runtime/Renderer/Private/RendererScene.cpp
--- a/Engine/Source/Runtime/Renderer/Private/RendererScene.cpp
+++ b/Engine/Source/Runtime/Renderer/Private/RendererScene.cpp
@@ -1099,13 +1099,11 @@ public:
 			const FPrimitiveSceneInfoCompact PrimitiveSceneInfoCompact(PrimInfo);
 
 			// Find local lights that affect the primitive in the light octree.
-			for (FSceneLightOctree::TConstElementBoxIterator<SceneRenderingAllocator> LightIt(Scene->LocalShadowCastingLightOctree, Bounds.GetBox());
-				LightIt.HasPendingElements();
-				LightIt.Advance())
+			Scene->LocalShadowCastingLightOctree.IterateElementsWithBoundsTest(Bounds.GetBox(), [&PrimitiveSceneInfoCompact](const FLightSceneInfoCompact& LightSceneInfoCompact)
 			{
-				const FLightSceneInfoCompact& LightSceneInfoCompact = LightIt.GetCurrentElement();
 				LightSceneInfoCompact.LightSceneInfo->CreateLightPrimitiveInteraction(LightSceneInfoCompact, PrimitiveSceneInfoCompact);
-			}
+			});
+
 			// Also loop through non-local (directional) shadow-casting lights
 			for (int32 LightID : Scene->DirectionalShadowCastingLightIDs)
 			{
@@ -4224,13 +4222,11 @@ void FScene::CreateLightPrimitiveInteractionsForPrimitive(FPrimitiveSceneInfo* P
 		const FPrimitiveSceneInfoCompact PrimitiveSceneInfoCompact(PrimitiveInfo);
 
 		// Find local lights that affect the primitive in the light octree.
-		for (FSceneLightOctree::TConstElementBoxIterator<SceneRenderingAllocator> LightIt(LocalShadowCastingLightOctree, Bounds.GetBox());
-			LightIt.HasPendingElements();
-			LightIt.Advance())
+		LocalShadowCastingLightOctree.IterateElementsWithBoundsTest(Bounds.GetBox(), [&PrimitiveSceneInfoCompact](const FLightSceneInfoCompact& LightSceneInfoCompact)
 		{
-			const FLightSceneInfoCompact& LightSceneInfoCompact = LightIt.GetCurrentElement();
 			LightSceneInfoCompact.LightSceneInfo->CreateLightPrimitiveInteraction(LightSceneInfoCompact, PrimitiveSceneInfoCompact);
-		}
+		});
+	
 		// Also loop through non-local (directional) shadow-casting lights
 		for (int32 LightID : DirectionalShadowCastingLightIDs)
 		{
diff --git a/Engine/Source/Runtime/Renderer/Private/ShadowSetup.cpp b/Engine/Source/Runtime/Renderer/Private/ShadowSetup.cpp
--- a/Engine/Source/Runtime/Renderer/Private/ShadowSetup.cpp
+++ b/Engine/Source/Runtime/Renderer/Private/ShadowSetup.cpp
@@ -3697,7 +3697,7 @@ struct FGatherShadowPrimitivesPacket
 	// Inputs
 	const FScene* Scene;
 	TArray<FViewInfo>& Views;
-	const FScenePrimitiveOctree::FNode* Node;
+	const FScenePrimitiveOctree::FNodeIndex NodeIndex;
 	int32 StartPrimitiveIndex;
 	int32 NumPrimitives;
 	const TArray<FProjectedShadowInfo*, SceneRenderingAllocator>& PreShadows;
@@ -3718,7 +3718,7 @@ struct FGatherShadowPrimitivesPacket
 	FGatherShadowPrimitivesPacket(
 		const FScene* InScene,
 		TArray<FViewInfo>& InViews,
-		const FScenePrimitiveOctree::FNode* InNode,
+		const FScenePrimitiveOctree::FNodeIndex InNodeIndex,
 		int32 InStartPrimitiveIndex,
 		int32 InNumPrimitives,
 		const TArray<FProjectedShadowInfo*, SceneRenderingAllocator>& InPreShadows,
@@ -3728,7 +3728,7 @@ struct FGatherShadowPrimitivesPacket
 		FPerShadowGatherStats& OutGlobalStats) 
 		: Scene(InScene)
 		, Views(InViews)
-		, Node(InNode)
+		, NodeIndex(InNodeIndex)
 		, StartPrimitiveIndex(InStartPrimitiveIndex)
 		, NumPrimitives(InNumPrimitives)
 		, PreShadows(InPreShadows)
@@ -3761,14 +3761,14 @@ struct FGatherShadowPrimitivesPacket
 	{
 		QUICK_SCOPE_CYCLE_COUNTER(STAT_GatherShadowPrimitivesPacket);
 
-		if (Node)
+		if (NodeIndex != INDEX_NONE)
 		{
 			// Check all the primitives in this octree node.
-			for (FScenePrimitiveOctree::ElementConstIt NodePrimitiveIt(Node->GetElementIt()); NodePrimitiveIt; ++NodePrimitiveIt)
+			for (const FPrimitiveSceneInfoCompact& PrimitiveSceneInfoCompact : Scene->PrimitiveOctree.GetElementsForNode(NodeIndex))
 			{
-				if (NodePrimitiveIt->PrimitiveFlagsCompact.bCastDynamicShadow)
+				if (PrimitiveSceneInfoCompact.PrimitiveFlagsCompact.bCastDynamicShadow)
 				{
-					FilterPrimitiveForShadows(*NodePrimitiveIt);
+					FilterPrimitiveForShadows(PrimitiveSceneInfoCompact);
 				}
 			}
 		}
@@ -3979,78 +3979,49 @@ void FSceneRenderer::GatherShadowPrimitives(
 			Packets.Reserve(100);
 
 			// Find primitives that are in a shadow frustum in the octree.
-			for(FScenePrimitiveOctree::TConstIterator<SceneRenderingAllocator> PrimitiveOctreeIt(Scene->PrimitiveOctree);
-				PrimitiveOctreeIt.HasPendingNodes();
-				PrimitiveOctreeIt.Advance())
+			Scene->PrimitiveOctree.IterateNodesWithPredicate([&PreShadows, &ViewDependentWholeSceneShadows](const FBoxCenterAndExtent& NodeBounds)
 			{
-				const FScenePrimitiveOctree::FNode& PrimitiveOctreeNode = PrimitiveOctreeIt.GetCurrentNode();
-				const FOctreeNodeContext& PrimitiveOctreeNodeContext = PrimitiveOctreeIt.GetCurrentContext();
+				// Check that the child node is in the frustum for at least one shadow.
 
+				// Check for subjects of preshadows.
+				for (int32 ShadowIndex = 0, Num = PreShadows.Num(); ShadowIndex < Num; ShadowIndex++)
 				{
-					// Find children of this octree node that may contain relevant primitives.
-					FOREACH_OCTREE_CHILD_NODE(ChildRef)
-					{
-						if(PrimitiveOctreeNode.HasChild(ChildRef))
-						{
-							// Check that the child node is in the frustum for at least one shadow.
-							const FOctreeNodeContext ChildContext = PrimitiveOctreeNodeContext.GetChildContext(ChildRef);
-							bool bIsInFrustum = false;
-
-							// Check for subjects of preshadows.
-							if(!bIsInFrustum)
-							{
-								for(int32 ShadowIndex = 0, Num = PreShadows.Num(); ShadowIndex < Num; ShadowIndex++)
-								{
-									FProjectedShadowInfo* ProjectedShadowInfo = PreShadows[ShadowIndex];
-
-									check(ProjectedShadowInfo->CasterFrustum.PermutedPlanes.Num());
-									// Check if this primitive is in the shadow's frustum.
-									if(ProjectedShadowInfo->CasterFrustum.IntersectBox(
-										ChildContext.Bounds.Center + ProjectedShadowInfo->PreShadowTranslation,
-										ChildContext.Bounds.Extent
-										))
-									{
-										bIsInFrustum = true;
-										break;
-									}
-								}
-							}
-
-							if (!bIsInFrustum)
-							{
-								for(int32 ShadowIndex = 0, Num = ViewDependentWholeSceneShadows.Num(); ShadowIndex < Num; ShadowIndex++)
-								{
-									FProjectedShadowInfo* ProjectedShadowInfo = ViewDependentWholeSceneShadows[ShadowIndex];
-
-									//check(ProjectedShadowInfo->CasterFrustum.PermutedPlanes.Num());
-									// Check if this primitive is in the shadow's frustum.
-									if(ProjectedShadowInfo->CasterFrustum.IntersectBox(
-										ChildContext.Bounds.Center + ProjectedShadowInfo->PreShadowTranslation,
-										ChildContext.Bounds.Extent
-										))
-									{
-										bIsInFrustum = true;
-										break;
-									}
-								}
-							}
+					FProjectedShadowInfo* ProjectedShadowInfo = PreShadows[ShadowIndex];
+
+					check(ProjectedShadowInfo->CasterFrustum.PermutedPlanes.Num());
+					// Check if this primitive is in the shadow's frustum.
+					if (ProjectedShadowInfo->CasterFrustum.IntersectBox(
+						NodeBounds.Center + ProjectedShadowInfo->PreShadowTranslation,
+						NodeBounds.Extent
+					))
+					return true;
+				}
 
-							if(bIsInFrustum)
-							{
-								// If the child node was in the frustum of at least one preshadow, push it on
-								// the iterator's pending node stack.
-								PrimitiveOctreeIt.PushChild(ChildRef);
-							}
-						}
-					}
+				for (int32 ShadowIndex = 0, Num = ViewDependentWholeSceneShadows.Num(); ShadowIndex < Num; ShadowIndex++)
+				{
+					FProjectedShadowInfo* ProjectedShadowInfo = ViewDependentWholeSceneShadows[ShadowIndex];
+
+					//check(ProjectedShadowInfo->CasterFrustum.PermutedPlanes.Num());
+					// Check if this primitive is in the shadow's frustum.
+					if (ProjectedShadowInfo->CasterFrustum.IntersectBox(
+						NodeBounds.Center + ProjectedShadowInfo->PreShadowTranslation,
+						NodeBounds.Extent
+					))
+					return true;
 				}
 
-				if (PrimitiveOctreeNode.GetElementCount() > 0)
+				// If the child node was in the frustum of at least one preshadow, push it on
+				// the iterator's pending node stack.
+				return false;
+			},
+			[this, &PreShadows, &ViewDependentWholeSceneShadows, bStaticSceneOnly, &GatherStats, &Packets](FScenePrimitiveOctree::FNodeIndex NodeIndex)
+			{
+				if (Scene->PrimitiveOctree.GetElementsForNode(NodeIndex).Num() > 0)
 				{
 					FGatherShadowPrimitivesPacket* Packet = new(FMemStack::Get()) FGatherShadowPrimitivesPacket(
 						Scene,
 						Views,
-						&PrimitiveOctreeNode,
+						NodeIndex,
 						0,
 						0,
 						PreShadows,
@@ -4060,7 +4031,7 @@ void FSceneRenderer::GatherShadowPrimitives(
 						GatherStats);
 					Packets.Add(Packet);
 				}
-			}
+			});
 		}
 		else
 		{
@@ -4076,7 +4047,7 @@ void FSceneRenderer::GatherShadowPrimitives(
 				FGatherShadowPrimitivesPacket* Packet = new(FMemStack::Get()) FGatherShadowPrimitivesPacket(
 					Scene,
 					Views,
-					NULL,
+					INDEX_NONE,
 					StartPrimitiveIndex,
 					NumPrimitives,
 					PreShadows,
diff --git a/Engine/Source/Runtime/Renderer/Private/ShadowSetupMobile.cpp b/Engine/Source/Runtime/Renderer/Private/ShadowSetupMobile.cpp
--- a/Engine/Source/Runtime/Renderer/Private/ShadowSetupMobile.cpp
+++ b/Engine/Source/Runtime/Renderer/Private/ShadowSetupMobile.cpp
@@ -169,38 +169,16 @@ static bool MobileDetermineStaticMeshesCSMVisibilityState(FScene* Scene, FViewIn
 	{
 		QUICK_SCOPE_CYCLE_COUNTER(STAT_ShadowOctreeTraversal);
 		// Find primitives that are in a shadow frustum in the octree.
-		for (FScenePrimitiveOctree::TConstIterator<SceneRenderingAllocator> PrimitiveOctreeIt(Scene->PrimitiveOctree);
-		PrimitiveOctreeIt.HasPendingNodes();
-			PrimitiveOctreeIt.Advance())
+		Scene->PrimitiveOctree.IterateElementsWithPredicate([&IsReceiverFunc](const FBoxCenterAndExtent& NodeBounds)
 		{
-			const FScenePrimitiveOctree::FNode& PrimitiveOctreeNode = PrimitiveOctreeIt.GetCurrentNode();
-			const FOctreeNodeContext& PrimitiveOctreeNodeContext = PrimitiveOctreeIt.GetCurrentContext();
-
-			// Find children of this octree node that may contain relevant primitives.
-			FOREACH_OCTREE_CHILD_NODE(ChildRef)
-			{
-				if (PrimitiveOctreeNode.HasChild(ChildRef))
-				{
-					// Check that the child node is in the frustum for at least one shadow.
-					const FOctreeNodeContext ChildContext = PrimitiveOctreeNodeContext.GetChildContext(ChildRef);
-					bool bCanReceiveDynamicShadow = IsReceiverFunc(FVector(ChildContext.Bounds.Center), FVector(ChildContext.Bounds.Extent), ChildContext.Bounds.Extent.Size3());
-
-					if (bCanReceiveDynamicShadow)
-					{
-						PrimitiveOctreeIt.PushChild(ChildRef);
-					}
-				}
-			}
-
-			// Check all the primitives in this octree node.
-			for (FScenePrimitiveOctree::ElementConstIt NodePrimitiveIt(PrimitiveOctreeNode.GetElementIt()); NodePrimitiveIt; ++NodePrimitiveIt)
-			{
-				// gather the shadows for this one primitive
-				bFoundReceiver = MobileDetermineStaticMeshesCSMVisibilityStateInner(Scene, View, *NodePrimitiveIt, WholeSceneShadow, IsReceiverFunc) || bFoundReceiver;
-			}
-		}
+			return IsReceiverFunc(FVector(NodeBounds.Center), FVector(NodeBounds.Extent), NodeBounds.Extent.Size3());
+		},
+		[&bFoundReceiver, Scene, &View, WholeSceneShadow, &IsReceiverFunc](const FPrimitiveSceneInfoCompact& Primitive)
+		{
+			// gather the shadows for this one primitive
+			bFoundReceiver = MobileDetermineStaticMeshesCSMVisibilityStateInner(Scene, View, Primitive, WholeSceneShadow, IsReceiverFunc) || bFoundReceiver;
+		});
 	}
-
 	return bFoundReceiver;
 }
 
