FNEM - fixing error: taking the address of a temporary object of type 'VectorRegister' [-Waddress-of-temporary]

diff --git a/Engine/Source/Runtime/Core/Public/Math/UnrealMathFPU.h b/Engine/Source/Runtime/Core/Public/Math/UnrealMathFPU.h
--- a/Engine/Source/Runtime/Core/Public/Math/UnrealMathFPU.h
+++ b/Engine/Source/Runtime/Core/Public/Math/UnrealMathFPU.h
@@ -204,7 +204,10 @@ FORCEINLINE VectorRegisterInt MakeVectorRegisterInt(int32 X, int32 Y, int32 Z, i
  * @param Vec	Vector to store
  * @param Ptr	Aligned memory pointer
  */
-#define VectorStoreAligned( Vec, Ptr )	FMemory::Memcpy( Ptr, &(Vec), 16 )
+FORCEINLINE void VectorStoreAligned(const VectorRegister& Vec, void* Ptr)
+{
+	FMemory::Memcpy(Ptr, &(Vec), 16);
+}
 
 /**
  * Performs non-temporal store of a vector to aligned memory without polluting the caches
