FNEM - remove unnecessary brace

diff --git a/Engine/Source/Runtime/Core/Public/Math/UnrealMathFPU.h b/Engine/Source/Runtime/Core/Public/Math/UnrealMathFPU.h
--- a/Engine/Source/Runtime/Core/Public/Math/UnrealMathFPU.h
+++ b/Engine/Source/Runtime/Core/Public/Math/UnrealMathFPU.h
@@ -206,7 +206,7 @@ FORCEINLINE VectorRegisterInt MakeVectorRegisterInt(int32 X, int32 Y, int32 Z, i
  */
 FORCEINLINE void VectorStoreAligned(const VectorRegister& Vec, void* Ptr)
 {
-	FMemory::Memcpy(Ptr, &(Vec), 16);
+	FMemory::Memcpy(Ptr, &Vec, 16);
 }
 
 /**
