FNEM - Meshpaint was looking for nearby nodes and not intersecting nodes as such.

diff --git a/Engine/Plugins/MeshPainting/Source/MeshPaintingToolset/Private/MeshPaintHelpers.cpp b/Engine/Plugins/MeshPainting/Source/MeshPaintingToolset/Private/MeshPaintHelpers.cpp
--- a/Engine/Plugins/MeshPainting/Source/MeshPaintingToolset/Private/MeshPaintHelpers.cpp
+++ b/Engine/Plugins/MeshPainting/Source/MeshPaintingToolset/Private/MeshPaintHelpers.cpp
@@ -1330,7 +1330,7 @@ void UMeshPaintingToolset::ApplyVertexColorsToAllLODs(IMeshPaintComponentAdapter
 						FVector CurNormal = VertexTangentZ.ToFVector();
 
 						// Iterate through the octree attempting to find the vertices closest to the current new point
-						VertPosOctree.IterateElementsWithBoundsTest(FBoxCenterAndExtent(CurPosition, FVector::ZeroVector), [&PointsToConsider](const FPaintedMeshVertex& Vertex)
+						VertPosOctree.FindNearbyElements(FBoxCenterAndExtent(CurPosition, FVector::ZeroVector), [&PointsToConsider](const FPaintedMeshVertex& Vertex)
 						{
 							// Add all of the elements in the current node to the list of points to consider for closest point calculations
 							PointsToConsider.Add(Vertex);
diff --git a/Engine/Source/Editor/MeshPaint/Private/MeshPaintHelpers.cpp b/Engine/Source/Editor/MeshPaint/Private/MeshPaintHelpers.cpp
--- a/Engine/Source/Editor/MeshPaint/Private/MeshPaintHelpers.cpp
+++ b/Engine/Source/Editor/MeshPaint/Private/MeshPaintHelpers.cpp
@@ -1843,7 +1843,7 @@ void MeshPaintHelpers::ApplyVertexColorsToAllLODs(IMeshPaintGeometryAdapter& Geo
 						FVector CurNormal = VertexTangentZ.ToFVector();
 
 						// Iterate through the octree attempting to find the vertices closest to the current new point
-						VertPosOctree.IterateElementsWithBoundsTest(FBoxCenterAndExtent(CurPosition, FVector::ZeroVector), [&PointsToConsider](const FPaintedMeshVertex& Vertex)
+						VertPosOctree.FindNearbyElements(FBoxCenterAndExtent(CurPosition, FVector::ZeroVector), [&PointsToConsider](const FPaintedMeshVertex& Vertex)
 						{
 							PointsToConsider.Add(Vertex);
 						});
diff --git a/Engine/Source/Runtime/Core/Public/Math/GenericOctree.h b/Engine/Source/Runtime/Core/Public/Math/GenericOctree.h
--- a/Engine/Source/Runtime/Core/Public/Math/GenericOctree.h
+++ b/Engine/Source/Runtime/Core/Public/Math/GenericOctree.h
@@ -655,6 +655,50 @@ private:
 			}
 		}
 	}
+
+	template<typename IterateFunc>
+	void FindNearbyElementsInternal(FNodeIndex CurrentNodeIndex, const FOctreeNodeContext& NodeContext, const FBoxCenterAndExtent& BoxBounds, const IterateFunc& Func) const
+	{
+		if (TreeNodes[CurrentNodeIndex].InclusiveNumElements > 0)
+		{
+			for (int Index = 0; Index < TreeElements[CurrentNodeIndex].Num(); Index++)
+			{
+				typename TCallTraits<ElementType>::ConstReference Element = TreeElements[CurrentNodeIndex][Index];
+				if (Intersect(OctreeSemantics::GetBoundingBox(Element), BoxBounds))
+				{
+					Func(Element);
+				}
+			}
+
+			if (!TreeNodes[CurrentNodeIndex].IsLeaf())
+			{
+				const FOctreeChildNodeSubset IntersectingChildSubset = NodeContext.GetIntersectingChildren(BoxBounds);
+				FNodeIndex ChildStartIndex = TreeNodes[CurrentNodeIndex].ChildNodes;
+				
+				bool FoundMatchingNode = false;
+				for (int i = 0; i < 8; i++)
+				{
+					if (TreeNodes[ChildStartIndex + i].InclusiveNumElements > 0 && IntersectingChildSubset.Contains(FOctreeChildNodeRef(i)))
+					{
+						FoundMatchingNode = true;
+						FindNearbyElementsInternal(ChildStartIndex + i, NodeContext.GetChildContext(FOctreeChildNodeRef(i)), BoxBounds, Func);
+					}
+				}
+
+				//look in the near node vicinity for a match
+				if (!FoundMatchingNode)
+				{
+					for (int i = 0; i < 8; i++)
+					{
+						if (!IntersectingChildSubset.Contains(FOctreeChildNodeRef(i)))
+						{
+							FindNearbyElementsInternal(ChildStartIndex + i, NodeContext.GetChildContext(FOctreeChildNodeRef(i)), BoxBounds, Func);
+						}
+					}
+				}
+			}
+		}
+	}
 public:
 
 	/**
@@ -724,6 +768,18 @@ public:
 		FindFirstElementWithBoundsTestInternal(0, RootNodeContext, BoxBounds, Func, ContinueTraversal);
 	}
 
+	/**
+	* this function will traverse the Octree using a fast (and relaxed) box-box intersection but will also try and search in the neighborhood nodes if the intersected node is empty.
+	* @param BoxBounds - the bounds to test if a node is traversed or skipped.
+	* @param Func - Function to call with each Element for nodes that passed bounds test.
+	*/
+	template<typename IterateBoundsFunc>
+	inline void FindNearbyElements(const FBoxCenterAndExtent& BoxBounds, const IterateBoundsFunc& Func) const
+	{
+		FindNearbyElementsInternal(0, RootNodeContext, BoxBounds, Func);
+	}
+
+
 	/**
 	 * Adds an element to the octree.
 	 * @param Element - The element to add.
diff --git a/Engine/Source/Runtime/Engine/Private/StaticMeshBuild.cpp b/Engine/Source/Runtime/Engine/Private/StaticMeshBuild.cpp
--- a/Engine/Source/Runtime/Engine/Private/StaticMeshBuild.cpp
+++ b/Engine/Source/Runtime/Engine/Private/StaticMeshBuild.cpp
@@ -570,7 +570,7 @@ void RemapPaintedVertexColors(const TArray<FPaintedVertex>& InPaintedVertices,
 		}
 
 		// Iterate through the octree attempting to find the vertices closest to the current new point
-		VertPosOctree.IterateElementsWithBoundsTest(FBoxCenterAndExtent(CurPosition, FVector::ZeroVector), [&PointsToConsider](const FPaintedVertex& Vertex)
+		VertPosOctree.FindNearbyElements(FBoxCenterAndExtent(CurPosition, FVector::ZeroVector), [&PointsToConsider](const FPaintedVertex& Vertex)
 		{
 			PointsToConsider.Add(Vertex);
 		});
