FNEM - Meshpaint was looking for nearby nodes and not intersecting nodes as such.

diff --git a/Engine/Plugins/MeshPainting/Source/MeshPaintingToolset/Private/MeshPaintHelpers.cpp b/Engine/Plugins/MeshPainting/Source/MeshPaintingToolset/Private/MeshPaintHelpers.cpp
--- a/Engine/Plugins/MeshPainting/Source/MeshPaintingToolset/Private/MeshPaintHelpers.cpp
+++ b/Engine/Plugins/MeshPainting/Source/MeshPaintingToolset/Private/MeshPaintHelpers.cpp
@@ -1330,7 +1330,7 @@ void UMeshPaintingToolset::ApplyVertexColorsToAllLODs(IMeshPaintComponentAdapter
 						FVector CurNormal = VertexTangentZ.ToFVector();
 
 						// Iterate through the octree attempting to find the vertices closest to the current new point
-						VertPosOctree.FindNearbyElements(FBoxCenterAndExtent(CurPosition, FVector::ZeroVector), [&PointsToConsider](const FPaintedMeshVertex& Vertex)
+						VertPosOctree.FindNearbyElements(CurPosition, [&PointsToConsider](const FPaintedMeshVertex& Vertex)
 						{
 							// Add all of the elements in the current node to the list of points to consider for closest point calculations
 							PointsToConsider.Add(Vertex);
diff --git a/Engine/Source/Editor/MeshPaint/Private/MeshPaintHelpers.cpp b/Engine/Source/Editor/MeshPaint/Private/MeshPaintHelpers.cpp
--- a/Engine/Source/Editor/MeshPaint/Private/MeshPaintHelpers.cpp
+++ b/Engine/Source/Editor/MeshPaint/Private/MeshPaintHelpers.cpp
@@ -1843,7 +1843,7 @@ void MeshPaintHelpers::ApplyVertexColorsToAllLODs(IMeshPaintGeometryAdapter& Geo
 						FVector CurNormal = VertexTangentZ.ToFVector();
 
 						// Iterate through the octree attempting to find the vertices closest to the current new point
-						VertPosOctree.FindNearbyElements(FBoxCenterAndExtent(CurPosition, FVector::ZeroVector), [&PointsToConsider](const FPaintedMeshVertex& Vertex)
+						VertPosOctree.FindNearbyElements(CurPosition, [&PointsToConsider](const FPaintedMeshVertex& Vertex)
 						{
 							PointsToConsider.Add(Vertex);
 						});
diff --git a/Engine/Source/Runtime/Core/Public/Math/GenericOctree.h b/Engine/Source/Runtime/Core/Public/Math/GenericOctree.h
--- a/Engine/Source/Runtime/Core/Public/Math/GenericOctree.h
+++ b/Engine/Source/Runtime/Core/Public/Math/GenericOctree.h
@@ -663,34 +663,26 @@ private:
 		{
 			for (int Index = 0; Index < TreeElements[CurrentNodeIndex].Num(); Index++)
 			{
-				typename TCallTraits<ElementType>::ConstReference Element = TreeElements[CurrentNodeIndex][Index];
-				if (Intersect(OctreeSemantics::GetBoundingBox(Element), BoxBounds))
-				{
-					Func(Element);
-				}
+				Func(TreeElements[CurrentNodeIndex][Index]);
 			}
 
 			if (!TreeNodes[CurrentNodeIndex].IsLeaf())
 			{
-				const FOctreeChildNodeSubset IntersectingChildSubset = NodeContext.GetIntersectingChildren(BoxBounds);
-				FNodeIndex ChildStartIndex = TreeNodes[CurrentNodeIndex].ChildNodes;
-				
-				bool FoundMatchingNode = false;
-				for (int i = 0; i < 8; i++)
+				// Find the child of the current node, if any, that contains the current new point
+				FOctreeChildNodeRef ChildRef = NodeContext.GetContainingChild(BoxBounds);
+				if (!ChildRef.IsNULL())
 				{
-					if (TreeNodes[ChildStartIndex + i].InclusiveNumElements > 0 && IntersectingChildSubset.Contains(FOctreeChildNodeRef(i)))
+					FNodeIndex ChildStartIndex = TreeNodes[CurrentNodeIndex].ChildNodes;
+					// If the specified child node exists and contains any match, push it than process it
+					if (TreeNodes[ChildStartIndex + ChildRef.Index].InclusiveNumElements > 0)
 					{
-						FoundMatchingNode = true;
-						FindNearbyElementsInternal(ChildStartIndex + i, NodeContext.GetChildContext(FOctreeChildNodeRef(i)), BoxBounds, Func);
+						FindNearbyElementsInternal(ChildStartIndex + ChildRef.Index, NodeContext.GetChildContext(ChildRef), BoxBounds, Func);
 					}
-				}
-
-				//look in the near node vicinity for a match
-				if (!FoundMatchingNode)
-				{
-					for (int i = 0; i < 8; i++)
+					// If the child node doesn't is a match, it's not worth pursuing any further. In an attempt to find
+					// anything to match vs. the new point, process all of the children of the current octree node
+					else
 					{
-						if (!IntersectingChildSubset.Contains(FOctreeChildNodeRef(i)))
+						for (int i = 0; i < 8; i++)
 						{
 							FindNearbyElementsInternal(ChildStartIndex + i, NodeContext.GetChildContext(FOctreeChildNodeRef(i)), BoxBounds, Func);
 						}
@@ -769,14 +761,14 @@ public:
 	}
 
 	/**
-	* this function will traverse the Octree using a fast (and relaxed) box-box intersection but will also try and search in the neighborhood nodes if the intersected node is empty.
-	* @param BoxBounds - the bounds to test if a node is traversed or skipped.
+	* this function will traverse the Octree using a tryint to find nearby nodes that contain any elements.
+	* @param Position - the position to look nearby.
 	* @param Func - Function to call with each Element for nodes that passed bounds test.
 	*/
 	template<typename IterateBoundsFunc>
-	inline void FindNearbyElements(const FBoxCenterAndExtent& BoxBounds, const IterateBoundsFunc& Func) const
+	inline void FindNearbyElements(const FVector& Position, const IterateBoundsFunc& Func) const
 	{
-		FindNearbyElementsInternal(0, RootNodeContext, BoxBounds, Func);
+		FindNearbyElementsInternal(0, RootNodeContext, FBoxCenterAndExtent(Position, FVector::ZeroVector), Func);
 	}
 
 
diff --git a/Engine/Source/Runtime/Engine/Private/StaticMeshBuild.cpp b/Engine/Source/Runtime/Engine/Private/StaticMeshBuild.cpp
--- a/Engine/Source/Runtime/Engine/Private/StaticMeshBuild.cpp
+++ b/Engine/Source/Runtime/Engine/Private/StaticMeshBuild.cpp
@@ -570,7 +570,7 @@ void RemapPaintedVertexColors(const TArray<FPaintedVertex>& InPaintedVertices,
 		}
 
 		// Iterate through the octree attempting to find the vertices closest to the current new point
-		VertPosOctree.FindNearbyElements(FBoxCenterAndExtent(CurPosition, FVector::ZeroVector), [&PointsToConsider](const FPaintedVertex& Vertex)
+		VertPosOctree.FindNearbyElements(CurPosition, [&PointsToConsider](const FPaintedVertex& Vertex)
 		{
 			PointsToConsider.Add(Vertex);
 		});
