FNEM - Better naming

diff --git a/Engine/Plugins/MeshPainting/Source/MeshPaintingToolset/Private/BaseMeshPaintComponentAdapter.cpp b/Engine/Plugins/MeshPainting/Source/MeshPaintingToolset/Private/BaseMeshPaintComponentAdapter.cpp
--- a/Engine/Plugins/MeshPainting/Source/MeshPaintingToolset/Private/BaseMeshPaintComponentAdapter.cpp
+++ b/Engine/Plugins/MeshPainting/Source/MeshPaintingToolset/Private/BaseMeshPaintComponentAdapter.cpp
@@ -85,7 +85,7 @@ TArray<uint32> FBaseMeshPaintComponentAdapter::SphereIntersectTriangles(const fl
 	// definitely don't want our brush to be cut off by a hard triangle edge
 	const float SquaredRadiusBias = ComponentSpaceSquaredBrushRadius * 0.025f;
 
-	MeshTriOctree->IterateElementsWithBoundsTest(FBoxCenterAndExtent(ComponentSpaceBrushPosition, FVector(FMath::Sqrt(ComponentSpaceSquaredBrushRadius + SquaredRadiusBias))), [&OutTriangles, bOnlyFrontFacing, &ComponentSpaceCameraPosition](const FMeshPaintTriangle& CurrentTri)
+	MeshTriOctree->FindElementsWithBoundsTest(FBoxCenterAndExtent(ComponentSpaceBrushPosition, FVector(FMath::Sqrt(ComponentSpaceSquaredBrushRadius + SquaredRadiusBias))), [&OutTriangles, bOnlyFrontFacing, &ComponentSpaceCameraPosition](const FMeshPaintTriangle& CurrentTri)
 	{
 		const float SignedPlaneDist = FVector::PointPlaneDist(ComponentSpaceCameraPosition, CurrentTri.Vertices[0], CurrentTri.Normal);
 		if (!bOnlyFrontFacing || SignedPlaneDist < 0.0f)
diff --git a/Engine/Plugins/Runtime/EditableMesh/Source/EditableMesh/EditableMesh.cpp b/Engine/Plugins/Runtime/EditableMesh/Source/EditableMesh/EditableMesh.cpp
--- a/Engine/Plugins/Runtime/EditableMesh/Source/EditableMesh/EditableMesh.cpp
+++ b/Engine/Plugins/Runtime/EditableMesh/Source/EditableMesh/EditableMesh.cpp
@@ -8099,7 +8099,7 @@ void UEditableMesh::SearchSpatialDatabaseWithPredicate( TFunctionRef< bool( cons
 	if( IsSpatialDatabaseAllowed() && ensure( Octree.IsValid() ) )
 	{
 		// @todo mesheditor perf: Do we need to use a custom stack allocator for iterating?  The default should probably be okay.
-		Octree->IterateElementsWithPredicate([&Predicate](const FBoxCenterAndExtent& NodeBounds)
+		Octree->FindElementsWithPredicate([&Predicate](const FBoxCenterAndExtent& NodeBounds)
 		{
 			return Predicate(NodeBounds.GetBox());
 		},
@@ -8153,7 +8153,7 @@ void UEditableMesh::SearchSpatialDatabaseForPolygonsPotentiallyIntersectingPlane
 	if( IsSpatialDatabaseAllowed() && ensure( Octree.IsValid() ) )
 	{
 		// @todo mesheditor perf: Do we need to use a custom stack allocator for iterating?  The default should probably be okay.
-		Octree->IterateElementsWithPredicate([&InPlane](const FBoxCenterAndExtent& NodeBounds)
+		Octree->FindElementsWithPredicate([&InPlane](const FBoxCenterAndExtent& NodeBounds)
 		{
 			const bool bIsOverlappingLineSegment =
 				FMath::PlaneAABBIntersection(
diff --git a/Engine/Source/Developer/FunctionalTesting/Private/FunctionalAITest.cpp b/Engine/Source/Developer/FunctionalTesting/Private/FunctionalAITest.cpp
--- a/Engine/Source/Developer/FunctionalTesting/Private/FunctionalAITest.cpp
+++ b/Engine/Source/Developer/FunctionalTesting/Private/FunctionalAITest.cpp
@@ -118,7 +118,7 @@ void AFunctionalAITest::OnTimeout()
 		const FVector TransformedOrigin = GetTransform().TransformPosition(NavMeshDebugOrigin);
 		const FBox DebugBounds = FBox::BuildAABB(TransformedOrigin, NavMeshDebugExtent);
 
-		NavigationOctree->IterateElementsWithBoundsTest(DebugBounds, [&AreaFilter, &Navmesh](const FNavigationOctreeElement& Element)
+		NavigationOctree->FindElementsWithBoundsTest(DebugBounds, [&AreaFilter, &Navmesh](const FNavigationOctreeElement& Element)
 		{
 			if (Element.IsMatchingFilter(AreaFilter))
 			{
diff --git a/Engine/Source/Developer/SkeletalMeshUtilitiesCommon/Private/LODUtilities.cpp b/Engine/Source/Developer/SkeletalMeshUtilitiesCommon/Private/LODUtilities.cpp
--- a/Engine/Source/Developer/SkeletalMeshUtilitiesCommon/Private/LODUtilities.cpp
+++ b/Engine/Source/Developer/SkeletalMeshUtilitiesCommon/Private/LODUtilities.cpp
@@ -1443,7 +1443,7 @@ void MatchVertexIndexUsingPosition(
 		FBox CurBox(PositionSrc - Extent, PositionSrc + Extent);
 		while (OcTreeTriangleResults.Num() <= 0)
 		{
-			OcTree.IterateElementsWithBoundsTest(CurBox, [&OcTreeTriangleResults](const FTriangleElement& Element)
+			OcTree.FindElementsWithBoundsTest(CurBox, [&OcTreeTriangleResults](const FTriangleElement& Element)
 			{
 				// Add all of the elements in the current node to the list of points to consider for closest point calculations
 				OcTreeTriangleResults.Add(Element);
diff --git a/Engine/Source/Editor/MeshPaint/Private/BaseMeshPaintGeometryAdapter.cpp b/Engine/Source/Editor/MeshPaint/Private/BaseMeshPaintGeometryAdapter.cpp
--- a/Engine/Source/Editor/MeshPaint/Private/BaseMeshPaintGeometryAdapter.cpp
+++ b/Engine/Source/Editor/MeshPaint/Private/BaseMeshPaintGeometryAdapter.cpp
@@ -78,7 +78,7 @@ TArray<uint32> FBaseMeshPaintGeometryAdapter::SphereIntersectTriangles(const flo
 	// definitely don't want our brush to be cut off by a hard triangle edge
 	const float SquaredRadiusBias = ComponentSpaceSquaredBrushRadius * 0.025f;
 
-	MeshTriOctree->IterateElementsWithBoundsTest(FBoxCenterAndExtent(ComponentSpaceBrushPosition, FVector(FMath::Sqrt(ComponentSpaceSquaredBrushRadius + SquaredRadiusBias))), [bOnlyFrontFacing , &OutTriangles, &ComponentSpaceCameraPosition](const FMeshPaintTriangle& CurrentTri)
+	MeshTriOctree->FindElementsWithBoundsTest(FBoxCenterAndExtent(ComponentSpaceBrushPosition, FVector(FMath::Sqrt(ComponentSpaceSquaredBrushRadius + SquaredRadiusBias))), [bOnlyFrontFacing , &OutTriangles, &ComponentSpaceCameraPosition](const FMeshPaintTriangle& CurrentTri)
 	{
 		const float SignedPlaneDist = FVector::PointPlaneDist(ComponentSpaceCameraPosition, CurrentTri.Vertices[0], CurrentTri.Normal);
 		if (!bOnlyFrontFacing || SignedPlaneDist < 0.0f)
diff --git a/Engine/Source/Editor/UnrealEd/Private/Fbx/FbxSkeletalMeshImport.cpp b/Engine/Source/Editor/UnrealEd/Private/Fbx/FbxSkeletalMeshImport.cpp
--- a/Engine/Source/Editor/UnrealEd/Private/Fbx/FbxSkeletalMeshImport.cpp
+++ b/Engine/Source/Editor/UnrealEd/Private/Fbx/FbxSkeletalMeshImport.cpp
@@ -220,7 +220,7 @@ void RemapSkeletalMeshVertexColorToImportData(const USkeletalMesh* SkeletalMesh,
 		const FVector& Normal = WedgeIndexToNormal.FindChecked(WedgeIndex);
 
 		TArray<FSoftSkinVertex> PointsToConsider;
-		VertPosOctree.IterateElementsWithBoundsTest(FBoxCenterAndExtent(Position, FVector::ZeroVector), [&PointsToConsider](const FSoftSkinVertex& Vertex)
+		VertPosOctree.FindNearbyElements(Position, [&PointsToConsider](const FSoftSkinVertex& Vertex)
 		{
 			PointsToConsider.Add(Vertex);
 		});
diff --git a/Engine/Source/Runtime/Core/Public/Math/GenericOctree.h b/Engine/Source/Runtime/Core/Public/Math/GenericOctree.h
--- a/Engine/Source/Runtime/Core/Public/Math/GenericOctree.h
+++ b/Engine/Source/Runtime/Core/Public/Math/GenericOctree.h
@@ -579,7 +579,7 @@ private:
 	}
 
 	template<typename PredicateFunc, typename IterateFunc>
-	void IterateNodesWithPredicateInternal(FNodeIndex CurrentNodeIndex, const FOctreeNodeContext& NodeContext, const PredicateFunc& Predicate, const IterateFunc& Func) const
+	void FindNodesWithPredicateInternal(FNodeIndex CurrentNodeIndex, const FOctreeNodeContext& NodeContext, const PredicateFunc& Predicate, const IterateFunc& Func) const
 	{
 		if (TreeNodes[CurrentNodeIndex].InclusiveNumElements > 0)
 		{
@@ -592,7 +592,7 @@ private:
 					FNodeIndex ChildStartIndex = TreeNodes[CurrentNodeIndex].ChildNodes;
 					for (int i = 0; i < 8; i++)
 					{
-						IterateNodesWithPredicateInternal(ChildStartIndex + i, NodeContext.GetChildContext(FOctreeChildNodeRef(i)), Predicate, Func);
+						FindNodesWithPredicateInternal(ChildStartIndex + i, NodeContext.GetChildContext(FOctreeChildNodeRef(i)), Predicate, Func);
 					}
 				}
 			}
@@ -600,7 +600,7 @@ private:
 	}
 
 	template<typename IterateFunc>
-	void IterateElementsWithBoundsTestInternal(FNodeIndex CurrentNodeIndex, const FOctreeNodeContext& NodeContext, const FBoxCenterAndExtent& BoxBounds, const IterateFunc& Func) const
+	void FindElementsWithBoundsTestInternal(FNodeIndex CurrentNodeIndex, const FOctreeNodeContext& NodeContext, const FBoxCenterAndExtent& BoxBounds, const IterateFunc& Func) const
 	{
 		if (TreeNodes[CurrentNodeIndex].InclusiveNumElements > 0)
 		{
@@ -620,7 +620,7 @@ private:
 				{
 					if(IntersectingChildSubset.Contains(FOctreeChildNodeRef(i)))
 					{
-						IterateElementsWithBoundsTestInternal(ChildStartIndex + i, NodeContext.GetChildContext(FOctreeChildNodeRef(i)), BoxBounds, Func);
+						FindElementsWithBoundsTestInternal(ChildStartIndex + i, NodeContext.GetChildContext(FOctreeChildNodeRef(i)), BoxBounds, Func);
 					}
 				}
 			}
@@ -698,7 +698,7 @@ public:
 	 * @param Func - Function to call with each Element.
 	 */
 	template<typename IterateAllElementsFunc>
-	inline void IterateAllElements(const IterateAllElementsFunc& Func) const
+	inline void FindAllElements(const IterateAllElementsFunc& Func) const
 	{
 		for (const ElementArrayType& Elements : TreeElements)
 		{
@@ -715,9 +715,9 @@ public:
 	 * @param Func - Function that will receive the node ID which can be stored and later used to get the elements using GetElementsForNode for all nodes that passed the predicate.
 	 */
 	template<typename PredicateFunc, typename IterateFunc>
-	inline void IterateNodesWithPredicate(const PredicateFunc& Predicate, const IterateFunc& Func) const
+	inline void FindNodesWithPredicate(const PredicateFunc& Predicate, const IterateFunc& Func) const
 	{
-		IterateNodesWithPredicateInternal(0, RootNodeContext, Predicate, Func);
+		FindNodesWithPredicateInternal(0, RootNodeContext, Predicate, Func);
 	}
 
 	/**
@@ -726,9 +726,9 @@ public:
 	 * @param Func - Function to call with each Element for nodes that passed the predicate.
 	 */
 	template<typename PredicateFunc, typename IterateFunc>
-	inline void IterateElementsWithPredicate(const PredicateFunc& Predicate, const IterateFunc& Func) const
+	inline void FindElementsWithPredicate(const PredicateFunc& Predicate, const IterateFunc& Func) const
 	{
-		IterateNodesWithPredicateInternal(0, RootNodeContext, Predicate, [&Func, this](FNodeIndex NodeIndex)
+		FindNodesWithPredicateInternal(0, RootNodeContext, Predicate, [&Func, this](FNodeIndex NodeIndex)
 		{
 			for (typename TCallTraits<ElementType>::ConstReference Element : TreeElements[NodeIndex])
 			{
@@ -743,9 +743,9 @@ public:
 	 * @param Func - Function to call with each Element for nodes that passed bounds test.
 	 */
 	template<typename IterateBoundsFunc>
-	inline void IterateElementsWithBoundsTest(const FBoxCenterAndExtent& BoxBounds, const IterateBoundsFunc& Func) const
+	inline void FindElementsWithBoundsTest(const FBoxCenterAndExtent& BoxBounds, const IterateBoundsFunc& Func) const
 	{
-		IterateElementsWithBoundsTestInternal(0, RootNodeContext, BoxBounds, Func);
+		FindElementsWithBoundsTestInternal(0, RootNodeContext, BoxBounds, Func);
 	}
 
 	/**
@@ -881,7 +881,7 @@ public:
 
 	/**
 	 * return all elements for a given node.
-	 * @param NodeIndex - The the index of the node can be obtained using IterateNodesWithPredicate.
+	 * @param NodeIndex - The the index of the node can be obtained using FindNodesWithPredicate.
 	 */
 	TArrayView<const ElementType> GetElementsForNode(FNodeIndex NodeIndex) const
 	{
@@ -897,7 +897,7 @@ public:
 		int32 MaxElementsPerNode = 0;
 		TArray<int32> NodeElementDistribution;
 
-		IterateNodesWithPredicateInternal(0, RootNodeContext, [](const FBoxCenterAndExtent&) { return true; }, [&, this](FNodeIndex NodeIndex)
+		FindNodesWithPredicateInternal(0, RootNodeContext, [](const FBoxCenterAndExtent&) { return true; }, [&, this](FNodeIndex NodeIndex)
 		{
 			const int32 CurrentNodeElementCount = GetElementsForNode(NodeIndex).Num();
 
diff --git a/Engine/Source/Runtime/Engine/Private/PrecomputedLightVolume.cpp b/Engine/Source/Runtime/Engine/Private/PrecomputedLightVolume.cpp
--- a/Engine/Source/Runtime/Engine/Private/PrecomputedLightVolume.cpp
+++ b/Engine/Source/Runtime/Engine/Private/PrecomputedLightVolume.cpp
@@ -221,7 +221,7 @@ FArchive& operator<<(FArchive& Ar,FPrecomputedLightVolumeData& Volume)
 			if (!Ar.IsCooking() || Ar.CookingTarget()->SupportsFeature(ETargetPlatformFeatures::HighQualityLightmaps))
 			{
 				// Gather an array of samples from the octree
-				Volume.HighQualityLightmapOctree.IterateAllElements([&HighQualitySamples](const FVolumeLightingSample& Sample)
+				Volume.HighQualityLightmapOctree.FindAllElements([&HighQualitySamples](const FVolumeLightingSample& Sample)
 				{
 					HighQualitySamples.Add(Sample);
 				});
@@ -234,7 +234,7 @@ FArchive& operator<<(FArchive& Ar,FPrecomputedLightVolumeData& Volume)
 			if (!Ar.IsCooking() || Ar.CookingTarget()->SupportsFeature(ETargetPlatformFeatures::LowQualityLightmaps))
 			{
 				// Gather an array of samples from the octree
-				Volume.LowQualityLightmapOctree.IterateAllElements([&LowQualitySamples](const FVolumeLightingSample& Sample)
+				Volume.LowQualityLightmapOctree.FindAllElements([&LowQualitySamples](const FVolumeLightingSample& Sample)
 				{
 					LowQualitySamples.Add(Sample);
 				});
@@ -397,7 +397,7 @@ void FPrecomputedLightVolume::InterpolateIncidentRadiancePoint(
 		FVector WorldPosition = InWorldPosition - WorldOriginOffset; // relocate from world to volume space
 		FBoxCenterAndExtent BoundingBox( WorldPosition, FVector::ZeroVector );
 		// Iterate over the octree nodes containing the query point.
-		OctreeForRendering->IterateElementsWithBoundsTest(BoundingBox, [&AccumulatedIncidentRadiance, &SkyBentNormal, &AccumulatedDirectionalLightShadowing, &AccumulatedWeight, &WorldPosition](const FVolumeLightingSample& VolumeSample)
+		OctreeForRendering->FindElementsWithBoundsTest(BoundingBox, [&AccumulatedIncidentRadiance, &SkyBentNormal, &AccumulatedDirectionalLightShadowing, &AccumulatedWeight, &WorldPosition](const FVolumeLightingSample& VolumeSample)
 		{
 			const float DistanceSquared = (VolumeSample.Position - WorldPosition).SizeSquared();
 			const float RadiusSquared = FMath::Square(VolumeSample.Radius);
@@ -441,7 +441,7 @@ void FPrecomputedLightVolume::InterpolateIncidentRadianceBlock(
 		PotentiallyIntersectingSamples.Reset(100);
 
 		// Iterate over the octree nodes containing the query point.
-		OctreeForRendering->IterateElementsWithBoundsTest(BoundingBox, [&](const FVolumeLightingSample& VolumeSample)
+		OctreeForRendering->FindElementsWithBoundsTest(BoundingBox, [&](const FVolumeLightingSample& VolumeSample)
 		{
 			PotentiallyIntersectingSamples.Add(&VolumeSample);
 		});
@@ -500,7 +500,7 @@ void FPrecomputedLightVolume::InterpolateIncidentRadianceBlock(
 
 void FPrecomputedLightVolume::DebugDrawSamples(FPrimitiveDrawInterface* PDI, bool bDrawDirectionalShadowing) const
 {
-	OctreeForRendering->IterateAllElements([bDrawDirectionalShadowing, PDI, this](const FVolumeLightingSample& VolumeSample)
+	OctreeForRendering->FindAllElements([bDrawDirectionalShadowing, PDI, this](const FVolumeLightingSample& VolumeSample)
 	{
 		const FLinearColor AverageColor = bDrawDirectionalShadowing
 			? FLinearColor(VolumeSample.DirectionalLightShadowing, VolumeSample.DirectionalLightShadowing, VolumeSample.DirectionalLightShadowing)
diff --git a/Engine/Source/Runtime/Engine/Private/Rendering/SkeletalMeshLODImporterData.cpp b/Engine/Source/Runtime/Engine/Private/Rendering/SkeletalMeshLODImporterData.cpp
--- a/Engine/Source/Runtime/Engine/Private/Rendering/SkeletalMeshLODImporterData.cpp
+++ b/Engine/Source/Runtime/Engine/Private/Rendering/SkeletalMeshLODImporterData.cpp
@@ -939,7 +939,7 @@ void FOctreeQueryHelper::FindNearestWedgeIndexes(const FVector& SearchPosition,
 	{
 		// Iterate through the octree attempting to find the vertices closest to the current new point
 		// The first shot is an intersection with a 1 CM cube box around the search position, this ensure we dont fall in the wrong neighbourg
-		WedgePosOctree->IterateElementsWithBoundsTest(FBoxCenterAndExtent(SearchPosition, Extend), [&OutNearestWedges, &MinSquaredDistance, &SearchPosition](const FWedgeInfo& WedgeInfo)
+		WedgePosOctree->FindElementsWithBoundsTest(FBoxCenterAndExtent(SearchPosition, Extend), [&OutNearestWedges, &MinSquaredDistance, &SearchPosition](const FWedgeInfo& WedgeInfo)
 		{
 			// Add all of the elements in the current node to the list of points to consider for closest point calculations
 			float VectorDelta = FVector::DistSquared(SearchPosition, WedgeInfo.Position);
diff --git a/Engine/Source/Runtime/Experimental/GeometryCollectionCore/Private/GeometryCollection/GeometryCollectionProximityUtility.cpp b/Engine/Source/Runtime/Experimental/GeometryCollectionCore/Private/GeometryCollection/GeometryCollectionProximityUtility.cpp
--- a/Engine/Source/Runtime/Experimental/GeometryCollectionCore/Private/GeometryCollection/GeometryCollectionProximityUtility.cpp
+++ b/Engine/Source/Runtime/Experimental/GeometryCollectionCore/Private/GeometryCollection/GeometryCollectionProximityUtility.cpp
@@ -274,7 +274,7 @@ void FGeometryCollectionProximityUtility::UpdateProximity(FGeometryCollection* G
 		// 	Query the Octree
 		OtherFaceTransformDataArray.Reserve(NumFaces);
 
-		MeshTriOctree->IterateElementsWithBoundsTest(ThisFaceBounds, [&OtherFaceTransformDataArray, &FaceTransformDataArray](const FProximityTriangle& OctreePolygon)
+		MeshTriOctree->FindElementsWithBoundsTest(ThisFaceBounds, [&OtherFaceTransformDataArray, &FaceTransformDataArray](const FProximityTriangle& OctreePolygon)
 		{
 			OtherFaceTransformDataArray.Add(FaceTransformDataArray[OctreePolygon.ArrayIndex]);
 		});
diff --git a/Engine/Source/Runtime/NavigationSystem/Private/NavMesh/NavMeshRenderingComponent.cpp b/Engine/Source/Runtime/NavigationSystem/Private/NavMesh/NavMeshRenderingComponent.cpp
--- a/Engine/Source/Runtime/NavigationSystem/Private/NavMesh/NavMeshRenderingComponent.cpp
+++ b/Engine/Source/Runtime/NavigationSystem/Private/NavMesh/NavMeshRenderingComponent.cpp
@@ -752,7 +752,7 @@ void FNavMeshSceneProxyData::GatherData(const ARecastNavMesh* NavMesh, int32 InN
 				TArray<uint32> CollidingIndices;
 
 				FBox CurrentNodeBoundsBox;
-				NavOctree->IterateElementsWithPredicate([bGatherOctree, bGatherPathCollidingGeometry, this, &CurrentNodeBoundsBox](const FBoxCenterAndExtent& NodeBounds)
+				NavOctree->FindElementsWithPredicate([bGatherOctree, bGatherPathCollidingGeometry, this, &CurrentNodeBoundsBox](const FBoxCenterAndExtent& NodeBounds)
 				{
 					if (bGatherOctree)
 					{
diff --git a/Engine/Source/Runtime/NavigationSystem/Private/NavMesh/RecastNavMeshGenerator.cpp b/Engine/Source/Runtime/NavigationSystem/Private/NavMesh/RecastNavMeshGenerator.cpp
--- a/Engine/Source/Runtime/NavigationSystem/Private/NavMesh/RecastNavMeshGenerator.cpp
+++ b/Engine/Source/Runtime/NavigationSystem/Private/NavMesh/RecastNavMeshGenerator.cpp
@@ -2193,11 +2193,8 @@ void FRecastTileGenerator::PrepareGeometrySources(const FRecastNavMeshGenerator&
 	NavSystem = NavSys;
 	bUpdateGeometry = bGeometryChanged;
 
-	for (FNavigationOctree::TConstElementBoxIterator<FNavigationOctree::DefaultStackAllocator> It(*NavOctreeInstance, ParentGenerator.GrowBoundingBox(TileBB, /*bIncludeAgentHeight*/ false));
-		It.HasPendingElements();
-		It.Advance())
+	NavOctreeInstance->FindElementsWithBoundsTest(ParentGenerator.GrowBoundingBox(TileBB, /*bIncludeAgentHeight*/ false), [this, bGeometryChanged](const FNavigationOctreeElement& Element)
 	{
-		const FNavigationOctreeElement& Element = It.GetCurrentElement();
 		const bool bShouldUse = Element.ShouldUseGeometry(NavDataConfig);
 		if (bShouldUse)
 		{
@@ -2210,7 +2207,7 @@ void FRecastTileGenerator::PrepareGeometrySources(const FRecastNavMeshGenerator&
 				NavigationRelevantData.Add(Element.Data);
 			}
 		}
-	}
+	});
 }
 
 void FRecastTileGenerator::GatherGeometry(const FRecastNavMeshGenerator& ParentGenerator, bool bGeometryChanged)
@@ -2225,7 +2222,7 @@ void FRecastTileGenerator::GatherGeometry(const FRecastNavMeshGenerator& ParentG
 	}
 	const FNavDataConfig& OwnerNavDataConfig = ParentGenerator.GetOwner()->GetConfig();
 
-	NavigationOctree->IterateElementsWithBoundsTest(ParentGenerator.GrowBoundingBox(TileBB, /*bIncludeAgentHeight*/ false), [&OwnerNavDataConfig, &NavigationOctree, this, bGeometryChanged](const FNavigationOctreeElement& Element)
+	NavigationOctree->FindElementsWithBoundsTest(ParentGenerator.GrowBoundingBox(TileBB, /*bIncludeAgentHeight*/ false), [&OwnerNavDataConfig, &NavigationOctree, this, bGeometryChanged](const FNavigationOctreeElement& Element)
 	{
 		const bool bShouldUse = Element.ShouldUseGeometry(OwnerNavDataConfig);
 		if (bShouldUse)
@@ -6351,11 +6348,8 @@ void FRecastNavMeshGenerator::GrabDebugSnapshot(struct FVisualLogEntry* Snapshot
 		const ARecastNavMesh* NavData = Cast<const ARecastNavMesh>(NavSys->NavDataSet[Index]);
 		if (NavData)
 		{
-			for (FNavigationOctree::TConstElementBoxIterator<FNavigationOctree::DefaultStackAllocator> It(*NavOctree, BoundingBox);
-				It.HasPendingElements();
-				It.Advance())
+			NavOctree->FindElementsWithBoundsTest(BoundingBox, [this, NavData, &Indices, &CoordBuffer, Snapshot, &CategoryName, LogVerbosity, NavAreaVerbosity](const FNavigationOctreeElement& Element)
 			{
-				const FNavigationOctreeElement& Element = It.GetCurrentElement();
 				const bool bExportGeometry = Element.Data->HasGeometry() && Element.ShouldUseGeometry(DestNavMesh->GetConfig());
 
 				TArray<FTransform> InstanceTransforms;
@@ -6458,8 +6452,7 @@ void FRecastNavMeshGenerator::GrabDebugSnapshot(struct FVisualLogEntry* Snapshot
 						}
 					}
 				}
-			}
-
+			});
 		}
 
 	}
@@ -6495,11 +6488,8 @@ void FRecastNavMeshGenerator::ExportNavigationData(const FString& FileName) cons
 			};
 			TArray<FAreaExportData> AreaExport;
 
-			for(FNavigationOctree::TConstElementBoxIterator<FNavigationOctree::DefaultStackAllocator> It(*NavOctree, TotalNavBounds);
-				It.HasPendingElements();
-				It.Advance())
+			NavOctree->FindElementsWithBoundsTest(TotalNavBounds, [this, NavData, &IndexBuffer, &CoordBuffer, &AreaExport](const FNavigationOctreeElement& Element)
 			{
-				const FNavigationOctreeElement& Element = It.GetCurrentElement();
 				const bool bExportGeometry = Element.Data->HasGeometry() && Element.ShouldUseGeometry(DestNavMesh->GetConfig());
 
 				TArray<FTransform> InstanceTransforms;
@@ -6584,7 +6574,7 @@ void FRecastNavMeshGenerator::ExportNavigationData(const FString& FileName) cons
 						}
 					}
 				}
-			}
+			});
 			
 			UWorld* NavigationWorld = GetWorld();
 			for (int32 LevelIndex = 0; LevelIndex < NavigationWorld->GetNumLevels(); ++LevelIndex) 
diff --git a/Engine/Source/Runtime/NavigationSystem/Private/NavigationDataHandler.cpp b/Engine/Source/Runtime/NavigationSystem/Private/NavigationDataHandler.cpp
--- a/Engine/Source/Runtime/NavigationSystem/Private/NavigationDataHandler.cpp
+++ b/Engine/Source/Runtime/NavigationSystem/Private/NavigationDataHandler.cpp
@@ -360,7 +360,7 @@ void FNavigationDataHandler::FindElementsInNavOctree(const FBox& QueryBox, const
 		return;
 	}
 
-	OctreeController.NavOctree->IterateElementsWithBoundsTest(QueryBox, [&Elements, &Filter](const FNavigationOctreeElement& Element)
+	OctreeController.NavOctree->FindElementsWithBoundsTest(QueryBox, [&Elements, &Filter](const FNavigationOctreeElement& Element)
 	{
 		if (Element.IsMatchingFilter(Filter))
 		{
diff --git a/Engine/Source/Runtime/Renderer/Private/LightSceneInfo.cpp b/Engine/Source/Runtime/Renderer/Private/LightSceneInfo.cpp
--- a/Engine/Source/Runtime/Renderer/Private/LightSceneInfo.cpp
+++ b/Engine/Source/Runtime/Renderer/Private/LightSceneInfo.cpp
@@ -117,7 +117,7 @@ void FLightSceneInfo::AddToScene()
 			// Find primitives that the light affects in the primitive octree.
 			FMemMark MemStackMark(FMemStack::Get());
 
-			Scene->PrimitiveOctree.IterateElementsWithBoundsTest(GetBoundingBox(), [&LightSceneInfoCompact, this](const FPrimitiveSceneInfoCompact& PrimitiveSceneInfoCompact)
+			Scene->PrimitiveOctree.FindElementsWithBoundsTest(GetBoundingBox(), [&LightSceneInfoCompact, this](const FPrimitiveSceneInfoCompact& PrimitiveSceneInfoCompact)
 			{
 				CreateLightPrimitiveInteraction(LightSceneInfoCompact, PrimitiveSceneInfoCompact);
 			});
diff --git a/Engine/Source/Runtime/Renderer/Private/RendererScene.cpp b/Engine/Source/Runtime/Renderer/Private/RendererScene.cpp
--- a/Engine/Source/Runtime/Renderer/Private/RendererScene.cpp
+++ b/Engine/Source/Runtime/Renderer/Private/RendererScene.cpp
@@ -1099,7 +1099,7 @@ public:
 			const FPrimitiveSceneInfoCompact PrimitiveSceneInfoCompact(PrimInfo);
 
 			// Find local lights that affect the primitive in the light octree.
-			Scene->LocalShadowCastingLightOctree.IterateElementsWithBoundsTest(Bounds.GetBox(), [&PrimitiveSceneInfoCompact](const FLightSceneInfoCompact& LightSceneInfoCompact)
+			Scene->LocalShadowCastingLightOctree.FindElementsWithBoundsTest(Bounds.GetBox(), [&PrimitiveSceneInfoCompact](const FLightSceneInfoCompact& LightSceneInfoCompact)
 			{
 				LightSceneInfoCompact.LightSceneInfo->CreateLightPrimitiveInteraction(LightSceneInfoCompact, PrimitiveSceneInfoCompact);
 			});
@@ -4222,7 +4222,7 @@ void FScene::CreateLightPrimitiveInteractionsForPrimitive(FPrimitiveSceneInfo* P
 		const FPrimitiveSceneInfoCompact PrimitiveSceneInfoCompact(PrimitiveInfo);
 
 		// Find local lights that affect the primitive in the light octree.
-		LocalShadowCastingLightOctree.IterateElementsWithBoundsTest(Bounds.GetBox(), [&PrimitiveSceneInfoCompact](const FLightSceneInfoCompact& LightSceneInfoCompact)
+		LocalShadowCastingLightOctree.FindElementsWithBoundsTest(Bounds.GetBox(), [&PrimitiveSceneInfoCompact](const FLightSceneInfoCompact& LightSceneInfoCompact)
 		{
 			LightSceneInfoCompact.LightSceneInfo->CreateLightPrimitiveInteraction(LightSceneInfoCompact, PrimitiveSceneInfoCompact);
 		});
diff --git a/Engine/Source/Runtime/Renderer/Private/ShadowSetup.cpp b/Engine/Source/Runtime/Renderer/Private/ShadowSetup.cpp
--- a/Engine/Source/Runtime/Renderer/Private/ShadowSetup.cpp
+++ b/Engine/Source/Runtime/Renderer/Private/ShadowSetup.cpp
@@ -3979,7 +3979,7 @@ void FSceneRenderer::GatherShadowPrimitives(
 			Packets.Reserve(100);
 
 			// Find primitives that are in a shadow frustum in the octree.
-			Scene->PrimitiveOctree.IterateNodesWithPredicate([&PreShadows, &ViewDependentWholeSceneShadows](const FBoxCenterAndExtent& NodeBounds)
+			Scene->PrimitiveOctree.FindNodesWithPredicate([&PreShadows, &ViewDependentWholeSceneShadows](const FBoxCenterAndExtent& NodeBounds)
 			{
 				// Check that the child node is in the frustum for at least one shadow.
 
diff --git a/Engine/Source/Runtime/Renderer/Private/ShadowSetupMobile.cpp b/Engine/Source/Runtime/Renderer/Private/ShadowSetupMobile.cpp
--- a/Engine/Source/Runtime/Renderer/Private/ShadowSetupMobile.cpp
+++ b/Engine/Source/Runtime/Renderer/Private/ShadowSetupMobile.cpp
@@ -169,7 +169,7 @@ static bool MobileDetermineStaticMeshesCSMVisibilityState(FScene* Scene, FViewIn
 	{
 		QUICK_SCOPE_CYCLE_COUNTER(STAT_ShadowOctreeTraversal);
 		// Find primitives that are in a shadow frustum in the octree.
-		Scene->PrimitiveOctree.IterateElementsWithPredicate([&IsReceiverFunc](const FBoxCenterAndExtent& NodeBounds)
+		Scene->PrimitiveOctree.FindElementsWithPredicate([&IsReceiverFunc](const FBoxCenterAndExtent& NodeBounds)
 		{
 			return IsReceiverFunc(FVector(NodeBounds.Center), FVector(NodeBounds.Extent), NodeBounds.Extent.Size3());
 		},
