FNEM - Various Compileer fixes

diff --git a/Engine/Source/Runtime/Core/Public/Math/GenericOctree.h b/Engine/Source/Runtime/Core/Public/Math/GenericOctree.h
--- a/Engine/Source/Runtime/Core/Public/Math/GenericOctree.h
+++ b/Engine/Source/Runtime/Core/Public/Math/GenericOctree.h
@@ -550,16 +550,6 @@ private:
 				return;
 			}
 		}
-
-		UE_LOG(LogGenericOctree, Fatal,
-			TEXT("Failed to find an octree node for an element with bounds (%f,%f,%f) +/- (%f,%f,%f)!"),
-			ElementBounds.Center.X,
-			ElementBounds.Center.Y,
-			ElementBounds.Center.Z,
-			ElementBounds.Extent.X,
-			ElementBounds.Extent.Y,
-			ElementBounds.Extent.Z
-		);
 	}
 
 	void CollapseNodesInternal(FNodeIndex CurrentNodeIndex, ElementArrayType& CollapsedNodeElements)
diff --git a/Engine/Source/Runtime/Core/Public/Math/GenericOctree.inl b/Engine/Source/Runtime/Core/Public/Math/GenericOctree.inl
--- a/Engine/Source/Runtime/Core/Public/Math/GenericOctree.inl
+++ b/Engine/Source/Runtime/Core/Public/Math/GenericOctree.inl
@@ -50,8 +50,8 @@ FORCEINLINE FOctreeChildNodeSubset FOctreeNodeContext::GetIntersectingChildren(c
 		);
 
 	// Intersect the query bounds with the node's children's bounds.
-	Result.PositiveChildBits = VectorMaskBits(VectorMask_GT(QueryBoundsMax, PositiveChildBoundsMin)) & 0x7;
-	Result.NegativeChildBits = VectorMaskBits(VectorMask_LE(QueryBoundsMin, NegativeChildBoundsMax)) & 0x7;
+	Result.PositiveChildBits = VectorMaskBits(VectorCompareGT(QueryBoundsMax, PositiveChildBoundsMin)) & 0x7;
+	Result.NegativeChildBits = VectorMaskBits(VectorCompareLE(QueryBoundsMin, NegativeChildBoundsMax)) & 0x7;
 	return Result;
 }
 
@@ -78,7 +78,7 @@ FORCEINLINE FOctreeChildNodeRef FOctreeNodeContext::GetContainingChild(const FBo
 	else
 	{
 		// Return the child node that the query is closest to as the containing child.
-		Result.Index = VectorMaskBits(VectorMask_GT(QueryBoundsCenter, BoundsCenter)) & 0x7;
+		Result.Index = VectorMaskBits(VectorCompareGT(QueryBoundsCenter, BoundsCenter)) & 0x7;
 	}
 
 	return Result;
