fix for infinite loading screen

diff --git a/Engine/Source/Runtime/Core/Private/Misc/PreLoadFile.cpp b/Engine/Source/Runtime/Core/Private/Misc/PreLoadFile.cpp
--- a/Engine/Source/Runtime/Core/Private/Misc/PreLoadFile.cpp
+++ b/Engine/Source/Runtime/Core/Private/Misc/PreLoadFile.cpp
@@ -103,7 +103,7 @@ void FPreLoadFile::KickOffRead()
 	else
 	{
 		// if we are failing a second time, then just mark the event complete, we have exhausted attempts
-		if (bFailedToOpenInKickOff)
+//		if (bFailedToOpenInKickOff)
 		{
 			CompletionEvent->Trigger();
 		}
