Turning this message off for now. If you are running debug or dev this ends up in RunGetVarArgsTest which tries to allocate memory.

diff --git a/Engine/Source/Runtime/Core/Private/Apple/ApplePlatformMemory.cpp b/Engine/Source/Runtime/Core/Private/Apple/ApplePlatformMemory.cpp
--- a/Engine/Source/Runtime/Core/Private/Apple/ApplePlatformMemory.cpp
+++ b/Engine/Source/Runtime/Core/Private/Apple/ApplePlatformMemory.cpp
@@ -252,14 +252,14 @@ void FApplePlatformMemory::NanoMallocInit()
 		}
 	}
 	
-	if(NanoRegionStart == 0 && NanoRegionEnd == 0)
-	{
-		FPlatformMisc::LowLevelOutputDebugStringf(TEXT("WARNING: No nano malloc region found. We will always use UE allocators\n"));
-	}
-	else
-	{
-		FPlatformMisc::LowLevelOutputDebugStringf(TEXT("Detected nanozone %p - %p\n"), (void*) NanoRegionStart, (void*) NanoRegionEnd);
-	}
+//	if(NanoRegionStart == 0 && NanoRegionEnd == 0)
+//	{
+//		FPlatformMisc::LowLevelOutputDebugStringf(TEXT("WARNING: No nano malloc region found. We will always use UE allocators\n"));
+//	}
+//	else
+//	{
+//		FPlatformMisc::LowLevelOutputDebugStringf(TEXT("Detected nanozone %p - %p\n"), (void*) NanoRegionStart, (void*) NanoRegionEnd);
+//	}
 }
 
 void FApplePlatformMemory::Init()
