RT polling fixes. Still disabled by default.
- Fixed assertion related to stall queue being used in stalling mode when the RT uses it in polling mode.
- Added csv stat to show the delay between FDrawSceneCommand enqueue and execution on the RT.
- Changed ProcessTasksNamedThread RT check to use ThreadId instead of IsInRenderingThread. This makes it consistent with the RT check in EnqueueFromOtherThread.

#rnx

diff --git a/Engine/Source/Runtime/Core/Private/Async/TaskGraph.cpp b/Engine/Source/Runtime/Core/Private/Async/TaskGraph.cpp
--- a/Engine/Source/Runtime/Core/Private/Async/TaskGraph.cpp
+++ b/Engine/Source/Runtime/Core/Private/Async/TaskGraph.cpp
@@ -680,9 +680,11 @@ public:
 			ProcessingTasks.Start(StatName);
 		}
 #endif
+		const bool bIsRenderThread = (ENamedThreads::GetThreadIndex(ThreadId) == ENamedThreads::GetRenderThread());
 		while (!Queue(QueueIndex).QuitForReturn)
 		{
-			FBaseGraphTask* Task = Queue(QueueIndex).StallQueue.Pop(0, bAllowStall);
+			const bool bStallQueueAllowStall = bAllowStall && (!bIsRenderThread || GDoRenderThreadWakeupTrigger);
+			FBaseGraphTask* Task = Queue(QueueIndex).StallQueue.Pop(0, bStallQueueAllowStall);
 			TestRandomizedThreads();
 			if (!Task)
 			{
@@ -697,7 +699,7 @@ public:
 				{
 					{
 						FScopeCycleCounter Scope(StallStatId);
-						Queue(QueueIndex).StallRestartEvent->Wait(IsInRenderingThread() ? GRenderThreadPollPeriodMs : MAX_uint32, bCountAsStall);
+						Queue(QueueIndex).StallRestartEvent->Wait(bIsRenderThread ? GRenderThreadPollPeriodMs : MAX_uint32, bCountAsStall);
 						if (Queue(QueueIndex).QuitForShutdown)
 						{
 							return ProcessedTasks;
@@ -777,7 +779,8 @@ public:
 		if (ThreadToStart >= 0)
 		{
 			checkThreadGraph(ThreadToStart == 0);
-			if ((ENamedThreads::GetThreadIndex(Task->ThreadToExecuteOn) != ENamedThreads::GetRenderThread()) || GDoRenderThreadWakeupTrigger)
+			const bool bExecuteOnRenderThread = (ENamedThreads::GetThreadIndex(Task->ThreadToExecuteOn) == ENamedThreads::GetRenderThread());
+			if (!bExecuteOnRenderThread || GDoRenderThreadWakeupTrigger)
 			{
 				QUICK_SCOPE_CYCLE_COUNTER(STAT_TaskGraph_EnqueueFromOtherThread_Trigger);
 				TASKGRAPH_SCOPE_CYCLE_COUNTER(1, STAT_TaskGraph_EnqueueFromOtherThread_Trigger);
diff --git a/Engine/Source/Runtime/Renderer/Private/SceneRendering.cpp b/Engine/Source/Runtime/Renderer/Private/SceneRendering.cpp
--- a/Engine/Source/Runtime/Renderer/Private/SceneRendering.cpp
+++ b/Engine/Source/Runtime/Renderer/Private/SceneRendering.cpp
@@ -3716,9 +3716,13 @@ void FRendererModule::BeginRenderingViewFamily(FCanvas* Canvas, FSceneViewFamily
 
 		SceneRenderer->ViewFamily.DisplayInternalsData.Setup(World);
 
+		const uint32 DrawSceneEnqueue = FPlatformTime::Cycles();
 		ENQUEUE_RENDER_COMMAND(FDrawSceneCommand)(
-			[SceneRenderer](FRHICommandListImmediate& RHICmdList)
+			[SceneRenderer, DrawSceneEnqueue](FRHICommandListImmediate& RHICmdList)
 			{
+				const float StartDelayMillisec = FPlatformTime::ToMilliseconds(FPlatformTime::Cycles() - DrawSceneEnqueue);
+				CSV_CUSTOM_STAT_GLOBAL(DrawSceneCommand_StartDelay, StartDelayMillisec, ECsvCustomStatOp::Set);
+
 				RenderViewFamily_RenderThread(RHICmdList, SceneRenderer);
 				FlushPendingDeleteRHIResources_RenderThread();
 			});
