Fix programs under the Engine/Restricted folder not being able to find the correct engine directory.

diff --git a/Engine/Source/Programs/UnrealBuildTool/Configuration/UEBuildTarget.cs b/Engine/Source/Programs/UnrealBuildTool/Configuration/UEBuildTarget.cs
--- a/Engine/Source/Programs/UnrealBuildTool/Configuration/UEBuildTarget.cs
+++ b/Engine/Source/Programs/UnrealBuildTool/Configuration/UEBuildTarget.cs
@@ -3540,9 +3540,15 @@ namespace UnrealBuildTool
 			// By default, shadow source files for this target in the root OutputDirectory
 			GlobalLinkEnvironment.LocalShadowDirectory = GlobalLinkEnvironment.OutputDirectory;
 
-			if(!String.IsNullOrEmpty(Rules.ExeBinariesSubFolder))
+			DirectoryReference OutputDir = Binaries[0].OutputDir;
+			if (OutputDir.IsUnderDirectory(UnrealBuildTool.EngineDirectory))
 			{
-				GlobalCompileEnvironment.Definitions.Add(String.Format("ENGINE_BASE_DIR_ADJUST={0}", Rules.ExeBinariesSubFolder.Replace('\\', '/').Trim('/').Count(x => x == '/') + 1));
+				DirectoryReference BaseDir = DirectoryReference.Combine(UnrealBuildTool.EngineDirectory, "Binaries", Platform.ToString());
+				if (BaseDir != OutputDir)
+				{
+					string RelativeBaseDir = BaseDir.MakeRelativeTo(OutputDir).Replace(Path.DirectorySeparatorChar, '/');
+					GlobalCompileEnvironment.Definitions.Add(String.Format("UE_RELATIVE_BASE_DIR=\"{0}\"", RelativeBaseDir));
+				}
 			}
 
 			if (Rules.bForceCompileDevelopmentAutomationTests)
diff --git a/Engine/Source/Runtime/Core/Private/Windows/WindowsPlatformProcess.cpp b/Engine/Source/Runtime/Core/Private/Windows/WindowsPlatformProcess.cpp
--- a/Engine/Source/Runtime/Core/Private/Windows/WindowsPlatformProcess.cpp
+++ b/Engine/Source/Runtime/Core/Private/Windows/WindowsPlatformProcess.cpp
@@ -924,12 +924,9 @@ const TCHAR* FWindowsPlatformProcess::BaseDir()
 			GetModuleFileName(hCurrentModule, Result, UE_ARRAY_COUNT(Result));
 			FString TempResult(Result);
 			TempResult = TempResult.Replace(TEXT("\\"), TEXT("/"));
+
 			FCString::Strcpy(Result, *TempResult);
 			int32 StringLength = FCString::Strlen(Result);
-			int32 NumSubDirectories = 0;
-#ifdef ENGINE_BASE_DIR_ADJUST
-			NumSubDirectories = ENGINE_BASE_DIR_ADJUST;
-#endif
 			if (StringLength > 0)
 			{
 				--StringLength;
@@ -937,16 +934,16 @@ const TCHAR* FWindowsPlatformProcess::BaseDir()
 				{
 					if (Result[StringLength - 1] == TEXT('/') || Result[StringLength - 1] == TEXT('\\'))
 					{
-						if(--NumSubDirectories < 0) //-V547
-						{
-							break;
-						}
+						break;
 					}
 				}
 			}
 			Result[StringLength] = 0;
 
 			FString CollapseResult(Result);
+#ifdef UE_RELATIVE_BASE_DIR
+			CollapseResult /= UE_RELATIVE_BASE_DIR;
+#endif
 			FPaths::CollapseRelativeDirectories(CollapseResult);
 			FCString::Strcpy(Result, *CollapseResult);
 		}
