Adding iPhoneSE2 device profile. Defaulting to the iPhone11.

diff --git a/Engine/Config/BaseDeviceProfiles.ini b/Engine/Config/BaseDeviceProfiles.ini
--- a/Engine/Config/BaseDeviceProfiles.ini
+++ b/Engine/Config/BaseDeviceProfiles.ini
@@ -34,6 +34,7 @@
 +DeviceProfileNameAndTypes=iPhone11,IOS
 +DeviceProfileNameAndTypes=iPhone11Pro,IOS
 +DeviceProfileNameAndTypes=iPhone11ProMax,IOS
++DeviceProfileNameAndTypes=iPhoneSE2,IOS
 +DeviceProfileNameAndTypes=iPadPro,IOS
 +DeviceProfileNameAndTypes=iPadPro105,IOS
 +DeviceProfileNameAndTypes=iPadPro129,IOS
@@ -320,6 +321,10 @@ DeviceType=IOS
 BaseProfileName=iPhone8
 +CVars=r.CustomUnsafeZones="(L:free[0,-15][812,15]);(P:fixed[83,0][206,30])"
 
+[iPhoneSE2 DeviceProfile]
+DeviceType=IOS
+BaseProfileName=iPhone11
+
 [iPadPro DeviceProfile]
 DeviceType=IOS
 BaseProfileName=IOS
diff --git a/Engine/Source/Runtime/ApplicationCore/Private/IOS/IOSPlatformApplicationMisc.cpp b/Engine/Source/Runtime/ApplicationCore/Private/IOS/IOSPlatformApplicationMisc.cpp
--- a/Engine/Source/Runtime/ApplicationCore/Private/IOS/IOSPlatformApplicationMisc.cpp
+++ b/Engine/Source/Runtime/ApplicationCore/Private/IOS/IOSPlatformApplicationMisc.cpp
@@ -282,6 +282,7 @@ EScreenPhysicalAccuracy FIOSPlatformApplicationMisc::ComputePhysicalScreenDensit
 	case FPlatformMisc::IOS_IPhone8:
     case FPlatformMisc::IOS_IPhoneXR:
 	case FPlatformMisc::IOS_IPhone11:
+	case FPlatformMisc::IOS_IPhoneSE2:
 		ScreenDensity = 326;
 		Accuracy = EScreenPhysicalAccuracy::Truth;
 		break;
diff --git a/Engine/Source/Runtime/Core/Private/IOS/IOSPlatformMisc.cpp b/Engine/Source/Runtime/Core/Private/IOS/IOSPlatformMisc.cpp
--- a/Engine/Source/Runtime/Core/Private/IOS/IOSPlatformMisc.cpp
+++ b/Engine/Source/Runtime/Core/Private/IOS/IOSPlatformMisc.cpp
@@ -619,6 +619,10 @@ FIOSPlatformMisc::EIOSDevice FIOSPlatformMisc::GetIOSDeviceType()
 			{
 				DeviceType = IOS_IPhone11ProMax;
 			}
+			else if (Minor == 8)
+			{
+				DeviceType = IOS_IPhoneSE2;
+			}
 		}
 		else if (Major >= 13)
 		{
diff --git a/Engine/Source/Runtime/Core/Public/IOS/IOSPlatformMisc.h b/Engine/Source/Runtime/Core/Public/IOS/IOSPlatformMisc.h
--- a/Engine/Source/Runtime/Core/Public/IOS/IOSPlatformMisc.h
+++ b/Engine/Source/Runtime/Core/Public/IOS/IOSPlatformMisc.h
@@ -155,6 +155,7 @@ struct CORE_API FIOSPlatformMisc : public FApplePlatformMisc
         IOS_IPadMini5,
 		IOS_IPodTouch7,
 		IOS_IPad7,
+		IOS_IPhoneSE2,
 		IOS_Unknown,
 	};
 
@@ -209,6 +210,7 @@ struct CORE_API FIOSPlatformMisc : public FApplePlatformMisc
             TEXT("IPadMini5"),
 			TEXT("IPodTouch7"),
 			TEXT("IPad7"),
+			TEXT("IPhoneSE2"),
 			TEXT("Unknown"),
 		};
 		static_assert((sizeof(IOSDeviceNames) / sizeof(IOSDeviceNames[0])) == ((int32)IOS_Unknown + 1), "Mismatched IOSDeviceNames and EIOSDevice.");
