- Don't scan for .uplugin files for platformextensions where there is no support for the platform compiled in to the editor. This helps with build machines that don't have NDAd platform support that makes precompiled binaries, but users sync those platforms so have the .uplugin locally
- Optimized GetExtensionDirs()
- Added a FDataDrivenPlatformInfoRegistry::HasCompiledSupportForPlatform() function that can be used in other places to detect if a platform has been compiled or not

diff --git a/Engine/Source/Developer/DesktopPlatform/Private/PlatformInfo.cpp b/Engine/Source/Developer/DesktopPlatform/Private/PlatformInfo.cpp
--- a/Engine/Source/Developer/DesktopPlatform/Private/PlatformInfo.cpp
+++ b/Engine/Source/Developer/DesktopPlatform/Private/PlatformInfo.cpp
@@ -6,6 +6,7 @@
 #include "HAL/FileManager.h"
 #include "Misc/ConfigCacheIni.h"
 #include "Misc/Paths.h"
+#include "Misc/DelayedAutoRegister.h"
 
 #define LOCTEXT_NAMESPACE "PlatformInfo"
 
@@ -202,18 +203,12 @@ void LoadDataDrivenPlatforms()
 	}
 }
 
-struct FPlatformInfoAutoInit
+FDelayedAutoRegisterHelper GPlatformInfoInit(EDelayedRegisterRunPhase::FileSystemReady, []()
 {
-	FPlatformInfoAutoInit()
-	{
-		FCoreDelegates::ConfigReadyForUse.AddLambda([]
-		{
-			BuildHardcodedPlatforms();
-			LoadDataDrivenPlatforms();
-		});
-	}
+	BuildHardcodedPlatforms();
+	LoadDataDrivenPlatforms();
+});
 
-} GPlatformInfoAutoInit;
 #endif
 
 } // anonymous namespace
diff --git a/Engine/Source/Runtime/Core/Private/Misc/DataDrivenPlatformInfoRegistry.cpp b/Engine/Source/Runtime/Core/Private/Misc/DataDrivenPlatformInfoRegistry.cpp
--- a/Engine/Source/Runtime/Core/Private/Misc/DataDrivenPlatformInfoRegistry.cpp
+++ b/Engine/Source/Runtime/Core/Private/Misc/DataDrivenPlatformInfoRegistry.cpp
@@ -5,6 +5,7 @@
 #include "HAL/FileManager.h"
 #include "Misc/FileHelper.h"
 #include "Misc/ConfigCacheIni.h"
+#include "Modules/ModuleManager.h"
 
 
 
@@ -167,6 +168,31 @@ static void LoadDDPIIniSettings(const FConfigFile& IniFile, FDataDrivenPlatformI
 	DDPIGetBool(IniFile, TEXT("Freezing_bWithRayTracing"), Info.Freezing_bWithRayTracing);
 
 	// NOTE: add more settings here!
+
+
+#if DDPI_HAS_EXTENDED_PLATFORMINFO_DATA
+	// now look in the PlatformInfo objects in the file for TP and UBT names
+
+	FName TargetPlatformKey("TargetPlatformName");
+	FName UBTNameKey("UBTTargetID");
+	for (auto& Pair : IniFile)
+	{
+		if (Pair.Key.StartsWith(TEXT("PlatformInfo")))
+		{
+			const FConfigValue* PlatformName = Pair.Value.Find(TargetPlatformKey);
+			if (PlatformName)
+			{
+				Info.AllTargetPlatformNames.AddUnique(PlatformName->GetValue());
+			}
+			
+			PlatformName = Pair.Value.Find(UBTNameKey);
+			if (PlatformName)
+			{
+				Info.AllUBTPlatformNames.AddUnique(PlatformName->GetValue());
+			}
+		}
+	}
+#endif
 }
 
 
@@ -222,6 +248,41 @@ const TMap<FString, FDataDrivenPlatformInfoRegistry::FPlatformInfo>& FDataDriven
 	return DataDrivenPlatforms;
 }
 
+const TArray<FString>& FDataDrivenPlatformInfoRegistry::GetValidPlatformDirectoryNames()
+{
+	static bool bHasSearchedForPlatforms = false;
+	static TArray<FString> ValidPlatformDirectories;
+
+	if (bHasSearchedForPlatforms == false)
+	{
+		bHasSearchedForPlatforms = true;
+
+		// look for possible platforms
+		const TMap<FString, FDataDrivenPlatformInfoRegistry::FPlatformInfo>& Infos = GetAllPlatformInfos();
+		for (auto Pair : Infos)
+		{
+#if DDPI_HAS_EXTENDED_PLATFORMINFO_DATA
+			// if the editor hasn't compiled in support for the platform, it's not "valid"
+			if (!HasCompiledSupportForPlatform(Pair.Key, EPlatformNameType::Ini))
+			{
+				continue;
+			}
+#endif
+
+			// add ourself as valid
+			ValidPlatformDirectories.AddUnique(Pair.Key);
+
+			// now add additional directories
+			for (FString& AdditionalDir : Pair.Value.AdditionalRestrictedFolders)
+			{
+				ValidPlatformDirectories.AddUnique(AdditionalDir);
+			}
+		}
+	}
+
+	return ValidPlatformDirectories;
+}
+
 
 const FDataDrivenPlatformInfoRegistry::FPlatformInfo& FDataDrivenPlatformInfoRegistry::GetPlatformInfo(const FString& PlatformName)
 {
@@ -252,4 +313,51 @@ const TArray<FString>& FDataDrivenPlatformInfoRegistry::GetConfidentialPlatforms
 
 	// return whatever we have already found
 	return FoundPlatforms;
-}
\ No newline at end of file
+}
+
+
+#if DDPI_HAS_EXTENDED_PLATFORMINFO_DATA
+bool FDataDrivenPlatformInfoRegistry::HasCompiledSupportForPlatform(const FString& PlatformName, EPlatformNameType PlatformNameType)
+{
+	if (PlatformNameType == EPlatformNameType::Ini)
+	{
+		// get the DDPI info object
+		const FPlatformInfo& Info = GetPlatformInfo(PlatformName);
+
+		// look to see if any of the TPs in the Info are valid - if at least one is, we are good
+		for (const FString& TPName : Info.AllTargetPlatformNames)
+		{
+			if (HasCompiledSupportForPlatform(TPName, EPlatformNameType::TargetPlatform))
+			{
+				return true;
+			}
+		}
+		return false;
+
+	}
+	else if (PlatformNameType == EPlatformNameType::UBT)
+	{
+		FName PlatformFName(*PlatformName);
+
+		// find all the DataDrivenPlatformInfo objects and find a matching the UBT name
+		for (auto& Pair : FDataDrivenPlatformInfoRegistry::GetAllPlatformInfos())
+		{
+			// if this platform contains the UBT platform name, then check the info for it's TPs
+			// (we could be tricky and match UBT platforms with TPs just for thse UBT platforms, but that complexity does not seem needed)
+			if (Pair.Value.AllUBTPlatformNames.Contains(PlatformName))
+			{
+				return HasCompiledSupportForPlatform(Pair.Key, EPlatformNameType::Ini);
+			}
+		}
+
+		return false;
+	}
+	else if (PlatformNameType == EPlatformNameType::TargetPlatform)
+	{
+		// was this TP compiled?
+		return FModuleManager::Get().ModuleExists(*FString::Printf(TEXT("%sTargetPlatform"), *PlatformName));
+	}
+
+	return false;
+}
+#endif
diff --git a/Engine/Source/Runtime/Core/Private/Misc/Paths.cpp b/Engine/Source/Runtime/Core/Private/Misc/Paths.cpp
--- a/Engine/Source/Runtime/Core/Private/Misc/Paths.cpp
+++ b/Engine/Source/Runtime/Core/Private/Misc/Paths.cpp
@@ -257,19 +257,14 @@ static void AddIfDirectoryExists(TArray<FString>& ExtensionDirs, FString&& Dir)
 	}
 }
 
-TArray<FString> FPaths::GetExtensionDirs(const FString& BaseDir, const FString& SubDir)
+static void GetExtensionDirsInternal(TArray<FString>& ExtensionDirs, const FString& BaseDir, const FString& SubDir)
 {
-	TArray<FString> ExtensionDirs;
 	AddIfDirectoryExists(ExtensionDirs, FPaths::Combine(BaseDir, SubDir));
 
 	FString PlatformExtensionBaseDir = FPaths::Combine(BaseDir, TEXT("Platforms"));
-	for (auto const& Platform : FDataDrivenPlatformInfoRegistry::GetAllPlatformInfos())
+	for (const FString& PlatformName : FDataDrivenPlatformInfoRegistry::GetValidPlatformDirectoryNames())
 	{
-		for (auto const& AdditionalFolder : Platform.Value.AdditionalRestrictedFolders)
-		{
-			AddIfDirectoryExists(ExtensionDirs, FPaths::Combine(PlatformExtensionBaseDir, AdditionalFolder, SubDir));
-		}
-		AddIfDirectoryExists(ExtensionDirs, FPaths::Combine(PlatformExtensionBaseDir, Platform.Key, SubDir));
+		AddIfDirectoryExists(ExtensionDirs, FPaths::Combine(PlatformExtensionBaseDir, PlatformName, SubDir));
 	}
 
 	FString RestrictedBaseDir = FPaths::Combine(BaseDir, TEXT("Restricted"));
@@ -280,12 +275,28 @@ TArray<FString> FPaths::GetExtensionDirs(const FString& BaseDir, const FString&
 			AddIfDirectoryExists(ExtensionDirs, FPaths::Combine(FilenameOrDirectory, SubDir));
 
 			// now look for platforms under it
-			ExtensionDirs.Append(GetExtensionDirs(FilenameOrDirectory, SubDir));
+			GetExtensionDirsInternal(ExtensionDirs, FilenameOrDirectory, SubDir);
 		}
 		return true;
 	});
+}
+
+const TArray<FString>& FPaths::GetExtensionDirs(const FString& BaseDir, const FString& SubDir)
+{
+	static TMap<TTuple<FString, FString>, TArray<FString>> CachedExtensionDirs;
+
+	// look up a cached set of this BaseDir/SubDir combo
+	TTuple<FString, FString> DirKey(BaseDir, SubDir);
+	TArray<FString>* CachedDirs = CachedExtensionDirs.Find(DirKey);
+	if (CachedDirs == nullptr)
+	{
+		CachedDirs = &CachedExtensionDirs.Emplace(DirKey);
+
+		// get the directories
+		GetExtensionDirsInternal(*CachedDirs, BaseDir, SubDir);
+	}
 
-	return ExtensionDirs;
+	return *CachedDirs;
 }
 
 FString FPaths::RootDir()
diff --git a/Engine/Source/Runtime/Core/Public/Misc/DataDrivenPlatformInfoRegistry.h b/Engine/Source/Runtime/Core/Public/Misc/DataDrivenPlatformInfoRegistry.h
--- a/Engine/Source/Runtime/Core/Public/Misc/DataDrivenPlatformInfoRegistry.h
+++ b/Engine/Source/Runtime/Core/Public/Misc/DataDrivenPlatformInfoRegistry.h
@@ -7,6 +7,7 @@
 #include "Containers/Map.h"
 #include "Containers/Array.h"
 
+#define DDPI_HAS_EXTENDED_PLATFORMINFO_DATA WITH_EDITOR && !IS_MONOLITHIC
 
 struct CORE_API FDataDrivenPlatformInfoRegistry
 {
@@ -37,6 +38,13 @@ struct CORE_API FDataDrivenPlatformInfoRegistry
 
 		// NOTE: add more settings here (and read them in in the LoadDDPIIniSettings() function in the .cpp)
 
+
+#if DDPI_HAS_EXTENDED_PLATFORMINFO_DATA
+		// list of TP names (WindowsNoEditor, etc)
+		TArray<FString> AllTargetPlatformNames;
+		// list of UBT platform names (Win32, Win64, etc)
+		TArray<FString> AllUBTPlatformNames;
+#endif
 	};
 
 	/**
@@ -44,6 +52,12 @@ struct CORE_API FDataDrivenPlatformInfoRegistry
 	*/
 	static const TMap<FString, FDataDrivenPlatformInfoRegistry::FPlatformInfo>& GetAllPlatformInfos();
 
+	/**
+	 * Gets a set of platform names based on GetAllPlatformInfos, their AdditionalRestrictedFolders, and possibly filtered based on what editor has support compiled for
+	 * This is not necessarily the same as IniParents, although there is overlap - IniParents come from chaining DDPIs, so those will be in GetAllPlatformInfos already to be checked 
+	 */
+	static const TArray<FString>& GetValidPlatformDirectoryNames();
+
 	/**
 	 * Get the data driven platform info for a given platform. If the platform doesn't have any on disk,
 	 * this will return a default constructed FConfigDataDrivenPlatformInfo
@@ -65,5 +79,23 @@ struct CORE_API FDataDrivenPlatformInfoRegistry
 	 * Load the given ini file, and 
 	 */
 	static bool LoadDataDrivenIniFile(int32 Index, FConfigFile& IniFile, FString& PlatformName);
+
+
+#if DDPI_HAS_EXTENDED_PLATFORMINFO_DATA
+	/**
+	 * Checks for the existence of compiled modules for a given (usually another, target, platform)
+	 * Since there are different types of platform names, it is necessary pass in the type of name
+	 */
+	enum class EPlatformNameType
+	{
+		// for instance Win64
+		UBT,
+		// for instance Windows
+		Ini,
+		// for instance WindowsNoEditor
+		TargetPlatform,
+	};
+	static bool HasCompiledSupportForPlatform(const FString& PlatformName, EPlatformNameType PlatformNameType);
+#endif
 };
 
diff --git a/Engine/Source/Runtime/Core/Public/Misc/Paths.h b/Engine/Source/Runtime/Core/Public/Misc/Paths.h
--- a/Engine/Source/Runtime/Core/Public/Misc/Paths.h
+++ b/Engine/Source/Runtime/Core/Public/Misc/Paths.h
@@ -143,7 +143,7 @@ public:
 	 *
 	 * @return BaseDir and usable extension directories under BaseDir (either Engine or Project)
 	 */
-	static TArray<FString> GetExtensionDirs(const FString& BaseDir, const FString& SubDir=FString());
+	static const TArray<FString>& GetExtensionDirs(const FString& BaseDir, const FString& SubDir=FString());
 
 	/**
 	 * Returns the root directory of the engine directory tree
