Version bump for underwater post process material refactor

kevin.ortegren

diff --git a/Engine/Source/Runtime/Core/Public/UObject/FortniteMainBranchObjectVersion.h b/Engine/Source/Runtime/Core/Public/UObject/FortniteMainBranchObjectVersion.h
--- a/Engine/Source/Runtime/Core/Public/UObject/FortniteMainBranchObjectVersion.h
+++ b/Engine/Source/Runtime/Core/Public/UObject/FortniteMainBranchObjectVersion.h
@@ -108,6 +108,9 @@ struct CORE_API FFortniteMainBranchObjectVersion
 		// FFieldPath will serialize the owner struct reference and only a short path to its property
 		FFieldPathOwnerSerialization,
 
+		// Simplified WaterBody post process material handling : 
+		FixUpUnderwaterPostProcessMaterial,
+
 		// -----<new versions can be added above this line>-------------------------------------------------
 		VersionPlusOne,
 		LatestVersion = VersionPlusOne - 1
