Fix CIS error/build break on HoloLens with new Octree code.

diff --git a/Engine/Source/Runtime/Core/Public/Math/GenericOctree.h b/Engine/Source/Runtime/Core/Public/Math/GenericOctree.h
--- a/Engine/Source/Runtime/Core/Public/Math/GenericOctree.h
+++ b/Engine/Source/Runtime/Core/Public/Math/GenericOctree.h
@@ -282,7 +282,15 @@ public:
 
 	inline VectorRegister GetChildOffsetVec(int i) const
 	{
-		union { VectorRegister v; VectorRegisterInt i; } Mask;
+		union MaskType { VectorRegister v;  VectorRegisterInt i; 
+#if PLATFORM_HOLOLENS
+			MaskType() 
+				: v(MakeVectorRegister(0.0f, 0.0f, 0.0f, 0.0f))
+				, i(MakeVectorRegisterInt(0, 0, 0, 0))
+				{}
+#endif
+		} Mask;
+
 		Mask.v = MakeVectorRegister(1u, 2u, 4u, 8u);
 		VectorRegisterInt X = VectorIntLoad1(&i);
 		VectorRegisterInt A = VectorIntAnd(X, Mask.i);
