TInvokeResult trait for getting the return type when invoking a callable with given argument types.

diff --git a/Engine/Source/Runtime/Core/Public/Templates/Invoke.h b/Engine/Source/Runtime/Core/Public/Templates/Invoke.h
--- a/Engine/Source/Runtime/Core/Public/Templates/Invoke.h
+++ b/Engine/Source/Runtime/Core/Public/Templates/Invoke.h
@@ -106,4 +106,31 @@ FORCEINLINE auto Invoke(ReturnType (ObjType::*PtrMemFun)(PMFArgTypes...) const,
 	[](auto&& Obj, auto&&... Args) \
 	{ \
 		return UE4Invoke_Private::DereferenceIfNecessary<Type>(Forward<decltype(Obj)>(Obj)).FuncName(Forward<decltype(Args)>(Args)...); \
-	}
\ No newline at end of file
+	}
+
+
+namespace UE4Invoke_Private
+{
+	template <typename, typename FuncType, typename... ArgTypes>
+	struct TInvokeResult_Impl
+	{
+	};
+
+	template <typename FuncType, typename... ArgTypes>
+	struct TInvokeResult_Impl<decltype((void)Invoke(DeclVal<FuncType>(), DeclVal<ArgTypes>()...)), FuncType, ArgTypes...>
+	{
+		using Type = decltype(Invoke(DeclVal<FuncType>(), DeclVal<ArgTypes>()...));
+	};
+}
+
+/**
+ * Trait for the type of the result when invoking a callable with the given argument types.
+ * Not defined (as thus usable in SFINAE contexts) when the callable cannot be invoked with the given argument types.
+ */
+template <typename FuncType, typename... ArgTypes>
+struct TInvokeResult : UE4Invoke_Private::TInvokeResult_Impl<void, FuncType, ArgTypes...>
+{
+};
+
+template <typename FuncType, typename... ArgTypes>
+using TInvokeResult_T = typename TInvokeResult<FuncType, ArgTypes...>::Type;
