- Fixed the USE_BOOT_PROFILING define to be back to 0 by default, it was on mistakenly in embedded apps

diff --git a/Engine/Source/Runtime/Core/Private/Misc/CoreGlobals.cpp b/Engine/Source/Runtime/Core/Private/Misc/CoreGlobals.cpp
--- a/Engine/Source/Runtime/Core/Private/Misc/CoreGlobals.cpp
+++ b/Engine/Source/Runtime/Core/Private/Misc/CoreGlobals.cpp
@@ -417,7 +417,7 @@ static struct FBootTimingStart
 } GBootTimingStart;
 
 
-#define USE_BOOT_PROFILING (BUILD_EMBEDDED_APP)
+#define USE_BOOT_PROFILING 0
 
 #if !USE_BOOT_PROFILING
 FScopedBootTiming::FScopedBootTiming(const ANSICHAR *InMessage)
