Fix issues causing material DDC values to be non-determinisic
- Shader code within resources sorted by hash
- Debug source code values sorted by type
Also add mechanism to serialize any MemoryLayout described type to string
- Makes it easier to inspect material data to find sources of non-determinism

diff --git a/Engine/Plugins/Compositing/OpenColorIO/Source/OpenColorIO/Private/OpenColorIOShaderMap.cpp b/Engine/Plugins/Compositing/OpenColorIO/Source/OpenColorIO/Private/OpenColorIOShaderMap.cpp
--- a/Engine/Plugins/Compositing/OpenColorIO/Source/OpenColorIO/Private/OpenColorIOShaderMap.cpp
+++ b/Engine/Plugins/Compositing/OpenColorIO/Source/OpenColorIO/Private/OpenColorIOShaderMap.cpp
@@ -237,18 +237,13 @@ FShaderCompileJob* FOpenColorIOShaderType::BeginCompileShader(
 FShader* FOpenColorIOShaderType::FinishCompileShader(
 	const FSHAHash& InShaderMapHash,
 	const FShaderCompileJob& InCurrentJob,
-	const FString& InDebugDescription,
-	FShaderMapResourceBuilder& ResourceBuilder
+	const FString& InDebugDescription
 	)
 {
 	check(InCurrentJob.bSucceeded);
 
-	// Reuse an existing resource with the same key or create a new one based on the compile output
-	// This allows FShaders to share compiled bytecode and RHI shader references
-	const int32 ResourceIndex = ResourceBuilder.FindOrAddCode(InCurrentJob.Output);
-
 	const int32 PermutationId = 0;
-	FShader* Shader = ConstructCompiled(FOpenColorIOShaderType::CompiledShaderInitializerType(this, PermutationId, InCurrentJob.Output, ResourceIndex, InShaderMapHash, InDebugDescription));
+	FShader* Shader = ConstructCompiled(FOpenColorIOShaderType::CompiledShaderInitializerType(this, PermutationId, InCurrentJob.Output, InShaderMapHash, InDebugDescription));
 	InCurrentJob.Output.ParameterMap.VerifyBindingsAreComplete(GetName(), InCurrentJob.Output.Target, InCurrentJob.VFType);
 
 	return Shader;
@@ -476,15 +471,17 @@ void FOpenColorIOShaderMap::Compile(FOpenColorIOTransformResource* InColorTransf
 	}
 }
 
-FShader* FOpenColorIOShaderMap::ProcessCompilationResultsForSingleJob(FShaderCompileJob& CurrentJob, const FSHAHash& InShaderMapHash, FShaderMapResourceBuilder& InResourceBuilder)
+FShader* FOpenColorIOShaderMap::ProcessCompilationResultsForSingleJob(FShaderCompileJob& CurrentJob, const FSHAHash& InShaderMapHash)
 {
 	check(CurrentJob.Id == CompilingId);
 
+	GetResourceCode()->AddShaderCompilerOutput(CurrentJob.Output);
+
 	FShader* Shader = nullptr;
 
 	FOpenColorIOShaderType* OpenColorIOShaderType = CurrentJob.ShaderType->GetOpenColorIOShaderType();
 	check(OpenColorIOShaderType);
-	Shader = OpenColorIOShaderType->FinishCompileShader(InShaderMapHash, CurrentJob, GetContent()->FriendlyName, InResourceBuilder);
+	Shader = OpenColorIOShaderType->FinishCompileShader(InShaderMapHash, CurrentJob, GetContent()->FriendlyName);
 	bCompiledSuccessfully = CurrentJob.bSucceeded;
 
 	FOpenColorIOPixelShader *OpenColorIOShader = static_cast<FOpenColorIOPixelShader*>(Shader);
@@ -502,10 +499,9 @@ bool FOpenColorIOShaderMap::ProcessCompilationResults(const TArray<TSharedRef<FS
 	FSHAHash ShaderMapHash;
 	GetContent()->ShaderMapId.GetOpenColorIOHash(ShaderMapHash);
 
-	FShaderMapResourceBuilder ResourceBuilder(GetResourceCode());
 	do
 	{
-		ProcessCompilationResultsForSingleJob(StaticCastSharedRef<FShaderCompileJob>(InCompilationResults[InOutJobIndex]).Get(), ShaderMapHash, ResourceBuilder);
+		ProcessCompilationResultsForSingleJob(StaticCastSharedRef<FShaderCompileJob>(InCompilationResults[InOutJobIndex]).Get(), ShaderMapHash);
 
 		InOutJobIndex++;
 		
diff --git a/Engine/Plugins/Compositing/OpenColorIO/Source/OpenColorIO/Public/OpenColorIOShaderType.h b/Engine/Plugins/Compositing/OpenColorIO/Source/OpenColorIO/Public/OpenColorIOShaderType.h
--- a/Engine/Plugins/Compositing/OpenColorIO/Source/OpenColorIO/Public/OpenColorIOShaderType.h
+++ b/Engine/Plugins/Compositing/OpenColorIO/Source/OpenColorIO/Public/OpenColorIOShaderType.h
@@ -54,11 +54,10 @@ public:
 			FShaderType* InType,
 			int32 InPermutationId,
 			const FShaderCompilerOutput& CompilerOutput,
-			int32 InResourceIndex,
 			const FSHAHash& InOCIOShaderMapHash,
 			const FString& InDebugDescription
 			)
-		: FGlobalShaderType::CompiledShaderInitializerType(InType,InPermutationId,CompilerOutput, InResourceIndex, InOCIOShaderMapHash,nullptr,nullptr)
+		: FGlobalShaderType::CompiledShaderInitializerType(InType,InPermutationId,CompilerOutput, InOCIOShaderMapHash,nullptr,nullptr)
 		, DebugDescription(InDebugDescription)
 		{}
 	};
@@ -110,8 +109,7 @@ public:
 	FShader* FinishCompileShader(
 		const FSHAHash& InOCIOShaderMapHash,
 		const FShaderCompileJob& CurrentJob,
-		const FString& InDebugDescription,
-		FShaderMapResourceBuilder& ResourceBuilder
+		const FString& InDebugDescription
 		);
 
 	/**
diff --git a/Engine/Plugins/Compositing/OpenColorIO/Source/OpenColorIO/Public/OpenColorIOShared.h b/Engine/Plugins/Compositing/OpenColorIO/Source/OpenColorIO/Public/OpenColorIOShared.h
--- a/Engine/Plugins/Compositing/OpenColorIO/Source/OpenColorIO/Public/OpenColorIOShared.h
+++ b/Engine/Plugins/Compositing/OpenColorIO/Source/OpenColorIO/Public/OpenColorIOShared.h
@@ -280,7 +280,7 @@ private:
 
 	uint32 bHasFrozenContent : 1;
 
-	FShader* ProcessCompilationResultsForSingleJob(FShaderCompileJob& InSingleJob, const FSHAHash& InShaderMapHash, FShaderMapResourceBuilder& InResourceBuilder);
+	FShader* ProcessCompilationResultsForSingleJob(FShaderCompileJob& InSingleJob, const FSHAHash& InShaderMapHash);
 
 	bool IsOpenColorIOShaderComplete(const FOpenColorIOTransformResource* InColorTransform, const FOpenColorIOShaderType* InShaderType, bool bSilent);
 
diff --git a/Engine/Plugins/FX/Niagara/Source/NiagaraShader/Private/NiagaraShader.cpp b/Engine/Plugins/FX/Niagara/Source/NiagaraShader/Private/NiagaraShader.cpp
--- a/Engine/Plugins/FX/Niagara/Source/NiagaraShader/Private/NiagaraShader.cpp
+++ b/Engine/Plugins/FX/Niagara/Source/NiagaraShader/Private/NiagaraShader.cpp
@@ -486,16 +486,11 @@ FShader* FNiagaraShaderType::FinishCompileShader(
 	//const FUniformExpressionSet& UniformExpressionSet,
 	const FSHAHash& ShaderMapHash,
 	const FShaderCompileJob& CurrentJob,
-	const FString& InDebugDescription,
-	FShaderMapResourceBuilder& ResourceBuilder
+	const FString& InDebugDescription
 	)
 {
 	check(CurrentJob.bSucceeded);
 
-	// Reuse an existing resource with the same key or create a new one based on the compile output
-	// This allows FShaders to share compiled bytecode and RHI shader references
-	const int32 ResourceIndex = ResourceBuilder.FindOrAddCode(CurrentJob.Output);
-
 	const int32 PermutationId = 0;
 	//CurrentJob.Id
 	TArray<FNiagaraDataInterfaceGPUParamInfo> DIParamInfo;
@@ -504,7 +499,7 @@ FShader* FNiagaraShaderType::FinishCompileShader(
 		check(false);
 	}
 
-	FShader* Shader = ConstructCompiled(FNiagaraShaderType::CompiledShaderInitializerType(this, PermutationId, CurrentJob.Output, ResourceIndex, ShaderMapHash, InDebugDescription, DIParamInfo));
+	FShader* Shader = ConstructCompiled(FNiagaraShaderType::CompiledShaderInitializerType(this, PermutationId, CurrentJob.Output, ShaderMapHash, InDebugDescription, DIParamInfo));
 	//CurrentJob.Output.ParameterMap.VerifyBindingsAreComplete(GetName(), CurrentJob.Output.Target, nullptr); // b/c we don't bind data interfaces yet...
 
 	return Shader;
@@ -779,16 +774,18 @@ void FNiagaraShaderMap::Compile(
 	}
 }
 
-FShader* FNiagaraShaderMap::ProcessCompilationResultsForSingleJob(TSharedRef<FShaderCommonCompileJob, ESPMode::ThreadSafe> SingleJob, const FSHAHash& ShaderMapHash, FShaderMapResourceBuilder& InResourceBuilder)
+FShader* FNiagaraShaderMap::ProcessCompilationResultsForSingleJob(TSharedRef<FShaderCommonCompileJob, ESPMode::ThreadSafe> SingleJob, const FSHAHash& ShaderMapHash)
 {
 	TSharedRef<FShaderCompileJob, ESPMode::ThreadSafe> CurrentJob = StaticCastSharedRef<FShaderCompileJob>(SingleJob);
 	check(CurrentJob->Id == CompilingId);
 
+	GetResourceCode()->AddShaderCompilerOutput(CurrentJob->Output);
+
 	FShader* Shader = nullptr;
 
 	FNiagaraShaderType* NiagaraShaderType = CurrentJob->ShaderType->GetNiagaraShaderType();
 	check(NiagaraShaderType);
-	Shader = NiagaraShaderType->FinishCompileShader(ShaderMapHash, *CurrentJob, GetContent()->FriendlyName, InResourceBuilder);
+	Shader = NiagaraShaderType->FinishCompileShader(ShaderMapHash, *CurrentJob, GetContent()->FriendlyName);
 	bCompiledSuccessfully = CurrentJob->bSucceeded;
 
 	FNiagaraShader *NiagaraShader = static_cast<FNiagaraShader*>(Shader);
@@ -808,11 +805,10 @@ bool FNiagaraShaderMap::ProcessCompilationResults(const TArray<TSharedRef<FShade
 	FSHAHash ShaderMapHash;
 	GetContent()->ShaderMapId.GetScriptHash(ShaderMapHash);
 
-	FShaderMapResourceBuilder ResourceBuilder(GetResourceCode());
 	do
 	{
 		{
-			ProcessCompilationResultsForSingleJob(InCompilationResults[InOutJobIndex], ShaderMapHash, ResourceBuilder);
+			ProcessCompilationResultsForSingleJob(InCompilationResults[InOutJobIndex], ShaderMapHash);
 		}
 
 		InOutJobIndex++;
diff --git a/Engine/Plugins/FX/Niagara/Source/NiagaraShader/Public/NiagaraShaderType.h b/Engine/Plugins/FX/Niagara/Source/NiagaraShader/Public/NiagaraShaderType.h
--- a/Engine/Plugins/FX/Niagara/Source/NiagaraShader/Public/NiagaraShaderType.h
+++ b/Engine/Plugins/FX/Niagara/Source/NiagaraShader/Public/NiagaraShaderType.h
@@ -63,12 +63,11 @@ public:
 			FShaderType* InType,
 			int32 InPermutationId,
 			const FShaderCompilerOutput& CompilerOutput,
-			int32 InResourceIndex,
 			const FSHAHash& InNiagaraShaderMapHash,
 			const FString& InDebugDescription,
 			const TArray< FNiagaraDataInterfaceGPUParamInfo > &InDIParamInfo
 			)
-		: FGlobalShaderType::CompiledShaderInitializerType(InType,InPermutationId,CompilerOutput, InResourceIndex, InNiagaraShaderMapHash,nullptr,nullptr)
+		: FGlobalShaderType::CompiledShaderInitializerType(InType,InPermutationId,CompilerOutput, InNiagaraShaderMapHash,nullptr,nullptr)
 		, DebugDescription(InDebugDescription)
 		, DIParamInfo(InDIParamInfo)
 		{}
@@ -123,8 +122,7 @@ public:
 	FShader* FinishCompileShader(
 		const FSHAHash& NiagaraShaderMapHash,
 		const FShaderCompileJob& CurrentJob,
-		const FString& InDebugDescription,
-		FShaderMapResourceBuilder& ResourceBuilder
+		const FString& InDebugDescription
 		);
 
 	/**
diff --git a/Engine/Plugins/FX/Niagara/Source/NiagaraShader/Public/NiagaraShared.h b/Engine/Plugins/FX/Niagara/Source/NiagaraShader/Public/NiagaraShared.h
--- a/Engine/Plugins/FX/Niagara/Source/NiagaraShader/Public/NiagaraShared.h
+++ b/Engine/Plugins/FX/Niagara/Source/NiagaraShader/Public/NiagaraShared.h
@@ -544,7 +544,7 @@ private:
 	/** Indicates whether the shader map should be stored in the shader cache. */
 	uint32 bIsPersistent : 1;
 
-	FShader* ProcessCompilationResultsForSingleJob(TSharedRef<class FShaderCommonCompileJob, ESPMode::ThreadSafe> SingleJob, const FSHAHash& ShaderMapHash, FShaderMapResourceBuilder& InResourceBuilder);
+	FShader* ProcessCompilationResultsForSingleJob(TSharedRef<class FShaderCommonCompileJob, ESPMode::ThreadSafe> SingleJob, const FSHAHash& ShaderMapHash);
 
 	bool IsNiagaraShaderComplete(const FNiagaraShaderScript* Script, const FNiagaraShaderType* ShaderType, bool bSilent);
 
diff --git a/Engine/Shaders/Public/ShaderVersion.ush b/Engine/Shaders/Public/ShaderVersion.ush
--- a/Engine/Shaders/Public/ShaderVersion.ush
+++ b/Engine/Shaders/Public/ShaderVersion.ush
@@ -3,4 +3,4 @@
 // in Platform.ush (which should be included in any shader) it allows to invalidate the shader DDC.
 // 
 // If you are merging streams and there is a conflict with this GUID you should make a new GUID rather than taking one or the other.
-// GUID = 90D1B596-2512-4A90-9FE2-F88BD15F9E7A
\ No newline at end of file
+// GUID = C721CC12-7BBF-48B2-97AC-D1B54745EDD4
\ No newline at end of file
diff --git a/Engine/Source/Runtime/Core/Private/Containers/MemoryImageAllocator.cpp b/Engine/Source/Runtime/Core/Private/Containers/MemoryImageAllocator.cpp
--- a/Engine/Source/Runtime/Core/Private/Containers/MemoryImageAllocator.cpp
+++ b/Engine/Source/Runtime/Core/Private/Containers/MemoryImageAllocator.cpp
@@ -2,6 +2,7 @@
 
 #include "Containers/ContainerAllocationPolicies.h"
 #include "Containers/ResourceArray.h"
+#include "Misc/StringBuilder.h"
 #include "Serialization/MemoryImage.h"
 
 IMPLEMENT_ABSTRACT_TYPE_LAYOUT(FResourceArrayInterface);
@@ -82,3 +83,25 @@ void FMemoryImageAllocatorBase::CopyUnfrozen(const FMemoryUnfreezeContent& Conte
 		}
 	}
 }
+
+void FMemoryImageAllocatorBase::ToString(const FTypeLayoutDesc& TypeDesc, int32 NumAllocatedElements, const FPlatformTypeLayoutParameters& LayoutParams, FMemoryToStringContext& OutContext) const
+{
+	OutContext.String->Appendf(TEXT("TArray<%s>, Num: %d\n"), TypeDesc.Name, NumAllocatedElements);
+	++OutContext.Indent;
+
+	const void* RawPtr = GetAllocation();
+	FTypeLayoutDesc::FToStringFunc* Func = TypeDesc.ToStringFunc;
+	const uint32 ElementSize = TypeDesc.Size;
+
+	for (int32 i = 0; i < NumAllocatedElements; ++i)
+	{
+		OutContext.AppendIndent();
+		OutContext.String->Appendf(TEXT("[%d]: "), i);
+		Func((uint8*)RawPtr + ElementSize * i,
+			TypeDesc,
+			LayoutParams,
+			OutContext);
+	}
+
+	--OutContext.Indent;
+}
diff --git a/Engine/Source/Runtime/Core/Private/Misc/SecureHash.cpp b/Engine/Source/Runtime/Core/Private/Misc/SecureHash.cpp
--- a/Engine/Source/Runtime/Core/Private/Misc/SecureHash.cpp
+++ b/Engine/Source/Runtime/Core/Private/Misc/SecureHash.cpp
@@ -1,6 +1,7 @@
 // Copyright Epic Games, Inc. All Rights Reserved.
 
 #include "Misc/SecureHash.h"
+#include "Misc/StringBuilder.h"
 #include "HAL/FileManager.h"
 #include "Misc/Paths.h"
 
@@ -615,6 +616,11 @@ void LexFromString(FSHAHash& InHash, const TCHAR* InString)
 	InHash.FromString(InString);
 }
 
+void Freeze::IntrinsicToString(const FSHAHash& Object, const FTypeLayoutDesc& TypeDesc, const FPlatformTypeLayoutParameters& LayoutParams, FMemoryToStringContext& OutContext)
+{
+	OutContext.String->Appendf(TEXT("%s\n"), *Object.ToString());
+}
+
 FSHA1::FSHA1()
 {
 	m_block = (SHA1_WORKSPACE_BLOCK *)m_workspace;
diff --git a/Engine/Source/Runtime/Core/Private/Serialization/MemoryImage.cpp b/Engine/Source/Runtime/Core/Private/Serialization/MemoryImage.cpp
--- a/Engine/Source/Runtime/Core/Private/Serialization/MemoryImage.cpp
+++ b/Engine/Source/Runtime/Core/Private/Serialization/MemoryImage.cpp
@@ -3,6 +3,7 @@
 #include "Serialization/MemoryImage.h"
 #include "Containers/UnrealString.h"
 #include "Misc/SecureHash.h"
+#include "Misc/StringBuilder.h"
 #include "Interfaces/ITargetPlatform.h"
 #include "UObject/NameTypes.h"
 #include "Hash/CityHash.h"
@@ -641,6 +642,132 @@ uint32 Freeze::DefaultGetTargetAlignment(const FTypeLayoutDesc& TypeDesc, const
 	return Alignment;
 }
 
+void FMemoryToStringContext::AppendNullptr()
+{
+	String->Append(TEXT("nullptr\n"));
+}
+
+void FMemoryToStringContext::AppendIndent()
+{
+	for (int32 i = 0; i < Indent; ++i)
+	{
+		String->Append(TEXT("    "));
+	}
+}
+
+void Freeze::DefaultToString(const void* Object, const FTypeLayoutDesc& TypeDesc, const FPlatformTypeLayoutParameters& LayoutParams, FMemoryToStringContext& OutContext)
+{
+	OutContext.String->Appendf(TEXT("%s\n"), TypeDesc.Name);
+	++OutContext.Indent;
+
+	const FFieldLayoutDesc* FieldDesc = TypeDesc.Fields;
+	while (FieldDesc)
+	{
+		if (Freeze::IncludeField(FieldDesc, LayoutParams))
+		{
+			const FTypeLayoutDesc& FieldType = *FieldDesc->Type;
+			const uint8* FieldObject = (uint8*)Object + FieldDesc->Offset;
+
+			OutContext.AppendIndent();
+			OutContext.String->Appendf(TEXT("%s: "), FieldDesc->Name);
+
+			if (FieldDesc->BitFieldSize == 0u)
+			{
+				FieldType.ToStringFunc(FieldObject, FieldType, LayoutParams, OutContext);
+			}
+			else
+			{
+				OutContext.String->Append(TEXT("(BITFIELD)\n"));
+			}
+		}
+
+		FieldDesc = FieldDesc->Next;
+	}
+
+	--OutContext.Indent;
+}
+
+void Freeze::IntrinsicToString(char Object, const FTypeLayoutDesc& TypeDesc, const FPlatformTypeLayoutParameters& LayoutParams, FMemoryToStringContext& OutContext)
+{
+	OutContext.String->Appendf(TEXT("%d\n"), Object);
+}
+void Freeze::IntrinsicToString(short Object, const FTypeLayoutDesc& TypeDesc, const FPlatformTypeLayoutParameters& LayoutParams, FMemoryToStringContext& OutContext)
+{
+	OutContext.String->Appendf(TEXT("%d\n"), Object);
+}
+void Freeze::IntrinsicToString(int Object, const FTypeLayoutDesc& TypeDesc, const FPlatformTypeLayoutParameters& LayoutParams, FMemoryToStringContext& OutContext)
+{
+	OutContext.String->Appendf(TEXT("%d\n"), Object);
+}
+void Freeze::IntrinsicToString(int8 Object, const FTypeLayoutDesc& TypeDesc, const FPlatformTypeLayoutParameters& LayoutParams, FMemoryToStringContext& OutContext)
+{
+	OutContext.String->Appendf(TEXT("%d\n"), Object);
+}
+void Freeze::IntrinsicToString(long Object, const FTypeLayoutDesc& TypeDesc, const FPlatformTypeLayoutParameters& LayoutParams, FMemoryToStringContext& OutContext)
+{
+	OutContext.String->Appendf(TEXT("%d\n"), Object);
+}
+void Freeze::IntrinsicToString(long long Object, const FTypeLayoutDesc& TypeDesc, const FPlatformTypeLayoutParameters& LayoutParams, FMemoryToStringContext& OutContext)
+{
+	OutContext.String->Appendf(TEXT("%d\n"), Object);
+}
+void Freeze::IntrinsicToString(unsigned char Object, const FTypeLayoutDesc& TypeDesc, const FPlatformTypeLayoutParameters& LayoutParams, FMemoryToStringContext& OutContext)
+{
+	OutContext.String->Appendf(TEXT("%d\n"), Object);
+}
+void Freeze::IntrinsicToString(unsigned short Object, const FTypeLayoutDesc& TypeDesc, const FPlatformTypeLayoutParameters& LayoutParams, FMemoryToStringContext& OutContext)
+{
+	OutContext.String->Appendf(TEXT("%d\n"), Object);
+}
+void Freeze::IntrinsicToString(unsigned int Object, const FTypeLayoutDesc& TypeDesc, const FPlatformTypeLayoutParameters& LayoutParams, FMemoryToStringContext& OutContext)
+{
+	OutContext.String->Appendf(TEXT("%d\n"), Object);
+}
+void Freeze::IntrinsicToString(unsigned long Object, const FTypeLayoutDesc& TypeDesc, const FPlatformTypeLayoutParameters& LayoutParams, FMemoryToStringContext& OutContext)
+{
+	OutContext.String->Appendf(TEXT("%d\n"), Object);
+}
+void Freeze::IntrinsicToString(unsigned long long Object, const FTypeLayoutDesc& TypeDesc, const FPlatformTypeLayoutParameters& LayoutParams, FMemoryToStringContext& OutContext)
+{
+	OutContext.String->Appendf(TEXT("%d\n"), Object);
+}
+void Freeze::IntrinsicToString(float Object, const FTypeLayoutDesc& TypeDesc, const FPlatformTypeLayoutParameters& LayoutParams, FMemoryToStringContext& OutContext)
+{
+	OutContext.String->Appendf(TEXT("%g\n"), Object);
+}
+void Freeze::IntrinsicToString(double Object, const FTypeLayoutDesc& TypeDesc, const FPlatformTypeLayoutParameters& LayoutParams, FMemoryToStringContext& OutContext)
+{
+	OutContext.String->Appendf(TEXT("%g\n"), Object);
+}
+void Freeze::IntrinsicToString(wchar_t Object, const FTypeLayoutDesc& TypeDesc, const FPlatformTypeLayoutParameters& LayoutParams, FMemoryToStringContext& OutContext)
+{
+	OutContext.String->Appendf(TEXT("%d\n"), Object);
+}
+void Freeze::IntrinsicToString(char16_t Object, const FTypeLayoutDesc& TypeDesc, const FPlatformTypeLayoutParameters& LayoutParams, FMemoryToStringContext& OutContext)
+{
+	OutContext.String->Appendf(TEXT("%d\n"), Object);
+}
+void Freeze::IntrinsicToString(void* Object, const FTypeLayoutDesc& TypeDesc, const FPlatformTypeLayoutParameters& LayoutParams, FMemoryToStringContext& OutContext)
+{
+	OutContext.String->Appendf(TEXT("%p\n"), Object);
+}
+void Freeze::IntrinsicToString(const FHashedName& Object, const FTypeLayoutDesc& TypeDesc, const FPlatformTypeLayoutParameters& LayoutParams, FMemoryToStringContext& OutContext)
+{
+	OutContext.String->Appendf(TEXT("%016llX\n"), Object.GetHash());
+}
+
+void FMemoryImageString::ToString(FMemoryToStringContext& OutContext) const
+{
+	if (Data.Num() > 0)
+	{
+		OutContext.String->Appendf(TEXT("\"%s\"\n"), Data.GetData());
+	}
+	else
+	{
+		OutContext.String->Append(TEXT("\"\"\n"));
+	}
+}
+
 uint32 Freeze::AppendHash(const FTypeLayoutDesc& TypeDesc, const FPlatformTypeLayoutParameters& LayoutParams, FSHA1& Hasher)
 {
 	return TypeDesc.AppendHashFunc(TypeDesc, LayoutParams, Hasher);
diff --git a/Engine/Source/Runtime/Core/Public/Containers/Array.h b/Engine/Source/Runtime/Core/Public/Containers/Array.h
--- a/Engine/Source/Runtime/Core/Public/Containers/Array.h
+++ b/Engine/Source/Runtime/Core/Public/Containers/Array.h
@@ -2632,6 +2632,7 @@ private:
 
 		static void CopyUnfrozen(const FMemoryUnfreezeContent& Context, const TArray&, void* Dst) { new(Dst) TArray(); }
 		static void AppendHash(const FPlatformTypeLayoutParameters& LayoutParams, FSHA1& Hasher) {}
+		static void ToString(const FPlatformTypeLayoutParameters& LayoutParams, FMemoryToStringContext& OutContext, const TArray& Object) {}
 	};
 
 	template<typename Dummy>
@@ -2653,6 +2654,10 @@ private:
 		{
 			Freeze::AppendHash(StaticGetTypeLayoutDesc<ElementType>(), LayoutParams, Hasher);
 		}
+		static void ToString(const FPlatformTypeLayoutParameters& LayoutParams, FMemoryToStringContext& OutContext, const TArray& Object)
+		{
+			Object.AllocatorInstance.ToString(StaticGetTypeLayoutDesc<ElementType>(), Object.ArrayNum, LayoutParams, OutContext);
+		}
 	};
 
 public:
@@ -2674,6 +2679,12 @@ public:
 		TSupportsFreezeMemoryImageHelper<bSupportsFreezeMemoryImage>::AppendHash(LayoutParams, Hasher);
 	}
 
+	void ToString(const FPlatformTypeLayoutParameters& LayoutParams, FMemoryToStringContext& OutContext) const
+	{
+		static const bool bSupportsFreezeMemoryImage = TAllocatorTraits<Allocator>::SupportsFreezeMemoryImage && THasTypeLayout<ElementType>::Value;
+		TSupportsFreezeMemoryImageHelper<bSupportsFreezeMemoryImage>::ToString(LayoutParams, OutContext, *this);
+	}
+
 	/**
 	* Implicit heaps
 	*/
@@ -2987,6 +2998,12 @@ namespace Freeze
 		// Assume alignment of array is drive by pointer
 		return FMath::Min(LayoutParams.GetMemoryImagePointerSize(), LayoutParams.MaxFieldAlignment);
 	}
+
+	template<typename T, typename AllocatorType>
+	void IntrinsicToString(const TArray<T, AllocatorType>& Object, const FTypeLayoutDesc& TypeDesc, const FPlatformTypeLayoutParameters& LayoutParams, FMemoryToStringContext& OutContext)
+	{
+		Object.ToString(LayoutParams, OutContext);
+	}
 }
 
 DECLARE_TEMPLATE_INTRINSIC_TYPE_LAYOUT((template <typename T, typename AllocatorType>), (TArray<T, AllocatorType>));
diff --git a/Engine/Source/Runtime/Core/Public/Misc/SecureHash.h b/Engine/Source/Runtime/Core/Public/Misc/SecureHash.h
--- a/Engine/Source/Runtime/Core/Public/Misc/SecureHash.h
+++ b/Engine/Source/Runtime/Core/Public/Misc/SecureHash.h
@@ -241,6 +241,12 @@ public:
 	friend CORE_API FString LexToString(const FSHAHash&);
 	friend CORE_API void LexFromString(FSHAHash& Hash, const TCHAR*);
 };
+
+namespace Freeze
+{
+	CORE_API void IntrinsicToString(const FSHAHash& Object, const FTypeLayoutDesc& TypeDesc, const FPlatformTypeLayoutParameters& LayoutParams, FMemoryToStringContext& OutContext);
+}
+
 DECLARE_INTRINSIC_TYPE_LAYOUT(FSHAHash);
 
 inline FStringBuilderBase& operator<<(FStringBuilderBase& Builder, const FSHAHash& Hash) { UE::String::BytesToHex(Hash.Hash, Builder); return Builder; }
diff --git a/Engine/Source/Runtime/Core/Public/Serialization/MemoryImage.h b/Engine/Source/Runtime/Core/Public/Serialization/MemoryImage.h
--- a/Engine/Source/Runtime/Core/Public/Serialization/MemoryImage.h
+++ b/Engine/Source/Runtime/Core/Public/Serialization/MemoryImage.h
@@ -317,6 +317,27 @@ namespace Freeze
 	{
 		return FMath::Min(LayoutParams.GetMemoryImagePointerSize(), LayoutParams.MaxFieldAlignment);
 	}
+
+	template<typename T>
+	void IntrinsicToString(const TMemoryImagePtr<T>& Object, const FTypeLayoutDesc& TypeDesc, const FPlatformTypeLayoutParameters& LayoutParams, FMemoryToStringContext& OutContext)
+	{
+		T* RawPtr = Object.Get();
+		if (RawPtr)
+		{
+			// Compile-time type of the thing we're pointing to
+			const FTypeLayoutDesc& StaticTypeDesc = StaticGetTypeLayoutDesc<T>();
+			// Actual run-time type of the thing we're pointing to
+			const FTypeLayoutDesc& DerivedTypeDesc = GetTypeLayoutDesc(OutContext.TryGetPrevPointerTable(), *RawPtr);
+			// 'this' offset to adjust from the compile-time type to the run-time type
+			const uint32 OffsetToBase = DerivedTypeDesc.GetOffsetToBase(StaticTypeDesc);
+
+			DerivedTypeDesc.ToStringFunc((uint8*)RawPtr - OffsetToBase, DerivedTypeDesc, LayoutParams, OutContext);
+		}
+		else
+		{
+			OutContext.AppendNullptr();
+		}
+	}
 }
 
 DECLARE_TEMPLATE_INTRINSIC_TYPE_LAYOUT(template<typename T>, TMemoryImagePtr<T>);
@@ -396,6 +417,7 @@ public:
 
 	void ResizeAllocation(int32 PreviousNumElements, int32 NumElements, SIZE_T NumBytesPerElement, uint32 Alignment);
 	void WriteMemoryImage(FMemoryImageWriter& Writer, const FTypeLayoutDesc& TypeDesc, int32 NumAllocatedElements, uint32 Alignment) const;
+	void ToString(const FTypeLayoutDesc& TypeDesc, int32 NumAllocatedElements, const FPlatformTypeLayoutParameters& LayoutParams, FMemoryToStringContext& OutContext) const;
 	void CopyUnfrozen(const FMemoryUnfreezeContent& Context, const FTypeLayoutDesc& TypeDesc, int32 NumAllocatedElements, void* Dst) const;
 
 private:
@@ -488,7 +510,10 @@ class FMemoryImageString
 private:
 	/** Array holding the character data */
 	using DataType = TMemoryImageArray<TCHAR>;
-	LAYOUT_FIELD(DataType, Data)
+	LAYOUT_FIELD(DataType, Data);
+
+	CORE_API void ToString(FMemoryToStringContext& OutContext) const;
+	LAYOUT_TOSTRING(ToString);
 public:
 	using ElementType = TCHAR;
 
@@ -603,6 +628,12 @@ public:
 private:
 	uint64 Hash;
 };
+
+namespace Freeze
+{
+	CORE_API void IntrinsicToString(const FHashedName& Object, const FTypeLayoutDesc& TypeDesc, const FPlatformTypeLayoutParameters& LayoutParams, FMemoryToStringContext& OutContext);
+}
+
 DECLARE_INTRINSIC_TYPE_LAYOUT(FHashedName);
 
 class CORE_API FPtrTableBase
diff --git a/Engine/Source/Runtime/Core/Public/Serialization/MemoryLayout.h b/Engine/Source/Runtime/Core/Public/Serialization/MemoryLayout.h
--- a/Engine/Source/Runtime/Core/Public/Serialization/MemoryLayout.h
+++ b/Engine/Source/Runtime/Core/Public/Serialization/MemoryLayout.h
@@ -3,6 +3,7 @@
 #pragma once
 
 #include "Containers/EnumAsByte.h"
+#include "Containers/StringFwd.h"
 #include "Templates/IsAbstract.h"
 #include "Templates/IsPolymorphic.h"
 #include "Templates/IsTriviallyDestructible.h"
@@ -95,6 +96,18 @@ struct FPlatformTypeLayoutParameters
 	CORE_API void AppendKeyString(FString& KeyString) const;
 };
 
+struct FMemoryToStringContext
+{
+	CORE_API void AppendNullptr();
+	CORE_API void AppendIndent();
+
+	const FPointerTableBase* TryGetPrevPointerTable() const { return PrevPointerTable; }
+
+	FStringBuilderBase* String = nullptr;
+	const FPointerTableBase* PrevPointerTable = nullptr;
+	int32 Indent = 0;
+};
+
 namespace ETypeLayoutInterface
 {
 	enum Type : uint8
@@ -145,6 +158,7 @@ struct FTypeLayoutDesc
 	typedef void (FUnfrozenCopyFunc)(const FMemoryUnfreezeContent& Context, const void* Object, const FTypeLayoutDesc& TypeDesc, void* OutDst);
 	typedef uint32 (FAppendHashFunc)(const FTypeLayoutDesc& TypeDesc, const FPlatformTypeLayoutParameters& LayoutParams, FSHA1& Hasher);
 	typedef uint32 (FGetTargetAlignmentFunc)(const FTypeLayoutDesc& TypeDesc, const FPlatformTypeLayoutParameters& LayoutParams);
+	typedef void (FToStringFunc)(const void* Object, const FTypeLayoutDesc& TypeDesc, const FPlatformTypeLayoutParameters& LayoutParams, FMemoryToStringContext& OutContext);
 	typedef const void* (FGetDefaultFunc)();
 
 	static CORE_API const FTypeLayoutDesc& GetInvalidTypeLayout();
@@ -163,6 +177,7 @@ struct FTypeLayoutDesc
 	FUnfrozenCopyFunc* UnfrozenCopyFunc;
 	FAppendHashFunc* AppendHashFunc;
 	FGetTargetAlignmentFunc* GetTargetAlignmentFunc;
+	FToStringFunc* ToStringFunc;
 	FGetDefaultFunc* GetDefaultObjectFunc;
 
 	uint64 NameHash; // from FHashedName(Name)
@@ -289,6 +304,7 @@ namespace Freeze
 	CORE_API void DefaultUnfrozenCopy(const FMemoryUnfreezeContent& Context, const void* Object, const FTypeLayoutDesc& TypeDesc, void* OutDst);
 	CORE_API uint32 DefaultAppendHash(const FTypeLayoutDesc& TypeDesc, const FPlatformTypeLayoutParameters& LayoutParams, FSHA1& Hasher);
 	CORE_API uint32 DefaultGetTargetAlignment(const FTypeLayoutDesc& TypeDesc, const FPlatformTypeLayoutParameters& LayoutParams);
+	CORE_API void DefaultToString(const void* Object, const FTypeLayoutDesc& TypeDesc, const FPlatformTypeLayoutParameters& LayoutParams, FMemoryToStringContext& OutContext);
 	CORE_API uint32 AppendHashForNameAndSize(const TCHAR* Name, uint32 Size, FSHA1& Hasher);
 	CORE_API void IntrinsicWriteMemoryImage(FMemoryImageWriter& Writer, const void* Object, uint32 Size);
 	CORE_API void IntrinsicWriteMemoryImage(FMemoryImageWriter& Writer, void*, const FTypeLayoutDesc&);
@@ -327,6 +343,29 @@ namespace Freeze
 		return LayoutParams.GetRawPointerSize();
 	}
 
+	template<typename T>
+	inline void IntrinsicToString(const T& Object, const FTypeLayoutDesc& TypeDesc, const FPlatformTypeLayoutParameters& LayoutParams, FMemoryToStringContext& OutContext)
+	{
+		DefaultToString(&Object, TypeDesc, LayoutParams, OutContext);
+	}
+
+	CORE_API void IntrinsicToString(char Object, const FTypeLayoutDesc& TypeDesc, const FPlatformTypeLayoutParameters& LayoutParams, FMemoryToStringContext& OutContext);
+	CORE_API void IntrinsicToString(short Object, const FTypeLayoutDesc& TypeDesc, const FPlatformTypeLayoutParameters& LayoutParams, FMemoryToStringContext& OutContext);
+	CORE_API void IntrinsicToString(int Object, const FTypeLayoutDesc& TypeDesc, const FPlatformTypeLayoutParameters& LayoutParams, FMemoryToStringContext& OutContext);
+	CORE_API void IntrinsicToString(int8 Object, const FTypeLayoutDesc& TypeDesc, const FPlatformTypeLayoutParameters& LayoutParams, FMemoryToStringContext& OutContext);
+	CORE_API void IntrinsicToString(long Object, const FTypeLayoutDesc& TypeDesc, const FPlatformTypeLayoutParameters& LayoutParams, FMemoryToStringContext& OutContext);
+	CORE_API void IntrinsicToString(long long Object, const FTypeLayoutDesc& TypeDesc, const FPlatformTypeLayoutParameters& LayoutParams, FMemoryToStringContext& OutContext);
+	CORE_API void IntrinsicToString(unsigned char Object, const FTypeLayoutDesc& TypeDesc, const FPlatformTypeLayoutParameters& LayoutParams, FMemoryToStringContext& OutContext);
+	CORE_API void IntrinsicToString(unsigned short Object, const FTypeLayoutDesc& TypeDesc, const FPlatformTypeLayoutParameters& LayoutParams, FMemoryToStringContext& OutContext);
+	CORE_API void IntrinsicToString(unsigned int Object, const FTypeLayoutDesc& TypeDesc, const FPlatformTypeLayoutParameters& LayoutParams, FMemoryToStringContext& OutContext);
+	CORE_API void IntrinsicToString(unsigned long Object, const FTypeLayoutDesc& TypeDesc, const FPlatformTypeLayoutParameters& LayoutParams, FMemoryToStringContext& OutContext);
+	CORE_API void IntrinsicToString(unsigned long long Object, const FTypeLayoutDesc& TypeDesc, const FPlatformTypeLayoutParameters& LayoutParams, FMemoryToStringContext& OutContext);
+	CORE_API void IntrinsicToString(float Object, const FTypeLayoutDesc& TypeDesc, const FPlatformTypeLayoutParameters& LayoutParams, FMemoryToStringContext& OutContext);
+	CORE_API void IntrinsicToString(double Object, const FTypeLayoutDesc& TypeDesc, const FPlatformTypeLayoutParameters& LayoutParams, FMemoryToStringContext& OutContext);
+	CORE_API void IntrinsicToString(wchar_t Object, const FTypeLayoutDesc& TypeDesc, const FPlatformTypeLayoutParameters& LayoutParams, FMemoryToStringContext& OutContext);
+	CORE_API void IntrinsicToString(char16_t Object, const FTypeLayoutDesc& TypeDesc, const FPlatformTypeLayoutParameters& LayoutParams, FMemoryToStringContext& OutContext);
+	CORE_API void IntrinsicToString(void* Object, const FTypeLayoutDesc& TypeDesc, const FPlatformTypeLayoutParameters& LayoutParams, FMemoryToStringContext& OutContext);
+
 	// helper functions
 	CORE_API void ExtractBitFieldValue(const void* Value, uint32 SrcBitOffset, uint32 DestBitOffset, uint32 NumBits, uint64& InOutValue);
 	CORE_API bool IncludeField(const FFieldLayoutDesc* FieldDesc, const FPlatformTypeLayoutParameters& LayoutParams);
@@ -491,6 +530,18 @@ PRAGMA_ENABLE_DEPRECATION_WARNINGS
 		} \
 	};
 
+#define INTERNAL_LAYOUT_TOSTRING(Func, Counter) \
+	UE_DECLARE_INTERNAL_LINK_SPECIALIZATION(InternalLinkType, Counter - CounterBase) { \
+		UE_STATIC_ONLY(InternalLinkType); \
+		static void CallToString(const void* Object, const FTypeLayoutDesc& TypeDesc, const FPlatformTypeLayoutParameters& LayoutParams, FMemoryToStringContext& OutContext) { \
+			static_cast<const DerivedType*>(Object)->Func(OutContext); \
+		} \
+		static void Initialize(FTypeLayoutDesc& TypeDesc) { \
+			InternalLinkType<Counter - CounterBase + 1>::Initialize(TypeDesc); \
+			TypeDesc.ToStringFunc = &CallToString; \
+		} \
+	};
+
 #define LAYOUT_FIELD(T, Name, ...) PREPROCESSOR_REMOVE_OPTIONAL_PARENS(T) Name; INTERNAL_LAYOUT_FIELD(T, Name, STRUCT_OFFSET(DerivedType, Name), EFieldLayoutFlags::MakeFlags(__VA_ARGS__), 1u, 0u, __COUNTER__)
 #define LAYOUT_MUTABLE_FIELD(T, Name, ...) mutable PREPROCESSOR_REMOVE_OPTIONAL_PARENS(T) Name; INTERNAL_LAYOUT_FIELD(T, Name, STRUCT_OFFSET(DerivedType, Name), EFieldLayoutFlags::MakeFlags(__VA_ARGS__), 1u, 0u, __COUNTER__)
 #define LAYOUT_FIELD_INITIALIZED(T, Name, Value, ...) PREPROCESSOR_REMOVE_OPTIONAL_PARENS(T) Name = Value; INTERNAL_LAYOUT_FIELD(T, Name, STRUCT_OFFSET(DerivedType, Name), EFieldLayoutFlags::MakeFlags(__VA_ARGS__), 1u, 0u, __COUNTER__)
@@ -501,6 +552,7 @@ PRAGMA_ENABLE_DEPRECATION_WARNINGS
 #define LAYOUT_FIELD_WITH_WRITER(T, Name, Func) PREPROCESSOR_REMOVE_OPTIONAL_PARENS(T) Name; INTERNAL_LAYOUT_FIELD_WITH_WRITER(T, Name, Func, __COUNTER__)
 #define LAYOUT_MUTABLE_FIELD_WITH_WRITER(T, Name, Func) mutable PREPROCESSOR_REMOVE_OPTIONAL_PARENS(T) Name; INTERNAL_LAYOUT_FIELD_WITH_WRITER(T, Name, Func, __COUNTER__)
 #define LAYOUT_WRITE_MEMORY_IMAGE(Func) INTERNAL_LAYOUT_WRITE_MEMORY_IMAGE(Func, __COUNTER__)
+#define LAYOUT_TOSTRING(Func) INTERNAL_LAYOUT_TOSTRING(Func, __COUNTER__)
 
 #if WITH_EDITORONLY_DATA
 #define LAYOUT_FIELD_EDITORONLY(T, Name, ...) PREPROCESSOR_REMOVE_OPTIONAL_PARENS(T) Name; INTERNAL_LAYOUT_FIELD(T, Name, STRUCT_OFFSET(DerivedType, Name), EFieldLayoutFlags::MakeFlagsEditorOnly(__VA_ARGS__), 1u, 0u, __COUNTER__)
@@ -555,6 +607,7 @@ PRAGMA_ENABLE_DEPRECATION_WARNINGS
 			TypeDesc.UnfrozenCopyFunc = &Freeze::DefaultUnfrozenCopy; \
 			TypeDesc.AppendHashFunc = &Freeze::DefaultAppendHash; \
 			TypeDesc.GetTargetAlignmentFunc = &Freeze::DefaultGetTargetAlignment; \
+			TypeDesc.ToStringFunc = &Freeze::DefaultToString; \
 			TypeDesc.Size = sizeof(PREPROCESSOR_REMOVE_OPTIONAL_PARENS(T)); \
 			TypeDesc.Alignment = alignof(PREPROCESSOR_REMOVE_OPTIONAL_PARENS(T)); \
 			TypeDesc.Interface = ETypeLayoutInterface::InInterface; \
@@ -601,6 +654,7 @@ PRAGMA_ENABLE_DEPRECATION_WARNINGS
 			TypeDesc.UnfrozenCopyFunc = &Freeze::DefaultUnfrozenCopy; \
 			TypeDesc.AppendHashFunc = &Freeze::DefaultAppendHash; \
 			TypeDesc.GetTargetAlignmentFunc = &Freeze::DefaultGetTargetAlignment; \
+			TypeDesc.ToStringFunc = &Freeze::DefaultToString; \
 			TypeDesc.Size = sizeof(PREPROCESSOR_REMOVE_OPTIONAL_PARENS(T)); \
 			TypeDesc.Alignment = alignof(PREPROCESSOR_REMOVE_OPTIONAL_PARENS(T)); \
 			TypeDesc.Interface = InterfaceType; \
@@ -697,6 +751,9 @@ inline void DeleteObjectFromLayout(T* Object, FPointerTableBase* PtrTable = null
 		static uint32 CallGetTargetAlignment(const FTypeLayoutDesc& TypeDesc, const FPlatformTypeLayoutParameters& LayoutParams) { \
 			return Freeze::IntrinsicGetTargetAlignment(static_cast<PREPROCESSOR_REMOVE_OPTIONAL_PARENS(T) const*>(nullptr), TypeDesc, LayoutParams); \
 		} \
+		static void CallToString(const void* Object, const FTypeLayoutDesc& TypeDesc, const FPlatformTypeLayoutParameters& LayoutParams, FMemoryToStringContext& OutContext) { \
+			return Freeze::IntrinsicToString(*static_cast<PREPROCESSOR_REMOVE_OPTIONAL_PARENS(T) const*>(Object), TypeDesc, LayoutParams, OutContext); \
+		} \
 		static const FTypeLayoutDesc& Do() { \
 			alignas(FTypeLayoutDesc) static uint8 TypeBuffer[sizeof(FTypeLayoutDesc)] = { 0 }; \
 			FTypeLayoutDesc& TypeDesc = *(FTypeLayoutDesc*)TypeBuffer; \
@@ -708,6 +765,7 @@ inline void DeleteObjectFromLayout(T* Object, FPointerTableBase* PtrTable = null
 				TypeDesc.UnfrozenCopyFunc = &CallUnfrozenCopy; \
 				TypeDesc.AppendHashFunc = &CallAppendHash; \
 				TypeDesc.GetTargetAlignmentFunc = &CallGetTargetAlignment; \
+				TypeDesc.ToStringFunc = &CallToString; \
 				TypeDesc.Size = sizeof(PREPROCESSOR_REMOVE_OPTIONAL_PARENS(T)); \
 				TypeDesc.Alignment = alignof(PREPROCESSOR_REMOVE_OPTIONAL_PARENS(T)); \
 				TypeDesc.Interface = ETypeLayoutInterface::NonVirtual; \
@@ -727,6 +785,7 @@ inline void DeleteObjectFromLayout(T* Object, FPointerTableBase* PtrTable = null
 		static void CallUnfrozenCopy(const FMemoryUnfreezeContent& Context, const void* Object, const FTypeLayoutDesc& TypeDesc, void* OutDst); \
 		static uint32 CallAppendHash(const FTypeLayoutDesc& TypeDesc, const FPlatformTypeLayoutParameters& LayoutParams, FSHA1& Hasher); \
 		static uint32 CallGetTargetAlignment(const FTypeLayoutDesc& TypeDesc, const FPlatformTypeLayoutParameters& LayoutParams); \
+		static void CallToString(const void* Object, const FTypeLayoutDesc& TypeDesc, const FPlatformTypeLayoutParameters& LayoutParams, FMemoryToStringContext& OutContext); \
 		RequiredAPI static const FTypeLayoutDesc& Do(); }; \
 	PREPROCESSOR_REMOVE_OPTIONAL_PARENS(TemplatePrefix) struct TGetTypeLayoutHelper<PREPROCESSOR_REMOVE_OPTIONAL_PARENS(T)> { \
 		UE_STATIC_ONLY(TGetTypeLayoutHelper); \
@@ -745,6 +804,9 @@ inline void DeleteObjectFromLayout(T* Object, FPointerTableBase* PtrTable = null
 		uint32 TStaticGetTypeLayoutHelper<PREPROCESSOR_REMOVE_OPTIONAL_PARENS(T)>::CallGetTargetAlignment(const FTypeLayoutDesc& TypeDesc, const FPlatformTypeLayoutParameters& LayoutParams) { \
 			return Freeze::IntrinsicGetTargetAlignment(static_cast<PREPROCESSOR_REMOVE_OPTIONAL_PARENS(T) const*>(nullptr), TypeDesc, LayoutParams); \
 		} \
+		void TStaticGetTypeLayoutHelper<PREPROCESSOR_REMOVE_OPTIONAL_PARENS(T)>::CallToString(const void* Object, const FTypeLayoutDesc& TypeDesc, const FPlatformTypeLayoutParameters& LayoutParams, FMemoryToStringContext& OutContext) { \
+			return Freeze::IntrinsicToString(*static_cast<PREPROCESSOR_REMOVE_OPTIONAL_PARENS(T) const*>(Object), TypeDesc, LayoutParams, OutContext); \
+		} \
 		const FTypeLayoutDesc& TStaticGetTypeLayoutHelper<PREPROCESSOR_REMOVE_OPTIONAL_PARENS(T)>::Do() { \
 			alignas(FTypeLayoutDesc) static uint8 TypeBuffer[sizeof(FTypeLayoutDesc)] = { 0 }; \
 			FTypeLayoutDesc& TypeDesc = *(FTypeLayoutDesc*)TypeBuffer; \
@@ -756,6 +818,7 @@ inline void DeleteObjectFromLayout(T* Object, FPointerTableBase* PtrTable = null
 				TypeDesc.UnfrozenCopyFunc = &CallUnfrozenCopy; \
 				TypeDesc.AppendHashFunc = &CallAppendHash; \
 				TypeDesc.GetTargetAlignmentFunc = &CallGetTargetAlignment; \
+				TypeDesc.ToStringFunc = &CallToString; \
 				TypeDesc.Size = sizeof(PREPROCESSOR_REMOVE_OPTIONAL_PARENS(T)); \
 				TypeDesc.Alignment = alignof(PREPROCESSOR_REMOVE_OPTIONAL_PARENS(T)); \
 				TypeDesc.Interface = ETypeLayoutInterface::NonVirtual; \
diff --git a/Engine/Source/Runtime/Engine/Private/Materials/MaterialShader.cpp b/Engine/Source/Runtime/Engine/Private/Materials/MaterialShader.cpp
--- a/Engine/Source/Runtime/Engine/Private/Materials/MaterialShader.cpp
+++ b/Engine/Source/Runtime/Engine/Private/Materials/MaterialShader.cpp
@@ -932,23 +932,18 @@ FShader* FMaterialShaderType::FinishCompileShader(
 	const FSHAHash& MaterialShaderMapHash,
 	const FShaderCompileJob& CurrentJob,
 	const FShaderPipelineType* ShaderPipelineType,
-	const FString& InDebugDescription,
-	FShaderMapResourceBuilder& ResourceBuilder
+	const FString& InDebugDescription
 	)
 {
 	check(CurrentJob.bSucceeded);
 
-	// Reuse an existing resource with the same key or create a new one based on the compile output
-	// This allows FShaders to share compiled bytecode and RHI shader references
-	const int32 ResourceIndex = ResourceBuilder.FindOrAddCode(CurrentJob.Output);
-
 	if (ShaderPipelineType && !ShaderPipelineType->ShouldOptimizeUnusedOutputs(CurrentJob.Input.Target.GetPlatform()))
 	{
 		// If sharing shaders in this pipeline, remove it from the type/id so it uses the one in the shared shadermap list
 		ShaderPipelineType = nullptr;
 	}
 
-	FShader* Shader = ConstructCompiled(CompiledShaderInitializerType(this, CurrentJob.PermutationId, CurrentJob.Output, ResourceIndex, UniformExpressionSet, MaterialShaderMapHash, ShaderPipelineType, nullptr, InDebugDescription));
+	FShader* Shader = ConstructCompiled(CompiledShaderInitializerType(this, CurrentJob.PermutationId, CurrentJob.Output, UniformExpressionSet, MaterialShaderMapHash, ShaderPipelineType, nullptr, InDebugDescription));
 	CurrentJob.Output.ParameterMap.VerifyBindingsAreComplete(GetName(), CurrentJob.Output.Target, CurrentJob.VFType);
 
 	return Shader;
@@ -1076,7 +1071,7 @@ void FMaterialShaderMap::LoadFromDerivedDataCache(const FMaterial* Material, con
 				// Deserialize from the cached data
 				InOutShaderMap->Serialize(Ar);
 				//InOutShaderMap->RegisterSerializedShaders(false);
-
+		
 				const FString InDataKey = GetMaterialShaderMapKeyString(InOutShaderMap->GetShaderMapId(), InPlatform, TargetPlatform);
 				checkSlow(InOutShaderMap->GetShaderMapId() == ShaderMapId);
 
@@ -1567,12 +1562,14 @@ void FMaterialShaderMap::Compile(
 	}
 }
 
-FShader* FMaterialShaderMap::ProcessCompilationResultsForSingleJob(FShaderCompileJob* SingleJob, const FShaderPipelineType* ShaderPipeline, const FSHAHash& MaterialShaderMapHash, FShaderMapResourceBuilder& ResourceBuilder)
+FShader* FMaterialShaderMap::ProcessCompilationResultsForSingleJob(FShaderCompileJob* SingleJob, const FShaderPipelineType* ShaderPipeline, const FSHAHash& MaterialShaderMapHash)
 {
 	check(SingleJob);
 	const FShaderCompileJob& CurrentJob = *SingleJob;
 	check(CurrentJob.Id == CompilingId);
 
+	GetResourceCode()->AddShaderCompilerOutput(CurrentJob.Output);
+
 #if ALLOW_SHADERMAP_DEBUG_DATA
 	CompileTime += SingleJob->Output.CompileTime;
 #endif
@@ -1586,7 +1583,7 @@ FShader* FMaterialShaderMap::ProcessCompilationResultsForSingleJob(FShaderCompil
 		check(MeshShaderMap);
 		FMeshMaterialShaderType* MeshMaterialShaderType = CurrentJob.ShaderType->GetMeshMaterialShaderType();
 		check(MeshMaterialShaderType);
-		Shader = MeshMaterialShaderType->FinishCompileShader(GetContent()->MaterialCompilationOutput.UniformExpressionSet, MaterialShaderMapHash, CurrentJob, ShaderPipeline, GetFriendlyName(), ResourceBuilder);
+		Shader = MeshMaterialShaderType->FinishCompileShader(GetContent()->MaterialCompilationOutput.UniformExpressionSet, MaterialShaderMapHash, CurrentJob, ShaderPipeline, GetFriendlyName());
 		check(Shader);
 		if (!ShaderPipeline)
 		{
@@ -1598,7 +1595,7 @@ FShader* FMaterialShaderMap::ProcessCompilationResultsForSingleJob(FShaderCompil
 	{
 		FMaterialShaderType* MaterialShaderType = CurrentJob.ShaderType->GetMaterialShaderType();
 		check(MaterialShaderType);
-		Shader = MaterialShaderType->FinishCompileShader(GetContent()->MaterialCompilationOutput.UniformExpressionSet, MaterialShaderMapHash, CurrentJob, ShaderPipeline, GetFriendlyName(), ResourceBuilder);
+		Shader = MaterialShaderType->FinishCompileShader(GetContent()->MaterialCompilationOutput.UniformExpressionSet, MaterialShaderMapHash, CurrentJob, ShaderPipeline, GetFriendlyName());
 
 		check(Shader);
 		if (!ShaderPipeline)
@@ -1609,8 +1606,12 @@ FShader* FMaterialShaderMap::ProcessCompilationResultsForSingleJob(FShaderCompil
 	}
 
 #if WITH_EDITOR
-	// add shader source to 
-	new(GetMutableContent()->ShaderProcessedSource) FMaterialProcessedSource(CurrentJob.ShaderType->GetFName(), *CurrentJob.Output.OptionalFinalShaderSource);
+	// add shader source
+	{
+		// Keep the preprocessed source list sorted by type name
+		const int32 Index = Algo::LowerBoundBy(GetMutableContent()->ShaderProcessedSource, CurrentJob.ShaderType->GetHashedName(), [](const FMaterialProcessedSource& Value) { return Value.Name; });
+		GetMutableContent()->ShaderProcessedSource.EmplaceAt(Index, CurrentJob.ShaderType->GetHashedName(), *CurrentJob.Output.OptionalFinalShaderSource);
+	}
 #endif
 
 	return Shader;
@@ -1627,13 +1628,12 @@ bool FMaterialShaderMap::ProcessCompilationResults(const TArray<TSharedRef<FShad
 	FSHAHash MaterialShaderMapHash;
 	ShaderMapId.GetMaterialHash(MaterialShaderMapHash);
 
-	FShaderMapResourceBuilder ResourceBuilder(GetResourceCode());
 	do
 	{
 		FShaderCompileJob* SingleJob = InCompilationResults[InOutJobIndex]->GetSingleShaderJob();
 		if (SingleJob)
 		{
-			ProcessCompilationResultsForSingleJob(SingleJob, nullptr, MaterialShaderMapHash, ResourceBuilder);
+			ProcessCompilationResultsForSingleJob(SingleJob, nullptr, MaterialShaderMapHash);
 			for (auto Pair : SingleJob->SharingPipelines)
 			{
 				auto& SharedPipelinesPerVF = SharedPipelines.FindOrAdd(SingleJob->VFType);
@@ -1656,7 +1656,7 @@ bool FMaterialShaderMap::ProcessCompilationResults(const TArray<TSharedRef<FShad
 			for (int32 Index = 0; Index < CurrentJob.StageJobs.Num(); ++Index)
 			{
 				SingleJob = CurrentJob.StageJobs[Index]->GetSingleShaderJob();
-				FShader* Shader = ProcessCompilationResultsForSingleJob(SingleJob, PipelineJob->ShaderPipeline, MaterialShaderMapHash, ResourceBuilder);
+				FShader* Shader = ProcessCompilationResultsForSingleJob(SingleJob, PipelineJob->ShaderPipeline, MaterialShaderMapHash);
 				ShaderPipeline->AddShader(Shader, SingleJob->PermutationId);
 				check(VertexFactoryType == CurrentJob.StageJobs[Index]->GetSingleShaderJob()->VFType);
 			}
diff --git a/Engine/Source/Runtime/Engine/Private/Materials/MeshMaterialShader.cpp b/Engine/Source/Runtime/Engine/Private/Materials/MeshMaterialShader.cpp
--- a/Engine/Source/Runtime/Engine/Private/Materials/MeshMaterialShader.cpp
+++ b/Engine/Source/Runtime/Engine/Private/Materials/MeshMaterialShader.cpp
@@ -139,23 +139,18 @@ FShader* FMeshMaterialShaderType::FinishCompileShader(
 	const FSHAHash& MaterialShaderMapHash,
 	const FShaderCompileJob& CurrentJob,
 	const FShaderPipelineType* ShaderPipelineType,
-	const FString& InDebugDescription,
-	FShaderMapResourceBuilder& ResourceBuilder)
+	const FString& InDebugDescription)
 {
 	checkf(CurrentJob.bSucceeded, TEXT("Failed MeshMaterialType compilation job: %s"), *GetJobName(&CurrentJob, ShaderPipelineType, InDebugDescription));
 	checkf(CurrentJob.VFType, TEXT("No VF on MeshMaterialType compilation job: %s"), *GetJobName(&CurrentJob, ShaderPipelineType, InDebugDescription));
 
-	// Reuse an existing resource with the same key or create a new one based on the compile output
-	// This allows FShaders to share compiled bytecode and RHI shader references
-	const int32 ResourceIndex = ResourceBuilder.FindOrAddCode(CurrentJob.Output);
-
 	if (ShaderPipelineType && !ShaderPipelineType->ShouldOptimizeUnusedOutputs(CurrentJob.Input.Target.GetPlatform()))
 	{
 		// If sharing shaders in this pipeline, remove it from the type/id so it uses the one in the shared shadermap list
 		ShaderPipelineType = nullptr;
 	}
 
-	FShader* Shader = ConstructCompiled(CompiledShaderInitializerType(this, CurrentJob.PermutationId, CurrentJob.Output, ResourceIndex, UniformExpressionSet, MaterialShaderMapHash, InDebugDescription, ShaderPipelineType, CurrentJob.VFType));
+	FShader* Shader = ConstructCompiled(CompiledShaderInitializerType(this, CurrentJob.PermutationId, CurrentJob.Output, UniformExpressionSet, MaterialShaderMapHash, InDebugDescription, ShaderPipelineType, CurrentJob.VFType));
 	ValidateAfterBind(this, (FMeshMaterialShader*)Shader);
 	CurrentJob.Output.ParameterMap.VerifyBindingsAreComplete(GetName(), CurrentJob.Output.Target, CurrentJob.VFType);
 
diff --git a/Engine/Source/Runtime/Engine/Private/ShaderCompiler/ShaderCompiler.cpp b/Engine/Source/Runtime/Engine/Private/ShaderCompiler/ShaderCompiler.cpp
--- a/Engine/Source/Runtime/Engine/Private/ShaderCompiler/ShaderCompiler.cpp
+++ b/Engine/Source/Runtime/Engine/Private/ShaderCompiler/ShaderCompiler.cpp
@@ -4313,9 +4313,7 @@ FShader* FGlobalShaderTypeCompiler::FinishCompileShader(FGlobalShaderType* Shade
 		EShaderPlatform Platform = CurrentJob.Input.Target.GetPlatform();
 		FGlobalShaderMapSection* Section = GGlobalShaderMap[Platform]->FindOrAddSection(ShaderType);
 
-		// Reuse an existing resource with the same key or create a new one based on the compile output
-		// This allows FShaders to share compiled bytecode and RHI shader references
-		const int32 ResourceIndex = Section->GetResourceBuilder().FindOrAddCode(CurrentJob.Output);
+		Section->GetResourceCode()->AddShaderCompilerOutput(CurrentJob.Output);
 
 		if (ShaderPipelineType && !ShaderPipelineType->ShouldOptimizeUnusedOutputs(CurrentJob.Input.Target.GetPlatform()))
 		{
@@ -4323,7 +4321,6 @@ FShader* FGlobalShaderTypeCompiler::FinishCompileShader(FGlobalShaderType* Shade
 			ShaderPipelineType = nullptr;
 		}
 
-
 		// Create the global shader map hash
 		FSHAHash GlobalShaderMapHash;
 		{
@@ -4334,7 +4331,7 @@ FShader* FGlobalShaderTypeCompiler::FinishCompileShader(FGlobalShaderType* Shade
 			HashState.GetHash(&GlobalShaderMapHash.Hash[0]);
 		}
 
-		Shader = ShaderType->ConstructCompiled(FGlobalShaderType::CompiledShaderInitializerType(ShaderType, CurrentJob.PermutationId, CurrentJob.Output, ResourceIndex, GlobalShaderMapHash, ShaderPipelineType, nullptr));
+		Shader = ShaderType->ConstructCompiled(FGlobalShaderType::CompiledShaderInitializerType(ShaderType, CurrentJob.PermutationId, CurrentJob.Output, GlobalShaderMapHash, ShaderPipelineType, nullptr));
 		CurrentJob.Output.ParameterMap.VerifyBindingsAreComplete(ShaderType->GetName(), CurrentJob.Output.Target, CurrentJob.VFType);
 	}
 
@@ -4577,8 +4574,7 @@ static void SaveGlobalShaderMapToDerivedDataCache(EShaderPlatform Platform)
 		if (Section)
 		{
 			Section->FinalizeContent();
-			Section->ReleaseResourceBuilder();
-
+	
 			SaveData.Reset();
 			FMemoryWriter Ar(SaveData, true);
 			Section->Serialize(Ar);
diff --git a/Engine/Source/Runtime/Engine/Public/MaterialShaderType.h b/Engine/Source/Runtime/Engine/Public/MaterialShaderType.h
--- a/Engine/Source/Runtime/Engine/Public/MaterialShaderType.h
+++ b/Engine/Source/Runtime/Engine/Public/MaterialShaderType.h
@@ -67,14 +67,13 @@ public:
 			FShaderType* InType,
 			int32 InPermutationId,
 			const FShaderCompilerOutput& CompilerOutput,
-			int32 InResourceIndex,
 			const FUniformExpressionSet& InUniformExpressionSet,
 			const FSHAHash& InMaterialShaderMapHash,
 			const FShaderPipelineType* InShaderPipeline,
 			FVertexFactoryType* InVertexFactoryType,
 			const FString& InDebugDescription
 			)
-		: FGlobalShaderType::CompiledShaderInitializerType(InType,InPermutationId,CompilerOutput, InResourceIndex,InMaterialShaderMapHash,InShaderPipeline,InVertexFactoryType)
+		: FGlobalShaderType::CompiledShaderInitializerType(InType,InPermutationId,CompilerOutput,InMaterialShaderMapHash,InShaderPipeline,InVertexFactoryType)
 		, UniformExpressionSet(InUniformExpressionSet)
 		, DebugDescription(InDebugDescription)
 		{}
@@ -147,8 +146,7 @@ public:
 		const FSHAHash& MaterialShaderMapHash,
 		const FShaderCompileJob& CurrentJob,
 		const FShaderPipelineType* ShaderPipeline,
-		const FString& InDebugDescription,
-		FShaderMapResourceBuilder& ResourceBuilder
+		const FString& InDebugDescription
 		);
 
 	/**
diff --git a/Engine/Source/Runtime/Engine/Public/MaterialShared.h b/Engine/Source/Runtime/Engine/Public/MaterialShared.h
--- a/Engine/Source/Runtime/Engine/Public/MaterialShared.h
+++ b/Engine/Source/Runtime/Engine/Public/MaterialShared.h
@@ -1302,7 +1302,7 @@ private:
 	/** Indicates whether the shader map should be stored in the shader cache. */
 	uint32 bIsPersistent : 1;
 
-	FShader* ProcessCompilationResultsForSingleJob(class FShaderCompileJob* SingleJob, const FShaderPipelineType* ShaderPipeline, const FSHAHash& MaterialShaderMapHash, FShaderMapResourceBuilder& ResourceBuilder);
+	FShader* ProcessCompilationResultsForSingleJob(class FShaderCompileJob* SingleJob, const FShaderPipelineType* ShaderPipeline, const FSHAHash& MaterialShaderMapHash);
 
 	friend ENGINE_API void DumpMaterialStats( EShaderPlatform Platform );
 	friend class FShaderCompilingManager;
diff --git a/Engine/Source/Runtime/Engine/Public/MeshMaterialShaderType.h b/Engine/Source/Runtime/Engine/Public/MeshMaterialShaderType.h
--- a/Engine/Source/Runtime/Engine/Public/MeshMaterialShaderType.h
+++ b/Engine/Source/Runtime/Engine/Public/MeshMaterialShaderType.h
@@ -29,14 +29,13 @@ public:
 			FShaderType* InType,
 			int32 PermutationId,
 			const FShaderCompilerOutput& CompilerOutput,
-			int32 InResourceIndex,
 			const FUniformExpressionSet& InUniformExpressionSet,
 			const FSHAHash& InMaterialShaderMapHash,
 			const FString& InDebugDescription,
 			const FShaderPipelineType* InShaderPipeline,
 			FVertexFactoryType* InVertexFactoryType
 			):
-			FMaterialShaderType::CompiledShaderInitializerType(InType,PermutationId,CompilerOutput, InResourceIndex,InUniformExpressionSet,InMaterialShaderMapHash,InShaderPipeline,InVertexFactoryType,InDebugDescription),
+			FMaterialShaderType::CompiledShaderInitializerType(InType,PermutationId,CompilerOutput,InUniformExpressionSet,InMaterialShaderMapHash,InShaderPipeline,InVertexFactoryType,InDebugDescription),
 			VertexFactoryType(InVertexFactoryType)
 		{}
 	};
@@ -113,8 +112,7 @@ public:
 		const FSHAHash& MaterialShaderMapHash,
 		const FShaderCompileJob& CurrentJob,
 		const FShaderPipelineType* ShaderPipeline,
-		const FString& InDebugDescription,
-		FShaderMapResourceBuilder& ResourceBuilder
+		const FString& InDebugDescription
 		);
 
 	/**
diff --git a/Engine/Source/Runtime/RHI/Public/RHIResources.h b/Engine/Source/Runtime/RHI/Public/RHIResources.h
--- a/Engine/Source/Runtime/RHI/Public/RHIResources.h
+++ b/Engine/Source/Runtime/RHI/Public/RHIResources.h
@@ -358,6 +358,7 @@ public:
 
 	explicit FRHIUniformBufferLayout(const TCHAR* InName) :
 		ConstantBufferSize(0),
+		NumUsesForDebugging(0),
 		Name(InName),
 		Hash(0)
 	{
@@ -369,6 +370,7 @@ public:
 	};
 	explicit FRHIUniformBufferLayout(EInit) :
 		ConstantBufferSize(0),
+		NumUsesForDebugging(0),
 		Hash(0)
 	{
 	}
diff --git a/Engine/Source/Runtime/RenderCore/Private/Shader.cpp b/Engine/Source/Runtime/RenderCore/Private/Shader.cpp
--- a/Engine/Source/Runtime/RenderCore/Private/Shader.cpp
+++ b/Engine/Source/Runtime/RenderCore/Private/Shader.cpp
@@ -6,6 +6,7 @@
 
 #include "Shader.h"
 #include "Misc/CoreMisc.h"
+#include "Misc/StringBuilder.h"
 #include "Stats/StatsMisc.h"
 #include "Serialization/MemoryWriter.h"
 #include "VertexFactory.h"
@@ -39,6 +40,32 @@ IMPLEMENT_TYPE_LAYOUT(FShaderParameterInfo);
 IMPLEMENT_TYPE_LAYOUT(FShaderLooseParameterBufferInfo);
 IMPLEMENT_TYPE_LAYOUT(FShaderParameterMapInfo);
 
+void Freeze::IntrinsicToString(const TIndexedPtr<FShaderType>& Object, const FTypeLayoutDesc& TypeDesc, const FPlatformTypeLayoutParameters& LayoutParams, FMemoryToStringContext& OutContext)
+{
+	const FShaderType* Type = Object.Get(OutContext.TryGetPrevPointerTable());
+	if (Type)
+	{
+		OutContext.String->Appendf(TEXT("%s\n"), Type->GetName());
+	}
+	else
+	{
+		OutContext.AppendNullptr();
+	}
+}
+
+void Freeze::IntrinsicToString(const TIndexedPtr<FVertexFactoryType>& Object, const FTypeLayoutDesc& TypeDesc, const FPlatformTypeLayoutParameters& LayoutParams, FMemoryToStringContext& OutContext)
+{
+	const FVertexFactoryType* Type = Object.Get(OutContext.TryGetPrevPointerTable());
+	if (Type)
+	{
+		OutContext.String->Appendf(TEXT("%s\n"), Type->GetName());
+	}
+	else
+	{
+		OutContext.AppendNullptr();
+	}
+}
+
 IMPLEMENT_EXPORTED_INTRINSIC_TYPE_LAYOUT(TIndexedPtr<FShaderType>);
 IMPLEMENT_EXPORTED_INTRINSIC_TYPE_LAYOUT(TIndexedPtr<FVertexFactoryType>);
 
@@ -474,7 +501,7 @@ FShader::FShader(const CompiledShaderInitializerType& Initializer)
 	: Type(Initializer.Type)
 	, VFType(Initializer.VertexFactoryType)
 	, Target(Initializer.Target)
-	, ResourceIndex(Initializer.ResourceIndex)
+	, ResourceIndex(INDEX_NONE)
 #if WITH_EDITORONLY_DATA
 	, NumInstructions(Initializer.NumInstructions)
 	, NumTextureSamplers(Initializer.NumTextureSamplers)
@@ -518,6 +545,13 @@ FShader::~FShader()
 {
 }
 
+void FShader::Finalize(const FShaderMapResourceCode* Code)
+{
+	check(ResourceIndex == INDEX_NONE);
+	ResourceIndex = Code->FindShaderIndex(GetOutputHash());
+	check(ResourceIndex != INDEX_NONE);
+}
+
 void FShader::BuildParameterMapInfo(const TMap<FString, FParameterAllocation>& ParameterMap)
 {
 	for (int32 ParameterTypeIndex = 0; ParameterTypeIndex < (int32)EShaderParameterType::Num; ParameterTypeIndex++)
@@ -936,6 +970,17 @@ void FShaderPipeline::Validate(const FShaderPipelineType* InPipelineType) const
 	}
 }
 
+void FShaderPipeline::Finalize(FShaderMapResourceCode* Code)
+{
+	for (uint32 i = 0u; i < SF_NumGraphicsFrequencies; ++i)
+	{
+		if (Shaders[i])
+		{
+			Shaders[i]->Finalize(Code);
+		}
+	}
+}
+
 #if WITH_EDITOR
 void FShaderPipeline::SaveShaderStableKeys(const FShaderMapPointerTable& InPtrTable, EShaderPlatform TargetShaderPlatform, const struct FStableShaderKeyAndValue& InSaveKeyVal) const
 {
diff --git a/Engine/Source/Runtime/RenderCore/Private/ShaderCodeLibrary.cpp b/Engine/Source/Runtime/RenderCore/Private/ShaderCodeLibrary.cpp
--- a/Engine/Source/Runtime/RenderCore/Private/ShaderCodeLibrary.cpp
+++ b/Engine/Source/Runtime/RenderCore/Private/ShaderCodeLibrary.cpp
@@ -924,13 +924,14 @@ struct FEditorShaderCodeArchive
 					const FShaderMapResourceCode::FShaderEntry& SourceShaderEntry = Code->ShaderEntries[i];
 					FShaderCodeEntry& SerializedShaderEntry = SerializedShaders.ShaderEntries[ShaderIndex];
 					SerializedShaderEntry.Frequency = SourceShaderEntry.Frequency;
-					SerializedShaderEntry.Size = SourceShaderEntry.CompressedSize;
+					SerializedShaderEntry.Size = SourceShaderEntry.Code.Num();
 					SerializedShaderEntry.UncompressedSize = SourceShaderEntry.UncompressedSize;
-					ShaderCode.Emplace(Code->ShaderCode.GetData() + SourceShaderEntry.Offset, SourceShaderEntry.CompressedSize);
+					ShaderCode.Add(SourceShaderEntry.Code);
 					check(ShaderCode.Num() == SerializedShaders.ShaderEntries.Num());
 
 					CodeStats.NumUniqueShaders++;
-					CodeStats.ShadersUniqueSize += SourceShaderEntry.CompressedSize;
+					CodeStats.ShadersUniqueSize += SourceShaderEntry.Code.Num();
+					CodeStats.ShadersSize += SourceShaderEntry.Code.Num();
 				}
 				SerializedShaders.ShaderIndices[ShaderMapEntry.ShaderIndicesOffset + i] = ShaderIndex;
 			}
@@ -1845,7 +1846,6 @@ public:
 			static_cast<int32>(UE_ARRAY_COUNT(EditorShaderCodeStats)), static_cast<int32>(Platform));
 		FShaderCodeStats& CodeStats = EditorShaderCodeStats[Platform];
 		CodeStats.NumShaders += Code->ShaderEntries.Num();
-		CodeStats.ShadersSize += Code->ShaderCode.Num();
 
 		checkf(Platform < UE_ARRAY_COUNT(EditorShaderCodeArchive), TEXT("FShaderCodeLibrary::AddShaderCode can only be called with a valid shader platform (expected no more than %d, passed: %d)"),
 			static_cast<int32>(UE_ARRAY_COUNT(EditorShaderCodeArchive)), static_cast<int32>(Platform));
diff --git a/Engine/Source/Runtime/RenderCore/Private/ShaderMap.cpp b/Engine/Source/Runtime/RenderCore/Private/ShaderMap.cpp
--- a/Engine/Source/Runtime/RenderCore/Private/ShaderMap.cpp
+++ b/Engine/Source/Runtime/RenderCore/Private/ShaderMap.cpp
@@ -6,6 +6,7 @@
 
 #include "Shader.h"
 #include "Misc/CoreMisc.h"
+#include "Misc/StringBuilder.h"
 #include "Stats/StatsMisc.h"
 #include "Serialization/MemoryWriter.h"
 #include "VertexFactory.h"
@@ -63,8 +64,20 @@ void FShaderMapBase::AssignContent(FShaderMapContent* InContent)
 
 void FShaderMapBase::FinalizeContent()
 {
+	check(Code);
+
 	if (Content && FrozenContentSize == 0u)
 	{
+		for (FShader* Shader : Content->GetShaders())
+		{
+			Shader->Finalize(Code);
+		}
+
+		for (FShaderPipeline* Pipeline : Content->GetShaderPipelines())
+		{
+			Pipeline->Finalize(Code);
+		}
+
 		Content->Validate(*this);
 
 		FMemoryImage MemoryImage;
@@ -93,7 +106,6 @@ void FShaderMapBase::FinalizeContent()
 		INC_DWORD_STAT_BY(STAT_Shaders_NumShadersLoaded, NumFrozenShaders);
 	}
 
-	check(Code);
 	Code->Finalize();
 	Resource = new FShaderMapResource_InlineCode(GetShaderPlatform(), Code);
 	BeginInitResource(Resource);
@@ -290,6 +302,28 @@ bool FShaderMapBase::Serialize(FArchive& Ar, bool bInlineShaderResources, bool b
 	return bContentValid;
 }
 
+FString FShaderMapBase::ToString() const
+{
+	TStringBuilder<32000> String;
+	{
+		FMemoryToStringContext Context;
+		Context.PrevPointerTable = PointerTable;
+		Context.String = &String;
+
+		FPlatformTypeLayoutParameters LayoutParams;
+		LayoutParams.InitializeForCurrent();
+
+		ContentTypeLayout.ToStringFunc(Content, ContentTypeLayout, LayoutParams, Context);
+	}
+
+	if (Code)
+	{
+		Code->ToString(String);
+	}
+
+	return String.ToString();
+}
+
 void FShaderMapBase::DestroyContent()
 {
 	if (Content)
@@ -680,6 +714,7 @@ void FShaderMapContent::Finalize()
 	for (int32 SortedIndex = 0; SortedIndex < SortedEntries.Num(); ++SortedIndex)
 	{
 		const FSortedShaderEntry& SortedEntry = SortedEntries[SortedIndex];
+
 		const uint16 Key = MakeShaderHash(SortedEntry.TypeName, SortedEntry.PermutationId);
 		NewShaders.Add(Shaders[SortedEntry.Index]);
 		ShaderTypes.Add(SortedEntry.TypeName);
diff --git a/Engine/Source/Runtime/RenderCore/Private/ShaderResource.cpp b/Engine/Source/Runtime/RenderCore/Private/ShaderResource.cpp
--- a/Engine/Source/Runtime/RenderCore/Private/ShaderResource.cpp
+++ b/Engine/Source/Runtime/RenderCore/Private/ShaderResource.cpp
@@ -6,6 +6,7 @@
 
 #include "Shader.h"
 #include "Misc/CoreMisc.h"
+#include "Misc/StringBuilder.h"
 #include "Stats/StatsMisc.h"
 #include "Serialization/MemoryWriter.h"
 #include "VertexFactory.h"
@@ -110,9 +111,9 @@ static void ApplyResourceStats(FShaderMapResourceCode& Resource)
 {
 #if STATS
 	INC_DWORD_STAT_BY(STAT_Shaders_ShaderResourceMemory, Resource.GetSizeBytes());
-	for (FShaderMapResourceCode::FShaderEntry& Shader : Resource.ShaderEntries)
+	for (const FShaderMapResourceCode::FShaderEntry& Shader : Resource.ShaderEntries)
 	{
-		INC_DWORD_STAT_BY_FName(GetMemoryStatType(Shader.Frequency).GetName(), Shader.CompressedSize);
+		INC_DWORD_STAT_BY_FName(GetMemoryStatType(Shader.Frequency).GetName(), Shader.Code.Num());
 	}
 #endif // STATS
 }
@@ -121,9 +122,9 @@ static void RemoveResourceStats(FShaderMapResourceCode& Resource)
 {
 #if STATS
 	DEC_DWORD_STAT_BY(STAT_Shaders_ShaderResourceMemory, Resource.GetSizeBytes());
-	for (FShaderMapResourceCode::FShaderEntry& Shader : Resource.ShaderEntries)
+	for (const FShaderMapResourceCode::FShaderEntry& Shader : Resource.ShaderEntries)
 	{
-		DEC_DWORD_STAT_BY_FName(GetMemoryStatType(Shader.Frequency).GetName(), Shader.CompressedSize);
+		DEC_DWORD_STAT_BY_FName(GetMemoryStatType(Shader.Frequency).GetName(), Shader.Code.Num());
 	}
 #endif // STATS
 }
@@ -142,18 +143,100 @@ void FShaderMapResourceCode::Finalize()
 	ApplyResourceStats(*this);
 }
 
+uint32 FShaderMapResourceCode::GetSizeBytes() const
+{
+	uint32 Size = sizeof(*this) + ShaderHashes.GetAllocatedSize() + ShaderEntries.GetAllocatedSize();
+	for (const FShaderEntry& Entry : ShaderEntries)
+	{
+		Size += Entry.Code.GetAllocatedSize();
+	}
+	return Size;
+}
+
+int32 FShaderMapResourceCode::FindShaderIndex(const FSHAHash& InHash) const
+{
+	return Algo::BinarySearch(ShaderHashes, InHash);
+}
+
+void FShaderMapResourceCode::AddShaderCode(EShaderFrequency InFrequency, const FSHAHash& InHash, TConstArrayView<uint8> InCode)
+{
+	const int32 Index = Algo::LowerBound(ShaderHashes, InHash);
+	if (Index >= ShaderHashes.Num() || ShaderHashes[Index] != InHash)
+	{
+		ShaderHashes.Insert(InHash, Index);
+
+		FShaderEntry& Entry = ShaderEntries.InsertDefaulted_GetRef(Index);
+		Entry.Frequency = InFrequency;
+		Entry.UncompressedSize = InCode.Num();
+
+		bool bAllowShaderCompression = true;
+#if !(UE_BUILD_SHIPPING || UE_BUILD_TEST)
+		static const IConsoleVariable* CVarSkipCompression = IConsoleManager::Get().FindConsoleVariable(TEXT("r.Shaders.SkipCompression"));
+		bAllowShaderCompression = CVarSkipCompression ? CVarSkipCompression->GetInt() == 0 : true;
+#endif
+
+		int32 CompressedSize = InCode.Num();
+		Entry.Code.AddUninitialized(CompressedSize);
+
+		if (bAllowShaderCompression && FCompression::CompressMemory(ShaderCompressionFormat, Entry.Code.GetData(), CompressedSize, InCode.GetData(), InCode.Num()))
+		{
+			// resize to fit reduced compressed size, but don't reallocate memory
+			Entry.Code.SetNum(CompressedSize, false);
+		}
+		else
+		{
+			FMemory::Memcpy(Entry.Code.GetData(), InCode.GetData(), InCode.Num());
+		}
+	}
+}
+
+#if WITH_EDITORONLY_DATA
+void FShaderMapResourceCode::AddPlatformDebugData(TConstArrayView<uint8> InPlatformDebugData)
+{
+	if (InPlatformDebugData.Num() == 0)
+	{
+		return;
+	}
+
+	FSHAHash Hash;
+	{
+		FSHA1 Hasher;
+		Hasher.Update(InPlatformDebugData.GetData(), InPlatformDebugData.Num());
+		Hasher.Final();
+		Hasher.GetHash(Hash.Hash);
+	}
+
+	const int32 Index = Algo::LowerBound(PlatformDebugDataHashes, Hash);
+	if (Index >= PlatformDebugDataHashes.Num() || PlatformDebugDataHashes[Index] != Hash)
+	{
+		PlatformDebugDataHashes.Insert(Hash, Index);
+		PlatformDebugData.EmplaceAt(Index, InPlatformDebugData.GetData(), InPlatformDebugData.Num());
+	}
+}
+#endif // WITH_EDITORONLY_DATA
+
+void FShaderMapResourceCode::ToString(FStringBuilderBase& OutString) const
+{
+	OutString.Appendf(TEXT("Shaders: Num=%d\n"), ShaderHashes.Num());
+	for (int32 i = 0; i < ShaderHashes.Num(); ++i)
+	{
+		const FShaderEntry& Entry = ShaderEntries[i];
+		OutString.Appendf(TEXT("    [%d]: { Hash: %s, Freq: %s, Size: %d, UncompressedSize: %d }\n"),
+			i, *ShaderHashes[i].ToString(), GetShaderFrequencyString(Entry.Frequency), Entry.Code.Num(), Entry.UncompressedSize);
+	}
+}
+
 void FShaderMapResourceCode::Serialize(FArchive& Ar, bool bLoadedByCookedMaterial)
 {
 	Ar << ResourceHash;
 	Ar << ShaderHashes;
 	Ar << ShaderEntries;
-	Ar << ShaderCode;
 	check(ShaderEntries.Num() == ShaderHashes.Num());
 #if WITH_EDITORONLY_DATA
 	const bool bSerializePlatformData = !bLoadedByCookedMaterial && (!Ar.IsCooking() || Ar.CookingTarget()->HasEditorOnlyData());
 	if (bSerializePlatformData)
 	{
-		Ar << PlatformDebugEntries;
+		Ar << PlatformDebugDataHashes;
 		Ar << PlatformDebugData;
 	}
 #endif // WITH_EDITORONLY_DATA
@@ -167,7 +250,7 @@ void FShaderMapResourceCode::NotifyShadersCooked(const ITargetPlatform* TargetPl
 	// Notify the platform shader format that this particular shader is being used in the cook.
 	// We discard this data in cooked builds unless Ar.CookingTarget()->HasEditorOnlyData() is true.
 	check(TargetPlatform);
-	if (PlatformDebugEntries.Num())
+	if (PlatformDebugData.Num())
 	{
 		TArray<FName> ShaderFormatNames;
 		TargetPlatform->GetAllTargetedShaderFormats(ShaderFormatNames);
@@ -176,9 +259,9 @@ void FShaderMapResourceCode::NotifyShadersCooked(const ITargetPlatform* TargetPl
 			const IShaderFormat* ShaderFormat = GetTargetPlatformManagerRef().FindShaderFormat(FormatName);
 			if (ShaderFormat)
 			{
-				for (const FPlatformDebugEntry& Entry : PlatformDebugEntries)
+				for (const auto& Entry : PlatformDebugData)
 				{
-					ShaderFormat->NotifyShaderCooked(MakeArrayView(PlatformDebugData.GetData() + Entry.Offset, Entry.Size), FormatName);
+					ShaderFormat->NotifyShaderCooked(Entry, FormatName);
 				}
 			}
 		}
@@ -187,83 +270,6 @@ void FShaderMapResourceCode::NotifyShadersCooked(const ITargetPlatform* TargetPl
 }
 #endif // WITH_EDITORONLY_DATA
 
-FShaderMapResourceBuilder::FShaderMapResourceBuilder(FShaderMapResourceCode* InCode) : ShaderHashTable(1024, 256), Code(InCode)
-{
-	check(InCode->ShaderEntries.Num() == InCode->ShaderHashes.Num());
-	for(int32 ShaderIndex = 0; ShaderIndex < InCode->ShaderHashes.Num(); ++ShaderIndex)
-	{
-		const uint32 Key = GetTypeHash(InCode->ShaderHashes[ShaderIndex]);
-		ShaderHashTable.Add(Key, ShaderIndex);
-	}
-}
-
-int32 FShaderMapResourceBuilder::FindCode(const FSHAHash& InHash, uint32 InKey) const
-{
-	for (int32 Index = ShaderHashTable.First(InKey); ShaderHashTable.IsValid(Index); Index = ShaderHashTable.Next(Index))
-	{
-		if (Code->ShaderHashes[Index] == InHash)
-		{
-			return Index;
-		}
-	}
-	return INDEX_NONE;
-}
-
-int32 FShaderMapResourceBuilder::FindOrAddCode(EShaderFrequency InFrequency, const FSHAHash& InHash, const TConstArrayView<uint8>& InCode)
-{
-	const uint32 Key = GetTypeHash(InHash);
-	int32 Index = FindCode(InHash, Key);
-	if(Index == INDEX_NONE)
-	{
-		bool bAllowShaderCompression = true;
-
-#if !(UE_BUILD_SHIPPING || UE_BUILD_TEST)
-		static const IConsoleVariable* CVarSkipCompression = IConsoleManager::Get().FindConsoleVariable(TEXT("r.Shaders.SkipCompression"));
-		bAllowShaderCompression = CVarSkipCompression ? CVarSkipCompression->GetInt() == 0 : true;
-#endif
-
-		int32 CompressedSize = InCode.Num();
-		TArray<uint8> CompressedCode;
-		CompressedCode.AddUninitialized(CompressedSize);
-
-		if (bAllowShaderCompression && FCompression::CompressMemory(ShaderCompressionFormat, CompressedCode.GetData(), CompressedSize, InCode.GetData(), InCode.Num()))
-		{
-			CompressedCode.SetNum(CompressedSize, false);
-		}
-		else
-		{
-			CompressedCode = InCode;
-		}
-
-		Index = Code->ShaderEntries.Num();
-		ShaderHashTable.Add(Key, Index);
-
-		Code->ShaderHashes.Add(InHash);
-		FShaderMapResourceCode::FShaderEntry& Entry = Code->ShaderEntries.AddDefaulted_GetRef();
-		Entry.CompressedSize = CompressedSize;
-		Entry.UncompressedSize = InCode.Num();
-		Entry.Frequency = InFrequency;
-		Entry.Offset = Code->ShaderCode.Num();
-
-		Code->ShaderCode.Append(CompressedCode);
-	}
-
-	return Index;
-}
-
-#if WITH_EDITORONLY_DATA
-void FShaderMapResourceBuilder::AddPlatformDebugData(TConstArrayView<uint8> InPlatformDebugData)
-{
-	if (InPlatformDebugData.Num() > 0)
-	{
-		FShaderMapResourceCode::FPlatformDebugEntry& Entry = Code->PlatformDebugEntries.AddDefaulted_GetRef();
-		Entry.Offset = Code->PlatformDebugData.Num();
-		Entry.Size = InPlatformDebugData.Num();
-		Code->PlatformDebugData.Append(InPlatformDebugData.GetData(), InPlatformDebugData.Num());
-	}
-}
-#endif // WITH_EDITORONLY_DATA
-
 FShaderMapResource::FShaderMapResource(EShaderPlatform InPlatform, int32 NumShaders)
 	: Platform(InPlatform)
 	, NumRefs(0)
@@ -354,13 +360,13 @@ TRefCountPtr<FRHIShader> FShaderMapResource_InlineCode::CreateRHIShader(int32 Sh
 
 	FMemStackBase& MemStack = FMemStack::Get();
 	const FShaderMapResourceCode::FShaderEntry& ShaderEntry = Code->ShaderEntries[ShaderIndex];
-	const uint8* ShaderCode = Code->ShaderCode.GetData() + ShaderEntry.Offset;
+	const uint8* ShaderCode = ShaderEntry.Code.GetData();
 
 	FMemMark Mark(MemStack);
-	if (ShaderEntry.CompressedSize != ShaderEntry.UncompressedSize)
+	if (ShaderEntry.Code.Num() != ShaderEntry.UncompressedSize)
 	{
 		void* UncompressedCode = MemStack.Alloc(ShaderEntry.UncompressedSize, 16);
-		auto bSucceed = FCompression::UncompressMemory(ShaderCompressionFormat, UncompressedCode, ShaderEntry.UncompressedSize, ShaderCode, ShaderEntry.CompressedSize);
+		auto bSucceed = FCompression::UncompressMemory(ShaderCompressionFormat, UncompressedCode, ShaderEntry.UncompressedSize, ShaderCode, ShaderEntry.Code.Num());
 		check(bSucceed);
 		ShaderCode = (uint8*)UncompressedCode;
 	}
diff --git a/Engine/Source/Runtime/RenderCore/Public/GlobalShader.h b/Engine/Source/Runtime/RenderCore/Public/GlobalShader.h
--- a/Engine/Source/Runtime/RenderCore/Public/GlobalShader.h
+++ b/Engine/Source/Runtime/RenderCore/Public/GlobalShader.h
@@ -130,28 +130,7 @@ class RENDERCORE_API FGlobalShaderMapSection : public TShaderMap<FGlobalShaderMa
 public:
 	static FGlobalShaderMapSection* CreateFromArchive(FArchive& Ar);
 
-	~FGlobalShaderMapSection() { ReleaseResourceBuilder(); }
-
 	bool Serialize(FArchive& Ar);
-
-	FShaderMapResourceBuilder& GetResourceBuilder()
-	{
-		if (ResourceBuilder == nullptr)
-		{
-			ResourceBuilder = new FShaderMapResourceBuilder(GetResourceCode());
-		}
-		return *ResourceBuilder;
-	}
-
-	void ReleaseResourceBuilder()
-	{
-		if (ResourceBuilder)
-		{
-			delete ResourceBuilder;
-			ResourceBuilder = nullptr;
-		}
-	}
-
 private:
 	inline FGlobalShaderMapSection() {}
 
@@ -162,8 +141,6 @@ private:
 
 	TShaderRef<FShader> GetShader(FShaderType* ShaderType, int32 PermutationId = 0) const;
 	FShaderPipelineRef GetShaderPipeline(const FShaderPipelineType* PipelineType) const;
-
-	FShaderMapResourceBuilder* ResourceBuilder = nullptr;
 };
 
 class RENDERCORE_API FGlobalShaderMap
diff --git a/Engine/Source/Runtime/RenderCore/Public/Shader.h b/Engine/Source/Runtime/RenderCore/Public/Shader.h
--- a/Engine/Source/Runtime/RenderCore/Public/Shader.h
+++ b/Engine/Source/Runtime/RenderCore/Public/Shader.h
@@ -281,15 +281,14 @@ class FShaderMapResourceCode : public FThreadSafeRefCountedObject
 public:
 	struct FShaderEntry
 	{
-		uint32 Offset;
-		uint32 UncompressedSize;
-		uint32 CompressedSize;
+		TArray<uint8> Code;
+		int32 UncompressedSize;
 		EShaderFrequency Frequency;
 
 		friend FArchive& operator<<(FArchive& Ar, FShaderEntry& Entry)
 		{
 			uint8 Freq = Entry.Frequency;
-			Ar << Entry.Offset << Entry.UncompressedSize << Entry.CompressedSize << Freq;
+			Ar << Entry.Code << Entry.UncompressedSize << Freq;
 			Entry.Frequency = (EShaderFrequency)Freq;
 			return Ar;
 		}
@@ -317,18 +316,31 @@ public:
 	RENDERCORE_API void NotifyShadersCooked(const ITargetPlatform* TargetPlatform);
 #endif // WITH_EDITORONLY_DATA
 
-	inline const uint32 GetSizeBytes() const
+	RENDERCORE_API uint32 GetSizeBytes() const;
+
+	inline void AddShaderCompilerOutput(const FShaderCompilerOutput& Output)
 	{
-		return sizeof(*this) + ShaderHashes.GetAllocatedSize() + ShaderEntries.GetAllocatedSize() + ShaderCode.GetAllocatedSize();
+#if WITH_EDITORONLY_DATA
+		AddPlatformDebugData(Output.PlatformDebugData);
+#endif
+		AddShaderCode(Output.Target.GetFrequency(), Output.OutputHash, Output.ShaderCode.GetReadAccess());
 	}
 
+	int32 FindShaderIndex(const FSHAHash& InHash) const;
+
+	RENDERCORE_API void AddShaderCode(EShaderFrequency InFrequency, const FSHAHash& InHash, TConstArrayView<uint8> InCode);
+#if WITH_EDITORONLY_DATA
+	RENDERCORE_API void AddPlatformDebugData(TConstArrayView<uint8> InPlatformDebugData);
+#endif
+
+	RENDERCORE_API void ToString(FStringBuilderBase& OutString) const;
+
 	FSHAHash ResourceHash;
 	TArray<FSHAHash> ShaderHashes;
 	TArray<FShaderEntry> ShaderEntries;
-	TArray<uint8> ShaderCode;
 #if WITH_EDITORONLY_DATA
-	TArray<FPlatformDebugEntry> PlatformDebugEntries;
-	TArray<uint8> PlatformDebugData;
+	TArray<TArray<uint8>> PlatformDebugData;
+	TArray<FSHAHash> PlatformDebugDataHashes;
 #endif // WITH_EDITORONLY_DATA
 };
 	
@@ -347,29 +359,6 @@ public:
 	TRefCountPtr<FShaderMapResourceCode> Code;
 };
 
-class FShaderMapResourceBuilder
-{
-public:
-	RENDERCORE_API explicit FShaderMapResourceBuilder(FShaderMapResourceCode* InCode);
-
-	inline int32 FindOrAddCode(const FShaderCompilerOutput& Output)
-	{
-#if WITH_EDITORONLY_DATA
-		AddPlatformDebugData(Output.PlatformDebugData);
-#endif
-		return FindOrAddCode(Output.Target.GetFrequency(), Output.OutputHash, Output.ShaderCode.GetReadAccess());
-	}
-
-	RENDERCORE_API int32 FindCode(const FSHAHash& InHash, uint32 InKey) const;
-	RENDERCORE_API int32 FindOrAddCode(EShaderFrequency InFrequency, const FSHAHash& InHash, const TConstArrayView<uint8>& InCode);
-#if WITH_EDITORONLY_DATA
-	RENDERCORE_API void AddPlatformDebugData(TConstArrayView<uint8> InPlatformDebugData);
-#endif // WITH_EDITORONLY_DATA
-private:
-	FHashTable ShaderHashTable;
-	TRefCountPtr<FShaderMapResourceCode> Code;
-};
-
 class RENDERCORE_API FShaderMapPointerTable : public FPointerTableBase
 {
 public:
@@ -622,7 +611,6 @@ struct FShadereCompiledShaderInitializerType
 	FSHAHash MaterialShaderMapHash;
 	const FShaderPipelineType* ShaderPipeline;
 	FVertexFactoryType* VertexFactoryType;
-	int32 ResourceIndex;
 	uint32 NumInstructions;
 	uint32 NumTextureSamplers;
 	uint32 CodeSize;
@@ -632,7 +620,6 @@ struct FShadereCompiledShaderInitializerType
 		FShaderType* InType,
 		int32 InPermutationId,
 		const FShaderCompilerOutput& CompilerOutput,
-		int32 InResourceIndex,
 		const FSHAHash& InMaterialShaderMapHash,
 		const FShaderPipelineType* InShaderPipeline,
 		FVertexFactoryType* InVertexFactoryType
@@ -645,7 +632,6 @@ struct FShadereCompiledShaderInitializerType
 		MaterialShaderMapHash(InMaterialShaderMapHash),
 		ShaderPipeline(InShaderPipeline),
 		VertexFactoryType(InVertexFactoryType),
-		ResourceIndex(InResourceIndex),
 		NumInstructions(CompilerOutput.NumInstructions),
 		NumTextureSamplers(CompilerOutput.NumTextureSamplers),
 		CodeSize(CompilerOutput.ShaderCode.GetShaderCodeSize()),
@@ -653,6 +639,12 @@ struct FShadereCompiledShaderInitializerType
 	{}
 };
 
+namespace Freeze
+{
+	RENDERCORE_API void IntrinsicToString(const TIndexedPtr<FShaderType>& Object, const FTypeLayoutDesc& TypeDesc, const FPlatformTypeLayoutParameters& LayoutParams, FMemoryToStringContext& OutContext);
+	RENDERCORE_API void IntrinsicToString(const TIndexedPtr<FVertexFactoryType>& Object, const FTypeLayoutDesc& TypeDesc, const FPlatformTypeLayoutParameters& LayoutParams, FMemoryToStringContext& OutContext);
+}
+
 DECLARE_EXPORTED_TEMPLATE_INTRINSIC_TYPE_LAYOUT(template<>, TIndexedPtr<FShaderType>, RENDERCORE_API);
 DECLARE_EXPORTED_TEMPLATE_INTRINSIC_TYPE_LAYOUT(template<>, TIndexedPtr<FVertexFactoryType>, RENDERCORE_API);
 
@@ -696,13 +688,15 @@ public:
 	
 	const FSHAHash& GetOutputHash() const;
 
+	void Finalize(const FShaderMapResourceCode* Code);
+
 	// Accessors.
 	inline FShaderType* GetType(const FShaderMapPointerTable& InPointerTable) const { return Type.Get(InPointerTable.ShaderTypes); }
 	inline FShaderType* GetType(const FPointerTableBase* InPointerTable) const { return Type.Get(InPointerTable); }
 	inline FVertexFactoryType* GetVertexFactoryType(const FShaderMapPointerTable& InPointerTable) const { return VFType.Get(InPointerTable.VFTypes); }
 	inline FVertexFactoryType* GetVertexFactoryType(const FPointerTableBase* InPointerTable) const { return VFType.Get(InPointerTable); }
 	inline FShaderType* GetTypeUnfrozen() const { return Type.GetUnfrozen(); }
-	inline int32 GetResourceIndex() const { return ResourceIndex; }
+	inline int32 GetResourceIndex() const { checkSlow(ResourceIndex != INDEX_NONE); return ResourceIndex; }
 	inline EShaderPlatform GetShaderPlatform() const { return Target.GetPlatform(); }
 	inline EShaderFrequency GetFrequency() const { return Target.GetFrequency(); }
 	inline const FShaderTarget GetTarget() const { return Target; }
@@ -1661,6 +1655,8 @@ public:
 
 	void Validate(const FShaderPipelineType* InPipelineType) const;
 
+	void Finalize(FShaderMapResourceCode* Code);
+
 	enum EFilter
 	{
 		EAll,			// All pipelines
@@ -1905,6 +1901,8 @@ public:
 	void UnfreezeContent();
 	bool Serialize(FArchive& Ar, bool bInlineShaderResources, bool bLoadedByCookedMaterial);
 
+	FString ToString() const;
+
 #if WITH_EDITOR
 	inline void GetOutdatedTypes(TArray<const FShaderType*>& OutdatedShaderTypes, TArray<const FShaderPipelineType*>& OutdatedShaderPipelineTypes, TArray<const FVertexFactoryType*>& OutdatedFactoryTypes) const
 	{
