Added some placeholder slots in EIOSDevice so we can add newly released iDevice support in a hotfix. To be cleaned up in UE-92275.

diff --git a/Engine/Source/Runtime/ApplicationCore/Private/IOS/IOSPlatformApplicationMisc.cpp b/Engine/Source/Runtime/ApplicationCore/Private/IOS/IOSPlatformApplicationMisc.cpp
--- a/Engine/Source/Runtime/ApplicationCore/Private/IOS/IOSPlatformApplicationMisc.cpp
+++ b/Engine/Source/Runtime/ApplicationCore/Private/IOS/IOSPlatformApplicationMisc.cpp
@@ -258,7 +258,7 @@ void FIOSPlatformApplicationMisc::ClipboardPaste(class FString& Result)
 EScreenPhysicalAccuracy FIOSPlatformApplicationMisc::ComputePhysicalScreenDensity(int32& ScreenDensity)
 {
 	FPlatformMisc::EIOSDevice Device = FPlatformMisc::GetIOSDeviceType();
-	static_assert( FPlatformMisc::EIOSDevice::IOS_Unknown == 46, "Every device needs to be handled here." );
+	static_assert( FPlatformMisc::EIOSDevice::IOS_Unknown == 54, "Every device needs to be handled here." );
 
 	ScreenDensity = 0;
 	EScreenPhysicalAccuracy Accuracy = EScreenPhysicalAccuracy::Unknown;
@@ -331,6 +331,14 @@ EScreenPhysicalAccuracy FIOSPlatformApplicationMisc::ComputePhysicalScreenDensit
 	case FPlatformMisc::IOS_AppleTV4K:
 		Accuracy = EScreenPhysicalAccuracy::Unknown;
 		break;
+	case FPlatformMisc::IOS_NewDevice1:
+	case FPlatformMisc::IOS_NewDevice2:
+	case FPlatformMisc::IOS_NewDevice3:
+	case FPlatformMisc::IOS_NewDevice4:
+	case FPlatformMisc::IOS_NewDevice5:
+	case FPlatformMisc::IOS_NewDevice6:
+	case FPlatformMisc::IOS_NewDevice7:
+	case FPlatformMisc::IOS_NewDevice8:
 	default:
 		// If we don't know, assume that the density is a multiple of the 
 		// native Content Scaling Factor.  Won't be exact, but should be close enough.
diff --git a/Engine/Source/Runtime/Core/Private/IOS/IOSPlatformMisc.cpp b/Engine/Source/Runtime/Core/Private/IOS/IOSPlatformMisc.cpp
--- a/Engine/Source/Runtime/Core/Private/IOS/IOSPlatformMisc.cpp
+++ b/Engine/Source/Runtime/Core/Private/IOS/IOSPlatformMisc.cpp
@@ -337,6 +337,72 @@ bool FIOSPlatformMisc::HasPlatformFeature(const TCHAR* FeatureName)
 	return FGenericPlatformMisc::HasPlatformFeature(FeatureName);
 }
 
+const TCHAR* FIOSPlatformMisc::GetDefaultDeviceProfileName()
+{
+	static const TCHAR* IOSDeviceNames[] =
+	{
+		TEXT("IPhone4"),
+		TEXT("IPhone4S"),
+		TEXT("IPhone5"),
+		TEXT("IPhone5S"),
+		TEXT("IPodTouch5"),
+		TEXT("IPodTouch6"),
+		TEXT("IPad2"),
+		TEXT("IPad3"),
+		TEXT("IPad4"),
+		TEXT("IPadMini"),
+		TEXT("IPadMini2"),
+		TEXT("IPadMini4"),
+		TEXT("IPadAir"),
+		TEXT("IPadAir2"),
+		TEXT("IPhone6"),
+		TEXT("IPhone6Plus"),
+		TEXT("IPhone6S"),
+		TEXT("IPhone6SPlus"),
+		TEXT("IPhone7"),
+		TEXT("IPhone7Plus"),
+		TEXT("IPhone8"),
+		TEXT("IPhone8Plus"),
+		TEXT("IPhoneX"),
+		TEXT("IPadPro"),
+		TEXT("AppleTV"),
+		TEXT("AppleTV4K"),
+		TEXT("IPhoneSE"),
+		TEXT("IPadPro129"),
+		TEXT("IPadPro97"),
+		TEXT("IPadPro105"),
+		TEXT("IPadPro2_129"),
+		TEXT("IPad5"),
+		TEXT("IPhoneXS"),
+		TEXT("IPhoneXSMax"),
+		TEXT("IPhoneXR"),
+		TEXT("IPhone11"),
+		TEXT("IPhone11Pro"),
+		TEXT("IPhone11ProMax"),
+		TEXT("IPad6"),
+		TEXT("IPadPro11"),
+		TEXT("IPadPro3_129"),
+		TEXT("IPadAir3"),
+		TEXT("IPadMini5"),
+		TEXT("IPodTouch7"),
+		TEXT("IPad7"),
+		TEXT("IPhoneSE2"),
+		TEXT("NewIDevice1"),
+		TEXT("NewIDevice2"),
+		TEXT("NewIDevice3"),
+		TEXT("NewIDevice4"),
+		TEXT("NewIDevice5"),
+		TEXT("NewIDevice6"),
+		TEXT("NewIDevice7"),
+		TEXT("NewIDevice8"),
+		TEXT("Unknown"),
+	};
+	static_assert((sizeof(IOSDeviceNames) / sizeof(IOSDeviceNames[0])) == ((int32)IOS_Unknown + 1), "Mismatched IOSDeviceNames and EIOSDevice.");
+
+	// look up into the string array by the enum
+	return IOSDeviceNames[(int32)GetIOSDeviceType()];
+}
+
 FString GetIOSDeviceIDString()
 {
 	static FString CachedResult;
diff --git a/Engine/Source/Runtime/Core/Public/IOS/IOSPlatformMisc.h b/Engine/Source/Runtime/Core/Public/IOS/IOSPlatformMisc.h
--- a/Engine/Source/Runtime/Core/Public/IOS/IOSPlatformMisc.h
+++ b/Engine/Source/Runtime/Core/Public/IOS/IOSPlatformMisc.h
@@ -156,68 +156,24 @@ struct CORE_API FIOSPlatformMisc : public FApplePlatformMisc
 		IOS_IPodTouch7,
 		IOS_IPad7,
 		IOS_IPhoneSE2,
+
+		// We can use the entries below for any iOS devices released during the hotfix cycle
+		// They should be moved to real device enum above these values in the next full release.
+		IOS_NewDevice1,
+		IOS_NewDevice2,
+		IOS_NewDevice3,
+		IOS_NewDevice4,
+		IOS_NewDevice5,
+		IOS_NewDevice6,
+		IOS_NewDevice7,
+		IOS_NewDevice8,
+
 		IOS_Unknown,
 	};
 
 	static EIOSDevice GetIOSDeviceType();
 
-	static FORCENOINLINE const TCHAR* GetDefaultDeviceProfileName()
-	{
-		static const TCHAR* IOSDeviceNames[] = 
-		{
-			TEXT("IPhone4"),
-			TEXT("IPhone4S"),
-			TEXT("IPhone5"),
-			TEXT("IPhone5S"),
-			TEXT("IPodTouch5"),
-			TEXT("IPodTouch6"),
-			TEXT("IPad2"),
-			TEXT("IPad3"),
-			TEXT("IPad4"),
-			TEXT("IPadMini"),
-			TEXT("IPadMini2"),
-			TEXT("IPadMini4"),
-			TEXT("IPadAir"),
-			TEXT("IPadAir2"),
-			TEXT("IPhone6"),
-			TEXT("IPhone6Plus"),
-			TEXT("IPhone6S"),
-			TEXT("IPhone6SPlus"),
-			TEXT("IPhone7"),
-			TEXT("IPhone7Plus"),
-			TEXT("IPhone8"),
-			TEXT("IPhone8Plus"),
-			TEXT("IPhoneX"),
-			TEXT("IPadPro"),
-			TEXT("AppleTV"),
-			TEXT("AppleTV4K"),
-			TEXT("IPhoneSE"),
-			TEXT("IPadPro129"),
-			TEXT("IPadPro97"),
-			TEXT("IPadPro105"),
-			TEXT("IPadPro2_129"),
-			TEXT("IPad5"),
-            TEXT("IPhoneXS"),
-            TEXT("IPhoneXSMax"),
-            TEXT("IPhoneXR"),
-			TEXT("IPhone11"),
-			TEXT("IPhone11Pro"),
-			TEXT("IPhone11ProMax"),
-			TEXT("IPad6"),
-			TEXT("IPadPro11"),
-			TEXT("IPadPro3_129"),
-            TEXT("IPadAir3"),
-            TEXT("IPadMini5"),
-			TEXT("IPodTouch7"),
-			TEXT("IPad7"),
-			TEXT("IPhoneSE2"),
-			TEXT("Unknown"),
-		};
-		static_assert((sizeof(IOSDeviceNames) / sizeof(IOSDeviceNames[0])) == ((int32)IOS_Unknown + 1), "Mismatched IOSDeviceNames and EIOSDevice.");
-		
-		// look up into the string array by the enum
-		return IOSDeviceNames[(int32)GetIOSDeviceType()];
-	}
+	static const TCHAR* GetDefaultDeviceProfileName();
 
 	static FString GetCPUVendor();
 	static FString GetCPUBrand();
