M - More renameing the Octree to Octree2 in preparation to add a deprecation for the original class

diff --git a/Engine/Plugins/MeshPainting/Source/MeshPaintingToolset/Private/MeshPaintHelpers.cpp b/Engine/Plugins/MeshPainting/Source/MeshPaintingToolset/Private/MeshPaintHelpers.cpp
--- a/Engine/Plugins/MeshPainting/Source/MeshPaintingToolset/Private/MeshPaintHelpers.cpp
+++ b/Engine/Plugins/MeshPainting/Source/MeshPaintingToolset/Private/MeshPaintHelpers.cpp
@@ -1244,7 +1244,7 @@ struct FVertexColorPropogationOctreeSemantics
 	}
 
 	/** Ignored for this implementation */
-	FORCEINLINE static void SetElementId( const FPaintedMeshVertex& Element, FOctreeElementId Id )
+	FORCEINLINE static void SetElementId( const FPaintedMeshVertex& Element, FOctreeElementId2 Id )
 	{
 	}
 };
diff --git a/Engine/Plugins/MeshPainting/Source/MeshPaintingToolset/Public/TMeshPaintOctree.h b/Engine/Plugins/MeshPainting/Source/MeshPaintingToolset/Public/TMeshPaintOctree.h
--- a/Engine/Plugins/MeshPainting/Source/MeshPaintingToolset/Public/TMeshPaintOctree.h
+++ b/Engine/Plugins/MeshPainting/Source/MeshPaintingToolset/Public/TMeshPaintOctree.h
@@ -51,7 +51,7 @@ struct FMeshPaintTriangleOctreeSemantics
 	}
 
 	/** Ignored for this implementation */
-	FORCEINLINE static void SetElementId(const FMeshPaintTriangle& Element, FOctreeElementId Id)
+	FORCEINLINE static void SetElementId(const FMeshPaintTriangle& Element, FOctreeElementId2 Id)
 	{
 
 	}
diff --git a/Engine/Plugins/Runtime/EditableMesh/Source/EditableMesh/EditableMesh.cpp b/Engine/Plugins/Runtime/EditableMesh/Source/EditableMesh/EditableMesh.cpp
--- a/Engine/Plugins/Runtime/EditableMesh/Source/EditableMesh/EditableMesh.cpp
+++ b/Engine/Plugins/Runtime/EditableMesh/Source/EditableMesh/EditableMesh.cpp
@@ -512,7 +512,7 @@ void UEditableMesh::RemapOctreeIDs( const FElementIDRemappings& Remappings )
 
 	{
 		// Make a temporary copy of the original map, and clear the stored version
-		static TMap<FPolygonID, FOctreeElementId> OldPolygonIDToOctreeElementIDMap;
+		static TMap<FPolygonID, FOctreeElementId2> OldPolygonIDToOctreeElementIDMap;
 		OldPolygonIDToOctreeElementIDMap = PolygonIDToOctreeElementIDMap;
 		PolygonIDToOctreeElementIDMap.Reset();
 
@@ -520,7 +520,7 @@ void UEditableMesh::RemapOctreeIDs( const FElementIDRemappings& Remappings )
 		for( const auto& OldPolygonIDAndOctreeElementID : OldPolygonIDToOctreeElementIDMap )
 		{
 			const FPolygonID OldPolygonID = OldPolygonIDAndOctreeElementID.Key;
-			const FOctreeElementId OctreeElementID = OldPolygonIDAndOctreeElementID.Value;
+			const FOctreeElementId2 OctreeElementID = OldPolygonIDAndOctreeElementID.Value;
 
 			const FPolygonID NewPolygonID = Remappings.GetRemappedPolygonID( OldPolygonID );
 			PolygonIDToOctreeElementIDMap.Add( NewPolygonID, OctreeElementID );
@@ -7989,7 +7989,7 @@ void UEditableMesh::UpdateOrRebuildOctree()
 			{
 				for( const FPolygonID PolygonID : DeletedOctreePolygonIDs )
 				{
-					const FOctreeElementId OctreeElementID = PolygonIDToOctreeElementIDMap.FindAndRemoveChecked( PolygonID );
+					const FOctreeElementId2 OctreeElementID = PolygonIDToOctreeElementIDMap.FindAndRemoveChecked( PolygonID );
 					Octree->RemoveElement( OctreeElementID );
 				}
 				DeletedOctreePolygonIDs.Reset();
diff --git a/Engine/Plugins/Runtime/EditableMesh/Source/EditableMesh/EditableMeshOctree.h b/Engine/Plugins/Runtime/EditableMesh/Source/EditableMesh/EditableMeshOctree.h
--- a/Engine/Plugins/Runtime/EditableMesh/Source/EditableMesh/EditableMeshOctree.h
+++ b/Engine/Plugins/Runtime/EditableMesh/Source/EditableMesh/EditableMeshOctree.h
@@ -48,7 +48,7 @@ struct FEditableMeshOctreeSemantics
 		return ( A.PolygonID == B.PolygonID );
 	}
 
-	FORCEINLINE static void SetElementId( const FEditableMeshOctreePolygon& Element, FOctreeElementId OctreeElementID )
+	FORCEINLINE static void SetElementId( const FEditableMeshOctreePolygon& Element, FOctreeElementId2 OctreeElementID )
 	{
 		Element.EditableMesh.PolygonIDToOctreeElementIDMap.Add( Element.PolygonID, OctreeElementID );
 	}
diff --git a/Engine/Plugins/Runtime/EditableMesh/Source/EditableMesh/Public/EditableMesh.h b/Engine/Plugins/Runtime/EditableMesh/Source/EditableMesh/Public/EditableMesh.h
--- a/Engine/Plugins/Runtime/EditableMesh/Source/EditableMesh/Public/EditableMesh.h
+++ b/Engine/Plugins/Runtime/EditableMesh/Source/EditableMesh/Public/EditableMesh.h
@@ -10,7 +10,7 @@
 #include "Misc/Change.h"		// For TUniquePtr<FChange>
 #include "Logging/LogMacros.h"
 #include "Materials/MaterialInterface.h"
-#include "Math/GenericOctreePublic.h"	// For FOctreeElementId
+#include "Math/GenericOctreePublic.h"	// For FOctreeElementId2
 #include "GeometryHitTest.h"
 #include "EditableMesh.generated.h"
 
@@ -625,7 +625,7 @@ public:
 	TSharedPtr<class FEditableMeshOctree> Octree;
 
 	/** Maps our polygon IDs to octree element IDs */
-	TMap<FPolygonID, FOctreeElementId> PolygonIDToOctreeElementIDMap;
+	TMap<FPolygonID, FOctreeElementId2> PolygonIDToOctreeElementIDMap;
 
 	/** Polygons that were deleted since the last time our octree was refreshed */
 	TSet<FPolygonID> DeletedOctreePolygonIDs;
diff --git a/Engine/Source/Developer/SkeletalMeshUtilitiesCommon/Private/LODUtilities.cpp b/Engine/Source/Developer/SkeletalMeshUtilitiesCommon/Private/LODUtilities.cpp
--- a/Engine/Source/Developer/SkeletalMeshUtilitiesCommon/Private/LODUtilities.cpp
+++ b/Engine/Source/Developer/SkeletalMeshUtilitiesCommon/Private/LODUtilities.cpp
@@ -1341,7 +1341,7 @@ struct FTriangleOctreeSemantics
 		return (A.TriangleIndex == B.TriangleIndex);
 	}
 
-	FORCEINLINE static void SetElementId(const FTriangleElement& Element, FOctreeElementId OctreeElementID)
+	FORCEINLINE static void SetElementId(const FTriangleElement& Element, FOctreeElementId2 OctreeElementID)
 	{
 	}
 };
diff --git a/Engine/Source/Editor/MeshPaint/Private/MeshPaintHelpers.cpp b/Engine/Source/Editor/MeshPaint/Private/MeshPaintHelpers.cpp
--- a/Engine/Source/Editor/MeshPaint/Private/MeshPaintHelpers.cpp
+++ b/Engine/Source/Editor/MeshPaint/Private/MeshPaintHelpers.cpp
@@ -1755,7 +1755,7 @@ struct FVertexColorPropogationOctreeSemantics
 	}
 
 	/** Ignored for this implementation */
-	FORCEINLINE static void SetElementId( const FPaintedMeshVertex& Element, FOctreeElementId Id )
+	FORCEINLINE static void SetElementId( const FPaintedMeshVertex& Element, FOctreeElementId2 Id )
 	{
 	}
 };
diff --git a/Engine/Source/Editor/MeshPaint/Public/TMeshPaintOctree.h b/Engine/Source/Editor/MeshPaint/Public/TMeshPaintOctree.h
--- a/Engine/Source/Editor/MeshPaint/Public/TMeshPaintOctree.h
+++ b/Engine/Source/Editor/MeshPaint/Public/TMeshPaintOctree.h
@@ -51,7 +51,7 @@ struct FMeshPaintTriangleOctreeSemantics
 	}
 
 	/** Ignored for this implementation */
-	FORCEINLINE static void SetElementId(const FMeshPaintTriangle& Element, FOctreeElementId Id)
+	FORCEINLINE static void SetElementId(const FMeshPaintTriangle& Element, FOctreeElementId2 Id)
 	{
 
 	}
diff --git a/Engine/Source/Editor/UnrealEd/Private/Fbx/FbxSkeletalMeshImport.cpp b/Engine/Source/Editor/UnrealEd/Private/Fbx/FbxSkeletalMeshImport.cpp
--- a/Engine/Source/Editor/UnrealEd/Private/Fbx/FbxSkeletalMeshImport.cpp
+++ b/Engine/Source/Editor/UnrealEd/Private/Fbx/FbxSkeletalMeshImport.cpp
@@ -156,7 +156,7 @@ struct FSkeletalMeshVertPosOctreeSemantics
 	}
 
 	/** Ignored for this implementation */
-	FORCEINLINE static void SetElementId(const FSoftSkinVertex& Element, FOctreeElementId Id)
+	FORCEINLINE static void SetElementId(const FSoftSkinVertex& Element, FOctreeElementId2 Id)
 	{
 	}
 };
diff --git a/Engine/Source/Programs/UnrealLightmass/Private/Launch/UnitTest.cpp b/Engine/Source/Programs/UnrealLightmass/Private/Launch/UnitTest.cpp
--- a/Engine/Source/Programs/UnrealLightmass/Private/Launch/UnitTest.cpp
+++ b/Engine/Source/Programs/UnrealLightmass/Private/Launch/UnitTest.cpp
@@ -217,7 +217,7 @@ void TestLightmass()
 
 	UE_LOG(LogLightmass, Display, TEXT("sizeof FDirectionalLight = %d, FLight = %d, FDirectionalLightData = %d"), sizeof(FDirectionalLight), sizeof(FLight), sizeof(FDirectionalLightData));
 
-	TOctree2<float, FTestOctreeSemantics> TestOctree(FVector4(0), 10.0f);
+	TOctree<float, FTestOctreeSemantics> TestOctree(FVector4(0), 10.0f);
 	TestOctree.AddElement(5);
 
 
diff --git a/Engine/Source/Programs/UnrealLightmass/Private/Lighting/LightingCache.h b/Engine/Source/Programs/UnrealLightmass/Private/Lighting/LightingCache.h
--- a/Engine/Source/Programs/UnrealLightmass/Private/Lighting/LightingCache.h
+++ b/Engine/Source/Programs/UnrealLightmass/Private/Lighting/LightingCache.h
@@ -205,7 +205,7 @@ public:
 	struct FRecordOctreeSemantics;
 
 	/** The type of lighting cache octree nodes. */
-	typedef TOctree2<FRecord<SampleType>,FRecordOctreeSemantics> LightingOctreeType;
+	typedef TOctree<FRecord<SampleType>,FRecordOctreeSemantics> LightingOctreeType;
 	
 	/** The octree semantics for irradiance records. */
 	struct FRecordOctreeSemantics
diff --git a/Engine/Source/Programs/UnrealLightmass/Private/Lighting/LightingSystem.h b/Engine/Source/Programs/UnrealLightmass/Private/Lighting/LightingSystem.h
--- a/Engine/Source/Programs/UnrealLightmass/Private/Lighting/LightingSystem.h
+++ b/Engine/Source/Programs/UnrealLightmass/Private/Lighting/LightingSystem.h
@@ -231,7 +231,7 @@ struct FPhotonElement
 	{}
 };
 
-typedef TOctree2<FPhotonElement,struct FPhotonMapOctreeSemantics> FPhotonOctree;
+typedef TOctree<FPhotonElement,struct FPhotonMapOctreeSemantics> FPhotonOctree;
 
 /** Octree semantic definitions. */
 struct FPhotonMapOctreeSemantics
@@ -299,7 +299,7 @@ struct FPhotonSegmentMapOctreeSemantics
 	}
 };
 
-typedef TOctree2<FPhotonSegmentElement,struct FPhotonSegmentMapOctreeSemantics> FPhotonSegmentOctree;
+typedef TOctree<FPhotonSegmentElement,struct FPhotonSegmentMapOctreeSemantics> FPhotonSegmentOctree;
 
 /** A photon which stores a precalculated irradiance estimate. */
 class FIrradiancePhoton : public FIrradiancePhotonData
@@ -376,7 +376,7 @@ private:
 	TArray<FIrradiancePhoton>& PhotonArray;
 };
 
-typedef TOctree2<FIrradiancePhotonElement,struct FIrradiancePhotonMapOctreeSemantics> FIrradiancePhotonOctree;
+typedef TOctree<FIrradiancePhotonElement,struct FIrradiancePhotonMapOctreeSemantics> FIrradiancePhotonOctree;
 
 /** Octree semantic definitions. */
 struct FIrradiancePhotonMapOctreeSemantics
@@ -438,7 +438,7 @@ struct FVolumeSampleInterpolationElement
 	{}
 };
 
-typedef TOctree2<FVolumeSampleInterpolationElement,struct FVolumeLightingInterpolationOctreeSemantics> FVolumeLightingInterpolationOctree;
+typedef TOctree<FVolumeSampleInterpolationElement,struct FVolumeLightingInterpolationOctreeSemantics> FVolumeLightingInterpolationOctree;
 
 /** Octree semantic definitions. */
 struct FVolumeLightingInterpolationOctreeSemantics
diff --git a/Engine/Source/Programs/UnrealLightmass/Private/Lighting/SampleVolume.cpp b/Engine/Source/Programs/UnrealLightmass/Private/Lighting/SampleVolume.cpp
--- a/Engine/Source/Programs/UnrealLightmass/Private/Lighting/SampleVolume.cpp
+++ b/Engine/Source/Programs/UnrealLightmass/Private/Lighting/SampleVolume.cpp
@@ -12,7 +12,7 @@ namespace Lightmass
 
 typedef FVolumeSampleInterpolationElement FVolumeSampleProximityElement;
 
-typedef TOctree2<FVolumeSampleProximityElement,struct FVolumeLightingProximityOctreeSemantics> FVolumeLightingProximityOctree;
+typedef TOctree<FVolumeSampleProximityElement,struct FVolumeLightingProximityOctreeSemantics> FVolumeLightingProximityOctree;
 
 struct FVolumeLightingProximityOctreeSemantics
 {
diff --git a/Engine/Source/Programs/UnrealLightmass/Private/LightmassCore/Math/LMOctree.h b/Engine/Source/Programs/UnrealLightmass/Private/LightmassCore/Math/LMOctree.h
--- a/Engine/Source/Programs/UnrealLightmass/Private/LightmassCore/Math/LMOctree.h
+++ b/Engine/Source/Programs/UnrealLightmass/Private/LightmassCore/Math/LMOctree.h
@@ -702,13 +702,13 @@ FORCEINLINE FOctreeChildNodeRef FOctreeNodeContext::GetContainingChild(const FBo
 }
 
 template<typename ElementType,typename OctreeSemantics>
-void TOctree2<ElementType,OctreeSemantics>::AddElement(typename TTypeTraits<ElementType>::ConstInitType Element)
+void TOctree<ElementType,OctreeSemantics>::AddElement(typename TTypeTraits<ElementType>::ConstInitType Element)
 {
 	AddElementToNode(Element,RootNode,RootNodeContext);
 }
 															
 template<typename ElementType,typename OctreeSemantics>
-void TOctree2<ElementType,OctreeSemantics>::AddElementToNode(
+void TOctree<ElementType,OctreeSemantics>::AddElementToNode(
 	typename TTypeTraits<ElementType>::ConstInitType Element,
 	const FNode& InNode,
 	const FOctreeNodeContext& InContext
@@ -765,7 +765,7 @@ void TOctree2<ElementType,OctreeSemantics>::AddElementToNode(
 				// Create the child node if it hasn't been created yet.
 				if(!Node.Children[ChildRef.Index])
 				{
-					Node.Children[ChildRef.Index] = new typename TOctree2<ElementType,OctreeSemantics>::FNode(&Node);
+					Node.Children[ChildRef.Index] = new typename TOctree<ElementType,OctreeSemantics>::FNode(&Node);
 				}
 
 				// Push the node onto the stack to visit.
@@ -787,7 +787,7 @@ void TOctree2<ElementType,OctreeSemantics>::AddElementToNode(
 }
 
 template<typename ElementType,typename OctreeSemantics>
-void TOctree2<ElementType,OctreeSemantics>::DumpStats(bool bDetailed) const
+void TOctree<ElementType,OctreeSemantics>::DumpStats(bool bDetailed) const
 {
 	int32 NumNodes = 0;
 	int32 NumLeaves = 0;
@@ -857,7 +857,7 @@ void TOctree2<ElementType,OctreeSemantics>::DumpStats(bool bDetailed) const
 }
 
 template<typename ElementType,typename OctreeSemantics>
-void TOctree2<ElementType,OctreeSemantics>::GetMemoryUsage(size_t& OutSizeBytes) const
+void TOctree<ElementType,OctreeSemantics>::GetMemoryUsage(size_t& OutSizeBytes) const
 {
 	size_t SizeBytes = 0;
 
@@ -880,7 +880,7 @@ void TOctree2<ElementType,OctreeSemantics>::GetMemoryUsage(size_t& OutSizeBytes)
 }
 
 template<typename ElementType,typename OctreeSemantics>
-TOctree2<ElementType,OctreeSemantics>::TOctree(const FVector4& InOrigin,float InExtent)
+TOctree<ElementType,OctreeSemantics>::TOctree(const FVector4& InOrigin,float InExtent)
 :	RootNode(NULL)
 ,	RootNodeContext(FBoxCenterAndExtent(InOrigin,FVector4(InExtent,InExtent,InExtent,0)), OctreeSemantics::LoosenessDenominator)
 ,	MinLeafExtent(InExtent * FMath::Pow((1.0f + 1.0f / (float)OctreeSemantics::LoosenessDenominator) / 2.0f,OctreeSemantics::MaxNodeDepth))
diff --git a/Engine/Source/Runtime/Core/Public/Math/GenericOctree.h b/Engine/Source/Runtime/Core/Public/Math/GenericOctree.h
--- a/Engine/Source/Runtime/Core/Public/Math/GenericOctree.h
+++ b/Engine/Source/Runtime/Core/Public/Math/GenericOctree.h
@@ -537,7 +537,7 @@ private:
 			{
 				int ElementIndex = TreeElements[CurrentNodeIndex].Emplace(Element);
 
-				SetElementId(Element, FOctreeElementId(CurrentNodeIndex, ElementIndex));	
+				SetElementId(Element, FOctreeElementId2(CurrentNodeIndex, ElementIndex));	
 				return;
 			}
 		}
@@ -547,7 +547,7 @@ private:
 			if (ChildRef.IsNULL())
 			{
 				int ElementIndex = TreeElements[CurrentNodeIndex].Emplace(Element);
-				SetElementId(Element, FOctreeElementId(CurrentNodeIndex, ElementIndex));
+				SetElementId(Element, FOctreeElementId2(CurrentNodeIndex, ElementIndex));
 				return;
 			}
 			else
@@ -785,7 +785,7 @@ public:
 	 * Removes an element from the octree.
 	 * @param ElementId - The element to remove from the octree.
 	 */
-	inline void RemoveElement(FOctreeElementId ElementId)
+	inline void RemoveElement(FOctreeElementId2 ElementId)
 	{
 		checkSlow(ElementId.IsValidId());
 
@@ -834,7 +834,7 @@ public:
 				for (int ElementIndex = 0; ElementIndex < TreeElements[CollapseNodeIndex].Num(); ElementIndex++)
 				{
 					// Update the external element id for the element that's being collapsed.
-					SetElementId(TreeElements[CollapseNodeIndex][ElementIndex], FOctreeElementId(CollapseNodeIndex, ElementIndex));
+					SetElementId(TreeElements[CollapseNodeIndex][ElementIndex], FOctreeElementId2(CollapseNodeIndex, ElementIndex));
 				}
 
 				// Mark the node as a leaf.
@@ -857,22 +857,22 @@ public:
 	}
 
 	/** Accesses an octree element by ID. */
-	ElementType& GetElementById(FOctreeElementId ElementId)
+	ElementType& GetElementById(FOctreeElementId2 ElementId)
 	{
 		return TreeElements[ElementId.NodeIndex][ElementId.ElementIndex];
 	}
 
 	/** Accesses an octree element by ID. */
-	const ElementType& GetElementById(FOctreeElementId ElementId) const
+	const ElementType& GetElementById(FOctreeElementId2 ElementId) const
 	{
 		return TreeElements[ElementId.NodeIndex][ElementId.ElementIndex];
 	}
 
 	/**
-	 * check if a FOctreeElementId is valid.
+	 * check if a FOctreeElementId2 is valid.
 	 * @param ElementId - The ElementId to check.
 	 */
-	bool IsValidElementId(FOctreeElementId ElementId) const
+	bool IsValidElementId(FOctreeElementId2 ElementId) const
 	{
 		return ElementId.IsValidId() && ElementId.ElementIndex < TreeElements[ElementId.NodeIndex].Num();
 	};
@@ -1008,25 +1008,25 @@ private:
 	struct COctreeSemanticsV2
 	{
 		template<typename Semantics>
-		auto Requires(typename Semantics::FOctree& OctreeInstance, const ElementType& Element, FOctreeElementId Id)
+		auto Requires(typename Semantics::FOctree& OctreeInstance, const ElementType& Element, FOctreeElementId2 Id)
 			-> decltype(Semantics::SetElementId(OctreeInstance, Element, Id));
 	};
 
 	// Function overload set which calls the V2 version if it's supported or the old version if it's not
 	template <typename Semantics>
-	typename TEnableIf<!TModels<COctreeSemanticsV2, Semantics>::Value>::Type SetOctreeSemanticsElementId(const ElementType& Element, FOctreeElementId Id)
+	typename TEnableIf<!TModels<COctreeSemanticsV2, Semantics>::Value>::Type SetOctreeSemanticsElementId(const ElementType& Element, FOctreeElementId2 Id)
 	{
 		Semantics::SetElementId(Element, Id);
 	}
 	template <typename Semantics>
-	typename TEnableIf<TModels<COctreeSemanticsV2, Semantics>::Value>::Type SetOctreeSemanticsElementId(const ElementType& Element, FOctreeElementId Id)
+	typename TEnableIf<TModels<COctreeSemanticsV2, Semantics>::Value>::Type SetOctreeSemanticsElementId(const ElementType& Element, FOctreeElementId2 Id)
 	{
 		Semantics::SetElementId(*this, Element, Id);
 	}
 
 protected:
 	// redirects SetElementId call to the proper implementation
-	void SetElementId(const ElementType& Element, FOctreeElementId Id)
+	void SetElementId(const ElementType& Element, FOctreeElementId2 Id)
 	{
 		SetOctreeSemanticsElementId<OctreeSemantics>(Element, Id);
 	}
diff --git a/Engine/Source/Runtime/Core/Public/Math/GenericOctree.inl b/Engine/Source/Runtime/Core/Public/Math/GenericOctree.inl
--- a/Engine/Source/Runtime/Core/Public/Math/GenericOctree.inl
+++ b/Engine/Source/Runtime/Core/Public/Math/GenericOctree.inl
@@ -13,7 +13,7 @@
 class FBoxCenterAndExtent;
 class FOctreeChildNodeRef;
 class FOctreeChildNodeSubset;
-class FOctreeElementId;
+class FOctreeElementId2;
 class FOctreeNodeContext;
 struct FMath;
 template<typename ElementType,typename OctreeSemantics> class TOctree;
diff --git a/Engine/Source/Runtime/Core/Public/Math/GenericOctreePublic.h b/Engine/Source/Runtime/Core/Public/Math/GenericOctreePublic.h
--- a/Engine/Source/Runtime/Core/Public/Math/GenericOctreePublic.h
+++ b/Engine/Source/Runtime/Core/Public/Math/GenericOctreePublic.h
@@ -11,7 +11,7 @@
 /** 
  *	An identifier for an element in the octree. 
  */
-class FOctreeElementId
+class FOctreeElementId2
 {
 public:
 
@@ -22,7 +22,7 @@ public:
 	friend class TOctree2;
 
 	/** Default constructor. */
-	FOctreeElementId()
+	FOctreeElementId2()
 		:	NodeIndex(INDEX_NONE)
 		,	ElementIndex(INDEX_NONE)
 	{}
@@ -42,7 +42,7 @@ private:
 	int32 ElementIndex;
 
 	/** Initialization constructor. */
-	FOctreeElementId(uint32 InNodeIndex, int32 InElementIndex)
+	FOctreeElementId2(uint32 InNodeIndex, int32 InElementIndex)
 		:	NodeIndex(InNodeIndex)
 		,	ElementIndex(InElementIndex)
 	{}
diff --git a/Engine/Source/Runtime/Engine/Private/StaticMeshBuild.cpp b/Engine/Source/Runtime/Engine/Private/StaticMeshBuild.cpp
--- a/Engine/Source/Runtime/Engine/Private/StaticMeshBuild.cpp
+++ b/Engine/Source/Runtime/Engine/Private/StaticMeshBuild.cpp
@@ -488,7 +488,7 @@ struct FStaticMeshComponentVertPosOctreeSemantics
 	}
 
 	/** Ignored for this implementation */
-	FORCEINLINE static void SetElementId( const FPaintedVertex& Element, FOctreeElementId Id )
+	FORCEINLINE static void SetElementId( const FPaintedVertex& Element, FOctreeElementId2 Id )
 	{
 	}
 };
diff --git a/Engine/Source/Runtime/Engine/Public/PrecomputedLightVolume.h b/Engine/Source/Runtime/Engine/Public/PrecomputedLightVolume.h
--- a/Engine/Source/Runtime/Engine/Public/PrecomputedLightVolume.h
+++ b/Engine/Source/Runtime/Engine/Public/PrecomputedLightVolume.h
@@ -80,7 +80,7 @@ struct FLightVolumeOctreeSemantics
 		return &Sample.Position.X;
 	}
 
-	static void SetElementId(const FVolumeLightingSample& Element, FOctreeElementId Id)
+	static void SetElementId(const FVolumeLightingSample& Element, FOctreeElementId2 Id)
 	{
 	}
 
diff --git a/Engine/Source/Runtime/Engine/Public/Rendering/SkeletalMeshLODImporterData.h b/Engine/Source/Runtime/Engine/Public/Rendering/SkeletalMeshLODImporterData.h
--- a/Engine/Source/Runtime/Engine/Public/Rendering/SkeletalMeshLODImporterData.h
+++ b/Engine/Source/Runtime/Engine/Public/Rendering/SkeletalMeshLODImporterData.h
@@ -605,7 +605,7 @@ struct FWedgeInfoOctreeSemantics
 	}
 
 	/** Ignored for this implementation */
-	FORCEINLINE static void SetElementId(const FWedgeInfo& Element, FOctreeElementId Id)
+	FORCEINLINE static void SetElementId(const FWedgeInfo& Element, FOctreeElementId2 Id)
 	{
 	}
 };
diff --git a/Engine/Source/Runtime/Experimental/GeometryCollectionCore/Private/GeometryCollection/GeometryCollectionProximityUtility.cpp b/Engine/Source/Runtime/Experimental/GeometryCollectionCore/Private/GeometryCollection/GeometryCollectionProximityUtility.cpp
--- a/Engine/Source/Runtime/Experimental/GeometryCollectionCore/Private/GeometryCollection/GeometryCollectionProximityUtility.cpp
+++ b/Engine/Source/Runtime/Experimental/GeometryCollectionCore/Private/GeometryCollection/GeometryCollectionProximityUtility.cpp
@@ -64,7 +64,7 @@ struct FMeshProximityTriangleOctreeSemantics
 	}
 
 	/** Ignored for this implementation */
-	FORCEINLINE static void SetElementId(const FProximityTriangle& Element, FOctreeElementId Id)
+	FORCEINLINE static void SetElementId(const FProximityTriangle& Element, FOctreeElementId2 Id)
 	{
 
 	}
diff --git a/Engine/Source/Runtime/NavigationSystem/Private/NavigationDataHandler.cpp b/Engine/Source/Runtime/NavigationSystem/Private/NavigationDataHandler.cpp
--- a/Engine/Source/Runtime/NavigationSystem/Private/NavigationDataHandler.cpp
+++ b/Engine/Source/Runtime/NavigationSystem/Private/NavigationDataHandler.cpp
@@ -20,7 +20,7 @@ FNavigationDataHandler::FNavigationDataHandler(FNavigationOctreeController& InOc
 		: OctreeController(InOctreeController), DirtyAreasController(InDirtyAreasController)
 {}
 
-void FNavigationDataHandler::RemoveNavOctreeElementId(const FOctreeElementId& ElementId, int32 UpdateFlags)
+void FNavigationDataHandler::RemoveNavOctreeElementId(const FOctreeElementId2& ElementId, int32 UpdateFlags)
 {
 	if (ensure(OctreeController.IsValidElement(ElementId)))
 	{
@@ -116,7 +116,7 @@ void FNavigationDataHandler::AddElementToNavOctree(const FNavigationDirtyElement
 	{
 		// check if parent node is waiting in queue
 		const FSetElementId ParentRequestId = OctreeController.PendingOctreeUpdates.FindId(FNavigationDirtyElement(NavigationParent));
-		const FOctreeElementId* ParentId = OctreeController.GetObjectsNavOctreeId(*NavigationParent);
+		const FOctreeElementId2* ParentId = OctreeController.GetObjectsNavOctreeId(*NavigationParent);
 		if (ParentRequestId.IsValidId() && ParentId == nullptr)
 		{
 			FNavigationDirtyElement& ParentNode = OctreeController.PendingOctreeUpdates[ParentRequestId];
@@ -126,7 +126,7 @@ void FNavigationDataHandler::AddElementToNavOctree(const FNavigationDirtyElement
 			ParentNode.bInvalidRequest = true;
 		}
 
-		const FOctreeElementId* ElementId = ParentId ? ParentId : OctreeController.GetObjectsNavOctreeId(*NavigationParent);
+		const FOctreeElementId2* ElementId = ParentId ? ParentId : OctreeController.GetObjectsNavOctreeId(*NavigationParent);
 		if (ElementId && ensure(OctreeController.IsValidElement(*ElementId)))
 		{
 			UE_LOG(LogNavOctree, Log, TEXT("ADD %s to %s"), *GetNameSafe(ElementOwner), *GetNameSafe(NavigationParent));
@@ -168,7 +168,7 @@ bool FNavigationDataHandler::UnregisterNavOctreeElement(UObject& ElementOwner, I
 	}
 
 	bool bUnregistered = false;
-	const FOctreeElementId* ElementId = OctreeController.GetObjectsNavOctreeId(ElementOwner);
+	const FOctreeElementId2* ElementId = OctreeController.GetObjectsNavOctreeId(ElementOwner);
 	UE_LOG(LogNavOctree, Log, TEXT("UNREG %s %s"), *ElementOwner.GetName(), ElementId ? TEXT("[exists]") : TEXT("[does\'t exist]"));
 
 	if (ElementId != nullptr)
@@ -329,7 +329,7 @@ void FNavigationDataHandler::UpdateNavOctreeParentChain(UObject& ElementOwner, b
 
 bool FNavigationDataHandler::UpdateNavOctreeElementBounds(UActorComponent& Comp, const FBox& NewBounds, const FBox& DirtyArea)
 {
-	const FOctreeElementId* ElementId = OctreeController.GetObjectsNavOctreeId(Comp);
+	const FOctreeElementId2* ElementId = OctreeController.GetObjectsNavOctreeId(Comp);
 	if (ElementId != nullptr && ensure(OctreeController.IsValidElement(*ElementId)))
 	{
 		OctreeController.NavOctree->UpdateNode(*ElementId, NewBounds);
@@ -420,7 +420,7 @@ void FNavigationDataHandler::AddLevelCollisionToOctree(ULevel& Level)
 		OctreeController.NavOctree->GetNavGeometryStoringMode() == FNavigationOctree::StoreNavGeometry)
 	{
 		const TArray<FVector>* LevelGeom = Level.GetStaticNavigableGeometry();
-		const FOctreeElementId* ElementId = OctreeController.GetObjectsNavOctreeId(Level);
+		const FOctreeElementId2* ElementId = OctreeController.GetObjectsNavOctreeId(Level);
 
 		if (!ElementId && LevelGeom && LevelGeom->Num() > 0)
 		{
@@ -444,7 +444,7 @@ void FNavigationDataHandler::RemoveLevelCollisionFromOctree(ULevel& Level)
 {
 	if (OctreeController.NavOctree.IsValid())
 	{
-		const FOctreeElementId* ElementId = OctreeController.GetObjectsNavOctreeId(Level);
+		const FOctreeElementId2* ElementId = OctreeController.GetObjectsNavOctreeId(Level);
 		UE_LOG(LogNavOctree, Log, TEXT("UNREG %s %s"), *Level.GetName(), ElementId ? TEXT("[exists]") : TEXT(""));
 
 		if (ElementId != nullptr)
diff --git a/Engine/Source/Runtime/NavigationSystem/Private/NavigationOctree.cpp b/Engine/Source/Runtime/NavigationSystem/Private/NavigationOctree.cpp
--- a/Engine/Source/Runtime/NavigationSystem/Private/NavigationOctree.cpp
+++ b/Engine/Source/Runtime/NavigationSystem/Private/NavigationOctree.cpp
@@ -149,7 +149,7 @@ void FNavigationOctree::AddNode(UObject* ElementOb, INavRelevantInterface* NavEl
 	AddElement(Element);
 }
 
-void FNavigationOctree::AppendToNode(const FOctreeElementId& Id, INavRelevantInterface* NavElement, const FBox& Bounds, FNavigationOctreeElement& Element)
+void FNavigationOctree::AppendToNode(const FOctreeElementId2& Id, INavRelevantInterface* NavElement, const FBox& Bounds, FNavigationOctreeElement& Element)
 {
 	FNavigationOctreeElement OrgData = GetElementById(Id);
 
@@ -187,7 +187,7 @@ void FNavigationOctree::AppendToNode(const FOctreeElementId& Id, INavRelevantInt
 	AddElement(Element);
 }
 
-void FNavigationOctree::UpdateNode(const FOctreeElementId& Id, const FBox& NewBounds)
+void FNavigationOctree::UpdateNode(const FOctreeElementId2& Id, const FBox& NewBounds)
 {
 	FNavigationOctreeElement ElementCopy = GetElementById(Id);
 	RemoveElement(Id);
@@ -195,7 +195,7 @@ void FNavigationOctree::UpdateNode(const FOctreeElementId& Id, const FBox& NewBo
 	AddElement(ElementCopy);
 }
 
-void FNavigationOctree::RemoveNode(const FOctreeElementId& Id)
+void FNavigationOctree::RemoveNode(const FOctreeElementId2& Id)
 {
 	const FNavigationOctreeElement& Element = GetElementById(Id);
 	const int32 ElementMemory = Element.GetAllocatedSize();
@@ -205,7 +205,7 @@ void FNavigationOctree::RemoveNode(const FOctreeElementId& Id)
 	RemoveElement(Id);
 }
 
-const FNavigationRelevantData* FNavigationOctree::GetDataForID(const FOctreeElementId& Id) const
+const FNavigationRelevantData* FNavigationOctree::GetDataForID(const FOctreeElementId2& Id) const
 {
 	if (Id.IsValidId() == false)
 	{
@@ -217,7 +217,7 @@ const FNavigationRelevantData* FNavigationOctree::GetDataForID(const FOctreeElem
 	return &*OctreeElement.Data;
 }
 
-void FNavigationOctree::SetElementIdImpl(const uint32 OwnerUniqueId, FOctreeElementId Id)
+void FNavigationOctree::SetElementIdImpl(const uint32 OwnerUniqueId, FOctreeElementId2 Id)
 {
 	ObjectToOctreeId.Add(OwnerUniqueId, Id);
 }
@@ -228,7 +228,7 @@ void FNavigationOctree::SetElementIdImpl(const uint32 OwnerUniqueId, FOctreeElem
 #if NAVSYS_DEBUG
 FORCENOINLINE
 #endif // NAVSYS_DEBUG
-void FNavigationOctreeSemantics::SetElementId(FNavigationOctreeSemantics::FOctree& OctreeOwner, const FNavigationOctreeElement& Element, FOctreeElementId Id)
+void FNavigationOctreeSemantics::SetElementId(FNavigationOctreeSemantics::FOctree& OctreeOwner, const FNavigationOctreeElement& Element, FOctreeElementId2 Id)
 {
 	((FNavigationOctree&)OctreeOwner).SetElementIdImpl(Element.OwnerUniqueId, Id);
 }
diff --git a/Engine/Source/Runtime/NavigationSystem/Private/NavigationOctreeController.cpp b/Engine/Source/Runtime/NavigationSystem/Private/NavigationOctreeController.cpp
--- a/Engine/Source/Runtime/NavigationSystem/Private/NavigationOctreeController.cpp
+++ b/Engine/Source/Runtime/NavigationSystem/Private/NavigationOctreeController.cpp
@@ -26,7 +26,7 @@ void FNavigationOctreeController::SetNavigableGeometryStoringMode(FNavigationOct
 
 bool FNavigationOctreeController::GetNavOctreeElementData(const UObject& NodeOwner, int32& DirtyFlags, FBox& DirtyBounds)
 {
-	const FOctreeElementId* ElementId = GetObjectsNavOctreeId(NodeOwner);
+	const FOctreeElementId2* ElementId = GetObjectsNavOctreeId(NodeOwner);
 	if (ElementId != nullptr && IsValidElement(*ElementId))
 	{
 		// mark area occupied by given actor as dirty
@@ -41,7 +41,7 @@ bool FNavigationOctreeController::GetNavOctreeElementData(const UObject& NodeOwn
 
 const FNavigationRelevantData* FNavigationOctreeController::GetDataForObject(const UObject& Object) const
 {
-	const FOctreeElementId* ElementId = GetObjectsNavOctreeId(Object);
+	const FOctreeElementId2* ElementId = GetObjectsNavOctreeId(Object);
 
 	if (ElementId != nullptr && IsValidElement(*ElementId))
 	{
diff --git a/Engine/Source/Runtime/NavigationSystem/Private/NavigationSystem.cpp b/Engine/Source/Runtime/NavigationSystem/Private/NavigationSystem.cpp
--- a/Engine/Source/Runtime/NavigationSystem/Private/NavigationSystem.cpp
+++ b/Engine/Source/Runtime/NavigationSystem/Private/NavigationSystem.cpp
@@ -2698,7 +2698,7 @@ void UNavigationSystemV1::UnregisterNavOctreeElement(UObject* ElementOwner, INav
 	}
 }
 
-void UNavigationSystemV1::RemoveNavOctreeElementId(const FOctreeElementId& ElementId, int32 UpdateFlags)
+void UNavigationSystemV1::RemoveNavOctreeElementId(const FOctreeElementId2& ElementId, int32 UpdateFlags)
 {
 	FNavigationDataHandler(DefaultOctreeController, DefaultDirtyAreasController).RemoveNavOctreeElementId(ElementId, UpdateFlags);
 }
diff --git a/Engine/Source/Runtime/NavigationSystem/Public/NavigationDataHandler.h b/Engine/Source/Runtime/NavigationSystem/Public/NavigationDataHandler.h
--- a/Engine/Source/Runtime/NavigationSystem/Public/NavigationDataHandler.h
+++ b/Engine/Source/Runtime/NavigationSystem/Public/NavigationDataHandler.h
@@ -12,7 +12,7 @@ struct NAVIGATIONSYSTEM_API FNavigationDataHandler
 
 	FNavigationDataHandler(FNavigationOctreeController& InOctreeController, FNavigationDirtyAreasController& InDirtyAreasController);
 
-	void RemoveNavOctreeElementId(const FOctreeElementId& ElementId, int32 UpdateFlags);
+	void RemoveNavOctreeElementId(const FOctreeElementId2& ElementId, int32 UpdateFlags);
 	FSetElementId RegisterNavOctreeElement(UObject& ElementOwner, INavRelevantInterface& ElementInterface, int32 UpdateFlags);
 	void AddElementToNavOctree(const FNavigationDirtyElement& DirtyElement);
 
diff --git a/Engine/Source/Runtime/NavigationSystem/Public/NavigationOctree.h b/Engine/Source/Runtime/NavigationSystem/Public/NavigationOctree.h
--- a/Engine/Source/Runtime/NavigationSystem/Public/NavigationOctree.h
+++ b/Engine/Source/Runtime/NavigationSystem/Public/NavigationOctree.h
@@ -112,7 +112,7 @@ struct FNavigationOctreeSemantics
 #if NAVSYS_DEBUG
 	FORCENOINLINE 
 #endif // NAVSYS_DEBUG
-	static void SetElementId(FOctree& OctreeOwner, const FNavigationOctreeElement& Element, FOctreeElementId Id);
+	static void SetElementId(FOctree& OctreeOwner, const FNavigationOctreeElement& Element, FOctreeElementId2 Id);
 };
 
 class NAVIGATIONSYSTEM_API FNavigationOctree : public TOctree2<FNavigationOctreeElement, FNavigationOctreeSemantics>, public TSharedFromThis<FNavigationOctree, ESPMode::ThreadSafe>
@@ -144,7 +144,7 @@ public:
 	 * Removes an element from the octree.
 	 * @param ElementId - The element to remove from the octree.
 	 */
-	inline void RemoveElement(FOctreeElementId ElementId)
+	inline void RemoveElement(FOctreeElementId2 ElementId)
 	{
 		DEC_MEMORY_STAT_BY(STAT_NavigationMemory, OctreeSizeBytes);
 		DEC_MEMORY_STAT_BY(STAT_Navigation_CollisionTreeMemory, OctreeSizeBytes);
@@ -161,17 +161,17 @@ public:
 	void AddNode(UObject* ElementOb, INavRelevantInterface* NavElement, const FBox& Bounds, FNavigationOctreeElement& Data);
 
 	/** Append new data to existing node */
-	void AppendToNode(const FOctreeElementId& Id, INavRelevantInterface* NavElement, const FBox& Bounds, FNavigationOctreeElement& Data);
+	void AppendToNode(const FOctreeElementId2& Id, INavRelevantInterface* NavElement, const FBox& Bounds, FNavigationOctreeElement& Data);
 
 	/** Updates element bounds remove/add operation */
-	void UpdateNode(const FOctreeElementId& Id, const FBox& NewBounds);
+	void UpdateNode(const FOctreeElementId2& Id, const FBox& NewBounds);
 
 	/** Remove node */
-	void RemoveNode(const FOctreeElementId& Id);
+	void RemoveNode(const FOctreeElementId2& Id);
 
 	void SetNavigableGeometryStoringMode(ENavGeometryStoringMode NavGeometryMode);
 
-	const FNavigationRelevantData* GetDataForID(const FOctreeElementId& Id) const;
+	const FNavigationRelevantData* GetDataForID(const FOctreeElementId2& Id) const;
 
 	ENavGeometryStoringMode GetNavGeometryStoringMode() const
 	{
@@ -194,9 +194,9 @@ protected:
 	friend struct FNavigationOctreeController;
 	friend struct FNavigationOctreeSemantics;
 
-	void SetElementIdImpl(const uint32 OwnerUniqueId, FOctreeElementId Id);
+	void SetElementIdImpl(const uint32 OwnerUniqueId, FOctreeElementId2 Id);
 
-	TMap<uint32, FOctreeElementId> ObjectToOctreeId;
+	TMap<uint32, FOctreeElementId2> ObjectToOctreeId;
 	ENavDataGatheringMode DefaultGeometryGatheringMode;
 	uint32 bGatherGeometry : 1;
 	uint32 NodesMemory;
diff --git a/Engine/Source/Runtime/NavigationSystem/Public/NavigationOctreeController.h b/Engine/Source/Runtime/NavigationSystem/Public/NavigationOctreeController.h
--- a/Engine/Source/Runtime/NavigationSystem/Public/NavigationOctreeController.h
+++ b/Engine/Source/Runtime/NavigationSystem/Public/NavigationOctreeController.h
@@ -31,7 +31,7 @@ struct NAVIGATIONSYSTEM_API FNavigationOctreeController
 
 	const FNavigationOctree* GetOctree() const;
 	FNavigationOctree* GetMutableOctree();
-	const FOctreeElementId* GetObjectsNavOctreeId(const UObject& Object) const;
+	const FOctreeElementId2* GetObjectsNavOctreeId(const UObject& Object) const;
 	bool GetNavOctreeElementData(const UObject& NodeOwner, int32& DirtyFlags, FBox& DirtyBounds);
 	const FNavigationRelevantData* GetDataForObject(const UObject& Object) const;
 	FNavigationRelevantData* GetMutableDataForObject(const UObject& Object);
@@ -39,7 +39,7 @@ struct NAVIGATIONSYSTEM_API FNavigationOctreeController
 	bool IsNavigationOctreeLocked() const;
 	/** basically says if navoctree has been created already */
 	bool IsValid() const { return NavOctree.IsValid(); }
-	bool IsValidElement(const FOctreeElementId& ElementId) const;
+	bool IsValidElement(const FOctreeElementId2& ElementId) const;
 	bool IsEmpty() const { return (IsValid() == false) || NavOctree->GetSizeBytes() == 0; }
 
 private:
@@ -54,7 +54,7 @@ FORCEINLINE uint32 FNavigationOctreeController::HashObject(const UObject& Object
 	return FNavigationOctree::HashObject(Object);
 }
 
-FORCEINLINE_DEBUGGABLE const FOctreeElementId* FNavigationOctreeController::GetObjectsNavOctreeId(const UObject& Object) const
+FORCEINLINE_DEBUGGABLE const FOctreeElementId2* FNavigationOctreeController::GetObjectsNavOctreeId(const UObject& Object) const
 { 
 	return NavOctree.IsValid()
 		? NavOctree->ObjectToOctreeId.Find(HashObject(Object))
@@ -99,7 +99,7 @@ FORCEINLINE void FNavigationOctreeController::SetNavigationOctreeLock(bool bLock
 	bNavOctreeLock = bLock; 
 }
 
-FORCEINLINE bool FNavigationOctreeController::IsValidElement(const FOctreeElementId& ElementId) const
+FORCEINLINE bool FNavigationOctreeController::IsValidElement(const FOctreeElementId2& ElementId) const
 {
 	return IsValid() && NavOctree->IsValidElementId(ElementId);
 }
diff --git a/Engine/Source/Runtime/NavigationSystem/Public/NavigationSystem.h b/Engine/Source/Runtime/NavigationSystem/Public/NavigationSystem.h
--- a/Engine/Source/Runtime/NavigationSystem/Public/NavigationSystem.h
+++ b/Engine/Source/Runtime/NavigationSystem/Public/NavigationSystem.h
@@ -741,11 +741,11 @@ public:
 	{
 		return FNavigationOctree::HashObject(Object);
 	}
-	FORCEINLINE const FOctreeElementId* GetObjectsNavOctreeId(const UObject& Object) const { return DefaultOctreeController.GetObjectsNavOctreeId(Object); }
+	FORCEINLINE const FOctreeElementId2* GetObjectsNavOctreeId(const UObject& Object) const { return DefaultOctreeController.GetObjectsNavOctreeId(Object); }
 	FORCEINLINE bool HasPendingObjectNavOctreeId(UObject* Object) const { return Object && DefaultOctreeController.HasPendingObjectNavOctreeId(*Object); }
 	FORCEINLINE void RemoveObjectsNavOctreeId(const UObject& Object) { DefaultOctreeController.RemoveObjectsNavOctreeId(Object); }
 
-	void RemoveNavOctreeElementId(const FOctreeElementId& ElementId, int32 UpdateFlags);
+	void RemoveNavOctreeElementId(const FOctreeElementId2& ElementId, int32 UpdateFlags);
 
 	const FNavigationRelevantData* GetDataForObject(const UObject& Object) const;
 	FNavigationRelevantData* GetMutableDataForObject(const UObject& Object);
@@ -1204,10 +1204,10 @@ public:
 	virtual ANavigationData* CreateNavigationDataInstance(const FNavDataConfig& NavConfig);
 	
 	UE_DEPRECATED(4.24, "This function is deprecated and no longer used. NavigationSystem is no longer involved in storing navoctree element IDs. See FNavigationOctree for more details.")
-	void SetObjectsNavOctreeId(const UObject& Object, FOctreeElementId Id) {}
+	void SetObjectsNavOctreeId(const UObject& Object, FOctreeElementId2 Id) {}
 
 	UE_DEPRECATED(4.24, "This member is deprecated and no longer used.  Please access OctreeController instead")
-	TMap<uint32, FOctreeElementId> ObjectToOctreeId;
+	TMap<uint32, FOctreeElementId2> ObjectToOctreeId;
 	UE_DEPRECATED(4.24, "This member is deprecated and no longer used.  Please access OctreeController instead")
 	TSet<FNavigationDirtyElement> PendingOctreeUpdates;
 	UE_DEPRECATED(4.24, "This member is deprecated and no longer used.  Please access OctreeController instead")
diff --git a/Engine/Source/Runtime/Renderer/Private/LightSceneInfo.h b/Engine/Source/Runtime/Renderer/Private/LightSceneInfo.h
--- a/Engine/Source/Runtime/Renderer/Private/LightSceneInfo.h
+++ b/Engine/Source/Runtime/Renderer/Private/LightSceneInfo.h
@@ -169,7 +169,7 @@ public:
 	int32 Id;
 
 	/** The identifier for the primitive in Scene->PrimitiveOctree. */
-	FOctreeElementId OctreeId;
+	FOctreeElementId2 OctreeId;
 
 	/** Tile intersection buffer for distance field shadowing, stored on the light to avoid reallocating each frame. */
 	mutable TUniquePtr<class FLightTileIntersectionResources> TileIntersectionResources;
@@ -340,7 +340,7 @@ struct FLightOctreeSemantics
 		return A.LightSceneInfo == B.LightSceneInfo;
 	}
 	
-	FORCEINLINE static void SetElementId(const FLightSceneInfoCompact& Element,FOctreeElementId Id)
+	FORCEINLINE static void SetElementId(const FLightSceneInfoCompact& Element,FOctreeElementId2 Id)
 	{
 		Element.LightSceneInfo->OctreeId = Id;
 	}
diff --git a/Engine/Source/Runtime/Renderer/Private/PrimitiveSceneInfo.cpp b/Engine/Source/Runtime/Renderer/Private/PrimitiveSceneInfo.cpp
--- a/Engine/Source/Runtime/Renderer/Private/PrimitiveSceneInfo.cpp
+++ b/Engine/Source/Runtime/Renderer/Private/PrimitiveSceneInfo.cpp
@@ -855,7 +855,7 @@ void FPrimitiveSceneInfo::RemoveFromScene(bool bUpdateStaticDrawLists)
 	check(OctreeId.IsValidId());
 	check(Scene->PrimitiveOctree.GetElementById(OctreeId).PrimitiveSceneInfo == this);
 	Scene->PrimitiveOctree.RemoveElement(OctreeId);
-	OctreeId = FOctreeElementId();
+	OctreeId = FOctreeElementId2();
 
 	if (LightmapDataOffset != INDEX_NONE && UseGPUScene(GMaxRHIShaderPlatform, Scene->GetFeatureLevel()))
 	{
diff --git a/Engine/Source/Runtime/Renderer/Public/PrimitiveSceneInfo.h b/Engine/Source/Runtime/Renderer/Public/PrimitiveSceneInfo.h
--- a/Engine/Source/Runtime/Renderer/Public/PrimitiveSceneInfo.h
+++ b/Engine/Source/Runtime/Renderer/Public/PrimitiveSceneInfo.h
@@ -255,7 +255,7 @@ public:
 	TArray<class FStaticMeshBatch> StaticMeshes;
 
 	/** The identifier for the primitive in Scene->PrimitiveOctree. */
-	FOctreeElementId OctreeId;
+	FOctreeElementId2 OctreeId;
 
 	/** 
 	 * Caches the primitive's indirect lighting cache allocation.
@@ -565,7 +565,7 @@ struct FPrimitiveOctreeSemantics
 		return A.PrimitiveSceneInfo == B.PrimitiveSceneInfo;
 	}
 
-	FORCEINLINE static void SetElementId(const FPrimitiveSceneInfoCompact& Element,FOctreeElementId Id)
+	FORCEINLINE static void SetElementId(const FPrimitiveSceneInfoCompact& Element,FOctreeElementId2 Id)
 	{
 		Element.PrimitiveSceneInfo->OctreeId = Id;
 	}
