M - Deprecated version of the original Octree class

diff --git a/Engine/Source/Runtime/Core/Public/Math/GenericOctree.h b/Engine/Source/Runtime/Core/Public/Math/GenericOctree.h
--- a/Engine/Source/Runtime/Core/Public/Math/GenericOctree.h
+++ b/Engine/Source/Runtime/Core/Public/Math/GenericOctree.h
@@ -1032,4 +1032,740 @@ protected:
 	}
 };
 
+/** An octree. */
+template<typename ElementType, typename OctreeSemantics>
+class UE_DEPRECATED(4.26, "The old Octree is deprecated use TOctree2.") TOctree
+{
+public:
+
+	typedef TArray<ElementType, typename OctreeSemantics::ElementAllocator> ElementArrayType;
+	typedef typename ElementArrayType::TConstIterator ElementConstIt;
+
+	/** A node in the octree. */
+	class FNode
+	{
+	public:
+
+		friend class TOctree;
+
+		/** Initialization constructor. */
+		explicit FNode(const FNode* InParent)
+			: Parent(InParent)
+			, InclusiveNumElements(0)
+			, bIsLeaf(true)
+		{
+			FOREACH_OCTREE_CHILD_NODE(ChildRef)
+			{
+				Children[ChildRef.Index] = NULL;
+			}
+		}
+
+		/** Destructor. */
+		~FNode()
+		{
+			FOREACH_OCTREE_CHILD_NODE(ChildRef)
+			{
+				delete Children[ChildRef.Index];
+			}
+		}
+
+		// Accessors.
+		FORCEINLINE ElementConstIt GetElementIt() const { return ElementConstIt(Elements); }
+		FORCEINLINE bool IsLeaf() const { return bIsLeaf; }
+		FORCEINLINE bool HasChild(FOctreeChildNodeRef ChildRef) const
+		{
+			return Children[ChildRef.Index] != NULL && Children[ChildRef.Index]->InclusiveNumElements > 0;
+		}
+		FORCEINLINE FNode* GetChild(FOctreeChildNodeRef ChildRef) const
+		{
+			return Children[ChildRef.Index];
+		}
+		FORCEINLINE int32 GetElementCount() const
+		{
+			return Elements.Num();
+		}
+		FORCEINLINE int32 GetInclusiveElementCount() const
+		{
+			return InclusiveNumElements;
+		}
+		FORCEINLINE const ElementArrayType& GetElements() const
+		{
+			return Elements;
+		}
+		void ShrinkElements()
+		{
+			Elements.Shrink();
+			FOREACH_OCTREE_CHILD_NODE(ChildRef)
+			{
+				if (Children[ChildRef.Index])
+				{
+					Children[ChildRef.Index]->ShrinkElements();
+				}
+			}
+		}
+
+		void ApplyOffset(const FVector& InOffset)
+		{
+			for (auto& Element : Elements)
+			{
+				OctreeSemantics::ApplyOffset(Element, InOffset);
+			}
+
+			FOREACH_OCTREE_CHILD_NODE(ChildRef)
+			{
+				if (Children[ChildRef.Index])
+				{
+					Children[ChildRef.Index]->ApplyOffset(InOffset);
+				}
+			}
+		}
+
+	private:
+
+		/** The elements in this node. */
+		mutable ElementArrayType Elements;
+
+		/** The parent of this node. */
+		const FNode* Parent;
+
+		/** The children of the node. */
+		mutable FNode* Children[8];
+
+		/** The number of elements contained by the node and its child nodes. */
+		mutable uint32 InclusiveNumElements : 31;
+
+		/** true if the meshes should be added directly to the node, rather than subdividing when possible. */
+		mutable uint32 bIsLeaf : 1;
+	};
+
+
+
+	/** A reference to an octree node, its context, and a read lock. */
+	class FNodeReference
+	{
+	public:
+
+		const FNode* Node;
+		FOctreeNodeContext Context;
+
+		/** Default constructor. */
+		FNodeReference() :
+			Node(NULL),
+			Context()
+		{}
+
+		/** Initialization constructor. */
+		FNodeReference(const FNode* InNode, const FOctreeNodeContext& InContext) :
+			Node(InNode),
+			Context(InContext)
+		{}
+	};
+
+	/** The default iterator allocator gives the stack enough inline space to contain a path and its siblings from root to leaf. */
+	typedef TInlineAllocator<7 * (14 - 1) + 8> DefaultStackAllocator;
+
+	/** An octree node iterator. */
+	template<typename StackAllocator = DefaultStackAllocator>
+	class TConstIterator
+	{
+	public:
+
+		/** Pushes a child of the current node onto the stack of nodes to visit. */
+		void PushChild(FOctreeChildNodeRef ChildRef)
+		{
+			FNodeReference* NewNode = new (NodeStack) FNodeReference;
+			NewNode->Node = CurrentNode.Node->GetChild(ChildRef);
+			CurrentNode.Context.GetChildContext(ChildRef, &NewNode->Context);
+		}
+		/** Pushes a child of the current node onto the stack of nodes to visit. */
+		void PushChild(FOctreeChildNodeRef ChildRef, uint32 FullyInsideView, uint32 FullyOutsideView)
+		{
+			FNodeReference* NewNode = new (NodeStack) FNodeReference;
+			NewNode->Node = CurrentNode.Node->GetChild(ChildRef);
+			CurrentNode.Context.GetChildContext(ChildRef, &NewNode->Context);
+			NewNode->Context.InCullBits = FullyInsideView;
+			NewNode->Context.OutCullBits = FullyOutsideView;
+		}
+		/** Pushes a child of the current node onto the stack of nodes to visit. */
+		void PushChild(FOctreeChildNodeRef ChildRef, const FOctreeNodeContext& Context)
+		{
+			new (NodeStack) FNodeReference(CurrentNode.Node->GetChild(ChildRef), Context);
+		}
+
+
+		/** Iterates to the next node. */
+		void Advance()
+		{
+			if (NodeStack.Num())
+			{
+				CurrentNode = NodeStack[NodeStack.Num() - 1];
+				NodeStack.RemoveAt(NodeStack.Num() - 1);
+			}
+			else
+			{
+				CurrentNode = FNodeReference();
+			}
+		}
+
+		/** Checks if there are any nodes left to iterate over. */
+		bool HasPendingNodes() const
+		{
+			return CurrentNode.Node != NULL;
+		}
+
+		/** Starts iterating at the root of an octree. */
+		TConstIterator(const TOctree& Tree)
+			: CurrentNode(FNodeReference(&Tree.RootNode, Tree.RootNodeContext))
+		{}
+
+		/** Starts iterating at a particular node of an octree. */
+		TConstIterator(const FNode& Node, const FOctreeNodeContext& Context) :
+			CurrentNode(FNodeReference(&Node, Context))
+		{}
+
+		// Accessors.
+		const FNode& GetCurrentNode() const
+		{
+			return *CurrentNode.Node;
+		}
+		const FOctreeNodeContext& GetCurrentContext() const
+		{
+			return CurrentNode.Context;
+		}
+
+	private:
+
+		/** The node that is currently being visited. */
+		FNodeReference CurrentNode;
+
+		/** The nodes which are pending iteration. */
+		TArray<FNodeReference, StackAllocator> NodeStack;
+	};
+
+	/** Iterates over the elements in the octree that intersect a bounding box. */
+	template<typename StackAllocator = DefaultStackAllocator>
+	class TConstElementBoxIterator
+	{
+	public:
+
+		/** Iterates to the next element. */
+		void Advance()
+		{
+			++ElementIt;
+			AdvanceToNextIntersectingElement();
+		}
+
+		/** Checks if there are any elements left to iterate over. */
+		bool HasPendingElements() const
+		{
+			return NodeIt.HasPendingNodes();
+		}
+
+		/** Initialization constructor. */
+		TConstElementBoxIterator(const TOctree& Tree, const FBoxCenterAndExtent& InBoundingBox)
+			: IteratorBounds(InBoundingBox)
+			, NodeIt(Tree)
+			, ElementIt(Tree.RootNode.GetElementIt())
+		{
+			ProcessChildren();
+			AdvanceToNextIntersectingElement();
+		}
+
+		// Accessors.
+		const ElementType& GetCurrentElement() const
+		{
+			return *ElementIt;
+		}
+
+	private:
+
+		/** The bounding box to check for intersection with. */
+		FBoxCenterAndExtent IteratorBounds;
+
+		/** The octree node iterator. */
+		TConstIterator<StackAllocator> NodeIt;
+
+		/** The element iterator for the current node. */
+		ElementConstIt ElementIt;
+
+		/** Processes the children of the current node. */
+		void ProcessChildren()
+		{
+			// Add the child nodes that intersect the bounding box to the node iterator's stack.
+			const FNode& CurrentNode = NodeIt.GetCurrentNode();
+			const FOctreeNodeContext& Context = NodeIt.GetCurrentContext();
+			const FOctreeChildNodeSubset IntersectingChildSubset = Context.GetIntersectingChildren(IteratorBounds);
+			FOREACH_OCTREE_CHILD_NODE(ChildRef)
+			{
+				if (IntersectingChildSubset.Contains(ChildRef) && CurrentNode.HasChild(ChildRef))
+				{
+					NodeIt.PushChild(ChildRef);
+				}
+			}
+		}
+
+		/** Advances the iterator to the next intersecting primitive, starting at a primitive in the current node. */
+		void AdvanceToNextIntersectingElement()
+		{
+			check(NodeIt.HasPendingNodes()); // please don't call this after iteration has ended
+
+			while (1)
+			{
+				ElementConstIt LocalElementIt(ElementIt);
+				if (LocalElementIt)
+				{
+					FPlatformMisc::Prefetch(&(*LocalElementIt));
+					FPlatformMisc::Prefetch(&(*LocalElementIt), PLATFORM_CACHE_LINE_SIZE);
+
+					// this is redundantly pull out of the while loop to prevent LHS on the iterator
+					// Check if the current element intersects the bounding box.
+					if (Intersect(OctreeSemantics::GetBoundingBox(*LocalElementIt), IteratorBounds))
+					{
+						// If it intersects, break out of the advancement loop.
+						Move(ElementIt, LocalElementIt);
+						return;
+					}
+
+					// Check if we've advanced past the elements in the current node.
+					while (++LocalElementIt)
+					{
+						FPlatformMisc::Prefetch(&(*LocalElementIt), PLATFORM_CACHE_LINE_SIZE);
+
+						// Check if the current element intersects the bounding box.
+						if (Intersect(OctreeSemantics::GetBoundingBox(*LocalElementIt), IteratorBounds))
+
+						{
+							// If it intersects, break out of the advancement loop.
+							Move(ElementIt, LocalElementIt);
+							return;
+						}
+					}
+				}
+				// Advance to the next node.
+				NodeIt.Advance();
+				if (!NodeIt.HasPendingNodes())
+				{
+					Move(ElementIt, LocalElementIt);
+					return;
+				}
+				ProcessChildren();
+				// The element iterator can't be assigned to, but it can be replaced by Move.
+				Move(ElementIt, NodeIt.GetCurrentNode().GetElementIt());
+			}
+		}
+	};
+
+	/**
+	 * Adds an element to the octree.
+	 * @param Element - The element to add.
+	 * @return An identifier for the element in the octree.
+	 */
+	void AddElement(typename TTypeTraits<ElementType>::ConstInitType Element)
+	{
+		AddElementToNode(Element, RootNode, RootNodeContext);
+	}
+
+	/**
+	 * Removes an element from the octree.
+	 * @param ElementId - The element to remove from the octree.
+	 */
+	void RemoveElement(FOctreeElementId ElementId)
+	{
+		check(ElementId.IsValidId());
+
+		FNode* ElementIdNode = (FNode*)ElementId.Node;
+
+		// Remove the element from the node's element list.
+		ElementIdNode->Elements.RemoveAtSwap(ElementId.ElementIndex);
+
+		SetOctreeMemoryUsage(this, TotalSizeBytes - sizeof(ElementType));
+
+		if (ElementId.ElementIndex < ElementIdNode->Elements.Num())
+		{
+			// Update the external element id for the element that was swapped into the vacated element index.
+			SetElementId(ElementIdNode->Elements[ElementId.ElementIndex], ElementId);
+		}
+
+		// Update the inclusive element counts for the nodes between the element and the root node,
+		// and find the largest node that is small enough to collapse.
+		const FNode* CollapseNode = NULL;
+		for (const FNode* Node = ElementIdNode; Node; Node = Node->Parent)
+		{
+			--Node->InclusiveNumElements;
+			if (Node->InclusiveNumElements < OctreeSemantics::MinInclusiveElementsPerNode)
+			{
+				CollapseNode = Node;
+			}
+		}
+
+		// Collapse the largest node that was pushed below the threshold for collapse by the removal.
+		if (CollapseNode && !CollapseNode->bIsLeaf)
+		{
+			if (CollapseNode->Elements.Num() < (int32)CollapseNode->InclusiveNumElements)
+			{
+				CollapseNode->Elements.Reserve(CollapseNode->InclusiveNumElements);
+
+				// Gather the elements contained in this node and its children.
+				for (TConstIterator<> ChildNodeIt(*CollapseNode, RootNodeContext); ChildNodeIt.HasPendingNodes(); ChildNodeIt.Advance())
+				{
+					const FNode& ChildNode = ChildNodeIt.GetCurrentNode();
+
+					if (&ChildNode != CollapseNode)
+					{
+						// Child node will be collapsed so move the child's elements to the collapse node element list.
+						for (ElementType& Element : ChildNode.Elements)
+						{
+							const int32 NewElementIndex = CollapseNode->Elements.Add(MoveTemp(Element));
+
+							// Update the external element id for the element that's being collapsed.
+							SetElementId(CollapseNode->Elements[NewElementIndex], FOctreeElementId(CollapseNode, NewElementIndex));
+						}
+					}
+
+					// Recursively visit all child nodes.
+					FOREACH_OCTREE_CHILD_NODE(ChildRef)
+					{
+						if (ChildNode.HasChild(ChildRef))
+						{
+							ChildNodeIt.PushChild(ChildRef);
+						}
+					}
+				}
+
+				// Free the child nodes.
+				FOREACH_OCTREE_CHILD_NODE(ChildRef)
+				{
+					if (CollapseNode->Children[ChildRef.Index])
+					{
+						SetOctreeMemoryUsage(this, TotalSizeBytes - sizeof(*CollapseNode->Children[ChildRef.Index]));
+					}
+
+					delete CollapseNode->Children[ChildRef.Index];
+					CollapseNode->Children[ChildRef.Index] = NULL;
+				}
+			}
+
+			// Mark the node as a leaf.
+			CollapseNode->bIsLeaf = true;
+		}
+	}
+
+
+	void Destroy()
+	{
+		RootNode.~FNode();
+		new (&RootNode) FNode(NULL);
+
+		// this looks a bit @hacky, but FNode's destructor doesn't 
+		// update TotalSizeBytes so better to set it to 0 than
+		// not do nothing with it (would contain obviously false value)
+		SetOctreeMemoryUsage(this, 0);
+	}
+
+	/** Accesses an octree element by ID. */
+	ElementType& GetElementById(FOctreeElementId ElementId)
+	{
+		check(ElementId.IsValidId());
+		FNode* ElementIdNode = (FNode*)ElementId.Node;
+		return ElementIdNode->Elements[ElementId.ElementIndex];
+	}
+
+	/** Accesses an octree element by ID. */
+	const ElementType& GetElementById(FOctreeElementId ElementId) const
+	{
+		check(ElementId.IsValidId());
+		FNode* ElementIdNode = (FNode*)ElementId.Node;
+		return ElementIdNode->Elements[ElementId.ElementIndex];
+	}
+
+	/** Checks if given ElementId represents a valid Octree element */
+	bool IsValidElementId(FOctreeElementId ElementId) const
+	{
+		return ElementId.IsValidId()
+			&& ElementId.ElementIndex != INDEX_NONE
+			&& ElementId.ElementIndex < ((FNode*)ElementId.Node)->Elements.Num();
+	};
+
+	/** Writes stats for the octree to the log. */
+	void DumpStats() const
+	{
+		int32 NumNodes = 0;
+		int32 NumLeaves = 0;
+		int32 NumElements = 0;
+		int32 MaxElementsPerNode = 0;
+		TArray<int32> NodeElementDistribution;
+
+		for (TConstIterator<> NodeIt(*this); NodeIt.HasPendingNodes(); NodeIt.Advance())
+		{
+			const FNode& CurrentNode = NodeIt.GetCurrentNode();
+			const int32 CurrentNodeElementCount = CurrentNode.GetElementCount();
+
+			NumNodes++;
+			if (CurrentNode.IsLeaf())
+			{
+				NumLeaves++;
+			}
+
+			NumElements += CurrentNodeElementCount;
+			MaxElementsPerNode = FMath::Max(MaxElementsPerNode, CurrentNodeElementCount);
+
+			if (CurrentNodeElementCount >= NodeElementDistribution.Num())
+			{
+				NodeElementDistribution.AddZeroed(CurrentNodeElementCount - NodeElementDistribution.Num() + 1);
+			}
+			NodeElementDistribution[CurrentNodeElementCount]++;
+
+			FOREACH_OCTREE_CHILD_NODE(ChildRef)
+			{
+				if (CurrentNode.HasChild(ChildRef))
+				{
+					NodeIt.PushChild(ChildRef);
+				}
+			}
+		}
+
+		UE_LOG(LogGenericOctree, Log, TEXT("Octree overview:"));
+		UE_LOG(LogGenericOctree, Log, TEXT("\t%i nodes"), NumNodes);
+		UE_LOG(LogGenericOctree, Log, TEXT("\t%i leaves"), NumLeaves);
+		UE_LOG(LogGenericOctree, Log, TEXT("\t%i elements"), NumElements);
+		UE_LOG(LogGenericOctree, Log, TEXT("\t%i >= elements per node"), MaxElementsPerNode);
+		UE_LOG(LogGenericOctree, Log, TEXT("Octree node element distribution:"));
+		for (int32 i = 0; i < NodeElementDistribution.Num(); i++)
+		{
+			if (NodeElementDistribution[i] > 0)
+			{
+				UE_LOG(LogGenericOctree, Log, TEXT("\tElements: %3i, Nodes: %3i"), i, NodeElementDistribution[i]);
+			}
+		}
+	}
+
+
+	SIZE_T GetSizeBytes() const
+	{
+		return TotalSizeBytes;
+	}
+
+	float GetNodeLevelExtent(int32 Level) const
+	{
+		const int32 ClampedLevel = FMath::Clamp<uint32>(Level, 0, OctreeSemantics::MaxNodeDepth);
+		return RootNodeContext.Bounds.Extent.X * FMath::Pow((1.0f + 1.0f / (float)FOctreeNodeContext::LoosenessDenominator) / 2.0f, ClampedLevel);
+	}
+
+	FBoxCenterAndExtent GetRootBounds() const
+	{
+		return RootNodeContext.Bounds;
+	}
+
+	void ShrinkElements()
+	{
+		RootNode.ShrinkElements();
+	}
+
+	/**
+	 * Apply an arbitrary offset to all elements in the tree
+	 * InOffset - offset to apply
+	 * bGlobalOctree - hint that this octree is used as a boundless global volume,
+	 *  so only content will be shifted but not origin of the octree
+	 */
+	void ApplyOffset(const FVector& InOffset, bool bGlobalOctree)
+	{
+		// Shift elements
+		RootNode.ApplyOffset(InOffset);
+
+		// Make a local copy of all nodes
+		FNode OldRootNode = RootNode;
+
+		// Assign to root node a NULL node, so Destroy procedure will not delete child nodes
+		RootNode = FNode(nullptr);
+		// Call destroy to clean up octree
+		Destroy();
+
+		if (!bGlobalOctree)
+		{
+			RootNodeContext.Bounds.Center += FVector4(InOffset, 0.0f);
+		}
+
+		// Add all elements from a saved nodes to a new empty octree
+		for (TConstIterator<> NodeIt(OldRootNode, RootNodeContext); NodeIt.HasPendingNodes(); NodeIt.Advance())
+		{
+			const auto& CurrentNode = NodeIt.GetCurrentNode();
+
+			FOREACH_OCTREE_CHILD_NODE(ChildRef)
+			{
+				if (CurrentNode.HasChild(ChildRef))
+				{
+					NodeIt.PushChild(ChildRef);
+				}
+			}
+
+			for (auto ElementIt = CurrentNode.GetElementIt(); ElementIt; ++ElementIt)
+			{
+				AddElement(*ElementIt);
+			}
+		}
+
+		// Saved nodes will be deleted on scope exit
+	}
+
+
+	/** Initialization constructor. */
+	TOctree(const FVector& InOrigin, float InExtent)
+		: RootNode(NULL)
+		, RootNodeContext(FBoxCenterAndExtent(InOrigin, FVector(InExtent, InExtent, InExtent)), 0, 0)
+		, MinLeafExtent(InExtent* FMath::Pow((1.0f + 1.0f / (float)FOctreeNodeContext::LoosenessDenominator) / 2.0f, OctreeSemantics::MaxNodeDepth))
+		, TotalSizeBytes(0)
+	{
+	}
+
+	/** DO NOT USE. This constructor is for internal usage only for hot-reload purposes. */
+	TOctree() : RootNode(nullptr)
+	{
+		EnsureRetrievingVTablePtrDuringCtor(TEXT("TOctree()"));
+	}
+
+private:
+
+	/** The octree's root node. */
+	FNode RootNode;
+
+	/** The octree's root node's context. */
+	FOctreeNodeContext RootNodeContext;
+
+	/** The extent of a leaf at the maximum allowed depth of the tree. */
+	float MinLeafExtent;
+
+	/** this function basically set TotalSizeBytes, but gives opportunity to
+	 *	include this Octree in memory stats */
+	void SetOctreeMemoryUsage(TOctree<ElementType, OctreeSemantics>* Octree, int32 NewSize)
+	{
+		Octree->TotalSizeBytes = NewSize;
+	}
+
+	SIZE_T TotalSizeBytes;
+
+	/** Adds an element to a node or its children. */
+	void AddElementToNode(
+		typename TTypeTraits<ElementType>::ConstInitType Element,
+		const FNode& InNode,
+		const FOctreeNodeContext& InContext
+	)
+	{
+		const FBoxCenterAndExtent ElementBounds(OctreeSemantics::GetBoundingBox(Element));
+
+		for (TConstIterator<TInlineAllocator<1> > NodeIt(InNode, InContext); NodeIt.HasPendingNodes(); NodeIt.Advance())
+		{
+			const FNode& Node = NodeIt.GetCurrentNode();
+			const FOctreeNodeContext& Context = NodeIt.GetCurrentContext();
+			const bool bIsLeaf = Node.IsLeaf();
+
+			bool bAddElementToThisNode = false;
+
+			// Increment the number of elements included in this node and its children.
+			Node.InclusiveNumElements++;
+
+			if (bIsLeaf)
+			{
+				// If this is a leaf, check if adding this element would turn it into a node by overflowing its element list.
+				if (Node.Elements.Num() + 1 > OctreeSemantics::MaxElementsPerLeaf && Context.Bounds.Extent.X > MinLeafExtent)
+				{
+					// Copy the leaf's elements, remove them from the leaf, and turn it into a node.
+					ElementArrayType ChildElements;
+					Exchange(ChildElements, Node.Elements);
+					SetOctreeMemoryUsage(this, TotalSizeBytes - ChildElements.Num() * sizeof(ElementType));
+					Node.InclusiveNumElements = 0;
+
+					// Allow elements to be added to children of this node.
+					Node.bIsLeaf = false;
+
+					// Re-add all of the node's child elements, potentially creating children of this node for them.
+					for (ElementConstIt ElementIt(ChildElements); ElementIt; ++ElementIt)
+					{
+						AddElementToNode(*ElementIt, Node, Context);
+					}
+
+					// Add the element to this node.
+					AddElementToNode(Element, Node, Context);
+					return;
+				}
+				else
+				{
+					// If the leaf has room for the new element, simply add it to the list.
+					bAddElementToThisNode = true;
+				}
+			}
+			else
+			{
+				// If this isn't a leaf, find a child that entirely contains the element.
+				const FOctreeChildNodeRef ChildRef = Context.GetContainingChild(ElementBounds);
+				if (ChildRef.IsNULL())
+				{
+					// If none of the children completely contain the element, add it to this node directly.
+					bAddElementToThisNode = true;
+				}
+				else
+				{
+					// Create the child node if it hasn't been created yet.
+					if (!Node.Children[ChildRef.Index])
+					{
+						Node.Children[ChildRef.Index] = new typename TOctree<ElementType, OctreeSemantics>::FNode(&Node);
+						SetOctreeMemoryUsage(this, TotalSizeBytes + sizeof(*Node.Children[ChildRef.Index]));
+					}
+
+					// Push the node onto the stack to visit.
+					NodeIt.PushChild(ChildRef);
+				}
+			}
+
+			if (bAddElementToThisNode)
+			{
+				// Add the element to this node.
+				new(Node.Elements) ElementType(Element);
+
+				SetOctreeMemoryUsage(this, TotalSizeBytes + sizeof(ElementType));
+
+				// Set the element's ID.
+				SetElementId(Element, FOctreeElementId(&Node, Node.Elements.Num() - 1));
+				return;
+			}
+		}
+
+		UE_LOG(LogGenericOctree, Fatal,
+			TEXT("Failed to find an octree node for an element with bounds (%f,%f,%f) +/- (%f,%f,%f)!"),
+			ElementBounds.Center.X,
+			ElementBounds.Center.Y,
+			ElementBounds.Center.Z,
+			ElementBounds.Extent.X,
+			ElementBounds.Extent.Y,
+			ElementBounds.Extent.Z
+		);
+	}
+
+	// Concept definition for the new octree semantics, which adds a new TOctree parameter
+	struct COctreeSemanticsV2
+	{
+		template<typename Semantics>
+		auto Requires(typename Semantics::FOctree& OctreeInstance, const ElementType& Element, FOctreeElementId Id)
+			-> decltype(Semantics::SetElementId(OctreeInstance, Element, Id));
+	};
+
+	// Function overload set which calls the V2 version if it's supported or the old version if it's not
+	template <typename Semantics>
+	typename TEnableIf<!TModels<COctreeSemanticsV2, Semantics>::Value>::Type SetOctreeSemanticsElementId(const ElementType& Element, FOctreeElementId Id)
+	{
+		Semantics::SetElementId(Element, Id);
+	}
+	template <typename Semantics>
+	typename TEnableIf<TModels<COctreeSemanticsV2, Semantics>::Value>::Type SetOctreeSemanticsElementId(const ElementType& Element, FOctreeElementId Id)
+	{
+		Semantics::SetElementId(*this, Element, Id);
+	}
+
+protected:
+	// redirects SetElementId call to the proper implementation
+	void SetElementId(const ElementType& Element, FOctreeElementId Id)
+	{
+		SetOctreeSemanticsElementId<OctreeSemantics>(Element, Id);
+	}
+};
+
 #include "GenericOctree.inl"
diff --git a/Engine/Source/Runtime/Core/Public/Math/GenericOctreePublic.h b/Engine/Source/Runtime/Core/Public/Math/GenericOctreePublic.h
--- a/Engine/Source/Runtime/Core/Public/Math/GenericOctreePublic.h
+++ b/Engine/Source/Runtime/Core/Public/Math/GenericOctreePublic.h
@@ -53,3 +53,46 @@ private:
 		return ElementIndex;
 	}
 };
+
+/**
+ *	An identifier for an element in the octree.
+ */
+class FOctreeElementId
+{
+public:
+
+	template<typename, typename>
+	friend class TOctree;
+
+	/** Default constructor. */
+	FOctreeElementId()
+		: Node(NULL)
+		, ElementIndex(INDEX_NONE)
+	{}
+
+	/** @return a boolean value representing whether the id is NULL. */
+	bool IsValidId() const
+	{
+		return Node != NULL;
+	}
+
+private:
+
+	/** The node the element is in. */
+	const void* Node;
+
+	/** The index of the element in the node's element array. */
+	int32 ElementIndex;
+
+	/** Initialization constructor. */
+	FOctreeElementId(const void* InNode, int32 InElementIndex)
+		: Node(InNode)
+		, ElementIndex(InElementIndex)
+	{}
+
+	/** Implicit conversion to the element index. */
+	operator int32() const
+	{
+		return ElementIndex;
+	}
+};
\ No newline at end of file
