Optimize UTF8ToTCHAR converter
  - Add fast path for ASCII only text and reduce loop iteration cost

  - 550ms -> 133ms of aggregated time during editor startup

  - Other tests to ensure no regression on worst case scenarios
  - 5.0s -> 1.5s (1GB UTF8 with ASCII Only)
  - 4.8s -> 2.4s (1GB UTF8 with extended char every 16 chars)
  - 4.2s -> 3.6s (1GB UTF8 with extended char every 8 chars)

#rnx

diff --git a/Engine/Source/Runtime/Core/Public/Containers/StringConv.h b/Engine/Source/Runtime/Core/Public/Containers/StringConv.h
--- a/Engine/Source/Runtime/Core/Public/Containers/StringConv.h
+++ b/Engine/Source/Runtime/Core/Public/Containers/StringConv.h
@@ -654,40 +654,78 @@ private:
 	static void Convert_Impl(DestBufferType& ConvertedBuffer, int32 DestLen, const ANSICHAR* Source, const int32 SourceLen)
 	{
 		const ANSICHAR* SourceEnd = Source + SourceLen;
+
+		const uint64 ExtendedCharMask = 0x8080808080808080;
 		while (Source < SourceEnd && DestLen > 0)
 		{
-			// Read our codepoint, advancing the source pointer
-			uint32 Codepoint = CodepointFromUtf8(Source, UE_PTRDIFF_TO_UINT32(SourceEnd - Source));
-
-#if !PLATFORM_TCHAR_IS_4_BYTES
-			// We want to write out two chars
-			if (StringConv::IsEncodedSurrogate(Codepoint))
+			// In case we're given an unaligned pointer, we'll
+			// fallback to the slow path until properly aligned.
+			if (IsAligned(Source, 8))
 			{
-				// We need two characters to write the surrogate pair
-				if (DestLen >= 2)
+				// Fast path for most common case
+				while (Source < SourceEnd - 8 && DestLen >= 8)
 				{
-					uint16 HighSurrogate = 0;
-					uint16 LowSurrogate = 0;
-					StringConv::DecodeSurrogate(Codepoint, HighSurrogate, LowSurrogate);
+					// Detect any extended characters 8 chars at a time
+					if ((*(const uint64*)Source) & ExtendedCharMask)
+					{
+						// Move to slow path since we got extended characters to process
+						break;
+					}
 
-					*(ConvertedBuffer++) = (ToType)HighSurrogate;
-					*(ConvertedBuffer++) = (ToType)LowSurrogate;
-					DestLen -= 2;
-					continue;
+					// This should get unrolled on most compiler
+					// ROI of diminished return to vectorize this as we 
+					// would have to deal with alignment, endianness and
+					// rewrite the iterators to support bulk writes
+					for (int32 Index = 0; Index < 8; ++Index)
+					{
+						*(ConvertedBuffer++) = (ToType)(uint8)*(Source++);
+					}
+					DestLen -= 8;
 				}
-
-				// If we don't have space, write a bogus character instead (we should have space for it)
-				Codepoint = UNICODE_BOGUS_CHAR_CODEPOINT;
 			}
-			else if (Codepoint > StringConv::ENCODED_SURROGATE_END_CODEPOINT)
+
+			// Slow path for extended characters
+			while (Source < SourceEnd && DestLen > 0)
 			{
-				// Ignore values higher than the supplementary plane range
-				Codepoint = UNICODE_BOGUS_CHAR_CODEPOINT;
-			}
-#endif	// !PLATFORM_TCHAR_IS_4_BYTES
+				// Read our codepoint, advancing the source pointer
+				uint32 Codepoint = CodepointFromUtf8(Source, UE_PTRDIFF_TO_UINT32(SourceEnd - Source));
 
-			*(ConvertedBuffer++) = (ToType)Codepoint;
-			--DestLen;
+	#if !PLATFORM_TCHAR_IS_4_BYTES
+				// We want to write out two chars
+				if (StringConv::IsEncodedSurrogate(Codepoint))
+				{
+					// We need two characters to write the surrogate pair
+					if (DestLen >= 2)
+					{
+						uint16 HighSurrogate = 0;
+						uint16 LowSurrogate = 0;
+						StringConv::DecodeSurrogate(Codepoint, HighSurrogate, LowSurrogate);
+
+						*(ConvertedBuffer++) = (ToType)HighSurrogate;
+						*(ConvertedBuffer++) = (ToType)LowSurrogate;
+						DestLen -= 2;
+						continue;
+					}
+
+					// If we don't have space, write a bogus character instead (we should have space for it)
+					Codepoint = UNICODE_BOGUS_CHAR_CODEPOINT;
+				}
+				else if (Codepoint > StringConv::ENCODED_SURROGATE_END_CODEPOINT)
+				{
+					// Ignore values higher than the supplementary plane range
+					Codepoint = UNICODE_BOGUS_CHAR_CODEPOINT;
+				}
+	#endif	// !PLATFORM_TCHAR_IS_4_BYTES
+
+				*(ConvertedBuffer++) = (ToType)Codepoint;
+				--DestLen;
+
+				// Return to the fast path once aligned and back to simple ASCII chars
+				if (Codepoint < 128 && IsAligned(Source, 8))
+				{
+					break;
+				}
+			}
 		}
 	}
 };
