Fixed all size variants of IAsyncReadRequest so they correctly wait inside of WaitCompletionImpl for a callback to be completed before returning.
Normal read requests have this behavior, but size requests were not respecting this, which caused problems when the constructor launches a callback on another thread and then immediately calls WaitCompletion, such as in the callback lambda inside FFileCacheHandle::FFileCacheHandle / PushCompletedRequest.
#rnx

diff --git a/Engine/Source/Runtime/Core/Private/FileCache/FileCache.cpp b/Engine/Source/Runtime/Core/Private/FileCache/FileCache.cpp
--- a/Engine/Source/Runtime/Core/Private/FileCache/FileCache.cpp
+++ b/Engine/Source/Runtime/Core/Private/FileCache/FileCache.cpp
@@ -160,7 +160,7 @@ public:
 				{
 					// Requests are added to this list from the completed callback, but the final completion flag is not set until after callback is finished
 					// This means that there's a narrow window where the request is not technically considered to be complete yet
-					CompletedRequest->WaitCompletion();
+					verify(CompletedRequest->WaitCompletion());
 					delete CompletedRequest;
 				}
 			}, TStatId());
diff --git a/Engine/Source/Runtime/Core/Private/Windows/WindowsAsyncIO.h b/Engine/Source/Runtime/Core/Private/Windows/WindowsAsyncIO.h
--- a/Engine/Source/Runtime/Core/Private/Windows/WindowsAsyncIO.h
+++ b/Engine/Source/Runtime/Core/Private/Windows/WindowsAsyncIO.h
@@ -341,9 +341,14 @@ public:
 		Size = InFileSize;
 		SetComplete();
 	}
+
 	virtual void WaitCompletionImpl(float TimeLimitSeconds) override
 	{
+		// Even though SetComplete called in the constructor and sets bCompleteAndCallbackCalled=true, we still need to implement WaitComplete as
+		// the CompleteCallback can end up starting async tasks that can overtake the constructor execution and need to wait for the constructor to finish.
+		while (!bCompleteAndCallbackCalled);
 	}
+
 	virtual void CancelImpl()
 	{
 	}
@@ -357,9 +362,14 @@ public:
 	{
 		SetComplete();
 	}
+
 	virtual void WaitCompletionImpl(float TimeLimitSeconds) override
 	{
+		// Even though SetComplete called in the constructor and sets bCompleteAndCallbackCalled=true, we still need to implement WaitComplete as
+		// the CompleteCallback can end up starting async tasks that can overtake the constructor execution and need to wait for the constructor to finish.
+		while (!bCompleteAndCallbackCalled);
 	}
+
 	virtual void CancelImpl()
 	{
 	}
diff --git a/Engine/Source/Runtime/CoreUObject/Private/Serialization/BulkData2.cpp b/Engine/Source/Runtime/CoreUObject/Private/Serialization/BulkData2.cpp
--- a/Engine/Source/Runtime/CoreUObject/Private/Serialization/BulkData2.cpp
+++ b/Engine/Source/Runtime/CoreUObject/Private/Serialization/BulkData2.cpp
@@ -253,7 +253,9 @@ private:
 
 	virtual void WaitCompletionImpl(float TimeLimitSeconds) override
 	{
-		// No need to wait as the work is done in the constructor
+		// Even though SetComplete called in the constructor and sets bCompleteAndCallbackCalled=true, we still need to implement WaitComplete as
+		// the CompleteCallback can end up starting async tasks that can overtake the constructor execution and need to wait for the constructor to finish.
+		while (!bCompleteAndCallbackCalled);
 	}
 
 	virtual void CancelImpl() override
diff --git a/Engine/Source/Runtime/PakFile/Private/IPlatformFilePak.cpp b/Engine/Source/Runtime/PakFile/Private/IPlatformFilePak.cpp
--- a/Engine/Source/Runtime/PakFile/Private/IPlatformFilePak.cpp
+++ b/Engine/Source/Runtime/PakFile/Private/IPlatformFilePak.cpp
@@ -386,9 +386,14 @@ public:
 		Size = InFileSize;
 		SetComplete();
 	}
+
 	virtual void WaitCompletionImpl(float TimeLimitSeconds) override
 	{
+		// Even though SetComplete called in the constructor and sets bCompleteAndCallbackCalled=true, we still need to implement WaitComplete as
+		// the CompleteCallback can end up starting async tasks that can overtake the constructor execution and need to wait for the constructor to finish.
+		while (!bCompleteAndCallbackCalled);
 	}
+
 	virtual void CancelImpl()
 	{
 	}
