Avoid FName -> FString allocation in FBlacklistPaths

diff --git a/Engine/Source/Runtime/Core/Private/Misc/BlacklistNames.cpp b/Engine/Source/Runtime/Core/Private/Misc/BlacklistNames.cpp
--- a/Engine/Source/Runtime/Core/Private/Misc/BlacklistNames.cpp
+++ b/Engine/Source/Runtime/Core/Private/Misc/BlacklistNames.cpp
@@ -1,6 +1,7 @@
 // Copyright Epic Games, Inc. All Rights Reserved.
 
 #include "Misc/BlacklistNames.h"
+#include "Misc/StringBuilder.h"
 
 FBlacklistNames::FBlacklistNames() :
 	bSuppressOnFilterChanged(false)
@@ -129,14 +130,16 @@ FBlacklistPaths::FBlacklistPaths() :
 {
 }
 
-bool FBlacklistPaths::PassesFilter(const FString& Item) const
+bool FBlacklistPaths::PassesFilter(const FStringView Item) const
 {
-	if (Whitelist.Num() > 0 && !Whitelist.Contains(Item))
+	const uint32 ItemHash = GetTypeHash(Item);
+
+	if (Whitelist.Num() > 0 && !Whitelist.ContainsByHash(ItemHash, Item))
 	{
 		return false;
 	}
 
-	if (Blacklist.Contains(Item))
+	if (Blacklist.ContainsByHash(ItemHash, Item))
 	{
 		return false;
 	}
@@ -151,10 +154,17 @@ bool FBlacklistPaths::PassesFilter(const FString& Item) const
 
 bool FBlacklistPaths::PassesFilter(const FName Item) const
 {
-	return PassesFilter(Item.ToString());
+	TStringBuilder<FName::StringBufferSize> ItemStr;
+	Item.ToString(ItemStr);
+	return PassesFilter(FStringView(ItemStr));
 }
 
-bool FBlacklistPaths::PassesStartsWithFilter(const FString& Item) const
+bool FBlacklistPaths::PassesFilter(const TCHAR* Item) const
+{
+	return PassesFilter(FStringView(Item));
+}
+
+bool FBlacklistPaths::PassesStartsWithFilter(const FStringView Item) const
 {
 	if (Whitelist.Num() > 0)
 	{
@@ -195,12 +205,28 @@ bool FBlacklistPaths::PassesStartsWithFilter(const FString& Item) const
 
 bool FBlacklistPaths::PassesStartsWithFilter(const FName Item) const
 {
-	return PassesStartsWithFilter(Item.ToString());
+	TStringBuilder<FName::StringBufferSize> ItemStr;
+	Item.ToString(ItemStr);
+	return PassesStartsWithFilter(FStringView(ItemStr));
 }
 
-void FBlacklistPaths::AddBlacklistItem(const FName OwnerName, const FString& Item)
+bool FBlacklistPaths::PassesStartsWithFilter(const TCHAR* Item) const
 {
-	Blacklist.FindOrAdd(Item).AddUnique(OwnerName);
+	return PassesStartsWithFilter(FStringView(Item));
+}
+
+void FBlacklistPaths::AddBlacklistItem(const FName OwnerName, const FStringView Item)
+{
+	const uint32 ItemHash = GetTypeHash(Item);
+
+	if (FBlacklistOwners* FoundOwners = Blacklist.FindByHash(ItemHash, Item))
+	{
+		FoundOwners->AddUnique(OwnerName);
+	}
+	else
+	{
+		Blacklist.AddByHash(ItemHash, FString(Item)).Add(OwnerName);
+	}
 	
 	if (!bSuppressOnFilterChanged)
 	{
@@ -208,9 +234,30 @@ void FBlacklistPaths::AddBlacklistItem(const FName OwnerName, const FString& Ite
 	}
 }
 
-void FBlacklistPaths::AddWhitelistItem(const FName OwnerName, const FString& Item)
+void FBlacklistPaths::AddBlacklistItem(const FName OwnerName, const FName Item)
 {
-	Whitelist.FindOrAdd(Item).AddUnique(OwnerName);	
+	TStringBuilder<FName::StringBufferSize> ItemStr;
+	Item.ToString(ItemStr);
+	return AddBlacklistItem(OwnerName, FStringView(ItemStr));
+}
+
+void FBlacklistPaths::AddBlacklistItem(const FName OwnerName, const TCHAR* Item)
+{
+	return AddBlacklistItem(OwnerName, FStringView(Item));
+}
+
+void FBlacklistPaths::AddWhitelistItem(const FName OwnerName, const FStringView Item)
+{
+	const uint32 ItemHash = GetTypeHash(Item);
+
+	if (FBlacklistOwners* FoundOwners = Whitelist.FindByHash(ItemHash, Item))
+	{
+		FoundOwners->AddUnique(OwnerName);
+	}
+	else
+	{
+		Whitelist.AddByHash(ItemHash, FString(Item)).Add(OwnerName);
+	}
 	
 	if (!bSuppressOnFilterChanged)
 	{
@@ -218,6 +265,18 @@ void FBlacklistPaths::AddWhitelistItem(const FName OwnerName, const FString& Ite
 	}
 }
 
+void FBlacklistPaths::AddWhitelistItem(const FName OwnerName, const FName Item)
+{
+	TStringBuilder<FName::StringBufferSize> ItemStr;
+	Item.ToString(ItemStr);
+	return AddWhitelistItem(OwnerName, FStringView(ItemStr));
+}
+
+void FBlacklistPaths::AddWhitelistItem(const FName OwnerName, const TCHAR* Item)
+{
+	return AddWhitelistItem(OwnerName, FStringView(Item));
+}
+
 void FBlacklistPaths::AddBlacklistAll(const FName OwnerName)
 {
 	BlacklistAll.AddUnique(OwnerName);
diff --git a/Engine/Source/Runtime/Core/Public/Misc/BlacklistNames.h b/Engine/Source/Runtime/Core/Public/Misc/BlacklistNames.h
--- a/Engine/Source/Runtime/Core/Public/Misc/BlacklistNames.h
+++ b/Engine/Source/Runtime/Core/Public/Misc/BlacklistNames.h
@@ -3,6 +3,7 @@
 #pragma once
 
 #include "CoreTypes.h"
+#include "Containers/StringView.h"
 #include "Templates/SharedPointer.h"
 #include "Delegates/Delegate.h"
 
@@ -73,22 +74,40 @@ public:
 	~FBlacklistPaths() {}
 	
 	/** Returns true if passes filter restrictions using exact match */
-	bool PassesFilter(const FString& Item) const;
+	bool PassesFilter(const FStringView Item) const;
 
 	/** Returns true if passes filter restrictions using exact match */
 	bool PassesFilter(const FName Item) const;
 
+	/** Returns true if passes filter restrictions using exact match */
+	bool PassesFilter(const TCHAR* Item) const;
+
 	/** Returns true if passes filter restrictions for path */
-	bool PassesStartsWithFilter(const FString& Item) const;
+	bool PassesStartsWithFilter(const FStringView Item) const;
 
 	/** Returns true if passes filter restrictions for path */
 	bool PassesStartsWithFilter(const FName Item) const;
 
+	/** Returns true if passes filter restrictions for path */
+	bool PassesStartsWithFilter(const TCHAR* Item) const;
+
+	/** Add item to blacklist, this specific item will be filtered out */
+	void AddBlacklistItem(const FName OwnerName, const FStringView Item);
+
+	/** Add item to blacklist, this specific item will be filtered out */
+	void AddBlacklistItem(const FName OwnerName, const FName Item);
+
 	/** Add item to blacklist, this specific item will be filtered out */
-	void AddBlacklistItem(const FName OwnerName, const FString& Item);
+	void AddBlacklistItem(const FName OwnerName, const TCHAR* Item);
+
+	/** Add item to whitelist after which all items not in the whitelist will be filtered out */
+	void AddWhitelistItem(const FName OwnerName, const FStringView Item);
+
+	/** Add item to whitelist after which all items not in the whitelist will be filtered out */
+	void AddWhitelistItem(const FName OwnerName, const FName Item);
 
 	/** Add item to whitelist after which all items not in the whitelist will be filtered out */
-	void AddWhitelistItem(const FName OwnerName, const FString& Item);
+	void AddWhitelistItem(const FName OwnerName, const TCHAR* Item);
 
 	/** Set to filter out all items */
 	void AddBlacklistAll(const FName OwnerName);
