M - some older compilers complain about use of deprecation inside the depreacted class.

diff --git a/Engine/Source/Runtime/Core/Public/Math/GenericOctree.h b/Engine/Source/Runtime/Core/Public/Math/GenericOctree.h
--- a/Engine/Source/Runtime/Core/Public/Math/GenericOctree.h
+++ b/Engine/Source/Runtime/Core/Public/Math/GenericOctree.h
@@ -1034,7 +1034,7 @@ protected:
 
 /** An octree. */
 template<typename ElementType, typename OctreeSemantics>
-class UE_DEPRECATED(4.26, "The old Octree is deprecated use TOctree2.") TOctree
+class TOctree_DEPRECATED
 {
 public:
 
@@ -1046,7 +1046,7 @@ public:
 	{
 	public:
 
-		friend class TOctree;
+		friend class TOctree_DEPRECATED;
 
 		/** Initialization constructor. */
 		explicit FNode(const FNode* InParent)
@@ -1214,7 +1214,7 @@ public:
 		}
 
 		/** Starts iterating at the root of an octree. */
-		TConstIterator(const TOctree& Tree)
+		TConstIterator(const TOctree_DEPRECATED& Tree)
 			: CurrentNode(FNodeReference(&Tree.RootNode, Tree.RootNodeContext))
 		{}
 
@@ -1262,7 +1262,7 @@ public:
 		}
 
 		/** Initialization constructor. */
-		TConstElementBoxIterator(const TOctree& Tree, const FBoxCenterAndExtent& InBoundingBox)
+		TConstElementBoxIterator(const TOctree_DEPRECATED& Tree, const FBoxCenterAndExtent& InBoundingBox)
 			: IteratorBounds(InBoundingBox)
 			, NodeIt(Tree)
 			, ElementIt(Tree.RootNode.GetElementIt())
@@ -1609,7 +1609,7 @@ public:
 
 
 	/** Initialization constructor. */
-	TOctree(const FVector& InOrigin, float InExtent)
+	TOctree_DEPRECATED(const FVector& InOrigin, float InExtent)
 		: RootNode(NULL)
 		, RootNodeContext(FBoxCenterAndExtent(InOrigin, FVector(InExtent, InExtent, InExtent)), 0, 0)
 		, MinLeafExtent(InExtent* FMath::Pow((1.0f + 1.0f / (float)FOctreeNodeContext::LoosenessDenominator) / 2.0f, OctreeSemantics::MaxNodeDepth))
@@ -1618,7 +1618,7 @@ public:
 	}
 
 	/** DO NOT USE. This constructor is for internal usage only for hot-reload purposes. */
-	TOctree() : RootNode(nullptr)
+	TOctree_DEPRECATED() : RootNode(nullptr)
 	{
 		EnsureRetrievingVTablePtrDuringCtor(TEXT("TOctree()"));
 	}
@@ -1636,7 +1636,7 @@ private:
 
 	/** this function basically set TotalSizeBytes, but gives opportunity to
 	 *	include this Octree in memory stats */
-	void SetOctreeMemoryUsage(TOctree<ElementType, OctreeSemantics>* Octree, int32 NewSize)
+	void SetOctreeMemoryUsage(TOctree_DEPRECATED<ElementType, OctreeSemantics>* Octree, int32 NewSize)
 	{
 		Octree->TotalSizeBytes = NewSize;
 	}
@@ -1707,7 +1707,7 @@ private:
 					// Create the child node if it hasn't been created yet.
 					if (!Node.Children[ChildRef.Index])
 					{
-						Node.Children[ChildRef.Index] = new typename TOctree<ElementType, OctreeSemantics>::FNode(&Node);
+						Node.Children[ChildRef.Index] = new typename TOctree_DEPRECATED<ElementType, OctreeSemantics>::FNode(&Node);
 						SetOctreeMemoryUsage(this, TotalSizeBytes + sizeof(*Node.Children[ChildRef.Index]));
 					}
 
@@ -1768,4 +1768,9 @@ protected:
 	}
 };
 
+template<typename ElementType, typename OctreeSemantics>
+class UE_DEPRECATED(4.26, "The old Octree is deprecated use TOctree2.") TOctree : public TOctree_DEPRECATED<ElementType, OctreeSemantics>
+{
+};
+
 #include "GenericOctree.inl"
diff --git a/Engine/Source/Runtime/Core/Public/Math/GenericOctree.inl b/Engine/Source/Runtime/Core/Public/Math/GenericOctree.inl
--- a/Engine/Source/Runtime/Core/Public/Math/GenericOctree.inl
+++ b/Engine/Source/Runtime/Core/Public/Math/GenericOctree.inl
@@ -13,11 +13,8 @@
 class FBoxCenterAndExtent;
 class FOctreeChildNodeRef;
 class FOctreeChildNodeSubset;
-class FOctreeElementId2;
 class FOctreeNodeContext;
 struct FMath;
-template<typename ElementType,typename OctreeSemantics> class TOctree;
-template<typename T> struct TTypeTraits;
 
 
 FORCEINLINE bool FOctreeChildNodeSubset::Contains(FOctreeChildNodeRef ChildRef) const
diff --git a/Engine/Source/Runtime/Core/Public/Math/GenericOctreePublic.h b/Engine/Source/Runtime/Core/Public/Math/GenericOctreePublic.h
--- a/Engine/Source/Runtime/Core/Public/Math/GenericOctreePublic.h
+++ b/Engine/Source/Runtime/Core/Public/Math/GenericOctreePublic.h
@@ -15,9 +15,6 @@ class FOctreeElementId2
 {
 public:
 
-	template<typename,typename>
-	friend class TOctree;
-
 	template<typename, typename>
 	friend class TOctree2;
 
@@ -62,7 +59,7 @@ class FOctreeElementId
 public:
 
 	template<typename, typename>
-	friend class TOctree;
+	friend class TOctree_DEPRECATED;
 
 	/** Default constructor. */
 	FOctreeElementId()
