Fixed: TM-ShaderModels renders black when packaging QAGame on Mac for Android

diff --git a/Engine/Source/Runtime/Core/Private/HAL/ConsoleManager.cpp b/Engine/Source/Runtime/Core/Private/HAL/ConsoleManager.cpp
--- a/Engine/Source/Runtime/Core/Private/HAL/ConsoleManager.cpp
+++ b/Engine/Source/Runtime/Core/Private/HAL/ConsoleManager.cpp
@@ -2743,6 +2743,14 @@ static TAutoConsoleVariable<int32> CVarDisableOpenGLTextureStreamingSupport(
 	TEXT("  1 = Texture streaming will be disabled."),
 	ECVF_ReadOnly);
 
+// Moved here from OpenGLRHI module to make sure its always accessible on all platforms
+static FAutoConsoleVariable CVarOpenGLUseEmulatedUBs(
+	TEXT("OpenGL.UseEmulatedUBs"),
+	1,
+	TEXT("If true, enable using emulated uniform buffers on OpenGL ES3.1 mode."),
+	ECVF_ReadOnly
+	);
+
 static TAutoConsoleVariable<int32> CVarAndroidOverrideExternalTextureSupport(
 	TEXT("r.Android.OverrideExternalTextureSupport"),
 	0,
diff --git a/Engine/Source/Runtime/OpenGLDrv/Private/OpenGLCommands.cpp b/Engine/Source/Runtime/OpenGLDrv/Private/OpenGLCommands.cpp
--- a/Engine/Source/Runtime/OpenGLDrv/Private/OpenGLCommands.cpp
+++ b/Engine/Source/Runtime/OpenGLDrv/Private/OpenGLCommands.cpp
@@ -49,13 +49,6 @@ namespace OpenGLConsoleVariables
 		TEXT("If true, use glMapBuffer otherwise use glBufferSubdata.")
 		);
 
-	static FAutoConsoleVariable CVarUseEmulatedUBs(
-		TEXT("OpenGL.UseEmulatedUBs"),
-		1,
-		TEXT("If true, enable using emulated uniform buffers on ES3.1 mode."),
-		ECVF_ReadOnly
-		);
-
 	int32 bSkipCompute = 0;
 	static FAutoConsoleVariableRef CVarSkipCompute(
 		TEXT("OpenGL.SkipCompute"),
