Change default linux directory permissions to 775 from 755.

diff --git a/Engine/Source/Runtime/Core/Private/Unix/UnixPlatformFile.cpp b/Engine/Source/Runtime/Core/Private/Unix/UnixPlatformFile.cpp
--- a/Engine/Source/Runtime/Core/Private/Unix/UnixPlatformFile.cpp
+++ b/Engine/Source/Runtime/Core/Private/Unix/UnixPlatformFile.cpp
@@ -1017,7 +1017,7 @@ bool FUnixPlatformFile::DirectoryExists(const TCHAR* Directory)
 
 bool FUnixPlatformFile::CreateDirectory(const TCHAR* Directory)
 {
-	return mkdir(TCHAR_TO_UTF8(*NormalizeFilename(Directory, true)), 0755) == 0 || (errno == EEXIST);
+	return mkdir(TCHAR_TO_UTF8(*NormalizeFilename(Directory, true)), 0775) == 0 || (errno == EEXIST);
 }
 
 bool FUnixPlatformFile::DeleteDirectory(const TCHAR* Directory)
@@ -1161,7 +1161,7 @@ bool FUnixPlatformFile::CreateDirectoriesFromPath(const TCHAR* Path)
 		{
 			SubPath[i+1] = 0;
 
-			if (mkdir(SubPath, 0755) == -1)
+			if (mkdir(SubPath, 0775) == -1)
 			{
 				int ErrNo = errno;
 
