Added volatile keyword to WaitCompletionImpl to ensure compilers don't just optimize the check away.
#rnx

diff --git a/Engine/Source/Runtime/Core/Private/Windows/WindowsAsyncIO.h b/Engine/Source/Runtime/Core/Private/Windows/WindowsAsyncIO.h
--- a/Engine/Source/Runtime/Core/Private/Windows/WindowsAsyncIO.h
+++ b/Engine/Source/Runtime/Core/Private/Windows/WindowsAsyncIO.h
@@ -346,7 +346,7 @@ public:
 	{
 		// Even though SetComplete called in the constructor and sets bCompleteAndCallbackCalled=true, we still need to implement WaitComplete as
 		// the CompleteCallback can end up starting async tasks that can overtake the constructor execution and need to wait for the constructor to finish.
-		while (!bCompleteAndCallbackCalled);
+		while (!*(volatile bool*)&bCompleteAndCallbackCalled);
 	}
 
 	virtual void CancelImpl()
@@ -367,7 +367,7 @@ public:
 	{
 		// Even though SetComplete called in the constructor and sets bCompleteAndCallbackCalled=true, we still need to implement WaitComplete as
 		// the CompleteCallback can end up starting async tasks that can overtake the constructor execution and need to wait for the constructor to finish.
-		while (!bCompleteAndCallbackCalled);
+		while (!*(volatile bool*)&bCompleteAndCallbackCalled);
 	}
 
 	virtual void CancelImpl()
diff --git a/Engine/Source/Runtime/CoreUObject/Private/Serialization/BulkData2.cpp b/Engine/Source/Runtime/CoreUObject/Private/Serialization/BulkData2.cpp
--- a/Engine/Source/Runtime/CoreUObject/Private/Serialization/BulkData2.cpp
+++ b/Engine/Source/Runtime/CoreUObject/Private/Serialization/BulkData2.cpp
@@ -255,7 +255,7 @@ private:
 	{
 		// Even though SetComplete called in the constructor and sets bCompleteAndCallbackCalled=true, we still need to implement WaitComplete as
 		// the CompleteCallback can end up starting async tasks that can overtake the constructor execution and need to wait for the constructor to finish.
-		while (!bCompleteAndCallbackCalled);
+		while (!*(volatile bool*)&bCompleteAndCallbackCalled);
 	}
 
 	virtual void CancelImpl() override
diff --git a/Engine/Source/Runtime/PakFile/Private/IPlatformFilePak.cpp b/Engine/Source/Runtime/PakFile/Private/IPlatformFilePak.cpp
--- a/Engine/Source/Runtime/PakFile/Private/IPlatformFilePak.cpp
+++ b/Engine/Source/Runtime/PakFile/Private/IPlatformFilePak.cpp
@@ -391,7 +391,7 @@ public:
 	{
 		// Even though SetComplete called in the constructor and sets bCompleteAndCallbackCalled=true, we still need to implement WaitComplete as
 		// the CompleteCallback can end up starting async tasks that can overtake the constructor execution and need to wait for the constructor to finish.
-		while (!bCompleteAndCallbackCalled);
+		while (!*(volatile bool*)&bCompleteAndCallbackCalled);
 	}
 
 	virtual void CancelImpl()
