Fix to avoid parsing the same Restricted plugin folders multiple times.

diff --git a/Engine/Source/Runtime/Core/Private/Misc/Paths.cpp b/Engine/Source/Runtime/Core/Private/Misc/Paths.cpp
--- a/Engine/Source/Runtime/Core/Private/Misc/Paths.cpp
+++ b/Engine/Source/Runtime/Core/Private/Misc/Paths.cpp
@@ -272,9 +272,6 @@ static void GetExtensionDirsInternal(TArray<FString>& ExtensionDirs, const FStri
 	{
 		if (bIsDirectory)
 		{
-			AddIfDirectoryExists(ExtensionDirs, FPaths::Combine(FilenameOrDirectory, SubDir));
-
-			// now look for platforms under it
 			GetExtensionDirsInternal(ExtensionDirs, FilenameOrDirectory, SubDir);
 		}
 		return true;
