A hopeful incremental CIS fix (UE-90852).

diff --git a/Engine/Source/Runtime/Core/Private/Serialization/MemoryImage.cpp b/Engine/Source/Runtime/Core/Private/Serialization/MemoryImage.cpp
--- a/Engine/Source/Runtime/Core/Private/Serialization/MemoryImage.cpp
+++ b/Engine/Source/Runtime/Core/Private/Serialization/MemoryImage.cpp
@@ -9,6 +9,7 @@
 #include "Hash/CityHash.h"
 #include "ProfilingDebugging/LoadTimeTracker.h"
 #include "Misc/DataDrivenPlatformInfoRegistry.h"
+#include "Serialization/Archive.h"
 
 IMPLEMENT_TYPE_LAYOUT(FMemoryImageString);
 
@@ -79,6 +80,19 @@ void FPlatformTypeLayoutParameters::AppendKeyString(FString& KeyString) const
 	}
 }
 
+void FPlatformTypeLayoutParameters::Serialize(FArchive& Ar)
+{
+	// if you change this code, please bump MATERIALSHADERMAP_DERIVEDDATA_VER (see FMaterialShaderMap::Serialize)
+	// since this is a part of ShaderMapId
+	Ar << MaxFieldAlignment;
+	Ar << b32Bit;
+	Ar << bForce64BitMemoryImagePointers;
+	Ar << bAlignBases;
+	Ar << bWithEditorOnly;
+	Ar << bWithRayTracing;
+	Ar << bIsCurrentPlatform;
+}
+
 // evaluated during static-initialization, so logging from regular check() macros won't work correctly
 static void InitializeSizeFromFields(FTypeLayoutDesc& TypeLayout, const FPlatformTypeLayoutParameters& PlatformLayoutParams)
 {
diff --git a/Engine/Source/Runtime/Core/Public/Serialization/MemoryLayout.h b/Engine/Source/Runtime/Core/Public/Serialization/MemoryLayout.h
--- a/Engine/Source/Runtime/Core/Public/Serialization/MemoryLayout.h
+++ b/Engine/Source/Runtime/Core/Public/Serialization/MemoryLayout.h
@@ -98,18 +98,7 @@ struct FPlatformTypeLayoutParameters
 	/**
 	 * This is used for serializing into/from the DDC
 	 */
-	void Serialize(FArchive& Ar)
-	{
-		// if you change this code, please bump MATERIALSHADERMAP_DERIVEDDATA_VER (see FMaterialShaderMap::Serialize)
-		// since this is a part of ShaderMapId
-		Ar << MaxFieldAlignment;
-		Ar << b32Bit;
-		Ar << bForce64BitMemoryImagePointers;
-		Ar << bAlignBases;
-		Ar << bWithEditorOnly;
-		Ar << bWithRayTracing;
-		Ar << bIsCurrentPlatform;
-	}
+	void Serialize(FArchive& Ar);
 
 	bool operator==(const FPlatformTypeLayoutParameters& ReferenceSet) const
 	{
