Enable CSVProfiler processing thread on multithread forked servers

diff --git a/Engine/Source/Runtime/Core/Private/ProfilingDebugging/CsvProfiler.cpp b/Engine/Source/Runtime/Core/Private/ProfilingDebugging/CsvProfiler.cpp
--- a/Engine/Source/Runtime/Core/Private/ProfilingDebugging/CsvProfiler.cpp
+++ b/Engine/Source/Runtime/Core/Private/ProfilingDebugging/CsvProfiler.cpp
@@ -26,6 +26,7 @@
 #include "Stats/Stats.h"
 #include "HAL/LowLevelMemTracker.h"
 #include "Misc/Compression.h"
+#include "Misc/Fork.h"
 
 #if CSV_PROFILER
 
@@ -2188,9 +2189,9 @@ public:
 		: CsvProfiler(InCsvProfiler)
 	{
 #if CSV_THREAD_HIGH_PRI
-		Thread = FRunnableThread::Create(this, TEXT("CSVProfiler"), 0, TPri_Highest, FPlatformAffinity::GetTaskGraphThreadMask());
+		Thread = FForkProcessHelper::CreateForkableThread(this, TEXT("CSVProfiler"), 0, TPri_Highest, FPlatformAffinity::GetTaskGraphThreadMask());
 #else
-		Thread = FRunnableThread::Create(this, TEXT("CSVProfiler"), 0, TPri_Lowest, FPlatformAffinity::GetTaskGraphBackgroundTaskMask());
+		Thread = FForkProcessHelper::CreateForkableThread(this, TEXT("CSVProfiler"), 0, TPri_Lowest, FPlatformAffinity::GetTaskGraphBackgroundTaskMask());
 #endif
 	}
 
@@ -2204,6 +2205,11 @@ public:
 		}
 	}
 
+	bool IsValid() const
+	{
+		return Thread != nullptr;
+	}
+
 	// FRunnable interface
 	virtual bool Init() override
 	{
@@ -2602,6 +2608,13 @@ void FCsvProfiler::BeginFrame()
 					{
 						// Lazily create the CSV processing thread
 						ProcessingThread = new FCsvProfilerProcessingThread(*this);
+						if (ProcessingThread->IsValid() == false)
+						{
+							UE_LOG(LogCsvProfiler, Error, TEXT("CSV Processing Thread could not be created due to being in a single-thread environment "));
+							delete ProcessingThread;
+							ProcessingThread = nullptr;
+							GCsvUseProcessingThread = false;
+						}
 					}
 
 					// Figure out the target framerate
@@ -3231,7 +3244,7 @@ void FCsvProfiler::Init()
 #endif // CSV_PROFILER_ALLOW_DEBUG_FEATURES
 
 	// Always disable the CSV profiling thread if the platform does not support threading.
-	if (!FPlatformProcess::SupportsMultithreading())
+	if (!FPlatformProcess::SupportsMultithreading() && FForkProcessHelper::SupportsMultithreadingPostFork() == false)
 	{
 		GCsvUseProcessingThread = false;
 	}
