Update FNiagaraShaderMapId to hold FPlatformTypeLayoutParameters, should fix problems when cooking niagara shaders for mulitple platforms at the same time

diff --git a/Engine/Plugins/FX/Niagara/Source/Niagara/Classes/NiagaraScript.h b/Engine/Plugins/FX/Niagara/Source/Niagara/Classes/NiagaraScript.h
--- a/Engine/Plugins/FX/Niagara/Source/Niagara/Classes/NiagaraScript.h
+++ b/Engine/Plugins/FX/Niagara/Source/Niagara/Classes/NiagaraScript.h
@@ -580,12 +580,12 @@ public:
 
 	// Infrastructure for GPU compute Shaders
 #if WITH_EDITOR
-	NIAGARA_API void CacheResourceShadersForCooking(EShaderPlatform ShaderPlatform, TArray<FNiagaraShaderScript*>& InOutCachedResources);
+	NIAGARA_API void CacheResourceShadersForCooking(EShaderPlatform ShaderPlatform, TArray<FNiagaraShaderScript*>& InOutCachedResources, const ITargetPlatform* TargetPlatform = nullptr);
 
 	NIAGARA_API void CacheResourceShadersForRendering(bool bRegenerateId, bool bForceRecompile=false);
 	void BeginCacheForCookedPlatformData(const ITargetPlatform *TargetPlatform);
 	virtual bool IsCachedCookedPlatformDataLoaded(const ITargetPlatform* TargetPlatform) override;
-	void CacheShadersForResources(FNiagaraShaderScript* ResourceToCache, bool bApplyCompletedShaderMapForRendering, bool bForceRecompile = false, bool bCooking=false);
+	void CacheShadersForResources(FNiagaraShaderScript* ResourceToCache, bool bApplyCompletedShaderMapForRendering, bool bForceRecompile = false, bool bCooking=false, const ITargetPlatform* TargetPlatform = nullptr);
 #endif // WITH_EDITOR
 	FNiagaraShaderScript* AllocateResource();
 	FNiagaraShaderScript* GetRenderThreadScript()
diff --git a/Engine/Plugins/FX/Niagara/Source/Niagara/Private/NiagaraScript.cpp b/Engine/Plugins/FX/Niagara/Source/Niagara/Private/NiagaraScript.cpp
--- a/Engine/Plugins/FX/Niagara/Source/Niagara/Private/NiagaraScript.cpp
+++ b/Engine/Plugins/FX/Niagara/Source/Niagara/Private/NiagaraScript.cpp
@@ -1783,7 +1783,7 @@ void UNiagaraScript::BeginCacheForCookedPlatformData(const ITargetPlatform *Targ
 			const EShaderPlatform LegacyShaderPlatform = ShaderFormatToLegacyShaderPlatform(DesiredShaderFormats[FormatIndex]);
 			if (FNiagaraUtilities::SupportsGPUParticles(LegacyShaderPlatform))
 			{
-				CacheResourceShadersForCooking(LegacyShaderPlatform, CachedScriptResourcesForPlatform);
+				CacheResourceShadersForCooking(LegacyShaderPlatform, CachedScriptResourcesForPlatform, TargetPlatform);
 			}
 		}
 	}
@@ -1831,7 +1831,7 @@ bool UNiagaraScript::IsCachedCookedPlatformDataLoaded(const ITargetPlatform* Tar
 	return true;
 }
 
-void UNiagaraScript::CacheResourceShadersForCooking(EShaderPlatform ShaderPlatform, TArray<FNiagaraShaderScript*>& InOutCachedResources)
+void UNiagaraScript::CacheResourceShadersForCooking(EShaderPlatform ShaderPlatform, TArray<FNiagaraShaderScript*>& InOutCachedResources, const ITargetPlatform* TargetPlatform)
 {
 	if (CanBeRunOnGpu())
 	{
@@ -1862,7 +1862,7 @@ void UNiagaraScript::CacheResourceShadersForCooking(EShaderPlatform ShaderPlatfo
 
 			check(ResourceToCache);
 
-			CacheShadersForResources(ResourceToCache, false, false, true);
+			CacheShadersForResources(ResourceToCache, false, false, true, TargetPlatform);
 
 			INiagaraModule NiagaraModule = FModuleManager::GetModuleChecked<INiagaraModule>(TEXT("Niagara"));
 			NiagaraModule.ProcessShaderCompilationQueue();
@@ -1874,12 +1874,13 @@ void UNiagaraScript::CacheResourceShadersForCooking(EShaderPlatform ShaderPlatfo
 
 
 
-void UNiagaraScript::CacheShadersForResources(FNiagaraShaderScript* ResourceToCache, bool bApplyCompletedShaderMapForRendering, bool bForceRecompile, bool bCooking)
+void UNiagaraScript::CacheShadersForResources(FNiagaraShaderScript* ResourceToCache, bool bApplyCompletedShaderMapForRendering, bool bForceRecompile, bool bCooking, const ITargetPlatform* TargetPlatform)
 {
 	if (CanBeRunOnGpu())
 	{
 		// When not running in the editor, the shaders are created in-sync (in the postload) to avoid update issues.
-		const bool bSuccess = ResourceToCache->CacheShaders(bApplyCompletedShaderMapForRendering, bForceRecompile, bCooking || !GIsEditor || GIsAutomationTesting);
+		const bool bSync = bCooking || !GIsEditor || GIsAutomationTesting;
+		const bool bSuccess = ResourceToCache->CacheShaders(bApplyCompletedShaderMapForRendering, bForceRecompile, bSync, TargetPlatform);
 
 #if defined(NIAGARA_SCRIPT_COMPILE_LOGGING_MEDIUM)
 		if (!bSuccess)
diff --git a/Engine/Plugins/FX/Niagara/Source/NiagaraShader/Private/NiagaraShader.cpp b/Engine/Plugins/FX/Niagara/Source/NiagaraShader/Private/NiagaraShader.cpp
--- a/Engine/Plugins/FX/Niagara/Source/NiagaraShader/Private/NiagaraShader.cpp
+++ b/Engine/Plugins/FX/Niagara/Source/NiagaraShader/Private/NiagaraShader.cpp
@@ -15,6 +15,8 @@
 	#include "Interfaces/ITargetPlatformManagerModule.h"
 	#include "TickableEditorObject.h"
 	#include "DerivedDataCacheInterface.h"
+	#include "Interfaces/ITargetPlatformManagerModule.h"
+	#include "Interfaces/ITargetPlatform.h"
 #endif
 #include "ProfilingDebugging/CookStats.h"
 
@@ -240,17 +242,22 @@ bool FNiagaraShaderMapId::operator==(const FNiagaraShaderMapId& ReferenceSet) co
 		|| */BaseCompileHash != ReferenceSet.BaseCompileHash
 		|| FeatureLevel != ReferenceSet.FeatureLevel
 		|| CompilerVersionID != ReferenceSet.CompilerVersionID 
-		|| bUsesRapidIterationParams != ReferenceSet.bUsesRapidIterationParams)
+		|| bUsesRapidIterationParams != ReferenceSet.bUsesRapidIterationParams
+		|| LayoutParams != ReferenceSet.LayoutParams)
 	{
 		return false;
 	}
-	
-	
+
 	if (AdditionalDefines.Num() != ReferenceSet.AdditionalDefines.Num())
 	{
 		return false;
 	}
 
+	if (ReferencedCompileHashes.Num() != ReferenceSet.ReferencedCompileHashes.Num())
+	{
+		return false;
+	}
+
 	for (int32 Idx = 0; Idx < ReferenceSet.AdditionalDefines.Num(); Idx++)
 	{
 		const FString& ReferenceStr = ReferenceSet.AdditionalDefines[Idx];
@@ -261,11 +268,6 @@ bool FNiagaraShaderMapId::operator==(const FNiagaraShaderMapId& ReferenceSet) co
 		}
 	}
 
-	if (ReferencedCompileHashes.Num() != ReferenceSet.ReferencedCompileHashes.Num())
-	{
-		return false;
-	}
-
 	for (int32 i = 0; i < ReferencedCompileHashes.Num(); i++)
 	{
 		if (ReferencedCompileHashes[i] != ReferenceSet.ReferencedCompileHashes[i])
@@ -304,8 +306,34 @@ void FNiagaraShaderMapId::AppendKeyString(FString& KeyString) const
 
 	FString FeatureLevelString;
 	GetFeatureLevelName(FeatureLevel, FeatureLevelString);
-	KeyString += FeatureLevelString + TEXT("_");
 
+	{
+		const FSHAHash LayoutHash = Freeze::HashLayout(StaticGetTypeLayoutDesc<FNiagaraShaderMapContent>(), LayoutParams);
+		KeyString += TEXT("_");
+		KeyString += LayoutHash.ToString();
+		KeyString += TEXT("_");
+	}
+
+	{
+		const FSHAHash LayoutHash = Freeze::HashLayout(StaticGetTypeLayoutDesc<FNiagaraShader>(), LayoutParams);
+		KeyString += TEXT("_");
+		KeyString += LayoutHash.ToString();
+		KeyString += TEXT("_");
+	}
+
+	/*for (const FHashedName& DITypeName : DITypeNames)
+	{
+		const FTypeLayoutDesc* DIType = FTypeLayoutDesc::Find(DITypeName.GetHash());
+		if (DIType)
+		{
+			const FSHAHash LayoutHash = Freeze::HashLayout(*DIType, LayoutParams);
+			KeyString += TEXT("_");
+			KeyString += LayoutHash.ToString();
+			KeyString += TEXT("_");
+		}
+	}*/
+
+	KeyString += FeatureLevelString + TEXT("_");
 	KeyString += CompilerVersionID.ToString();
 	KeyString += TEXT("_");
 
diff --git a/Engine/Plugins/FX/Niagara/Source/NiagaraShader/Private/NiagaraShared.cpp b/Engine/Plugins/FX/Niagara/Source/NiagaraShader/Private/NiagaraShared.cpp
--- a/Engine/Plugins/FX/Niagara/Source/NiagaraShader/Private/NiagaraShared.cpp
+++ b/Engine/Plugins/FX/Niagara/Source/NiagaraShader/Private/NiagaraShared.cpp
@@ -19,6 +19,9 @@
 #include "RendererInterface.h"
 #include "Modules/ModuleManager.h"
 #include "NiagaraCustomVersion.h"
+#if WITH_EDITOR
+#include "Interfaces/ITargetPlatform.h"
+#endif
 
 IMPLEMENT_TYPE_LAYOUT(FNiagaraDataInterfaceParamRef);
 IMPLEMENT_TYPE_LAYOUT(FNiagaraShaderMapContent);
@@ -120,6 +123,7 @@ bool FNiagaraShaderScript::IsSame(const FNiagaraShaderMapId& InId) const
 	{
 		return false;
 	}
+
 	for (int32 i = 0; i < ReferencedCompileHashes.Num(); ++i)
 	{
 		if (ReferencedCompileHashes[i] != InId.ReferencedCompileHashes[i])
@@ -159,7 +163,7 @@ void FNiagaraShaderScript::GetDependentShaderTypes(EShaderPlatform Platform, TAr
 
 
 
-NIAGARASHADER_API void FNiagaraShaderScript::GetShaderMapId(EShaderPlatform Platform, FNiagaraShaderMapId& OutId) const
+NIAGARASHADER_API void FNiagaraShaderScript::GetShaderMapId(EShaderPlatform Platform, const ITargetPlatform* TargetPlatform, FNiagaraShaderMapId& OutId) const
 {
 	if (bLoadedCookedShaderMapId)
 	{
@@ -167,6 +171,8 @@ NIAGARASHADER_API void FNiagaraShaderScript::GetShaderMapId(EShaderPlatform Plat
 	}
 	else
 	{
+		INiagaraShaderModule* Module = INiagaraShaderModule::Get();
+
 		TArray<FShaderType*> ShaderTypes;
 		GetDependentShaderTypes(Platform, ShaderTypes);
 		OutId.FeatureLevel = GetFeatureLevel();/*
@@ -186,6 +192,19 @@ NIAGARASHADER_API void FNiagaraShaderScript::GetShaderMapId(EShaderPlatform Plat
 		{
 			OutId.AdditionalDefines.Emplace(Define);
 		}
+
+		if (TargetPlatform)
+		{
+#if WITH_EDITOR
+			OutId.LayoutParams.InitializeForPlatform(TargetPlatform->IniPlatformName(), TargetPlatform->HasEditorOnlyData());
+#else
+			UE_LOG(LogShaders, Error, TEXT("FNiagaraShaderScript::GetShaderMapId: TargetPlatform is not null, but a cooked executable cannot target platforms other than its own."));
+#endif
+		}
+		else
+		{
+			OutId.LayoutParams.InitializeForCurrent();
+		}
 	}
 }
 
@@ -337,10 +356,10 @@ NIAGARASHADER_API  bool FNiagaraShaderScript::IsCompilationFinished() const
 */
 #if WITH_EDITOR
 
-bool FNiagaraShaderScript::CacheShaders(bool bApplyCompletedShaderMapForRendering, bool bForceRecompile, bool bSynchronous)
+bool FNiagaraShaderScript::CacheShaders(bool bApplyCompletedShaderMapForRendering, bool bForceRecompile, bool bSynchronous, const ITargetPlatform* TargetPlatform)
 {
 	FNiagaraShaderMapId NoStaticParametersId;
-	GetShaderMapId(ShaderPlatform, NoStaticParametersId);
+	GetShaderMapId(ShaderPlatform, TargetPlatform, NoStaticParametersId);
 	return CacheShaders(NoStaticParametersId, bApplyCompletedShaderMapForRendering, bForceRecompile, bSynchronous);
 }
 
diff --git a/Engine/Plugins/FX/Niagara/Source/NiagaraShader/Public/NiagaraShared.h b/Engine/Plugins/FX/Niagara/Source/NiagaraShader/Public/NiagaraShared.h
--- a/Engine/Plugins/FX/Niagara/Source/NiagaraShader/Public/NiagaraShared.h
+++ b/Engine/Plugins/FX/Niagara/Source/NiagaraShader/Public/NiagaraShared.h
@@ -228,6 +228,11 @@ public:
 	/** The compile hashes of the top level scripts the script is dependent on. */
 	LAYOUT_FIELD(TMemoryImageArray<FSHAHash>, ReferencedCompileHashes);
 
+	/*
+	 * Type layout parameters of the memory image
+	 */
+	LAYOUT_FIELD(FPlatformTypeLayoutParameters, LayoutParams);
+
 	/** Whether or not we need to bake Rapid Iteration params. True to keep params, false to bake.*/
 	LAYOUT_FIELD_INITIALIZED(bool, bUsesRapidIterationParams, true);
 
@@ -584,7 +589,7 @@ public:
 	 * Caches the shaders for this script with no static parameters on the given platform.
 	 * This is used by FNiagaraShaderScript
 	 */
-	NIAGARASHADER_API  bool CacheShaders(bool bApplyCompletedShaderMapForRendering, bool bForceRecompile, bool bSynchronous = false);
+	NIAGARASHADER_API  bool CacheShaders(bool bApplyCompletedShaderMapForRendering, bool bForceRecompile, bool bSynchronous = false, const ITargetPlatform* TargetPlatform = nullptr);
 	bool CacheShaders(const FNiagaraShaderMapId& ShaderMapId, bool bApplyCompletedShaderMapForRendering, bool bForceRecompile, bool bSynchronous = false);
 
 
@@ -609,7 +614,7 @@ public:
 	NIAGARASHADER_API void ReleaseShaderMap();
 
 	void GetDependentShaderTypes(EShaderPlatform Platform, TArray<FShaderType*>& OutShaderTypes) const;
-	NIAGARASHADER_API  virtual void GetShaderMapId(EShaderPlatform Platform, FNiagaraShaderMapId& OutId) const;
+	NIAGARASHADER_API  virtual void GetShaderMapId(EShaderPlatform Platform, const ITargetPlatform* TargetPlatform, FNiagaraShaderMapId& OutId) const;
 
 	NIAGARASHADER_API void Invalidate();
 
diff --git a/Engine/Shaders/Public/ShaderVersion.ush b/Engine/Shaders/Public/ShaderVersion.ush
--- a/Engine/Shaders/Public/ShaderVersion.ush
+++ b/Engine/Shaders/Public/ShaderVersion.ush
@@ -3,4 +3,4 @@
 // in Platform.ush (which should be included in any shader) it allows to invalidate the shader DDC.
 // 
 // If you are merging streams and there is a conflict with this GUID you should make a new GUID rather than taking one or the other.
-// GUID = 984E1C09-96B0-441E-9ADB-82971B2A2767
\ No newline at end of file
+// GUID = 9D884D03-FB42-4CD9-81B0-B81CDF22CEF9
\ No newline at end of file
diff --git a/Engine/Source/Runtime/Core/Private/Serialization/MemoryImage.cpp b/Engine/Source/Runtime/Core/Private/Serialization/MemoryImage.cpp
--- a/Engine/Source/Runtime/Core/Private/Serialization/MemoryImage.cpp
+++ b/Engine/Source/Runtime/Core/Private/Serialization/MemoryImage.cpp
@@ -12,11 +12,19 @@
 #include "Serialization/Archive.h"
 
 IMPLEMENT_TYPE_LAYOUT(FMemoryImageString);
+IMPLEMENT_TYPE_LAYOUT(FPlatformTypeLayoutParameters);
 
 static const uint32 NumTypeLayoutDescHashBuckets = 4357u;
 static const FTypeLayoutDesc* GTypeLayoutHashBuckets[NumTypeLayoutDescHashBuckets] = { nullptr };
 static uint32 GNumTypeLayoutsRegistered = 0u;
 
+bool FPlatformTypeLayoutParameters::IsCurrentPlatform() const
+{
+	FPlatformTypeLayoutParameters CurrentPlatform;
+	CurrentPlatform.InitializeForCurrent();
+	return *this == CurrentPlatform;
+}
+
 void FPlatformTypeLayoutParameters::InitializeForArchive(FArchive& Ar)
 {
 	check(Ar.IsSaving());
@@ -34,23 +42,26 @@ void FPlatformTypeLayoutParameters::InitializeForPlatform(const FString& Platfor
 {
 	const FDataDrivenPlatformInfoRegistry::FPlatformInfo& PlatformInfo = FDataDrivenPlatformInfoRegistry::GetPlatformInfo(PlatformName);
 
-	bWithEditorOnly = bHasEditorOnlyData;
-	bWithRayTracing = PlatformInfo.Freezing_bWithRayTracing;
-	b32Bit = PlatformInfo.Freezing_b32Bit;
-	bForce64BitMemoryImagePointers = PlatformInfo.Freezing_bForce64BitMemoryImagePointers;
-	bAlignBases = PlatformInfo.Freezing_bAlignBases;
+	Flags = Flag_Initialized;
+	if (bHasEditorOnlyData) Flags |= Flag_WithEditorOnly;
+	if (PlatformInfo.Freezing_bWithRayTracing) Flags |= Flag_WithRaytracing;
+	if (PlatformInfo.Freezing_b32Bit) Flags |= Flag_Is32Bit;
+	if (PlatformInfo.Freezing_bForce64BitMemoryImagePointers) Flags |= Flag_Force64BitMemoryImagePointers;
+	if (PlatformInfo.Freezing_bAlignBases) Flags |= Flag_AlignBases;
+
 	MaxFieldAlignment = PlatformInfo.Freezing_MaxFieldAlignment;
 }
 
 void FPlatformTypeLayoutParameters::InitializeForCurrent()
 {
-	bWithEditorOnly = WITH_EDITORONLY_DATA;
-	bWithRayTracing = WITH_RAYTRACING;
-	b32Bit = PLATFORM_32BITS;
-	bForce64BitMemoryImagePointers = UE_FORCE_64BIT_MEMORY_IMAGE_POINTERS;
+	Flags = Flag_Initialized;
+	if (WITH_EDITORONLY_DATA) Flags |= Flag_WithEditorOnly;
+	if (WITH_RAYTRACING) Flags |= Flag_WithRaytracing;
+	if (PLATFORM_32BITS) Flags |= Flag_Is32Bit;
+	if (UE_FORCE_64BIT_MEMORY_IMAGE_POINTERS) Flags |= Flag_Force64BitMemoryImagePointers;
+
 	check(GetRawPointerSize() == sizeof(void*));
 	check(GetMemoryImagePointerSize() == sizeof(FMemoryImagePtrInt));
-	bIsCurrentPlatform = true;
 
 #if defined(__clang__)
 	InitializeForClang();
@@ -61,38 +72,32 @@ void FPlatformTypeLayoutParameters::InitializeForCurrent()
 
 void FPlatformTypeLayoutParameters::InitializeForMSVC()
 {
-	bAlignBases = true;
+	Flags |= Flag_AlignBases;
 
 	// This corresponds to the value used by /Zp#
-	MaxFieldAlignment = b32Bit ? 4u : 8u;
+	MaxFieldAlignment = Is32Bit() ? 4u : 8u;
 }
 
 void FPlatformTypeLayoutParameters::InitializeForClang()
 {
-	bAlignBases = false;
+	// nothing
+}
+
+FArchive& FPlatformTypeLayoutParameters::Serialize(FArchive& Ar)
+{
+	// if you change this code, please bump MATERIALSHADERMAP_DERIVEDDATA_VER (see FMaterialShaderMap::Serialize)
+	// since this is a part of ShaderMapId
+	return Ar << MaxFieldAlignment << Flags;
 }
 
 void FPlatformTypeLayoutParameters::AppendKeyString(FString& KeyString) const
 {
-	if (b32Bit && bForce64BitMemoryImagePointers)
+	if (Is32Bit() && HasForce64BitMemoryImagePointers())
 	{
 		KeyString += TEXT("FIX_");
 	}
 }
 
-void FPlatformTypeLayoutParameters::Serialize(FArchive& Ar)
-{
-	// if you change this code, please bump MATERIALSHADERMAP_DERIVEDDATA_VER (see FMaterialShaderMap::Serialize)
-	// since this is a part of ShaderMapId
-	Ar << MaxFieldAlignment;
-	Ar << b32Bit;
-	Ar << bForce64BitMemoryImagePointers;
-	Ar << bAlignBases;
-	Ar << bWithEditorOnly;
-	Ar << bWithRayTracing;
-	Ar << bIsCurrentPlatform;
-}
-
 // evaluated during static-initialization, so logging from regular check() macros won't work correctly
 static void InitializeSizeFromFields(FTypeLayoutDesc& TypeLayout, const FPlatformTypeLayoutParameters& PlatformLayoutParams)
 {
@@ -163,7 +168,7 @@ static void InitializeSizeFromFields(FTypeLayoutDesc& TypeLayout, const FPlatfor
 
 			if (PaddedFieldSize > 0u)
 			{
-				if (!bIsBase || PlatformLayoutParams.bAlignBases)
+				if (!bIsBase || PlatformLayoutParams.HasAlignBases())
 				{
 					const uint32 FieldSize = Align(PaddedFieldSize, FieldTypeAlignment);
 					check(FieldSize == FieldType.Size);
@@ -321,11 +326,11 @@ bool Freeze::IncludeField(const FFieldLayoutDesc* FieldDesc, const FPlatformType
 	const bool bIsEditorOnly = (FieldDesc->Flags & EFieldLayoutFlags::WithEditorOnly) != 0u;
 	const bool bIsRayTracing = (FieldDesc->Flags & EFieldLayoutFlags::WithRayTracing) != 0u;
 
-	if (bIsEditorOnly && !LayoutParams.bWithEditorOnly)
+	if (bIsEditorOnly && !LayoutParams.WithEditorOnly())
 	{
 		return false;
 	}
-	if (bIsRayTracing && !LayoutParams.bWithRayTracing)
+	if (bIsRayTracing && !LayoutParams.WithRaytracing())
 	{
 		return false;
 	}
@@ -440,7 +445,7 @@ void Freeze::DefaultWriteMemoryImage(FMemoryImageWriter& Writer, const void* Obj
 								Object,
 								FieldObject + ArrayIndex * FieldType.Size,
 								FieldType, bIsBase ? DerivedTypeDesc : FieldType);
-							if (!bIsBase || TargetLayoutParams.bAlignBases)
+							if (!bIsBase || TargetLayoutParams.HasAlignBases())
 							{
 								// Align the field size
 								const uint32 FieldSize = Writer.GetOffset() - FieldOffset;
@@ -548,10 +553,10 @@ uint32 Freeze::DefaultAppendHash(const FTypeLayoutDesc& TypeLayout, const FPlatf
 				const bool bIsBase = (FieldIndex < TypeLayout.NumBases);
 				const uint32 FieldTypeAlignment = GetTargetAlignment(FieldType, LayoutParams);
 				const uint32 FieldAlignment = FMath::Min(FieldTypeAlignment, LayoutParams.MaxFieldAlignment);
-				check(!LayoutParams.bIsCurrentPlatform || FieldTypeAlignment == FieldType.Alignment);
+				check(!LayoutParams.IsCurrentPlatform() || FieldTypeAlignment == FieldType.Alignment);
 
 				Offset = Align(Offset, FieldAlignment);
-				check(!LayoutParams.bIsCurrentPlatform || CheckOffsetMatch(Offset, Field->Offset, FieldType, bIsBase));
+				check(!LayoutParams.IsCurrentPlatform() || CheckOffsetMatch(Offset, Field->Offset, FieldType, bIsBase));
 
 				Hasher.Update((uint8*)&Offset, sizeof(Offset));
 				Hasher.Update((uint8*)&Field->NumArray, sizeof(Field->NumArray));
@@ -580,7 +585,7 @@ uint32 Freeze::DefaultAppendHash(const FTypeLayoutDesc& TypeLayout, const FPlatf
 
 				if (PaddedFieldSize > 0u)
 				{
-					if (!bIsBase || LayoutParams.bAlignBases)
+					if (!bIsBase || LayoutParams.HasAlignBases())
 					{
 						const uint32 FieldSize = Align(PaddedFieldSize, FieldTypeAlignment);
 						Offset += FieldSize * Field->NumArray;
@@ -621,7 +626,7 @@ uint32 Freeze::DefaultAppendHash(const FTypeLayoutDesc& TypeLayout, const FPlatf
 		Field = Field->Next;
 	}
 
-	check(!LayoutParams.bIsCurrentPlatform || Offset == TypeLayout.SizeFromFields);
+	check(!LayoutParams.IsCurrentPlatform() || Offset == TypeLayout.SizeFromFields);
 	return Offset;
 }
 
@@ -701,6 +706,16 @@ void Freeze::DefaultToString(const void* Object, const FTypeLayoutDesc& TypeDesc
 	--OutContext.Indent;
 }
 
+uint32 Freeze::IntrinsicAppendHash(void* const* DummyObject, const FTypeLayoutDesc& TypeDesc, const FPlatformTypeLayoutParameters& LayoutParams, FSHA1& Hasher)
+{
+	return AppendHashForNameAndSize(TypeDesc.Name, LayoutParams.GetRawPointerSize(), Hasher);
+}
+
+uint32 Freeze::IntrinsicGetTargetAlignment(void* const* DummyObject, const FTypeLayoutDesc& TypeDesc, const FPlatformTypeLayoutParameters& LayoutParams)
+{
+	return LayoutParams.GetRawPointerSize();
+}
+
 void Freeze::IntrinsicToString(char Object, const FTypeLayoutDesc& TypeDesc, const FPlatformTypeLayoutParameters& LayoutParams, FMemoryToStringContext& OutContext)
 {
 	OutContext.String->Appendf(TEXT("%d\n"), Object);
@@ -1123,7 +1138,7 @@ FMemoryImageSection* FMemoryImageSection::WritePointer(const FString& SectionNam
 
 uint32 FMemoryImageSection::WriteRawPointerSizedBytes(uint64 PointerValue)
 {
-	if (ParentImage->TargetLayoutParameters.b32Bit)
+	if (ParentImage->TargetLayoutParameters.Is32Bit())
 	{
 		return WriteBytes((uint32)PointerValue);
 	}
@@ -1162,7 +1177,7 @@ uint32 FMemoryImageSection::WriteFName(const FName& Name)
 {
 	const FPlatformTypeLayoutParameters& TargetLayoutParameters = ParentImage->TargetLayoutParameters;
 	uint32 Offset = 0u;
-	if (TargetLayoutParameters.bWithEditorOnly)
+	if (TargetLayoutParameters.WithEditorOnly())
 	{
 		Offset = WriteBytes(FScriptName());
 	}
diff --git a/Engine/Source/Runtime/Core/Public/Containers/BitArray.h b/Engine/Source/Runtime/Core/Public/Containers/BitArray.h
--- a/Engine/Source/Runtime/Core/Public/Containers/BitArray.h
+++ b/Engine/Source/Runtime/Core/Public/Containers/BitArray.h
@@ -1034,7 +1034,7 @@ public:
 	void WriteMemoryImage(FMemoryImageWriter& Writer) const
 	{
 		static const bool bSupportsFreezeMemoryImage = TAllocatorTraits<Allocator>::SupportsFreezeMemoryImage;
-		checkf(!Writer.GetTargetLayoutParams().b32Bit, TEXT("TBitArray does not currently support freezing for 32bits"));
+		checkf(!Writer.Is32BitTarget(), TEXT("TBitArray does not currently support freezing for 32bits"));
 		TSupportsFreezeMemoryImageHelper<bSupportsFreezeMemoryImage>::WriteMemoryImage(Writer, *this);
 	}
 };
diff --git a/Engine/Source/Runtime/Core/Public/Containers/Set.h b/Engine/Source/Runtime/Core/Public/Containers/Set.h
--- a/Engine/Source/Runtime/Core/Public/Containers/Set.h
+++ b/Engine/Source/Runtime/Core/Public/Containers/Set.h
@@ -1260,7 +1260,7 @@ private:
 public:
 	void WriteMemoryImage(FMemoryImageWriter& Writer) const
 	{
-		checkf(!Writer.GetTargetLayoutParams().b32Bit, TEXT("TSet does not currently support freezing for 32bits"));
+		checkf(!Writer.Is32BitTarget(), TEXT("TSet does not currently support freezing for 32bits"));
 		TSupportsFreezeMemoryImageHelper<SupportsFreezeMemoryImage>::WriteMemoryImage(Writer, *this);
 	}
 
diff --git a/Engine/Source/Runtime/Core/Public/Containers/SparseArray.h b/Engine/Source/Runtime/Core/Public/Containers/SparseArray.h
--- a/Engine/Source/Runtime/Core/Public/Containers/SparseArray.h
+++ b/Engine/Source/Runtime/Core/Public/Containers/SparseArray.h
@@ -1157,7 +1157,7 @@ private:
 public:
 	void WriteMemoryImage(FMemoryImageWriter& Writer) const
 	{
-		checkf(!Writer.GetTargetLayoutParams().b32Bit, TEXT("TSparseArray does not currently support freezing for 32bits"));
+		checkf(!Writer.Is32BitTarget(), TEXT("TSparseArray does not currently support freezing for 32bits"));
 		TSupportsFreezeMemoryImageHelper<SupportsFreezeMemoryImage>::WriteMemoryImage(Writer, *this);
 	}
 
diff --git a/Engine/Source/Runtime/Core/Public/Serialization/MemoryImageWriter.h b/Engine/Source/Runtime/Core/Public/Serialization/MemoryImageWriter.h
--- a/Engine/Source/Runtime/Core/Public/Serialization/MemoryImageWriter.h
+++ b/Engine/Source/Runtime/Core/Public/Serialization/MemoryImageWriter.h
@@ -22,7 +22,7 @@ public:
 	FPointerTableBase& GetPointerTable() const;
 	const FPointerTableBase* TryGetPrevPointerTable() const;
 
-	inline bool Is32BitTarget() const { return GetTargetLayoutParams().b32Bit; }
+	inline bool Is32BitTarget() const { return GetTargetLayoutParams().Is32Bit(); }
 	inline bool Is64BitTarget() const { return !Is32BitTarget(); }
 
 	void AddDependency(const FTypeLayoutDesc& TypeDesc);
diff --git a/Engine/Source/Runtime/Core/Public/Serialization/MemoryLayout.h b/Engine/Source/Runtime/Core/Public/Serialization/MemoryLayout.h
--- a/Engine/Source/Runtime/Core/Public/Serialization/MemoryLayout.h
+++ b/Engine/Source/Runtime/Core/Public/Serialization/MemoryLayout.h
@@ -16,6 +16,7 @@ class FMemoryImageWriter;
 class FMemoryUnfreezeContent;
 class FPointerTableBase;
 struct FTypeLayoutDesc;
+struct FPlatformTypeLayoutParameters;
 
 // Duplicated from RHIDefinitions.h [AND MUST MATCH THE Freezing_bWithRayTracing in DataDrivenPlatformInfo.ini]
 #ifndef WITH_RAYTRACING
@@ -66,57 +67,6 @@ using FMemoryImageUPtrInt = UPTRINT;
 
 template<typename T> const FTypeLayoutDesc& StaticGetTypeLayoutDesc();
 
-struct FPlatformTypeLayoutParameters
-{
-	uint32 MaxFieldAlignment = 0xffffffff;
-	bool b32Bit = false;
-	bool bForce64BitMemoryImagePointers = true;
-	bool bAlignBases = false;
-	bool bWithEditorOnly = false;
-	bool bWithRayTracing = false;
-	bool bIsCurrentPlatform = false;
-
-	inline bool Has32BitMemoryImagePointers() const { return b32Bit && !bForce64BitMemoryImagePointers; }
-	inline bool Has64BitMemoryImagePointers() const { return !Has32BitMemoryImagePointers(); }
-
-	inline uint32 GetRawPointerSize() const { return b32Bit ? sizeof(uint32) : sizeof(uint64); }
-	inline uint32 GetMemoryImagePointerSize() const { return Has32BitMemoryImagePointers() ? sizeof(uint32) : sizeof(uint64); }
-
-	CORE_API void InitializeForArchive(FArchive& Ar);
-	CORE_API void InitializeForPlatform(const FString& PlatformName, bool bHasEditorOnlyData);
-	CORE_API void InitializeForCurrent();
-	CORE_API void InitializeForMSVC();
-	CORE_API void InitializeForClang();
-
-	/**
-	 * Allow the layout parameters to modify the given DDC key string
-	 * Normally this should not be used, and instead hashed layouts of any dependent types should be preferred
-	 * This is mostly needed when fixing bugs related to writing memory images with certain parameter combinations
-	 */
-	CORE_API void AppendKeyString(FString& KeyString) const;
-
-	/**
-	 * This is used for serializing into/from the DDC
-	 */
-	void Serialize(FArchive& Ar);
-
-	bool operator==(const FPlatformTypeLayoutParameters& ReferenceSet) const
-	{
-		return MaxFieldAlignment == ReferenceSet.MaxFieldAlignment &&
-			b32Bit == ReferenceSet.b32Bit &&
-			bForce64BitMemoryImagePointers == ReferenceSet.bForce64BitMemoryImagePointers &&
-			bAlignBases == ReferenceSet.bAlignBases &&
-			bWithEditorOnly == ReferenceSet.bWithEditorOnly &&
-			bWithRayTracing == ReferenceSet.bWithRayTracing &&
-			bIsCurrentPlatform == ReferenceSet.bIsCurrentPlatform;
-	}
-
-	bool operator!=(const FPlatformTypeLayoutParameters& ReferenceSet) const
-	{
-		return !(*this == ReferenceSet);
-	}
-};
-
 struct FMemoryToStringContext
 {
 	CORE_API void AppendNullptr();
@@ -354,22 +304,15 @@ namespace Freeze
 		return TypeDesc.Alignment;
 	}
 
-	inline uint32 IntrinsicAppendHash(void* const* DummyObject, const FTypeLayoutDesc& TypeDesc, const FPlatformTypeLayoutParameters& LayoutParams, FSHA1& Hasher)
-	{
-		return AppendHashForNameAndSize(TypeDesc.Name, LayoutParams.GetRawPointerSize(), Hasher);
-	}
-
-	inline uint32 IntrinsicGetTargetAlignment(void* const* DummyObject, const FTypeLayoutDesc& TypeDesc, const FPlatformTypeLayoutParameters& LayoutParams)
-	{
-		return LayoutParams.GetRawPointerSize();
-	}
-
 	template<typename T>
 	inline void IntrinsicToString(const T& Object, const FTypeLayoutDesc& TypeDesc, const FPlatformTypeLayoutParameters& LayoutParams, FMemoryToStringContext& OutContext)
 	{
 		DefaultToString(&Object, TypeDesc, LayoutParams, OutContext);
 	}
 
+	CORE_API uint32 IntrinsicAppendHash(void* const* DummyObject, const FTypeLayoutDesc& TypeDesc, const FPlatformTypeLayoutParameters& LayoutParams, FSHA1& Hasher);
+	CORE_API uint32 IntrinsicGetTargetAlignment(void* const* DummyObject, const FTypeLayoutDesc& TypeDesc, const FPlatformTypeLayoutParameters& LayoutParams);
+
 	CORE_API void IntrinsicToString(char Object, const FTypeLayoutDesc& TypeDesc, const FPlatformTypeLayoutParameters& LayoutParams, FMemoryToStringContext& OutContext);
 	CORE_API void IntrinsicToString(short Object, const FTypeLayoutDesc& TypeDesc, const FPlatformTypeLayoutParameters& LayoutParams, FMemoryToStringContext& OutContext);
 	CORE_API void IntrinsicToString(int Object, const FTypeLayoutDesc& TypeDesc, const FPlatformTypeLayoutParameters& LayoutParams, FMemoryToStringContext& OutContext);
@@ -884,3 +827,68 @@ ALIAS_TEMPLATE_TYPE_LAYOUT(template<typename T>, const T, T);
 
 // All raw pointer types map to void*, since they're all handled the same way
 ALIAS_TEMPLATE_TYPE_LAYOUT(template<typename T>, T*, void*);
+
+struct FPlatformTypeLayoutParameters
+{
+	DECLARE_EXPORTED_TYPE_LAYOUT(FPlatformTypeLayoutParameters, CORE_API, NonVirtual);
+
+	enum Flags
+	{
+		Flag_Initialized = (1 << 0),
+		Flag_Is32Bit = (1 << 1),
+		Flag_Force64BitMemoryImagePointers = (1 << 2),
+		Flag_AlignBases = (1 << 3),
+		Flag_WithEditorOnly = (1 << 4),
+		Flag_WithRaytracing = (1 << 5),
+	};
+
+	LAYOUT_FIELD_INITIALIZED(uint32, MaxFieldAlignment, 0xffffffff);
+	LAYOUT_FIELD_INITIALIZED(uint32, Flags, 0u);
+
+	inline bool IsInitialized() const { return (Flags & Flag_Initialized) != 0u; }
+	inline bool Is32Bit() const { return (Flags & Flag_Is32Bit) != 0u; }
+	inline bool HasForce64BitMemoryImagePointers() const { return (Flags & Flag_Force64BitMemoryImagePointers) != 0u; }
+	inline bool HasAlignBases() const { return (Flags & Flag_AlignBases) != 0u; }
+	inline bool WithEditorOnly() const { return (Flags & Flag_WithEditorOnly) != 0u; }
+	inline bool WithRaytracing() const { return (Flags & Flag_WithRaytracing) != 0u; }
+
+	inline bool Has32BitMemoryImagePointers() const { return Is32Bit() && !HasForce64BitMemoryImagePointers(); }
+	inline bool Has64BitMemoryImagePointers() const { return !Has32BitMemoryImagePointers(); }
+
+	inline uint32 GetRawPointerSize() const { return Is32Bit() ? sizeof(uint32) : sizeof(uint64); }
+	inline uint32 GetMemoryImagePointerSize() const { return Has32BitMemoryImagePointers() ? sizeof(uint32) : sizeof(uint64); }
+
+	friend inline FArchive& operator<<(FArchive& Ar, FPlatformTypeLayoutParameters& Ref)
+	{
+		return Ref.Serialize(Ar);
+	}
+
+	friend inline bool operator==(const FPlatformTypeLayoutParameters& Lhs, const FPlatformTypeLayoutParameters& Rhs)
+	{
+		return Lhs.Flags == Rhs.Flags && Lhs.MaxFieldAlignment == Rhs.MaxFieldAlignment;
+	}
+
+	friend inline bool operator!=(const FPlatformTypeLayoutParameters& Lhs, const FPlatformTypeLayoutParameters& Rhs)
+	{
+		return !operator==(Lhs, Rhs);
+	}
+
+	CORE_API bool IsCurrentPlatform() const;
+	CORE_API void InitializeForArchive(FArchive& Ar);
+	CORE_API void InitializeForPlatform(const FString& PlatformName, bool bHasEditorOnlyData);
+	CORE_API void InitializeForCurrent();
+	CORE_API void InitializeForMSVC();
+	CORE_API void InitializeForClang();
+
+	/**
+	 * This is used for serializing into/from the DDC
+	 */
+	CORE_API FArchive& Serialize(FArchive& Ar);
+
+	/**
+	 * Allow the layout parameters to modify the given DDC key string
+	 * Normally this should not be used, and instead hashed layouts of any dependent types should be preferred
+	 * This is mostly needed when fixing bugs related to writing memory images with certain parameter combinations
+	 */
+	CORE_API void AppendKeyString(FString& KeyString) const;
+};
diff --git a/Engine/Source/Runtime/Engine/Private/Materials/MaterialShader.cpp b/Engine/Source/Runtime/Engine/Private/Materials/MaterialShader.cpp
--- a/Engine/Source/Runtime/Engine/Private/Materials/MaterialShader.cpp
+++ b/Engine/Source/Runtime/Engine/Private/Materials/MaterialShader.cpp
@@ -394,7 +394,7 @@ void FMaterialShaderMapId::Serialize(FArchive& Ar, bool bLoadedByCookedMaterial)
 	// the LayoutParams serialized. The other option (old DDC) should be prevented by us having mutated MATERIALSHADERMAP_DERIVEDDATA_VER
 	if (Ar.UE4Ver() >= VER_UE4_PURGED_FMATERIAL_COMPILE_OUTPUTS)
 	{
-		LayoutParams.Serialize(Ar);
+		Ar << LayoutParams;
 	}
 	else
 	{
@@ -601,6 +601,11 @@ bool FMaterialShaderMapId::operator==(const FMaterialShaderMapId& ReferenceSet)
 bool FMaterialShaderMapId::IsContentValid() const
 {
 #if WITH_EDITOR
+	if (!LayoutParams.IsInitialized())
+	{
+		return false;
+	}
+
 	//We expect overrides to be set to false
 	for (const FStaticSwitchParameter& StaticSwitchParameter : StaticSwitchParameters)
 	{
