fix t7xx t6xx Stretching Geometry GPUScene

diff --git a/Engine/Config/BaseDeviceProfiles.ini b/Engine/Config/BaseDeviceProfiles.ini
--- a/Engine/Config/BaseDeviceProfiles.ini
+++ b/Engine/Config/BaseDeviceProfiles.ini
@@ -574,15 +574,17 @@ BaseProfileName=Android_PowerVR_GM9xxx
 [Android_Mali_T6xx DeviceProfile]
 DeviceType=Android
 BaseProfileName=Android_Low
++CVars=r.Android.MaliMidgardIndexingBug=1
 
 [Android_Mali_T7xx DeviceProfile]
 DeviceType=Android
 BaseProfileName=Android_Low
++CVars=r.Android.MaliMidgardIndexingBug=1
 
 [Android_Mali_T8xx DeviceProfile]
 DeviceType=Android
 BaseProfileName=Android_Mid
-+CVars=r.Android.MaliT8Bug=1
++CVars=r.Android.MaliMidgardIndexingBug=1
 ; using early_fragment_tests in a fragment shader does not work correctly on this device
 +CVars=r.Android.DisableEarlyFragmentTests=1
 
diff --git a/Engine/Source/Runtime/Core/Private/Android/AndroidPlatformMisc.cpp b/Engine/Source/Runtime/Core/Private/Android/AndroidPlatformMisc.cpp
--- a/Engine/Source/Runtime/Core/Private/Android/AndroidPlatformMisc.cpp
+++ b/Engine/Source/Runtime/Core/Private/Android/AndroidPlatformMisc.cpp
@@ -75,8 +75,8 @@ static FAutoConsoleVariableRef CAndroidLowPowerBatteryThreshold(
 	ECVF_Default
 );
 
-static TAutoConsoleVariable<int32> CVarMaliT8Bug(
-	TEXT("r.Android.MaliT8Bug"),
+static TAutoConsoleVariable<int32> CVarMaliMidgardIndexingBug(
+	TEXT("r.Android.MaliMidgardIndexingBug"),
 	0,
 	TEXT("For an indexed instance draw, the OpenGL ES driver does not handle attributes correctly. This issue only happens on Mali T8xx GPU when the difference between two adjacent index values are larger than 16.\n")
 	TEXT("  0 = off\n")
@@ -2739,5 +2739,5 @@ uint32 FAndroidMisc::GetCoreFrequency(int32 CoreIndex, ECoreFrequencyProperty Co
 
 bool FAndroidMisc::Expand16BitIndicesTo32BitOnLoad()
 {
-	return  (CVarMaliT8Bug.GetValueOnAnyThread() > 0);
+	return  (CVarMaliMidgardIndexingBug.GetValueOnAnyThread() > 0);
 }
