Fixed build from missing CORE_API on exposed vars.

#rnx

diff --git a/Engine/Source/Runtime/Core/Private/Async/TaskGraph.cpp b/Engine/Source/Runtime/Core/Private/Async/TaskGraph.cpp
--- a/Engine/Source/Runtime/Core/Private/Async/TaskGraph.cpp
+++ b/Engine/Source/Runtime/Core/Private/Async/TaskGraph.cpp
@@ -62,8 +62,8 @@ namespace ENamedThreads
 }
 
 // RenderingThread.cpp sets these values if needed
-TAtomic<bool> GDoRenderThreadWakeupTrigger(true);	// Accessed by any thread, modified by GT.
-int32 GRenderThreadPollPeriodMs = -1;				// Accessed/Modified by RT only.
+CORE_API TAtomic<bool> GDoRenderThreadWakeupTrigger(true);	// Accessed by any thread, modified by GT/RT.
+CORE_API int32 GRenderThreadPollPeriodMs = -1;				// Accessed/Modified by RT only.
 
 static int32 GIgnoreThreadToDoGatherOn = 0;
 static FAutoConsoleVariableRef CVarIgnoreThreadToDoGatherOn(
diff --git a/Engine/Source/Runtime/RenderCore/Private/RenderingThread.cpp b/Engine/Source/Runtime/RenderCore/Private/RenderingThread.cpp
--- a/Engine/Source/Runtime/RenderCore/Private/RenderingThread.cpp
+++ b/Engine/Source/Runtime/RenderCore/Private/RenderingThread.cpp
@@ -60,8 +60,8 @@ float GRenderingThreadMaxIdleTickFrequency = 40.f;
  * RT Task Graph polling.
  */
 
-extern TAtomic<bool> GDoRenderThreadWakeupTrigger;
-extern int32 GRenderThreadPollPeriodMs;
+extern CORE_API TAtomic<bool> GDoRenderThreadWakeupTrigger;
+extern CORE_API int32 GRenderThreadPollPeriodMs;
 
 static void OnRenderThreadPollPeriodMsChanged(IConsoleVariable* Var)
 {
