- Some more PlatformExtension moves
#rnx

diff --git a/Engine/Source/Runtime/Core/Public/MemPro/MemPro.h b/Engine/Source/Runtime/Core/Public/MemPro/MemPro.h
--- a/Engine/Source/Runtime/Core/Public/MemPro/MemPro.h
+++ b/Engine/Source/Runtime/Core/Public/MemPro/MemPro.h
@@ -297,7 +297,7 @@ namespace
 #elif defined(MEMPRO_PLATFORM_XBOXONE)
 
 	#ifdef __UNREAL__
-		#include "XboxOne/MemProXboxOne.h"
+		#include "MemPro/MemProXboxOne.h"
 	#else
 		#include "MemProXboxOne.hpp"		// contact [REDACTED_EMAIL] for this platform
 	#endif
