FParse::Line and LineExtended will reuse memory of passed in Result if already allocated
#rnx

diff --git a/Engine/Source/Runtime/Core/Private/Misc/Parse.cpp b/Engine/Source/Runtime/Core/Private/Misc/Parse.cpp
--- a/Engine/Source/Runtime/Core/Private/Misc/Parse.cpp
+++ b/Engine/Source/Runtime/Core/Private/Misc/Parse.cpp
@@ -919,10 +919,14 @@ bool FParse::Token( const TCHAR*& Str, FString& Arg, bool UseEscape )
 FString FParse::Token( const TCHAR*& Str, bool UseEscape )
 {
 	TCHAR Buffer[1024];
-	if( FParse::Token( Str, Buffer, UE_ARRAY_COUNT(Buffer), UseEscape ) )
+	if (FParse::Token(Str, Buffer, UE_ARRAY_COUNT(Buffer), UseEscape))
+	{
 		return Buffer;
+	}
 	else
+	{
 		return TEXT("");
+	}
 }
 
 bool FParse::AlnumToken(const TCHAR*& Str, FString& Arg)
@@ -1016,7 +1020,7 @@ bool FParse::Line(const TCHAR** Stream, FString& Result, bool bExact)
 	bool bIsQuoted = false;
 	bool bIgnore = false;
 
-	Result = TEXT("");
+	Result.Reset();
 
 	while (**Stream != TEXT('\0') && **Stream != TEXT('\n') && **Stream != TEXT('\r'))
 	{
@@ -1078,7 +1082,7 @@ bool FParse::LineExtended(const TCHAR** Stream, FString& Result, int32& LinesCon
 	bool bIgnore = false;
 	int32 BracketDepth = 0;
 
-	Result = TEXT("");
+	Result.Reset();
 	LinesConsumed = 0;
 
 	while (**Stream != TEXT('\0') && ((**Stream != TEXT('\n') && **Stream != TEXT('\r')) || BracketDepth > 0))
