TBA: When getting the cacheable archive from an archive->structuredarchive adapter, avoid using the internal root to get to the underlying archive as it won't have been created if we haven't written anything yet. The adapter itself has a reference to the inner archive so just use that.

diff --git a/Engine/Source/Runtime/Core/Private/Serialization/ArchiveFromStructuredArchive.cpp b/Engine/Source/Runtime/Core/Private/Serialization/ArchiveFromStructuredArchive.cpp
--- a/Engine/Source/Runtime/Core/Private/Serialization/ArchiveFromStructuredArchive.cpp
+++ b/Engine/Source/Runtime/Core/Private/Serialization/ArchiveFromStructuredArchive.cpp
@@ -308,7 +308,14 @@ void FArchiveFromStructuredArchiveImpl::OpenArchive()
 
 FArchive* FArchiveFromStructuredArchiveImpl::GetCacheableArchive()
 {
-	return IsTextFormat() ? nullptr : Pimpl->Root->GetUnderlyingArchive().GetCacheableArchive();
+	if (IsTextFormat())
+	{
+		return nullptr;
+	}
+	else
+	{
+		return InnerArchive.GetCacheableArchive();
+	}
 }
 
 bool FArchiveFromStructuredArchiveImpl::ContainsData() const
