Add a way for a function call to be monitored for timely completion, if it doesn't complete within a time limit treat it as a hang.

Andriy.Tylychko

diff --git a/Engine/Source/Runtime/Core/Private/HAL/ThreadHeartBeat.cpp b/Engine/Source/Runtime/Core/Private/HAL/ThreadHeartBeat.cpp
--- a/Engine/Source/Runtime/Core/Private/HAL/ThreadHeartBeat.cpp
+++ b/Engine/Source/Runtime/Core/Private/HAL/ThreadHeartBeat.cpp
@@ -279,6 +279,10 @@ uint32 FThreadHeartBeat::Run()
 	{
 		double HangDuration;
 		uint32 ThreadThatHung = CheckHeartBeat(HangDuration);
+		if (ThreadThatHung == FThreadHeartBeat::InvalidThreadId)
+		{
+			ThreadThatHung = CheckFunctionHeartBeat(HangDuration);
+		}
 
 		if (ThreadThatHung == FThreadHeartBeat::InvalidThreadId)
 		{
@@ -398,19 +402,72 @@ void FThreadHeartBeat::PresentFrame()
 #endif
 }
 
-uint32 FThreadHeartBeat::CheckHeartBeat(double& OutHangDuration)
+void FThreadHeartBeat::MonitorFunctionStart()
 {
-	// Editor and debug builds run too slow to measure them correctly
+#if USE_HANG_DETECTION
+	// disable on platforms that don't start the thread
+	if (FPlatformMisc::AllowThreadHeartBeat() == false)
+	{
+		return;
+	}
+
+	uint32 ThreadId = FPlatformTLS::GetCurrentThreadId();
+	FScopeLock FunctionHeartBeatLock(&FunctionHeartBeatCritical);
+
+	FHeartBeatInfo& HeartBeatInfo = FunctionHeartBeat.FindOrAdd(ThreadId);
+	HeartBeatInfo.LastHeartBeatTime = Clock.Seconds();
+	HeartBeatInfo.HangDuration = CurrentHangDuration;
+	HeartBeatInfo.SuspendedCount = 0;
+#endif
+}
+
+void FThreadHeartBeat::MonitorFunctionEnd()
+{
+#if USE_HANG_DETECTION
+	// disable on platforms that don't start the thread
+	if (FPlatformMisc::AllowThreadHeartBeat() == false)
+	{
+		return;
+	}
+	uint32 ThreadId = FPlatformTLS::GetCurrentThreadId();
+	FScopeLock FunctionHeartBeatLock(&FunctionHeartBeatCritical);
+	FHeartBeatInfo* HeartBeatInfo = FunctionHeartBeat.Find(ThreadId);
+	if (HeartBeatInfo != nullptr)
+	{
+		HeartBeatInfo->SuspendedCount = 1;
+	}
+	else
+	{
+		// has to have been there otherwise the Start/End are out of order/unbalanced
+		check(0);
+	}
+#endif
+}
+
+
+bool FThreadHeartBeat::IsEnabled()
+{
+	bool CheckBeats = false;
 #if USE_HANG_DETECTION
 	static bool bForceEnabled = FParse::Param(FCommandLine::Get(), TEXT("debughangdetection"));
 	static bool bDisabled = !bForceEnabled && FParse::Param(FCommandLine::Get(), TEXT("nothreadtimeout"));
 
-	bool CheckBeats = (ConfigHangDuration > 0.0 || ConfigPresentDuration > 0.0)
+	CheckBeats = (ConfigHangDuration > 0.0 || ConfigPresentDuration > 0.0)
 		&& bReadyToCheckHeartbeat
 		&& !IsEngineExitRequested()
 		&& (bForceEnabled || !FPlatformMisc::IsDebuggerPresent())
 		&& !bDisabled
 		&& !GlobalSuspendCount.GetValue();
+#endif
+	return CheckBeats;
+}
+
+
+uint32 FThreadHeartBeat::CheckHeartBeat(double& OutHangDuration)
+{
+	// Editor and debug builds run too slow to measure them correctly
+#if USE_HANG_DETECTION
+	bool CheckBeats = IsEnabled();
 
 	if (CheckBeats)
 	{
@@ -450,6 +507,39 @@ uint32 FThreadHeartBeat::CheckHeartBeat(double& OutHangDuration)
 	return InvalidThreadId;
 }
 
+uint32 FThreadHeartBeat::CheckFunctionHeartBeat(double& OutHangDuration)
+{
+	// Editor and debug builds run too slow to measure them correctly
+#if USE_HANG_DETECTION
+	bool CheckBeats = IsEnabled();
+
+	if (CheckBeats)
+	{
+		const double CurrentTime = Clock.Seconds();
+		FScopeLock HeartBeatLock(&FunctionHeartBeatCritical);
+		if (ConfigHangDuration > 0.0)
+		{
+			// Check heartbeat for all functions and return thread ID of the thread that was running the function when it hung.
+			// Note: We only return a thread id for a thread that has updated since the last hang, i.e. is still alive
+			// This avoids the case where a user may be in a deep and minorly varying callstack and flood us with reports
+			for (TPair<uint32, FHeartBeatInfo>& LastHeartBeat : FunctionHeartBeat)
+			{
+				FHeartBeatInfo& HeartBeatInfo = LastHeartBeat.Value;
+				if (HeartBeatInfo.SuspendedCount == 0 && (CurrentTime - HeartBeatInfo.LastHeartBeatTime) > HeartBeatInfo.HangDuration && HeartBeatInfo.LastHeartBeatTime >= HeartBeatInfo.LastHangTime)
+				{
+					HeartBeatInfo.LastHangTime = CurrentTime;
+					OutHangDuration = HeartBeatInfo.HangDuration;
+					return LastHeartBeat.Key;
+				}
+			}
+		}
+	}
+#endif
+	return InvalidThreadId;
+}
+
+
+
 void FThreadHeartBeat::KillHeartBeat()
 {
 #if USE_HANG_DETECTION
diff --git a/Engine/Source/Runtime/Core/Public/HAL/ThreadHeartBeat.h b/Engine/Source/Runtime/Core/Public/HAL/ThreadHeartBeat.h
--- a/Engine/Source/Runtime/Core/Public/HAL/ThreadHeartBeat.h
+++ b/Engine/Source/Runtime/Core/Public/HAL/ThreadHeartBeat.h
@@ -86,6 +86,12 @@ class CORE_API FThreadHeartBeat : public FRunnable
 	TMap<uint32, FHeartBeatInfo> ThreadHeartBeat;
 	/** The last heartbeat time for the rendering or RHI thread frame present. */
 	FHeartBeatInfo PresentHeartBeat;
+
+	/** Synch object for the function heartbeat */
+	FCriticalSection FunctionHeartBeatCritical;
+	/** Keeps track of the last heartbeat time for a function, can't be nested */
+	TMap<uint32, FHeartBeatInfo> FunctionHeartBeat;
+	
 	/** True if heartbeat should be measured */
 	FThreadSafeBool bReadyToCheckHeartbeat;
 	/** Max time the thread is allowed to not send the heartbeat*/
@@ -115,6 +121,8 @@ class CORE_API FThreadHeartBeat : public FRunnable
 	void FORCENOINLINE OnHang(double HangDuration, uint32 ThreadThatHung);
 	void FORCENOINLINE OnPresentHang(double HangDuration);
 
+	bool IsEnabled();
+
 public:
 
 	enum EConstants
@@ -140,6 +148,14 @@ public:
 	uint32 CheckHeartBeat(double& OutHangDuration);
 	/** Called by a thread when it's no longer expecting to be ticked */
 	void KillHeartBeat();
+
+	/** Called from a thread once on entry to a function to be monitored */
+	void MonitorFunctionStart();
+	/** Called by a thread when a function has completed and no longer needs to be monitored */
+	void MonitorFunctionEnd();
+	/** Called by a supervising thread to check all function calls' being monitored health */
+	uint32 CheckFunctionHeartBeat(double& OutHangDuration);
+
 	/** 
 	 * Suspend heartbeat measuring for the current thread if the thread has already had a heartbeat 
 	 * @param bAllThreads If true, suspends heartbeat for all threads, not only the current one
@@ -198,6 +214,27 @@ public:
 	}
 };
 
+/** Simple scope object to put at the top of a function to monitor it completes in a timely fashion */
+struct FFunctionHeartBeatScope
+{
+public:
+	FORCEINLINE FFunctionHeartBeatScope()
+	{
+		if (FThreadHeartBeat* HB = FThreadHeartBeat::GetNoInit())
+		{
+			HB->MonitorFunctionStart();
+		}
+	}
+	FORCEINLINE ~FFunctionHeartBeatScope()
+	{
+		if (FThreadHeartBeat* HB = FThreadHeartBeat::GetNoInit())
+		{
+			HB->MonitorFunctionEnd();
+		}
+	}
+};
+
+
 // When 1, performs a full symbol lookup in hitch call stacks, otherwise only
 // a backtrace is performed and the raw addresses are written to the log.
 #ifndef LOOKUP_SYMBOLS_IN_HITCH_STACK_WALK
