Put all code into the .text segment on ARM. Having the uedbg section present confuses clang when the code requires branches larger than the ARM branch instruction can handle. Clang will generate trampolines in the .text segment only.

diff --git a/Engine/Source/Runtime/Core/Public/Misc/AssertionMacros.h b/Engine/Source/Runtime/Core/Public/Misc/AssertionMacros.h
--- a/Engine/Source/Runtime/Core/Public/Misc/AssertionMacros.h
+++ b/Engine/Source/Runtime/Core/Public/Misc/AssertionMacros.h
@@ -10,12 +10,18 @@
 #include "Templates/IsValidVariadicFunctionArg.h"
 #include "Misc/VarArgs.h"
 
-#if DO_CHECK || DO_GUARD_SLOW
+#if (DO_CHECK || DO_GUARD_SLOW) && !PLATFORM_CPU_ARM_FAMILY
 	// We'll put all assert implementation code into a separate section in the linked
 	// executable. This code should never execute so using a separate section keeps
 	// it well off the hot path and hopefully out of the instruction cache. It also
 	// facilitates reasoning about the makeup of a compiled/linked binary.
 	#define UE_DEBUG_SECTION PLATFORM_CODE_SECTION(".uedbg")
+#else
+	// On ARM we can't do this because the executable will require jumps larger
+	// than the branch instruction can handle. Clang will only generate
+	// the trampolines in the .text segment of the binary. If the uedbg segment
+	// is present it will generate code that it cannot link.
+	#define UE_DEBUG_SECTION
 #endif // DO_CHECK || DO_GUARD_SLOW
 
 namespace ELogVerbosity
