add explicit PLATFORM_SUPPORTS_COLORIZED_OUTPUT_DEVICE instead of assuming only PLATFORM_DESKTOP has support.

diff --git a/Engine/Source/Runtime/Core/Public/HAL/Platform.h b/Engine/Source/Runtime/Core/Public/HAL/Platform.h
--- a/Engine/Source/Runtime/Core/Public/HAL/Platform.h
+++ b/Engine/Source/Runtime/Core/Public/HAL/Platform.h
@@ -591,6 +591,10 @@
 	#define PLATFORM_MANAGES_HDR_SETTING 0
 #endif
 
+#ifndef PLATFORM_SUPPORTS_COLORIZED_OUTPUT_DEVICE
+	#define PLATFORM_SUPPORTS_COLORIZED_OUTPUT_DEVICE PLATFORM_DESKTOP
+#endif
+
 // deprecated, do not use
 #define PLATFORM_HAS_THREADSAFE_RHIGetRenderQueryResult	#
 #define PLATFORM_SUPPORTS_RHI_THREAD #
diff --git a/Engine/Source/Runtime/Core/Public/Misc/OutputDevice.h b/Engine/Source/Runtime/Core/Public/Misc/OutputDevice.h
--- a/Engine/Source/Runtime/Core/Public/Misc/OutputDevice.h
+++ b/Engine/Source/Runtime/Core/Public/Misc/OutputDevice.h
@@ -11,7 +11,7 @@
 #include "Templates/AndOrNot.h"
 #include "Templates/IsArrayOrRefOfType.h"
 
-#if !PLATFORM_DESKTOP
+#if !PLATFORM_SUPPORTS_COLORIZED_OUTPUT_DEVICE
 	// don't support colorized text on consoles
 	#define SET_WARN_COLOR(Color)
 	#define SET_WARN_COLOR_AND_BACKGROUND(Color, Bkgrnd)
