Allow ensures to be compiled independantly of checks via USE_ENSURES_IN_SHIPPING

diff --git a/Engine/Source/Runtime/Core/Private/GenericPlatform/GenericPlatformMisc.cpp b/Engine/Source/Runtime/Core/Private/GenericPlatform/GenericPlatformMisc.cpp
--- a/Engine/Source/Runtime/Core/Private/GenericPlatform/GenericPlatformMisc.cpp
+++ b/Engine/Source/Runtime/Core/Private/GenericPlatform/GenericPlatformMisc.cpp
@@ -1230,7 +1230,7 @@ FString FGenericPlatformMisc::GetTimeZoneId()
 	return FString();
 }
 
-#if DO_CHECK
+#if DO_ENSURE
 namespace GenericPlatformMisc
 {
 	/** Chances for handling an ensure (0.0 - never, 1.0 - always). */
@@ -1277,7 +1277,7 @@ void FGenericPlatformMisc::UpdateHotfixableEnsureSettings()
 
 	GenericPlatformMisc::GEnsureSettingsEverUpdated = true;
 }
-#endif // !DO_CHECK
+#endif // #if DO_ENSURE
 
 void FGenericPlatformMisc::TickHotfixables()
 {
diff --git a/Engine/Source/Runtime/Core/Private/Misc/AssertionMacros.cpp b/Engine/Source/Runtime/Core/Private/Misc/AssertionMacros.cpp
--- a/Engine/Source/Runtime/Core/Private/Misc/AssertionMacros.cpp
+++ b/Engine/Source/Runtime/Core/Private/Misc/AssertionMacros.cpp
@@ -233,7 +233,7 @@ void FDebug::LogFormattedMessageWithCallstack(const FName& InLogName, const ANSI
 	}
 }
 
-#if DO_CHECK || DO_GUARD_SLOW
+#if DO_CHECK || DO_GUARD_SLOW || DO_ENSURE
 //
 // Failed assertion handler.
 //warning: May be called at library startup time.
@@ -457,7 +457,7 @@ void FORCENOINLINE FDebug::CheckVerifyFailedImpl(
 	}
 }
 
-#endif // DO_CHECK || DO_GUARD_SLOW
+#endif // DO_CHECK || DO_GUARD_SLOW || DO_ENSURE
 
 void VARARGS FDebug::AssertFailed(const ANSICHAR* Expr, const ANSICHAR* File, int32 Line, const TCHAR* Format/* = TEXT("")*/, ...)
 {
@@ -483,7 +483,7 @@ void FDebug::ProcessFatalError()
 	GError->Logf(TEXT("%s"), GErrorHist);
 }
 
-#if DO_CHECK || DO_GUARD_SLOW
+#if DO_CHECK || DO_GUARD_SLOW || DO_ENSURE
 FORCENOINLINE bool VARARGS FDebug::OptionallyLogFormattedEnsureMessageReturningFalseImpl( bool bLog, const ANSICHAR* Expr, const ANSICHAR* File, int32 Line, const TCHAR* FormattedMsg, ... )
 {
 	if (bLog)
diff --git a/Engine/Source/Runtime/Core/Public/GenericPlatform/GenericPlatformMisc.h b/Engine/Source/Runtime/Core/Public/GenericPlatform/GenericPlatformMisc.h
--- a/Engine/Source/Runtime/Core/Public/GenericPlatform/GenericPlatformMisc.h
+++ b/Engine/Source/Runtime/Core/Public/GenericPlatform/GenericPlatformMisc.h
@@ -1051,20 +1051,20 @@ public:
 	 *
 	 * @return ensure is allowed
 	 */
-#if DO_CHECK
+#if DO_ENSURE
 	static bool IsEnsureAllowed();
 #else
 	static bool IsEnsureAllowed() { return true; }
-#endif // DO_CHECK
+#endif // DO_ENSURE
 
 	/**
 	 * Updates hotfixable ensure settings from config and commandline (config takes priority).
 	 */
-#if DO_CHECK
+#if DO_ENSURE
 	static void UpdateHotfixableEnsureSettings();
 #else
 	static void UpdateHotfixableEnsureSettings() {}
-#endif // DO_CHECK
+#endif // DO_ENSURE
 
 	/**
 	 * Ticks values that can be hotfixable in the config.
diff --git a/Engine/Source/Runtime/Core/Public/Misc/AssertionMacros.h b/Engine/Source/Runtime/Core/Public/Misc/AssertionMacros.h
--- a/Engine/Source/Runtime/Core/Public/Misc/AssertionMacros.h
+++ b/Engine/Source/Runtime/Core/Public/Misc/AssertionMacros.h
@@ -10,7 +10,7 @@
 #include "Templates/IsValidVariadicFunctionArg.h"
 #include "Misc/VarArgs.h"
 
-#if (DO_CHECK || DO_GUARD_SLOW) && !PLATFORM_CPU_ARM_FAMILY
+#if (DO_CHECK || DO_GUARD_SLOW || DO_ENSURE) && !PLATFORM_CPU_ARM_FAMILY
 	// We'll put all assert implementation code into a separate section in the linked
 	// executable. This code should never execute so using a separate section keeps
 	// it well off the hot path and hopefully out of the instruction cache. It also
@@ -62,7 +62,7 @@ struct CORE_API FDebug
 	/** Dumps the stack trace into the log with a custom heading, meant to be used for debugging purposes. */
 	static void DumpStackTraceToLog(const TCHAR* Heading);
 
-#if DO_CHECK || DO_GUARD_SLOW
+#if DO_CHECK || DO_GUARD_SLOW || DO_ENSURE
 private:
 	static void VARARGS CheckVerifyFailedImpl(const ANSICHAR* Expr, const char* File, int32 Line, const TCHAR* Format, ...);
 	static void VARARGS LogAssertFailedMessageImpl(const ANSICHAR* Expr, const ANSICHAR* File, int32 Line, const TCHAR* Fmt, ...);
@@ -141,7 +141,7 @@ public:
 // "verify" expressions are always evaluated, but only cause an error if enabled.
 //
 
-#if DO_CHECK || DO_GUARD_SLOW
+#if DO_CHECK || DO_GUARD_SLOW || DO_ENSURE
 	template <typename FmtType, typename... Types>
 	void FORCENOINLINE UE_DEBUG_SECTION FDebug::CheckVerifyFailed(
 		const ANSICHAR* Expr,
@@ -317,7 +317,7 @@ public:
  * ensure is hit in a session; ensureAlways can be used instead if you want to handle every failure
  */
 
-#if DO_CHECK && !USING_CODE_ANALYSIS // The Visual Studio 2013 analyzer doesn't understand these complex conditionals
+#if DO_ENSURE && !USING_CODE_ANALYSIS // The Visual Studio 2013 analyzer doesn't understand these complex conditionals
 
 	#define UE_ENSURE_IMPL(Capture, Always, InExpression, ...) \
 		(LIKELY(!!(InExpression)) || (DispatchCheckVerify<bool>([Capture] () FORCENOINLINE UE_DEBUG_SECTION \
@@ -342,7 +342,7 @@ public:
 	#define ensureAlways(     InExpression                ) UE_ENSURE_IMPL( , true,  InExpression, TEXT(""))
 	#define ensureAlwaysMsgf( InExpression, InFormat, ... ) UE_ENSURE_IMPL(&, true,  InExpression, InFormat, ##__VA_ARGS__)
 
-#else	// DO_CHECK
+#else	// DO_ENSURE
 
 	#define ensure(           InExpression                ) (!!(InExpression))
 	#define ensureMsgf(       InExpression, InFormat, ... ) (!!(InExpression))
diff --git a/Engine/Source/Runtime/Core/Public/Misc/Build.h b/Engine/Source/Runtime/Core/Public/Misc/Build.h
--- a/Engine/Source/Runtime/Core/Public/Misc/Build.h
+++ b/Engine/Source/Runtime/Core/Public/Misc/Build.h
@@ -184,6 +184,11 @@
 	#define USE_CHECKS_IN_SHIPPING 0
 #endif
 
+/** If not defined follow the CHECK behavior since previously ensures were compiled in with checks */
+#ifndef USE_ENSURES_IN_SHIPPING
+	#define USE_ENSURES_IN_SHIPPING USE_CHECKS_IN_SHIPPING
+#endif
+
 #ifndef ALLOW_CONSOLE_IN_SHIPPING
 	#define ALLOW_CONSOLE_IN_SHIPPING 0
 #endif
@@ -209,7 +214,8 @@
 	Basic options that by default depend on the build configuration and platform
 
 	DO_GUARD_SLOW									If true, then checkSlow, checkfSlow and verifySlow are compiled into the executable.
-	DO_CHECK										If true, then checkCode, checkf, verify, check, checkNoEntry, checkNoReentry, checkNoRecursion, verifyf, checkf, ensure, ensureAlways, ensureMsgf and ensureAlwaysMsgf are compiled into the executables
+	DO_CHECK										If true, then checkCode, checkf, verify, check, checkNoEntry, checkNoReentry, checkNoRecursion, verifyf, checkf are compiled into the executables
+	DO_ENSURE										If true, then ensure, ensureAlways, ensureMsgf and ensureAlwaysMsgf are compiled into the executables
 	STATS											If true, then the stats system is compiled into the executable.
 	ALLOW_DEBUG_FILES								If true, then debug files like screen shots and profiles can be saved from the executable.
 	NO_LOGGING										If true, then no logs or text output will be produced
@@ -223,6 +229,9 @@
 	#ifndef DO_CHECK
 		#define DO_CHECK										1
 	#endif
+	#ifndef DO_ENSURE
+		#define DO_ENSURE										1
+	#endif
 	#ifndef STATS
 		#define STATS											((WITH_UNREAL_DEVELOPER_TOOLS || !WITH_EDITORONLY_DATA || USE_STATS_WITHOUT_ENGINE || USE_MALLOC_PROFILER || FORCE_USE_STATS) && !ENABLE_STATNAMEDEVENTS)
 	#endif
@@ -242,6 +251,9 @@
 	#ifndef DO_CHECK
 		#define DO_CHECK										1
 	#endif
+	#ifndef DO_ENSURE
+		#define DO_ENSURE										1
+	#endif
 	#ifndef STATS
 		#define STATS											((WITH_UNREAL_DEVELOPER_TOOLS || !WITH_EDITORONLY_DATA || USE_STATS_WITHOUT_ENGINE || USE_MALLOC_PROFILER || FORCE_USE_STATS) && !ENABLE_STATNAMEDEVENTS)
 	#endif
@@ -261,6 +273,9 @@
 	#ifndef DO_CHECK
 		#define DO_CHECK										USE_CHECKS_IN_SHIPPING
 	#endif
+	#ifndef DO_ENSURE
+		#define DO_ENSURE										USE_ENSURES_IN_SHIPPING
+	#endif
 	#ifndef STATS
 		#define STATS											((USE_MALLOC_PROFILER || FORCE_USE_STATS) && !ENABLE_STATNAMEDEVENTS)
 	#endif
@@ -281,6 +296,9 @@
 		#ifndef DO_CHECK
 			#define DO_CHECK									1
 		#endif
+		#ifndef DO_ENSURE
+			#define DO_ENSURE									1
+		#endif
 		#ifndef STATS
 			#define STATS										1
 		#endif
@@ -300,6 +318,9 @@
 		#ifndef DO_CHECK
 			#define DO_CHECK									USE_CHECKS_IN_SHIPPING
 		#endif
+		#ifndef DO_ENSURE
+			#define DO_ENSURE									USE_ENSURES_IN_SHIPPING
+		#endif
 		#ifndef STATS
 			#define STATS										(FORCE_USE_STATS && !ENABLE_STATNAMEDEVENTS)
 		#endif
