Move location of GameOpenOrder.log, EditorOpenOrder.log and CookerOpenOrder.log to platform extensions folder
#rnx

diff --git a/Engine/Source/Editor/UnrealEd/Private/Commandlets/AssetRegistryGenerator.cpp b/Engine/Source/Editor/UnrealEd/Private/Commandlets/AssetRegistryGenerator.cpp
--- a/Engine/Source/Editor/UnrealEd/Private/Commandlets/AssetRegistryGenerator.cpp
+++ b/Engine/Source/Editor/UnrealEd/Private/Commandlets/AssetRegistryGenerator.cpp
@@ -26,6 +26,7 @@
 #include "Stats/StatsMisc.h"
 #include "Templates/UniquePtr.h"
 #include "Engine/AssetManager.h"
+#include "Misc/DataDrivenPlatformInfoRegistry.h"
 
 #include "Serialization/JsonWriter.h"
 #include "Serialization/JsonReader.h"
@@ -1366,7 +1367,15 @@ bool FAssetRegistryGenerator::WriteCookerOpenOrder()
 
 	if (CookerFileOrderString.Len())
 	{
-		FString OpenOrderFilename = FString::Printf(TEXT("%sBuild/%s/FileOpenOrder/CookerOpenOrder.log"), *FPaths::ProjectDir(), *TargetPlatform->PlatformName());
+		FString OpenOrderFilename;
+		if (FDataDrivenPlatformInfoRegistry::GetPlatformInfo(TargetPlatform->PlatformName()).bIsConfidential)
+		{
+			OpenOrderFilename = FString::Printf(TEXT("%sPlatforms/%s/Build/FileOpenOrder/CookerOpenOrder.log"), *FPaths::ProjectDir(), *TargetPlatform->PlatformName());
+		}
+		else
+		{
+			OpenOrderFilename = FString::Printf(TEXT("%sBuild/%s/FileOpenOrder/CookerOpenOrder.log"), *FPaths::ProjectDir(), *TargetPlatform->PlatformName());
+		}
 		FFileHelper::SaveStringToFile(CookerFileOrderString, *OpenOrderFilename);
 	}
 
diff --git a/Engine/Source/Programs/AutomationTool/Scripts/CopyBuildToStagingDirectory.Automation.cs b/Engine/Source/Programs/AutomationTool/Scripts/CopyBuildToStagingDirectory.Automation.cs
--- a/Engine/Source/Programs/AutomationTool/Scripts/CopyBuildToStagingDirectory.Automation.cs
+++ b/Engine/Source/Programs/AutomationTool/Scripts/CopyBuildToStagingDirectory.Automation.cs
@@ -2197,7 +2197,11 @@ public partial class Project : CommandUtils
 				foreach (string OrderLocation in OrderLocations)
 				{
 					// Add input file to control order of file within the pak
-					DirectoryReference PakOrderFileLocationBase = DirectoryReference.Combine(SC.ProjectRoot, "Build", OrderLocation, "FileOpenOrder");
+					DirectoryReference PakOrderFileLocationBase = DirectoryReference.Combine(SC.ProjectRoot, "Platforms", OrderLocation, "Build", "FileOpenOrder");
+					if (!DirectoryReference.Exists(PakOrderFileLocationBase))
+					{
+						PakOrderFileLocationBase = DirectoryReference.Combine(SC.ProjectRoot, "Build", OrderLocation, "FileOpenOrder");
+					}
 
 					FileReference FileLocation = FileReference.Combine(PakOrderFileLocationBase, OrderFileName);
 
@@ -2408,7 +2412,11 @@ public partial class Project : CommandUtils
 
 			foreach (string OrderLocation in OrderLocations.Reverse())
 			{
-				DirectoryReference IoStoreOrderFileLocationBase = DirectoryReference.Combine(SC.ProjectRoot, "Build", OrderLocation, "FileOpenOrder");
+				DirectoryReference IoStoreOrderFileLocationBase = DirectoryReference.Combine(SC.ProjectRoot, "Platforms", OrderLocation, "Build", "FileOpenOrder");
+				if (!DirectoryReference.Exists(IoStoreOrderFileLocationBase))
+				{
+					IoStoreOrderFileLocationBase = DirectoryReference.Combine(SC.ProjectRoot, "Build", OrderLocation, "FileOpenOrder");
+				}
 				FileReference FileLocation = FileReference.Combine(IoStoreOrderFileLocationBase, "GameOpenOrder.log");
 				if (FileExists_NoExceptions(FileLocation.FullName))
 				{
diff --git a/Engine/Source/Runtime/Core/Public/HAL/IPlatformFileOpenLogWrapper.h b/Engine/Source/Runtime/Core/Public/HAL/IPlatformFileOpenLogWrapper.h
--- a/Engine/Source/Runtime/Core/Public/HAL/IPlatformFileOpenLogWrapper.h
+++ b/Engine/Source/Runtime/Core/Public/HAL/IPlatformFileOpenLogWrapper.h
@@ -12,6 +12,7 @@
 #include "GenericPlatform/GenericPlatformFile.h"
 #include "Misc/Paths.h"
 #include "Async/AsyncFileHandle.h"
+#include "Misc/DataDrivenPlatformInfoRegistry.h"
 
 class IAsyncReadFileHandle;
 class FPlatformFileOpenLog;
@@ -93,7 +94,14 @@ public:
 
 			for (int32 Platform = 0;Platform < PlatformNames.Num(); ++Platform)
 			{
-				LogFileDirectory = FPaths::Combine( FPlatformMisc::ProjectDir(), TEXT( "Build" ), *PlatformNames[Platform], TEXT("FileOpenOrder"));
+				if (FDataDrivenPlatformInfoRegistry::GetPlatformInfo(PlatformNames[Platform]).bIsConfidential)
+				{
+					LogFileDirectory = FPaths::Combine(FPlatformMisc::ProjectDir(), TEXT("Platforms"), *PlatformNames[Platform], TEXT("Build"), TEXT("FileOpenOrder"));
+				}
+				else
+				{
+					LogFileDirectory = FPaths::Combine( FPlatformMisc::ProjectDir(), TEXT( "Build" ), *PlatformNames[Platform], TEXT("FileOpenOrder"));
+				}
 #if WITH_EDITOR
 				LogFilePath = FPaths::Combine( *LogFileDirectory, TEXT("EditorOpenOrder.log"));
 #else 
@@ -109,7 +117,14 @@ public:
 		}
 		else
 		{
-			LogFileDirectory = FPaths::Combine( FPlatformMisc::ProjectDir(), TEXT( "Build" ), StringCast<TCHAR>(FPlatformProperties::PlatformName()).Get(), TEXT("FileOpenOrder"));
+			if (FDataDrivenPlatformInfoRegistry::GetPlatformInfo(FPlatformProperties::PlatformName()).bIsConfidential)
+			{
+				LogFileDirectory = FPaths::Combine(FPlatformMisc::ProjectDir(), TEXT("Platforms"), StringCast<TCHAR>(FPlatformProperties::PlatformName()).Get(), TEXT("Build"), TEXT("FileOpenOrder"));
+			}
+			else
+			{
+				LogFileDirectory = FPaths::Combine(FPlatformMisc::ProjectDir(), TEXT("Build"), StringCast<TCHAR>(FPlatformProperties::PlatformName()).Get(), TEXT("FileOpenOrder"));
+			}
 #if WITH_EDITOR
 			LogFilePath = FPaths::Combine( *LogFileDirectory, TEXT("EditorOpenOrder.log"));
 #else 
diff --git a/Engine/Source/Runtime/RenderCore/Private/ShaderCodeLibrary.cpp b/Engine/Source/Runtime/RenderCore/Private/ShaderCodeLibrary.cpp
--- a/Engine/Source/Runtime/RenderCore/Private/ShaderCodeLibrary.cpp
+++ b/Engine/Source/Runtime/RenderCore/Private/ShaderCodeLibrary.cpp
@@ -1758,7 +1758,11 @@ public:
 
 				for (int32 Platform = 0; Platform < PlatformNames.Num(); ++Platform)
 				{
-					FString LogFileDirectory = FPaths::Combine(FPlatformMisc::ProjectDir(), TEXT("Build"), *PlatformNames[Platform], TEXT("FileOpenOrder"));
+					FString LogFileDirectory = FPaths::Combine(FPlatformMisc::ProjectDir(), TEXT("Platforms"), *PlatformNames[Platform], TEXT("Build"), TEXT("FileOpenOrder"));
+					if (!FPaths::DirectoryExists(LogFileDirectory))
+					{
+						LogFileDirectory = FPaths::Combine(FPlatformMisc::ProjectDir(), TEXT("Build"), *PlatformNames[Platform], TEXT("FileOpenOrder"));
+					}
 					FString LogFilePath = FPaths::Combine(*LogFileDirectory, TEXT("GameOpenOrder.log"));
 					UE_LOG(LogShaderLibrary, Display, TEXT("Checking if '%s' exists..."), *LogFilePath);
 					if (FPaths::FileExists(LogFilePath))
