Live Stream Animation Plugin
Currently, this plugin supports sending Live Link between UE4 instances, with no compression or quantization.

diff --git a/Engine/Source/Runtime/Core/Public/UObject/UnrealNames.inl b/Engine/Source/Runtime/Core/Public/UObject/UnrealNames.inl
--- a/Engine/Source/Runtime/Core/Public/UObject/UnrealNames.inl
+++ b/Engine/Source/Runtime/Core/Public/UObject/UnrealNames.inl
@@ -155,6 +155,7 @@ REGISTER_NAME(290,BeaconPort)
 REGISTER_NAME(291,MeshPort)
 REGISTER_NAME(292,MeshNetDriver)
 REGISTER_NAME(293,LiveStreamVoice)
+REGISTER_NAME(294,LiveStreamAnimation)
 
 // Texture settings.
 REGISTER_NAME(300,Linear)
