Allow null tasks in array supplied to WaitUntilTasksComplete

diff --git a/Engine/Source/Runtime/Core/Private/Async/TaskGraph.cpp b/Engine/Source/Runtime/Core/Private/Async/TaskGraph.cpp
--- a/Engine/Source/Runtime/Core/Private/Async/TaskGraph.cpp
+++ b/Engine/Source/Runtime/Core/Private/Async/TaskGraph.cpp
@@ -1505,7 +1505,8 @@ public:
 				bool bAnyPending = false;
 				for (int32 Index = 0; Index < Tasks.Num(); Index++)
 				{
-					if (!Tasks[Index]->IsComplete())
+					FGraphEvent* Task = Tasks[Index].GetReference();
+					if (Task && !Task->IsComplete())
 					{
 						bAnyPending = true;
 						break;
@@ -1527,7 +1528,8 @@ public:
 				bool bAnyPending = false;
 				for (int32 Index = 0; Index < Tasks.Num(); Index++)
 				{
-					if (!Tasks[Index]->IsComplete())
+					FGraphEvent* Task = Tasks[Index].GetReference();
+					if (Task && !Task->IsComplete())
 					{
 						bAnyPending = true;
 						break;
@@ -1554,7 +1556,8 @@ public:
 			bAnyPending = false;
 			for (int32 Index = 0; Index < Tasks.Num(); Index++)
 			{
-				if (!Tasks[Index]->IsComplete())
+				FGraphEvent* Task = Tasks[Index].GetReference();
+				if (Task && !Task->IsComplete())
 				{
 					bAnyPending = true;
 					break;
diff --git a/Engine/Source/Runtime/Core/Public/Async/TaskGraphInterfaces.h b/Engine/Source/Runtime/Core/Public/Async/TaskGraphInterfaces.h
--- a/Engine/Source/Runtime/Core/Public/Async/TaskGraphInterfaces.h
+++ b/Engine/Source/Runtime/Core/Public/Async/TaskGraphInterfaces.h
@@ -920,8 +920,8 @@ private:
 		{
 			for (int32 Index = 0; Index < Prerequisites->Num(); Index++)
 			{
-				check((*Prerequisites)[Index]);
-				if (!(*Prerequisites)[Index]->AddSubsequent(this))
+				FGraphEvent* Prerequisite = (*Prerequisites)[Index];
+				if (Prerequisite == nullptr || !Prerequisite->AddSubsequent(this))
 				{
 					AlreadyCompletedPrerequisites++;
 				}
