Fix data race when IFileCache handle is used by multiple threads

diff --git a/Engine/Source/Runtime/Core/Private/FileCache/FileCache.cpp b/Engine/Source/Runtime/Core/Private/FileCache/FileCache.cpp
--- a/Engine/Source/Runtime/Core/Private/FileCache/FileCache.cpp
+++ b/Engine/Source/Runtime/Core/Private/FileCache/FileCache.cpp
@@ -828,24 +828,20 @@ IMemoryReadStreamRef FFileCacheHandle::ReadData(FGraphEventArray& OutCompletionE
 	const CacheLineID StartLine = GetBlock<CacheLineID>(Offset);
 	const CacheLineID EndLine = GetBlock<CacheLineID>(Offset + BytesToRead - 1);
 
-	CheckForSizeRequestComplete();
-
-	int32 NumSlotsNeeded = EndLine.Get() + 1 - StartLine.Get();
+	const int32 NumCacheSlots = EndLine.Get() + 1 - StartLine.Get();
+	check(NumCacheSlots > 0);
+	const uint32 AllocSize = sizeof(FMemoryReadStreamCache) + sizeof(CacheSlotID) * (NumCacheSlots - 1);
+	void* ResultMemory = FMemory::Malloc(AllocSize, alignof(FMemoryReadStreamCache));
+	FMemoryReadStreamCache* Result = new(ResultMemory) FMemoryReadStreamCache();
+	Result->NumCacheSlots = NumCacheSlots;
+	Result->InitialSlotOffset = GetBlockOffset<CacheLineID>(Offset);
+	Result->Size = BytesToRead;
 
 	FFileCache& Cache = GetCache();
 
 	FScopeLock CacheLock(&Cache.CriticalSection);
-	/*if (NumSlotsNeeded > Cache.NumFreeSlots)
-	{
-		// not enough free slots in the cache to service this request
-		CacheLock.Unlock();
-		if((Priority & AIOP_PRIORITY_MASK) >= AIOP_CriticalPath)
-		{
-			// Critical request can't fail, just read the requested data directly
-			return ReadDataUncached(OutCompletionEvents, Offset, BytesToRead, Priority);
-		}
-		return nullptr;
-	}*/
+
+	CheckForSizeRequestComplete();
 
 	if (EndLine.Get() >= NumSlots)
 	{
@@ -859,15 +855,6 @@ IMemoryReadStreamRef FFileCacheHandle::ReadData(FGraphEventArray& OutCompletionE
 		LineToRequest.SetNum((int32)NumSlots, false);
 	}
 
-	const int32 NumCacheSlots = EndLine.Get() + 1 - StartLine.Get();
-	check(NumCacheSlots > 0);
-	const uint32 AllocSize = sizeof(FMemoryReadStreamCache) + sizeof(CacheSlotID) * (NumCacheSlots - 1);
-	void* ResultMemory = FMemory::Malloc(AllocSize, alignof(FMemoryReadStreamCache));
-	FMemoryReadStreamCache* Result = new(ResultMemory) FMemoryReadStreamCache();
-	Result->NumCacheSlots = NumCacheSlots;
-	Result->InitialSlotOffset = GetBlockOffset<CacheLineID>(Offset);
-	Result->Size = BytesToRead;
-
 	bool bHasPendingSlot = false;
 	for (CacheLineID LineID = StartLine; LineID.Get() <= EndLine.Get(); ++LineID)
 	{
@@ -930,96 +917,96 @@ FGraphEventRef FFileCacheHandle::PreloadData(const FFileCachePreloadEntry* Prelo
 
 	check(NumEntries > 0);
 
-	CheckForSizeRequestComplete();
-
-	FFileCache& Cache = GetCache();
-
-	FScopeLock CacheLock(&Cache.CriticalSection);
-
-	{
-		const FFileCachePreloadEntry& LastEntry = PreloadEntries[NumEntries - 1];
-		const CacheLineID LastEndLine = GetBlock<CacheLineID>(LastEntry.Offset + LastEntry.Size - 1);
-		if (LastEndLine.Get() >= NumSlots)
-		{
-			// If we're still waiting on SizeRequest, may need to lazily allocate some slots to service this request
-			// If this happens after SizeRequest has completed, that means something must have gone wrong
-			check(SizeRequestEvent);
-			NumSlots = LastEndLine.Get() + 1;
-			// TArray is max signed int
-			check(NumSlots < MAX_int32);
-			LineToSlot.SetNum((int32)NumSlots, false);
-			LineToRequest.SetNum((int32)NumSlots, false);
-		}
-	}
-
 	FGraphEventArray CompletionEvents;
 	TArray<CacheSlotID> LockedSlots;
 	LockedSlots.Empty(NumEntries);
 
-	CacheLineID CurrentLine(0);
-	int64 PrevOffset = -1;
-	uint32 NumSlotsLoaded = 0u;
-	for (int32 EntryIndex = 0; EntryIndex < NumEntries; ++EntryIndex)
+	FFileCache& Cache = GetCache();
 	{
-		const FFileCachePreloadEntry& Entry = PreloadEntries[EntryIndex];
-		int64 EntryOffset = InOffset + Entry.Offset;
-		const int64 EndOffset = EntryOffset + Entry.Size;
-		const CacheLineID StartLine = GetBlock<CacheLineID>(EntryOffset);
-		const CacheLineID EndLine = GetBlock<CacheLineID>(EndOffset - 1);
-
-		checkf(Entry.Offset > PrevOffset, TEXT("Preload entries must be sorted by Offset [%lld, %lld), %lld"),
-			Entry.Offset, Entry.Offset + Entry.Size, PrevOffset);
-		PrevOffset = Entry.Offset;
-
-		int64 OffsetInSlot = EntryOffset - StartLine.Get() * CacheLineSize;
-		int64 SizeInSlot = FMath::Min<int64>(Entry.Size, CacheLineSize - OffsetInSlot);
-		if (CurrentLine.Get() > StartLine.Get())
-		{
-			// Will hit this case if the last line of the previous entry is the same as the first line of this entry
-			// In this case, we'll already have a slot allocated for the line, just need to mark additional preload region
-			check(EntryIndex > 0);
-			check(CurrentLine.Get() == StartLine.Get() + 1);
-			const CacheSlotID SlotID = LineToSlot[StartLine.Get()];
-			check(SlotID.IsValid());
-			Cache.MarkSlotPreloadedRegion(SlotID, OffsetInSlot, SizeInSlot);
-			EntryOffset += SizeInSlot;
-			OffsetInSlot = 0;
-			SizeInSlot = FMath::Min<int64>(EndOffset - EntryOffset, CacheLineSize);
-		}
-		else
-		{
-			CurrentLine = StartLine;
-		}
+		FScopeLock CacheLock(&Cache.CriticalSection);
 
-		while (CurrentLine.Get() <= EndLine.Get())
+		CheckForSizeRequestComplete();
 		{
-			CacheSlotID& SlotID = LineToSlot[CurrentLine.Get()];
-			if (!SlotID.IsValid())
+			const FFileCachePreloadEntry& LastEntry = PreloadEntries[NumEntries - 1];
+			const CacheLineID LastEndLine = GetBlock<CacheLineID>(LastEntry.Offset + LastEntry.Size - 1);
+			if (LastEndLine.Get() >= NumSlots)
 			{
-				// no valid slot for this line, grab a new slot from cache and start a read request
-				SlotID = AcquireSlotAndReadLine(Cache, CurrentLine, Priority);
-				LockedSlots.Add(SlotID);
-				++NumSlotsLoaded;
+				// If we're still waiting on SizeRequest, may need to lazily allocate some slots to service this request
+				// If this happens after SizeRequest has completed, that means something must have gone wrong
+				check(SizeRequestEvent);
+				NumSlots = LastEndLine.Get() + 1;
+				// TArray is max signed int
+				check(NumSlots < MAX_int32);
+				LineToSlot.SetNum((int32)NumSlots, false);
+				LineToRequest.SetNum((int32)NumSlots, false);
 			}
-		
-			Cache.MarkSlotPreloadedRegion(SlotID, OffsetInSlot, SizeInSlot);
+		}
 
-			FPendingRequest& PendingRequest = LineToRequest[CurrentLine.Get()];
-			if (PendingRequest.Event && !PendingRequest.Event->IsComplete())
+		CacheLineID CurrentLine(0);
+		int64 PrevOffset = -1;
+		uint32 NumSlotsLoaded = 0u;
+		for (int32 EntryIndex = 0; EntryIndex < NumEntries; ++EntryIndex)
+		{
+			const FFileCachePreloadEntry& Entry = PreloadEntries[EntryIndex];
+			int64 EntryOffset = InOffset + Entry.Offset;
+			const int64 EndOffset = EntryOffset + Entry.Size;
+			const CacheLineID StartLine = GetBlock<CacheLineID>(EntryOffset);
+			const CacheLineID EndLine = GetBlock<CacheLineID>(EndOffset - 1);
+
+			checkf(Entry.Offset > PrevOffset, TEXT("Preload entries must be sorted by Offset [%lld, %lld), %lld"),
+				Entry.Offset, Entry.Offset + Entry.Size, PrevOffset);
+			PrevOffset = Entry.Offset;
+
+			int64 OffsetInSlot = EntryOffset - StartLine.Get() * CacheLineSize;
+			int64 SizeInSlot = FMath::Min<int64>(Entry.Size, CacheLineSize - OffsetInSlot);
+			if (CurrentLine.Get() > StartLine.Get())
 			{
-				// this line has a pending async request to read data
-				// will need to wait for this request to complete before data is valid
-				CompletionEvents.Add(PendingRequest.Event);
+				// Will hit this case if the last line of the previous entry is the same as the first line of this entry
+				// In this case, we'll already have a slot allocated for the line, just need to mark additional preload region
+				check(EntryIndex > 0);
+				check(CurrentLine.Get() == StartLine.Get() + 1);
+				const CacheSlotID SlotID = LineToSlot[StartLine.Get()];
+				check(SlotID.IsValid());
+				Cache.MarkSlotPreloadedRegion(SlotID, OffsetInSlot, SizeInSlot);
+				EntryOffset += SizeInSlot;
+				OffsetInSlot = 0;
+				SizeInSlot = FMath::Min<int64>(EndOffset - EntryOffset, CacheLineSize);
 			}
 			else
 			{
-				PendingRequest.Event.SafeRelease();
+				CurrentLine = StartLine;
 			}
 
-			++CurrentLine;
-			EntryOffset += SizeInSlot;
-			OffsetInSlot = 0;
-			SizeInSlot = FMath::Min<int64>(EndOffset - EntryOffset, CacheLineSize);
+			while (CurrentLine.Get() <= EndLine.Get())
+			{
+				CacheSlotID& SlotID = LineToSlot[CurrentLine.Get()];
+				if (!SlotID.IsValid())
+				{
+					// no valid slot for this line, grab a new slot from cache and start a read request
+					SlotID = AcquireSlotAndReadLine(Cache, CurrentLine, Priority);
+					LockedSlots.Add(SlotID);
+					++NumSlotsLoaded;
+				}
+
+				Cache.MarkSlotPreloadedRegion(SlotID, OffsetInSlot, SizeInSlot);
+
+				FPendingRequest& PendingRequest = LineToRequest[CurrentLine.Get()];
+				if (PendingRequest.Event && !PendingRequest.Event->IsComplete())
+				{
+					// this line has a pending async request to read data
+					// will need to wait for this request to complete before data is valid
+					CompletionEvents.Add(PendingRequest.Event);
+				}
+				else
+				{
+					PendingRequest.Event.SafeRelease();
+				}
+
+				++CurrentLine;
+				EntryOffset += SizeInSlot;
+				OffsetInSlot = 0;
+				SizeInSlot = FMath::Min<int64>(EndOffset - EntryOffset, CacheLineSize);
+			}
 		}
 	}
 
diff --git a/Engine/Source/Runtime/Core/Public/FileCache/FileCache.h b/Engine/Source/Runtime/Core/Public/FileCache/FileCache.h
--- a/Engine/Source/Runtime/Core/Public/FileCache/FileCache.h
+++ b/Engine/Source/Runtime/Core/Public/FileCache/FileCache.h
@@ -17,10 +17,7 @@ struct FFileCachePreloadEntry
 };
 
 /**
- * Thready safety note: Once created a IFileCacheHandle is assumed to be only used from a single thread.
- * (i.e. the IFileCacheHandle interface is not thread safe, and the user will need to ensure serialization).
- * Of course you can create several IFileCacheHandle's on separate threads if needed. And obviously Internally threading
- * will also be used to do async IO and cache management.
+ * All methods may be safely called from multiple threads simultaneously, unless otherwise noted
  * 
  * Also note, if you create several IFileCacheHandle's to the same file on separate threads these will be considered
  * as individual separate files from the cache point of view and thus each will have their own cache data allocated.
