TBA
 * Fix global fstructuredarchive->farchive insertion operator not closing the adapter after use
 * Add slot insertion overloads for TEnumAsByte and enum class types

diff --git a/Engine/Source/Runtime/Core/Public/Serialization/StructuredArchive.h b/Engine/Source/Runtime/Core/Public/Serialization/StructuredArchive.h
--- a/Engine/Source/Runtime/Core/Public/Serialization/StructuredArchive.h
+++ b/Engine/Source/Runtime/Core/Public/Serialization/StructuredArchive.h
@@ -262,6 +262,22 @@ public:
 	void operator << (FSoftObjectPath& Value);
 	void operator << (FLazyObjectPtr& Value);
 
+	template <typename T>
+	FORCEINLINE void operator<<(TEnumAsByte<T>& Value)
+	{
+		uint8 Tmp = (uint8)Value.GetValue();
+		*this << Tmp;
+		Value = (T)Tmp;
+	}
+
+	template <
+		typename EnumType,
+		typename = typename TEnableIf<TIsEnumClass<EnumType>::Value>::Type
+	>
+	FORCEINLINE void operator<<(EnumType& Value)
+	{
+		*this << (__underlying_type(EnumType)&)Value;
+	}
 
 	template <typename T>
 	FORCEINLINE void operator<<(TNamedAttribute<T> Item)
@@ -760,6 +776,9 @@ typename TEnableIf<
 	FArchive& Ar = Slot.GetUnderlyingArchive();
 #endif
 	Ar << Obj;
+#if WITH_TEXT_ARCHIVE_SUPPORT
+	Adapter.Close();
+#endif
 }
 
 #if !WITH_TEXT_ARCHIVE_SUPPORT
