Add missing include.

diff --git a/Engine/Source/Runtime/Core/Private/IO/IoDispatcherFileBackend.cpp b/Engine/Source/Runtime/Core/Private/IO/IoDispatcherFileBackend.cpp
--- a/Engine/Source/Runtime/Core/Private/IO/IoDispatcherFileBackend.cpp
+++ b/Engine/Source/Runtime/Core/Private/IO/IoDispatcherFileBackend.cpp
@@ -12,6 +12,7 @@
 #include "HAL/RunnableThread.h"
 #include "Misc/ScopeLock.h"
 #include "Misc/CoreDelegates.h"
+#include "Misc/Paths.h"
 
 TRACE_DECLARE_MEMORY_COUNTER(IoDispatcherTotalBytesRead, TEXT("IoDispatcher/TotalBytesRead"));
 TRACE_DECLARE_MEMORY_COUNTER(IoDispatcherTotalBytesScattered, TEXT("IoDispatcher/TotalBytesScattered"));
