IoDispatcher: Process blocks asynchronous if compressed or encrypted.
IoDispatcher: Fix link error.

diff --git a/Engine/Source/Runtime/Core/Private/IO/IoDispatcherFileBackend.cpp b/Engine/Source/Runtime/Core/Private/IO/IoDispatcherFileBackend.cpp
--- a/Engine/Source/Runtime/Core/Private/IO/IoDispatcherFileBackend.cpp
+++ b/Engine/Source/Runtime/Core/Private/IO/IoDispatcherFileBackend.cpp
@@ -67,12 +67,12 @@ private:
 	IMappedFileHandle* SharedMappedFileHandle;
 };
 
-FEncryptionKeyCache::FEncryptionKeyCache()
+FFileIoStoreEncryptionKeys::FFileIoStoreEncryptionKeys()
 {
-	FCoreDelegates::GetRegisterEncryptionKeyDelegate().BindRaw(this, &FEncryptionKeyCache::RegisterEncryptionKey);
+	FCoreDelegates::GetRegisterEncryptionKeyDelegate().BindRaw(this, &FFileIoStoreEncryptionKeys::RegisterEncryptionKey);
 }
 
-bool FEncryptionKeyCache::GetEncryptionKey(const FGuid& Guid, FAES::FAESKey& OutKey) const
+bool FFileIoStoreEncryptionKeys::GetEncryptionKey(const FGuid& Guid, FAES::FAESKey& OutKey) const
 {
 	OutKey.Reset();
 
@@ -94,7 +94,7 @@ bool FEncryptionKeyCache::GetEncryptionKey(const FGuid& Guid, FAES::FAESKey& Out
 	return false;
 }
 
-void FEncryptionKeyCache::RegisterEncryptionKey(const FGuid& Guid, const FAES::FAESKey& Key)
+void FFileIoStoreEncryptionKeys::RegisterEncryptionKey(const FGuid& Guid, const FAES::FAESKey& Key)
 {
 	{
 		FScopeLock _(&EncryptionKeysCritical);
@@ -615,7 +615,7 @@ void FFileIoStore::ProcessCompletedBlocks()
 		{
 			AllocMemoryForRequest(Scatter.Request);
 		}
-		if (BlockToDecompress->CompressionMethod.IsNone())
+		if (BlockToDecompress->CompressionMethod.IsNone() && !BlockToDecompress->EncryptionKey.IsValid())
 		{
 			ScatterBlock(BlockToDecompress, false);
 			FinalizeCompressedBlock(BlockToDecompress);
diff --git a/Engine/Source/Runtime/Core/Private/IO/IoDispatcherFileBackend.h b/Engine/Source/Runtime/Core/Private/IO/IoDispatcherFileBackend.h
--- a/Engine/Source/Runtime/Core/Private/IO/IoDispatcherFileBackend.h
+++ b/Engine/Source/Runtime/Core/Private/IO/IoDispatcherFileBackend.h
@@ -122,12 +122,12 @@ struct FFileIoStoreResolvedRequest
 	uint64 ResolvedSize;
 };
 
-class FEncryptionKeyCache
+class FFileIoStoreEncryptionKeys
 {
 public:
 	using FKeyRegisteredCallback = TFunction<void(const FGuid&, const FAES::FAESKey&)>;
 
-	FEncryptionKeyCache();
+	FFileIoStoreEncryptionKeys();
 	bool GetEncryptionKey(const FGuid& Guid, FAES::FAESKey& OutKey) const;
 	void SetKeyRegisteredCallback(FKeyRegisteredCallback&& Callback)
 	{
@@ -260,5 +260,5 @@ private:
 	FFileIoStoreRawBlock* ScheduledBlocksTail = nullptr;
 	FCriticalSection DecompressedBlocksCritical;
 	FFileIoStoreCompressedBlock* FirstDecompressedBlock = nullptr;
-	FEncryptionKeyCache EncryptionKeys;
+	FFileIoStoreEncryptionKeys EncryptionKeys;
 };
