use PLATFORM_SUPPORTS_COLORIZED_OUTPUT_DEVICE for picking log colour. Related to CL 13111174

diff --git a/Engine/Source/Runtime/Core/Private/Logging/MessageLog.cpp b/Engine/Source/Runtime/Core/Private/Logging/MessageLog.cpp
--- a/Engine/Source/Runtime/Core/Private/Logging/MessageLog.cpp
+++ b/Engine/Source/Runtime/Core/Private/Logging/MessageLog.cpp
@@ -218,7 +218,7 @@ ELogVerbosity::Type FMessageLog::GetLogVerbosity( EMessageSeverity::Type InSever
 
 const TCHAR* const FMessageLog::GetLogColor( EMessageSeverity::Type InSeverity )
 {
-#if !PLATFORM_DESKTOP
+#if !PLATFORM_SUPPORTS_COLORIZED_OUTPUT_DEVICE
 	return NULL;
 #else
 	switch(InSeverity)
