Fix LoadTextFileFromPlatformPackage and FileExistsInPlatformPackage sometimes failing in shipping builds when using FGenericPlatformMisc

diff --git a/Engine/Source/Runtime/Core/Private/GenericPlatform/GenericPlatformMisc.cpp b/Engine/Source/Runtime/Core/Private/GenericPlatform/GenericPlatformMisc.cpp
--- a/Engine/Source/Runtime/Core/Private/GenericPlatform/GenericPlatformMisc.cpp
+++ b/Engine/Source/Runtime/Core/Private/GenericPlatform/GenericPlatformMisc.cpp
@@ -1451,14 +1451,19 @@ FString FGenericPlatformMisc::LoadTextFileFromPlatformPackage(const FString& Rel
 {
 	FString Path = RootDir() / RelativePath;
 	FString Result;
-	FFileHelper::LoadFileToString(Result, *Path);
+	if (FFileHelper::LoadFileToString(Result, &IPlatformFile::GetPlatformPhysical(), *Path))
+	{
+		return Result;
+	}
+
+	Result.Empty();
 	return Result;
 }
 
 bool FGenericPlatformMisc::FileExistsInPlatformPackage(const FString& RelativePath)
 {
 	FString Path = RootDir() / RelativePath;
-	return IFileManager::Get().FileExists(*Path);
+	return IPlatformFile::GetPlatformPhysical().FileExists(*Path);
 }
 
 void FGenericPlatformMisc::TearDown()
