Add priority buckets in thread pools for better scheduling within a single pool.
  - This allows to use the GLargeThreadPool and benefit from all available cores even for low priority tasks.
  - Using this implementation, DistanceFields computations, Async resources building, AssetRegistry scanning could all run from the same pool at max cores efficiency.
  - Add synthetic micro-benchmark to isolate new implementation performance and compare it against the old one.

  BenchmarkTool.exe -Benchmark=BM_ThreadPoolOverhead

  Best 3 out of 10 runs
  Before  BM_ThreadPoolOverhead              100000 iterations took   946 ms (9.465282 us/iteration)
               BM_ThreadPoolOverhead              100000 iterations took   947 ms (9.473544 us/iteration)
               BM_ThreadPoolOverhead              100000 iterations took   948 ms (9.485538 us/iteration)

  After    BM_ThreadPoolOverhead              100000 iterations took   944 ms (9.441136 us/iteration)
              BM_ThreadPoolOverhead              100000 iterations took   940 ms (9.407212 us/iteration)
              BM_ThreadPoolOverhead              100000 iterations took   941 ms (9.410743 us/iteration)

diff --git a/Engine/Source/Programs/BenchmarkTool/Private/BenchmarkTool.cpp b/Engine/Source/Programs/BenchmarkTool/Private/BenchmarkTool.cpp
--- a/Engine/Source/Programs/BenchmarkTool/Private/BenchmarkTool.cpp
+++ b/Engine/Source/Programs/BenchmarkTool/Private/BenchmarkTool.cpp
@@ -204,8 +204,15 @@ public:
 		TArray<BenchmarkReporter::Run> RunResults;
 		RunResults.Reserve(Benchmarks.Num());
 
+		FString BenchName;
+		FParse::Value(FCommandLine::Get(), TEXT("-Benchmark="), BenchName);
 		for (auto& Bench : Benchmarks)
 		{
+			if (BenchName.Len() > 0 && BenchName != Bench->Name)
+			{
+				continue;
+			}
+
 			BenchmarkState State;
 			State.SetIterationCount(Bench->IterationCount);
 
@@ -220,7 +227,7 @@ public:
 
 		for (BenchmarkReporter::Run& Line : RunResults)
 		{
-			UE_LOG(LogBenchmarkTool, Log, 
+			UE_LOG(LogBenchmarkTool, Display, 
 				TEXT("%-30s %10ld iterations took %5ld ms (%f us/iteration)"),
 				*Line.Name,
 				Line.IterationCount,
@@ -602,12 +609,70 @@ void BM_TRefCountAssign(BenchmarkState& State)
 	}
 }
 
+// This test is probably only meaningful when comparing relative speed
+// of threadpool implementations and to profile the current one.
+void BM_ThreadPoolOverhead(BenchmarkState& State)
+{
+	struct FEventTask : public IQueuedWork
+	{
+		FEvent* Event;
+		FEventTask()  { Event = FPlatformProcess::GetSynchEventFromPool(true); }
+		~FEventTask() { FPlatformProcess::ReturnSynchEventToPool(Event); }
+		void Abandon() override { }
+	};
+
+	struct FWaitTask : public FEventTask
+	{
+		void DoThreadedWork() override { Event->Wait(); };
+	};
+
+	struct FTriggerTask : public FEventTask
+	{
+		void DoThreadedWork() override { Event->Trigger(); };
+	};
+	
+	// A single thread is enough to exercise the Queuing and Dequeuing code paths.
+	// More threads will only increase contention as the current implementation of
+	// those functions is under critical sections anyway...
+	TUniquePtr<FQueuedThreadPool> ThreadPool(FQueuedThreadPool::Allocate());
+	check(ThreadPool && ThreadPool->Create(1));
+
+	FWaitTask DefaultTask;
+	// Stall the tasks so we can benchmark the queuing code in AddQueuedWork.
+	// Otherwise, most threads will be able to execute as fast as the queuing
+	// happens and the dispatching will occur directly on each thread, 
+	// exercising a different code path.
+	DefaultTask.Event->Reset();
+
+	for (auto _ : State)
+	{
+		// Always enqueue the same empty task to avoid allocations.
+		ThreadPool->AddQueuedWork(&DefaultTask);
+	}
+
+	// Unstall the task processing so we can properly exercise the code
+	// path where all the threads need to pick another job to process.
+	// (i.e. ReturnToPoolOrGetNextJob)
+	DefaultTask.Event->Trigger();
+
+	// Enqueue a sentinel task so we know when everything has been unqueued
+	FTriggerTask LastTask;
+	ThreadPool->AddQueuedWork(&LastTask);
+	LastTask.Event->Wait();
+
+	// If run with more than 1 thread, there might be up to (number of threads) empty task 
+	// still being processed at this point, but since work is not meaningful we don't really care.
+	// The thread pool destructor will abandon any task left and clear up threads before going
+	// out of scope so we're safe to exit without any additional cleanup.
+}
+
 UE_BENCHMARK(BM_TSharedPtr)->Iterations(100000000);
 UE_BENCHMARK(BM_TRefCountPtr)->Iterations(100000000);
 UE_BENCHMARK(BM_TSharedPtr_NoTS)->Iterations(100000000);
 UE_BENCHMARK(BM_TSharedPtrAssign)->Iterations(100000000);
 UE_BENCHMARK(BM_TRefCountAssign)->Iterations(100000000);
 UE_BENCHMARK(BM_TSharedPtrAssign_NoTS)->Iterations(100000000);
+UE_BENCHMARK(BM_ThreadPoolOverhead)->Iterations(100000);
 
 //////////////////////////////////////////////////////////////////////////
 
diff --git a/Engine/Source/Runtime/Core/Private/HAL/ThreadingBase.cpp b/Engine/Source/Runtime/Core/Private/HAL/ThreadingBase.cpp
--- a/Engine/Source/Runtime/Core/Private/HAL/ThreadingBase.cpp
+++ b/Engine/Source/Runtime/Core/Private/HAL/ThreadingBase.cpp
@@ -646,7 +646,10 @@ class FQueuedThreadPoolBase : public FQueuedThreadPool
 protected:
 
 	/** The work queue to pull from. */
-	TArray<IQueuedWork*> QueuedWork;
+	TArray<TArray<IQueuedWork*>, TInlineAllocator<static_cast<int32>(EQueuedWorkPriority::Lowest) + 1>> PriorityQueuedWork;
+
+	/** The first queue to extract a work item from to avoid scanning all priorities when unqueuing. */
+	int32 FirstNonEmptyQueueIndex = 0;
 	
 	/** The thread pool to dole work out to. */
 	TArray<FQueuedThread*> QueuedThreads;
@@ -657,6 +660,9 @@ protected:
 	/** The synchronization object used to protect access to the queued work. */
 	FCriticalSection* SynchQueue;
 
+	/** Thread-safe counter to keep tab on number of queued work. */
+	TAtomic<int32> NumQueuedWork;
+
 	/** If true, indicates the destruction process has taken place. */
 	bool TimeToDie;
 
@@ -684,6 +690,8 @@ public:
 		// Presize the array so there is no extra memory allocated
 		check(QueuedThreads.Num() == 0);
 		QueuedThreads.Empty(InNumQueuedThreads);
+		NumQueuedWork = 0;
+		FirstNonEmptyQueueIndex = 0;
 
 		// Check for stack size override.
 		if( OverrideStackSize > StackSize )
@@ -726,12 +734,17 @@ public:
 				TimeToDie = 1;
 				FPlatformMisc::MemoryBarrier();
 				// Clean up all queued objects
-				for (int32 Index = 0; Index < QueuedWork.Num(); Index++)
+				for (TArray<IQueuedWork*>& Queue : PriorityQueuedWork)
 				{
-					QueuedWork[Index]->Abandon();
+					for (IQueuedWork* WorkItem : Queue)
+					{
+						WorkItem->Abandon();
+					}
 				}
+				
 				// Empty out the invalid pointers
-				QueuedWork.Empty();
+				PriorityQueuedWork.Empty();
+				NumQueuedWork = 0;
 			}
 			// wait for all threads to finish up
 			while (1)
@@ -765,14 +778,15 @@ public:
 	int32 GetNumQueuedJobs() const
 	{
 		// this is a estimate of the number of queued jobs. 
-		// no need for thread safe lock as the queuedWork array isn't moved around in memory so unless this class is being destroyed then we don't need to wrory about it
-		return QueuedWork.Num();
+		return NumQueuedWork.Load(EMemoryOrder::Relaxed);
 	}
+
 	virtual int32 GetNumThreads() const 
 	{
 		return AllThreads.Num();
 	}
-	void AddQueuedWork(IQueuedWork* InQueuedWork) override
+
+	void AddQueuedWork(IQueuedWork* InQueuedWork, EQueuedWorkPriority InQueuedWorkPriority) override
 	{
 		check(InQueuedWork != nullptr);
 
@@ -802,8 +816,18 @@ public:
 			{
 				// No thread available, queue the work to be done
 				// as soon as one does become available
-				QueuedWork.Add(InQueuedWork);
-
+				int32 QueueIndex = static_cast<int32>(InQueuedWorkPriority);
+				if (PriorityQueuedWork.Num() <= QueueIndex)
+				{
+					PriorityQueuedWork.SetNum(QueueIndex + 1);
+				}
+				
+				NumQueuedWork++;
+				if (QueueIndex < FirstNonEmptyQueueIndex)
+				{
+					FirstNonEmptyQueueIndex = QueueIndex;
+				}
+				PriorityQueuedWork[QueueIndex].Add(InQueuedWork);
 				return;
 			}
 
@@ -827,7 +851,16 @@ public:
 		check(InQueuedWork != nullptr);
 		check(SynchQueue);
 		FScopeLock sl(SynchQueue);
-		return !!QueuedWork.RemoveSingle(InQueuedWork);
+		for (int32 QueueIndex = FirstNonEmptyQueueIndex, Num = PriorityQueuedWork.Num(); QueueIndex < Num; ++QueueIndex)
+		{
+			if (PriorityQueuedWork[QueueIndex].RemoveSingle(InQueuedWork))
+			{
+				NumQueuedWork--;
+				return true;
+			}
+		}
+
+		return false;
 	}
 
 	IQueuedWork* ReturnToPoolOrGetNextJob(FQueuedThread* InQueuedThread)
@@ -838,17 +871,27 @@ public:
 		FScopeLock sl(SynchQueue);
 		if (TimeToDie)
 		{
-			check(!QueuedWork.Num());  // we better not have anything if we are dying
+			check(!NumQueuedWork);  // we better not have anything if we are dying
 		}
-		if (QueuedWork.Num() > 0)
+		
+		for (int32 QueueIndex = FirstNonEmptyQueueIndex, Num = PriorityQueuedWork.Num(); QueueIndex < Num; ++QueueIndex)
 		{
-			// Grab the oldest work in the queue. This is slower than
-			// getting the most recent but prevents work from being
-			// queued and never done
-			Work = QueuedWork[0];
-			// Remove it from the list so no one else grabs it
-			QueuedWork.RemoveAt(0, 1, /* do not allow shrinking */ false);
+			TArray<IQueuedWork*>& QueuedWork = PriorityQueuedWork[QueueIndex];
+			if (QueuedWork.Num() > 0)
+			{
+				// Grab the oldest work in the queue. This is slower than
+				// getting the most recent but prevents work from being
+				// queued and never done
+				Work = QueuedWork[0];
+				// Remove it from the list so no one else grabs it
+				QueuedWork.RemoveAt(0, 1, /* do not allow shrinking */ false);
+
+				FirstNonEmptyQueueIndex = QueueIndex;
+				NumQueuedWork--;
+				break;
+			}
 		}
+
 		if (!Work)
 		{
 			// There was no work to be done, so add the thread to the pool
diff --git a/Engine/Source/Runtime/Core/Public/Async/AsyncWork.h b/Engine/Source/Runtime/Core/Public/Async/AsyncWork.h
--- a/Engine/Source/Runtime/Core/Public/Async/AsyncWork.h
+++ b/Engine/Source/Runtime/Core/Public/Async/AsyncWork.h
@@ -229,7 +229,7 @@ class FAsyncTask
 	/* Generic start function, not called directly
 		* @param bForceSynchronous if true, this job will be started synchronously, now, on this thread
 	**/
-	void Start(bool bForceSynchronous, FQueuedThreadPool* InQueuedPool)
+	void Start(bool bForceSynchronous, FQueuedThreadPool* InQueuedPool, EQueuedWorkPriority InQueuedWorkPriority)
 	{
 		FScopeCycleCounter Scope( Task.GetStatId(), true );
 		DECLARE_SCOPE_CYCLE_COUNTER( TEXT( "FAsyncTask::Start" ), STAT_FAsyncTask_Start, STATGROUP_ThreadPoolAsyncTasks );
@@ -250,7 +250,7 @@ class FAsyncTask
 				DoneEvent = FPlatformProcess::GetSynchEventFromPool(true);
 			}
 			DoneEvent->Reset();
-			QueuedPool->AddQueuedWork(this);
+			QueuedPool->AddQueuedWork(this, InQueuedWorkPriority);
 		}
 		else 
 		{
@@ -399,17 +399,17 @@ public:
 	* Run this task on this thread
 	* @param bDoNow if true then do the job now instead of at EnsureCompletion
 	**/
-	void StartSynchronousTask()
+	void StartSynchronousTask(EQueuedWorkPriority InQueuedWorkPriority = EQueuedWorkPriority::Normal)
 	{
-		Start(true, GThreadPool);
+		Start(true, GThreadPool, InQueuedWorkPriority);
 	}
 
 	/** 
 	* Queue this task for processing by the background thread pool
 	**/
-	void StartBackgroundTask(FQueuedThreadPool* InQueuedPool = GThreadPool)
+	void StartBackgroundTask(FQueuedThreadPool* InQueuedPool = GThreadPool, EQueuedWorkPriority InQueuedWorkPriority = EQueuedWorkPriority::Normal)
 	{
-		Start(false, InQueuedPool);
+		Start(false, InQueuedPool, InQueuedWorkPriority);
 	}
 
 	/** 
diff --git a/Engine/Source/Runtime/Core/Public/Misc/QueuedThreadPool.h b/Engine/Source/Runtime/Core/Public/Misc/QueuedThreadPool.h
--- a/Engine/Source/Runtime/Core/Public/Misc/QueuedThreadPool.h
+++ b/Engine/Source/Runtime/Core/Public/Misc/QueuedThreadPool.h
@@ -7,6 +7,16 @@
 
 class IQueuedWork;
 
+/** Higher priority are picked up first by the task thread pool. */
+enum class EQueuedWorkPriority : uint8
+{
+	Highest = 0,
+	High = 1,
+	Normal = 2,
+	Low = 3,
+	Lowest = 4
+};
+
 /**
  * Interface for queued thread pools.
  *
@@ -34,10 +44,11 @@ public:
 	 * Checks to see if there is a thread available to perform the task. If not,
 	 * it queues the work for later. Otherwise it is immediately dispatched.
 	 *
-	 * @param InQueuedWork The work that needs to be done asynchronously
+	 * @param InQueuedWork         The work that needs to be done asynchronously
+	 * @param InQueuedWorkPriority The priority at which to process this task
 	 * @see RetractQueuedWork
 	 */
-	virtual void AddQueuedWork( IQueuedWork* InQueuedWork ) = 0;
+	virtual void AddQueuedWork( IQueuedWork* InQueuedWork, EQueuedWorkPriority InQueuedWorkPriority = EQueuedWorkPriority::Normal) = 0;
 
 	/**
 	 * Attempts to retract a previously queued task.
