Merge GThreadPool and GLargeThreadPool in Editor using new ThreadPool wrapping functionality
  - Supports limiting concurrency essentially giving the option of carving smaller threadpools out of a bigger one
  - Supports priority mapping function to allow per instance priority scheduling
  - In Editor, GThreadPool is now carved up from GLargeThreadPool by wrapping it and limiting concurrency
  - Can also reduce latency on the main thread pool because work items are scheduled just-in-time giving
    the chance to other wrapped queues to take turn and each make progress even if they have hundreds
    of tasks in their respective queue.

diff --git a/Engine/Source/Programs/BenchmarkTool/Private/BenchmarkTool.cpp b/Engine/Source/Programs/BenchmarkTool/Private/BenchmarkTool.cpp
--- a/Engine/Source/Programs/BenchmarkTool/Private/BenchmarkTool.cpp
+++ b/Engine/Source/Programs/BenchmarkTool/Private/BenchmarkTool.cpp
@@ -6,7 +6,7 @@
 
 #include "Templates/RefCounting.h"
 #include "Templates/SharedPointer.h"
-
+#include "Misc/QueuedThreadPoolWrapper.h"
 #include "RequiredProgramMainCPPInclude.h"
 #include <locale.h>
 #include <xutility>
@@ -609,14 +609,12 @@ void BM_TRefCountAssign(BenchmarkState& State)
 	}
 }
 
-// This test is probably only meaningful when comparing relative speed
-// of threadpool implementations and to profile the current one.
-void BM_ThreadPoolOverhead(BenchmarkState& State)
+void BM_ThreadPoolOverhead_Impl(BenchmarkState& State, FQueuedThreadPool* ThreadPool)
 {
 	struct FEventTask : public IQueuedWork
 	{
 		FEvent* Event;
-		FEventTask()  { Event = FPlatformProcess::GetSynchEventFromPool(true); }
+		FEventTask() { Event = FPlatformProcess::GetSynchEventFromPool(true); }
 		~FEventTask() { FPlatformProcess::ReturnSynchEventToPool(Event); }
 		void Abandon() override { }
 	};
@@ -630,12 +628,6 @@ void BM_ThreadPoolOverhead(BenchmarkState& State)
 	{
 		void DoThreadedWork() override { Event->Trigger(); };
 	};
-	
-	// A single thread is enough to exercise the Queuing and Dequeuing code paths.
-	// More threads will only increase contention as the current implementation of
-	// those functions is under critical sections anyway...
-	TUniquePtr<FQueuedThreadPool> ThreadPool(FQueuedThreadPool::Allocate());
-	check(ThreadPool && ThreadPool->Create(1));
 
 	FWaitTask DefaultTask;
 	// Stall the tasks so we can benchmark the queuing code in AddQueuedWork.
@@ -666,6 +658,34 @@ void BM_ThreadPoolOverhead(BenchmarkState& State)
 	// out of scope so we're safe to exit without any additional cleanup.
 }
 
+// This test is probably only meaningful when comparing relative speed
+// of threadpool implementations and to profile the current one.
+void BM_ThreadPoolOverhead(BenchmarkState& State)
+{
+	// A single thread is enough to exercise the Queuing and Dequeuing code paths.
+	// More threads will only increase contention as the current implementation of
+	// those functions is under critical sections anyway...
+	TUniquePtr<FQueuedThreadPool> ThreadPool(FQueuedThreadPool::Allocate());
+	check(ThreadPool && ThreadPool->Create(1));
+
+	BM_ThreadPoolOverhead_Impl(State, ThreadPool.Get());
+}
+
+// This test is probably only meaningful when comparing relative speed
+// of threadpool implementations and to profile the current one.
+void BM_ThreadPoolWrapperOverhead(BenchmarkState& State)
+{
+	// A single thread is enough to exercise the Queuing and Dequeuing code paths.
+	// More threads will only increase contention as the current implementation of
+	// those functions is under critical sections anyway...
+	TUniquePtr<FQueuedThreadPool> ThreadPool(FQueuedThreadPool::Allocate());
+	check(ThreadPool && ThreadPool->Create(1));
+
+	FQueuedThreadPoolWrapper ThreadPoolWrapper(ThreadPool.Get());
+
+	BM_ThreadPoolOverhead_Impl(State, &ThreadPoolWrapper);
+}
+
 UE_BENCHMARK(BM_TSharedPtr)->Iterations(100000000);
 UE_BENCHMARK(BM_TRefCountPtr)->Iterations(100000000);
 UE_BENCHMARK(BM_TSharedPtr_NoTS)->Iterations(100000000);
@@ -673,6 +693,7 @@ UE_BENCHMARK(BM_TSharedPtrAssign)->Iterations(100000000);
 UE_BENCHMARK(BM_TRefCountAssign)->Iterations(100000000);
 UE_BENCHMARK(BM_TSharedPtrAssign_NoTS)->Iterations(100000000);
 UE_BENCHMARK(BM_ThreadPoolOverhead)->Iterations(100000);
+UE_BENCHMARK(BM_ThreadPoolWrapperOverhead)->Iterations(100000);
 
 //////////////////////////////////////////////////////////////////////////
 
diff --git a/Engine/Source/Runtime/Core/Private/HAL/ThreadingBase.cpp b/Engine/Source/Runtime/Core/Private/HAL/ThreadingBase.cpp
--- a/Engine/Source/Runtime/Core/Private/HAL/ThreadingBase.cpp
+++ b/Engine/Source/Runtime/Core/Private/HAL/ThreadingBase.cpp
@@ -528,6 +528,77 @@ void FRunnableThread::FreeTls()
 	}
 }
 
+/*-----------------------------------------------------------------------------
+	FThreadPoolPriorityQueue
+-----------------------------------------------------------------------------*/
+
+void FThreadPoolPriorityQueue::Enqueue(IQueuedWork* InQueuedWork, EQueuedWorkPriority InPriority)
+{
+	int32 QueueIndex = static_cast<int32>(InPriority);
+	if (PriorityQueuedWork.Num() <= QueueIndex)
+	{
+		PriorityQueuedWork.SetNum(QueueIndex + 1);
+	}
+
+	NumQueuedWork++;
+	if (QueueIndex < FirstNonEmptyQueueIndex)
+	{
+		FirstNonEmptyQueueIndex = QueueIndex;
+	}
+	PriorityQueuedWork[QueueIndex].Add(InQueuedWork);
+}
+
+bool FThreadPoolPriorityQueue::Retract(IQueuedWork* InQueuedWork)
+{
+	for (int32 QueueIndex = FirstNonEmptyQueueIndex, Num = PriorityQueuedWork.Num(); QueueIndex < Num; ++QueueIndex)
+	{
+		if (PriorityQueuedWork[QueueIndex].RemoveSingle(InQueuedWork))
+		{
+			NumQueuedWork--;
+			return true;
+		}
+	}
+
+	return false;
+}
+
+IQueuedWork* FThreadPoolPriorityQueue::Dequeue(EQueuedWorkPriority* OutDequeuedWorkPriority)
+{
+	IQueuedWork* Work = nullptr;
+	for (int32 QueueIndex = FirstNonEmptyQueueIndex, Num = PriorityQueuedWork.Num(); QueueIndex < Num; ++QueueIndex)
+	{
+		TArray<IQueuedWork*>& QueuedWork = PriorityQueuedWork[QueueIndex];
+		if (QueuedWork.Num() > 0)
+		{
+			// Grab the oldest work in the queue. This is slower than
+			// getting the most recent but prevents work from being
+			// queued and never done
+			Work = QueuedWork[0];
+			// Remove it from the list so no one else grabs it
+			QueuedWork.RemoveAt(0, 1, /* do not allow shrinking */ false);
+
+			FirstNonEmptyQueueIndex = QueueIndex;
+			NumQueuedWork--;
+
+			if (OutDequeuedWorkPriority)
+			{
+				*OutDequeuedWorkPriority = (EQueuedWorkPriority)QueueIndex;
+			}
+
+			break;
+		}
+	}
+
+	return Work;
+}
+
+void FThreadPoolPriorityQueue::Reset()
+{
+	PriorityQueuedWork.Empty();
+	FirstNonEmptyQueueIndex = 0;
+	NumQueuedWork = 0;
+}
+
 /*-----------------------------------------------------------------------------
 	FQueuedThread
 -----------------------------------------------------------------------------*/
@@ -646,11 +717,8 @@ class FQueuedThreadPoolBase : public FQueuedThreadPool
 protected:
 
 	/** The work queue to pull from. */
-	TArray<TArray<IQueuedWork*>, TInlineAllocator<static_cast<int32>(EQueuedWorkPriority::Lowest) + 1>> PriorityQueuedWork;
+	FThreadPoolPriorityQueue QueuedWork;
 
-	/** The first queue to extract a work item from to avoid scanning all priorities when unqueuing. */
-	int32 FirstNonEmptyQueueIndex = 0;
-	
 	/** The thread pool to dole work out to. */
 	TArray<FQueuedThread*> QueuedThreads;
 
@@ -660,9 +728,6 @@ protected:
 	/** The synchronization object used to protect access to the queued work. */
 	FCriticalSection* SynchQueue;
 
-	/** Thread-safe counter to keep tab on number of queued work. */
-	TAtomic<int32> NumQueuedWork;
-
 	/** If true, indicates the destruction process has taken place. */
 	bool TimeToDie;
 
@@ -690,8 +755,7 @@ public:
 		// Presize the array so there is no extra memory allocated
 		check(QueuedThreads.Num() == 0);
 		QueuedThreads.Empty(InNumQueuedThreads);
-		NumQueuedWork = 0;
-		FirstNonEmptyQueueIndex = 0;
+		QueuedWork.Reset();
 
 		// Check for stack size override.
 		if( OverrideStackSize > StackSize )
@@ -734,17 +798,12 @@ public:
 				TimeToDie = 1;
 				FPlatformMisc::MemoryBarrier();
 				// Clean up all queued objects
-				for (TArray<IQueuedWork*>& Queue : PriorityQueuedWork)
+				while (IQueuedWork * WorkItem = QueuedWork.Dequeue())
 				{
-					for (IQueuedWork* WorkItem : Queue)
-					{
-						WorkItem->Abandon();
-					}
+					WorkItem->Abandon();
 				}
 				
-				// Empty out the invalid pointers
-				PriorityQueuedWork.Empty();
-				NumQueuedWork = 0;
+				QueuedWork.Reset();
 			}
 			// wait for all threads to finish up
 			while (1)
@@ -778,7 +837,7 @@ public:
 	int32 GetNumQueuedJobs() const
 	{
 		// this is a estimate of the number of queued jobs. 
-		return NumQueuedWork.Load(EMemoryOrder::Relaxed);
+		return QueuedWork.Num();
 	}
 
 	virtual int32 GetNumThreads() const 
@@ -816,18 +875,7 @@ public:
 			{
 				// No thread available, queue the work to be done
 				// as soon as one does become available
-				int32 QueueIndex = static_cast<int32>(InQueuedWorkPriority);
-				if (PriorityQueuedWork.Num() <= QueueIndex)
-				{
-					PriorityQueuedWork.SetNum(QueueIndex + 1);
-				}
-				
-				NumQueuedWork++;
-				if (QueueIndex < FirstNonEmptyQueueIndex)
-				{
-					FirstNonEmptyQueueIndex = QueueIndex;
-				}
-				PriorityQueuedWork[QueueIndex].Add(InQueuedWork);
+				QueuedWork.Enqueue(InQueuedWork);
 				return;
 			}
 
@@ -851,16 +899,7 @@ public:
 		check(InQueuedWork != nullptr);
 		check(SynchQueue);
 		FScopeLock sl(SynchQueue);
-		for (int32 QueueIndex = FirstNonEmptyQueueIndex, Num = PriorityQueuedWork.Num(); QueueIndex < Num; ++QueueIndex)
-		{
-			if (PriorityQueuedWork[QueueIndex].RemoveSingle(InQueuedWork))
-			{
-				NumQueuedWork--;
-				return true;
-			}
-		}
-
-		return false;
+		return QueuedWork.Retract(InQueuedWork);
 	}
 
 	IQueuedWork* ReturnToPoolOrGetNextJob(FQueuedThread* InQueuedThread)
@@ -871,26 +910,10 @@ public:
 		FScopeLock sl(SynchQueue);
 		if (TimeToDie)
 		{
-			check(!NumQueuedWork);  // we better not have anything if we are dying
+			check(!QueuedWork.Num());  // we better not have anything if we are dying
 		}
 		
-		for (int32 QueueIndex = FirstNonEmptyQueueIndex, Num = PriorityQueuedWork.Num(); QueueIndex < Num; ++QueueIndex)
-		{
-			TArray<IQueuedWork*>& QueuedWork = PriorityQueuedWork[QueueIndex];
-			if (QueuedWork.Num() > 0)
-			{
-				// Grab the oldest work in the queue. This is slower than
-				// getting the most recent but prevents work from being
-				// queued and never done
-				Work = QueuedWork[0];
-				// Remove it from the list so no one else grabs it
-				QueuedWork.RemoveAt(0, 1, /* do not allow shrinking */ false);
-
-				FirstNonEmptyQueueIndex = QueueIndex;
-				NumQueuedWork--;
-				break;
-			}
-		}
+		Work = QueuedWork.Dequeue();
 
 		if (!Work)
 		{
diff --git a/Engine/Source/Runtime/Core/Private/Misc/QueuedThreadPoolWrapper.cpp b/Engine/Source/Runtime/Core/Private/Misc/QueuedThreadPoolWrapper.cpp
new file mode 100644
--- /dev/null
+++ b/Engine/Source/Runtime/Core/Private/Misc/QueuedThreadPoolWrapper.cpp
@@ -0,0 +1,180 @@
+// Copyright Epic Games, Inc. All Rights Reserved.
+
+#include "Misc/QueuedThreadPoolWrapper.h"
+#include "Misc/ScopeRWLock.h"
+#include "Misc/IQueuedWork.h"
+#include "HAL/PlatformProcess.h"
+
+struct FQueuedThreadPoolWrapper::FScheduledWork : public IQueuedWork
+{
+	FScheduledWork(FQueuedThreadPoolWrapper* InParentPool)
+		: ParentPool(InParentPool)
+	{
+	}
+
+	void DoThreadedWork() override
+	{
+		InnerWork->DoThreadedWork();
+		InnerWork = nullptr;
+		ParentPool->Schedule(this);
+	}
+
+	virtual void Abandon() override
+	{
+		InnerWork->Abandon();
+		InnerWork = nullptr;
+		ParentPool->Schedule(this);
+	}
+
+	IQueuedWork* InnerWork = nullptr;
+
+private:
+	FQueuedThreadPoolWrapper* ParentPool;
+};
+
+FQueuedThreadPoolWrapper::FQueuedThreadPoolWrapper(FQueuedThreadPool* InWrappedQueuedThreadPool, int32 InMaxConcurrency, TFunction<EQueuedWorkPriority(EQueuedWorkPriority)> InPriorityMapper)
+	: PriorityMapper(InPriorityMapper)
+	, WrappedQueuedThreadPool(InWrappedQueuedThreadPool)
+	, MaxConcurrency(InMaxConcurrency == -1 ? InWrappedQueuedThreadPool->GetNumThreads() : InMaxConcurrency)
+	, MaxTaskToSchedule(-1)
+	, CurrentConcurrency(0)
+{
+}
+
+FQueuedThreadPoolWrapper::~FQueuedThreadPoolWrapper()
+{
+	Destroy();
+}
+
+bool FQueuedThreadPoolWrapper::Create(uint32 InNumQueuedThreads, uint32 StackSize, EThreadPriority ThreadPriority)
+{
+	return true;
+}
+
+void FQueuedThreadPoolWrapper::Destroy()
+{
+	{
+		FRWScopeLock ScopeLock(Lock, SLT_Write);
+		// Clean up all queued objects
+		while (IQueuedWork* WorkItem = QueuedWork.Dequeue())
+		{
+			WorkItem->Abandon();
+		}
+
+		QueuedWork.Reset();
+	}
+
+	// We can't delete our WorkPool elements
+	// if they're still referenced by a threadpool.
+	// No choice to wait until they're all finished.
+	if (CurrentConcurrency)
+	{
+		TRACE_CPUPROFILER_EVENT_SCOPE(FQueuedThreadPoolWrapper::DestroyWait);
+		while (CurrentConcurrency)
+		{
+			FPlatformProcess::Sleep(0.1f);
+		}
+	}
+
+	{
+		FRWScopeLock ScopeLock(Lock, SLT_Write);
+		for (FScheduledWork* ScheduledWork : WorkPool)
+		{
+			check(ScheduledWork->InnerWork == nullptr);
+			delete ScheduledWork;
+		}
+		WorkPool.Empty();
+	}
+}
+
+void FQueuedThreadPoolWrapper::SetMaxConcurrency(int32 InMaxConcurrency)
+{
+	MaxConcurrency = InMaxConcurrency == -1 ? WrappedQueuedThreadPool->GetNumThreads() : InMaxConcurrency;
+
+	// In case we just increased the concurrency, we might need to schedule new tasks
+	Schedule();
+}
+
+void FQueuedThreadPoolWrapper::Pause()
+{
+	FRWScopeLock ScopeLock(Lock, SLT_Write);
+	MaxTaskToSchedule = 0;
+}
+
+void FQueuedThreadPoolWrapper::Resume(int32 InNumQueuedWork)
+{
+	{
+		FRWScopeLock ScopeLock(Lock, SLT_Write);
+		MaxTaskToSchedule = InNumQueuedWork;
+	}
+
+	Schedule();
+}
+
+void FQueuedThreadPoolWrapper::AddQueuedWork(IQueuedWork* InQueuedWork, EQueuedWorkPriority InPriority)
+{
+	{
+		FRWScopeLock ScopeLock(Lock, SLT_Write);
+		QueuedWork.Enqueue(InQueuedWork, InPriority);
+	}
+
+	Schedule();
+}
+
+bool FQueuedThreadPoolWrapper::RetractQueuedWork(IQueuedWork* InQueuedWork)
+{
+	FRWScopeLock ScopeLock(Lock, SLT_Write);
+	return QueuedWork.Retract(InQueuedWork);
+}
+
+int32 FQueuedThreadPoolWrapper::GetNumThreads() const
+{
+	return MaxConcurrency;
+}
+
+void FQueuedThreadPoolWrapper::Schedule(FScheduledWork* Work)
+{
+	FRWScopeLock ScopeLock(Lock, SLT_Write);
+
+	// It's important to reduce the current concurrency before entering the loop
+	// especially when MaxConcurrency is 1 to ensure new work is scheduled and don't 
+	// end up never being called again.
+	if (Work != nullptr)
+	{
+		CurrentConcurrency--;
+		WorkPool.Push(Work);
+		Work = nullptr;
+	}
+
+	while (CurrentConcurrency < MaxConcurrency.Load(EMemoryOrder::Relaxed) && (MaxTaskToSchedule == -1 || MaxTaskToSchedule > 0))
+	{
+		EQueuedWorkPriority WorkPriority;
+		IQueuedWork* InnerWork = QueuedWork.Dequeue(&WorkPriority);
+
+		if (InnerWork)
+		{
+			CurrentConcurrency++;
+			if (WorkPool.Num() > 0)
+			{
+				Work = WorkPool.Pop(false);
+			}
+			else
+			{
+				Work = new FScheduledWork(this);
+			}
+
+			Work->InnerWork = InnerWork;
+			WrappedQueuedThreadPool->AddQueuedWork(Work, PriorityMapper(WorkPriority));
+			Work = nullptr;
+
+			if (MaxTaskToSchedule > 0)
+			{
+				MaxTaskToSchedule--;
+			}
+		}
+		else
+		{
+			break;
+		}
+	}
+}
diff --git a/Engine/Source/Runtime/Core/Public/Misc/QueuedThreadPool.h b/Engine/Source/Runtime/Core/Public/Misc/QueuedThreadPool.h
--- a/Engine/Source/Runtime/Core/Public/Misc/QueuedThreadPool.h
+++ b/Engine/Source/Runtime/Core/Public/Misc/QueuedThreadPool.h
@@ -17,6 +17,45 @@ enum class EQueuedWorkPriority : uint8
 	Lowest = 4
 };
 
+/** 
+ *  Priority Queue tailored for FQueuedThreadPool implementation
+ *
+ *  This class is NOT thread-safe and must be properly protected.
+ */
+class CORE_API FThreadPoolPriorityQueue
+{
+public:
+	/**
+	 * Enqueue a work item at specified priority
+	 */
+	void Enqueue(IQueuedWork* InQueuedWork, EQueuedWorkPriority InPriority = EQueuedWorkPriority::Normal);
+
+	/**
+	 * Search and remove a queued work item from the list
+	 */
+	bool Retract(IQueuedWork* InQueuedWork);
+
+	/**
+	 * Get the next work item in priority order.
+	 */
+	IQueuedWork* Dequeue(EQueuedWorkPriority* OutDequeuedWorkPriority = nullptr);
+
+	/**
+	 * Empty the queue.
+	 */
+	void Reset();
+
+	/**
+	 * Get the total number of queued items.
+	 */
+	int32 Num() const { return NumQueuedWork; }
+private:
+	/** The first queue to extract a work item from to avoid scanning all priorities when unqueuing. */
+	int32 FirstNonEmptyQueueIndex = 0;
+	TArray<TArray<IQueuedWork*>, TInlineAllocator<static_cast<int32>(EQueuedWorkPriority::Lowest) + 1>> PriorityQueuedWork;
+	TAtomic<int32> NumQueuedWork;
+};
+
 /**
  * Interface for queued thread pools.
  *
@@ -86,7 +125,6 @@ public:
 	static uint32 OverrideStackSize;
 };
 
-
 /**
  *  Global thread pool for shared async operations
  */
diff --git a/Engine/Source/Runtime/Core/Public/Misc/QueuedThreadPoolWrapper.h b/Engine/Source/Runtime/Core/Public/Misc/QueuedThreadPoolWrapper.h
new file mode 100644
--- /dev/null
+++ b/Engine/Source/Runtime/Core/Public/Misc/QueuedThreadPoolWrapper.h
@@ -0,0 +1,61 @@
+// Copyright Epic Games, Inc. All Rights Reserved.
+
+#pragma once
+
+#include "CoreMinimal.h"
+#include "GenericPlatform/GenericPlatformAffinity.h"
+#include "QueuedThreadPool.h"
+#include "ScopeRWLock.h"
+
+/** ThreadPool wrapper implementation allowing to schedule
+  * up to MaxConcurrency tasks at a time making sub-partitioning
+  * another thread-pool a breeze and allowing more fine-grained control 
+  * over scheduling by effectively giving another set of priorities.
+  */
+class CORE_API FQueuedThreadPoolWrapper : public FQueuedThreadPool
+{
+public:
+	/**
+	 * InWrappedQueuedThreadPool  Underlying thread pool to schedule task to.
+	 * InMaxConcurrency           Maximum number of concurrent tasks allowed, -1 will limit concurrency to number of threads available in the underlying thread pool.
+	 * InPriorityMapper           Thread-safe function used to map any priority from this Queue to the priority that should be used when scheduling the task on the underlying thread pool.
+	 */
+	FQueuedThreadPoolWrapper(FQueuedThreadPool* InWrappedQueuedThreadPool, int32 InMaxConcurrency = -1, TFunction<EQueuedWorkPriority (EQueuedWorkPriority)> InPriorityMapper = [](EQueuedWorkPriority InPriority) { return InPriority; });
+	~FQueuedThreadPoolWrapper();
+
+	/**
+	 *  Queued task are not scheduled against the wrapped thread-pool until resumed
+	 */
+	void Pause();
+
+	/**
+	 *  Resume a specified amount of queued work, or -1 to unpause.
+	 */
+	void Resume(int32 InNumQueuedWork = -1);
+
+	/**
+	 *  Dynamically adjust the maximum number of concurrent tasks, -1 for unlimited.
+	 */
+	void SetMaxConcurrency(int32 MaxConcurrency = -1);
+
+	void AddQueuedWork(IQueuedWork* InQueuedWork, EQueuedWorkPriority InPriority = EQueuedWorkPriority::Normal) override;
+	bool RetractQueuedWork(IQueuedWork* InQueuedWork) override;
+	int32 GetNumThreads() const override;
+
+private:
+	struct FScheduledWork;
+
+	bool Create(uint32 InNumQueuedThreads, uint32 StackSize, EThreadPriority ThreadPriority) override;
+	void Destroy() override;
+	void Schedule(FScheduledWork* Work = nullptr);
+
+	FRWLock Lock;
+	TFunction<EQueuedWorkPriority(EQueuedWorkPriority)> PriorityMapper;
+	FThreadPoolPriorityQueue QueuedWork;
+	FQueuedThreadPool* WrappedQueuedThreadPool;
+	TArray<FScheduledWork*> WorkPool;
+	TAtomic<int32> MaxConcurrency;
+	int32 MaxTaskToSchedule;
+	TAtomic<int32> CurrentConcurrency;
+	EQueuedWorkPriority WrappedQueuePriority;
+};
diff --git a/Engine/Source/Runtime/Launch/Private/LaunchEngineLoop.cpp b/Engine/Source/Runtime/Launch/Private/LaunchEngineLoop.cpp
--- a/Engine/Source/Runtime/Launch/Private/LaunchEngineLoop.cpp
+++ b/Engine/Source/Runtime/Launch/Private/LaunchEngineLoop.cpp
@@ -89,6 +89,7 @@
 	#include "Settings/EditorExperimentalSettings.h"
 	#include "Interfaces/IEditorStyleModule.h"
 	#include "PIEPreviewDeviceProfileSelectorModule.h"
+	#include "Misc/QueuedThreadPoolWrapper.h"
 
 	#if PLATFORM_WINDOWS
 		#include "Windows/AllowWindowsPlatformTypes.h"
@@ -2154,7 +2155,28 @@ int32 FEngineLoop::PreInitPreStartupScreen(const TCHAR* CmdLine)
 			StackSize = 1000;
 		}
 
+#if WITH_EDITOR
+		{
+			TRACE_THREAD_GROUP_SCOPE("LargeThreadPool");
+			// when we are in the editor we like to do things like build lighting and such
+			// this thread pool can be used for those purposes
+			GLargeThreadPool = FQueuedThreadPool::Allocate();
+			int32 NumThreadsInLargeThreadPool = FMath::Max(FPlatformMisc::NumberOfCoresIncludingHyperthreads() - 2, 2);
+
+			verify(GLargeThreadPool->Create(NumThreadsInLargeThreadPool, StackSize * 1024));
+
+			int32 NumThreadsInThreadPool = FPlatformMisc::NumberOfWorkerThreadsToSpawn();
 
+			// we are only going to give dedicated servers one pool thread
+			if (FPlatformProperties::IsServerOnly())
+			{
+				NumThreadsInThreadPool = 1;
+			}
+
+			// GThreadPool will schedule on the LargeThreadPool but limit max concurrency to the given number.
+			GThreadPool = new FQueuedThreadPoolWrapper(GLargeThreadPool, NumThreadsInThreadPool);
+		}
+#else
 		{
 			TRACE_THREAD_GROUP_SCOPE("ThreadPool");
 			GThreadPool = FQueuedThreadPool::Allocate();
@@ -2167,6 +2189,7 @@ int32 FEngineLoop::PreInitPreStartupScreen(const TCHAR* CmdLine)
 			}
 			verify(GThreadPool->Create(NumThreadsInThreadPool, StackSize * 1024, TPri_SlightlyBelowNormal));
 		}
+#endif
 		{
 			TRACE_THREAD_GROUP_SCOPE("BackgroundThreadPool");
 			GBackgroundPriorityThreadPool = FQueuedThreadPool::Allocate();
@@ -2179,17 +2202,6 @@ int32 FEngineLoop::PreInitPreStartupScreen(const TCHAR* CmdLine)
 			verify(GBackgroundPriorityThreadPool->Create(NumThreadsInThreadPool, StackSize * 1024, TPri_Lowest));
 		}
 
-#if WITH_EDITOR
-		{
-			TRACE_THREAD_GROUP_SCOPE("LargeThreadPool");
-			// when we are in the editor we like to do things like build lighting and such
-			// this thread pool can be used for those purposes
-			GLargeThreadPool = FQueuedThreadPool::Allocate();
-			int32 NumThreadsInLargeThreadPool = FMath::Max(FPlatformMisc::NumberOfCoresIncludingHyperthreads() - 2, 2);
-
-			verify(GLargeThreadPool->Create(NumThreadsInLargeThreadPool, StackSize * 1024));
-		}
-#endif
 	}
 
 #if WITH_APPLICATION_CORE
