- Delayed Turnkey SDK querying until later so it can check if it's a commandlet, in which case it won't call UAT since UAT could be calling the commandlet
- Added a helper function that can take a commandline and modify it for the three host platforms (add .bat/.sh, run a .exe through mono, etc) to reduce the mess of if WINDOWS / elif MAC / etc stuff

diff --git a/Engine/Source/Developer/TargetPlatform/Private/TargetPlatformManagerModule.cpp b/Engine/Source/Developer/TargetPlatform/Private/TargetPlatformManagerModule.cpp
--- a/Engine/Source/Developer/TargetPlatform/Private/TargetPlatformManagerModule.cpp
+++ b/Engine/Source/Developer/TargetPlatform/Private/TargetPlatformManagerModule.cpp
@@ -990,33 +990,15 @@ protected:
 		DECLARE_SCOPE_CYCLE_COUNTER( TEXT( "FTargetPlatformManagerModule::SetupSDKStatus" ), STAT_FTargetPlatformManagerModule_SetupSDKStatus, STATGROUP_TargetPlatform );
 
 		// run UBT with -validate -allplatforms and read the output
-		FString CmdExe, CommandLine;
-		
-		if (PLATFORM_MAC)
-		{
-			CmdExe = TEXT("/bin/sh");
-			FString ScriptPath = FPaths::ConvertRelativePathToFull(FPaths::EngineDir() / TEXT("Build/BatchFiles/Mac/RunMono.sh"));
-			CommandLine = TEXT("\"") + ScriptPath + TEXT("\" \"") + FPaths::ConvertRelativePathToFull(FPaths::EngineDir() / TEXT("Binaries/DotNET/UnrealBuildTool.exe")) + TEXT("\" -Mode=ValidatePlatforms");
-		}
-		else if (PLATFORM_WINDOWS)
-		{
-			CmdExe = FPaths::ConvertRelativePathToFull(FPaths::EngineDir() / TEXT("Binaries/DotNET/UnrealBuildTool.exe"));
-			CommandLine = TEXT("-Mode=ValidatePlatforms");
-		}
-		else if (PLATFORM_LINUX)
-		{
-			CmdExe = TEXT("/bin/bash");	// bash and not sh because of pushd
-			FString ScriptPath = FPaths::ConvertRelativePathToFull(FPaths::EngineDir() / TEXT("Build/BatchFiles/Linux/RunMono.sh"));
-			CommandLine = TEXT("\"") + ScriptPath + TEXT("\" \"") + FPaths::ConvertRelativePathToFull(FPaths::EngineDir() / TEXT("Binaries/DotNET/UnrealBuildTool.exe")) + TEXT("\" -Mode=ValidatePlatforms");
-		}
-		else
-		{
-			checkf(false, TEXT("FTargetPlatformManagerModule::SetupSDKStatus(): Unsupported platform!"));
-		}
-
+		FString CmdExe = TEXT("{EngineDir}/Binaries/DotNET/UnrealBuildTool.exe");
+		FString CommandLine = TEXT("-Mode=ValidatePlatforms");
 		// Allow for only a subset of platforms to be reparsed - needed when kicking a change from the UI
 		CommandLine += TargetPlatforms.IsEmpty() ? TEXT(" -allplatforms") : (TEXT(" -platforms=") + TargetPlatforms);
 
+		// convert into appropriate calls for the current platform
+		FPlatformProcess::ModifyCreateProcParams(CmdExe, CommandLine, FGenericPlatformProcess::ECreateProcHelperFlags::None);
+
+
 		TSharedPtr<FMonitoredProcess> UBTProcess = MakeShareable(new FMonitoredProcess(CmdExe, CommandLine, true));
 		UBTProcess->OnOutput().BindStatic(&FTargetPlatformManagerModule::OnStatusOutput);
 		SDKStatusMessage = TEXT("");
diff --git a/Engine/Source/Runtime/Core/Private/GenericPlatform/GenericPlatformProcess.cpp b/Engine/Source/Runtime/Core/Private/GenericPlatform/GenericPlatformProcess.cpp
--- a/Engine/Source/Runtime/Core/Private/GenericPlatform/GenericPlatformProcess.cpp
+++ b/Engine/Source/Runtime/Core/Private/GenericPlatform/GenericPlatformProcess.cpp
@@ -256,6 +256,68 @@ void FGenericPlatformProcess::TerminateProc( FProcHandle & ProcessHandle, bool K
 	UE_LOG(LogHAL, Fatal, TEXT("FGenericPlatformProcess::TerminateProc not implemented on this platform"));
 }
 
+static FString ReplaceProcHelperVariables(const FString& Str)
+{
+	return Str.Replace(TEXT("{Platform}"), ANSI_TO_TCHAR(FPlatformProperties::IniPlatformName())).
+		Replace(TEXT("{BaseDir}"), FPlatformProcess::BaseDir()).
+		Replace(TEXT("{EngineDir}"), *FPaths::EngineDir()).
+		Replace(TEXT("{ProjectDir}"), *FPaths::ProjectDir());
+}
+
+void FGenericPlatformProcess::ModifyCreateProcParams(FString& InOutURL, FString& InOutParams, ECreateProcHelperFlags Flags)
+{
+	FString Command = ReplaceProcHelperVariables(InOutURL);
+	FString Arguments = ReplaceProcHelperVariables(InOutParams);
+
+	if ((Flags & ECreateProcHelperFlags::AppendPlatformScriptExtension) != ECreateProcHelperFlags::None)
+	{
+#if PLATFORM_WINDOWS
+		Command += TEXT(".bat");
+#elif PLATFORM_MAC || PLATFORM_LINUX
+		Command += TEXT(".sh");
+#else
+	#error Unknown platform - override this function!
+#endif
+	}
+
+	Command = FPaths::ConvertRelativePathToFull(Command);
+
+	// wrap in quotes if needed
+	if (Command.Contains(TEXT(" ")))
+	{
+		Command = FString::Printf(TEXT("\"%s\""), *Command);
+	}
+
+	// Mac and Linux need to run .exe files through Mono
+#if PLATFORM_MAC || PLATFORM_LINUX
+	if (Command.EndsWith(TEXT(".exe")))
+	{	
+		Command = FPaths::EngineDir() / TEXT("Build/BatchFiles/Mac/RunMono.sh ") + Command;
+		Flags = Flags | ECreateProcHelperFlags::RunThroughShell;
+	}
+#endif
+
+
+
+	if ((Flags & ECreateProcHelperFlags::RunThroughShell) != ECreateProcHelperFlags::None)
+	{
+#if PLATFORM_WINDOWS
+		InOutURL = TEXT("cmd.exe");
+		InOutParams = FString::Printf(TEXT("/c %s %s"), *Command, *Arguments);
+#elif PLATFORM_MAC || PLATFORM_LINUX
+		InOutURL = TEXT("/bin/bash");
+		InOutParams = FString::Printf(TEXT("-c '%s %s'"), *Command, *Arguments);
+#endif
+	}
+	else
+	{
+		// otherwise, we have done all we need to do, just return them!
+		InOutURL = Command;
+		InOutParams = Arguments;
+	}
+}
+
+
 FGenericPlatformProcess::EWaitAndForkResult FGenericPlatformProcess::WaitAndFork()
 {
 	UE_LOG(LogHAL, Fatal, TEXT("FGenericPlatformProcess::WaitAndFork not implemented on this platform"));
diff --git a/Engine/Source/Runtime/Core/Private/Misc/DataDrivenPlatformInfoRegistry.cpp b/Engine/Source/Runtime/Core/Private/Misc/DataDrivenPlatformInfoRegistry.cpp
--- a/Engine/Source/Runtime/Core/Private/Misc/DataDrivenPlatformInfoRegistry.cpp
+++ b/Engine/Source/Runtime/Core/Private/Misc/DataDrivenPlatformInfoRegistry.cpp
@@ -13,6 +13,13 @@
 
 TMap<FString, FDataDrivenPlatformInfoRegistry::FPlatformInfo> FDataDrivenPlatformInfoRegistry::DataDrivenPlatforms;
 
+#if DDPI_HAS_EXTENDED_PLATFORMINFO_DATA
+// delay the kick off of running UAT so we can check IsRunningCommandlet()
+FDelayedAutoRegisterHelper GPlatformInfoInit(EDelayedRegisterRunPhase::TaskGraphSystemReady, []()
+{
+	FDataDrivenPlatformInfoRegistry::UpdateSdkStatus();
+});
+#endif
 
 static const TArray<FString>& GetDataDrivenIniFilenames()
 {
@@ -247,11 +254,6 @@ const TMap<FString, FDataDrivenPlatformInfoRegistry::FPlatformInfo>& FDataDriven
 				It.Value.IniParentChain.Insert(CurrentPlatform, 0);
 			}
 		}
-
-
-#if DDPI_HAS_EXTENDED_PLATFORMINFO_DATA
-		UpdateSdkStatus();
-#endif
 	}
 
 	return DataDrivenPlatforms;
@@ -309,15 +311,37 @@ static FString ConvertToDDPIDeviceId(const FString& DeviceId)
 
 void FDataDrivenPlatformInfoRegistry::UpdateSdkStatus()
 {
+	// don't run UAT from commandlets (like the cooker) that are often launched from UAT and this will go poorly
+	if (IsRunningCommandlet())
+	{
+		for (auto& It : DataDrivenPlatforms)
+		{
+			It.Value.SdkStatus = DDPIPlatformSdkStatus::Unknown;
+
+			// reset the per-device status when querying general Sdk status
+			It.Value.ClearDeviceStatus();
+		}
+
+		return;
+	}
+
+
 	FString AutomationTool, ReportFilename;
 	PrepForTurnkeyReport(AutomationTool, ReportFilename);
 
 	// run Turnkey to get all of the platform statuses
-	FString Commandline = FString::Printf(TEXT("/c %s Turnkey -command=VerifySdk -platform=%s -ReportFilename=\"%s\""),
-		*AutomationTool,
+// 	FString ProcessURL = TEXT("{EngineDir}/Binaries/DotNET/AutomationTool.exe");
+// 	FString ProcessArguments = FString::Printf(TEXT("Turnkey -command=VerifySdk -platform=%s -ReportFilename=\"%s\""),
+// 		*FString::JoinBy(DataDrivenPlatforms, TEXT("+"), [](TPair<FString, FDataDrivenPlatformInfoRegistry::FPlatformInfo> Pair) { return ConvertToUATPlatform(Pair.Key); }),
+// 		*ReportFilename);
+
+	FString ProcessURL = TEXT("{EngineDir}/Build/BatchFiles/RunuAT");
+	FString ProcessArguments = FString::Printf(TEXT("Turnkey -command=VerifySdk -platform=%s -ReportFilename=\"%s\""),
 		*FString::JoinBy(DataDrivenPlatforms, TEXT("+"), [](TPair<FString, FDataDrivenPlatformInfoRegistry::FPlatformInfo> Pair) { return ConvertToUATPlatform(Pair.Key); }),
 		*ReportFilename);
 
+	// convert into appropriate calls for the current platform
+	FPlatformProcess::ModifyCreateProcParams(ProcessURL, ProcessArguments, FGenericPlatformProcess::ECreateProcHelperFlags::RunThroughShell | FGenericPlatformProcess::ECreateProcHelperFlags::AppendPlatformScriptExtension);
 
 	// reset status to unknown
 	for (auto& It : DataDrivenPlatforms)
@@ -328,7 +352,7 @@ void FDataDrivenPlatformInfoRegistry::UpdateSdkStatus()
 		It.Value.ClearDeviceStatus();
 	}
 
-	FMonitoredProcess* TurnkeyProcess = new FMonitoredProcess(TEXT("cmd.exe"), Commandline, true, false);
+	FMonitoredProcess* TurnkeyProcess = new FMonitoredProcess(ProcessURL, ProcessArguments, true, false);
 	TurnkeyProcess->OnCompleted().BindLambda([ReportFilename, TurnkeyProcess](int32 ExitCode)
 	{
 		AsyncTask(ENamedThreads::GameThread, [ReportFilename, TurnkeyProcess, ExitCode]()
diff --git a/Engine/Source/Runtime/Core/Public/GenericPlatform/GenericPlatformProcess.h b/Engine/Source/Runtime/Core/Public/GenericPlatform/GenericPlatformProcess.h
--- a/Engine/Source/Runtime/Core/Public/GenericPlatform/GenericPlatformProcess.h
+++ b/Engine/Source/Runtime/Core/Public/GenericPlatform/GenericPlatformProcess.h
@@ -407,6 +407,36 @@ struct CORE_API FGenericPlatformProcess
 	 */
 	static void TerminateProc( FProcHandle & ProcessHandle, bool KillTree = false );
 
+
+	enum class ECreateProcHelperFlags : uint8
+	{
+		None = 0,
+		// If this is set, the command will be run via cmd.exe, sh, or similar
+		RunThroughShell = 1,
+		// If this is set, the Command will have a platform-specific script extension appended (.bat, .sh, etc)
+		AppendPlatformScriptExtension = 2,
+	};
+
+	/**
+	 * Generates the URL and Parms fields for CreateProc (or FMonitoredProcess) usable on any (generally host) platform. Target platforms that need this will likely need to override it.
+	 *
+	 * Flags are used to determine out to modify the passed in InURL and InParams
+	 * This function will replace instances (in Command or Arguments) of a few variables in {} with their values:
+	 *    {Platform} replaced with the name of the Host platform (Windows, Mac, Linux)
+	 *    {BaseDir} replaced with path to root directory (one up from Engine)
+	 *    {EngineDir} replaced with path to engine directory
+	 *    {ProjectDir} replaced with path to project directory
+	 *
+	 * Ex: CreateProcHelper("{BaseDir}/RunUAT", "MyScript -param", ECreateProcHelperFlags::RunThroughShell | ECreateProcHelperFlags::AppendPlatformScriptExtension)
+	 * Would run "cmd.exe /c RunUAT.bat" on Windows and "sh -c RunUAT.sh" on Mac and Linux
+	 *
+	 * @param InOutURL URL that will be modified and passed to CreateProc or FMonitoredProcess or the like
+	 * @param InOutParams Params to the program that will be modified
+	 * @param Flags to control how the strings are modified
+	 *
+	 */
+	static void ModifyCreateProcParams(FString& InOutURL, FString& InOutParams, ECreateProcHelperFlags Flags);
+
 	enum class EWaitAndForkResult : uint8
 	{
 		Error,
@@ -669,3 +699,4 @@ struct CORE_API FGenericPlatformProcess
 };
 
 
+ENUM_CLASS_FLAGS(FGenericPlatformProcess::ECreateProcHelperFlags);
