- Fixed CIS errors

diff --git a/Engine/Source/Runtime/Core/Private/GenericPlatform/GenericPlatformProcess.cpp b/Engine/Source/Runtime/Core/Private/GenericPlatform/GenericPlatformProcess.cpp
--- a/Engine/Source/Runtime/Core/Private/GenericPlatform/GenericPlatformProcess.cpp
+++ b/Engine/Source/Runtime/Core/Private/GenericPlatform/GenericPlatformProcess.cpp
@@ -276,7 +276,7 @@ void FGenericPlatformProcess::ModifyCreateProcParams(FString& InOutURL, FString&
 #elif PLATFORM_MAC || PLATFORM_LINUX
 		Command += TEXT(".sh");
 #else
-	#error Unknown platform - override this function!
+		UE_LOG(LogInit, Fatal, TEXT("ModifyCreateProcParams being called on a platform that is not supported"));
 #endif
 	}
 
diff --git a/Engine/Source/Runtime/Core/Public/GenericPlatform/GenericPlatformProcess.h b/Engine/Source/Runtime/Core/Public/GenericPlatform/GenericPlatformProcess.h
--- a/Engine/Source/Runtime/Core/Public/GenericPlatform/GenericPlatformProcess.h
+++ b/Engine/Source/Runtime/Core/Public/GenericPlatform/GenericPlatformProcess.h
@@ -731,5 +731,5 @@ struct CORE_API FGenericPlatformProcess
 	}
 };
 
-
+#include "Misc/EnumClassFlags.h"
 ENUM_CLASS_FLAGS(FGenericPlatformProcess::ECreateProcHelperFlags);
