- Post merge fixups

diff --git a/Engine/Source/Runtime/Core/Private/Misc/QueuedThreadPoolWrapper.cpp b/Engine/Source/Runtime/Core/Private/Misc/QueuedThreadPoolWrapper.cpp
--- a/Engine/Source/Runtime/Core/Private/Misc/QueuedThreadPoolWrapper.cpp
+++ b/Engine/Source/Runtime/Core/Private/Misc/QueuedThreadPoolWrapper.cpp
@@ -46,7 +46,7 @@ FQueuedThreadPoolWrapper::~FQueuedThreadPoolWrapper()
 	Destroy();
 }
 
-bool FQueuedThreadPoolWrapper::Create(uint32 InNumQueuedThreads, uint32 StackSize, EThreadPriority ThreadPriority)
+bool FQueuedThreadPoolWrapper::Create(uint32 InNumQueuedThreads, uint32 StackSize, EThreadPriority ThreadPriority, const TCHAR* Name)
 {
 	return true;
 }
diff --git a/Engine/Source/Runtime/Core/Public/Misc/QueuedThreadPoolWrapper.h b/Engine/Source/Runtime/Core/Public/Misc/QueuedThreadPoolWrapper.h
--- a/Engine/Source/Runtime/Core/Public/Misc/QueuedThreadPoolWrapper.h
+++ b/Engine/Source/Runtime/Core/Public/Misc/QueuedThreadPoolWrapper.h
@@ -45,7 +45,7 @@ public:
 private:
 	struct FScheduledWork;
 
-	bool Create(uint32 InNumQueuedThreads, uint32 StackSize, EThreadPriority ThreadPriority) override;
+	bool Create(uint32 InNumQueuedThreads, uint32 StackSize, EThreadPriority ThreadPriority, const TCHAR* Name) override;
 	void Destroy() override;
 	void Schedule(FScheduledWork* Work = nullptr);
 
