- Non-unity CIS fix

diff --git a/Engine/Source/Runtime/Core/Public/Misc/DataDrivenPlatformInfoRegistry.h b/Engine/Source/Runtime/Core/Public/Misc/DataDrivenPlatformInfoRegistry.h
--- a/Engine/Source/Runtime/Core/Public/Misc/DataDrivenPlatformInfoRegistry.h
+++ b/Engine/Source/Runtime/Core/Public/Misc/DataDrivenPlatformInfoRegistry.h
@@ -6,6 +6,7 @@
 #include "Containers/UnrealString.h"
 #include "Containers/Map.h"
 #include "Containers/Array.h"
+#include "UObject/NameTypes.h"
 
 #define DDPI_HAS_EXTENDED_PLATFORMINFO_DATA (WITH_EDITOR || IS_PROGRAM) && !IS_MONOLITHIC
 
