Explicitly tagging IsInRenderingThread, IsInGameThread etc. functions to be able to move them between theads in the future. The old behaviour is still in packe to verify the new system, but will be removed in the future.

diff --git a/Engine/Plugins/FX/Niagara/Source/Niagara/Private/NiagaraDataInterface.cpp b/Engine/Plugins/FX/Niagara/Source/Niagara/Private/NiagaraDataInterface.cpp
--- a/Engine/Plugins/FX/Niagara/Source/Niagara/Private/NiagaraDataInterface.cpp
+++ b/Engine/Plugins/FX/Niagara/Source/Niagara/Private/NiagaraDataInterface.cpp
@@ -21,9 +21,10 @@ UNiagaraDataInterface::~UNiagaraDataInterface()
 	if ( Proxy.IsValid() )
 	{
 		ENQUEUE_RENDER_COMMAND(FDeleteProxyRT) (
-			[RT_Proxy=MoveTemp(Proxy)](FRHICommandListImmediate& CmdList)
+			[RT_Proxy=MoveTemp(Proxy)](FRHICommandListImmediate& CmdList) mutable
 			{
 				// This will release RT_Proxy on the RT
+				RT_Proxy = nullptr;
 			}
 		);
 		check(Proxy.IsValid() == false);
diff --git a/Engine/Source/Programs/ShaderCompileWorker/Private/ShaderCompileWorker.cpp b/Engine/Source/Programs/ShaderCompileWorker/Private/ShaderCompileWorker.cpp
--- a/Engine/Source/Programs/ShaderCompileWorker/Private/ShaderCompileWorker.cpp
+++ b/Engine/Source/Programs/ShaderCompileWorker/Private/ShaderCompileWorker.cpp
@@ -961,6 +961,7 @@ static int32 GuardedMain(int32 argc, TCHAR* argv[], bool bDirectMode)
 
 static int32 GuardedMainWrapper(int32 ArgC, TCHAR* ArgV[], const TCHAR* CrashOutputFile, bool bDirectMode)
 {
+	FTaskTagScope Scope(ETaskTag::EGameThread);
 	// We need to know whether we are using XGE now, in case an exception
 	// is thrown before we parse the command line inside GuardedMain.
 	if ((ArgC > 6) && FCString::Strcmp(ArgV[6], TEXT("-xge_int")) == 0)
diff --git a/Engine/Source/Programs/ShaderCompileWorker/Private/XGEControlWorker.cpp b/Engine/Source/Programs/ShaderCompileWorker/Private/XGEControlWorker.cpp
--- a/Engine/Source/Programs/ShaderCompileWorker/Private/XGEControlWorker.cpp
+++ b/Engine/Source/Programs/ShaderCompileWorker/Private/XGEControlWorker.cpp
@@ -185,6 +185,7 @@ void FXGEControlWorker::InputThreadProc()
 
 int32 XGEController_GuardedMain(int32 ArgC, TCHAR* ArgV[])
 {
+	FTaskTagScope Scope(ETaskTag::EGameThread);
 	GEngineLoop.PreInit(ArgC, ArgV, TEXT("-NOPACKAGECACHE -Multiprocess"));
 
 	if (ArgC != 3)
diff --git a/Engine/Source/Programs/UnrealHeaderTool/Private/UnrealHeaderToolMain.cpp b/Engine/Source/Programs/UnrealHeaderTool/Private/UnrealHeaderToolMain.cpp
--- a/Engine/Source/Programs/UnrealHeaderTool/Private/UnrealHeaderToolMain.cpp
+++ b/Engine/Source/Programs/UnrealHeaderTool/Private/UnrealHeaderToolMain.cpp
@@ -30,6 +30,7 @@ bool GUHTErrorLogged = false;
  */
 INT32_MAIN_INT32_ARGC_TCHAR_ARGV()
 {
+	FTaskTagScope Scope(ETaskTag::EGameThread);
 	FString CmdLine;
 
 	for (int32 Arg = 0; Arg < ArgC; Arg++)
diff --git a/Engine/Source/Programs/UnrealPak/Private/UnrealPak.cpp b/Engine/Source/Programs/UnrealPak/Private/UnrealPak.cpp
--- a/Engine/Source/Programs/UnrealPak/Private/UnrealPak.cpp
+++ b/Engine/Source/Programs/UnrealPak/Private/UnrealPak.cpp
@@ -9,6 +9,8 @@ IMPLEMENT_APPLICATION(UnrealPak, "UnrealPak");
 
 INT32_MAIN_INT32_ARGC_TCHAR_ARGV()
 {
+	FTaskTagScope Scope(ETaskTag::EGameThread);
+
 	// start up the main loop
 	GEngineLoop.PreInit(ArgC, ArgV);
 
diff --git a/Engine/Source/Runtime/Core/Private/HAL/ThreadingBase.cpp b/Engine/Source/Runtime/Core/Private/HAL/ThreadingBase.cpp
--- a/Engine/Source/Runtime/Core/Private/HAL/ThreadingBase.cpp
+++ b/Engine/Source/Runtime/Core/Private/HAL/ThreadingBase.cpp
@@ -1,6 +1,5 @@
 // Copyright Epic Games, Inc. All Rights Reserved.
 
-
 #include "HAL/ThreadingBase.h"
 #include "UObject/NameTypes.h"
 #include "Stats/Stats.h"
@@ -14,6 +13,8 @@
 #include "HAL/PlatformStackWalk.h"
 #include "ProfilingDebugging/MiscTrace.h"
 
+#include <atomic>
+
 DEFINE_STAT( STAT_EventWaitWithId );
 DEFINE_STAT( STAT_EventTriggerWithId );
 
@@ -38,10 +39,111 @@ FQueuedThreadPool* GBackgroundPriorityThreadPool = nullptr;
 FQueuedThreadPool* GLargeThreadPool = nullptr;
 #endif
 
+thread_local ETaskTag FTaskTagScope::ActiveTaskTag = ETaskTag::ENone;
+static std::atomic_int ActiveNamedThreads = {};
+
+FTaskTagScope::FTaskTagScope(bool InTagOnlyIfNone, ETaskTag InTag) : Tag(InTag), TagOnlyIfNone(InTagOnlyIfNone)
+{
+	checkf(Tag != ETaskTag::ENone, TEXT("None cannot be used as a Tag"));
+	checkf(Tag != ETaskTag::EParallelThread, TEXT("Parallel cannot be used on it's own"));
+
+	if (!EnumHasAllFlags(Tag, ETaskTag::EParallelThread))
+	{
+		ETaskTag NamedThreadBits = (Tag & ETaskTag::ENamedThreadBits);
+		static_assert(sizeof(ETaskTag) == sizeof(int32), "EnumSize must match interlockedOr");
+		ETaskTag OldTag = ETaskTag(ActiveNamedThreads.fetch_or(int32(NamedThreadBits)));
+		checkf((OldTag & NamedThreadBits) == ETaskTag::ENone, TEXT("Only Scopes tagged with ETaskTag::EParallelThread can be tagged multiple times. ActiveNamedThreads(%h) cannot be tagged multiple times in the same callstack you can use FOptionalTaskTagScope to avoid retagging check the ActiveNamedThreads(%h) with the current Tag(%h)"), OldTag, FTaskTagScope::GetCurrentTag(), Tag);
+	}
+	
+	ParentTag = ActiveTaskTag;
+	if (!TagOnlyIfNone || ActiveTaskTag == ETaskTag::ENone)
+	{
+		ActiveTaskTag = Tag;
+	}
+	else if (TagOnlyIfNone)
+	{
+		if (EnumHasAllFlags(Tag, ETaskTag::EParallelRenderingThread))
+		{
+			checkf(IsInRenderingThread(), TEXT("ETaskTag::EParallelRenderingThread can only be retagged if they are in a parallel for on the RenderingThread or not tagged check the ActiveNamedThreads(%h)"), FTaskTagScope::GetCurrentTag());
+		}
+
+		if (EnumHasAllFlags(Tag, ETaskTag::EParallelGameThread))
+		{
+			checkf(IsInGameThread(), TEXT("ETaskTag::EParallelGameThread can only be retagged if they are in a parallel for on the GameThread or not tagged check the ActiveNamedThreads(%h)"), FTaskTagScope::GetCurrentTag());
+		}
+	}
+}
+
+FTaskTagScope::~FTaskTagScope()
+{
+	checkf(TagOnlyIfNone || ActiveTaskTag == Tag, TEXT("ActiveTaskTag(%h) corrupted needs to be Tag(%h)"), FTaskTagScope::GetCurrentTag(), Tag);
+	if (!TagOnlyIfNone || ParentTag == ETaskTag::ENone)
+	{
+		ActiveTaskTag = ParentTag;
+	}
+
+	if (!EnumHasAllFlags(Tag, ETaskTag::EParallelThread))
+	{
+		ETaskTag NamedThreadBits = (Tag & ETaskTag::ENamedThreadBits);
+		static_assert(sizeof(ETaskTag) == sizeof(int32), "EnumSize must match interlockedAnd");
+		ETaskTag OldTag = ETaskTag(ActiveNamedThreads.fetch_and(int32(~NamedThreadBits)));
+		checkf((OldTag & NamedThreadBits) == (NamedThreadBits), TEXT("Currently active Threads(%h) got corrupted check the ActiveNamedThreads(%h)"), OldTag, FTaskTagScope::GetCurrentTag());
+	}
+
+	//prolong the scope of the GT for static variable destructors
+	if (Tag == ETaskTag::EGameThread && ActiveTaskTag == ETaskTag::ENone)
+	{
+		ActiveTaskTag = ETaskTag::EGameThread;
+	}
+}
+
+bool FTaskTagScope::HasCurrentTag(ETaskTag InTag)
+{
+	if (ActiveTaskTag != ETaskTag::ENone)
+	{
+		return EnumHasAllFlags(ActiveTaskTag, InTag);
+	}
+	else
+	{
+		return InTag == ETaskTag::ENone;
+	}
+}
+
+ETaskTag FTaskTagScope::GetCurrentTag()
+{
+	return ActiveTaskTag;
+}
+
+CORE_API bool IsInGameThread()
+{
+	if (GIsGameThreadIdInitialized)
+	{
+		bool newValue = FTaskTagScope::HasCurrentTag(ETaskTag::EGameThread);
+#if !UE_BUILD_SHIPPING && !UE_BUILD_TEST
+		const uint32 CurrentThreadId = FPlatformTLS::GetCurrentThreadId();
+		bool oldValue = CurrentThreadId == GGameThreadId;
+		checkf(oldValue == newValue, TEXT("If this check fails make sure that there is a FTaskTagScope(ETaskTag::EGameThread) as deep as possible on the current callstack, you can see the current value in ActiveNamedThreads(%h)"), FTaskTagScope::GetCurrentTag());
+#endif
+		return newValue;
+	}
+
+	return true;
+}
+
+CORE_API bool IsInParallelGameThread()
+{
+	return FTaskTagScope::HasCurrentTag(ETaskTag::EParallelGameThread);
+}
+
 CORE_API bool IsInSlateThread()
 {
 	// If this explicitly is a slate thread, not just the main thread running slate
-	return GSlateLoadingThreadId != 0 && FPlatformTLS::GetCurrentThreadId() == GSlateLoadingThreadId;
+	bool newValue = FTaskTagScope::HasCurrentTag(ETaskTag::ESlateThread) && !FTaskTagScope::HasCurrentTag(ETaskTag::EParallelThread);
+#if !UE_BUILD_SHIPPING && !UE_BUILD_TEST
+	bool oldValue = GSlateLoadingThreadId != 0 && FPlatformTLS::GetCurrentThreadId() == GSlateLoadingThreadId;
+	checkf(oldValue == newValue, TEXT("If this check fails make sure that there is a FTaskTagScope(ETaskTag::ESlateThread) as as deep as possible on the current callstack, you can see the current value in ActiveNamedThreads(%h)"), FTaskTagScope::GetCurrentTag());
+#endif
+	return newValue;
 }
 
 CORE_API FRunnableThread* GAudioThread = nullptr;
@@ -49,7 +151,14 @@ CORE_API FRunnableThread* GAudioThread = nullptr;
 CORE_API bool IsInAudioThread()
 {
 	// True if this is the audio thread or if there is no audio thread, then if it is the game thread
-	return FPlatformTLS::GetCurrentThreadId() == (GAudioThreadId ? GAudioThreadId : GGameThreadId);
+	bool newValue = GAudioThreadId
+		? FTaskTagScope::HasCurrentTag(ETaskTag::EAudioThread) && !FTaskTagScope::HasCurrentTag(ETaskTag::EParallelThread)
+		: FTaskTagScope::HasCurrentTag(ETaskTag::EGameThread) && !FTaskTagScope::HasCurrentTag(ETaskTag::EParallelThread);
+#if !UE_BUILD_SHIPPING && !UE_BUILD_TEST
+	bool oldValue = FPlatformTLS::GetCurrentThreadId() == (GAudioThreadId ? GAudioThreadId : GGameThreadId);
+	checkf(oldValue == newValue, TEXT("If this check fails make sure that there is a FTaskTagScope(ETaskTag::EAudioThread) as deep as possible on the current callstack, you can see the current value in ActiveNamedThreads(%h)"), FTaskTagScope::GetCurrentTag());
+#endif
+	return newValue;
 }
 
 CORE_API TAtomic<int32> GIsRenderingThreadSuspended(0);
@@ -58,24 +167,55 @@ CORE_API FRunnableThread* GRenderingThread = nullptr;
 
 CORE_API bool IsInActualRenderingThread()
 {
-	return GRenderingThread && FPlatformTLS::GetCurrentThreadId() == GRenderingThread->GetThreadID();
+	bool newValue = FTaskTagScope::HasCurrentTag(ETaskTag::ERenderingThread) && !FTaskTagScope::HasCurrentTag(ETaskTag::EParallelThread);
+#if !UE_BUILD_SHIPPING && !UE_BUILD_TEST
+	bool oldValue = GRenderingThread && FPlatformTLS::GetCurrentThreadId() == GRenderingThread->GetThreadID();
+	checkf(oldValue == newValue, TEXT("If this check fails make sure that there is a FTaskTagScope(ETaskTag::ERenderingThread) as deep as possible on the current callstack, you can see the current value in ActiveNamedThreads(%h)"), FTaskTagScope::GetCurrentTag());
+#endif
+	return newValue;
 }
 
 CORE_API bool IsInRenderingThread()
 {
-	return !GRenderingThread || GIsRenderingThreadSuspended.Load(EMemoryOrder::Relaxed) || (FPlatformTLS::GetCurrentThreadId() == GRenderingThread->GetThreadID());
+	bool newValue = !GRenderingThread || GIsRenderingThreadSuspended.Load(EMemoryOrder::Relaxed)
+		? FTaskTagScope::HasCurrentTag(ETaskTag::EGameThread) || FTaskTagScope::HasCurrentTag(ETaskTag::ENone)
+		: FTaskTagScope::HasCurrentTag(ETaskTag::ERenderingThread) && !FTaskTagScope::HasCurrentTag(ETaskTag::EParallelThread);
+
+#if !UE_BUILD_SHIPPING && !UE_BUILD_TEST
+	bool oldValue = !GRenderingThread || GIsRenderingThreadSuspended.Load(EMemoryOrder::Relaxed) || (FPlatformTLS::GetCurrentThreadId() == GRenderingThread->GetThreadID());
+	checkf(oldValue == newValue, TEXT("If this check fails make sure that there is a FTaskTagScope(ETaskTag::ERenderingThread) as deep as possible on the current callstack, you can see the current value in ActiveNamedThreads(%h)"), FTaskTagScope::GetCurrentTag());
+#endif
+	return newValue;
 }
 
 CORE_API bool IsInParallelRenderingThread()
 {
+	bool newValue = false;
 	if (!GRenderingThread || GIsRenderingThreadSuspended.Load(EMemoryOrder::Relaxed))
 	{
-		return true;
+		newValue = FTaskTagScope::HasCurrentTag(ETaskTag::EParallelRenderingThread) || FTaskTagScope::HasCurrentTag(ETaskTag::EGameThread);
+	}
+	else
+	{
+		newValue = FTaskTagScope::HasCurrentTag(ETaskTag::EParallelRenderingThread)
+			|| FTaskTagScope::HasCurrentTag(ETaskTag::ERenderingThread)
+			|| FTaskTagScope::HasCurrentTag(ETaskTag::ERhiThread); //TODO lots of RHI functions rely on our broken IsInParallelRenderingThread;
+			//|| FTaskTagScope::HasCurrentTag(ETaskTag::EParallelGameThread); // this is very broken indeed
+	}
+
+#if !UE_BUILD_SHIPPING && !UE_BUILD_TEST
+	bool oldValue = false;
+	if (!GRenderingThread || GIsRenderingThreadSuspended.Load(EMemoryOrder::Relaxed))
+	{
+		oldValue = true;
 	}
 	else
 	{
-		return FPlatformTLS::GetCurrentThreadId() != GGameThreadId;
+		oldValue = FPlatformTLS::GetCurrentThreadId() != GGameThreadId;
 	}
+	checkf(oldValue == newValue, TEXT("If this check fails make sure that there is a FTaskTagScope(ETaskTag::EParallelRenderingThread) as deep as possible on the current callstack, you can see the current value in ActiveNamedThreads(%h)"), FTaskTagScope::GetCurrentTag());
+#endif
+	return newValue;
 }
 
 CORE_API uint32 GRHIThreadId = 0;
@@ -83,7 +223,12 @@ CORE_API FRunnableThread* GRHIThread_InternalUseOnly = nullptr;
 
 CORE_API bool IsInRHIThread()
 {
-	return GRHIThreadId && FPlatformTLS::GetCurrentThreadId() == GRHIThreadId;
+	bool newValue = FTaskTagScope::HasCurrentTag(ETaskTag::ERhiThread);
+#if !UE_BUILD_SHIPPING && !UE_BUILD_TEST
+	bool oldValue = GRHIThreadId && FPlatformTLS::GetCurrentThreadId() == GRHIThreadId;	
+	checkf(oldValue == newValue, TEXT("If this check fails make sure that there is a FTaskTagScope(ETaskTag::ERhiThread) as deep as possible on the current callstack, you can see the current value in ActiveNamedThreads(%h)"), FTaskTagScope::GetCurrentTag());
+#endif
+	return newValue;
 }
 
 // Fake threads
diff --git a/Engine/Source/Runtime/Core/Public/CoreGlobals.h b/Engine/Source/Runtime/Core/Public/CoreGlobals.h
--- a/Engine/Source/Runtime/Core/Public/CoreGlobals.h
+++ b/Engine/Source/Runtime/Core/Public/CoreGlobals.h
@@ -3,6 +3,7 @@
 
 #include "CoreTypes.h"
 #include "Containers/UnrealString.h"
+#include "Misc/EnumClassFlags.h"
 #include "UObject/NameTypes.h"
 #include "Logging/LogMacros.h"
 #include "HAL/PlatformTLS.h"
@@ -460,6 +461,68 @@ extern CORE_API bool GPumpingMessages;
 /** Enables various editor and HMD hacks that allow the experimental VR editor feature to work, perhaps at the expense of other systems */
 extern CORE_API bool GEnableVREditorHacks;
 
+enum class ETaskTag : int32
+{
+	ENone						= 0 << 0,
+	EGameThread					= 1 << 1,
+	ESlateThread				= 1 << 2,
+	EAudioThread				= 1 << 3,
+	ERenderingThread			= 1 << 4,
+	ERhiThread					= 1 << 5,
+	EAsyncLoadingThread			= 1 << 6,
+
+	ENamedThreadBits			= (EAsyncLoadingThread << 1) - 1,
+	EParallelThread				= 1 << 7,										//This can be used when multipe threads or jobs are involved (usually a parallel for) It will avoid the check for uniqieness of the named thread tag.
+	EParallelRenderingThread	= ERenderingThread | EParallelThread,
+	EParallelGameThread			= EGameThread | EParallelThread,
+};
+
+ENUM_CLASS_FLAGS(ETaskTag);
+
+
+/**
+ * This class can be used to Tag an execution context aka Thead or Job and allows us to later querry the state when we are in the callstack
+ * It is usually used for the IsInRendering/GamethreadFunctions.
+ *
+ * @param CtorSignature InTag the Tag to use
+ */
+class FTaskTagScope
+{
+	static thread_local ETaskTag ActiveTaskTag;
+	ETaskTag ParentTag;
+	ETaskTag Tag;
+	bool TagOnlyIfNone;
+
+protected:
+	FTaskTagScope(bool InTagOnlyIfNone, ETaskTag InTag);
+
+public:
+	CORE_API FTaskTagScope(ETaskTag InTag = ETaskTag::ENone) : FTaskTagScope(false, InTag)
+	{
+
+	}
+
+	CORE_API ~FTaskTagScope();
+
+	static ETaskTag GetCurrentTag();
+	static bool HasCurrentTag(ETaskTag InTag);
+};
+
+/**
+ * This class can be used to Tag an execution context but only in case it has not already been tagged
+ * It is usually used for the IsInRendering/GamethreadFunctions.
+ *
+ * @param CtorSignature InTag the Tag to use
+ */
+class FOptionalTaskTagScope : public FTaskTagScope
+{
+public:
+	CORE_API FOptionalTaskTagScope(ETaskTag InTag = ETaskTag::ENone) : FTaskTagScope(true, InTag)
+	{
+
+	}
+};
+
 /**
  * Ensures that current thread is during retrieval of vtable ptr of some
  * UClass.
@@ -470,16 +533,10 @@ extern CORE_API bool GEnableVREditorHacks;
 CORE_API void EnsureRetrievingVTablePtrDuringCtor(const TCHAR* CtorSignature);
 
 /** @return True if called from the game thread. */
-FORCEINLINE bool IsInGameThread()
-{
-	if(GIsGameThreadIdInitialized)
-	{
-		const uint32 CurrentThreadId = FPlatformTLS::GetCurrentThreadId();
-		return CurrentThreadId == GGameThreadId;
-	}
+extern CORE_API bool IsInGameThread();
 
-	return true;
-}
+/** @return True if called from the game thread in a parallel for. */
+extern CORE_API bool IsInParallelGameThread();
 
 /** @return True if called from the audio thread, and not merely a thread calling audio functions. */
 extern CORE_API bool IsInAudioThread();
diff --git a/Engine/Source/Runtime/Core/Public/HAL/IConsoleManager.h b/Engine/Source/Runtime/Core/Public/HAL/IConsoleManager.h
--- a/Engine/Source/Runtime/Core/Public/HAL/IConsoleManager.h
+++ b/Engine/Source/Runtime/Core/Public/HAL/IConsoleManager.h
@@ -1292,7 +1292,7 @@ private: // ----------------------------------------------------
 			cvarCheckCode(ensure(!IsInActualRenderingThread()));
 			return 0;
 		}
-		return IsInGameThread() ? 0 : 1;
+		return IsInParallelGameThread() || IsInGameThread() ? 0 : 1;
 	}
 
 	// needed for FConsoleVariable and FConsoleVariableRef2, intentionally not public
diff --git a/Engine/Source/Runtime/Engine/Private/AudioThread.cpp b/Engine/Source/Runtime/Engine/Private/AudioThread.cpp
--- a/Engine/Source/Runtime/Engine/Private/AudioThread.cpp
+++ b/Engine/Source/Runtime/Engine/Private/AudioThread.cpp
@@ -274,22 +274,23 @@ struct FAudioAsyncBatcher
 
 static FAudioAsyncBatcher GAudioAsyncBatcher;
 
-void FAudioThread::RunCommandOnAudioThread(TFunction<void()> InFunction, const TStatId InStatId)
+void FAudioThread::RunCommandOnAudioThread(TUniqueFunction<void()> InFunction, const TStatId InStatId)
 {
 	check(FPlatformTLS::GetCurrentThreadId() == GGameThreadId);
 	if (bIsAudioThreadRunning)
 	{
 		if (GCVarEnableAudioCommandLogging == 1)
 		{
-			TFunction<void()> FuncWrapper = [InFunction, InStatId]()
+			TUniqueFunction<void()> FuncWrapper = [Function = MoveTemp(InFunction), InStatId]()
 			{
+				FTaskTagScope Scope(ETaskTag::EAudioThread);
 				FAudioThread::SetCurrentAudioThreadStatId(InStatId);
 
 				// Time the execution of the function
 				const double StartTime = FPlatformTime::Seconds();
 
 				// Execute the function
-				InFunction();
+				Function();
 
 				// Track the longest one
 				const double DeltaTime = (FPlatformTime::Seconds() - StartTime) * 1000.0f;
@@ -303,7 +304,13 @@ void FAudioThread::RunCommandOnAudioThread(TFunction<void()> InFunction, const T
 		}
 		else
 		{
-			FFunctionGraphTask::CreateAndDispatchWhenReady(MoveTemp(InFunction), InStatId, GAudioAsyncBatcher.GetAsyncPrereq(), ENamedThreads::AudioThread);
+			TUniqueFunction<void()> FuncWrapper = [Function = MoveTemp(InFunction), InStatId]()
+			{
+				FTaskTagScope Scope(ETaskTag::EAudioThread);
+				// Execute the function
+				Function();
+			};
+			FFunctionGraphTask::CreateAndDispatchWhenReady(MoveTemp(FuncWrapper), InStatId, GAudioAsyncBatcher.GetAsyncPrereq(), ENamedThreads::AudioThread);
 		}
 	}
 	else
@@ -366,7 +373,7 @@ void FAudioThread::ProcessAllCommands()
 	}
 }
 
-void FAudioThread::RunCommandOnGameThread(TFunction<void()> InFunction, const TStatId InStatId)
+void FAudioThread::RunCommandOnGameThread(TUniqueFunction<void()> InFunction, const TStatId InStatId)
 {
 	if (bIsAudioThreadRunning)
 	{
diff --git a/Engine/Source/Runtime/Engine/Private/Components/ActorComponent.cpp b/Engine/Source/Runtime/Engine/Private/Components/ActorComponent.cpp
--- a/Engine/Source/Runtime/Engine/Private/Components/ActorComponent.cpp
+++ b/Engine/Source/Runtime/Engine/Private/Components/ActorComponent.cpp
@@ -97,6 +97,7 @@ void FRegisterComponentContext::Process()
 	ParallelFor(AddPrimitiveBatches.Num(),
 		[&](int32 Index)
 		{
+			FOptionalTaskTagScope Scope(ETaskTag::EParallelGameThread);
 			if (!AddPrimitiveBatches[Index]->IsPendingKill())
 			{
 				Scene->AddPrimitive(AddPrimitiveBatches[Index]);
diff --git a/Engine/Source/Runtime/Engine/Private/Materials/MaterialShared.cpp b/Engine/Source/Runtime/Engine/Private/Materials/MaterialShared.cpp
--- a/Engine/Source/Runtime/Engine/Private/Materials/MaterialShared.cpp
+++ b/Engine/Source/Runtime/Engine/Private/Materials/MaterialShared.cpp
@@ -779,7 +779,7 @@ bool FMaterial::MaterialModifiesMeshPosition_RenderThread() const
 
 bool FMaterial::MaterialModifiesMeshPosition_GameThread() const
 {
-	check(IsInGameThread());
+	check(IsInParallelGameThread() || IsInGameThread());
 	FMaterialShaderMap* ShaderMap = GameThreadShaderMap.GetReference();
 	bool bUsesWPO = ShaderMap ? ShaderMap->ModifiesMeshPosition() : false;
 
diff --git a/Engine/Source/Runtime/Engine/Private/StaticMeshRender.cpp b/Engine/Source/Runtime/Engine/Private/StaticMeshRender.cpp
--- a/Engine/Source/Runtime/Engine/Private/StaticMeshRender.cpp
+++ b/Engine/Source/Runtime/Engine/Private/StaticMeshRender.cpp
@@ -2128,7 +2128,7 @@ FStaticMeshSceneProxy::FLODInfo::FLODInfo(const UStaticMeshComponent* InComponen
 		FMaterialResource const* MaterialResource = const_cast<UMaterialInterface const*>(SectionInfo.Material)->GetMaterial_Concurrent()->GetMaterialResource(FeatureLevel);
 		if(MaterialResource)
 		{
-			if (IsInGameThread())
+			if (IsInParallelGameThread() || IsInGameThread())
 			{
 				if (MaterialResource->MaterialModifiesMeshPosition_GameThread())
 				{
diff --git a/Engine/Source/Runtime/Engine/Public/AudioThread.h b/Engine/Source/Runtime/Engine/Public/AudioThread.h
--- a/Engine/Source/Runtime/Engine/Public/AudioThread.h
+++ b/Engine/Source/Runtime/Engine/Public/AudioThread.h
@@ -76,13 +76,13 @@ public:
 	static ENGINE_API void StopAudioThread();
 
 	/** Execute a command on the audio thread. If GIsAudioThreadRunning is false the command will execute immediately */
-	static ENGINE_API void RunCommandOnAudioThread(TFunction<void()> InFunction, const TStatId InStatId = TStatId());
+	static ENGINE_API void RunCommandOnAudioThread(TUniqueFunction<void()> InFunction, const TStatId InStatId = TStatId());
 
 	/** Processes all enqueued audio thread commands. */
 	static ENGINE_API void ProcessAllCommands();
 
 	/** Execute a (presumably audio) command on the game thread. If GIsAudioThreadRunning is false the command will execute immediately */
-	static ENGINE_API void RunCommandOnGameThread(TFunction<void()> InFunction, const TStatId InStatId = TStatId());
+	static ENGINE_API void RunCommandOnGameThread(TUniqueFunction<void()> InFunction, const TStatId InStatId = TStatId());
 
 	static ENGINE_API void SetUseThreadedAudio(bool bInUseThreadedAudio);
 	static ENGINE_API bool IsUsingThreadedAudio() { return bUseThreadedAudio; }
diff --git a/Engine/Source/Runtime/Landscape/Public/LandscapeRender.h b/Engine/Source/Runtime/Landscape/Public/LandscapeRender.h
--- a/Engine/Source/Runtime/Landscape/Public/LandscapeRender.h
+++ b/Engine/Source/Runtime/Landscape/Public/LandscapeRender.h
@@ -599,6 +599,7 @@ struct FLandscapeRenderSystem
 
 		void AnyThreadTask()
 		{
+			FOptionalTaskTagScope Scope(ETaskTag::EParallelRenderingThread);
 			RenderSystem.ComputeSectionPerViewParameters(
 				ViewPtrAsIdentifier, ViewLODOverride, ViewLODDistanceFactor, 
 				ViewEngineShowFlagCollisionPawn, ViewEngineShowFlagCollisionVisibility, 
@@ -637,6 +638,7 @@ struct FLandscapeRenderSystem
 
 		void AnyThreadTask()
 		{
+			FOptionalTaskTagScope Scope(ETaskTag::EParallelRenderingThread);
 			RenderSystem.FetchHeightmapLODBiases();
 		}
 
diff --git a/Engine/Source/Runtime/Launch/Private/Launch.cpp b/Engine/Source/Runtime/Launch/Private/Launch.cpp
--- a/Engine/Source/Runtime/Launch/Private/Launch.cpp
+++ b/Engine/Source/Runtime/Launch/Private/Launch.cpp
@@ -92,6 +92,8 @@ extern UNREALED_API FSecondsCounterData BlueprintCompileAndLoadTimerData;
  */
 int32 GuardedMain( const TCHAR* CmdLine )
 {
+	FTaskTagScope Scope(ETaskTag::EGameThread);
+
 #if !(UE_BUILD_SHIPPING)
 	if (FParse::Param(CmdLine, TEXT("waitforattach")))
 	{
diff --git a/Engine/Source/Runtime/MoviePlayer/Private/MoviePlayerThreading.cpp b/Engine/Source/Runtime/MoviePlayer/Private/MoviePlayerThreading.cpp
--- a/Engine/Source/Runtime/MoviePlayer/Private/MoviePlayerThreading.cpp
+++ b/Engine/Source/Runtime/MoviePlayer/Private/MoviePlayerThreading.cpp
@@ -169,6 +169,8 @@ bool FSlateLoadingThreadTask::Init()
 
 uint32 FSlateLoadingThreadTask::Run()
 {
+	FTaskTagScope Scope(ETaskTag::ESlateThread);
+
 	check( GSlateLoadingThreadId == FPlatformTLS::GetCurrentThreadId() );
 
 	SyncMechanism->SlateThreadRunMainLoop();
diff --git a/Engine/Source/Runtime/RHI/Private/RHICommandList.cpp b/Engine/Source/Runtime/RHI/Private/RHICommandList.cpp
--- a/Engine/Source/Runtime/RHI/Private/RHICommandList.cpp
+++ b/Engine/Source/Runtime/RHI/Private/RHICommandList.cpp
@@ -412,6 +412,7 @@ public:
 
 	void DoTask(ENamedThreads::Type CurrentThread, const FGraphEventRef& MyCompletionGraphEvent)
 	{
+		FTaskTagScope Scope(ETaskTag::ERhiThread);
 		SCOPE_CYCLE_COUNTER(STAT_RHIThreadExecute);
 		if (IsRunningRHIInTaskThread())
 		{
@@ -997,6 +998,7 @@ public:
 
 	void DoTask(ENamedThreads::Type CurrentThread, const FGraphEventRef& MyCompletionGraphEvent)
 	{
+		FTaskTagScope Scope(ETaskTag::ERhiThread | ETaskTag::EParallelThread);
 		SCOPE_CYCLE_COUNTER(STAT_ParallelChainTranslate);
 		SCOPED_NAMED_EVENT(FParallelTranslateCommandList_DoTask, FColor::Magenta);
 		check(ContextContainer && RHICmdLists);
diff --git a/Engine/Source/Runtime/RenderCore/Private/RenderingThread.cpp b/Engine/Source/Runtime/RenderCore/Private/RenderingThread.cpp
--- a/Engine/Source/Runtime/RenderCore/Private/RenderingThread.cpp
+++ b/Engine/Source/Runtime/RenderCore/Private/RenderingThread.cpp
@@ -463,6 +463,7 @@ public:
 
 	virtual void Exit(void) override
 	{
+		FTaskTagScope Scope(ETaskTag::ERenderingThread);
 		// Release rendering context ownership on the current thread if we had acquired it
 		if (bAcquiredThreadOwnership)
 		{
diff --git a/Engine/Source/Runtime/RenderCore/Public/RenderingThread.h b/Engine/Source/Runtime/RenderCore/Public/RenderingThread.h
--- a/Engine/Source/Runtime/RenderCore/Public/RenderingThread.h
+++ b/Engine/Source/Runtime/RenderCore/Public/RenderingThread.h
@@ -190,6 +190,7 @@ public:
 
 	void DoTask(ENamedThreads::Type CurrentThread, const FGraphEventRef& MyCompletionGraphEvent)
 	{
+		FTaskTagScope Scope(ETaskTag::ERenderingThread);
 		TRACE_CPUPROFILER_EVENT_SCOPE_ON_CHANNEL_STR(TSTR::TStr(), RenderCommandsChannel);
 		FRHICommandListImmediate& RHICmdList = GetImmediateCommandList_ForRenderCommand();
 		Lambda(RHICmdList);
diff --git a/Engine/Source/Runtime/Renderer/Private/MeshDrawCommands.cpp b/Engine/Source/Runtime/Renderer/Private/MeshDrawCommands.cpp
--- a/Engine/Source/Runtime/Renderer/Private/MeshDrawCommands.cpp
+++ b/Engine/Source/Runtime/Renderer/Private/MeshDrawCommands.cpp
@@ -770,6 +770,7 @@ public:
 
 	void AnyThreadTask()
 	{
+		FOptionalTaskTagScope Scope(ETaskTag::EParallelRenderingThread);
 		TRACE_CPUPROFILER_EVENT_SCOPE(MeshDrawCommandPassSetupTask);
 		// Mobile base pass is a special case, as final lists is created from two mesh passes based on CSM visibility.
 		const bool bMobileShadingBasePass = Context.ShadingPath == EShadingPath::Mobile && Context.PassType == EMeshPass::BasePass;
@@ -1215,6 +1216,7 @@ public:
 
 	void DoTask(ENamedThreads::Type CurrentThread, const FGraphEventRef& MyCompletionGraphEvent)
 	{
+		FOptionalTaskTagScope Scope(ETaskTag::EParallelRenderingThread);
 		TRACE_CPUPROFILER_EVENT_SCOPE(DrawVisibleMeshCommandsAnyThreadTask);
 		checkSlow(RHICmdList.IsInsideRenderPass());
 
diff --git a/Engine/Source/Runtime/Renderer/Private/PrimitiveSceneInfo.cpp b/Engine/Source/Runtime/Renderer/Private/PrimitiveSceneInfo.cpp
--- a/Engine/Source/Runtime/Renderer/Private/PrimitiveSceneInfo.cpp
+++ b/Engine/Source/Runtime/Renderer/Private/PrimitiveSceneInfo.cpp
@@ -264,6 +264,7 @@ void FPrimitiveSceneInfo::CacheMeshDrawCommands(FRHICommandListImmediate& RHICmd
 
 			if ((FPassProcessorManager::GetPassFlags(ShadingPath, PassType) & EMeshPassFlags::CachedMeshCommands) != EMeshPassFlags::None)
 			{
+				FOptionalTaskTagScope Scope(ETaskTag::EParallelRenderingThread);
 				FCachedMeshDrawCommandInfo CommandInfo(PassType);
 
 				FCriticalSection& CachedMeshDrawCommandLock = Scene->CachedMeshDrawCommandLock[PassType];
@@ -501,6 +502,7 @@ void FPrimitiveSceneInfo::AddStaticMeshes(FRHICommandListImmediate& RHICmdList,
 	{
 		ParallelForTemplate(SceneInfos.Num(), [Scene, &SceneInfos](int32 Index)
 		{
+			FOptionalTaskTagScope Scope(ETaskTag::EParallelRenderingThread);
 			SCOPED_NAMED_EVENT(FPrimitiveSceneInfo_AddStaticMeshes_DrawStaticElements, FColor::Magenta);
 			FPrimitiveSceneInfo* SceneInfo = SceneInfos[Index];
 			// Cache the primitive's static mesh elements.
diff --git a/Engine/Source/Runtime/Renderer/Private/SceneVisibility.cpp b/Engine/Source/Runtime/Renderer/Private/SceneVisibility.cpp
--- a/Engine/Source/Runtime/Renderer/Private/SceneVisibility.cpp
+++ b/Engine/Source/Runtime/Renderer/Private/SceneVisibility.cpp
@@ -1949,6 +1949,7 @@ struct FRelevancePacket
 
 	void AnyThreadTask()
 	{
+		FOptionalTaskTagScope Scope(ETaskTag::EParallelRenderingThread);
 		ComputeRelevance();
 		MarkRelevant();
 	}
diff --git a/Engine/Source/Runtime/Renderer/Private/ShadowSetup.cpp b/Engine/Source/Runtime/Renderer/Private/ShadowSetup.cpp
--- a/Engine/Source/Runtime/Renderer/Private/ShadowSetup.cpp
+++ b/Engine/Source/Runtime/Renderer/Private/ShadowSetup.cpp
@@ -3653,6 +3653,7 @@ struct FGatherShadowPrimitivesPacket
 
 	void AnyThreadTask()
 	{
+		FOptionalTaskTagScope Scope(ETaskTag::EParallelRenderingThread);
 		QUICK_SCOPE_CYCLE_COUNTER(STAT_GatherShadowPrimitivesPacket);
 
 		if (NodeIndex != INDEX_NONE)
