- Fixed Monolithic editor build errors about PlatformInfo (there are still two more errors I saw in QAGame)

diff --git a/Engine/Source/Runtime/Core/Public/Misc/DataDrivenPlatformInfoRegistry.h b/Engine/Source/Runtime/Core/Public/Misc/DataDrivenPlatformInfoRegistry.h
--- a/Engine/Source/Runtime/Core/Public/Misc/DataDrivenPlatformInfoRegistry.h
+++ b/Engine/Source/Runtime/Core/Public/Misc/DataDrivenPlatformInfoRegistry.h
@@ -8,7 +8,7 @@
 #include "Containers/Array.h"
 #include "UObject/NameTypes.h"
 
-#define DDPI_HAS_EXTENDED_PLATFORMINFO_DATA (WITH_EDITOR || IS_PROGRAM) && !IS_MONOLITHIC
+#define DDPI_HAS_EXTENDED_PLATFORMINFO_DATA (WITH_EDITOR || IS_PROGRAM)
 
 #if DDPI_HAS_EXTENDED_PLATFORMINFO_DATA
 
