Fix Mac undefined symbol compile error
#rnx

diff --git a/Engine/Source/Runtime/Core/Public/CoreGlobals.h b/Engine/Source/Runtime/Core/Public/CoreGlobals.h
--- a/Engine/Source/Runtime/Core/Public/CoreGlobals.h
+++ b/Engine/Source/Runtime/Core/Public/CoreGlobals.h
@@ -494,7 +494,7 @@ class FTaskTagScope
 	bool TagOnlyIfNone;
 
 protected:
-	FTaskTagScope(bool InTagOnlyIfNone, ETaskTag InTag);
+	CORE_API FTaskTagScope(bool InTagOnlyIfNone, ETaskTag InTag);
 
 public:
 	CORE_API FTaskTagScope(ETaskTag InTag = ETaskTag::ENone) : FTaskTagScope(false, InTag)
