Added FUE5MainStreamObjectVersion to support serialization changes made in //UE5/Main

diff --git a/Engine/Source/Runtime/Core/Private/UObject/DevObjectVersion.cpp b/Engine/Source/Runtime/Core/Private/UObject/DevObjectVersion.cpp
--- a/Engine/Source/Runtime/Core/Private/UObject/DevObjectVersion.cpp
+++ b/Engine/Source/Runtime/Core/Private/UObject/DevObjectVersion.cpp
@@ -28,6 +28,7 @@
 #include "UObject/ExternalPhysicsMaterialCustomObjectVersion.h"
 #include "UObject/CineCameraObjectVersion.h"
 #include "UObject/VirtualProductionObjectVersion.h"
+#include "UObject/UE5MainStreamObjectVersion.h"
 
 DEFINE_LOG_CATEGORY_STATIC(LogDevObjectVersion, Log, All);
 
@@ -189,4 +190,7 @@ const FGuid FVirtualProductionObjectVersion::GUID(0x64F58936, 0xFD1B42BA, 0xBA96
 // Register VirtualProduction custom version with Core
 FDevVersionRegistration GRegisterVirtualProductionObjectVersion(FVirtualProductionObjectVersion::GUID, FVirtualProductionObjectVersion::LatestVersion, TEXT("Dev-VirtualProduction"));
 
-
+// Unique UE5 version id
+const FGuid FUE5MainStreamObjectVersion::GUID(0x697DD581, 0xE64f41AB, 0xAA4A51EC, 0xBEB7B628);
+// Register UE5 main stream custom version with Core
+FDevVersionRegistration GRegisterUE5MainStreamObjectVersion(FUE5MainStreamObjectVersion::GUID, FUE5MainStreamObjectVersion::LatestVersion, TEXT("UE5-Main"));
diff --git a/Engine/Source/Runtime/Core/Public/UObject/UE5MainStreamObjectVersion.h b/Engine/Source/Runtime/Core/Public/UObject/UE5MainStreamObjectVersion.h
new file mode 100644
--- /dev/null
+++ b/Engine/Source/Runtime/Core/Public/UObject/UE5MainStreamObjectVersion.h
@@ -0,0 +1,24 @@
+// Copyright Epic Games, Inc. All Rights Reserved.
+#pragma once
+
+#include "CoreTypes.h"
+#include "Misc/Guid.h"
+
+// Custom serialization version for changes made in //UE5/Main stream
+struct CORE_API FUE5MainStreamObjectVersion
+{
+	enum Type
+	{
+		// Before any version changes were made
+		BeforeCustomVersionWasAdded = 0,
+
+		// -----<new versions can be added above this line>-------------------------------------------------
+		VersionPlusOne,
+		LatestVersion = VersionPlusOne - 1
+	};
+
+	// The GUID for this custom version number
+	const static FGuid GUID;
+
+	FUE5MainStreamObjectVersion() = delete;
+};
