SS - Fix tagging of parallel RHI Threads

diff --git a/Engine/Source/Runtime/Core/Private/HAL/ThreadingBase.cpp b/Engine/Source/Runtime/Core/Private/HAL/ThreadingBase.cpp
--- a/Engine/Source/Runtime/Core/Private/HAL/ThreadingBase.cpp
+++ b/Engine/Source/Runtime/Core/Private/HAL/ThreadingBase.cpp
@@ -223,7 +223,7 @@ CORE_API FRunnableThread* GRHIThread_InternalUseOnly = nullptr;
 
 CORE_API bool IsInRHIThread()
 {
-	bool newValue = FTaskTagScope::HasCurrentTag(ETaskTag::ERhiThread);
+	bool newValue = FTaskTagScope::HasCurrentTag(ETaskTag::ERhiThread) && !FTaskTagScope::HasCurrentTag(ETaskTag::EParallelThread);
 #if !UE_BUILD_SHIPPING && !UE_BUILD_TEST
 	bool oldValue = GRHIThreadId && FPlatformTLS::GetCurrentThreadId() == GRHIThreadId;	
 	checkf(oldValue == newValue, TEXT("If this check fails make sure that there is a FTaskTagScope(ETaskTag::ERhiThread) as deep as possible on the current callstack, you can see the current value in ActiveNamedThreads(%h)"), FTaskTagScope::GetCurrentTag());
diff --git a/Engine/Source/Runtime/Core/Public/CoreGlobals.h b/Engine/Source/Runtime/Core/Public/CoreGlobals.h
--- a/Engine/Source/Runtime/Core/Public/CoreGlobals.h
+++ b/Engine/Source/Runtime/Core/Public/CoreGlobals.h
@@ -475,6 +475,7 @@ enum class ETaskTag : int32
 	EParallelThread				= 1 << 7,										//This can be used when multipe threads or jobs are involved (usually a parallel for) It will avoid the check for uniqieness of the named thread tag.
 	EParallelRenderingThread	= ERenderingThread | EParallelThread,
 	EParallelGameThread			= EGameThread | EParallelThread,
+	EParallelRhiThread			= ERhiThread | EParallelThread,
 };
 
 ENUM_CLASS_FLAGS(ETaskTag);
diff --git a/Engine/Source/Runtime/RHI/Private/RHICommandList.cpp b/Engine/Source/Runtime/RHI/Private/RHICommandList.cpp
--- a/Engine/Source/Runtime/RHI/Private/RHICommandList.cpp
+++ b/Engine/Source/Runtime/RHI/Private/RHICommandList.cpp
@@ -1004,7 +1004,7 @@ public:
 
 	void DoTask(ENamedThreads::Type CurrentThread, const FGraphEventRef& MyCompletionGraphEvent)
 	{
-		FTaskTagScope Scope(ETaskTag::ERhiThread | ETaskTag::EParallelThread);
+		FOptionalTaskTagScope Scope(ETaskTag::EParallelRhiThread);
 		SCOPE_CYCLE_COUNTER(STAT_ParallelChainTranslate);
 		SCOPED_NAMED_EVENT(FParallelTranslateCommandList_DoTask, FColor::Magenta);
 		check(ContextContainer && RHICmdLists);
