Fixed trace initialisation after a bad merge from UE-Main.

#rnx

diff --git a/Engine/Source/Runtime/Core/Private/ProfilingDebugging/TraceAuxiliary.cpp b/Engine/Source/Runtime/Core/Private/ProfilingDebugging/TraceAuxiliary.cpp
--- a/Engine/Source/Runtime/Core/Private/ProfilingDebugging/TraceAuxiliary.cpp
+++ b/Engine/Source/Runtime/Core/Private/ProfilingDebugging/TraceAuxiliary.cpp
@@ -32,14 +32,23 @@
 #include "Windows/HideWindowsPlatformTypes.h"
 #endif
 
+////////////////////////////////////////////////////////////////////////////////
+enum class ETraceConnectType
+{
+	Network,
+	File
+};
+
 ////////////////////////////////////////////////////////////////////////////////
 class FTraceAuxiliaryImpl
 {
 public:
-	void					ParseCommandLine(const TCHAR* CommandLine);
-	bool 					Start(const TCHAR* ChannelSet);
-	bool 					Stop();
-	const TCHAR*			GetPath() const;
+	const TCHAR*			GetDest() const;
+	template <class T> void	ReadChannels(T&& Callback) const;
+	void					AddChannels(const TCHAR* ChannelList);
+	bool					Connect(ETraceConnectType Type, const TCHAR* Parameter);
+	void					EnableChannels();
+	void					DisableChannels();
 
 private:
 	enum class EState : uint8
@@ -49,137 +58,161 @@ private:
 		Stopped,
 	};
 
+	struct FChannel
+	{
+		FString				Name;
+		bool				bActive = false;
+	};
+
+	void					AddChannel(const TCHAR* Name);
+	void					AddChannels(const TCHAR* Name, bool bResolvePresets);
+	void					EnableChannel(FChannel& Channel);
 	bool					SendToHost(const TCHAR* Host);
 	bool					WriteToFile(const TCHAR* Path=nullptr);
-	void					ToggleChannels(const TCHAR* Channels);
-	FString					GetChannels(const TCHAR* ChannelSet) const;
-	TMap<uint32, FString>	ActiveChannels;
-	FString					TracePath;
+
+	TMap<uint32, FChannel>	Channels;
+	FString					TraceDest;
 	EState					State = EState::None;
 };
 
 static FTraceAuxiliaryImpl GTraceAuxiliary;
 
 ////////////////////////////////////////////////////////////////////////////////
-void FTraceAuxiliaryImpl::ToggleChannels(const TCHAR* Channels)
+void FTraceAuxiliaryImpl::AddChannels(const TCHAR* ChannelList)
 {
-	UE::String::ParseTokens(Channels, TEXT(","), [this] (const FStringView& Token)
-	{
-		TCHAR ChannelName[64];
-		const size_t ChannelNameSize = Token.CopyString(ChannelName, 63);
-		ChannelName[ChannelNameSize] = '\0';
+	AddChannels(ChannelList, true);
+}
 
-		uint32 ChannelHash = 5381;
-		for (const TCHAR* c = ChannelName; *c; ++c)
-		{
-            ChannelHash = ((ChannelHash << 5) + ChannelHash) + *c;
-		}
+////////////////////////////////////////////////////////////////////////////////
+void FTraceAuxiliaryImpl::AddChannels(const TCHAR* ChannelList, bool bResolvePresets)
+{
+	UE::String::ParseTokens(ChannelList, TEXT(","), [this, bResolvePresets] (const FStringView& Token)
+	{
+		TCHAR Name[80];
+		const size_t ChannelNameSize = Token.CopyString(Name, UE_ARRAY_COUNT(Name) - 1);
+		Name[ChannelNameSize] = '\0';
 
-		if (ActiveChannels.Find(ChannelHash) != nullptr)
+		if (bResolvePresets)
 		{
-			return;
+			FString Value;
+			if (GConfig->GetString(TEXT("Trace.ChannelPresets"), Name, Value, GEngineIni))
+			{
+				AddChannels(*Value, false);
+				return;
+			}
 		}
 
-		ActiveChannels.Add(ChannelHash, ChannelName);
-
-		Trace::ToggleChannel(ChannelName, true);
+		AddChannel(Name);
 	});
 }
 
 ////////////////////////////////////////////////////////////////////////////////
-FString FTraceAuxiliaryImpl::GetChannels(const TCHAR* ChannelSet) const
+void FTraceAuxiliaryImpl::AddChannel(const TCHAR* Name)
 {
-	FString Value;
-
-	if (ChannelSet == nullptr)
-	{
-		if (!GConfig->GetString(TEXT("Trace.ChannelPresets"), TEXT("Default"), Value, GEngineIni))
-		{
-			Value = TEXT("cpu,frame,log,bookmark");
-		}
-	}
-	else if (!GConfig->GetString(TEXT("Trace.ChannelPresets"), ChannelSet, Value, GEngineIni))
+	uint32 Hash = 5381;
+	for (const TCHAR* c = Name; *c; ++c)
 	{
-		Value = ChannelSet;
+		uint32 LowerC = *c | 0x20;
+        Hash = ((Hash << 5) + Hash) + LowerC;
 	}
 
-	return Value;
-}
-
-////////////////////////////////////////////////////////////////////////////////
-void FTraceAuxiliaryImpl::ParseCommandLine(const TCHAR* CommandLine)
-{
-	if (State >= EState::Tracing)
+	if (Channels.Find(Hash) != nullptr)
 	{
 		return;
 	}
 
-	bool bOk = false;
-	FString Parameter;
+	FChannel& Value = Channels.Add(Hash, {});
+	Value.Name = Name;
 
-	// Start tracing if it isn't already
-	if (FParse::Value(CommandLine, TEXT("-tracehost="), Parameter))
+	if (State >= EState::Tracing)
 	{
-		bOk = SendToHost(*Parameter);
+		EnableChannel(Value);
 	}
+}
 
-	else if (FParse::Value(CommandLine, TEXT("-tracefile="), Parameter))
+////////////////////////////////////////////////////////////////////////////////
+bool FTraceAuxiliaryImpl::Connect(ETraceConnectType Type, const TCHAR* Parameter)
+{
+	// Connect/write to file. But only if we're not already sending/writing
+	bool bConnected = (State >= EState::Tracing);
+	if (!bConnected)
 	{
-		bOk = WriteToFile(*Parameter);
+		if (Type == ETraceConnectType::Network)
+		{
+			bConnected = SendToHost(Parameter);
+		}
+
+		else if (Type == ETraceConnectType::File)
+		{
+			bConnected = WriteToFile(Parameter);
+		}
 	}
 
-	else if (FParse::Param(CommandLine, TEXT("tracefile")))
+	if (!bConnected)
 	{
-		bOk = WriteToFile();
+		return false;
 	}
 
-	const TCHAR* ChannelSet = nullptr;
-	if (FParse::Value(CommandLine, TEXT("-trace="), Parameter, false))
+	// We're now connected. If we don't appear to have any channels we'll set
+	// some defaults for the user. Less futzing.
+	if (!Channels.Num())
 	{
-		ChannelSet = *Parameter;
+		FString Value;
+		if (!GConfig->GetString(TEXT("Trace.ChannelPresets"), TEXT("Default"), Value, GEngineIni))
+		{
+			Value = TEXT("cpu,frame,log,bookmark");
+		}
+
+		AddChannels(*Value);
 	}
 
-	FString Channels = GetChannels(ChannelSet);
-	ToggleChannels(*Channels);
+	EnableChannels();
 
-	State = bOk ? EState::Tracing : EState::None;
+	State = EState::Tracing;
+	return true;
 }
 
 ////////////////////////////////////////////////////////////////////////////////
-bool FTraceAuxiliaryImpl::Start(const TCHAR* ChannelSet)
+void FTraceAuxiliaryImpl::EnableChannel(FChannel& Channel)
 {
-	if (State < EState::Tracing)
+	if (Channel.bActive)
 	{
-		if (!WriteToFile())
-		{
-			return false;
-		}
+		return;
 	}
 
-	FString Channels = GetChannels(ChannelSet);
-	ToggleChannels(*Channels);
+	// Channel names have been provided by the user and may not exist yet. As
+	// we want to maintain bActive accurately (channels toggles are reference
+	// counted), we will first check Trace knows of the channel.
+	if (!Trace::IsChannel(*Channel.Name))
+	{
+		return;
+	}
 
-	State = EState::Tracing;
-	return true;
+	Trace::ToggleChannel(*Channel.Name, true);
+	Channel.bActive = true;
 }
 
 ////////////////////////////////////////////////////////////////////////////////
-bool FTraceAuxiliaryImpl::Stop()
+void FTraceAuxiliaryImpl::EnableChannels()
 {
-	if (State < EState::Tracing)
+	for (auto& ChannelPair : Channels)
 	{
-		return false;
+		EnableChannel(ChannelPair.Value);
 	}
+}
 
-	for (const auto& ChannelPair : ActiveChannels)
+////////////////////////////////////////////////////////////////////////////////
+void FTraceAuxiliaryImpl::DisableChannels()
+{
+	for (auto& ChannelPair : Channels)
 	{
-		const FString& Name = ChannelPair.Value;
-		Trace::ToggleChannel(*Name, false);
+		FChannel& Channel = ChannelPair.Value;
+		if (Channel.bActive)
+		{
+			Trace::ToggleChannel(*Channel.Name, false);
+			Channel.bActive = false;
+		}
 	}
-	ActiveChannels.Reset();
-
-	State = EState::Stopped;
-	return true;
 }
 
 ////////////////////////////////////////////////////////////////////////////////
@@ -191,13 +224,14 @@ bool FTraceAuxiliaryImpl::SendToHost(const TCHAR* Host)
 		return false;
 	}
 
+	TraceDest = Host;
 	return true;
 }
 
 ////////////////////////////////////////////////////////////////////////////////
 bool FTraceAuxiliaryImpl::WriteToFile(const TCHAR* Path)
 {
-	if (Path == nullptr)
+	if (Path == nullptr || *Path == '\0')
 	{
 		FString Name = FDateTime::Now().ToString(TEXT("%Y%m%d_%H%M%S.utrace"));
 		return WriteToFile(*Name);
@@ -246,14 +280,24 @@ bool FTraceAuxiliaryImpl::WriteToFile(const TCHAR* Path)
 		return false;
 	}
 
-	TracePath = MoveTemp(WritePath);
+	TraceDest = MoveTemp(NativePath);
 	return true;
 }
 
 ////////////////////////////////////////////////////////////////////////////////
-const TCHAR* FTraceAuxiliaryImpl::GetPath() const
+const TCHAR* FTraceAuxiliaryImpl::GetDest() const
 {
-	return *TracePath;
+	return *TraceDest;
+}
+
+////////////////////////////////////////////////////////////////////////////////
+template <class T>
+void FTraceAuxiliaryImpl::ReadChannels(T&& Callback) const
+{
+	for (const auto& ChannelPair : Channels)
+	{
+		Callback(*(ChannelPair.Value.Name));
+	}
 }
 
 
@@ -261,30 +305,37 @@ const TCHAR* FTraceAuxiliaryImpl::GetPath() const
 ////////////////////////////////////////////////////////////////////////////////
 static void TraceAuxiliaryStart(const TArray<FString>& Args)
 {
-	const TCHAR* Channels = (Args.Num() > 0) ? *(Args[0]) : nullptr;
+	if (Args.Num() > 0)
+	{
+		GTraceAuxiliary.AddChannels(*(Args[0]));
+	}
 
-	if (!GTraceAuxiliary.Start(Channels))
+	if (!GTraceAuxiliary.Connect(ETraceConnectType::File, nullptr))
 	{
 		UE_LOG(LogConsoleResponse, Warning, TEXT("Failed to start tracing to a file"));
 		return;
 	}
 
 	// Give the user some feedback that everything's underway.
-	Channels = (Channels != nullptr) ? Channels : TEXT("[default]");
-	UE_LOG(LogConsoleResponse, Log, TEXT("Tracing to; %s"), GTraceAuxiliary.GetPath());
-	UE_LOG(LogConsoleResponse, Log, TEXT("Trace channels; %s"), Channels);
+	FString Channels;
+	GTraceAuxiliary.ReadChannels([&Channels] (const TCHAR* Channel)
+	{
+		if (Channels.Len())
+		{
+			Channels += TEXT(",");
+		}
+
+		Channels += Channel;
+	});
+	UE_LOG(LogConsoleResponse, Log, TEXT("Tracing to; %s"), GTraceAuxiliary.GetDest());
+	UE_LOG(LogConsoleResponse, Log, TEXT("Trace channels; %s"), *Channels);
 }
 
 ////////////////////////////////////////////////////////////////////////////////
 static void TraceAuxiliaryStop()
 {
-	if (!GTraceAuxiliary.Stop())
-	{
-		UE_LOG(LogConsoleResponse, Warning, TEXT("Unable to stop tracing"));
-		return;
-	}
-
-	UE_LOG(LogConsoleResponse, Log, TEXT("Tracing paused. Use 'Trace.Start' to resume"));
+	GTraceAuxiliary.DisableChannels();
+	UE_LOG(LogConsoleResponse, Log, TEXT("Tracing stopped. Use 'Trace.Start' to resume"));
 }
 
 ////////////////////////////////////////////////////////////////////////////////
@@ -321,6 +372,16 @@ UE_TRACE_EVENT_END()
 void FTraceAuxiliary::Initialize(const TCHAR* CommandLine)
 {
 #if UE_TRACE_ENABLED
+	// Trace out information about this session. This is done before initialisation
+	// so that it is always sent (all channels are enabled prior to initialisation)
+	UE_TRACE_LOG(Diagnostics, Session2, Trace::TraceLogChannel)
+		<< Session2.Platform(PREPROCESSOR_TO_STRING(UBT_COMPILED_PLATFORM))
+		<< Session2.AppName(UE_APP_NAME)
+		<< Session2.CommandLine(CommandLine)
+		<< Session2.ConfigurationType(uint8(FApp::GetBuildConfiguration()))
+		<< Session2.TargetType(uint8(FApp::GetBuildTargetType()));
+
+	// Initialize Trace
 	Trace::FInitializeDesc Desc;
 	Desc.bUseWorkerThread = FPlatformProcess::SupportsMultithreading();
 
@@ -333,14 +394,28 @@ void FTraceAuxiliary::Initialize(const TCHAR* CommandLine)
 
 	FCoreDelegates::OnEndFrame.AddStatic(Trace::Update);
 
-	// Trace out information about this session
-	UE_TRACE_LOG(Diagnostics, Session2, Trace::TraceLogChannel)
-		<< Session2.Platform(PREPROCESSOR_TO_STRING(UBT_COMPILED_PLATFORM))
-		<< Session2.AppName(UE_APP_NAME)
-		<< Session2.CommandLine(CommandLine)
-		<< Session2.ConfigurationType(uint8(FApp::GetBuildConfiguration()))
-		<< Session2.TargetType(uint8(FApp::GetBuildTargetType()));
+	// Extract an explicit channel set from the command line.
+	if (FParse::Value(CommandLine, TEXT("-trace="), Parameter, false))
+	{
+		GTraceAuxiliary.AddChannels(*Parameter);
+		GTraceAuxiliary.EnableChannels();
+	}
+
+	// Attempt to send trace data somewhere from the command line
+	if (FParse::Value(CommandLine, TEXT("-tracehost="), Parameter))
+	{
+		GTraceAuxiliary.Connect(ETraceConnectType::Network, *Parameter);
+	}
+	else if (FParse::Value(CommandLine, TEXT("-tracefile="), Parameter))
+	{
+		GTraceAuxiliary.Connect(ETraceConnectType::File, *Parameter);
+	}
+	else if (FParse::Param(CommandLine, TEXT("tracefile")))
+	{
+		GTraceAuxiliary.Connect(ETraceConnectType::File, nullptr);
+	}
 
+	// Insights' instrumentation.
 	TRACE_CPUPROFILER_INIT(CommandLine);
 	TRACE_PLATFORMFILE_INIT(CommandLine);
 	TRACE_COUNTERS_INIT(CommandLine);
@@ -348,10 +423,10 @@ void FTraceAuxiliary::Initialize(const TCHAR* CommandLine)
 }
 
 ////////////////////////////////////////////////////////////////////////////////
-void FTraceAuxiliary::ParseCommandLine(const TCHAR* CommandLine)
+void FTraceAuxiliary::EnableChannels()
 {
 #if UE_TRACE_ENABLED
-	GTraceAuxiliary.ParseCommandLine(CommandLine);
+	GTraceAuxiliary.EnableChannels();
 #endif
 }
 
@@ -364,7 +439,7 @@ void FTraceAuxiliary::TryAutoConnect()
 	HANDLE KnownEvent = ::OpenEvent(EVENT_ALL_ACCESS, false, TEXT("Local\\UnrealInsightsRecorder"));
 	if (KnownEvent != nullptr)
 	{
-		GTraceAuxiliary.ParseCommandLine(TEXT("-tracehost=127.0.0.1"));
+		GTraceAuxiliary.Connect(ETraceConnectType::Network, TEXT("127.0.0.1"));
 		::CloseHandle(KnownEvent);
 	}
 	#endif
diff --git a/Engine/Source/Runtime/Core/Public/ProfilingDebugging/TraceAuxiliary.h b/Engine/Source/Runtime/Core/Public/ProfilingDebugging/TraceAuxiliary.h
--- a/Engine/Source/Runtime/Core/Public/ProfilingDebugging/TraceAuxiliary.h
+++ b/Engine/Source/Runtime/Core/Public/ProfilingDebugging/TraceAuxiliary.h
@@ -9,6 +9,6 @@ class CORE_API FTraceAuxiliary
 {
 public:
 	static void Initialize(const TCHAR* CommandLine);
-	static void ParseCommandLine(const TCHAR* CommandLine);
 	static void TryAutoConnect();
+	static void EnableChannels();
 };
diff --git a/Engine/Source/Runtime/Launch/Private/LaunchEngineLoop.cpp b/Engine/Source/Runtime/Launch/Private/LaunchEngineLoop.cpp
--- a/Engine/Source/Runtime/Launch/Private/LaunchEngineLoop.cpp
+++ b/Engine/Source/Runtime/Launch/Private/LaunchEngineLoop.cpp
@@ -4030,7 +4030,7 @@ int32 FEngineLoop::Init()
 		GetMoviePlayer()->WaitForMovieToFinish();
     }
 
-	FTraceAuxiliary::ParseCommandLine(FCommandLine::Get());
+	FTraceAuxiliary::EnableChannels();
 
 #if !UE_SERVER
 	// initialize media framework
