Fixed a copy-paste error in FBox::ComputeSquareDistanceToBox that showed up in the PVS log.

#rnx

diff --git a/Engine/Source/Runtime/Core/Public/Math/Box.h b/Engine/Source/Runtime/Core/Public/Math/Box.h
--- a/Engine/Source/Runtime/Core/Public/Math/Box.h
+++ b/Engine/Source/Runtime/Core/Public/Math/Box.h
@@ -155,7 +155,7 @@ public:
 	 */
 	FORCEINLINE float ComputeSquaredDistanceToBox(const FBox& Box) const
 	{
-		FVector AxisDistances = (Box.GetCenter() - Box.GetCenter()).GetAbs() - (GetExtent() + Box.GetExtent());
+		FVector AxisDistances = (GetCenter() - Box.GetCenter()).GetAbs() - (GetExtent() + Box.GetExtent());
 		AxisDistances = FVector::Max(AxisDistances, FVector(0.0f, 0.0f, 0.0f));
 		return FVector::DotProduct(AxisDistances, AxisDistances);
 	}
