SS - Fix IsInParallelRenderingThread

diff --git a/Engine/Source/Runtime/Core/Private/HAL/ThreadingBase.cpp b/Engine/Source/Runtime/Core/Private/HAL/ThreadingBase.cpp
--- a/Engine/Source/Runtime/Core/Private/HAL/ThreadingBase.cpp
+++ b/Engine/Source/Runtime/Core/Private/HAL/ThreadingBase.cpp
@@ -193,7 +193,7 @@ CORE_API bool IsInParallelRenderingThread()
 	bool newValue = false;
 	if (!GRenderingThread || GIsRenderingThreadSuspended.Load(EMemoryOrder::Relaxed))
 	{
-		newValue = FTaskTagScope::HasCurrentTag(ETaskTag::EParallelRenderingThread) || FTaskTagScope::HasCurrentTag(ETaskTag::EGameThread);
+		newValue = FTaskTagScope::HasCurrentTag(ETaskTag::ERenderingThread) || FTaskTagScope::HasCurrentTag(ETaskTag::EGameThread) || FTaskTagScope::HasCurrentTag(ETaskTag::EParallelRenderingThread);
 	}
 	else
 	{
