SS -Try to catch non parallel scopes that are quickly opened and closed on other threads in the debugger.

diff --git a/Engine/Source/Runtime/Core/Private/HAL/ThreadingBase.cpp b/Engine/Source/Runtime/Core/Private/HAL/ThreadingBase.cpp
--- a/Engine/Source/Runtime/Core/Private/HAL/ThreadingBase.cpp
+++ b/Engine/Source/Runtime/Core/Private/HAL/ThreadingBase.cpp
@@ -52,7 +52,13 @@ FTaskTagScope::FTaskTagScope(bool InTagOnlyIfNone, ETaskTag InTag) : Tag(InTag),
 		ETaskTag NamedThreadBits = (Tag & ETaskTag::ENamedThreadBits);
 		static_assert(sizeof(ETaskTag) == sizeof(int32), "EnumSize must match interlockedOr");
 		ETaskTag OldTag = ETaskTag(ActiveNamedThreads.fetch_or(int32(NamedThreadBits)));
-		checkf((OldTag & NamedThreadBits) == ETaskTag::ENone, TEXT("Only Scopes tagged with ETaskTag::EParallelThread can be tagged multiple times. ActiveNamedThreads(%h) cannot be tagged multiple times in the same callstack you can use FOptionalTaskTagScope to avoid retagging check the ActiveNamedThreads(%h) with the current Tag(%h)"), OldTag, FTaskTagScope::GetCurrentTag(), Tag);
+		bool IsOK = (OldTag & NamedThreadBits) == ETaskTag::ENone;
+		if (!IsOK)
+		{
+			//Try to catch other Threads that already opened a non parallel scope
+			ActiveNamedThreads.store(int32(ETaskTag::ENone));
+		}
+		checkf(IsOK, TEXT("Only Scopes tagged with ETaskTag::EParallelThread can be tagged multiple times. ActiveNamedThreads(%h) cannot be tagged multiple times in the same callstack you can use FOptionalTaskTagScope to avoid retagging check the ActiveNamedThreads(%h) with the current Tag(%h)"), OldTag, FTaskTagScope::GetCurrentTag(), Tag);
 	}
 	
 	ParentTag = ActiveTaskTag;
