SS - Streamline the Tagging system and exclude GT from ParallelGT. In the future we should proably also exclude RT from parallelRT but than er would have to change all the usages.

diff --git a/Engine/Source/Runtime/Core/Private/HAL/ThreadingBase.cpp b/Engine/Source/Runtime/Core/Private/HAL/ThreadingBase.cpp
--- a/Engine/Source/Runtime/Core/Private/HAL/ThreadingBase.cpp
+++ b/Engine/Source/Runtime/Core/Private/HAL/ThreadingBase.cpp
@@ -103,16 +103,9 @@ FTaskTagScope::~FTaskTagScope()
 	}
 }
 
-bool FTaskTagScope::HasCurrentTag(ETaskTag InTag)
+bool FTaskTagScope::IsCurrentTag(ETaskTag InTag)
 {
-	if (ActiveTaskTag != ETaskTag::ENone)
-	{
-		return EnumHasAllFlags(ActiveTaskTag, InTag);
-	}
-	else
-	{
-		return InTag == ETaskTag::ENone;
-	}
+	return ActiveTaskTag == InTag;
 }
 
 ETaskTag FTaskTagScope::GetCurrentTag()
@@ -124,7 +117,7 @@ CORE_API bool IsInGameThread()
 {
 	if (GIsGameThreadIdInitialized)
 	{
-		bool newValue = FTaskTagScope::HasCurrentTag(ETaskTag::EGameThread);
+		bool newValue = FTaskTagScope::IsCurrentTag(ETaskTag::EGameThread);
 #if !UE_BUILD_SHIPPING && !UE_BUILD_TEST
 		const uint32 CurrentThreadId = FPlatformTLS::GetCurrentThreadId();
 		bool oldValue = CurrentThreadId == GGameThreadId;
@@ -138,13 +131,13 @@ CORE_API bool IsInGameThread()
 
 CORE_API bool IsInParallelGameThread()
 {
-	return FTaskTagScope::HasCurrentTag(ETaskTag::EParallelGameThread);
+	return FTaskTagScope::IsCurrentTag(ETaskTag::EParallelGameThread);
 }
 
 CORE_API bool IsInSlateThread()
 {
 	// If this explicitly is a slate thread, not just the main thread running slate
-	bool newValue = FTaskTagScope::HasCurrentTag(ETaskTag::ESlateThread) && !FTaskTagScope::HasCurrentTag(ETaskTag::EParallelThread);
+	bool newValue = FTaskTagScope::IsCurrentTag(ETaskTag::ESlateThread);
 #if !UE_BUILD_SHIPPING && !UE_BUILD_TEST
 	bool oldValue = GSlateLoadingThreadId != 0 && FPlatformTLS::GetCurrentThreadId() == GSlateLoadingThreadId;
 	checkf(oldValue == newValue, TEXT("If this check fails make sure that there is a FTaskTagScope(ETaskTag::ESlateThread) as as deep as possible on the current callstack, you can see the current value in ActiveNamedThreads(%h)"), FTaskTagScope::GetCurrentTag());
@@ -158,8 +151,8 @@ CORE_API bool IsInAudioThread()
 {
 	// True if this is the audio thread or if there is no audio thread, then if it is the game thread
 	bool newValue = GAudioThreadId
-		? FTaskTagScope::HasCurrentTag(ETaskTag::EAudioThread) && !FTaskTagScope::HasCurrentTag(ETaskTag::EParallelThread)
-		: FTaskTagScope::HasCurrentTag(ETaskTag::EGameThread) && !FTaskTagScope::HasCurrentTag(ETaskTag::EParallelThread);
+		? FTaskTagScope::IsCurrentTag(ETaskTag::EAudioThread)
+		: FTaskTagScope::IsCurrentTag(ETaskTag::EGameThread);
 #if !UE_BUILD_SHIPPING && !UE_BUILD_TEST
 	bool oldValue = FPlatformTLS::GetCurrentThreadId() == (GAudioThreadId ? GAudioThreadId : GGameThreadId);
 	checkf(oldValue == newValue, TEXT("If this check fails make sure that there is a FTaskTagScope(ETaskTag::EAudioThread) as deep as possible on the current callstack, you can see the current value in ActiveNamedThreads(%h)"), FTaskTagScope::GetCurrentTag());
@@ -173,7 +166,7 @@ CORE_API FRunnableThread* GRenderingThread = nullptr;
 
 CORE_API bool IsInActualRenderingThread()
 {
-	bool newValue = FTaskTagScope::HasCurrentTag(ETaskTag::ERenderingThread) && !FTaskTagScope::HasCurrentTag(ETaskTag::EParallelThread);
+	bool newValue = FTaskTagScope::IsCurrentTag(ETaskTag::ERenderingThread);
 #if !UE_BUILD_SHIPPING && !UE_BUILD_TEST
 	bool oldValue = GRenderingThread && FPlatformTLS::GetCurrentThreadId() == GRenderingThread->GetThreadID();
 	checkf(oldValue == newValue, TEXT("If this check fails make sure that there is a FTaskTagScope(ETaskTag::ERenderingThread) as deep as possible on the current callstack, you can see the current value in ActiveNamedThreads(%h)"), FTaskTagScope::GetCurrentTag());
@@ -184,8 +177,8 @@ CORE_API bool IsInActualRenderingThread()
 CORE_API bool IsInRenderingThread()
 {
 	bool newValue = !GRenderingThread || GIsRenderingThreadSuspended.Load(EMemoryOrder::Relaxed)
-		? FTaskTagScope::HasCurrentTag(ETaskTag::EGameThread) || FTaskTagScope::HasCurrentTag(ETaskTag::ENone)
-		: FTaskTagScope::HasCurrentTag(ETaskTag::ERenderingThread) && !FTaskTagScope::HasCurrentTag(ETaskTag::EParallelThread);
+		? FTaskTagScope::IsCurrentTag(ETaskTag::EGameThread) || FTaskTagScope::IsCurrentTag(ETaskTag::ENone)
+		: FTaskTagScope::IsCurrentTag(ETaskTag::ERenderingThread);
 
 #if !UE_BUILD_SHIPPING && !UE_BUILD_TEST
 	bool oldValue = !GRenderingThread || GIsRenderingThreadSuspended.Load(EMemoryOrder::Relaxed) || (FPlatformTLS::GetCurrentThreadId() == GRenderingThread->GetThreadID());
@@ -199,14 +192,13 @@ CORE_API bool IsInParallelRenderingThread()
 	bool newValue = false;
 	if (!GRenderingThread || GIsRenderingThreadSuspended.Load(EMemoryOrder::Relaxed))
 	{
-		newValue = FTaskTagScope::HasCurrentTag(ETaskTag::ERenderingThread) || FTaskTagScope::HasCurrentTag(ETaskTag::EGameThread) || FTaskTagScope::HasCurrentTag(ETaskTag::EParallelRenderingThread);
+		newValue = FTaskTagScope::IsCurrentTag(ETaskTag::ERenderingThread) || FTaskTagScope::IsCurrentTag(ETaskTag::EGameThread) || FTaskTagScope::IsCurrentTag(ETaskTag::EParallelRenderingThread);
 	}
 	else
 	{
-		newValue = FTaskTagScope::HasCurrentTag(ETaskTag::EParallelRenderingThread)
-			|| FTaskTagScope::HasCurrentTag(ETaskTag::ERenderingThread)
-			|| FTaskTagScope::HasCurrentTag(ETaskTag::ERhiThread); //TODO lots of RHI functions rely on our broken IsInParallelRenderingThread;
-			//|| FTaskTagScope::HasCurrentTag(ETaskTag::EParallelGameThread); // this is very broken indeed
+		newValue = FTaskTagScope::IsCurrentTag(ETaskTag::EParallelRenderingThread)
+			|| FTaskTagScope::IsCurrentTag(ETaskTag::ERenderingThread)
+			|| FTaskTagScope::IsCurrentTag(ETaskTag::ERhiThread); //TODO lots of RHI functions rely on our broken IsInParallelRenderingThread;
 	}
 
 #if !UE_BUILD_SHIPPING && !UE_BUILD_TEST
@@ -229,7 +221,7 @@ CORE_API FRunnableThread* GRHIThread_InternalUseOnly = nullptr;
 
 CORE_API bool IsInRHIThread()
 {
-	bool newValue = FTaskTagScope::HasCurrentTag(ETaskTag::ERhiThread) && !FTaskTagScope::HasCurrentTag(ETaskTag::EParallelThread);
+	bool newValue = FTaskTagScope::IsCurrentTag(ETaskTag::ERhiThread);
 #if !UE_BUILD_SHIPPING && !UE_BUILD_TEST
 	bool oldValue = GRHIThreadId && FPlatformTLS::GetCurrentThreadId() == GRHIThreadId;	
 	checkf(oldValue == newValue, TEXT("If this check fails make sure that there is a FTaskTagScope(ETaskTag::ERhiThread) as deep as possible on the current callstack, you can see the current value in ActiveNamedThreads(%h)"), FTaskTagScope::GetCurrentTag());
diff --git a/Engine/Source/Runtime/Core/Public/CoreGlobals.h b/Engine/Source/Runtime/Core/Public/CoreGlobals.h
--- a/Engine/Source/Runtime/Core/Public/CoreGlobals.h
+++ b/Engine/Source/Runtime/Core/Public/CoreGlobals.h
@@ -506,7 +506,7 @@ public:
 	CORE_API ~FTaskTagScope();
 
 	static ETaskTag GetCurrentTag();
-	static bool HasCurrentTag(ETaskTag InTag);
+	static bool IsCurrentTag(ETaskTag InTag);
 };
 
 /**
diff --git a/Engine/Source/Runtime/Engine/Public/InGamePerformanceTracker.h b/Engine/Source/Runtime/Engine/Public/InGamePerformanceTracker.h
--- a/Engine/Source/Runtime/Engine/Public/InGamePerformanceTracker.h
+++ b/Engine/Source/Runtime/Engine/Public/InGamePerformanceTracker.h
@@ -130,12 +130,7 @@ public:
 	{
 		if (Tracker && FInGamePerformanceTracker::CachedEnabled)
 		{
-			if (IsInParallelGameThread())
-			{
-				//On other threads we don't allow (or check for) re-entrance and store our own cycle data.
-				BeginCycles = FPlatformTime::Cycles();
-			}
-			else if (IsInGameThread())
+			if (IsInGameThread())
 			{
 				Tracker->EnterTimedSection();//On the GT we use the tracker directly so we can allow re-entrance.
 			}
@@ -150,14 +145,7 @@ public:
 	{
 		if (Tracker && FInGamePerformanceTracker::CachedEnabled)
 		{
-			if (IsInParallelGameThread())
-			{
-				//On other threads we don't allow (or check for) re-entrance and store our own cycle data.
-				uint32 Cycles = FPlatformTime::Cycles() - BeginCycles;
-				Tracker->AddCycles(Cycles);
-				BeginCycles = 0;
-			}
-			else if (IsInGameThread())
+			if (IsInGameThread())
 			{
 				Tracker->ExitTimedSection();//On the GT we use the tracker directly so we can allow re-entrance.
 			}
