Supress false positives about TFieldPtrAccessor that was specifically created to support this scenario

diff --git a/Engine/Source/Runtime/Core/Public/Windows/WindowsPlatformCodeAnalysis.h b/Engine/Source/Runtime/Core/Public/Windows/WindowsPlatformCodeAnalysis.h
--- a/Engine/Source/Runtime/Core/Public/Windows/WindowsPlatformCodeAnalysis.h
+++ b/Engine/Source/Runtime/Core/Public/Windows/WindowsPlatformCodeAnalysis.h
@@ -229,5 +229,5 @@
 	//-V:Linker:678
 	//-V:self:678
 	//-V:AccumulateParentID:678
-
+	//-V:Resource:623
 #endif
