Properly name the game thread on Windows

diff --git a/Engine/Source/Runtime/Core/Private/Windows/WindowsPlatformProcess.cpp b/Engine/Source/Runtime/Core/Private/Windows/WindowsPlatformProcess.cpp
--- a/Engine/Source/Runtime/Core/Private/Windows/WindowsPlatformProcess.cpp
+++ b/Engine/Source/Runtime/Core/Private/Windows/WindowsPlatformProcess.cpp
@@ -1938,6 +1938,11 @@ FString FWindowsPlatformProcess::FProcEnumInfo::GetFullPath() const
 	return GetApplicationName(GetPID());
 }
 
+void FWindowsPlatformProcess::SetupGameThread()
+{
+	SetThreadName(TEXT("GameThread"));
+}
+
 namespace WindowsPlatformProcessImpl
 {
 	static void SetThreadName(LPCSTR ThreadName)
diff --git a/Engine/Source/Runtime/Core/Public/Windows/WindowsPlatformProcess.h b/Engine/Source/Runtime/Core/Public/Windows/WindowsPlatformProcess.h
--- a/Engine/Source/Runtime/Core/Public/Windows/WindowsPlatformProcess.h
+++ b/Engine/Source/Runtime/Core/Public/Windows/WindowsPlatformProcess.h
@@ -186,6 +186,7 @@ public:
 	static FSemaphore* NewInterprocessSynchObject(const TCHAR* Name, bool bCreate, uint32 MaxLocks = 1);
 	static bool DeleteInterprocessSynchObject(FSemaphore * Object);
 	static bool Daemonize();
+	static void SetupGameThread();
 	static void SetupAudioThread();
 	static void TeardownAudioThread();
 protected:
