Add an IsEmpty() method to the standard containers for consistency and readability.

diff --git a/Engine/Source/Runtime/Core/Public/Containers/Array.h b/Engine/Source/Runtime/Core/Public/Containers/Array.h
--- a/Engine/Source/Runtime/Core/Public/Containers/Array.h
+++ b/Engine/Source/Runtime/Core/Public/Containers/Array.h
@@ -686,6 +686,17 @@ public:
 		return Index >= 0 && Index < ArrayNum;
 	}
 
+	/**
+	 * Returns true if the array is empty and contains no elements. 
+	 *
+	 * @returns True if the array is empty.
+	 * @see Num
+	 */
+	bool IsEmpty() const
+	{
+		return ArrayNum == 0;
+	}
+
 	/**
 	 * Returns number of elements in array.
 	 *
diff --git a/Engine/Source/Runtime/Core/Public/Containers/ArrayView.h b/Engine/Source/Runtime/Core/Public/Containers/ArrayView.h
--- a/Engine/Source/Runtime/Core/Public/Containers/ArrayView.h
+++ b/Engine/Source/Runtime/Core/Public/Containers/ArrayView.h
@@ -203,6 +203,17 @@ public:
 		return (Index >= 0) && (Index < ArrayNum);
 	}
 
+	/**
+	 * Returns true if the array view is empty and contains no elements. 
+	 *
+	 * @returns True if the array view is empty.
+	 * @see Num
+	 */
+	bool IsEmpty() const
+	{
+		return ArrayNum == 0;
+	}
+
 	/**
 	 * Returns number of elements in array.
 	 *
diff --git a/Engine/Source/Runtime/Core/Public/Containers/BasicArray.h b/Engine/Source/Runtime/Core/Public/Containers/BasicArray.h
--- a/Engine/Source/Runtime/Core/Public/Containers/BasicArray.h
+++ b/Engine/Source/Runtime/Core/Public/Containers/BasicArray.h
@@ -85,6 +85,11 @@ public:
 		}
 	}
 
+	bool IsEmpty() const
+	{
+		return Num() == 0;
+	}
+
 	int32 Num() const
 	{
 		FData* LocalData = Data;
diff --git a/Engine/Source/Runtime/Core/Public/Containers/BinaryHeap.h b/Engine/Source/Runtime/Core/Public/Containers/BinaryHeap.h
--- a/Engine/Source/Runtime/Core/Public/Containers/BinaryHeap.h
+++ b/Engine/Source/Runtime/Core/Public/Containers/BinaryHeap.h
@@ -23,6 +23,7 @@ public:
 	void		Free();
 	void		Resize( uint32 NewHeapSize, uint32 NewIndexSize );
 
+	bool		IsEmpty() const			{ return HeapNum == 0; }
 	uint32		Num() const				{ return HeapNum; }
 	uint32		GetHeapSize() const		{ return HeapSize; }
 	uint32		GetIndexSize() const	{ return IndexSize; }
diff --git a/Engine/Source/Runtime/Core/Public/Containers/BitArray.h b/Engine/Source/Runtime/Core/Public/Containers/BitArray.h
--- a/Engine/Source/Runtime/Core/Public/Containers/BitArray.h
+++ b/Engine/Source/Runtime/Core/Public/Containers/BitArray.h
@@ -846,6 +846,17 @@ public:
 		return InIndex >= 0 && InIndex < NumBits;
 	}
 
+	/**
+	 * Returns true if the bit array is empty and contains no elements. 
+	 *
+	 * @returns True if the bit array is empty.
+	 * @see Num
+	 */
+	bool IsEmpty() const 
+	{ 
+		return NumBits == 0; 
+	}
+
 	FORCEINLINE int32 Num() const { return NumBits; }
 	FORCEINLINE FBitReference operator[](int32 Index)
 	{
diff --git a/Engine/Source/Runtime/Core/Public/Containers/ChunkedArray.h b/Engine/Source/Runtime/Core/Public/Containers/ChunkedArray.h
--- a/Engine/Source/Runtime/Core/Public/Containers/ChunkedArray.h
+++ b/Engine/Source/Runtime/Core/Public/Containers/ChunkedArray.h
@@ -128,6 +128,18 @@ public:
 		const int32 ChunkElementIndex = ElementIndex % NumElementsPerChunk;
 		return Chunks[ChunkIndex].Elements[ChunkElementIndex];
 	}
+
+	/**
+	 * Returns true if the chunked array is empty and contains no elements. 
+	 *
+	 * @returns True if the chunked array is empty.
+	 * @see Num
+	 */
+	bool IsEmpty() const
+	{
+		return NumElements == 0;
+	}
+
 	int32 Num() const 
 	{ 
 		return NumElements; 
diff --git a/Engine/Source/Runtime/Core/Public/Containers/IndirectArray.h b/Engine/Source/Runtime/Core/Public/Containers/IndirectArray.h
--- a/Engine/Source/Runtime/Core/Public/Containers/IndirectArray.h
+++ b/Engine/Source/Runtime/Core/Public/Containers/IndirectArray.h
@@ -74,6 +74,17 @@ public:
 		Empty();
 	}
 
+	/**
+	 * Returns true if the array is empty and contains no elements. 
+	 *
+	 * @returns True if the array is empty.
+	 * @see Num
+	 */
+	bool IsEmpty() const
+	{
+		return Array.IsEmpty();
+	}
+
 	/**
 	 * Gets number of elements in array.
 	 *
diff --git a/Engine/Source/Runtime/Core/Public/Containers/List.h b/Engine/Source/Runtime/Core/Public/Containers/List.h
--- a/Engine/Source/Runtime/Core/Public/Containers/List.h
+++ b/Engine/Source/Runtime/Core/Public/Containers/List.h
@@ -740,6 +740,17 @@ public:
 		return (FindNode(InElement) != nullptr);
 	}
 
+	/**
+	 * Returns true if the list is empty and contains no elements. 
+	 *
+	 * @returns True if the list is empty.
+	 * @see Num
+	 */
+	bool IsEmpty() const
+	{
+		return ListSize == 0;
+	}
+
 	/**
 	 * Returns the number of items in the list.
 	 *
diff --git a/Engine/Source/Runtime/Core/Public/Containers/LruCache.h b/Engine/Source/Runtime/Core/Public/Containers/LruCache.h
--- a/Engine/Source/Runtime/Core/Public/Containers/LruCache.h
+++ b/Engine/Source/Runtime/Core/Public/Containers/LruCache.h
@@ -425,6 +425,17 @@ public:
 		return MaxNumElements;
 	}
 
+	/**
+	 * Returns true if the cache is empty and contains no elements. 
+	 *
+	 * @returns True if the cache is empty.
+	 * @see Num
+	 */
+	bool IsEmpty() const
+	{
+		return LookupSet.IsEmpty();
+	}
+
 	/**
 	 * Get the number of entries in the cache.
 	 *
diff --git a/Engine/Source/Runtime/Core/Public/Containers/Map.h b/Engine/Source/Runtime/Core/Public/Containers/Map.h
--- a/Engine/Source/Runtime/Core/Public/Containers/Map.h
+++ b/Engine/Source/Runtime/Core/Public/Containers/Map.h
@@ -266,6 +266,17 @@ public:
 		Pairs.Reserve(Number);
 	}
 
+	/**
+	 * Returns true if the map is empty and contains no elements. 
+	 *
+	 * @returns True if the map is empty.
+	 * @see Num
+	 */
+	bool IsEmpty() const
+	{
+		return Pairs.IsEmpty();
+	}
+
 	/** @return The number of elements in the map. */
 	FORCEINLINE int32 Num() const
 	{
@@ -1600,6 +1611,11 @@ public:
 		return Pairs.IsValidIndex(Index);
 	}
 
+	bool IsEmpty() const
+	{
+		return Pairs.IsEmpty();
+	}
+
 	int32 Num() const
 	{
 		return Pairs.Num();
diff --git a/Engine/Source/Runtime/Core/Public/Containers/ScriptArray.h b/Engine/Source/Runtime/Core/Public/Containers/ScriptArray.h
--- a/Engine/Source/Runtime/Core/Public/Containers/ScriptArray.h
+++ b/Engine/Source/Runtime/Core/Public/Containers/ScriptArray.h
@@ -32,6 +32,10 @@ public:
 	{
 		return i>=0 && i<ArrayNum;
 	}
+	bool IsEmpty() const
+	{
+		return ArrayNum == 0;
+	}
 	FORCEINLINE int32 Num() const
 	{
 		checkSlow(ArrayNum>=0);
diff --git a/Engine/Source/Runtime/Core/Public/Containers/Set.h b/Engine/Source/Runtime/Core/Public/Containers/Set.h
--- a/Engine/Source/Runtime/Core/Public/Containers/Set.h
+++ b/Engine/Source/Runtime/Core/Public/Containers/Set.h
@@ -541,6 +541,17 @@ public:
 		Ar.CountBytes(HashSize * sizeof(int32),HashSize * sizeof(FSetElementId));
 	}
 
+	/**
+	 * Returns true if the sets is empty and contains no elements. 
+	 *
+	 * @returns True if the set is empty.
+	 * @see Num
+	 */
+	bool IsEmpty() const
+	{
+		return Elements.IsEmpty();
+	}
+
 	/** @return the number of elements. */
 	FORCEINLINE int32 Num() const
 	{
@@ -1704,6 +1715,11 @@ public:
 		return Elements.IsValidIndex(Index);
 	}
 
+	bool IsEmpty() const
+	{
+		return Elements.IsEmpty();
+	}
+
 	int32 Num() const
 	{
 		return Elements.Num();
diff --git a/Engine/Source/Runtime/Core/Public/Containers/SortedMap.h b/Engine/Source/Runtime/Core/Public/Containers/SortedMap.h
--- a/Engine/Source/Runtime/Core/Public/Containers/SortedMap.h
+++ b/Engine/Source/Runtime/Core/Public/Containers/SortedMap.h
@@ -104,6 +104,17 @@ public:
 		Pairs.Reserve(Number);
 	}
 
+	/**
+	 * Returns true if the map is empty and contains no elements. 
+	 *
+	 * @returns True if the map is empty.
+	 * @see Num
+	 */
+	bool IsEmpty() const
+	{
+		return Pairs.IsEmpty();
+	}
+
 	/** @return The number of elements in the map. */
 	FORCEINLINE int32 Num() const
 	{
diff --git a/Engine/Source/Runtime/Core/Public/Containers/SparseArray.h b/Engine/Source/Runtime/Core/Public/Containers/SparseArray.h
--- a/Engine/Source/Runtime/Core/Public/Containers/SparseArray.h
+++ b/Engine/Source/Runtime/Core/Public/Containers/SparseArray.h
@@ -785,6 +785,7 @@ public:
 	}
 	bool IsAllocated(int32 Index) const { return AllocationFlags[Index]; }
 	int32 GetMaxIndex() const { return Data.Num(); }
+	bool IsEmpty() const { return Data.Num() == NumFreeIndices; }
 	int32 Num() const { return Data.Num() - NumFreeIndices; }
 
 	/**
@@ -1255,6 +1256,11 @@ public:
 		return AllocationFlags.IsValidIndex(Index) && AllocationFlags[Index];
 	}
 
+	bool IsEmpty() const
+	{
+		return Data.Num() == NumFreeIndices;
+	}
+
 	int32 Num() const
 	{
 		return Data.Num() - NumFreeIndices;
diff --git a/Engine/Source/Runtime/Core/Public/Containers/StaticArray.h b/Engine/Source/Runtime/Core/Public/Containers/StaticArray.h
--- a/Engine/Source/Runtime/Core/Public/Containers/StaticArray.h
+++ b/Engine/Source/Runtime/Core/Public/Containers/StaticArray.h
@@ -75,6 +75,17 @@ public:
 		return false;
 	}
 
+	/**
+	 * Returns true if the array is empty and contains no elements. 
+	 *
+	 * @returns True if the array is empty.
+	 * @see Num
+	 */
+	bool IsEmpty() const
+	{
+		return NumElements == 0;
+	}
+
 	/** The number of elements in the array. */
 	FORCEINLINE_DEBUGGABLE int32 Num() const { return NumElements; }
 
