Implement missing task retraction in the threadpool wrapper causing stalls on the game-thread when the taskpools are full
Add fully dynamic threadpool wrapper with sort functionality
Fix FQueuedThreadPool not properly honoring the given priority

diff --git a/Engine/Source/Runtime/Core/Private/HAL/ThreadingBase.cpp b/Engine/Source/Runtime/Core/Private/HAL/ThreadingBase.cpp
--- a/Engine/Source/Runtime/Core/Private/HAL/ThreadingBase.cpp
+++ b/Engine/Source/Runtime/Core/Private/HAL/ThreadingBase.cpp
@@ -782,6 +782,15 @@ void FThreadPoolPriorityQueue::Reset()
 	NumQueuedWork = 0;
 }
 
+void FThreadPoolPriorityQueue::Sort(EQueuedWorkPriority InPriorityBucket, TFunctionRef<bool(const IQueuedWork* A, const IQueuedWork* B)> Predicate)
+{
+	int32 QueueIndex = static_cast<int32>(InPriorityBucket);
+	if (QueueIndex < PriorityQueuedWork.Num())
+	{
+		Algo::Sort(PriorityQueuedWork[QueueIndex], Predicate);
+	}
+}
+
 /*-----------------------------------------------------------------------------
 	FQueuedThread
 -----------------------------------------------------------------------------*/
@@ -1062,7 +1071,7 @@ public:
 			{
 				// No thread available, queue the work to be done
 				// as soon as one does become available
-				QueuedWork.Enqueue(InQueuedWork);
+				QueuedWork.Enqueue(InQueuedWork, InQueuedWorkPriority);
 				return;
 			}
 
diff --git a/Engine/Source/Runtime/Core/Private/Misc/QueuedThreadPoolWrapper.cpp b/Engine/Source/Runtime/Core/Private/Misc/QueuedThreadPoolWrapper.cpp
--- a/Engine/Source/Runtime/Core/Private/Misc/QueuedThreadPoolWrapper.cpp
+++ b/Engine/Source/Runtime/Core/Private/Misc/QueuedThreadPoolWrapper.cpp
@@ -15,14 +15,12 @@ struct FQueuedThreadPoolWrapper::FScheduledWork : public IQueuedWork
 	void DoThreadedWork() override
 	{
 		InnerWork->DoThreadedWork();
-		InnerWork = nullptr;
 		ParentPool->Schedule(this);
 	}
 
-	virtual void Abandon() override
+	void Abandon() override
 	{
 		InnerWork->Abandon();
-		InnerWork = nullptr;
 		ParentPool->Schedule(this);
 	}
 
@@ -62,13 +60,26 @@ void FQueuedThreadPoolWrapper::Destroy()
 		}
 
 		QueuedWork.Reset();
+
+		// Try to retract anything already in flight
+		TArray<FScheduledWork*> ScheduledWorkToRetract;
+		ScheduledWork.GenerateValueArray(ScheduledWorkToRetract);
+
+		for (FScheduledWork* Work : ScheduledWorkToRetract)
+		{
+			if (WrappedQueuedThreadPool->RetractQueuedWork(Work))
+			{
+				Work->InnerWork->Abandon();
+				ReleaseWorkNoLock(Work);
+			}
+		}
 	}
 
-	// We can't delete our WorkPool elements
-	// if they're still referenced by a threadpool.
-	// No choice to wait until they're all finished.
 	if (CurrentConcurrency)
 	{
+		// We can't delete our WorkPool elements
+		// if they're still referenced by a threadpool.
+		// Retraction didn't work, so no choice to wait until they're all finished.
 		TRACE_CPUPROFILER_EVENT_SCOPE(FQueuedThreadPoolWrapper::DestroyWait);
 		while (CurrentConcurrency)
 		{
@@ -78,10 +89,10 @@ void FQueuedThreadPoolWrapper::Destroy()
 
 	{
 		FRWScopeLock ScopeLock(Lock, SLT_Write);
-		for (FScheduledWork* ScheduledWork : WorkPool)
+		for (FScheduledWork* Work : WorkPool)
 		{
-			check(ScheduledWork->InnerWork == nullptr);
-			delete ScheduledWork;
+			check(Work->InnerWork == nullptr);
+			delete Work;
 		}
 		WorkPool.Empty();
 	}
@@ -123,8 +134,26 @@ void FQueuedThreadPoolWrapper::AddQueuedWork(IQueuedWork* InQueuedWork, EQueuedW
 
 bool FQueuedThreadPoolWrapper::RetractQueuedWork(IQueuedWork* InQueuedWork)
 {
-	FRWScopeLock ScopeLock(Lock, SLT_Write);
-	return QueuedWork.Retract(InQueuedWork);
+	FScheduledWork* Retracted = nullptr;
+	{
+		FRWScopeLock ScopeLock(Lock, SLT_Write);
+		if (QueuedWork.Retract(InQueuedWork))
+		{
+			return true;
+		}
+		
+		Retracted = ScheduledWork.FindRef(InQueuedWork);
+		if (Retracted && !WrappedQueuedThreadPool->RetractQueuedWork(Retracted))
+		{
+			Retracted = nullptr;
+		}
+	}
+
+	if (Retracted)
+	{
+		Schedule(Retracted);
+	}
+	return Retracted != nullptr;
 }
 
 int32 FQueuedThreadPoolWrapper::GetNumThreads() const
@@ -132,6 +161,14 @@ int32 FQueuedThreadPoolWrapper::GetNumThreads() const
 	return MaxConcurrency;
 }
 
+void FQueuedThreadPoolWrapper::ReleaseWorkNoLock(FScheduledWork* Work)
+{
+	CurrentConcurrency--;
+	ScheduledWork.Remove(Work->InnerWork);
+	Work->InnerWork = nullptr;
+	WorkPool.Push(Work);
+}
+
 void FQueuedThreadPoolWrapper::Schedule(FScheduledWork* Work)
 {
 	FRWScopeLock ScopeLock(Lock, SLT_Write);
@@ -141,9 +178,7 @@ void FQueuedThreadPoolWrapper::Schedule(FScheduledWork* Work)
 	// end up never being called again.
 	if (Work != nullptr)
 	{
-		CurrentConcurrency--;
-		WorkPool.Push(Work);
-		Work = nullptr;
+		ReleaseWorkNoLock(Work);
 	}
 
 	while (CurrentConcurrency < MaxConcurrency.Load(EMemoryOrder::Relaxed) && (MaxTaskToSchedule == -1 || MaxTaskToSchedule > 0))
@@ -164,8 +199,8 @@ void FQueuedThreadPoolWrapper::Schedule(FScheduledWork* Work)
 			}
 
 			Work->InnerWork = InnerWork;
+			ScheduledWork.Add(InnerWork, Work);
 			WrappedQueuedThreadPool->AddQueuedWork(Work, PriorityMapper(WorkPriority));
-			Work = nullptr;
 
 			if (MaxTaskToSchedule > 0)
 			{
diff --git a/Engine/Source/Runtime/Core/Public/Misc/QueuedThreadPool.h b/Engine/Source/Runtime/Core/Public/Misc/QueuedThreadPool.h
--- a/Engine/Source/Runtime/Core/Public/Misc/QueuedThreadPool.h
+++ b/Engine/Source/Runtime/Core/Public/Misc/QueuedThreadPool.h
@@ -49,6 +49,11 @@ public:
 	 * Get the total number of queued items.
 	 */
 	int32 Num() const { return NumQueuedWork; }
+
+	/**
+	 * Sort Priority Bucket given Predicate
+	 */
+	void Sort(EQueuedWorkPriority InPriorityBucket, TFunctionRef<bool (const IQueuedWork* A, const IQueuedWork* B)> Predicate);
 private:
 	/** The first queue to extract a work item from to avoid scanning all priorities when unqueuing. */
 	int32 FirstNonEmptyQueueIndex = 0;
diff --git a/Engine/Source/Runtime/Core/Public/Misc/QueuedThreadPoolWrapper.h b/Engine/Source/Runtime/Core/Public/Misc/QueuedThreadPoolWrapper.h
--- a/Engine/Source/Runtime/Core/Public/Misc/QueuedThreadPoolWrapper.h
+++ b/Engine/Source/Runtime/Core/Public/Misc/QueuedThreadPoolWrapper.h
@@ -9,7 +9,7 @@
 
 /** ThreadPool wrapper implementation allowing to schedule
   * up to MaxConcurrency tasks at a time making sub-partitioning
-  * another thread-pool a breeze and allowing more fine-grained control 
+  * another thread-pool a breeze and allowing more fine-grained control
   * over scheduling by effectively giving another set of priorities.
   */
 class CORE_API FQueuedThreadPoolWrapper : public FQueuedThreadPool
@@ -20,7 +20,7 @@ public:
 	 * InMaxConcurrency           Maximum number of concurrent tasks allowed, -1 will limit concurrency to number of threads available in the underlying thread pool.
 	 * InPriorityMapper           Thread-safe function used to map any priority from this Queue to the priority that should be used when scheduling the task on the underlying thread pool.
 	 */
-	FQueuedThreadPoolWrapper(FQueuedThreadPool* InWrappedQueuedThreadPool, int32 InMaxConcurrency = -1, TFunction<EQueuedWorkPriority (EQueuedWorkPriority)> InPriorityMapper = [](EQueuedWorkPriority InPriority) { return InPriority; });
+	FQueuedThreadPoolWrapper(FQueuedThreadPool* InWrappedQueuedThreadPool, int32 InMaxConcurrency = -1, TFunction<EQueuedWorkPriority(EQueuedWorkPriority)> InPriorityMapper = [](EQueuedWorkPriority InPriority) { return InPriority; });
 	~FQueuedThreadPoolWrapper();
 
 	/**
@@ -42,20 +42,59 @@ public:
 	bool RetractQueuedWork(IQueuedWork* InQueuedWork) override;
 	int32 GetNumThreads() const override;
 
+protected:
+	FRWLock Lock;
+	FThreadPoolPriorityQueue QueuedWork;
+
 private:
 	struct FScheduledWork;
 
 	bool Create(uint32 InNumQueuedThreads, uint32 StackSize, EThreadPriority ThreadPriority, const TCHAR* Name) override;
 	void Destroy() override;
 	void Schedule(FScheduledWork* Work = nullptr);
-
-	FRWLock Lock;
+	void ReleaseWorkNoLock(FScheduledWork* Work);
+	
 	TFunction<EQueuedWorkPriority(EQueuedWorkPriority)> PriorityMapper;
-	FThreadPoolPriorityQueue QueuedWork;
+
 	FQueuedThreadPool* WrappedQueuedThreadPool;
 	TArray<FScheduledWork*> WorkPool;
+	TMap<IQueuedWork*, FScheduledWork*> ScheduledWork;
 	TAtomic<int32> MaxConcurrency;
 	int32 MaxTaskToSchedule;
 	TAtomic<int32> CurrentConcurrency;
 	EQueuedWorkPriority WrappedQueuePriority;
 };
+
+/** ThreadPool wrapper implementation allowing to schedule
+  * up to MaxConcurrency tasks at a time making sub-partitioning
+  * another thread-pool a breeze and allowing more fine-grained control
+  * over scheduling by giving full control of task reordering.
+  */
+class CORE_API FQueuedThreadPoolDynamicWrapper : public FQueuedThreadPoolWrapper
+{
+public:
+	/**
+	 * InWrappedQueuedThreadPool  Underlying thread pool to schedule task to.
+	 * InMaxConcurrency           Maximum number of concurrent tasks allowed, -1 will limit concurrency to number of threads available in the underlying thread pool.
+	 * InPriorityMapper           Thread-safe function used to map any priority from this Queue to the priority that should be used when scheduling the task on the underlying thread pool.
+	 */
+	FQueuedThreadPoolDynamicWrapper(FQueuedThreadPool* InWrappedQueuedThreadPool, int32 InMaxConcurrency = -1, TFunction<EQueuedWorkPriority(EQueuedWorkPriority)> InPriorityMapper = [](EQueuedWorkPriority InPriority) { return InPriority; })
+		: FQueuedThreadPoolWrapper(InWrappedQueuedThreadPool, InMaxConcurrency, InPriorityMapper)
+	{
+	}
+
+	void AddQueuedWork(IQueuedWork* InQueuedWork, EQueuedWorkPriority InPriority = EQueuedWorkPriority::Normal) override
+	{
+		// Override priority to make sure all elements are in the same buckets and can then be sorted all together.
+		FQueuedThreadPoolWrapper::AddQueuedWork(InQueuedWork, EQueuedWorkPriority::Normal);
+	}
+
+	/**
+	 * Apply sort predicate to reorder the queued tasks
+	 */
+	void Sort(TFunctionRef<bool(const IQueuedWork* Lhs, const IQueuedWork* Rhs)> Predicate)
+	{
+		FRWScopeLock ScopeLock(Lock, SLT_Write);
+		QueuedWork.Sort(EQueuedWorkPriority::Normal, Predicate);
+	}
+};
\ No newline at end of file
