SS - Identify Static initialization for task scopes.

diff --git a/Engine/Source/Runtime/Core/Private/HAL/ThreadingBase.cpp b/Engine/Source/Runtime/Core/Private/HAL/ThreadingBase.cpp
--- a/Engine/Source/Runtime/Core/Private/HAL/ThreadingBase.cpp
+++ b/Engine/Source/Runtime/Core/Private/HAL/ThreadingBase.cpp
@@ -39,7 +39,13 @@ FQueuedThreadPool* GBackgroundPriorityThreadPool = nullptr;
 FQueuedThreadPool* GLargeThreadPool = nullptr;
 #endif
 
-thread_local ETaskTag FTaskTagScope::ActiveTaskTag = ETaskTag::ENone;
+int32 FTaskTagScope::GetStaticThreadId()
+{
+	static int32 ThreadID = FPlatformTLS::GetCurrentThreadId();
+	return ThreadID;
+}
+
+thread_local ETaskTag FTaskTagScope::ActiveTaskTag = ETaskTag::EStaticInit;
 static std::atomic_int ActiveNamedThreads = {};
 
 FTaskTagScope::FTaskTagScope(bool InTagOnlyIfNone, ETaskTag InTag) : Tag(InTag), TagOnlyIfNone(InTagOnlyIfNone)
@@ -47,6 +53,11 @@ FTaskTagScope::FTaskTagScope(bool InTagOnlyIfNone, ETaskTag InTag) : Tag(InTag),
 	checkf(Tag != ETaskTag::ENone, TEXT("None cannot be used as a Tag"));
 	checkf(Tag != ETaskTag::EParallelThread, TEXT("Parallel cannot be used on it's own"));
 
+	if (ActiveTaskTag == ETaskTag::EStaticInit)
+	{
+		checkf(Tag == ETaskTag::EGameThread && IsRunningDuringStaticInit(), TEXT("The Gamethread can only be tagged on the inital thread of the application"));
+	}
+
 	if (!EnumHasAllFlags(Tag, ETaskTag::EParallelThread))
 	{
 		ETaskTag NamedThreadBits = (Tag & ETaskTag::ENamedThreadBits);
@@ -97,7 +108,7 @@ FTaskTagScope::~FTaskTagScope()
 	}
 
 	//prolong the scope of the GT for static variable destructors
-	if (Tag == ETaskTag::EGameThread && ActiveTaskTag == ETaskTag::ENone)
+	if (Tag == ETaskTag::EGameThread && ActiveTaskTag == ETaskTag::EStaticInit)
 	{
 		ActiveTaskTag = ETaskTag::EGameThread;
 	}
@@ -113,11 +124,16 @@ ETaskTag FTaskTagScope::GetCurrentTag()
 	return ActiveTaskTag;
 }
 
+bool FTaskTagScope::IsRunningDuringStaticInit()
+{
+	return ActiveTaskTag == ETaskTag::EStaticInit && GetStaticThreadId() == FPlatformTLS::GetCurrentThreadId();
+}
+
 CORE_API bool IsInGameThread()
 {
 	if (GIsGameThreadIdInitialized)
 	{
-		bool newValue = FTaskTagScope::IsCurrentTag(ETaskTag::EGameThread);
+		bool newValue = FTaskTagScope::IsCurrentTag(ETaskTag::EGameThread) || FTaskTagScope::IsRunningDuringStaticInit();
 #if !UE_BUILD_SHIPPING && !UE_BUILD_TEST
 		const uint32 CurrentThreadId = FPlatformTLS::GetCurrentThreadId();
 		bool oldValue = CurrentThreadId == GGameThreadId;
@@ -177,7 +193,7 @@ CORE_API bool IsInActualRenderingThread()
 CORE_API bool IsInRenderingThread()
 {
 	bool newValue = !GRenderingThread || GIsRenderingThreadSuspended.Load(EMemoryOrder::Relaxed)
-		? FTaskTagScope::IsCurrentTag(ETaskTag::EGameThread) || FTaskTagScope::IsCurrentTag(ETaskTag::ENone)
+		? FTaskTagScope::IsCurrentTag(ETaskTag::EGameThread) || FTaskTagScope::IsRunningDuringStaticInit()
 		: FTaskTagScope::IsCurrentTag(ETaskTag::ERenderingThread);
 
 #if !UE_BUILD_SHIPPING && !UE_BUILD_TEST
@@ -694,6 +710,9 @@ void FRunnableThread::SetTls()
 	check( ThreadID == FPlatformTLS::GetCurrentThreadId() );
 	check( FPlatformTLS::IsValidTlsSlot(RunnableTlsSlot) );
 	FPlatformTLS::SetTlsValue( RunnableTlsSlot, this );
+
+	check(FTaskTagScope::ActiveTaskTag == ETaskTag::EStaticInit);
+	FTaskTagScope::ActiveTaskTag = ETaskTag::ENone;
 }
 
 void FRunnableThread::FreeTls()
diff --git a/Engine/Source/Runtime/Core/Public/CoreGlobals.h b/Engine/Source/Runtime/Core/Public/CoreGlobals.h
--- a/Engine/Source/Runtime/Core/Public/CoreGlobals.h
+++ b/Engine/Source/Runtime/Core/Public/CoreGlobals.h
@@ -464,6 +464,7 @@ extern CORE_API bool GEnableVREditorHacks;
 enum class ETaskTag : int32
 {
 	ENone						= 0 << 0,
+	EStaticInit					= 1 << 0,
 	EGameThread					= 1 << 1,
 	ESlateThread				= 1 << 2,
 	EAudioThread				= 1 << 3,
@@ -487,26 +488,29 @@ ENUM_CLASS_FLAGS(ETaskTag);
  *
  * @param CtorSignature InTag the Tag to use
  */
-class FTaskTagScope
+class CORE_API FTaskTagScope
 {
+	friend class FRunnableThread;
 	static thread_local ETaskTag ActiveTaskTag;
+	static int32 GetStaticThreadId();
 	ETaskTag ParentTag;
 	ETaskTag Tag;
 	bool TagOnlyIfNone;
 
 protected:
-	CORE_API FTaskTagScope(bool InTagOnlyIfNone, ETaskTag InTag);
+	FTaskTagScope(bool InTagOnlyIfNone, ETaskTag InTag);
 
 public:
-	CORE_API FTaskTagScope(ETaskTag InTag = ETaskTag::ENone) : FTaskTagScope(false, InTag)
+	FTaskTagScope(ETaskTag InTag = ETaskTag::ENone) : FTaskTagScope(false, InTag)
 	{
 
 	}
 
-	CORE_API ~FTaskTagScope();
+	~FTaskTagScope();
 
 	static ETaskTag GetCurrentTag();
 	static bool IsCurrentTag(ETaskTag InTag);
+	static bool IsRunningDuringStaticInit();
 };
 
 /**
@@ -515,10 +519,10 @@ public:
  *
  * @param CtorSignature InTag the Tag to use
  */
-class FOptionalTaskTagScope : public FTaskTagScope
+class CORE_API FOptionalTaskTagScope : public FTaskTagScope
 {
 public:
-	CORE_API FOptionalTaskTagScope(ETaskTag InTag = ETaskTag::ENone) : FTaskTagScope(true, InTag)
+	FOptionalTaskTagScope(ETaskTag InTag = ETaskTag::ENone) : FTaskTagScope(true, InTag)
 	{
 
 	}
