Backout changelist 14019473 as it is breaking CIS
#rnx

diff --git a/Engine/Source/Runtime/Core/Private/HAL/ThreadingBase.cpp b/Engine/Source/Runtime/Core/Private/HAL/ThreadingBase.cpp
--- a/Engine/Source/Runtime/Core/Private/HAL/ThreadingBase.cpp
+++ b/Engine/Source/Runtime/Core/Private/HAL/ThreadingBase.cpp
@@ -39,13 +39,7 @@ FQueuedThreadPool* GBackgroundPriorityThreadPool = nullptr;
 FQueuedThreadPool* GLargeThreadPool = nullptr;
 #endif
 
-int32 FTaskTagScope::GetStaticThreadId()
-{
-	static int32 ThreadID = FPlatformTLS::GetCurrentThreadId();
-	return ThreadID;
-}
-
-thread_local ETaskTag FTaskTagScope::ActiveTaskTag = ETaskTag::EStaticInit;
+thread_local ETaskTag FTaskTagScope::ActiveTaskTag = ETaskTag::ENone;
 static std::atomic_int ActiveNamedThreads = {};
 
 FTaskTagScope::FTaskTagScope(bool InTagOnlyIfNone, ETaskTag InTag) : Tag(InTag), TagOnlyIfNone(InTagOnlyIfNone)
@@ -53,11 +47,6 @@ FTaskTagScope::FTaskTagScope(bool InTagOnlyIfNone, ETaskTag InTag) : Tag(InTag),
 	checkf(Tag != ETaskTag::ENone, TEXT("None cannot be used as a Tag"));
 	checkf(Tag != ETaskTag::EParallelThread, TEXT("Parallel cannot be used on it's own"));
 
-	if (ActiveTaskTag == ETaskTag::EStaticInit)
-	{
-		checkf(Tag == ETaskTag::EGameThread && IsRunningDuringStaticInit(), TEXT("The Gamethread can only be tagged on the inital thread of the application"));
-	}
-
 	if (!EnumHasAllFlags(Tag, ETaskTag::EParallelThread))
 	{
 		ETaskTag NamedThreadBits = (Tag & ETaskTag::ENamedThreadBits);
@@ -108,7 +97,7 @@ FTaskTagScope::~FTaskTagScope()
 	}
 
 	//prolong the scope of the GT for static variable destructors
-	if (Tag == ETaskTag::EGameThread && ActiveTaskTag == ETaskTag::EStaticInit)
+	if (Tag == ETaskTag::EGameThread && ActiveTaskTag == ETaskTag::ENone)
 	{
 		ActiveTaskTag = ETaskTag::EGameThread;
 	}
@@ -124,16 +113,11 @@ ETaskTag FTaskTagScope::GetCurrentTag()
 	return ActiveTaskTag;
 }
 
-bool FTaskTagScope::IsRunningDuringStaticInit()
-{
-	return ActiveTaskTag == ETaskTag::EStaticInit && GetStaticThreadId() == FPlatformTLS::GetCurrentThreadId();
-}
-
 CORE_API bool IsInGameThread()
 {
 	if (GIsGameThreadIdInitialized)
 	{
-		bool newValue = FTaskTagScope::IsCurrentTag(ETaskTag::EGameThread) || FTaskTagScope::IsRunningDuringStaticInit();
+		bool newValue = FTaskTagScope::IsCurrentTag(ETaskTag::EGameThread);
 #if !UE_BUILD_SHIPPING && !UE_BUILD_TEST
 		const uint32 CurrentThreadId = FPlatformTLS::GetCurrentThreadId();
 		bool oldValue = CurrentThreadId == GGameThreadId;
@@ -193,7 +177,7 @@ CORE_API bool IsInActualRenderingThread()
 CORE_API bool IsInRenderingThread()
 {
 	bool newValue = !GRenderingThread || GIsRenderingThreadSuspended.Load(EMemoryOrder::Relaxed)
-		? FTaskTagScope::IsCurrentTag(ETaskTag::EGameThread) || FTaskTagScope::IsRunningDuringStaticInit()
+		? FTaskTagScope::IsCurrentTag(ETaskTag::EGameThread) || FTaskTagScope::IsCurrentTag(ETaskTag::ENone)
 		: FTaskTagScope::IsCurrentTag(ETaskTag::ERenderingThread);
 
 #if !UE_BUILD_SHIPPING && !UE_BUILD_TEST
@@ -710,9 +694,6 @@ void FRunnableThread::SetTls()
 	check( ThreadID == FPlatformTLS::GetCurrentThreadId() );
 	check( FPlatformTLS::IsValidTlsSlot(RunnableTlsSlot) );
 	FPlatformTLS::SetTlsValue( RunnableTlsSlot, this );
-
-	check(FTaskTagScope::ActiveTaskTag == ETaskTag::EStaticInit);
-	FTaskTagScope::ActiveTaskTag = ETaskTag::ENone;
 }
 
 void FRunnableThread::FreeTls()
diff --git a/Engine/Source/Runtime/Core/Public/CoreGlobals.h b/Engine/Source/Runtime/Core/Public/CoreGlobals.h
--- a/Engine/Source/Runtime/Core/Public/CoreGlobals.h
+++ b/Engine/Source/Runtime/Core/Public/CoreGlobals.h
@@ -464,7 +464,6 @@ extern CORE_API bool GEnableVREditorHacks;
 enum class ETaskTag : int32
 {
 	ENone						= 0 << 0,
-	EStaticInit					= 1 << 0,
 	EGameThread					= 1 << 1,
 	ESlateThread				= 1 << 2,
 	EAudioThread				= 1 << 3,
@@ -488,29 +487,26 @@ ENUM_CLASS_FLAGS(ETaskTag);
  *
  * @param CtorSignature InTag the Tag to use
  */
-class CORE_API FTaskTagScope
+class FTaskTagScope
 {
-	friend class FRunnableThread;
 	static thread_local ETaskTag ActiveTaskTag;
-	static int32 GetStaticThreadId();
 	ETaskTag ParentTag;
 	ETaskTag Tag;
 	bool TagOnlyIfNone;
 
 protected:
-	FTaskTagScope(bool InTagOnlyIfNone, ETaskTag InTag);
+	CORE_API FTaskTagScope(bool InTagOnlyIfNone, ETaskTag InTag);
 
 public:
-	FTaskTagScope(ETaskTag InTag = ETaskTag::ENone) : FTaskTagScope(false, InTag)
+	CORE_API FTaskTagScope(ETaskTag InTag = ETaskTag::ENone) : FTaskTagScope(false, InTag)
 	{
 
 	}
 
-	~FTaskTagScope();
+	CORE_API ~FTaskTagScope();
 
 	static ETaskTag GetCurrentTag();
 	static bool IsCurrentTag(ETaskTag InTag);
-	static bool IsRunningDuringStaticInit();
 };
 
 /**
@@ -519,10 +515,10 @@ public:
  *
  * @param CtorSignature InTag the Tag to use
  */
-class CORE_API FOptionalTaskTagScope : public FTaskTagScope
+class FOptionalTaskTagScope : public FTaskTagScope
 {
 public:
-	FOptionalTaskTagScope(ETaskTag InTag = ETaskTag::ENone) : FTaskTagScope(true, InTag)
+	CORE_API FOptionalTaskTagScope(ETaskTag InTag = ETaskTag::ENone) : FTaskTagScope(true, InTag)
 	{
 
 	}
