SS - Identify Static initialization for task scopes.

diff --git a/Engine/Source/Runtime/Core/Private/HAL/ThreadingBase.cpp b/Engine/Source/Runtime/Core/Private/HAL/ThreadingBase.cpp
--- a/Engine/Source/Runtime/Core/Private/HAL/ThreadingBase.cpp
+++ b/Engine/Source/Runtime/Core/Private/HAL/ThreadingBase.cpp
@@ -39,14 +39,30 @@ FQueuedThreadPool* GBackgroundPriorityThreadPool = nullptr;
 FQueuedThreadPool* GLargeThreadPool = nullptr;
 #endif
 
-thread_local ETaskTag FTaskTagScope::ActiveTaskTag = ETaskTag::ENone;
+int32 FTaskTagScope::GetStaticThreadId()
+{
+	static int32 ThreadID = FPlatformTLS::GetCurrentThreadId();
+	return ThreadID;
+}
+
+thread_local ETaskTag FTaskTagScope::ActiveTaskTag = ETaskTag::EStaticInit;
 static std::atomic_int ActiveNamedThreads = {};
 
+void FTaskTagScope::SetTagNone()
+{
+	ActiveTaskTag = ETaskTag::ENone;
+}
+
 FTaskTagScope::FTaskTagScope(bool InTagOnlyIfNone, ETaskTag InTag) : Tag(InTag), TagOnlyIfNone(InTagOnlyIfNone)
 {
 	checkf(Tag != ETaskTag::ENone, TEXT("None cannot be used as a Tag"));
 	checkf(Tag != ETaskTag::EParallelThread, TEXT("Parallel cannot be used on it's own"));
 
+	if (ActiveTaskTag == ETaskTag::EStaticInit)
+	{
+		checkf(Tag == ETaskTag::EGameThread && IsRunningDuringStaticInit(), TEXT("The Gamethread can only be tagged on the inital thread of the application"));
+	}
+
 	if (!EnumHasAllFlags(Tag, ETaskTag::EParallelThread))
 	{
 		ETaskTag NamedThreadBits = (Tag & ETaskTag::ENamedThreadBits);
@@ -97,7 +113,7 @@ FTaskTagScope::~FTaskTagScope()
 	}
 
 	//prolong the scope of the GT for static variable destructors
-	if (Tag == ETaskTag::EGameThread && ActiveTaskTag == ETaskTag::ENone)
+	if (Tag == ETaskTag::EGameThread && ActiveTaskTag == ETaskTag::EStaticInit)
 	{
 		ActiveTaskTag = ETaskTag::EGameThread;
 	}
@@ -113,11 +129,16 @@ ETaskTag FTaskTagScope::GetCurrentTag()
 	return ActiveTaskTag;
 }
 
+bool FTaskTagScope::IsRunningDuringStaticInit()
+{
+	return ActiveTaskTag == ETaskTag::EStaticInit && GetStaticThreadId() == FPlatformTLS::GetCurrentThreadId();
+}
+
 CORE_API bool IsInGameThread()
 {
 	if (GIsGameThreadIdInitialized)
 	{
-		bool newValue = FTaskTagScope::IsCurrentTag(ETaskTag::EGameThread);
+		bool newValue = FTaskTagScope::IsCurrentTag(ETaskTag::EGameThread) || FTaskTagScope::IsRunningDuringStaticInit();
 #if !UE_BUILD_SHIPPING && !UE_BUILD_TEST
 		const uint32 CurrentThreadId = FPlatformTLS::GetCurrentThreadId();
 		bool oldValue = CurrentThreadId == GGameThreadId;
@@ -177,7 +198,7 @@ CORE_API bool IsInActualRenderingThread()
 CORE_API bool IsInRenderingThread()
 {
 	bool newValue = !GRenderingThread || GIsRenderingThreadSuspended.Load(EMemoryOrder::Relaxed)
-		? FTaskTagScope::IsCurrentTag(ETaskTag::EGameThread) || FTaskTagScope::IsCurrentTag(ETaskTag::ENone)
+		? FTaskTagScope::IsCurrentTag(ETaskTag::EGameThread) || FTaskTagScope::IsCurrentTag(ETaskTag::ERenderingThread) || FTaskTagScope::IsRunningDuringStaticInit()
 		: FTaskTagScope::IsCurrentTag(ETaskTag::ERenderingThread);
 
 #if !UE_BUILD_SHIPPING && !UE_BUILD_TEST
@@ -694,6 +715,7 @@ void FRunnableThread::SetTls()
 	check( ThreadID == FPlatformTLS::GetCurrentThreadId() );
 	check( FPlatformTLS::IsValidTlsSlot(RunnableTlsSlot) );
 	FPlatformTLS::SetTlsValue( RunnableTlsSlot, this );
+	FTaskTagScope::SetTagNone();
 }
 
 void FRunnableThread::FreeTls()
diff --git a/Engine/Source/Runtime/Core/Public/CoreGlobals.h b/Engine/Source/Runtime/Core/Public/CoreGlobals.h
--- a/Engine/Source/Runtime/Core/Public/CoreGlobals.h
+++ b/Engine/Source/Runtime/Core/Public/CoreGlobals.h
@@ -464,6 +464,7 @@ extern CORE_API bool GEnableVREditorHacks;
 enum class ETaskTag : int32
 {
 	ENone						= 0 << 0,
+	EStaticInit					= 1 << 0,
 	EGameThread					= 1 << 1,
 	ESlateThread				= 1 << 2,
 	EAudioThread				= 1 << 3,
@@ -489,11 +490,16 @@ ENUM_CLASS_FLAGS(ETaskTag);
  */
 class FTaskTagScope
 {
+	friend class FRunnableThread;
+	friend class FRenderingThread;
 	static thread_local ETaskTag ActiveTaskTag;
+	static int32 GetStaticThreadId();
 	ETaskTag ParentTag;
 	ETaskTag Tag;
 	bool TagOnlyIfNone;
 
+	static void CORE_API SetTagNone();
+
 protected:
 	CORE_API FTaskTagScope(bool InTagOnlyIfNone, ETaskTag InTag);
 
@@ -505,8 +511,9 @@ public:
 
 	CORE_API ~FTaskTagScope();
 
-	static ETaskTag GetCurrentTag();
-	static bool IsCurrentTag(ETaskTag InTag);
+	static CORE_API ETaskTag GetCurrentTag();
+	static CORE_API bool IsCurrentTag(ETaskTag InTag);
+	static CORE_API bool IsRunningDuringStaticInit();
 };
 
 /**
diff --git a/Engine/Source/Runtime/RenderCore/Private/RenderingThread.cpp b/Engine/Source/Runtime/RenderCore/Private/RenderingThread.cpp
--- a/Engine/Source/Runtime/RenderCore/Private/RenderingThread.cpp
+++ b/Engine/Source/Runtime/RenderCore/Private/RenderingThread.cpp
@@ -450,6 +450,8 @@ public:
 	virtual bool Init(void) override
 	{ 
 		GRenderThreadId = FPlatformTLS::GetCurrentThreadId();
+		FTaskTagScope::SetTagNone();
+		FTaskTagScope Scope(ETaskTag::ERenderingThread);
 
 		// Acquire rendering context ownership on the current thread, unless using an RHI thread, which will be the real owner
 		if (!IsRunningRHIInSeparateThread())
