Fix assertion on Mac caused by IsRunningDuringStaticInit().

#rnx

diff --git a/Engine/Source/Runtime/Core/Core.Build.cs b/Engine/Source/Runtime/Core/Core.Build.cs
--- a/Engine/Source/Runtime/Core/Core.Build.cs
+++ b/Engine/Source/Runtime/Core/Core.Build.cs
@@ -23,7 +23,7 @@ public class Core : ModuleRules
 				"Runtime/SynthBenchmark/Public",
 				"Runtime/Core/Private",
 				"Runtime/Core/Private/Misc",
-                "Runtime/Core/Private/Internationalization",
+				"Runtime/Core/Private/Internationalization",
 				"Runtime/Core/Private/Internationalization/Cultures",
 				"Runtime/Engine/Public",
 			}
@@ -33,8 +33,8 @@ public class Core : ModuleRules
 			new string[] {
 				"TargetPlatform",
 				"DerivedDataCache",
-                "InputDevice",
-                "Analytics",
+				"InputDevice",
+				"Analytics",
 				"RHI"
 			}
 			);
@@ -52,7 +52,7 @@ public class Core : ModuleRules
 			AddEngineThirdPartyPrivateStaticDependencies(Target,
 				"zlib");
 
-			AddEngineThirdPartyPrivateStaticDependencies(Target, 
+			AddEngineThirdPartyPrivateStaticDependencies(Target,
 				"IntelTBB",
 				"IntelVTune"
 				);
@@ -60,7 +60,7 @@ public class Core : ModuleRules
 			AddEngineThirdPartyPrivateStaticDependencies(Target,
 				"mimalloc");
 
-			if(Target.Platform != UnrealTargetPlatform.Win32 && Target.WindowsPlatform.bUseBundledDbgHelp)
+			if (Target.Platform != UnrealTargetPlatform.Win32 && Target.WindowsPlatform.bUseBundledDbgHelp)
 			{
 				PublicDelayLoadDLLs.Add("DBGHELP.DLL");
 				PrivateDefinitions.Add("USE_BUNDLED_DBGHELP=1");
@@ -92,7 +92,7 @@ public class Core : ModuleRules
 				"rd_route"
 				);
 			PublicFrameworks.AddRange(new string[] { "Cocoa", "Carbon", "IOKit", "Security" });
-			
+
 			if (Target.bBuildEditor == true)
 			{
 				PublicAdditionalLibraries.Add("/System/Library/PrivateFrameworks/MultitouchSupport.framework/Versions/Current/MultitouchSupport");
@@ -107,9 +107,9 @@ public class Core : ModuleRules
 			if (Target.Platform == UnrealTargetPlatform.IOS)
 			{
 				PublicFrameworks.AddRange(new string[] { "CoreMotion", "AdSupport", "WebKit" });
-                AddEngineThirdPartyPrivateStaticDependencies(Target,
-                    "PLCrashReporter"
-                    );
+				AddEngineThirdPartyPrivateStaticDependencies(Target,
+					"PLCrashReporter"
+					);
 			}
 
 			PrivateIncludePathModuleNames.Add("ApplicationCore");
@@ -132,23 +132,23 @@ public class Core : ModuleRules
 				);
 		}
 		else if (Target.IsInPlatformGroup(UnrealPlatformGroup.Unix))
-        {
+		{
 			AddEngineThirdPartyPrivateStaticDependencies(Target,
 				"zlib",
 				"jemalloc"
-                );
+				);
 
 			// Core uses dlopen()
 			PublicSystemLibraries.Add("dl");
-        }
+		}
 
-		if ( Target.bCompileICU == true )
-        {
+		if (Target.bCompileICU == true)
+		{
 			AddEngineThirdPartyPrivateStaticDependencies(Target, "ICU");
-        }
-        PublicDefinitions.Add("UE_ENABLE_ICU=" + (Target.bCompileICU ? "1" : "0")); // Enable/disable (=1/=0) ICU usage in the codebase. NOTE: This flag is for use while integrating ICU and will be removed afterward.
+		}
+		PublicDefinitions.Add("UE_ENABLE_ICU=" + (Target.bCompileICU ? "1" : "0")); // Enable/disable (=1/=0) ICU usage in the codebase. NOTE: This flag is for use while integrating ICU and will be removed afterward.
 
-        // If we're compiling with the engine, then add Core's engine dependencies
+		// If we're compiling with the engine, then add Core's engine dependencies
 		if (Target.bCompileAgainstEngine == true)
 		{
 			if (!Target.bBuildRequiresCookedData)
@@ -157,12 +157,12 @@ public class Core : ModuleRules
 			}
 		}
 
-		
+
 		// On Windows platform, VSPerfExternalProfiler.cpp needs access to "VSPerf.h".  This header is included with Visual Studio, but it's not in a standard include path.
-		if(Target.Platform.IsInGroup(UnrealPlatformGroup.Windows))
+		if (Target.Platform.IsInGroup(UnrealPlatformGroup.Windows))
 		{
 			var VisualStudioVersionNumber = "11.0";
-			var SubFolderName = ( Target.Platform == UnrealTargetPlatform.Win32 ) ? "PerfSDK" : "x64/PerfSDK";
+			var SubFolderName = (Target.Platform == UnrealTargetPlatform.Win32) ? "PerfSDK" : "x64/PerfSDK";
 
 			string PerfIncludeDirectory = Path.Combine(Environment.GetFolderPath(Environment.SpecialFolder.ProgramFilesX86), String.Format("Microsoft Visual Studio {0}/Team Tools/Performance Tools/{1}", VisualStudioVersionNumber, SubFolderName));
 
@@ -177,22 +177,31 @@ public class Core : ModuleRules
 			}
 		}
 
-		if( Target.Platform == UnrealTargetPlatform.HoloLens)
+		if (Target.Platform == UnrealTargetPlatform.HoloLens)
 		{
 			PublicDefinitions.Add("WITH_VS_PERF_PROFILER=0");
 		}
 
-        if (Target.bWithDirectXMath && (Target.Platform == UnrealTargetPlatform.Win64 || Target.Platform == UnrealTargetPlatform.Win32))
-        {
-            PublicDefinitions.Add("WITH_DIRECTXMATH=1");
-        }
-        else
-        {
-            PublicDefinitions.Add("WITH_DIRECTXMATH=0");
-        }
+		if (Target.bWithDirectXMath && (Target.Platform == UnrealTargetPlatform.Win64 || Target.Platform == UnrealTargetPlatform.Win32))
+		{
+			PublicDefinitions.Add("WITH_DIRECTXMATH=1");
+		}
+		else
+		{
+			PublicDefinitions.Add("WITH_DIRECTXMATH=0");
+		}
+
+		if (Target.Platform == UnrealTargetPlatform.Mac)
+		{
+			PublicDefinitions.Add("IS_RUNNING_GAMETHREAD_ON_EXTERNAL_THREAD=1");
+		}
+		else
+		{
+			PublicDefinitions.Add("IS_RUNNING_GAMETHREAD_ON_EXTERNAL_THREAD=0");
+		}
 
 		// Set a macro to allow FApp::GetBuildTargetType() to detect client targts
-		if(Target.Type == TargetRules.TargetType.Client)
+		if (Target.Type == TargetRules.TargetType.Client)
 		{
 			PrivateDefinitions.Add("IS_CLIENT_TARGET=1");
 		}
@@ -205,8 +214,8 @@ public class Core : ModuleRules
 		// Run-time validation must be enabled through '-stompmalloc' command line argument.
 
 		bool bWithMallocStomp = false;
-        if (Target.Configuration != UnrealTargetConfiguration.Shipping)
-        {
+		if (Target.Configuration != UnrealTargetConfiguration.Shipping)
+		{
 			if (Target.Platform == UnrealTargetPlatform.Mac ||
 				Target.Platform == UnrealTargetPlatform.Linux ||
 				Target.Platform == UnrealTargetPlatform.LinuxAArch64 ||
@@ -216,9 +225,9 @@ public class Core : ModuleRules
 			{
 				bWithMallocStomp = true;
 			}
-        }
+		}
 
-        PublicDefinitions.Add("WITH_MALLOC_STOMP=" + (bWithMallocStomp ? "1" : "0"));
+		PublicDefinitions.Add("WITH_MALLOC_STOMP=" + (bWithMallocStomp ? "1" : "0"));
 
 		PrivateDefinitions.Add("PLATFORM_COMPILER_OPTIMIZATION_LTCG=" + (Target.bAllowLTCG ? "1" : "0"));
 		PrivateDefinitions.Add("PLATFORM_COMPILER_OPTIMIZATION_PG=" + (Target.bPGOOptimize ? "1" : "0"));
diff --git a/Engine/Source/Runtime/Core/Private/HAL/ThreadingBase.cpp b/Engine/Source/Runtime/Core/Private/HAL/ThreadingBase.cpp
--- a/Engine/Source/Runtime/Core/Private/HAL/ThreadingBase.cpp
+++ b/Engine/Source/Runtime/Core/Private/HAL/ThreadingBase.cpp
@@ -60,7 +60,11 @@ FTaskTagScope::FTaskTagScope(bool InTagOnlyIfNone, ETaskTag InTag) : Tag(InTag),
 
 	if (ActiveTaskTag == ETaskTag::EStaticInit)
 	{
+#if IS_RUNNING_GAMETHREAD_ON_EXTERNAL_THREAD
+		checkf(Tag == ETaskTag::EGameThread, TEXT("The Gamethread can only be tagged on the inital thread of the application"));
+#else
 		checkf(Tag == ETaskTag::EGameThread && IsRunningDuringStaticInit(), TEXT("The Gamethread can only be tagged on the inital thread of the application"));
+#endif
 	}
 
 	if (!EnumHasAllFlags(Tag, ETaskTag::EParallelThread))
