fixed the IOdispatcher problem at mounting utoc

diff --git a/Engine/Source/Runtime/Core/Private/IO/IoDispatcherFileBackend.cpp b/Engine/Source/Runtime/Core/Private/IO/IoDispatcherFileBackend.cpp
--- a/Engine/Source/Runtime/Core/Private/IO/IoDispatcherFileBackend.cpp
+++ b/Engine/Source/Runtime/Core/Private/IO/IoDispatcherFileBackend.cpp
@@ -136,7 +136,8 @@ FIoStatus FFileIoStoreReader::Initialize(const FIoStoreEnvironment& Environment)
 
 	ContainerFile.FilePath = ContainerFilePath;
 
-	FIoStoreTocResource TocResource;
+	TUniquePtr<FIoStoreTocResource> TocResourcePtr = MakeUnique<FIoStoreTocResource>();
+	FIoStoreTocResource& TocResource = *TocResourcePtr;
 	FIoStatus Status = FIoStoreTocResource::Read(*TocFilePath, EIoStoreTocReadOptions::ExcludeTocMeta, TocResource);
 	if (!Status.IsOk())
 	{
