SS - Added more debug information for task scopes

diff --git a/Engine/Source/Runtime/Core/Private/HAL/ThreadingBase.cpp b/Engine/Source/Runtime/Core/Private/HAL/ThreadingBase.cpp
--- a/Engine/Source/Runtime/Core/Private/HAL/ThreadingBase.cpp
+++ b/Engine/Source/Runtime/Core/Private/HAL/ThreadingBase.cpp
@@ -78,7 +78,7 @@ FTaskTagScope::FTaskTagScope(bool InTagOnlyIfNone, ETaskTag InTag) : Tag(InTag),
 			//Try to catch other Threads that already opened a non parallel scope
 			ActiveNamedThreads.store(int32(ETaskTag::ENone));
 		}
-		checkf(IsOK, TEXT("Only Scopes tagged with ETaskTag::EParallelThread can be tagged multiple times. ActiveNamedThreads(%h) cannot be tagged multiple times in the same callstack you can use FOptionalTaskTagScope to avoid retagging check the ActiveNamedThreads(%h) with the current Tag(%h)"), OldTag, FTaskTagScope::GetCurrentTag(), Tag);
+		checkf(IsOK, TEXT("Only Scopes tagged with ETaskTag::EParallelThread can be tagged multiple times. ActiveNamedThreads(%x) cannot be tagged multiple times in the same callstack you can use FOptionalTaskTagScope to avoid retagging check the ActiveNamedThreads(%x) with the current Tag(%x)"), OldTag, FTaskTagScope::GetCurrentTag(), Tag);
 	}
 	
 	ParentTag = ActiveTaskTag;
@@ -90,19 +90,19 @@ FTaskTagScope::FTaskTagScope(bool InTagOnlyIfNone, ETaskTag InTag) : Tag(InTag),
 	{
 		if (EnumHasAllFlags(Tag, ETaskTag::EParallelRenderingThread))
 		{
-			checkf(IsInRenderingThread(), TEXT("ETaskTag::EParallelRenderingThread can only be retagged if they are in a parallel for on the RenderingThread or not tagged check the ActiveNamedThreads(%h)"), FTaskTagScope::GetCurrentTag());
+			checkf(IsInRenderingThread(), TEXT("ETaskTag::EParallelRenderingThread can only be retagged if they are in a parallel for on the RenderingThread or not tagged check the ActiveNamedThreads(%x)"), FTaskTagScope::GetCurrentTag());
 		}
 
 		if (EnumHasAllFlags(Tag, ETaskTag::EParallelGameThread))
 		{
-			checkf(IsInGameThread(), TEXT("ETaskTag::EParallelGameThread can only be retagged if they are in a parallel for on the GameThread or not tagged check the ActiveNamedThreads(%h)"), FTaskTagScope::GetCurrentTag());
+			checkf(IsInGameThread(), TEXT("ETaskTag::EParallelGameThread can only be retagged if they are in a parallel for on the GameThread or not tagged check the ActiveNamedThreads(%x)"), FTaskTagScope::GetCurrentTag());
 		}
 	}
 }
 
 FTaskTagScope::~FTaskTagScope()
 {
-	checkf(TagOnlyIfNone || ActiveTaskTag == Tag, TEXT("ActiveTaskTag(%h) corrupted needs to be Tag(%h)"), FTaskTagScope::GetCurrentTag(), Tag);
+	checkf(TagOnlyIfNone || ActiveTaskTag == Tag, TEXT("ActiveTaskTag(%x) corrupted needs to be Tag(%x)"), FTaskTagScope::GetCurrentTag(), Tag);
 	if (!TagOnlyIfNone || ParentTag == ETaskTag::ENone)
 	{
 		ActiveTaskTag = ParentTag;
@@ -113,7 +113,7 @@ FTaskTagScope::~FTaskTagScope()
 		ETaskTag NamedThreadBits = (Tag & ETaskTag::ENamedThreadBits);
 		static_assert(sizeof(ETaskTag) == sizeof(int32), "EnumSize must match interlockedAnd");
 		ETaskTag OldTag = ETaskTag(ActiveNamedThreads.fetch_and(int32(~NamedThreadBits)));
-		checkf((OldTag & NamedThreadBits) == (NamedThreadBits), TEXT("Currently active Threads(%h) got corrupted check the ActiveNamedThreads(%h)"), OldTag, FTaskTagScope::GetCurrentTag());
+		checkf((OldTag & NamedThreadBits) == (NamedThreadBits), TEXT("Currently active Threads(%x) got corrupted check the ActiveNamedThreads(%x)"), OldTag, FTaskTagScope::GetCurrentTag());
 	}
 
 	//prolong the scope of the GT for static variable destructors
@@ -146,7 +146,8 @@ CORE_API bool IsInGameThread()
 #if !UE_BUILD_SHIPPING && !UE_BUILD_TEST
 		const uint32 CurrentThreadId = FPlatformTLS::GetCurrentThreadId();
 		bool oldValue = CurrentThreadId == GGameThreadId;
-		checkf(oldValue == newValue, TEXT("If this check fails make sure that there is a FTaskTagScope(ETaskTag::EGameThread) as deep as possible on the current callstack, you can see the current value in ActiveNamedThreads(%h)"), FTaskTagScope::GetCurrentTag());
+		ensureMsgf(oldValue == newValue, TEXT("oldValue(%i) newValue(%i) If this check fails make sure that there is a FTaskTagScope(ETaskTag::EGameThread) as deep as possible on the current callstack, you can see the current value in ActiveNamedThreads(%x)"), oldValue, newValue, FTaskTagScope::GetCurrentTag());
+		newValue = oldValue;
 #endif
 		return newValue;
 	}
@@ -165,7 +166,8 @@ CORE_API bool IsInSlateThread()
 	bool newValue = FTaskTagScope::IsCurrentTag(ETaskTag::ESlateThread);
 #if !UE_BUILD_SHIPPING && !UE_BUILD_TEST
 	bool oldValue = GSlateLoadingThreadId != 0 && FPlatformTLS::GetCurrentThreadId() == GSlateLoadingThreadId;
-	checkf(oldValue == newValue, TEXT("If this check fails make sure that there is a FTaskTagScope(ETaskTag::ESlateThread) as as deep as possible on the current callstack, you can see the current value in ActiveNamedThreads(%h)"), FTaskTagScope::GetCurrentTag());
+	ensureMsgf(oldValue == newValue, TEXT("oldValue(%i) newValue(%i) If this check fails make sure that there is a FTaskTagScope(ETaskTag::ESlateThread) as as deep as possible on the current callstack, you can see the current value in ActiveNamedThreads(%x)"), oldValue, newValue, FTaskTagScope::GetCurrentTag());
+	newValue = oldValue;
 #endif
 	return newValue;
 }
@@ -180,7 +182,8 @@ CORE_API bool IsInAudioThread()
 		: FTaskTagScope::IsCurrentTag(ETaskTag::EGameThread);
 #if !UE_BUILD_SHIPPING && !UE_BUILD_TEST
 	bool oldValue = FPlatformTLS::GetCurrentThreadId() == (GAudioThreadId ? GAudioThreadId : GGameThreadId);
-	checkf(oldValue == newValue, TEXT("If this check fails make sure that there is a FTaskTagScope(ETaskTag::EAudioThread) as deep as possible on the current callstack, you can see the current value in ActiveNamedThreads(%h)"), FTaskTagScope::GetCurrentTag());
+	ensureMsgf(oldValue == newValue, TEXT("oldValue(%i) newValue(%i) If this check fails make sure that there is a FTaskTagScope(ETaskTag::EAudioThread) as deep as possible on the current callstack, you can see the current value in ActiveNamedThreads(%x)"), oldValue, newValue, FTaskTagScope::GetCurrentTag());
+	newValue = oldValue;
 #endif
 	return newValue;
 }
@@ -194,7 +197,8 @@ CORE_API bool IsInActualRenderingThread()
 	bool newValue = FTaskTagScope::IsCurrentTag(ETaskTag::ERenderingThread);
 #if !UE_BUILD_SHIPPING && !UE_BUILD_TEST
 	bool oldValue = GRenderingThread && FPlatformTLS::GetCurrentThreadId() == GRenderingThread->GetThreadID();
-	checkf(oldValue == newValue, TEXT("If this check fails make sure that there is a FTaskTagScope(ETaskTag::ERenderingThread) as deep as possible on the current callstack, you can see the current value in ActiveNamedThreads(%h)"), FTaskTagScope::GetCurrentTag());
+	ensureMsgf(oldValue == newValue, TEXT("oldValue(%i) newValue(%i) If this check fails make sure that there is a FTaskTagScope(ETaskTag::ERenderingThread) as deep as possible on the current callstack, you can see the current value in ActiveNamedThreads(%x)"), oldValue, newValue, FTaskTagScope::GetCurrentTag());
+	newValue = oldValue;
 #endif
 	return newValue;
 }
@@ -207,7 +211,8 @@ CORE_API bool IsInRenderingThread()
 
 #if !UE_BUILD_SHIPPING && !UE_BUILD_TEST
 	bool oldValue = !GRenderingThread || GIsRenderingThreadSuspended.Load(EMemoryOrder::Relaxed) || (FPlatformTLS::GetCurrentThreadId() == GRenderingThread->GetThreadID());
-	checkf(oldValue == newValue, TEXT("If this check fails make sure that there is a FTaskTagScope(ETaskTag::ERenderingThread) as deep as possible on the current callstack, you can see the current value in ActiveNamedThreads(%h)"), FTaskTagScope::GetCurrentTag());
+	ensureMsgf(oldValue == newValue, TEXT("oldValue(%i) newValue(%i) If this check fails make sure that there is a FTaskTagScope(ETaskTag::ERenderingThread) as deep as possible on the current callstack, you can see the current value in ActiveNamedThreads(%x), GRenderingThread(%x), GIsRenderingThreadSuspended(%d)"), oldValue, newValue, FTaskTagScope::GetCurrentTag(), GRenderingThread, GIsRenderingThreadSuspended.Load(EMemoryOrder::Relaxed));
+	newValue = oldValue;
 #endif
 	return newValue;
 }
@@ -237,7 +242,8 @@ CORE_API bool IsInParallelRenderingThread()
 	{
 		oldValue = FPlatformTLS::GetCurrentThreadId() != GGameThreadId;
 	}
-	checkf(oldValue == newValue, TEXT("If this check fails make sure that there is a FTaskTagScope(ETaskTag::EParallelRenderingThread) as deep as possible on the current callstack, you can see the current value in ActiveNamedThreads(%h)"), FTaskTagScope::GetCurrentTag());
+	ensureMsgf(oldValue == newValue, TEXT("oldValue(%i) newValue(%i) If this check fails make sure that there is a FTaskTagScope(ETaskTag::EParallelRenderingThread) as deep as possible on the current callstack, you can see the current value in ActiveNamedThreads(%x), GRenderingThread(%x), GIsRenderingThreadSuspended(%d)"), oldValue, newValue, FTaskTagScope::GetCurrentTag(), GRenderingThread, GIsRenderingThreadSuspended.Load(EMemoryOrder::Relaxed));
+	newValue = oldValue;
 #endif
 	return newValue;
 }
@@ -250,7 +256,8 @@ CORE_API bool IsInRHIThread()
 	bool newValue = FTaskTagScope::IsCurrentTag(ETaskTag::ERhiThread);
 #if !UE_BUILD_SHIPPING && !UE_BUILD_TEST
 	bool oldValue = GRHIThreadId && FPlatformTLS::GetCurrentThreadId() == GRHIThreadId;	
-	checkf(oldValue == newValue, TEXT("If this check fails make sure that there is a FTaskTagScope(ETaskTag::ERhiThread) as deep as possible on the current callstack, you can see the current value in ActiveNamedThreads(%h)"), FTaskTagScope::GetCurrentTag());
+	ensureMsgf(oldValue == newValue, TEXT("oldValue(%i) newValue(%i) If this check fails make sure that there is a FTaskTagScope(ETaskTag::ERhiThread) as deep as possible on the current callstack, you can see the current value in ActiveNamedThreads(%x)"), oldValue, newValue, FTaskTagScope::GetCurrentTag());
+	newValue = oldValue;
 #endif
 	return newValue;
 }
