Removed HAVE_RUNTIME_THREADING_SWITCHES so all platforms correctly respect -onethread, preventing a number of crashes from mismatches between code using ShouldUseThreadingForPerformance() (forced true) and directly querying -onethread on the commandline. The original intent of this code is likely lost over time (and conflicting changes), and the current implementation is incredibly broken and unnecessary.

diff --git a/Engine/Source/Runtime/Core/Private/Misc/App.cpp b/Engine/Source/Runtime/Core/Private/Misc/App.cpp
--- a/Engine/Source/Runtime/Core/Private/Misc/App.cpp
+++ b/Engine/Source/Runtime/Core/Private/Misc/App.cpp
@@ -236,8 +236,6 @@ bool FApp::IsUnattended() // @todo clang: Workaround for missing symbol export
 }
 #endif
 
-#if HAVE_RUNTIME_THREADING_SWITCHES
-
 #if PLATFORM_LUMIN
 #define MIN_CORE_COUNT 2
 #else
@@ -247,7 +245,7 @@ bool FApp::IsUnattended() // @todo clang: Workaround for missing symbol export
 bool FApp::ShouldUseThreadingForPerformance()
 {
 	static bool OnlyOneThread = 
-		FParse::Param(FCommandLine::Get(), TEXT("ONETHREAD")) ||
+		FParse::Param(FCommandLine::Get(), TEXT("onethread")) ||
 		FParse::Param(FCommandLine::Get(), TEXT("noperfthreads")) ||
 		IsRunningDedicatedServer() ||
 		!FPlatformProcess::SupportsMultithreading() ||
@@ -259,8 +257,6 @@ bool FApp::ShouldUseThreadingForPerformance()
 }
 #undef MIN_CORE_COUNT
 
-#endif // HAVE_RUNTIME_THREADING_SWITCHES
-
 FTimecode FApp::GetTimecode()
 {
 	if (CurrentFrameTime.IsSet())
diff --git a/Engine/Source/Runtime/Core/Public/Misc/App.h b/Engine/Source/Runtime/Core/Public/Misc/App.h
--- a/Engine/Source/Runtime/Core/Public/Misc/App.h
+++ b/Engine/Source/Runtime/Core/Public/Misc/App.h
@@ -17,11 +17,6 @@
 #include "Misc/QualifiedFrameTime.h"
 #include "HAL/PlatformProcess.h"
 
-// platforms which can have runtime threading switches
-#if !defined(HAVE_RUNTIME_THREADING_SWITCHES)
-	#define HAVE_RUNTIME_THREADING_SWITCHES			(!PLATFORM_PS4 && !PLATFORM_ANDROID && !PLATFORM_TVOS && !PLATFORM_SWITCH && !PLATFORM_LUMIN)
-#endif
-
 /**
  * Provides information about the application.
  */
@@ -447,14 +442,7 @@ public:
 	 *
 	 * @return true if this isn't a server, has more than one core, does not have a -onethread command line options, etc.
 	 */
-#if HAVE_RUNTIME_THREADING_SWITCHES
 	static bool ShouldUseThreadingForPerformance();
-#else
-	FORCEINLINE static bool ShouldUseThreadingForPerformance()
-	{
-		return true;
-	}
-#endif // HAVE_RUNTIME_THREADING_SWITCHES
 
 	/**
 	 * Checks whether application is in benchmark mode.
