- Big speedup to scanning for DataDrivenPlatformInfo.ini files (editor starup went from 270ms to ~1ms)

diff --git a/Engine/Source/Runtime/Core/Private/Misc/DataDrivenPlatformInfoRegistry.cpp b/Engine/Source/Runtime/Core/Private/Misc/DataDrivenPlatformInfoRegistry.cpp
--- a/Engine/Source/Runtime/Core/Private/Misc/DataDrivenPlatformInfoRegistry.cpp
+++ b/Engine/Source/Runtime/Core/Private/Misc/DataDrivenPlatformInfoRegistry.cpp
@@ -38,7 +38,19 @@ static const TArray<FString>& GetDataDrivenIniFilenames()
 
 		// look for the special files in any congfig subdirectories
 		IFileManager::Get().FindFilesRecursive(DataDrivenIniFilenames, *FPaths::EngineConfigDir(), TEXT("DataDrivenPlatformInfo.ini"), true, false);
-		IFileManager::Get().FindFilesRecursive(DataDrivenIniFilenames, *FPaths::EnginePlatformExtensionsDir(), TEXT("DataDrivenPlatformInfo.ini"), true, false, false);
+
+		// manually look through the platform directories - we can't use GetExtensionDirs(), since that function uses the results of this function 
+		TArray<FString> PlatformDirs;
+		IFileManager::Get().FindFiles(PlatformDirs, *FPaths::Combine(FPaths::EnginePlatformExtensionsDir(), TEXT("*")), false, true);
+
+		for (const FString& PlatformDir : PlatformDirs)
+		{
+			FString IniPath = FPaths::Combine(FPaths::EnginePlatformExtensionsDir(), PlatformDir, TEXT("Config/DataDrivenPlatformInfo.ini"));
+			if (IFileManager::Get().FileExists(*IniPath))
+			{
+				DataDrivenIniFilenames.Add(IniPath);
+			}
+		}
 	}
 
 	return DataDrivenIniFilenames;
