- Fixed a variable assigning to itself bug (thanks PVS!)

diff --git a/Engine/Source/Runtime/Core/Private/Misc/DataDrivenPlatformInfoRegistry.cpp b/Engine/Source/Runtime/Core/Private/Misc/DataDrivenPlatformInfoRegistry.cpp
--- a/Engine/Source/Runtime/Core/Private/Misc/DataDrivenPlatformInfoRegistry.cpp
+++ b/Engine/Source/Runtime/Core/Private/Misc/DataDrivenPlatformInfoRegistry.cpp
@@ -229,7 +229,7 @@ static void LoadDDPIIniSettings(const FConfigFile& IniFile, FDataDrivenPlatformI
 	DDPIGetString(IniFile, TEXT("XLargeIconPath"), Info.IconPaths.XLargePath);
 	if (Info.IconPaths.XLargePath == TEXT(""))
 	{
-		Info.IconPaths.XLargePath = Info.IconPaths.XLargePath;
+		Info.IconPaths.XLargePath = Info.IconPaths.LargePath;
 	}
 
 	FString PlatformString = PlatformName.ToString();
