Fixed incorrect rounding in FLinearColor::ToFColor

diff --git a/Engine/Source/Runtime/Core/Private/Math/Color.cpp b/Engine/Source/Runtime/Core/Private/Math/Color.cpp
--- a/Engine/Source/Runtime/Core/Private/Math/Color.cpp
+++ b/Engine/Source/Runtime/Core/Private/Math/Color.cpp
@@ -134,10 +134,10 @@ FColor FLinearColor::ToFColor(const bool bSRGB) const
 
 	FColor Result;
 
-	Result.A = (uint8)FMath::FloorToInt(FloatA * 255.999f);
-	Result.R = (uint8)FMath::FloorToInt(FloatR * 255.999f);
-	Result.G = (uint8)FMath::FloorToInt(FloatG * 255.999f);
-	Result.B = (uint8)FMath::FloorToInt(FloatB * 255.999f);
+	Result.A = (uint8)FMath::RoundToInt(FloatA * 255.0f);
+	Result.R = (uint8)FMath::RoundToInt(FloatR * 255.0f);
+	Result.G = (uint8)FMath::RoundToInt(FloatG * 255.0f);
+	Result.B = (uint8)FMath::RoundToInt(FloatB * 255.0f);
 
 	return Result;
 }
