- Added a checkf() to try to track down a crash opening the Turnkey LaunchOn Menu

diff --git a/Engine/Source/Runtime/Core/Private/Misc/DataDrivenPlatformInfoRegistry.cpp b/Engine/Source/Runtime/Core/Private/Misc/DataDrivenPlatformInfoRegistry.cpp
--- a/Engine/Source/Runtime/Core/Private/Misc/DataDrivenPlatformInfoRegistry.cpp
+++ b/Engine/Source/Runtime/Core/Private/Misc/DataDrivenPlatformInfoRegistry.cpp
@@ -363,7 +363,7 @@ static void PrepForTurnkeyReport(FString& Command, FString& BaseCommandline, FSt
 
 static FString ConvertToDDPIPlatform(const FString& Platform)
 {
-	FString  New = Platform.Replace(TEXT("Client"), TEXT("")).Replace(TEXT("Server"), TEXT(""));
+	FString  New = Platform.Replace(TEXT("Editor"), TEXT("")).Replace(TEXT("Client"), TEXT("")).Replace(TEXT("Server"), TEXT(""));
 	if (New == TEXT("Win64"))
 	{
 		New = TEXT("Windows");
@@ -518,8 +518,12 @@ FDataDrivenPlatformInfo& FDataDrivenPlatformInfoRegistry::DeviceIdToInfo(FString
 		*OutDeviceName = PlatformAndDevice[1];
 	}
 
+	FString DDPIPlatformName = ConvertToDDPIPlatform(PlatformAndDevice[0]);
+
+	checkf(DataDrivenPlatforms.Contains(*DDPIPlatformName), TEXT("DataDrivenPlatforms map did not contain the DDPI Platform %s"), *DDPIPlatformName);
+
 	// have to convert back to Windows from Win64
-	return DataDrivenPlatforms[*ConvertToDDPIPlatform(PlatformAndDevice[0])];
+	return DataDrivenPlatforms[*DDPIPlatformName];
 
 }
 
