- Fixed issue with new project not packacing on first run
#jria UE-97359

diff --git a/Engine/Source/Runtime/Core/Private/Misc/DataDrivenPlatformInfoRegistry.cpp b/Engine/Source/Runtime/Core/Private/Misc/DataDrivenPlatformInfoRegistry.cpp
--- a/Engine/Source/Runtime/Core/Private/Misc/DataDrivenPlatformInfoRegistry.cpp
+++ b/Engine/Source/Runtime/Core/Private/Misc/DataDrivenPlatformInfoRegistry.cpp
@@ -353,6 +353,9 @@ static void PrepForTurnkeyReport(FString& Command, FString& BaseCommandline, FSt
 	FString LogFilename = FPaths::ConvertRelativePathToFull(FPaths::Combine(FPaths::ProjectIntermediateDir(), *FString::Printf(TEXT("TurnkeyLog_%d.log"), ReportIndex)));
 	ReportFilename = FPaths::ConvertRelativePathToFull(FPaths::Combine(FPaths::ProjectIntermediateDir(), *FString::Printf(TEXT("TurnkeyReport_%d.log"), ReportIndex++)));
 
+	// make sure intermediate directory exists
+	IFileManager::Get().MakeDirectory(*FPaths::ProjectIntermediateDir());
+
 	Command = TEXT("{EngineDir}/Build/BatchFiles/RunuAT");
 //	Command = TEXT("{EngineDir}/Binaries/DotNET/AutomationTool.exe");
 	BaseCommandline = FString::Printf(TEXT("Turnkey -command=VerifySdk -ReportFilename=\"%s\" -log=\"%s\""), *ReportFilename, *LogFilename);
