Improve ThreadPoolWrapper to retract lower priority task to make place for higher priority ones when scheduling

A threadpoolwrapper only schedules MaxConcurrency task on its inner thread pool. If all those tasks already scheduled are low-priority and still waiting to be processed, we will replace them by new ones with higher priorities coming in.
This is especially important when the threadpool is running at full load and low priority task could be stalled for a while before being processed.

diff --git a/Engine/Source/Runtime/Core/Private/HAL/ThreadingBase.cpp b/Engine/Source/Runtime/Core/Private/HAL/ThreadingBase.cpp
--- a/Engine/Source/Runtime/Core/Private/HAL/ThreadingBase.cpp
+++ b/Engine/Source/Runtime/Core/Private/HAL/ThreadingBase.cpp
@@ -748,6 +748,11 @@ void FRunnableThread::FreeTls()
 	FThreadPoolPriorityQueue
 -----------------------------------------------------------------------------*/
 
+FThreadPoolPriorityQueue::FThreadPoolPriorityQueue()
+	: NumQueuedWork(0)
+{
+}
+
 void FThreadPoolPriorityQueue::Enqueue(IQueuedWork* InQueuedWork, EQueuedWorkPriority InPriority)
 {
 	int32 QueueIndex = static_cast<int32>(InPriority);
@@ -808,6 +813,28 @@ IQueuedWork* FThreadPoolPriorityQueue::Dequeue(EQueuedWorkPriority* OutDequeuedW
 	return Work;
 }
 
+IQueuedWork* FThreadPoolPriorityQueue::Peek(EQueuedWorkPriority* OutDequeuedWorkPriority)
+{
+	IQueuedWork* Work = nullptr;
+	for (int32 QueueIndex = FirstNonEmptyQueueIndex, Num = PriorityQueuedWork.Num(); QueueIndex < Num; ++QueueIndex)
+	{
+		TArray<IQueuedWork*>& QueuedWork = PriorityQueuedWork[QueueIndex];
+		if (QueuedWork.Num() > 0)
+		{
+			Work = QueuedWork[0];
+
+			if (OutDequeuedWorkPriority)
+			{
+				*OutDequeuedWorkPriority = (EQueuedWorkPriority)QueueIndex;
+			}
+
+			break;
+		}
+	}
+
+	return Work;
+}
+
 void FThreadPoolPriorityQueue::Reset()
 {
 	PriorityQueuedWork.Empty();
diff --git a/Engine/Source/Runtime/Core/Private/Misc/QueuedThreadPoolWrapper.cpp b/Engine/Source/Runtime/Core/Private/Misc/QueuedThreadPoolWrapper.cpp
--- a/Engine/Source/Runtime/Core/Private/Misc/QueuedThreadPoolWrapper.cpp
+++ b/Engine/Source/Runtime/Core/Private/Misc/QueuedThreadPoolWrapper.cpp
@@ -25,7 +25,7 @@ struct FQueuedThreadPoolWrapper::FScheduledWork : public IQueuedWork
 	}
 
 	IQueuedWork* InnerWork = nullptr;
-
+	EQueuedWorkPriority Priority = EQueuedWorkPriority::Normal;
 private:
 	FQueuedThreadPoolWrapper* ParentPool;
 };
@@ -181,6 +181,36 @@ void FQueuedThreadPoolWrapper::Schedule(FScheduledWork* Work)
 		ReleaseWorkNoLock(Work);
 	}
 
+	// If a higher priority task comes in, try to retract a lower priority one if possible to make room
+	if (CurrentConcurrency >= MaxConcurrency.Load(EMemoryOrder::Relaxed) && (MaxTaskToSchedule == -1 || MaxTaskToSchedule > 0))
+	{
+		EQueuedWorkPriority NextWorkPriority;
+		IQueuedWork* NextWork = QueuedWork.Peek(&NextWorkPriority);
+		if (NextWork)
+		{
+			// Scheduled work is bound by MaxConcurrency which is normally limited by Core count. 
+			// The linear scan should be small and pretty fast.
+			for (TTuple<IQueuedWork*, FScheduledWork*>& Pair : ScheduledWork)
+			{
+				// higher number means lower priority
+				if (Pair.Value->Priority > NextWorkPriority)
+				{
+					if (WrappedQueuedThreadPool->RetractQueuedWork(Pair.Value))
+					{
+						QueuedWork.Enqueue(Pair.Key, Pair.Value->Priority);
+						ReleaseWorkNoLock(Pair.Value);
+
+						if (MaxTaskToSchedule != -1)
+						{
+							MaxTaskToSchedule++;
+						}
+						break;
+					}
+				}
+			}
+		}
+	}
+
 	while (CurrentConcurrency < MaxConcurrency.Load(EMemoryOrder::Relaxed) && (MaxTaskToSchedule == -1 || MaxTaskToSchedule > 0))
 	{
 		EQueuedWorkPriority WorkPriority;
@@ -198,6 +228,7 @@ void FQueuedThreadPoolWrapper::Schedule(FScheduledWork* Work)
 				Work = new FScheduledWork(this);
 			}
 
+			Work->Priority = WorkPriority;
 			Work->InnerWork = InnerWork;
 			ScheduledWork.Add(InnerWork, Work);
 			WrappedQueuedThreadPool->AddQueuedWork(Work, PriorityMapper(WorkPriority));
diff --git a/Engine/Source/Runtime/Core/Public/Misc/QueuedThreadPool.h b/Engine/Source/Runtime/Core/Public/Misc/QueuedThreadPool.h
--- a/Engine/Source/Runtime/Core/Public/Misc/QueuedThreadPool.h
+++ b/Engine/Source/Runtime/Core/Public/Misc/QueuedThreadPool.h
@@ -25,6 +25,8 @@ enum class EQueuedWorkPriority : uint8
 class CORE_API FThreadPoolPriorityQueue
 {
 public:
+	FThreadPoolPriorityQueue();
+
 	/**
 	 * Enqueue a work item at specified priority
 	 */
@@ -40,6 +42,11 @@ public:
 	 */
 	IQueuedWork* Dequeue(EQueuedWorkPriority* OutDequeuedWorkPriority = nullptr);
 
+	/**
+	 * Get the next work item in priority order without actually dequeuing.
+	 */
+	IQueuedWork* Peek(EQueuedWorkPriority* OutDequeuedWorkPriority = nullptr);
+
 	/**
 	 * Empty the queue.
 	 */
