Add ability to enable/disable "unused private field" warnings.

diff --git a/Engine/Source/Runtime/Core/Public/Clang/ClangPlatformCompilerPreSetup.h b/Engine/Source/Runtime/Core/Public/Clang/ClangPlatformCompilerPreSetup.h
--- a/Engine/Source/Runtime/Core/Public/Clang/ClangPlatformCompilerPreSetup.h
+++ b/Engine/Source/Runtime/Core/Public/Clang/ClangPlatformCompilerPreSetup.h
@@ -149,6 +149,17 @@
 		_Pragma("clang diagnostic ignored \"-Wmacro-redefined\"")
 #endif // PRAGMA_DISABLE_MACRO_REDEFINED_WARNINGS
 
+#ifndef PRAGMA_DISABLE_UNUSED_PRIVATE_FIELDS_WARNINGS
+#define PRAGMA_DISABLE_UNUSED_PRIVATE_FIELDS_WARNINGS \
+		_Pragma("clang diagnostic push") \
+		_Pragma("clang diagnostic ignored \"-Wunused-private-field\"")
+#endif // PRAGMA_DISABLE_UNUSED_PRIVATE_FIELDS_WARNINGS
+
+#ifndef PRAGMA_ENABLE_UNUSED_PRIVATE_FIELDS_WARNINGS
+#define PRAGMA_ENABLE_UNUSED_PRIVATE_FIELDS_WARNINGS \
+		_Pragma("clang diagnostic pop")
+#endif // PRAGMA_ENABLE_UNUSED_PRIVATE_FIELDS_WARNINGS
+
 #ifndef PRAGMA_ENABLE_MACRO_REDEFINED_WARNINGS
 	#define PRAGMA_ENABLE_MACRO_REDEFINED_WARNINGS \
 		_Pragma("clang diagnostic pop")
diff --git a/Engine/Source/Runtime/Core/Public/HoloLens/HoloLensPlatformCompilerPreSetup.h b/Engine/Source/Runtime/Core/Public/HoloLens/HoloLensPlatformCompilerPreSetup.h
--- a/Engine/Source/Runtime/Core/Public/HoloLens/HoloLensPlatformCompilerPreSetup.h
+++ b/Engine/Source/Runtime/Core/Public/HoloLens/HoloLensPlatformCompilerPreSetup.h
@@ -84,6 +84,16 @@
 		__pragma(warning(pop))
 #endif // PRAGMA_ENABLE_REGISTER_WARNINGS
 
+#ifndef PRAGMA_DISABLE_UNUSED_PRIVATE_FIELDS_WARNINGS
+#define PRAGMA_DISABLE_UNUSED_PRIVATE_FIELDS_WARNINGS
+	 // MSVC doesn't seem to have a warning similar to -Wunused-private-field
+#endif // PRAGMA_DISABLE_UNUSED_PRIVATE_FIELDS_WARNINGS
+
+#ifndef PRAGMA_ENABLE_UNUSED_PRIVATE_FIELDS_WARNINGS
+#define PRAGMA_ENABLE_UNUSED_PRIVATE_FIELDS_WARNINGS
+		// MSVC doesn't seem to have a warning similar to -Wunused-private-field
+#endif // PRAGMA_ENABLE_UNUSED_PRIVATE_FIELDS_WARNINGS
+
 #ifndef PRAGMA_POP
 	#define PRAGMA_POP \
 		__pragma(warning(pop))
diff --git a/Engine/Source/Runtime/Core/Public/Windows/WindowsPlatformCompilerPreSetup.h b/Engine/Source/Runtime/Core/Public/Windows/WindowsPlatformCompilerPreSetup.h
--- a/Engine/Source/Runtime/Core/Public/Windows/WindowsPlatformCompilerPreSetup.h
+++ b/Engine/Source/Runtime/Core/Public/Windows/WindowsPlatformCompilerPreSetup.h
@@ -129,6 +129,16 @@
 			__pragma(warning(pop))
 	#endif // PRAGMA_ENABLE_REGISTER_WARNINGS
 
+	#ifndef PRAGMA_DISABLE_UNUSED_PRIVATE_FIELDS_WARNINGS
+	#define PRAGMA_DISABLE_UNUSED_PRIVATE_FIELDS_WARNINGS
+			// MSVC doesn't seem to have a warning similar to -Wunused-private-field
+	#endif // PRAGMA_DISABLE_UNUSED_PRIVATE_FIELDS_WARNINGS
+
+	#ifndef PRAGMA_ENABLE_UNUSED_PRIVATE_FIELDS_WARNINGS
+	#define PRAGMA_ENABLE_UNUSED_PRIVATE_FIELDS_WARNINGS
+			// MSVC doesn't seem to have a warning similar to -Wunused-private-field
+	#endif // PRAGMA_ENABLE_UNUSED_PRIVATE_FIELDS_WARNINGS
+
 	#ifndef PRAGMA_POP
 		#define PRAGMA_POP \
 			__pragma(warning(pop))
