Remove unused private fields - ThreadHeartBeat

diff --git a/Engine/Source/Runtime/Core/Private/HAL/ThreadHeartBeat.cpp b/Engine/Source/Runtime/Core/Private/HAL/ThreadHeartBeat.cpp
--- a/Engine/Source/Runtime/Core/Private/HAL/ThreadHeartBeat.cpp
+++ b/Engine/Source/Runtime/Core/Private/HAL/ThreadHeartBeat.cpp
@@ -757,11 +757,12 @@ void FThreadHeartBeat::SetDurationMultiplier(double NewMultiplier)
 
 FGameThreadHitchHeartBeatThreaded::FGameThreadHitchHeartBeatThreaded()
 	: Thread(nullptr)
-	, HangDuration(-1.f)
+#if USE_HITCH_DETECTION
 	, bWalkStackOnHitch(false)
 	, FirstStartTime(0.0)
 	, FrameStartTime(0.0)
 	, SuspendedCount(0)
+#endif
 	, Clock(HitchDetectorClock_MaxTimeStep_MS / 1000.0)
 {
 	// We don't care about programs for now so no point in spawning the extra thread
@@ -1053,7 +1054,11 @@ void FGameThreadHitchHeartBeatThreaded::ResumeHeartBeat()
 
 double FGameThreadHitchHeartBeatThreaded::GetFrameStartTime()
 {
+#if USE_HITCH_DETECTION
 	return FrameStartTime;
+#else
+	return 0.0;
+#endif
 }
 
 double FGameThreadHitchHeartBeatThreaded::GetCurrentTime()
diff --git a/Engine/Source/Runtime/Core/Public/HAL/ThreadHeartBeat.h b/Engine/Source/Runtime/Core/Public/HAL/ThreadHeartBeat.h
--- a/Engine/Source/Runtime/Core/Public/HAL/ThreadHeartBeat.h
+++ b/Engine/Source/Runtime/Core/Public/HAL/ThreadHeartBeat.h
@@ -283,24 +283,19 @@ class CORE_API FGameThreadHitchHeartBeatThreaded : public FRunnable
 	FThreadSafeCounter StopTaskCounter;
 	/** Synch object for the heartbeat */
 	FCriticalSection HeartBeatCritical;
-	/** Max time the game thread is allowed to not send the heartbeat*/
-	float HangDuration;
 
+#if USE_HITCH_DETECTION
 	bool bWalkStackOnHitch;
-
 	double FirstStartTime;
 	double FrameStartTime;
-
 	int32 SuspendedCount;
+#endif
 
 	FThreadHeartBeatClock Clock;
 
 #if LOOKUP_SYMBOLS_IN_HITCH_STACK_WALK
 	static const SIZE_T StackTraceSize = 65535;
 	ANSICHAR StackTrace[StackTraceSize];
-#else
-	static const uint32 MaxStackDepth = 128;
-	uint64 StackTrace[MaxStackDepth];
 #endif
 
 	void InitSettings();
