Remove unused private fields - Localization & Docs

diff --git a/Engine/Source/Developer/Localization/Private/LocalizationTargetTypes.cpp b/Engine/Source/Developer/Localization/Private/LocalizationTargetTypes.cpp
--- a/Engine/Source/Developer/Localization/Private/LocalizationTargetTypes.cpp
+++ b/Engine/Source/Developer/Localization/Private/LocalizationTargetTypes.cpp
@@ -347,7 +347,6 @@ namespace
 
 	private:
 		FWordCountCSVParser(FLocalizationTargetSettings& InTarget)
-			: Target(InTarget)
 		{
 			MandatoryColumnNames.Add( TEXT("Date/Time") );
 			MandatoryColumnNames.Add( TEXT("Word Count") );
@@ -425,7 +424,6 @@ namespace
 		}
 
 	private:
-		FLocalizationTargetSettings& Target;
 		TArray<FString> MandatoryColumnNames;
 		TArray<FString> CultureNames;
 	};
diff --git a/Engine/Source/Editor/Documentation/Private/DocumentationPage.h b/Engine/Source/Editor/Documentation/Private/DocumentationPage.h
--- a/Engine/Source/Editor/Documentation/Private/DocumentationPage.h
+++ b/Engine/Source/Editor/Documentation/Private/DocumentationPage.h
@@ -43,6 +43,4 @@ private:
 	TArray<FExcerpt> StoredExcerpts;
 	/** The UDN meta data contained in this page */
 	FUDNPageMetadata StoredMetadata;
-	/** Signals if the page has been loaded */
-	bool IsLoaded;
 };
diff --git a/Engine/Source/Runtime/Core/Private/Internationalization/ICUText.cpp b/Engine/Source/Runtime/Core/Private/Internationalization/ICUText.cpp
--- a/Engine/Source/Runtime/Core/Private/Internationalization/ICUText.cpp
+++ b/Engine/Source/Runtime/Core/Private/Internationalization/ICUText.cpp
@@ -123,8 +123,7 @@ class FText::FSortPredicate::FSortPredicateImplementation
 {
 public:
 	FSortPredicateImplementation(const ETextComparisonLevel::Type InComparisonLevel)
-		: ComparisonLevel(InComparisonLevel)
-		, ICUCollator(FInternationalization::Get().GetCurrentLanguage()->Implementation->GetCollator(InComparisonLevel))
+		: ICUCollator(FInternationalization::Get().GetCurrentLanguage()->Implementation->GetCollator(InComparisonLevel))
 	{
 	}
 
@@ -147,7 +146,6 @@ public:
 	}
 
 private:
-	const ETextComparisonLevel::Type ComparisonLevel;
 	const TSharedRef<const icu::Collator, ESPMode::ThreadSafe> ICUCollator;
 };
 
