Fix hang detection, some paths were broken in the #ifdef maze.

diff --git a/Engine/Source/Runtime/Core/Private/HAL/ThreadHeartBeat.cpp b/Engine/Source/Runtime/Core/Private/HAL/ThreadHeartBeat.cpp
--- a/Engine/Source/Runtime/Core/Private/HAL/ThreadHeartBeat.cpp
+++ b/Engine/Source/Runtime/Core/Private/HAL/ThreadHeartBeat.cpp
@@ -30,10 +30,6 @@
 	#define UE_ASSERT_ON_HANG 0
 #endif
 
-#ifndef WALK_STACK_ON_HITCH_DETECTED
-	#define WALK_STACK_ON_HITCH_DETECTED 0
-#endif
-
 #ifndef NEEDS_DEBUG_INFO_ON_PRESENT_HANG
 #define NEEDS_DEBUG_INFO_ON_PRESENT_HANG 0
 #endif
diff --git a/Engine/Source/Runtime/Core/Public/HAL/ThreadHeartBeat.h b/Engine/Source/Runtime/Core/Public/HAL/ThreadHeartBeat.h
--- a/Engine/Source/Runtime/Core/Public/HAL/ThreadHeartBeat.h
+++ b/Engine/Source/Runtime/Core/Public/HAL/ThreadHeartBeat.h
@@ -273,6 +273,10 @@ public:
 #define LOOKUP_SYMBOLS_IN_HITCH_STACK_WALK 0
 #endif
 
+#ifndef WALK_STACK_ON_HITCH_DETECTED
+#define WALK_STACK_ON_HITCH_DETECTED 0
+#endif
+
 class CORE_API FGameThreadHitchHeartBeatThreaded : public FRunnable
 {
 	static FGameThreadHitchHeartBeatThreaded* Singleton;
@@ -285,18 +289,24 @@ class CORE_API FGameThreadHitchHeartBeatThreaded : public FRunnable
 	FCriticalSection HeartBeatCritical;
 
 #if USE_HITCH_DETECTION
+	float HangDuration = -1.0f;
 	bool bWalkStackOnHitch;
 	double FirstStartTime;
 	double FrameStartTime;
 	int32 SuspendedCount;
-#endif
-
-	FThreadHeartBeatClock Clock;
 
+#if WALK_STACK_ON_HITCH_DETECTED
 #if LOOKUP_SYMBOLS_IN_HITCH_STACK_WALK
 	static const SIZE_T StackTraceSize = 65535;
 	ANSICHAR StackTrace[StackTraceSize];
-#endif
+#else
+	static const uint32 MaxStackDepth = 128;
+	uint64 StackTrace[MaxStackDepth]; 
+#endif // LOOKUP_SYMBOLS_IN_HITCH_STACK_WALK
+#endif // WALK_STACK_ON_HITCH_DETECTED
+#endif // USE_HITCH_DETECTION
+
+	FThreadHeartBeatClock Clock;
 
 	void InitSettings();
 
