[GenericOctree] Fixed typo

#rnx

diff --git a/Engine/Source/Runtime/Core/Public/Math/GenericOctree.h b/Engine/Source/Runtime/Core/Public/Math/GenericOctree.h
--- a/Engine/Source/Runtime/Core/Public/Math/GenericOctree.h
+++ b/Engine/Source/Runtime/Core/Public/Math/GenericOctree.h
@@ -736,7 +736,7 @@ public:
 	}
 
 	/**
-	 * this function will traverse the Octree using a fast box-box interssection this should be the preferred way of traversing the tree.
+	 * this function will traverse the Octree using a fast box-box intersection this should be the preferred way of traversing the tree.
 	 * @param BoxBounds - the bounds to test if a node is traversed or skipped.
 	 * @param Func - Function to call with each Element for nodes that passed bounds test.
 	 */
