- Fixed Turnkey SDK detection when engine has a space in its pathname

diff --git a/Engine/Source/Runtime/Core/Private/Misc/DataDrivenPlatformInfoRegistry.cpp b/Engine/Source/Runtime/Core/Private/Misc/DataDrivenPlatformInfoRegistry.cpp
--- a/Engine/Source/Runtime/Core/Private/Misc/DataDrivenPlatformInfoRegistry.cpp
+++ b/Engine/Source/Runtime/Core/Private/Misc/DataDrivenPlatformInfoRegistry.cpp
@@ -404,12 +404,14 @@ static void PrepForTurnkeyReport(FString& Command, FString& BaseCommandline, FSt
 	// make sure intermediate directory exists
 	IFileManager::Get().MakeDirectory(*FPaths::ProjectIntermediateDir());
 
-	Command = TEXT("{EngineDir}/Build/BatchFiles/RunuAT");
+	Command = TEXT("{EngineDir}Build/BatchFiles/RunuAT");
 //	Command = TEXT("{EngineDir}/Binaries/DotNET/AutomationTool.exe");
 	BaseCommandline = FString::Printf(TEXT("Turnkey -command=VerifySdk -ReportFilename=\"%s\" -log=\"%s\""), *ReportFilename, *LogFilename);
 
 	// convert into appropriate calls for the current platform
-	FPlatformProcess::ModifyCreateProcParams(Command, BaseCommandline, FGenericPlatformProcess::ECreateProcHelperFlags::AppendPlatformScriptExtension);
+	FPlatformProcess::ModifyCreateProcParams(Command, BaseCommandline, FGenericPlatformProcess::ECreateProcHelperFlags::AppendPlatformScriptExtension | FGenericPlatformProcess::ECreateProcHelperFlags::RunThroughShell);
+
+	UE_LOG(LogInit, Log, TEXT("Running Turnkey SDK detection: '%s %s'"), *Command, *BaseCommandline);
 }
 
 static FString ConvertToDDPIPlatform(const FString& Platform)
