- Hardened up the async SDK querying in the Launch On menu

diff --git a/Engine/Source/Developer/LauncherServices/Private/Launcher/LauncherWorker.cpp b/Engine/Source/Developer/LauncherServices/Private/Launcher/LauncherWorker.cpp
--- a/Engine/Source/Developer/LauncherServices/Private/Launcher/LauncherWorker.cpp
+++ b/Engine/Source/Developer/LauncherServices/Private/Launcher/LauncherWorker.cpp
@@ -953,7 +953,7 @@ void FLauncherWorker::CreateAndExecuteTasks( const ILauncherProfileRef& InProfil
 	FString TurnkeyCommand;
 	if (Profile->ShouldUpdateDeviceFlash() && DeviceGroup->GetDeviceIDs().Num() >= 0)
 	{
-		TurnkeyCommand = FString::Printf(TEXT("Turnkey -command=VerifySdk -type=Flash -device=%s -UpdateIfNeeded -EditorIO"), *FString::Join(DeviceGroup->GetDeviceIDs(), TEXT("+")));
+		TurnkeyCommand = FString::Printf(TEXT("Turnkey -command=VerifySdk -type=Flash -device=%s -UpdateIfNeeded -EditorIO -utf8output -WaitForUATMutex"), *FString::Join(DeviceGroup->GetDeviceIDs(), TEXT("+")));
 	}
 
 	TSharedPtr<FLauncherTask> BuildTask = MakeShareable(new FLauncherUATTask(UATCommand, TEXT("Build Task"), TEXT("Launching UAT..."), ReadPipe, WritePipe, InProfile->GetEditorExe(), ProcHandle, this, StartString, TurnkeyCommand));
diff --git a/Engine/Source/Editor/MainFrame/Private/Frame/MainFrameActions.cpp b/Engine/Source/Editor/MainFrame/Private/Frame/MainFrameActions.cpp
--- a/Engine/Source/Editor/MainFrame/Private/Frame/MainFrameActions.cpp
+++ b/Engine/Source/Editor/MainFrame/Private/Frame/MainFrameActions.cpp
@@ -390,7 +390,7 @@ static bool ShowBadSDKDialog(FName IniPlatformName)
 		FFormatNamedArguments Args;
 		//		Args.Add(TEXT("DisplayName"), PlatformInfo->DisplayName);
 		Args.Add(TEXT("DisplayName"), FText::FromName(IniPlatformName));
-		FText WarningText = FText::Format(LOCTEXT("BadSDK_Message", "The SDK for {DisplayName} appears to not be installed properly, which is needed to generate data. Check the SDK section of the Launch On menu in the main toolbar to update SDK.\n\nWould you like to attempt to continue anyway?"), Args);
+		FText WarningText = FText::Format(LOCTEXT("BadSDK_Message", "The SDK for {DisplayName} appears to not be installed properly, which is needed to generate data. Check the SDK section of the Launch On menu in the main toolbar to update SDK.\n\nAlternatively, you may continue anyway, which will attempt to install/update the SDK before cookingt the content. This is a new feature, and is not guaranteed to work.\n\nWould you like to continue?"), Args);
 
 		FSuppressableWarningDialog::FSetupInfo Info(
 			WarningText,
@@ -457,7 +457,7 @@ void FMainFrameActionCallbacks::CookContent(const FName InPlatformInfoName)
 	}
 
 	FString ProjectPath = FPaths::IsProjectFilePathSet() ? FPaths::ConvertRelativePathToFull(FPaths::GetProjectFilePath()) : FPaths::RootDir() / FApp::GetProjectName() / FApp::GetProjectName() + TEXT(".uproject");
-	FString CommandLine = FString::Printf(TEXT("-ScriptsForProject=\"%s\" BuildCookRun %s%s -nop4 -project=\"%s\" -cook -skipstage -ue4exe=\"%s\" %s -utf8output"),
+	FString CommandLine = FString::Printf(TEXT("-ScriptsForProject=\"%s\" BuildCookRun %s%s -nop4 -project=\"%s\" -cook -skipstage -ue4exe=\"%s\" %s -utf8output -WaitForUATMutex"),
 		*ProjectPath,
 		GetUATCompilationFlags(),
 		FApp::IsEngineInstalled() ? TEXT(" -installed") : TEXT(""),
@@ -859,7 +859,7 @@ void FMainFrameActionCallbacks::PackageProject( const FName InPlatformInfoName )
 	FString TurnkeyCommandline = FString::Printf(TEXT("Turnkey -command=VerifySdk -platform=%s -UpdateIfNeeded -EditorIO"), *PlatformInfo->DataDrivenPlatformInfo->UBTPlatformString);
 
 	FString ProjectPath = FPaths::IsProjectFilePathSet() ? FPaths::ConvertRelativePathToFull(FPaths::GetProjectFilePath()) : FPaths::RootDir() / FApp::GetProjectName() / FApp::GetProjectName() + TEXT(".uproject");
-	FString CommandLine = FString::Printf(TEXT("-ScriptsForProject=\"%s\" %s BuildCookRun %s%s -nop4 -project=\"%s\" -cook -stage -archive -archivedirectory=\"%s\" -package -ue4exe=\"%s\" %s -utf8output"),
+	FString CommandLine = FString::Printf(TEXT("-ScriptsForProject=\"%s\" %s BuildCookRun %s%s -nop4 -project=\"%s\" -cook -stage -archive -archivedirectory=\"%s\" -package -ue4exe=\"%s\" %s -utf8output -WaitForUATMutex"),
 		*ProjectPath,
 		*TurnkeyCommandline,
 		GetUATCompilationFlags(),
diff --git a/Engine/Source/Editor/UnrealEd/Private/Kismet2/DebuggerCommands.cpp b/Engine/Source/Editor/UnrealEd/Private/Kismet2/DebuggerCommands.cpp
--- a/Engine/Source/Editor/UnrealEd/Private/Kismet2/DebuggerCommands.cpp
+++ b/Engine/Source/Editor/UnrealEd/Private/Kismet2/DebuggerCommands.cpp
@@ -1072,7 +1072,7 @@ static void MakeAllDevicesSubMenu(FMenuBuilder& InMenuBuilder, const PlatformInf
 
 static void TurnkeyInstallSdk(FString PlatformName, bool bPreferFull, bool bForceInstall, const TCHAR* DeviceName=nullptr)
 {
-//	FString CommandLine = FString::Printf(TEXT("Turnkey -command=InstallSdk -platform=%s -BestAvailable -AllowAutoSdk -EditorIO -noturnkeyvariables"), *PlatformName);
+//	FString CommandLine = FString::Printf(TEXT("Turnkey -command=InstallSdk -platform=%s -BestAvailable -AllowAutoSdk -EditorIO -noturnkeyvariables -utf8output -WaitForUATMutex"), *PlatformName);
 
 	FString OptionalOptions;
 	if (bPreferFull)
@@ -1088,7 +1088,7 @@ static void TurnkeyInstallSdk(FString PlatformName, bool bPreferFull, bool bForc
 		OptionalOptions += FString::Printf(TEXT(" -Device=%s"), DeviceName);
 	}
 
-	FString CommandLine = FString::Printf(TEXT("Turnkey -command=VerifySdk -UpdateIfNeeded -platform=%s %s -EditorIO -noturnkeyvariables"), *PlatformName, *OptionalOptions);
+	FString CommandLine = FString::Printf(TEXT("Turnkey -command=VerifySdk -UpdateIfNeeded -platform=%s %s -EditorIO -noturnkeyvariables -utf8output -WaitForUATMutex"), *PlatformName, *OptionalOptions);
 
 	FText TaskName = LOCTEXT("InstallingSdk", "Installing Sdk");
 	IUATHelperModule::Get().CreateUatTask(CommandLine, FText::FromString(PlatformName), TaskName, TaskName, FEditorStyle::GetBrush(TEXT("MainFrame.PackageProject")),
@@ -1170,7 +1170,7 @@ static TAttribute<FText> MakeSdkStatusAttribute(FName IniPlatformName, TSharedPt
 	{
 		// get the status, or Unknown if it's not there
 		const FDataDrivenPlatformInfo& Info = FDataDrivenPlatformInfoRegistry::GetPlatformInfo(IniPlatformName);
-		DDPIPlatformSdkStatus Status = DeviceId.Len() ? Info.GetStatusForDeviceId(DeviceId) : Info.GetSdkStatus();
+		DDPIPlatformSdkStatus Status = DeviceId.Len() ? Info.GetStatusForDeviceId(DeviceId) : Info.GetSdkStatus(false);
 
 		// @todo turnkey: Have premade FText's by SdkStatus for speed
 		const TCHAR* Desc =
@@ -3167,7 +3167,7 @@ void FInternalPlayWorldCommandCallbacks::LaunchOnDevice(const FString& DeviceId,
 			// if we want to check device flash before we start cooking, kick it off now. we could delay this 
 			if (bUseTurnkey)
 			{
-				FString CommandLine = FString::Printf(TEXT("Turnkey -command=VerifySdk -UpdateIfNeeded -platform=%s -EditorIO -noturnkeyvariables -device=%s"), *UBTPlatformName, *DeviceName);
+				FString CommandLine = FString::Printf(TEXT("Turnkey -command=VerifySdk -UpdateIfNeeded -platform=%s -EditorIO -noturnkeyvariables -device=%s -utf8output -WaitForUATMutex"), *UBTPlatformName, *DeviceName);
 				FText TaskName = LOCTEXT("VerifyingSDK", "Verifying SDK and Device");
 
 				IUATHelperModule::Get().CreateUatTask(CommandLine, FText::FromString(IniPlatformName), TaskName, TaskName, FEditorStyle::GetBrush(TEXT("MainFrame.PackageProject")),
diff --git a/Engine/Source/Runtime/Core/Private/Misc/DataDrivenPlatformInfoRegistry.cpp b/Engine/Source/Runtime/Core/Private/Misc/DataDrivenPlatformInfoRegistry.cpp
--- a/Engine/Source/Runtime/Core/Private/Misc/DataDrivenPlatformInfoRegistry.cpp
+++ b/Engine/Source/Runtime/Core/Private/Misc/DataDrivenPlatformInfoRegistry.cpp
@@ -15,6 +15,7 @@ namespace
 	TMap<FName, FDataDrivenPlatformInfo> DataDrivenPlatforms;
 	TArray<FName> SortedPlatformNames;
 	TArray<const FDataDrivenPlatformInfo*> SortedPlatformInfos;
+	FCriticalSection DDPILocker;
 }
 
 #if DDPI_HAS_EXTENDED_PLATFORMINFO_DATA
@@ -406,12 +407,10 @@ static void PrepForTurnkeyReport(FString& Command, FString& BaseCommandline, FSt
 
 	Command = TEXT("{EngineDir}Build/BatchFiles/RunuAT");
 //	Command = TEXT("{EngineDir}/Binaries/DotNET/AutomationTool.exe");
-	BaseCommandline = FString::Printf(TEXT("Turnkey -command=VerifySdk -ReportFilename=\"%s\" -log=\"%s\""), *ReportFilename, *LogFilename);
+	BaseCommandline = FString::Printf(TEXT("Turnkey -utf8output -WaitForUATMutex -command=VerifySdk -ReportFilename=\"%s\" -log=\"%s\""), *ReportFilename, *LogFilename);
 
 	// convert into appropriate calls for the current platform
 	FPlatformProcess::ModifyCreateProcParams(Command, BaseCommandline, FGenericPlatformProcess::ECreateProcHelperFlags::AppendPlatformScriptExtension | FGenericPlatformProcess::ECreateProcHelperFlags::RunThroughShell);
-
-	UE_LOG(LogInit, Log, TEXT("Running Turnkey SDK detection: '%s %s'"), *Command, *BaseCommandline);
 }
 
 static FString ConvertToDDPIPlatform(const FString& Platform)
@@ -474,13 +473,19 @@ void FDataDrivenPlatformInfoRegistry::UpdateSdkStatus()
 	PrepForTurnkeyReport(Command, BaseCommandline, ReportFilename);
 	FString Commandline = BaseCommandline + FString(TEXT(" -platform=")) + FString::JoinBy(DataDrivenPlatforms, TEXT("+"), [](TPair<FName, FDataDrivenPlatformInfo> Pair) { return ConvertToUATPlatform(Pair.Key.ToString()); });
 
-	// reset status to unknown
-	for (auto& It : DataDrivenPlatforms)
+	UE_LOG(LogInit, Log, TEXT("Running Turnkey SDK detection: '%s %s'"), *Command, *Commandline);
+
 	{
-		It.Value.SdkStatus = DDPIPlatformSdkStatus::Querying;
-		
-		// reset the per-device status when querying general Sdk status
-		It.Value.ClearDeviceStatus();
+		FScopeLock Lock(&DDPILocker);
+
+		// reset status to unknown
+		for (auto& It : DataDrivenPlatforms)
+		{
+			It.Value.SdkStatus = DDPIPlatformSdkStatus::Querying;
+
+			// reset the per-device status when querying general Sdk status
+			It.Value.ClearDeviceStatus();
+		}
 	}
 
 	FMonitoredProcess* TurnkeyProcess = new FMonitoredProcess(Command, Commandline, true, false);
@@ -488,6 +493,8 @@ void FDataDrivenPlatformInfoRegistry::UpdateSdkStatus()
 	{
 		AsyncTask(ENamedThreads::GameThread, [ReportFilename, TurnkeyProcess, ExitCode]()
 		{
+			FScopeLock Lock(&DDPILocker);
+
 			if (ExitCode == 0 || ExitCode == 10)
 			{
 				TArray<FString> Contents;
@@ -589,10 +596,16 @@ void FDataDrivenPlatformInfoRegistry::UpdateDeviceSdkStatus(TArray<FString> Plat
 
 	FString Commandline = BaseCommandline + FString(TEXT(" -Device=")) + FString::JoinBy(PlatformDeviceIds, TEXT("+"), [](FString Id) { return ConvertToUATDeviceId(Id); });
 
-	// set status to querying
-	for (const FString& Id : PlatformDeviceIds)
+	UE_LOG(LogInit, Log, TEXT("Running Turnkey SDK detection: '%s %s'"), *Command, *Commandline);
+
 	{
-		DeviceIdToInfo(Id).PerDeviceStatus.Add(ConvertToDDPIDeviceId(Id), DDPIPlatformSdkStatus::Querying);
+		FScopeLock Lock(&DDPILocker);
+
+		// set status to querying
+		for (const FString& Id : PlatformDeviceIds)
+		{
+			DeviceIdToInfo(Id).PerDeviceStatus.Add(ConvertToDDPIDeviceId(Id), DDPIPlatformSdkStatus::Querying);
+		}
 	}
 
 	FMonitoredProcess* TurnkeyProcess = new FMonitoredProcess(Command, Commandline, true, false);
@@ -600,6 +613,8 @@ void FDataDrivenPlatformInfoRegistry::UpdateDeviceSdkStatus(TArray<FString> Plat
 	{
 		AsyncTask(ENamedThreads::GameThread, [ReportFilename, TurnkeyProcess, PlatformDeviceIds, ExitCode]()
 		{
+			FScopeLock Lock(&DDPILocker);
+
 			if (ExitCode == 0 || ExitCode == 10)
 			{
 				TArray<FString> Contents;
@@ -653,6 +668,8 @@ void FDataDrivenPlatformInfoRegistry::UpdateDeviceSdkStatus(TArray<FString> Plat
 
 void FDataDrivenPlatformInfoRegistry::ClearDeviceStatus(FName PlatformName)
 {
+	FScopeLock Lock(&DDPILocker);
+
 	if (PlatformName != NAME_None)
 	{
 		for (auto It : DataDrivenPlatforms)
@@ -677,7 +694,10 @@ DDPIPlatformSdkStatus FDataDrivenPlatformInfo::GetStatusForDeviceId(const FStrin
 
 void FDataDrivenPlatformInfo::ClearDeviceStatus()
 {
-	PerDeviceStatus.Empty();
+	for (auto& Pair : PerDeviceStatus)
+	{
+		Pair.Value = DDPIPlatformSdkStatus::Unknown;
+	}
 }
 
 #endif
