Fix issues with cpu trace events emitted on other channels.

Trace cpu spec events were only gated on the given channel, but must also depend on the CpuChannel. Also clarified some details in documentation.

#rnx

diff --git a/Engine/Source/Runtime/Core/Public/ProfilingDebugging/CpuProfilerTrace.h b/Engine/Source/Runtime/Core/Public/ProfilingDebugging/CpuProfilerTrace.h
--- a/Engine/Source/Runtime/Core/Public/ProfilingDebugging/CpuProfilerTrace.h
+++ b/Engine/Source/Runtime/Core/Public/ProfilingDebugging/CpuProfilerTrace.h
@@ -19,15 +19,52 @@
 
 UE_TRACE_CHANNEL_EXTERN(CpuChannel, CORE_API);
 
+/*
+ * Facilities for tracing timed cpu events. Two types of events are supported, static events where the identifier is
+ * known at compile time, and dynamic event were identifiers can be constructed in runtime. Static events have lower overhead
+ * so always prefer to use them if possible.
+ *
+ * Events are tracked per thread, so begin/end calls must be matched and called on the same thread. It is possible to use any channel
+ * to emit the events, but both that channel and the CpuChannel must then be enabled.
+ *
+ * Usage of the scope macros is highly encouraged in order to avoid mistakes.
+ */
 struct FCpuProfilerTrace
 {
 	CORE_API static void Init(const TCHAR* CmdLine);
 	CORE_API static void Shutdown();
+	/*
+	 * Output cpu event definition (spec).
+	 * @param Name Event name
+	 * @return Event definition id
+	 */
 	FORCENOINLINE CORE_API static uint32 OutputEventType(const ANSICHAR* Name);
+	/*
+	 * Output cpu event definition (spec).
+	 * @param Name Event name
+	 * @return Event definition id
+	 */
 	FORCENOINLINE CORE_API static uint32 OutputEventType(const TCHAR* Name);
+	/*
+	 * Output begin event marker for a given spec. Must always be matched with an end event.
+	 * @param SpecId Event definition id.
+	 */
 	CORE_API static void OutputBeginEvent(uint32 SpecId);
+	/*
+	 * Output begin event marker for a dynamic event name. This is more expensive than statically known event
+	 * names using \ref OutputBeginEvent. Must always be matched with an end event.
+	 * @param Name Name of event
+	 */
 	CORE_API static void OutputBeginDynamicEvent(const ANSICHAR* Name);
+	/*
+	 * Output begin event marker for a dynamic event name. This is more expensive than statically known event
+	 * names using \ref OutputBeginEvent. Must always be matched with an end event.
+	 * @param Name Name of event
+	 */
 	CORE_API static void OutputBeginDynamicEvent(const TCHAR* Name);
+	/*
+	 * Output end event marker for static or dynamic event for the currently open scope.
+	 */
 	CORE_API static void OutputEndEvent();
 
 	struct FEventScope
@@ -90,13 +127,15 @@ struct FCpuProfilerTrace
 #define TRACE_CPUPROFILER_SHUTDOWN() \
 	FCpuProfilerTrace::Shutdown();
 
+// Note: The event will be emitted only if both the given channel and CpuChannel is enabled.
 #define TRACE_CPUPROFILER_EVENT_SCOPE_ON_CHANNEL_STR(NameStr, Channel) \
 	static uint32 PREPROCESSOR_JOIN(__CpuProfilerEventSpecId, __LINE__); \
-	if (bool(Channel) && PREPROCESSOR_JOIN(__CpuProfilerEventSpecId, __LINE__) == 0) { \
+	if (bool(Channel|CpuChannel) && PREPROCESSOR_JOIN(__CpuProfilerEventSpecId, __LINE__) == 0) { \
 		PREPROCESSOR_JOIN(__CpuProfilerEventSpecId, __LINE__) = FCpuProfilerTrace::OutputEventType(NameStr); \
 	} \
 	FCpuProfilerTrace::FEventScope PREPROCESSOR_JOIN(__CpuProfilerEventScope, __LINE__)(PREPROCESSOR_JOIN(__CpuProfilerEventSpecId, __LINE__), Channel);
 
+// Note: The event will be emitted only if both the given channel and CpuChannel is enabled.
 #define TRACE_CPUPROFILER_EVENT_SCOPE_ON_CHANNEL(Name, Channel) \
 	TRACE_CPUPROFILER_EVENT_SCOPE_ON_CHANNEL_STR(#Name, Channel)
 
