Profiler updates

1) Fix Visual Studio Profiling PerfAPI build script pathing for modern era of Visual Studio, add IDEDir to UBT for base path of PerfAPI lib/headers

2) Add initial support for Superluminal instrumentation API, if it's installed and -Superluminal is passed on the command line. Dial back color saturation a bit to go easier on the eyes (my eyes!!). This is WIP support, we are technically violating the API rules by feeding in changing strings to the same profile marker. So, with this change you can see the instrumentation markers in the timeline, but using all the event features in Superluminal causes weird artifacts like event strings changing as you use profiler features. Subsequent change will attempt to pry apart static and instance string data in the instrumentation, as this is a common trope across profilers (and non-PC platforms).

Updated FExternalProfiler API to forward along color information from the instrumentation. This is a minor breaking change as it's very unlikely there are direct callers of these APIs.

diff --git a/Engine/Source/Programs/UnrealBuildTool/Platform/Windows/UEBuildWindows.cs b/Engine/Source/Programs/UnrealBuildTool/Platform/Windows/UEBuildWindows.cs
--- a/Engine/Source/Programs/UnrealBuildTool/Platform/Windows/UEBuildWindows.cs
+++ b/Engine/Source/Programs/UnrealBuildTool/Platform/Windows/UEBuildWindows.cs
@@ -354,6 +354,14 @@ namespace UnrealBuildTool
 			get { return WindowsPlatform.FindDiaSdkDirs(Environment.Compiler).Select(x => x.FullName).FirstOrDefault(); }
 		}
 
+		/// <summary>
+		/// Directory containing the IDE package (Professional, Community, etc...)
+		/// </summary>
+		public string IDEDir
+		{
+			get { return WindowsPlatform.FindVisualStudioInstallations(Environment.Compiler).Select(x => x.BaseDir.FullName).FirstOrDefault(); }
+		}
+
 		/// <summary>
 		/// When using a Visual Studio compiler, returns the version name as a string
 		/// </summary>
@@ -594,6 +602,11 @@ namespace UnrealBuildTool
 			get { return Inner.DiaSdkDir; }
 		}
 		
+		public string IDEDir
+		{
+			get { return Inner.IDEDir; }
+		}
+
 		public string GetArchitectureSubpath()
 		{
 			return WindowsExports.GetArchitectureSubpath(Architecture);
@@ -611,7 +624,7 @@ namespace UnrealBuildTool
 		/// Information about a particular Visual Studio installation
 		/// </summary>
 		[DebuggerDisplay("{BaseDir}")]
-		class VisualStudioInstallation
+		public class VisualStudioInstallation
 		{
 			/// <summary>
 			/// Base directory for the installation
@@ -995,7 +1008,7 @@ namespace UnrealBuildTool
 		/// </summary>
 		/// <param name="Compiler">Version of the toolchain to look for.</param>
 		/// <returns>List of directories containing Visual Studio installations</returns>
-		static List<VisualStudioInstallation> FindVisualStudioInstallations(WindowsCompiler Compiler)
+		public static List<VisualStudioInstallation> FindVisualStudioInstallations(WindowsCompiler Compiler)
 		{
 			List<VisualStudioInstallation> Installations;
 			if(!CachedVisualStudioInstallations.TryGetValue(Compiler, out Installations))
diff --git a/Engine/Source/Runtime/Core/Core.Build.cs b/Engine/Source/Runtime/Core/Core.Build.cs
--- a/Engine/Source/Runtime/Core/Core.Build.cs
+++ b/Engine/Source/Runtime/Core/Core.Build.cs
@@ -158,16 +158,15 @@ public class Core : ModuleRules
 			}
 		}
 
-
 		// On Windows platform, VSPerfExternalProfiler.cpp needs access to "VSPerf.h".  This header is included with Visual Studio, but it's not in a standard include path.
 		if (Target.Platform.IsInGroup(UnrealPlatformGroup.Windows))
 		{
-			var VisualStudioVersionNumber = "11.0";
-			var SubFolderName = (Target.Platform == UnrealTargetPlatform.Win32) ? "PerfSDK" : "x64/PerfSDK";
-
-			string PerfIncludeDirectory = Path.Combine(Environment.GetFolderPath(Environment.SpecialFolder.ProgramFilesX86), String.Format("Microsoft Visual Studio {0}/Team Tools/Performance Tools/{1}", VisualStudioVersionNumber, SubFolderName));
+			string VisualStudioInstallation = Target.WindowsPlatform.IDEDir;
+			string SubFolderName = (Target.Platform == UnrealTargetPlatform.Win32) ? "PerfSDK/" : "x64/PerfSDK/";
+			string PerfIncludeDirectory = Path.Combine(VisualStudioInstallation, String.Format("Team Tools/Performance Tools/{0}", SubFolderName));
 
-			if (File.Exists(Path.Combine(PerfIncludeDirectory, "VSPerf.h")))
+			if (File.Exists(Path.Combine(PerfIncludeDirectory, "VSPerf.h"))
+				&& Target.Configuration != UnrealTargetConfiguration.Shipping)
 			{
 				PrivateIncludePaths.Add(PerfIncludeDirectory);
 				PublicDefinitions.Add("WITH_VS_PERF_PROFILER=1");
@@ -183,6 +182,26 @@ public class Core : ModuleRules
 			PublicDefinitions.Add("WITH_VS_PERF_PROFILER=0");
 		}
 
+		// Superluminal instrumentation support, if one has it installed
+		if (Target.Platform.IsInGroup(UnrealPlatformGroup.Windows))
+		{
+			string SuperluminalRootDir = Path.Combine(Environment.GetFolderPath(Environment.SpecialFolder.ProgramFiles), "Superluminal/Performance/API/");
+			string SubFolderName = (Target.Platform == UnrealTargetPlatform.Win32) ? "lib/x86/" : "lib/x64/";
+			string SuperluminalLibDir = Path.Combine(SuperluminalRootDir, SubFolderName);
+
+			if (Directory.Exists(SuperluminalRootDir)
+				&& Target.Configuration != UnrealTargetConfiguration.Shipping)
+			{
+				PublicSystemIncludePaths.Add(Path.Combine(SuperluminalRootDir, "include/"));
+				PublicAdditionalLibraries.Add(Path.Combine(SuperluminalLibDir, "PerformanceAPI_MD.lib"));
+				PublicDefinitions.Add("WITH_SUPERLUMINAL_PROFILER=1");
+			}
+			else
+			{
+				PublicDefinitions.Add("WITH_SUPERLUMINAL_PROFILER=0");
+			}
+		}
+
 		if (Target.bWithDirectXMath && (Target.Platform == UnrealTargetPlatform.Win64 || Target.Platform == UnrealTargetPlatform.Win32))
 		{
 			PublicDefinitions.Add("WITH_DIRECTXMATH=1");
diff --git a/Engine/Source/Runtime/Core/Private/GenericPlatform/GenericPlatformMisc.cpp b/Engine/Source/Runtime/Core/Private/GenericPlatform/GenericPlatformMisc.cpp
--- a/Engine/Source/Runtime/Core/Private/GenericPlatform/GenericPlatformMisc.cpp
+++ b/Engine/Source/Runtime/Core/Private/GenericPlatform/GenericPlatformMisc.cpp
@@ -480,7 +480,7 @@ void FGenericPlatformMisc::BeginNamedEvent(const struct FColor& Color, const ANS
 	FExternalProfiler* CurrentProfiler = FActiveExternalProfilerBase::GetActiveProfiler();
 	if (CurrentProfiler != NULL)
 	{
-		CurrentProfiler->StartScopedEvent(ANSI_TO_TCHAR(Text));
+		CurrentProfiler->StartScopedEvent(Color, ANSI_TO_TCHAR(Text));
 	}
 #endif
 #if CPUPROFILERTRACE_ENABLED
@@ -499,7 +499,7 @@ void FGenericPlatformMisc::BeginNamedEvent(const struct FColor& Color, const TCH
 
 	if (CurrentProfiler != NULL)
 	{
-		CurrentProfiler->StartScopedEvent(Text);
+		CurrentProfiler->StartScopedEvent(Color, Text);
 	}
 #endif
 #if CPUPROFILERTRACE_ENABLED
diff --git a/Engine/Source/Runtime/Core/Private/ProfilingDebugging/Windows/SuperluminalExternalProfiler.cpp b/Engine/Source/Runtime/Core/Private/ProfilingDebugging/Windows/SuperluminalExternalProfiler.cpp
new file mode 100644
--- /dev/null
+++ b/Engine/Source/Runtime/Core/Private/ProfilingDebugging/Windows/SuperluminalExternalProfiler.cpp
@@ -0,0 +1,86 @@
+// Copyright Epic Games, Inc. All Rights Reserved.
+
+#include "CoreTypes.h"
+#include "HAL/PlatformProcess.h"
+#include "Misc/AssertionMacros.h"
+#include "ProfilingDebugging/ExternalProfiler.h"
+#include "Features/IModularFeatures.h"
+#include "Templates/UniquePtr.h"
+
+#if WITH_SUPERLUMINAL_PROFILER && UE_EXTERNAL_PROFILING_ENABLED
+
+#define PERFORMANCEAPI_ENABLED 1
+#include "Superluminal/PerformanceAPI_capi.h"
+
+// Use a custom color pack function here because our colors are full intensity
+//  and are painful/difficult to read behind text in the app's flame graph
+static uint32_t UPerformanceAPI_Pack(uint8_t R, uint8_t G, uint8_t B)
+{
+	const float Desaturate = 0.33f;
+	float Grey = (R * 0.299f) + (G * 0.587f) + (B * 0.144f);
+	R = (uint8_t)(Grey * Desaturate + ((float)R) * (1.f - Desaturate));
+	G = (uint8_t)(Grey * Desaturate + ((float)G) * (1.f - Desaturate));
+	B = (uint8_t)(Grey * Desaturate + ((float)B) * (1.f - Desaturate));
+	return ((((uint32_t)(R)) << 24) | (((uint32_t)(G)) << 16) | (((uint32_t)(B)) << 8));
+}
+
+/**
+ * Superluminal Performance Profiler implementation of FExternalProfiler
+ */
+class SuperluminalExternalProfiler : public FExternalProfiler
+{
+public:
+	/** Constructor */
+	SuperluminalExternalProfiler()
+	{
+		// Register as a modular feature
+		IModularFeatures::Get().RegisterModularFeature( FExternalProfiler::GetFeatureName(), this );
+	}
+
+	/** Destructor */
+	virtual ~SuperluminalExternalProfiler() {}
+
+	/** Gets the name of this profiler as a string.  This is used to allow the user to select this profiler in a system configuration file or on the command-line */
+	virtual const TCHAR* GetProfilerName() const override
+	{
+		return TEXT( "Superluminal" );
+	}
+
+	virtual void FrameSync() override {}
+	virtual void ProfilerPauseFunction() override {}
+	virtual void ProfilerResumeFunction() override {}
+
+	/** Starts a scoped event specific to the profiler. */
+	virtual void StartScopedEvent(const struct FColor& Color, const TCHAR* Text) override
+	{
+		PerformanceAPI_BeginEvent_Wide(Text, nullptr, UPerformanceAPI_Pack(Color.R, Color.G, Color.B));
+	}
+
+	/** Ends a scoped event specific to the profiler. */
+	virtual void EndScopedEvent() override
+	{
+		PerformanceAPI_EndEvent();
+	}
+
+	/** Dead code stripper will cull our static below unless we do something. */
+	virtual bool Initialize()
+	{
+		return true;
+	}
+};
+
+namespace SuperluminalProfiler
+{
+	struct FAtModuleInit
+	{
+		FAtModuleInit()
+		{
+			static TUniquePtr<SuperluminalExternalProfiler> ProfilerSuperluminal = MakeUnique<SuperluminalExternalProfiler>();
+			ProfilerSuperluminal->Initialize();
+		}
+	};
+
+	static FAtModuleInit AtModuleInit;
+}
+
+#endif	// WITH_SUPERLUMINAL_PROFILER && UE_EXTERNAL_PROFILING_ENABLED
diff --git a/Engine/Source/Runtime/Core/Private/ProfilingDebugging/Windows/VTuneExternalProfiler.cpp b/Engine/Source/Runtime/Core/Private/ProfilingDebugging/Windows/VTuneExternalProfiler.cpp
--- a/Engine/Source/Runtime/Core/Private/ProfilingDebugging/Windows/VTuneExternalProfiler.cpp
+++ b/Engine/Source/Runtime/Core/Private/ProfilingDebugging/Windows/VTuneExternalProfiler.cpp
@@ -77,7 +77,7 @@ public:
 		VTResume();
 	}
 
-	void StartScopedEvent(const TCHAR* Text) 
+	void StartScopedEvent(const struct FColor& Color, const TCHAR* Text) override
 	{
 		uint32 Hash = GetTypeHash(Text);
 
@@ -93,7 +93,7 @@ public:
 		__itt_task_begin(Domain, __itt_null, __itt_null, Handle);
 	}
 
-	void EndScopedEvent()
+	void EndScopedEvent() override
 	{
 		//Deactivate last event
 		__itt_task_end(Domain);
diff --git a/Engine/Source/Runtime/Core/Private/Windows/WindowsPlatformMisc.cpp b/Engine/Source/Runtime/Core/Private/Windows/WindowsPlatformMisc.cpp
--- a/Engine/Source/Runtime/Core/Private/Windows/WindowsPlatformMisc.cpp
+++ b/Engine/Source/Runtime/Core/Private/Windows/WindowsPlatformMisc.cpp
@@ -868,7 +868,7 @@ void FWindowsPlatformMisc::BeginNamedEvent(const struct FColor& Color, const TCH
 	FExternalProfiler* Profiler = FActiveExternalProfilerBase::GetActiveProfiler();
 	if (Profiler)
 	{
-		Profiler->StartScopedEvent(Text);
+		Profiler->StartScopedEvent(Color, Text);
 	}
 #endif
 #if CPUPROFILERTRACE_ENABLED
@@ -887,7 +887,7 @@ void FWindowsPlatformMisc::BeginNamedEvent(const struct FColor& Color, const ANS
 	FExternalProfiler* Profiler = FActiveExternalProfilerBase::GetActiveProfiler();
 	if (Profiler)
 	{
-		Profiler->StartScopedEvent(ANSI_TO_TCHAR(Text));
+		Profiler->StartScopedEvent(Color, ANSI_TO_TCHAR(Text));
 	}
 #endif
 #if CPUPROFILERTRACE_ENABLED
diff --git a/Engine/Source/Runtime/Core/Public/ProfilingDebugging/ExternalProfiler.h b/Engine/Source/Runtime/Core/Public/ProfilingDebugging/ExternalProfiler.h
--- a/Engine/Source/Runtime/Core/Public/ProfilingDebugging/ExternalProfiler.h
+++ b/Engine/Source/Runtime/Core/Public/ProfilingDebugging/ExternalProfiler.h
@@ -60,7 +60,7 @@ public:
 	static FName GetFeatureName();
 
 	/** Starts a scoped event specific to the profiler. */
-	virtual void StartScopedEvent(const TCHAR* Text) {};
+	virtual void StartScopedEvent(const struct FColor& Color, const TCHAR* Text) {};
 
 	/** Ends a scoped event specific to the profiler. */
 	virtual void EndScopedEvent() {};
