Defensive fix for Visual Studio not being installed on the machine

diff --git a/Engine/Source/Runtime/Core/Core.Build.cs b/Engine/Source/Runtime/Core/Core.Build.cs
--- a/Engine/Source/Runtime/Core/Core.Build.cs
+++ b/Engine/Source/Runtime/Core/Core.Build.cs
@@ -161,17 +161,23 @@ public class Core : ModuleRules
 		// On Windows platform, VSPerfExternalProfiler.cpp needs access to "VSPerf.h".  This header is included with Visual Studio, but it's not in a standard include path.
 		if (Target.Platform.IsInGroup(UnrealPlatformGroup.Windows))
 		{
+			bool VSPerfDefined = false;
 			string VisualStudioInstallation = Target.WindowsPlatform.IDEDir;
-			string SubFolderName = (Target.Platform == UnrealTargetPlatform.Win32) ? "PerfSDK/" : "x64/PerfSDK/";
-			string PerfIncludeDirectory = Path.Combine(VisualStudioInstallation, String.Format("Team Tools/Performance Tools/{0}", SubFolderName));
-
-			if (File.Exists(Path.Combine(PerfIncludeDirectory, "VSPerf.h"))
-				&& Target.Configuration != UnrealTargetConfiguration.Shipping)
+			if (VisualStudioInstallation != null && VisualStudioInstallation != string.Empty && Directory.Exists(VisualStudioInstallation))
 			{
-				PrivateIncludePaths.Add(PerfIncludeDirectory);
-				PublicDefinitions.Add("WITH_VS_PERF_PROFILER=1");
+				string SubFolderName = (Target.Platform == UnrealTargetPlatform.Win32) ? "PerfSDK/" : "x64/PerfSDK/";
+				string PerfIncludeDirectory = Path.Combine(VisualStudioInstallation, String.Format("Team Tools/Performance Tools/{0}", SubFolderName));
+
+				if (File.Exists(Path.Combine(PerfIncludeDirectory, "VSPerf.h"))
+					&& Target.Configuration != UnrealTargetConfiguration.Shipping)
+				{
+					PrivateIncludePaths.Add(PerfIncludeDirectory);
+					PublicDefinitions.Add("WITH_VS_PERF_PROFILER=1");
+					VSPerfDefined = true;
+				}
 			}
-			else
+
+			if (!VSPerfDefined)
 			{
 				PublicDefinitions.Add("WITH_VS_PERF_PROFILER=0");
 			}
