Superluminal: Make a more precise detection of the appropriate Superluminal installation by checking for the API header we need.

diff --git a/Engine/Source/Runtime/Core/Core.Build.cs b/Engine/Source/Runtime/Core/Core.Build.cs
--- a/Engine/Source/Runtime/Core/Core.Build.cs
+++ b/Engine/Source/Runtime/Core/Core.Build.cs
@@ -195,7 +195,7 @@ public class Core : ModuleRules
 			string SubFolderName = (Target.Platform == UnrealTargetPlatform.Win32) ? "lib/x86/" : "lib/x64/";
 			string SuperluminalLibDir = Path.Combine(SuperluminalRootDir, SubFolderName);
 
-			if (Directory.Exists(SuperluminalRootDir)
+			if (File.Exists(Path.Combine(SuperluminalRootDir, "include/Superluminal/PerformanceAPI_capi.h"))
 				&& Target.Configuration != UnrealTargetConfiguration.Shipping)
 			{
 				PublicSystemIncludePaths.Add(Path.Combine(SuperluminalRootDir, "include/"));
