M2 - Renderdoc is loading the executable in a way so that the static initzialization happens on another thread than the thread from which the executable is running from.

diff --git a/Engine/Source/Runtime/Core/Private/HAL/ThreadingBase.cpp b/Engine/Source/Runtime/Core/Private/HAL/ThreadingBase.cpp
--- a/Engine/Source/Runtime/Core/Private/HAL/ThreadingBase.cpp
+++ b/Engine/Source/Runtime/Core/Private/HAL/ThreadingBase.cpp
@@ -60,10 +60,9 @@ FTaskTagScope::FTaskTagScope(bool InTagOnlyIfNone, ETaskTag InTag) : Tag(InTag),
 
 	if (ActiveTaskTag == ETaskTag::EStaticInit)
 	{
-#if IS_RUNNING_GAMETHREAD_ON_EXTERNAL_THREAD
 		checkf(Tag == ETaskTag::EGameThread, TEXT("The Gamethread can only be tagged on the inital thread of the application"));
-#else
-		checkf(Tag == ETaskTag::EGameThread && IsRunningDuringStaticInit(), TEXT("The Gamethread can only be tagged on the inital thread of the application"));
+#if !IS_RUNNING_GAMETHREAD_ON_EXTERNAL_THREAD
+		ensureMsgf(IsRunningDuringStaticInit(), TEXT("Static initialization should have happened on the same thread as the main thread"));
 #endif
 	}
 
