Set IS_RUNNING_GAMETHREAD_ON_EXTERNAL_THREAD on Switch and set the EGameThread tag from the Switch's RealMain()

#rnx

tested:
* ShooterGame on Switch
* ShooterGame on Win64

diff --git a/Engine/Source/Runtime/Core/Core.Build.cs b/Engine/Source/Runtime/Core/Core.Build.cs
--- a/Engine/Source/Runtime/Core/Core.Build.cs
+++ b/Engine/Source/Runtime/Core/Core.Build.cs
@@ -221,10 +221,6 @@ public class Core : ModuleRules
 		{
 			PublicDefinitions.Add("IS_RUNNING_GAMETHREAD_ON_EXTERNAL_THREAD=1");
 		}
-		else
-		{
-			PublicDefinitions.Add("IS_RUNNING_GAMETHREAD_ON_EXTERNAL_THREAD=0");
-		}
 
 		// Set a macro to allow FApp::GetBuildTargetType() to detect client targts
 		if (Target.Type == TargetRules.TargetType.Client)
diff --git a/Engine/Source/Runtime/Core/Private/HAL/ThreadingBase.cpp b/Engine/Source/Runtime/Core/Private/HAL/ThreadingBase.cpp
--- a/Engine/Source/Runtime/Core/Private/HAL/ThreadingBase.cpp
+++ b/Engine/Source/Runtime/Core/Private/HAL/ThreadingBase.cpp
@@ -15,6 +15,10 @@
 
 #include <atomic>
 
+#ifndef IS_RUNNING_GAMETHREAD_ON_EXTERNAL_THREAD
+#define IS_RUNNING_GAMETHREAD_ON_EXTERNAL_THREAD 0
+#endif
+
 DEFINE_STAT( STAT_EventWaitWithId );
 DEFINE_STAT( STAT_EventTriggerWithId );
 
