Build actor descriptors hash through serialization.

diff --git a/Engine/Source/Runtime/Core/Public/Misc/HashBuilder.h b/Engine/Source/Runtime/Core/Public/Misc/HashBuilder.h
--- a/Engine/Source/Runtime/Core/Public/Misc/HashBuilder.h
+++ b/Engine/Source/Runtime/Core/Public/Misc/HashBuilder.h
@@ -13,10 +13,15 @@ public:
 		: Hash(~InHash)
 	{}
 
+	void Append(const void* Data, int64 Num)
+	{
+		Hash = FCrc::MemCrc32(Data, Num, Hash);
+	}
+
 	template <typename T>
 	typename TEnableIf<TIsPODType<T>::Value, FHashBuilder&>::Type Append(const T& InData)
 	{
-		Hash = FCrc::MemCrc32(&InData, sizeof(T), Hash);
+		Append(&InData, sizeof(T));
 		return *this;
 	}
 
@@ -65,4 +70,31 @@ public:
 
 private:
 	uint32 Hash;
-};
\ No newline at end of file
+};
+
+/**
+ * FArchive adapter for FHashBuilder
+ */
+class FHashBuilderArchive : public FArchive
+{
+public:
+	FHashBuilderArchive()
+	{
+		SetIsLoading(false);
+		SetIsSaving(true);
+		SetIsPersistent(false);
+	}
+
+	virtual FString GetArchiveName() const { return TEXT("FHashBuilderArchive"); }
+
+	void Serialize(void* Data, int64 Num) override { HashBuilder.Append(Data, Num); }
+
+	using FArchive::operator<<;
+	virtual FArchive& operator<<(class FName& Value) override { HashBuilder << Value; return *this; }
+	virtual FArchive& operator<<(class UObject*& Value) override { check(0); return *this; }
+
+	uint32 GetHash() const { return HashBuilder.GetHash(); }
+
+protected:
+	FHashBuilder HashBuilder;
+};
diff --git a/Engine/Source/Runtime/Engine/Private/WorldPartition/ActorPartition/PartitionActorDesc.cpp b/Engine/Source/Runtime/Engine/Private/WorldPartition/ActorPartition/PartitionActorDesc.cpp
--- a/Engine/Source/Runtime/Engine/Private/WorldPartition/ActorPartition/PartitionActorDesc.cpp
+++ b/Engine/Source/Runtime/Engine/Private/WorldPartition/ActorPartition/PartitionActorDesc.cpp
@@ -26,10 +26,4 @@ void FPartitionActorDesc::Serialize(FArchive& Ar)
 
 	Ar << GridSize << GridIndexX << GridIndexY << GridIndexZ;
 }
-
-void FPartitionActorDesc::BuildHash(FHashBuilder& HashBuilder)
-{
-	FWorldPartitionActorDesc::BuildHash(HashBuilder);
-	HashBuilder << GridSize << GridIndexX << GridIndexY << GridIndexZ;
-}
 #endif
diff --git a/Engine/Source/Runtime/Engine/Private/WorldPartition/HLOD/HLODActorDesc.cpp b/Engine/Source/Runtime/Engine/Private/WorldPartition/HLOD/HLODActorDesc.cpp
--- a/Engine/Source/Runtime/Engine/Private/WorldPartition/HLOD/HLODActorDesc.cpp
+++ b/Engine/Source/Runtime/Engine/Private/WorldPartition/HLOD/HLODActorDesc.cpp
@@ -36,12 +36,4 @@ void FHLODActorDesc::Serialize(FArchive& Ar)
 		HLODLayer = HLODLayerStr;
 	}
 }
-
-void FHLODActorDesc::BuildHash(FHashBuilder& HashBuilder)
-{
-	FWorldPartitionActorDesc::BuildHash(HashBuilder);
-	HashBuilder << SubActors;
-	HashBuilder << HLODLayer.ToString();
-}
-
 #endif
diff --git a/Engine/Source/Runtime/Engine/Private/WorldPartition/LevelInstance/LevelInstanceActorDesc.cpp b/Engine/Source/Runtime/Engine/Private/WorldPartition/LevelInstance/LevelInstanceActorDesc.cpp
--- a/Engine/Source/Runtime/Engine/Private/WorldPartition/LevelInstance/LevelInstanceActorDesc.cpp
+++ b/Engine/Source/Runtime/Engine/Private/WorldPartition/LevelInstance/LevelInstanceActorDesc.cpp
@@ -40,10 +40,4 @@ void FLevelInstanceActorDesc::Serialize(FArchive& Ar)
 		}
 	}
 }
-
-void FLevelInstanceActorDesc::BuildHash(FHashBuilder& HashBuilder)
-{
-	FWorldPartitionActorDesc::BuildHash(HashBuilder);
-	HashBuilder << LevelPackage << LevelInstanceTransform;
-}
 #endif
diff --git a/Engine/Source/Runtime/Engine/Private/WorldPartition/WorldPartitionActorDesc.cpp b/Engine/Source/Runtime/Engine/Private/WorldPartition/WorldPartitionActorDesc.cpp
--- a/Engine/Source/Runtime/Engine/Private/WorldPartition/WorldPartitionActorDesc.cpp
+++ b/Engine/Source/Runtime/Engine/Private/WorldPartition/WorldPartitionActorDesc.cpp
@@ -140,19 +140,19 @@ FString FWorldPartitionActorDesc::ToString() const
 
 void FWorldPartitionActorDesc::UpdateHash()
 {
-	FHashBuilder HashBuilder;
-	BuildHash(HashBuilder);
-	Hash = HashBuilder.GetHash();
-}
-
-void FWorldPartitionActorDesc::BuildHash(FHashBuilder& HashBuilder)
-{
-	HashBuilder << Guid << Class << ActorPackage << ActorPath << BoundsLocation << BoundsExtent << GridPlacement << RuntimeGrid << bActorIsEditorOnly << bLevelBoundsRelevant << Layers << References;
+	FHashBuilderArchive HashBuilderAr;
+	Serialize(HashBuilderAr);
+	Hash = HashBuilderAr.GetHash();
 }
 
 void FWorldPartitionActorDesc::Serialize(FArchive& Ar)
 {
 	Ar << Class << Guid << BoundsLocation << BoundsExtent << GridPlacement << RuntimeGrid << bActorIsEditorOnly << bLevelBoundsRelevant << Layers << References;
+
+	if (!Ar.IsPersistent())
+	{
+		Ar << ActorPackage << ActorPath;
+	}
 }
 
 FBox FWorldPartitionActorDesc::GetBounds() const
diff --git a/Engine/Source/Runtime/Engine/Public/WorldPartition/ActorPartition/PartitionActorDesc.h b/Engine/Source/Runtime/Engine/Public/WorldPartition/ActorPartition/PartitionActorDesc.h
--- a/Engine/Source/Runtime/Engine/Public/WorldPartition/ActorPartition/PartitionActorDesc.h
+++ b/Engine/Source/Runtime/Engine/Public/WorldPartition/ActorPartition/PartitionActorDesc.h
@@ -19,8 +19,6 @@ public:
 	int64 GridIndexZ;
 
 protected:
-	virtual void BuildHash(FHashBuilder& HashBuilder) override;
-
 	virtual void InitFrom(const AActor* InActor) override;
 	virtual void Serialize(FArchive& Ar) override;
 #endif
diff --git a/Engine/Source/Runtime/Engine/Public/WorldPartition/HLOD/HLODActorDesc.h b/Engine/Source/Runtime/Engine/Public/WorldPartition/HLOD/HLODActorDesc.h
--- a/Engine/Source/Runtime/Engine/Public/WorldPartition/HLOD/HLODActorDesc.h
+++ b/Engine/Source/Runtime/Engine/Public/WorldPartition/HLOD/HLODActorDesc.h
@@ -22,8 +22,6 @@ public:
 	inline const FSoftObjectPath& GetHLODLayer() const { return HLODLayer; }
 
 protected:
-	virtual void BuildHash(FHashBuilder& HashBuilder) override;
-
 	virtual void InitFrom(const AActor* InActor) override;
 	virtual void Serialize(FArchive& Ar) override;
 
diff --git a/Engine/Source/Runtime/Engine/Public/WorldPartition/LevelInstance/LevelInstanceActorDesc.h b/Engine/Source/Runtime/Engine/Public/WorldPartition/LevelInstance/LevelInstanceActorDesc.h
--- a/Engine/Source/Runtime/Engine/Public/WorldPartition/LevelInstance/LevelInstanceActorDesc.h
+++ b/Engine/Source/Runtime/Engine/Public/WorldPartition/LevelInstance/LevelInstanceActorDesc.h
@@ -22,8 +22,6 @@ public:
 	inline FName GetLevelPackage() const { return LevelPackage; }
 
 protected:
-	virtual void BuildHash(FHashBuilder& HashBuilder) override;
-
 	virtual void InitFrom(const AActor* InActor) override;
 	virtual void Serialize(FArchive& Ar) override;
 
diff --git a/Engine/Source/Runtime/Engine/Public/WorldPartition/WorldPartitionActorDesc.h b/Engine/Source/Runtime/Engine/Public/WorldPartition/WorldPartitionActorDesc.h
--- a/Engine/Source/Runtime/Engine/Public/WorldPartition/WorldPartitionActorDesc.h
+++ b/Engine/Source/Runtime/Engine/Public/WorldPartition/WorldPartitionActorDesc.h
@@ -89,7 +89,6 @@ protected:
 	FWorldPartitionActorDesc();
 	
 	void UpdateHash();
-	virtual void BuildHash(FHashBuilder& HashBuilder);
 
 	virtual void InitFrom(const AActor* Actor);
 	virtual void Serialize(FArchive& Ar);
