fix build

diff --git a/Engine/Source/Runtime/Core/Private/Windows/WindowsPlatformMisc.cpp b/Engine/Source/Runtime/Core/Private/Windows/WindowsPlatformMisc.cpp
--- a/Engine/Source/Runtime/Core/Private/Windows/WindowsPlatformMisc.cpp
+++ b/Engine/Source/Runtime/Core/Private/Windows/WindowsPlatformMisc.cpp
@@ -868,7 +868,7 @@ void FWindowsPlatformMisc::BeginNamedEvent(const struct FColor& Color, const TCH
 	FExternalProfiler* Profiler = FActiveExternalProfilerBase::GetActiveProfiler();
 	if (Profiler)
 	{
-		Profiler->StartScopedEvent(Text);
+		Profiler->StartScopedEvent(Color, Text);
 	}
 #endif
 #if CPUPROFILERTRACE_ENABLED
@@ -887,7 +887,7 @@ void FWindowsPlatformMisc::BeginNamedEvent(const struct FColor& Color, const ANS
 	FExternalProfiler* Profiler = FActiveExternalProfilerBase::GetActiveProfiler();
 	if (Profiler)
 	{
-		Profiler->StartScopedEvent(ANSI_TO_TCHAR(Text));
+		Profiler->StartScopedEvent(Color, ANSI_TO_TCHAR(Text));
 	}
 #endif
 #if CPUPROFILERTRACE_ENABLED
