Fixed an issue causing the command line '-handleensurepercent=' to never be applied.

- The condition checking if the command line is initialized or not was the wrong way around.

diff --git a/Engine/Source/Runtime/Core/Private/GenericPlatform/GenericPlatformMisc.cpp b/Engine/Source/Runtime/Core/Private/GenericPlatform/GenericPlatformMisc.cpp
--- a/Engine/Source/Runtime/Core/Private/GenericPlatform/GenericPlatformMisc.cpp
+++ b/Engine/Source/Runtime/Core/Private/GenericPlatform/GenericPlatformMisc.cpp
@@ -1283,7 +1283,7 @@ void FGenericPlatformMisc::UpdateHotfixableEnsureSettings()
 	else
 	{
 		float HandleEnsurePercentOnCmdLine = 100.0f;
-		if (!FCommandLine::IsInitialized() && FParse::Value(FCommandLine::Get(), TEXT("handleensurepercent="), HandleEnsurePercentOnCmdLine))
+		if (FCommandLine::IsInitialized() && FParse::Value(FCommandLine::Get(), TEXT("handleensurepercent="), HandleEnsurePercentOnCmdLine))
 		{
 			GenericPlatformMisc::GEnsureChance = HandleEnsurePercentOnCmdLine / 100.0f;
 		}
