HashBuilder: automatically route types that implements GetTypeHash to it instead of going through CRC32.

diff --git a/Engine/Source/Runtime/Core/Public/Misc/HashBuilder.h b/Engine/Source/Runtime/Core/Public/Misc/HashBuilder.h
--- a/Engine/Source/Runtime/Core/Public/Misc/HashBuilder.h
+++ b/Engine/Source/Runtime/Core/Public/Misc/HashBuilder.h
@@ -13,28 +13,29 @@ public:
 		: Hash(~InHash)
 	{}
 
-	void Append(const void* Data, int64 Num)
+	void AppendRaw(const void* Data, int64 Num)
 	{
 		Hash = FCrc::MemCrc32(Data, Num, Hash);
 	}
 
 	template <typename T>
-	typename TEnableIf<TIsPODType<T>::Value, FHashBuilder&>::Type Append(const T& InData)
+	typename TEnableIf<TIsPODType<T>::Value, FHashBuilder&>::Type AppendRaw(const T& InData)
 	{
-		Append(&InData, sizeof(T));
+		AppendRaw(&InData, sizeof(T), Hash);
 		return *this;
 	}
 
-	FHashBuilder& Append(const FString& InString)
+	template <typename T>
+	typename TEnableIf<!TModels<CGetTypeHashable, T>::Value, FHashBuilder&>::Type Append(const T& InData)
 	{
-		uint32 StringHash = GetTypeHash(InString);
-		return Append(StringHash);
+		return AppendRaw(InData);
 	}
 
-	FHashBuilder& Append(const FName& InName)
+	template <typename T>
+	typename TEnableIf<TModels<CGetTypeHashable, T>::Value, FHashBuilder&>::Type Append(const T& InData)
 	{
-		uint32 NameHash = GetTypeHash(InName);
-		return Append(NameHash);
+		Hash = HashCombine(Hash, GetTypeHash(InData));
+		return *this;
 	}
 
 	template <typename T>
@@ -87,7 +88,7 @@ public:
 
 	virtual FString GetArchiveName() const { return TEXT("FHashBuilderArchive"); }
 
-	void Serialize(void* Data, int64 Num) override { HashBuilder.Append(Data, Num); }
+	void Serialize(void* Data, int64 Num) override { HashBuilder.AppendRaw(Data, Num); }
 
 	using FArchive::operator<<;
 	virtual FArchive& operator<<(class FName& Value) override { HashBuilder << Value; return *this; }
@@ -97,4 +98,4 @@ public:
 
 protected:
 	FHashBuilder HashBuilder;
-};
+};
\ No newline at end of file
