M5 - remove assignment

diff --git a/Engine/Source/Runtime/Core/Private/HAL/ThreadingBase.cpp b/Engine/Source/Runtime/Core/Private/HAL/ThreadingBase.cpp
--- a/Engine/Source/Runtime/Core/Private/HAL/ThreadingBase.cpp
+++ b/Engine/Source/Runtime/Core/Private/HAL/ThreadingBase.cpp
@@ -50,7 +50,7 @@ int32 FTaskTagScope::GetStaticThreadId()
 }
 
 thread_local ETaskTag FTaskTagScope::ActiveTaskTag = ETaskTag::EStaticInit;
-static std::atomic_int ActiveNamedThreads = {};
+static std::atomic_int ActiveNamedThreads {};
 
 void FTaskTagScope::SetTagNone()
 {
