Added a FCommandLine::Reset() call that restores FCommandLine to a state as if it wasn't initialised.

#rnx

diff --git a/Engine/Source/Runtime/Core/Private/Misc/CommandLine.cpp b/Engine/Source/Runtime/Core/Private/Misc/CommandLine.cpp
--- a/Engine/Source/Runtime/Core/Private/Misc/CommandLine.cpp
+++ b/Engine/Source/Runtime/Core/Private/Misc/CommandLine.cpp
@@ -29,6 +29,15 @@ bool FCommandLine::IsInitialized()
 	return bIsInitialized;
 }
 
+void FCommandLine::Reset()
+{
+	CmdLine[0] = '\0';
+	OriginalCmdLine[0] = '\0';
+	LoggingCmdLine[0] = '\0';
+	LoggingOriginalCmdLine[0] = '\0';
+	bIsInitialized = false;
+}
+
 const TCHAR* FCommandLine::Get()
 {
 	UE_CLOG(!bIsInitialized, LogInit, Fatal, TEXT("Attempting to get the command line but it hasn't been initialized yet."));
diff --git a/Engine/Source/Runtime/Core/Public/Misc/CommandLine.h b/Engine/Source/Runtime/Core/Public/Misc/CommandLine.h
--- a/Engine/Source/Runtime/Core/Public/Misc/CommandLine.h
+++ b/Engine/Source/Runtime/Core/Public/Misc/CommandLine.h
@@ -19,6 +19,11 @@ struct CORE_API FCommandLine
 	/** maximum size of the command line */
 	static const uint32 MaxCommandLineSize = 16384;
 
+	/** 
+	 * Resets FCommandLine to an uninitialised state as if Set() has never been called
+	 */
+	static void Reset();
+
 	/** 
 	 * Returns an edited version of the executable's command line with the game name and certain other parameters removed. 
 	 */
