Code that subsequently runs when an assert or ensure fails calls FCommandLine::Get() *a lot*. If the failed assert is before a command line has been set then the many Get() calls will themselves in turn throw asserts. It is impractical to chase these and guard against calling Get() and inappropriate in many instances. So instead we'll induce an empty command line before dispatching a failed assert or ensure.

#rnx

diff --git a/Engine/Source/Runtime/Core/Private/Misc/AssertionMacros.cpp b/Engine/Source/Runtime/Core/Private/Misc/AssertionMacros.cpp
--- a/Engine/Source/Runtime/Core/Private/Misc/AssertionMacros.cpp
+++ b/Engine/Source/Runtime/Core/Private/Misc/AssertionMacros.cpp
@@ -37,6 +37,33 @@ namespace
 		static FCriticalSection FailDebugCriticalSection;
 		return FailDebugCriticalSection;
 	}	
+
+	struct FTempCommandLineScope
+	{
+		// The code which is run when an assert or ensure fails (without a
+		// debugger attached) calls FCommandLine::Get() *a lot*. If the failed
+		// assert is before a command line has been set then the many Get()
+		// calls will in turn throw asserts.  It is impractical to chase these
+		// and guard against calling Get() and inappropriate in many instances.
+		FTempCommandLineScope()
+		{
+			if (!FCommandLine::IsInitialized())
+			{
+				FCommandLine::Set(TEXT(""));
+				bShouldReset = true;
+			}
+		}
+
+		~FTempCommandLineScope()
+		{
+			if (bShouldReset)
+			{
+				FCommandLine::Reset();
+			}
+		}
+
+		bool bShouldReset = false;
+	};
 }
 
 #define FILE_LINE_DESC TEXT(" [File:%s] [Line: %i] ")
@@ -88,6 +115,8 @@ void PrintScriptCallstack()
 
 static void AssertFailedImplV(const ANSICHAR* Expr, const ANSICHAR* File, int32 Line, const TCHAR* Format, va_list Args)
 {
+	FTempCommandLineScope TempCommandLine;
+
 	// This is not perfect because another thread might crash and be handled before this assert
 	// but this static variable will report the crash as an assert. Given complexity of a thread
 	// aware solution, this should be good enough. If crash reports are obviously wrong we can
@@ -274,6 +303,8 @@ void FDebug::LogAssertFailedMessageImplV(const ANSICHAR* Expr, const ANSICHAR* F
  */
 FORCENOINLINE void FDebug::EnsureFailed(const ANSICHAR* Expr, const ANSICHAR* File, int32 Line, const TCHAR* Msg, int NumStackFramesToIgnore)
 {
+	FTempCommandLineScope TempCommandLine;
+
 	// if time isn't ready yet, we better not continue
 	if (FPlatformTime::GetSecondsPerCycle() == 0.0)
 	{
