M5 - integrate fix from U5S

diff --git a/Engine/Source/Runtime/Core/Private/HAL/ThreadingBase.cpp b/Engine/Source/Runtime/Core/Private/HAL/ThreadingBase.cpp
--- a/Engine/Source/Runtime/Core/Private/HAL/ThreadingBase.cpp
+++ b/Engine/Source/Runtime/Core/Private/HAL/ThreadingBase.cpp
@@ -199,7 +199,7 @@ CORE_API bool IsInActualRenderingThread()
 {
 	bool newValue = FTaskTagScope::IsCurrentTag(ETaskTag::ERenderingThread);
 #if !UE_BUILD_SHIPPING && !UE_BUILD_TEST
-	bool oldValue = GRenderingThread && FPlatformTLS::GetCurrentThreadId() == GRenderingThread->GetThreadID();
+	bool oldValue = FPlatformTLS::GetCurrentThreadId() == GRenderThreadId;
 	ensureMsgf(oldValue == newValue, TEXT("oldValue(%i) newValue(%i) If this check fails make sure that there is a FTaskTagScope(ETaskTag::ERenderingThread) as deep as possible on the current callstack, you can see the current value in ActiveNamedThreads(%x)"), oldValue, newValue, FTaskTagScope::GetCurrentTag());
 	newValue = oldValue;
 #endif
