- Moved valid device software (flash) versions into UEBuildPlatformSDK subclasses to put alongside the SDK version, and made them min-max
- General fixes for some platforms
- DDPI now remembers the SDK version information from Turnkey
- Updated Turnkey Launch menu to have SDK information, including device software versions,
- Moved the initial turnkey query earlier before AutoSDKs are setup so envvars are not blown away for the child process
- Added concept of "Prepare For Debugging", but not enabled on any platforms yet (what it means is per platform, and is somewhere between cooking and packaging)
- VerifySdk Turnkey command now puts its output into a (x=y, z=w) format for easier parsing in C++
- NullCopyProvider now copies large remote files, and remote directories, locally before using

diff --git a/Engine/Source/Developer/TargetDeviceServices/Private/Services/TargetDeviceServiceManager.cpp b/Engine/Source/Developer/TargetDeviceServices/Private/Services/TargetDeviceServiceManager.cpp
--- a/Engine/Source/Developer/TargetDeviceServices/Private/Services/TargetDeviceServiceManager.cpp
+++ b/Engine/Source/Developer/TargetDeviceServices/Private/Services/TargetDeviceServiceManager.cpp
@@ -134,7 +134,7 @@ bool FTargetDeviceServiceManager::AddTargetDevice(ITargetDevicePtr InDevice)
 		return false;
 	}
 
-	const FString& DeviceName = InDevice->GetName();
+	FString DeviceName = InDevice->GetName();
 	ITargetDeviceServicePtr DeviceService = AddService(DeviceName);
 
 	if (DeviceService.IsValid())
diff --git a/Engine/Source/Developer/TargetPlatform/Private/TargetPlatformManagerModule.cpp b/Engine/Source/Developer/TargetPlatform/Private/TargetPlatformManagerModule.cpp
--- a/Engine/Source/Developer/TargetPlatform/Private/TargetPlatformManagerModule.cpp
+++ b/Engine/Source/Developer/TargetPlatform/Private/TargetPlatformManagerModule.cpp
@@ -88,6 +88,8 @@ public:
 		, bHasInitErrors(false)
 		, bIgnoreFirstDelegateCall(true)
 	{
+		FDataDrivenPlatformInfoRegistry::UpdateSdkStatus();
+
 #if AUTOSDKS_ENABLED		
 		
 		// AutoSDKs only enabled if UE_SDKS_ROOT is set.
@@ -998,7 +1000,7 @@ protected:
 
 	bool SetupSDKStatus(const FString& TargetPlatforms)
 	{
-		FDataDrivenPlatformInfoRegistry::UpdateSdkStatus();
+//		FDataDrivenPlatformInfoRegistry::UpdateSdkStatus();
 #if 0
 		DECLARE_SCOPE_CYCLE_COUNTER( TEXT( "FTargetPlatformManagerModule::SetupSDKStatus" ), STAT_FTargetPlatformManagerModule_SetupSDKStatus, STATGROUP_TargetPlatform );
 
diff --git a/Engine/Source/Editor/MainFrame/Private/Menus/MainMenu.cpp b/Engine/Source/Editor/MainFrame/Private/Menus/MainMenu.cpp
--- a/Engine/Source/Editor/MainFrame/Private/Menus/MainMenu.cpp
+++ b/Engine/Source/Editor/MainFrame/Private/Menus/MainMenu.cpp
@@ -492,8 +492,6 @@ void FMainMenu::RegisterFileProjectMenu()
 	);
 
 	Section.AddMenuEntry(FMainFrameCommands::Get().ZipUpProject);
-
-	Section.AddMenuEntry("TempNote", FText::FromString(TEXT("[See Launch for Packaging]")), FText(), FSlateIcon(), FUIAction(FExecuteAction(), FCanExecuteAction::CreateLambda([] { return false; })), EUserInterfaceActionType::None);
 }
 
 void FMainMenu::RegisterRecentFileAndExitMenuItems()
diff --git a/Engine/Source/Editor/TurnkeySupport/Private/TurnkeySupportModule.cpp b/Engine/Source/Editor/TurnkeySupport/Private/TurnkeySupportModule.cpp
--- a/Engine/Source/Editor/TurnkeySupport/Private/TurnkeySupportModule.cpp
+++ b/Engine/Source/Editor/TurnkeySupport/Private/TurnkeySupportModule.cpp
@@ -196,8 +196,18 @@ public:
 
 	static bool CanCookOrPackage(FName IniPlatformName, EPrepareContentMode Mode)
 	{
-		// prep for debugging not supported yet
-		return Mode != EPrepareContentMode::PrepareForDebugging && GetTargetPlatformManager()->FindTargetPlatform(IniPlatformName.ToString()) != nullptr;
+		if (GetTargetPlatformManager()->FindTargetPlatform(IniPlatformName.ToString()) == nullptr)
+		{
+			return false;
+		}
+
+		// PrepForDebugging needs the platform to specify how
+		if (Mode == EPrepareContentMode::PrepareForDebugging)
+		{
+			return FDataDrivenPlatformInfoRegistry::GetPlatformInfo(IniPlatformName).PrepareForDebuggingOptions != TEXT("");
+		}
+
+		return true;
 	}
 
 	static void CookOrPackage(FName IniPlatformName, EPrepareContentMode Mode)
@@ -332,11 +342,6 @@ public:
 				BuildCookRunParams += TEXT(" -pak");
 			}
 
-			if (PackagingSettings->bUseIoStore)
-			{
-				BuildCookRunParams += TEXT(" -iostore");
-			}
-
 			if (PackagingSettings->IncludePrerequisites)
 			{
 				BuildCookRunParams += TEXT(" -prereqs");
@@ -389,6 +394,16 @@ public:
 				BuildCookRunParams += FString::Printf(TEXT(" -clientconfig=%s"), LexToString(ConfigurationInfo.Configuration));
 			}
 		}
+		else if (Mode == EPrepareContentMode::PrepareForDebugging)
+		{
+			const ITargetPlatform* TargetPlatform = GetTargetPlatformManager()->FindTargetPlatform(PlatformInfo->Name);
+			if (ShouldBuildProject(PackagingSettings, TargetPlatform))
+			{
+				BuildCookRunParams += TEXT(" -build");
+			}
+
+			BuildCookRunParams += FString::Printf(TEXT(" %s"), *FDataDrivenPlatformInfoRegistry::GetPlatformInfo(IniPlatformName).PrepareForDebuggingOptions);
+		}
 		else if (Mode == EPrepareContentMode::CookOnly)
 		{
 			ContentPrepDescription = LOCTEXT("CookingContentTaskName", "Cooking content");
@@ -756,31 +771,53 @@ static void MakeTurnkeyPlatformMenu(FMenuBuilder& MenuBuilder, FName IniPlatform
 		for (const TSharedPtr<ITargetDeviceProxy> Proxy : DeviceProxies)
 		{
 			FString DeviceName = Proxy->GetName();
+			FString DeviceId = Proxy->GetTargetDeviceId(NAME_None);
 			MenuBuilder.AddSubMenu(
 				MakeSdkStatusAttribute(IniPlatformName, Proxy),
 				FText(),
-				FNewMenuDelegate::CreateLambda([UBTPlatformString, DeviceName](FMenuBuilder& SubMenuBuilder) {
-				SubMenuBuilder.AddMenuEntry(
-					LOCTEXT("Turnkey_RepairDevice", "Repair Device as Needed"),
-					LOCTEXT("TurnkeyTooltip_RepairDevice", "Perform any fixup that may be needed on this device. If up to date already, nothing will be done."),
-					FSlateIcon(),
-					FExecuteAction::CreateStatic(TurnkeyInstallSdk, UBTPlatformString, false, false, *DeviceName)
-				);
+				FNewMenuDelegate::CreateLambda([UBTPlatformString, IniPlatformName, DeviceName, DeviceId](FMenuBuilder& SubMenuBuilder)
+				{
+					const FDataDrivenPlatformInfo& DDPI = FDataDrivenPlatformInfoRegistry::GetPlatformInfo(IniPlatformName);
+					const FDDPISdkInfo& SdkInfo = DDPI.GetSdkInfoForDeviceId(DeviceId);
+					FFormatOrderedArguments Args = { FText::FromString(SdkInfo.InstalledVersion), FText::FromString(SdkInfo.MinAllowedVersion),FText::FromString(SdkInfo.MaxAllowedVersion) };
+					SubMenuBuilder.AddWidget(
+						SNew(STextBlock)
+						.ColorAndOpacity(FSlateColor::UseSubduedForeground())
+						.Text(FText::Format(LOCTEXT("SdkInfo", "Manual Installed SDK: {0}\nAllowedVersions: {1}-{2}"), Args)),
+						FText::GetEmpty()
+					);
 
-				SubMenuBuilder.AddMenuEntry(
-					LOCTEXT("Turnkey_ForceRepairDevice", "Force Repair Device"),
-					LOCTEXT("TurnkeyTooltip_ForceRepairDevice", "Force repairing anything on the device needed (update firmware, etc). Will perform all steps possible, even if not needed."),
-					FSlateIcon(),
-					FExecuteAction::CreateStatic(TurnkeyInstallSdk, UBTPlatformString, true, false, *DeviceName)
-				);
-			})
+					SubMenuBuilder.AddMenuEntry(
+						LOCTEXT("Turnkey_RepairDevice", "Repair Device as Needed"),
+						LOCTEXT("TurnkeyTooltip_RepairDevice", "Perform any fixup that may be needed on this device. If up to date already, nothing will be done."),
+						FSlateIcon(),
+						FExecuteAction::CreateStatic(TurnkeyInstallSdk, UBTPlatformString, false, false, *DeviceName)
+					);
+
+					SubMenuBuilder.AddMenuEntry(
+						LOCTEXT("Turnkey_ForceRepairDevice", "Force Repair Device"),
+						LOCTEXT("TurnkeyTooltip_ForceRepairDevice", "Force repairing anything on the device needed (update firmware, etc). Will perform all steps possible, even if not needed."),
+						FSlateIcon(),
+						FExecuteAction::CreateStatic(TurnkeyInstallSdk, UBTPlatformString, true, false, *DeviceName)
+					);
+				})
 			);
 		}
 
+
 		MenuBuilder.EndSection();
-		}
+	}
+
+	MenuBuilder.BeginSection("SdkManagement", LOCTEXT("TurnkeySection_Sdks", "Sdk Managment"));
 
-	MenuBuilder.BeginSection("AllDevices", LOCTEXT("TurnkeySection_Sdks", "Sdk Managment"));
+	const FDDPISdkInfo& SdkInfo = DDPI.GetSdkInfo(true);
+	FFormatOrderedArguments Args = { FText::FromString(SdkInfo.InstalledVersion), FText::FromString(SdkInfo.AutoSDKVersion),FText::FromString(SdkInfo.MinAllowedVersion),FText::FromString(SdkInfo.MaxAllowedVersion) };
+	MenuBuilder.AddWidget(
+		SNew(STextBlock)
+		.ColorAndOpacity(FSlateColor::UseSubduedForeground())
+		.Text(FText::Format(LOCTEXT("SdkInfo", "Manual Installed SDK: {0}\nAutoSDK: {1}\nAllowedVersions: {2}-{3}"), Args)),
+		FText::GetEmpty()
+	);
 
 	const TCHAR* NoDevice = nullptr;
 	if (DDPI.GetSdkStatus() == DDPIPlatformSdkStatus::OutOfDate)
@@ -881,7 +918,7 @@ bool CanLaunchOnDevice(const FString& DeviceName)
 }
 
 
-void LaunchOnDevice(const FString& DeviceId, const FString& DeviceName, bool bUseTurnkey)
+static void LaunchOnDevice(const FString& DeviceId, const FString& DeviceName, bool bUseTurnkey)
 {
 	FTargetDeviceId TargetDeviceId;
 	if (FTargetDeviceId::Parse(DeviceId, TargetDeviceId))
@@ -969,23 +1006,18 @@ static void PrepareLaunchOn(FString DeviceId, FString DeviceName)
 	PlaySettings->SaveConfig();
 }
 
-void HandleLaunchOnDeviceActionExecute(FString DeviceId, FString DeviceName, bool bUseTurnkey)
+static void HandleLaunchOnDeviceActionExecute(FString DeviceId, FString DeviceName, bool bUseTurnkey)
 {
 	PrepareLaunchOn(DeviceId, DeviceName);
 	LaunchOnDevice(DeviceId, DeviceName, bUseTurnkey);
 }
 
 
-bool HandleLaunchOnDeviceActionCanExecute(FString DeviceName)
+static bool HandleLaunchOnDeviceActionCanExecute(FString DeviceName)
 {
 	return CanLaunchOnDevice(DeviceName);
 }
 
-bool HandleLaunchOnDeviceActionIsChecked(FString DeviceName)
-{
-	return (DeviceName == GetDefault<ULevelEditorPlaySettings>()->LastExecutedLaunchName);
-}
-
 static void GenerateDeviceProxyMenuParams(TSharedPtr<ITargetDeviceProxy> DeviceProxy, FName PlatformName, FUIAction& OutAction, FText& OutTooltip)
 {
 	// 	// create an All_<platform>_devices_on_<host> submenu
@@ -1004,10 +1036,10 @@ static void GenerateDeviceProxyMenuParams(TSharedPtr<ITargetDeviceProxy> DeviceP
 	// 	}
 
 		// ... create an action...
+	FString DeviceId = DeviceProxy->GetTargetDeviceId(NAME_None);
 	OutAction = FUIAction(
-		FExecuteAction::CreateStatic(&HandleLaunchOnDeviceActionExecute, DeviceProxy->GetTargetDeviceId(NAME_None), DeviceProxy->GetName(), true)
+		FExecuteAction::CreateStatic(&HandleLaunchOnDeviceActionExecute, DeviceId, DeviceProxy->GetName(), true)
 		//		, FCanExecuteAction::CreateStatic(&HandleLaunchOnDeviceActionCanExecute, DeviceProxy->GetName())
-		//		, FIsActionChecked::CreateStatic(&HandleLaunchOnDeviceActionIsChecked, DeviceProxy->GetName())
 	);
 
 	// ... generate tooltip text
@@ -1455,6 +1487,18 @@ if (!FModuleManager::LoadModuleChecked<IProjectTargetPlatformEditorModule>("Proj
 }
 
 
+
+void FInternalPlayWorldCommandCallbacks::HandleShowSDKTutorial(FString PlatformName, FString NotInstalledDocLink)
+{
+	// broadcast this, and assume someone will pick it up
+	IMainFrameModule& MainFrameModule = FModuleManager::GetModuleChecked<IMainFrameModule>(TEXT("MainFrame"));
+	MainFrameModule.BroadcastMainFrameSDKNotInstalled(PlatformName, NotInstalledDocLink);
+}
+
+
+
+
+
 #endif
 
 
diff --git a/Engine/Source/Editor/UnrealEd/Private/Kismet2/DebuggerCommands.cpp b/Engine/Source/Editor/UnrealEd/Private/Kismet2/DebuggerCommands.cpp
--- a/Engine/Source/Editor/UnrealEd/Private/Kismet2/DebuggerCommands.cpp
+++ b/Engine/Source/Editor/UnrealEd/Private/Kismet2/DebuggerCommands.cpp
@@ -154,15 +154,6 @@ public:
 
 	static void PlayInSettings_Clicked();
 
-	// Launch On
-	static void HandleLaunchOnDeviceActionExecute(FString DevicedId, FString DeviceName, bool bUseTurnkey);
-	static bool HandleLaunchOnDeviceActionCanExecute(FString DeviceName);
-	static bool HandleLaunchOnDeviceActionIsChecked(FString DeviceName);
-
-	// No Device
-	static void HandleNoDeviceFoundActionExecute() {}
-	static bool HandleNoDeviceFoundActionCanExecute() { return false; }
-
 	static void HandleShowSDKTutorial(FString PlatformName, FString NotInstalledDocLink);
 
 	static void RepeatLastLaunch_Clicked();
@@ -1671,45 +1662,6 @@ void FInternalPlayWorldCommandCallbacks::PlayInSettings_Clicked()
 
 
 
-static void PrepareLaunchOn(FString DeviceId, FString DeviceName)
-{
-	ULevelEditorPlaySettings* PlaySettings = GetMutableDefault<ULevelEditorPlaySettings>();
-
-	PlaySettings->LastExecutedLaunchModeType = LaunchMode_OnDevice;
-	PlaySettings->LastExecutedLaunchDevice = DeviceId;
-	PlaySettings->LastExecutedLaunchName = DeviceName;
-
-	PlaySettings->PostEditChange();
-
-	PlaySettings->SaveConfig();
-}
-
-void FInternalPlayWorldCommandCallbacks::HandleLaunchOnDeviceActionExecute(FString DeviceId, FString DeviceName, bool bUseTurnkey)
-{
-	PrepareLaunchOn(DeviceId, DeviceName);
-	LaunchOnDevice(DeviceId, DeviceName, bUseTurnkey);
-}
-
-
-bool FInternalPlayWorldCommandCallbacks::HandleLaunchOnDeviceActionCanExecute(FString DeviceName)
-{
-	return CanLaunchOnDevice(DeviceName);
-}
-
-
-bool FInternalPlayWorldCommandCallbacks::HandleLaunchOnDeviceActionIsChecked(FString DeviceName)
-{
-	return (DeviceName == GetDefault<ULevelEditorPlaySettings>()->LastExecutedLaunchName);
-}
-
-
-void FInternalPlayWorldCommandCallbacks::HandleShowSDKTutorial(FString PlatformName, FString NotInstalledDocLink)
-{
-	// broadcast this, and assume someone will pick it up
-	IMainFrameModule& MainFrameModule = FModuleManager::GetModuleChecked<IMainFrameModule>(TEXT("MainFrame"));
-	MainFrameModule.BroadcastMainFrameSDKNotInstalled(PlatformName, NotInstalledDocLink);
-}
-
 void FInternalPlayWorldCommandCallbacks::GetMouseControlExecute()
 {
 	if (IsInPIE()) {
@@ -1956,6 +1908,35 @@ void FInternalPlayWorldCommandCallbacks::AddMessageLog(const FText& Text, const
 }
 
 
+
+EPlayModeLocations FInternalPlayWorldCommandCallbacks::GetPlayModeLocation()
+{
+	// We can't use PlayLocation_DefaultPlayerStart without a player start position
+	return GEditor->CheckForPlayerStart()
+		? static_cast<EPlayModeLocations>(GetDefault<ULevelEditorPlaySettings>()->LastExecutedPlayModeLocation)
+		: PlayLocation_CurrentCameraLocation;
+}
+
+
+
+
+
+
+
+
+
+
+
+
+
+
+
+
+
+/////// Will move these into TurnkeySupportModule soon:
+
+
+
 bool FInternalPlayWorldCommandCallbacks::CanLaunchOnDevice(const FString& DeviceName)
 {
 	if (!GUnrealEd->IsPlayingViaLauncher())
@@ -2044,29 +2025,29 @@ void FInternalPlayWorldCommandCallbacks::LaunchOnDevice(const FString& DeviceId,
 
 				IUATHelperModule::Get().CreateUatTask(CommandLine, FText::FromString(IniPlatformName), TaskName, TaskName, FEditorStyle::GetBrush(TEXT("MainFrame.PackageProject")),
 					[SessionParams](FString Result, double)
+				{
+					// unfortunate string comparison for success
+					bool bWasSuccessful = Result == TEXT("Completed");
+					AsyncTask(ENamedThreads::GameThread, [SessionParams, bWasSuccessful]()
 					{
-						// unfortunate string comparison for success
-						bool bWasSuccessful = Result == TEXT("Completed");
-						AsyncTask(ENamedThreads::GameThread, [SessionParams, bWasSuccessful]()
+						if (bWasSuccessful)
+						{
+							GUnrealEd->RequestPlaySession(SessionParams);
+						}
+						else
 						{
- 							if (bWasSuccessful)
- 							{
- 								GUnrealEd->RequestPlaySession(SessionParams);
- 							}
- 							else
+							TSharedRef<SWindow> Win = OpenMsgDlgInt_NonModal(EAppMsgType::YesNo, LOCTEXT("SDKCheckFailed", "SDK Verification failed. Would you like to attempt the Launch On anyway?"), LOCTEXT("SDKCheckFailedTitle", "SDK Verification"),
+								FOnMsgDlgResult::CreateLambda([SessionParams](const TSharedRef<SWindow>&, EAppReturnType::Type Choice)
 							{
-								TSharedRef<SWindow> Win = OpenMsgDlgInt_NonModal(EAppMsgType::YesNo, LOCTEXT("SDKCheckFailed", "SDK Verification failed. Would you like to attempt the Launch On anyway?"), LOCTEXT("SDKCheckFailedTitle", "SDK Verification"),
-									FOnMsgDlgResult::CreateLambda([SessionParams](const TSharedRef<SWindow>&, EAppReturnType::Type Choice)
-									{
-										if (Choice == EAppReturnType::Yes)
-										{
-											GUnrealEd->RequestPlaySession(SessionParams);
-										}
-									}));
-								Win->ShowWindow();
-							}
-						});
+								if (Choice == EAppReturnType::Yes)
+								{
+									GUnrealEd->RequestPlaySession(SessionParams);
+								}
+							}));
+							Win->ShowWindow();
+						}
 					});
+				});
 			}
 			else
 			{
@@ -2077,33 +2058,6 @@ void FInternalPlayWorldCommandCallbacks::LaunchOnDevice(const FString& DeviceId,
 }
 
 
-EPlayModeLocations FInternalPlayWorldCommandCallbacks::GetPlayModeLocation()
-{
-	// We can't use PlayLocation_DefaultPlayerStart without a player start position
-	return GEditor->CheckForPlayerStart()
-		? static_cast<EPlayModeLocations>(GetDefault<ULevelEditorPlaySettings>()->LastExecutedPlayModeLocation)
-		: PlayLocation_CurrentCameraLocation;
-}
-
-
-
-
-
-
-
-
-
-
-
-
-
-
-
-
-
-/////// Will move these into TurnkeySupportModule soon:
-
-
 void FInternalPlayWorldCommandCallbacks::RepeatLastLaunch_Clicked()
 {
 	const ULevelEditorPlaySettings* PlaySettings = GetDefault<ULevelEditorPlaySettings>();
diff --git a/Engine/Source/Programs/AutomationTool/AutomationUtils/Platform.cs b/Engine/Source/Programs/AutomationTool/AutomationUtils/Platform.cs
--- a/Engine/Source/Programs/AutomationTool/AutomationUtils/Platform.cs
+++ b/Engine/Source/Programs/AutomationTool/AutomationUtils/Platform.cs
@@ -174,11 +174,6 @@ namespace AutomationTool
 			return null;
 		}
 
-		public virtual string GetAllowedSoftwareVersions()
-		{
-			return "";
-		}
-
 		public virtual bool InstallSDK(BuildCommand BuildCommand, FileRetriever Retriever, object HintObject, DeviceInfo Device = null)
 		{
 			string DownloadedSDK = Retriever.RetrieveFileSource(HintObject);
diff --git a/Engine/Source/Programs/AutomationTool/IOS/IOSPlatform.Automation.cs b/Engine/Source/Programs/AutomationTool/IOS/IOSPlatform.Automation.cs
--- a/Engine/Source/Programs/AutomationTool/IOS/IOSPlatform.Automation.cs
+++ b/Engine/Source/Programs/AutomationTool/IOS/IOSPlatform.Automation.cs
@@ -169,11 +169,6 @@ public class IOSPlatform : Platform
 		SDKName = (TargetPlatform == UnrealTargetPlatform.TVOS) ? "appletvos" : "iphoneos";
 	}
 
-	public override string GetAllowedSoftwareVersions()
-	{
-		return @"regex:^13\.5.*";
-	}
-
 	public override bool GetSDKInstallCommand(out string Command, out string Params, ref string Preamble, ref string SuccessPostamble, ref string FailurePostamble, FileRetriever Retriever)
 	{
 		Preamble = "Moving your original Xcode application from /Applications to the Trash, and unzipping the new version into /Applications";
diff --git a/Engine/Source/Programs/AutomationTool/Turnkey/Commands/InstallSdk.cs b/Engine/Source/Programs/AutomationTool/Turnkey/Commands/InstallSdk.cs
--- a/Engine/Source/Programs/AutomationTool/Turnkey/Commands/InstallSdk.cs
+++ b/Engine/Source/Programs/AutomationTool/Turnkey/Commands/InstallSdk.cs
@@ -103,6 +103,27 @@ namespace Turnkey.Commands
 // 					Sdks = Sdks.FindAll(x => x.IsNeeded(Platform, DeviceName));
 // 				}
 
+				Sdks.Sort((a,b) => 
+				{
+					int Val = a.Type.CompareTo(b.Type);
+					if (Val != 0)
+					{
+						return Val;
+					}
+
+					UInt64 VersionA, VersionB;
+					if (SDK.TryConvertVersionToInt(a.Version, out VersionA) && SDK.TryConvertVersionToInt(b.Version, out VersionB))
+					{
+						Val = VersionA.CompareTo(VersionB);
+						if (Val != 0)
+						{
+							return Val;
+						}
+						return a.AllowedFlashDeviceTypes.CompareTo(b.AllowedFlashDeviceTypes);
+					}
+
+					return 0;
+				});
 
 				if (bBestAvailable && Sdks.Count > 0)
 				{
diff --git a/Engine/Source/Programs/AutomationTool/Turnkey/Commands/ListPlatforms.cs b/Engine/Source/Programs/AutomationTool/Turnkey/Commands/ListPlatforms.cs
--- a/Engine/Source/Programs/AutomationTool/Turnkey/Commands/ListPlatforms.cs
+++ b/Engine/Source/Programs/AutomationTool/Turnkey/Commands/ListPlatforms.cs
@@ -34,10 +34,11 @@ namespace Turnkey.Commands
 
 				string ManualSDKVersion, AutoSDKVersion;
 				string MinAllowedVersion, MaxAllowedVersion;
+				string MinAllowedSoftwareVersion, MaxAllowedSoftwareVersion;
 				SDK.GetInstalledVersions(out ManualSDKVersion, out AutoSDKVersion);
 				SDK.GetValidVersionRange(out MinAllowedVersion, out MaxAllowedVersion);
+				SDK.GetValidSoftwareVersionRange(out MinAllowedSoftwareVersion, out MaxAllowedSoftwareVersion);
 				
-				string AllowedSoftware = Platform.GetAllowedSoftwareVersions();
 				bool bIsAutoSdkValid = SDK.IsVersionValid(AutoSDKVersion, bForAutoSDK:true);
 				bool bIsManualSdkValid = SDK.IsVersionValid(ManualSDKVersion, bForAutoSDK:false);
 				TurnkeyUtils.Log("  Platform: {0}", TargetPlatform.ToString());
@@ -63,7 +64,7 @@ namespace Turnkey.Commands
 					}
 				}
 
-				TurnkeyUtils.Log("  Allowed Device Software Version(s): {0}", AllowedSoftware);
+				TurnkeyUtils.Log("  Allowed Device Software Range: {0}-{1}", MinAllowedSoftwareVersion, MaxAllowedSoftwareVersion);
 				TurnkeyUtils.Log("  Devices: ");
 
 				DeviceInfo[] Devices = Platform.GetDevices();
@@ -75,7 +76,7 @@ namespace Turnkey.Commands
 				{
 					foreach (DeviceInfo Device in Devices)
 					{
-						bool bIsSoftwareValid = TurnkeyUtils.IsValueValid(Device.SoftwareVersion, AllowedSoftware, Platform);
+						bool bIsSoftwareValid = SDK.IsSoftwareVersionValid(Device.SoftwareVersion);
 
 						TurnkeyUtils.Log("    Name: {0}{1}", Device.Name, Device.bIsDefault ? "*" : "");
 						TurnkeyUtils.Log("      Id: {0}", Device.Id);
diff --git a/Engine/Source/Programs/AutomationTool/Turnkey/Commands/VerifySdk.cs b/Engine/Source/Programs/AutomationTool/Turnkey/Commands/VerifySdk.cs
--- a/Engine/Source/Programs/AutomationTool/Turnkey/Commands/VerifySdk.cs
+++ b/Engine/Source/Programs/AutomationTool/Turnkey/Commands/VerifySdk.cs
@@ -108,23 +108,40 @@ namespace Turnkey.Commands
 
 				SdkUtils.LocalAvailability LocalState = SdkUtils.GetLocalAvailability(AutomationPlatform, bUpdateIfNeeded);
 
+				string ManualSDKVersion, AutoSDKVersion;
+				string MinAllowedVersion, MaxAllowedVersion;
+				string MinSoftwareAllowedVersion, MaxSoftwareAllowedVersion;
+				PlatformSDK.GetInstalledVersions(out ManualSDKVersion, out AutoSDKVersion);
+				PlatformSDK.GetValidVersionRange(out MinAllowedVersion, out MaxAllowedVersion);
+				PlatformSDK.GetValidSoftwareVersionRange(out MinSoftwareAllowedVersion, out MaxSoftwareAllowedVersion);
+
+//				if (ManualSDKVersion == "1")
+				{
+					TurnkeyUtils.Report("How did i get {0} {1}", ManualSDKVersion, AutoSDKVersion);
+				}
+
+				SdkUtils.LocalAvailability ReportedState = LocalState;
+				string StatusString;
+
 				if ((LocalState & SdkUtils.LocalAvailability.Platform_ValidHostPrerequisites) == 0)
 				{
-					TurnkeyUtils.Report("{0}: Invalid: [Host Prerequisites are not valid]", Platform);
+					StatusString = "Invalid";
 				}
 				else if ((LocalState & (SdkUtils.LocalAvailability.AutoSdk_ValidVersionExists | SdkUtils.LocalAvailability.InstalledSdk_ValidVersionExists)) == 0)
 				{
-					string MinAllowedVersion, MaxAllowedVersion;
-					PlatformSDK.GetValidVersionRange(out MinAllowedVersion, out MaxAllowedVersion);
-					TurnkeyUtils.Report("{0}: Invalid: [No AutoSdk or Installed Sdk in range {1}-{2} - {3}]", Platform, MinAllowedVersion, MaxAllowedVersion, LocalState.ToString());
+					StatusString = "Invalid";
 					TurnkeyUtils.ExitCode = AutomationTool.ExitCode.Error_SDKNotFound;
 				}
 				else
 				{
-					TurnkeyUtils.Report("{0}: Valid: [{1}]", Platform, (LocalState & (SdkUtils.LocalAvailability.AutoSdk_ValidVersionExists | SdkUtils.LocalAvailability.InstalledSdk_ValidVersionExists)).ToString());
+					StatusString = "Valid";
+					ReportedState &= (SdkUtils.LocalAvailability.AutoSdk_ValidVersionExists | SdkUtils.LocalAvailability.InstalledSdk_ValidVersionExists);
 					//					TurnkeyUtils.Log("{0}: Valid [Installed: '{1}', Required: '{2}']", Platform, PlatformObject.GetInstalledSdk(), PlatformObject.GetAllowedSdks());
 				}
 
+				TurnkeyUtils.Report("{0}: (Status={1}, Installed={2}, AutoSDK={3}, MinAllowed={4}, MaxAllowed={5}, Flags=\"{6}\")", Platform, StatusString, ManualSDKVersion, AutoSDKVersion,
+					MinAllowedVersion, MaxAllowedVersion, ReportedState.ToString());
+
 				// install if out of date, or if forcing it
 				if (bForceSdkInstall || (bUpdateIfNeeded && TurnkeyUtils.ExitCode != AutomationTool.ExitCode.Success))
 				{
@@ -168,7 +185,7 @@ namespace Turnkey.Commands
 
 
 					// update LocalState
-					LocalState = SdkUtils.GetLocalAvailability(AutomationPlatform, false);
+//					LocalState = SdkUtils.GetLocalAvailability(AutomationPlatform, false);
 
 					// @todo turnkey: validate!
 				}
@@ -187,7 +204,7 @@ namespace Turnkey.Commands
 					TurnkeyUtils.Log("Installed Device validity:");
 
 					// a single device named all means all devices
-					if (!(DeviceNames.Count == 1 && DeviceNames[0].CompareTo("All") == 0))
+					if (!(DeviceNames.Count == 1 && DeviceNames[0].ToLower() == "all"))
 					{
 						Devices = Devices.Where(x => DeviceNames.Contains(x.Name, StringComparer.OrdinalIgnoreCase)).ToArray();
 					}
@@ -196,26 +213,35 @@ namespace Turnkey.Commands
 					foreach (DeviceInfo Device in Devices)
 					{
 						bool bArePrerequisitesValid = AutomationPlatform.UpdateDevicePrerequisites(Device, TurnkeyUtils.CommandUtilHelper, Retriever, !bUpdateIfNeeded);
+						bool bIsSoftwareValid = PlatformSDK.IsSoftwareVersionValid(Device.SoftwareVersion);
 
+						SdkUtils.LocalAvailability DeviceState = SdkUtils.LocalAvailability.None;
 						if (!bArePrerequisitesValid)
 						{
-							TurnkeyUtils.Report("{0}@{1}: Invalid: [Device Prerequisites are not valid]", Platform, Device.Name, Device.SoftwareVersion, AutomationPlatform.GetAllowedSoftwareVersions());
-							continue;
+							StatusString = "Invalid";
+							DeviceState |= SdkUtils.LocalAvailability.Device_InvalidPrerequisites;
 						}
-
-						bool bIsSoftwareValid = TurnkeyUtils.IsValueValid(Device.SoftwareVersion, AutomationPlatform.GetAllowedSoftwareVersions(), AutomationPlatform);
-
-
-						if (!bForceDeviceInstall && bIsSoftwareValid)
+						else if (bIsSoftwareValid)
 						{
-							TurnkeyUtils.Report("{0}@{1}: Valid: [{2}]", Platform, Device.Name, Device.SoftwareVersion);
+							StatusString = "Valid";
+							DeviceState |= SdkUtils.LocalAvailability.Device_InstallSoftwareValid;
 						}
 						else
 						{
-							if (!bForceDeviceInstall)
-							{
-								TurnkeyUtils.Report("{0}@{1}: Invalid: [Has {2}, needs {3}]", Platform, Device.Name, Device.SoftwareVersion, AutomationPlatform.GetAllowedSoftwareVersions());
-							}
+							StatusString = "Invalid";
+							DeviceState |= SdkUtils.LocalAvailability.Device_InstallSoftwareInvalid;
+						}
+
+						if (Device.bCanConnect== false)
+						{
+							DeviceState |= SdkUtils.LocalAvailability.Device_CannotConnect;
+						}
+
+						TurnkeyUtils.Report("{0}@{1}: (Status={2}, Installed={3}, MinAllowed={4}, MaxAllowed={5}, Flags=\"{6}\")", Platform, Device.Name, StatusString, Device.SoftwareVersion, 
+							MinSoftwareAllowedVersion, MaxSoftwareAllowedVersion, DeviceState.ToString());
+
+						if (bForceDeviceInstall || !bIsSoftwareValid)
+						{
 							TurnkeyUtils.ExitCode = AutomationTool.ExitCode.Error_SDKNotFound;
 
 							if (bUpdateIfNeeded)
diff --git a/Engine/Source/Programs/AutomationTool/Turnkey/FileSource.cs b/Engine/Source/Programs/AutomationTool/Turnkey/FileSource.cs
--- a/Engine/Source/Programs/AutomationTool/Turnkey/FileSource.cs
+++ b/Engine/Source/Programs/AutomationTool/Turnkey/FileSource.cs
@@ -212,12 +212,12 @@ namespace Turnkey
 // 
 			if (Type == SourceType.Flash)
 			{
-				AutomationTool.Platform AutomationPlatform = AutomationTool.Platform.GetPlatform(Platform);
-				bool bIsValid = TurnkeyUtils.IsValueValid(Version, AutomationPlatform.GetAllowedSoftwareVersions(), AutomationPlatform);
+				UEBuildPlatformSDK SDK = UEBuildPlatformSDK.GetSDKForPlatform(Platform.ToString());
+				bool bIsValid = SDK.IsSoftwareVersionValid(Version);
 				// if we were passed a device, also check if this Sdk is valid for that device
-				//				if (DeviceName != null)
+				if (Device != null)
 				{
-					bIsValid = bIsValid && Device != null && TurnkeyUtils.IsValueValid(Device.Type, AllowedFlashDeviceTypes, AutomationPlatform);
+					bIsValid &= TurnkeyUtils.IsValueValid(Device.Type, AllowedFlashDeviceTypes, AutomationTool.Platform.GetPlatform(Platform));
 				}
 				return bIsValid;
 			}
@@ -251,7 +251,7 @@ namespace Turnkey
 				//				Sdks = Sdks.FindAll(x => x.Version == null || (Type == SourceType.Flash ? TurnkeyUtils.IsValueValid(x.Version, Platform.GetAllowedSoftwareVersions(), Platform) : SDK.IsVersionValid(x.Version, bForAutoSDK: x.Type == SourceType.AutoSdk)));
 				Sdks = Sdks.FindAll(x => x.Version == null ||
 					(Type == SourceType.Flash ?
-						TurnkeyUtils.IsValueValid(x.Version, Platform.GetAllowedSoftwareVersions(), Platform) :
+						SDK.IsSoftwareVersionValid(x.Version) :
 						SDK.IsVersionValid(x.Version, bForAutoSDK: (bSelectBest && x.Type == SourceType.AutoSdk))
 					));
 
diff --git a/Engine/Source/Programs/AutomationTool/Turnkey/Providers/NullCopyProvider.cs b/Engine/Source/Programs/AutomationTool/Turnkey/Providers/NullCopyProvider.cs
--- a/Engine/Source/Programs/AutomationTool/Turnkey/Providers/NullCopyProvider.cs
+++ b/Engine/Source/Programs/AutomationTool/Turnkey/Providers/NullCopyProvider.cs
@@ -4,6 +4,7 @@ using System.IO;
 using System.Collections.Generic;
 using System.Linq;
 using System.Text.RegularExpressions;
+using Tools.DotNETCommon;
 
 namespace Turnkey
 {
@@ -56,7 +57,41 @@ namespace Turnkey
 				TurnkeyUtils.Log("Reqeusted local path {0} does not exist!", OutputPath);
 				return null;
 			}
-			
+
+			// accessing large remote files direcly can be crazy slow, so we cache large-ish files locally before using
+			bool bIsRemotePath = (new System.Uri(OutputPath).IsUnc) || (new DriveInfo(OutputPath).DriveType == DriveType.Network);
+			bool bIsLargeFile = File.Exists(OutputPath) && new FileInfo(OutputPath).Length > 5 * 1024 * 1024;
+			bool bIsDirectory = Directory.Exists(OutputPath);
+
+			if (bIsRemotePath && (bIsLargeFile || bIsDirectory))
+			{
+				string CopyLocation = Path.Combine(LocalCache.CreateTempDirectory(), Path.GetFileName(OutputPath));
+
+				if (bIsDirectory)
+				{
+					TurnkeyUtils.Log("Copying remote directory structure to local temp location {0}...", CopyLocation);
+					if (AutomationTool.CommandUtils.CopyDirectory_NoExceptions(OutputPath, CopyLocation) == false)
+					{
+						TurnkeyUtils.Log("Copy failed, unable to continue...");
+						return null;
+					}
+
+				}
+				else
+				{
+					TurnkeyUtils.Log("Copying large remote file to local temp location {0}...", CopyLocation);
+					if (AutomationTool.CommandUtils.CopyFile_NoExceptions(OutputPath, CopyLocation) == false)
+					{
+						TurnkeyUtils.Log("Copy failed, unable to continue...");
+						return null;
+					}
+				}
+
+				TurnkeyUtils.Log("Done!");
+
+				return CopyLocation;
+			}
+
 			return OutputPath;
 		}
 
diff --git a/Engine/Source/Programs/AutomationTool/Turnkey/SdkInfo.cs b/Engine/Source/Programs/AutomationTool/Turnkey/SdkInfo.cs
--- a/Engine/Source/Programs/AutomationTool/Turnkey/SdkInfo.cs
+++ b/Engine/Source/Programs/AutomationTool/Turnkey/SdkInfo.cs
@@ -26,6 +26,11 @@ namespace Turnkey
 			InstalledSdk_InvalidVersionExists = 32,
 			Platform_ValidHostPrerequisites = 64,
 			Platform_InvalidHostPrerequisites = 128,
+
+			Device_InvalidPrerequisites = 256,
+			Device_InstallSoftwareValid,
+			Device_InstallSoftwareInvalid,
+			Device_CannotConnect,
 		}
 
 		static public LocalAvailability GetLocalAvailability(AutomationTool.Platform AutomationPlatform, bool bAllowUpdatingPrerequisites)
@@ -46,42 +51,47 @@ namespace Turnkey
 
 			UEBuildPlatformSDK SDK = UEBuildPlatformSDK.GetSDKForPlatform(AutomationPlatform.PlatformType.ToString());
 
-			// if we have the variable at all, 
-			string AutoSdkVar = Environment.GetEnvironmentVariable("UE_SDKS_ROOT");
-			if (AutoSdkVar != null)
-			{
-				Result |= LocalAvailability.AutoSdk_VariableExists;
+			string ManualSDKVersion, AutoSDKVersion;
+			SDK.GetInstalledVersions(out ManualSDKVersion, out AutoSDKVersion);
 
-				// get platform subdirectory
-				string AutoSubdir = string.Format("Host{0}/{1}", HostPlatform.Current.HostEditorPlatform.ToString(), SDK.GetAutoSDKPlatformName());
-				DirectoryInfo PlatformDir = new DirectoryInfo(Path.Combine(AutoSdkVar, AutoSubdir));
-				if (PlatformDir.Exists)
+			if (AutoSDKVersion == null)
+			{
+				// look to see if other versions are around
+				string AutoSdkVar = Environment.GetEnvironmentVariable("UE_SDKS_ROOT");
+				if (AutoSdkVar != null)
 				{
-					bool bValidVersionFound = false;
-					foreach (DirectoryInfo Dir in PlatformDir.EnumerateDirectories())
+					// no matter what, remember we have the variable set
+					Result |= LocalAvailability.AutoSdk_VariableExists;
+
+					// get platform subdirectory
+					string AutoSubdir = string.Format("Host{0}/{1}", HostPlatform.Current.HostEditorPlatform.ToString(), SDK.GetAutoSDKPlatformName());
+					DirectoryInfo PlatformDir = new DirectoryInfo(Path.Combine(AutoSdkVar, AutoSubdir));
+					if (PlatformDir.Exists)
 					{
-						// check for valid version number (auto SDK must match version exactly)
-						if (string.Compare(Dir.Name, SDK.GetAutoSDKDirectoryForMasterVersion()) == 0)
+						foreach (DirectoryInfo Dir in PlatformDir.EnumerateDirectories())
 						{
-							// make sure it actually has buits in it
+							// look to see if other versions have been synced, but are bad version (otherwise, we would have had AutoSDKVersion above!)
 							if (File.Exists(Path.Combine(Dir.FullName, "setup.bat")) || File.Exists(Path.Combine(Dir.FullName, "setup.sh")))
 							{
-								bValidVersionFound = true;
+								Result |= LocalAvailability.AutoSdk_InvalidVersionExists;
 								break;
 							}
 						}
 					}
-
-					// if we had a platform directory, but no valid versions, note that
-					Result |= bValidVersionFound ? LocalAvailability.AutoSdk_ValidVersionExists : LocalAvailability.AutoSdk_InvalidVersionExists;
 				}
+
+			}
+			else
+			{
+				Result |= LocalAvailability.AutoSdk_ValidVersionExists | LocalAvailability.AutoSdk_VariableExists;
 			}
 
+			// if we have the variable at all, 
+
 			// if anything is installed, this will return a value
-			string InstalledVersion = SDK.GetInstalledVersion();
-			if (!string.IsNullOrEmpty(InstalledVersion))
+			if (!string.IsNullOrEmpty(ManualSDKVersion))
 			{
-				if (SDK.IsVersionValid(InstalledVersion, bForAutoSDK: false))
+				if (SDK.IsVersionValid(ManualSDKVersion, bForAutoSDK: false))
 				{
 					Result |= LocalAvailability.InstalledSdk_ValidVersionExists;
 				}
diff --git a/Engine/Source/Programs/AutomationTool/Turnkey/TurnkeyUtils.cs b/Engine/Source/Programs/AutomationTool/Turnkey/TurnkeyUtils.cs
--- a/Engine/Source/Programs/AutomationTool/Turnkey/TurnkeyUtils.cs
+++ b/Engine/Source/Programs/AutomationTool/Turnkey/TurnkeyUtils.cs
@@ -34,8 +34,8 @@ namespace Turnkey
 			IEnumerable<string> AutoSDKPlatforms = UEBuildPlatformSDK.AllSDKs.Select(x => x.GetAutoSDKPlatformName()).Distinct();
 			SetVariable("AutoSDKPlatforms", string.Join(",", AutoSDKPlatforms));
 
-			TurnkeyUtils.Log("AllPlatforms = {0}", GetVariableValue("AllPlatforms"));
-			TurnkeyUtils.Log("AutoSDKPlatforms = {0}", GetVariableValue("AutoSDKPlatforms"));
+// 			TurnkeyUtils.Log("AllPlatforms = {0}", GetVariableValue("AllPlatforms"));
+// 			TurnkeyUtils.Log("AutoSDKPlatforms = {0}", GetVariableValue("AutoSDKPlatforms"));
 
 			SetVariable("HOST_PLATFORM_NAME", HostPlatform.Current.HostEditorPlatform.ToString());
 		}
@@ -140,9 +140,8 @@ namespace Turnkey
 					Platforms.Add(PossiblePlatforms[PlatformChoice - 1]);
 				}
 			}
-			else if (PlatformString == "All")
+			else if (PlatformString.ToLower() == "all")
 			{
-				// HACK UNTIL WIN32 IS GONE
 				Platforms = PossiblePlatforms;
 			}
 			else
@@ -166,7 +165,7 @@ namespace Turnkey
 				}
 			}
 
-			return Platforms;
+			return Platforms.OrderBy(x => x.ToString()).ToList();
 		}
 
 		public static DeviceInfo GetDeviceFromCommandLineOrUser(string[] CommandOptions, UnrealTargetPlatform Platform)
diff --git a/Engine/Source/Programs/DotNETCommon/BuildUtilities/UEBuildPlatformSDK.cs b/Engine/Source/Programs/DotNETCommon/BuildUtilities/UEBuildPlatformSDK.cs
--- a/Engine/Source/Programs/DotNETCommon/BuildUtilities/UEBuildPlatformSDK.cs
+++ b/Engine/Source/Programs/DotNETCommon/BuildUtilities/UEBuildPlatformSDK.cs
@@ -120,9 +120,8 @@ namespace Tools.DotNETCommon
 			// if we support AutoSDKs, then return both versions
 			if (PlatformSupportsAutoSDKs())
 			{
-				// we don't need ask the platform for the installed version because we know the AutoSDK version has to match the DesiredVersion if it was setup properly
-				//AutoSDKVersion = (HasRequiredAutoSDKInstalled() == SDKStatus.Valid) ? GetMainVersion() : null;
-				AutoSDKVersion = GetInstalledSDKVersion();
+				AutoSDKVersion = (HasRequiredAutoSDKInstalled() == SDKStatus.Valid) ? GetInstalledSDKVersion() : null;
+//				AutoSDKVersion = GetInstalledSDKVersion();
 			}
 			else
 			{
@@ -141,6 +140,10 @@ namespace Tools.DotNETCommon
 		{
 			return IsVersionValidInternal(Version, bForAutoSDK);
 		}
+		public virtual bool IsSoftwareVersionValid(string Version)
+		{
+			return IsSoftwareVersionValidInternal(Version);
+		}
 
 		public void ReactivateAutoSDK()
 		{
@@ -169,9 +172,16 @@ namespace Tools.DotNETCommon
 		/// </summary>
 		/// <param name="MinVersion">Smallest version allowed</param>
 		/// <param name="MaxVersion">Largest version allowed (inclusive)</param>
-		/// <returns>True if the versions are valid, false if the platform is unable to convert its versions into an integer</returns>
 		public abstract void GetValidVersionRange(out string MinVersion, out string MaxVersion);
 
+
+		/// <summary>
+		/// Gets the valid string range of software/flash versions. TryConvertVersionToInt() will need to succeed to make this usable for range checks
+		/// </summary>
+		/// <param name="MinVersion">Smallest version allowed, or null if no minmum (in other words, 0 - MaxVersion)</param>
+		/// <param name="MaxVersion">Largest version allowed (inclusive), or null if no maximum (in other words, MinVersion - infinity)y</param>
+		public abstract void GetValidSoftwareVersionRange(out string MinVersion, out string MaxVersion);
+
 		/// <summary>
 		/// For a platform that doesn't use properly named AutoSDK directories, the directory name may not be convertible to an integer,
 		/// and IsVersionValid checks could fail when checking AutoSDK version for an exact match. GetMainVersion() would return the 
@@ -190,23 +200,36 @@ namespace Tools.DotNETCommon
 		/// <param name="MinVersion">Smallest version allowed</param>
 		/// <param name="MaxVersion">Largest version allowed (inclusive)</param>
 		/// <returns>True if the versions are valid, false if the platform is unable to convert its versions into an integer</returns>
-		public virtual bool GetValidVersionRange(out UInt64 MinVersion, out UInt64 MaxVersion)
+		public virtual void GetValidVersionRange(out UInt64 MinVersion, out UInt64 MaxVersion)
 		{
 			string MinVersionString, MaxVersionString;
 			GetValidVersionRange(out MinVersionString, out MaxVersionString);
 
-			if (TryConvertVersionToInt(MinVersionString, out MinVersion) && TryConvertVersionToInt(MaxVersionString, out MaxVersion))
+			// failures to convert here are bad
+			if (!TryConvertVersionToInt(MinVersionString, out MinVersion) || !TryConvertVersionToInt(MaxVersionString, out MaxVersion))
 			{
-				return true;
+				throw new Exception(string.Format("Unable to convert Min and Max valid versions to integers in {0} (Versions are {1} - {2})", GetType().Name, MinVersionString, MaxVersionString));
 			}
+		}
+		public virtual void GetValidSoftwareVersionRange(out UInt64 MinVersion, out UInt64 MaxVersion)
+		{
+			string MinVersionString, MaxVersionString;
+			GetValidSoftwareVersionRange(out MinVersionString, out MaxVersionString);
 
-			// if a version can't be converted to an int, we can't get any versions
-			MinVersion = MaxVersion = 0;
-			return false;
+			MinVersion = UInt64.MinValue;
+			MaxVersion = UInt64.MaxValue - 1; // MaxValue is always bad
+
+			// failures to convert here are bad
+			if ((MinVersionString != null && !TryConvertVersionToInt(MinVersionString, out MinVersion)) || 
+				(MaxVersionString != null && !TryConvertVersionToInt(MaxVersionString, out MaxVersion)))
+			{
+				throw new Exception(string.Format("Unable to convert Min and Max valid Software versions to integers in {0} (Versions are {1} - {2})", GetType().Name, MinVersionString, MaxVersionString));
+			}
 		}
 
+
 		// Let platform override behavior to determine if a version is a valid (useful for non-numeric versions)
-		public virtual bool IsVersionValidInternal(string Version, bool bForAutoSDK)
+		protected virtual bool IsVersionValidInternal(string Version, bool bForAutoSDK)
 		{
 			// we could have null if no SDK is installed at all, etc, which is always a failure
 			if (Version == null)
@@ -247,6 +270,27 @@ namespace Tools.DotNETCommon
 			return IntVersion >= MinVersion && IntVersion <= MaxVersion;
 		}
 
+		protected virtual bool IsSoftwareVersionValidInternal(string Version)
+		{
+			// we could have null if no SDK is installed at all, etc, which is always a failure
+			if (Version == null)
+			{
+				return false;
+			}
+
+			// convert it to an integer
+			UInt64 IntVersion;
+			if (!TryConvertVersionToInt(Version, out IntVersion))
+			{
+				return false;
+			}
+
+			// get numeric range
+			UInt64 MinVersion, MaxVersion;
+			GetValidSoftwareVersionRange(out MinVersion, out MaxVersion);
+			return IntVersion >= MinVersion && IntVersion <= MaxVersion;
+		}
+
 
 		/// <summary>
 		/// Only the platform can convert a version string into an integer that is usable for comparison
@@ -308,6 +352,7 @@ namespace Tools.DotNETCommon
 
 		protected const string AutoSetupEnvVar = "AutoSDKSetup";
 
+
 		protected static bool IsWindows()
 		{
 			return Environment.OSVersion.Platform == PlatformID.Win32NT || Environment.OSVersion.Platform == PlatformID.Win32S || Environment.OSVersion.Platform == PlatformID.Win32Windows;
@@ -800,18 +845,21 @@ namespace Tools.DotNETCommon
 					if (bNeedsToWriteAutoSetupEnvVar)
 					{
 						// write out the manual sdk version since child processes won't be able to detect manual with AutoSDK messing up env vars
-						string ValueToWrite = CachedManualSDKVersion != null ? CachedManualSDKVersion : "__None";
 						using (StreamWriter Writer = File.AppendText(EnvVarFile))
 						{
-							Writer.WriteLine("{0}={1}", PlatformSetupEnvVar, ValueToWrite);
+							Writer.WriteLine("{0}={1}", PlatformSetupEnvVar, "1");
 						}
 						// set the variable in the local environment in case this process spawns any others.
-						Environment.SetEnvironmentVariable(PlatformSetupEnvVar, ValueToWrite);
+						Environment.SetEnvironmentVariable(PlatformSetupEnvVar, "1");
 					}
 
 					// make sure we know that we've modified the local environment, invalidating manual installs for this run.
 					bLocalProcessSetupAutoSDK = true;
 
+					// tell any child processes what our manual version was before setting up autosdk
+					string ValueToWrite = CachedManualSDKVersion != null ? CachedManualSDKVersion : "__None";
+					Environment.SetEnvironmentVariable(GetPlatformAutoSDKSetupEnvVar(), ValueToWrite);
+
 					return true;
 				}
 			}
diff --git a/Engine/Source/Programs/UnrealBuildTool/Platform/Android/AndroidPlatformSDK.cs b/Engine/Source/Programs/UnrealBuildTool/Platform/Android/AndroidPlatformSDK.cs
--- a/Engine/Source/Programs/UnrealBuildTool/Platform/Android/AndroidPlatformSDK.cs
+++ b/Engine/Source/Programs/UnrealBuildTool/Platform/Android/AndroidPlatformSDK.cs
@@ -26,6 +26,13 @@ namespace UnrealBuildTool
 			MaxVersion = "r23a";
 		}
 
+		public override void GetValidSoftwareVersionRange(out string MinVersion, out string MaxVersion)
+		{
+			MinVersion = MaxVersion = null;
+		}
+
+
+
 		// the version we prefer
 		private string RecommendedVersion = "r21b";
 
diff --git a/Engine/Source/Programs/UnrealBuildTool/Platform/HoloLens/UEBuildHoloLens.cs b/Engine/Source/Programs/UnrealBuildTool/Platform/HoloLens/UEBuildHoloLens.cs
--- a/Engine/Source/Programs/UnrealBuildTool/Platform/HoloLens/UEBuildHoloLens.cs
+++ b/Engine/Source/Programs/UnrealBuildTool/Platform/HoloLens/UEBuildHoloLens.cs
@@ -756,6 +756,7 @@ namespace UnrealBuildTool
 		/// </summary>
 		public override void RegisterBuildPlatforms()
 		{
+			// for GetValidSoftwareVersionRange reasons we probably want a HoloLensePlatformSDK class
 			MicrosoftPlatformSDK SDK = new MicrosoftPlatformSDK();
 
 			UEBuildPlatform.RegisterBuildPlatform(new HoloLens(SDK));
diff --git a/Engine/Source/Programs/UnrealBuildTool/Platform/Linux/LinuxPlatformSDK.cs b/Engine/Source/Programs/UnrealBuildTool/Platform/Linux/LinuxPlatformSDK.cs
--- a/Engine/Source/Programs/UnrealBuildTool/Platform/Linux/LinuxPlatformSDK.cs
+++ b/Engine/Source/Programs/UnrealBuildTool/Platform/Linux/LinuxPlatformSDK.cs
@@ -22,6 +22,11 @@ namespace UnrealBuildTool
 			MaxVersion = "v16_clang-9.0.1-centos7";
 		}
 
+		public override void GetValidSoftwareVersionRange(out string MinVersion, out string MaxVersion)
+		{
+			MinVersion = MaxVersion = null;
+		}
+
 		public override string GetAutoSDKPlatformName()
 		{
 			return TargetPlatformName;
diff --git a/Engine/Source/Programs/UnrealBuildTool/Platform/Lumin/LuminPlatformSDK.cs b/Engine/Source/Programs/UnrealBuildTool/Platform/Lumin/LuminPlatformSDK.cs
--- a/Engine/Source/Programs/UnrealBuildTool/Platform/Lumin/LuminPlatformSDK.cs
+++ b/Engine/Source/Programs/UnrealBuildTool/Platform/Lumin/LuminPlatformSDK.cs
@@ -21,6 +21,10 @@ namespace UnrealBuildTool
 			MinVersion = MaxVersion = GetMainVersion();
 		}
 
+		public override void GetValidSoftwareVersionRange(out string MinVersion, out string MaxVersion)
+		{
+			MinVersion = MaxVersion = null;
+		}
 
 		public override string GetInstalledSDKVersion()
 		{
diff --git a/Engine/Source/Programs/UnrealBuildTool/Platform/Mac/ApplePlatformSDK.cs b/Engine/Source/Programs/UnrealBuildTool/Platform/Mac/ApplePlatformSDK.cs
--- a/Engine/Source/Programs/UnrealBuildTool/Platform/Mac/ApplePlatformSDK.cs
+++ b/Engine/Source/Programs/UnrealBuildTool/Platform/Mac/ApplePlatformSDK.cs
@@ -32,6 +32,10 @@ namespace UnrealBuildTool
 			}
 		}
 
+		public override void GetValidSoftwareVersionRange(out string MinVersion, out string MaxVersion)
+		{
+			MinVersion = MaxVersion = null;
+		}
 
 		public override string GetInstalledSDKVersion()
 		{
diff --git a/Engine/Source/Programs/UnrealBuildTool/Platform/Windows/MicrosoftPlatformSDK.cs b/Engine/Source/Programs/UnrealBuildTool/Platform/Windows/MicrosoftPlatformSDK.cs
--- a/Engine/Source/Programs/UnrealBuildTool/Platform/Windows/MicrosoftPlatformSDK.cs
+++ b/Engine/Source/Programs/UnrealBuildTool/Platform/Windows/MicrosoftPlatformSDK.cs
@@ -22,6 +22,12 @@ namespace UnrealBuildTool
 			MaxVersion = "10.9.99999";
 		}
 
+		public override void GetValidSoftwareVersionRange(out string MinVersion, out string MaxVersion)
+		{
+			// @todo we may want to split this for Win64 vs other things
+			MinVersion = MaxVersion = null;
+		}
+
 		public override string GetInstalledSDKVersion()
 		{
 			if (Utils.IsRunningOnMono)
diff --git a/Engine/Source/Runtime/Core/Private/Misc/DataDrivenPlatformInfoRegistry.cpp b/Engine/Source/Runtime/Core/Private/Misc/DataDrivenPlatformInfoRegistry.cpp
--- a/Engine/Source/Runtime/Core/Private/Misc/DataDrivenPlatformInfoRegistry.cpp
+++ b/Engine/Source/Runtime/Core/Private/Misc/DataDrivenPlatformInfoRegistry.cpp
@@ -270,7 +270,8 @@ static void LoadDDPIIniSettings(const FConfigFile& IniFile, FDataDrivenPlatformI
 	DDPIGetString(IniFile, TEXT("TutorialPath"), Info.SDKTutorial);
 	DDPIGetName(IniFile, TEXT("PlatformGroupName"), Info.PlatformGroupName);
 	DDPIGetName(IniFile, TEXT("PlatformSubMenu"), Info.PlatformSubMenu);
-	
+	DDPIGetString(IniFile, TEXT("PrepareForDebuggingOptions"), Info.PrepareForDebuggingOptions);
+
 
 	DDPIGetString(IniFile, TEXT("NormalIconPath"), Info.IconPaths.NormalPath);
 	DDPIGetString(IniFile, TEXT("LargeIconPath"), Info.IconPaths.LargePath);
@@ -449,6 +450,60 @@ static FString ConvertToDDPIDeviceId(const FString& DeviceId)
 	return FString::Printf(TEXT("%s@%s"), *ConvertToDDPIPlatform(PlatformAndDevice[0]), *PlatformAndDevice[1]);
 }
 
+bool GetSdkInfoFromTurnkey(FString Line, FString& PlatformName, FString& DeviceId, FDDPISdkInfo& SdkInfo)
+{
+	int32 Colon = Line.Find(TEXT(": "));
+
+	if (Colon < 0)
+	{
+		return false;
+	}
+
+	// break up the string
+	PlatformName = Line.Mid(0, Colon);
+	FString Info = Line.Mid(Colon + 2);
+
+	int32 AtSign = PlatformName.Find(TEXT("@"));
+	if (AtSign > 0)
+	{
+		// return the platform@name as the deviceId, then remove the @name part for the platform
+		DeviceId = ConvertToDDPIDeviceId(PlatformName);
+		PlatformName = PlatformName.Mid(0, AtSign);
+	}
+
+	// get the DDPI name
+	PlatformName = ConvertToDDPIPlatform(PlatformName);
+
+	// parse out the results from the (key=val, key=val) result from turnkey
+	FString StatusString;
+	FString FlagsString;
+	FParse::Value(*Info, TEXT("Status="), StatusString);
+	FParse::Value(*Info, TEXT("Flags="), FlagsString);
+	FParse::Value(*Info, TEXT("Installed="), SdkInfo.InstalledVersion);
+	FParse::Value(*Info, TEXT("AutoSDK="), SdkInfo.AutoSDKVersion);
+	FParse::Value(*Info, TEXT("MinAllowed="), SdkInfo.MinAllowedVersion);
+	FParse::Value(*Info, TEXT("MaxAllowed="), SdkInfo.MaxAllowedVersion);
+
+	SdkInfo.Status = DDPIPlatformSdkStatus::Unknown;
+	if (StatusString == TEXT("Valid"))
+	{
+		SdkInfo.Status = DDPIPlatformSdkStatus::Valid;
+	}
+	else
+	{
+		if (FlagsString.Contains(TEXT("AutoSdk_InvalidVersionExists")) || FlagsString.Contains(TEXT("InstalledSdk_InvalidVersionExists")))
+		{
+			SdkInfo.Status = DDPIPlatformSdkStatus::OutOfDate;
+		}
+		else
+		{
+			SdkInfo.Status = DDPIPlatformSdkStatus::NoSdk;
+		}
+	}
+
+	return true;
+}
+
 void FDataDrivenPlatformInfoRegistry::UpdateSdkStatus()
 {
 	// make sure we've read in the inis
@@ -459,7 +514,7 @@ void FDataDrivenPlatformInfoRegistry::UpdateSdkStatus()
 	{
 		for (auto& It : DataDrivenPlatforms)
 		{
-			It.Value.SdkStatus = DDPIPlatformSdkStatus::Unknown;
+			It.Value.SdkInfo.Status = DDPIPlatformSdkStatus::Unknown;
 
 			// reset the per-device status when querying general Sdk status
 			It.Value.ClearDeviceStatus();
@@ -481,7 +536,7 @@ void FDataDrivenPlatformInfoRegistry::UpdateSdkStatus()
 		// reset status to unknown
 		for (auto& It : DataDrivenPlatforms)
 		{
-			It.Value.SdkStatus = DDPIPlatformSdkStatus::Querying;
+			It.Value.SdkInfo.Status = DDPIPlatformSdkStatus::Querying;
 
 			// reset the per-device status when querying general Sdk status
 			It.Value.ClearDeviceStatus();
@@ -502,32 +557,34 @@ void FDataDrivenPlatformInfoRegistry::UpdateSdkStatus()
 				{
 					for (FString& Line : Contents)
 					{
-						TArray<FString> Tokens;
-						Line.ParseIntoArray(Tokens, TEXT(": "), true);
-						// Tokens [0] is the platform, [1] is the status, [2] is information
+						UE_LOG(LogTemp, Log, TEXT("Turnkey Platform: %s"), *Line);
 
-						DDPIPlatformSdkStatus Status = DDPIPlatformSdkStatus::Unknown;
-						if (Tokens[1] == TEXT("Valid"))
+						// parse a Turnkey line
+						FString PlatformName, Unused;
+						FDDPISdkInfo SdkInfo;
+						if (GetSdkInfoFromTurnkey(Line, PlatformName, Unused, SdkInfo) == false)
 						{
-							Status = DDPIPlatformSdkStatus::Valid;
+							continue;
 						}
-						else
+
+						// check if we had already set a ManualSDK - and don't set it again. Because of the way AutoSDKs are activated in the editor after the first call to Turnkey,
+						// future calls to Turnkey will inherit the AutoSDK env vars, and it won't be able to determine the manual SDK versions anymore. If we use the editor to
+						// install an SDK via Turnkey, it will directly update the installed version based on the result of that command, not this Update operation
+
+						FString OriginalManualInstallValue = DataDrivenPlatforms[*PlatformName].SdkInfo.InstalledVersion;
+
+						// set it into the platform
+						DataDrivenPlatforms[*PlatformName].SdkInfo = SdkInfo;
+
+						// restore the original installed version if it set after the first time
+						if (OriginalManualInstallValue.Len() > 0)
 						{
-							if (Tokens[2].Contains(TEXT("AutoSdk_InvalidVersionExists")) || Tokens[2].Contains(TEXT("InstalledSdk_InvalidVersionExists")))
-							{
-								Status = DDPIPlatformSdkStatus::OutOfDate;
-							}
-							else
-							{
-								Status = DDPIPlatformSdkStatus::NoSdk;
-							}
+							DataDrivenPlatforms[*PlatformName].SdkInfo.InstalledVersion = OriginalManualInstallValue;
 						}
 
-						// have to convert back to WIndows from Win64
-						FString PlatformName = ConvertToDDPIPlatform(Tokens[0]);
-						DataDrivenPlatforms[*PlatformName].SdkStatus = Status;
 
-						UE_LOG(LogTemp, Log, TEXT("Turnkey Platform: %s - %d - %s"), *PlatformName, (int)Status, *Tokens[2]);
+						UE_LOG(LogTemp, Log, TEXT("[TEST] Turnkey Platform: %s - %d, Installed: %s, AudoSDK: %s, Allowed: %s-%s"), *PlatformName, (int)SdkInfo.Status, *SdkInfo.InstalledVersion,
+							*SdkInfo.AutoSDKVersion, *SdkInfo.MinAllowedVersion, *SdkInfo.MaxAllowedVersion);
 					}
 				}
 			}
@@ -535,7 +592,9 @@ void FDataDrivenPlatformInfoRegistry::UpdateSdkStatus()
 			{
 				for (auto& It : DataDrivenPlatforms)
 				{
-					It.Value.SdkStatus = DDPIPlatformSdkStatus::Error;
+					It.Value.SdkInfo.Status = DDPIPlatformSdkStatus::Error;
+					It.Value.SdkInfo.SdkErrorInformation = FText::Format(NSLOCTEXT("Turnkey", "TurnkeyError_ReturnedError", "Turnkey returned an error, code {0}"), { ExitCode });
+
 					// @todo turnkey error description!
 				}
 			}
@@ -543,17 +602,16 @@ void FDataDrivenPlatformInfoRegistry::UpdateSdkStatus()
 
 			for (auto& It : DataDrivenPlatforms)
 			{
-				if (It.Value.SdkStatus == DDPIPlatformSdkStatus::Querying)
+				if (It.Value.SdkInfo.Status == DDPIPlatformSdkStatus::Querying)
 				{
 					if (It.Value.bIsFakePlatform)
 					{
-						It.Value.SdkStatus = DDPIPlatformSdkStatus::Unknown;
+						It.Value.SdkInfo.Status = DDPIPlatformSdkStatus::Unknown;
 					}
 					else
 					{
-						It.Value.SdkStatus = DDPIPlatformSdkStatus::Error;
-						It.Value.SdkErrorInformation = "The platform's Sdk status was not returned from Turnkey";
-						//					It.Value.SdkErrorInformation = NSLOCTEXT("Turnkey", "PlatformNotReturned", "The platform's Sdk status was not returned from Turnkey");
+						It.Value.SdkInfo.Status = DDPIPlatformSdkStatus::Error;
+						It.Value.SdkInfo.SdkErrorInformation = NSLOCTEXT("Turnkey", "TurnkeyError_NotReturned", "The platform's Sdk status was not returned from Turnkey");
 					}
 				}
 			}
@@ -602,9 +660,11 @@ void FDataDrivenPlatformInfoRegistry::UpdateDeviceSdkStatus(TArray<FString> Plat
 		FScopeLock Lock(&DDPILocker);
 
 		// set status to querying
+		FDDPISdkInfo DefaultInfo;
+		DefaultInfo.Status = DDPIPlatformSdkStatus::Querying;
 		for (const FString& Id : PlatformDeviceIds)
 		{
-			DeviceIdToInfo(Id).PerDeviceStatus.Add(ConvertToDDPIDeviceId(Id), DDPIPlatformSdkStatus::Querying);
+			DeviceIdToInfo(Id).PerDeviceStatus.Add(ConvertToDDPIDeviceId(Id), DefaultInfo);
 		}
 	}
 
@@ -622,24 +682,25 @@ void FDataDrivenPlatformInfoRegistry::UpdateDeviceSdkStatus(TArray<FString> Plat
 				{
 					for (FString& Line : Contents)
 					{
-						TArray<FString> Tokens;
-						Line.ParseIntoArray(Tokens, TEXT(": "), true);
-						// Tokens [0] is platform:device, [1] is the status, [2] is information
+						FString PlatformName, DDPIDeviceId;
+						FDDPISdkInfo SdkInfo;
+						if (GetSdkInfoFromTurnkey(Line, PlatformName, DDPIDeviceId, SdkInfo) == false)
+						{
+							continue;
+						}
 
-						// token[0] without an @ is just platform status, which we don't care about now
-						if (!Tokens[0].Contains(TEXT("@")))
+						// skip over non-device lines
+						if (DDPIDeviceId.Len() == 0)
 						{
 							continue;
 						}
 
-						TArray<FString> PlatformAndDevice;
-						Tokens[0].ParseIntoArray(PlatformAndDevice, TEXT("@"), true);
-						
-						DDPIPlatformSdkStatus Status = Tokens[1] == TEXT("Valid") ? DDPIPlatformSdkStatus::FlashValid : DDPIPlatformSdkStatus::FlashOutOfDate;
+						UE_LOG(LogTemp, Log, TEXT("Turnkey Device: %s"), *Line);
 
-						DeviceIdToInfo(Tokens[0]).PerDeviceStatus[ConvertToDDPIDeviceId(Tokens[0])] = Status;
+						DeviceIdToInfo(DDPIDeviceId).PerDeviceStatus[DDPIDeviceId] = SdkInfo;
 
-						UE_LOG(LogTemp, Log, TEXT("Turnkey Device: %s - %d - %s"), *Tokens[0], (int)Status, *Tokens[2]);
+						UE_LOG(LogTemp, Log, TEXT("[TEST] Turnkey Device: %s - %d, Installed: %s, Allowed: %s-%s"), *DDPIDeviceId, (int)SdkInfo.Status, *SdkInfo.InstalledVersion,
+							*SdkInfo.MinAllowedVersion, *SdkInfo.MaxAllowedVersion);
 					}
 				}
 			}
@@ -648,11 +709,11 @@ void FDataDrivenPlatformInfoRegistry::UpdateDeviceSdkStatus(TArray<FString> Plat
 			{
 				FDataDrivenPlatformInfo& Info = DeviceIdToInfo(Id);
 				
-				DDPIPlatformSdkStatus& Status = Info.PerDeviceStatus[ConvertToDDPIDeviceId(Id)];
-				if (Status == DDPIPlatformSdkStatus::Querying)
+				FDDPISdkInfo& SdkInfo = Info.PerDeviceStatus[ConvertToDDPIDeviceId(Id)];
+				if (SdkInfo.Status == DDPIPlatformSdkStatus::Querying)
 				{
-					Status = DDPIPlatformSdkStatus::Error;
-					Info.SdkErrorInformation = "A device's Sdk status was not returned from Turnkey";
+					SdkInfo.Status = DDPIPlatformSdkStatus::Error;
+					SdkInfo.SdkErrorInformation = NSLOCTEXT("Turnkey", "TurnkeyError_DeviceNotReturned", "A device's Sdk status was not returned from Turnkey");
 				}
 			}
 
@@ -689,14 +750,20 @@ void FDataDrivenPlatformInfoRegistry::ClearDeviceStatus(FName PlatformName)
 DDPIPlatformSdkStatus FDataDrivenPlatformInfo::GetStatusForDeviceId(const FString& DeviceId) const
 {
 	// return the status, or Unknown if not known
-	return PerDeviceStatus.FindRef(ConvertToDDPIDeviceId(DeviceId));
+	return PerDeviceStatus.FindRef(ConvertToDDPIDeviceId(DeviceId)).Status;
 }
 
+const FDDPISdkInfo& FDataDrivenPlatformInfo::GetSdkInfoForDeviceId(const FString& DeviceId) const
+{
+	return *PerDeviceStatus.Find(ConvertToDDPIDeviceId(DeviceId));
+}
+
+
 void FDataDrivenPlatformInfo::ClearDeviceStatus()
 {
 	for (auto& Pair : PerDeviceStatus)
 	{
-		Pair.Value = DDPIPlatformSdkStatus::Unknown;
+		Pair.Value.Status = DDPIPlatformSdkStatus::Unknown;
 	}
 }
 
diff --git a/Engine/Source/Runtime/Core/Public/Misc/DataDrivenPlatformInfoRegistry.h b/Engine/Source/Runtime/Core/Public/Misc/DataDrivenPlatformInfoRegistry.h
--- a/Engine/Source/Runtime/Core/Public/Misc/DataDrivenPlatformInfoRegistry.h
+++ b/Engine/Source/Runtime/Core/Public/Misc/DataDrivenPlatformInfoRegistry.h
@@ -13,7 +13,6 @@
 
 #if DDPI_HAS_EXTENDED_PLATFORMINFO_DATA
 
-
 enum class DDPIPlatformSdkStatus : uint8
 {
 	Unknown,
@@ -27,6 +26,15 @@ enum class DDPIPlatformSdkStatus : uint8
 	// @todo turnkey: add AutoSdkValid and ManualSdkValid, with Valid a Combination of both
 };
 
+struct FDDPISdkInfo
+{
+	DDPIPlatformSdkStatus Status;
+	FText SdkErrorInformation;
+	FString InstalledVersion;
+	FString AutoSDKVersion; // only valid for platform, not device
+	FString MinAllowedVersion;
+	FString MaxAllowedVersion;
+};
 
 /** Available icon sizes (see FPlatformIconPaths) */
 enum class EPlatformIconSize : uint8
@@ -143,16 +151,27 @@ struct FDataDrivenPlatformInfo
 
 #if DDPI_HAS_EXTENDED_PLATFORMINFO_DATA
 
-	// Information about the validity of using a platform, discovered via Turnkey
 private:
-	DDPIPlatformSdkStatus SdkStatus;
+	// Information about the validity of using a platform, discovered via Turnkey
+	FDDPISdkInfo SdkInfo;
+
+	// Information about the validity of each connected device (by string, discovered by Turnkey)
+	TMap<FString, FDDPISdkInfo> PerDeviceStatus;
+
 public:
 	DDPIPlatformSdkStatus GetSdkStatus(bool bBlockIfQuerying = true) const
 	{
-		check(!(bBlockIfQuerying && SdkStatus == DDPIPlatformSdkStatus::Querying));
-		return SdkStatus;
+		check(!(bBlockIfQuerying && SdkInfo.Status == DDPIPlatformSdkStatus::Querying));
+		return SdkInfo.Status;
 	}
-	FString SdkErrorInformation;
+
+	const FDDPISdkInfo& GetSdkInfo(bool bBlockIfQuerying = true) const
+	{
+		check(!(bBlockIfQuerying && SdkInfo.Status == DDPIPlatformSdkStatus::Querying));
+		return SdkInfo;
+	}
+
+
 
 	// setting moved from PlatformInfo::FTargetPlatformInfo
 
@@ -176,6 +195,9 @@ public:
 	FName UBTPlatformName;
 	FString UBTPlatformString;
 
+	/** If this is set, then it is used to prepare a project for debugging */
+	FString PrepareForDebuggingOptions;
+
 	/** Whether or not the platform can use Crash Reporter */
 	bool bCanUseCrashReporter;
 
@@ -192,9 +214,9 @@ public:
 	bool bHasCompiledTargetSupport;
 
 
-
 	// Get the status of a device, or Unknown if not specified
 	CORE_API DDPIPlatformSdkStatus GetStatusForDeviceId(const FString& DeviceId) const;
+	CORE_API const FDDPISdkInfo& GetSdkInfoForDeviceId(const FString& DeviceId) const;
 	CORE_API void ClearDeviceStatus();
 
 
@@ -236,8 +258,6 @@ public:
 private:
 	friend struct FDataDrivenPlatformInfoRegistry;
 
-	// Information about the validity of each connected device (by string, discovered by Turnkey)
-	TMap<FString, DDPIPlatformSdkStatus> PerDeviceStatus;
 #endif
 };
 
