Fixes to get HoloLens working on device

diff --git a/Engine/Source/Runtime/Core/Core.Build.cs b/Engine/Source/Runtime/Core/Core.Build.cs
--- a/Engine/Source/Runtime/Core/Core.Build.cs
+++ b/Engine/Source/Runtime/Core/Core.Build.cs
@@ -217,7 +217,7 @@ public class Core : ModuleRules
 			PublicDefinitions.Add("WITH_DIRECTXMATH=0");
 		}
 
-		if (Target.Platform == UnrealTargetPlatform.Mac)
+		if ((Target.Platform == UnrealTargetPlatform.Mac) || (Target.Platform == UnrealTargetPlatform.HoloLens))
 		{
 			PublicDefinitions.Add("IS_RUNNING_GAMETHREAD_ON_EXTERNAL_THREAD=1");
 		}
diff --git a/Engine/Source/Runtime/Core/Private/HoloLens/HoloLensRunnableThread.cpp b/Engine/Source/Runtime/Core/Private/HoloLens/HoloLensRunnableThread.cpp
--- a/Engine/Source/Runtime/Core/Private/HoloLens/HoloLensRunnableThread.cpp
+++ b/Engine/Source/Runtime/Core/Private/HoloLens/HoloLensRunnableThread.cpp
@@ -60,6 +60,9 @@ uint32 FRunnableThreadHoloLens::Run()
 	{
 		// Initialization has completed, release the sync event
 		ThreadInitSyncEvent->Trigger();
+
+		SetTls();
+
 		// Now run the task that needs to be done
 		ExitCode = Runnable->Run();
 		// Allow any allocated resources to be cleaned up
diff --git a/Engine/Source/Runtime/Launch/Private/HoloLens/HoloLensLaunch.cpp b/Engine/Source/Runtime/Launch/Private/HoloLens/HoloLensLaunch.cpp
--- a/Engine/Source/Runtime/Launch/Private/HoloLens/HoloLensLaunch.cpp
+++ b/Engine/Source/Runtime/Launch/Private/HoloLens/HoloLensLaunch.cpp
@@ -1047,6 +1047,8 @@ void ViewProvider::InitOptionalPackages()
 // this method is called after Load
 void ViewProvider::Run()
 {
+	FTaskTagScope Scope(ETaskTag::EGameThread);
+
 	// Wait for activation to complete
 	while (!ActivationComplete)
 	{
