Add an option to restore the 512MB sized MAX_MEMORY_PER_BLOCK_SIZE for platforms that don't support larger amounts of virtual address space

#rnx

tested:
* ShooterGame on Switch, Win64, PS4

diff --git a/Engine/Source/Runtime/Core/Public/HAL/MallocBinned3.h b/Engine/Source/Runtime/Core/Public/HAL/MallocBinned3.h
--- a/Engine/Source/Runtime/Core/Public/HAL/MallocBinned3.h
+++ b/Engine/Source/Runtime/Core/Public/HAL/MallocBinned3.h
@@ -20,6 +20,9 @@
 
 #define USE_CACHED_PAGE_ALLOCATOR_FOR_LARGE_ALLOCS (0)
 
+#ifndef USE_512MB_MAX_MEMORY_PER_BLOCK_SIZE 
+#define USE_512MB_MAX_MEMORY_PER_BLOCK_SIZE 0
+#endif
 
 #define BINNED3_BASE_PAGE_SIZE				4096			// Minimum "page size" for binned3
 #define BINNED3_MINIMUM_ALIGNMENT_SHIFT		4				// Alignment of blocks, expressed as a shift
@@ -30,7 +33,13 @@
 #define BINNED3_MAX_SMALL_POOL_SIZE			(128 * 1024)	// Maximum medium block size
 #endif
 #define BINNED3_SMALL_POOL_COUNT			(BINNEDCOMMON_NUM_LISTED_SMALL_POOLS + (BINNED3_MAX_SMALL_POOL_SIZE - BINNEDCOMMON_MAX_LISTED_SMALL_POOL_SIZE) / BINNED3_BASE_PAGE_SIZE)
-#define MAX_MEMORY_PER_BLOCK_SIZE_SHIFT (30) // maximum of 512MB per block size
+
+#if USE_512MB_MAX_MEMORY_PER_BLOCK_SIZE
+#define MAX_MEMORY_PER_BLOCK_SIZE_SHIFT (29) // maximum of 512MB per block size
+#else
+#define MAX_MEMORY_PER_BLOCK_SIZE_SHIFT (30) // maximum of 1GB per block size
+#endif
+
 #define MAX_MEMORY_PER_BLOCK_SIZE (1ull << MAX_MEMORY_PER_BLOCK_SIZE_SHIFT) 
 
 // This choice depends on how efficient the OS is with sparse commits in large VM blocks
