Fix for clang compile error on AVX instruction use, which appears to to be disabled

diff --git a/Engine/Source/Runtime/Core/Private/Windows/WindowsPlatformMisc.cpp b/Engine/Source/Runtime/Core/Private/Windows/WindowsPlatformMisc.cpp
--- a/Engine/Source/Runtime/Core/Private/Windows/WindowsPlatformMisc.cpp
+++ b/Engine/Source/Runtime/Core/Private/Windows/WindowsPlatformMisc.cpp
@@ -2311,6 +2311,7 @@ namespace WindowsPlatformMiscImpl
 {
 	static bool TryAVX2Instruction()
 	{
+#if defined(__AVX__)
 		__try
 		{
 			_mm256_set1_ps(1.0f);
@@ -2325,6 +2326,9 @@ namespace WindowsPlatformMiscImpl
 		}
 
 		return true;
+#else 
+		return false;
+#endif 
 	}
 }
 PRAGMA_ENABLE_OPTIMIZATION
