Fix clang compiler builds.

#rnx

diff --git a/Engine/Source/Runtime/Core/Public/Misc/QueuedThreadPool.h b/Engine/Source/Runtime/Core/Public/Misc/QueuedThreadPool.h
--- a/Engine/Source/Runtime/Core/Public/Misc/QueuedThreadPool.h
+++ b/Engine/Source/Runtime/Core/Public/Misc/QueuedThreadPool.h
@@ -74,7 +74,7 @@ private:
  * This interface is used by all queued thread pools. It used as a callback by
  * FQueuedThreads and is used to queue asynchronous work for callers.
  */
-class FQueuedThreadPool
+class CORE_API FQueuedThreadPool
 {
 public:
 	/**
