Fix compile issues in QueuedThreadPool.h

diff --git a/Engine/Source/Runtime/Core/Public/Misc/QueuedThreadPool.h b/Engine/Source/Runtime/Core/Public/Misc/QueuedThreadPool.h
--- a/Engine/Source/Runtime/Core/Public/Misc/QueuedThreadPool.h
+++ b/Engine/Source/Runtime/Core/Public/Misc/QueuedThreadPool.h
@@ -126,14 +126,14 @@ public:
 	 *
 	 * @return A new thread pool.
 	 */
-	static CORE_API FQueuedThreadPool* Allocate();
+	static FQueuedThreadPool* Allocate();
 
 	/**
 	 *	Stack size for threads created for the thread pool. 
 	 *	Can be overridden by other projects.
 	 *	If 0 means to use the value passed in the Create method.
 	 */
-	static CORE_API uint32 OverrideStackSize;
+	static uint32 OverrideStackSize;
 };
 
 /**
