Fix AVX2 texture compression that had been disabled because of Clang compilation errors
   - Removed overly cautious check because it can't be implemented properly on Clang without enabling AVX2 code generation... which we definitely don't want!

diff --git a/Engine/Source/Runtime/Core/Private/Windows/WindowsPlatformMisc.cpp b/Engine/Source/Runtime/Core/Private/Windows/WindowsPlatformMisc.cpp
--- a/Engine/Source/Runtime/Core/Private/Windows/WindowsPlatformMisc.cpp
+++ b/Engine/Source/Runtime/Core/Private/Windows/WindowsPlatformMisc.cpp
@@ -2305,35 +2305,6 @@ FString FWindowsPlatformMisc::GetCPUBrand()
 	return FCPUIDQueriedData::GetBrand();
 }
 
-#if !PLATFORM_SEH_EXCEPTIONS_DISABLED
-PRAGMA_DISABLE_OPTIMIZATION
-namespace WindowsPlatformMiscImpl
-{
-	static bool TryAVX2Instruction()
-	{
-#if defined(__AVX__)
-		__try
-		{
-			_mm256_set1_ps(1.0f);
-
-			// Avoid state-switch penalties on Intel CPUs since UE is only SSE2 optimized by default.
-			// https://software.intel.com/content/www/us/en/develop/articles/intel-avx-state-transitions-migrating-sse-code-to-avx.html
-			_mm256_zeroupper();
-		}
-		__except (EXCEPTION_EXECUTE_HANDLER)
-		{
-			return false;
-		}
-
-		return true;
-#else 
-		return false;
-#endif 
-	}
-}
-PRAGMA_ENABLE_OPTIMIZATION
-#endif
-
 bool FWindowsPlatformMisc::HasAVX2InstructionSupport()
 {
 	if (!HasCPUIDInstruction())
@@ -2378,14 +2349,6 @@ bool FWindowsPlatformMisc::HasAVX2InstructionSupport()
 		return false;
 	}
 
-#if !PLATFORM_SEH_EXCEPTIONS_DISABLED
-	// One last try in case we run inside an hypervisor that's lying to us.
-	if (!WindowsPlatformMiscImpl::TryAVX2Instruction())
-	{
-		return false;
-	}
-#endif
-
 	return true;
 }
 
