PR #7375: fix for `FMath::Eval` not returning after encountering error (Contributed by [REDACTED_PERSON])

diff --git a/Engine/Source/Runtime/Core/Private/Math/UnrealMath.cpp b/Engine/Source/Runtime/Core/Private/Math/UnrealMath.cpp
--- a/Engine/Source/Runtime/Core/Private/Math/UnrealMath.cpp
+++ b/Engine/Source/Runtime/Core/Private/Math/UnrealMath.cpp
@@ -3280,6 +3280,7 @@ PrecLoop:
 	else
 	{
 		UE_LOG(LogUnrealMath, Log, TEXT("Expression Error : Unrecognized Operator"));
+		return 0;
 	}
 
 	*pResult = V;
