Added a memory view type

The new const and mutable view types, FConstMemoryView and FMutableMemoryView, support familiar operations to modify the view, such as Left, Right, Mid, Reset. Other operations include checking if one view contains or intersects another, comparing views, and comparing the bytes within views.

#rn minor Core

diff --git a/Engine/Source/Runtime/Core/Private/Tests/Memory/MemoryViewTest.cpp b/Engine/Source/Runtime/Core/Private/Tests/Memory/MemoryViewTest.cpp
new file mode 100644
--- /dev/null
+++ b/Engine/Source/Runtime/Core/Private/Tests/Memory/MemoryViewTest.cpp
@@ -0,0 +1,226 @@
+// Copyright Epic Games, Inc. All Rights Reserved.
+
+#include "Memory/MemoryView.h"
+
+#include "Misc/AutomationTest.h"
+
+#include <type_traits>
+
+static_assert(std::is_trivially_copyable<FMutableMemoryView>::value, "FMutableMemoryView must be trivially copyable");
+static_assert(std::is_trivially_copy_constructible<FMutableMemoryView>::value, "FMutableMemoryView must be trivially copy constructible");
+static_assert(std::is_trivially_move_constructible<FMutableMemoryView>::value, "FMutableMemoryView must be trivially move constructible");
+static_assert(std::is_trivially_copy_assignable<FMutableMemoryView>::value, "FMutableMemoryView must be trivially copy assignable");
+static_assert(std::is_trivially_move_assignable<FMutableMemoryView>::value, "FMutableMemoryView must be trivially move assignable");
+static_assert(std::is_trivially_destructible<FMutableMemoryView>::value, "FMutableMemoryView must be trivially destructible");
+
+static_assert(std::is_trivially_copyable<FConstMemoryView>::value, "FConstMemoryView must be trivially copyable");
+static_assert(std::is_trivially_copy_constructible<FConstMemoryView>::value, "FConstMemoryView must be trivially copy constructible");
+static_assert(std::is_trivially_move_constructible<FConstMemoryView>::value, "FConstMemoryView must be trivially move constructible");
+static_assert(std::is_trivially_copy_assignable<FConstMemoryView>::value, "FConstMemoryView must be trivially copy assignable");
+static_assert(std::is_trivially_move_assignable<FConstMemoryView>::value, "FConstMemoryView must be trivially move assignable");
+static_assert(std::is_trivially_destructible<FConstMemoryView>::value, "FConstMemoryView must be trivially destructible");
+
+static_assert(std::is_constructible<FConstMemoryView, const FMutableMemoryView&>::value, "Missing constructor");
+static_assert(!std::is_constructible<FMutableMemoryView, const FConstMemoryView&>::value, "Invalid constructor");
+static_assert(!std::is_constructible<FMutableMemoryView, const void*, uint64>::value, "Invalid constructor");
+
+static_assert(std::is_assignable<FConstMemoryView, const FMutableMemoryView&>::value, "Missing assignment");
+static_assert(!std::is_assignable<FMutableMemoryView, const FConstMemoryView&>::value, "Invalid assignment");
+
+#if WITH_DEV_AUTOMATION_TESTS
+
+IMPLEMENT_SIMPLE_AUTOMATION_TEST(FMemoryViewTest, "System.Core.Memory.MemoryView", EAutomationTestFlags::ApplicationContextMask | EAutomationTestFlags::SmokeFilter)
+bool FMemoryViewTest::RunTest(const FString& Parameters)
+{
+	auto TestMutableMemoryView = [this](const FMutableMemoryView& View, void* Data, uint64 Size)
+	{
+		TestEqual(TEXT("MemoryView.GetData()"), View.GetData(), Data);
+		TestEqual(TEXT("MemoryView.Size()"), View.Size(), Size);
+		TestEqual(TEXT("MemoryView.IsEmpty()"), View.IsEmpty(), Size == 0);
+	};
+
+	auto TestConstMemoryView = [this](const FConstMemoryView& View, const void* Data, uint64 Size)
+	{
+		TestEqual(TEXT("MemoryView.GetData()"), View.GetData(), Data);
+		TestEqual(TEXT("MemoryView.Size()"), View.Size(), Size);
+		TestEqual(TEXT("MemoryView.IsEmpty()"), View.IsEmpty(), Size == 0);
+	};
+
+	uint8 ByteArray[16]{};
+	uint32 IntArray[12]{};
+
+	// Test Empty Views
+	TestMutableMemoryView(FMutableMemoryView(), nullptr, 0);
+	TestConstMemoryView(FMutableMemoryView(), nullptr, 0);
+	TestConstMemoryView(FConstMemoryView(), nullptr, 0);
+
+	// Test Construction from Type[], TArrayView, (Type*, uint64)
+	TestMutableMemoryView(MakeMemoryView(IntArray), IntArray, sizeof(IntArray));
+	TestMutableMemoryView(MakeMemoryView(MakeArrayView(IntArray)), IntArray, sizeof(IntArray));
+	TestMutableMemoryView(MakeMemoryView(IntArray, sizeof(IntArray)), IntArray, sizeof(IntArray));
+	TestConstMemoryView(MakeMemoryView(AsConst(IntArray)), IntArray, sizeof(IntArray));
+	TestConstMemoryView(MakeMemoryView(MakeArrayView(AsConst(IntArray))), IntArray, sizeof(IntArray));
+	TestConstMemoryView(MakeMemoryView(AsConst(IntArray), sizeof(IntArray)), IntArray, sizeof(IntArray));
+
+	// Test Construction from std::initializer_list
+	//MakeMemoryView({1, 2, 3}); // fail because the type must be deduced
+	std::initializer_list<uint8> InitializerList{1, 2, 3};
+	TestConstMemoryView(MakeMemoryView(InitializerList), GetData(InitializerList), GetNum(InitializerList) * sizeof(uint8));
+
+	// Test Reset
+	{
+		FMutableMemoryView View = MakeMemoryView(IntArray);
+		View.Reset();
+		TestEqual(TEXT("MemoryView.Reset()"), View, FMutableMemoryView());
+	}
+
+	// Test Left
+	static_assert(MakeMemoryView(IntArray).Left(0).IsEmpty(), "Error in Left");
+	static_assert(MakeMemoryView(IntArray).Left(1) == MakeMemoryView(IntArray, 1), "Error in Left");
+	static_assert(MakeMemoryView(IntArray).Left(sizeof(IntArray)) == MakeMemoryView(IntArray), "Error in Left");
+	static_assert(MakeMemoryView(IntArray).Left(sizeof(IntArray) + 1) == MakeMemoryView(IntArray), "Error in Left");
+	static_assert(MakeMemoryView(IntArray).Left(MAX_uint64) == MakeMemoryView(IntArray), "Error in Left");
+
+	// Test LeftChop
+	static_assert(MakeMemoryView(IntArray).LeftChop(0) == MakeMemoryView(IntArray), "Error in LeftChop");
+	static_assert(MakeMemoryView(IntArray).LeftChop(1) == MakeMemoryView(IntArray, sizeof(IntArray) - 1), "Error in LeftChop");
+	static_assert(MakeMemoryView(IntArray).LeftChop(sizeof(IntArray)).IsEmpty(), "Error in LeftChop");
+	static_assert(MakeMemoryView(IntArray).LeftChop(sizeof(IntArray) + 1).IsEmpty(), "Error in LeftChop");
+	static_assert(MakeMemoryView(IntArray).LeftChop(MAX_uint64).IsEmpty(), "Error in LeftChop");
+
+	// Test Right
+	TestEqual(TEXT("MemoryView.Right(0)"), MakeMemoryView(IntArray).Right(0), FMutableMemoryView());
+	TestEqual(TEXT("MemoryView.Right(1)"), MakeMemoryView(IntArray).Right(1), MakeMemoryView(reinterpret_cast<uint8*>(IntArray) + sizeof(IntArray) - 1, 1));
+	TestEqual(TEXT("MemoryView.Right(Size)"), MakeMemoryView(IntArray).Right(sizeof(IntArray)), MakeMemoryView(IntArray));
+	TestEqual(TEXT("MemoryView.Right(Size+1)"), MakeMemoryView(IntArray).Right(sizeof(IntArray) + 1), MakeMemoryView(IntArray));
+	TestEqual(TEXT("MemoryView.Right(MaxSize)"), MakeMemoryView(IntArray).Right(MAX_uint64), MakeMemoryView(IntArray));
+
+	// Test RightChop
+	TestEqual(TEXT("MemoryView.RightChop(0)"), MakeMemoryView(IntArray).RightChop(0), MakeMemoryView(IntArray));
+	TestEqual(TEXT("MemoryView.RightChop(1)"), MakeMemoryView(IntArray).RightChop(1), MakeMemoryView(reinterpret_cast<uint8*>(IntArray) + 1, sizeof(IntArray) - 1));
+	TestEqual(TEXT("MemoryView.RightChop(Size)"), MakeMemoryView(IntArray).RightChop(sizeof(IntArray)), FMutableMemoryView());
+	TestEqual(TEXT("MemoryView.RightChop(Size+1)"), MakeMemoryView(IntArray).RightChop(sizeof(IntArray) + 1), FMutableMemoryView());
+	TestEqual(TEXT("MemoryView.RightChop(MaxSize)"), MakeMemoryView(IntArray).RightChop(MAX_uint64), FMutableMemoryView());
+
+	// Test Mid
+	TestEqual(TEXT("MemoryView.Mid(0)"), MakeMemoryView(IntArray).Mid(0), MakeMemoryView(IntArray));
+	TestEqual(TEXT("MemoryView.Mid(1)"), MakeMemoryView(IntArray).Mid(1), MakeMemoryView(reinterpret_cast<uint8*>(IntArray) + 1, sizeof(IntArray) - 1));
+	TestEqual(TEXT("MemoryView.Mid(Size)"), MakeMemoryView(IntArray).Mid(sizeof(IntArray)), FMutableMemoryView());
+	TestEqual(TEXT("MemoryView.Mid(Size+1)"), MakeMemoryView(IntArray).Mid(sizeof(IntArray) + 1), FMutableMemoryView());
+	TestEqual(TEXT("MemoryView.Mid(MaxSize)"), MakeMemoryView(IntArray).Mid(MAX_uint64), FMutableMemoryView());
+	TestEqual(TEXT("MemoryView.Mid(0,0)"), MakeMemoryView(IntArray).Mid(0, 0), FMutableMemoryView());
+	TestEqual(TEXT("MemoryView.Mid(0,1)"), MakeMemoryView(IntArray).Mid(0, 1), MakeMemoryView(IntArray, 1));
+	TestEqual(TEXT("MemoryView.Mid(1,Size-2)"), MakeMemoryView(IntArray).Mid(1, sizeof(IntArray) - 2), MakeMemoryView(reinterpret_cast<uint8*>(IntArray) + 1, sizeof(IntArray) - 2));
+	TestEqual(TEXT("MemoryView.Mid(1,Size-1)"), MakeMemoryView(IntArray).Mid(1, sizeof(IntArray) - 1), MakeMemoryView(reinterpret_cast<uint8*>(IntArray) + 1, sizeof(IntArray) - 1));
+	TestEqual(TEXT("MemoryView.Mid(1,Size)"), MakeMemoryView(IntArray).Mid(1, sizeof(IntArray)), MakeMemoryView(reinterpret_cast<uint8*>(IntArray) + 1, sizeof(IntArray) - 1));
+	TestEqual(TEXT("MemoryView.Mid(0,MaxSize)"), MakeMemoryView(IntArray).Mid(0, MAX_uint64), MakeMemoryView(IntArray));
+	TestEqual(TEXT("MemoryView.Mid(MaxSize,MaxSize)"), MakeMemoryView(IntArray).Mid(MAX_uint64, MAX_uint64), FMutableMemoryView());
+
+	// Test Contains
+	TestTrue(TEXT("MemoryView.Contains(Empty)"), FMutableMemoryView().Contains(FConstMemoryView()));
+	TestTrue(TEXT("MemoryView.Contains(Empty)"), FConstMemoryView().Contains(FMutableMemoryView()));
+	TestTrue(TEXT("MemoryView.Contains(Equal)"), MakeMemoryView(ByteArray).Contains(MakeMemoryView(ByteArray)));
+	TestTrue(TEXT("MemoryView.Contains(SmallerBy1Left)"), MakeMemoryView(ByteArray).Contains(MakeMemoryView(ByteArray + 1, 15)));
+	TestTrue(TEXT("MemoryView.Contains(SmallerBy1Right)"), MakeMemoryView(ByteArray).Contains(MakeMemoryView(ByteArray, 15)));
+	TestTrue(TEXT("MemoryView.Contains(SmallerBy2Both)"), MakeMemoryView(ByteArray).Contains(MakeMemoryView(ByteArray + 1, 14)));
+	TestTrue(TEXT("MemoryView.Contains(EmptyContained)"), MakeMemoryView(ByteArray).Contains(MakeMemoryView(ByteArray, 0)));
+	TestTrue(TEXT("MemoryView.Contains(EmptyContained)"), MakeMemoryView(ByteArray).Contains(MakeMemoryView(ByteArray + 8, 0)));
+	TestTrue(TEXT("MemoryView.Contains(EmptyContained)"), MakeMemoryView(ByteArray).Contains(MakeMemoryView(ByteArray + 16, 0)));
+	TestFalse(TEXT("MemoryView.Contains(EmptyOutside)"), MakeMemoryView(ByteArray).Contains(MakeMemoryView(ByteArray - 1, 0)));
+	TestFalse(TEXT("MemoryView.Contains(EmptyOutside)"), MakeMemoryView(ByteArray).Contains(MakeMemoryView(ByteArray + 17, 0)));
+	TestFalse(TEXT("MemoryView.Contains(OutsideBy1Left)"), MakeMemoryView(ByteArray).Contains(MakeMemoryView(ByteArray - 1, 1)));
+	TestFalse(TEXT("MemoryView.Contains(OutsideBy1Right)"), MakeMemoryView(ByteArray).Contains(MakeMemoryView(ByteArray + 16, 1)));
+	TestFalse(TEXT("MemoryView.Contains(LargerBy1Left)"), MakeMemoryView(ByteArray).Contains(MakeMemoryView(ByteArray - 1, 17)));
+	TestFalse(TEXT("MemoryView.Contains(LargerBy1Right)"), MakeMemoryView(ByteArray).Contains(MakeMemoryView(ByteArray, 17)));
+	TestFalse(TEXT("MemoryView.Contains(LargerBy2Both)"), MakeMemoryView(ByteArray).Contains(MakeMemoryView(ByteArray - 1, 18)));
+	TestFalse(TEXT("MemoryView.Contains(SmallerOverlapLeft)"), MakeMemoryView(ByteArray).Contains(MakeMemoryView(ByteArray - 1, 8)));
+	TestFalse(TEXT("MemoryView.Contains(SmallerOverlapRight)"), MakeMemoryView(ByteArray).Contains(MakeMemoryView(ByteArray + 9, 8)));
+
+	// Test Intersects
+	TestTrue(TEXT("MemoryView.Intersects(Equal)"), MakeMemoryView(ByteArray).Intersects(MakeMemoryView(ByteArray)));
+	TestTrue(TEXT("MemoryView.Intersects(SmallerBy1Left)"), MakeMemoryView(ByteArray).Intersects(MakeMemoryView(ByteArray + 1, 15)));
+	TestTrue(TEXT("MemoryView.Intersects(SmallerBy1Right)"), MakeMemoryView(ByteArray).Intersects(MakeMemoryView(ByteArray, 15)));
+	TestTrue(TEXT("MemoryView.Intersects(SmallerBy2Both)"), MakeMemoryView(ByteArray).Intersects(MakeMemoryView(ByteArray + 1, 14)));
+	TestTrue(TEXT("MemoryView.Intersects(SmallerOverlapLeft)"), MakeMemoryView(ByteArray).Intersects(MakeMemoryView(ByteArray - 1, 8)));
+	TestTrue(TEXT("MemoryView.Intersects(SmallerOverlapRight)"), MakeMemoryView(ByteArray).Intersects(MakeMemoryView(ByteArray + 9, 8)));
+	TestTrue(TEXT("MemoryView.Intersects(LargerBy1Left)"), MakeMemoryView(ByteArray).Intersects(MakeMemoryView(ByteArray - 1, 17)));
+	TestTrue(TEXT("MemoryView.Intersects(LargerBy1Right)"), MakeMemoryView(ByteArray).Intersects(MakeMemoryView(ByteArray, 17)));
+	TestTrue(TEXT("MemoryView.Intersects(LargerBy2Both)"), MakeMemoryView(ByteArray).Intersects(MakeMemoryView(ByteArray - 1, 18)));
+	TestTrue(TEXT("MemoryView.Intersects(EmptyMiddle)"), MakeMemoryView(ByteArray).Intersects(MakeMemoryView(ByteArray + 8, 0)));
+	TestFalse(TEXT("MemoryView.Intersects(Empty)"), FMutableMemoryView().Intersects(FConstMemoryView()));
+	TestFalse(TEXT("MemoryView.Intersects(Empty)"), FConstMemoryView().Intersects(FMutableMemoryView()));
+	TestFalse(TEXT("MemoryView.Intersects(EmptyLeft)"), MakeMemoryView(ByteArray).Intersects(MakeMemoryView(ByteArray, 0)));
+	TestFalse(TEXT("MemoryView.Intersects(EmptyRight)"), MakeMemoryView(ByteArray).Intersects(MakeMemoryView(ByteArray + 16, 0)));
+	TestFalse(TEXT("MemoryView.Intersects(EmptyOutside)"), MakeMemoryView(ByteArray).Intersects(MakeMemoryView(ByteArray - 1, 0)));
+	TestFalse(TEXT("MemoryView.Intersects(EmptyOutside)"), MakeMemoryView(ByteArray).Intersects(MakeMemoryView(ByteArray + 17, 0)));
+	TestFalse(TEXT("MemoryView.Intersects(OutsideBy1Left)"), MakeMemoryView(ByteArray).Intersects(MakeMemoryView(ByteArray - 1, 1)));
+	TestFalse(TEXT("MemoryView.Intersects(OutsideBy1Right)"), MakeMemoryView(ByteArray).Intersects(MakeMemoryView(ByteArray + 16, 1)));
+
+	// Test CompareBytes
+	const uint8 CompareBytes[8] = { 5, 4, 6, 2, 4, 7, 1, 3 };
+	TestEqual(TEXT("MemoryView.CompareBytes(Empty)"), FMutableMemoryView().CompareBytes(FConstMemoryView()), 0);
+	TestEqual(TEXT("MemoryView.CompareBytes(Empty)"), FConstMemoryView().CompareBytes(FMutableMemoryView()), 0);
+	TestEqual(TEXT("MemoryView.CompareBytes(EqualView)"), MakeMemoryView(ByteArray).CompareBytes(MakeMemoryView(ByteArray)), 0);
+	TestEqual(TEXT("MemoryView.CompareBytes(EqualBytes)"), MakeMemoryView(ByteArray, 8).CompareBytes(MakeMemoryView(ByteArray + 8, 8)), 0);
+	TestTrue(TEXT("MemoryView.CompareBytes(EmptyLeft)"), FConstMemoryView().CompareBytes(MakeMemoryView(ByteArray)) < 0);
+	TestTrue(TEXT("MemoryView.CompareBytes(EmptyRight)"), MakeMemoryView(ByteArray).CompareBytes(FConstMemoryView()) > 0);
+	TestTrue(TEXT("MemoryView.CompareBytes(ShorterLeft)"), MakeMemoryView(ByteArray, 8).CompareBytes(MakeMemoryView(ByteArray)) < 0);
+	TestTrue(TEXT("MemoryView.CompareBytes(ShorterRight)"), MakeMemoryView(ByteArray).CompareBytes(MakeMemoryView(ByteArray, 8)) > 0);
+	TestTrue(TEXT("MemoryView.CompareBytes(ShorterLeft)"), MakeMemoryView(IntArray, 8).CompareBytes(MakeMemoryView(ByteArray)) < 0);
+	TestTrue(TEXT("MemoryView.CompareBytes(ShorterRight)"), MakeMemoryView(ByteArray).CompareBytes(MakeMemoryView(IntArray, 8)) > 0);
+	TestTrue(TEXT("MemoryView.CompareBytes(DifferentSize)"), MakeMemoryView(ByteArray, 4).CompareBytes(MakeMemoryView(ByteArray, 8)) < 0);
+	TestTrue(TEXT("MemoryView.CompareBytes(DifferentSize)"), MakeMemoryView(ByteArray, 8).CompareBytes(MakeMemoryView(ByteArray, 4)) > 0);
+	TestTrue(TEXT("MemoryView.CompareBytes(SameSizeLeftLess)"), MakeMemoryView(CompareBytes, 2).CompareBytes(MakeMemoryView(CompareBytes + 2, 2)) < 0);
+	TestTrue(TEXT("MemoryView.CompareBytes(SameSizeLeftGreater)"), MakeMemoryView(CompareBytes, 3).CompareBytes(MakeMemoryView(CompareBytes + 3, 3)) > 0);
+
+	// Test EqualBytes
+	TestTrue(TEXT("MemoryView.EqualBytes(Empty)"), FMutableMemoryView().EqualBytes(FConstMemoryView()));
+	TestTrue(TEXT("MemoryView.EqualBytes(Empty)"), FConstMemoryView().EqualBytes(FMutableMemoryView()));
+	TestTrue(TEXT("MemoryView.EqualBytes(EqualView)"), MakeMemoryView(ByteArray).EqualBytes(MakeMemoryView(ByteArray)));
+	TestTrue(TEXT("MemoryView.EqualBytes(EqualBytes)"), MakeMemoryView(ByteArray, 8).EqualBytes(MakeMemoryView(ByteArray + 8, 8)));
+	TestFalse(TEXT("MemoryView.EqualBytes(DifferentSize)"), MakeMemoryView(ByteArray, 8).EqualBytes(MakeMemoryView(ByteArray, 4)));
+	TestFalse(TEXT("MemoryView.EqualBytes(DifferentSize)"), MakeMemoryView(ByteArray, 4).EqualBytes(MakeMemoryView(ByteArray, 8)));
+	TestFalse(TEXT("MemoryView.EqualBytes(DifferentBytes)"), MakeMemoryView(CompareBytes, 4).EqualBytes(MakeMemoryView(CompareBytes + 4, 4)));
+
+	// Test Equals
+	TestTrue(TEXT("MemoryView.Equals(Empty)"), FMutableMemoryView().Equals(FMutableMemoryView()));
+	TestTrue(TEXT("MemoryView.Equals(Empty)"), FMutableMemoryView().Equals(FConstMemoryView()));
+	TestTrue(TEXT("MemoryView.Equals(Empty)"), FConstMemoryView().Equals(FMutableMemoryView()));
+	TestTrue(TEXT("MemoryView.Equals(Empty)"), FConstMemoryView().Equals(FConstMemoryView()));
+	TestTrue(TEXT("MemoryView.Equals(Equal)"), MakeMemoryView(IntArray).Equals(MakeMemoryView(AsConst(IntArray))));
+	TestFalse(TEXT("MemoryView.Equals(DataDiff)"), MakeMemoryView(IntArray).Equals(MakeMemoryView(IntArray + 1, sizeof(IntArray) - sizeof(*IntArray))));
+	TestFalse(TEXT("MemoryView.Equals(SizeDiff)"), MakeMemoryView(IntArray).Equals(MakeMemoryView(IntArray, sizeof(*IntArray))));
+	TestFalse(TEXT("MemoryView.Equals(BothDiff)"), MakeMemoryView(IntArray).Equals(FMutableMemoryView()));
+
+	// Test operator==
+	static_assert(MakeMemoryView(ByteArray) == MakeMemoryView(ByteArray), "Error in MemoryView == MemoryView");
+	static_assert(MakeMemoryView(ByteArray) == MakeMemoryView(AsConst(ByteArray)), "Error in MemoryView == MemoryView");
+	static_assert(MakeMemoryView(AsConst(ByteArray)) == MakeMemoryView(ByteArray), "Error in MemoryView == MemoryView");
+	static_assert(MakeMemoryView(AsConst(ByteArray)) == MakeMemoryView(AsConst(ByteArray)), "Error in MemoryView == MemoryView");
+
+	// Test operator!=
+	static_assert(MakeMemoryView(ByteArray) != MakeMemoryView(IntArray), "Error in MemoryView != MemoryView");
+	static_assert(MakeMemoryView(ByteArray) != MakeMemoryView(AsConst(IntArray)), "Error in MemoryView != MemoryView");
+	static_assert(MakeMemoryView(AsConst(ByteArray)) != MakeMemoryView(IntArray), "Error in MemoryView != MemoryView");
+	static_assert(MakeMemoryView(AsConst(ByteArray)) != MakeMemoryView(AsConst(IntArray)), "Error in MemoryView != MemoryView");
+
+	// Test operator+=
+	TestEqual(TEXT("MemoryView += 0"), MakeMemoryView(ByteArray) += 0, MakeMemoryView(ByteArray));
+	TestEqual(TEXT("MemoryView += Half"), MakeMemoryView(ByteArray) += 8, MakeMemoryView(ByteArray + 8, 8));
+	TestEqual(TEXT("MemoryView += Size"), MakeMemoryView(ByteArray) += 16, MakeMemoryView(ByteArray + 16, 0));
+	TestEqual(TEXT("MemoryView += OutOfBounds"), MakeMemoryView(ByteArray) += 32, MakeMemoryView(ByteArray + 16, 0));
+
+	// Test operator+
+	TestEqual(TEXT("MemoryView + 0"), MakeMemoryView(ByteArray) + 0, MakeMemoryView(ByteArray));
+	TestEqual(TEXT("0 + MemoryView"), 0 + MakeMemoryView(ByteArray), MakeMemoryView(ByteArray));
+	TestEqual(TEXT("MemoryView + Half"), MakeMemoryView(ByteArray) + 8, MakeMemoryView(ByteArray + 8, 8));
+	TestEqual(TEXT("Half + MemoryView"), 8 + MakeMemoryView(ByteArray), MakeMemoryView(ByteArray + 8, 8));
+	TestEqual(TEXT("MemoryView + Size"), MakeMemoryView(ByteArray) + 16, MakeMemoryView(ByteArray + 16, 0));
+	TestEqual(TEXT("Size + MemoryView"), 16 + MakeMemoryView(ByteArray), MakeMemoryView(ByteArray + 16, 0));
+	TestEqual(TEXT("MemoryView + OutOfBounds"), MakeMemoryView(ByteArray) + 32, MakeMemoryView(ByteArray + 16, 0));
+	TestEqual(TEXT("OutOfBounds + MemoryView"), 32 + MakeMemoryView(ByteArray), MakeMemoryView(ByteArray + 16, 0));
+
+	return true;
+}
+
+#endif // WITH_DEV_AUTOMATION_TESTS
diff --git a/Engine/Source/Runtime/Core/Public/Memory/MemoryView.h b/Engine/Source/Runtime/Core/Public/Memory/MemoryView.h
new file mode 100644
--- /dev/null
+++ b/Engine/Source/Runtime/Core/Public/Memory/MemoryView.h
@@ -0,0 +1,254 @@
+// Copyright Epic Games, Inc. All Rights Reserved.
+
+#pragma once
+
+#include "CoreTypes.h"
+#include "Math/NumericLimits.h"
+#include "Math/UnrealMathUtility.h"
+#include "Templates/AndOrNot.h"
+#include "Templates/ChooseClass.h"
+#include "Templates/EnableIf.h"
+#include "Templates/Identity.h"
+#include "Templates/IsConst.h"
+#include "Templates/PointerIsConvertibleFromTo.h"
+#include "Templates/UnrealTemplate.h"
+#include "Traits/IsContiguousContainer.h"
+#include "Traits/IsVoidType.h"
+
+#include <initializer_list>
+
+/**
+ * A non-owning view of a contiguous region of memory.
+ *
+ * Prefer to use the aliases \ref FConstMemoryView or \ref FMutableMemoryView over this type.
+ *
+ * Functions that modify a view clamp sizes and offsets to always return a sub-view of the input.
+ */
+template <typename DataType>
+class TMemoryView
+{
+	static_assert(TIsVoidType<DataType>::Value, "DataMust be cv-qualified void");
+
+public:
+	/** Construct an empty view. */
+	constexpr TMemoryView() = default;
+
+	/** Construct a view of InSize bytes starting at InData. */
+	constexpr inline TMemoryView(DataType* InData, uint64 InSize) : Data(InData), DataSize(InSize) {}
+
+	template <typename OtherDataType, typename = typename TEnableIf<TPointerIsConvertibleFromTo<OtherDataType, DataType>::Value>::Type>
+	constexpr inline TMemoryView(const TMemoryView<OtherDataType>& InView) : Data(InView.Data), DataSize(InView.DataSize) {}
+
+	/** Returns a pointer to the start of the view. */
+	constexpr inline DataType* GetData() const { return Data; }
+
+	/** Returns the number of bytes in the view. */
+	constexpr inline uint64 Size() const { return DataSize; }
+
+	/** Returns whether the view has a size of 0 regardless of its data pointer. */
+	constexpr inline bool IsEmpty() const { return DataSize == 0; }
+
+	/** Resets to an empty view. */
+	constexpr inline void Reset() { *this = TMemoryView(); }
+
+	/** Returns the left-most part of the view by taking the given number of bytes from the left. */
+	constexpr inline TMemoryView Left(uint64 InSize) const
+	{
+		TMemoryView View(*this);
+		View.LeftInline(InSize);
+		return View;
+	}
+
+	/** Returns the left-most part of the view by chopping the given number of bytes from the right. */
+	constexpr inline TMemoryView LeftChop(uint64 InSize) const
+	{
+		TMemoryView View(*this);
+		View.LeftChopInline(InSize);
+		return View;
+	}
+
+	/** Returns the right-most part of the view by taking the given number of bytes from the right. */
+	inline TMemoryView Right(uint64 InSize) const
+	{
+		TMemoryView View(*this);
+		View.RightInline(InSize);
+		return View;
+	}
+
+	/** Returns the right-most part of the view by chopping the given number of bytes from the left. */
+	inline TMemoryView RightChop(uint64 InSize) const
+	{
+		TMemoryView View(*this);
+		View.RightChopInline(InSize);
+		return View;
+	}
+
+	/** Returns the middle part of the view by taking up to the given number of bytes from the given position. */
+	inline TMemoryView Mid(uint64 InOffset, uint64 InSize = TNumericLimits<uint64>::Max()) const
+	{
+		TMemoryView View(*this);
+		View.MidInline(InOffset, InSize);
+		return View;
+	}
+
+	/** Modifies the view to be the given number of bytes from the left. */
+	constexpr inline void LeftInline(uint64 InSize)
+	{
+		DataSize = FMath::Min(DataSize, InSize);
+	}
+
+	/** Modifies the view by chopping the given number of bytes from the right. */
+	constexpr inline void LeftChopInline(uint64 InSize)
+	{
+		DataSize -= FMath::Min(DataSize, InSize);
+	}
+
+	/** Modifies the view to be the given number of bytes from the right. */
+	inline void RightInline(uint64 InSize)
+	{
+		const uint64 OldSize = DataSize;
+		const uint64 NewSize = FMath::Min(OldSize, InSize);
+		Data = GetDataAtOffsetNoCheck(OldSize - NewSize);
+		DataSize = NewSize;
+	}
+
+	/** Modifies the view by chopping the given number of bytes from the left. */
+	inline void RightChopInline(uint64 InSize)
+	{
+		const uint64 Offset = FMath::Min(DataSize, InSize);
+		Data = GetDataAtOffsetNoCheck(Offset);
+		DataSize -= Offset;
+	}
+
+	/** Modifies the view to be the middle part by taking up to the given number of bytes from the given offset. */
+	inline void MidInline(uint64 InOffset, uint64 InSize = TNumericLimits<uint64>::Max())
+	{
+		RightChopInline(InOffset);
+		LeftInline(InSize);
+	}
+
+	/** Returns whether this view fully contains the other view. */
+	template <typename OtherDataType>
+	inline bool Contains(const TMemoryView<OtherDataType>& InView) const
+	{
+		return Data <= InView.Data && GetDataAtOffsetNoCheck(DataSize) >= InView.GetDataAtOffsetNoCheck(InView.DataSize);
+	}
+
+	/** Returns whether this view intersects the other view. */
+	template <typename OtherDataType>
+	inline bool Intersects(const TMemoryView<OtherDataType>& InView) const
+	{
+		return Data < InView.GetDataAtOffsetNoCheck(InView.DataSize) && InView.Data < GetDataAtOffsetNoCheck(DataSize);
+	}
+
+	/** Returns whether the bytes of this view are equal or less/greater than the bytes of the other view. */
+	template <typename OtherDataType>
+	inline int32 CompareBytes(const TMemoryView<OtherDataType>& InView) const
+	{
+		const int32 Compare = Data == InView.Data ? 0 : FMemory::Memcmp(Data, InView.Data, FMath::Min(DataSize, InView.DataSize));
+		return Compare || DataSize == InView.DataSize ? Compare : DataSize < InView.DataSize ? -1 : 1;
+	}
+
+	/** Returns whether the bytes of this views are equal to the bytes of the other view. */
+	template <typename OtherDataType>
+	inline bool EqualBytes(const TMemoryView<OtherDataType>& InView) const
+	{
+		return DataSize == InView.DataSize && (Data == InView.Data || FMemory::Memcmp(Data, InView.Data, DataSize) == 0);
+	}
+
+	/** Returns whether the data pointers and sizes of this view and the other view are equal. */
+	template <typename OtherDataType>
+	constexpr inline bool Equals(const TMemoryView<OtherDataType>& InView) const
+	{
+		return DataSize == InView.DataSize && (DataSize == 0 || Data == InView.Data);
+	}
+
+	/** Returns whether the data pointers and sizes of this view and the other view are equal. */
+	template <typename OtherDataType>
+	constexpr inline bool operator==(const TMemoryView<OtherDataType>& InView) const
+	{
+		return Equals(InView);
+	}
+
+	/** Returns whether the data pointers and sizes of this view and the other view are not equal. */
+	template <typename OtherDataType>
+	constexpr inline bool operator!=(const TMemoryView<OtherDataType>& InView) const
+	{
+		return !Equals(InView);
+	}
+
+	/** Advances the start of the view by an offset, which is clamped to stay within the view. */
+	constexpr inline TMemoryView& operator+=(uint64 InOffset)
+	{
+		RightChopInline(InOffset);
+		return *this;
+	}
+
+private:
+	/** Returns the data pointer advanced by an offset in bytes. */
+	inline DataType* GetDataAtOffsetNoCheck(uint64 InOffset) const
+	{
+		using ByteType = typename TChooseClass<TIsConst<DataType>::Value, const uint8, uint8>::Result;
+		return reinterpret_cast<ByteType*>(Data) + InOffset;
+	}
+
+	template <typename OtherDataType>
+	friend class TMemoryView;
+
+private:
+	DataType* Data = nullptr;
+	uint64 DataSize = 0;
+};
+
+/** Advances the start of the view by an offset, which is clamped to stay within the view. */
+template <typename DataType>
+constexpr inline TMemoryView<DataType> operator+(const TMemoryView<DataType>& View, uint64 Offset)
+{
+	return TMemoryView<DataType>(View) += Offset;
+}
+
+/** Advances the start of the view by an offset, which is clamped to stay within the view. */
+template <typename DataType>
+constexpr inline TMemoryView<DataType> operator+(uint64 Offset, const TMemoryView<DataType>& View)
+{
+	return TMemoryView<DataType>(View) += Offset;
+}
+
+/** Make a non-owning mutable view of Size bytes starting at Data. */
+constexpr inline TMemoryView<void> MakeMemoryView(void* Data, uint64 Size)
+{
+	return TMemoryView<void>(Data, Size);
+}
+
+/** Make a non-owning const view of Size bytes starting at Data. */
+constexpr inline TMemoryView<const void> MakeMemoryView(const void* Data, uint64 Size)
+{
+	return TMemoryView<const void>(Data, Size);
+}
+
+/**
+ * Make a non-owning view of the memory of the initializer list.
+ *
+ * This overload is only available when the element type does not need to be deduced.
+ */
+template <typename T>
+constexpr inline TMemoryView<const void> MakeMemoryView(std::initializer_list<typename TIdentity<T>::Type> List)
+{
+	return TMemoryView<const void>(GetData(List), GetNum(List) * sizeof(T));
+}
+
+/** Make a non-owning view of the memory of the contiguous container. */
+template <typename ContainerType, typename TEnableIf<TIsContiguousContainer<ContainerType>::Value>::Type* = nullptr>
+constexpr inline auto MakeMemoryView(ContainerType&& Container)
+{
+	using ElementType = typename TRemovePointer<decltype(GetData(DeclVal<ContainerType>()))>::Type;
+	constexpr bool bIsConst = TIsConst<ElementType>::Value;
+	using DataType = typename TChooseClass<bIsConst, const void, void>::Result;
+	return TMemoryView<DataType>(GetData(Container), GetNum(Container) * sizeof(ElementType));
+}
+
+/** A non-owning mutable view of a contiguous region of memory. */
+using FMutableMemoryView = TMemoryView<void>;
+
+/** A non-owning const view of a contiguous region of memory. */
+using FConstMemoryView = TMemoryView<const void>;
