Added a BLAKE3 integration to Core

#rn minor Core

diff --git a/Engine/Source/Runtime/Core/Core.Build.cs b/Engine/Source/Runtime/Core/Core.Build.cs
--- a/Engine/Source/Runtime/Core/Core.Build.cs
+++ b/Engine/Source/Runtime/Core/Core.Build.cs
@@ -14,6 +14,8 @@ public class Core : ModuleRules
 
 		PrivateDependencyModuleNames.Add("BuildSettings");
 
+		PrivateDependencyModuleNames.Add("BLAKE3");
+
 		PublicDependencyModuleNames.Add("TraceLog");
 		PublicIncludePaths.Add("Runtime/TraceLog/Public");
 
diff --git a/Engine/Source/Runtime/Core/Private/Misc/Blake3.cpp b/Engine/Source/Runtime/Core/Private/Misc/Blake3.cpp
new file mode 100644
--- /dev/null
+++ b/Engine/Source/Runtime/Core/Private/Misc/Blake3.cpp
@@ -0,0 +1,63 @@
+// Copyright Epic Games, Inc. All Rights Reserved.
+
+#include "Misc/Blake3.h"
+
+#include "Containers/StringView.h"
+#include "Memory/MemoryView.h"
+#include "Serialization/Archive.h"
+#include "String/HexToBytes.h"
+
+#include "blake3.h"
+
+///////////////////////////////////////////////////////////////////////////////////////////////////////////////////////
+
+FBlake3Hash::FBlake3Hash(const FAnsiStringView HexHash)
+{
+	check(HexHash.Len() == sizeof(ByteArray) * 2);
+	UE::String::HexToBytes(HexHash, Hash);
+}
+
+FBlake3Hash::FBlake3Hash(const FWideStringView HexHash)
+{
+	check(HexHash.Len() == sizeof(ByteArray) * 2);
+	UE::String::HexToBytes(HexHash, Hash);
+}
+
+FArchive& operator<<(FArchive& Ar, FBlake3Hash& Hash)
+{
+	Ar.Serialize(Hash.Hash, sizeof(Hash.Hash));
+	return Ar;
+}
+
+///////////////////////////////////////////////////////////////////////////////////////////////////////////////////////
+
+static_assert(sizeof(FBlake3) == sizeof(blake3_hasher), "Adjust the allocation in FBlake3 to match blake3_hasher");
+
+void FBlake3::Reset()
+{
+	blake3_hasher& Hasher = reinterpret_cast<blake3_hasher&>(HasherBytes);
+	blake3_hasher_init(&Hasher);
+}
+
+void FBlake3::Update(const void* Data, uint64 Size)
+{
+	blake3_hasher& Hasher = reinterpret_cast<blake3_hasher&>(HasherBytes);
+	blake3_hasher_update(&Hasher, Data, Size);
+}
+
+FBlake3Hash FBlake3::Finalize() const
+{
+	uint8 Output[BLAKE3_OUT_LEN];
+	const blake3_hasher& Hasher = reinterpret_cast<const blake3_hasher&>(HasherBytes);
+	blake3_hasher_finalize(&Hasher, Output, BLAKE3_OUT_LEN);
+	return FBlake3Hash(Output);
+}
+
+FBlake3Hash FBlake3::HashBuffer(const void* Data, uint64 Size)
+{
+	FBlake3 Hash;
+	Hash.Update(Data, Size);
+	return Hash.Finalize();
+}
+
+///////////////////////////////////////////////////////////////////////////////////////////////////////////////////////
diff --git a/Engine/Source/Runtime/Core/Private/Tests/Misc/Blake3Test.cpp b/Engine/Source/Runtime/Core/Private/Tests/Misc/Blake3Test.cpp
new file mode 100644
--- /dev/null
+++ b/Engine/Source/Runtime/Core/Private/Tests/Misc/Blake3Test.cpp
@@ -0,0 +1,65 @@
+// Copyright Epic Games, Inc. All Rights Reserved.
+
+#include "Misc/Blake3.h"
+
+#include "Misc/AutomationTest.h"
+#include "Misc/StringBuilder.h"
+
+#if WITH_DEV_AUTOMATION_TESTS
+
+IMPLEMENT_SIMPLE_AUTOMATION_TEST(FBlake3Test, "System.Core.Misc.Blake3", EAutomationTestFlags::ApplicationContextMask | EAutomationTestFlags::SmokeFilter)
+bool FBlake3Test::RunTest(const FString& Parameters)
+{
+	struct FBlake3TestCase
+	{
+		uint32 InputLength;
+		FStringView Hash;
+	};
+	constexpr FBlake3TestCase TestCases[] =
+	{
+		{     0u, TEXT("af1349b9f5f9a1a6a0404dea36dcc9499bcb25c9adc112b7cc9a93cae41f3262"_SV)},
+		{     1u, TEXT("2d3adedff11b61f14c886e35afa036736dcd87a74d27b5c1510225d0f592e213"_SV)},
+		{ 1'023u, TEXT("10108970eeda3eb932baac1428c7a2163b0e924c9a9e25b35bba72b28f70bd11"_SV)},
+		{ 1'024u, TEXT("42214739f095a406f3fc83deb889744ac00df831c10daa55189b5d121c855af7"_SV)},
+		{ 1'025u, TEXT("d00278ae47eb27b34faecf67b4fe263f82d5412916c1ffd97c8cb7fb814b8444"_SV)},
+		{ 2'048u, TEXT("e776b6028c7cd22a4d0ba182a8bf62205d2ef576467e838ed6f2529b85fba24a"_SV)},
+		{ 2'049u, TEXT("5f4d72f40d7a5f82b15ca2b2e44b1de3c2ef86c426c95c1af0b6879522563030"_SV)},
+		{ 3'072u, TEXT("b98cb0ff3623be03326b373de6b9095218513e64f1ee2edd2525c7ad1e5cffd2"_SV)},
+		{ 3'073u, TEXT("7124b49501012f81cc7f11ca069ec9226cecb8a2c850cfe644e327d22d3e1cd3"_SV)},
+		{ 4'096u, TEXT("015094013f57a5277b59d8475c0501042c0b642e531b0a1c8f58d2163229e969"_SV)},
+		{ 4'097u, TEXT("9b4052b38f1c5fc8b1f9ff7ac7b27cd242487b3d890d15c96a1c25b8aa0fb995"_SV)},
+		{ 5'120u, TEXT("9cadc15fed8b5d854562b26a9536d9707cadeda9b143978f319ab34230535833"_SV)},
+		{ 5'121u, TEXT("628bd2cb2004694adaab7bbd778a25df25c47b9d4155a55f8fbd79f2fe154cff"_SV)},
+		{ 6'144u, TEXT("3e2e5b74e048f3add6d21faab3f83aa44d3b2278afb83b80b3c35164ebeca205"_SV)},
+		{ 6'145u, TEXT("f1323a8631446cc50536a9f705ee5cb619424d46887f3c376c695b70e0f0507f"_SV)},
+		{ 7'168u, TEXT("61da957ec2499a95d6b8023e2b0e604ec7f6b50e80a9678b89d2628e99ada77a"_SV)},
+		{ 7'169u, TEXT("a003fc7a51754a9b3c7fae0367ab3d782dccf28855a03d435f8cfe74605e7817"_SV)},
+		{ 8'192u, TEXT("aae792484c8efe4f19e2ca7d371d8c467ffb10748d8a5a1ae579948f718a2a63"_SV)},
+		{ 8'193u, TEXT("bab6c09cb8ce8cf459261398d2e7aef35700bf488116ceb94a36d0f5f1b7bc3b"_SV)},
+		{16'384u, TEXT("f875d6646de28985646f34ee13be9a576fd515f76b5b0a26bb324735041ddde4"_SV)},
+		{31'744u, TEXT("62b6960e1a44bcc1eb1a611a8d6235b6b4b78f32e7abc4fb4c6cdcce94895c47"_SV)},
+	};
+
+	constexpr uint32 MaxInputLength = TestCases[UE_ARRAY_COUNT(TestCases) - 1].InputLength;
+	uint8 Input[MaxInputLength];
+	for (uint32 Index = 0; Index < MaxInputLength; ++Index)
+	{
+		Input[Index] = uint8(Index % 251);
+	}
+
+	FBlake3 Hash;
+	uint32 InputIndex = 0;
+	for (const FBlake3TestCase& TestCase : TestCases)
+	{
+		check(InputIndex <= TestCase.InputLength);
+		Hash.Update(Input + InputIndex, TestCase.InputLength - InputIndex);
+		InputIndex = TestCase.InputLength;
+		TStringBuilder<70> HashHex;
+		HashHex << Hash.Finalize();
+		TestEqual(FString::Printf(TEXT("BLAKE3(%u)"), TestCase.InputLength), FStringView(HashHex), TestCase.Hash);
+	}
+
+	return true;
+}
+
+#endif //WITH_DEV_AUTOMATION_TESTS
diff --git a/Engine/Source/Runtime/Core/Public/Misc/Blake3.h b/Engine/Source/Runtime/Core/Public/Misc/Blake3.h
new file mode 100644
--- /dev/null
+++ b/Engine/Source/Runtime/Core/Public/Misc/Blake3.h
@@ -0,0 +1,98 @@
+// Copyright Epic Games, Inc. All Rights Reserved.
+
+#pragma once
+
+#include "Containers/StringFwd.h"
+#include "HAL/UnrealMemory.h"
+#include "String/BytesToHex.h"
+#include "Templates/TypeCompatibleBytes.h"
+
+class FArchive;
+
+///////////////////////////////////////////////////////////////////////////////////////////////////////////////////////
+
+/** Stores a BLAKE3 hash. */
+struct FBlake3Hash
+{
+public:
+	using ByteArray = uint8[32];
+
+	FBlake3Hash() = default;
+	explicit FBlake3Hash(const ByteArray& Hash);
+
+	CORE_API explicit FBlake3Hash(FAnsiStringView HexHash);
+	CORE_API explicit FBlake3Hash(FWideStringView HexHash);
+
+	CORE_API friend FArchive& operator<<(FArchive& Ar, FBlake3Hash& Hash);
+
+	inline const ByteArray& Bytes() const { return Hash; }
+
+private:
+	alignas(uint32) ByteArray Hash{};
+};
+
+///////////////////////////////////////////////////////////////////////////////////////////////////////////////////////
+
+/** Calculates a BLAKE3 hash. */
+class FBlake3
+{
+public:
+	inline FBlake3() { Reset(); }
+
+	FBlake3(const FBlake3&) = delete;
+	FBlake3& operator=(const FBlake3&) = delete;
+
+	/** Reset to the default state in which no input has been written. */
+	CORE_API void Reset();
+
+	/** Add the data as input to the hash. May be called any number of times. */
+	CORE_API void Update(const void* Data, uint64 Size);
+
+	/** Finalize the hash of the input data. May be called any number of times, and more input may be added after. */
+	CORE_API FBlake3Hash Finalize() const;
+
+	/** Calculate the hash of the input data. */
+	CORE_API static FBlake3Hash HashBuffer(const void* Data, uint64 Size);
+
+private:
+	TAlignedBytes<1912, 8> HasherBytes;
+};
+
+///////////////////////////////////////////////////////////////////////////////////////////////////////////////////////
+
+inline FBlake3Hash::FBlake3Hash(const ByteArray& InHash)
+{
+	FMemory::Memcpy(Hash, InHash, sizeof(ByteArray));
+}
+
+inline bool operator==(const FBlake3Hash& A, const FBlake3Hash& B)
+{
+	return FMemory::Memcmp(A.Bytes(), B.Bytes(), sizeof(A.Bytes())) == 0;
+}
+
+inline bool operator!=(const FBlake3Hash& A, const FBlake3Hash& B)
+{
+	return FMemory::Memcmp(A.Bytes(), B.Bytes(), sizeof(A.Bytes())) != 0;
+}
+
+inline bool operator<(const FBlake3Hash& A, const FBlake3Hash& B)
+{
+	return FMemory::Memcmp(A.Bytes(), B.Bytes(), sizeof(A.Bytes())) < 0;
+}
+
+inline uint32 GetTypeHash(const FBlake3Hash& Hash)
+{
+	return *reinterpret_cast<const uint32*>(Hash.Bytes());
+}
+
+inline FAnsiStringBuilderBase& operator<<(FAnsiStringBuilderBase& Builder, const FBlake3Hash& Hash)
+{
+	UE::String::BytesToHex(Hash.Bytes(), Builder);
+	return Builder;
+}
+
+inline FWideStringBuilderBase& operator<<(FWideStringBuilderBase& Builder, const FBlake3Hash& Hash)
+{
+	UE::String::BytesToHex(Hash.Bytes(), Builder);
+	return Builder;
+}
