Use std::numeric_limits to generate special FP constants for the ULP unit test.

#trivial
#rnx

diff --git a/Engine/Source/Runtime/Core/Private/Tests/Math/UnrealMathTest.cpp b/Engine/Source/Runtime/Core/Private/Tests/Math/UnrealMathTest.cpp
--- a/Engine/Source/Runtime/Core/Private/Tests/Math/UnrealMathTest.cpp
+++ b/Engine/Source/Runtime/Core/Private/Tests/Math/UnrealMathTest.cpp
@@ -1887,16 +1887,18 @@ bool FMathRoundHalfFromZeroTests::RunTest(const FString& Parameters)
 IMPLEMENT_SIMPLE_AUTOMATION_TEST(FIsNearlyEqualByULPTest, "System.Core.Math.IsNearlyEqualByULP", EAutomationTestFlags::ApplicationContextMask | EAutomationTestFlags::SmokeFilter)
 bool FIsNearlyEqualByULPTest::RunTest(const FString& Parameters)
 {
-	static float FloatNan = FMath::Sqrt(-1.0f);
-	static double DoubleNan = double(FloatNan);
+	static const float FloatNan = std::numeric_limits<float>::quiet_NaN();
+	static const double DoubleNan = std::numeric_limits<double>::quiet_NaN();
 
-	static float FloatInf = 1.0f / 0.0f;
-	static double DoubleInf = 1.0 / 0.0;
+	static const float FloatInf = std::numeric_limits<float>::infinity();
+	static const double DoubleInf = std::numeric_limits<double>::infinity();
 
 	float FloatTrueMin;
 	double DoubleTrueMin;
 
-	// Construct our own true minimum float constants (aka FLT_TRUE_MIN), bypassing any value parsing.
+	// Construct our own true minimum float constants (aka std::numeric_limits<float>::denorm_min), 
+	// to ensure we don't get caught by any application or system-wide flush-to-zero or 
+	// denormals-are-zero settings.
 	{
 		uint32 FloatTrueMinInt = 0x00000001U;
 		uint64 DoubleTrueMinInt = 0x0000000000000001ULL;
