- ConfigCacheIni overhaul to allow for multiple platform's "GConfigs" at once.
- FConfigCacheIni::ForPlatform(PlatformName) is the way to get the in-memory (won't write to disk) ini system for a different platform. Can pass it GEngineIni, etc for the filename (  FConfigCacheIni::ForPlatform("IOS")->GetString(....., GEngineIni);  )
- GEngineIni and some others (the "stadnard, runtime ini files") are no longer filenames, but identifiers that will work with other platform Ini systems
- Modified "GETINI" console command to allow for a <Platform>@ prefix, to get another platform's ini value (using new multi-platform config) (GetIni [Platform@]IniFile:Section.SubSection Key)

diff --git a/Engine/Plugins/Developer/RenderDocPlugin/Source/RenderDocPlugin/Private/RenderDocPluginModule.cpp b/Engine/Plugins/Developer/RenderDocPlugin/Source/RenderDocPlugin/Private/RenderDocPluginModule.cpp
--- a/Engine/Plugins/Developer/RenderDocPlugin/Source/RenderDocPlugin/Private/RenderDocPluginModule.cpp
+++ b/Engine/Plugins/Developer/RenderDocPlugin/Source/RenderDocPlugin/Private/RenderDocPluginModule.cpp
@@ -334,18 +334,7 @@ void FRenderDocPluginModule::ShowNotification(const FText& Message, bool bForceN
 void FRenderDocPluginModule::InjectDebugExecKeybind()
 {
 	// Inject our key bind into the debug execs
-	FConfigFile* ConfigFile = nullptr;
-	// Look for the first matching INI file entry
-	for (TMap<FString, FConfigFile>::TIterator It(*GConfig); It; ++It)
-	{
-		if (It.Key().EndsWith(TEXT("Input.ini")))
-		{
-			ConfigFile = &It.Value();
-			break;
-		}
-	}
-	check(ConfigFile != nullptr);
-	FConfigSection* Section = ConfigFile->Find(TEXT("/Script/Engine.PlayerInput"));
+	FConfigSection* Section = GConfig->GetSectionPrivate(TEXT("/Script/Engine.PlayerInput"), false, false, GInputIni);
 	if (Section != nullptr)
 	{
 		Section->HandleAddCommand(TEXT("DebugExecBindings"), TEXT("(Key=F12,Command=\"RenderDoc.CaptureFrame\", Alt=true)"), false);
diff --git a/Engine/Source/Developer/LocalizationService/Private/LocalizationServiceHelpers.cpp b/Engine/Source/Developer/LocalizationService/Private/LocalizationServiceHelpers.cpp
--- a/Engine/Source/Developer/LocalizationService/Private/LocalizationServiceHelpers.cpp
+++ b/Engine/Source/Developer/LocalizationService/Private/LocalizationServiceHelpers.cpp
@@ -23,7 +23,7 @@ namespace LocalizationServiceHelpers
 			if (LocalizationServiceSettingsIni.Len() == 0)
 			{
 				const FString LocalizationServiceSettingsDir = FPaths::GeneratedConfigDir();
-				FConfigCacheIni::LoadGlobalIniFile(LocalizationServiceSettingsIni, TEXT("LocalizationServiceSettings"), NULL, false, false, true, *LocalizationServiceSettingsDir);
+				FConfigCacheIni::LoadGlobalIniFile(LocalizationServiceSettingsIni, TEXT("LocalizationServiceSettings"), NULL, false, false, true, true, *LocalizationServiceSettingsDir);
 			}
 			return LocalizationServiceSettingsIni;
 		}
@@ -35,7 +35,7 @@ namespace LocalizationServiceHelpers
 		if (LocalizationServiceGlobalSettingsIni.Len() == 0)
 		{
 			const FString LocalizationServiceSettingsDir = FPaths::EngineSavedDir() + TEXT("Config/");
-			FConfigCacheIni::LoadGlobalIniFile(LocalizationServiceGlobalSettingsIni, TEXT("LocalizationServiceSettings"), NULL, false, false, true, *LocalizationServiceSettingsDir);
+			FConfigCacheIni::LoadGlobalIniFile(LocalizationServiceGlobalSettingsIni, TEXT("LocalizationServiceSettings"), NULL, false, false, true, true, *LocalizationServiceSettingsDir);
 		}
 		return LocalizationServiceGlobalSettingsIni;
 	}
diff --git a/Engine/Source/Developer/ScreenShotComparisonTools/Private/ScreenShotManager.cpp b/Engine/Source/Developer/ScreenShotComparisonTools/Private/ScreenShotManager.cpp
--- a/Engine/Source/Developer/ScreenShotComparisonTools/Private/ScreenShotManager.cpp
+++ b/Engine/Source/Developer/ScreenShotComparisonTools/Private/ScreenShotManager.cpp
@@ -611,9 +611,9 @@ void FScreenShotManager::BuildFallbackPlatformsListFromConfig(const UScreenShotC
 	// legacy path
 	if (GConfig)
 	{
-		for (const TPair<FString,FConfigFile>& Config : *GConfig)
+		for (const FString& ConfigFilename : GConfig->GetFilenames())
 		{
-			FConfigSection* FallbackSection = GConfig->GetSectionPrivate(TEXT("AutomationTestFallbackHierarchy"), false, true, Config.Key);
+			FConfigSection* FallbackSection = GConfig->GetSectionPrivate(TEXT("AutomationTestFallbackHierarchy"), false, true, ConfigFilename);
 			if (FallbackSection)
 			{
 				UE_LOG(LogScreenShotManager, Warning, TEXT("Please move FallbackPlatform entries in [AutomationTestFallbackHierarchy] to +ScreenshotFallbackPlatforms= under[/Script/ScreenShotComparisonTools.ScreenShotComparisonSettings] in DefaultEngine.ini"));
diff --git a/Engine/Source/Developer/SourceControl/Private/SourceControlHelpers.cpp b/Engine/Source/Developer/SourceControl/Private/SourceControlHelpers.cpp
--- a/Engine/Source/Developer/SourceControl/Private/SourceControlHelpers.cpp
+++ b/Engine/Source/Developer/SourceControl/Private/SourceControlHelpers.cpp
@@ -1154,7 +1154,7 @@ const FString& USourceControlHelpers::GetSettingsIni()
 		if (SourceControlSettingsIni.Len() == 0)
 		{
 			const FString SourceControlSettingsDir = FPaths::GeneratedConfigDir();
-			FConfigCacheIni::LoadGlobalIniFile(SourceControlSettingsIni, TEXT("SourceControlSettings"), nullptr, false, false, true, *SourceControlSettingsDir);
+			FConfigCacheIni::LoadGlobalIniFile(SourceControlSettingsIni, TEXT("SourceControlSettings"), nullptr, false, false, true, true, *SourceControlSettingsDir);
 		}
 		return SourceControlSettingsIni;
 	}
@@ -1167,7 +1167,7 @@ const FString& USourceControlHelpers::GetGlobalSettingsIni()
 	if (SourceControlGlobalSettingsIni.Len() == 0)
 	{
 		const FString SourceControlSettingsDir = FPaths::EngineSavedDir() + TEXT("Config/");
-		FConfigCacheIni::LoadGlobalIniFile(SourceControlGlobalSettingsIni, TEXT("SourceControlSettings"), nullptr, false, false, true, *SourceControlSettingsDir);
+		FConfigCacheIni::LoadGlobalIniFile(SourceControlGlobalSettingsIni, TEXT("SourceControlSettings"), nullptr, false, false, true, true, *SourceControlSettingsDir);
 	}
 	return SourceControlGlobalSettingsIni;
 }
diff --git a/Engine/Source/Editor/DetailCustomizations/Private/ConfigEditorPropertyDetails.cpp b/Engine/Source/Editor/DetailCustomizations/Private/ConfigEditorPropertyDetails.cpp
--- a/Engine/Source/Editor/DetailCustomizations/Private/ConfigEditorPropertyDetails.cpp
+++ b/Engine/Source/Editor/DetailCustomizations/Private/ConfigEditorPropertyDetails.cpp
@@ -89,8 +89,6 @@ void FConfigPropertyHelperDetails::OnPropertyValueChanged(UObject* Object, FProp
 			OriginalProperty->GetOwnerClass()->GetDefaultObject()->SaveConfig(CPF_Config, *ConfigIniName, &Config);
 
 			// Take the saved section for this object and have the config system process and write out the one property we care about.
-			ensureMsgf(Config.Num() == 1, TEXT("UObject::UpdateDefaultConfig() caused more files than expected in the Sandbox config cache!"));
-
 			TArray<FString> Keys;
 			NewFile.GetKeys(Keys);
 
diff --git a/Engine/Source/Editor/UnrealEd/Private/CookOnTheFlyServer.cpp b/Engine/Source/Editor/UnrealEd/Private/CookOnTheFlyServer.cpp
--- a/Engine/Source/Editor/UnrealEd/Private/CookOnTheFlyServer.cpp
+++ b/Engine/Source/Editor/UnrealEd/Private/CookOnTheFlyServer.cpp
@@ -4408,11 +4408,12 @@ bool UCookOnTheFlyServer::IniSettingsOutOfDate(const ITargetPlatform* TargetPlat
 		FName FileFName = FName(*Filename);
 		if ( !ConfigFile )
 		{
-			for( const auto& File : *GConfig )
+			for( const FString& ConfigFilename : GConfig->GetFilenames() )
 			{
-				if (File.Value.Name == FileFName)
+				FConfigFile* File = GConfig->FindConfigFile(ConfigFilename);
+				if (File->Name == FileFName)
 				{
-					ConfigFile = &File.Value;
+					ConfigFile = File;
 					break;
 				}
 			}
diff --git a/Engine/Source/Runtime/Core/Private/Misc/AutomationTest.cpp b/Engine/Source/Runtime/Core/Private/Misc/AutomationTest.cpp
--- a/Engine/Source/Runtime/Core/Private/Misc/AutomationTest.cpp
+++ b/Engine/Source/Runtime/Core/Private/Misc/AutomationTest.cpp
@@ -449,9 +449,9 @@ void FAutomationTestFramework::BuildTestBlacklistFromConfig()
 
 		const FString CommandLine = FCommandLine::Get();
 
-		for (const TPair<FString, FConfigFile>& Config : *GConfig)
+		for (const FString& ConfigFilename : GConfig->GetFilenames())
 		{
-			FConfigSection* BlacklistSection = GConfig->GetSectionPrivate(TEXT("AutomationTestBlacklist"), false, true, Config.Key);
+			FConfigSection* BlacklistSection = GConfig->GetSectionPrivate(TEXT("AutomationTestBlacklist"), false, true, ConfigFilename);
 			if (BlacklistSection)
 			{
 				// Parse all blacklist definitions of the format "BlacklistTest=(Map=/Game/Tests/MapName, Test=TestName, Reason="Foo")"
diff --git a/Engine/Source/Runtime/Core/Private/Misc/ConfigCacheIni.cpp b/Engine/Source/Runtime/Core/Private/Misc/ConfigCacheIni.cpp
--- a/Engine/Source/Runtime/Core/Private/Misc/ConfigCacheIni.cpp
+++ b/Engine/Source/Runtime/Core/Private/Misc/ConfigCacheIni.cpp
@@ -21,6 +21,8 @@
 #include "Misc/StringBuilder.h"
 #include "Serialization/MemoryReader.h"
 #include "Serialization/MemoryWriter.h"
+#include "Serialization/LargeMemoryReader.h"
+#include "Async/Async.h"
 
 #if WITH_EDITOR
 	#define INI_CACHE 1
@@ -2075,9 +2077,17 @@ FConfigCacheIni::~FConfigCacheIni()
 	Flush( 1 );
 }
 
+
 FConfigFile* FConfigCacheIni::FindConfigFile( const FString& Filename )
 {
-	return TMap<FString,FConfigFile>::Find( Filename );
+	// look for a known file, if there's no ini extension
+	FConfigFile* Result = Filename.EndsWith(TEXT(".ini")) ? nullptr : KnownFiles.GetMutableFile(FName(*Filename));
+
+	if (Result == nullptr)
+	{
+		Result = Super::Find(Filename);
+	}
+	return Result;
 }
 
 FConfigFile* FConfigCacheIni::Find(const FString& Filename)
@@ -2088,8 +2098,9 @@ FConfigFile* FConfigCacheIni::Find(const FString& Filename)
 		return nullptr;
 	}
 
-	// Get file.
-	FConfigFile* Result = TMap<FString,FConfigFile>::Find(Filename);
+	// Get the file if it exists
+	FConfigFile* Result = FindConfigFile(Filename);
+
 	// this is || filesize so we load up .int files if file IO is allowed
 	if (!Result && !bAreFileOperationsDisabled)
 	{
@@ -2101,6 +2112,7 @@ FConfigFile* FConfigCacheIni::Find(const FString& Filename)
 			UE_LOG(LogConfig, Verbose, TEXT("GConfig::Find has loaded file:  %s"), *Filename);
 		}
 	}
+
 	return Result;
 }
 
@@ -2122,6 +2134,20 @@ FConfigFile* FConfigCacheIni::FindConfigFileWithBaseName(FName BaseName)
 	return nullptr;
 }
 
+TArray<FString> FConfigCacheIni::GetFilenames()
+{
+	TArray<FString> Result;
+	Super::GetKeys(Result);
+
+	for (const FConfigCacheIni::FKnownConfigFiles::FKnownConfigFile& File : KnownFiles.Files)
+	{
+		Result.Add(File.IniName.ToString());
+	}
+
+	return Result;
+}
+
+
 void FConfigCacheIni::Flush( bool Read, const FString& Filename )
 {
 	// never Flush temporary cache objects
@@ -2140,7 +2166,17 @@ void FConfigCacheIni::Flush( bool Read, const FString& Filename )
 				It.Value().Write(*It.Key());
 			}
 		}
+
+		// now flush the known files (all or a single file)
+		for (FConfigCacheIni::FKnownConfigFiles::FKnownConfigFile& File : KnownFiles.Files)
+		{
+			if (Filename.Len() == 0 || Filename == File.IniName.ToString())
+			{
+				File.IniFile.Write(*File.IniPath);
+			}
+		}
 	}
+
 	if( Read )
 	{
 		// we can't read it back in if file operations are disabled
@@ -2698,6 +2734,22 @@ void FConfigCacheIni::Exit()
 	Flush( 1 );
 }
 
+static void DumpFile(FOutputDevice& Ar, const FString& Filename, const FConfigFile& File)
+{
+	Ar.Logf(TEXT("FileName: %s"), *Filename);
+	for (FConfigFile::TConstIterator FileIt(File); FileIt; ++FileIt)
+	{
+		const FConfigSection& Sec = FileIt.Value();
+		Ar.Logf(TEXT("   [%s]"), *FileIt.Key());
+		for (FConfigSectionMap::TConstIterator SecIt(Sec); SecIt; ++SecIt)
+		{
+			Ar.Logf(TEXT("   %s=%s"), *SecIt.Key().ToString(), *SecIt.Value().GetValue());
+		}
+
+		Ar.Log(LINE_TERMINATOR);
+	}
+}
+
 void FConfigCacheIni::Dump(FOutputDevice& Ar, const TCHAR* BaseIniName)
 {
 	if (BaseIniName == nullptr)
@@ -2706,23 +2758,16 @@ void FConfigCacheIni::Dump(FOutputDevice& Ar, const TCHAR* BaseIniName)
 		TMap<FString,FConfigFile>::Dump( Ar );
 	}
 
+	for (const FConfigCacheIni::FKnownConfigFiles::FKnownConfigFile& File : KnownFiles.Files)
+	{
+		DumpFile(Ar, File.IniPath, File.IniFile);
+	}
+
 	for ( TIterator It(*this); It; ++It )
 	{
 		if (BaseIniName == nullptr || FPaths::GetBaseFilename(It.Key()) == BaseIniName)
 		{
-			Ar.Logf(TEXT("FileName: %s"), *It.Key());
-			FConfigFile& File = It.Value();
-			for ( FConfigFile::TIterator FileIt(File); FileIt; ++FileIt )
-			{
-				FConfigSection& Sec = FileIt.Value();
-				Ar.Logf(TEXT("   [%s]"), *FileIt.Key());
-				for (FConfigSectionMap::TConstIterator SecIt(Sec); SecIt; ++SecIt)
-				{
-					Ar.Logf(TEXT("   %s=%s"), *SecIt.Key().ToString(), *SecIt.Value().GetValue());
-				}
-
-				Ar.Log(LINE_TERMINATOR);
-			}
+			DumpFile(Ar, It.Key(), It.Value());
 		}
 	}
 }
@@ -3593,6 +3638,9 @@ static FString PerformFinalExpansions(const FString InString, const FString& Pla
 	static FString PlatformExtensionProjectDir;
 	static FString ProjectNotForLicenseesDir;
 	static FString ProjectNoRedistDir;
+	static FCriticalSection ExpansionSection;
+
+	FScopeLock Lock(&ExpansionSection);
 
 	if (LastPlatform != PlatformName)
 	{
@@ -3800,13 +3848,14 @@ void FConfigCacheIni::SaveCurrentStateForBootstrap(const TCHAR* Filename)
 	FFileHelper::SaveArrayToFile(FileContent, Filename);
 }
 
-FArchive& operator<<(FArchive& Ar, FConfigCacheIni& ConfigCacheIni)
+void FConfigCacheIni::Serialize(FArchive& Ar)
 {
-	Ar << static_cast<FConfigCacheIni::Super&>(ConfigCacheIni);
-	Ar << ConfigCacheIni.bAreFileOperationsDisabled;
-	Ar << ConfigCacheIni.bIsReadyForUse;
-	Ar << ConfigCacheIni.Type;
-	return Ar;
+	FConfigCacheIni::Super* Parent = static_cast<FConfigCacheIni::Super*>(this);
+	Ar << *Parent;
+	Ar << KnownFiles;
+	Ar << bAreFileOperationsDisabled;
+	Ar << bIsReadyForUse;
+	Ar << Type;
 }
 
 void FConfigCacheIni::SerializeStateForBootstrap_Impl(FArchive& Ar)
@@ -3817,7 +3866,7 @@ void FConfigCacheIni::SerializeStateForBootstrap_Impl(FArchive& Ar)
 	// the same binary executable for both the parent and 
 	// children processes. It also takes care of saving/restoring
 	// global ini variables.
-	Ar << *this;
+	Serialize(Ar);
 	Ar << GEditorIni;
 	Ar << GEditorKeyBindingsIni;
 	Ar << GEditorLayoutIni;
@@ -3834,8 +3883,157 @@ void FConfigCacheIni::SerializeStateForBootstrap_Impl(FArchive& Ar)
 	Ar << GEngineIni;
 }
 
+
+bool FConfigCacheIni::InitializeKnownConfigFiles(const TCHAR* PlatformName, bool bDefaultEngineIniRequired)
+{
+	// check for scalability platform override.
+	const TCHAR* ScalabilityPlatformOverride = nullptr;
+#if !UE_BUILD_SHIPPING && WITH_EDITOR
+	if (PlatformName == nullptr)
+	{
+		FString ScalabilityPlatformOverrideCommandLine;
+		FParse::Value(FCommandLine::Get(), TEXT("ScalabilityIniPlatformOverride="), ScalabilityPlatformOverrideCommandLine);
+		ScalabilityPlatformOverride = ScalabilityPlatformOverrideCommandLine.Len() ? *ScalabilityPlatformOverrideCommandLine : nullptr;
+	}
+#endif
+
+	// when initializing GConfig, allow standard behavior (PlatformName will be null in that case)
+	bool bAllowGeneratedIniWhenCooked = PlatformName == nullptr;
+	bool bAllowRemoteConfig = PlatformName == nullptr;
+
+	bool bEngineConfigCreated = false;
+	for (uint8 KnownIndex = 0; KnownIndex < (uint8)EKnownIniFile::NumKnownFiles; KnownIndex++)
+	{
+		// allo for scalability to come from another 
+		const TCHAR* PlatformNameForThisIni = KnownIndex == (uint8)EKnownIniFile::Scalability ? ScalabilityPlatformOverride : PlatformName;
+		// load the hierarchy!
+		bool bConfigCreated = FConfigCacheIni::LoadGlobalIniFile(KnownFiles.Files[KnownIndex].IniPath, *KnownFiles.Files[KnownIndex].IniName.ToString(), PlatformName, 
+			bDefaultEngineIniRequired, false, bAllowGeneratedIniWhenCooked, bAllowRemoteConfig, *FPaths::GeneratedConfigDir(), this);
+
+		// we want to return if the Engine config was successfully created (to not remove any functionality from old code)
+		if (KnownIndex == (uint8)EKnownIniFile::Engine)
+		{
+			bEngineConfigCreated = bConfigCreated;
+		}
+	}
+
+	bIsReadyForUse = true;
+
+	return bEngineConfigCreated;
+}
+
+
+const FConfigFile* FConfigCacheIni::FKnownConfigFiles::GetFile(FName Name)
+{
+	// sharing logic
+	return GetMutableFile(Name);
+}
+
+FConfigFile* FConfigCacheIni::FKnownConfigFiles::GetMutableFile(FName Name)
+{
+	// walk the list of files looking for matching FName (a TMap was a bit slower)
+	for (FKnownConfigFile& File : Files)
+	{
+		if (File.IniName == Name)
+		{
+			return &File.IniFile;
+		}
+	}
+
+	return nullptr;
+}
+const FString& FConfigCacheIni::FKnownConfigFiles::GetFilename(FName Name)
+{
+	for (FKnownConfigFile& File : Files)
+	{
+		if (File.IniName == Name)
+		{
+			return File.IniPath;
+		}
+	}
+
+	static FString Empty;
+	return Empty;
+}
+
+FConfigCacheIni::FKnownConfigFiles::FKnownConfigFiles()
+{
+	// set the FNames associated with each file
+
+	// 	Files[(uint8)EKnownIniFile::Engine].IniName = FName("Engine");
+	#define SET_KNOWN_NAME(Ini) Files[(uint8)EKnownIniFile::Ini].IniName = FName(#Ini);
+		ENUMERATE_KNOWN_INI_FILES(SET_KNOWN_NAME);
+	#undef SET_KNOWN_NAME
+}
+
+FArchive& operator<<(FArchive& Ar, FConfigCacheIni::FKnownConfigFiles& Names)
+{
+	for (FConfigCacheIni::FKnownConfigFiles::FKnownConfigFile& File : Names.Files)
+	{
+		Ar << File.IniPath << File.IniFile;
+	}
+
+	return Ar;
+}
+
+#if PRELOAD_BINARY_CONFIG
+
+#include "Misc/PreLoadFile.h"
+static FPreLoadFile GPreLoadConfigBin(TEXT("{PROJECT}Config/Config.bin"));
+
+bool FConfigCacheIni::CreateGConfigFromSaved(const TCHAR* Filename)
+{
+	SCOPED_BOOT_TIMING("FConfigCacheIni::CreateGConfigFromSaved");
+	// get the already loaded file
+	int64 Size;
+	void* PreloadedData = GPreLoadConfigBin.TakeOwnershipOfLoadedData(&Size);
+	if (PreloadedData == nullptr)
+	{
+		return false;
+	}
+
+	UE_LOG(LogInit, Display, TEXT("Loading binary GConfig..."));
+
+	// serialize right out of the preloaded data
+	FLargeMemoryReader MemoryReader((uint8*)PreloadedData, Size);
+	FKnownConfigFiles Names;
+	GConfig = new FConfigCacheIni(EConfigCacheType::Temporary);
+
+	// make an object that we can use to pass to delegates for any extra binary data they want to write
+//	FCoreDelegates::FExtraBinaryConfigData ExtraData(*GConfig, false);
+
+	GConfig->Serialize(MemoryReader);
+	MemoryReader << Names;// << ExtraData.Data;
+
+	// now let the delegates pull their data out, after GConfig is set up
+//	FCoreDelegates::AccessExtraBinaryConfigData.Broadcast(ExtraData);
+
+	FCoreDelegates::ConfigReadyForUse.Broadcast();
+
+	FMemory::Free(PreloadedData);
+	return true;
+}
+
+#endif
+
 void FConfigCacheIni::InitializeConfigSystem()
 {
+	// assign the G***Ini strings for the known ini's
+	#define ASSIGN_GLOBAL_INI_STRING(IniName) G##IniName##Ini = FString(#IniName);
+		// GEngineIni = FString("Engine")
+		ENUMERATE_KNOWN_INI_FILES(ASSIGN_GLOBAL_INI_STRING);
+	#undef ASSIGN_GLOBAL_INI_STRING
+
+#if PRELOAD_BINARY_CONFIG
+	// attempt to load from staged binary config data
+	FString BinaryConfigFile = FPaths::Combine(FPaths::SourceConfigDir(), TEXT("Config.bin"));
+	if (IFileManager::Get().FileExists(*BinaryConfigFile) &&
+		FConfigCacheIni::CreateGConfigFromSaved(*BinaryConfigFile))
+	{
+		return;
+	}
+#endif
+
 	// Bootstrap the Ini config cache
 	FString IniBootstrapFilename;
 	if (FParse::Value( FCommandLine::Get(), TEXT("IniBootstrap="), IniBootstrapFilename))
@@ -3856,6 +4054,8 @@ void FConfigCacheIni::InitializeConfigSystem()
 		}
 	}
 
+	UE_LOG(LogInit, Display, TEXT("Loading text-based GConfig..."));
+
 	// Perform any upgrade we need before we load any configuration files
 	FConfigManifest::UpgradeFromPreviousVersions();
 
@@ -3865,8 +4065,11 @@ void FConfigCacheIni::InitializeConfigSystem()
 	// load the main .ini files (unless we're running a program or a gameless UnrealEditor.exe, DefaultEngine.ini is required).
 	const bool bIsGamelessExe = !FApp::HasProjectName();
 	const bool bDefaultEngineIniRequired = !bIsGamelessExe && (GIsGameAgnosticExe || FApp::IsProjectNameEmpty());
-	bool bEngineConfigCreated = FConfigCacheIni::LoadGlobalIniFile(GEngineIni, TEXT("Engine"), nullptr, bDefaultEngineIniRequired);
 
+	// load in the default ini files
+	bool bEngineConfigCreated = GConfig->InitializeKnownConfigFiles(nullptr, bDefaultEngineIniRequired);
+
+	// verify if needed
 	if ( !bIsGamelessExe )
 	{
 		// Now check and see if our game is correct if this is a game agnostic binary
@@ -3888,8 +4091,6 @@ void FConfigCacheIni::InitializeConfigSystem()
 		}
 	}
 
-	FConfigCacheIni::LoadGlobalIniFile(GGameIni, TEXT("Game"));
-	FConfigCacheIni::LoadGlobalIniFile(GInputIni, TEXT("Input"));
 #if WITH_EDITOR
 	// load some editor specific .ini files
 
@@ -3901,9 +4102,9 @@ void FConfigCacheIni::InitializeConfigSystem()
 
 	// Project agnostic editor ini files
 	const FString EditorSettingsDir = FPaths::EngineEditorSettingsDir();
-	FConfigCacheIni::LoadGlobalIniFile(GEditorSettingsIni, TEXT("EditorSettings"), nullptr, false, false, true, *EditorSettingsDir);
-	FConfigCacheIni::LoadGlobalIniFile(GEditorKeyBindingsIni, TEXT("EditorKeyBindings"), nullptr, false, false, true, *EditorSettingsDir);
-	FConfigCacheIni::LoadGlobalIniFile(GEditorLayoutIni, TEXT("EditorLayout"), nullptr, false, false, true, *EditorSettingsDir);
+	FConfigCacheIni::LoadGlobalIniFile(GEditorSettingsIni, TEXT("EditorSettings"), nullptr, false, false, true, true, *EditorSettingsDir);
+	FConfigCacheIni::LoadGlobalIniFile(GEditorKeyBindingsIni, TEXT("EditorKeyBindings"), nullptr, false, false, true, true, *EditorSettingsDir);
+	FConfigCacheIni::LoadGlobalIniFile(GEditorLayoutIni, TEXT("EditorLayout"), nullptr, false, false, true, true, *EditorSettingsDir);
 
 #endif
 #if PLATFORM_DESKTOP
@@ -3912,56 +4113,53 @@ void FConfigCacheIni::InitializeConfigSystem()
 	FConfigCacheIni::LoadGlobalIniFile(GLightmassIni, TEXT("Lightmass"));
 #endif
 
-	// check for scalability platform override.
-	const TCHAR* ScalabilityPlatformOverride = nullptr;
-#if !UE_BUILD_SHIPPING && WITH_EDITOR
-	FString ScalabilityPlatformOverrideCommandLine;
-	FParse::Value(FCommandLine::Get(), TEXT("ScalabilityIniPlatformOverride="), ScalabilityPlatformOverrideCommandLine);
-	ScalabilityPlatformOverride = ScalabilityPlatformOverrideCommandLine.Len() ? *ScalabilityPlatformOverrideCommandLine : nullptr;
-#endif
-
-	// Load scalability settings.
-	FConfigCacheIni::LoadGlobalIniFile(GScalabilityIni, TEXT("Scalability"), ScalabilityPlatformOverride);
-	// Load driver blacklist
-	FConfigCacheIni::LoadGlobalIniFile(GHardwareIni, TEXT("Hardware"));	
-	// Load runtime options
-	FConfigCacheIni::LoadGlobalIniFile(GRuntimeOptionsIni, TEXT("RuntimeOptions"));
-	// Load install bundle config
-	FConfigCacheIni::LoadGlobalIniFile(GInstallBundleIni, TEXT("InstallBundle"));
-	// Load device profiles for current platform
-	FConfigCacheIni::LoadGlobalIniFile(GDeviceProfilesIni, TEXT("DeviceProfiles"));
-
-
-
-	// Load user game settings .ini, allowing merging. This also updates the user .ini if necessary.
-	FConfigCacheIni::LoadGlobalIniFile(GGameUserSettingsIni, TEXT("GameUserSettings"));
 
 	// now we can make use of GConfig
 	GConfig->bIsReadyForUse = true;
 	FCoreDelegates::ConfigReadyForUse.Broadcast();
+
+#if WITH_EDITOR
+	AsyncInitializeConfigForPlatforms();
+#endif
 }
 
-bool FConfigCacheIni::LoadGlobalIniFile(FString& FinalIniFilename, const TCHAR* BaseIniName, const TCHAR* Platform, bool bForceReload, bool bRequireDefaultIni, bool bAllowGeneratedIniWhenCooked, const TCHAR* GeneratedConfigDir)
+bool FConfigCacheIni::LoadGlobalIniFile(FString& OutFinalIniFilename, const TCHAR* BaseIniName, const TCHAR* Platform, bool bForceReload, bool bRequireDefaultIni, bool bAllowGeneratedIniWhenCooked, bool bAllowRemoteConfig, const TCHAR* GeneratedConfigDir, FConfigCacheIni* ConfigSystem)
 {
 	// figure out where the end ini file is
-	FinalIniFilename = GetDestIniFilename(BaseIniName, Platform, GeneratedConfigDir);
-	
-	// Start the loading process for the remote config file when appropriate
-	if (FRemoteConfig::Get()->ShouldReadRemoteFile(*FinalIniFilename))
+	FString FinalIniFilename = GetDestIniFilename(BaseIniName, Platform, GeneratedConfigDir);
+
+	// if we are reloading a known ini file (where OutFinalIniFilename already has a value), then we need to leave the OutFinalFilename alone until we can remove LoadGlobalIniFile completely
+	if (OutFinalIniFilename.Len() > 0 && OutFinalIniFilename == BaseIniName)
 	{
-		FRemoteConfig::Get()->Read(*FinalIniFilename, BaseIniName);
+		// don't do anything to the output string
 	}
-
-	FRemoteConfigAsyncIOInfo* RemoteInfo = FRemoteConfig::Get()->FindConfig(*FinalIniFilename);
-	if (RemoteInfo && (!RemoteInfo->bWasProcessed || !FRemoteConfig::Get()->IsFinished(*FinalIniFilename)))
+	else
 	{
-		// Defer processing this remote config file to until it has finish its IO operation
-		return false;
+		OutFinalIniFilename = FinalIniFilename;
+	}
+	
+	if (bAllowRemoteConfig)
+	{
+		// Start the loading process for the remote config file when appropriate
+		if (FRemoteConfig::Get()->ShouldReadRemoteFile(*FinalIniFilename))
+		{
+			FRemoteConfig::Get()->Read(*FinalIniFilename, BaseIniName);
+		}
+
+		FRemoteConfigAsyncIOInfo* RemoteInfo = FRemoteConfig::Get()->FindConfig(*FinalIniFilename);
+		if (RemoteInfo && (!RemoteInfo->bWasProcessed || !FRemoteConfig::Get()->IsFinished(*FinalIniFilename)))
+		{
+			// Defer processing this remote config file to until it has finish its IO operation
+			return false;
+		}
 	}
 
+	// look for a known file to see if it's been loaded yet
+	FConfigFile* KnownFile = ConfigSystem->KnownFiles.GetMutableFile(BaseIniName);
+
 	// need to check to see if the file already exists in the GConfigManager's cache
 	// if it does exist then we are done, nothing else to do
-	if (!bForceReload && GConfig->FindConfigFile(*FinalIniFilename) != nullptr)
+	if (!bForceReload && ((KnownFile != nullptr && KnownFile->Num() > 0) || (KnownFile == nullptr && ConfigSystem->FindConfigFile(*FinalIniFilename) != nullptr)))
 	{
 		//UE_LOG(LogConfig, Log,  TEXT( "Request to load a config file that was already loaded: %s" ), GeneratedIniFile );
 		return true;
@@ -3973,7 +4171,7 @@ bool FConfigCacheIni::LoadGlobalIniFile(FString& FinalIniFilename, const TCHAR*
 	if (bForceReload) // If reloading we should preserve the existing config dirs
 	{
 		// If base ini, try to use an existing GConfig file to set the config directories instead of assuming defaults
-		FConfigFile* BaseConfig = GConfig->FindConfigFileWithBaseName(BaseIniName);
+		FConfigFile* BaseConfig = ConfigSystem->FindConfigFileWithBaseName(BaseIniName);
 		if (BaseConfig)
 		{
 			if (BaseConfig->SourceEngineConfigDir.Len())
@@ -3988,8 +4186,8 @@ bool FConfigCacheIni::LoadGlobalIniFile(FString& FinalIniFilename, const TCHAR*
 		}
 	}
 
-	// make a new entry in GConfig (overwriting what's already there)
-	FConfigFile& NewConfigFile = GConfig->Add(FinalIniFilename, FConfigFile());
+	// if this was a known file, use it directly, otherwise, add one and use it
+	FConfigFile& NewConfigFile = KnownFile != nullptr ? *KnownFile : ConfigSystem->Add(FinalIniFilename, FConfigFile());
 
 	return LoadExternalIniFile(
 		NewConfigFile, 
@@ -3999,7 +4197,7 @@ bool FConfigCacheIni::LoadGlobalIniFile(FString& FinalIniFilename, const TCHAR*
 		/*bIsBaseIniName*/ true,
 		Platform, 
 		bForceReload, 
-		/*bWriteDestIni*/ true,
+		/*bWriteDestIni*/ ConfigSystem->Type == EConfigCacheType::DiskBacked,
 		bAllowGeneratedIniWhenCooked, 
 		GeneratedConfigDir);
 }
@@ -4088,7 +4286,8 @@ bool FConfigCacheIni::LoadExternalIniFile(FConfigFile& ConfigFile, const TCHAR*
 		// now generate and make sure it's up to date (using IniName as a Base for an ini filename)
 		// @todo This bNeedsWrite afaict is always true even if it loaded a completely valid generated/final .ini, and the write below will
 		// just write out the exact same thing it read in!
-		bool bNeedsWrite = GenerateDestIniFile(ConfigFile, DestIniFilename, ConfigFile.SourceIniHierarchy, bAllowGeneratedIniWhenCooked, true);
+		// don't use the ini cache (via bUseHierarchyCache param) if we are not writing to disk - the cache is shared resource, and shouldn't be used for temp FConfigCacheInis
+		bool bNeedsWrite = GenerateDestIniFile(ConfigFile, DestIniFilename, ConfigFile.SourceIniHierarchy, bAllowGeneratedIniWhenCooked, bWriteDestIni);
 
 		ConfigFile.Name = IniName;
 
@@ -4608,18 +4807,18 @@ void ForEachCVarInSectionFromIni(const TCHAR* InSectionName, const TCHAR* InIniF
 	if (FConfigSection* Section = GConfig->GetSectionPrivate(InSectionName, false, true, InIniFilename))
 	{
 		for (FConfigSectionMap::TConstIterator It(*Section); It; ++It)
-		{
-			const FString& KeyString = It.Key().GetPlainNameString();
-			const FString& ValueString = ConvertValueFromHumanFriendlyValue(*It.Value().GetValue());
+	{
+		const FString& KeyString = It.Key().GetPlainNameString();
+		const FString& ValueString = ConvertValueFromHumanFriendlyValue(*It.Value().GetValue());
 
-			IConsoleVariable* CVar = IConsoleManager::Get().FindConsoleVariable(*KeyString);
-			if (CVar)
-			{
-				InEvaluationFunction(CVar, KeyString, ValueString);
-			}
+		IConsoleVariable* CVar = IConsoleManager::Get().FindConsoleVariable(*KeyString);
+		if (CVar)
+		{
+			InEvaluationFunction(CVar, KeyString, ValueString);
 		}
 	}
 }
+}
 
 
 void ApplyCVarSettingsGroupFromIni(const TCHAR* InSectionBaseName, int32 InGroupNumber, const TCHAR* InIniFilename, uint32 SetBy)
@@ -4873,14 +5072,14 @@ void RecordApplyCVarSettingsFromIni()
 void ReapplyRecordedCVarSettingsFromIni()
 {
 	// first we need to reload the inis 
-	for (TPair<FString, FConfigFile>& IniPair : *GConfig)
+	for (const FString& Filename : GConfig->GetFilenames())
 	{
-		FConfigFile& ConfigFile = IniPair.Value;
-		if (ConfigFile.Num() > 0)
+		FConfigFile* ConfigFile = GConfig->FindConfigFile(Filename);
+		if (ConfigFile->Num() > 0)
 		{
-			FName BaseName = ConfigFile.Name;
+			FName BaseName = ConfigFile->Name;
 			// Must call LoadLocalIniFile (NOT LoadGlobalIniFile) to preserve original enginedir/sourcedir for plugins
-			verify(FConfigCacheIni::LoadLocalIniFile(ConfigFile, *BaseName.ToString(), true, nullptr, true));
+			verify(FConfigCacheIni::LoadLocalIniFile(*ConfigFile, *BaseName.ToString(), true, nullptr, true));
 		}
 	}
 
@@ -4917,3 +5116,80 @@ void DeleteRecordedConfigReadsFromIni()
 	ConfigHistoryHelper = nullptr;
 #endif
 }
+
+
+
+
+
+static TMap<FName, FConfigCacheIni> GConfigForPlatform;
+static TFuture<void> GConfigFuture;
+static FCriticalSection GConfigForPlatformLock;
+
+#if WITH_EDITOR
+void FConfigCacheIni::AsyncInitializeConfigForPlatforms()
+{
+	GConfigFuture = Async(EAsyncExecution::ThreadPool, []
+	{
+		double Start = FPlatformTime::Seconds();
+
+		for (auto Pair : FDataDrivenPlatformInfoRegistry::GetAllPlatformInfos())
+		{
+			FName PlatformName = Pair.Key;
+
+			FConfigCacheIni& NewConfig = GConfigForPlatform.Emplace(PlatformName, EConfigCacheType::Temporary);
+
+			NewConfig.InitializeKnownConfigFiles(*PlatformName.ToString(), false);
+		}
+
+		UE_LOG(LogConfig, Display, TEXT("Loading platform ini files took %.2f seconds"), FPlatformTime::Seconds() - Start);
+	});
+}
+
+#endif
+
+#if WITH_UNREAL_DEVELOPER_TOOLS
+
+FConfigCacheIni* FConfigCacheIni::ForPlatform(FName PlatformName)
+{
+	// use GConfig when no platform is specified
+	if (PlatformName == NAME_None)
+	{
+		return GConfig;
+	}
+
+	FScopeLock Lock(&GConfigForPlatformLock);
+
+	// they are likely already loaded, but just block to make sure
+	if (GConfigFuture.IsValid())
+	{
+		double Start = FPlatformTime::Seconds();
+
+		GConfigFuture.Get();
+		GConfigFuture = TFuture<void>();
+
+		UE_LOG(LogConfig, Display, TEXT("Blocking on platform ini files took %.2f seconds"), FPlatformTime::Seconds() - Start);
+	}
+
+	// look for existing config
+	FConfigCacheIni* PlatformConfig = GConfigForPlatform.Find(PlatformName);
+
+	// read any missing platform configs now, on demand
+	if (PlatformConfig == nullptr)
+	{
+		PlatformConfig = &GConfigForPlatform.Emplace(PlatformName, EConfigCacheType::Temporary);
+		PlatformConfig->InitializeKnownConfigFiles(*PlatformName.ToString(), false);
+	}
+
+	return GConfigForPlatform.Find(PlatformName);
+}
+
+void FConfigCacheIni::ClearOtherPlatformConfigs()
+{
+	FScopeLock Lock(&GConfigForPlatformLock);
+
+	checkf(!GConfigFuture.IsValid(), TEXT("Trying to clear platform configs while still reading them in is unexpected"));
+
+	GConfigForPlatform.Empty();
+}
+
+#endif
diff --git a/Engine/Source/Runtime/Core/Public/Misc/Build.h b/Engine/Source/Runtime/Core/Public/Misc/Build.h
--- a/Engine/Source/Runtime/Core/Public/Misc/Build.h
+++ b/Engine/Source/Runtime/Core/Public/Misc/Build.h
@@ -437,3 +437,8 @@
 #ifndef GET_DEVICE_ID_UNAVAILABLE
 	#define GET_DEVICE_ID_UNAVAILABLE 0
 #endif
+
+// CONtrols whether or not to make a global object to load COnfig.bin as soon as possible
+#ifndef PRELOAD_BINARY_CONFIG
+	#define PRELOAD_BINARY_CONFIG 1
+#endif
\ No newline at end of file
diff --git a/Engine/Source/Runtime/Core/Public/Misc/ConfigCacheIni.h b/Engine/Source/Runtime/Core/Public/Misc/ConfigCacheIni.h
--- a/Engine/Source/Runtime/Core/Public/Misc/ConfigCacheIni.h
+++ b/Engine/Source/Runtime/Core/Public/Misc/ConfigCacheIni.h
@@ -26,6 +26,41 @@ CORE_API DECLARE_LOG_CATEGORY_EXTERN(LogConfig, Log, All);
 // Server builds should be tweakable even in Shipping
 #define ALLOW_INI_OVERRIDE_FROM_COMMANDLINE			(UE_SERVER || !(UE_BUILD_SHIPPING))
 #define CONFIG_REMEMBER_ACCESS_PATTERN (WITH_EDITOR || 0)
+
+
+///////////////////////////////////////////////////////////////////////////////
+//
+// This is the master list of known ini files that are used and processed 
+// on all platforms (specifically for runtime/binary speedups. Other, editor-
+// specific inis, or non-hierarchical ini files will still work with the 
+// old system, but they won't have any optimizations applied
+//
+// These should be listed in the order of highest to lowest use, for optimization
+//
+///////////////////////////////////////////////////////////////////////////////
+
+#define ENUMERATE_KNOWN_INI_FILES(op) \
+	op(Engine) \
+	op(Game) \
+	op(Input) \
+	op(DeviceProfiles) \
+	op(GameUserSettings) \
+	op(Scalability) \
+	op(RuntimeOptions) \
+	op(InstallBundle) \
+	op(Hardware)
+
+#define KNOWN_INI_ENUM(IniName) IniName,
+enum class EKnownIniFile : uint8
+{
+	// make an enum entry for each known ini above
+	ENUMERATE_KNOWN_INI_FILES(KNOWN_INI_ENUM)
+	
+	// conventient counter for the above list
+	NumKnownFiles,
+};
+
+
 struct FConfigValue
 {
 public:
@@ -503,7 +538,7 @@ enum class EConfigCacheType : uint8
 };
 
 // Set of all cached config files.
-class CORE_API FConfigCacheIni : public TMap<FString,FConfigFile>
+class CORE_API FConfigCacheIni : private TMap<FString,FConfigFile>
 {
 public:
 	// Basic functions.
@@ -587,6 +622,18 @@ public:
 	/** Finds Config file that matches the base name such as "Engine" */
 	FConfigFile* FindConfigFileWithBaseName(FName BaseName);
 
+	using Super = TMap<FString, FConfigFile>;
+	FConfigFile& Add(const FString& Filename, const FConfigFile& File)
+	{
+		return Super::Add(Filename, File);
+	}
+	int32 Remove(const FString& Filename)
+	{
+		return Super::Remove(Filename);
+	}
+	TArray<FString> GetFilenames();
+
+
 	void Flush( bool Read, const FString& Filename=TEXT("") );
 
 	void LoadFile( const FString& InFilename, const FConfigFile* Fallback = NULL, const TCHAR* PlatformString = NULL );
@@ -839,9 +886,6 @@ public:
 		const FString&	Filename
 	);
 
-	// Static allocator.
-	static FConfigCacheIni* Factory();
-
 	// Static helper functions
 
 	/**
@@ -879,7 +923,7 @@ public:
 	 * @param GeneratedConfigDir The location where generated config files are made.
 	 * @return true if the final ini was created successfully.
 	 */
-	static bool LoadGlobalIniFile(FString& FinalIniFilename, const TCHAR* BaseIniName, const TCHAR* Platform=NULL, bool bForceReload=false, bool bRequireDefaultIni=false, bool bAllowGeneratedIniWhenCooked=true, const TCHAR* GeneratedConfigDir = *FPaths::GeneratedConfigDir());
+	static bool LoadGlobalIniFile(FString& FinalIniFilename, const TCHAR* BaseIniName, const TCHAR* Platform = NULL, bool bForceReload = false, bool bRequireDefaultIni = false, bool bAllowGeneratedIniWhenCooked = true, bool bAllowRemoteConfig = true, const TCHAR* GeneratedConfigDir = *FPaths::GeneratedConfigDir(), FConfigCacheIni* ConfigSystem=GConfig);
 
 	/**
 	 * Load an ini file directly into an FConfigFile, and nothing is written to GConfig or disk. 
@@ -926,8 +970,74 @@ public:
 	 */
 	void SaveCurrentStateForBootstrap(const TCHAR* Filename);
 
-	friend FArchive& operator<<(FArchive& Ar, FConfigCacheIni& ConfigCacheIni);
+	void Serialize(FArchive& Ar);
+
+
+	struct FKnownConfigFiles
+	{
+		FKnownConfigFiles();
+
+		// Write out this for binary config serialization
+		friend FArchive& operator<<(FArchive& Ar, FKnownConfigFiles& Names);
+
+		// setup GEngineIni based on this structure's values
+		void SetGlobalIniStringsFromMembers();
+
+		// given an name ("Engine") return the FConfigFile for it
+		const FConfigFile* GetFile(FName Name);
+		// given an name ("Engine") return the modifiable FConfigFile for it
+		FConfigFile* GetMutableFile(FName Name);
+
+		// get the disk-based filename for the given known ini name
+		const FString& GetFilename(FName Name);
+
+
+		// create the list of members for the known inis (Engine, Game, etc) See the top of this file for the list
+		struct FKnownConfigFile
+		{
+			FName IniName;
+			FString IniPath;
+			FConfigFile IniFile;
+		};
+
+		// array of all known filesd
+		FKnownConfigFile Files[(uint8)EKnownIniFile::NumKnownFiles];
+	};
+
+	/**
+	 * Load the standard (used on all platforms) ini files, like Engine, Input, etc
+	 *
+	 * @param PlatformName Ini name of the platform this is loadenig for (or nullptr for current platform)
+	 * @param bDefaultEngineIniRequired True if the engine ini is required
+	 *
+	 * return True if the engine ini was loaded
+	 */
+	bool InitializeKnownConfigFiles(const TCHAR* PlatformName, bool bDefaultEngineIniRequired);
+
+	/**
+	 * Create GConfig from a saved file
+	 */
+	static bool CreateGConfigFromSaved(const TCHAR* Filename);
+
+#if WITH_UNREAL_DEVELOPER_TOOLS
+	/**
+	 * Retrieve the fully processed ini system for another platform. The editor will start loading these 
+	 * in the background on startup
+	 */
+	static FConfigCacheIni* ForPlatform(FName PlatformName);
+
+	/**
+	 * Wipe all cached platform configs. Next ForPlatform call will load on-demand the platform configs
+	 */
+	static void ClearOtherPlatformConfigs();
+#endif
+
 private:
+#if WITH_EDITOR
+	/** We only auto-initialize other platform configs in the editor to not slow down programs like ShaderCOmpileWorker */
+	static void AsyncInitializeConfigForPlatforms();
+#endif
+
 	/** Serialize a bootstrapping state into or from an archive */
 	void SerializeStateForBootstrap_Impl(FArchive& Ar);
 
@@ -939,9 +1049,10 @@ private:
 	
 	/** The type of the cache (basically, do we call Flush in the destructor) */
 	EConfigCacheType Type;
-};
 
-FArchive& operator<<(FArchive& Ar, FConfigCacheIni& ConfigCacheIni);
+	/** The filenames for the known files in this config */
+	FKnownConfigFiles KnownFiles;
+};
 
 UE_DEPRECATED(4.24, "This functionality to generate Scalability@Level section string has been moved to Scalability.cpp. Explictly construct section you need manually.")
 CORE_API void ApplyCVarSettingsGroupFromIni(const TCHAR* InSectionBaseName, int32 InGroupNumber, const TCHAR* InIniFilename, uint32 SetBy);
diff --git a/Engine/Source/Runtime/CoreUObject/Private/UObject/Obj.cpp b/Engine/Source/Runtime/CoreUObject/Private/UObject/Obj.cpp
--- a/Engine/Source/Runtime/CoreUObject/Private/UObject/Obj.cpp
+++ b/Engine/Source/Runtime/CoreUObject/Private/UObject/Obj.cpp
@@ -2645,6 +2645,7 @@ FString UObject::GetProjectUserConfigFilename() const
 // @todo ini: Verify per object config objects
 void UObject::UpdateSingleSectionOfConfigFile(const FString& ConfigIniName)
 {
+#if WITH_UNREAL_DEVELOPER_TOOLS
 	// create a sandbox FConfigCache
 	FConfigCacheIni Config(EConfigCacheType::Temporary);
 
@@ -2654,8 +2655,6 @@ void UObject::UpdateSingleSectionOfConfigFile(const FString& ConfigIniName)
 	// save the object properties to this file
 	SaveConfig(CPF_Config, *ConfigIniName, &Config);
 
-	ensureMsgf(Config.Num() == 1, TEXT("UObject::UpdateDefaultConfig() caused more files than expected in the Sandbox config cache!"));
-
 	// do we need to use a special platform hierarchy?
 	FString OverridePlatform = GetFinalOverridePlatform(this);
 
@@ -2669,6 +2668,9 @@ void UObject::UpdateSingleSectionOfConfigFile(const FString& ConfigIniName)
 		FString FinalIniFileName;
 		GConfig->LoadGlobalIniFile(FinalIniFileName, *GetClass()->ClassConfigName.ToString(), NULL, true);
 	}
+#else
+	checkNoEntry();
+#endif
 }
 
 void UObject::UpdateDefaultConfigFile(const FString& SpecificFileLocation)
@@ -2701,8 +2703,6 @@ void UObject::UpdateSinglePropertyInConfigFile(const FProperty* InProperty, cons
 		SaveConfig(CPF_Config, *InConfigIniName, &Config);
 
 		// Take the saved section for this object and have the config system process and write out the one property we care about.
-		ensureMsgf(Config.Num() == 1, TEXT("UObject::UpdateDefaultConfig() caused more files than expected in the Sandbox config cache!"));
-
 		TArray<FString> Keys;
 		NewFile.GetKeys(Keys);
 
@@ -4383,10 +4383,10 @@ void InitUObject()
 
 	// Initialize redirects map
 	FCoreRedirects::Initialize();
-	for (const TPair<FString,FConfigFile>& It : *GConfig)
+	for (const FString& Filename : GConfig->GetFilenames())
 	{
-		FCoreRedirects::ReadRedirectsFromIni(It.Key);
-		FLinkerLoad::CreateActiveRedirectsMap(It.Key);
+		FCoreRedirects::ReadRedirectsFromIni(Filename);
+		FLinkerLoad::CreateActiveRedirectsMap(Filename);
 	}
 
 	FCoreDelegates::OnShutdownAfterError.AddStatic(StaticShutdownAfterError);
diff --git a/Engine/Source/Runtime/Engine/Private/ContentStreaming.cpp b/Engine/Source/Runtime/Engine/Private/ContentStreaming.cpp
--- a/Engine/Source/Runtime/Engine/Private/ContentStreaming.cpp
+++ b/Engine/Source/Runtime/Engine/Private/ContentStreaming.cpp
@@ -116,7 +116,7 @@ TArray<FTrackedRenderAssetEvent> GTrackedRenderAssets;
  */
 void TrackRenderAssetInit()
 {
-	if ( GConfig && GConfig->Num() > 0 )
+	if ( GConfig && GConfig->IsReadyForUse() )
 	{
 		GTrackedRenderAssetNames.Empty();
 		GTrackedRenderAssetsInitialized = true;
diff --git a/Engine/Source/Runtime/Engine/Private/GameUserSettings.cpp b/Engine/Source/Runtime/Engine/Private/GameUserSettings.cpp
--- a/Engine/Source/Runtime/Engine/Private/GameUserSettings.cpp
+++ b/Engine/Source/Runtime/Engine/Private/GameUserSettings.cpp
@@ -556,7 +556,7 @@ void UGameUserSettings::SaveSettings()
 void UGameUserSettings::LoadConfigIni(bool bForceReload/*=false*/)
 {
 	// Load .ini, allowing merging
-	FConfigCacheIni::LoadGlobalIniFile(GGameUserSettingsIni, TEXT("GameUserSettings"), nullptr, bForceReload, false, true, *FPaths::GeneratedConfigDir());
+	FConfigCacheIni::LoadGlobalIniFile(GGameUserSettingsIni, TEXT("GameUserSettings"), nullptr, bForceReload, false, true, true, *FPaths::GeneratedConfigDir());
 }
 
 void UGameUserSettings::PreloadResolutionSettings()
diff --git a/Engine/Source/Runtime/Engine/Private/UnrealEngine.cpp b/Engine/Source/Runtime/Engine/Private/UnrealEngine.cpp
--- a/Engine/Source/Runtime/Engine/Private/UnrealEngine.cpp
+++ b/Engine/Source/Runtime/Engine/Private/UnrealEngine.cpp
@@ -8262,10 +8262,10 @@ bool UEngine::HandleConfigHashCommand( const TCHAR* Cmd, FOutputDevice& Ar )
 		if ( ConfigFilename == TEXT("NAMESONLY") )
 		{
 			Ar.Log( TEXT("Files map:") );
-			for ( FConfigCacheIni::TIterator It(*GConfig); It; ++It )
+			for ( const FString& Filename : GConfig->GetFilenames() )
 			{
 				// base filename is what Dump() compares against
-				Ar.Logf(TEXT("FileName: %s (%s)"), *FPaths::GetBaseFilename(It.Key()), *It.Key());
+				Ar.Logf(TEXT("FileName: %s (%s)"), *FPaths::GetBaseFilename(Filename), *Filename);
 			}
 		}
 		else
@@ -8289,13 +8289,30 @@ bool UEngine::HandleConfigMemCommand( const TCHAR* Cmd, FOutputDevice& Ar )
 
 bool UEngine::HandleGetIniCommand(const TCHAR* Cmd, FOutputDevice& Ar)
 {
-	// Format: GetIni IniFile:Section.SubSection Key
+	// Format: GetIni [Platform@]IniFile:Section.SubSection Key
 	TCHAR IniPlusSectionName[256];
 	TCHAR KeyName[256];
 
 	if (FParse::Token(Cmd, IniPlusSectionName, UE_ARRAY_COUNT(IniPlusSectionName), true))
 	{
 		FString IniPlusSection = IniPlusSectionName;
+
+		// look for the optional Platform prefix which means to look up another platform's version of this ini setting
+		int32 PlatformDelim = IniPlusSection.Find(TEXT("@"));
+		FConfigCacheIni* ConfigSystem = GConfig;
+		if (PlatformDelim > 0)
+		{
+			FString PlatformName = IniPlusSection.Mid(0, PlatformDelim);
+			IniPlusSection = IniPlusSection.Mid(PlatformDelim + 1);
+
+			ConfigSystem = FConfigCacheIni::ForPlatform(*PlatformName);
+			if (ConfigSystem == nullptr)
+			{
+				Ar.Logf(TEXT("Failed to find config system for platform %s"), *PlatformName);
+				return false;
+			}
+		}
+
 		int32 IniDelim = IniPlusSection.Find(TEXT(":"));
 		FString IniName;
 		FString SectionName = (IniDelim != INDEX_NONE ? IniPlusSection.Mid(IniDelim+1) : IniPlusSection);
@@ -8314,7 +8331,7 @@ bool UEngine::HandleGetIniCommand(const TCHAR* Cmd, FOutputDevice& Ar)
 				TArray<FString> ConfigList;
 				FString SearchStr = IniPlusSection.Left(IniDelim) + TEXT(".ini");
 
-				GConfig->GetConfigFilenames(ConfigList);
+				ConfigSystem->GetConfigFilenames(ConfigList);
 
 				for (auto CurConfig : ConfigList)
 				{
@@ -8343,7 +8360,7 @@ bool UEngine::HandleGetIniCommand(const TCHAR* Cmd, FOutputDevice& Ar)
 			{
 				TArray<FString> Values;
 
-				bool bSuccess = !!GConfig->GetArray(*SectionName, KeyName, Values, IniName);
+				bool bSuccess = !!ConfigSystem->GetArray(*SectionName, KeyName, Values, IniName);
 
 				if (bSuccess)
 				{
