Make RGBE encoding match original publication by [REDACTED_PERSON]

diff --git a/Engine/Source/Runtime/Core/Private/Math/Color.cpp b/Engine/Source/Runtime/Core/Private/Math/Color.cpp
--- a/Engine/Source/Runtime/Core/Private/Math/Color.cpp
+++ b/Engine/Source/Runtime/Core/Private/Math/Color.cpp
@@ -92,7 +92,10 @@ FLinearColor FLinearColor::FromPow22Color(const FColor& Color)
 	return LinearColor;
 }
 
-// Convert from float to RGBE as outlined in Gregory Ward's Real Pixels article, Graphics Gems II, page 80.
+/**
+ * Converts from a linear float color to RGBE as outlined in Gregory Ward's Real Pixels article, Graphics Gems II, page 80.
+ * Implementation details in https://cbloomrants.blogspot.com/2020/06/widespread-error-in-radiance-hdr-rgbe.html
+ */ 
 FColor FLinearColor::ToRGBE() const
 {
 	const float	Primary = FMath::Max3( R, G, B );
@@ -105,12 +108,13 @@ FColor FLinearColor::ToRGBE() const
 	else
 	{
 		int32 Exponent;
-		const float Scale = (float)frexp(Primary, &Exponent) / Primary * 255.f;
+		frexpf(Primary, &Exponent);
+		const float Scale = ldexpf(1.f, -Exponent + 8);
 
-		Color.R		= (uint8)FMath::Clamp(FMath::TruncToInt(R * Scale), 0, 255);
-		Color.G		= (uint8)FMath::Clamp(FMath::TruncToInt(G * Scale), 0, 255);
-		Color.B		= (uint8)FMath::Clamp(FMath::TruncToInt(B * Scale), 0, 255);
-		Color.A		= (uint8)(FMath::Clamp(Exponent,-128,127) + 128);
+		Color.R = (uint8)FMath::Clamp(FMath::TruncToInt(R * Scale), 0, 255);
+		Color.G = (uint8)FMath::Clamp(FMath::TruncToInt(G * Scale), 0, 255);
+		Color.B = (uint8)FMath::Clamp(FMath::TruncToInt(B * Scale), 0, 255);
+		Color.A = (uint8)(FMath::Clamp(Exponent,-128,127) + 128);
 	}
 
 	return Color;
@@ -215,7 +219,10 @@ FColor FColor::FromHex( const FString& HexString )
 	return FColor(ForceInitToZero);
 }
 
-// Convert from RGBE to float as outlined in Gregory Ward's Real Pixels article, Graphics Gems II, page 80.
+/**
+ * Converts from RGBE to a linear float color as outlined in Gregory Ward's Real Pixels article, Graphics Gems II, page 80.
+ * Implementation details in https://cbloomrants.blogspot.com/2020/06/widespread-error-in-radiance-hdr-rgbe.html
+ */ 
 FLinearColor FColor::FromRGBE() const
 {
 	if (A == 0)
@@ -224,7 +231,8 @@ FLinearColor FColor::FromRGBE() const
 	}
 	else
 	{
-		const float Scale = (float)ldexp( 1 / 255.0f, A - 128 );
+		// the extra 8 here does the /256
+		const float Scale = ldexpf(1.f, A - (128 + 8));
 		return FLinearColor( R * Scale, G * Scale, B * Scale, 1.0f );
 	}
 }
diff --git a/Engine/Source/Runtime/Engine/Private/ImageUtils.cpp b/Engine/Source/Runtime/Engine/Private/ImageUtils.cpp
--- a/Engine/Source/Runtime/Engine/Private/ImageUtils.cpp
+++ b/Engine/Source/Runtime/Engine/Private/ImageUtils.cpp
@@ -23,6 +23,12 @@ DEFINE_LOG_CATEGORY_STATIC(LogImageUtils, Log, All);
 
 #define LOCTEXT_NAMESPACE "ImageUtils"
 
+/**
+ * Returns data containing the pixmap of the passed in rendertarget.
+ * @param TexRT - The 2D rendertarget from which to read pixmap data.
+ * @param RawData - an array to be filled with pixel data.
+ * @return true if RawData has been successfully filled.
+ */
 static bool GetRawData(UTextureRenderTarget2D* TexRT, TArray64<uint8>& RawData)
 {
 	FRenderTarget* RenderTarget = TexRT->GameThread_GetRenderTargetResource();
@@ -686,13 +692,9 @@ private:
 	}
 
 	/**
-	* Returns data containing the pixmap of the passed in rendertarget.
-	* @param TexRT - The 2D rendertarget from which to read pixmap data.
-	* @param RawData - an array to be filled with pixel data.
-	* @return true if RawData has been successfully filled.
-	*/
-	
-
+	 * Converts from a linear float color to RGBE as outlined in Gregory Ward's Real Pixels article, Graphics Gems II, page 80.
+	 * Implementation details in https://cbloomrants.blogspot.com/2020/06/widespread-error-in-radiance-hdr-rgbe.html
+	 */ 
 	static FColor ToRGBEDithered(const FLinearColor& ColorIN, const FRandomStream& Rand)
 	{
 		const float R = ColorIN.R;
@@ -708,11 +710,12 @@ private:
 		else
 		{
 			int32 Exponent;
-			const float Scale = frexp(Primary, &Exponent) / Primary * 255.f;
+			frexpf(Primary, &Exponent);
+			const float Scale = ldexpf(1.f, -Exponent + 8);
 
-			ReturnColor.R = FMath::Clamp(FMath::TruncToInt((R* Scale) + Rand.GetFraction()), 0, 255);
-			ReturnColor.G = FMath::Clamp(FMath::TruncToInt((G* Scale) + Rand.GetFraction()), 0, 255);
-			ReturnColor.B = FMath::Clamp(FMath::TruncToInt((B* Scale) + Rand.GetFraction()), 0, 255);
+			ReturnColor.R = FMath::Clamp(FMath::TruncToInt((R * Scale) + Rand.GetFraction()), 0, 255);
+			ReturnColor.G = FMath::Clamp(FMath::TruncToInt((G * Scale) + Rand.GetFraction()), 0, 255);
+			ReturnColor.B = FMath::Clamp(FMath::TruncToInt((B * Scale) + Rand.GetFraction()), 0, 255);
 			ReturnColor.A = FMath::Clamp(FMath::TruncToInt(Exponent), -128, 127) + 128;
 		}
 
