added a deprecation comment above FThreadSafeCounter(64) and FThreadSafeBool

diff --git a/Engine/Source/Runtime/Core/Public/HAL/ThreadSafeBool.h b/Engine/Source/Runtime/Core/Public/HAL/ThreadSafeBool.h
--- a/Engine/Source/Runtime/Core/Public/HAL/ThreadSafeBool.h
+++ b/Engine/Source/Runtime/Core/Public/HAL/ThreadSafeBool.h
@@ -5,6 +5,10 @@
 #include "CoreTypes.h"
 #include "HAL/ThreadSafeCounter.h"
 
+////////////////////////////////////////////////////////////////////////////////////////////////////////////////////////////
+// DEPRECATED. Please use `std::atomic<bool>`
+////////////////////////////////////////////////////////////////////////////////////////////////////////////////////////////
+
 /**
  * Thread safe bool, wraps FThreadSafeCounter
  */
diff --git a/Engine/Source/Runtime/Core/Public/HAL/ThreadSafeCounter.h b/Engine/Source/Runtime/Core/Public/HAL/ThreadSafeCounter.h
--- a/Engine/Source/Runtime/Core/Public/HAL/ThreadSafeCounter.h
+++ b/Engine/Source/Runtime/Core/Public/HAL/ThreadSafeCounter.h
@@ -5,6 +5,10 @@
 #include "CoreTypes.h"
 #include "HAL/PlatformAtomics.h"
 
+////////////////////////////////////////////////////////////////////////////////////////////////////////////////////////////
+// DEPRECATED. Please use `std::atomic<int32>`
+////////////////////////////////////////////////////////////////////////////////////////////////////////////////////////////
+
 /** Thread safe counter */
 class FThreadSafeCounter
 {
diff --git a/Engine/Source/Runtime/Core/Public/HAL/ThreadSafeCounter64.h b/Engine/Source/Runtime/Core/Public/HAL/ThreadSafeCounter64.h
--- a/Engine/Source/Runtime/Core/Public/HAL/ThreadSafeCounter64.h
+++ b/Engine/Source/Runtime/Core/Public/HAL/ThreadSafeCounter64.h
@@ -5,6 +5,10 @@
 #include "CoreTypes.h"
 #include "HAL/ThreadSafeCounter.h"
 
+////////////////////////////////////////////////////////////////////////////////////////////////////////////////////////////
+// DEPRECATED. Please use `std::atomic<int64>`
+////////////////////////////////////////////////////////////////////////////////////////////////////////////////////////////
+
 /** Thread safe counter for 64bit ints */
 class FThreadSafeCounter64
 {
