Added FArchiveMD5 adapter over FMD5 to compute MD5 hashes from serializable types.

diff --git a/Engine/Source/Runtime/Core/Public/Misc/ArchiveMD5.h b/Engine/Source/Runtime/Core/Public/Misc/ArchiveMD5.h
new file mode 100644
--- /dev/null
+++ b/Engine/Source/Runtime/Core/Public/Misc/ArchiveMD5.h
@@ -0,0 +1,53 @@
+// Copyright Epic Games, Inc. All Rights Reserved.
+
+#pragma once
+
+#include "CoreMinimal.h"
+#include "Misc/SecureHash.h"
+
+/**
+ * FArchive adapter for FMD5
+ */
+class FArchiveMD5 : public FArchive
+{
+public:
+	inline FArchiveMD5()
+	{
+		SetIsLoading(false);
+		SetIsSaving(true);
+		SetIsPersistent(false);
+	}
+
+	virtual FString GetArchiveName() const
+	{
+		return TEXT("FArchiveMD5");
+	}
+
+	void Serialize(void* Data, int64 Num) override
+	{
+		MD5.Update((uint8*)Data, Num);
+	}
+
+	using FArchive::operator<<;
+
+	virtual FArchive& operator<<(class FName& Value) override
+	{
+		FString NameAsString = Value.ToString();
+		*this << NameAsString;
+		return *this;
+	}
+
+	virtual FArchive& operator<<(class UObject*& Value) override
+	{
+		check(0);
+		return *this;
+	}
+
+	void GetHash(FMD5Hash& Hash)
+	{
+		Hash.Set(MD5);
+	}
+
+protected:
+	FMD5 MD5;
+};
\ No newline at end of file
