Fix crash caused by race condition between smoke tests and async tasks

diff --git a/Engine/Source/Runtime/Core/Private/Misc/AutomationTest.cpp b/Engine/Source/Runtime/Core/Private/Misc/AutomationTest.cpp
--- a/Engine/Source/Runtime/Core/Private/Misc/AutomationTest.cpp
+++ b/Engine/Source/Runtime/Core/Private/Misc/AutomationTest.cpp
@@ -100,7 +100,9 @@ void FAutomationTestFramework::FAutomationTestOutputDevice::Serialize( const TCH
 
 void FAutomationTestFramework::FAutomationTestMessageFilter::Serialize(const TCHAR* V, ELogVerbosity::Type Verbosity, const class FName& Category)
 {
-	if (DestinationContext)
+	// Prevent null dereference if logging happens in async tasks while changing DestinationContext
+	FFeedbackContext* const LocalDestinationContext = DestinationContext.load(std::memory_order_relaxed);
+	if (LocalDestinationContext)
 	{
 		if ((Verbosity == ELogVerbosity::Warning) || (Verbosity == ELogVerbosity::Error))
 		{
@@ -109,8 +111,7 @@ void FAutomationTestFramework::FAutomationTestMessageFilter::Serialize(const TCH
 				Verbosity = ELogVerbosity::Verbose;
 			}
 		}
-
-		DestinationContext->Serialize(V, Verbosity, Category);
+		LocalDestinationContext->Serialize(V, Verbosity, Category);
 	}
 }
 
diff --git a/Engine/Source/Runtime/Core/Public/Misc/AutomationTest.h b/Engine/Source/Runtime/Core/Public/Misc/AutomationTest.h
--- a/Engine/Source/Runtime/Core/Public/Misc/AutomationTest.h
+++ b/Engine/Source/Runtime/Core/Public/Misc/AutomationTest.h
@@ -28,7 +28,7 @@
 #include "HAL/PlatformProcess.h"
 #include "Misc/AutomationEvent.h"
 #include "Internationalization/Regex.h"
-
+#include <atomic>
 
 #ifndef WITH_AUTOMATION_TESTS
 	#define WITH_AUTOMATION_TESTS (WITH_DEV_AUTOMATION_TESTS || WITH_PERF_AUTOMATION_TESTS)
@@ -922,10 +922,12 @@ private:
 	 {
 	 public:
 		 FAutomationTestMessageFilter()
-			 : CurTest(nullptr) {}
+			: CurTest(nullptr)
+			, DestinationContext(nullptr) {}
 
 		 ~FAutomationTestMessageFilter()
 		 {
+			 DestinationContext = nullptr;
 			 CurTest = nullptr;
 		 }
 
@@ -962,7 +964,7 @@ private:
 
 	 private:
 		 class FAutomationTestBase* CurTest;
-		 FFeedbackContext* DestinationContext = nullptr;
+		 std::atomic<FFeedbackContext*> DestinationContext;
 	 };
 
 	//** Store information about blacklisted test */
