Change default value for UBTTargetName from "UE4" to "Unreal" for UE5, Found this issue when I click reompile on a module in the module manager

diff --git a/Engine/Source/Runtime/Core/Private/GenericPlatform/GenericPlatformMisc.cpp b/Engine/Source/Runtime/Core/Private/GenericPlatform/GenericPlatformMisc.cpp
--- a/Engine/Source/Runtime/Core/Private/GenericPlatform/GenericPlatformMisc.cpp
+++ b/Engine/Source/Runtime/Core/Private/GenericPlatform/GenericPlatformMisc.cpp
@@ -1100,7 +1100,7 @@ const TCHAR* FGenericPlatformMisc::GetUBTTarget()
 
 /** The name of the UBT target that the current executable was built from. Defaults to the UE4 default target for this type to make content only projects work, 
 	but will be overridden by the primary game module if it exists */
-TCHAR GUBTTargetName[128] = TEXT("UE4" PREPROCESSOR_TO_STRING(UBT_COMPILED_TARGET));
+TCHAR GUBTTargetName[128] = TEXT("Unreal" PREPROCESSOR_TO_STRING(UBT_COMPILED_TARGET));
 
 void FGenericPlatformMisc::SetUBTTargetName(const TCHAR* InTargetName)
 {
