Fix IsInRenderingThread/IsInParallelRenderingThread checks by caching GIsRenderingThreadSuspended between old and new values computations, since it could have changed between the two statements.

diff --git a/Engine/Source/Runtime/Core/Private/HAL/ThreadingBase.cpp b/Engine/Source/Runtime/Core/Private/HAL/ThreadingBase.cpp
--- a/Engine/Source/Runtime/Core/Private/HAL/ThreadingBase.cpp
+++ b/Engine/Source/Runtime/Core/Private/HAL/ThreadingBase.cpp
@@ -222,13 +222,15 @@ PRAGMA_ENABLE_DEPRECATION_WARNINGS
 CORE_API bool IsInRenderingThread()
 {
 PRAGMA_DISABLE_DEPRECATION_WARNINGS
-	bool newValue = !GRenderingThread || GIsRenderingThreadSuspended.Load(EMemoryOrder::Relaxed)
+	const bool bLocalIsLoadingThreadSuspended = GIsRenderingThreadSuspended.Load(EMemoryOrder::Relaxed) != 0;
+
+	bool newValue = !GRenderingThread || bLocalIsLoadingThreadSuspended
 		? FTaskTagScope::IsCurrentTag(ETaskTag::EGameThread) || FTaskTagScope::IsCurrentTag(ETaskTag::ERenderingThread) || FTaskTagScope::IsRunningDuringStaticInit()
 		: FTaskTagScope::IsCurrentTag(ETaskTag::ERenderingThread);
 
 #if !UE_BUILD_SHIPPING && !UE_BUILD_TEST
-	bool oldValue = !GRenderingThread || GIsRenderingThreadSuspended.Load(EMemoryOrder::Relaxed) || (FPlatformTLS::GetCurrentThreadId() == GRenderingThread->GetThreadID());
-	ensureMsgf(oldValue == newValue, TEXT("oldValue(%i) newValue(%i) If this check fails make sure that there is a FTaskTagScope(ETaskTag::ERenderingThread) as deep as possible on the current callstack, you can see the current value in ActiveNamedThreads(%x), GRenderingThread(%x), GIsRenderingThreadSuspended(%d)"), oldValue, newValue, FTaskTagScope::GetCurrentTag(), GRenderingThread, GIsRenderingThreadSuspended.Load(EMemoryOrder::Relaxed));
+	bool oldValue = !GRenderingThread || bLocalIsLoadingThreadSuspended || (FPlatformTLS::GetCurrentThreadId() == GRenderingThread->GetThreadID());
+	ensureMsgf(oldValue == newValue, TEXT("oldValue(%i) newValue(%i) If this check fails make sure that there is a FTaskTagScope(ETaskTag::ERenderingThread) as deep as possible on the current callstack, you can see the current value in ActiveNamedThreads(%x), GRenderingThread(%x), GIsRenderingThreadSuspended(%d)"), oldValue, newValue, FTaskTagScope::GetCurrentTag(), GRenderingThread, bLocalIsLoadingThreadSuspended);
 	newValue = oldValue;
 #endif
 	return newValue;
@@ -238,8 +240,10 @@ PRAGMA_ENABLE_DEPRECATION_WARNINGS
 CORE_API bool IsInParallelRenderingThread()
 {
 PRAGMA_DISABLE_DEPRECATION_WARNINGS
+	const bool bLocalIsLoadingThreadSuspended = GIsRenderingThreadSuspended.Load(EMemoryOrder::Relaxed) != 0;
+
 	bool newValue = false;
-	if (!GRenderingThread || GIsRenderingThreadSuspended.Load(EMemoryOrder::Relaxed))
+	if (!GRenderingThread || bLocalIsLoadingThreadSuspended)
 	{
 		newValue = FTaskTagScope::IsCurrentTag(ETaskTag::ERenderingThread) || FTaskTagScope::IsCurrentTag(ETaskTag::EGameThread) || FTaskTagScope::IsCurrentTag(ETaskTag::EParallelRenderingThread);
 	}
@@ -253,7 +257,7 @@ PRAGMA_DISABLE_DEPRECATION_WARNINGS
 
 #if !UE_BUILD_SHIPPING && !UE_BUILD_TEST
 	bool oldValue = false;
-	if (!GRenderingThread || GIsRenderingThreadSuspended.Load(EMemoryOrder::Relaxed))
+	if (!GRenderingThread || bLocalIsLoadingThreadSuspended)
 	{
 		oldValue = true;
 	}
@@ -261,7 +265,7 @@ PRAGMA_DISABLE_DEPRECATION_WARNINGS
 	{
 		oldValue = FPlatformTLS::GetCurrentThreadId() != GGameThreadId;
 	}
-	ensureMsgf(oldValue == newValue, TEXT("oldValue(%i) newValue(%i) If this check fails make sure that there is a FTaskTagScope(ETaskTag::EParallelRenderingThread) as deep as possible on the current callstack, you can see the current value in ActiveNamedThreads(%x), GRenderingThread(%x), GIsRenderingThreadSuspended(%d)"), oldValue, newValue, FTaskTagScope::GetCurrentTag(), GRenderingThread, GIsRenderingThreadSuspended.Load(EMemoryOrder::Relaxed));
+	ensureMsgf(oldValue == newValue, TEXT("oldValue(%i) newValue(%i) If this check fails make sure that there is a FTaskTagScope(ETaskTag::EParallelRenderingThread) as deep as possible on the current callstack, you can see the current value in ActiveNamedThreads(%x), GRenderingThread(%x), GIsRenderingThreadSuspended(%d)"), oldValue, newValue, FTaskTagScope::GetCurrentTag(), GRenderingThread, bLocalIsLoadingThreadSuspended);
 	newValue = oldValue;
 #endif
 	return newValue;
