Fixed static analysis warnings about sizeof(A.Bytes()) in BLAKE3

#rnx

diff --git a/Engine/Source/Runtime/Core/Public/Misc/Blake3.h b/Engine/Source/Runtime/Core/Public/Misc/Blake3.h
--- a/Engine/Source/Runtime/Core/Public/Misc/Blake3.h
+++ b/Engine/Source/Runtime/Core/Public/Misc/Blake3.h
@@ -67,17 +67,17 @@ inline FBlake3Hash::FBlake3Hash(const ByteArray& InHash)
 
 inline bool operator==(const FBlake3Hash& A, const FBlake3Hash& B)
 {
-	return FMemory::Memcmp(A.Bytes(), B.Bytes(), sizeof(A.Bytes())) == 0;
+	return FMemory::Memcmp(A.Bytes(), B.Bytes(), sizeof(decltype(A.Bytes()))) == 0;
 }
 
 inline bool operator!=(const FBlake3Hash& A, const FBlake3Hash& B)
 {
-	return FMemory::Memcmp(A.Bytes(), B.Bytes(), sizeof(A.Bytes())) != 0;
+	return FMemory::Memcmp(A.Bytes(), B.Bytes(), sizeof(decltype(A.Bytes()))) != 0;
 }
 
 inline bool operator<(const FBlake3Hash& A, const FBlake3Hash& B)
 {
-	return FMemory::Memcmp(A.Bytes(), B.Bytes(), sizeof(A.Bytes())) < 0;
+	return FMemory::Memcmp(A.Bytes(), B.Bytes(), sizeof(decltype(A.Bytes()))) < 0;
 }
 
 inline uint32 GetTypeHash(const FBlake3Hash& Hash)
