Fixed static analysis warnings about MemoryViewTest

#rnx

diff --git a/Engine/Source/Runtime/Core/Private/Tests/Memory/MemoryViewTest.cpp b/Engine/Source/Runtime/Core/Private/Tests/Memory/MemoryViewTest.cpp
--- a/Engine/Source/Runtime/Core/Private/Tests/Memory/MemoryViewTest.cpp
+++ b/Engine/Source/Runtime/Core/Private/Tests/Memory/MemoryViewTest.cpp
@@ -46,7 +46,14 @@ bool FMemoryViewTest::RunTest(const FString& Parameters)
 		TestEqual(TEXT("MemoryView.IsEmpty()"), View.IsEmpty(), Size == 0);
 	};
 
-	uint8 ByteArray[16]{};
+	struct
+	{
+		uint8 BeforeByteArray[4];
+		uint8 ByteArray[16]{};
+		uint8 AfterByteArray[4];
+	} ByteArrayContainer;
+
+	uint8 (&ByteArray)[16] = ByteArrayContainer.ByteArray;
 	uint32 IntArray[12]{};
 
 	// Test Empty Views
@@ -126,14 +133,14 @@ bool FMemoryViewTest::RunTest(const FString& Parameters)
 	TestTrue(TEXT("MemoryView.Contains(EmptyContained)"), MakeMemoryView(ByteArray).Contains(MakeMemoryView(ByteArray, 0)));
 	TestTrue(TEXT("MemoryView.Contains(EmptyContained)"), MakeMemoryView(ByteArray).Contains(MakeMemoryView(ByteArray + 8, 0)));
 	TestTrue(TEXT("MemoryView.Contains(EmptyContained)"), MakeMemoryView(ByteArray).Contains(MakeMemoryView(ByteArray + 16, 0)));
-	TestFalse(TEXT("MemoryView.Contains(EmptyOutside)"), MakeMemoryView(ByteArray).Contains(MakeMemoryView(ByteArray - 1, 0)));
-	TestFalse(TEXT("MemoryView.Contains(EmptyOutside)"), MakeMemoryView(ByteArray).Contains(MakeMemoryView(ByteArray + 17, 0)));
-	TestFalse(TEXT("MemoryView.Contains(OutsideBy1Left)"), MakeMemoryView(ByteArray).Contains(MakeMemoryView(ByteArray - 1, 1)));
+	TestFalse(TEXT("MemoryView.Contains(EmptyOutside)"), MakeMemoryView(ByteArray).Contains(MakeMemoryView(ByteArrayContainer.BeforeByteArray + 3, 0)));
+	TestFalse(TEXT("MemoryView.Contains(EmptyOutside)"), MakeMemoryView(ByteArray).Contains(MakeMemoryView(ByteArrayContainer.AfterByteArray + 1, 0)));
+	TestFalse(TEXT("MemoryView.Contains(OutsideBy1Left)"), MakeMemoryView(ByteArray).Contains(MakeMemoryView(ByteArrayContainer.BeforeByteArray + 3, 1)));
 	TestFalse(TEXT("MemoryView.Contains(OutsideBy1Right)"), MakeMemoryView(ByteArray).Contains(MakeMemoryView(ByteArray + 16, 1)));
-	TestFalse(TEXT("MemoryView.Contains(LargerBy1Left)"), MakeMemoryView(ByteArray).Contains(MakeMemoryView(ByteArray - 1, 17)));
+	TestFalse(TEXT("MemoryView.Contains(LargerBy1Left)"), MakeMemoryView(ByteArray).Contains(MakeMemoryView(ByteArrayContainer.BeforeByteArray + 3, 17)));
 	TestFalse(TEXT("MemoryView.Contains(LargerBy1Right)"), MakeMemoryView(ByteArray).Contains(MakeMemoryView(ByteArray, 17)));
-	TestFalse(TEXT("MemoryView.Contains(LargerBy2Both)"), MakeMemoryView(ByteArray).Contains(MakeMemoryView(ByteArray - 1, 18)));
-	TestFalse(TEXT("MemoryView.Contains(SmallerOverlapLeft)"), MakeMemoryView(ByteArray).Contains(MakeMemoryView(ByteArray - 1, 8)));
+	TestFalse(TEXT("MemoryView.Contains(LargerBy2Both)"), MakeMemoryView(ByteArray).Contains(MakeMemoryView(ByteArrayContainer.BeforeByteArray + 3, 18)));
+	TestFalse(TEXT("MemoryView.Contains(SmallerOverlapLeft)"), MakeMemoryView(ByteArray).Contains(MakeMemoryView(ByteArrayContainer.BeforeByteArray + 3, 8)));
 	TestFalse(TEXT("MemoryView.Contains(SmallerOverlapRight)"), MakeMemoryView(ByteArray).Contains(MakeMemoryView(ByteArray + 9, 8)));
 
 	// Test Intersects
@@ -141,19 +148,19 @@ bool FMemoryViewTest::RunTest(const FString& Parameters)
 	TestTrue(TEXT("MemoryView.Intersects(SmallerBy1Left)"), MakeMemoryView(ByteArray).Intersects(MakeMemoryView(ByteArray + 1, 15)));
 	TestTrue(TEXT("MemoryView.Intersects(SmallerBy1Right)"), MakeMemoryView(ByteArray).Intersects(MakeMemoryView(ByteArray, 15)));
 	TestTrue(TEXT("MemoryView.Intersects(SmallerBy2Both)"), MakeMemoryView(ByteArray).Intersects(MakeMemoryView(ByteArray + 1, 14)));
-	TestTrue(TEXT("MemoryView.Intersects(SmallerOverlapLeft)"), MakeMemoryView(ByteArray).Intersects(MakeMemoryView(ByteArray - 1, 8)));
+	TestTrue(TEXT("MemoryView.Intersects(SmallerOverlapLeft)"), MakeMemoryView(ByteArray).Intersects(MakeMemoryView(ByteArrayContainer.BeforeByteArray + 3, 8)));
 	TestTrue(TEXT("MemoryView.Intersects(SmallerOverlapRight)"), MakeMemoryView(ByteArray).Intersects(MakeMemoryView(ByteArray + 9, 8)));
-	TestTrue(TEXT("MemoryView.Intersects(LargerBy1Left)"), MakeMemoryView(ByteArray).Intersects(MakeMemoryView(ByteArray - 1, 17)));
+	TestTrue(TEXT("MemoryView.Intersects(LargerBy1Left)"), MakeMemoryView(ByteArray).Intersects(MakeMemoryView(ByteArrayContainer.BeforeByteArray + 3, 17)));
 	TestTrue(TEXT("MemoryView.Intersects(LargerBy1Right)"), MakeMemoryView(ByteArray).Intersects(MakeMemoryView(ByteArray, 17)));
-	TestTrue(TEXT("MemoryView.Intersects(LargerBy2Both)"), MakeMemoryView(ByteArray).Intersects(MakeMemoryView(ByteArray - 1, 18)));
+	TestTrue(TEXT("MemoryView.Intersects(LargerBy2Both)"), MakeMemoryView(ByteArray).Intersects(MakeMemoryView(ByteArrayContainer.BeforeByteArray + 3, 18)));
 	TestTrue(TEXT("MemoryView.Intersects(EmptyMiddle)"), MakeMemoryView(ByteArray).Intersects(MakeMemoryView(ByteArray + 8, 0)));
 	TestFalse(TEXT("MemoryView.Intersects(Empty)"), FMutableMemoryView().Intersects(FConstMemoryView()));
 	TestFalse(TEXT("MemoryView.Intersects(Empty)"), FConstMemoryView().Intersects(FMutableMemoryView()));
 	TestFalse(TEXT("MemoryView.Intersects(EmptyLeft)"), MakeMemoryView(ByteArray).Intersects(MakeMemoryView(ByteArray, 0)));
 	TestFalse(TEXT("MemoryView.Intersects(EmptyRight)"), MakeMemoryView(ByteArray).Intersects(MakeMemoryView(ByteArray + 16, 0)));
-	TestFalse(TEXT("MemoryView.Intersects(EmptyOutside)"), MakeMemoryView(ByteArray).Intersects(MakeMemoryView(ByteArray - 1, 0)));
-	TestFalse(TEXT("MemoryView.Intersects(EmptyOutside)"), MakeMemoryView(ByteArray).Intersects(MakeMemoryView(ByteArray + 17, 0)));
-	TestFalse(TEXT("MemoryView.Intersects(OutsideBy1Left)"), MakeMemoryView(ByteArray).Intersects(MakeMemoryView(ByteArray - 1, 1)));
+	TestFalse(TEXT("MemoryView.Intersects(EmptyOutside)"), MakeMemoryView(ByteArray).Intersects(MakeMemoryView(ByteArrayContainer.BeforeByteArray + 3, 0)));
+	TestFalse(TEXT("MemoryView.Intersects(EmptyOutside)"), MakeMemoryView(ByteArray).Intersects(MakeMemoryView(ByteArrayContainer.AfterByteArray + 1, 0)));
+	TestFalse(TEXT("MemoryView.Intersects(OutsideBy1Left)"), MakeMemoryView(ByteArray).Intersects(MakeMemoryView(ByteArrayContainer.BeforeByteArray + 3, 1)));
 	TestFalse(TEXT("MemoryView.Intersects(OutsideBy1Right)"), MakeMemoryView(ByteArray).Intersects(MakeMemoryView(ByteArray + 16, 1)));
 
 	// Test CompareBytes
@@ -193,16 +200,16 @@ bool FMemoryViewTest::RunTest(const FString& Parameters)
 	TestFalse(TEXT("MemoryView.Equals(BothDiff)"), MakeMemoryView(IntArray).Equals(FMutableMemoryView()));
 
 	// Test operator==
-	static_assert(MakeMemoryView(ByteArray) == MakeMemoryView(ByteArray), "Error in MemoryView == MemoryView");
-	static_assert(MakeMemoryView(ByteArray) == MakeMemoryView(AsConst(ByteArray)), "Error in MemoryView == MemoryView");
-	static_assert(MakeMemoryView(AsConst(ByteArray)) == MakeMemoryView(ByteArray), "Error in MemoryView == MemoryView");
-	static_assert(MakeMemoryView(AsConst(ByteArray)) == MakeMemoryView(AsConst(ByteArray)), "Error in MemoryView == MemoryView");
+	static_assert(MakeMemoryView(ByteArrayContainer.ByteArray) == MakeMemoryView(ByteArrayContainer.ByteArray), "Error in MemoryView == MemoryView");
+	static_assert(MakeMemoryView(ByteArrayContainer.ByteArray) == MakeMemoryView(AsConst(ByteArrayContainer.ByteArray)), "Error in MemoryView == MemoryView");
+	static_assert(MakeMemoryView(AsConst(ByteArrayContainer.ByteArray)) == MakeMemoryView(ByteArrayContainer.ByteArray), "Error in MemoryView == MemoryView");
+	static_assert(MakeMemoryView(AsConst(ByteArrayContainer.ByteArray)) == MakeMemoryView(AsConst(ByteArrayContainer.ByteArray)), "Error in MemoryView == MemoryView");
 
 	// Test operator!=
-	static_assert(MakeMemoryView(ByteArray) != MakeMemoryView(IntArray), "Error in MemoryView != MemoryView");
-	static_assert(MakeMemoryView(ByteArray) != MakeMemoryView(AsConst(IntArray)), "Error in MemoryView != MemoryView");
-	static_assert(MakeMemoryView(AsConst(ByteArray)) != MakeMemoryView(IntArray), "Error in MemoryView != MemoryView");
-	static_assert(MakeMemoryView(AsConst(ByteArray)) != MakeMemoryView(AsConst(IntArray)), "Error in MemoryView != MemoryView");
+	static_assert(MakeMemoryView(ByteArrayContainer.ByteArray) != MakeMemoryView(IntArray), "Error in MemoryView != MemoryView");
+	static_assert(MakeMemoryView(ByteArrayContainer.ByteArray) != MakeMemoryView(AsConst(IntArray)), "Error in MemoryView != MemoryView");
+	static_assert(MakeMemoryView(AsConst(ByteArrayContainer.ByteArray)) != MakeMemoryView(IntArray), "Error in MemoryView != MemoryView");
+	static_assert(MakeMemoryView(AsConst(ByteArrayContainer.ByteArray)) != MakeMemoryView(AsConst(IntArray)), "Error in MemoryView != MemoryView");
 
 	// Test operator+=
 	TestEqual(TEXT("MemoryView += 0"), MakeMemoryView(ByteArray) += 0, MakeMemoryView(ByteArray));
