HashBuilder: replace HashCombine by faster hash combine function similar to what boost::hash_combine_impl uses.

diff --git a/Engine/Source/Runtime/Core/Public/Misc/HashBuilder.h b/Engine/Source/Runtime/Core/Public/Misc/HashBuilder.h
--- a/Engine/Source/Runtime/Core/Public/Misc/HashBuilder.h
+++ b/Engine/Source/Runtime/Core/Public/Misc/HashBuilder.h
@@ -35,7 +35,7 @@ public:
 	template <typename T>
 	typename TEnableIf<TModels<CGetTypeHashable, T>::Value, FHashBuilder&>::Type Append(const T& InData)
 	{
-		Hash = HashCombine(Hash, GetTypeHash(InData));
+		Hash ^= GetTypeHash(InData) + 0x9e3779b9 + (Hash << 6) + (Hash >> 2);
 		return *this;
 	}
 
