- Added FPlatformMemory::RegisterCustomLLMTags() to allow a platform to add their own tags, and cleaned up some platform code that was doing it manually

diff --git a/Engine/Source/Runtime/Core/Private/HAL/LowLevelMemTracker.cpp b/Engine/Source/Runtime/Core/Private/HAL/LowLevelMemTracker.cpp
--- a/Engine/Source/Runtime/Core/Private/HAL/LowLevelMemTracker.cpp
+++ b/Engine/Source/Runtime/Core/Private/HAL/LowLevelMemTracker.cpp
@@ -1314,6 +1314,9 @@ void FLowLevelMemTracker::InitialiseTagDatas()
 		FLLMTagDeclaration::SetCreationCallback(GlobalRegisterTagDeclaration);
 	}
 
+	// now let the platform add any custom tags
+	FPlatformMemory::RegisterCustomLLMTags();
+
 	// All parents in the ELLMTags and the initial modules' list of LLM_DEFINE_TAG must be contained within that same set, so we can FinishConstruct them now, which we do in UpdateTags
 	bTagAdded = true;
 	UpdateTags();
diff --git a/Engine/Source/Runtime/Core/Public/GenericPlatform/GenericPlatformMemory.h b/Engine/Source/Runtime/Core/Public/GenericPlatform/GenericPlatformMemory.h
--- a/Engine/Source/Runtime/Core/Public/GenericPlatform/GenericPlatformMemory.h
+++ b/Engine/Source/Runtime/Core/Public/GenericPlatform/GenericPlatformMemory.h
@@ -732,6 +732,11 @@ public:
 	FORCEINLINE static void OnLowLevelMemory_Alloc(void const* Pointer, uint64 Size, uint64 Tag) { }
 	FORCEINLINE static void OnLowLevelMemory_Free(void const* Pointer, uint64 Size, uint64 Tag) { }
 
+	/**
+	* Called once at LLM initialization time to let the platform add any custom tags
+	*/
+	static void RegisterCustomLLMTags() { };
+
 	/**
 	* Called once per frame when LLM is collating the data for the current frame.
 	* Can be used to set platform-specific calculated tag data via SetTagAmountForTracker
