FString::operator+ overrides that take FStringView
#rnx

diff --git a/Engine/Source/Runtime/Core/Private/Containers/StringTest.cpp b/Engine/Source/Runtime/Core/Private/Containers/StringTest.cpp
--- a/Engine/Source/Runtime/Core/Private/Containers/StringTest.cpp
+++ b/Engine/Source/Runtime/Core/Private/Containers/StringTest.cpp
@@ -311,7 +311,9 @@ bool FStringFromStringViewTest::RunTest(const FString& Parameters)
 	// Verify basic construction and assignment from a string view.
 	{
 		const TCHAR* Literal = TEXT("Literal");
+		const ANSICHAR* AnsiLiteral = "Literal";
 		TestEqual(TEXT("String(StringView)"), FString(FStringView(Literal)), Literal);
+		TestEqual(TEXT("String(AnsiStringView)"), FString(FAnsiStringView(AnsiLiteral)), Literal);
 		TestEqual(TEXT("String = StringView"), FString(TEXT("Temp")) = FStringView(Literal), Literal);
 
 		FStringView EmptyStringView;
@@ -344,6 +346,37 @@ bool FStringFromStringViewTest::RunTest(const FString& Parameters)
 		TestEqual(TEXT("String = StringView(String).Mid"), AssignMiddleOfString, TEXT("Middle"));
 	}
 
+	// Verify operators taking string views and character arrays
+	{
+		FStringView RhsStringView = FStringView(TEXT("RhsNotSZ"), 3);
+		FString MovePlusSVResult = FString(TEXT("Lhs")) + RhsStringView;
+		TestEqual(TEXT("Move String + StringView"), MovePlusSVResult, TEXT("LhsRhs"));
+
+		FString CopyLhs(TEXT("Lhs"));
+		FString CopyPlusSVResult = CopyLhs + RhsStringView;
+		TestEqual(TEXT("Copy String + StringView"), CopyPlusSVResult, TEXT("LhsRhs"));
+
+		FString MovePlusTCHARsResult = FString(TEXT("Lhs")) + TEXT("Rhs");
+		TestEqual(TEXT("Move String + TCHAR*"), MovePlusTCHARsResult, TEXT("LhsRhs"));
+
+		FString CopyPlusTCHARsResult = CopyLhs + TEXT("Rhs");
+		TestEqual(TEXT("Copy String + TCHAR*"), CopyPlusTCHARsResult, TEXT("LhsRhs"));
+
+		FStringView LhsStringView = FStringView(TEXT("LhsNotSZ"), 3);
+		FString SVPlusMoveResult = LhsStringView + FString(TEXT("Rhs"));
+		TestEqual(TEXT("StringView + Move String"), SVPlusMoveResult, TEXT("LhsRhs"));
+
+		FString CopyRhs(TEXT("Rhs"));
+		FString SVPlusCopyResult = LhsStringView + CopyRhs;
+		TestEqual(TEXT("StringView + Copy String"), SVPlusCopyResult, TEXT("LhsRhs"));
+
+		FString TCHARsPlusMoveResult = TEXT("Lhs") + FString(TEXT("Rhs"));
+		TestEqual(TEXT("TCHAR* + Move String"), TCHARsPlusMoveResult, TEXT("LhsRhs"));
+
+		FString TCHARsPlusCopyResult = TEXT("Lhs") + CopyRhs;
+		TestEqual(TEXT("TCHAR* + Copy String"), TCHARsPlusCopyResult, TEXT("LhsRhs"));
+	}
+
 	return true;
 }
 
diff --git a/Engine/Source/Runtime/Core/Public/Containers/UnrealString.h b/Engine/Source/Runtime/Core/Public/Containers/UnrealString.h
--- a/Engine/Source/Runtime/Core/Public/Containers/UnrealString.h
+++ b/Engine/Source/Runtime/Core/Public/Containers/UnrealString.h
@@ -68,6 +68,12 @@ private:
 		TNot<TIsArray<typename TRemoveReference<CharRangeType>::Type>>,
 		TIsRangeOfCharType<CharRangeType>>;
 
+	/** Trait testing whether a type is a contiguous range of characters, and not CharType[] and not FString. */
+	template <typename CharRangeType>
+	using TIsCharRangeNotCArrayNotFString = TAnd<
+		TIsCharRangeNotCArray<CharRangeType>,
+		TNot<TIsDerivedFrom<typename TDecay<CharRangeType>::Type, FString>>>;
+
 	/** Trait testing whether a type is a contiguous range of TCHAR, and not TCHAR[]. */
 	template <typename CharRangeType>
 	using TIsTCharRangeNotCArray = TAnd<
@@ -75,6 +81,12 @@ private:
 		TNot<TIsArray<typename TRemoveReference<CharRangeType>::Type>>,
 		TIsRangeOfTCHAR<CharRangeType>>;
 
+	/** Trait testing whether a type is a contiguous range of TCHAR, and not TCHAR[] and not FString. */
+	template <typename CharRangeType>
+	using TIsTCharRangeNotCArrayNotFString = TAnd<
+		TIsTCharRangeNotCArray<CharRangeType>,
+		TNot<TIsDerivedFrom<typename TDecay<CharRangeType>::Type, FString>>>;
+
 public:
 	using ElementType = TCHAR;
 
@@ -188,7 +200,7 @@ public:
 	 *
 	 * @param Other The contiguous character range to copy from
 	 */
-	template <typename CharRangeType, typename TEnableIf<TIsCharRangeNotCArray<CharRangeType>::Value>::Type* = nullptr>
+	template <typename CharRangeType, typename TEnableIf<TIsCharRangeNotCArrayNotFString<CharRangeType>::Value>::Type* = nullptr>
 	explicit FString(CharRangeType&& Other)
 	{
 		if (const auto OtherNum = GetNum(Other))
@@ -208,7 +220,7 @@ public:
 	 * @param Other The contiguous character range to copy from
 	 * @param ExtraSlack The number of extra characters to reserve space for in the new string
 	 */
-	template <typename CharRangeType, typename TEnableIf<TIsCharRangeNotCArray<CharRangeType>::Value>::Type* = nullptr>
+	template <typename CharRangeType, typename TEnableIf<TIsCharRangeNotCArrayNotFString<CharRangeType>::Value>::Type* = nullptr>
 	explicit FString(CharRangeType&& Other, int32 ExtraSlack)
 	{
 		const auto OtherNum = GetNum(Other);
@@ -689,14 +701,23 @@ private:
 	FORCEINLINE static FString ConcatTCHARsToFString(const TCHAR* Lhs, typename TIdentity<RhsType>::Type Rhs)
 	{
 		checkSlow(Lhs);
-		Rhs.CheckInvariants();
+		if (!Lhs)
+		{
+			return MoveTempIfPossible(Rhs);
+		}
+		return ConcatTCHARRangeToFString<RhsType>(Lhs, FCString::Strlen(Lhs), MoveTempIfPossible(Rhs));
+	}
 
-		if (!Lhs || !*Lhs)
+	template <typename RhsType>
+	FORCEINLINE static FString ConcatTCHARRangeToFString(const TCHAR* Lhs, int32 LhsLen, typename TIdentity<RhsType>::Type Rhs)
+	{
+		checkSlow(LhsLen >= 0);
+		Rhs.CheckInvariants();
+		if (LhsLen == 0)
 		{
 			return MoveTempIfPossible(Rhs);
 		}
 
-		int32 LhsLen = FCString::Strlen(Lhs);
 		int32 RhsLen = Rhs.Len();
 
 		// This is not entirely optimal, as if the Rhs is an rvalue and has enough slack space to hold Lhs, then
@@ -717,19 +738,27 @@ private:
 	template <typename LhsType>
 	FORCEINLINE static FString ConcatFStringToTCHARs(typename TIdentity<LhsType>::Type Lhs, const TCHAR* Rhs)
 	{
-		Lhs.CheckInvariants();
 		checkSlow(Rhs);
-
-		if (!Rhs || !*Rhs)
+		if (!Rhs)
 		{
 			return MoveTempIfPossible(Lhs);
 		}
+		return ConcatFStringToTCHARRange<LhsType>(MoveTempIfPossible(Lhs), Rhs, FCString::Strlen(Rhs));
+	}
 
-		int32 RhsLen = FCString::Strlen(Rhs);
+	template <typename LhsType>
+	FORCEINLINE static FString ConcatFStringToTCHARRange(typename TIdentity<LhsType>::Type Lhs, const TCHAR* Rhs, int32 RhsLen)
+	{
+		Lhs.CheckInvariants();
+		checkSlow(RhsLen >= 0);
+		if (RhsLen == 0)
+		{
+			return MoveTempIfPossible(Lhs);
+		}
 
 		FString Result(MoveTempIfPossible(Lhs), RhsLen);
 		Result.AppendChars(Rhs, RhsLen);
-		
+
 		return Result;
 	}
 
@@ -799,6 +828,20 @@ public:
 		return ConcatTCHARsToFString<const FString&>(Lhs, Rhs);
 	}
 
+	/**
+	 * Concatenates a Range of TCHARs to an FString (e.g. FStringView).
+	 *
+	 * @param Lhs The CharRangeType on the left-hand-side of the expression.
+	 * @param Rhs The FString on the right-hand-side of the expression.
+	 *
+	 * @return The concatenated string.
+	 */
+	template <typename CharRangeType, typename TEnableIf<TIsTCharRangeNotCArrayNotFString<CharRangeType>::Value>::Type* = nullptr>
+	FORCEINLINE friend FString operator+(CharRangeType&& Lhs, const FString& Rhs)
+	{
+		return ConcatTCHARRangeToFString<const FString&>(GetData(Lhs), GetNum(Lhs), Rhs);
+	}
+
 	/**
 	 * Concatenates a TCHAR array to an FString.
 	 * 
@@ -812,6 +855,20 @@ public:
 		return ConcatTCHARsToFString<FString&&>(Lhs, MoveTemp(Rhs));
 	}
 
+	/**
+	 * Concatenates a Range of TCHARs to an FString (e.g. FStringView).
+	 *
+	 * @param Lhs The CharRangeType on the left-hand-side of the expression.
+	 * @param Rhs The FString on the right-hand-side of the expression.
+	 *
+	 * @return The concatenated string.
+	 */
+	template <typename CharRangeType, typename TEnableIf<TIsTCharRangeNotCArrayNotFString<CharRangeType>::Value>::Type* = nullptr>
+	FORCEINLINE friend FString operator+(CharRangeType&& Lhs, FString&& Rhs)
+	{
+		return ConcatTCHARRangeToFString<FString&&>(GetData(Lhs), GetNum(Lhs), MoveTemp(Rhs));
+	}
+
 	/**
 	 * Concatenates an FString to a TCHAR array.
 	 * 
@@ -825,6 +882,20 @@ public:
 		return ConcatFStringToTCHARs<const FString&>(Lhs, Rhs);
 	}
 
+	/**
+	 * Concatenates an FString to a range of TCHARs (e.g. FStringView).
+	 *
+	 * @param Lhs The FString on the left-hand-side of the expression.
+	 * @param Rhs The CharRangeType on the right-hand-side of the expression.
+	 *
+	 * @return The concatenated string.
+	 */
+	template <typename CharRangeType, typename TEnableIf<TIsTCharRangeNotCArrayNotFString<CharRangeType>::Value>::Type* = nullptr>
+	FORCEINLINE friend FString operator+(const FString& Lhs, CharRangeType&& Rhs)
+	{
+		return ConcatFStringToTCHARRange<const FString&>(Lhs, GetData(Rhs), GetNum(Rhs));
+	}
+
 	/**
 	 * Concatenates an FString to a TCHAR array.
 	 * 
@@ -838,6 +909,20 @@ public:
 		return ConcatFStringToTCHARs<FString&&>(MoveTemp(Lhs), Rhs);
 	}
 
+	/**
+	 * Concatenates an FString to a range of TCHARs (e.g. FStringView).
+	 *
+	 * @param Lhs The FString on the left-hand-side of the expression.
+	 * @param Rhs The CharRangeType on the right-hand-side of the expression.
+	 *
+	 * @return The concatenated string.
+	 */
+	template <typename CharRangeType, typename TEnableIf<TIsTCharRangeNotCArrayNotFString<CharRangeType>::Value>::Type* = nullptr>
+	FORCEINLINE friend FString operator+(FString&& Lhs, CharRangeType&& Rhs)
+	{
+		return ConcatFStringToTCHARRange<FString&&>(MoveTemp(Lhs), GetData(Rhs), GetNum(Rhs));
+	}
+
 	/**
 	 * Concatenate this path with given path ensuring the / character is used between them
 	 * 
