Fix static analysis warning in UnrealMath.cpp for unreachable code

#rnx

diff --git a/Engine/Source/Runtime/Core/Private/Math/UnrealMath.cpp b/Engine/Source/Runtime/Core/Private/Math/UnrealMath.cpp
--- a/Engine/Source/Runtime/Core/Private/Math/UnrealMath.cpp
+++ b/Engine/Source/Runtime/Core/Private/Math/UnrealMath.cpp
@@ -3282,9 +3282,6 @@ PrecLoop:
 		UE_LOG(LogUnrealMath, Log, TEXT("Expression Error : Unrecognized Operator"));
 		return 0;
 	}
-
-	*pResult = V;
-	return 1;
 }
 
 bool FMath::Eval( FString Str, float& OutValue )
