Copying //UE5/Dev-Cooker@14539516 to Main (//UE5/Main)

diff --git a/Engine/Source/Programs/UnrealHeaderTool/Private/HeaderParser.cpp b/Engine/Source/Programs/UnrealHeaderTool/Private/HeaderParser.cpp
--- a/Engine/Source/Programs/UnrealHeaderTool/Private/HeaderParser.cpp
+++ b/Engine/Source/Programs/UnrealHeaderTool/Private/HeaderParser.cpp
@@ -107,6 +107,7 @@ const FName FHeaderParserNames::NAME_IsConversionRoot(TEXT("IsConversionRoot"));
 const FName FHeaderParserNames::NAME_AdvancedClassDisplay(TEXT("AdvancedClassDisplay"));
 
 EGeneratedCodeVersion FHeaderParser::DefaultGeneratedCodeVersion = EGeneratedCodeVersion::V1;
+ENativePointerMemberBehavior FHeaderParser::NativePointerMemberBehavior = ENativePointerMemberBehavior::AllowSilently;
 TArray<FString> FHeaderParser::StructsWithNoPrefix;
 TArray<FString> FHeaderParser::StructsWithTPrefix;
 FRigVMStructMap FHeaderParser::StructRigVMMap;
@@ -868,6 +869,15 @@ namespace
 				return Result;
 			}
 
+			case CPT_ObjectPtrReference:
+			{
+				check(VarProperty.PropertyClass);
+
+				FObjectPtrProperty* Result = new FObjectPtrProperty(Scope, Name, ObjectFlags);
+				Result->PropertyClass = VarProperty.PropertyClass;
+				return Result;
+			}
+
 			case CPT_SoftObjectReference:
 				check(VarProperty.PropertyClass);
 
@@ -1222,6 +1232,7 @@ namespace
 			|| Property->IsA<FClassProperty>()
 			|| Property->IsA<FSoftObjectProperty>()
 			|| Property->IsA<FObjectProperty>()
+			|| Property->IsA<FObjectPtrProperty>()
 			|| Property->IsA<FFloatProperty>()
 			|| Property->IsA<FDoubleProperty>()
 			|| Property->IsA<FIntProperty>()
@@ -1292,6 +1303,27 @@ namespace
 		SkipAlignasIfNecessary(Parser);
 	}
 
+	inline ENativePointerMemberBehavior ToNativePointerMemberBehavior(const FString& InString)
+	{
+		if (InString.Compare(TEXT("Disallow")) == 0)
+		{
+			return ENativePointerMemberBehavior::Disallow;
+		}
+
+		if (InString.Compare(TEXT("AllowSilently")) == 0)
+		{
+			return ENativePointerMemberBehavior::AllowSilently;
+		}
+
+		if (InString.Compare(TEXT("AllowAndLog")) == 0)
+		{
+			return ENativePointerMemberBehavior::AllowAndLog;
+		}
+
+		FError::Throwf(TEXT("Unrecognized native pointer member behavior: %s"), *InString );
+		return ENativePointerMemberBehavior::Disallow;
+	}
+
 	static const TCHAR* GLayoutMacroNames[] = {
 		TEXT("LAYOUT_ARRAY"),
 		TEXT("LAYOUT_ARRAY_EDITORONLY"),
@@ -1319,15 +1351,18 @@ FScriptLocation::FScriptLocation()
 /////////////////////////////////////////////////////
 // FHeaderParser
 
-FString FHeaderParser::GetContext()
+FString FHeaderParser::GetSourceFileContext() const
 {
 	FFileScope* FileScope = GetCurrentFileScope();
 	FUnrealSourceFile* SourceFile = FileScope ? FileScope->GetSourceFile() : GetCurrentSourceFile();
-	FString ScopeFilename = SourceFile
+	return SourceFile
 		? IFileManager::Get().ConvertToAbsolutePathForExternalAppForRead(*SourceFile->GetFilename())
 		: TEXT("UNKNOWN");
+}
 
-	return FString::Printf(TEXT("%s(%i)"), *ScopeFilename, InputLine);
+FString FHeaderParser::GetContext()
+{
+	return FString::Printf(TEXT("%s(%i)"), *GetSourceFileContext(), InputLine);
 }
 
 /*-----------------------------------------------------------------------------
@@ -4185,6 +4220,13 @@ void FHeaderParser::GetVarType(
 		}
 	}
 
+	if (MatchIdentifier(TEXT("mutable"), ESearchCase::CaseSensitive))
+	{
+		//@TODO: Should flag as settable from a const context, but this is at least good enough to allow use for C++ land
+	}
+
+	const int32 VarStartPos = InputPos;
+
 	if (MatchIdentifier(TEXT("const"), ESearchCase::CaseSensitive))
 	{
 		//@TODO: UCREMOVAL: Should use this to set the new (currently non-existent) CPF_Const flag appropriately!
@@ -4192,11 +4234,6 @@ void FHeaderParser::GetVarType(
 		bNativeConst = true;
 	}
 
-	if (MatchIdentifier(TEXT("mutable"), ESearchCase::CaseSensitive))
-	{
-		//@TODO: Should flag as settable from a const context, but this is at least good enough to allow use for C++ land
-	}
-
 	if (MatchIdentifier(TEXT("struct"), ESearchCase::CaseSensitive))
 	{
 		bUnconsumedStructKeyword = true;
@@ -4767,16 +4804,9 @@ void FHeaderParser::GetVarType(
 			// An object reference of some type (maybe a restricted class?)
 			UClass* TempClass = NULL;
 
-			const bool bIsLazyPtrTemplate        = VarType.Matches(TEXT("TLazyObjectPtr"), ESearchCase::CaseSensitive);
+			EPropertyType PropertyType = CPT_ObjectReference;
 			const bool bIsSoftObjectPtrTemplate  = VarType.Matches(TEXT("TSoftObjectPtr"), ESearchCase::CaseSensitive);
-			const bool bIsSoftClassPtrTemplate   = VarType.Matches(TEXT("TSoftClassPtr"), ESearchCase::CaseSensitive);
-			const bool bIsWeakPtrTemplate        = VarType.Matches(TEXT("TWeakObjectPtr"), ESearchCase::CaseSensitive);
-			const bool bIsAutoweakPtrTemplate    = VarType.Matches(TEXT("TAutoWeakObjectPtr"), ESearchCase::CaseSensitive);
-			const bool bIsScriptInterfaceWrapper = VarType.Matches(TEXT("TScriptInterface"), ESearchCase::CaseSensitive);
-
-			bool bIsWeak     = false;
-			bool bIsLazy     = false;
-			bool bIsSoft     = false;
+
 			bool bWeakIsAuto = false;
 
 			if (VarType.Matches(TEXT("TSubclassOf"), ESearchCase::CaseSensitive))
@@ -4788,71 +4818,97 @@ void FHeaderParser::GetVarType(
 				TempClass = UInterface::StaticClass();
 				Flags |= CPF_UObjectWrapper;
 			}
-			else if (bIsSoftClassPtrTemplate)
+			else if (VarType.Matches(TEXT("TSoftClassPtr"), ESearchCase::CaseSensitive))
 			{
 				TempClass = UClass::StaticClass();
-				bIsSoft = true;
+				PropertyType = CPT_SoftObjectReference;
 			}
-			else if (bIsLazyPtrTemplate || bIsWeakPtrTemplate || bIsAutoweakPtrTemplate || bIsScriptInterfaceWrapper || bIsSoftObjectPtrTemplate)
+			else
 			{
-				RequireSymbol(TEXT('<'), VarType.Identifier);
-
-				// Consume a forward class declaration 'class' if present
-				MatchIdentifier(TEXT("class"), ESearchCase::CaseSensitive);
+				const bool bIsLazyPtrTemplate        = VarType.Matches(TEXT("TLazyObjectPtr"), ESearchCase::CaseSensitive);
+				const bool bIsObjectPtrTemplate      = VarType.Matches(TEXT("TObjectPtr"), ESearchCase::CaseSensitive);
+				const bool bIsWeakPtrTemplate        = VarType.Matches(TEXT("TWeakObjectPtr"), ESearchCase::CaseSensitive);
+				const bool bIsAutoweakPtrTemplate    = VarType.Matches(TEXT("TAutoWeakObjectPtr"), ESearchCase::CaseSensitive);
+				const bool bIsScriptInterfaceWrapper = VarType.Matches(TEXT("TScriptInterface"), ESearchCase::CaseSensitive);
 
-				// Also consume const
-				bNativeConstTemplateArg |= MatchIdentifier(TEXT("const"), ESearchCase::CaseSensitive);
-				
-				// Find the lazy/weak class
-				FToken InnerClass;
-				if (GetIdentifier(InnerClass))
+				if (bIsLazyPtrTemplate || bIsObjectPtrTemplate || bIsWeakPtrTemplate || bIsAutoweakPtrTemplate || bIsScriptInterfaceWrapper || bIsSoftObjectPtrTemplate)
 				{
-					RedirectTypeIdentifier(InnerClass);
+					RequireSymbol(TEXT('<'), VarType.Identifier);
 
-					TempClass = AllClasses.FindScriptClass(InnerClass.Identifier);
-					if (TempClass == nullptr)
-					{
-						FError::Throwf(TEXT("Unrecognized type '%s' (in expression %s<%s>) - type must be a UCLASS"), InnerClass.Identifier, VarType.Identifier, InnerClass.Identifier);
-					}
+					// Also consume const
+					bNativeConstTemplateArg |= MatchIdentifier(TEXT("const"), ESearchCase::CaseSensitive);
 
-					if (bIsAutoweakPtrTemplate)
+					// Consume a forward class declaration 'class' if present
+					MatchIdentifier(TEXT("class"), ESearchCase::CaseSensitive);
+					
+					// Find the lazy/weak class
+					FToken InnerClass;
+					if (GetIdentifier(InnerClass))
 					{
-						bIsWeak = true;
-						bWeakIsAuto = true;
-					}
-					else if (bIsLazyPtrTemplate)
-					{
-						bIsLazy = true;
-					}
-					else if (bIsWeakPtrTemplate)
-					{
-						bIsWeak = true;
+						RedirectTypeIdentifier(InnerClass);
+
+						TempClass = AllClasses.FindScriptClass(InnerClass.Identifier);
+						if (TempClass == nullptr)
+						{
+							FError::Throwf(TEXT("Unrecognized type '%s' (in expression %s<%s>) - type must be a UCLASS"), InnerClass.Identifier, VarType.Identifier, InnerClass.Identifier);
+						}
+
+						if (bIsAutoweakPtrTemplate)
+						{
+							PropertyType = CPT_WeakObjectReference;
+							bWeakIsAuto = true;
+						}
+						else if (bIsLazyPtrTemplate)
+						{
+							PropertyType = CPT_LazyObjectReference;
+						}
+						else if (bIsObjectPtrTemplate)
+						{
+							PropertyType = CPT_ObjectPtrReference;
+						}
+						else if (bIsWeakPtrTemplate)
+						{
+							PropertyType = CPT_WeakObjectReference;
+						}
+						else if (bIsSoftObjectPtrTemplate)
+						{
+							PropertyType = CPT_SoftObjectReference;
+						}
+
+						Flags |= CPF_UObjectWrapper;
 					}
-					else if (bIsSoftObjectPtrTemplate)
+					else
 					{
-						bIsSoft = true;
+						FError::Throwf(TEXT("%s: Missing template type"), VarType.Identifier);
 					}
 
-					Flags |= CPF_UObjectWrapper;
+					// Const after template argument type but before end of template symbol
+					bNativeConstTemplateArg |= MatchIdentifier(TEXT("const"), ESearchCase::CaseSensitive);
+
+					RequireSymbol(TEXT('>'), VarType.Identifier, ESymbolParseOption::CloseTemplateBracket);
 				}
 				else
 				{
-					FError::Throwf(TEXT("%s: Missing template type"), VarType.Identifier);
+					TempClass = AllClasses.FindScriptClass(VarType.Identifier);
 				}
-
-				RequireSymbol(TEXT('>'), VarType.Identifier, ESymbolParseOption::CloseTemplateBracket);
-			}
-			else
-			{
-				TempClass = AllClasses.FindScriptClass(VarType.Identifier);
 			}
 
 			if (TempClass != NULL)
 			{
 				bHandledType = true;
 
-				bool bAllowWeak = !(Disallow & CPF_AutoWeak); // if it is not allowing anything, force it strong. this is probably a function arg
-				VarProperty = FPropertyBase(TempClass, bAllowWeak && bIsWeak, bWeakIsAuto, bIsLazy, bIsSoft);
+				if ((PropertyType == CPT_WeakObjectReference) && (Disallow & CPF_AutoWeak)) // if it is not allowing anything, force it strong. this is probably a function arg
+				{
+					PropertyType = CPT_ObjectReference;
+				}
+
+				// if this is an interface class, we use the FInterfaceProperty class instead of FObjectProperty
+				if ((PropertyType == CPT_ObjectReference) && TempClass->HasAnyClassFlags(CLASS_Interface))
+				{
+					PropertyType = CPT_Interface;
+				}
+
+				VarProperty = FPropertyBase(TempClass, PropertyType, bWeakIsAuto);
 				if (TempClass->IsChildOf(UClass::StaticClass()))
 				{
 					if ( MatchSymbol(TEXT('<')) )
@@ -4880,12 +4936,11 @@ void FHeaderParser::GetVarType(
 						VarProperty.MetaClass = UObject::StaticClass();
 					}
 
-					if (bIsWeak)
+					if (PropertyType == CPT_WeakObjectReference)
 					{
 						FError::Throwf(TEXT("Class variables cannot be weak, they are always strong."));
 					}
-
-					if (bIsLazy)
+					else if (PropertyType == CPT_LazyObjectReference)
 					{
 						FError::Throwf(TEXT("Class variables cannot be lazy, they are always strong."));
 					}
@@ -4910,9 +4965,35 @@ void FHeaderParser::GetVarType(
 
 					RequireSymbol(TEXT('*'), TEXT("Expected a pointer type"));
 
-					// Swallow trailing 'const' after pointer properties
+					// Optionally emit messages about native pointer members and swallow trailing 'const' after pointer properties
 					if (VariableCategory == EVariableCategory::Member)
 					{
+						switch (NativePointerMemberBehavior)
+						{
+							case ENativePointerMemberBehavior::Disallow:
+							{
+								FString TypeName = FString(InputPos - VarStartPos, Input + VarStartPos).TrimStartAndEnd().ReplaceCharWithEscapedChar();
+								UE_LOG(LogCompile, Error, TEXT("Native pointer usage in member declaration detected in '%s', line %d [[%s]]"), *GetSourceFileContext(), InputLine, *TypeName);
+							}
+							break;
+							case ENativePointerMemberBehavior::AllowAndLog:
+							{
+								FString TypeName = FString(InputPos - VarStartPos, Input + VarStartPos).TrimStartAndEnd().ReplaceCharWithEscapedChar();
+								UE_LOG(LogCompile, Log, TEXT("Native pointer usage in member declaration detected in '%s', line %d [[%s]]"), *GetSourceFileContext(), InputLine, *TypeName);
+							}
+							break;
+							case ENativePointerMemberBehavior::AllowSilently:
+							{
+								// Do nothing
+							}
+							break;
+							default:
+							{
+								checkf(false, TEXT("Unhandled case"));
+							}
+							break;
+						}
+
 						MatchIdentifier(TEXT("const"), ESearchCase::CaseSensitive);
 					}
 
@@ -5062,7 +5143,7 @@ void FHeaderParser::GetVarType(
 	// Perform some more specific validation on the property flags
 	if (VarProperty.PropertyFlags & CPF_PersistentInstance)
 	{
-		if (VarProperty.Type == CPT_ObjectReference)
+		if ((VarProperty.Type == CPT_ObjectReference) || (VarProperty.Type == CPT_ObjectPtrReference))
 		{
 			if (VarProperty.PropertyClass->IsChildOf<UClass>())
 			{
@@ -5433,6 +5514,11 @@ FProperty* FHeaderParser::GetVarNameAndDim
 		{
 			FError::Throwf(TEXT("UFunctions cannot take a lazy pointer as a parameter."));
 		}
+		else if (VarProperty.Type == CPT_ObjectPtrReference)
+		{
+			// @TODO: OBJPTR: Investigate TObjectPtr support for UFunction parameters.
+			FError::Throwf(TEXT("UFunctions cannot take a TObjectPtr as a parameter."));
+		}
 	}
 
 	// If this is the first time seeing the property name, then flag it for replace instead of add
@@ -7995,7 +8081,7 @@ void FHeaderParser::ValidatePropertyIsDeprecatedIfNecessary(const FPropertyBase&
 
 	// check to see if we have a FObjectProperty using a deprecated class.
 	// PropertyClass is part of a union, so only check PropertyClass if this token represents an object property
-	if ( (VarProperty.Type == CPT_ObjectReference || VarProperty.Type == CPT_WeakObjectReference || VarProperty.Type == CPT_LazyObjectReference || VarProperty.Type == CPT_SoftObjectReference) && VarProperty.PropertyClass != NULL
+	if ( (VarProperty.Type == CPT_ObjectReference || VarProperty.Type == CPT_WeakObjectReference || VarProperty.Type == CPT_LazyObjectReference || VarProperty.Type == CPT_ObjectPtrReference || VarProperty.Type == CPT_SoftObjectReference) && VarProperty.PropertyClass != NULL
 		&&	VarProperty.PropertyClass->HasAnyClassFlags(CLASS_Deprecated)	// and the object class being used has been deprecated
 		&& (VarProperty.PropertyFlags&CPF_Deprecated) == 0					// and this property isn't marked deprecated as well
 		&& (OuterPropertyType == NULL || !(OuterPropertyType->PropertyFlags & CPF_Deprecated)) ) // and this property isn't in an array that was marked deprecated either
@@ -8019,6 +8105,7 @@ struct FExposeOnSpawnValidator
 		case CPT_Bool:
 		case CPT_Bool8:
 		case CPT_ObjectReference:
+		case CPT_ObjectPtrReference:
 		case CPT_String:
 		case CPT_Text:
 		case CPT_Name:
@@ -8834,6 +8921,7 @@ FHeaderParser::FHeaderParser(FFeedbackContext* InWarn, const FManifestModule& In
 		const FName StructsWithTPrefixKey(TEXT("StructsWithTPrefix"));
 		const FName DelegateParameterCountStringsKey(TEXT("DelegateParameterCountStrings"));
 		const FName GeneratedCodeVersionKey(TEXT("GeneratedCodeVersion"));
+		const FName NativePointerMemberBehaviorKey(TEXT("NativePointerMemberBehavior"));
 
 		FConfigSection* ConfigSection = GConfig->GetSectionPrivate(TEXT("UnrealHeaderTool"), false, true, GEngineIni);
 		if (ConfigSection)
@@ -8866,6 +8954,10 @@ FHeaderParser::FHeaderParser(FFeedbackContext* InWarn, const FManifestModule& In
 				{
 					DefaultGeneratedCodeVersion = ToGeneratedCodeVersion(It.Value().GetValue());
 				}
+				else if (It.Key() == NativePointerMemberBehaviorKey)
+				{
+					NativePointerMemberBehavior = ToNativePointerMemberBehavior(It.Value().GetValue());
+				}
 			}
 		}
 		bConfigOptionsInitialized = true;
diff --git a/Engine/Source/Programs/UnrealHeaderTool/Private/HeaderParser.h b/Engine/Source/Programs/UnrealHeaderTool/Private/HeaderParser.h
--- a/Engine/Source/Programs/UnrealHeaderTool/Private/HeaderParser.h
+++ b/Engine/Source/Programs/UnrealHeaderTool/Private/HeaderParser.h
@@ -417,6 +417,13 @@ struct FDocumentationPolicy
 	bool bFloatRangesRequired = UHT_DOCUMENTATION_POLICY_DEFAULT;
 };
 
+enum class ENativePointerMemberBehavior
+{
+	Disallow,
+	AllowSilently,
+	AllowAndLog,
+};
+
 /////////////////////////////////////////////////////
 // FHeaderParser
 
@@ -430,6 +437,8 @@ public:
 	// Default version of generated code. Defaults to oldest possible, unless specified otherwise in config.
 	static EGeneratedCodeVersion DefaultGeneratedCodeVersion;
 
+	static ENativePointerMemberBehavior NativePointerMemberBehavior;
+
 	// Compute the function parameter size and save the return offset
 	static void ComputeFunctionParametersSize(UClass* InClass);
 
@@ -769,6 +778,8 @@ protected:
 		const FManifestModule& Module
 	);
 
+	FString GetSourceFileContext() const;
+
 	// FContextSupplier interface.
 	virtual FString GetContext() override;
 	// End of FContextSupplier interface.
diff --git a/Engine/Source/Programs/UnrealHeaderTool/Private/ParserHelper.cpp b/Engine/Source/Programs/UnrealHeaderTool/Private/ParserHelper.cpp
--- a/Engine/Source/Programs/UnrealHeaderTool/Private/ParserHelper.cpp
+++ b/Engine/Source/Programs/UnrealHeaderTool/Private/ParserHelper.cpp
@@ -127,6 +127,7 @@ const TCHAR* FPropertyBase::GetPropertyTypeText( EPropertyType Type )
 		CASE_TEXT(CPT_SoftObjectReference);
 		CASE_TEXT(CPT_WeakObjectReference);
 		CASE_TEXT(CPT_LazyObjectReference);
+		CASE_TEXT(CPT_ObjectPtrReference);
 		CASE_TEXT(CPT_Map);
 		CASE_TEXT(CPT_Set);
 		CASE_TEXT(CPT_FieldPath);
diff --git a/Engine/Source/Programs/UnrealHeaderTool/Private/ParserHelper.h b/Engine/Source/Programs/UnrealHeaderTool/Private/ParserHelper.h
--- a/Engine/Source/Programs/UnrealHeaderTool/Private/ParserHelper.h
+++ b/Engine/Source/Programs/UnrealHeaderTool/Private/ParserHelper.h
@@ -203,8 +203,8 @@ public:
 	{
 	}
 
-	explicit FPropertyBase(UClass* InClass, bool bIsWeak = false, bool bWeakIsAuto = false, bool bIsLazy = false, bool bIsSoft = false)
-	: Type                (CPT_ObjectReference)
+	explicit FPropertyBase(UClass* InClass, EPropertyType InType, bool bWeakIsAuto = false)
+	: Type                (InType)
 	, ArrayType           (EArrayType::None)
 	, PropertyFlags       (CPF_None)
 	, ImpliedPropertyFlags(CPF_None)
@@ -218,26 +218,9 @@ public:
 	, PointerType         (EPointerType::None)
 	, IntType             (EIntType::None)
 	{
-		// if this is an interface class, we use the FInterfaceProperty class instead of FObjectProperty
-		if ( InClass->HasAnyClassFlags(CLASS_Interface) )
-		{
-			Type = CPT_Interface;
-		}
-		if (bIsLazy)
+		if ((Type == CPT_WeakObjectReference) && bWeakIsAuto)
 		{
-			Type = CPT_LazyObjectReference;
-		}
-		else if (bIsSoft)
-		{
-			Type = CPT_SoftObjectReference;
-		}
-		else if (bIsWeak)
-		{
-			Type = CPT_WeakObjectReference;
-			if (bWeakIsAuto)
-			{
-				PropertyFlags |= CPF_AutoWeak;
-			}
+			PropertyFlags |= CPF_AutoWeak;
 		}
 	}
 
@@ -413,6 +396,11 @@ public:
 			*this = FPropertyBase(CPT_LazyObjectReference);
 			PropertyClass = CastField<FLazyObjectProperty>(Property)->PropertyClass;
 		}
+		else if( ClassOfProperty==FObjectPtrProperty::StaticClass() )
+		{
+			*this = FPropertyBase(CPT_ObjectPtrReference);
+			PropertyClass = CastField<FObjectPtrProperty>(Property)->PropertyClass;
+		}
 		else if( ClassOfProperty==FSoftClassProperty::StaticClass() )
 		{
 			*this = FPropertyBase(CPT_SoftObjectReference);
@@ -499,7 +487,7 @@ public:
 	 */
 	bool IsObject() const
 	{
-		return Type == CPT_ObjectReference || Type == CPT_Interface || Type == CPT_WeakObjectReference || Type == CPT_LazyObjectReference || Type == CPT_SoftObjectReference;
+		return Type == CPT_ObjectReference || Type == CPT_Interface || Type == CPT_WeakObjectReference || Type == CPT_LazyObjectReference || Type == CPT_ObjectPtrReference || Type == CPT_SoftObjectReference;
 	}
 
 	bool IsContainer() const
@@ -569,7 +557,7 @@ public:
 				return false;
 			}
 		}
-		else if ((Type == CPT_ObjectReference || Type == CPT_WeakObjectReference || Type == CPT_LazyObjectReference || Type == CPT_SoftObjectReference) && Other.Type != CPT_Interface && (PropertyFlags & CPF_ReturnParm))
+		else if ((Type == CPT_ObjectReference || Type == CPT_WeakObjectReference || Type == CPT_LazyObjectReference || Type == CPT_ObjectPtrReference || Type == CPT_SoftObjectReference) && Other.Type != CPT_Interface && (PropertyFlags & CPF_ReturnParm))
 		{
 			bReverseClassChainCheck = false;
 		}
diff --git a/Engine/Source/Runtime/AssetRegistry/Private/PackageBuildDependencyTracker.cpp b/Engine/Source/Runtime/AssetRegistry/Private/PackageBuildDependencyTracker.cpp
new file mode 100644
--- /dev/null
+++ b/Engine/Source/Runtime/AssetRegistry/Private/PackageBuildDependencyTracker.cpp
@@ -0,0 +1,86 @@
+// Copyright Epic Games, Inc. All Rights Reserved.
+
+#include "PackageBuildDependencyTracker.h"
+#include "HAL/Platform.h"
+#include "Logging/LogMacros.h"
+#include "UObject/Package.h"
+#include "Misc/PackageAccessTrackingOps.h"
+
+#if UE_WITH_PACKAGE_ACCESS_TRACKING
+
+DEFINE_LOG_CATEGORY_STATIC(LogPackageBuildDependencyTracker, Log, All);
+
+FPackageBuildDependencyTracker FPackageBuildDependencyTracker::Singleton;
+
+void FPackageBuildDependencyTracker::DumpData() const
+{
+	FScopeLock RecordsScopeLock(&RecordsLock);
+	uint64 ReferencingPackageCount = 0;
+	uint64 ReferenceCount = 0;
+	for (const auto& PackageAccessRecord : Records)
+	{
+		++ReferencingPackageCount;
+		for (const auto& AccessedData : PackageAccessRecord.Value)
+		{
+			++ReferenceCount;
+		}
+	}
+	UE_LOG(LogPackageBuildDependencyTracker, Display, TEXT("Package Accesses (%u referencing packages with a total of %u unique accesses):"), ReferencingPackageCount, ReferenceCount);
+	UE_LOG(LogPackageBuildDependencyTracker, Display, TEXT("========================================================================="));
+	for (const auto& PackageAccessRecord : Records)
+	{
+		UE_LOG(LogPackageBuildDependencyTracker, Display, TEXT("%s:"), *PackageAccessRecord.Key.ToString());
+		for (const auto& AccessedData : PackageAccessRecord.Value)
+		{
+			UE_LOG(LogPackageBuildDependencyTracker, Display, TEXT("    %s"), *AccessedData.ToString());
+		}
+	}
+}
+
+FPackageBuildDependencyTracker::FPackageBuildDependencyTracker()
+{
+	PreviousObjectHandleReadFunction = SetObjectHandleReadCallback(StaticOnObjectHandleRead);
+}
+
+void FPackageBuildDependencyTracker::StaticOnObjectHandleRead(UObject* ReadObject)
+{
+	if (ReadObject)
+	{
+		if (PackageAccessTracking_Private::FPackageAccessRefScope* InnermostThreadScope = PackageAccessTracking_Private::FPackageAccessRefScope::GetCurrentThreadScope())
+		{
+			PackageAccessTracking_Private::FPackageAccessRefScope* SearchThreadScope = InnermostThreadScope;
+			
+			do
+			{
+				FName OpName = SearchThreadScope->GetOpName();
+				if ((OpName == PackageAccessTrackingOps::NAME_Load) || (OpName == PackageAccessTrackingOps::NAME_Save) /*|| (OpName == PackageAccessTrackingOps::NAME_CreateDefaultObject)*/)
+				{
+					//If this was an Package Access done in the context of a save or load, we record it as a build dependency.
+					//If we want to capture script package dependencies, we may also need to record accesses under the "PackageAccessTrackingOps::NAME_CreateDefaultObject" operation
+					//which occurs from "UObjectLoadAllCompiledInDefaultProperties" and possibly elsewhere.
+
+					//NOTE: The referencer is the InnermostThreadScope, not the SearchThreadScope.  Consider a situation where the in the process of loading package A,
+					//		we call PostLoad on an object on package B which references an object in package C.  The SearchThreadScope with OpName == Load is for package A,
+					//		but we want to record that B depends on C, not that A depends on C.  The InnermostThreadScope will have a package name of B and an OpName of PostLoad
+					//		and so we want the package name from the InnermostThreadScope while having searched upwards for a scope with an OpName == Load.
+					FName Referencer = InnermostThreadScope->GetPackageName();
+					FName Referenced = ReadObject->GetOutermost()->GetFName();
+					if (Referencer != Referenced)
+					{
+						FScopeLock RecordsScopeLock(&Singleton.RecordsLock);
+						Singleton.Records.FindOrAdd(Referencer).Add(Referenced);
+					}
+					break;
+				}
+				SearchThreadScope = SearchThreadScope->GetOuter();
+			} while (SearchThreadScope);
+		}
+	}
+
+	if (Singleton.PreviousObjectHandleReadFunction)
+	{
+		Singleton.PreviousObjectHandleReadFunction(ReadObject);
+	}
+}
+
+#endif // UE_WITH_OBJECT_HANDLE_TRACKING
diff --git a/Engine/Source/Runtime/AssetRegistry/Private/PackageBuildDependencyTracker.h b/Engine/Source/Runtime/AssetRegistry/Private/PackageBuildDependencyTracker.h
new file mode 100644
--- /dev/null
+++ b/Engine/Source/Runtime/AssetRegistry/Private/PackageBuildDependencyTracker.h
@@ -0,0 +1,42 @@
+// Copyright Epic Games, Inc. All Rights Reserved.
+
+#pragma once
+
+#include "HAL/CriticalSection.h"
+#include "Misc/PackageAccessTracking.h"
+#include "Templates/UnrealTemplate.h"
+#include "UObject/ObjectHandle.h"
+
+#if UE_WITH_PACKAGE_ACCESS_TRACKING
+
+class FPackageBuildDependencyTracker: public FNoncopyable
+{
+public:
+	static ASSETREGISTRY_API FPackageBuildDependencyTracker& Get() { return Singleton; }
+
+	ASSETREGISTRY_API void DumpData() const;
+
+private:
+	FPackageBuildDependencyTracker();
+
+	/** Track object reference reads */
+	static void StaticOnObjectHandleRead(UObject* ReadObject);
+	ObjectHandleReadFunction* PreviousObjectHandleReadFunction = nullptr;
+
+	mutable FCriticalSection RecordsLock;
+	TMap<FName, TSet<FName>> Records;
+	static FPackageBuildDependencyTracker Singleton;
+};
+
+ASSETREGISTRY_API void DumpBuildDependencyTrackerData()
+{
+	FPackageBuildDependencyTracker::Get().DumpData();
+}
+
+#else
+
+ASSETREGISTRY_API void DumpBuildDependencyTrackerData()
+{
+}
+
+#endif // UE_WITH_PACKAGE_ACCESS_TRACKING
diff --git a/Engine/Source/Runtime/ClothingSystemRuntimeCommon/Private/ClothingAsset.cpp b/Engine/Source/Runtime/ClothingSystemRuntimeCommon/Private/ClothingAsset.cpp
--- a/Engine/Source/Runtime/ClothingSystemRuntimeCommon/Private/ClothingAsset.cpp
+++ b/Engine/Source/Runtime/ClothingSystemRuntimeCommon/Private/ClothingAsset.cpp
@@ -885,19 +885,19 @@ void UClothingAssetCommon::PostLoad()
 				}
 			}
 			// And keep the old config too
-			SetClothConfig(ClothSimConfig_DEPRECATED);
+			SetClothConfig(ToRawPtr(ClothSimConfig_DEPRECATED));
 			ClothSimConfig_DEPRECATED = nullptr;
 			bMigrateSharedConfigToConfig = false;
 		}
 		if (ChaosClothSimConfig_DEPRECATED)
 		{
-			SetClothConfig(ChaosClothSimConfig_DEPRECATED);
+			SetClothConfig(ToRawPtr(ChaosClothSimConfig_DEPRECATED));
 			ChaosClothSimConfig_DEPRECATED = nullptr;
 			bMigrateSharedConfigToConfig = false;
 		}
 		if (ClothSharedSimConfig_DEPRECATED)
 		{
-			SetClothConfig(ClothSharedSimConfig_DEPRECATED);
+			SetClothConfig(ToRawPtr(ClothSharedSimConfig_DEPRECATED));
 			ClothSharedSimConfig_DEPRECATED = nullptr;
 			bMigrateSharedConfigToConfig = false;
 		}
diff --git a/Engine/Source/Runtime/Core/Private/Memory/SharedBuffer.cpp b/Engine/Source/Runtime/Core/Private/Memory/SharedBuffer.cpp
new file mode 100644
--- /dev/null
+++ b/Engine/Source/Runtime/Core/Private/Memory/SharedBuffer.cpp
@@ -0,0 +1,41 @@
+// Copyright Epic Games, Inc. All Rights Reserved.
+
+#include "Memory/SharedBuffer.h"
+
+FSharedBuffer::FSharedBuffer(const uint64 InSize)
+{
+	checkf(InSize < (uint64(1) << 48), TEXT("Size %" UINT64_FMT " exceeds maximum shared buffer size of 256 TiB"), InSize);
+	Data = FMemory::Malloc(InSize);
+	DataSize = InSize;
+	Flags = uint16(ESharedBufferFlags::Owned);
+}
+
+FSharedBuffer::FSharedBuffer(EAssumeOwnershipTag, void* const InData, const uint64 InSize)
+{
+	checkf(InSize < (uint64(1) << 48), TEXT("Size %" UINT64_FMT " exceeds maximum shared buffer size of 256 TiB"), InSize);
+	Data = InData;
+	DataSize = InSize;
+	Flags = uint16(ESharedBufferFlags::Owned);
+}
+
+FSharedBuffer::FSharedBuffer(ECloneTag, const void* const InData, const uint64 InSize)
+	: FSharedBuffer(InSize)
+{
+	FMemory::Memcpy(Data, InData, InSize);
+}
+
+FSharedBuffer::FSharedBuffer(EWrapTag, void* const InData, const uint64 InSize)
+{
+	checkf(InSize < (uint64(1) << 48), TEXT("Size %" UINT64_FMT " exceeds maximum shared buffer size of 256 TiB"), InSize);
+	Data = InData;
+	DataSize = InSize;
+	Flags = uint16(ESharedBufferFlags::None);
+}
+
+FSharedBuffer::~FSharedBuffer()
+{
+	if (IsOwned())
+	{
+		FMemory::Free(Data);
+	}
+}
diff --git a/Engine/Source/Runtime/Core/Private/Misc/CoreGlobals.cpp b/Engine/Source/Runtime/Core/Private/Misc/CoreGlobals.cpp
--- a/Engine/Source/Runtime/Core/Private/Misc/CoreGlobals.cpp
+++ b/Engine/Source/Runtime/Core/Private/Misc/CoreGlobals.cpp
@@ -230,9 +230,15 @@ float					GNearClippingPlane				= 10.0f;				/* Near clipping plane */
 bool					GExitPurge						= false;
 
 FChunkedFixedUObjectArray* GCoreObjectArrayForDebugVisualizers = nullptr;
+template class CORE_API TArray<FMinimalName, TInlineAllocator<3>>;
+
+TArray<FMinimalName, TInlineAllocator<3>>* GCoreComplexObjectPathDebug = nullptr;
+FObjectHandlePackageDebugData* GCoreObjectHandlePackageDebug = nullptr;
 #if PLATFORM_UNIX
 uint8** CORE_API GNameBlocksDebug = FNameDebugVisualizer::GetBlocks();
 FChunkedFixedUObjectArray*& CORE_API GObjectArrayForDebugVisualizers = GCoreObjectArrayForDebugVisualizers;
+TArray<FMinimalName, TInlineAllocator<3>>*& GComplexObjectPathDebug = GCoreComplexObjectPathDebug;
+FObjectHandlePackageDebugData*& CORE_API GObjectHandlePackageDebug = GCoreObjectHandlePackageDebug;
 #endif
 
 /** Game name, used for base game directory and ini among other things										*/
diff --git a/Engine/Source/Runtime/Core/Private/Serialization/Archive.cpp b/Engine/Source/Runtime/Core/Private/Serialization/Archive.cpp
--- a/Engine/Source/Runtime/Core/Private/Serialization/Archive.cpp
+++ b/Engine/Source/Runtime/Core/Private/Serialization/Archive.cpp
@@ -457,6 +457,13 @@ FArchive& FArchive::operator<<(struct FLazyObjectPtr& Value)
 	return *this;
 }
 
+FArchive& FArchive::operator<<(struct FObjectPtr& Value)
+{
+	// The base FArchive does not implement this method. Use FArchiveUObject instead.
+	UE_LOG(LogSerialization, Fatal, TEXT("FArchive does not support FObjectPtr serialization. Use FArchiveUObject instead."));
+	return *this;
+}
+
 FArchive& FArchive::operator<<(struct FSoftObjectPtr& Value)
 {
 	// The base FArchive does not implement this method. Use FArchiveUObject instead.
@@ -678,6 +685,7 @@ public:
 
 void FArchive::SerializeCompressed(void* V, int64 Length, FName CompressionFormat, ECompressionFlags Flags, bool bTreatBufferAsFileReader)
 {
+	TRACE_CPUPROFILER_EVENT_SCOPE(FArchive::SerializeCompressed);
 	if( IsLoading() )
 	{
 		if (CompressionFormat == NAME_Zlib && FPlatformProperties::GetZlibReplacementFormat() != nullptr)
diff --git a/Engine/Source/Runtime/Core/Private/Serialization/CompactBinary.cpp b/Engine/Source/Runtime/Core/Private/Serialization/CompactBinary.cpp
new file mode 100644
--- /dev/null
+++ b/Engine/Source/Runtime/Core/Private/Serialization/CompactBinary.cpp
@@ -0,0 +1,1054 @@
+// Copyright Epic Games, Inc. All Rights Reserved.
+
+#include "Serialization/CompactBinary.h"
+
+#include "Math/UnrealMathUtility.h"
+#include "Misc/ByteSwap.h"
+#include "Misc/DateTime.h"
+#include "Misc/Guid.h"
+#include "Misc/Timespan.h"
+#include "Serialization/VarInt.h"
+
+///////////////////////////////////////////////////////////////////////////////////////////////////////////////////////
+
+namespace CbObjectPrivate
+{
+	static constexpr const uint8 EmptyObjectPayload[] = { 0x00 };
+	static constexpr const uint8 EmptyArrayPayload[] = { 0x01, 0x00 };
+
+	template <typename T>
+	static constexpr FORCEINLINE T ReadUnaligned(const void* const Memory)
+	{
+	#if PLATFORM_SUPPORTS_UNALIGNED_LOADS
+		return *static_cast<const T*>(Memory);
+	#else
+		T Value;
+		FMemory::Memcpy(&Value, Memory, sizeof(Value));
+		return Value;
+	#endif
+	}
+}
+
+///////////////////////////////////////////////////////////////////////////////////////////////////////////////////////
+
+void FCbFieldType::StaticAssertTypeConstants()
+{
+	constexpr ECbFieldType AllFlags = ECbFieldType::HasFieldName | ECbFieldType::HasFieldType;
+	static_assert(!(TypeMask & AllFlags), "TypeMask is invalid!");
+
+	static_assert(ObjectBase == ECbFieldType::Object, "ObjectBase is invalid!");
+	static_assert((ObjectMask & (AllFlags | ECbFieldType::UniformObject)) == ECbFieldType::Object, "ObjectMask is invalid!");
+	static_assert(!(ObjectMask & (ObjectBase ^ ECbFieldType::UniformObject)), "ObjectMask or ObjectBase is invalid!");
+	static_assert(TypeMask == (ObjectMask | (ObjectBase ^ ECbFieldType::UniformObject)), "ObjectMask or ObjectBase is invalid!");
+
+	static_assert(ArrayBase == ECbFieldType::Array, "ArrayBase is invalid!");
+	static_assert((ArrayMask & (AllFlags | ECbFieldType::UniformArray)) == ECbFieldType::Array, "ArrayMask is invalid!");
+	static_assert(!(ArrayMask & (ArrayBase ^ ECbFieldType::UniformArray)), "ArrayMask or ArrayBase is invalid!");
+	static_assert(TypeMask == (ArrayMask | (ArrayBase ^ ECbFieldType::UniformArray)), "ArrayMask or ArrayBase is invalid!");
+
+	static_assert(IntegerBase == ECbFieldType::IntegerPositive, "IntegerBase is invalid!");
+	static_assert((IntegerMask & (AllFlags | ECbFieldType::IntegerNegative)) == ECbFieldType::IntegerPositive, "IntegerMask is invalid!");
+	static_assert(!(IntegerMask & (IntegerBase ^ ECbFieldType::IntegerNegative)), "IntegerMask or IntegerBase is invalid!");
+	static_assert(TypeMask == (IntegerMask | (IntegerBase ^ ECbFieldType::IntegerNegative)), "IntegerMask or IntegerBase is invalid!");
+
+	static_assert(FloatBase == ECbFieldType::IntegerPositive, "FloatBase is invalid!");
+	static_assert((FloatMask & (AllFlags | ECbFieldType::IntegerPositive)) == ECbFieldType::IntegerPositive, "FloatMask is invalid!");
+	static_assert(!(FloatMask & (FloatBase ^ ECbFieldType::Float64)), "FloatMask or FloatBase is invalid!");
+	static_assert(TypeMask == (FloatMask | (FloatBase ^ ECbFieldType::Float64)), "FloatMask or FloatBase is invalid!");
+
+	static_assert(BoolBase == ECbFieldType::BoolFalse, "BoolBase is invalid!");
+	static_assert((BoolMask & (AllFlags | ECbFieldType::BoolTrue)) == ECbFieldType::BoolFalse, "BoolMask is invalid!");
+	static_assert(!(BoolMask & (BoolBase ^ ECbFieldType::BoolTrue)), "BoolMask or BoolBase is invalid!");
+	static_assert(TypeMask == (BoolMask | (BoolBase ^ ECbFieldType::BoolTrue)), "BoolMask or BoolBase is invalid!");
+
+	static_assert(ExternHashBase == ECbFieldType::BinaryHash, "ExternHashBase is invalid!");
+	static_assert((ExternHashMask & (AllFlags | ECbFieldType::FieldHash)) == ECbFieldType::BinaryHash, "ExternHashMask is invalid!");
+	static_assert(!(ExternHashMask & (ExternHashBase ^ ECbFieldType::FieldHash)), "ExternHashMask or ExternHashBase is invalid!");
+	static_assert(TypeMask == (ExternHashMask | (ExternHashBase ^ ECbFieldType::FieldHash)), "ExternHashMask or ExternHashBase is invalid!");
+}
+
+///////////////////////////////////////////////////////////////////////////////////////////////////////////////////////
+
+FCbField::FCbField(const void* const InData, const ECbFieldType InType)
+{
+	const uint8* Bytes = static_cast<const uint8*>(InData);
+	const ECbFieldType LocalType = FCbFieldType::HasFieldType(InType) ? (ECbFieldType(*Bytes++) | ECbFieldType::HasFieldType) : InType;
+	uint32 NameLenByteCount = 0;
+	const uint64 NameLen64 = FCbFieldType::HasFieldName(LocalType) ? ReadVarUInt(Bytes, NameLenByteCount) : 0;
+	Bytes += NameLen64 + NameLenByteCount;
+
+	Type = LocalType;
+	NameLen = uint32(FMath::Clamp<uint64>(NameLen64, 0, ~uint32(0)));
+	Payload = Bytes;
+}
+
+FConstMemoryView FCbField::AsMemoryView() const
+{
+	const uint32 TypeSize = FCbFieldType::HasFieldType(Type) ? sizeof(ECbFieldType) : 0;
+	const uint32 NameSize = FCbFieldType::HasFieldName(Type) ? NameLen + MeasureVarUInt(NameLen) : 0;
+	const uint64 LocalPayloadSize = PayloadSize();
+	return MakeMemoryView(static_cast<const uint8*>(Payload) - TypeSize - NameSize, TypeSize + NameSize + PayloadSize());
+}
+
+uint64 FCbField::Size() const
+{
+	return AsMemoryView().Size();
+}
+
+uint64 FCbField::PayloadSize() const
+{
+	switch (FCbFieldType::GetType(Type))
+	{
+	case ECbFieldType::None:
+	case ECbFieldType::Null:
+		return 0;
+	case ECbFieldType::Object:
+	case ECbFieldType::UniformObject:
+	case ECbFieldType::Array:
+	case ECbFieldType::UniformArray:
+	case ECbFieldType::Binary:
+	case ECbFieldType::String:
+	{
+		uint32 PayloadSizeByteCount;
+		const uint64 PayloadSize = ReadVarUInt(Payload, PayloadSizeByteCount);
+		return PayloadSize + PayloadSizeByteCount;
+	}
+	case ECbFieldType::IntegerPositive:
+	case ECbFieldType::IntegerNegative:
+		return MeasureVarUInt(Payload);
+	case ECbFieldType::Float32:
+		return 4;
+	case ECbFieldType::Float64:
+		return 8;
+	case ECbFieldType::BoolFalse:
+	case ECbFieldType::BoolTrue:
+		return 0;
+	case ECbFieldType::BinaryHash:
+	case ECbFieldType::FieldHash:
+		return 32;
+	case ECbFieldType::Uuid:
+		return 16;
+	case ECbFieldType::DateTime:
+	case ECbFieldType::TimeSpan:
+		return 8;
+	default:
+		return 0;
+	}
+}
+
+bool FCbField::Equals(const FCbField& Other) const
+{
+	return Type == Other.Type && AsMemoryView().EqualBytes(Other.AsMemoryView());
+}
+
+FCbObject FCbField::AsObject()
+{
+	const ECbFieldType LocalType = Type;
+	if (FCbFieldType::IsObject(LocalType))
+	{
+		Error = ECbFieldError::None;
+		return FCbObject(Payload, FCbFieldType::GetType(LocalType));
+	}
+	else
+	{
+		Error = ECbFieldError::TypeError;
+		return FCbObject();
+	}
+}
+
+FCbArray FCbField::AsArray()
+{
+	const ECbFieldType LocalType = Type;
+	if (FCbFieldType::IsArray(LocalType))
+	{
+		Error = ECbFieldError::None;
+		return FCbArray(Payload, FCbFieldType::GetType(LocalType));
+	}
+	else
+	{
+		Error = ECbFieldError::TypeError;
+		return FCbArray();
+	}
+}
+
+FConstMemoryView FCbField::AsBinary(const FConstMemoryView Default)
+{
+	if (FCbFieldType::IsBinary(Type))
+	{
+		const uint8* const PayloadBytes = static_cast<const uint8*>(Payload);
+		uint32 ValueSizeByteCount;
+		const uint64 ValueSize = ReadVarUInt(PayloadBytes, ValueSizeByteCount);
+
+		Error = ECbFieldError::None;
+		return MakeMemoryView(PayloadBytes + ValueSizeByteCount, ValueSize);
+	}
+	else
+	{
+		Error = ECbFieldError::TypeError;
+		return Default;
+	}
+}
+
+FAnsiStringView FCbField::AsString(const FAnsiStringView Default)
+{
+	if (FCbFieldType::IsString(Type))
+	{
+		const ANSICHAR* const PayloadChars = static_cast<const ANSICHAR*>(Payload);
+		uint32 ValueSizeByteCount;
+		const uint64 ValueSize = ReadVarUInt(PayloadChars, ValueSizeByteCount);
+
+		if (ValueSize >= (uint64(1) << 31))
+		{
+			Error = ECbFieldError::RangeError;
+			return Default;
+		}
+
+		Error = ECbFieldError::None;
+		return FAnsiStringView(PayloadChars + ValueSizeByteCount, int32(ValueSize));
+	}
+	else
+	{
+		Error = ECbFieldError::TypeError;
+		return Default;
+	}
+}
+
+uint64 FCbField::AsInteger(const uint64 Default, const FIntegerParams Params)
+{
+	const ECbFieldType LocalType = Type;
+	if (FCbFieldType::IsInteger(LocalType))
+	{
+		// A shift of a 64-bit value by 64 is undefined so shift by one less because magnitude is never zero.
+		const uint64 OutOfRangeMask = uint64(-2) << (Params.MagnitudeBits - 1);
+		const uint64 IsNegative = uint8(Type) & 1;
+
+		uint32 MagnitudeByteCount;
+		const uint64 Magnitude = ReadVarUInt(Payload, MagnitudeByteCount);
+		const uint64 Value = Magnitude ^ -int64(IsNegative);
+
+		const uint64 IsInRange = (!(Magnitude & OutOfRangeMask)) & ((!IsNegative) | Params.IsSigned);
+		Error = IsInRange ? ECbFieldError::None : ECbFieldError::RangeError;
+
+		const uint64 UseValueMask = -int64(IsInRange);
+		return (Value & UseValueMask) | (Default & ~UseValueMask);
+	}
+	else
+	{
+		Error = ECbFieldError::TypeError;
+		return Default;
+	}
+}
+
+float FCbField::AsFloat(const float Default)
+{
+	switch (FCbFieldType::GetType(Type))
+	{
+	case ECbFieldType::IntegerPositive:
+	case ECbFieldType::IntegerNegative:
+	{
+		const uint64 IsNegative = uint8(Type) & 1;
+		constexpr uint64 OutOfRangeMask = ~((uint64(1) << /*FLT_MANT_DIG*/ 24) - 1);
+
+		uint32 MagnitudeByteCount;
+		const int64 Magnitude = ReadVarUInt(Payload, MagnitudeByteCount) + IsNegative;
+		const uint64 IsInRange = !(Magnitude & OutOfRangeMask);
+		Error = IsInRange ? ECbFieldError::None : ECbFieldError::RangeError;
+		return IsInRange ? float(IsNegative ? -Magnitude : Magnitude) : Default;
+	}
+	case ECbFieldType::Float32:
+	{
+		Error = ECbFieldError::None;
+		const uint32 Value = NETWORK_ORDER32(CbObjectPrivate::ReadUnaligned<uint32>(Payload));
+		return reinterpret_cast<const float&>(Value);
+	}
+	case ECbFieldType::Float64:
+		Error = ECbFieldError::RangeError;
+		return Default;
+	default:
+		Error = ECbFieldError::TypeError;
+		return Default;
+	}
+}
+
+double FCbField::AsDouble(const double Default)
+{
+	switch (FCbFieldType::GetType(Type))
+	{
+	case ECbFieldType::IntegerPositive:
+	case ECbFieldType::IntegerNegative:
+	{
+		const uint64 IsNegative = uint8(Type) & 1;
+		constexpr uint64 OutOfRangeMask = ~((uint64(1) << /*DBL_MANT_DIG*/ 53) - 1);
+
+		uint32 MagnitudeByteCount;
+		const int64 Magnitude = ReadVarUInt(Payload, MagnitudeByteCount) + IsNegative;
+		const uint64 IsInRange = !(Magnitude & OutOfRangeMask);
+		Error = IsInRange ? ECbFieldError::None : ECbFieldError::RangeError;
+		return IsInRange ? double(IsNegative ? -Magnitude : Magnitude) : Default;
+	}
+	case ECbFieldType::Float32:
+	{
+		Error = ECbFieldError::None;
+		const uint32 Value = NETWORK_ORDER32(CbObjectPrivate::ReadUnaligned<uint32>(Payload));
+		return reinterpret_cast<const float&>(Value);
+	}
+	case ECbFieldType::Float64:
+	{
+		Error = ECbFieldError::None;
+		const uint64 Value = NETWORK_ORDER64(CbObjectPrivate::ReadUnaligned<uint64>(Payload));
+		return reinterpret_cast<const double&>(Value);
+	}
+	default:
+		Error = ECbFieldError::TypeError;
+		return Default;
+	}
+}
+
+bool FCbField::AsBool(const bool bDefault)
+{
+	const ECbFieldType LocalType = Type;
+	const bool bIsBool = FCbFieldType::IsBool(LocalType);
+	Error = bIsBool ? ECbFieldError::None : ECbFieldError::TypeError;
+	return (uint8(bIsBool) & uint8(LocalType) & 1) | ((!bIsBool) & bDefault);
+}
+
+FBlake3Hash FCbField::AsBinaryHash(const FBlake3Hash& Default)
+{
+	if (FCbFieldType::IsBinaryHash(Type))
+	{
+		Error = ECbFieldError::None;
+		return FBlake3Hash(*static_cast<const FBlake3Hash::ByteArray*>(Payload));
+	}
+	else
+	{
+		Error = ECbFieldError::TypeError;
+		return Default;
+	}
+}
+
+FBlake3Hash FCbField::AsFieldHash(const FBlake3Hash& Default)
+{
+	if (FCbFieldType::IsFieldHash(Type))
+	{
+		Error = ECbFieldError::None;
+		return FBlake3Hash(*static_cast<const FBlake3Hash::ByteArray*>(Payload));
+	}
+	else
+	{
+		Error = ECbFieldError::TypeError;
+		return Default;
+	}
+}
+
+FGuid FCbField::AsUuid()
+{
+	return AsUuid(FGuid());
+}
+
+FGuid FCbField::AsUuid(const FGuid& Default)
+{
+	if (FCbFieldType::IsUuid(Type))
+	{
+		Error = ECbFieldError::None;
+		FGuid Value;
+		FMemory::Memcpy(&Value, Payload, sizeof(FGuid));
+		Value.A = NETWORK_ORDER32(Value.A);
+		Value.B = NETWORK_ORDER32(Value.B);
+		Value.C = NETWORK_ORDER32(Value.C);
+		Value.D = NETWORK_ORDER32(Value.D);
+		return Value;
+	}
+	else
+	{
+		Error = ECbFieldError::TypeError;
+		return Default;
+	}
+}
+
+int64 FCbField::AsDateTimeTicks(const int64 Default)
+{
+	if (FCbFieldType::IsDateTime(Type))
+	{
+		Error = ECbFieldError::None;
+		return NETWORK_ORDER64(CbObjectPrivate::ReadUnaligned<int64>(Payload));
+	}
+	else
+	{
+		Error = ECbFieldError::TypeError;
+		return Default;
+	}
+}
+
+FDateTime FCbField::AsDateTime()
+{
+	return FDateTime(AsDateTimeTicks(0));
+}
+
+FDateTime FCbField::AsDateTime(FDateTime Default)
+{
+	return FDateTime(AsDateTimeTicks(Default.GetTicks()));
+}
+
+int64 FCbField::AsTimeSpanTicks(const int64 Default)
+{
+	if (FCbFieldType::IsTimeSpan(Type))
+	{
+		Error = ECbFieldError::None;
+		return NETWORK_ORDER64(CbObjectPrivate::ReadUnaligned<int64>(Payload));
+	}
+	else
+	{
+		Error = ECbFieldError::TypeError;
+		return Default;
+	}
+}
+
+FTimespan FCbField::AsTimeSpan()
+{
+	return FTimespan(AsTimeSpanTicks(0));
+}
+
+FTimespan FCbField::AsTimeSpan(FTimespan Default)
+{
+	return FTimespan(AsTimeSpanTicks(Default.GetTicks()));
+}
+
+void FCbField::CopyTo(const FMutableMemoryView Buffer) const
+{
+	const FConstMemoryView Source = AsMemoryView();
+	FMemory::Memcpy(Buffer.GetData(), Source.GetData(), FMath::Min(Buffer.Size(), Source.Size()));
+}
+
+///////////////////////////////////////////////////////////////////////////////////////////////////////////////////////
+
+FCbArray::FCbArray()
+	: Payload(CbObjectPrivate::EmptyArrayPayload)
+	, Type(ECbFieldType::Array)
+{
+}
+
+FCbArray::FCbArray(const void* const InData, const ECbFieldType InType)
+{
+	const uint8* Bytes = static_cast<const uint8*>(InData);
+	const ECbFieldType LocalType = FCbFieldType::HasFieldType(InType) ? (ECbFieldType(*Bytes++) | ECbFieldType::HasFieldType) : InType;
+	uint32 NameLenByteCount = 0;
+	const uint64 NameLen = FCbFieldType::HasFieldName(LocalType) ? ReadVarUInt(Bytes, NameLenByteCount) : 0;
+	Bytes += NameLen + NameLenByteCount;
+
+	Type = LocalType;
+	Payload = Bytes;
+}
+
+FConstMemoryView FCbArray::AsMemoryView() const
+{
+	uint32 PayloadSizeByteCount;
+	const uint64 PayloadSize = ReadVarUInt(Payload, PayloadSizeByteCount);
+	return MakeMemoryView(Payload, PayloadSizeByteCount + PayloadSize);
+}
+
+uint64 FCbArray::Size() const
+{
+	return sizeof(ECbFieldType) + AsMemoryView().Size();
+}
+
+uint64 FCbArray::Num() const
+{
+	const uint8* PayloadBytes = static_cast<const uint8*>(Payload);
+	PayloadBytes += MeasureVarUInt(PayloadBytes);
+	uint32 NumByteCount;
+	return ReadVarUInt(PayloadBytes, NumByteCount);
+}
+
+FCbFieldIterator FCbArray::CreateIterator() const
+{
+	const uint8* PayloadBytes = static_cast<const uint8*>(Payload);
+	uint32 PayloadSizeByteCount;
+	const uint64 PayloadSize = ReadVarUInt(PayloadBytes, PayloadSizeByteCount);
+	PayloadBytes += PayloadSizeByteCount;
+	if (PayloadSize > PayloadSizeByteCount)
+	{
+		const void* const PayloadEnd = PayloadBytes + PayloadSize;
+		PayloadBytes += MeasureVarUInt(PayloadBytes);
+		const ECbFieldType UniformType = FCbFieldType::GetType(Type) == ECbFieldType::UniformArray ?
+			ECbFieldType(*PayloadBytes++) : ECbFieldType::HasFieldType;
+		return FCbFieldIterator(FCbField(PayloadBytes, UniformType), PayloadEnd);
+	}
+	return FCbFieldIterator();
+}
+
+bool FCbArray::Equals(const FCbArray& Other) const
+{
+	return FCbFieldType::GetType(Type) == FCbFieldType::GetType(Other.Type) && AsMemoryView().EqualBytes(Other.AsMemoryView());
+}
+
+void FCbArray::CopyTo(const FMutableMemoryView Buffer) const
+{
+	uint8* BufferBytes = static_cast<uint8*>(Buffer.GetData());
+	*BufferBytes++ = uint8(FCbFieldType::GetType(Type));
+	const FConstMemoryView Source = AsMemoryView();
+	FMemory::Memcpy(BufferBytes, Source.GetData(), FMath::Min(Buffer.Size() - 1, Source.Size()));
+}
+
+///////////////////////////////////////////////////////////////////////////////////////////////////////////////////////
+
+FCbObject::FCbObject()
+	: Payload(CbObjectPrivate::EmptyObjectPayload)
+	, Type(ECbFieldType::Object)
+{
+}
+
+FCbObject::FCbObject(const void* const InData, const ECbFieldType InType)
+{
+	const uint8* Bytes = static_cast<const uint8*>(InData);
+	const ECbFieldType LocalType = FCbFieldType::HasFieldType(InType) ? (ECbFieldType(*Bytes++) | ECbFieldType::HasFieldType) : InType;
+	uint32 NameLenByteCount = 0;
+	const uint64 NameLen = FCbFieldType::HasFieldName(LocalType) ? ReadVarUInt(Bytes, NameLenByteCount) : 0;
+	Bytes += NameLen + NameLenByteCount;
+
+	Type = LocalType;
+	Payload = Bytes;
+}
+
+FConstMemoryView FCbObject::AsMemoryView() const
+{
+	uint32 PayloadSizeByteCount;
+	const uint64 PayloadSize = ReadVarUInt(Payload, PayloadSizeByteCount);
+	return MakeMemoryView(Payload, PayloadSizeByteCount + PayloadSize);
+}
+
+uint64 FCbObject::Size() const
+{
+	return sizeof(ECbFieldType) + AsMemoryView().Size();
+}
+
+FCbFieldIterator FCbObject::CreateIterator() const
+{
+	const uint8* PayloadBytes = static_cast<const uint8*>(Payload);
+	uint32 PayloadSizeByteCount;
+	const uint64 PayloadSize = ReadVarUInt(PayloadBytes, PayloadSizeByteCount);
+	PayloadBytes += PayloadSizeByteCount;
+	if (PayloadSize)
+	{
+		const void* const PayloadEnd = PayloadBytes + PayloadSize;
+		const ECbFieldType UniformType = FCbFieldType::GetType(Type) == ECbFieldType::UniformObject ?
+			ECbFieldType(*PayloadBytes++) : ECbFieldType::HasFieldType;
+		return FCbFieldIterator(FCbField(PayloadBytes, UniformType), PayloadEnd);
+	}
+	return FCbFieldIterator();
+}
+
+bool FCbObject::Equals(const FCbObject& Other) const
+{
+	return FCbFieldType::GetType(Type) == FCbFieldType::GetType(Other.Type) && AsMemoryView().EqualBytes(Other.AsMemoryView());
+}
+
+FCbField FCbObject::Find(const FAnsiStringView Name) const
+{
+	for (const FCbField Field : *this)
+	{
+		if (Name.Equals(Field.Name(), ESearchCase::CaseSensitive))
+		{
+			return Field;
+		}
+	}
+	return FCbField();
+}
+
+FCbField FCbObject::FindIgnoreCase(const FAnsiStringView Name) const
+{
+	for (const FCbField Field : *this)
+	{
+		if (Name.Equals(Field.Name(), ESearchCase::IgnoreCase))
+		{
+			return Field;
+		}
+	}
+	return FCbField();
+}
+
+void FCbObject::CopyTo(const FMutableMemoryView Buffer) const
+{
+	uint8* BufferBytes = static_cast<uint8*>(Buffer.GetData());
+	*BufferBytes++ = uint8(FCbFieldType::GetType(Type));
+	const FConstMemoryView Source = AsMemoryView();
+	FMemory::Memcpy(BufferBytes, Source.GetData(), FMath::Min(Buffer.Size() - 1, Source.Size()));
+}
+
+///////////////////////////////////////////////////////////////////////////////////////////////////////////////////////
+
+/**
+ * Adds the given error(s) to the error mask.
+ *
+ * This function exists to make validation errors easier to debug by providing one location to set a breakpoint.
+ */
+FORCENOINLINE static void AddError(ECbValidateError& OutError, const ECbValidateError InError)
+{
+	OutError |= InError;
+}
+
+/**
+ * Validate and read a field type from the view.
+ *
+ * A type argument with the HasFieldType flag indicates that the type will not be read from the view.
+ */
+static ECbFieldType ValidateCbFieldType(FConstMemoryView& View, ECbValidateMode Mode, ECbValidateError& Error, ECbFieldType Type = ECbFieldType::HasFieldType)
+{
+	if (FCbFieldType::HasFieldType(Type))
+	{
+		if (View.Size() >= 1)
+		{
+			Type = *static_cast<const ECbFieldType*>(View.GetData());
+			View += 1;
+			if (FCbFieldType::HasFieldType(Type))
+			{
+				AddError(Error, ECbValidateError::InvalidType);
+			}
+		}
+		else
+		{
+			AddError(Error, ECbValidateError::OutOfBounds);
+			View.Reset();
+			return ECbFieldType::None;
+		}
+	}
+
+	if (FCbFieldType::GetType(Type) > ECbFieldType::TimeSpan)
+	{
+		AddError(Error, ECbValidateError::InvalidType);
+		View.Reset();
+	}
+
+	return Type;
+}
+
+/**
+ * Validate and read an unsigned integer from the view.
+ *
+ * Modifies the view to start at the end of the value, and adds error flags if applicable.
+ */
+static uint64 ValidateCbUInt(FConstMemoryView& View, ECbValidateMode Mode, ECbValidateError& Error)
+{
+	if (View.Size() > 0 && View.Size() >= MeasureVarUInt(View.GetData()))
+	{
+		uint32 ValueByteCount;
+		const uint64 Value = ReadVarUInt(View.GetData(), ValueByteCount);
+		if (EnumHasAnyFlags(Mode, ECbValidateMode::Format) && ValueByteCount > MeasureVarUInt(Value))
+		{
+			AddError(Error, ECbValidateError::InvalidInteger);
+		}
+		View += ValueByteCount;
+		return Value;
+	}
+	else
+	{
+		AddError(Error, ECbValidateError::OutOfBounds);
+		View.Reset();
+		return 0;
+	}
+}
+
+/**
+ * Validate a 64-bit floating point value from the view.
+ *
+ * Modifies the view to start at the end of the value, and adds error flags if applicable.
+ */
+static void ValidateCbFloat64(FConstMemoryView& View, ECbValidateMode Mode, ECbValidateError& Error)
+{
+	if (View.Size() >= sizeof(double))
+	{
+		if (EnumHasAnyFlags(Mode, ECbValidateMode::Format))
+		{
+			const uint64 RawValue = NETWORK_ORDER64(CbObjectPrivate::ReadUnaligned<uint64>(View.GetData()));
+			const double Value = reinterpret_cast<const double&>(RawValue);
+			if (Value == double(float(Value)))
+			{
+				AddError(Error, ECbValidateError::InvalidFloat);
+			}
+		}
+		View += sizeof(double);
+	}
+	else
+	{
+		AddError(Error, ECbValidateError::OutOfBounds);
+		View.Reset();
+	}
+}
+
+/**
+ * Validate and read a string from the view.
+ *
+ * Modifies the view to start at the end of the string, and adds error flags if applicable.
+ */
+static FAnsiStringView ValidateCbString(FConstMemoryView& View, ECbValidateMode Mode, ECbValidateError& Error)
+{
+	const uint64 NameSize = ValidateCbUInt(View, Mode, Error);
+	if (View.Size() >= NameSize)
+	{
+		const FAnsiStringView Name(static_cast<const ANSICHAR*>(View.GetData()), static_cast<int32>(NameSize));
+		View += NameSize;
+		return Name;
+	}
+	else
+	{
+		AddError(Error, ECbValidateError::OutOfBounds);
+		View.Reset();
+		return FAnsiStringView();
+	}
+}
+
+static FCbField ValidateCbField(FConstMemoryView& View, ECbValidateMode Mode, ECbValidateError& Error, ECbFieldType ExternalType);
+
+/** A type that checks whether all validated fields are of the same type. */
+class FCbUniformFieldsValidator
+{
+public:
+	inline explicit FCbUniformFieldsValidator(ECbFieldType InExternalType)
+		: ExternalType(InExternalType)
+	{
+	}
+
+	inline FCbField ValidateField(FConstMemoryView& View, ECbValidateMode Mode, ECbValidateError& Error)
+	{
+		const void* const FieldData = View.GetData();
+		if (FCbField Field = ValidateCbField(View, Mode, Error, ExternalType))
+		{
+			++FieldCount;
+			if (FCbFieldType::HasFieldType(ExternalType))
+			{
+				const ECbFieldType FieldType = *static_cast<const ECbFieldType*>(FieldData);
+				if (FieldCount == 1)
+				{
+					FirstType = FieldType;
+				}
+				else if (FieldType != FirstType)
+				{
+					bUniform = false;
+				}
+			}
+			return Field;
+		}
+
+		// It may not safe to check for uniformity if the field was invalid.
+		bUniform = false;
+		return FCbField();
+	}
+
+	inline bool IsUniform() const { return FieldCount > 0 && bUniform; }
+
+private:
+	uint32 FieldCount = 0;
+	bool bUniform = true;
+	ECbFieldType FirstType = ECbFieldType::None;
+	ECbFieldType ExternalType;
+};
+
+static void ValidateCbObject(FConstMemoryView& View, ECbValidateMode Mode, ECbValidateError& Error, ECbFieldType ObjectType)
+{
+	const uint64 Size = ValidateCbUInt(View, Mode, Error);
+	FConstMemoryView ObjectView = View.Left(Size);
+	View += Size;
+
+	if (Size > 0)
+	{
+		TArray<FAnsiStringView, TInlineAllocator<16>> Names;
+
+		const bool bUniformObject = FCbFieldType::GetType(ObjectType) == ECbFieldType::UniformObject;
+		const ECbFieldType ExternalType = bUniformObject ? ValidateCbFieldType(ObjectView, Mode, Error) : ECbFieldType::HasFieldType;
+		FCbUniformFieldsValidator UniformValidator(ExternalType);
+		do
+		{
+			if (FCbField Field = UniformValidator.ValidateField(ObjectView, Mode, Error))
+			{
+				if (EnumHasAnyFlags(Mode, ECbValidateMode::Names))
+				{
+					if (Field.HasName())
+					{
+						Names.Add(Field.Name());
+					}
+					else
+					{
+						AddError(Error, ECbValidateError::MissingName);
+					}
+				}
+			}
+		}
+		while (!ObjectView.IsEmpty());
+
+		if (EnumHasAnyFlags(Mode, ECbValidateMode::Names) && Names.Num() > 1)
+		{
+			Algo::Sort(Names, [](FAnsiStringView L, FAnsiStringView R) { return L.Compare(R) < 0; });
+			for (const FAnsiStringView* NamesIt = Names.GetData(), *NamesEnd = NamesIt + Names.Num() - 1; NamesIt != NamesEnd; ++NamesIt)
+			{
+				if (NamesIt[0].Equals(NamesIt[1]))
+				{
+					AddError(Error, ECbValidateError::DuplicateName);
+					break;
+				}
+			}
+		}
+
+		if (!bUniformObject && EnumHasAnyFlags(Mode, ECbValidateMode::Format) && UniformValidator.IsUniform())
+		{
+			AddError(Error, ECbValidateError::NonUniformObject);
+		}
+	}
+}
+
+static void ValidateCbArray(FConstMemoryView& View, ECbValidateMode Mode, ECbValidateError& Error, ECbFieldType ArrayType)
+{
+	const uint64 Size = ValidateCbUInt(View, Mode, Error);
+	FConstMemoryView ArrayView = View.Left(Size);
+	View += Size;
+
+	const uint64 Count = ValidateCbUInt(ArrayView, Mode, Error);
+	const uint64 FieldsSize = ArrayView.Size();
+	const bool bUniformArray = FCbFieldType::GetType(ArrayType) == ECbFieldType::UniformArray;
+	const ECbFieldType ExternalType = bUniformArray ? ValidateCbFieldType(ArrayView, Mode, Error) : ECbFieldType::HasFieldType;
+	FCbUniformFieldsValidator UniformValidator(ExternalType);
+
+	for (uint64 Index = 0; Index < Count; ++Index)
+	{
+		if (FCbField Field = UniformValidator.ValidateField(ArrayView, Mode, Error))
+		{
+			if (Field.HasName() && EnumHasAnyFlags(Mode, ECbValidateMode::Names))
+			{
+				AddError(Error, ECbValidateError::ArrayName);
+			}
+		}
+	}
+
+	if (!bUniformArray && EnumHasAnyFlags(Mode, ECbValidateMode::Format) && UniformValidator.IsUniform() && FieldsSize > Count)
+	{
+		AddError(Error, ECbValidateError::NonUniformArray);
+	}
+}
+
+static FCbField ValidateCbField(FConstMemoryView& View, ECbValidateMode Mode, ECbValidateError& Error, const ECbFieldType ExternalType = ECbFieldType::HasFieldType)
+{
+	const FConstMemoryView FieldView = View;
+	const ECbFieldType Type = ValidateCbFieldType(View, Mode, Error, ExternalType);
+	const FAnsiStringView Name = FCbFieldType::HasFieldName(Type) ? ValidateCbString(View, Mode, Error) : FAnsiStringView();
+
+	auto ValidateFixedPayload = [&View, &Error](uint32 PayloadSize)
+	{
+		if (View.Size() >= PayloadSize)
+		{
+			View += PayloadSize;
+		}
+		else
+		{
+			AddError(Error, ECbValidateError::OutOfBounds);
+			View.Reset();
+		}
+	};
+
+	if (EnumHasAnyFlags(Error, ECbValidateError::OutOfBounds | ECbValidateError::InvalidType))
+	{
+		return FCbField();
+	}
+
+	switch (FCbFieldType::GetType(Type))
+	{
+	default:
+	case ECbFieldType::None:
+		AddError(Error, ECbValidateError::InvalidType);
+		View.Reset();
+		break;
+	case ECbFieldType::Null:
+	case ECbFieldType::BoolFalse:
+	case ECbFieldType::BoolTrue:
+		if (FieldView == View)
+		{
+			// Reset the view because a zero-sized field can cause infinite field iteration.
+			AddError(Error, ECbValidateError::InvalidType);
+			View.Reset();
+		}
+		break;
+	case ECbFieldType::Object:
+	case ECbFieldType::UniformObject:
+		ValidateCbObject(View, Mode, Error, FCbFieldType::GetType(Type));
+		break;
+	case ECbFieldType::Array:
+	case ECbFieldType::UniformArray:
+		ValidateCbArray(View, Mode, Error, FCbFieldType::GetType(Type));
+		break;
+	case ECbFieldType::Binary:
+	{
+		const uint64 ValueSize = ValidateCbUInt(View, Mode, Error);
+		if (View.Size() < ValueSize)
+		{
+			AddError(Error, ECbValidateError::OutOfBounds);
+			View.Reset();
+		}
+		else
+		{
+			View += ValueSize;
+		}
+		break;
+	}
+	case ECbFieldType::String:
+		ValidateCbString(View, Mode, Error);
+		break;
+	case ECbFieldType::IntegerPositive:
+		ValidateCbUInt(View, Mode, Error);
+		break;
+	case ECbFieldType::IntegerNegative:
+		ValidateCbUInt(View, Mode, Error);
+		break;
+	case ECbFieldType::Float32:
+		ValidateFixedPayload(4);
+		break;
+	case ECbFieldType::Float64:
+		ValidateCbFloat64(View, Mode, Error);
+		break;
+	case ECbFieldType::BinaryHash:
+	case ECbFieldType::FieldHash:
+		ValidateFixedPayload(32);
+		break;
+	case ECbFieldType::Uuid:
+		ValidateFixedPayload(16);
+		break;
+	case ECbFieldType::DateTime:
+	case ECbFieldType::TimeSpan:
+		ValidateFixedPayload(8);
+		break;
+	}
+
+	if (EnumHasAnyFlags(Error, ECbValidateError::OutOfBounds | ECbValidateError::InvalidType))
+	{
+		return FCbField();
+	}
+
+	return FCbField(FieldView.GetData(), ExternalType);
+}
+
+ECbValidateError ValidateCompactBinary(FConstMemoryView View, ECbValidateMode Mode, ECbFieldType Type)
+{
+	ECbValidateError Error = ECbValidateError::None;
+	ValidateCbField(View, Mode, Error, Type);
+	if (!View.IsEmpty() && EnumHasAnyFlags(Mode, ECbValidateMode::Padding))
+	{
+		AddError(Error, ECbValidateError::Padding);
+	}
+	return Error;
+}
+
+ECbValidateError ValidateCompactBinaryRange(FConstMemoryView View, ECbValidateMode Mode)
+{
+	ECbValidateError Error = ECbValidateError::None;
+	while (!View.IsEmpty())
+	{
+		ValidateCbField(View, Mode, Error);
+	}
+	return Error;
+}
+
+///////////////////////////////////////////////////////////////////////////////////////////////////////////////////////
+
+uint64 MeasureCompactBinary(FConstMemoryView View, ECbFieldType Type)
+{
+	uint64 Size = 0;
+
+	if (FCbFieldType::HasFieldType(Type))
+	{
+		if (View.Size() == 0)
+		{
+			return 0;
+		}
+
+		Type = *static_cast<const ECbFieldType*>(View.GetData());
+		View += 1;
+		Size += 1;
+	}
+
+	bool bVariableSize = false;
+	switch (FCbFieldType::GetType(Type))
+	{
+	case ECbFieldType::None:
+	case ECbFieldType::Null:
+		break;
+	case ECbFieldType::Object:
+	case ECbFieldType::UniformObject:
+	case ECbFieldType::Array:
+	case ECbFieldType::UniformArray:
+	case ECbFieldType::Binary:
+	case ECbFieldType::String:
+	case ECbFieldType::IntegerPositive:
+	case ECbFieldType::IntegerNegative:
+		bVariableSize = true;
+		break;
+	case ECbFieldType::Float32:
+		Size += 4;
+		break;
+	case ECbFieldType::Float64:
+		Size += 8;
+		break;
+	case ECbFieldType::BoolFalse:
+	case ECbFieldType::BoolTrue:
+		break;
+	case ECbFieldType::BinaryHash:
+	case ECbFieldType::FieldHash:
+		Size += 32;
+		break;
+	case ECbFieldType::Uuid:
+		Size += 16;
+		break;
+	case ECbFieldType::DateTime:
+	case ECbFieldType::TimeSpan:
+		Size += 8;
+		break;
+	default:
+		return 0;
+	}
+
+	if (FCbFieldType::HasFieldName(Type))
+	{
+		if (View.Size() == 0 || View.Size() < MeasureVarUInt(View.GetData()))
+		{
+			return 0;
+		}
+
+		uint32 NameLenByteCount;
+		const uint64 NameLen = ReadVarUInt(View.GetData(), NameLenByteCount);
+		const uint64 NameSize = NameLen + NameLenByteCount;
+
+		if (bVariableSize && View.Size() < NameSize)
+		{
+			return 0;
+		}
+
+		View += NameSize;
+		Size += NameSize;
+	}
+
+	switch (FCbFieldType::GetType(Type))
+	{
+	case ECbFieldType::Object:
+	case ECbFieldType::UniformObject:
+	case ECbFieldType::Array:
+	case ECbFieldType::UniformArray:
+	case ECbFieldType::Binary:
+	case ECbFieldType::String:
+		if (View.Size() == 0 || View.Size() < MeasureVarUInt(View.GetData()))
+		{
+			return 0;
+		}
+		else
+		{
+			uint32 PayloadSizeByteCount;
+			const uint64 PayloadSize = ReadVarUInt(View.GetData(), PayloadSizeByteCount);
+			return Size + PayloadSize + PayloadSizeByteCount;
+		}
+	case ECbFieldType::IntegerPositive:
+	case ECbFieldType::IntegerNegative:
+		if (View.Size() == 0)
+		{
+			return 0;
+		}
+		return Size + MeasureVarUInt(View.GetData());
+	default:
+		return Size;
+	}
+}
+
+///////////////////////////////////////////////////////////////////////////////////////////////////////////////////////
diff --git a/Engine/Source/Runtime/Core/Private/Serialization/CompactBinaryWriter.cpp b/Engine/Source/Runtime/Core/Private/Serialization/CompactBinaryWriter.cpp
new file mode 100644
--- /dev/null
+++ b/Engine/Source/Runtime/Core/Private/Serialization/CompactBinaryWriter.cpp
@@ -0,0 +1,533 @@
+// Copyright Epic Games, Inc. All Rights Reserved.
+
+#include "Serialization/CompactBinaryWriter.h"
+
+#include "Containers/StringConv.h"
+#include "Misc/Blake3.h"
+#include "Misc/ByteSwap.h"
+#include "Misc/DateTime.h"
+#include "Misc/Guid.h"
+#include "Misc/Timespan.h"
+#include "Serialization/VarInt.h"
+
+enum class FCbWriter::EStateFlags : uint8
+{
+	None = 0,
+	/** Whether a name has been written for the current field. */
+	Name = 1 << 0,
+	/** Whether this state is in the process of writing a field. */
+	Field = 1 << 1,
+	/** Whether this state is for array fields. */
+	Array = 1 << 2,
+	/** Whether this state is for object fields. */
+	Object = 1 << 3,
+};
+
+ENUM_CLASS_FLAGS(FCbWriter::EStateFlags);
+
+/** Whether the field type can be used in a uniform array or uniform object. */
+static constexpr bool IsUniformType(const ECbFieldType Type)
+{
+	if (FCbFieldType::HasFieldName(Type))
+	{
+		return true;
+	}
+
+	switch (Type)
+	{
+	case ECbFieldType::None:
+	case ECbFieldType::Null:
+	case ECbFieldType::BoolFalse:
+	case ECbFieldType::BoolTrue:
+		return false;
+	default:
+		return true;
+	}
+}
+
+FCbWriter::FCbWriter()
+{
+	States.Emplace();
+}
+
+FCbWriter::FCbWriter(const int64 InitialSize)
+	: FCbWriter()
+{
+	Data.Reserve(InitialSize);
+}
+
+FCbWriter::~FCbWriter()
+{
+}
+
+void FCbWriter::Reset()
+{
+	Data.Reset();
+	States.Reset();
+	States.Emplace();
+}
+
+FCbFieldRefIterator FCbWriter::Save() const
+{
+	const uint64 Size = SaveSize();
+	const FSharedBufferRef Buffer = MakeSharedBuffer(Size);
+	const FMutableMemoryView View = Buffer->GetView();
+	SaveToMemory(View);
+	return FCbFieldRefIterator(FCbFieldRef(Buffer), View.RightChop(Size).GetData());
+}
+
+FCbFieldIterator FCbWriter::SaveToMemory(const FMutableMemoryView Buffer) const
+{
+	checkf(States.Num() == 1 && States.Last().Flags == EStateFlags::None,
+		TEXT("It is invalid to save while there are incomplete write operations."));
+	checkf(Data.Num() > 0, TEXT("It is invalid to save when nothing has been written."));
+	checkf(Buffer.Size() == Data.Num(),
+		TEXT("Buffer is %" UINT64_FMT " bytes but %" INT64_FMT " is required."), Buffer.Size(), Data.Num());
+	FMemory::Memcpy(Buffer.GetData(), Data.GetData(), Data.Num());
+	return FCbFieldIterator(FCbField(Buffer.GetData()), Buffer.RightChop(Buffer.Size()).GetData());
+}
+
+uint64 FCbWriter::SaveSize() const
+{
+	return Data.Num();
+}
+
+void FCbWriter::BeginField()
+{
+	FState& State = States.Last();
+	if ((State.Flags & EStateFlags::Field) == EStateFlags::None)
+	{
+		State.Flags |= EStateFlags::Field;
+		State.Offset = Data.AddZeroed(sizeof(ECbFieldType));
+	}
+	else
+	{
+		checkf((State.Flags & EStateFlags::Name) == EStateFlags::Name,
+			TEXT("A new field cannot be written until the previous field '%.*hs' is finished."),
+			CurrentName().Len(), CurrentName().GetData());
+	}
+}
+
+void FCbWriter::EndField(ECbFieldType Type)
+{
+	FState& State = States.Last();
+
+	if ((State.Flags & EStateFlags::Name) == EStateFlags::Name)
+	{
+		Type |= ECbFieldType::HasFieldName;
+	}
+	else
+	{
+		checkf((State.Flags & EStateFlags::Object) == EStateFlags::None,
+			TEXT("It is invalid to write an object field without a unique non-empty name."));
+	}
+
+	if (State.Count == 0)
+	{
+		State.UniformType = Type;
+	}
+	else if (State.UniformType != Type)
+	{
+		State.UniformType = ECbFieldType::None;
+	}
+
+	State.Flags &= ~(EStateFlags::Name | EStateFlags::Field);
+	++State.Count;
+	Data[State.Offset] = uint8(Type);
+}
+
+FCbWriter& FCbWriter::Name(const FAnsiStringView Name)
+{
+	FState& State = States.Last();
+	checkf((State.Flags & EStateFlags::Array) != EStateFlags::Array,
+		TEXT("It is invalid to write a name for an array field. Name '%.*hs'"), Name.Len(), Name.GetData());
+	checkf(!Name.IsEmpty(), TEXT("%s"), (State.Flags & EStateFlags::Object) == EStateFlags::Object
+		? TEXT("It is invalid to write an empty name for an object field. Specify a unique non-empty name.")
+		: TEXT("It is invalid to write an empty name for a top-level field. Specify a name or avoid this call."));
+	checkf((State.Flags & (EStateFlags::Name | EStateFlags::Field)) == EStateFlags::None,
+		TEXT("A new field '%.*hs' cannot be written until the previous field '%.*hs' is finished."),
+		Name.Len(), Name.GetData(), CurrentName().Len(), CurrentName().GetData());
+
+	BeginField();
+	State.Flags |= EStateFlags::Name;
+	const uint32 NameLenByteCount = MeasureVarUInt(uint32(Name.Len()));
+	const int64 NameLenOffset = Data.AddUninitialized(NameLenByteCount);
+	WriteVarUInt(uint64(Name.Len()), Data.GetData() + NameLenOffset);
+	Data.Append(reinterpret_cast<const uint8*>(Name.GetData()), Name.Len());
+	return *this;
+}
+
+void FCbWriter::NameOrString(const FAnsiStringView NameOrValue)
+{
+	// A name is only written if it would begin a new field inside of an object.
+	if ((States.Last().Flags & (EStateFlags::Name | EStateFlags::Field | EStateFlags::Object)) == EStateFlags::Object)
+	{
+		Name(NameOrValue);
+	}
+	else
+	{
+		String(NameOrValue);
+	}
+}
+
+FAnsiStringView FCbWriter::CurrentName() const
+{
+	const FState& State = States.Last();
+	if ((State.Flags & EStateFlags::Name) == EStateFlags::Name)
+	{
+		const uint8* const EncodedName = Data.GetData() + State.Offset + sizeof(ECbFieldType);
+		uint32 NameLenByteCount;
+		const uint64 NameLen = ReadVarUInt(EncodedName, NameLenByteCount);
+		using NameLenSizeType = FAnsiStringView::SizeType;
+		const NameLenSizeType ClampedNameLen = NameLenSizeType(FMath::Clamp<uint64>(NameLen, 0, ~NameLenSizeType(0)));
+		return FAnsiStringView(reinterpret_cast<const ANSICHAR*>(EncodedName + NameLenByteCount), ClampedNameLen);
+	}
+	return FAnsiStringView();
+}
+
+void FCbWriter::MakeFieldsUniform(const int64 FieldBeginOffset, const int64 FieldEndOffset)
+{
+	FMutableMemoryView SourceView(Data.GetData() + FieldBeginOffset, uint64(FieldEndOffset - FieldBeginOffset));
+	FMutableMemoryView TargetView = SourceView + sizeof(ECbFieldType);
+	while (!SourceView.IsEmpty())
+	{
+		const uint64 FieldSize = MeasureCompactBinary(SourceView) - sizeof(ECbFieldType);
+		SourceView += sizeof(ECbFieldType);
+		if (TargetView.GetData() != SourceView.GetData())
+		{
+			FMemory::Memmove(TargetView.GetData(), SourceView.GetData(), FieldSize);
+		}
+		SourceView += FieldSize;
+		TargetView += FieldSize;
+	}
+	if (!TargetView.IsEmpty())
+	{
+		Data.RemoveAt(FieldEndOffset - TargetView.Size(), TargetView.Size());
+	}
+}
+
+void FCbWriter::Field(const FCbField& Value)
+{
+	checkf(Value.HasValue(), TEXT("It is invalid to write a field with no value."));
+	BeginField();
+	struct FCbFieldCopy : public FCbField
+	{
+		using FCbField::GetCopyType;
+		using FCbField::GetPayload;
+		using FCbField::GetPayloadEnd;
+	};
+	const FCbFieldCopy& ValueCopy = static_cast<const FCbFieldCopy&>(Value);
+	const void* const SourceData = ValueCopy.GetPayload();
+	const int64 SourceSize = PTRINT(ValueCopy.GetPayloadEnd()) - PTRINT(SourceData);
+	const int64 TargetOffset = Data.AddUninitialized(SourceSize);
+	FMemory::Memcpy(Data.GetData() + TargetOffset, SourceData, SourceSize);
+	EndField(FCbFieldType::GetType(ValueCopy.GetCopyType()));
+}
+
+void FCbWriter::BeginObject()
+{
+	BeginField();
+	States.Push(FState());
+	States.Last().Flags |= EStateFlags::Object;
+}
+
+void FCbWriter::EndObject()
+{
+	checkf(States.Num() > 1 && (States.Last().Flags & EStateFlags::Object) == EStateFlags::Object,
+		TEXT("It is invalid to end an object when an object is not at the top of the stack."));
+	checkf((States.Last().Flags & EStateFlags::Field) == EStateFlags::None,
+		TEXT("It is invalid to end an object until the previous field is finished."));
+	const bool bUniform = IsUniformType(States.Last().UniformType);
+	const uint64 Count = States.Last().Count;
+	States.Pop();
+
+	// Calculate the offset of the payload.
+	const FState& State = States.Last();
+	int64 PayloadOffset = State.Offset + 1;
+	if ((State.Flags & EStateFlags::Name) == EStateFlags::Name)
+	{
+		uint32 NameLenByteCount;
+		const uint64 NameLen = ReadVarUInt(Data.GetData() + PayloadOffset, NameLenByteCount);
+		PayloadOffset += NameLen + NameLenByteCount;
+	}
+
+	// Remove redundant field types for uniform objects.
+	if (bUniform && Count > 1)
+	{
+		MakeFieldsUniform(PayloadOffset, Data.Num());
+	}
+
+	// Insert the object size.
+	const uint64 Size = uint64(Data.Num() - PayloadOffset);
+	const uint32 SizeByteCount = MeasureVarUInt(Size);
+	Data.InsertUninitialized(PayloadOffset, SizeByteCount);
+	WriteVarUInt(Size, Data.GetData() + PayloadOffset);
+
+	EndField(bUniform ? ECbFieldType::UniformObject : ECbFieldType::Object);
+}
+
+void FCbWriter::Object(const FCbObject& Value)
+{
+	BeginField();
+	struct FCbObjectCopy : public FCbObject
+	{
+		using FCbObject::GetCopyType;
+		using FCbObject::AsMemoryView;
+	};
+	const FCbObjectCopy& ValueCopy = static_cast<const FCbObjectCopy&>(Value);
+	const FConstMemoryView SourceView = ValueCopy.AsMemoryView();
+	const int64 TargetOffset = Data.AddUninitialized(SourceView.Size());
+	FMemory::Memcpy(Data.GetData() + TargetOffset, SourceView.GetData(), SourceView.Size());
+	EndField(FCbFieldType::GetType(ValueCopy.GetCopyType()));
+}
+
+void FCbWriter::BeginArray()
+{
+	BeginField();
+	States.Push(FState());
+	States.Last().Flags |= EStateFlags::Array;
+}
+
+void FCbWriter::EndArray()
+{
+	checkf(States.Num() > 1 && (States.Last().Flags & EStateFlags::Array) == EStateFlags::Array,
+		TEXT("Invalid attempt to end an array when an array is not at the top of the stack."));
+	checkf((States.Last().Flags & EStateFlags::Field) == EStateFlags::None,
+		TEXT("It is invalid to end an array until the previous field is finished."));
+	const bool bUniform = IsUniformType(States.Last().UniformType);
+	const uint64 Count = States.Last().Count;
+	States.Pop();
+
+	// Calculate the offset of the payload.
+	const FState& State = States.Last();
+	int64 PayloadOffset = State.Offset + 1;
+	if ((State.Flags & EStateFlags::Name) == EStateFlags::Name)
+	{
+		uint32 NameLenByteCount;
+		const uint64 NameLen = ReadVarUInt(Data.GetData() + PayloadOffset, NameLenByteCount);
+		PayloadOffset += NameLen + NameLenByteCount;
+	}
+
+	// Remove redundant field types for uniform arrays.
+	if (bUniform && Count > 1)
+	{
+		MakeFieldsUniform(PayloadOffset, Data.Num());
+	}
+
+	// Insert the array size and field count.
+	const uint32 CountByteCount = MeasureVarUInt(Count);
+	const uint64 Size = uint64(Data.Num() - PayloadOffset) + CountByteCount;
+	const uint32 SizeByteCount = MeasureVarUInt(Size);
+	Data.InsertUninitialized(PayloadOffset, SizeByteCount + CountByteCount);
+	WriteVarUInt(Size, Data.GetData() + PayloadOffset);
+	WriteVarUInt(Count, Data.GetData() + PayloadOffset + SizeByteCount);
+
+	EndField(bUniform ? ECbFieldType::UniformArray : ECbFieldType::Array);
+}
+
+void FCbWriter::Array(const FCbArray& Value)
+{
+	BeginField();
+	struct FCbArrayCopy : public FCbArray
+	{
+		using FCbArray::GetCopyType;
+		using FCbArray::AsMemoryView;
+	};
+	const FCbArrayCopy& ValueCopy = static_cast<const FCbArrayCopy&>(Value);
+	const FConstMemoryView SourceView = ValueCopy.AsMemoryView();
+	const int64 TargetOffset = Data.AddUninitialized(SourceView.Size());
+	FMemory::Memcpy(Data.GetData() + TargetOffset, SourceView.GetData(), SourceView.Size());
+	EndField(FCbFieldType::GetType(ValueCopy.GetCopyType()));
+}
+
+void FCbWriter::Null()
+{
+	BeginField();
+	EndField(ECbFieldType::Null);
+}
+
+void FCbWriter::Binary(const void* const Value, const uint64 Size)
+{
+	BeginField();
+	const uint32 SizeByteCount = MeasureVarUInt(Size);
+	const int64 SizeOffset = Data.AddUninitialized(SizeByteCount);
+	WriteVarUInt(Size, Data.GetData() + SizeOffset);
+	Data.Append(static_cast<const uint8*>(Value), Size);
+	EndField(ECbFieldType::Binary);
+}
+
+void FCbWriter::String(const FAnsiStringView Value)
+{
+	BeginField();
+	const uint64 Size = uint64(Value.Len());
+	const uint32 SizeByteCount = MeasureVarUInt(Size);
+	const int64 Offset = Data.AddUninitialized(SizeByteCount + Size);
+	uint8* StringData = Data.GetData() + Offset;
+	WriteVarUInt(Size, StringData);
+	StringData += SizeByteCount;
+	if (Size > 0)
+	{
+		FMemory::Memcpy(StringData, Value.GetData(), Value.Len() * sizeof(ANSICHAR));
+	}
+	EndField(ECbFieldType::String);
+}
+
+void FCbWriter::String(const FWideStringView Value)
+{
+	BeginField();
+	const uint32 Size = uint32(FTCHARToUTF8_Convert::ConvertedLength(Value.GetData(), Value.Len()));
+	const uint32 SizeByteCount = MeasureVarUInt(Size);
+	const int64 Offset = Data.AddUninitialized(SizeByteCount + Size);
+	uint8* StringData = Data.GetData() + Offset;
+	WriteVarUInt(Size, StringData);
+	StringData += SizeByteCount;
+	if (Size > 0)
+	{
+		FTCHARToUTF8_Convert::Convert(reinterpret_cast<ANSICHAR*>(StringData), Size, Value.GetData(), Value.Len());
+	}
+	EndField(ECbFieldType::String);
+}
+
+void FCbWriter::Integer(const int32 Value)
+{
+	if (Value >= 0)
+	{
+		return Integer(uint32(Value));
+	}
+	BeginField();
+	const uint32 Magnitude = ~uint32(Value);
+	const uint32 MagnitudeByteCount = MeasureVarUInt(Magnitude);
+	const int64 Offset = Data.AddUninitialized(MagnitudeByteCount);
+	WriteVarUInt(Magnitude, Data.GetData() + Offset);
+	EndField(ECbFieldType::IntegerNegative);
+}
+
+void FCbWriter::Integer(const int64 Value)
+{
+	if (Value >= 0)
+	{
+		return Integer(uint64(Value));
+	}
+	BeginField();
+	const uint64 Magnitude = ~uint64(Value);
+	const uint32 MagnitudeByteCount = MeasureVarUInt(Magnitude);
+	const int64 Offset = Data.AddUninitialized(MagnitudeByteCount);
+	WriteVarUInt(Magnitude, Data.GetData() + Offset);
+	EndField(ECbFieldType::IntegerNegative);
+}
+
+void FCbWriter::Integer(const uint32 Value)
+{
+	BeginField();
+	const uint32 ValueByteCount = MeasureVarUInt(Value);
+	const int64 Offset = Data.AddUninitialized(ValueByteCount);
+	WriteVarUInt(Value, Data.GetData() + Offset);
+	EndField(ECbFieldType::IntegerPositive);
+}
+
+void FCbWriter::Integer(const uint64 Value)
+{
+	BeginField();
+	const uint32 ValueByteCount = MeasureVarUInt(Value);
+	const int64 Offset = Data.AddUninitialized(ValueByteCount);
+	WriteVarUInt(Value, Data.GetData() + Offset);
+	EndField(ECbFieldType::IntegerPositive);
+}
+
+void FCbWriter::Float(const float Value)
+{
+	BeginField();
+	const uint32 RawValue = NETWORK_ORDER32(reinterpret_cast<const uint32&>(Value));
+	Data.Append(reinterpret_cast<const uint8*>(&RawValue), sizeof(uint32));
+	EndField(ECbFieldType::Float32);
+}
+
+void FCbWriter::Float(const double Value)
+{
+	const float Value32 = float(Value);
+	if (Value == double(Value32))
+	{
+		return Float(Value32);
+	}
+	BeginField();
+	const uint64 RawValue = NETWORK_ORDER64(reinterpret_cast<const uint64&>(Value));
+	Data.Append(reinterpret_cast<const uint8*>(&RawValue), sizeof(uint64));
+	EndField(ECbFieldType::Float64);
+}
+
+void FCbWriter::Bool(const bool bValue)
+{
+	BeginField();
+	EndField(bValue ? ECbFieldType::BoolTrue : ECbFieldType::BoolFalse);
+}
+
+void FCbWriter::BinaryHash(const FBlake3Hash& Value)
+{
+	BeginField();
+	Data.Append(Value.Bytes(), sizeof(decltype(Value.Bytes())));
+	EndField(ECbFieldType::BinaryHash);
+}
+
+void FCbWriter::FieldHash(const FBlake3Hash& Value)
+{
+	BeginField();
+	Data.Append(Value.Bytes(), sizeof(decltype(Value.Bytes())));
+	EndField(ECbFieldType::FieldHash);
+}
+
+void FCbWriter::Uuid(const FGuid& Value)
+{
+	const auto AppendSwappedBytes = [this](uint32 In)
+	{
+		In = NETWORK_ORDER32(In);
+		Data.Append(reinterpret_cast<const uint8*>(&In), sizeof(In));
+	};
+	BeginField();
+	AppendSwappedBytes(Value.A);
+	AppendSwappedBytes(Value.B);
+	AppendSwappedBytes(Value.C);
+	AppendSwappedBytes(Value.D);
+	EndField(ECbFieldType::Uuid);
+}
+
+void FCbWriter::DateTimeTicks(const int64 Ticks)
+{
+	BeginField();
+	const uint64 RawValue = NETWORK_ORDER64(uint64(Ticks));
+	Data.Append(reinterpret_cast<const uint8*>(&RawValue), sizeof(uint64));
+	EndField(ECbFieldType::DateTime);
+}
+
+void FCbWriter::DateTime(const FDateTime Value)
+{
+	DateTimeTicks(Value.GetTicks());
+}
+
+void FCbWriter::TimeSpanTicks(const int64 Ticks)
+{
+	BeginField();
+	const uint64 RawValue = NETWORK_ORDER64(uint64(Ticks));
+	Data.Append(reinterpret_cast<const uint8*>(&RawValue), sizeof(uint64));
+	EndField(ECbFieldType::TimeSpan);
+}
+
+void FCbWriter::TimeSpan(const FTimespan Value)
+{
+	TimeSpanTicks(Value.GetTicks());
+}
+
+///////////////////////////////////////////////////////////////////////////////////////////////////////////////////////
+
+FCbWriter& operator<<(FCbWriter& Writer, const FDateTime Value)
+{
+	Writer.DateTime(Value);
+	return Writer;
+}
+
+FCbWriter& operator<<(FCbWriter& Writer, const FTimespan Value)
+{
+	Writer.TimeSpan(Value);
+	return Writer;
+}
+
+///////////////////////////////////////////////////////////////////////////////////////////////////////////////////////
diff --git a/Engine/Source/Runtime/Core/Private/Serialization/StructuredArchive.cpp b/Engine/Source/Runtime/Core/Private/Serialization/StructuredArchive.cpp
--- a/Engine/Source/Runtime/Core/Private/Serialization/StructuredArchive.cpp
+++ b/Engine/Source/Runtime/Core/Private/Serialization/StructuredArchive.cpp
@@ -529,6 +529,13 @@ void FStructuredArchiveSlot::operator<< (FLazyObjectPtr& Value)
 	Ar.LeaveSlot();
 }
 
+void FStructuredArchiveSlot::operator<< (FObjectPtr& Value)
+{
+	Ar.EnterSlot(*this);
+	Ar.Formatter.Serialize(Value);
+	Ar.LeaveSlot();
+}
+
 void FStructuredArchiveSlot::operator<< (FSoftObjectPtr& Value)
 {
 	Ar.EnterSlot(*this);
diff --git a/Engine/Source/Runtime/Core/Private/Tests/Memory/SharedBufferTest.cpp b/Engine/Source/Runtime/Core/Private/Tests/Memory/SharedBufferTest.cpp
new file mode 100644
--- /dev/null
+++ b/Engine/Source/Runtime/Core/Private/Tests/Memory/SharedBufferTest.cpp
@@ -0,0 +1,237 @@
+// Copyright Epic Games, Inc. All Rights Reserved.
+
+#include "Memory/SharedBuffer.h"
+
+#include "Misc/AutomationTest.h"
+
+#include <type_traits>
+
+///////////////////////////////////////////////////////////////////////////////////////////////////////////////////////
+
+static_assert(!std::is_constructible<FSharedBufferRef>::value, "Invalid constructor");
+static_assert(!std::is_constructible<FSharedBufferRef, FSharedBufferPtr>::value, "Invalid constructor");
+static_assert(!std::is_constructible<FSharedBufferRef, FSharedBufferWeakPtr>::value, "Invalid constructor");
+static_assert(!std::is_constructible<FSharedBufferRef, FSharedBufferConstRef>::value, "Invalid constructor");
+static_assert(!std::is_constructible<FSharedBufferRef, FSharedBufferConstPtr>::value, "Invalid constructor");
+static_assert(!std::is_constructible<FSharedBufferRef, FSharedBufferConstWeakPtr>::value, "Invalid constructor");
+
+static_assert(!std::is_constructible<FSharedBufferConstRef>::value, "Invalid constructor");
+static_assert(std::is_constructible<FSharedBufferConstRef, FSharedBufferRef>::value, "Missing constructor");
+static_assert(!std::is_constructible<FSharedBufferConstRef, FSharedBufferPtr>::value, "Invalid constructor");
+static_assert(!std::is_constructible<FSharedBufferConstRef, FSharedBufferWeakPtr>::value, "Invalid constructor");
+static_assert(!std::is_constructible<FSharedBufferConstRef, FSharedBufferConstPtr>::value, "Invalid constructor");
+static_assert(!std::is_constructible<FSharedBufferConstRef, FSharedBufferConstWeakPtr>::value, "Invalid constructor");
+
+static_assert(std::is_constructible<FSharedBufferPtr>::value, "Missing constructor");
+static_assert(std::is_constructible<FSharedBufferPtr, FSharedBufferRef>::value, "Missing constructor");
+static_assert(!std::is_constructible<FSharedBufferPtr, FSharedBufferWeakPtr>::value, "Invalid constructor");
+static_assert(!std::is_constructible<FSharedBufferPtr, FSharedBufferConstRef>::value, "Invalid constructor");
+static_assert(!std::is_constructible<FSharedBufferPtr, FSharedBufferConstPtr>::value, "Invalid constructor");
+static_assert(!std::is_constructible<FSharedBufferPtr, FSharedBufferConstWeakPtr>::value, "Invalid constructor");
+
+static_assert(std::is_constructible<FSharedBufferConstPtr>::value, "Missing constructor");
+static_assert(std::is_constructible<FSharedBufferConstPtr, FSharedBufferRef>::value, "Missing constructor");
+static_assert(std::is_constructible<FSharedBufferConstPtr, FSharedBufferPtr>::value, "Missing constructor");
+static_assert(!std::is_constructible<FSharedBufferConstPtr, FSharedBufferWeakPtr>::value, "Invalid constructor");
+static_assert(std::is_constructible<FSharedBufferConstPtr, FSharedBufferConstRef>::value, "Missing constructor");
+static_assert(!std::is_constructible<FSharedBufferConstPtr, FSharedBufferConstWeakPtr>::value, "Invalid constructor");
+
+static_assert(std::is_constructible<FSharedBufferWeakPtr>::value, "Missing constructor");
+static_assert(std::is_constructible<FSharedBufferWeakPtr, FSharedBufferRef>::value, "Missing constructor");
+static_assert(std::is_constructible<FSharedBufferWeakPtr, FSharedBufferPtr>::value, "Missing constructor");
+static_assert(!std::is_constructible<FSharedBufferWeakPtr, FSharedBufferConstRef>::value, "Invalid constructor");
+static_assert(!std::is_constructible<FSharedBufferWeakPtr, FSharedBufferConstPtr>::value, "Invalid constructor");
+static_assert(!std::is_constructible<FSharedBufferWeakPtr, FSharedBufferConstWeakPtr>::value, "Invalid constructor");
+
+static_assert(std::is_constructible<FSharedBufferConstWeakPtr>::value, "Missing constructor");
+static_assert(std::is_constructible<FSharedBufferConstWeakPtr, FSharedBufferRef>::value, "Missing constructor");
+static_assert(std::is_constructible<FSharedBufferConstWeakPtr, FSharedBufferPtr>::value, "Missing constructor");
+static_assert(std::is_constructible<FSharedBufferConstWeakPtr, FSharedBufferWeakPtr>::value, "Missing constructor");
+static_assert(std::is_constructible<FSharedBufferConstWeakPtr, FSharedBufferConstRef>::value, "Missing constructor");
+static_assert(std::is_constructible<FSharedBufferConstWeakPtr, FSharedBufferConstPtr>::value, "Missing constructor");
+
+///////////////////////////////////////////////////////////////////////////////////////////////////////////////////////
+
+static_assert(std::is_assignable<FSharedBufferRef, FSharedBufferRef>::value, "Missing assignment");
+static_assert(!std::is_assignable<FSharedBufferRef, FSharedBufferPtr>::value, "Invalid assignment");
+static_assert(!std::is_assignable<FSharedBufferRef, FSharedBufferWeakPtr>::value, "Invalid assignment");
+static_assert(!std::is_assignable<FSharedBufferRef, FSharedBufferConstRef>::value, "Invalid assignment");
+static_assert(!std::is_assignable<FSharedBufferRef, FSharedBufferConstPtr>::value, "Invalid assignment");
+static_assert(!std::is_assignable<FSharedBufferRef, FSharedBufferConstWeakPtr>::value, "Invalid assignment");
+
+static_assert(std::is_assignable<FSharedBufferConstRef, FSharedBufferRef>::value, "Missing assignment");
+static_assert(!std::is_assignable<FSharedBufferConstRef, FSharedBufferPtr>::value, "Invalid assignment");
+static_assert(!std::is_assignable<FSharedBufferConstRef, FSharedBufferWeakPtr>::value, "Invalid assignment");
+static_assert(std::is_assignable<FSharedBufferConstRef, FSharedBufferConstRef>::value, "Missing assignment");
+static_assert(!std::is_assignable<FSharedBufferConstRef, FSharedBufferConstPtr>::value, "Invalid assignment");
+static_assert(!std::is_assignable<FSharedBufferConstRef, FSharedBufferConstWeakPtr>::value, "Invalid assignment");
+
+static_assert(std::is_assignable<FSharedBufferPtr, FSharedBufferRef>::value, "Missing assignment");
+static_assert(std::is_assignable<FSharedBufferPtr, FSharedBufferPtr>::value, "Missing assignment");
+static_assert(!std::is_assignable<FSharedBufferPtr, FSharedBufferWeakPtr>::value, "Invalid assignment");
+static_assert(!std::is_assignable<FSharedBufferPtr, FSharedBufferConstRef>::value, "Invalid assignment");
+static_assert(!std::is_assignable<FSharedBufferPtr, FSharedBufferConstPtr>::value, "Invalid assignment");
+static_assert(!std::is_assignable<FSharedBufferPtr, FSharedBufferConstWeakPtr>::value, "Invalid assignment");
+
+static_assert(std::is_assignable<FSharedBufferConstPtr, FSharedBufferRef>::value, "Missing assignment");
+static_assert(std::is_assignable<FSharedBufferConstPtr, FSharedBufferPtr>::value, "Missing assignment");
+static_assert(!std::is_assignable<FSharedBufferConstPtr, FSharedBufferWeakPtr>::value, "Invalid assignment");
+static_assert(std::is_assignable<FSharedBufferConstPtr, FSharedBufferConstRef>::value, "Missing assignment");
+static_assert(std::is_assignable<FSharedBufferConstPtr, FSharedBufferConstPtr>::value, "Missing assignment");
+static_assert(!std::is_assignable<FSharedBufferConstPtr, FSharedBufferConstWeakPtr>::value, "Invalid assignment");
+
+static_assert(std::is_assignable<FSharedBufferWeakPtr, FSharedBufferRef>::value, "Missing assignment");
+static_assert(std::is_assignable<FSharedBufferWeakPtr, FSharedBufferPtr>::value, "Missing assignment");
+static_assert(std::is_assignable<FSharedBufferWeakPtr, FSharedBufferWeakPtr>::value, "Missing assignment");
+static_assert(!std::is_assignable<FSharedBufferWeakPtr, FSharedBufferConstRef>::value, "Invalid assignment");
+static_assert(!std::is_assignable<FSharedBufferWeakPtr, FSharedBufferConstPtr>::value, "Invalid assignment");
+static_assert(!std::is_assignable<FSharedBufferWeakPtr, FSharedBufferConstWeakPtr>::value, "Invalid assignment");
+
+static_assert(std::is_assignable<FSharedBufferConstWeakPtr, FSharedBufferRef>::value, "Missing assignment");
+static_assert(std::is_assignable<FSharedBufferConstWeakPtr, FSharedBufferPtr>::value, "Missing assignment");
+static_assert(std::is_assignable<FSharedBufferConstWeakPtr, FSharedBufferWeakPtr>::value, "Missing assignment");
+static_assert(std::is_assignable<FSharedBufferConstWeakPtr, FSharedBufferConstRef>::value, "Missing assignment");
+static_assert(std::is_assignable<FSharedBufferConstWeakPtr, FSharedBufferConstPtr>::value, "Missing assignment");
+static_assert(std::is_assignable<FSharedBufferConstWeakPtr, FSharedBufferConstWeakPtr>::value, "Missing assignment");
+
+///////////////////////////////////////////////////////////////////////////////////////////////////////////////////////
+
+static_assert(std::is_same<FSharedBufferRef, decltype(MakeSharedBuffer())>::value, "Invalid constructor");
+static_assert(std::is_same<FSharedBufferRef, decltype(MakeSharedBuffer(DeclVal<uint64>()))>::value, "Invalid constructor");
+
+static_assert(std::is_same<FSharedBufferRef, decltype(MakeSharedBuffer(FSharedBuffer::AssumeOwnership, DeclVal<void*>(), 0))>::value, "Invalid constructor");
+static_assert(std::is_same<FSharedBufferConstRef, decltype(MakeSharedBuffer(FSharedBuffer::AssumeOwnership, DeclVal<const void*>(), 0))>::value, "Invalid constructor");
+
+static_assert(std::is_same<FSharedBufferRef, decltype(MakeSharedBuffer(FSharedBuffer::Clone, DeclVal<void*>(), 0))>::value, "Invalid constructor");
+static_assert(std::is_same<FSharedBufferRef, decltype(MakeSharedBuffer(FSharedBuffer::Clone, DeclVal<const void*>(), 0))>::value, "Invalid constructor");
+
+static_assert(std::is_same<FSharedBufferRef, decltype(MakeSharedBuffer(FSharedBuffer::Wrap, DeclVal<void*>(), 0))>::value, "Invalid constructor");
+static_assert(std::is_same<FSharedBufferConstRef, decltype(MakeSharedBuffer(FSharedBuffer::Wrap, DeclVal<const void*>(), 0))>::value, "Invalid constructor");
+
+static_assert(std::is_same<FSharedBufferRef, decltype(MakeSharedBuffer(FSharedBuffer::AssumeOwnership, FMutableMemoryView()))>::value, "Invalid constructor");
+static_assert(std::is_same<FSharedBufferConstRef, decltype(MakeSharedBuffer(FSharedBuffer::AssumeOwnership, FConstMemoryView()))>::value, "Invalid constructor");
+
+static_assert(std::is_same<FSharedBufferRef, decltype(MakeSharedBuffer(FSharedBuffer::Clone, FMutableMemoryView()))>::value, "Invalid constructor");
+static_assert(std::is_same<FSharedBufferRef, decltype(MakeSharedBuffer(FSharedBuffer::Clone, FConstMemoryView()))>::value, "Invalid constructor");
+
+static_assert(std::is_same<FSharedBufferRef, decltype(MakeSharedBuffer(FSharedBuffer::Wrap, FMutableMemoryView()))>::value, "Invalid constructor");
+static_assert(std::is_same<FSharedBufferConstRef, decltype(MakeSharedBuffer(FSharedBuffer::Wrap, FConstMemoryView()))>::value, "Invalid constructor");
+
+static_assert(std::is_same<void*, decltype(DeclVal<FSharedBuffer>().GetData())>::value, "Invalid accessor");
+static_assert(std::is_same<const void*, decltype(DeclVal<const FSharedBuffer>().GetData())>::value, "Invalid accessor");
+
+static_assert(std::is_same<FMutableMemoryView, decltype(DeclVal<FSharedBuffer>().GetView())>::value, "Invalid accessor");
+static_assert(std::is_same<FConstMemoryView, decltype(DeclVal<const FSharedBuffer>().GetView())>::value, "Invalid accessor");
+
+static_assert(std::is_convertible<FSharedBuffer, FMutableMemoryView>::value, "Missing conversion");
+static_assert(std::is_convertible<FSharedBuffer, FConstMemoryView>::value, "Missing conversion");
+static_assert(!std::is_convertible<const FSharedBuffer, FMutableMemoryView>::value, "Invalid conversion");
+static_assert(std::is_convertible<const FSharedBuffer, FConstMemoryView>::value, "Missing conversion");
+
+///////////////////////////////////////////////////////////////////////////////////////////////////////////////////////
+
+static_assert(std::is_same<FSharedBufferRef, decltype(MakeSharedBufferOwned(DeclVal<const FSharedBufferRef&>()))>::value, "Invalid conversion");
+static_assert(std::is_same<FSharedBufferPtr, decltype(MakeSharedBufferOwned(DeclVal<const FSharedBufferPtr&>()))>::value, "Invalid conversion");
+static_assert(std::is_same<FSharedBufferPtr, decltype(MakeSharedBufferOwned(DeclVal<FSharedBufferPtr&&>()))>::value, "Invalid conversion");
+
+static_assert(std::is_same<FSharedBufferConstRef, decltype(MakeSharedBufferOwned(DeclVal<const FSharedBufferConstRef&>()))>::value, "Invalid conversion");
+static_assert(std::is_same<FSharedBufferConstPtr, decltype(MakeSharedBufferOwned(DeclVal<const FSharedBufferConstPtr&>()))>::value, "Invalid conversion");
+static_assert(std::is_same<FSharedBufferConstPtr, decltype(MakeSharedBufferOwned(DeclVal<FSharedBufferConstPtr&&>()))>::value, "Invalid conversion");
+
+///////////////////////////////////////////////////////////////////////////////////////////////////////////////////////
+
+static_assert(std::is_same<uint32, decltype(GetTypeHash(DeclVal<const FSharedBufferRef>()))>::value, "Missing or invalid hash function");
+static_assert(std::is_same<uint32, decltype(GetTypeHash(DeclVal<const FSharedBufferConstRef>()))>::value, "Missing or invalid hash function");
+static_assert(std::is_same<uint32, decltype(GetTypeHash(DeclVal<const FSharedBufferPtr>()))>::value, "Missing or invalid hash function");
+static_assert(std::is_same<uint32, decltype(GetTypeHash(DeclVal<const FSharedBufferConstPtr>()))>::value, "Missing or invalid hash function");
+static_assert(std::is_same<uint32, decltype(GetTypeHash(DeclVal<const FSharedBufferWeakPtr>()))>::value, "Missing or invalid hash function");
+static_assert(std::is_same<uint32, decltype(GetTypeHash(DeclVal<const FSharedBufferConstWeakPtr>()))>::value, "Missing or invalid hash function");
+
+///////////////////////////////////////////////////////////////////////////////////////////////////////////////////////
+
+#if WITH_DEV_AUTOMATION_TESTS
+
+IMPLEMENT_SIMPLE_AUTOMATION_TEST(FSharedBufferTest, "System.Core.Memory.SharedBuffer", EAutomationTestFlags::ApplicationContextMask | EAutomationTestFlags::SmokeFilter)
+bool FSharedBufferTest::RunTest(const FString& Parameters)
+{
+	// Test Default
+	{
+		FSharedBufferRef Ref = MakeSharedBuffer();
+		TestFalse(TEXT("FSharedBuffer().IsOwned()"), Ref->IsOwned());
+		TestEqual(TEXT("FSharedBuffer().Size()"), Ref->Size(), uint64(0));
+	}
+
+	// Test Size
+	{
+		constexpr uint64 Size = 64;
+		FSharedBufferRef Ref = MakeSharedBuffer(64);
+		TestTrue(TEXT("FSharedBuffer(Size).IsOwned()"), Ref->IsOwned());
+		TestEqual(TEXT("FSharedBuffer(Size).Size()"), Ref->Size(), Size);
+	}
+
+	// Test AssumeOwnership
+	{
+		constexpr uint64 Size = 64;
+		void* const Data = FMemory::Malloc(Size);
+		FSharedBufferRef Ref = MakeSharedBuffer(FSharedBuffer::AssumeOwnership, Data, Size);
+		TestTrue(TEXT("FSharedBuffer(AssumeOwnership).IsOwned()"), Ref->IsOwned());
+		TestEqual(TEXT("FSharedBuffer(AssumeOwnership).Size()"), Ref->Size(), Size);
+		TestEqual(TEXT("FSharedBuffer(AssumeOwnership).GetData()"), Ref->GetData(), Data);
+	}
+
+	// Test Clone
+	{
+		constexpr uint64 Size = 64;
+		const uint8 Data[Size]{};
+		FSharedBufferRef Ref = MakeSharedBuffer(FSharedBuffer::Clone, Data, Size);
+		TestTrue(TEXT("FSharedBuffer(Clone).IsOwned()"), Ref->IsOwned());
+		TestEqual(TEXT("FSharedBuffer(Clone).Size()"), Ref->Size(), Size);
+		TestNotEqual(TEXT("FSharedBuffer(Clone).GetData()"), static_cast<const void*>(Ref->GetData()), static_cast<const void*>(Data));
+	}
+
+	// Test Wrap
+	{
+		constexpr uint64 Size = 64;
+		uint8 Data[Size]{};
+		FSharedBufferRef Ref = MakeSharedBuffer(FSharedBuffer::Wrap, Data, Size);
+		TestFalse(TEXT("FSharedBuffer(Wrap).IsOwned()"), Ref->IsOwned());
+		TestEqual(TEXT("FSharedBuffer(Wrap).Size()"), Ref->Size(), Size);
+		TestEqual(TEXT("FSharedBuffer(Wrap).GetData()"), Ref->GetData(), static_cast<void*>(Data));
+	}
+
+	// Test MakeSharedBufferOwned
+	{
+		constexpr uint64 Size = 64;
+		uint8 Data[Size]{};
+		FSharedBufferRef Ref = MakeSharedBufferOwned(MakeSharedBuffer(FSharedBuffer::Wrap, Data, Size));
+		TestTrue(TEXT("MakeSharedBufferOwned(Wrap).IsOwned()"), Ref->IsOwned());
+		TestEqual(TEXT("MakeSharedBufferOwned(Wrap).Size()"), Ref->Size(), Size);
+		TestNotEqual(TEXT("MakeSharedBufferOwned(Wrap).GetData()"), Ref->GetData(), static_cast<void*>(Data));
+		FSharedBufferRef RefCopy = MakeSharedBufferOwned(Ref);
+		TestTrue(TEXT("MakeSharedBufferOwned(Owned).IsOwned()"), RefCopy->IsOwned());
+		TestEqual(TEXT("MakeSharedBufferOwned(Owned).Size()"), RefCopy->Size(), Size);
+		TestEqual(TEXT("MakeSharedBufferOwned(Owned).GetData()"), RefCopy->GetData(), Ref->GetData());
+	}
+	{
+		FSharedBufferPtr Ptr;
+		TestFalse(TEXT("MakeSharedBufferOwned(Null)"), MakeSharedBufferOwned(Ptr).IsValid());
+		TestFalse(TEXT("MakeSharedBufferOwned(Null)"), MakeSharedBufferOwned(AsConst(Ptr)).IsValid());
+		TestFalse(TEXT("MakeSharedBufferOwned(Null)"), MakeSharedBufferOwned(FSharedBufferPtr()).IsValid());
+		TestFalse(TEXT("MakeSharedBufferOwned(Null)"), MakeSharedBufferOwned(FSharedBufferConstPtr()).IsValid());
+	}
+
+	// Test WeakPtr
+	{
+		FSharedBufferWeakPtr WeakPtr;
+		{
+			FSharedBufferRef Ref = MakeSharedBuffer();
+			WeakPtr = Ref;
+			TestTrue(TEXT("FSharedBufferWeakPtr(Ref).Pin().IsValid()"), WeakPtr.Pin().IsValid());
+		}
+		TestFalse(TEXT("FSharedBufferWeakPtr(Ref).Pin().IsValid()"), WeakPtr.Pin().IsValid());
+	}
+
+	return true;
+}
+
+#endif // WITH_DEV_AUTOMATION_TESTS
