Copying //UE5/Dev-Cooker@14539516 to Main (//UE5/Main)

diff --git a/Engine/Source/Runtime/Core/Private/Tests/Serialization/CompactBinaryTest.cpp b/Engine/Source/Runtime/Core/Private/Tests/Serialization/CompactBinaryTest.cpp
new file mode 100644
--- /dev/null
+++ b/Engine/Source/Runtime/Core/Private/Tests/Serialization/CompactBinaryTest.cpp
@@ -0,0 +1,1502 @@
+// Copyright Epic Games, Inc. All Rights Reserved.
+
+#include "Serialization/CompactBinary.h"
+
+#include "Misc/AutomationTest.h"
+#include "Misc/Blake3.h"
+#include "Misc/DateTime.h"
+#include "Misc/Guid.h"
+#include "Misc/Timespan.h"
+#include "Serialization/VarInt.h"
+
+#if WITH_DEV_AUTOMATION_TESTS
+
+///////////////////////////////////////////////////////////////////////////////////////////////////////////////////////
+
+static constexpr EAutomationTestFlags::Type CompactBinaryTestFlags = EAutomationTestFlags::Type(EAutomationTestFlags::ApplicationContextMask | EAutomationTestFlags::SmokeFilter);
+
+///////////////////////////////////////////////////////////////////////////////////////////////////////////////////////
+
+template <ECbFieldType FieldType>
+struct TCbFieldTypeAccessors;
+
+template <ECbFieldType FieldType>
+using TCbFieldValueType = typename TCbFieldTypeAccessors<FCbFieldType::GetType(FieldType)>::ValueType;
+
+template <ECbFieldType FieldType>
+constexpr bool (FCbField::*TCbFieldIsTypeFn)() const = TCbFieldTypeAccessors<FCbFieldType::GetType(FieldType)>::IsTypeFn;
+
+template <ECbFieldType FieldType>
+constexpr auto TCbFieldAsTypeFn = TCbFieldTypeAccessors<FCbFieldType::GetType(FieldType)>::AsTypeFn;
+
+#define UE_CBFIELD_TYPE_ACCESSOR(FieldType, InIsTypeFn, InAsTypeFn, InValueType)                                      \
+	template <>                                                                                                       \
+	struct TCbFieldTypeAccessors<ECbFieldType::FieldType>                                                             \
+	{                                                                                                                 \
+		using ValueType = InValueType;                                                                                \
+		static constexpr bool (FCbField::*IsTypeFn)() const = &FCbField::InIsTypeFn;                                  \
+		static constexpr auto AsTypeFn = &FCbField::InAsTypeFn;                                                       \
+	};
+
+#define UE_CBFIELD_TYPE_ACCESSOR_TYPED(FieldType, InIsTypeFn, InAsTypeFn, InValueType, InDefaultType)                 \
+	template <>                                                                                                       \
+	struct TCbFieldTypeAccessors<ECbFieldType::FieldType>                                                             \
+	{                                                                                                                 \
+		using ValueType = InValueType;                                                                                \
+		static constexpr bool (FCbField::*IsTypeFn)() const = &FCbField::InIsTypeFn;                                  \
+		static constexpr InValueType (FCbField::*AsTypeFn)(InDefaultType) = &FCbField::InAsTypeFn;                    \
+	};
+
+UE_CBFIELD_TYPE_ACCESSOR(Object, IsObject, AsObject, FCbObject);
+UE_CBFIELD_TYPE_ACCESSOR(Array, IsArray, AsArray, FCbArray);
+UE_CBFIELD_TYPE_ACCESSOR(Binary, IsBinary, AsBinary, FConstMemoryView);
+UE_CBFIELD_TYPE_ACCESSOR(String, IsString, AsString, FAnsiStringView);
+UE_CBFIELD_TYPE_ACCESSOR(IntegerPositive, IsInteger, AsUInt64, uint64);
+UE_CBFIELD_TYPE_ACCESSOR(IntegerNegative, IsInteger, AsInt64, int64);
+UE_CBFIELD_TYPE_ACCESSOR(Float32, IsFloat, AsFloat, float);
+UE_CBFIELD_TYPE_ACCESSOR(Float64, IsFloat, AsDouble, double);
+UE_CBFIELD_TYPE_ACCESSOR(BoolFalse, IsBool, AsBool, bool);
+UE_CBFIELD_TYPE_ACCESSOR(BoolTrue, IsBool, AsBool, bool);
+UE_CBFIELD_TYPE_ACCESSOR(BinaryHash, IsBinaryHash, AsBinaryHash, FBlake3Hash);
+UE_CBFIELD_TYPE_ACCESSOR(FieldHash, IsFieldHash, AsFieldHash, FBlake3Hash);
+UE_CBFIELD_TYPE_ACCESSOR_TYPED(Uuid, IsUuid, AsUuid, FGuid, const FGuid&);
+UE_CBFIELD_TYPE_ACCESSOR(DateTime, IsDateTime, AsDateTimeTicks, int64);
+UE_CBFIELD_TYPE_ACCESSOR(TimeSpan, IsTimeSpan, AsTimeSpanTicks, int64);
+
+///////////////////////////////////////////////////////////////////////////////////////////////////////////////////////
+
+class FCbFieldTestBase : public FAutomationTestBase
+{
+protected:
+	using FAutomationTestBase::FAutomationTestBase;
+	using FAutomationTestBase::TestEqual;
+
+	void TestEqualBytes(const TCHAR* What, FConstMemoryView Actual, TArrayView<const uint8> Expected)
+	{
+		TestTrue(What, Actual.EqualBytes(MakeMemoryView(Expected)));
+	}
+
+	template <typename T, typename Default>
+	void TestFieldAsTypeNoClone(FCbField& Field, T (FCbField::*AsTypeFn)(Default), T ExpectedValue = T(), T DefaultValue = T(), ECbFieldError ExpectedError = ECbFieldError::None)
+	{
+		TestTypeError(Field, ExpectedError);
+		TestEqual(TEXT("FCbField::As[Type]()"), (Field.*AsTypeFn)(DefaultValue), ExpectedValue);
+		TestEqual(TEXT("FCbField::As[Type]() -> HasError()"), Field.HasError(), ExpectedError != ECbFieldError::None);
+		TestEqual(TEXT("FCbField::As[Type]() -> GetError()"), Field.GetError(), ExpectedError);
+	}
+
+	template <typename T, typename Default>
+	void TestFieldAsType(FCbField& Field, T (FCbField::*AsTypeFn)(Default), T ExpectedValue = T(), T DefaultValue = T(), ECbFieldError ExpectedError = ECbFieldError::None)
+	{
+		TestFieldAsTypeNoClone(Field, AsTypeFn, ExpectedValue, DefaultValue, ExpectedError);
+		FCbFieldRef FieldClone(FCbFieldRef::Clone, Field);
+		TestFieldAsTypeNoClone(FieldClone, AsTypeFn, ExpectedValue, DefaultValue, ExpectedError);
+		TestTrue(TEXT("FCbField::Equals()"), Field.Equals(FieldClone));
+	}
+
+	template <typename T>
+	void TestFieldAsTypeNoClone(FCbField& Field, T (FCbField::*AsTypeFn)(), T ExpectedValue = T(), T DefaultValue = T(), ECbFieldError ExpectedError = ECbFieldError::None)
+	{
+		TestTypeError(Field, ExpectedError);
+		(Field.*AsTypeFn)();
+		TestEqual(TEXT("FCbField::As[Type]() -> HasError()"), Field.HasError(), ExpectedError != ECbFieldError::None);
+		TestEqual(TEXT("FCbField::As[Type]() -> GetError()"), Field.GetError(), ExpectedError);
+	}
+
+	template <typename T>
+	void TestFieldAsType(FCbField& Field, T (FCbField::*AsTypeFn)(), T ExpectedValue = T(), T DefaultValue = T(), ECbFieldError ExpectedError = ECbFieldError::None)
+	{
+		TestFieldAsTypeNoClone(Field, AsTypeFn, ExpectedValue, DefaultValue, ExpectedError);
+		FCbFieldRef FieldClone(FCbFieldRef::Clone, Field);
+		TestFieldAsTypeNoClone(FieldClone, AsTypeFn, ExpectedValue, DefaultValue, ExpectedError);
+		TestTrue(TEXT("FCbField::Equals()"), Field.Equals(FieldClone));
+	}
+
+	template <ECbFieldType FieldType, typename T = TCbFieldValueType<FieldType>>
+	void TestField(FCbField& Field, T ExpectedValue = T(), T DefaultValue = T(), ECbFieldError ExpectedError = ECbFieldError::None)
+	{
+		TestTrue(TEXT("FCbField::Is[Type]()"), (Field.*TCbFieldIsTypeFn<FieldType>)());
+		TestFieldAsType(Field, TCbFieldAsTypeFn<FieldType>, ExpectedValue, DefaultValue, ExpectedError);
+	}
+
+	template <ECbFieldType FieldType, typename T = TCbFieldValueType<FieldType>>
+	void TestField(TArrayView<const uint8> Payload, T ExpectedValue = T(), T DefaultValue = T())
+	{
+		FCbField Field(Payload.GetData(), FieldType);
+		TestEqual(TEXT("FCbField::Size()"), Field.Size(), uint64(Payload.Num()));
+		TestTrue(TEXT("FCbField::HasValue()"), Field.HasValue());
+		TestFalse(TEXT("FCbField::HasError() == false"), Field.HasError());
+		TestEqual(TEXT("FCbField::GetError() == None"), Field.GetError(), ECbFieldError::None);
+		TestField<FieldType>(Field, ExpectedValue, DefaultValue);
+	}
+
+	template <typename T, typename Default>
+	void TestFieldAsTypeError(FCbField& Field, T (FCbField::*AsTypeFn)(Default), ECbFieldError ExpectedError, T ExpectedValue = T())
+	{
+		TestFieldAsTypeNoClone(Field, AsTypeFn, ExpectedValue, ExpectedValue, ExpectedError);
+	}
+
+	template <typename T>
+	void TestFieldAsTypeError(FCbField& Field, T (FCbField::*AsTypeFn)(), ECbFieldError ExpectedError, T ExpectedValue = T())
+	{
+		TestFieldAsTypeNoClone(Field, AsTypeFn, ExpectedValue, ExpectedValue, ExpectedError);
+	}
+
+	template <ECbFieldType FieldType, typename T = TCbFieldValueType<FieldType>>
+	void TestFieldError(FCbField& Field, ECbFieldError ExpectedError, T ExpectedValue = T())
+	{
+		TestEqual(TEXT("FCbField::Is[Type]()"), (Field.*TCbFieldIsTypeFn<FieldType>)(), ExpectedError != ECbFieldError::TypeError);
+		TestFieldAsTypeError(Field, TCbFieldAsTypeFn<FieldType>, ExpectedError, ExpectedValue);
+	}
+
+	template <ECbFieldType FieldType, typename T = TCbFieldValueType<FieldType>>
+	void TestFieldError(TArrayView<const uint8> Payload, ECbFieldError ExpectedError, T ExpectedValue = T())
+	{
+		FCbField Field(Payload.GetData(), FieldType);
+		TestFieldError<FieldType>(Field, ExpectedError, ExpectedValue);
+	}
+
+private:
+	void TestTypeError(FCbField& Field, ECbFieldError ExpectedError)
+	{
+		if (ExpectedError == ECbFieldError::None && !Field.IsBool())
+		{
+			TestFalse(TEXT("FCbField::IsBool() == false"), Field.IsBool());
+			TestFalse(TEXT("FCbField::AsBool() == false"), Field.AsBool());
+			TestTrue(TEXT("FCbField::AsBool() -> HasError()"), Field.HasError());
+			TestEqual(TEXT("FCbField::AsBool() -> GetError() == TypeError"), Field.GetError(), ECbFieldError::TypeError);
+		}
+	}
+};
+
+///////////////////////////////////////////////////////////////////////////////////////////////////////////////////////
+
+IMPLEMENT_CUSTOM_SIMPLE_AUTOMATION_TEST(FCbFieldNoneTest, FCbFieldTestBase, "System.Core.Serialization.CbField.None", CompactBinaryTestFlags)
+bool FCbFieldNoneTest::RunTest(const FString& Parameters)
+{
+	// Test FCbField()
+	{
+		constexpr FCbField DefaultField;
+		static_assert(!DefaultField.HasName(), "Error in HasName()");
+		static_assert(!DefaultField.HasValue(), "Error in HasValue()");
+		static_assert(!DefaultField.HasError(), "Error in HasError()");
+		static_assert(DefaultField.GetError() == ECbFieldError::None, "Error in GetError()");
+		TestEqual(TEXT("FCbField()::Size() == 0"), DefaultField.Size(), uint64(0));
+		TestEqual(TEXT("FCbField()::Name().Len() == 0"), DefaultField.Name().Len(), 0);
+		TestFalse(TEXT("!FCbField()::HasName()"), DefaultField.HasName());
+		TestFalse(TEXT("!FCbField()::HasValue()"), DefaultField.HasValue());
+		TestFalse(TEXT("!FCbField()::HasError()"), DefaultField.HasError());
+		TestEqual(TEXT("FCbField()::GetError() == None"), DefaultField.GetError(), ECbFieldError::None);
+	}
+
+	// Test FCbField(None)
+	{
+		FCbField NoneField(nullptr, ECbFieldType::None);
+		TestEqual(TEXT("FCbField(None)::Size() == 0"), NoneField.Size(), uint64(0));
+		TestEqual(TEXT("FCbField(None)::Name().Len() == 0"), NoneField.Name().Len(), 0);
+		TestFalse(TEXT("!FCbField(None)::HasName()"), NoneField.HasName());
+		TestFalse(TEXT("!FCbField(None)::HasValue()"), NoneField.HasValue());
+		TestFalse(TEXT("!FCbField(None)::HasError()"), NoneField.HasError());
+		TestEqual(TEXT("FCbField(None)::GetError() == None"), NoneField.GetError(), ECbFieldError::None);
+	}
+
+	// Test FCbField(None|Type|Name)
+	{
+		constexpr ECbFieldType FieldType = ECbFieldType::None | ECbFieldType::HasFieldName;
+		constexpr const ANSICHAR NoneBytes[] = { ANSICHAR(FieldType), 4, 'N', 'a', 'm', 'e' };
+		FCbField NoneField(NoneBytes);
+		TestEqual(TEXT("FCbField(None|Type|Name)::Size()"), NoneField.Size(), uint64(sizeof(NoneBytes)));
+		TestEqual(TEXT("FCbField(None|Type|Name)::Name()"), NoneField.Name(), "Name"_ASV);
+		TestTrue(TEXT("FCbField(None|Type|Name)::HasName()"), NoneField.HasName());
+		TestFalse(TEXT("!FCbField(None|Type|Name)::HasValue()"), NoneField.HasValue());
+	}
+
+	// Test FCbField(None|Type)
+	{
+		constexpr ECbFieldType FieldType = ECbFieldType::None;
+		constexpr const ANSICHAR NoneBytes[] = { ANSICHAR(FieldType) };
+		FCbField NoneField(NoneBytes);
+		TestEqual(TEXT("FCbField(None|Type)::Size()"), NoneField.Size(), uint64(sizeof(NoneBytes)));
+		TestEqual(TEXT("FCbField(None|Type)::Name()"), NoneField.Name().Len(), 0);
+		TestFalse(TEXT("FCbField(None|Type)::HasName()"), NoneField.HasName());
+		TestFalse(TEXT("!FCbField(None|Type)::HasValue()"), NoneField.HasValue());
+	}
+
+	// Test FCbField(None|Name)
+	{
+		constexpr ECbFieldType FieldType = ECbFieldType::None | ECbFieldType::HasFieldName;
+		constexpr const ANSICHAR NoneBytes[] = { 4, 'N', 'a', 'm', 'e' };
+		FCbField NoneField(NoneBytes, FieldType);
+		TestEqual(TEXT("FCbField(None|Name)::Size()"), NoneField.Size(), uint64(sizeof(NoneBytes)));
+		TestEqual(TEXT("FCbField(None|Name)::Name()"), NoneField.Name(), "Name"_ASV);
+		TestTrue(TEXT("FCbField(None|Name)::HasName()"), NoneField.HasName());
+		TestFalse(TEXT("!FCbField(None|Name)::HasValue()"), NoneField.HasValue());
+	}
+
+	return true;
+}
+
+IMPLEMENT_CUSTOM_SIMPLE_AUTOMATION_TEST(FCbFieldNullTest, FCbFieldTestBase, "System.Core.Serialization.CbField.Null", CompactBinaryTestFlags)
+bool FCbFieldNullTest::RunTest(const FString& Parameters)
+{
+	// Test FCbField(Null)
+	{
+		FCbField NullField(nullptr, ECbFieldType::Null);
+		TestEqual(TEXT("FCbField(Null)::Size() == 0"), NullField.Size(), uint64(0));
+		TestTrue(TEXT("FCbField(Null)::IsNull()"), NullField.IsNull());
+		TestTrue(TEXT("FCbField(Null)::HasValue()"), NullField.HasValue());
+		TestFalse(TEXT("!FCbField(Null)::HasError()"), NullField.HasError());
+		TestEqual(TEXT("FCbField(Null)::GetError() == None"), NullField.GetError(), ECbFieldError::None);
+	}
+
+	// Test FCbField(None) as Null
+	{
+		FCbField Field;
+		TestFalse(TEXT("FCbField(None)::IsNull()"), Field.IsNull());
+	}
+
+	return true;
+}
+
+IMPLEMENT_CUSTOM_SIMPLE_AUTOMATION_TEST(FCbFieldObjectTest, FCbFieldTestBase, "System.Core.Serialization.CbField.Object", CompactBinaryTestFlags)
+bool FCbFieldObjectTest::RunTest(const FString& Parameters)
+{
+	auto TestIntObject = [this](const FCbObject& Object, int32 ExpectedNum, uint64 ExpectedPayloadSize)
+	{
+		TestEqual(TEXT("FCbField(Object)::AsObject().Size()"), Object.Size(), ExpectedPayloadSize + sizeof(ECbFieldType));
+
+		int32 ActualNum = 0;
+		for (FCbFieldIterator It = Object.CreateIterator(); It; ++It)
+		{
+			++ActualNum;
+			TestNotEqual(TEXT("FCbField(Object) Iterator Name"), It->Name().Len(), 0);
+			TestEqual(TEXT("FCbField(Object) Iterator"), It->AsInt32(), ActualNum);
+		}
+		TestEqual(TEXT("FCbField(Object)::AsObject().CreateIterator() -> Count"), ActualNum, ExpectedNum);
+
+		ActualNum = 0;
+		for (FCbField Field : Object)
+		{
+			++ActualNum;
+			TestNotEqual(TEXT("FCbField(Object) Iterator Name"), Field.Name().Len(), 0);
+			TestEqual(TEXT("FCbField(Object) Range"), Field.AsInt32(), ActualNum);
+		}
+		TestEqual(TEXT("FCbField(Object)::AsObject() Range -> Count"), ActualNum, ExpectedNum);
+	};
+
+	// Test FCbField(Object, Empty)
+	TestField<ECbFieldType::Object>({0});
+
+	// Test FCbField(Object, Empty)
+	{
+		FCbObject Object;
+		TestIntObject(Object, 0, 1);
+
+		// Find fields that do not exist.
+		TestFalse(TEXT("FCbObject()::Find(Missing)"), Object.Find("Field"_ASV).HasValue());
+		TestFalse(TEXT("FCbObject()::FindIgnoreCase(Missing)"), Object.FindIgnoreCase("Field"_ASV).HasValue());
+		TestFalse(TEXT("FCbObject()::operator[](Missing)"), Object["Field"_ASV].HasValue());
+
+		// Advance an iterator past the last field.
+		FCbFieldIterator It = Object.CreateIterator();
+		TestFalse(TEXT("FCbObject()::CreateIterator() At End"), bool(It));
+		TestTrue(TEXT("FCbObject()::CreateIterator() At End"), !It);
+		for (int Count = 16; Count > 0; --Count)
+		{
+			++It;
+			It->AsInt32();
+		}
+		TestFalse(TEXT("FCbObject()::CreateIterator() At End"), bool(It));
+		TestTrue(TEXT("FCbObject()::CreateIterator() At End"), !It);
+	}
+
+	// Test FCbField(Object, NotEmpty)
+	{
+		constexpr uint8 IntType = uint8(ECbFieldType::HasFieldName | ECbFieldType::IntegerPositive);
+		const uint8 Payload[] = { 12, IntType, 1, 'A', 1, IntType, 1, 'B', 2, IntType, 1, 'C', 3 };
+		FCbField Field(Payload, ECbFieldType::Object);
+		TestFieldAsType(Field, &FCbField::AsObject);
+		FCbObjectRef Object(FCbObjectRef::Clone, Field.AsObject());
+		TestIntObject(Object, 3, sizeof(Payload));
+		TestIntObject(Field.AsObject(), 3, sizeof(Payload));
+		TestTrue(TEXT("FCbObject::Equals()"), Object.Equals(Field.AsObject()));
+		TestEqual(TEXT("FCbObject::Find()"), Object.Find("B"_ASV).AsInt32(), 2);
+		TestEqual(TEXT("FCbObject::Find()"), Object.Find("b"_ASV).AsInt32(4), 4);
+		TestEqual(TEXT("FCbObject::FindIgnoreCase()"), Object.FindIgnoreCase("B"_ASV).AsInt32(), 2);
+		TestEqual(TEXT("FCbObject::FindIgnoreCase()"), Object.FindIgnoreCase("b"_ASV).AsInt32(), 2);
+		TestEqual(TEXT("FCbObject::operator[]"), Object["B"_ASV].AsInt32(), 2);
+		TestEqual(TEXT("FCbObject::operator[]"), Object["b"_ASV].AsInt32(4), 4);
+	}
+
+	// Test FCbField(UniformObject, NotEmpty)
+	{
+		constexpr uint8 IntType = uint8(ECbFieldType::HasFieldName | ECbFieldType::IntegerPositive);
+		const uint8 Payload[] = { 10, IntType, 1, 'A', 1, 1, 'B', 2, 1, 'C', 3 };
+		FCbField Field(Payload, ECbFieldType::UniformObject);
+		TestFieldAsType(Field, &FCbField::AsObject);
+		FCbObjectRef Object(FCbObjectRef::Clone, Field.AsObject());
+		TestIntObject(Object, 3, sizeof(Payload));
+		TestIntObject(Field.AsObject(), 3, sizeof(Payload));
+		TestTrue(TEXT("FCbObject{Uniform}::Equals()"), Object.Equals(Field.AsObject()));
+		TestEqual(TEXT("FCbObject{Uniform}::Find()"), Object.Find("B"_ASV).AsInt32(), 2);
+		TestEqual(TEXT("FCbObject{Uniform}::Find()"), Object.FindRef("B"_ASV).AsInt32(), 2);
+		TestEqual(TEXT("FCbObject{Uniform}::Find()"), Object.Find("b"_ASV).AsInt32(4), 4);
+		TestEqual(TEXT("FCbObject{Uniform}::Find()"), Object.FindRef("b"_ASV).AsInt32(4), 4);
+		TestEqual(TEXT("FCbObject{Uniform}::FindIgnoreCase()"), Object.FindIgnoreCase("B"_ASV).AsInt32(), 2);
+		TestEqual(TEXT("FCbObject{Uniform}::FindIgnoreCase()"), Object.FindRefIgnoreCase("B"_ASV).AsInt32(), 2);
+		TestEqual(TEXT("FCbObject{Uniform}::FindIgnoreCase()"), Object.FindIgnoreCase("b"_ASV).AsInt32(), 2);
+		TestEqual(TEXT("FCbObject{Uniform}::FindIgnoreCase()"), Object.FindRefIgnoreCase("b"_ASV).AsInt32(), 2);
+		TestEqual(TEXT("FCbObject{Uniform}::operator[]"), Object["B"_ASV].AsInt32(), 2);
+		TestEqual(TEXT("FCbObject{Uniform}::operator[]"), Object["b"_ASV].AsInt32(4), 4);
+		Object = FCbFieldRef(FCbFieldRef::Wrap, Field).AsObjectRef();
+	}
+
+	// Test FCbField(None) as Object
+	{
+		FCbField Field;
+		TestFieldError<ECbFieldType::Object>(Field, ECbFieldError::TypeError);
+		FCbFieldRef(FCbFieldRef::Wrap, Field).AsObjectRef();
+	}
+
+	// Test FCbObject(ObjectWithName) and CreateRefIterator
+	{
+		const uint8 ObjectType = uint8(ECbFieldType::Object | ECbFieldType::HasFieldName);
+		const uint8 Buffer[] = { ObjectType, 3, 'K', 'e', 'y', 4, uint8(ECbFieldType::HasFieldName | ECbFieldType::IntegerPositive), 1, 'F', 8 };
+		const FCbObject Object(Buffer);
+		TestEqual(TEXT("FCbObject(ObjectWithName)::Size()"), Object.Size(), uint64(6));
+		const FCbObjectRef ObjectClone(FCbObjectRef::Clone, Object);
+		TestEqual(TEXT("FCbObjectRef(ObjectWithName)::Size()"), ObjectClone.Size(), uint64(6));
+		TestTrue(TEXT("FCbObject::Equals()"), Object.Equals(ObjectClone));
+		for (FCbFieldRefIterator It = ObjectClone.CreateRefIterator(); It; ++It)
+		{
+			FCbFieldRef Field = *It;
+			TestEqual(TEXT("FCbObjectRef::CreateRefIterator().Name()"), Field.Name(), "F"_ASV);
+			TestEqual(TEXT("FCbObjectRef::CreateRefIterator().AsInt32()"), Field.AsInt32(), 8);
+			TestTrue(TEXT("FCbObjectRef::CreateRefIterator().IsOwned()"), Field.IsOwned());
+		}
+		for (FCbFieldRefIterator It = ObjectClone.CreateRefIterator(), End; It != End; ++It)
+		{
+		}
+	}
+
+	return true;
+}
+
+IMPLEMENT_CUSTOM_SIMPLE_AUTOMATION_TEST(FCbFieldArrayTest, FCbFieldTestBase, "System.Core.Serialization.CbField.Array", CompactBinaryTestFlags)
+bool FCbFieldArrayTest::RunTest(const FString& Parameters)
+{
+	auto TestIntArray = [this](FCbArray Array, int32 ExpectedNum, uint64 ExpectedPayloadSize)
+	{
+		TestEqual(TEXT("FCbField(Array)::AsArray().Size()"), Array.Size(), ExpectedPayloadSize + sizeof(ECbFieldType));
+		TestEqual(TEXT("FCbField(Array)::AsArray().Num()"), Array.Num(), uint64(ExpectedNum));
+
+		int32 ActualNum = 0;
+		for (FCbFieldIterator It = Array.CreateIterator(); It; ++It)
+		{
+			++ActualNum;
+			TestEqual(TEXT("FCbField(Array) Iterator"), It->AsInt32(), ActualNum);
+		}
+		TestEqual(TEXT("FCbField(Array)::AsArray().CreateIterator() -> Count"), ActualNum, ExpectedNum);
+
+		ActualNum = 0;
+		for (FCbField Field : Array)
+		{
+			++ActualNum;
+			TestEqual(TEXT("FCbField(Array) Range"), Field.AsInt32(), ActualNum);
+		}
+		TestEqual(TEXT("FCbField(Array)::AsArray() Range -> Count"), ActualNum, ExpectedNum);
+	};
+
+	// Test FCbField(Array, Empty)
+	TestField<ECbFieldType::Array>({1, 0});
+
+	// Test FCbField(Array, Empty)
+	{
+		FCbArray Array;
+		TestIntArray(Array, 0, 2);
+
+		// Advance an iterator past the last field.
+		FCbFieldIterator It = Array.CreateIterator();
+		TestFalse(TEXT("FCbArray()::CreateIterator() At End"), bool(It));
+		TestTrue(TEXT("FCbArray()::CreateIterator() At End"), !It);
+		for (int Count = 16; Count > 0; --Count)
+		{
+			++It;
+			It->AsInt32();
+		}
+		TestFalse(TEXT("FCbArray()::CreateIterator() At End"), bool(It));
+		TestTrue(TEXT("FCbArray()::CreateIterator() At End"), !It);
+	}
+
+	// Test FCbField(Array, NotEmpty)
+	{
+		constexpr uint8 IntType = uint8(ECbFieldType::IntegerPositive);
+		const uint8 Payload[] = { 7, 3, IntType, 1, IntType, 2, IntType, 3 };
+		FCbField Field(Payload, ECbFieldType::Array);
+		TestFieldAsType(Field, &FCbField::AsArray);
+		FCbArrayRef Array(FCbArrayRef::Clone, Field.AsArray());
+		TestIntArray(Array, 3, sizeof(Payload));
+		TestIntArray(Field.AsArray(), 3, sizeof(Payload));
+		TestTrue(TEXT("FCbArray::Equals()"), Array.Equals(Field.AsArray()));
+	}
+
+	// Test FCbField(UniformArray)
+	{
+		constexpr uint8 IntType = uint8(ECbFieldType::IntegerPositive);
+		const uint8 Payload[] = { 5, 3, IntType, 1, 2, 3 };
+		FCbField Field(Payload, ECbFieldType::UniformArray);
+		TestFieldAsType(Field, &FCbField::AsArray);
+		FCbArrayRef Array(FCbArrayRef::Clone, Field.AsArray());
+		TestIntArray(Array, 3, sizeof(Payload));
+		TestIntArray(Field.AsArray(), 3, sizeof(Payload));
+		TestTrue(TEXT("FCbArray::Equals()"), Array.Equals(Field.AsArray()));
+		Array = FCbFieldRef(FCbFieldRef::Wrap, Field).AsArrayRef();
+	}
+
+	// Test FCbField(None) as Array
+	{
+		FCbField Field;
+		TestFieldError<ECbFieldType::Array>(Field, ECbFieldError::TypeError);
+		FCbFieldRef(FCbFieldRef::Wrap, Field).AsArrayRef();
+	}
+
+	// Test FCbArray(ArrayWithName) and CreateRefIterator
+	{
+		const uint8 ArrayType = uint8(ECbFieldType::Array | ECbFieldType::HasFieldName);
+		const uint8 Buffer[] = { ArrayType, 3, 'K', 'e', 'y', 3, 1, uint8(ECbFieldType::IntegerPositive), 8 };
+		const FCbArray Array(Buffer);
+		TestEqual(TEXT("Array(ArrayWithName)::Size()"), Array.Size(), uint64(5));
+		const FCbArrayRef ArrayClone(FCbArrayRef::Clone, Array);
+		TestEqual(TEXT("FCbArrayRef(ArrayWithName)::Size()"), ArrayClone.Size(), uint64(5));
+		TestTrue(TEXT("FCbArray::Equals()"), Array.Equals(ArrayClone));
+		for (FCbFieldRefIterator It = ArrayClone.CreateRefIterator(); It; ++It)
+		{
+			FCbFieldRef Field = *It;
+			TestEqual(TEXT("FCbArrayRef::CreateRefIterator().AsInt32()"), Field.AsInt32(), 8);
+			TestTrue(TEXT("FCbArrayRef::CreateRefIterator().IsOwned()"), Field.IsOwned());
+		}
+		for (FCbFieldRefIterator It = ArrayClone.CreateRefIterator(), End; It != End; ++It)
+		{
+		}
+	}
+
+	return true;
+}
+
+IMPLEMENT_CUSTOM_SIMPLE_AUTOMATION_TEST(FCbFieldBinaryTest, FCbFieldTestBase, "System.Core.Serialization.CbField.Binary", CompactBinaryTestFlags)
+bool FCbFieldBinaryTest::RunTest(const FString& Parameters)
+{
+	// Test FCbField(Binary, Empty)
+	TestField<ECbFieldType::Binary>({0});
+
+	// Test FCbField(Binary, Value)
+	{
+		const uint8 Payload[] = { 3, 4, 5, 6 }; // Size: 3, Data: 4/5/6
+		FCbField Field(Payload, ECbFieldType::Binary);
+		TestFieldAsTypeNoClone(Field, &FCbField::AsBinary, MakeMemoryView(Payload + 1, 3));
+	}
+
+	// Test FCbField(None) as Binary
+	{
+		FCbField Field;
+		const uint8 Default[] = { 1, 2, 3 };
+		TestFieldError<ECbFieldType::Binary>(Field, ECbFieldError::TypeError, MakeMemoryView(Default));
+	}
+
+	return true;
+}
+
+IMPLEMENT_CUSTOM_SIMPLE_AUTOMATION_TEST(FCbFieldStringTest, FCbFieldTestBase, "System.Core.Serialization.CbField.String", CompactBinaryTestFlags)
+bool FCbFieldStringTest::RunTest(const FString& Parameters)
+{
+	// Test FCbField(String, Empty)
+	TestField<ECbFieldType::String>({0});
+
+	// Test FCbField(String, Value)
+	{
+		const uint8 Payload[] = { 3, 'A', 'B', 'C' }; // Size: 3, Data: ABC
+		TestField<ECbFieldType::String>(Payload, FAnsiStringView(reinterpret_cast<const ANSICHAR*>(Payload) + 1, 3));
+	}
+
+	// Test FCbField(String, OutOfRangeSize)
+	{
+		uint8 Payload[9];
+		WriteVarUInt(uint64(1) << 31, Payload);
+		TestFieldError<ECbFieldType::String>(Payload, ECbFieldError::RangeError, "ABC"_ASV);
+	}
+
+	// Test FCbField(None) as String
+	{
+		FCbField Field;
+		TestFieldError<ECbFieldType::String>(Field, ECbFieldError::TypeError, "ABC"_ASV);
+	}
+
+	return true;
+}
+
+class FCbFieldIntegerTestBase : public FCbFieldTestBase
+{
+protected:
+	using FCbFieldTestBase::FCbFieldTestBase;
+
+	enum class EIntType : uint8
+	{
+		None   = 0x00,
+		Int8   = 0x01,
+		Int16  = 0x02,
+		Int32  = 0x04,
+		Int64  = 0x08,
+		UInt8  = 0x10,
+		UInt16 = 0x20,
+		UInt32 = 0x40,
+		UInt64 = 0x80,
+		// Masks for positive values requiring the specified number of bits.
+		Pos64 = UInt64,
+		Pos63 = Pos64 |  Int64,
+		Pos32 = Pos63 | UInt32,
+		Pos31 = Pos32 |  Int32,
+		Pos16 = Pos31 | UInt16,
+		Pos15 = Pos16 |  Int16,
+		Pos8  = Pos15 | UInt8,
+		Pos7  = Pos8  |  Int8,
+		// Masks for negative values requiring the specified number of bits.
+		Neg63 = Int64,
+		Neg31 = Neg63 | Int32,
+		Neg15 = Neg31 | Int16,
+		Neg7  = Neg15 | Int8,
+	};
+
+	void TestIntegerField(ECbFieldType FieldType, EIntType ExpectedMask, uint64 Magnitude)
+	{
+		uint8 Payload[9];
+		const bool Negative = bool(uint8(FieldType) & 1);
+		WriteVarUInt(Magnitude - Negative, Payload);
+		constexpr uint64 DefaultValue = 8;
+		const uint64 ExpectedValue = Negative ? uint64(-int64(Magnitude)) : Magnitude;
+		FCbField Field(Payload, FieldType);
+		TestFieldAsType(Field, &FCbField::AsInt8, int8(EnumHasAnyFlags(ExpectedMask, EIntType::Int8) ? ExpectedValue : DefaultValue),
+			int8(DefaultValue), EnumHasAnyFlags(ExpectedMask, EIntType::Int8) ? ECbFieldError::None : ECbFieldError::RangeError);
+		TestFieldAsType(Field, &FCbField::AsInt16, int16(EnumHasAnyFlags(ExpectedMask, EIntType::Int16) ? ExpectedValue : DefaultValue),
+			int16(DefaultValue), EnumHasAnyFlags(ExpectedMask, EIntType::Int16) ? ECbFieldError::None : ECbFieldError::RangeError);
+		TestFieldAsType(Field, &FCbField::AsInt32, int32(EnumHasAnyFlags(ExpectedMask, EIntType::Int32) ? ExpectedValue : DefaultValue),
+			int32(DefaultValue), EnumHasAnyFlags(ExpectedMask, EIntType::Int32) ? ECbFieldError::None : ECbFieldError::RangeError);
+		TestFieldAsType(Field, &FCbField::AsInt64, int64(EnumHasAnyFlags(ExpectedMask, EIntType::Int64) ? ExpectedValue : DefaultValue),
+			int64(DefaultValue), EnumHasAnyFlags(ExpectedMask, EIntType::Int64) ? ECbFieldError::None : ECbFieldError::RangeError);
+		TestFieldAsType(Field, &FCbField::AsUInt8, uint8(EnumHasAnyFlags(ExpectedMask, EIntType::UInt8) ? ExpectedValue : DefaultValue),
+			uint8(DefaultValue), EnumHasAnyFlags(ExpectedMask, EIntType::UInt8) ? ECbFieldError::None : ECbFieldError::RangeError);
+		TestFieldAsType(Field, &FCbField::AsUInt16, uint16(EnumHasAnyFlags(ExpectedMask, EIntType::UInt16) ? ExpectedValue : DefaultValue),
+			uint16(DefaultValue), EnumHasAnyFlags(ExpectedMask, EIntType::UInt16) ? ECbFieldError::None : ECbFieldError::RangeError);
+		TestFieldAsType(Field, &FCbField::AsUInt32, uint32(EnumHasAnyFlags(ExpectedMask, EIntType::UInt32) ? ExpectedValue : DefaultValue),
+			uint32(DefaultValue), EnumHasAnyFlags(ExpectedMask, EIntType::UInt32) ? ECbFieldError::None : ECbFieldError::RangeError);
+		TestFieldAsType(Field, &FCbField::AsUInt64, uint64(EnumHasAnyFlags(ExpectedMask, EIntType::UInt64) ? ExpectedValue : DefaultValue),
+			uint64(DefaultValue), EnumHasAnyFlags(ExpectedMask, EIntType::UInt64) ? ECbFieldError::None : ECbFieldError::RangeError);
+	}
+};
+
+IMPLEMENT_CUSTOM_SIMPLE_AUTOMATION_TEST(FCbFieldIntegerTest, FCbFieldIntegerTestBase, "System.Core.Serialization.CbField.Integer", CompactBinaryTestFlags)
+bool FCbFieldIntegerTest::RunTest(const FString& Parameters)
+{
+	// Test FCbField(IntegerPositive)
+	TestIntegerField(ECbFieldType::IntegerPositive, EIntType::Pos7,  0x00);
+	TestIntegerField(ECbFieldType::IntegerPositive, EIntType::Pos7,  0x7f);
+	TestIntegerField(ECbFieldType::IntegerPositive, EIntType::Pos8,  0x80);
+	TestIntegerField(ECbFieldType::IntegerPositive, EIntType::Pos8,  0xff);
+	TestIntegerField(ECbFieldType::IntegerPositive, EIntType::Pos15, 0x0100);
+	TestIntegerField(ECbFieldType::IntegerPositive, EIntType::Pos15, 0x7fff);
+	TestIntegerField(ECbFieldType::IntegerPositive, EIntType::Pos16, 0x8000);
+	TestIntegerField(ECbFieldType::IntegerPositive, EIntType::Pos16, 0xffff);
+	TestIntegerField(ECbFieldType::IntegerPositive, EIntType::Pos31, 0x0001'0000);
+	TestIntegerField(ECbFieldType::IntegerPositive, EIntType::Pos31, 0x7fff'ffff);
+	TestIntegerField(ECbFieldType::IntegerPositive, EIntType::Pos32, 0x8000'0000);
+	TestIntegerField(ECbFieldType::IntegerPositive, EIntType::Pos32, 0xffff'ffff);
+	TestIntegerField(ECbFieldType::IntegerPositive, EIntType::Pos63, 0x0000'0001'0000'0000);
+	TestIntegerField(ECbFieldType::IntegerPositive, EIntType::Pos63, 0x7fff'ffff'ffff'ffff);
+	TestIntegerField(ECbFieldType::IntegerPositive, EIntType::Pos64, 0x8000'0000'0000'0000);
+	TestIntegerField(ECbFieldType::IntegerPositive, EIntType::Pos64, 0xffff'ffff'ffff'ffff);
+
+	// Test FCbField(IntegerNegative)
+	TestIntegerField(ECbFieldType::IntegerNegative, EIntType::Neg7,  0x01);
+	TestIntegerField(ECbFieldType::IntegerNegative, EIntType::Neg7,  0x80);
+	TestIntegerField(ECbFieldType::IntegerNegative, EIntType::Neg15, 0x81);
+	TestIntegerField(ECbFieldType::IntegerNegative, EIntType::Neg15, 0x8000);
+	TestIntegerField(ECbFieldType::IntegerNegative, EIntType::Neg31, 0x8001);
+	TestIntegerField(ECbFieldType::IntegerNegative, EIntType::Neg31, 0x8000'0000);
+	TestIntegerField(ECbFieldType::IntegerNegative, EIntType::Neg63, 0x8000'0001);
+	TestIntegerField(ECbFieldType::IntegerNegative, EIntType::Neg63, 0x8000'0000'0000'0000);
+	TestIntegerField(ECbFieldType::IntegerNegative, EIntType::None,  0x8000'0000'0000'0001);
+	TestIntegerField(ECbFieldType::IntegerNegative, EIntType::None,  0xffff'ffff'ffff'ffff);
+
+	// Test FCbField(None) as Integer
+	{
+		FCbField Field;
+		TestFieldError<ECbFieldType::IntegerPositive>(Field, ECbFieldError::TypeError, uint64(8));
+		TestFieldError<ECbFieldType::IntegerNegative>(Field, ECbFieldError::TypeError, int64(8));
+	}
+
+	return true;
+}
+
+IMPLEMENT_CUSTOM_SIMPLE_AUTOMATION_TEST(FCbFieldFloatTest, FCbFieldTestBase, "System.Core.Serialization.CbField.Float", CompactBinaryTestFlags)
+bool FCbFieldFloatTest::RunTest(const FString& Parameters)
+{
+	// Test FCbField(Float, 32-bit)
+	{
+		const uint8 Payload[] = { 0xc0, 0x12, 0x34, 0x56 }; // -2.28444433f
+		TestField<ECbFieldType::Float32>(Payload, -2.28444433f);
+
+		FCbField Field(Payload, ECbFieldType::Float32);
+		TestFieldAsType(Field, &FCbField::AsDouble, -2.28444433);
+	}
+
+	// Test FCbField(Float, 64-bit)
+	{
+		const uint8 Payload[] = { 0xc1, 0x23, 0x45, 0x67, 0x89, 0xab, 0xcd, 0xef }; // -631475.76888888876
+		TestField<ECbFieldType::Float64>(Payload, -631475.76888888876);
+
+		FCbField Field(Payload, ECbFieldType::Float64);
+		TestFieldAsTypeError(Field, &FCbField::AsFloat, ECbFieldError::RangeError, 8.0f);
+	}
+
+	// Test FCbField(Integer+, MaxBinary32) as Float
+	{
+		uint8 Payload[9];
+		WriteVarUInt((uint64(1) << 24) - 1, Payload); // 16,777,215
+		FCbField Field(Payload, ECbFieldType::IntegerPositive);
+		TestField<ECbFieldType::Float32>(Field, 16'777'215.0f);
+		TestField<ECbFieldType::Float64>(Field, 16'777'215.0);
+	}
+
+	// Test FCbField(Integer+, MaxBinary32+1) as Float
+	{
+		uint8 Payload[9];
+		WriteVarUInt(uint64(1) << 24, Payload); // 16,777,216
+		FCbField Field(Payload, ECbFieldType::IntegerPositive);
+		TestFieldError<ECbFieldType::Float32>(Field, ECbFieldError::RangeError, 8.0f);
+		TestField<ECbFieldType::Float64>(Field, 16'777'216.0);
+	}
+
+	// Test FCbField(Integer+, MaxBinary64) as Float
+	{
+		uint8 Payload[9];
+		WriteVarUInt((uint64(1) << 53) - 1, Payload); // 9,007,199,254,740,991
+		FCbField Field(Payload, ECbFieldType::IntegerPositive);
+		TestFieldError<ECbFieldType::Float32>(Field, ECbFieldError::RangeError, 8.0f);
+		TestField<ECbFieldType::Float64>(Field, 9'007'199'254'740'991.0);
+	}
+
+	// Test FCbField(Integer+, MaxBinary64+1) as Float
+	{
+		uint8 Payload[9];
+		WriteVarUInt(uint64(1) << 53, Payload); // 9,007,199,254,740,992
+		FCbField Field(Payload, ECbFieldType::IntegerPositive);
+		TestFieldError<ECbFieldType::Float32>(Field, ECbFieldError::RangeError, 8.0f);
+		TestFieldError<ECbFieldType::Float64>(Field, ECbFieldError::RangeError, 8.0);
+	}
+
+	// Test FCbField(Integer+, MaxUInt64) as Float
+	{
+		uint8 Payload[9];
+		WriteVarUInt(uint64(-1), Payload); // Max uint64
+		FCbField Field(Payload, ECbFieldType::IntegerPositive);
+		TestFieldError<ECbFieldType::Float32>(Field, ECbFieldError::RangeError, 8.0f);
+		TestFieldError<ECbFieldType::Float64>(Field, ECbFieldError::RangeError, 8.0);
+	}
+
+	// Test FCbField(Integer-, MaxBinary32) as Float
+	{
+		uint8 Payload[9];
+		WriteVarUInt((uint64(1) << 24) - 2, Payload); // -16,777,215
+		FCbField Field(Payload, ECbFieldType::IntegerNegative);
+		TestField<ECbFieldType::Float32>(Field, -16'777'215.0f);
+		TestField<ECbFieldType::Float64>(Field, -16'777'215.0);
+	}
+
+	// Test FCbField(Integer-, MaxBinary32+1) as Float
+	{
+		uint8 Payload[9];
+		WriteVarUInt((uint64(1) << 24) - 1, Payload); // -16,777,216
+		FCbField Field(Payload, ECbFieldType::IntegerNegative);
+		TestFieldError<ECbFieldType::Float32>(Field, ECbFieldError::RangeError, 8.0f);
+		TestField<ECbFieldType::Float64>(Field, -16'777'216.0);
+	}
+
+	// Test FCbField(Integer-, MaxBinary64) as Float
+	{
+		uint8 Payload[9];
+		WriteVarUInt((uint64(1) << 53) - 2, Payload); // -9,007,199,254,740,991
+		FCbField Field(Payload, ECbFieldType::IntegerNegative);
+		TestFieldError<ECbFieldType::Float32>(Field, ECbFieldError::RangeError, 8.0f);
+		TestField<ECbFieldType::Float64>(Field, -9'007'199'254'740'991.0);
+	}
+
+	// Test FCbField(Integer-, MaxBinary64+1) as Float
+	{
+		uint8 Payload[9];
+		WriteVarUInt((uint64(1) << 53) - 1, Payload); // -9,007,199,254,740,992
+		FCbField Field(Payload, ECbFieldType::IntegerNegative);
+		TestFieldError<ECbFieldType::Float32>(Field, ECbFieldError::RangeError, 8.0f);
+		TestFieldError<ECbFieldType::Float64>(Field, ECbFieldError::RangeError, 8.0);
+	}
+
+	// Test FCbField(None) as Float
+	{
+		FCbField Field;
+		TestFieldError<ECbFieldType::Float32>(Field, ECbFieldError::TypeError, 8.0f);
+		TestFieldError<ECbFieldType::Float64>(Field, ECbFieldError::TypeError, 8.0);
+	}
+
+	return true;
+}
+
+IMPLEMENT_CUSTOM_SIMPLE_AUTOMATION_TEST(FCbFieldBoolTest, FCbFieldTestBase, "System.Core.Serialization.CbField.Bool", CompactBinaryTestFlags)
+bool FCbFieldBoolTest::RunTest(const FString& Parameters)
+{
+	// Test FCbField(Bool, False)
+	TestField<ECbFieldType::BoolFalse>({}, false, true);
+
+	// Test FCbField(Bool, True)
+	TestField<ECbFieldType::BoolTrue>({}, true, false);
+
+	// Test FCbField(None) as Bool
+	{
+		FCbField DefaultField;
+		TestFieldError<ECbFieldType::BoolFalse>(DefaultField, ECbFieldError::TypeError, false);
+		TestFieldError<ECbFieldType::BoolTrue>(DefaultField, ECbFieldError::TypeError, true);
+	}
+
+	return true;
+}
+
+IMPLEMENT_CUSTOM_SIMPLE_AUTOMATION_TEST(FCbFieldBinaryHashTest, FCbFieldTestBase, "System.Core.Serialization.CbField.BinaryHash", CompactBinaryTestFlags)
+bool FCbFieldBinaryHashTest::RunTest(const FString& Parameters)
+{
+	const FBlake3Hash::ByteArray ZeroBytes{};
+	const FBlake3Hash::ByteArray SequentialBytes{1, 2, 3, 4, 5, 6, 7, 8, 9, 10, 11, 12, 13, 14, 15, 16, 17, 18, 19, 20, 21, 22, 23, 24, 25, 26, 27, 28, 29, 30, 31, 32};
+
+	// Test FCbField(BinaryHash, Zero)
+	TestField<ECbFieldType::BinaryHash>(ZeroBytes);
+
+	// Test FCbField(BinaryHash, NonZero)
+	TestField<ECbFieldType::BinaryHash>(SequentialBytes, FBlake3Hash(SequentialBytes));
+
+	// Test FCbField(None) as BinaryHash
+	{
+		FCbField DefaultField;
+		TestFieldError<ECbFieldType::BinaryHash>(DefaultField, ECbFieldError::TypeError, FBlake3Hash(SequentialBytes));
+	}
+
+	return true;
+}
+
+IMPLEMENT_CUSTOM_SIMPLE_AUTOMATION_TEST(FCbFieldFieldHashTest, FCbFieldTestBase, "System.Core.Serialization.CbField.FieldHash", CompactBinaryTestFlags)
+bool FCbFieldFieldHashTest::RunTest(const FString& Parameters)
+{
+	const FBlake3Hash::ByteArray ZeroBytes{};
+	const FBlake3Hash::ByteArray SequentialBytes{1, 2, 3, 4, 5, 6, 7, 8, 9, 10, 11, 12, 13, 14, 15, 16, 17, 18, 19, 20, 21, 22, 23, 24, 25, 26, 27, 28, 29, 30, 31, 32};
+
+	// Test FCbField(FieldHash, Zero)
+	TestField<ECbFieldType::FieldHash>(ZeroBytes);
+
+	// Test FCbField(FieldHash, NonZero)
+	TestField<ECbFieldType::FieldHash>(SequentialBytes, FBlake3Hash(SequentialBytes));
+
+	// Test FCbField(None) as FieldHash
+	{
+		FCbField DefaultField;
+		TestFieldError<ECbFieldType::FieldHash>(DefaultField, ECbFieldError::TypeError, FBlake3Hash(SequentialBytes));
+	}
+
+	return true;
+}
+
+IMPLEMENT_CUSTOM_SIMPLE_AUTOMATION_TEST(FCbFieldUuidTest, FCbFieldTestBase, "System.Core.Serialization.CbField.Uuid", CompactBinaryTestFlags)
+bool FCbFieldUuidTest::RunTest(const FString& Parameters)
+{
+	const uint8 ZeroBytes[]{0, 0, 0, 0, 0, 0, 0, 0, 0, 0, 0, 0, 0, 0, 0, 0};
+	const uint8 SequentialBytes[]{0, 1, 2, 3, 4, 5, 6, 7, 8, 9, 10, 11, 12, 13, 14, 15};
+	const FGuid SequentialGuid(TEXT("00010203-0405-0607-0809-0a0b0c0d0e0f"));
+
+	// Test FCbField(Uuid, Zero)
+	TestField<ECbFieldType::Uuid>(ZeroBytes, FGuid(), SequentialGuid);
+
+	// Test FCbField(Uuid, NonZero)
+	TestField<ECbFieldType::Uuid>(SequentialBytes, SequentialGuid, FGuid());
+
+	// Test FCbField(None) as Uuid
+	{
+		FCbField DefaultField;
+		TestFieldError<ECbFieldType::Uuid>(DefaultField, ECbFieldError::TypeError, FGuid::NewGuid());
+	}
+
+	return true;
+}
+
+IMPLEMENT_CUSTOM_SIMPLE_AUTOMATION_TEST(FCbFieldDateTimeTest, FCbFieldTestBase, "System.Core.Serialization.CbField.DateTime", CompactBinaryTestFlags)
+bool FCbFieldDateTimeTest::RunTest(const FString& Parameters)
+{
+	// Test FCbField(DateTime, Zero)
+	TestField<ECbFieldType::DateTime>({0, 0, 0, 0, 0, 0, 0, 0});
+
+	// Test FCbField(DateTime, 0x1020'3040'5060'7080)
+	TestField<ECbFieldType::DateTime>({0x10, 0x20, 0x30, 0x40, 0x50, 0x60, 0x70, 0x80}, int64(0x1020'3040'5060'7080));
+
+	// Test FCbField(DateTime, Zero) as FDateTime
+	{
+		const uint8 Payload[] = {0, 0, 0, 0, 0, 0, 0, 0};
+		FCbField Field(Payload, ECbFieldType::DateTime);
+		TestEqual(TEXT("FCbField()::AsDateTime()"), Field.AsDateTime(), FDateTime(0));
+	}
+
+	// Test FCbField(None) as DateTime
+	{
+		FCbField DefaultField;
+		TestFieldError<ECbFieldType::DateTime>(DefaultField, ECbFieldError::TypeError);
+		const FDateTime DefaultValue(0x1020'3040'5060'7080);
+		TestEqual(TEXT("FCbField()::AsDateTime()"), DefaultField.AsDateTime(DefaultValue), DefaultValue);
+	}
+
+	return true;
+}
+
+IMPLEMENT_CUSTOM_SIMPLE_AUTOMATION_TEST(FCbFieldTimeSpanTest, FCbFieldTestBase, "System.Core.Serialization.CbField.TimeSpan", CompactBinaryTestFlags)
+bool FCbFieldTimeSpanTest::RunTest(const FString& Parameters)
+{
+	// Test FCbField(TimeSpan, Zero)
+	TestField<ECbFieldType::TimeSpan>({0, 0, 0, 0, 0, 0, 0, 0});
+
+	// Test FCbField(TimeSpan, 0x1020'3040'5060'7080)
+	TestField<ECbFieldType::TimeSpan>({0x10, 0x20, 0x30, 0x40, 0x50, 0x60, 0x70, 0x80}, int64(0x1020'3040'5060'7080));
+
+	// Test FCbField(TimeSpan, Zero) as FTimeSpan
+	{
+		const uint8 Payload[] = {0, 0, 0, 0, 0, 0, 0, 0};
+		FCbField Field(Payload, ECbFieldType::TimeSpan);
+		TestEqual(TEXT("FCbField()::AsTimeSpan()"), Field.AsTimeSpan(), FTimespan(0));
+	}
+
+	// Test FCbField(None) as TimeSpan
+	{
+		FCbField DefaultField;
+		TestFieldError<ECbFieldType::TimeSpan>(DefaultField, ECbFieldError::TypeError);
+		const FTimespan DefaultValue(0x1020'3040'5060'7080);
+		TestEqual(TEXT("FCbField()::AsTimeSpan()"), DefaultField.AsTimeSpan(DefaultValue), DefaultValue);
+	}
+
+	return true;
+}
+
+///////////////////////////////////////////////////////////////////////////////////////////////////////////////////////
+
+IMPLEMENT_CUSTOM_SIMPLE_AUTOMATION_TEST(FCbFieldRefTest, FCbFieldTestBase, "System.Core.Serialization.CbFieldRef", CompactBinaryTestFlags)
+bool FCbFieldRefTest::RunTest(const FString& Parameters)
+{
+	static_assert(std::is_constructible<FCbFieldRef, const FSharedBufferRef&>::value, "Missing constructor for FCbFieldRef");
+	static_assert(std::is_constructible<FCbFieldRef, const FSharedBufferPtr&>::value, "Missing constructor for FCbFieldRef");
+	static_assert(std::is_constructible<FCbFieldRef, const FSharedBufferConstRef&>::value, "Missing constructor for FCbFieldRef");
+	static_assert(std::is_constructible<FCbFieldRef, const FSharedBufferConstPtr&>::value, "Missing constructor for FCbFieldRef");
+
+	static_assert(std::is_constructible<FCbFieldRef, FSharedBufferRef&&>::value, "Missing constructor for FCbFieldRef");
+	static_assert(std::is_constructible<FCbFieldRef, FSharedBufferPtr&&>::value, "Missing constructor for FCbFieldRef");
+	static_assert(std::is_constructible<FCbFieldRef, FSharedBufferConstRef&&>::value, "Missing constructor for FCbFieldRef");
+	static_assert(std::is_constructible<FCbFieldRef, FSharedBufferConstPtr&&>::value, "Missing constructor for FCbFieldRef");
+
+	static_assert(std::is_constructible<FCbFieldRef, const FCbField&, const FSharedBufferRef&>::value, "Missing constructor for FCbFieldRef");
+	static_assert(std::is_constructible<FCbFieldRef, const FCbField&, const FSharedBufferPtr&>::value, "Missing constructor for FCbFieldRef");
+	static_assert(std::is_constructible<FCbFieldRef, const FCbField&, const FSharedBufferConstRef&>::value, "Missing constructor for FCbFieldRef");
+	static_assert(std::is_constructible<FCbFieldRef, const FCbField&, const FSharedBufferConstPtr&>::value, "Missing constructor for FCbFieldRef");
+	static_assert(std::is_constructible<FCbFieldRef, const FCbField&, const FCbFieldRef&>::value, "Missing constructor for FCbFieldRef");
+	static_assert(std::is_constructible<FCbFieldRef, const FCbField&, const FCbArrayRef&>::value, "Missing constructor for FCbFieldRef");
+	static_assert(std::is_constructible<FCbFieldRef, const FCbField&, const FCbObjectRef&>::value, "Missing constructor for FCbFieldRef");
+
+	static_assert(std::is_constructible<FCbFieldRef, const FCbField&, FSharedBufferRef&&>::value, "Missing constructor for FCbFieldRef");
+	static_assert(std::is_constructible<FCbFieldRef, const FCbField&, FSharedBufferPtr&&>::value, "Missing constructor for FCbFieldRef");
+	static_assert(std::is_constructible<FCbFieldRef, const FCbField&, FSharedBufferConstRef&&>::value, "Missing constructor for FCbFieldRef");
+	static_assert(std::is_constructible<FCbFieldRef, const FCbField&, FSharedBufferConstPtr&&>::value, "Missing constructor for FCbFieldRef");
+	static_assert(std::is_constructible<FCbFieldRef, const FCbField&, FCbFieldRef&&>::value, "Missing constructor for FCbFieldRef");
+	static_assert(std::is_constructible<FCbFieldRef, const FCbField&, FCbArrayRef&&>::value, "Missing constructor for FCbFieldRef");
+	static_assert(std::is_constructible<FCbFieldRef, const FCbField&, FCbObjectRef&&>::value, "Missing constructor for FCbFieldRef");
+
+	static_assert(std::is_constructible<FCbFieldRef, TCbBufferRef<FCbField>&&>::value, "Missing constructor for FCbFieldRef");
+	static_assert(std::is_assignable<FCbFieldRef, TCbBufferRef<FCbField>&&>::value, "Missing assignment for FCbFieldRef");
+
+	// Test FCbFieldRef()
+	{
+		FCbFieldRef DefaultField;
+		TestFalse(TEXT("FCbFieldRef().HasValue()"), DefaultField.HasValue());
+		TestFalse(TEXT("FCbFieldRef().IsOwned()"), DefaultField.IsOwned());
+		DefaultField.MakeOwned();
+		TestTrue(TEXT("FCbFieldRef().MakeOwned().IsOwned()"), DefaultField.IsOwned());
+	}
+
+	// Test Field w/ Type from Shared Buffer
+	{
+		uint8 Payload[] = { uint8(ECbFieldType::Binary), 3, 4, 5, 6 }; // Size: 3, Data: 4/5/6
+		FSharedBufferRef WrapBuffer = MakeSharedBuffer(FSharedBuffer::Wrap, MakeMemoryView(Payload));
+		FSharedBufferRef OwnedBuffer = MakeSharedBufferOwned(WrapBuffer);
+		FSharedBufferConstPtr WrapBufferPtr = WrapBuffer;
+		FSharedBufferConstPtr OwnedBufferPtr = OwnedBuffer;
+
+		FCbFieldRef WrapRef(WrapBuffer);
+		FCbFieldRef WrapPtr(WrapBufferPtr);
+		FCbFieldRef WrapPtrMove{FSharedBufferConstPtr(WrapBufferPtr)};
+		FCbFieldRef WrapOuterFieldRef(ImplicitConv<FCbField>(WrapPtr), WrapPtrMove);
+		FCbFieldRef WrapOuterBufferRef(ImplicitConv<FCbField>(WrapRef), WrapPtr);
+		FCbFieldRef OwnedRef(OwnedBuffer);
+		FCbFieldRef OwnedPtr(OwnedBufferPtr);
+		FCbFieldRef OwnedPtrMove{FSharedBufferConstPtr(OwnedBufferPtr)};
+		FCbFieldRef OwnedOuterFieldRef(ImplicitConv<FCbField>(OwnedPtr), OwnedPtrMove);
+		FCbFieldRef OwnedOuterBufferRef(ImplicitConv<FCbField>(OwnedRef), OwnedPtr);
+
+		// These lines are expected to assert when uncommented.
+		//FCbFieldRef InvalidOuterBuffer(ImplicitConv<FCbField>(OwnedRef), WrapBufferPtr);
+		//FCbFieldRef InvalidOuterBufferMove(ImplicitConv<FCbField>(OwnedRef), FSharedBufferConstPtr(WrapBufferPtr));
+
+		Payload[UE_ARRAY_COUNT(Payload) - 1] = 4;
+
+		TestEqualBytes(TEXT("FCbFieldRef(WrapBufferRef)"), WrapRef.AsBinary(), {4, 5, 4});
+		TestEqualBytes(TEXT("FCbFieldRef(WrapBufferPtr)"), WrapPtr.AsBinary(), {4, 5, 4});
+		TestEqualBytes(TEXT("FCbFieldRef(WrapBufferPtr&&)"), WrapPtrMove.AsBinary(), {4, 5, 4});
+		TestEqualBytes(TEXT("FCbFieldRef(WrapOuterFieldRef)"), WrapOuterFieldRef.AsBinary(), {4, 5, 4});
+		TestEqualBytes(TEXT("FCbFieldRef(WrapOuterBufferRef)"), WrapOuterBufferRef.AsBinary(), {4, 5, 4});
+		TestEqualBytes(TEXT("FCbFieldRef(OwnedBufferRef)"), OwnedRef.AsBinary(), {4, 5, 6});
+		TestEqualBytes(TEXT("FCbFieldRef(OwnedBufferPtr)"), OwnedPtr.AsBinary(), {4, 5, 6});
+		TestEqualBytes(TEXT("FCbFieldRef(OwnedBufferPtr&&)"), OwnedPtrMove.AsBinary(), {4, 5, 6});
+		TestEqualBytes(TEXT("FCbFieldRef(OwnedOuterFieldRef)"), OwnedOuterFieldRef.AsBinary(), {4, 5, 6});
+		TestEqualBytes(TEXT("FCbFieldRef(OwnedOuterBufferRef)"), OwnedOuterBufferRef.AsBinary(), {4, 5, 6});
+
+		TestFalse(TEXT("FCbFieldRef(WrapBufferRef).IsOwned()"), WrapRef.IsOwned());
+		TestFalse(TEXT("FCbFieldRef(WrapBufferPtr).IsOwned()"), WrapPtr.IsOwned());
+		TestFalse(TEXT("FCbFieldRef(WrapBufferPtr&&).IsOwned()"), WrapPtrMove.IsOwned());
+		TestFalse(TEXT("FCbFieldRef(WrapOuterFieldRef).IsOwned()"), WrapOuterFieldRef.IsOwned());
+		TestFalse(TEXT("FCbFieldRef(WrapOuterBufferRef).IsOwned()"), WrapOuterBufferRef.IsOwned());
+		TestTrue(TEXT("FCbFieldRef(OwnedBufferRef).IsOwned()"), OwnedRef.IsOwned());
+		TestTrue(TEXT("FCbFieldRef(OwnedBufferPtr).IsOwned()"), OwnedPtr.IsOwned());
+		TestTrue(TEXT("FCbFieldRef(OwnedBufferPtr&&).IsOwned()"), OwnedPtrMove.IsOwned());
+		TestTrue(TEXT("FCbFieldRef(OwnedOuterFieldRef).IsOwned()"), OwnedOuterFieldRef.IsOwned());
+		TestTrue(TEXT("FCbFieldRef(OwnedOuterBufferRef).IsOwned()"), OwnedOuterBufferRef.IsOwned());
+
+		WrapRef.MakeOwned();
+		OwnedRef.MakeOwned();
+		static_cast<uint8*>(OwnedBuffer->GetData())[UE_ARRAY_COUNT(Payload) - 1] = 5;
+		TestEqualBytes(TEXT("FCbFieldRef(Wrap).MakeOwned()"), WrapRef.AsBinary(), {4, 5, 4});
+		TestTrue(TEXT("FCbFieldRef(Wrap).MakeOwned().IsOwned()"), WrapRef.IsOwned());
+		TestEqualBytes(TEXT("FCbFieldRef(Owned).MakeOwned()"), OwnedRef.AsBinary(), {4, 5, 5});
+		TestTrue(TEXT("FCbFieldRef(Owned).MakeOwned().IsOwned()"), OwnedRef.IsOwned());
+	}
+
+	// Test Field w/ Type
+	{
+		uint8 Payload[] = { uint8(ECbFieldType::Binary), 3, 4, 5, 6 }; // Size: 3, Data: 4/5/6
+		uint8* PayloadCopy = static_cast<uint8*>(FMemory::Malloc(sizeof(Payload)));
+		FMemory::Memcpy(PayloadCopy, Payload, sizeof(Payload));
+
+		FCbField Field(Payload);
+
+		FCbFieldRef VoidAssume(FCbFieldRef::AssumeOwnership, ImplicitConv<const void*>(PayloadCopy));
+		FCbFieldRef VoidWrap(FCbFieldRef::Wrap, Payload);
+		FCbFieldRef VoidClone(FCbFieldRef::Clone, Payload);
+		FCbFieldRef FieldWrap(FCbFieldRef::Wrap, Field);
+		FCbFieldRef FieldClone(FCbFieldRef::Clone, Field);
+		FCbFieldRef FieldRefClone(FCbFieldRef::Clone, FieldWrap);
+
+		Payload[UE_ARRAY_COUNT(Payload) - 1] = 4;
+
+		TestEqualBytes(TEXT("FCbFieldRef(AssumeOwnership, Void)"), VoidAssume.AsBinary(), {4, 5, 6});
+		TestEqualBytes(TEXT("FCbFieldRef(Wrap, Void)"), VoidWrap.AsBinary(), {4, 5, 4});
+		TestEqualBytes(TEXT("FCbFieldRef(Clone, Void)"), VoidClone.AsBinary(), {4, 5, 6});
+		TestEqualBytes(TEXT("FCbFieldRef(Wrap, Field)"), FieldWrap.AsBinary(), {4, 5, 4});
+		TestEqualBytes(TEXT("FCbFieldRef(Clone, Field)"), FieldClone.AsBinary(), {4, 5, 6});
+		TestEqualBytes(TEXT("FCbFieldRef(Clone, FieldRef)"), FieldRefClone.AsBinary(), {4, 5, 6});
+
+		TestTrue(TEXT("FCbFieldRef(AssumeOwnership, Void).IsOwned()"), VoidAssume.IsOwned());
+		TestFalse(TEXT("FCbFieldRef(Wrap, Void).IsOwned()"), VoidWrap.IsOwned());
+		TestTrue(TEXT("FCbFieldRef(Clone, Void).IsOwned()"), VoidClone.IsOwned());
+		TestFalse(TEXT("FCbFieldRef(Wrap, Field).IsOwned()"), FieldWrap.IsOwned());
+		TestTrue(TEXT("FCbFieldRef(Clone, Field).IsOwned()"), FieldClone.IsOwned());
+		TestTrue(TEXT("FCbFieldRef(Clone, FieldRef).IsOwned()"), FieldRefClone.IsOwned());
+	}
+
+	// Test Field w/o Type
+	{
+		uint8 Payload[] = { 3, 4, 5, 6 }; // Size: 3, Data: 4/5/6
+		FCbField Field(Payload, ECbFieldType::Binary);
+
+		FCbFieldRef FieldWrap(FCbFieldRef::Wrap, Field);
+		FCbFieldRef FieldClone(FCbFieldRef::Clone, Field);
+		FCbFieldRef FieldRefClone(FCbFieldRef::Clone, FieldWrap);
+
+		Payload[UE_ARRAY_COUNT(Payload) - 1] = 4;
+
+		TestEqualBytes(TEXT("FCbFieldRef(Wrap, Field, NoType)"), FieldWrap.AsBinary(), {4, 5, 4});
+		TestEqualBytes(TEXT("FCbFieldRef(Clone, Field, NoType)"), FieldClone.AsBinary(), {4, 5, 6});
+		TestEqualBytes(TEXT("FCbFieldRef(Clone, FieldRef, NoType)"), FieldRefClone.AsBinary(), {4, 5, 6});
+
+		TestFalse(TEXT("FCbFieldRef(Wrap, Field, NoType).IsOwned()"), FieldWrap.IsOwned());
+		TestTrue(TEXT("FCbFieldRef(Clone, Field, NoType).IsOwned()"), FieldClone.IsOwned());
+		TestTrue(TEXT("FCbFieldRef(Clone, FieldRef, NoType).IsOwned()"), FieldRefClone.IsOwned());
+
+		FieldWrap.MakeOwned();
+		TestEqualBytes(TEXT("FCbFieldRef(Wrap, NoType).MakeOwned()"), FieldWrap.AsBinary(), {4, 5, 4});
+		TestTrue(TEXT("FCbFieldRef(Wrap, NoType).MakeOwned().IsOwned()"), FieldWrap.IsOwned());
+	}
+
+	return true;
+}
+
+///////////////////////////////////////////////////////////////////////////////////////////////////////////////////////
+
+IMPLEMENT_SIMPLE_AUTOMATION_TEST(FCbArrayRefTest, "System.Core.Serialization.CbArrayRef", CompactBinaryTestFlags)
+bool FCbArrayRefTest::RunTest(const FString& Parameters)
+{
+	static_assert(std::is_constructible<FCbArrayRef, const FCbArray&, const FSharedBufferConstPtr&>::value, "Missing constructor for FCbArrayRef");
+	static_assert(std::is_constructible<FCbArrayRef, const FCbArray&, const FCbFieldRef&>::value, "Missing constructor for FCbArrayRef");
+	static_assert(std::is_constructible<FCbArrayRef, const FCbArray&, const FCbArrayRef&>::value, "Missing constructor for FCbArrayRef");
+	static_assert(std::is_constructible<FCbArrayRef, const FCbArray&, const FCbObjectRef&>::value, "Missing constructor for FCbArrayRef");
+
+	static_assert(std::is_constructible<FCbArrayRef, const FCbArray&, FSharedBufferConstPtr&&>::value, "Missing constructor for FCbArrayRef");
+	static_assert(std::is_constructible<FCbArrayRef, const FCbArray&, FCbFieldRef&&>::value, "Missing constructor for FCbArrayRef");
+	static_assert(std::is_constructible<FCbArrayRef, const FCbArray&, FCbArrayRef&&>::value, "Missing constructor for FCbArrayRef");
+	static_assert(std::is_constructible<FCbArrayRef, const FCbArray&, FCbObjectRef&&>::value, "Missing constructor for FCbArrayRef");
+
+	static_assert(std::is_constructible<FCbArrayRef, TCbBufferRef<FCbArray>&&>::value, "Missing constructor for FCbArrayRef");
+	static_assert(std::is_assignable<FCbArrayRef, TCbBufferRef<FCbArray>&&>::value, "Missing assignment for FCbArrayRef");
+
+	// Test FCbArrayRef()
+	{
+		FCbArrayRef DefaultArray;
+		TestFalse(TEXT("FCbArrayRef().IsOwned()"), DefaultArray.IsOwned());
+		DefaultArray.MakeOwned();
+		TestTrue(TEXT("FCbArrayRef().MakeOwned().IsOwned()"), DefaultArray.IsOwned());
+	}
+
+	return true;
+}
+
+///////////////////////////////////////////////////////////////////////////////////////////////////////////////////////
+
+IMPLEMENT_SIMPLE_AUTOMATION_TEST(FCbObjectRefTest, "System.Core.Serialization.CbObjectRef", CompactBinaryTestFlags)
+bool FCbObjectRefTest::RunTest(const FString& Parameters)
+{
+	static_assert(std::is_constructible<FCbObjectRef, const FCbObject&, const FSharedBufferConstPtr&>::value, "Missing constructor for FCbObjectRef");
+	static_assert(std::is_constructible<FCbObjectRef, const FCbObject&, const FCbFieldRef&>::value, "Missing constructor for FCbObjectRef");
+	static_assert(std::is_constructible<FCbObjectRef, const FCbObject&, const FCbArrayRef&>::value, "Missing constructor for FCbObjectRef");
+	static_assert(std::is_constructible<FCbObjectRef, const FCbObject&, const FCbObjectRef&>::value, "Missing constructor for FCbObjectRef");
+
+	static_assert(std::is_constructible<FCbObjectRef, const FCbObject&, FSharedBufferConstPtr&&>::value, "Missing constructor for FCbObjectRef");
+	static_assert(std::is_constructible<FCbObjectRef, const FCbObject&, FCbFieldRef&&>::value, "Missing constructor for FCbObjectRef");
+	static_assert(std::is_constructible<FCbObjectRef, const FCbObject&, FCbArrayRef&&>::value, "Missing constructor for FCbObjectRef");
+	static_assert(std::is_constructible<FCbObjectRef, const FCbObject&, FCbObjectRef&&>::value, "Missing constructor for FCbObjectRef");
+
+	static_assert(std::is_constructible<FCbObjectRef, TCbBufferRef<FCbObject>&&>::value, "Missing constructor for FCbObjectRef");
+	static_assert(std::is_assignable<FCbObjectRef, TCbBufferRef<FCbObject>&&>::value, "Missing assignment for FCbObjectRef");
+
+	// Test FCbObjectRef()
+	{
+		FCbObjectRef DefaultObject;
+		TestFalse(TEXT("FCbObjectRef().IsOwned()"), DefaultObject.IsOwned());
+		DefaultObject.MakeOwned();
+		TestTrue(TEXT("FCbObjectRef().MakeOwned().IsOwned()"), DefaultObject.IsOwned());
+	}
+
+	return true;
+}
+
+///////////////////////////////////////////////////////////////////////////////////////////////////////////////////////
+
+IMPLEMENT_SIMPLE_AUTOMATION_TEST(FCbValidateTest, "System.Core.Serialization.ValidateCompactBinary", CompactBinaryTestFlags)
+bool FCbValidateTest::RunTest(const FString& Parameters)
+{
+	auto Validate = [this](std::initializer_list<uint8> Data, ECbFieldType Type = ECbFieldType::HasFieldType) -> ECbValidateError
+	{
+		return ValidateCompactBinary(MakeMemoryView(Data), ECbValidateMode::All, Type);
+	};
+	auto ValidateMode = [this](std::initializer_list<uint8> Data, ECbValidateMode Mode, ECbFieldType Type = ECbFieldType::HasFieldType) -> ECbValidateError
+	{
+		return ValidateCompactBinary(MakeMemoryView(Data), Mode, Type);
+	};
+
+	auto AddName = [](ECbFieldType Type) -> uint8 { return uint8(Type | ECbFieldType::HasFieldName); };
+
+	constexpr uint8 NullNoName = uint8(ECbFieldType::Null);
+	constexpr uint8 NullWithName = uint8(ECbFieldType::Null | ECbFieldType::HasFieldName);
+	constexpr uint8 IntNoName = uint8(ECbFieldType::IntegerPositive);
+	constexpr uint8 IntWithName = uint8(ECbFieldType::IntegerPositive | ECbFieldType::HasFieldName);
+
+	// Test OutOfBounds
+	TestEqual(TEXT("ValidateCompactBinary(OutOfBounds, Empty)"), Validate({}), ECbValidateError::OutOfBounds);
+
+	TestEqual(TEXT("ValidateCompactBinary(Valid, Null)"), Validate({NullNoName}), ECbValidateError::None);
+	TestEqual(TEXT("ValidateCompactBinary(Valid, Null, Name)"), Validate({NullWithName, 1, 'N'}), ECbValidateError::None);
+	TestEqual(TEXT("ValidateCompactBinary(OutOfBounds, Null, Name)"), Validate({NullWithName}), ECbValidateError::OutOfBounds);
+	TestEqual(TEXT("ValidateCompactBinary(OutOfBounds, Null, Name)"), Validate({NullWithName, 1}), ECbValidateError::OutOfBounds);
+	TestEqual(TEXT("ValidateCompactBinary(OutOfBounds, Null, Name)"), Validate({NullWithName, 0x80}), ECbValidateError::OutOfBounds);
+	TestEqual(TEXT("ValidateCompactBinary(OutOfBounds, Null, Name)"), Validate({NullWithName, 0x80, 128}), ECbValidateError::OutOfBounds);
+
+	TestEqual(TEXT("ValidateCompactBinary(Valid, Object, Empty)"), Validate({uint8(ECbFieldType::Object), 0}), ECbValidateError::None);
+	TestEqual(TEXT("ValidateCompactBinary(Valid, Object, Empty, NoType)"), Validate({0}, ECbFieldType::Object), ECbValidateError::None);
+	TestEqual(TEXT("ValidateCompactBinary(Valid, Object, Field)"), Validate({uint8(ECbFieldType::Object), 7, NullWithName, 1, 'N', IntWithName, 1, 'I', 0}), ECbValidateError::None);
+	TestEqual(TEXT("ValidateCompactBinary(Valid, Object, Field, NoType)"), Validate({7, NullWithName, 1, 'N', IntWithName, 1, 'I', 0}, ECbFieldType::Object), ECbValidateError::None);
+	TestEqual(TEXT("ValidateCompactBinary(OutOfBounds, Object)"), Validate({uint8(ECbFieldType::Object)}), ECbValidateError::OutOfBounds);
+	TestEqual(TEXT("ValidateCompactBinary(OutOfBounds, Object, NoType)"), Validate({}, ECbFieldType::Object), ECbValidateError::OutOfBounds);
+	TestEqual(TEXT("ValidateCompactBinary(OutOfBounds, Object)"), Validate({uint8(ECbFieldType::Object), 1}), ECbValidateError::OutOfBounds);
+	TestEqual(TEXT("ValidateCompactBinary(OutOfBounds, Object, NoType)"), Validate({1}, ECbFieldType::Object), ECbValidateError::OutOfBounds);
+	TestEqual(TEXT("ValidateCompactBinary(OutOfBounds, Object, OOB Field)"), Validate({uint8(ECbFieldType::Object), 3, AddName(ECbFieldType::Float32), 1, 'N'}), ECbValidateError::OutOfBounds);
+	TestEqual(TEXT("ValidateCompactBinary(OutOfBounds, Object, OOB Field, NoType)"), Validate({3, AddName(ECbFieldType::Float32), 1, 'N'}, ECbFieldType::Object), ECbValidateError::OutOfBounds);
+
+	TestEqual(TEXT("ValidateCompactBinary(Valid, UniformObject, Field)"), Validate({uint8(ECbFieldType::UniformObject), 3, NullWithName, 1, 'N'}), ECbValidateError::None);
+	TestEqual(TEXT("ValidateCompactBinary(Valid, UniformObject, Field, NoType)"), Validate({3, NullWithName, 1, 'N'}, ECbFieldType::UniformObject), ECbValidateError::None);
+	TestEqual(TEXT("ValidateCompactBinary(OutOfBounds, UniformObject)"), Validate({uint8(ECbFieldType::UniformObject)}), ECbValidateError::OutOfBounds);
+	TestEqual(TEXT("ValidateCompactBinary(OutOfBounds, UniformObject, NoType)"), Validate({}, ECbFieldType::UniformObject), ECbValidateError::OutOfBounds);
+	TestEqual(TEXT("ValidateCompactBinary(OutOfBounds, UniformObject)"), Validate({uint8(ECbFieldType::UniformObject), 1}), ECbValidateError::OutOfBounds);
+	TestEqual(TEXT("ValidateCompactBinary(OutOfBounds, UniformObject, NoType)"), Validate({1}, ECbFieldType::UniformObject), ECbValidateError::OutOfBounds);
+	TestEqual(TEXT("ValidateCompactBinary(OutOfBounds, UniformObject, OOB Field)"), Validate({uint8(ECbFieldType::UniformObject), 3, AddName(ECbFieldType::Float32), 1, 'N'}), ECbValidateError::OutOfBounds);
+	TestEqual(TEXT("ValidateCompactBinary(OutOfBounds, UniformObject, OOB Field, NoType)"), Validate({3, AddName(ECbFieldType::Float32), 1, 'N'}, ECbFieldType::UniformObject), ECbValidateError::OutOfBounds);
+
+	TestEqual(TEXT("ValidateCompactBinary(Valid, Array, Empty)"), Validate({uint8(ECbFieldType::Array), 1, 0}), ECbValidateError::None);
+	TestEqual(TEXT("ValidateCompactBinary(Valid, Array, Empty, NoType)"), Validate({1, 0}, ECbFieldType::Array), ECbValidateError::None);
+	TestEqual(TEXT("ValidateCompactBinary(Valid, Array, Field)"), Validate({uint8(ECbFieldType::Array), 4, 2, NullNoName, IntNoName, 0}), ECbValidateError::None);
+	TestEqual(TEXT("ValidateCompactBinary(Valid, Array, Field, NoType)"), Validate({4, 2, NullNoName, IntNoName, 0}, ECbFieldType::Array), ECbValidateError::None);
+	TestEqual(TEXT("ValidateCompactBinary(OutOfBounds, Array)"), Validate({uint8(ECbFieldType::Array)}), ECbValidateError::OutOfBounds);
+	TestEqual(TEXT("ValidateCompactBinary(OutOfBounds, Array, NoType)"), Validate({}, ECbFieldType::Array), ECbValidateError::OutOfBounds);
+	TestEqual(TEXT("ValidateCompactBinary(OutOfBounds, Array)"), Validate({uint8(ECbFieldType::Array), 1}), ECbValidateError::OutOfBounds);
+	TestEqual(TEXT("ValidateCompactBinary(OutOfBounds, Array, NoType)"), Validate({1}, ECbFieldType::Array), ECbValidateError::OutOfBounds);
+	TestEqual(TEXT("ValidateCompactBinary(OutOfBounds, Array, OOB Field)"), Validate({uint8(ECbFieldType::Array), 2, 1, uint8(ECbFieldType::Float32)}), ECbValidateError::OutOfBounds);
+	TestEqual(TEXT("ValidateCompactBinary(OutOfBounds, Array, OOB Field, NoType)"), Validate({2, 1, uint8(ECbFieldType::Float32)}, ECbFieldType::Array), ECbValidateError::OutOfBounds);
+
+	TestEqual(TEXT("ValidateCompactBinary(Valid, UniformArray, Field)"), Validate({uint8(ECbFieldType::UniformArray), 3, 1, IntNoName, 0}), ECbValidateError::None);
+	TestEqual(TEXT("ValidateCompactBinary(Valid, UniformArray, Field, NoType)"), Validate({3, 1, IntNoName, 0}, ECbFieldType::UniformArray), ECbValidateError::None);
+	TestEqual(TEXT("ValidateCompactBinary(OutOfBounds, UniformArray)"), Validate({uint8(ECbFieldType::UniformArray)}), ECbValidateError::OutOfBounds);
+	TestEqual(TEXT("ValidateCompactBinary(OutOfBounds, UniformArray, NoType)"), Validate({}, ECbFieldType::UniformArray), ECbValidateError::OutOfBounds);
+	TestEqual(TEXT("ValidateCompactBinary(OutOfBounds, UniformArray)"), Validate({uint8(ECbFieldType::UniformArray), 1}), ECbValidateError::OutOfBounds);
+	TestEqual(TEXT("ValidateCompactBinary(OutOfBounds, UniformArray, NoType)"), Validate({1}, ECbFieldType::UniformArray), ECbValidateError::OutOfBounds);
+	TestEqual(TEXT("ValidateCompactBinary(OutOfBounds, UniformArray, OOB Field)"), Validate({uint8(ECbFieldType::UniformArray), 2, 1, uint8(ECbFieldType::Float32)}), ECbValidateError::OutOfBounds);
+	TestEqual(TEXT("ValidateCompactBinary(OutOfBounds, UniformArray, OOB Field, NoType)"), Validate({2, 1, uint8(ECbFieldType::Float32)}, ECbFieldType::UniformArray), ECbValidateError::OutOfBounds);
+
+	TestEqual(TEXT("ValidateCompactBinary(Valid, Binary, Empty)"), Validate({uint8(ECbFieldType::Binary), 0}), ECbValidateError::None);
+	TestEqual(TEXT("ValidateCompactBinary(Valid, Binary, Empty, NoType)"), Validate({0}, ECbFieldType::Binary), ECbValidateError::None);
+	TestEqual(TEXT("ValidateCompactBinary(Valid, Binary, Field)"), Validate({uint8(ECbFieldType::Binary), 1, 0}), ECbValidateError::None);
+	TestEqual(TEXT("ValidateCompactBinary(Valid, Binary, Field, NoType)"), Validate({1, 0}, ECbFieldType::Binary), ECbValidateError::None);
+	TestEqual(TEXT("ValidateCompactBinary(OutOfBounds, Binary)"), Validate({uint8(ECbFieldType::Binary)}), ECbValidateError::OutOfBounds);
+	TestEqual(TEXT("ValidateCompactBinary(OutOfBounds, Binary, NoType)"), Validate({}, ECbFieldType::Binary), ECbValidateError::OutOfBounds);
+	TestEqual(TEXT("ValidateCompactBinary(OutOfBounds, Binary)"), Validate({uint8(ECbFieldType::Binary), 1}), ECbValidateError::OutOfBounds);
+	TestEqual(TEXT("ValidateCompactBinary(OutOfBounds, Binary, NoType)"), Validate({1}, ECbFieldType::Binary), ECbValidateError::OutOfBounds);
+
+	TestEqual(TEXT("ValidateCompactBinary(Valid, String, Empty)"), Validate({uint8(ECbFieldType::String), 0}), ECbValidateError::None);
+	TestEqual(TEXT("ValidateCompactBinary(Valid, String, Empty, NoType)"), Validate({0}, ECbFieldType::String), ECbValidateError::None);
+	TestEqual(TEXT("ValidateCompactBinary(Valid, String, Field)"), Validate({uint8(ECbFieldType::String), 1, 'S'}), ECbValidateError::None);
+	TestEqual(TEXT("ValidateCompactBinary(Valid, String, Field, NoType)"), Validate({1, 'S'}, ECbFieldType::String), ECbValidateError::None);
+	TestEqual(TEXT("ValidateCompactBinary(OutOfBounds, String)"), Validate({uint8(ECbFieldType::String)}), ECbValidateError::OutOfBounds);
+	TestEqual(TEXT("ValidateCompactBinary(OutOfBounds, String, NoType)"), Validate({}, ECbFieldType::String), ECbValidateError::OutOfBounds);
+	TestEqual(TEXT("ValidateCompactBinary(OutOfBounds, String)"), Validate({uint8(ECbFieldType::String), 1}), ECbValidateError::OutOfBounds);
+	TestEqual(TEXT("ValidateCompactBinary(OutOfBounds, String, NoType)"), Validate({1}, ECbFieldType::String), ECbValidateError::OutOfBounds);
+
+	TestEqual(TEXT("ValidateCompactBinary(Valid, IntegerPositive, 1-byte)"), Validate({uint8(ECbFieldType::IntegerPositive), 0}), ECbValidateError::None);
+	TestEqual(TEXT("ValidateCompactBinary(Valid, IntegerPositive, 1-byte, NoType)"), Validate({0}, ECbFieldType::IntegerPositive), ECbValidateError::None);
+	TestEqual(TEXT("ValidateCompactBinary(Valid, IntegerPositive, 2-byte)"), Validate({uint8(ECbFieldType::IntegerPositive), 0x80, 0x80}), ECbValidateError::None);
+	TestEqual(TEXT("ValidateCompactBinary(Valid, IntegerPositive, 2-byte, NoType)"), Validate({0x80, 0x80}, ECbFieldType::IntegerPositive), ECbValidateError::None);
+	TestEqual(TEXT("ValidateCompactBinary(OutOfBounds, IntegerPositive, 1-byte)"), Validate({uint8(ECbFieldType::IntegerPositive)}), ECbValidateError::OutOfBounds);
+	TestEqual(TEXT("ValidateCompactBinary(OutOfBounds, IntegerPositive, 1-byte, NoType)"), Validate({}, ECbFieldType::IntegerPositive), ECbValidateError::OutOfBounds);
+	TestEqual(TEXT("ValidateCompactBinary(OutOfBounds, IntegerPositive, 2-byte)"), Validate({uint8(ECbFieldType::IntegerPositive), 0x80}), ECbValidateError::OutOfBounds);
+	TestEqual(TEXT("ValidateCompactBinary(OutOfBounds, IntegerPositive, 2-byte, NoType)"), Validate({0x80}, ECbFieldType::IntegerPositive), ECbValidateError::OutOfBounds);
+	TestEqual(TEXT("ValidateCompactBinary(OutOfBounds, IntegerPositive, 9-byte)"), Validate({uint8(ECbFieldType::IntegerPositive), 0xff, 0, 0, 0, 0, 0, 0, 0}), ECbValidateError::OutOfBounds);
+	TestEqual(TEXT("ValidateCompactBinary(OutOfBounds, IntegerPositive, 9-byte, NoType)"), Validate({0xff, 0, 0, 0, 0, 0, 0, 0}, ECbFieldType::IntegerPositive), ECbValidateError::OutOfBounds);
+
+	TestEqual(TEXT("ValidateCompactBinary(Valid, IntegerNegative, 1-byte)"), Validate({uint8(ECbFieldType::IntegerNegative), 0}), ECbValidateError::None);
+	TestEqual(TEXT("ValidateCompactBinary(Valid, IntegerNegative, 1-byte, NoType)"), Validate({0}, ECbFieldType::IntegerNegative), ECbValidateError::None);
+	TestEqual(TEXT("ValidateCompactBinary(Valid, IntegerNegative, 2-byte)"), Validate({uint8(ECbFieldType::IntegerNegative), 0x80, 0x80}), ECbValidateError::None);
+	TestEqual(TEXT("ValidateCompactBinary(Valid, IntegerNegative, 2-byte, NoType)"), Validate({0x80, 0x80}, ECbFieldType::IntegerNegative), ECbValidateError::None);
+	TestEqual(TEXT("ValidateCompactBinary(OutOfBounds, IntegerNegative, 1-byte)"), Validate({uint8(ECbFieldType::IntegerNegative)}), ECbValidateError::OutOfBounds);
+	TestEqual(TEXT("ValidateCompactBinary(OutOfBounds, IntegerNegative, 1-byte, NoType)"), Validate({}, ECbFieldType::IntegerNegative), ECbValidateError::OutOfBounds);
+	TestEqual(TEXT("ValidateCompactBinary(OutOfBounds, IntegerNegative, 2-byte)"), Validate({uint8(ECbFieldType::IntegerNegative), 0x80}), ECbValidateError::OutOfBounds);
+	TestEqual(TEXT("ValidateCompactBinary(OutOfBounds, IntegerNegative, 2-byte, NoType)"), Validate({0x80}, ECbFieldType::IntegerNegative), ECbValidateError::OutOfBounds);
+	TestEqual(TEXT("ValidateCompactBinary(OutOfBounds, IntegerNegative, 9-byte)"), Validate({uint8(ECbFieldType::IntegerNegative), 0xff, 0, 0, 0, 0, 0, 0, 0}), ECbValidateError::OutOfBounds);
+	TestEqual(TEXT("ValidateCompactBinary(OutOfBounds, IntegerNegative, 9-byte, NoType)"), Validate({0xff, 0, 0, 0, 0, 0, 0, 0}, ECbFieldType::IntegerNegative), ECbValidateError::OutOfBounds);
+
+	TestEqual(TEXT("ValidateCompactBinary(Valid, Float32)"), Validate({uint8(ECbFieldType::Float32), 0, 0, 0, 0}), ECbValidateError::None);
+	TestEqual(TEXT("ValidateCompactBinary(Valid, Float32, NoType)"), Validate({0, 0, 0, 0}, ECbFieldType::Float32), ECbValidateError::None);
+	TestEqual(TEXT("ValidateCompactBinary(OutOfBounds, Float32)"), Validate({uint8(ECbFieldType::Float32), 0, 0, 0}), ECbValidateError::OutOfBounds);
+	TestEqual(TEXT("ValidateCompactBinary(OutOfBounds, Float32, NoType)"), Validate({0, 0, 0}, ECbFieldType::Float32), ECbValidateError::OutOfBounds);
+
+	TestEqual(TEXT("ValidateCompactBinary(Valid, Float64)"), Validate({uint8(ECbFieldType::Float64), 0x3f, 0xff, 0xff, 0xff, 0xf0, 0x00, 0x00, 0x00}), ECbValidateError::None);
+	TestEqual(TEXT("ValidateCompactBinary(Valid, Float64, NoType)"), Validate({0x3f, 0xff, 0xff, 0xff, 0xf0, 0x00, 0x00, 0x00}, ECbFieldType::Float64), ECbValidateError::None);
+	TestEqual(TEXT("ValidateCompactBinary(OutOfBounds, Float64)"), Validate({uint8(ECbFieldType::Float64), 0, 0, 0, 0, 0, 0, 0}), ECbValidateError::OutOfBounds);
+	TestEqual(TEXT("ValidateCompactBinary(OutOfBounds, Float64, NoType)"), Validate({0, 0, 0, 0, 0, 0, 0}, ECbFieldType::Float64), ECbValidateError::OutOfBounds);
+
+	TestEqual(TEXT("ValidateCompactBinary(Valid, BoolFalse)"), Validate({uint8(ECbFieldType::BoolFalse)}), ECbValidateError::None);
+	TestEqual(TEXT("ValidateCompactBinary(Valid, BoolTrue)"), Validate({uint8(ECbFieldType::BoolTrue)}), ECbValidateError::None);
+
+	TestEqual(TEXT("ValidateCompactBinary(Valid, BinaryHash)"), Validate({uint8(ECbFieldType::BinaryHash), 0, 0, 0, 0, 0, 0, 0, 0, 0, 0, 0, 0, 0, 0, 0, 0, 0, 0, 0, 0, 0, 0, 0, 0, 0, 0, 0, 0, 0, 0, 0, 0}), ECbValidateError::None);
+	TestEqual(TEXT("ValidateCompactBinary(Valid, BinaryHash, NoType)"), Validate({0, 0, 0, 0, 0, 0, 0, 0, 0, 0, 0, 0, 0, 0, 0, 0, 0, 0, 0, 0, 0, 0, 0, 0, 0, 0, 0, 0, 0, 0, 0, 0}, ECbFieldType::BinaryHash), ECbValidateError::None);
+	TestEqual(TEXT("ValidateCompactBinary(OutOfBounds, BinaryHash)"), Validate({uint8(ECbFieldType::BinaryHash), 0, 0, 0, 0, 0, 0, 0, 0, 0, 0, 0, 0, 0, 0, 0, 0, 0, 0, 0, 0, 0, 0, 0, 0, 0, 0, 0, 0, 0, 0, 0}), ECbValidateError::OutOfBounds);
+	TestEqual(TEXT("ValidateCompactBinary(OutOfBounds, BinaryHash, NoType)"), Validate({0, 0, 0, 0, 0, 0, 0, 0, 0, 0, 0, 0, 0, 0, 0, 0, 0, 0, 0, 0, 0, 0, 0, 0, 0, 0, 0, 0, 0, 0, 0}, ECbFieldType::BinaryHash), ECbValidateError::OutOfBounds);
+
+	TestEqual(TEXT("ValidateCompactBinary(Valid, FieldHash)"), Validate({uint8(ECbFieldType::FieldHash), 0, 0, 0, 0, 0, 0, 0, 0, 0, 0, 0, 0, 0, 0, 0, 0, 0, 0, 0, 0, 0, 0, 0, 0, 0, 0, 0, 0, 0, 0, 0, 0}), ECbValidateError::None);
+	TestEqual(TEXT("ValidateCompactBinary(Valid, FieldHash, NoType)"), Validate({0, 0, 0, 0, 0, 0, 0, 0, 0, 0, 0, 0, 0, 0, 0, 0, 0, 0, 0, 0, 0, 0, 0, 0, 0, 0, 0, 0, 0, 0, 0, 0}, ECbFieldType::FieldHash), ECbValidateError::None);
+	TestEqual(TEXT("ValidateCompactBinary(OutOfBounds, FieldHash)"), Validate({uint8(ECbFieldType::FieldHash), 0, 0, 0, 0, 0, 0, 0, 0, 0, 0, 0, 0, 0, 0, 0, 0, 0, 0, 0, 0, 0, 0, 0, 0, 0, 0, 0, 0, 0, 0, 0}), ECbValidateError::OutOfBounds);
+	TestEqual(TEXT("ValidateCompactBinary(OutOfBounds, FieldHash, NoType)"), Validate({0, 0, 0, 0, 0, 0, 0, 0, 0, 0, 0, 0, 0, 0, 0, 0, 0, 0, 0, 0, 0, 0, 0, 0, 0, 0, 0, 0, 0, 0, 0}, ECbFieldType::FieldHash), ECbValidateError::OutOfBounds);
+
+	TestEqual(TEXT("ValidateCompactBinary(Valid, Uuid)"), Validate({uint8(ECbFieldType::Uuid), 0, 0, 0, 0, 0, 0, 0, 0, 0, 0, 0, 0, 0, 0, 0, 0}), ECbValidateError::None);
+	TestEqual(TEXT("ValidateCompactBinary(Valid, Uuid, NoType)"), Validate({0, 0, 0, 0, 0, 0, 0, 0, 0, 0, 0, 0, 0, 0, 0, 0}, ECbFieldType::Uuid), ECbValidateError::None);
+	TestEqual(TEXT("ValidateCompactBinary(OutOfBounds, Uuid)"), Validate({uint8(ECbFieldType::Uuid), 0, 0, 0, 0, 0, 0, 0, 0, 0, 0, 0, 0, 0, 0, 0}), ECbValidateError::OutOfBounds);
+	TestEqual(TEXT("ValidateCompactBinary(OutOfBounds, Uuid, NoType)"), Validate({0, 0, 0, 0, 0, 0, 0, 0, 0, 0, 0, 0, 0, 0, 0}, ECbFieldType::Uuid), ECbValidateError::OutOfBounds);
+
+	TestEqual(TEXT("ValidateCompactBinary(Valid, DateTime)"), Validate({uint8(ECbFieldType::DateTime), 0, 0, 0, 0, 0, 0, 0, 0}), ECbValidateError::None);
+	TestEqual(TEXT("ValidateCompactBinary(Valid, DateTime, NoType)"), Validate({0, 0, 0, 0, 0, 0, 0, 0}, ECbFieldType::DateTime), ECbValidateError::None);
+	TestEqual(TEXT("ValidateCompactBinary(OutOfBounds, DateTime)"), Validate({uint8(ECbFieldType::DateTime), 0, 0, 0, 0, 0, 0, 0}), ECbValidateError::OutOfBounds);
+	TestEqual(TEXT("ValidateCompactBinary(OutOfBounds, DateTime, NoType)"), Validate({0, 0, 0, 0, 0, 0, 0}, ECbFieldType::DateTime), ECbValidateError::OutOfBounds);
+
+	TestEqual(TEXT("ValidateCompactBinary(Valid, TimeSpan)"), Validate({uint8(ECbFieldType::TimeSpan), 0, 0, 0, 0, 0, 0, 0, 0}), ECbValidateError::None);
+	TestEqual(TEXT("ValidateCompactBinary(Valid, TimeSpan, NoType)"), Validate({0, 0, 0, 0, 0, 0, 0, 0}, ECbFieldType::TimeSpan), ECbValidateError::None);
+	TestEqual(TEXT("ValidateCompactBinary(OutOfBounds, TimeSpan)"), Validate({uint8(ECbFieldType::TimeSpan), 0, 0, 0, 0, 0, 0, 0}), ECbValidateError::OutOfBounds);
+	TestEqual(TEXT("ValidateCompactBinary(OutOfBounds, TimeSpan, NoType)"), Validate({0, 0, 0, 0, 0, 0, 0}, ECbFieldType::TimeSpan), ECbValidateError::OutOfBounds);
+
+	// Test InvalidType
+	TestEqual(TEXT("ValidateCompactBinary(InvalidType, Unknown)"), Validate({uint8(ECbFieldType::TimeSpan) + 1}), ECbValidateError::InvalidType);
+	TestEqual(TEXT("ValidateCompactBinary(InvalidType, Unknown)"), Validate({}, ECbFieldType(uint8(ECbFieldType::TimeSpan) + 1)), ECbValidateError::InvalidType);
+	TestEqual(TEXT("ValidateCompactBinary(InvalidType, HasFieldType)"), Validate({uint8(ECbFieldType::Null | ECbFieldType::HasFieldType)}), ECbValidateError::InvalidType);
+
+	TestEqual(TEXT("ValidateCompactBinary(InvalidType, ZeroSizeField)"), Validate({}, ECbFieldType::Null), ECbValidateError::InvalidType);
+	TestEqual(TEXT("ValidateCompactBinary(InvalidType, ZeroSizeField, BoolFalse)"), Validate({}, ECbFieldType::BoolFalse), ECbValidateError::InvalidType);
+	TestEqual(TEXT("ValidateCompactBinary(InvalidType, ZeroSizeField, BoolTrue)"), Validate({}, ECbFieldType::BoolTrue), ECbValidateError::InvalidType);
+
+	TestEqual(TEXT("ValidateCompactBinary(InvalidType, ZeroSizeField, Array)"), Validate({uint8(ECbFieldType::UniformArray), 2, 2, NullNoName}), ECbValidateError::InvalidType);
+	TestEqual(TEXT("ValidateCompactBinary(InvalidType, ZeroSizeField, Object)"), Validate({uint8(ECbFieldType::UniformObject), 2, NullNoName, 0}), ECbValidateError::InvalidType);
+
+	// Test DuplicateName
+	TestEqual(TEXT("ValidateCompactBinary(DuplicateName)"), Validate({uint8(ECbFieldType::UniformObject), 7, NullWithName, 1, 'A', 1, 'B', 1, 'A'}), ECbValidateError::DuplicateName);
+	TestEqual(TEXT("ValidateCompactBinary(DuplicateName, CaseSensitive)"), Validate({uint8(ECbFieldType::UniformObject), 7, NullWithName, 1, 'A', 1, 'B', 1, 'a'}), ECbValidateError::None);
+	TestEqual(TEXT("ValidateCompactBinary(DuplicateName, Mode)"), ValidateMode({uint8(ECbFieldType::UniformObject), 7, NullWithName, 1, 'A', 1, 'B', 1, 'A'}, ECbValidateMode::All & ~ECbValidateMode::Names), ECbValidateError::None);
+
+	// Test MissingName
+	TestEqual(TEXT("ValidateCompactBinary(MissingName)"), Validate({uint8(ECbFieldType::Object), 3, NullNoName, IntNoName, 0}), ECbValidateError::MissingName);
+	TestEqual(TEXT("ValidateCompactBinary(MissingName, Uniform)"), Validate({uint8(ECbFieldType::UniformObject), 3, IntNoName, 0, 0}), ECbValidateError::MissingName);
+	TestEqual(TEXT("ValidateCompactBinary(MissingName, Mode)"), ValidateMode({uint8(ECbFieldType::Object), 3, NullNoName, IntNoName, 0}, ECbValidateMode::All & ~ECbValidateMode::Names), ECbValidateError::None);
+	TestEqual(TEXT("ValidateCompactBinary(MissingName, Uniform, Mode)"), ValidateMode({uint8(ECbFieldType::UniformObject), 3, IntNoName, 0, 0}, ECbValidateMode::All & ~ECbValidateMode::Names), ECbValidateError::None);
+
+	// Test ArrayName
+	TestEqual(TEXT("ValidateCompactBinary(ArrayName)"), Validate({uint8(ECbFieldType::Array), 5, 2, NullNoName, NullWithName, 1, 'F'}), ECbValidateError::ArrayName);
+	TestEqual(TEXT("ValidateCompactBinary(ArrayName, Uniform)"), Validate({uint8(ECbFieldType::UniformArray), 4, 1, NullWithName, 1, 'F'}), ECbValidateError::ArrayName);
+	TestEqual(TEXT("ValidateCompactBinary(ArrayName, Mode)"), ValidateMode({uint8(ECbFieldType::Array), 5, 2, NullNoName, NullWithName, 1, 'F'}, ECbValidateMode::All & ~ECbValidateMode::Names), ECbValidateError::None);
+	TestEqual(TEXT("ValidateCompactBinary(ArrayName, Uniform, Mode)"), ValidateMode({uint8(ECbFieldType::UniformArray), 4, 1, NullWithName, 1, 'F'}, ECbValidateMode::All & ~ECbValidateMode::Names), ECbValidateError::None);
+
+	// Test InvalidString
+	// Not tested or implemented yet because the engine does not provide enough UTF-8 functionality.
+
+	// Test InvalidInteger
+	TestEqual(TEXT("ValidateCompactBinary(InvalidInteger, NameSize)"), Validate({NullWithName, 0x80, 1, 'N'}), ECbValidateError::InvalidInteger);
+	TestEqual(TEXT("ValidateCompactBinary(InvalidInteger, ObjectSize)"), Validate({uint8(ECbFieldType::Object), 0xc0, 0, 0}), ECbValidateError::InvalidInteger);
+	TestEqual(TEXT("ValidateCompactBinary(InvalidInteger, ArraySize)"), Validate({uint8(ECbFieldType::Array), 0xe0, 0, 0, 1, 0}), ECbValidateError::InvalidInteger);
+	TestEqual(TEXT("ValidateCompactBinary(InvalidInteger, ArrayCount)"), Validate({uint8(ECbFieldType::Array), 5, 0xf0, 0, 0, 0, 0}), ECbValidateError::InvalidInteger);
+	TestEqual(TEXT("ValidateCompactBinary(InvalidInteger, BinarySize)"), Validate({uint8(ECbFieldType::Binary), 0xf8, 0, 0, 0, 0, 0}), ECbValidateError::InvalidInteger);
+	TestEqual(TEXT("ValidateCompactBinary(InvalidInteger, StringSize)"), Validate({uint8(ECbFieldType::String), 0xfc, 0, 0, 0, 0, 0, 0}), ECbValidateError::InvalidInteger);
+	TestEqual(TEXT("ValidateCompactBinary(InvalidInteger, IntegerPositive)"), Validate({uint8(ECbFieldType::IntegerPositive), 0xfe, 0, 0, 0, 0, 0, 0, 0}), ECbValidateError::InvalidInteger);
+	TestEqual(TEXT("ValidateCompactBinary(InvalidInteger, IntegerNegative)"), Validate({uint8(ECbFieldType::IntegerNegative), 0xff, 0, 0, 0, 0, 0, 0, 0, 0}), ECbValidateError::InvalidInteger);
+	TestEqual(TEXT("ValidateCompactBinary(InvalidInteger, ArraySize)"), Validate({uint8(ECbFieldType::Array), 0x80, 1, 0}), ECbValidateError::InvalidInteger);
+	TestEqual(TEXT("ValidateCompactBinary(InvalidInteger, ArrayCount)"), Validate({uint8(ECbFieldType::Array), 3, 0xc0, 0, 0}), ECbValidateError::InvalidInteger);
+	TestEqual(TEXT("ValidateCompactBinary(InvalidInteger, ObjectSize)"), Validate({uint8(ECbFieldType::Object), 0xe0, 0, 0, 0}), ECbValidateError::InvalidInteger);
+
+	TestEqual(TEXT("ValidateCompactBinary(InvalidInteger, NameSize, Mode)"), ValidateMode({NullWithName, 0x80, 1, 'N'}, ECbValidateMode::All & ~ECbValidateMode::Format), ECbValidateError::None);
+	TestEqual(TEXT("ValidateCompactBinary(InvalidInteger, ArraySize, Mode)"), ValidateMode({uint8(ECbFieldType::Array), 0xc0, 0, 1, 0}, ECbValidateMode::All & ~ECbValidateMode::Format), ECbValidateError::None);
+	TestEqual(TEXT("ValidateCompactBinary(InvalidInteger, ObjectSize, Mode)"), ValidateMode({uint8(ECbFieldType::Object), 0xe0, 0, 0, 0}, ECbValidateMode::All & ~ECbValidateMode::Format), ECbValidateError::None);
+
+	// Test InvalidFloat
+	TestEqual(TEXT("ValidateCompactBinary(InvalidFloat, MaxSignificant+1)"), Validate({uint8(ECbFieldType::Float64), 0x3f, 0xff, 0xff, 0xff, 0xf0, 0x00, 0x00, 0x00}), ECbValidateError::None); // 1.9999999403953552
+	TestEqual(TEXT("ValidateCompactBinary(InvalidFloat, MaxExponent+1)"), Validate({uint8(ECbFieldType::Float64), 0x47, 0xff, 0xff, 0xff, 0xe0, 0x00, 0x00, 0x00}), ECbValidateError::None); // 6.8056469327705771e38
+	TestEqual(TEXT("ValidateCompactBinary(InvalidFloat, MaxSignificand)"), Validate({uint8(ECbFieldType::Float64), 0x3f, 0xff, 0xff, 0xff, 0xe0, 0x00, 0x00, 0x00}), ECbValidateError::InvalidFloat); // 1.9999998807907104
+	TestEqual(TEXT("ValidateCompactBinary(InvalidFloat, MaxExponent)"), Validate({uint8(ECbFieldType::Float64), 0x47, 0xef, 0xff, 0xff, 0xe0, 0x00, 0x00, 0x00}), ECbValidateError::InvalidFloat); // 3.4028234663852886e38
+	TestEqual(TEXT("ValidateCompactBinary(InvalidFloat, MaxSignificand, Mode)"), ValidateMode({uint8(ECbFieldType::Float64), 0x3f, 0xff, 0xff, 0xff, 0xe0, 0x00, 0x00, 0x00}, ECbValidateMode::All & ~ECbValidateMode::Format), ECbValidateError::None); // 1.9999998807907104
+	TestEqual(TEXT("ValidateCompactBinary(InvalidFloat, MaxExponent, Mode)"), ValidateMode({uint8(ECbFieldType::Float64), 0x47, 0xef, 0xff, 0xff, 0xe0, 0x00, 0x00, 0x00}, ECbValidateMode::All & ~ECbValidateMode::Format), ECbValidateError::None); // 3.4028234663852886e38
+
+	// Test NonUniformObject
+	TestEqual(TEXT("ValidateCompactBinary(NonUniformObject)"), Validate({uint8(ECbFieldType::Object), 3, NullWithName, 1, 'A'}), ECbValidateError::NonUniformObject);
+	TestEqual(TEXT("ValidateCompactBinary(NonUniformObject)"), Validate({uint8(ECbFieldType::Object), 6, NullWithName, 1, 'A', NullWithName, 1, 'B'}), ECbValidateError::NonUniformObject);
+	TestEqual(TEXT("ValidateCompactBinary(NonUniformObject, Mode)"), ValidateMode({uint8(ECbFieldType::Object), 3, NullWithName, 1, 'A'}, ECbValidateMode::All & ~ECbValidateMode::Format), ECbValidateError::None);
+	TestEqual(TEXT("ValidateCompactBinary(NonUniformObject, Mode)"), ValidateMode({uint8(ECbFieldType::Object), 6, NullWithName, 1, 'A', NullWithName, 1, 'B'}, ECbValidateMode::All & ~ECbValidateMode::Format), ECbValidateError::None);
+
+	// Test NonUniformArray
+	TestEqual(TEXT("ValidateCompactBinary(NonUniformArray)"), Validate({uint8(ECbFieldType::Array), 3, 1, IntNoName, 0}), ECbValidateError::NonUniformArray);
+	TestEqual(TEXT("ValidateCompactBinary(NonUniformArray)"), Validate({uint8(ECbFieldType::Array), 5, 2, IntNoName, 1, IntNoName, 2}), ECbValidateError::NonUniformArray);
+	TestEqual(TEXT("ValidateCompactBinary(NonUniformArray, Null)"), Validate({uint8(ECbFieldType::Array), 3, 2, NullNoName, NullNoName}), ECbValidateError::None);
+	TestEqual(TEXT("ValidateCompactBinary(NonUniformArray, Bool)"), Validate({uint8(ECbFieldType::Array), 3, 2, uint8(ECbFieldType::BoolFalse), uint8(ECbFieldType::BoolFalse)}), ECbValidateError::None);
+	TestEqual(TEXT("ValidateCompactBinary(NonUniformArray, Bool)"), Validate({uint8(ECbFieldType::Array), 3, 2, uint8(ECbFieldType::BoolTrue), uint8(ECbFieldType::BoolTrue)}), ECbValidateError::None);
+	TestEqual(TEXT("ValidateCompactBinary(NonUniformArray, Mode)"), ValidateMode({uint8(ECbFieldType::Array), 3, 1, IntNoName, 0}, ECbValidateMode::All & ~ECbValidateMode::Format), ECbValidateError::None);
+	TestEqual(TEXT("ValidateCompactBinary(NonUniformArray, Mode)"), ValidateMode({uint8(ECbFieldType::Array), 5, 2, IntNoName, 1, IntNoName, 2}, ECbValidateMode::All & ~ECbValidateMode::Format), ECbValidateError::None);
+
+	// Test Padding
+	TestEqual(TEXT("ValidateCompactBinary(Padding)"), Validate({NullNoName, 0}), ECbValidateError::Padding);
+	TestEqual(TEXT("ValidateCompactBinary(Padding)"), Validate({uint8(ECbFieldType::Array), 1, 0, 0}), ECbValidateError::Padding);
+	TestEqual(TEXT("ValidateCompactBinary(Padding)"), Validate({uint8(ECbFieldType::Object), 0, 0}), ECbValidateError::Padding);
+	TestEqual(TEXT("ValidateCompactBinary(Padding, Mode)"), ValidateMode({NullNoName, 0}, ECbValidateMode::All & ~ECbValidateMode::Padding), ECbValidateError::None);
+	TestEqual(TEXT("ValidateCompactBinary(Padding, Mode)"), ValidateMode({uint8(ECbFieldType::Array), 1, 0, 0}, ECbValidateMode::All & ~ECbValidateMode::Padding), ECbValidateError::None);
+	TestEqual(TEXT("ValidateCompactBinary(Padding, Mode)"), ValidateMode({uint8(ECbFieldType::Object), 0, 0}, ECbValidateMode::All & ~ECbValidateMode::Padding), ECbValidateError::None);
+
+	return true;
+}
+
+///////////////////////////////////////////////////////////////////////////////////////////////////////////////////////
+
+IMPLEMENT_SIMPLE_AUTOMATION_TEST(FCbValidateRangeTest, "System.Core.Serialization.ValidateCompactBinaryRange", CompactBinaryTestFlags)
+bool FCbValidateRangeTest::RunTest(const FString& Parameters)
+{
+	auto Validate = [this](std::initializer_list<uint8> Data) -> ECbValidateError
+	{
+		return ValidateCompactBinaryRange(MakeMemoryView(Data), ECbValidateMode::All);
+	};
+
+	// Test Empty
+	TestEqual(TEXT("ValidateCompactBinaryRange(Empty)"), Validate({}), ECbValidateError::None);
+
+	// Test Valid
+	TestEqual(TEXT("ValidateCompactBinaryRange(Null x2)"), Validate({uint8(ECbFieldType::Null), uint8(ECbFieldType::Null)}), ECbValidateError::None);
+
+	// Test Padding
+	TestEqual(TEXT("ValidateCompactBinaryRange(Padding InvalidType)"), Validate({uint8(ECbFieldType::Null), 0}), ECbValidateError::InvalidType);
+	TestEqual(TEXT("ValidateCompactBinaryRange(Padding OutOfBounds)"), Validate({uint8(ECbFieldType::Null), uint8(ECbFieldType::Binary)}), ECbValidateError::OutOfBounds);
+
+	return true;
+}
+
+///////////////////////////////////////////////////////////////////////////////////////////////////////////////////////
+
+IMPLEMENT_SIMPLE_AUTOMATION_TEST(FCbMeasureTest, "System.Core.Serialization.MeasureCompactBinary", CompactBinaryTestFlags)
+bool FCbMeasureTest::RunTest(const FString& Parameters)
+{
+	auto Measure = [this](std::initializer_list<uint8> Data, ECbFieldType Type = ECbFieldType::HasFieldType) -> uint64
+	{
+		return MeasureCompactBinary(MakeMemoryView(Data), Type);
+	};
+
+	TestEqual(TEXT("MeasureCompactBinary(Empty)"), Measure({}), uint64(0));
+
+	TestEqual(TEXT("MeasureCompactBinary(Null, NoType)"), Measure({}, ECbFieldType::Null), uint64(0));
+	TestEqual(TEXT("MeasureCompactBinary(Null, NameSize1B, NoType)"), Measure({30}, ECbFieldType::Null | ECbFieldType::HasFieldName), uint64(31));
+	TestEqual(TEXT("MeasureCompactBinary(Null, NameSize2B, NoType)"), Measure({0x80, 0x80}, ECbFieldType::Null | ECbFieldType::HasFieldName), uint64(130));
+	TestEqual(TEXT("MeasureCompactBinary(Null, NameSize2BShort, NoType)"), Measure({0x80}, ECbFieldType::Null | ECbFieldType::HasFieldName), uint64(0));
+	TestEqual(TEXT("MeasureCompactBinary(Null, MissingName, NoType)"), Measure({}, ECbFieldType::Null | ECbFieldType::HasFieldName), uint64(0));
+
+	TestEqual(TEXT("MeasureCompactBinary(Null)"), Measure({uint8(ECbFieldType::Null)}), uint64(1));
+	TestEqual(TEXT("MeasureCompactBinary(Null, NameSize1B)"), Measure({uint8(ECbFieldType::Null | ECbFieldType::HasFieldName), 30}), uint64(32));
+	TestEqual(TEXT("MeasureCompactBinary(Null, NameSize2B)"), Measure({uint8(ECbFieldType::Null | ECbFieldType::HasFieldName), 0x80, 0x80}), uint64(131));
+	TestEqual(TEXT("MeasureCompactBinary(Null, NameSize2BShort)"), Measure({uint8(ECbFieldType::Null | ECbFieldType::HasFieldName), 0x80}), uint64(0));
+	TestEqual(TEXT("MeasureCompactBinary(Null, MissingName)"), Measure({uint8(ECbFieldType::Null | ECbFieldType::HasFieldName)}), uint64(0));
+
+	TestEqual(TEXT("MeasureCompactBinary(Object, NoSize)"), Measure({uint8(ECbFieldType::Object)}), uint64(0));
+	TestEqual(TEXT("MeasureCompactBinary(Object, Size1B)"), Measure({uint8(ECbFieldType::Object), 30}), uint64(32));
+	TestEqual(TEXT("MeasureCompactBinary(UniformObject, NoSize)"), Measure({uint8(ECbFieldType::UniformObject)}), uint64(0));
+	TestEqual(TEXT("MeasureCompactBinary(UniformObject, Size1B)"), Measure({uint8(ECbFieldType::UniformObject), 30}), uint64(32));
+
+	TestEqual(TEXT("MeasureCompactBinary(Array, NoSize)"), Measure({uint8(ECbFieldType::Array)}), uint64(0));
+	TestEqual(TEXT("MeasureCompactBinary(Array, Size1B)"), Measure({uint8(ECbFieldType::Array), 30}), uint64(32));
+	TestEqual(TEXT("MeasureCompactBinary(UniformArray, NoSize)"), Measure({uint8(ECbFieldType::UniformArray)}), uint64(0));
+	TestEqual(TEXT("MeasureCompactBinary(UniformArray, Size1B)"), Measure({uint8(ECbFieldType::UniformArray), 30}), uint64(32));
+
+	TestEqual(TEXT("MeasureCompactBinary(Binary, NoSize)"), Measure({uint8(ECbFieldType::Binary)}), uint64(0));
+	TestEqual(TEXT("MeasureCompactBinary(Binary, Size1B)"), Measure({uint8(ECbFieldType::Binary), 30}), uint64(32));
+
+	TestEqual(TEXT("MeasureCompactBinary(String, NoSize)"), Measure({uint8(ECbFieldType::String)}), uint64(0));
+	TestEqual(TEXT("MeasureCompactBinary(String, Size1B)"), Measure({uint8(ECbFieldType::String), 30}), uint64(32));
+	TestEqual(TEXT("MeasureCompactBinary(String, Size2B)"), Measure({uint8(ECbFieldType::String), 0x80, 0x80}), uint64(131));
+	TestEqual(TEXT("MeasureCompactBinary(String, Size2BShort)"), Measure({uint8(ECbFieldType::String), 0x80}), uint64(0));
+	TestEqual(TEXT("MeasureCompactBinary(String, MissingNameSize)"), Measure({uint8(ECbFieldType::String | ECbFieldType::HasFieldName)}), uint64(0));
+	TestEqual(TEXT("MeasureCompactBinary(String, MissingName)"), Measure({uint8(ECbFieldType::String | ECbFieldType::HasFieldName), 1}), uint64(0));
+
+	TestEqual(TEXT("MeasureCompactBinary(IntegerPositive, NoValue)"), Measure({uint8(ECbFieldType::IntegerPositive)}), uint64(0));
+	TestEqual(TEXT("MeasureCompactBinary(IntegerPositive, Value1B)"), Measure({uint8(ECbFieldType::IntegerPositive), 0x7f}), uint64(2));
+	TestEqual(TEXT("MeasureCompactBinary(IntegerPositive, Value2B)"), Measure({uint8(ECbFieldType::IntegerPositive), 0x80}), uint64(3));
+
+	TestEqual(TEXT("MeasureCompactBinary(IntegerNegative, NoValue)"), Measure({uint8(ECbFieldType::IntegerNegative)}), uint64(0));
+	TestEqual(TEXT("MeasureCompactBinary(IntegerNegative, Value1B)"), Measure({uint8(ECbFieldType::IntegerNegative), 0x7f}), uint64(2));
+	TestEqual(TEXT("MeasureCompactBinary(IntegerNegative, Value2B)"), Measure({uint8(ECbFieldType::IntegerNegative), 0x80}), uint64(3));
+
+	TestEqual(TEXT("MeasureCompactBinary(Float32, NoType)"), Measure({}, ECbFieldType::Float32), uint64(4));
+	TestEqual(TEXT("MeasureCompactBinary(Float32, NameSize1B, NoType)"), Measure({30}, ECbFieldType::Float32 | ECbFieldType::HasFieldName), uint64(35));
+	TestEqual(TEXT("MeasureCompactBinary(Float32, NameSize2B, NoType)"), Measure({0x80, 0x80}, ECbFieldType::Float32 | ECbFieldType::HasFieldName), uint64(134));
+	TestEqual(TEXT("MeasureCompactBinary(Float32, NameSize2BShort, NoType)"), Measure({0x80}, ECbFieldType::Float32 | ECbFieldType::HasFieldName), uint64(0));
+	TestEqual(TEXT("MeasureCompactBinary(Float32, MissingName, NoType)"), Measure({}, ECbFieldType::Float32 | ECbFieldType::HasFieldName), uint64(0));
+
+	TestEqual(TEXT("MeasureCompactBinary(Float32)"), Measure({uint8(ECbFieldType::Float32)}), uint64(5));
+	TestEqual(TEXT("MeasureCompactBinary(Float32, NameSize1B)"), Measure({uint8(ECbFieldType::Float32 | ECbFieldType::HasFieldName), 30}), uint64(36));
+	TestEqual(TEXT("MeasureCompactBinary(Float32, NameSize2B)"), Measure({uint8(ECbFieldType::Float32 | ECbFieldType::HasFieldName), 0x80, 0x80}), uint64(135));
+	TestEqual(TEXT("MeasureCompactBinary(Float32, NameSize2BShort)"), Measure({uint8(ECbFieldType::Float32 | ECbFieldType::HasFieldName), 0x80}), uint64(0));
+	TestEqual(TEXT("MeasureCompactBinary(Float32, MissingName)"), Measure({uint8(ECbFieldType::Float32 | ECbFieldType::HasFieldName)}), uint64(0));
+
+	TestEqual(TEXT("MeasureCompactBinary(Float64)"), Measure({uint8(ECbFieldType::Float64)}), uint64(9));
+
+	TestEqual(TEXT("MeasureCompactBinary(BoolFalse)"), Measure({uint8(ECbFieldType::BoolFalse)}), uint64(1));
+	TestEqual(TEXT("MeasureCompactBinary(BoolTrue)"), Measure({uint8(ECbFieldType::BoolTrue)}), uint64(1));
+
+	TestEqual(TEXT("MeasureCompactBinary(BinaryHash)"), Measure({uint8(ECbFieldType::BinaryHash)}), uint64(33));
+	TestEqual(TEXT("MeasureCompactBinary(FieldHash)"), Measure({uint8(ECbFieldType::FieldHash)}), uint64(33));
+
+	TestEqual(TEXT("MeasureCompactBinary(Uuid)"), Measure({uint8(ECbFieldType::Uuid)}), uint64(17));
+
+	TestEqual(TEXT("MeasureCompactBinary(DateTime)"), Measure({uint8(ECbFieldType::DateTime)}), uint64(9));
+	TestEqual(TEXT("MeasureCompactBinary(TimeSpan)"), Measure({uint8(ECbFieldType::TimeSpan)}), uint64(9));
+
+	return true;
+}
+
+///////////////////////////////////////////////////////////////////////////////////////////////////////////////////////
+
+IMPLEMENT_SIMPLE_AUTOMATION_TEST(FCbFieldParseTest, "System.Core.Serialization.CbFieldParseTest", CompactBinaryTestFlags)
+bool FCbFieldParseTest::RunTest(const FString& Parameters)
+{
+	// Test the optimal object parsing loop because it is expected to be required for high performance.
+	// Under ideal conditions, when the fields are in the expected order and there are no extra fields,
+	// the loop will execute once and only one comparison will be performed for each field name. Either
+	// way, each field will only be visited once even if the loop needs to execute several times.
+	auto ParseObject = [this](const FCbObject& Object, uint32& A, uint32& B, uint32& C, uint32& D)
+	{
+		for (FCbFieldIterator It = Object.CreateIterator(); It;)
+		{
+			const FCbFieldIterator Last = It;
+			if (It.Name().Equals("A"_ASV))
+			{
+				A = It.AsUInt32();
+				++It;
+			}
+			if (It.Name().Equals("B"_ASV))
+			{
+				B = It.AsUInt32();
+				++It;
+			}
+			if (It.Name().Equals("C"_ASV))
+			{
+				C = It.AsUInt32();
+				++It;
+			}
+			if (It.Name().Equals("D"_ASV))
+			{
+				D = It.AsUInt32();
+				++It;
+			}
+			if (Last == It)
+			{
+				++It;
+			}
+		}
+	};
+
+	auto TestParseObject = [this, &ParseObject](std::initializer_list<uint8> Data, uint32 A, uint32 B, uint32 C, uint32 D) -> bool
+	{
+		uint32 ParsedA = 0, ParsedB = 0, ParsedC = 0, ParsedD = 0;
+		ParseObject(FCbObject(GetData(Data), ECbFieldType::Object), ParsedA, ParsedB, ParsedC, ParsedD);
+		return A == ParsedA && B == ParsedB && C == ParsedC && D == ParsedD;
+	};
+
+	constexpr uint8 T = uint8(ECbFieldType::IntegerPositive | ECbFieldType::HasFieldName);
+	TestTrue(TEXT("FCbObject Parse(None)"), TestParseObject({0}, 0, 0, 0, 0));
+	TestTrue(TEXT("FCbObject Parse(ABCD)"), TestParseObject({16, T, 1, 'A', 1, T, 1, 'B', 2, T, 1, 'C', 3, T, 1, 'D', 4}, 1, 2, 3, 4));
+	TestTrue(TEXT("FCbObject Parse(BCDA)"), TestParseObject({16, T, 1, 'B', 2, T, 1, 'C', 3, T, 1, 'D', 4, T, 1, 'A', 1}, 1, 2, 3, 4));
+	TestTrue(TEXT("FCbObject Parse(BCD)"), TestParseObject({12, T, 1, 'B', 2, T, 1, 'C', 3, T, 1, 'D', 4}, 0, 2, 3, 4));
+	TestTrue(TEXT("FCbObject Parse(BC)"), TestParseObject({8, T, 1, 'B', 2, T, 1, 'C', 3}, 0, 2, 3, 0));
+	TestTrue(TEXT("FCbObject Parse(ABCDE)"), TestParseObject({20, T, 1, 'A', 1, T, 1, 'B', 2, T, 1, 'C', 3, T, 1, 'D', 4, T, 1, 'E', 5}, 1, 2, 3, 4));
+	TestTrue(TEXT("FCbObject Parse(EABCD)"), TestParseObject({20, T, 1, 'E', 5, T, 1, 'A', 1, T, 1, 'B', 2, T, 1, 'C', 3, T, 1, 'D', 4}, 1, 2, 3, 4));
+	TestTrue(TEXT("FCbObject Parse(DCBA)"), TestParseObject({16, T, 1, 'D', 4, T, 1, 'C', 3, T, 1, 'B', 2, T, 1, 'A', 1}, 1, 2, 3, 4));
+
+	return true;
+}
+
+///////////////////////////////////////////////////////////////////////////////////////////////////////////////////////
+
+#endif // WITH_DEV_AUTOMATION_TESTS
