Copying //UE5/Dev-Cooker@14539516 to Main (//UE5/Main)

diff --git a/Engine/Source/Runtime/Core/Private/Tests/Serialization/CompactBinaryWriterTest.cpp b/Engine/Source/Runtime/Core/Private/Tests/Serialization/CompactBinaryWriterTest.cpp
new file mode 100644
--- /dev/null
+++ b/Engine/Source/Runtime/Core/Private/Tests/Serialization/CompactBinaryWriterTest.cpp
@@ -0,0 +1,1011 @@
+// Copyright Epic Games, Inc. All Rights Reserved.
+
+#include "Serialization/CompactBinaryWriter.h"
+
+#include "Misc/AutomationTest.h"
+#include "Misc/Blake3.h"
+#include "Misc/DateTime.h"
+#include "Misc/Guid.h"
+#include "Misc/Timespan.h"
+#include "Misc/StringBuilder.h"
+
+#if WITH_DEV_AUTOMATION_TESTS
+
+static constexpr EAutomationTestFlags::Type CompactBinaryWriterTestFlags = EAutomationTestFlags::Type(EAutomationTestFlags::ApplicationContextMask | EAutomationTestFlags::EngineFilter);
+
+///////////////////////////////////////////////////////////////////////////////////////////////////////////////////////
+
+IMPLEMENT_SIMPLE_AUTOMATION_TEST(FCbWriterObjectTest, "System.Core.Serialization.CbWriter.Object", CompactBinaryWriterTestFlags)
+bool FCbWriterObjectTest::RunTest(const FString& Parameters)
+{
+	TCbWriter<256> Writer;
+
+	// Test Empty Object
+	{
+		Writer.Reset();
+		Writer.BeginObject();
+		Writer.EndObject();
+		FCbFieldRef Field = Writer.Save();
+		if (TestEqual(TEXT("FCbWriter(Object, Empty) Validate"), ValidateCompactBinary(*Field.Buffer(), ECbValidateMode::All), ECbValidateError::None))
+		{
+			TestTrue(TEXT("FCbWriter(Object, Empty).IsObject()"), Field.IsObject());
+			TestFalse(TEXT("FCbWriter(Object, Empty).AsObject()"), Field.AsObject().CreateIterator().HasValue());
+		}
+	}
+
+	// Test Named Empty Object
+	{
+		Writer.Reset();
+		Writer.Name("Object"_ASV);
+		Writer.BeginObject();
+		Writer.EndObject();
+		FCbFieldRef Field = Writer.Save();
+		if (TestEqual(TEXT("FCbWriter(Object, Empty, Name) Validate"), ValidateCompactBinary(*Field.Buffer(), ECbValidateMode::All), ECbValidateError::None))
+		{
+			TestTrue(TEXT("FCbWriter(Object, Empty, Name).IsObject()"), Field.IsObject());
+			TestFalse(TEXT("FCbWriter(Object, Empty, Name).AsObject()"), Field.AsObject().CreateIterator().HasValue());
+		}
+	}
+
+	// Test Basic Object
+	{
+		Writer.Reset();
+		Writer.BeginObject();
+		Writer.Name("Integer"_ASV).Integer(0);
+		Writer.Name("Float"_ASV).Float(0.0f);
+		Writer.EndObject();
+		FCbFieldRef Field = Writer.Save();
+		if (TestEqual(TEXT("FCbWriter(Object, Basic) Validate"), ValidateCompactBinary(*Field.Buffer(), ECbValidateMode::All), ECbValidateError::None))
+		{
+			TestTrue(TEXT("FCbWriter(Object, Basic).IsObject()"), Field.IsObject());
+			FCbObject Object = Field.AsObject();
+			TestTrue(TEXT("FCbWriter(Object, Basic).AsObject()[Integer]"), Object["Integer"_ASV].IsInteger());
+			TestTrue(TEXT("FCbWriter(Object, Basic).AsObject()[Float]"), Object["Float"_ASV].IsFloat());
+		}
+	}
+
+	// Test Uniform Object
+	{
+		Writer.Reset();
+		Writer.BeginObject();
+		Writer.Name("Field1"_ASV).Integer(0);
+		Writer.Name("Field2"_ASV).Integer(1);
+		Writer.EndObject();
+		FCbFieldRef Field = Writer.Save();
+		if (TestEqual(TEXT("FCbWriter(Object, Uniform) Validate"), ValidateCompactBinary(*Field.Buffer(), ECbValidateMode::All), ECbValidateError::None))
+		{
+			TestTrue(TEXT("FCbWriter(Object, Uniform).IsObject()"), Field.IsObject());
+			FCbObject Object = Field.AsObject();
+			TestTrue(TEXT("FCbWriter(Object, Uniform).AsObject()[Field1]"), Object["Field1"_ASV].IsInteger());
+			TestTrue(TEXT("FCbWriter(Object, Uniform).AsObject()[Field2]"), Object["Field2"_ASV].IsInteger());
+		}
+	}
+
+	return true;
+}
+
+///////////////////////////////////////////////////////////////////////////////////////////////////////////////////////
+
+IMPLEMENT_SIMPLE_AUTOMATION_TEST(FCbWriterArrayTest, "System.Core.Serialization.CbWriter.Array", CompactBinaryWriterTestFlags)
+bool FCbWriterArrayTest::RunTest(const FString& Parameters)
+{
+	TCbWriter<256> Writer;
+
+	// Test Empty Array
+	{
+		Writer.Reset();
+		Writer.BeginArray();
+		Writer.EndArray();
+		FCbFieldRef Field = Writer.Save();
+		if (TestEqual(TEXT("FCbWriter(Array, Empty) Validate"), ValidateCompactBinary(*Field.Buffer(), ECbValidateMode::All), ECbValidateError::None))
+		{
+			TestTrue(TEXT("FCbWriter(Array, Empty).IsArray()"), Field.IsArray());
+			TestEqual(TEXT("FCbWriter(Array, Empty).AsArray()"), Field.AsArray().Num(), uint64(0));
+		}
+	}
+
+	// Test Named Empty Array
+	{
+		Writer.Reset();
+		Writer.Name("Array"_ASV);
+		Writer.BeginArray();
+		Writer.EndArray();
+		FCbFieldRef Field = Writer.Save();
+		if (TestEqual(TEXT("FCbWriter(Array, Empty, Name) Validate"), ValidateCompactBinary(*Field.Buffer(), ECbValidateMode::All), ECbValidateError::None))
+		{
+			TestTrue(TEXT("FCbWriter(Array, Empty, Name).IsArray()"), Field.IsArray());
+			TestEqual(TEXT("FCbWriter(Array, Empty, Name).AsArray()"), Field.AsArray().Num(), uint64(0));
+		}
+	}
+
+	// Test Basic Array
+	{
+		Writer.Reset();
+		Writer.BeginArray();
+		Writer.Integer(0);
+		Writer.Float(0.0f);
+		Writer.EndArray();
+		FCbFieldRef Field = Writer.Save();
+		if (TestEqual(TEXT("FCbWriter(Array, Basic) Validate"), ValidateCompactBinary(*Field.Buffer(), ECbValidateMode::All), ECbValidateError::None))
+		{
+			TestTrue(TEXT("FCbWriter(Array, Basic).IsArray()"), Field.IsArray());
+			FCbFieldIterator Iterator = Field.AsArray().CreateIterator();
+			TestTrue(TEXT("FCbWriter(Array, Basic).AsArray()[Integer]"), Iterator.IsInteger());
+			++Iterator;
+			TestTrue(TEXT("FCbWriter(Array, Basic).AsArray()[Float]"), Iterator.IsFloat());
+			++Iterator;
+			TestFalse(TEXT("FCbWriter(Array, Basic).AsArray()[End]"), Iterator.HasValue());
+		}
+	}
+
+	// Test Uniform Array
+	{
+		Writer.Reset();
+		Writer.BeginArray();
+		Writer.Integer(0);
+		Writer.Integer(1);
+		Writer.EndArray();
+		FCbFieldRef Field = Writer.Save();
+		if (TestEqual(TEXT("FCbWriter(Array, Uniform) Validate"), ValidateCompactBinary(*Field.Buffer(), ECbValidateMode::All), ECbValidateError::None))
+		{
+			TestTrue(TEXT("FCbWriter(Array, Uniform).IsArray()"), Field.IsArray());
+			FCbFieldIterator Iterator = Field.AsArray().CreateIterator();
+			TestTrue(TEXT("FCbWriter(Array, Basic).AsArray()[Field1]"), Iterator.IsInteger());
+			++Iterator;
+			TestTrue(TEXT("FCbWriter(Array, Basic).AsArray()[Field2]"), Iterator.IsInteger());
+			++Iterator;
+			TestFalse(TEXT("FCbWriter(Array, Basic).AsArray()[End]"), Iterator.HasValue());
+		}
+	}
+
+	return true;
+}
+
+///////////////////////////////////////////////////////////////////////////////////////////////////////////////////////
+
+IMPLEMENT_SIMPLE_AUTOMATION_TEST(FCbWriterNullTest, "System.Core.Serialization.CbWriter.Null", CompactBinaryWriterTestFlags)
+bool FCbWriterNullTest::RunTest(const FString& Parameters)
+{
+	TCbWriter<256> Writer;
+
+	// Test Null
+	{
+		Writer.Reset();
+		Writer.Null();
+		FCbFieldRef Field = Writer.Save();
+		if (TestEqual(TEXT("FCbWriter(Null) Validate"), ValidateCompactBinary(*Field.Buffer(), ECbValidateMode::All), ECbValidateError::None))
+		{
+			TestFalse(TEXT("FCbWriter(Null).HasName()"), Field.HasName());
+			TestTrue(TEXT("FCbWriter(Null).IsNull()"), Field.IsNull());
+		}
+	}
+
+	// Test Null with Name
+	{
+		Writer.Reset();
+		Writer.Name("Null"_ASV);
+		Writer.Null();
+		FCbFieldRef Field = Writer.Save();
+		if (TestEqual(TEXT("FCbWriter(Null, Name) Validate"), ValidateCompactBinary(*Field.Buffer(), ECbValidateMode::All), ECbValidateError::None))
+		{
+			TestEqual(TEXT("FCbWriter(Null, Name).Name()"), Field.Name(), "Null"_ASV);
+			TestTrue(TEXT("FCbWriter(Null, Name).HasName()"), Field.HasName());
+			TestTrue(TEXT("FCbWriter(Null, Name).IsNull()"), Field.IsNull());
+		}
+	}
+
+	// Test Null Array/Object Uniformity
+	{
+		Writer.Reset();
+
+		Writer.BeginArray();
+		Writer.Null();
+		Writer.Null();
+		Writer.Null();
+		Writer.EndArray();
+
+		Writer.BeginObject();
+		Writer.Name("N1"_ASV).Null();
+		Writer.Name("N2"_ASV).Null();
+		Writer.Name("N3"_ASV).Null();
+		Writer.EndObject();
+
+		FCbFieldRefIterator Fields = Writer.Save();
+		TestEqual(TEXT("FCbWriter(Null, Uniform) Validate"), ValidateCompactBinaryRange(*Fields.Buffer(), ECbValidateMode::All), ECbValidateError::None);
+	}
+
+	// Test Null with SaveToMemory
+	{
+		Writer.Reset();
+		constexpr int NullCount = 3;
+		for (int Index = 0; Index < NullCount; ++Index)
+		{
+			Writer.Null();
+		}
+		uint8 Buffer[NullCount]{};
+		FCbFieldIterator Fields = Writer.SaveToMemory(MakeMemoryView(Buffer));
+		if (TestEqual(TEXT("FCbWriter(Null, Memory) Validate"), ValidateCompactBinaryRange(MakeMemoryView(Buffer), ECbValidateMode::All), ECbValidateError::None))
+		{
+			for (int Index = 0; Index < NullCount; ++Index)
+			{
+				TestTrue(TEXT("FCbWriter(Null, Memory) IsNull"), Fields.IsNull());
+				++Fields;
+			}
+			TestFalse(TEXT("FCbWriter(Null, Memory) HasValue"), Fields.HasValue());
+		}
+	}
+
+	return true;
+}
+
+///////////////////////////////////////////////////////////////////////////////////////////////////////////////////////
+
+class FCbWriterBinaryTestBase : public FAutomationTestBase
+{
+protected:
+	using FAutomationTestBase::FAutomationTestBase;
+	using FAutomationTestBase::TestEqual;
+
+	template <typename T, typename Size>
+	void TestEqual(const TCHAR* What, TArrayView<T, Size> Actual, TArrayView<T, Size> Expected)
+	{
+		TestTrue(What, Actual.Num() == Expected.Num() && CompareItems(Actual.GetData(), Expected.GetData(), Actual.Num()));
+	}
+};
+
+IMPLEMENT_CUSTOM_SIMPLE_AUTOMATION_TEST(FCbWriterBinaryTest, FCbWriterBinaryTestBase, "System.Core.Serialization.CbWriter.Binary", CompactBinaryWriterTestFlags)
+bool FCbWriterBinaryTest::RunTest(const FString& Parameters)
+{
+	TCbWriter<256> Writer;
+
+	// Test Empty Binary
+	{
+		Writer.Reset();
+		Writer.Binary(nullptr, 0);
+		FCbFieldRef Field = Writer.Save();
+		if (TestEqual(TEXT("FCbWriter(Binary, Empty) Validate"), ValidateCompactBinary(*Field.Buffer(), ECbValidateMode::All), ECbValidateError::None))
+		{
+			TestFalse(TEXT("FCbWriter(Binary, Empty).HasName()"), Field.HasName());
+			TestTrue(TEXT("FCbWriter(Binary, Empty).IsBinary()"), Field.IsBinary());
+			TestTrue(TEXT("FCbWriter(Binary, Empty).AsBinary()"), Field.AsBinary().IsEmpty());
+		}
+	}
+
+	// Test Basic Binary
+	{
+		Writer.Reset();
+		const uint8 BinaryValue[] = { 1, 2, 3, 4, 5, 6 };
+		Writer.Name("Binary"_ASV);
+		Writer.Binary(BinaryValue, sizeof(BinaryValue));
+		FCbFieldRef Field = Writer.Save();
+		if (TestEqual(TEXT("FCbWriter(Binary, Array) Validate"), ValidateCompactBinary(*Field.Buffer(), ECbValidateMode::All), ECbValidateError::None))
+		{
+			TestEqual(TEXT("FCbWriter(Binary, Array).Name()"), Field.Name(), "Binary"_ASV);
+			TestTrue(TEXT("FCbWriter(Binary, Array).HasName()"), Field.HasName());
+			TestTrue(TEXT("FCbWriter(Binary, Array).IsBinary()"), Field.IsBinary());
+			TestTrue(TEXT("FCbWriter(Binary, Array).AsBinary()"), Field.AsBinary().EqualBytes(MakeMemoryView(BinaryValue)));
+		}
+	}
+
+	return true;
+}
+
+///////////////////////////////////////////////////////////////////////////////////////////////////////////////////////
+
+IMPLEMENT_SIMPLE_AUTOMATION_TEST(FCbWriterStringTest, "System.Core.Serialization.CbWriter.String", CompactBinaryWriterTestFlags)
+bool FCbWriterStringTest::RunTest(const FString& Parameters)
+{
+	TCbWriter<256> Writer;
+
+	// Test Empty Strings
+	{
+		Writer.Reset();
+		Writer.String(FAnsiStringView());
+		Writer.String(FWideStringView());
+		FCbFieldRefIterator Fields = Writer.Save();
+		if (TestEqual(TEXT("FCbWriter(String, Empty) Validate"), ValidateCompactBinaryRange(*Fields.Buffer(), ECbValidateMode::All), ECbValidateError::None))
+		{
+			for (FCbField Field : Fields)
+			{
+				TestFalse(TEXT("FCbWriter(String, Empty).HasName()"), Field.HasName());
+				TestTrue(TEXT("FCbWriter(String, Empty).IsString()"), Field.IsString());
+				TestTrue(TEXT("FCbWriter(String, Empty).AsString()"), Field.AsString().IsEmpty());
+			}
+		}
+	}
+
+	// Test Basic Strings
+	{
+		Writer.Reset();
+		Writer.Name("String"_ASV).String("Value"_ASV);
+		Writer.Name("String"_ASV).String(TEXT("Value"_SV));
+		FCbFieldRefIterator Fields = Writer.Save();
+		if (TestEqual(TEXT("FCbWriter(String, Basic) Validate"), ValidateCompactBinaryRange(*Fields.Buffer(), ECbValidateMode::All), ECbValidateError::None))
+		{
+			for (FCbField Field : Fields)
+			{
+				TestEqual(TEXT("FCbWriter(String, Basic).Name()"), Field.Name(), "String"_ASV);
+				TestTrue(TEXT("FCbWriter(String, Basic).HasName()"), Field.HasName());
+				TestTrue(TEXT("FCbWriter(String, Basic).IsString()"), Field.IsString());
+				TestEqual(TEXT("FCbWriter(String, Basic).AsString()"), Field.AsString(), "Value"_ASV);
+			}
+		}
+	}
+
+	// Test Long Strings
+	{
+		Writer.Reset();
+		constexpr int DotCount = 256;
+		TAnsiStringBuilder<DotCount + 1> Dots;
+		for (int Index = 0; Index < DotCount; ++Index)
+		{
+			Dots << '.';
+		}
+		Writer.String(Dots);
+		Writer.String(FString::ChrN(DotCount, TEXT('.')));
+		FCbFieldRefIterator Fields = Writer.Save();
+		if (TestEqual(TEXT("FCbWriter(String, Long) Validate"), ValidateCompactBinaryRange(*Fields.Buffer(), ECbValidateMode::All), ECbValidateError::None))
+		{
+			for (FCbField Field : Fields)
+			{
+				TestEqual(TEXT("FCbWriter(String, Long).AsString()"), Field.AsString(), FAnsiStringView(Dots));
+			}
+		}
+	}
+
+	// Test Non-ASCII String
+	{
+		Writer.Reset();
+		WIDECHAR Value[2] = { 0xd83d, 0xde00 };
+		Writer.String("\xf0\x9f\x98\x80"_ASV);
+		Writer.String(FWideStringView(Value, UE_ARRAY_COUNT(Value)));
+		FCbFieldRefIterator Fields = Writer.Save();
+		if (TestEqual(TEXT("FCbWriter(String, Unicode) Validate"), ValidateCompactBinaryRange(*Fields.Buffer(), ECbValidateMode::All), ECbValidateError::None))
+		{
+			for (FCbField Field : Fields)
+			{
+				TestEqual(TEXT("FCbWriter(String, Unicode).AsString()"), Field.AsString(), "\xf0\x9f\x98\x80"_ASV);
+			}
+		}
+	}
+
+	return true;
+}
+
+///////////////////////////////////////////////////////////////////////////////////////////////////////////////////////
+
+IMPLEMENT_SIMPLE_AUTOMATION_TEST(FCbWriterIntegerTest, "System.Core.Serialization.CbWriter.Integer", CompactBinaryWriterTestFlags)
+bool FCbWriterIntegerTest::RunTest(const FString& Parameters)
+{
+	TCbWriter<256> Writer;
+
+	auto TestInt32 = [this, &Writer](int32 Value)
+	{
+		Writer.Reset();
+		Writer.Integer(Value);
+		FCbFieldRef Field = Writer.Save();
+		if (TestEqual(TEXT("FCbWriter(Integer, Int32) Validate"), ValidateCompactBinary(*Field.Buffer(), ECbValidateMode::All), ECbValidateError::None))
+		{
+			TestEqual(TEXT("FCbWriter(Integer, Int32) Value"), Field.AsInt32(), Value);
+			TestFalse(TEXT("FCbWriter(Integer, Int32) Error"), Field.HasError());
+		}
+	};
+
+	auto TestUInt32 = [this, &Writer](uint32 Value)
+	{
+		Writer.Reset();
+		Writer.Integer(Value);
+		FCbFieldRef Field = Writer.Save();
+		if (TestEqual(TEXT("FCbWriter(Integer, UInt32) Validate"), ValidateCompactBinary(*Field.Buffer(), ECbValidateMode::All), ECbValidateError::None))
+		{
+			TestEqual(TEXT("FCbWriter(Integer, UInt32) Value"), Field.AsUInt32(), Value);
+			TestFalse(TEXT("FCbWriter(Integer, UInt32) Error"), Field.HasError());
+		}
+	};
+
+	auto TestInt64 = [this, &Writer](int64 Value)
+	{
+		Writer.Reset();
+		Writer.Integer(Value);
+		FCbFieldRef Field = Writer.Save();
+		if (TestEqual(TEXT("FCbWriter(Integer, Int64) Validate"), ValidateCompactBinary(*Field.Buffer(), ECbValidateMode::All), ECbValidateError::None))
+		{
+			TestEqual(TEXT("FCbWriter(Integer, Int64) Value"), Field.AsInt64(), Value);
+			TestFalse(TEXT("FCbWriter(Integer, Int64) Error"), Field.HasError());
+		}
+	};
+
+	auto TestUInt64 = [this, &Writer](uint64 Value)
+	{
+		Writer.Reset();
+		Writer.Integer(Value);
+		FCbFieldRef Field = Writer.Save();
+		if (TestEqual(TEXT("FCbWriter(Integer, UInt64) Validate"), ValidateCompactBinary(*Field.Buffer(), ECbValidateMode::All), ECbValidateError::None))
+		{
+			TestEqual(TEXT("FCbWriter(Integer, UInt64) Value"), Field.AsUInt64(), Value);
+			TestFalse(TEXT("FCbWriter(Integer, UInt64) Error"), Field.HasError());
+		}
+	};
+
+	TestUInt32(uint32(0x00));
+	TestUInt32(uint32(0x7f));
+	TestUInt32(uint32(0x80));
+	TestUInt32(uint32(0xff));
+	TestUInt32(uint32(0x0100));
+	TestUInt32(uint32(0x7fff));
+	TestUInt32(uint32(0x8000));
+	TestUInt32(uint32(0xffff));
+	TestUInt32(uint32(0x0001'0000));
+	TestUInt32(uint32(0x7fff'ffff));
+	TestUInt32(uint32(0x8000'0000));
+	TestUInt32(uint32(0xffff'ffff));
+
+	TestUInt64(uint64(0x0000'0001'0000'0000));
+	TestUInt64(uint64(0x7fff'ffff'ffff'ffff));
+	TestUInt64(uint64(0x8000'0000'0000'0000));
+	TestUInt64(uint64(0xffff'ffff'ffff'ffff));
+
+	TestInt32(int32(0x01));
+	TestInt32(int32(0x80));
+	TestInt32(int32(0x81));
+	TestInt32(int32(0x8000));
+	TestInt32(int32(0x8001));
+	TestInt32(int32(0x7fff'ffff));
+	TestInt32(int32(0x8000'0000));
+	TestInt32(int32(0x8000'0001));
+
+	TestInt64(int64(0x0000'0001'0000'0000));
+	TestInt64(int64(0x8000'0000'0000'0000));
+	TestInt64(int64(0x7fff'ffff'ffff'ffff));
+	TestInt64(int64(0x8000'0000'0000'0001));
+	TestInt64(int64(0xffff'ffff'ffff'ffff));
+
+	return true;
+}
+
+///////////////////////////////////////////////////////////////////////////////////////////////////////////////////////
+
+IMPLEMENT_SIMPLE_AUTOMATION_TEST(FCbWriterFloatTest, "System.Core.Serialization.CbWriter.Float", CompactBinaryWriterTestFlags)
+bool FCbWriterFloatTest::RunTest(const FString& Parameters)
+{
+	TCbWriter<256> Writer;
+
+	// Test Float32
+	{
+		Writer.Reset();
+		constexpr float Values[] = { 0.0f, 1.0f, -1.0f, PI };
+		for (float Value : Values)
+		{
+			Writer.Float(Value);
+		}
+		FCbFieldRefIterator Fields = Writer.Save();
+		if (TestEqual(TEXT("FCbWriter(Float, Single) Validate"), ValidateCompactBinaryRange(*Fields.Buffer(), ECbValidateMode::All), ECbValidateError::None))
+		{
+			const float* CheckValue = Values;
+			for (FCbField Field : Fields)
+			{
+				TestEqual(TEXT("FCbWriter(Float, Single).AsFloat()"), Field.AsFloat(), *CheckValue++);
+				TestFalse(TEXT("FCbWriter(Float, Single) Error"), Field.HasError());
+			}
+		}
+	}
+
+	// Test Float64
+	{
+		Writer.Reset();
+		constexpr double Values[] = { 0.0f, 1.0f, -1.0f, PI, 1.9999998807907104, 1.9999999403953552, 3.4028234663852886e38, 6.8056469327705771e38 };
+		for (double Value : Values)
+		{
+			Writer.Float(Value);
+		}
+		FCbFieldRefIterator Fields = Writer.Save();
+		if (TestEqual(TEXT("FCbWriter(Float, Double) Validate"), ValidateCompactBinaryRange(*Fields.Buffer(), ECbValidateMode::All), ECbValidateError::None))
+		{
+			const double* CheckValue = Values;
+			for (FCbField Field : Fields)
+			{
+				TestEqual(TEXT("FCbWriter(Float, Double).AsDouble()"), Field.AsDouble(), *CheckValue++);
+				TestFalse(TEXT("FCbWriter(Float, Double) Error"), Field.HasError());
+			}
+		}
+	}
+
+	return true;
+}
+
+///////////////////////////////////////////////////////////////////////////////////////////////////////////////////////
+
+IMPLEMENT_SIMPLE_AUTOMATION_TEST(FCbWriterBoolTest, "System.Core.Serialization.CbWriter.Bool", CompactBinaryWriterTestFlags)
+bool FCbWriterBoolTest::RunTest(const FString& Parameters)
+{
+	TCbWriter<256> Writer;
+
+	// Test Bool Values
+	{
+		Writer.Bool(true);
+		Writer.Bool(false);
+
+		FCbFieldRefIterator Fields = Writer.Save();
+		if (TestEqual(TEXT("FCbWriter(Bool) Validate"), ValidateCompactBinaryRange(*Fields.Buffer(), ECbValidateMode::All), ECbValidateError::None))
+		{
+			TestTrue(TEXT("FCbWriter(Bool).AsBool()"), Fields.AsBool());
+			TestFalse(TEXT("FCbWriter(Bool) Error"), Fields.HasError());
+			++Fields;
+			TestFalse(TEXT("FCbWriter(Bool).AsBool()"), Fields.AsBool());
+			TestFalse(TEXT("FCbWriter(Bool) Error"), Fields.HasError());
+		}
+	}
+
+	// Test Bool Array/Object Uniformity
+	{
+		Writer.Reset();
+
+		Writer.BeginArray();
+		Writer.Bool(false);
+		Writer.Bool(false);
+		Writer.Bool(false);
+		Writer.EndArray();
+
+		Writer.BeginObject();
+		Writer.Name("B1"_ASV).Bool(false);
+		Writer.Name("B2"_ASV).Bool(false);
+		Writer.Name("B3"_ASV).Bool(false);
+		Writer.EndObject();
+
+		FCbFieldRefIterator Fields = Writer.Save();
+		TestEqual(TEXT("FCbWriter(Bool, Uniform) Validate"), ValidateCompactBinaryRange(*Fields.Buffer(), ECbValidateMode::All), ECbValidateError::None);
+	}
+
+	return true;
+}
+
+///////////////////////////////////////////////////////////////////////////////////////////////////////////////////////
+
+///////////////////////////////////////////////////////////////////////////////////////////////////////////////////////
+
+IMPLEMENT_SIMPLE_AUTOMATION_TEST(FCbWriterBinaryHashTest, "System.Core.Serialization.CbWriter.BinaryHash", CompactBinaryWriterTestFlags)
+bool FCbWriterBinaryHashTest::RunTest(const FString& Parameters)
+{
+	TCbWriter<256> Writer;
+
+	const FBlake3Hash::ByteArray ZeroBytes{};
+	const FBlake3Hash::ByteArray SequentialBytes{1, 2, 3, 4, 5, 6, 7, 8, 9, 10, 11, 12, 13, 14, 15, 16, 17, 18, 19, 20, 21, 22, 23, 24, 25, 26, 27, 28, 29, 30, 31, 32};
+
+	const FBlake3Hash Values[] = { FBlake3Hash(ZeroBytes), FBlake3Hash(SequentialBytes) };
+	for (const FBlake3Hash& Value : Values)
+	{
+		Writer.BinaryHash(Value);
+	}
+
+	FCbFieldRefIterator Fields = Writer.Save();
+	if (TestEqual(TEXT("FCbWriter(BinaryHash) Validate"), ValidateCompactBinaryRange(*Fields.Buffer(), ECbValidateMode::All), ECbValidateError::None))
+	{
+		const FBlake3Hash* CheckValue = Values;
+		for (FCbField Field : Fields)
+		{
+			TestEqual(TEXT("FCbWriter(BinaryHash).AsBinaryHash()"), Field.AsBinaryHash(), *CheckValue++);
+			TestFalse(TEXT("FCbWriter(BinaryHash) Error"), Field.HasError());
+		}
+	}
+
+	return true;
+}
+
+///////////////////////////////////////////////////////////////////////////////////////////////////////////////////////
+
+IMPLEMENT_SIMPLE_AUTOMATION_TEST(FCbWriterFieldHashTest, "System.Core.Serialization.CbWriter.FieldHash", CompactBinaryWriterTestFlags)
+bool FCbWriterFieldHashTest::RunTest(const FString& Parameters)
+{
+	TCbWriter<256> Writer;
+
+	const FBlake3Hash::ByteArray ZeroBytes{};
+	const FBlake3Hash::ByteArray SequentialBytes{1, 2, 3, 4, 5, 6, 7, 8, 9, 10, 11, 12, 13, 14, 15, 16, 17, 18, 19, 20, 21, 22, 23, 24, 25, 26, 27, 28, 29, 30, 31, 32};
+
+	const FBlake3Hash Values[] = { FBlake3Hash(ZeroBytes), FBlake3Hash(SequentialBytes) };
+	for (const FBlake3Hash& Value : Values)
+	{
+		Writer.FieldHash(Value);
+	}
+
+	FCbFieldRefIterator Fields = Writer.Save();
+	if (TestEqual(TEXT("FCbWriter(FieldHash) Validate"), ValidateCompactBinaryRange(*Fields.Buffer(), ECbValidateMode::All), ECbValidateError::None))
+	{
+		const FBlake3Hash* CheckValue = Values;
+		for (FCbField Field : Fields)
+		{
+			TestEqual(TEXT("FCbWriter(FieldHash).AsFieldHash()"), Field.AsFieldHash(), *CheckValue++);
+			TestFalse(TEXT("FCbWriter(FieldHash) Error"), Field.HasError());
+		}
+	}
+
+	return true;
+}
+
+///////////////////////////////////////////////////////////////////////////////////////////////////////////////////////
+
+IMPLEMENT_SIMPLE_AUTOMATION_TEST(FCbWriterUuidTest, "System.Core.Serialization.CbWriter.Uuid", CompactBinaryWriterTestFlags)
+bool FCbWriterUuidTest::RunTest(const FString& Parameters)
+{
+	TCbWriter<256> Writer;
+
+	const FGuid Values[] = { FGuid(), FGuid::NewGuid() };
+	for (const FGuid& Value : Values)
+	{
+		Writer.Uuid(Value);
+	}
+
+	FCbFieldRefIterator Fields = Writer.Save();
+	if (TestEqual(TEXT("FCbWriter(Uuid) Validate"), ValidateCompactBinaryRange(*Fields.Buffer(), ECbValidateMode::All), ECbValidateError::None))
+	{
+		const FGuid* CheckValue = Values;
+		for (FCbField Field : Fields)
+		{
+			TestEqual(TEXT("FCbWriter(Uuid).AsUuid()"), Field.AsUuid(), *CheckValue++);
+			TestFalse(TEXT("FCbWriter(Uuid) Error"), Field.HasError());
+		}
+	}
+
+	return true;
+}
+
+///////////////////////////////////////////////////////////////////////////////////////////////////////////////////////
+
+IMPLEMENT_SIMPLE_AUTOMATION_TEST(FCbWriterDateTimeTest, "System.Core.Serialization.CbWriter.DateTime", CompactBinaryWriterTestFlags)
+bool FCbWriterDateTimeTest::RunTest(const FString& Parameters)
+{
+	TCbWriter<256> Writer;
+
+	const FDateTime Values[] = { FDateTime(0), FDateTime(2020, 5, 13, 15, 10) };
+	for (FDateTime Value : Values)
+	{
+		Writer.DateTime(Value);
+	}
+
+	FCbFieldRefIterator Fields = Writer.Save();
+	if (TestEqual(TEXT("FCbWriter(DateTime) Validate"), ValidateCompactBinaryRange(*Fields.Buffer(), ECbValidateMode::All), ECbValidateError::None))
+	{
+		const FDateTime* CheckValue = Values;
+		for (FCbField Field : Fields)
+		{
+			TestEqual(TEXT("FCbWriter(DateTime).AsDateTime()"), Field.AsDateTime(), *CheckValue++);
+			TestFalse(TEXT("FCbWriter(DateTime) Error"), Field.HasError());
+		}
+	}
+
+	return true;
+}
+
+///////////////////////////////////////////////////////////////////////////////////////////////////////////////////////
+
+IMPLEMENT_SIMPLE_AUTOMATION_TEST(FCbWriterTimeSpanTest, "System.Core.Serialization.CbWriter.TimeSpan", CompactBinaryWriterTestFlags)
+bool FCbWriterTimeSpanTest::RunTest(const FString& Parameters)
+{
+	TCbWriter<256> Writer;
+
+	const FTimespan Values[] = { FTimespan(0), FTimespan(1, 2, 4, 8) };
+	for (FTimespan Value : Values)
+	{
+		Writer.TimeSpan(Value);
+	}
+
+	FCbFieldRefIterator Fields = Writer.Save();
+	if (TestEqual(TEXT("FCbWriter(TimeSpan) Validate"), ValidateCompactBinaryRange(*Fields.Buffer(), ECbValidateMode::All), ECbValidateError::None))
+	{
+		const FTimespan* CheckValue = Values;
+		for (FCbField Field : Fields)
+		{
+			TestEqual(TEXT("FCbWriter(TimeSpan).AsTimeSpan()"), Field.AsTimeSpan(), *CheckValue++);
+			TestFalse(TEXT("FCbWriter(TimeSpan) Error"), Field.HasError());
+		}
+	}
+
+	return true;
+}
+
+///////////////////////////////////////////////////////////////////////////////////////////////////////////////////////
+
+IMPLEMENT_SIMPLE_AUTOMATION_TEST(FCbWriterComplexTest, "System.Core.Serialization.CbWriter.Complex", CompactBinaryWriterTestFlags)
+bool FCbWriterComplexTest::RunTest(const FString& Parameters)
+{
+	FCbObjectRef Object;
+	{
+		FCbWriter Writer;
+		Writer.BeginObject();
+
+		const uint8 LocalField[] = { uint8(ECbFieldType::IntegerPositive | ECbFieldType::HasFieldName), 1, 'I', 42 };
+		Writer.Name("FieldCopy"_ASV).Field(FCbField(LocalField));
+
+		const uint8 LocalObject[] = { uint8(ECbFieldType::Object | ECbFieldType::HasFieldName), 1, 'O', 7,
+			uint8(ECbFieldType::IntegerPositive | ECbFieldType::HasFieldName), 1, 'I', 42,
+			uint8(ECbFieldType::Null | ECbFieldType::HasFieldName), 1, 'N' };
+		Writer.Name("ObjectCopy"_ASV).Object(FCbObject(LocalObject));
+
+		const uint8 LocalArray[] = { uint8(ECbFieldType::UniformArray | ECbFieldType::HasFieldName), 1, 'A', 4, 2,
+			uint8(ECbFieldType::IntegerPositive), 42, 21 };
+		Writer.Name("ArrayCopy"_ASV).Array(FCbArray(LocalArray));
+
+		Writer.Name("Null"_ASV).Null();
+
+		Writer.Name("Binary"_ASV);
+		Writer.BeginObject();
+		{
+			Writer.Name("Empty"_ASV).Binary(FConstMemoryView());
+			Writer.Name("Value"_ASV).Binary(MakeMemoryView("BinaryValue"));
+			Writer.Name("LargeValue"_ASV).Binary(MakeMemoryView(FString::ChrN(256, TEXT('.'))));
+		}
+		Writer.EndObject();
+
+		Writer.Name("Strings"_ASV);
+		Writer.BeginObject();
+		{
+			Writer.Name("AnsiString"_ASV).String("AnsiValue"_ASV);
+			Writer.Name("WideString"_ASV).String(FString::ChrN(256, TEXT('.')));
+			Writer.Name("EmptyAnsiString"_ASV).String(FAnsiStringView());
+			Writer.Name("EmptyWideString"_ASV).String(FWideStringView());
+		}
+		Writer.EndObject();
+
+		Writer.Name("Integers"_ASV);
+		Writer.BeginArray();
+		{
+			Writer.Integer(int32(-1));
+			Writer.Integer(int64(-1));
+			Writer.Integer(uint32(1));
+			Writer.Integer(uint64(1));
+			Writer.Integer(MIN_int32);
+			Writer.Integer(MAX_int32);
+			Writer.Integer(MAX_uint32);
+			Writer.Integer(MIN_int64);
+			Writer.Integer(MAX_int64);
+			Writer.Integer(MAX_uint64);
+		}
+		Writer.EndArray();
+
+		Writer.Name("UniformIntegers"_ASV);
+		Writer.BeginArray();
+		{
+			Writer.Integer(0);
+			Writer.Integer(MAX_int32);
+			Writer.Integer(MAX_uint32);
+			Writer.Integer(MAX_int64);
+			Writer.Integer(MAX_uint64);
+		}
+		Writer.EndArray();
+
+		Writer.Name("Float32"_ASV).Float(1.0f);
+		Writer.Name("Float64as32"_ASV).Float(2.0);
+		Writer.Name("Float64"_ASV).Float(3.0e100);
+
+		Writer.Name("False"_ASV).Bool(false);
+		Writer.Name("True"_ASV).Bool(true);
+
+		Writer.Name("BinaryHash"_ASV).BinaryHash(FBlake3Hash());
+		Writer.Name("FieldHash"_ASV).FieldHash(FBlake3Hash());
+
+		Writer.Name("Uuid"_ASV).Uuid(FGuid());
+
+		Writer.Name("DateTimeZero"_ASV).DateTimeTicks(0);
+		Writer.Name("DateTime2020"_ASV).DateTime(FDateTime(2020, 5, 13, 15, 10));
+
+		Writer.Name("TimeSpanZero"_ASV).TimeSpanTicks(0);
+		Writer.Name("TimeSpan"_ASV).TimeSpan(FTimespan(1, 2, 4, 8));
+
+		Writer.Name("NestedObjects"_ASV);
+		Writer.BeginObject();
+		{
+			Writer.Name("Empty"_ASV);
+			Writer.BeginObject();
+			Writer.EndObject();
+
+			Writer.Name("Null"_ASV);
+			Writer.BeginObject();
+			Writer.Name("Null"_ASV).Null();
+			Writer.EndObject();
+		}
+		Writer.EndObject();
+
+		Writer.Name("NestedArrays"_ASV);
+		Writer.BeginArray();
+		{
+			Writer.BeginArray();
+			Writer.EndArray();
+
+			Writer.BeginArray();
+			Writer.Null();
+			Writer.Null();
+			Writer.Null();
+			Writer.EndArray();
+
+			Writer.BeginArray();
+			Writer.Bool(false);
+			Writer.Bool(false);
+			Writer.Bool(false);
+			Writer.EndArray();
+
+			Writer.BeginArray();
+			Writer.Bool(true);
+			Writer.Bool(true);
+			Writer.Bool(true);
+			Writer.EndArray();
+		}
+		Writer.EndArray();
+
+		Writer.Name("ArrayOfObjects"_ASV);
+		Writer.BeginArray();
+		{
+			Writer.BeginObject();
+			Writer.EndObject();
+
+			Writer.BeginObject();
+			Writer.Name("Null"_ASV).Null();
+			Writer.EndObject();
+		}
+		Writer.EndArray();
+
+		Writer.Name("LargeArray"_ASV);
+		Writer.BeginArray();
+		for (int Index = 0; Index < 256; ++Index)
+		{
+			Writer.Integer(Index - 128);
+		}
+		Writer.EndArray();
+
+		Writer.Name("LargeUniformArray"_ASV);
+		Writer.BeginArray();
+		for (int Index = 0; Index < 256; ++Index)
+		{
+			Writer.Integer(Index);
+		}
+		Writer.EndArray();
+
+		Writer.Name("NestedUniformArray"_ASV);
+		Writer.BeginArray();
+		for (int Index = 0; Index < 16; ++Index)
+		{
+			Writer.BeginArray();
+			for (int Value = 0; Value < 4; ++Value)
+			{
+				Writer.Integer(Value);
+			}
+			Writer.EndArray();
+		}
+		Writer.EndArray();
+
+		Writer.EndObject();
+		Object = Writer.Save().AsObjectRef();
+	}
+
+	TestEqual(TEXT("FCbWriter(Complex) Validate"), ValidateCompactBinary(*Object.Buffer(), ECbValidateMode::All), ECbValidateError::None);
+
+	return true;
+}
+
+///////////////////////////////////////////////////////////////////////////////////////////////////////////////////////
+
+IMPLEMENT_SIMPLE_AUTOMATION_TEST(FCbWriterStreamTest, "System.Core.Serialization.CbWriter.Stream", CompactBinaryWriterTestFlags)
+bool FCbWriterStreamTest::RunTest(const FString& Parameters)
+{
+	FCbObjectRef Object;
+	{
+		FCbWriter Writer;
+		Writer.BeginObject();
+
+		const uint8 LocalField[] = { uint8(ECbFieldType::IntegerPositive | ECbFieldType::HasFieldName), 1, 'I', 42 };
+		Writer << "FieldCopy"_ASV << FCbField(LocalField);
+
+		const uint8 LocalObject[] = { uint8(ECbFieldType::Object | ECbFieldType::HasFieldName), 1, 'O', 7,
+			uint8(ECbFieldType::IntegerPositive | ECbFieldType::HasFieldName), 1, 'I', 42,
+			uint8(ECbFieldType::Null | ECbFieldType::HasFieldName), 1, 'N' };
+		Writer << "ObjectCopy"_ASV << FCbObject(LocalObject);
+
+		const uint8 LocalArray[] = { uint8(ECbFieldType::UniformArray | ECbFieldType::HasFieldName), 1, 'A', 4, 2,
+			uint8(ECbFieldType::IntegerPositive), 42, 21 };
+		Writer << "ArrayCopy"_ASV << FCbArray(LocalArray);
+
+		Writer << "Null"_ASV << nullptr;
+
+		Writer << "Strings"_ASV;
+		Writer.BeginObject();
+		Writer
+			<< "AnsiString"_ASV << "AnsiValue"_ASV
+			<< "WideString"_ASV << TEXT("WideValue"_SV);
+		Writer.EndObject();
+
+		Writer << "Integers"_ASV;
+		Writer.BeginArray();
+		Writer << int32(-1) << int64(-1) << uint32(1) << uint64(1);
+		Writer.EndArray();
+
+		Writer << "Float32"_ASV << 1.0f;
+		Writer << "Float64"_ASV << 2.0;
+
+		Writer << "False"_ASV << false << "True"_ASV << true;
+
+		Writer << "Uuid"_ASV << FGuid();
+
+		Writer << "DateTime"_ASV << FDateTime(2020, 5, 13, 15, 10);
+		Writer << "TimeSpan"_ASV << FTimespan(1, 2, 4, 8);
+
+		Writer.EndObject();
+		Object = Writer.Save().AsObjectRef();
+	}
+
+	TestEqual(TEXT("FCbWriter(Stream) Validate"), ValidateCompactBinary(*Object.Buffer(), ECbValidateMode::All), ECbValidateError::None);
+
+	return true;
+}
+
+///////////////////////////////////////////////////////////////////////////////////////////////////////////////////////
+
+IMPLEMENT_SIMPLE_AUTOMATION_TEST(FCbWriterStateTest, "System.Core.Serialization.CbWriter.State", CompactBinaryWriterTestFlags)
+bool FCbWriterStateTest::RunTest(const FString& Parameters)
+{
+	FCbWriter Writer;
+
+	// Assert on saving an empty writer.
+	//uint8 EmptyField[1];
+	//Writer.Reset();
+	//Writer.SaveToMemory(MakeMemoryView(EmptyField));
+
+	// Assert on under-sized save buffer.
+	//uint8 ZeroFieldSmall[1];
+	//Writer.Reset();
+	//Writer.Integer(0);
+	//Writer.SaveToMemory(MakeMemoryView(ZeroFieldSmall));
+
+	// Assert on over-sized save buffer.
+	//uint8 ZeroFieldLarge[3];
+	//Writer.Reset();
+	//Writer.Integer(0);
+	//Writer.SaveToMemory(MakeMemoryView(ZeroFieldLarge));
+
+	// Assert on empty name.
+	//Writer.Name(""_ASV);
+
+	// Assert on name after name.
+	//Writer.Name("Field"_ASV).Name("Field"_ASV);
+
+	// Assert on missing name.
+	//Writer.BeginObject();
+	//Writer.Null();
+	//Writer.EndObject();
+
+	// Assert on name in array.
+	//Writer.BeginArray();
+	//Writer.Name("Field"_ASV);
+	//Writer.EndArray();
+
+	// Assert on save in object.
+	//uint8 InvalidObject[1];
+	//Writer.Reset();
+	//Writer.BeginObject();
+	//Writer.SaveToMemory(MakeMemoryView(InvalidObject));
+	//Writer.EndObject();
+
+	// Assert on save in array.
+	//uint8 InvalidArray[1];
+	//Writer.Reset();
+	//Writer.BeginArray();
+	//Writer.SaveToMemory(MakeMemoryView(InvalidArray));
+	//Writer.EndArray();
+
+	// Assert on object end with no begin.
+	//Writer.EndObject();
+
+	// Assert on array end with no begin.
+	//Writer.EndArray();
+
+	// Assert on object end after name with no value.
+	//Writer.BeginObject();
+	//Writer.Name("Field"_ASV);
+	//Writer.EndObject();
+
+	// Assert on writing a field with no value.
+	//Writer.Field(FCbField());
+
+	return true;
+}
+
+///////////////////////////////////////////////////////////////////////////////////////////////////////////////////////
+
+#endif // WITH_DEV_AUTOMATION_TESTS
diff --git a/Engine/Source/Runtime/Core/Private/Tests/Serialization/VarIntTest.cpp b/Engine/Source/Runtime/Core/Private/Tests/Serialization/VarIntTest.cpp
new file mode 100644
--- /dev/null
+++ b/Engine/Source/Runtime/Core/Private/Tests/Serialization/VarIntTest.cpp
@@ -0,0 +1,286 @@
+// Copyright Epic Games, Inc. All Rights Reserved.
+
+#include "Serialization/VarInt.h"
+
+#include "Misc/AutomationTest.h"
+
+#if WITH_DEV_AUTOMATION_TESTS
+
+IMPLEMENT_SIMPLE_AUTOMATION_TEST(FVarIntMeasureTest, "System.Core.Serialization.VarInt.Measure", EAutomationTestFlags::ApplicationContextMask | EAutomationTestFlags::SmokeFilter)
+bool FVarIntMeasureTest::RunTest(const FString& Parameters)
+{
+	// Test MeasureVarInt at signed 32-bit encoding boundaries.
+	TestEqual(TEXT("MeasureVarInt(0x0000'0000)"), MeasureVarInt(int32(0x0000'0000)), 1);
+	TestEqual(TEXT("MeasureVarInt(0x0000'0001)"), MeasureVarInt(int32(0x0000'0001)), 1);
+	TestEqual(TEXT("MeasureVarInt(0x0000'003f)"), MeasureVarInt(int32(0x0000'003f)), 1);
+	TestEqual(TEXT("MeasureVarInt(0x0000'0040)"), MeasureVarInt(int32(0x0000'0040)), 2);
+	TestEqual(TEXT("MeasureVarInt(0x0000'1fff)"), MeasureVarInt(int32(0x0000'1fff)), 2);
+	TestEqual(TEXT("MeasureVarInt(0x0000'2000)"), MeasureVarInt(int32(0x0000'2000)), 3);
+	TestEqual(TEXT("MeasureVarInt(0x000f'ffff)"), MeasureVarInt(int32(0x000f'ffff)), 3);
+	TestEqual(TEXT("MeasureVarInt(0x0010'0000)"), MeasureVarInt(int32(0x0010'0000)), 4);
+	TestEqual(TEXT("MeasureVarInt(0x07ff'ffff)"), MeasureVarInt(int32(0x07ff'ffff)), 4);
+	TestEqual(TEXT("MeasureVarInt(0x0800'0000)"), MeasureVarInt(int32(0x0800'0000)), 5);
+	TestEqual(TEXT("MeasureVarInt(0x7fff'ffff)"), MeasureVarInt(int32(0x7fff'ffff)), 5);
+	TestEqual(TEXT("MeasureVarInt(0xffff'ffff)"), MeasureVarInt(int32(0xffff'ffff)), 1); // -0x0000'0001
+	TestEqual(TEXT("MeasureVarInt(0xffff'ffc0)"), MeasureVarInt(int32(0xffff'ffc0)), 1); // -0x0000'0040
+	TestEqual(TEXT("MeasureVarInt(0xffff'ffbf)"), MeasureVarInt(int32(0xffff'ffbf)), 2); // -0x0000'0041
+	TestEqual(TEXT("MeasureVarInt(0xffff'e000)"), MeasureVarInt(int32(0xffff'e000)), 2); // -0x0000'2000
+	TestEqual(TEXT("MeasureVarInt(0xffff'dfff)"), MeasureVarInt(int32(0xffff'dfff)), 3); // -0x0000'2001
+	TestEqual(TEXT("MeasureVarInt(0xfff0'0000)"), MeasureVarInt(int32(0xfff0'0000)), 3); // -0x0010'0000
+	TestEqual(TEXT("MeasureVarInt(0xffef'ffff)"), MeasureVarInt(int32(0xffef'ffff)), 4); // -0x0010'0001
+	TestEqual(TEXT("MeasureVarInt(0xf800'0000)"), MeasureVarInt(int32(0xf800'0000)), 4); // -0x0800'0000
+	TestEqual(TEXT("MeasureVarInt(0xf7ff'ffff)"), MeasureVarInt(int32(0xf7ff'ffff)), 5); // -0x0800'0001
+	TestEqual(TEXT("MeasureVarInt(0x8000'0000)"), MeasureVarInt(int32(0x8000'0000)), 5); // -0x8000'0000
+
+	// Test MeasureVarUInt at unsigned 32-bit encoding boundaries.
+	TestEqual(TEXT("MeasureVarUInt(0x0000'0000)"), MeasureVarUInt(uint32(0x0000'0000)), 1);
+	TestEqual(TEXT("MeasureVarUInt(0x0000'0001)"), MeasureVarUInt(uint32(0x0000'0001)), 1);
+	TestEqual(TEXT("MeasureVarUInt(0x0000'007f)"), MeasureVarUInt(uint32(0x0000'007f)), 1);
+	TestEqual(TEXT("MeasureVarUInt(0x0000'0080)"), MeasureVarUInt(uint32(0x0000'0080)), 2);
+	TestEqual(TEXT("MeasureVarUInt(0x0000'3fff)"), MeasureVarUInt(uint32(0x0000'3fff)), 2);
+	TestEqual(TEXT("MeasureVarUInt(0x0000'4000)"), MeasureVarUInt(uint32(0x0000'4000)), 3);
+	TestEqual(TEXT("MeasureVarUInt(0x001f'ffff)"), MeasureVarUInt(uint32(0x001f'ffff)), 3);
+	TestEqual(TEXT("MeasureVarUInt(0x0020'0000)"), MeasureVarUInt(uint32(0x0020'0000)), 4);
+	TestEqual(TEXT("MeasureVarUInt(0x0fff'ffff)"), MeasureVarUInt(uint32(0x0fff'ffff)), 4);
+	TestEqual(TEXT("MeasureVarUInt(0x1000'0000)"), MeasureVarUInt(uint32(0x1000'0000)), 5);
+	TestEqual(TEXT("MeasureVarUInt(0xffff'ffff)"), MeasureVarUInt(uint32(0xffff'ffff)), 5);
+
+	// Test MeasureVarInt at signed 64-bit encoding boundaries.
+	TestEqual(TEXT("MeasureVarInt(0x0000'0000'0000'0000)"), MeasureVarInt(int64(0x0000'0000'0000'0000)), 1);
+	TestEqual(TEXT("MeasureVarInt(0x0000'0000'0000'0001)"), MeasureVarInt(int64(0x0000'0000'0000'0001)), 1);
+	TestEqual(TEXT("MeasureVarInt(0x0000'0000'0000'003f)"), MeasureVarInt(int64(0x0000'0000'0000'003f)), 1);
+	TestEqual(TEXT("MeasureVarInt(0x0000'0000'0000'0040)"), MeasureVarInt(int64(0x0000'0000'0000'0040)), 2);
+	TestEqual(TEXT("MeasureVarInt(0x0000'0000'0000'1fff)"), MeasureVarInt(int64(0x0000'0000'0000'1fff)), 2);
+	TestEqual(TEXT("MeasureVarInt(0x0000'0000'0000'2000)"), MeasureVarInt(int64(0x0000'0000'0000'2000)), 3);
+	TestEqual(TEXT("MeasureVarInt(0x0000'0000'000f'ffff)"), MeasureVarInt(int64(0x0000'0000'000f'ffff)), 3);
+	TestEqual(TEXT("MeasureVarInt(0x0000'0000'0010'0000)"), MeasureVarInt(int64(0x0000'0000'0010'0000)), 4);
+	TestEqual(TEXT("MeasureVarInt(0x0000'0000'07ff'ffff)"), MeasureVarInt(int64(0x0000'0000'07ff'ffff)), 4);
+	TestEqual(TEXT("MeasureVarInt(0x0000'0000'0800'0000)"), MeasureVarInt(int64(0x0000'0000'0800'0000)), 5);
+	TestEqual(TEXT("MeasureVarInt(0x0000'0003'ffff'ffff)"), MeasureVarInt(int64(0x0000'0003'ffff'ffff)), 5);
+	TestEqual(TEXT("MeasureVarInt(0x0000'0004'0000'0000)"), MeasureVarInt(int64(0x0000'0004'0000'0000)), 6);
+	TestEqual(TEXT("MeasureVarInt(0x0000'01ff'ffff'ffff)"), MeasureVarInt(int64(0x0000'01ff'ffff'ffff)), 6);
+	TestEqual(TEXT("MeasureVarInt(0x0000'0200'0000'0000)"), MeasureVarInt(int64(0x0000'0200'0000'0000)), 7);
+	TestEqual(TEXT("MeasureVarInt(0x0000'ffff'ffff'ffff)"), MeasureVarInt(int64(0x0000'ffff'ffff'ffff)), 7);
+	TestEqual(TEXT("MeasureVarInt(0x0001'0000'0000'0000)"), MeasureVarInt(int64(0x0001'0000'0000'0000)), 8);
+	TestEqual(TEXT("MeasureVarInt(0x007f'ffff'ffff'ffff)"), MeasureVarInt(int64(0x007f'ffff'ffff'ffff)), 8);
+	TestEqual(TEXT("MeasureVarInt(0x0080'0000'0000'0000)"), MeasureVarInt(int64(0x0080'0000'0000'0000)), 9);
+	TestEqual(TEXT("MeasureVarInt(0x7fff'ffff'ffff'ffff)"), MeasureVarInt(int64(0x7fff'ffff'ffff'ffff)), 9);
+	TestEqual(TEXT("MeasureVarInt(0xffff'ffff'ffff'ffff)"), MeasureVarInt(int64(0xffff'ffff'ffff'ffff)), 1); // -0x0000'0000'0000'0001
+	TestEqual(TEXT("MeasureVarInt(0xffff'ffff'ffff'ffc0)"), MeasureVarInt(int64(0xffff'ffff'ffff'ffc0)), 1); // -0x0000'0000'0000'0040
+	TestEqual(TEXT("MeasureVarInt(0xffff'ffff'ffff'ffbf)"), MeasureVarInt(int64(0xffff'ffff'ffff'ffbf)), 2); // -0x0000'0000'0000'0041
+	TestEqual(TEXT("MeasureVarInt(0xffff'ffff'ffff'e000)"), MeasureVarInt(int64(0xffff'ffff'ffff'e000)), 2); // -0x0000'0000'0000'2000
+	TestEqual(TEXT("MeasureVarInt(0xffff'ffff'ffff'dfff)"), MeasureVarInt(int64(0xffff'ffff'ffff'dfff)), 3); // -0x0000'0000'0000'2001
+	TestEqual(TEXT("MeasureVarInt(0xffff'ffff'fff0'0000)"), MeasureVarInt(int64(0xffff'ffff'fff0'0000)), 3); // -0x0000'0000'0010'0000
+	TestEqual(TEXT("MeasureVarInt(0xffff'ffff'ffef'ffff)"), MeasureVarInt(int64(0xffff'ffff'ffef'ffff)), 4); // -0x0000'0000'0010'0001
+	TestEqual(TEXT("MeasureVarInt(0xffff'ffff'f800'0000)"), MeasureVarInt(int64(0xffff'ffff'f800'0000)), 4); // -0x0000'0000'0800'0000
+	TestEqual(TEXT("MeasureVarInt(0xffff'ffff'f7ff'ffff)"), MeasureVarInt(int64(0xffff'ffff'f7ff'ffff)), 5); // -0x0000'0000'0800'0001
+	TestEqual(TEXT("MeasureVarInt(0xffff'fffc'0000'0000)"), MeasureVarInt(int64(0xffff'fffc'0000'0000)), 5); // -0x0000'0004'0000'0000
+	TestEqual(TEXT("MeasureVarInt(0xffff'fffb'ffff'ffff)"), MeasureVarInt(int64(0xffff'fffb'ffff'ffff)), 6); // -0x0000'0004'0000'0001
+	TestEqual(TEXT("MeasureVarInt(0xffff'fe00'0000'0000)"), MeasureVarInt(int64(0xffff'fe00'0000'0000)), 6); // -0x0000'0200'0000'0000
+	TestEqual(TEXT("MeasureVarInt(0xffff'fdff'ffff'ffff)"), MeasureVarInt(int64(0xffff'fdff'ffff'ffff)), 7); // -0x0000'0200'0000'0001
+	TestEqual(TEXT("MeasureVarInt(0xffff'0000'0000'0000)"), MeasureVarInt(int64(0xffff'0000'0000'0000)), 7); // -0x0001'0000'0000'0000
+	TestEqual(TEXT("MeasureVarInt(0xfffe'ffff'ffff'ffff)"), MeasureVarInt(int64(0xfffe'ffff'ffff'ffff)), 8); // -0x0001'0000'0000'0001
+	TestEqual(TEXT("MeasureVarInt(0xff80'0000'0000'0000)"), MeasureVarInt(int64(0xff80'0000'0000'0000)), 8); // -0x0080'0000'0000'0000
+	TestEqual(TEXT("MeasureVarInt(0xff7f'ffff'ffff'ffff)"), MeasureVarInt(int64(0xff7f'ffff'ffff'ffff)), 9); // -0x0080'0000'0000'0001
+	TestEqual(TEXT("MeasureVarInt(0x8000'0000'0000'0000)"), MeasureVarInt(int64(0x8000'0000'0000'0000)), 9); // -0x8000'0000'0000'0000
+
+	// Test MeasureVarUInt at unsigned 64-bit encoding boundaries.
+	TestEqual(TEXT("MeasureVarUInt(0x0000'0000'0000'0000)"), MeasureVarUInt(uint64(0x0000'0000'0000'0000)), 1);
+	TestEqual(TEXT("MeasureVarUInt(0x0000'0000'0000'0001)"), MeasureVarUInt(uint64(0x0000'0000'0000'0001)), 1);
+	TestEqual(TEXT("MeasureVarUInt(0x0000'0000'0000'007f)"), MeasureVarUInt(uint64(0x0000'0000'0000'007f)), 1);
+	TestEqual(TEXT("MeasureVarUInt(0x0000'0000'0000'0080)"), MeasureVarUInt(uint64(0x0000'0000'0000'0080)), 2);
+	TestEqual(TEXT("MeasureVarUInt(0x0000'0000'0000'3fff)"), MeasureVarUInt(uint64(0x0000'0000'0000'3fff)), 2);
+	TestEqual(TEXT("MeasureVarUInt(0x0000'0000'0000'4000)"), MeasureVarUInt(uint64(0x0000'0000'0000'4000)), 3);
+	TestEqual(TEXT("MeasureVarUInt(0x0000'0000'001f'ffff)"), MeasureVarUInt(uint64(0x0000'0000'001f'ffff)), 3);
+	TestEqual(TEXT("MeasureVarUInt(0x0000'0000'0020'0000)"), MeasureVarUInt(uint64(0x0000'0000'0020'0000)), 4);
+	TestEqual(TEXT("MeasureVarUInt(0x0000'0000'0fff'ffff)"), MeasureVarUInt(uint64(0x0000'0000'0fff'ffff)), 4);
+	TestEqual(TEXT("MeasureVarUInt(0x0000'0000'1000'0000)"), MeasureVarUInt(uint64(0x0000'0000'1000'0000)), 5);
+	TestEqual(TEXT("MeasureVarUInt(0x0000'0007'ffff'ffff)"), MeasureVarUInt(uint64(0x0000'0007'ffff'ffff)), 5);
+	TestEqual(TEXT("MeasureVarUInt(0x0000'0008'0000'0000)"), MeasureVarUInt(uint64(0x0000'0008'0000'0000)), 6);
+	TestEqual(TEXT("MeasureVarUInt(0x0000'03ff'ffff'ffff)"), MeasureVarUInt(uint64(0x0000'03ff'ffff'ffff)), 6);
+	TestEqual(TEXT("MeasureVarUInt(0x0000'0400'0000'0000)"), MeasureVarUInt(uint64(0x0000'0400'0000'0000)), 7);
+	TestEqual(TEXT("MeasureVarUInt(0x0001'ffff'ffff'ffff)"), MeasureVarUInt(uint64(0x0001'ffff'ffff'ffff)), 7);
+	TestEqual(TEXT("MeasureVarUInt(0x0002'0000'0000'0000)"), MeasureVarUInt(uint64(0x0002'0000'0000'0000)), 8);
+	TestEqual(TEXT("MeasureVarUInt(0x00ff'ffff'ffff'ffff)"), MeasureVarUInt(uint64(0x00ff'ffff'ffff'ffff)), 8);
+	TestEqual(TEXT("MeasureVarUInt(0x0100'0000'0000'0000)"), MeasureVarUInt(uint64(0x0100'0000'0000'0000)), 9);
+	TestEqual(TEXT("MeasureVarUInt(0xffff'ffff'ffff'ffff)"), MeasureVarUInt(uint64(0xffff'ffff'ffff'ffff)), 9);
+
+	// Test MeasureVarInt at encoding boundaries.
+	auto TestMeasureVarInt = [](uint8 FirstByte) { return MeasureVarInt(&FirstByte); };
+	TestEqual(TEXT("MeasureVarInt(1-byte array)"), TestMeasureVarInt(0b0000'0000), 1);
+	TestEqual(TEXT("MeasureVarInt(1-byte array)"), TestMeasureVarInt(0b0111'1111), 1);
+	TestEqual(TEXT("MeasureVarInt(2-byte array)"), TestMeasureVarInt(0b1000'0000), 2);
+	TestEqual(TEXT("MeasureVarInt(2-byte array)"), TestMeasureVarInt(0b1011'1111), 2);
+	TestEqual(TEXT("MeasureVarInt(3-byte array)"), TestMeasureVarInt(0b1100'0000), 3);
+	TestEqual(TEXT("MeasureVarInt(3-byte array)"), TestMeasureVarInt(0b1101'1111), 3);
+	TestEqual(TEXT("MeasureVarInt(4-byte array)"), TestMeasureVarInt(0b1110'0000), 4);
+	TestEqual(TEXT("MeasureVarInt(4-byte array)"), TestMeasureVarInt(0b1110'1111), 4);
+	TestEqual(TEXT("MeasureVarInt(5-byte array)"), TestMeasureVarInt(0b1111'0000), 5);
+	TestEqual(TEXT("MeasureVarInt(5-byte array)"), TestMeasureVarInt(0b1111'0111), 5);
+	TestEqual(TEXT("MeasureVarInt(6-byte array)"), TestMeasureVarInt(0b1111'1000), 6);
+	TestEqual(TEXT("MeasureVarInt(6-byte array)"), TestMeasureVarInt(0b1111'1011), 6);
+	TestEqual(TEXT("MeasureVarInt(7-byte array)"), TestMeasureVarInt(0b1111'1100), 7);
+	TestEqual(TEXT("MeasureVarInt(7-byte array)"), TestMeasureVarInt(0b1111'1101), 7);
+	TestEqual(TEXT("MeasureVarInt(8-byte array)"), TestMeasureVarInt(0b1111'1110), 8);
+	TestEqual(TEXT("MeasureVarInt(9-byte array)"), TestMeasureVarInt(0b1111'1111), 9);
+
+	// Test MeasureVarUInt at encoding boundaries.
+	auto TestMeasureVarUInt = [](uint8 FirstByte) { return MeasureVarUInt(&FirstByte); };
+	TestEqual(TEXT("MeasureVarUInt(1-byte array)"), TestMeasureVarUInt(0b0000'0000), 1);
+	TestEqual(TEXT("MeasureVarUInt(1-byte array)"), TestMeasureVarUInt(0b0111'1111), 1);
+	TestEqual(TEXT("MeasureVarUInt(2-byte array)"), TestMeasureVarUInt(0b1000'0000), 2);
+	TestEqual(TEXT("MeasureVarUInt(2-byte array)"), TestMeasureVarUInt(0b1011'1111), 2);
+	TestEqual(TEXT("MeasureVarUInt(3-byte array)"), TestMeasureVarUInt(0b1100'0000), 3);
+	TestEqual(TEXT("MeasureVarUInt(3-byte array)"), TestMeasureVarUInt(0b1101'1111), 3);
+	TestEqual(TEXT("MeasureVarUInt(4-byte array)"), TestMeasureVarUInt(0b1110'0000), 4);
+	TestEqual(TEXT("MeasureVarUInt(4-byte array)"), TestMeasureVarUInt(0b1110'1111), 4);
+	TestEqual(TEXT("MeasureVarUInt(5-byte array)"), TestMeasureVarUInt(0b1111'0000), 5);
+	TestEqual(TEXT("MeasureVarUInt(5-byte array)"), TestMeasureVarUInt(0b1111'0111), 5);
+	TestEqual(TEXT("MeasureVarUInt(6-byte array)"), TestMeasureVarUInt(0b1111'1000), 6);
+	TestEqual(TEXT("MeasureVarUInt(6-byte array)"), TestMeasureVarUInt(0b1111'1011), 6);
+	TestEqual(TEXT("MeasureVarUInt(7-byte array)"), TestMeasureVarUInt(0b1111'1100), 7);
+	TestEqual(TEXT("MeasureVarUInt(7-byte array)"), TestMeasureVarUInt(0b1111'1101), 7);
+	TestEqual(TEXT("MeasureVarUInt(8-byte array)"), TestMeasureVarUInt(0b1111'1110), 8);
+	TestEqual(TEXT("MeasureVarUInt(9-byte array)"), TestMeasureVarUInt(0b1111'1111), 9);
+
+	return true;
+}
+
+IMPLEMENT_SIMPLE_AUTOMATION_TEST(FVarIntSerializeTest, "System.Core.Serialization.VarInt.Serialize", EAutomationTestFlags::ApplicationContextMask | EAutomationTestFlags::SmokeFilter)
+bool FVarIntSerializeTest::RunTest(const FString& Parameters)
+{
+	// Test Read/WriteVarInt at signed 32-bit encoding boundaries.
+	auto TestSerializeVarInt32 = [](int32 Value) -> bool
+	{
+		uint8 Buffer[5];
+		const uint32 WriteByteCount = WriteVarInt(Value, Buffer);
+		uint32 ReadByteCount = 0;
+		return WriteByteCount <= 5 && ReadVarInt(Buffer, ReadByteCount) == Value && ReadByteCount == WriteByteCount;
+	};
+	TestTrue(TEXT("Read/WriteVarInt(0x0000'0000)"), TestSerializeVarInt32(int32(0x0000'0000)));
+	TestTrue(TEXT("Read/WriteVarInt(0x0000'0001)"), TestSerializeVarInt32(int32(0x0000'0001)));
+	TestTrue(TEXT("Read/WriteVarInt(0x0000'003f)"), TestSerializeVarInt32(int32(0x0000'003f)));
+	TestTrue(TEXT("Read/WriteVarInt(0x0000'0040)"), TestSerializeVarInt32(int32(0x0000'0040)));
+	TestTrue(TEXT("Read/WriteVarInt(0x0000'1fff)"), TestSerializeVarInt32(int32(0x0000'1fff)));
+	TestTrue(TEXT("Read/WriteVarInt(0x0000'2000)"), TestSerializeVarInt32(int32(0x0000'2000)));
+	TestTrue(TEXT("Read/WriteVarInt(0x000f'ffff)"), TestSerializeVarInt32(int32(0x000f'ffff)));
+	TestTrue(TEXT("Read/WriteVarInt(0x0010'0000)"), TestSerializeVarInt32(int32(0x0010'0000)));
+	TestTrue(TEXT("Read/WriteVarInt(0x07ff'ffff)"), TestSerializeVarInt32(int32(0x07ff'ffff)));
+	TestTrue(TEXT("Read/WriteVarInt(0x0800'0000)"), TestSerializeVarInt32(int32(0x0800'0000)));
+	TestTrue(TEXT("Read/WriteVarInt(0x7fff'ffff)"), TestSerializeVarInt32(int32(0x7fff'ffff)));
+	TestTrue(TEXT("Read/WriteVarInt(0xffff'ffff)"), TestSerializeVarInt32(int32(0xffff'ffff))); // -0x0000'0001
+	TestTrue(TEXT("Read/WriteVarInt(0xffff'ffc0)"), TestSerializeVarInt32(int32(0xffff'ffc0))); // -0x0000'0040
+	TestTrue(TEXT("Read/WriteVarInt(0xffff'ffbf)"), TestSerializeVarInt32(int32(0xffff'ffbf))); // -0x0000'0041
+	TestTrue(TEXT("Read/WriteVarInt(0xffff'e000)"), TestSerializeVarInt32(int32(0xffff'e000))); // -0x0000'2000
+	TestTrue(TEXT("Read/WriteVarInt(0xffff'dfff)"), TestSerializeVarInt32(int32(0xffff'dfff))); // -0x0000'2001
+	TestTrue(TEXT("Read/WriteVarInt(0xfff0'0000)"), TestSerializeVarInt32(int32(0xfff0'0000))); // -0x0010'0000
+	TestTrue(TEXT("Read/WriteVarInt(0xffef'ffff)"), TestSerializeVarInt32(int32(0xffef'ffff))); // -0x0010'0001
+	TestTrue(TEXT("Read/WriteVarInt(0xf800'0000)"), TestSerializeVarInt32(int32(0xf800'0000))); // -0x0800'0000
+	TestTrue(TEXT("Read/WriteVarInt(0xf7ff'ffff)"), TestSerializeVarInt32(int32(0xf7ff'ffff))); // -0x0800'0001
+	TestTrue(TEXT("Read/WriteVarInt(0x8000'0000)"), TestSerializeVarInt32(int32(0x8000'0000))); // -0x8000'0000
+
+	// Test Read/WriteVarUInt at unsigned 32-bit encoding boundaries.
+	auto TestSerializeVarUInt32 = [](uint32 Value) -> bool
+	{
+		uint8 Buffer[5];
+		const uint32 WriteByteCount = WriteVarUInt(Value, Buffer);
+		uint32 ReadByteCount = 0;
+		return WriteByteCount <= 5 && ReadVarUInt(Buffer, ReadByteCount) == Value && ReadByteCount == WriteByteCount;
+	};
+	TestTrue(TEXT("Read/WriteVarUInt(0x0000'0000)"), TestSerializeVarUInt32(uint32(0x0000'0000)));
+	TestTrue(TEXT("Read/WriteVarUInt(0x0000'007f)"), TestSerializeVarUInt32(uint32(0x0000'007f)));
+	TestTrue(TEXT("Read/WriteVarUInt(0x0000'0080)"), TestSerializeVarUInt32(uint32(0x0000'0080)));
+	TestTrue(TEXT("Read/WriteVarUInt(0x0000'3fff)"), TestSerializeVarUInt32(uint32(0x0000'3fff)));
+	TestTrue(TEXT("Read/WriteVarUInt(0x0000'4000)"), TestSerializeVarUInt32(uint32(0x0000'4000)));
+	TestTrue(TEXT("Read/WriteVarUInt(0x0000'7fff)"), TestSerializeVarUInt32(uint32(0x0000'7fff)));
+	TestTrue(TEXT("Read/WriteVarUInt(0x0000'8000)"), TestSerializeVarUInt32(uint32(0x0000'8000)));
+	TestTrue(TEXT("Read/WriteVarUInt(0x0000'ffff)"), TestSerializeVarUInt32(uint32(0x0000'ffff)));
+	TestTrue(TEXT("Read/WriteVarUInt(0x001f'ffff)"), TestSerializeVarUInt32(uint32(0x001f'ffff)));
+	TestTrue(TEXT("Read/WriteVarUInt(0x0020'0000)"), TestSerializeVarUInt32(uint32(0x0020'0000)));
+	TestTrue(TEXT("Read/WriteVarUInt(0x0fff'ffff)"), TestSerializeVarUInt32(uint32(0x0fff'ffff)));
+	TestTrue(TEXT("Read/WriteVarUInt(0x1000'0000)"), TestSerializeVarUInt32(uint32(0x1000'0000)));
+	TestTrue(TEXT("Read/WriteVarUInt(0xffff'ffff)"), TestSerializeVarUInt32(uint32(0xffff'ffff)));
+
+	// Test Read/WriteVarInt at signed 64-bit encoding boundaries.
+	auto TestSerializeVarInt64 = [](int64 Value) -> bool
+	{
+		uint8 Buffer[9];
+		const uint32 WriteByteCount = WriteVarInt(Value, Buffer);
+		uint32 ReadByteCount = 0;
+		return WriteByteCount <= 9 && ReadVarInt(Buffer, ReadByteCount) == Value && ReadByteCount == WriteByteCount;
+	};
+	TestTrue(TEXT("Read/WriteVarInt(0x0000'0000'0000'0000)"), TestSerializeVarInt64(int64(0x0000'0000'0000'0000)));
+	TestTrue(TEXT("Read/WriteVarInt(0x0000'0000'0000'0001)"), TestSerializeVarInt64(int64(0x0000'0000'0000'0001)));
+	TestTrue(TEXT("Read/WriteVarInt(0x0000'0000'0000'003f)"), TestSerializeVarInt64(int64(0x0000'0000'0000'003f)));
+	TestTrue(TEXT("Read/WriteVarInt(0x0000'0000'0000'0040)"), TestSerializeVarInt64(int64(0x0000'0000'0000'0040)));
+	TestTrue(TEXT("Read/WriteVarInt(0x0000'0000'0000'1fff)"), TestSerializeVarInt64(int64(0x0000'0000'0000'1fff)));
+	TestTrue(TEXT("Read/WriteVarInt(0x0000'0000'0000'2000)"), TestSerializeVarInt64(int64(0x0000'0000'0000'2000)));
+	TestTrue(TEXT("Read/WriteVarInt(0x0000'0000'000f'ffff)"), TestSerializeVarInt64(int64(0x0000'0000'000f'ffff)));
+	TestTrue(TEXT("Read/WriteVarInt(0x0000'0000'0010'0000)"), TestSerializeVarInt64(int64(0x0000'0000'0010'0000)));
+	TestTrue(TEXT("Read/WriteVarInt(0x0000'0000'07ff'ffff)"), TestSerializeVarInt64(int64(0x0000'0000'07ff'ffff)));
+	TestTrue(TEXT("Read/WriteVarInt(0x0000'0000'0800'0000)"), TestSerializeVarInt64(int64(0x0000'0000'0800'0000)));
+	TestTrue(TEXT("Read/WriteVarInt(0x0000'0003'ffff'ffff)"), TestSerializeVarInt64(int64(0x0000'0003'ffff'ffff)));
+	TestTrue(TEXT("Read/WriteVarInt(0x0000'0004'0000'0000)"), TestSerializeVarInt64(int64(0x0000'0004'0000'0000)));
+	TestTrue(TEXT("Read/WriteVarInt(0x0000'01ff'ffff'ffff)"), TestSerializeVarInt64(int64(0x0000'01ff'ffff'ffff)));
+	TestTrue(TEXT("Read/WriteVarInt(0x0000'0200'0000'0000)"), TestSerializeVarInt64(int64(0x0000'0200'0000'0000)));
+	TestTrue(TEXT("Read/WriteVarInt(0x0000'ffff'ffff'ffff)"), TestSerializeVarInt64(int64(0x0000'ffff'ffff'ffff)));
+	TestTrue(TEXT("Read/WriteVarInt(0x0001'0000'0000'0000)"), TestSerializeVarInt64(int64(0x0001'0000'0000'0000)));
+	TestTrue(TEXT("Read/WriteVarInt(0x007f'ffff'ffff'ffff)"), TestSerializeVarInt64(int64(0x007f'ffff'ffff'ffff)));
+	TestTrue(TEXT("Read/WriteVarInt(0x0080'0000'0000'0000)"), TestSerializeVarInt64(int64(0x0080'0000'0000'0000)));
+	TestTrue(TEXT("Read/WriteVarInt(0x7fff'ffff'ffff'ffff)"), TestSerializeVarInt64(int64(0x7fff'ffff'ffff'ffff)));
+	TestTrue(TEXT("Read/WriteVarInt(0xffff'ffff'ffff'ffff)"), TestSerializeVarInt64(int64(0xffff'ffff'ffff'ffff))); // -0x0000'0000'0000'0001
+	TestTrue(TEXT("Read/WriteVarInt(0xffff'ffff'ffff'ffc0)"), TestSerializeVarInt64(int64(0xffff'ffff'ffff'ffc0))); // -0x0000'0000'0000'0040
+	TestTrue(TEXT("Read/WriteVarInt(0xffff'ffff'ffff'ffbf)"), TestSerializeVarInt64(int64(0xffff'ffff'ffff'ffbf))); // -0x0000'0000'0000'0041
+	TestTrue(TEXT("Read/WriteVarInt(0xffff'ffff'ffff'e000)"), TestSerializeVarInt64(int64(0xffff'ffff'ffff'e000))); // -0x0000'0000'0000'2000
+	TestTrue(TEXT("Read/WriteVarInt(0xffff'ffff'ffff'dfff)"), TestSerializeVarInt64(int64(0xffff'ffff'ffff'dfff))); // -0x0000'0000'0000'2001
+	TestTrue(TEXT("Read/WriteVarInt(0xffff'ffff'fff0'0000)"), TestSerializeVarInt64(int64(0xffff'ffff'fff0'0000))); // -0x0000'0000'0010'0000
+	TestTrue(TEXT("Read/WriteVarInt(0xffff'ffff'ffef'ffff)"), TestSerializeVarInt64(int64(0xffff'ffff'ffef'ffff))); // -0x0000'0000'0010'0001
+	TestTrue(TEXT("Read/WriteVarInt(0xffff'ffff'f800'0000)"), TestSerializeVarInt64(int64(0xffff'ffff'f800'0000))); // -0x0000'0000'0800'0000
+	TestTrue(TEXT("Read/WriteVarInt(0xffff'ffff'f7ff'ffff)"), TestSerializeVarInt64(int64(0xffff'ffff'f7ff'ffff))); // -0x0000'0000'0800'0001
+	TestTrue(TEXT("Read/WriteVarInt(0xffff'fffc'0000'0000)"), TestSerializeVarInt64(int64(0xffff'fffc'0000'0000))); // -0x0000'0004'0000'0000
+	TestTrue(TEXT("Read/WriteVarInt(0xffff'fffb'ffff'ffff)"), TestSerializeVarInt64(int64(0xffff'fffb'ffff'ffff))); // -0x0000'0004'0000'0001
+	TestTrue(TEXT("Read/WriteVarInt(0xffff'fe00'0000'0000)"), TestSerializeVarInt64(int64(0xffff'fe00'0000'0000))); // -0x0000'0200'0000'0000
+	TestTrue(TEXT("Read/WriteVarInt(0xffff'fdff'ffff'ffff)"), TestSerializeVarInt64(int64(0xffff'fdff'ffff'ffff))); // -0x0000'0200'0000'0001
+	TestTrue(TEXT("Read/WriteVarInt(0xffff'0000'0000'0000)"), TestSerializeVarInt64(int64(0xffff'0000'0000'0000))); // -0x0001'0000'0000'0000
+	TestTrue(TEXT("Read/WriteVarInt(0xfffe'ffff'ffff'ffff)"), TestSerializeVarInt64(int64(0xfffe'ffff'ffff'ffff))); // -0x0001'0000'0000'0001
+	TestTrue(TEXT("Read/WriteVarInt(0xff80'0000'0000'0000)"), TestSerializeVarInt64(int64(0xff80'0000'0000'0000))); // -0x0080'0000'0000'0000
+	TestTrue(TEXT("Read/WriteVarInt(0xff7f'ffff'ffff'ffff)"), TestSerializeVarInt64(int64(0xff7f'ffff'ffff'ffff))); // -0x0080'0000'0000'0001
+	TestTrue(TEXT("Read/WriteVarInt(0x8000'0000'0000'0000)"), TestSerializeVarInt64(int64(0x8000'0000'0000'0000))); // -0x8000'0000'0000'0000
+
+	// Test Read/WriteVarUInt at unsigned 64-bit encoding boundaries.
+	auto TestSerializeVarUInt64 = [](uint64 Value) -> bool
+	{
+		uint8 Buffer[9];
+		const uint32 WriteByteCount = WriteVarUInt(Value, Buffer);
+		uint32 ReadByteCount = 0;
+		return WriteByteCount <= 9 && ReadVarUInt(Buffer, ReadByteCount) == Value && ReadByteCount == WriteByteCount;
+	};
+	TestTrue(TEXT("Read/WriteVarUInt(0x0000'0000'0000'0000)"), TestSerializeVarUInt64(uint64(0x0000'0000'0000'0000)));
+	TestTrue(TEXT("Read/WriteVarUInt(0x0000'0000'0000'007f)"), TestSerializeVarUInt64(uint64(0x0000'0000'0000'007f)));
+	TestTrue(TEXT("Read/WriteVarUInt(0x0000'0000'0000'0080)"), TestSerializeVarUInt64(uint64(0x0000'0000'0000'0080)));
+	TestTrue(TEXT("Read/WriteVarUInt(0x0000'0000'0000'3fff)"), TestSerializeVarUInt64(uint64(0x0000'0000'0000'3fff)));
+	TestTrue(TEXT("Read/WriteVarUInt(0x0000'0000'0000'4000)"), TestSerializeVarUInt64(uint64(0x0000'0000'0000'4000)));
+	TestTrue(TEXT("Read/WriteVarUInt(0x0000'0000'0000'7fff)"), TestSerializeVarUInt64(uint64(0x0000'0000'0000'7fff)));
+	TestTrue(TEXT("Read/WriteVarUInt(0x0000'0000'0000'8000)"), TestSerializeVarUInt64(uint64(0x0000'0000'0000'8000)));
+	TestTrue(TEXT("Read/WriteVarUInt(0x0000'0000'0000'ffff)"), TestSerializeVarUInt64(uint64(0x0000'0000'0000'ffff)));
+	TestTrue(TEXT("Read/WriteVarUInt(0x0000'0000'001f'ffff)"), TestSerializeVarUInt64(uint64(0x0000'0000'001f'ffff)));
+	TestTrue(TEXT("Read/WriteVarUInt(0x0000'0000'0020'0000)"), TestSerializeVarUInt64(uint64(0x0000'0000'0020'0000)));
+	TestTrue(TEXT("Read/WriteVarUInt(0x0000'0000'0fff'ffff)"), TestSerializeVarUInt64(uint64(0x0000'0000'0fff'ffff)));
+	TestTrue(TEXT("Read/WriteVarUInt(0x0000'0000'1000'0000)"), TestSerializeVarUInt64(uint64(0x0000'0000'1000'0000)));
+	TestTrue(TEXT("Read/WriteVarUInt(0x0000'0000'7fff'ffff)"), TestSerializeVarUInt64(uint64(0x0000'0000'7fff'ffff)));
+	TestTrue(TEXT("Read/WriteVarUInt(0x0000'0000'8000'0000)"), TestSerializeVarUInt64(uint64(0x0000'0000'8000'0000)));
+	TestTrue(TEXT("Read/WriteVarUInt(0x0000'0000'ffff'ffff)"), TestSerializeVarUInt64(uint64(0x0000'0000'ffff'ffff)));
+	TestTrue(TEXT("Read/WriteVarUInt(0x0000'0007'ffff'ffff)"), TestSerializeVarUInt64(uint64(0x0000'0007'ffff'ffff)));
+	TestTrue(TEXT("Read/WriteVarUInt(0x0000'0008'0000'0000)"), TestSerializeVarUInt64(uint64(0x0000'0008'0000'0000)));
+	TestTrue(TEXT("Read/WriteVarUInt(0x0000'03ff'ffff'ffff)"), TestSerializeVarUInt64(uint64(0x0000'03ff'ffff'ffff)));
+	TestTrue(TEXT("Read/WriteVarUInt(0x0000'0400'0000'0000)"), TestSerializeVarUInt64(uint64(0x0000'0400'0000'0000)));
+	TestTrue(TEXT("Read/WriteVarUInt(0x0001'ffff'ffff'ffff)"), TestSerializeVarUInt64(uint64(0x0001'ffff'ffff'ffff)));
+	TestTrue(TEXT("Read/WriteVarUInt(0x0002'0000'0000'0000)"), TestSerializeVarUInt64(uint64(0x0002'0000'0000'0000)));
+	TestTrue(TEXT("Read/WriteVarUInt(0x00ff'ffff'ffff'ffff)"), TestSerializeVarUInt64(uint64(0x00ff'ffff'ffff'ffff)));
+	TestTrue(TEXT("Read/WriteVarUInt(0x0100'0000'0000'0000)"), TestSerializeVarUInt64(uint64(0x0100'0000'0000'0000)));
+	TestTrue(TEXT("Read/WriteVarUInt(0xffff'ffff'ffff'ffff)"), TestSerializeVarUInt64(uint64(0xffff'ffff'ffff'ffff)));
+
+	return true;
+}
+
+#endif // WITH_DEV_AUTOMATION_TESTS
diff --git a/Engine/Source/Runtime/Core/Private/UObject/DevObjectVersion.cpp b/Engine/Source/Runtime/Core/Private/UObject/DevObjectVersion.cpp
--- a/Engine/Source/Runtime/Core/Private/UObject/DevObjectVersion.cpp
+++ b/Engine/Source/Runtime/Core/Private/UObject/DevObjectVersion.cpp
@@ -30,6 +30,7 @@
 #include "UObject/CineCameraObjectVersion.h"
 #include "UObject/VirtualProductionObjectVersion.h"
 #include "UObject/UE5MainStreamObjectVersion.h"
+#include "UObject/UE5CookerObjectVersion.h"
 
 DEFINE_LOG_CATEGORY_STATIC(LogDevObjectVersion, Log, All);
 
@@ -200,3 +201,8 @@ FDevVersionRegistration GRegisterVirtualProductionObjectVersion(FVirtualProducti
 const FGuid FUE5MainStreamObjectVersion::GUID(0x697DD581, 0xE64f41AB, 0xAA4A51EC, 0xBEB7B628);
 // Register UE5 main stream custom version with Core
 FDevVersionRegistration GRegisterUE5MainStreamObjectVersion(FUE5MainStreamObjectVersion::GUID, FUE5MainStreamObjectVersion::LatestVersion, TEXT("UE5-Main"));
+
+// Unique Cooker version id
+const FGuid FUE5CookerObjectVersion::GUID(0x26075A32, 0x730F4708, 0x88E98C32, 0xF1599D05);
+// Register Cooker custom version with Core
+FDevVersionRegistration GRegisterCookerObjectVersion(FUE5CookerObjectVersion::GUID, FUE5CookerObjectVersion::LatestVersion, TEXT("UE5-Dev-Cooker"));
diff --git a/Engine/Source/Runtime/Core/Public/Containers/ContainersFwd.h b/Engine/Source/Runtime/Core/Public/Containers/ContainersFwd.h
--- a/Engine/Source/Runtime/Core/Public/Containers/ContainersFwd.h
+++ b/Engine/Source/Runtime/Core/Public/Containers/ContainersFwd.h
@@ -2,6 +2,7 @@
 
 #pragma once
 
+#include "HAL/Platform.h"
 #include "CoreTypes.h"
 #include "Traits/IsContiguousContainer.h"
 
diff --git a/Engine/Source/Runtime/Core/Public/CoreGlobals.h b/Engine/Source/Runtime/Core/Public/CoreGlobals.h
--- a/Engine/Source/Runtime/Core/Public/CoreGlobals.h
+++ b/Engine/Source/Runtime/Core/Public/CoreGlobals.h
@@ -642,3 +642,11 @@ void CORE_API EnableEmitDrawEventsOnlyOnCommandlist();
 /** Array to help visualize weak pointers in the debugger */
 class FChunkedFixedUObjectArray;
 extern CORE_API FChunkedFixedUObjectArray* GCoreObjectArrayForDebugVisualizers;
+
+/** Array to help visualize object paths in the debugger */
+extern template class TArray<FMinimalName, TInlineAllocator<3>>;
+extern CORE_API TArray<FMinimalName, TInlineAllocator<3>>* GCoreComplexObjectPathDebug;
+
+/** Array to help visualize object handles in the debugger */
+struct FObjectHandlePackageDebugData;
+extern CORE_API FObjectHandlePackageDebugData* GCoreObjectHandlePackageDebug;
diff --git a/Engine/Source/Runtime/Core/Public/Memory/SharedBuffer.h b/Engine/Source/Runtime/Core/Public/Memory/SharedBuffer.h
new file mode 100644
--- /dev/null
+++ b/Engine/Source/Runtime/Core/Public/Memory/SharedBuffer.h
@@ -0,0 +1,390 @@
+// Copyright Epic Games, Inc. All Rights Reserved.
+
+#pragma once
+
+#include "CoreTypes.h"
+#include "Memory/MemoryView.h"
+#include "Templates/Decay.h"
+#include "Templates/EnableIf.h"
+#include "Templates/SharedPointer.h"
+#include "Templates/UnrealTemplate.h"
+
+///////////////////////////////////////////////////////////////////////////////////////////////////////////////////////
+
+/**
+ * A reference-counted shared buffer type.
+ *
+ * A shared buffer may either own its memory or provide a view into memory. When used as a view,
+ * it is the responsibility of the user to guarantee that the lifetime of the wrapped memory will
+ * exceed the lifetime of the shared buffer. \ref MakeSharedBufferOwned can be used to guarantee
+ * that the memory is owned by the buffer.
+ *
+ * A shared buffer must be constructed using \ref MakeSharedBuffer to invoke one of the provided
+ * constructors to allocate, assume ownership of, clone, or wrap a buffer.
+ *
+ * A shared buffer must be stored and accessed through one of its six reference types:
+ * FSharedBuffer[Const]Ref or FSharedBuffer[Const][Weak]Ptr.
+ */
+class FSharedBuffer
+{
+public:
+	enum EAssumeOwnershipTag { AssumeOwnership };
+	enum ECloneTag { Clone };
+	enum EWrapTag { Wrap };
+
+	FSharedBuffer() = default;
+
+	/** Use via MakeSharedBuffer. Allocate an owned buffer of the given size. */
+	CORE_API explicit FSharedBuffer(uint64 Size);
+
+	/** Use via MakeSharedBuffer. Assumes ownership of Data and calls FMemory::Free upon destruction. */
+	CORE_API FSharedBuffer(EAssumeOwnershipTag, void* Data, uint64 Size);
+	/** Use via MakeSharedBuffer. Clone the given buffer to an owned copy. */
+	CORE_API FSharedBuffer(ECloneTag, const void* Data, uint64 Size);
+	/** Use via MakeSharedBuffer. Wrap the given buffer without taking ownership. */
+	CORE_API FSharedBuffer(EWrapTag, void* Data, uint64 Size);
+
+	CORE_API ~FSharedBuffer();
+
+	FSharedBuffer(const FSharedBuffer&) = delete;
+	FSharedBuffer& operator=(const FSharedBuffer&) = delete;
+
+	inline void* GetData() { return Data; }
+	inline const void* GetData() const { return Data; }
+
+	inline uint64 Size() const { return DataSize; }
+
+	/** Whether the shared buffer owns the memory that it provides a view of. */
+	inline bool IsOwned() const { return (Flags & uint16(ESharedBufferFlags::Owned)) != 0; }
+
+	inline FMutableMemoryView GetView() { return FMutableMemoryView(GetData(), Size()); }
+	inline FConstMemoryView GetView() const { return FConstMemoryView(GetData(), Size()); }
+
+	inline operator FMutableMemoryView() { return GetView(); }
+	inline operator FConstMemoryView() const { return GetView(); }
+
+private:
+	enum class ESharedBufferFlags : uint16
+	{
+		None = 0,
+		Owned = 1 << 0,
+	};
+
+	void* Data = nullptr;
+	union
+	{
+		uint64 DataSizeAndFlags = 0;
+		struct
+		{
+			uint64 DataSize : 48;
+			uint64 Flags : 16;
+		};
+	};
+};
+
+///////////////////////////////////////////////////////////////////////////////////////////////////////////////////////
+
+class FSharedBufferRef;
+class FSharedBufferConstRef;
+
+namespace SharedBufferPrivate
+{
+	/** Wrapper to allow GetTypeHash to resolve from within a scope with another overload. */
+	template <typename... ArgTypes>
+	auto WrapGetTypeHash(ArgTypes&&... Args) -> decltype(GetTypeHash(Forward<ArgTypes>(Args)...))
+	{
+		return GetTypeHash(Forward<ArgTypes>(Args)...);
+	}
+
+	FSharedBufferRef MakeSharedBuffer(TSharedRef<FSharedBuffer, ESPMode::ThreadSafe>&& Ref);
+	FSharedBufferConstRef MakeSharedBuffer(TSharedRef<const FSharedBuffer, ESPMode::ThreadSafe>&& Ref);
+}
+
+///////////////////////////////////////////////////////////////////////////////////////////////////////////////////////
+
+// The reference types defined below provide an interface that is a subset of the shared pointer
+// types provided in Templates/SharedPointer.h. These types will be switched to alias the generic
+// types once the generic types have been updated to always provide thread safety and to deprecate
+// operations that are not safe in the presence of multi-threaded usage.
+
+///////////////////////////////////////////////////////////////////////////////////////////////////////////////////////
+
+/**
+ * A non-nullable, thread-safe, mutable reference to a shared buffer.
+ *
+ * Equivalent to TSharedRef<FSharedBuffer, ESPMode::ThreadSafe>
+ */
+class FSharedBufferRef
+{
+public:
+	inline FSharedBuffer& Get() const { return Ref.Get(); }
+
+	inline FSharedBuffer& operator*() const { return Ref.operator*(); }
+	inline FSharedBuffer* operator->() const { return Ref.operator->(); }
+
+	inline friend uint32 GetTypeHash(const FSharedBufferRef& InRef) { return SharedBufferPrivate::WrapGetTypeHash(InRef.Ref); }
+
+private:
+	constexpr inline friend bool IsValid(const FSharedBufferRef&) { return true; }
+
+	inline explicit FSharedBufferRef(TSharedRef<FSharedBuffer, ESPMode::ThreadSafe>&& InRef) : Ref(InRef) {}
+
+	friend FSharedBufferRef SharedBufferPrivate::MakeSharedBuffer(TSharedRef<FSharedBuffer, ESPMode::ThreadSafe>&& Ref);
+
+	friend class FSharedBufferConstRef;
+	friend class FSharedBufferPtr;
+	friend class FSharedBufferConstPtr;
+	friend class FSharedBufferWeakPtr;
+	friend class FSharedBufferConstWeakPtr;
+
+	TSharedRef<FSharedBuffer, ESPMode::ThreadSafe> Ref;
+};
+
+/**
+ * A non-nullable, thread-safe, const reference to a shared buffer.
+ *
+ * Equivalent to TSharedRef<const FSharedBuffer, ESPMode::ThreadSafe>
+ */
+class FSharedBufferConstRef
+{
+public:
+	inline FSharedBufferConstRef(const FSharedBufferRef& InRef) : Ref(InRef.Ref) {}
+
+	inline const FSharedBuffer& Get() const { return Ref.Get(); }
+
+	inline const FSharedBuffer& operator*() const { return Ref.operator*(); }
+	inline const FSharedBuffer* operator->() const { return Ref.operator->(); }
+
+	inline friend uint32 GetTypeHash(const FSharedBufferConstRef& InRef) { return SharedBufferPrivate::WrapGetTypeHash(InRef.Ref); }
+
+private:
+	constexpr inline friend bool IsValid(const FSharedBufferConstRef&) { return true; }
+
+	inline explicit FSharedBufferConstRef(TSharedRef<const FSharedBuffer, ESPMode::ThreadSafe>&& InRef) : Ref(InRef) {}
+
+	friend FSharedBufferConstRef SharedBufferPrivate::MakeSharedBuffer(TSharedRef<const FSharedBuffer, ESPMode::ThreadSafe>&& Ref);
+
+	friend class FSharedBufferConstPtr;
+	friend class FSharedBufferConstWeakPtr;
+
+	TSharedRef<const FSharedBuffer, ESPMode::ThreadSafe> Ref;
+};
+
+///////////////////////////////////////////////////////////////////////////////////////////////////////////////////////
+
+/**
+ * A non-nullable, thread-safe, mutable pointer to a shared buffer.
+ *
+ * Equivalent to TSharedPtr<FSharedBuffer, ESPMode::ThreadSafe>
+ */
+class FSharedBufferPtr
+{
+public:
+	FSharedBufferPtr() = default;
+
+	inline FSharedBufferPtr(const FSharedBufferRef& InRef) : Ptr(InRef.Ref) {}
+
+	inline FSharedBufferRef ToSharedRef() const& { return FSharedBufferRef(Ptr.ToSharedRef()); }
+	inline FSharedBufferRef ToSharedRef() && { return FSharedBufferRef(MoveTemp(Ptr).ToSharedRef()); }
+
+	inline bool IsValid() const { return Ptr.IsValid(); }
+	inline FSharedBuffer* Get() const { return Ptr.Get(); }
+
+	inline explicit operator bool() const { return Ptr.IsValid(); }
+	inline FSharedBuffer& operator*() const { return Ptr.operator*(); }
+	inline FSharedBuffer* operator->() const { return Ptr.operator->(); }
+
+	inline void Reset() { Ptr.Reset(); }
+
+	inline friend uint32 GetTypeHash(const FSharedBufferPtr& InPtr) { return SharedBufferPrivate::WrapGetTypeHash(InPtr.Ptr); }
+
+private:
+	inline friend bool IsValid(const FSharedBufferPtr& InPtr) { return InPtr.IsValid(); }
+
+	inline FSharedBufferPtr(TSharedPtr<FSharedBuffer, ESPMode::ThreadSafe>&& InPtr) : Ptr(InPtr) {}
+
+	friend class FSharedBufferConstPtr;
+	friend class FSharedBufferWeakPtr;
+	friend class FSharedBufferConstWeakPtr;
+
+	TSharedPtr<FSharedBuffer, ESPMode::ThreadSafe> Ptr;
+};
+
+/**
+ * A non-nullable, thread-safe, const pointer to a shared buffer.
+ *
+ * Equivalent to TSharedPtr<const FSharedBuffer, ESPMode::ThreadSafe>
+ */
+class FSharedBufferConstPtr
+{
+public:
+	FSharedBufferConstPtr() = default;
+
+	inline FSharedBufferConstPtr(const FSharedBufferRef& InRef) : Ptr(InRef.Ref) {}
+	inline FSharedBufferConstPtr(const FSharedBufferConstRef& InRef) : Ptr(InRef.Ref) {}
+	inline FSharedBufferConstPtr(const FSharedBufferPtr& InPtr) : Ptr(InPtr.Ptr) {}
+	inline FSharedBufferConstPtr(FSharedBufferPtr&& InPtr) : Ptr(MoveTemp(InPtr.Ptr)) {}
+
+	inline FSharedBufferConstRef ToSharedRef() const& { return FSharedBufferConstRef(Ptr.ToSharedRef()); }
+	inline FSharedBufferConstRef ToSharedRef() && { return FSharedBufferConstRef(MoveTemp(Ptr).ToSharedRef()); }
+
+	inline bool IsValid() const { return Ptr.IsValid(); }
+	inline const FSharedBuffer* Get() const { return Ptr.Get(); }
+
+	inline explicit operator bool() const { return Ptr.IsValid(); }
+	inline const FSharedBuffer& operator*() const { return Ptr.operator*(); }
+	inline const FSharedBuffer* operator->() const { return Ptr.operator->(); }
+
+	inline void Reset() { Ptr.Reset(); }
+
+	inline friend uint32 GetTypeHash(const FSharedBufferConstPtr& InPtr) { return SharedBufferPrivate::WrapGetTypeHash(InPtr.Ptr); }
+
+private:
+	inline friend bool IsValid(const FSharedBufferConstPtr& InPtr) { return InPtr.IsValid(); }
+
+	inline FSharedBufferConstPtr(TSharedPtr<const FSharedBuffer, ESPMode::ThreadSafe>&& InPtr) : Ptr(InPtr) {}
+
+	friend class FSharedBufferConstWeakPtr;
+
+	TSharedPtr<const FSharedBuffer, ESPMode::ThreadSafe> Ptr;
+};
+
+///////////////////////////////////////////////////////////////////////////////////////////////////////////////////////
+
+/**
+ * A non-nullable, thread-safe, weak, mutable pointer to a shared buffer.
+ *
+ * Equivalent to TWeakPtr<FSharedBuffer, ESPMode::ThreadSafe>
+ */
+class FSharedBufferWeakPtr
+{
+public:
+	FSharedBufferWeakPtr() = default;
+
+	inline FSharedBufferWeakPtr(const FSharedBufferRef& InRef) : Ptr(InRef.Ref) {}
+	inline FSharedBufferWeakPtr(const FSharedBufferPtr& InPtr) : Ptr(InPtr.Ptr) {}
+	inline FSharedBufferWeakPtr(FSharedBufferPtr&& InPtr) : Ptr(MoveTemp(InPtr.Ptr)) {}
+
+	inline FSharedBufferPtr Pin() const& { return FSharedBufferPtr(Ptr.Pin()); }
+	inline FSharedBufferPtr Pin() && { return FSharedBufferPtr(MoveTemp(Ptr).Pin()); }
+
+	inline void Reset() { Ptr.Reset(); }
+
+	inline friend uint32 GetTypeHash(const FSharedBufferWeakPtr& InPtr) { return SharedBufferPrivate::WrapGetTypeHash(InPtr.Ptr); }
+
+private:
+	friend class FSharedBufferConstWeakPtr;
+
+	TWeakPtr<FSharedBuffer, ESPMode::ThreadSafe> Ptr;
+};
+
+/**
+ * A non-nullable, thread-safe, weak, const pointer to a shared buffer.
+ *
+ * Equivalent to TWeakPtr<const FSharedBuffer, ESPMode::ThreadSafe>
+ */
+class FSharedBufferConstWeakPtr
+{
+public:
+	FSharedBufferConstWeakPtr() = default;
+
+	inline FSharedBufferConstWeakPtr(const FSharedBufferRef& InRef) : Ptr(InRef.Ref) {}
+	inline FSharedBufferConstWeakPtr(const FSharedBufferConstRef& InRef) : Ptr(InRef.Ref) {}
+	inline FSharedBufferConstWeakPtr(const FSharedBufferPtr& InPtr) : Ptr(InPtr.Ptr) {}
+	inline FSharedBufferConstWeakPtr(const FSharedBufferConstPtr& InPtr) : Ptr(InPtr.Ptr) {}
+	inline FSharedBufferConstWeakPtr(const FSharedBufferWeakPtr& InPtr) : Ptr(InPtr.Ptr) {}
+	inline FSharedBufferConstWeakPtr(FSharedBufferPtr&& InPtr) : Ptr(MoveTemp(InPtr.Ptr)) {}
+	inline FSharedBufferConstWeakPtr(FSharedBufferConstPtr&& InPtr) : Ptr(MoveTemp(InPtr.Ptr)) {}
+	inline FSharedBufferConstWeakPtr(FSharedBufferWeakPtr&& InPtr) : Ptr(MoveTemp(InPtr.Ptr)) {}
+
+	inline FSharedBufferConstPtr Pin() const& { return FSharedBufferConstPtr(Ptr.Pin()); }
+	inline FSharedBufferConstPtr Pin() && { return FSharedBufferConstPtr(MoveTemp(Ptr).Pin()); }
+
+	inline void Reset() { Ptr.Reset(); }
+
+	inline friend uint32 GetTypeHash(const FSharedBufferConstWeakPtr& InPtr) { return SharedBufferPrivate::WrapGetTypeHash(InPtr.Ptr); }
+
+private:
+	TWeakPtr<const FSharedBuffer, ESPMode::ThreadSafe> Ptr;
+};
+
+///////////////////////////////////////////////////////////////////////////////////////////////////////////////////////
+
+template<> struct TIsZeroConstructType<FSharedBufferRef> { enum { Value = true }; };
+template<> struct TIsZeroConstructType<FSharedBufferConstRef> { enum { Value = true }; };
+template<> struct TIsZeroConstructType<FSharedBufferPtr> { enum { Value = true }; };
+template<> struct TIsZeroConstructType<FSharedBufferConstPtr> { enum { Value = true }; };
+template<> struct TIsZeroConstructType<FSharedBufferWeakPtr> { enum { Value = true }; };
+template<> struct TIsZeroConstructType<FSharedBufferConstWeakPtr> { enum { Value = true }; };
+
+template<> struct TIsWeakPointerType<FSharedBufferWeakPtr> { enum { Value = true }; };
+template<> struct TIsWeakPointerType<FSharedBufferConstWeakPtr> { enum { Value = true }; };
+
+///////////////////////////////////////////////////////////////////////////////////////////////////////////////////////
+
+namespace SharedBufferPrivate
+{
+	inline FSharedBufferRef MakeSharedBuffer(TSharedRef<FSharedBuffer, ESPMode::ThreadSafe>&& Ref)
+	{
+		return FSharedBufferRef(MoveTemp(Ref));
+	}
+
+	inline FSharedBufferConstRef MakeSharedBuffer(TSharedRef<const FSharedBuffer, ESPMode::ThreadSafe>&& Ref)
+	{
+		return FSharedBufferConstRef(MoveTemp(Ref));
+	}
+}
+
+///////////////////////////////////////////////////////////////////////////////////////////////////////////////////////
+
+/**
+ * Construct a shared buffer. \see \ref FSharedBuffer for more details.
+ *
+ * MakeSharedBuffer(); an empty non-owned buffer
+ * MakeSharedBuffer(uint64 Size); an uninitialized buffer of Size bytes
+ * MakeSharedBuffer(FSharedBuffer::AssumeOwnership, Data, Size); take ownership of the buffer (Data, Size)
+ * MakeSharedBuffer(FSharedBuffer::Clone, Data, Size); an owned clone the buffer (Data, Size)
+ * MakeSharedBuffer(FSharedBuffer::Wrap, Data, Size); non-owning wrapper of the buffer (Data, Size)
+ */
+template <typename... ArgTypes, typename = decltype(FSharedBuffer(Forward<ArgTypes>(DeclVal<ArgTypes>())...))>
+inline FSharedBufferRef MakeSharedBuffer(ArgTypes&&... Args)
+{
+	return SharedBufferPrivate::MakeSharedBuffer(MakeShared<FSharedBuffer, ESPMode::ThreadSafe>(Forward<ArgTypes>(Args)...));
+}
+
+/** Construct a const buffer to take ownership of the buffer (Data, Size). \see \ref FSharedBuffer for details. */
+inline FSharedBufferConstRef MakeSharedBuffer(FSharedBuffer::EAssumeOwnershipTag, const void* Data, uint64 Size)
+{
+	// The const_cast is safe because it is returned as a const buffer that disallows mutable access.
+	return SharedBufferPrivate::MakeSharedBuffer(MakeShared<const FSharedBuffer, ESPMode::ThreadSafe>(FSharedBuffer::AssumeOwnership, const_cast<void*>(Data), Size));
+}
+
+/** Construct a const non-owning wrapper of the buffer (Data, Size) \see \ref FSharedBuffer for details. */
+inline FSharedBufferConstRef MakeSharedBuffer(FSharedBuffer::EWrapTag, const void* Data, uint64 Size)
+{
+	// The const_cast is safe because it is returned as a const buffer that disallows mutable access.
+	return SharedBufferPrivate::MakeSharedBuffer(MakeShared<const FSharedBuffer, ESPMode::ThreadSafe>(FSharedBuffer::Wrap, const_cast<void*>(Data), Size));
+}
+
+/** Construct a shared buffer from a memory view. \see \ref FSharedBuffer for details. */
+template <typename TagType, typename ViewType>
+inline auto MakeSharedBuffer(TagType Tag, TMemoryView<ViewType> View) -> decltype(MakeSharedBuffer(DeclVal<TagType>(), View.GetData(), View.Size()))
+{
+	return MakeSharedBuffer(Tag, View.GetData(), View.Size());
+}
+
+///////////////////////////////////////////////////////////////////////////////////////////////////////////////////////
+
+/**
+ * Return the buffer if it is owned, or a clone otherwise.
+ *
+ * @param Buffer A FSharedBuffer[Const]Ref or FSharedBuffer[Const][Weak]Ptr to make owned.
+ * @return An owned copy of the input buffer in the same reference type as the input.
+ */
+template <typename T, typename TEnableIf<TIsSame<FSharedBuffer, typename TDecay<decltype(*DeclVal<T>())>::Type>::Value>::Type* = nullptr>
+inline auto MakeSharedBufferOwned(T&& Buffer) -> decltype(typename TDecay<T>::Type(MakeSharedBuffer(FSharedBuffer::Clone, DeclVal<T>()->GetData(), DeclVal<T>()->Size())))
+{
+	return !IsValid(Buffer) || Buffer->IsOwned() ? Forward<T>(Buffer) : MakeSharedBuffer(FSharedBuffer::Clone, Buffer->GetData(), Buffer->Size());
+}
+
+///////////////////////////////////////////////////////////////////////////////////////////////////////////////////////
diff --git a/Engine/Source/Runtime/Core/Public/Modules/Boilerplate/ModuleBoilerplate.h b/Engine/Source/Runtime/Core/Public/Modules/Boilerplate/ModuleBoilerplate.h
--- a/Engine/Source/Runtime/Core/Public/Modules/Boilerplate/ModuleBoilerplate.h
+++ b/Engine/Source/Runtime/Core/Public/Modules/Boilerplate/ModuleBoilerplate.h
@@ -52,7 +52,9 @@ class FChunkedFixedUObjectArray;
 #else
 	#define UE4_VISUALIZERS_HELPERS \
 		uint8** GNameBlocksDebug = FNameDebugVisualizer::GetBlocks(); \
-		FChunkedFixedUObjectArray*& GObjectArrayForDebugVisualizers = GCoreObjectArrayForDebugVisualizers;
+		FChunkedFixedUObjectArray*& GObjectArrayForDebugVisualizers = GCoreObjectArrayForDebugVisualizers; \
+		TArray<FMinimalName, TInlineAllocator<3>>*& GComplexObjectPathDebug = GCoreComplexObjectPathDebug; \
+		FObjectHandlePackageDebugData*& GObjectHandlePackageDebug = GCoreObjectHandlePackageDebug;
 #endif
 
 // in DLL builds, these are done per-module, otherwise we just need one in the application
diff --git a/Engine/Source/Runtime/Core/Public/Serialization/Archive.h b/Engine/Source/Runtime/Core/Public/Serialization/Archive.h
--- a/Engine/Source/Runtime/Core/Public/Serialization/Archive.h
+++ b/Engine/Source/Runtime/Core/Public/Serialization/Archive.h
@@ -1143,6 +1143,16 @@ public:
 	 * @return This instance.
 	 */
 	virtual FArchive& operator<<(struct FLazyObjectPtr& Value);
+
+	/**
+	 * Serializes a wrapped object pointer value from or into this archive.
+	 *
+	 * Most of the time, FObjectPtrs are serialized as UObject*, but some archives need to override this.
+	 *
+	 * @param Value The value to serialize.
+	 * @return This instance.
+	 */
+	virtual FArchive& operator<<(struct FObjectPtr& Value);
 	
 	/**
 	 * Serializes asset pointer from or into this archive.
diff --git a/Engine/Source/Runtime/Core/Public/Serialization/ArchiveProxy.h b/Engine/Source/Runtime/Core/Public/Serialization/ArchiveProxy.h
--- a/Engine/Source/Runtime/Core/Public/Serialization/ArchiveProxy.h
+++ b/Engine/Source/Runtime/Core/Public/Serialization/ArchiveProxy.h
@@ -46,6 +46,12 @@ public:
 		return *this;
 	}
 
+	virtual FArchive& operator<<(FObjectPtr& Value) override
+	{
+		InnerArchive << Value;
+		return *this;
+	}
+
 	virtual FArchive& operator<<(FLazyObjectPtr& Value) override
 	{
 		InnerArchive << Value;
