Copying //UE5/Dev-Cooker@14539516 to Main (//UE5/Main)

diff --git a/Engine/Source/Runtime/Core/Public/Serialization/CompactBinary.h b/Engine/Source/Runtime/Core/Public/Serialization/CompactBinary.h
new file mode 100644
--- /dev/null
+++ b/Engine/Source/Runtime/Core/Public/Serialization/CompactBinary.h
@@ -0,0 +1,1142 @@
+// Copyright Epic Games, Inc. All Rights Reserved.
+
+#pragma once
+
+#include "CoreTypes.h"
+#include "Containers/StringView.h"
+#include "Memory/MemoryView.h"
+#include "Memory/SharedBuffer.h"
+#include "Misc/Blake3.h"
+#include "Misc/EnumClassFlags.h"
+#include "Templates/IsTriviallyDestructible.h"
+
+/**
+ * \file
+ *
+ * This file declares a compact binary data format that is compatible with JSON and only slightly
+ * more expressive. The format is designed to achieve fairly small encoded sizes while also being
+ * efficient to read both sequentially and through random access. An atom of data in this compact
+ * binary format is called a field, which can be an object, array, byte string, character string,
+ * integer, floating point value, boolean, null, uuid, date/time, time span, or a hash digest for
+ * referencing blobs or compact binary data that are stored externally.
+ *
+ * An object is a collection of name-field pairs, and an array is a collection of fields. Encoded
+ * objects and arrays are both written such that they may be interpreted as a field that can then
+ * be cast to an object or array. This attribute means that a blob containing compact binary data
+ * is always safe to interpret as a field, which allows for easy validation as described later.
+ *
+ * A field can be constructed as a view of the underlying memory with \ref FCbField or can either
+ * view (wrap) or own (clone or assume ownership) the underlying memory with \ref FCbFieldRef. An
+ * array provides the same functionality through \ref FCbArray or \ref FCbArrayRef, and an object
+ * uses \ref FCbObject or \ref FCbObjectRef.
+ *
+ * It is optimal use the view types when possible, and reference types only when they are needed,
+ * to avoid the overhead of the atomic reference counting of the shared buffer.
+ *
+ * A range of data validation functionality is provided through \ref ValidateCompactBinary. Using
+ * the Default mode of validation is sufficient to be able to read the data without crashing. The
+ * additional validation modes are required for compatibility with other formats (Names), correct
+ * byte-wise comparisons of the encoded data (Format), or storage without arbitrary padding bytes
+ * leading to non-deterministic values being stored (Padding).
+ *
+ * Example:
+ *
+ * \code
+ * void BeginBuild(FCbObjectRef Params)
+ * {
+ *     if (FSharedBufferConstPtr Data = Storage().Load(Params["Data"].AsBinaryHash()))
+ *     {
+ *         SetData(Data);
+ *     }
+ *
+ *     if (Params["Resize"].AsBool())
+ *     {
+ *         FCbField MaxWidthField = Params["MaxWidth"]
+ *         FCbField MaxHeightField = Params["MaxHeight"];
+ *         if (MaxWidthField && MaxHeightField)
+ *         {
+ *             Resize(MaxWidthField.AsInt32(), MaxHeightField.AsInt32());
+ *         }
+ *     }
+ *
+ *     for (FCbField Format : Params["Formats"].AsArray())
+ *     {
+ *         BeginCompress(FName(Format.AsString()));
+ *     }
+ * }
+ * \endcode
+ */
+
+class FCbArray;
+class FCbObject;
+struct FDateTime;
+struct FGuid;
+struct FTimespan;
+
+///////////////////////////////////////////////////////////////////////////////////////////////////////////////////////
+
+/**
+ * Field types and flags for \ref FCbField.
+ *
+ * This is a private type and is only declared here to enable inline use below.
+ *
+ * DO NOT CHANGE THE VALUE OF ANY MEMBERS OF THIS ENUM!
+ * BACKWARD COMPATIBILITY REQUIRES THAT THESE VALUES BE FIXED!
+ * SERIALIZATION USES HARD-CODED CONSTANTS BASED ON THESE VALUES!
+ */
+enum class ECbFieldType : uint8
+{
+	/** A field type that does not occur in a valid object. */
+	None            = 0x00,
+
+	/** Null. Payload is empty. */
+	Null            = 0x01,
+
+	/**
+	 * Object is an array of fields with unique non-empty names.
+	 *
+	 * Payload is a VarUInt byte count for the encoded fields followed by the fields.
+	 */
+	Object          = 0x02,
+	/**
+	 * UniformObject is an array of fields with the same field types and unique non-empty names.
+	 *
+	 * Payload is a VarUInt byte count for the encoded fields followed by the fields.
+	 */
+	UniformObject   = 0x03,
+
+	/**
+	 * Array is an array of fields with no name that may be of different types.
+	 *
+	 * Payload is a VarUInt byte count, followed by a VarUInt item count, followed by the fields.
+	 */
+	Array           = 0x04,
+	/**
+	 * UniformArray is an array of fields with no name and with the same field type.
+	 *
+	 * Payload is a VarUInt byte count, followed by a VarUInt item count, followed by field type,
+	 * followed by the fields without their field type.
+	 */
+	UniformArray    = 0x05,
+
+	/** Binary. Payload is a VarUInt byte count followed by the data. */
+	Binary          = 0x06,
+
+	/** String in UTF-8. Payload is a VarUInt byte count then an unterminated UTF-8 string. */
+	String          = 0x07,
+
+	/**
+	 * Non-negative integer with the range of a 64-bit unsigned integer.
+	 *
+	 * Payload is the value encoded as a VarUInt.
+	 */
+	IntegerPositive = 0x08,
+	/**
+	 * Negative integer with the range of a 64-bit signed integer.
+	 *
+	 * Payload is the ones' complement of the value encoded as a VarUInt.
+	 */
+	IntegerNegative = 0x09,
+
+	/** Single precision float. Payload is one big endian IEEE 754 binary32 float. */
+	Float32         = 0x0a,
+	/** Double precision float. Payload is one big endian IEEE 754 binary64 float. */
+	Float64         = 0x0b,
+
+	/** Boolean false value. Payload is empty. */
+	BoolFalse       = 0x0c,
+	/** Boolean true value. Payload is empty. */
+	BoolTrue        = 0x0d,
+
+	/**
+	 * BinaryHash is a reference to a blob stored externally.
+	 *
+	 * Payload is a 256-bit hash digest of the blob.
+	 */
+	BinaryHash      = 0x0e,
+	/**
+	 * FieldHash is a reference to a field/array/object stored externally.
+	 *
+	 * Payload is a 256-bit hash digest of the field.
+	 */
+	FieldHash       = 0x0f,
+
+	/** UUID/GUID. Payload is a 128-bit UUID as defined by RFC 4122. */
+	Uuid            = 0x10,
+
+	/**
+	 * Date and time between 0001-01-01 00:00:00.0000000 and 9999-12-31 23:59:59.9999999.
+	 *
+	 * Payload is a big endian int64 count of 100ns ticks since 0001-01-01 00:00:00.0000000.
+	 */
+	DateTime        = 0x11,
+	/**
+	 * Difference between two date/time values.
+	 *
+	 * Payload is a big endian int64 count of 100ns ticks in the span, and may be negative.
+	 */
+	TimeSpan        = 0x12,
+
+	/**
+	 * A transient flag which indicates that the object or array containing this field has stored
+	 * the field type before the payload and name. Non-uniform objects and fields will set this.
+	 *
+	 * Note: Since the flag must never be serialized, this bit may be repurposed in the future.
+	 */
+	HasFieldType    = 0x40,
+
+	/** A persisted flag which indicates that the field has a name stored before the payload. */
+	HasFieldName    = 0x80,
+};
+
+ENUM_CLASS_FLAGS(ECbFieldType);
+
+///////////////////////////////////////////////////////////////////////////////////////////////////////////////////////
+
+/** Functions that operate on \ref ECbFieldType. */
+class FCbFieldType
+{
+	static constexpr ECbFieldType SerializedTypeMask    = ECbFieldType(0b1011'1111);
+	static constexpr ECbFieldType TypeMask              = ECbFieldType(0b0011'1111);
+
+	static constexpr ECbFieldType ObjectMask            = ECbFieldType(0b0011'1110);
+	static constexpr ECbFieldType ObjectBase            = ECbFieldType(0b0000'0010);
+
+	static constexpr ECbFieldType ArrayMask             = ECbFieldType(0b0011'1110);
+	static constexpr ECbFieldType ArrayBase             = ECbFieldType(0b0000'0100);
+
+	static constexpr ECbFieldType IntegerMask           = ECbFieldType(0b0011'1110);
+	static constexpr ECbFieldType IntegerBase           = ECbFieldType(0b0000'1000);
+
+	static constexpr ECbFieldType FloatMask             = ECbFieldType(0b0011'1100);
+	static constexpr ECbFieldType FloatBase             = ECbFieldType(0b0000'1000);
+
+	static constexpr ECbFieldType BoolMask              = ECbFieldType(0b0011'1110);
+	static constexpr ECbFieldType BoolBase              = ECbFieldType(0b0000'1100);
+
+	static constexpr ECbFieldType ExternHashMask        = ECbFieldType(0b0011'1110);
+	static constexpr ECbFieldType ExternHashBase        = ECbFieldType(0b0000'1110);
+
+	static void StaticAssertTypeConstants();
+
+public:
+	/** The type with flags removed. */
+	static constexpr inline ECbFieldType GetType(ECbFieldType Type)             { return Type & TypeMask; }
+	/** The type with transient flags removed. */
+	static constexpr inline ECbFieldType GetSerializedType(ECbFieldType Type)   { return Type & SerializedTypeMask; }
+
+	static constexpr inline bool HasFieldType(ECbFieldType Type)                { return EnumHasAnyFlags(Type, ECbFieldType::HasFieldType); }
+	static constexpr inline bool HasFieldName(ECbFieldType Type)                { return EnumHasAnyFlags(Type, ECbFieldType::HasFieldName); }
+
+	static constexpr inline bool IsNone(ECbFieldType Type)                      { return GetType(Type) == ECbFieldType::None; }
+	static constexpr inline bool IsNull(ECbFieldType Type)                      { return GetType(Type) == ECbFieldType::Null; }
+
+	static constexpr inline bool IsObject(ECbFieldType Type)                    { return (Type & ObjectMask) == ObjectBase; }
+	static constexpr inline bool IsArray(ECbFieldType Type)                     { return (Type & ArrayMask) == ArrayBase; }
+
+	static constexpr inline bool IsBinary(ECbFieldType Type)                    { return GetType(Type) == ECbFieldType::Binary; }
+	static constexpr inline bool IsString(ECbFieldType Type)                    { return GetType(Type) == ECbFieldType::String; }
+
+	static constexpr inline bool IsInteger(ECbFieldType Type)                   { return (Type & IntegerMask) == IntegerBase; }
+	/** Whether the field is a float, or integer due to implicit conversion. */
+	static constexpr inline bool IsFloat(ECbFieldType Type)                     { return (Type & FloatMask) == FloatBase; }
+	static constexpr inline bool IsBool(ECbFieldType Type)                      { return (Type & BoolMask) == BoolBase; }
+
+	static constexpr inline bool IsExternHash(ECbFieldType Type)                { return (Type & ExternHashMask) == ExternHashBase; }
+	static constexpr inline bool IsBinaryHash(ECbFieldType Type)                { return GetType(Type) == ECbFieldType::BinaryHash; }
+	static constexpr inline bool IsFieldHash(ECbFieldType Type)                 { return GetType(Type) == ECbFieldType::FieldHash; }
+
+	static constexpr inline bool IsUuid(ECbFieldType Type)                      { return GetType(Type) == ECbFieldType::Uuid; }
+
+	static constexpr inline bool IsDateTime(ECbFieldType Type)                  { return GetType(Type) == ECbFieldType::DateTime; }
+	static constexpr inline bool IsTimeSpan(ECbFieldType Type)                  { return GetType(Type) == ECbFieldType::TimeSpan; }
+};
+
+///////////////////////////////////////////////////////////////////////////////////////////////////////////////////////
+
+/**
+ * Iterator for \ref FCbField[Ref] that can operate on any contiguous range of fields.
+ *
+ * The iterator *is* the current field that the iterator points to and exposes the full interface
+ * of \ref FCbField[Ref]. An iterator that is at the end is equivalent to a field with no value.
+ */
+template <typename FieldType>
+class TCbFieldIterator : public FieldType
+{
+public:
+	/** Construct a field iterator that has no fields to iterate. */
+	constexpr TCbFieldIterator() = default;
+
+	/**
+	 * Construct a field iterator.
+	 *
+	 * @param InField The first field, or the default field if there are no fields.
+	 * @param InFieldsEnd A pointer to the end of the payload of the last field, or null.
+	 */
+	constexpr inline TCbFieldIterator(const FieldType& InField, const void* InFieldsEnd)
+		: FieldType(InField)
+		, FieldsEnd(InFieldsEnd)
+	{
+	}
+
+	/**
+	 * Construct a field iterator.
+	 *
+	 * @param InField The first field, or the default field if there are no fields.
+	 * @param InIterator An iterator from which to take the pointer to the end of the payload of the last field.
+	 */
+	template <typename OtherType>
+	constexpr inline TCbFieldIterator(FieldType&& InField, const TCbFieldIterator<OtherType>& InIterator)
+		: FieldType(MoveTemp(InField))
+		, FieldsEnd(InIterator.FieldsEnd)
+	{
+	}
+
+	inline TCbFieldIterator& operator++()
+	{
+		const void* const PayloadEnd = FieldType::GetPayloadEnd();
+		const int64 AtEndMask = int64(PayloadEnd == FieldsEnd) - 1;
+		const ECbFieldType NextType = ECbFieldType(int64(FieldType::GetType()) & AtEndMask);
+		const void* const NextField = reinterpret_cast<const void*>(int64(PayloadEnd) & AtEndMask);
+		const void* const NextFieldsEnd = reinterpret_cast<const void*>(int64(FieldsEnd) & AtEndMask);
+		FieldType::Assign(NextField, NextType);
+		FieldsEnd = NextFieldsEnd;
+		return *this;
+	}
+
+	constexpr inline FieldType& operator*() { return *this; }
+	constexpr inline FieldType* operator->() { return this; }
+
+	constexpr inline bool operator!() const { return !bool(*this); }
+
+	constexpr inline bool operator==(const TCbFieldIterator& Other) const
+	{
+		return FieldType::GetPayload() == Other.FieldType::GetPayload();
+	}
+
+	constexpr inline bool operator!=(const TCbFieldIterator& Other) const { return !(*this == Other); }
+
+private:
+	friend inline TCbFieldIterator begin(const TCbFieldIterator& Iterator) { return Iterator; }
+	friend inline TCbFieldIterator end(const TCbFieldIterator&) { return TCbFieldIterator(); }
+
+private:
+	template <typename OtherType>
+	friend class TCbFieldIterator;
+
+	/** Pointer to the first byte past the end of the last field. Set to null at the end. */
+	const void* FieldsEnd = nullptr;
+};
+
+///////////////////////////////////////////////////////////////////////////////////////////////////////////////////////
+
+/** Errors that can occur when accessing a field. */
+enum class ECbFieldError : uint8
+{
+	/** The field is not in an error state. */
+	None,
+	/** The value type does not match the requested type. */
+	TypeError,
+	/** The value is out of range for the requested type. */
+	RangeError,
+};
+
+///////////////////////////////////////////////////////////////////////////////////////////////////////////////////////
+
+/**
+ * An atom of data in the compact binary format.
+ *
+ * Accessing the value of a field is always a safe operation, even if accessed as the wrong type.
+ * An invalid access will return a default value for the requested type, and set an error code on
+ * the field that can be checked with GetLastError and HasLastError. A valid access will clear an
+ * error from a previous invalid access.
+ *
+ * A field is encoded in one or more bytes, depending on its type and the type of object or array
+ * that contains it. A field of an object or array which is non-uniform encodes its field type in
+ * the first byte, and includes the HasFieldName flag for a field in an object. The field name is
+ * encoded in a variable-length unsigned integer of its size in bytes, for named fields, followed
+ * by that many bytes of the UTF-8 encoding of the name with no null terminator. The remainder of
+ * the field is the payload and is described in the field type enum. Every field must be uniquely
+ * addressable when encoded, which means a zero-byte field is not permitted, and only arises in a
+ * uniform array of fields with no payload, where the answer is to encode as a non-uniform array.
+ *
+ * This type only provides a view into memory and does not perform any memory management itself.
+ * Use \ref FCbFieldRef to hold a reference to the underlying memory when necessary.
+ */
+class FCbField
+{
+public:
+	/** Construct a field with no name and no value. */
+	constexpr FCbField() = default;
+
+	/**
+	 * Construct a field from a pointer to its data and an optional externally-provided type.
+	 *
+	 * @param Data Pointer to the start of the field data.
+	 * @param Type HasFieldType means that Data contains the type. Otherwise, use the given type.
+	 */
+	CORE_API explicit FCbField(const void* Data, ECbFieldType Type = ECbFieldType::HasFieldType);
+
+	/** Name of the field if it has one, otherwise empty. */
+	constexpr inline FAnsiStringView Name() const { return FAnsiStringView(static_cast<const ANSICHAR*>(Payload) - NameLen, NameLen); }
+
+	/** Size of the field in bytes. */
+	CORE_API uint64 Size() const;
+
+	/**
+	 * Whether this field is identical to the other field.
+	 *
+	 * Performs a deep comparison of any contained arrays or objects and their fields. Comparison
+	 * assumes that both fields are valid and are written in the canonical format. Fields must be
+	 * written in the same order in arrays and objects, and name comparison is case sensitive. If
+	 * these assumptions do not hold, this may return false for equivalent inputs. Validation can
+	 * be done with the `All` mode to check these assumptions about the format of the inputs.
+	 */
+	CORE_API bool Equals(const FCbField& Other) const;
+
+	/** Access the field as an object. Defaults to an empty object on error. */
+	CORE_API FCbObject AsObject();
+
+	/** Access the field as an array. Defaults to an empty array on error. */
+	CORE_API FCbArray AsArray();
+
+	/** Access the field as binary. Returns the provided default on error. */
+	CORE_API FConstMemoryView AsBinary(FConstMemoryView Default = FConstMemoryView());
+
+	/** Access the field as a UTF-8 string. Returns the provided default on error. */
+	CORE_API FAnsiStringView AsString(FAnsiStringView Default = FAnsiStringView());
+
+	/** Access the field as an int8. Returns the provided default on error. */
+	inline int8 AsInt8(int8 Default = 0)       { return AsInteger<int8>(Default); }
+	/** Access the field as an int16. Returns the provided default on error. */
+	inline int16 AsInt16(int16 Default = 0)    { return AsInteger<int16>(Default); }
+	/** Access the field as an int32. Returns the provided default on error. */
+	inline int32 AsInt32(int32 Default = 0)    { return AsInteger<int32>(Default); }
+	/** Access the field as an int64. Returns the provided default on error. */
+	inline int64 AsInt64(int64 Default = 0)    { return AsInteger<int64>(Default); }
+	/** Access the field as a uint8. Returns the provided default on error. */
+	inline uint8 AsUInt8(uint8 Default = 0)    { return AsInteger<uint8>(Default); }
+	/** Access the field as a uint16. Returns the provided default on error. */
+	inline uint16 AsUInt16(uint16 Default = 0) { return AsInteger<uint16>(Default); }
+	/** Access the field as a uint32. Returns the provided default on error. */
+	inline uint32 AsUInt32(uint32 Default = 0) { return AsInteger<uint32>(Default); }
+	/** Access the field as a uint64. Returns the provided default on error. */
+	inline uint64 AsUInt64(uint64 Default = 0) { return AsInteger<uint64>(Default); }
+
+	/** Access the field as a float. Returns the provided default on error. */
+	CORE_API float AsFloat(float Default = 0.0f);
+	/** Access the field as a double. Returns the provided default on error. */
+	CORE_API double AsDouble(double Default = 0.0);
+
+	/** Access the field as a bool. Returns the provided default on error. */
+	CORE_API bool AsBool(bool bDefault = false);
+
+	/** Access the field as a hash referencing a blob. Returns the provided default on error. */
+	CORE_API FBlake3Hash AsBinaryHash(const FBlake3Hash& Default = FBlake3Hash());
+	/** Access the field as a hash referencing a field/array/object. Returns the provided default on error. */
+	CORE_API FBlake3Hash AsFieldHash(const FBlake3Hash& Default = FBlake3Hash());
+
+	/** Access the field as a UUID. Returns a nil UUID on error. */
+	CORE_API FGuid AsUuid();
+	/** Access the field as a UUID. Returns the provided default on error. */
+	CORE_API FGuid AsUuid(const FGuid& Default);
+
+	/** Access the field as a date/time tick count. Returns the provided default on error. */
+	CORE_API int64 AsDateTimeTicks(int64 Default = 0);
+
+	/** Access the field as a date/time. Returns a date/time at the epoch on error. */
+	CORE_API FDateTime AsDateTime();
+	/** Access the field as a date/time. Returns the provided default on error. */
+	CORE_API FDateTime AsDateTime(FDateTime Default);
+
+	/** Access the field as a timespan tick count. Returns the provided default on error. */
+	CORE_API int64 AsTimeSpanTicks(int64 Default = 0);
+
+	/** Access the field as a timespan. Returns an empty timespan on error. */
+	CORE_API FTimespan AsTimeSpan();
+	/** Access the field as a timespan. Returns the provided default on error. */
+	CORE_API FTimespan AsTimeSpan(FTimespan Default);
+
+	/** True if the field has a name. */
+	constexpr inline bool HasName() const           { return FCbFieldType::HasFieldName(Type); }
+
+	constexpr inline bool IsNull() const            { return FCbFieldType::IsNull(Type); }
+
+	constexpr inline bool IsObject() const          { return FCbFieldType::IsObject(Type); }
+	constexpr inline bool IsArray() const           { return FCbFieldType::IsArray(Type); }
+
+	constexpr inline bool IsBinary() const          { return FCbFieldType::IsBinary(Type); }
+	constexpr inline bool IsString() const          { return FCbFieldType::IsString(Type); }
+
+	/** Whether the field is an integer of unspecified range and sign. */
+	constexpr inline bool IsInteger() const         { return FCbFieldType::IsInteger(Type); }
+	/** Whether the field is a float, or integer that supports implicit conversion. */
+	constexpr inline bool IsFloat() const           { return FCbFieldType::IsFloat(Type); }
+	constexpr inline bool IsBool() const            { return FCbFieldType::IsBool(Type); }
+
+	constexpr inline bool IsExternHash() const      { return FCbFieldType::IsExternHash(Type); }
+	constexpr inline bool IsBinaryHash() const      { return FCbFieldType::IsBinaryHash(Type); }
+	constexpr inline bool IsFieldHash() const       { return FCbFieldType::IsFieldHash(Type); }
+
+	constexpr inline bool IsUuid() const            { return FCbFieldType::IsUuid(Type); }
+
+	constexpr inline bool IsDateTime() const        { return FCbFieldType::IsDateTime(Type); }
+	constexpr inline bool IsTimeSpan() const        { return FCbFieldType::IsTimeSpan(Type); }
+
+	/** Whether the field has a value. \see \ref FCbField::HasValue() */
+	constexpr inline explicit operator bool() const { return HasValue(); }
+
+	/**
+	 * Whether the field has a value.
+	 *
+	 * All fields in a valid object or array have a value. A field with no value is returned when
+	 * finding a field by name fails or when accessing an iterator past the end.
+	 */
+	constexpr inline bool HasValue() const          { return !FCbFieldType::IsNone(Type); };
+
+	/** Whether the last field access encountered an error. */
+	constexpr inline bool HasError() const          { return Error != ECbFieldError::None; }
+	/** The type of error that occurred on the last field access, or None. */
+	constexpr inline ECbFieldError GetError() const { return Error; }
+
+protected:
+	constexpr inline ECbFieldType GetType() const { return Type; }
+	constexpr inline ECbFieldType GetCopyType() const { return Type; }
+	constexpr inline const void* GetPayload() const { return Payload; }
+	inline const void* GetPayloadEnd() const { return static_cast<const uint8*>(Payload) + PayloadSize(); }
+
+	inline void Assign(const void* InData, const ECbFieldType InType)
+	{
+		static_assert(TIsTriviallyDestructible<FCbField>::Value,
+			"This optimization requires FCbField to be trivially destructible!");
+		new(this) FCbField(InData, InType);
+	}
+
+	/**
+	 * Copy the field into a buffer of at least Size() bytes.
+	 *
+	 * The field type will only be copied if HasFieldType is set.
+	 */
+	CORE_API void CopyTo(FMutableMemoryView Buffer) const;
+
+	/** Create a view of the field, including the type and name when present. */
+	CORE_API FConstMemoryView AsMemoryView() const;
+
+private:
+	/** Size of the field payload in bytes. That is the field excluding its type and name. */
+	CORE_API uint64 PayloadSize() const;
+
+	/** Parameters for converting to an integer. */
+	struct FIntegerParams
+	{
+		/** Whether the output type has a sign bit. */
+		uint32 IsSigned : 1;
+		/** Bits of magnitude. (7 for int8) */
+		uint32 MagnitudeBits : 31;
+	};
+
+	/** Make integer params for the given integer type. */
+	template <typename IntType>
+	static constexpr inline FIntegerParams MakeIntegerParams()
+	{
+		FIntegerParams Params;
+		Params.IsSigned = IntType(-1) < IntType(0);
+		Params.MagnitudeBits = 8 * sizeof(IntType) - Params.IsSigned;
+		return Params;
+	}
+
+	/**
+	 * Access the field as the given integer type.
+	 *
+	 * Returns the provided default if the value cannot be represented in the output type.
+	 */
+	template <typename IntType>
+	inline IntType AsInteger(IntType Default)
+	{
+		return IntType(AsInteger(uint64(Default), MakeIntegerParams<IntType>()));
+	}
+
+	CORE_API uint64 AsInteger(uint64 Default, FIntegerParams Params);
+
+private:
+	/** The field type, with the transient HasFieldType flag if the field contains its type. */
+	ECbFieldType Type = ECbFieldType::None;
+	/** The error (if any) that occurred on the last field access. */
+	ECbFieldError Error = ECbFieldError::None;
+	/** The number of bytes for the name stored before the payload. */
+	uint32 NameLen = 0;
+	/** The value payload, which also points to the end of the name. */
+	const void* Payload = nullptr;
+};
+
+/** Iterator for \ref FCbField. \see \ref TCbFieldIterator */
+class FCbFieldIterator : public TCbFieldIterator<FCbField>
+{
+public:
+	using TCbFieldIterator::TCbFieldIterator;
+
+	FCbFieldIterator() = default;
+
+	inline FCbFieldIterator(TCbFieldIterator<FCbField>&& Field)
+		: TCbFieldIterator(MoveTemp(Field))
+	{
+	}
+};
+
+///////////////////////////////////////////////////////////////////////////////////////////////////////////////////////
+
+/**
+ * Array of \ref FCbField that have no names.
+ *
+ * Accessing a field of the array requires iteration. Access by index is not provided because the
+ * cost of accessing an item by index scales linearly with the index.
+ *
+ * This type only provides a view into memory and does not perform any memory management itself.
+ * Use \ref FCbArrayRef to hold a reference to the underlying memory when necessary.
+ */
+class FCbArray
+{
+public:
+	/** Construct an array with no fields. */
+	CORE_API FCbArray();
+
+	/**
+	 * Construct an array from a pointer to its data and an optional externally-provided type.
+	 *
+	 * @param Data Pointer to the start of the array data.
+	 * @param Type HasFieldType means that Data contains the type. Otherwise, use the given type.
+	 */
+	CORE_API explicit FCbArray(const void* Data, ECbFieldType Type = ECbFieldType::HasFieldType);
+
+	/** Size of the array in bytes if serialized by itself with no name. */
+	CORE_API uint64 Size() const;
+
+	/** Number of items in the array. */
+	CORE_API uint64 Num() const;
+
+	/** Create an iterator for the fields of this array. */
+	CORE_API FCbFieldIterator CreateIterator() const;
+
+	/**
+	 * Whether this array is identical to the other array.
+	 *
+	 * Performs a deep comparison of any contained arrays or objects and their fields. Comparison
+	 * assumes that both fields are valid and are written in the canonical format. Fields must be
+	 * written in the same order in arrays and objects, and name comparison is case sensitive. If
+	 * these assumptions do not hold, this may return false for equivalent inputs. Validation can
+	 * be done with the `All` mode to check these assumptions about the format of the inputs.
+	 */
+	CORE_API bool Equals(const FCbArray& Other) const;
+
+protected:
+	inline ECbFieldType GetCopyType() const { return FCbFieldType::GetType(Type) | ECbFieldType::HasFieldType; }
+
+	/** Copy the array into a buffer of at least Size() bytes. */
+	CORE_API void CopyTo(FMutableMemoryView Buffer) const;
+
+	/** Create a view of the payload, excluding the type or the name of the outer field. */
+	CORE_API FConstMemoryView AsMemoryView() const;
+
+private:
+	friend inline FCbFieldIterator begin(const FCbArray& Array) { return Array.CreateIterator(); }
+	friend inline FCbFieldIterator end(const FCbArray&) { return FCbFieldIterator(); }
+
+private:
+	const void* Payload;
+	ECbFieldType Type;
+};
+
+///////////////////////////////////////////////////////////////////////////////////////////////////////////////////////
+
+/**
+ * Array of \ref FCbField that have unique names.
+ *
+ * Accessing the fields of an object is always a safe operation, even if the requested field does
+ * not exist. Fields may be accessed by name or through iteration. When a field is requested that
+ * is not found in the object, the field that it returns has no value (evaluates to false) though
+ * attempting to access the empty field is also safe, as described by \ref FCbField.
+ *
+ * This type only provides a view into memory and does not perform any memory management itself.
+ * Use \ref FCbObjectRef to hold a reference to the underlying memory when necessary.
+ */
+class FCbObject
+{
+public:
+	/** Construct an object with no fields. */
+	CORE_API FCbObject();
+
+	/**
+	 * Construct an object from a pointer to its data and an optional externally-provided type.
+	 *
+	 * @param Data Pointer to the start of the object data.
+	 * @param Type HasFieldType means that Data contains the type. Otherwise, use the given type.
+	 */
+	CORE_API explicit FCbObject(const void* Data, ECbFieldType Type = ECbFieldType::HasFieldType);
+
+	/** Size of the object in bytes if serialized by itself with no name. */
+	CORE_API uint64 Size() const;
+
+	/** Create an iterator for the fields of this object. */
+	CORE_API FCbFieldIterator CreateIterator() const;
+
+	/**
+	 * Whether this object is identical to the other object.
+	 *
+	 * Performs a deep comparison of any contained arrays or objects and their fields. Comparison
+	 * assumes that both fields are valid and are written in the canonical format. Fields must be
+	 * written in the same order in arrays and objects, and name comparison is case sensitive. If
+	 * these assumptions do not hold, this may return false for equivalent inputs. Validation can
+	 * be done with the `All` mode to check these assumptions about the format of the inputs.
+	 */
+	CORE_API bool Equals(const FCbObject& Other) const;
+
+	/**
+	 * Find a field by case-sensitive name comparison.
+	 *
+	 * The cost of this operation scales linearly with the number of fields in the object. Prefer
+	 * to iterate over the fields only once when consuming an object.
+	 *
+	 * @param Name The name of the field.
+	 * @return The matching field if found, otherwise a field with no value.
+	 */
+	CORE_API FCbField Find(FAnsiStringView Name) const;
+
+	/** Find a field by case-insensitive name comparison. \see \ref FCbObject::Find */
+	CORE_API FCbField FindIgnoreCase(FAnsiStringView Name) const;
+
+	/** Find a field by case-sensitive name comparison. \see \ref FCbObject::Find */
+	inline FCbField operator[](FAnsiStringView Name) const { return Find(Name); }
+
+protected:
+	inline ECbFieldType GetCopyType() const { return FCbFieldType::GetType(Type) | ECbFieldType::HasFieldType; }
+
+	/** Copy the object into a buffer of at least Size() bytes. */
+	CORE_API void CopyTo(FMutableMemoryView Buffer) const;
+
+	/** Create a view of the payload, excluding the type or the name of the outer field. */
+	CORE_API FConstMemoryView AsMemoryView() const;
+
+private:
+	friend inline FCbFieldIterator begin(const FCbObject& Object) { return Object.CreateIterator(); }
+	friend inline FCbFieldIterator end(const FCbObject&) { return FCbFieldIterator(); }
+
+private:
+	const void* Payload;
+	ECbFieldType Type;
+};
+
+///////////////////////////////////////////////////////////////////////////////////////////////////////////////////////
+
+/**
+ * A wrapper that can hold a reference to the memory that contains its value.
+ *
+ * When constructed with an owned buffer, or by assuming ownership, or by cloning, this reference
+ * type will prevent the memory containing its value from being released. To minimize overhead, a
+ * buffer pointer is only stored for owned buffers.
+ */
+template <typename BaseType>
+class TCbBufferRef : public BaseType
+{
+public:
+	enum EAssumeOwnershipTag { AssumeOwnership };
+	enum ECloneTag           { Clone };
+	enum EWrapTag            { Wrap };
+
+	/** Construct a default value. */
+	TCbBufferRef() = default;
+
+	/** Construct a value and hold a reference to the buffer that contains it. */
+	inline explicit TCbBufferRef(const FSharedBufferConstRef& ValueBuffer)
+		: BaseType(ValueBuffer->GetData())
+	{
+		if (ValueBuffer->IsOwned())
+		{
+			BufferPtr = ValueBuffer;
+		}
+	}
+
+	/** Construct a value and hold a reference to the buffer that contains it. */
+	inline explicit TCbBufferRef(const FSharedBufferConstPtr& ValueBuffer)
+	{
+		if (ValueBuffer)
+		{
+			static_cast<BaseType&>(*this) = BaseType(ValueBuffer->GetData());
+			if (ValueBuffer->IsOwned())
+			{
+				BufferPtr = ValueBuffer;
+			}
+		}
+	}
+
+	/** Construct a value and hold a reference to the buffer that contains it. */
+	inline explicit TCbBufferRef(FSharedBufferConstPtr&& ValueBuffer)
+	{
+		if (ValueBuffer)
+		{
+			static_cast<BaseType&>(*this) = BaseType(ValueBuffer->GetData());
+			if (ValueBuffer->IsOwned())
+			{
+				BufferPtr = MoveTemp(ValueBuffer);
+			}
+		}
+	}
+
+	/** Construct a value and hold a reference to the buffer that contains it. */
+	inline explicit TCbBufferRef(const FSharedBufferRef& ValueBuffer)
+		: TCbBufferRef(FSharedBufferConstRef(ValueBuffer))
+	{
+	}
+
+	/** Construct a value that holds a reference to the buffer that contains it. */
+	inline TCbBufferRef(const BaseType& Value, const FSharedBufferConstPtr& ValueBuffer)
+		: BaseType(Value)
+	{
+		if (ValueBuffer)
+		{
+			check(ValueBuffer->GetView().Contains(BaseType::AsMemoryView()));
+			if (ValueBuffer->IsOwned())
+			{
+				BufferPtr = ValueBuffer;
+			}
+		}
+	}
+
+	/** Construct a value that holds a reference to the buffer that contains it. */
+	inline TCbBufferRef(const BaseType& Value, FSharedBufferConstPtr&& ValueBuffer)
+		: BaseType(Value)
+	{
+		if (ValueBuffer)
+		{
+			check(ValueBuffer->GetView().Contains(BaseType::AsMemoryView()));
+			if (ValueBuffer->IsOwned())
+			{
+				BufferPtr = MoveTemp(ValueBuffer);
+			}
+		}
+	}
+
+	/** Construct a value that holds a reference to the buffer of the outer that contains it. */
+	template <typename OtherBaseType>
+	inline TCbBufferRef(const BaseType& Value, const TCbBufferRef<OtherBaseType>& OuterRef)
+		: TCbBufferRef(Value, OuterRef.BufferPtr)
+	{
+	}
+
+	/** Construct a value that holds a reference to the buffer of the outer that contains it. */
+	template <typename OtherBaseType>
+	inline TCbBufferRef(const BaseType& Value, TCbBufferRef<OtherBaseType>&& OuterRef)
+		: TCbBufferRef(Value, MoveTemp(OuterRef.BufferPtr))
+	{
+	}
+
+	/** Construct a value and take ownership its memory. */
+	inline TCbBufferRef(EAssumeOwnershipTag, const void* Value)
+		: BaseType(Value)
+		, BufferPtr(MakeSharedBuffer(FSharedBuffer::AssumeOwnership, Value, BaseType::Size()))
+	{
+	}
+
+	/** Construct a value from an owned clone of its memory. */
+	inline TCbBufferRef(ECloneTag, const void* Value)
+		: TCbBufferRef(Clone, BaseType(Value))
+	{
+	}
+
+	/** Construct a value from an owned clone of its memory. */
+	inline TCbBufferRef(ECloneTag, const BaseType& Value)
+		: BaseType(Value)
+	{
+		const FSharedBufferRef BufferRef = MakeSharedBuffer(BaseType::Size());
+		BaseType::CopyTo(*BufferRef);
+		BufferPtr = BufferRef;
+		static_cast<BaseType&>(*this) = BaseType(BufferPtr->GetData(), BaseType::GetCopyType());
+	}
+
+	/** Construct a value by wrapping its memory without taking ownership of it. */
+	inline TCbBufferRef(EWrapTag, const void* Value)
+		: BaseType(Value)
+	{
+	}
+
+	/** Construct a value by wrapping its memory without taking ownership of it. */
+	inline TCbBufferRef(EWrapTag, const BaseType& Value)
+		: BaseType(Value)
+	{
+	}
+
+	/** Whether this reference has ownership of the memory in its buffer. */
+	inline bool IsOwned() const { return BufferPtr && BufferPtr->IsOwned(); }
+
+	/** Clone the value, if necessary, to a buffer that this reference has ownership of. */
+	inline void MakeOwned()
+	{
+		if (!IsOwned())
+		{
+			*this = TCbBufferRef(Clone, ImplicitConv<BaseType>(*this));
+		}
+	}
+
+	/** The buffer (if any) that contains this value. */
+	inline const FSharedBufferConstPtr& Buffer() const { return BufferPtr; }
+
+private:
+	template <typename OtherType>
+	friend class TCbBufferRef;
+
+	FSharedBufferConstPtr BufferPtr;
+};
+
+///////////////////////////////////////////////////////////////////////////////////////////////////////////////////////
+
+class FCbArrayRef;
+class FCbObjectRef;
+
+///////////////////////////////////////////////////////////////////////////////////////////////////////////////////////
+
+/** A field that can hold a reference to the memory that contains it. \see \ref TCbBufferRef */
+class FCbFieldRef : public TCbBufferRef<FCbField>
+{
+public:
+	using TCbBufferRef::TCbBufferRef;
+
+	FCbFieldRef() = default;
+
+	inline FCbFieldRef(TCbBufferRef<FCbField>&& Ref)
+		: TCbBufferRef(MoveTemp(Ref))
+	{
+	}
+
+	/** Access the field as an object. Defaults to an empty object on error. */
+	inline FCbObjectRef AsObjectRef();
+
+	/** Access the field as an array. Defaults to an empty array on error. */
+	inline FCbArrayRef AsArrayRef();
+};
+
+/** Iterator for \ref FCbFieldRef. \see \ref TCbFieldIterator */
+class FCbFieldRefIterator : public TCbFieldIterator<FCbFieldRef>
+{
+public:
+	using TCbFieldIterator::TCbFieldIterator;
+
+	FCbFieldRefIterator() = default;
+
+	inline FCbFieldRefIterator(TCbFieldIterator<FCbFieldRef>&& Field)
+		: TCbFieldIterator(MoveTemp(Field))
+	{
+	}
+};
+
+///////////////////////////////////////////////////////////////////////////////////////////////////////////////////////
+
+/** An array that can hold a reference to the memory that contains it. \see \ref TCbBufferRef */
+class FCbArrayRef : public TCbBufferRef<FCbArray>
+{
+public:
+	using TCbBufferRef::TCbBufferRef;
+
+	FCbArrayRef() = default;
+
+	inline FCbArrayRef(TCbBufferRef<FCbArray>&& Ref)
+		: TCbBufferRef(MoveTemp(Ref))
+	{
+	}
+
+	/** Create an iterator for the fields of this array. */
+	inline FCbFieldRefIterator CreateRefIterator() const
+	{
+		FCbFieldIterator It = CreateIterator();
+		return FCbFieldRefIterator(FCbFieldRef(It, Buffer()), It);
+	}
+};
+
+///////////////////////////////////////////////////////////////////////////////////////////////////////////////////////
+
+/** An object that can hold a reference to the memory that contains it. \see \ref TCbBufferRef */
+class FCbObjectRef : public TCbBufferRef<FCbObject>
+{
+public:
+	using TCbBufferRef::TCbBufferRef;
+
+	FCbObjectRef() = default;
+
+	inline FCbObjectRef(TCbBufferRef<FCbObject>&& Ref)
+		: TCbBufferRef(MoveTemp(Ref))
+	{
+	}
+
+	/** Create an iterator for the fields of this object. */
+	inline FCbFieldRefIterator CreateRefIterator() const
+	{
+		FCbFieldIterator It = CreateIterator();
+		return FCbFieldRefIterator(FCbFieldRef(It, Buffer()), It);
+	}
+
+	/** Find a field by case-sensitive name comparison. */
+	inline FCbFieldRef FindRef(FAnsiStringView Name) const
+	{
+		if (FCbField Field = Find(Name))
+		{
+			return FCbFieldRef(Field, *this);
+		}
+		return FCbFieldRef();
+	}
+
+	/** Find a field by case-insensitive name comparison. */
+	inline FCbFieldRef FindRefIgnoreCase(FAnsiStringView Name) const
+	{
+		if (FCbField Field = FindIgnoreCase(Name))
+		{
+			return FCbFieldRef(Field, *this);
+		}
+		return FCbFieldRef();
+	}
+};
+
+///////////////////////////////////////////////////////////////////////////////////////////////////////////////////////
+
+inline FCbObjectRef FCbFieldRef::AsObjectRef()
+{
+	return IsObject() ? FCbObjectRef(AsObject(), *this) : FCbObjectRef();
+}
+
+inline FCbArrayRef FCbFieldRef::AsArrayRef()
+{
+	return IsArray() ? FCbArrayRef(AsArray(), *this) : FCbArrayRef();
+}
+
+///////////////////////////////////////////////////////////////////////////////////////////////////////////////////////
+
+/** Flags for validating compact binary data. */
+enum class ECbValidateMode : uint32
+{
+	/**
+	 * Validate that the value can be read and stays inside the bounds of the memory view.
+	 *
+	 * This is the minimum level of validation required to be able to safely read a field, array,
+	 * or object without the risk of crashing or reading out of bounds.
+	 */
+	Default                 = 0,
+
+	/**
+	 * Validate that object fields have unique non-empty names and array fields have no names.
+	 *
+	 * Name validation failures typically do not inhibit reading the input, but duplicated fields
+	 * cannot be looked up by name other than the first, and converting to other data formats can
+	 * fail in the presence of naming issues.
+	 */
+	Names                   = 1 << 0,
+
+	/**
+	 * Validate that fields are serialized in the canonical format.
+	 *
+	 * Format validation failures typically do not inhibit reading the input. Values that fail in
+	 * this mode require more memory than in the canonical format, and comparisons of such values
+	 * for equality are not reliable. Examples of failures include uniform arrays or objects that
+	 * were not encoded uniformly, variable-length integers that could be encoded in fewer bytes,
+	 * or 64-bit floats that could be encoded in 32 bits without loss of precision.
+	 */
+	Format                  = 1 << 1,
+
+	/**
+	 * Validate that there is no padding after the value before the end of the memory view.
+	 *
+	 * Padding validation failures have no impact on the ability to read the input, but are using
+	 * more memory than necessary.
+	 */
+	Padding                 = 1 << 2,
+
+	/** Perform all validation described above. */
+	All                     = Default | Names | Format | Padding,
+};
+
+/** Flags for compact binary validation errors. Multiple flags may be combined. */
+enum class ECbValidateError : uint32
+{
+	/** The input had no validation errors. */
+	None                    = 0,
+
+	// Mode: Default
+
+	/** The input cannot be read without reading out of bounds. */
+	OutOfBounds             = 1 << 0,
+	/** The input has a field with an unrecognized or invalid type. */
+	InvalidType             = 1 << 1,
+
+	// Mode: Names
+
+	/** An object had more than one field with the same name. */
+	DuplicateName           = 1 << 2,
+	/** An object had a field with no name. */
+	MissingName             = 1 << 3,
+	/** An array field had a name. */
+	ArrayName               = 1 << 4,
+
+	// Mode: Format
+
+	/** A name or string payload is not valid UTF-8. */
+	InvalidString           = 1 << 5,
+	/** A size or integer payload can be encoded in fewer bytes. */
+	InvalidInteger          = 1 << 6,
+	/** A float64 payload can be encoded as a float32 without loss of precision. */
+	InvalidFloat            = 1 << 7,
+	/** An object has the same type for every field but is not uniform. */
+	NonUniformObject        = 1 << 8,
+	/** An array has the same type for every field and non-empty payloads but is not uniform. */
+	NonUniformArray         = 1 << 9,
+
+	// Mode: Padding
+
+	/** A value did not use the entire memory view given for validation. */
+	Padding                 = 1 << 10,
+};
+
+ENUM_CLASS_FLAGS(ECbValidateMode);
+ENUM_CLASS_FLAGS(ECbValidateError);
+
+/**
+ * Validate the compact binary data for one field in the view as specified by the mode flags.
+ *
+ * Only one top-level field is processed from the view, and validation recurses into any array or
+ * object within that field. To validate multiple consecutive top-level fields, call the function
+ * once for each top-level field. If the given view might contain multiple top-level fields, then
+ * either exclude the Padding flag from the Mode or use MeasureCompactBinary to break up the view
+ * into its constituent fields before validating.
+ *
+ * @param View A memory view containing at least one top-level field.
+ * @param Mode A combination of the flags for the types of validation to perform.
+ * @param Type HasFieldType means that View contains the type. Otherwise, use the given type.
+ * @return None on success, otherwise the flags for the types of errors that were detected.
+ */
+CORE_API ECbValidateError ValidateCompactBinary(FConstMemoryView View, ECbValidateMode Mode, ECbFieldType Type = ECbFieldType::HasFieldType);
+
+/**
+ * Validate the compact binary data for every field in the view as specified by the mode flags.
+ *
+ * This function expects the entire view to contain fields. Any trailing region of the view which
+ * does not contain a valid field will produce an OutOfBounds or InvalidType error instead of the
+ * Padding error that would be produced by the single field validation function.
+ *
+ * \see \ref ValidateCompactBinary for more detail.
+ */
+CORE_API ECbValidateError ValidateCompactBinaryRange(FConstMemoryView View, ECbValidateMode Mode);
+
+///////////////////////////////////////////////////////////////////////////////////////////////////////////////////////
+
+/**
+ * Determine the size in bytes of a one compact binary field at the start of the view.
+ *
+ * This may be called on an incomplete or invalid field, in which case the returned size is zero.
+ * A size can always be extracted from a valid field with no name if a view of at least the first
+ * 10 bytes is provided, regardless of field size. For fields with names, the size of view needed
+ * to calculate a size is at most 10 + MaxNameLen + MeasureVarUInt(MaxNameLen).
+ *
+ * This function can be used when streaming a field, for example, to determine the size of buffer
+ * to fill before attempting to construct a field from it.
+ *
+ * @param View A memory view that may contain the start of a field.
+ * @param Type HasFieldType means that View contains the type. Otherwise, use the given type.
+ */
+CORE_API uint64 MeasureCompactBinary(FConstMemoryView View, ECbFieldType Type = ECbFieldType::HasFieldType);
+
+///////////////////////////////////////////////////////////////////////////////////////////////////////////////////////
diff --git a/Engine/Source/Runtime/Core/Public/Serialization/CompactBinaryWriter.h b/Engine/Source/Runtime/Core/Public/Serialization/CompactBinaryWriter.h
new file mode 100644
--- /dev/null
+++ b/Engine/Source/Runtime/Core/Public/Serialization/CompactBinaryWriter.h
@@ -0,0 +1,332 @@
+// Copyright Epic Games, Inc. All Rights Reserved.
+
+#pragma once
+
+#include "CoreTypes.h"
+#include "Containers/StringView.h"
+#include "Serialization/CompactBinary.h"
+
+///////////////////////////////////////////////////////////////////////////////////////////////////////////////////////
+
+struct FBlake3Hash;
+struct FDateTime;
+struct FGuid;
+struct FTimespan;
+
+/**
+ * A writer for compact binary object, arrays, and fields.
+ *
+ * The writer produces a sequence of fields that can be saved to a provided memory buffer or into
+ * a new owned buffer. The typical use case is to write a single object, which can be accessed by
+ * calling Save().AsObjectRef() or SaveToMemory(Buffer).AsObject().
+ *
+ * The writer will assert on most incorrect usage and will always produce valid compact binary if
+ * provided with valid input. The writer does not check for invalid UTF-8 string encoding, object
+ * fields with duplicate names, or invalid compact binary being copied from another source.
+ *
+ * It is most convenient to use the streaming API for the writer, as demonstrated in the example.
+ *
+ * When writing a small amount of compact binary data, \see \ref TCbWriter can be more efficient,
+ * as it uses a fixed-size stack buffer for storage before spilling onto the heap.
+ *
+ * Example:
+ *
+ * \code
+ * FCbObjectRef WriteObject()
+ * {
+ *     TCbWriter<256> Writer;
+ *     Writer.BeginObject();
+ * 
+ *     Writer << "Resize" << true;
+ *     Writer << "MaxWidth" << 1024;
+ *     Writer << "MaxHeight" << 1024;
+ * 
+ *     Writer.BeginArray();
+ *     Writer << "FormatA" << "FormatB" << "FormatC";
+ *     Writer.EndArray();
+ * 
+ *     Writer.EndObject();
+ *     return Writer.Save().AsObjectRef();
+ * }
+ * \endcode
+ */
+class FCbWriter
+{
+public:
+	CORE_API FCbWriter();
+	CORE_API ~FCbWriter();
+
+	FCbWriter(const FCbWriter&) = delete;
+	FCbWriter& operator=(const FCbWriter&) = delete;
+
+	/** Empty the writer without releasing any allocated memory. */
+	CORE_API void Reset();
+
+	/**
+	 * Serialize the field(s) to an owned buffer and return it as an iterator.
+	 *
+	 * It is not valid to call this function in the middle of writing an object, array, or field.
+	 * The writer remains valid for further use when this function returns.
+	 */
+	CORE_API FCbFieldRefIterator Save() const;
+
+	/**
+	 * Serialize the field(s) to memory.
+	 *
+	 * It is not valid to call this function in the middle of writing an object, array, or field.
+	 * The writer remains valid for further use when this function returns.
+	 *
+	 * @param Buffer A mutable memory view to write to. Must be exactly SaveSize() bytes.
+	 * @return An iterator for the field(s) written to the buffer.
+	 */
+	CORE_API FCbFieldIterator SaveToMemory(FMutableMemoryView Buffer) const;
+
+	/**
+	 * The size of buffer (in bytes) required to serialize the fields that have been written.
+	 *
+	 * It is not valid to call this function in the middle of writing an object, array, or field.
+	 */
+	CORE_API uint64 SaveSize() const;
+
+	/**
+	 * Write a field name.
+	 *
+	 * It is not valid to call this function when writing an array field.
+	 * Names must be valid UTF-8 and must be unique within an object.
+	 */
+	CORE_API FCbWriter& Name(FAnsiStringView Name);
+
+	/** Copy the value (not the name) of an existing field. */
+	CORE_API void Field(const FCbField& Value);
+
+	/** Begin a new object. Must have a matching call to \ref EndObject. */
+	CORE_API void BeginObject();
+	/** End an object after its fields have been written. */
+	CORE_API void EndObject();
+
+	/** Copy the value (not the name) of an existing object. */
+	CORE_API void Object(const FCbObject& Value);
+
+	/** Begin a new array. Must have a matching call to \ref EndArray. */
+	CORE_API void BeginArray();
+	/** End an array after its fields have been written. */
+	CORE_API void EndArray();
+
+	/** Copy the value (not the name) of an existing array. */
+	CORE_API void Array(const FCbArray& Value);
+
+	/** Write a null field. */
+	CORE_API void Null();
+
+	/** Write a binary field by copying Size bytes from Value. */
+	CORE_API void Binary(const void* Value, uint64 Size);
+	/** Write a binary field by copying the view. */
+	inline void Binary(FConstMemoryView Value) { Binary(Value.GetData(), Value.Size()); }
+
+	/** Write a string field by copying the UTF-8 value. */
+	CORE_API void String(FAnsiStringView Value);
+	/** Write a string field by converting the UTF-16 value to UTF-8. */
+	CORE_API void String(FWideStringView Value);
+
+	/** Write an integer field. */
+	CORE_API void Integer(int32 Value);
+	/** Write an integer field. */
+	CORE_API void Integer(int64 Value);
+	/** Write an integer field. */
+	CORE_API void Integer(uint32 Value);
+	/** Write an integer field. */
+	CORE_API void Integer(uint64 Value);
+
+	/** Write a float field from a 32-bit float value. */
+	CORE_API void Float(float Value);
+	/** Write a float field from a 64-bit float value. */
+	CORE_API void Float(double Value);
+
+	/** Write a bool field. */
+	CORE_API void Bool(bool bValue);
+
+	/** Write a hash field referencing a blob. */
+	CORE_API void BinaryHash(const FBlake3Hash& Value);
+	/** Write a hash field referencing a field/array/object. */
+	CORE_API void FieldHash(const FBlake3Hash& Value);
+
+	/** Write a UUID field. */
+	CORE_API void Uuid(const FGuid& Value);
+
+	/** Write a date/time field with the specified count of 100ns ticks since the epoch. */
+	CORE_API void DateTimeTicks(int64 Ticks);
+
+	/** Write a date/time field. */
+	CORE_API void DateTime(FDateTime Value);
+
+	/** Write a time span field with the specified count of 100ns ticks. */
+	CORE_API void TimeSpanTicks(int64 Ticks);
+
+	/** Write a time span field. */
+	CORE_API void TimeSpan(FTimespan Value);
+
+	/** Private flags that are public to work with ENUM_CLASS_FLAGS. */
+	enum class EStateFlags : uint8;
+
+protected:
+	/** Reserve the specified size up front until the format is optimized. */
+	CORE_API explicit FCbWriter(int64 InitialSize);
+
+private:
+	friend FCbWriter& operator<<(FCbWriter& Writer, FAnsiStringView NameOrValue);
+
+	/** Begin writing a field. Called twice for fields with names and once otherwise. */
+	void BeginField();
+
+	/** Finish writing a field by writing its type. */
+	void EndField(ECbFieldType Type);
+
+	/** Write the field name if valid in this state, otherwise write the string value. */
+	CORE_API void NameOrString(FAnsiStringView NameOrValue);
+
+	/** A view of the name of the current field (if any), otherwise the empty view. */
+	FAnsiStringView CurrentName() const;
+
+	/** Remove field types after the first to make the sequence uniform. */
+	void MakeFieldsUniform(int64 FieldBeginOffset, int64 FieldEndOffset);
+
+	/** State of the object, array, or top-level field being written. */
+	struct FState
+	{
+		EStateFlags Flags{};
+		/** The type of the fields in the sequence if uniform, otherwise None. */
+		ECbFieldType UniformType{};
+		/** The offset of the start of the current field. */
+		int64 Offset{};
+		/** The number of fields written in this state. */
+		uint64 Count{};
+	};
+
+private:
+	// This is a prototype-quality format for the writer. Using an array of bytes is inefficient,
+	// and will lead to many unnecessary copies and moves of the data to resize the array, insert
+	// object and array sizes, and remove field types for uniform objects and uniform arrays. The
+	// optimized format will be a list of power-of-two blocks and an optional first block that is
+	// provided externally, such as on the stack. That format will store the offsets that require
+	// object or array sizes to be inserted and field types to be removed, and will perform those
+	// operations only when saving to a buffer.
+	TArray64<uint8> Data;
+	TArray<FState> States;
+};
+
+///////////////////////////////////////////////////////////////////////////////////////////////////////////////////////
+
+/**
+ * A writer for compact binary object, arrays, and fields that uses a fixed-size stack buffer.
+ *
+ * \see \ref FCbWriter
+ */
+template <uint32 InlineBufferSize>
+class TCbWriter : public FCbWriter
+{
+public:
+	inline TCbWriter()
+		: FCbWriter(InlineBufferSize)
+	{
+	}
+
+	TCbWriter(const TCbWriter&) = delete;
+	TCbWriter& operator=(const TCbWriter&) = delete;
+
+private:
+	// Reserve the inline buffer now even though we are unable to use it. This will avoid causing
+	// new stack overflows when this functionality is properly implemented in the future.
+	uint8 Buffer[InlineBufferSize];
+};
+
+///////////////////////////////////////////////////////////////////////////////////////////////////////////////////////
+
+/** Write the field name if valid in this state, otherwise write the string value. */
+inline FCbWriter& operator<<(FCbWriter& Writer, FAnsiStringView NameOrValue)
+{
+	Writer.NameOrString(NameOrValue);
+	return Writer;
+}
+
+inline FCbWriter& operator<<(FCbWriter& Writer, const FCbField& Value)
+{
+	Writer.Field(Value);
+	return Writer;
+}
+
+inline FCbWriter& operator<<(FCbWriter& Writer, const FCbObject& Value)
+{
+	Writer.Object(Value);
+	return Writer;
+}
+
+inline FCbWriter& operator<<(FCbWriter& Writer, const FCbArray& Value)
+{
+	Writer.Array(Value);
+	return Writer;
+}
+
+inline FCbWriter& operator<<(FCbWriter& Writer, nullptr_t)
+{
+	Writer.Null();
+	return Writer;
+}
+
+inline FCbWriter& operator<<(FCbWriter& Writer, FWideStringView Value)
+{
+	Writer.String(Value);
+	return Writer;
+}
+
+inline FCbWriter& operator<<(FCbWriter& Writer, int32 Value)
+{
+	Writer.Integer(Value);
+	return Writer;
+}
+
+inline FCbWriter& operator<<(FCbWriter& Writer, int64 Value)
+{
+	Writer.Integer(Value);
+	return Writer;
+}
+
+inline FCbWriter& operator<<(FCbWriter& Writer, uint32 Value)
+{
+	Writer.Integer(Value);
+	return Writer;
+}
+
+inline FCbWriter& operator<<(FCbWriter& Writer, uint64 Value)
+{
+	Writer.Integer(Value);
+	return Writer;
+}
+
+inline FCbWriter& operator<<(FCbWriter& Writer, float Value)
+{
+	Writer.Float(Value);
+	return Writer;
+}
+
+inline FCbWriter& operator<<(FCbWriter& Writer, double Value)
+{
+	Writer.Float(Value);
+	return Writer;
+}
+
+inline FCbWriter& operator<<(FCbWriter& Writer, bool Value)
+{
+	Writer.Bool(Value);
+	return Writer;
+}
+
+inline FCbWriter& operator<<(FCbWriter& Writer, const FGuid& Value)
+{
+	Writer.Uuid(Value);
+	return Writer;
+}
+
+CORE_API FCbWriter& operator<<(FCbWriter& Writer, FDateTime Value);
+CORE_API FCbWriter& operator<<(FCbWriter& Writer, FTimespan Value);
+
+///////////////////////////////////////////////////////////////////////////////////////////////////////////////////////
diff --git a/Engine/Source/Runtime/Core/Public/Serialization/Formatters/BinaryArchiveFormatter.h b/Engine/Source/Runtime/Core/Public/Serialization/Formatters/BinaryArchiveFormatter.h
--- a/Engine/Source/Runtime/Core/Public/Serialization/Formatters/BinaryArchiveFormatter.h
+++ b/Engine/Source/Runtime/Core/Public/Serialization/Formatters/BinaryArchiveFormatter.h
@@ -69,6 +69,7 @@ public:
 	virtual void Serialize(FSoftObjectPtr& Value) override;
 	virtual void Serialize(FSoftObjectPath& Value) override;
 	virtual void Serialize(FLazyObjectPtr& Value) override;
+	virtual void Serialize(FObjectPtr& Value) override;
 	virtual void Serialize(TArray<uint8>& Value) override;
 	virtual void Serialize(void* Data, uint64 DataSize) override;
 
@@ -291,6 +292,11 @@ inline void FBinaryArchiveFormatter::Serialize(FLazyObjectPtr& Value)
 	Inner << Value;
 }
 
+inline void FBinaryArchiveFormatter::Serialize(FObjectPtr& Value)
+{
+	Inner << Value;
+}
+
 inline void FBinaryArchiveFormatter::Serialize(TArray<uint8>& Data)
 {
 	Inner << Data;
diff --git a/Engine/Source/Runtime/Core/Public/Serialization/StructuredArchive.h b/Engine/Source/Runtime/Core/Public/Serialization/StructuredArchive.h
--- a/Engine/Source/Runtime/Core/Public/Serialization/StructuredArchive.h
+++ b/Engine/Source/Runtime/Core/Public/Serialization/StructuredArchive.h
@@ -261,6 +261,7 @@ public:
 	void operator << (FSoftObjectPtr& Value);
 	void operator << (FSoftObjectPath& Value);
 	void operator << (FLazyObjectPtr& Value);
+	void operator << (FObjectPtr& Value);
 
 	template <typename T>
 	FORCEINLINE void operator<<(TEnumAsByte<T>& Value)
@@ -963,6 +964,11 @@ typename TEnableIf<
 		Ar.Formatter.Serialize(Value);
 	}
 
+	FORCEINLINE void FStructuredArchiveSlot::operator<< (FObjectPtr& Value)
+	{
+		Ar.Formatter.Serialize(Value);
+	}
+
 	FORCEINLINE void FStructuredArchiveSlot::Serialize(TArray<uint8>& Data)
 	{
 		Ar.Formatter.Serialize(Data);
diff --git a/Engine/Source/Runtime/Core/Public/Serialization/StructuredArchiveFormatter.h b/Engine/Source/Runtime/Core/Public/Serialization/StructuredArchiveFormatter.h
--- a/Engine/Source/Runtime/Core/Public/Serialization/StructuredArchiveFormatter.h
+++ b/Engine/Source/Runtime/Core/Public/Serialization/StructuredArchiveFormatter.h
@@ -129,6 +129,7 @@ public:
 	virtual void Serialize(struct FSoftObjectPtr& Value) = 0;
 	virtual void Serialize(struct FSoftObjectPath& Value) = 0;
 	virtual void Serialize(struct FLazyObjectPtr& Value) = 0;
+	virtual void Serialize(struct FObjectPtr& Value) = 0;
 	virtual void Serialize(TArray<uint8>& Value) = 0;
 	virtual void Serialize(void* Data, uint64 DataSize) = 0;
 };
diff --git a/Engine/Source/Runtime/Core/Public/Templates/UnrealTemplate.h b/Engine/Source/Runtime/Core/Public/Templates/UnrealTemplate.h
--- a/Engine/Source/Runtime/Core/Public/Templates/UnrealTemplate.h
+++ b/Engine/Source/Runtime/Core/Public/Templates/UnrealTemplate.h
@@ -71,18 +71,18 @@ FORCEINLINE void Move(T& A,typename TMoveSupportTraits<T>::Move B)
  * Generically gets the data pointer of a contiguous container
  */
 template<typename T, typename = typename TEnableIf<TIsContiguousContainer<T>::Value>::Type>
-auto GetData(T&& Container) -> decltype(Container.GetData())
+constexpr auto GetData(T&& Container) -> decltype(Container.GetData())
 {
 	return Container.GetData();
 }
 
-template <typename T, SIZE_T N> CONSTEXPR       T* GetData(      T (& Container)[N]) { return Container; }
-template <typename T, SIZE_T N> CONSTEXPR       T* GetData(      T (&&Container)[N]) { return Container; }
-template <typename T, SIZE_T N> CONSTEXPR const T* GetData(const T (& Container)[N]) { return Container; }
-template <typename T, SIZE_T N> CONSTEXPR const T* GetData(const T (&&Container)[N]) { return Container; }
+template <typename T, SIZE_T N> constexpr T* GetData(      T (& Container)[N]) { return Container; }
+template <typename T, SIZE_T N> constexpr T* GetData(      T (&&Container)[N]) { return Container; }
+template <typename T, SIZE_T N> constexpr const T* GetData(const T (& Container)[N]) { return Container; }
+template <typename T, SIZE_T N> constexpr const T* GetData(const T (&&Container)[N]) { return Container; }
 
 template <typename T>
-CONSTEXPR const T* GetData(std::initializer_list<T> List)
+constexpr const T* GetData(std::initializer_list<T> List)
 {
 	return List.begin();
 }
@@ -91,15 +91,15 @@ CONSTEXPR const T* GetData(std::initializer_list<T> List)
 * Generically gets the number of items in a contiguous container
 */
 template<typename T, typename = typename TEnableIf<TIsContiguousContainer<T>::Value>::Type>
-auto GetNum(T&& Container) -> decltype(Container.Num())
+constexpr auto GetNum(T&& Container) -> decltype(Container.Num())
 {
 	return Container.Num();
 }
 
-template <typename T, SIZE_T N> CONSTEXPR SIZE_T GetNum(      T (& Container)[N]) { return N; }
-template <typename T, SIZE_T N> CONSTEXPR SIZE_T GetNum(      T (&&Container)[N]) { return N; }
-template <typename T, SIZE_T N> CONSTEXPR SIZE_T GetNum(const T (& Container)[N]) { return N; }
-template <typename T, SIZE_T N> CONSTEXPR SIZE_T GetNum(const T (&&Container)[N]) { return N; }
+template <typename T, SIZE_T N> constexpr SIZE_T GetNum(      T (& Container)[N]) { return N; }
+template <typename T, SIZE_T N> constexpr SIZE_T GetNum(      T (&&Container)[N]) { return N; }
+template <typename T, SIZE_T N> constexpr SIZE_T GetNum(const T (& Container)[N]) { return N; }
+template <typename T, SIZE_T N> constexpr SIZE_T GetNum(const T (&&Container)[N]) { return N; }
 
 /**
  * Gets the number of items in an initializer list.
@@ -108,7 +108,7 @@ template <typename T, SIZE_T N> CONSTEXPR SIZE_T GetNum(const T (&&Container)[N]
  * Realistically, an initializer list should not exceed the limits of int32.
  */
 template <typename T>
-CONSTEXPR int32 GetNum(std::initializer_list<T> List)
+constexpr int32 GetNum(std::initializer_list<T> List)
 {
 	return static_cast<int32>(List.size());
 }
diff --git a/Engine/Source/Runtime/Core/Public/UObject/NameTypes.h b/Engine/Source/Runtime/Core/Public/UObject/NameTypes.h
--- a/Engine/Source/Runtime/Core/Public/UObject/NameTypes.h
+++ b/Engine/Source/Runtime/Core/Public/UObject/NameTypes.h
@@ -352,6 +352,15 @@ struct FMinimalName
 		return !Index && Number == NAME_NO_NUMBER_INTERNAL;
 	}
 
+	friend FORCEINLINE FArchive& operator<<(FArchive& Ar, FMinimalName& E)
+	{
+		Ar << E.Index;
+		Ar << E.Number;
+		return Ar;
+	}
+
+	FORCEINLINE bool operator<(FMinimalName Rhs) const { return Index == Rhs.Index ? Number < Rhs.Number : Index < Rhs.Index; };
+
 	/** Index into the Names array (used to find String portion of the string/number pair) */
 	FNameEntryId	Index;
 	/** Number portion of the string/number pair (stored internally as 1 more than actual, so zero'd memory will be the default, no-instance case) */
diff --git a/Engine/Source/Runtime/Core/Public/UObject/UE5CookerObjectVersion.h b/Engine/Source/Runtime/Core/Public/UObject/UE5CookerObjectVersion.h
new file mode 100644
--- /dev/null
+++ b/Engine/Source/Runtime/Core/Public/UObject/UE5CookerObjectVersion.h
@@ -0,0 +1,40 @@
+// Copyright Epic Games, Inc. All Rights Reserved.
+#pragma once
+
+#include "CoreTypes.h"
+#include "Misc/Guid.h"
+
+// TODO: Remove this once we are sure we can run with FVirtualizedUntypedBulkData with no issues
+/** When set to 1 we use FVirtualizedUntypedBulkData, when 0 we use what ever code was there before. */
+#define UE_USE_VIRTUALBULKDATA 0
+
+// Custom serialization version for changes made in the UE5 Dev-Cooker stream
+struct CORE_API FUE5CookerObjectVersion
+{
+	enum Type
+	{
+		// Before any version changes were made
+		BeforeCustomVersionWasAdded = 0,
+		
+		// -----<new versions can be added above this line>-------------------------------------------------
+
+		// Virtualization version numbers are added at this point. 
+		// When UE_USE_VIRTUALBULKDATA is removed then they can be moved to the normal section
+
+#if UE_USE_VIRTUALBULKDATA
+		
+		// Switch FMeshDescriptionBulkData to use virtualized bulkdata
+		MeshDescriptionVirtualization,
+
+#endif //UE_USE_VIRTUALBULKDATA
+
+		VersionPlusOne,
+		LatestVersion = VersionPlusOne - 1
+	};
+
+	// The GUID for this custom version number
+	const static FGuid GUID;
+
+private:
+	FUE5CookerObjectVersion() {}
+};
diff --git a/Engine/Source/Runtime/Core/Public/UObject/UnrealNames.inl b/Engine/Source/Runtime/Core/Public/UObject/UnrealNames.inl
--- a/Engine/Source/Runtime/Core/Public/UObject/UnrealNames.inl
+++ b/Engine/Source/Runtime/Core/Public/UObject/UnrealNames.inl
@@ -44,6 +44,8 @@ REGISTER_NAME(33,CoreUObject)
 
 // More class properties
 REGISTER_NAME(34,EnumProperty)
+REGISTER_NAME(35,ObjectPtrProperty)
+
 
 // Special types.
 REGISTER_NAME(50,Cylinder)
diff --git a/Engine/Source/Runtime/Core/Public/UObject/WeakObjectPtrTemplates.h b/Engine/Source/Runtime/Core/Public/UObject/WeakObjectPtrTemplates.h
--- a/Engine/Source/Runtime/Core/Public/UObject/WeakObjectPtrTemplates.h
+++ b/Engine/Source/Runtime/Core/Public/UObject/WeakObjectPtrTemplates.h
@@ -10,6 +10,8 @@
 #include "Templates/LosesQualifiersFromTo.h"
 #include "Containers/Map.h"
 
+#include <type_traits>
+
 /**
  * FWeakObjectPtr is a weak pointer to a UObject. 
  * It can return nullptr later if the object is garbage collected.
@@ -55,9 +57,9 @@ public:
 	 */
 	template <
 		typename U,
-		typename = decltype(ImplicitConv<T*>((U*)nullptr))
+		decltype(ImplicitConv<T*>(std::declval<U>()))* = nullptr
 	>
-	FORCEINLINE TWeakObjectPtr(U* Object) :
+	FORCEINLINE TWeakObjectPtr(U Object) :
 		TWeakObjectPtrBase((const UObject*)Object)
 	{
 		// This static assert is in here rather than in the body of the class because we want
diff --git a/Engine/Source/Runtime/CoreUObject/Private/Serialization/ArchiveUObject.cpp b/Engine/Source/Runtime/CoreUObject/Private/Serialization/ArchiveUObject.cpp
--- a/Engine/Source/Runtime/CoreUObject/Private/Serialization/ArchiveUObject.cpp
+++ b/Engine/Source/Runtime/CoreUObject/Private/Serialization/ArchiveUObject.cpp
@@ -58,6 +58,22 @@ FArchive& FArchiveUObject::SerializeLazyObjectPtr(FArchive& Ar, FLazyObjectPtr&
 	return Ar;
 }
 
+FArchive& FArchiveUObject::SerializeObjectPtr(FArchive& Ar, FObjectPtr& Value)
+{
+	// Default behavior is to fully resolve the reference and send it through
+	// the raw UObject* serialization codepath and ensure that the result is saved back into
+	// the FObjectPtr afterwards.  There will be many use cases where this is
+	// not what we want to do, but this should be a reasonable default that
+	// allows FObjectPtrs to be treated like raw UObject*'s by default.
+	UObject* Object = Value.Get();
+	Ar << Object;
+	if (Ar.IsLoading() || (Object && Ar.IsModifyingWeakAndStrongReferences()))
+	{
+		Value = Object;
+	}
+	return Ar;
+}
+
 FArchive& FArchiveUObject::SerializeSoftObjectPtr(FArchive& Ar, FSoftObjectPtr& Value)
 {
 	if (Ar.IsSaving() || Ar.IsLoading())
diff --git a/Engine/Source/Runtime/CoreUObject/Private/Serialization/Formatters/JsonArchiveInputFormatter.cpp b/Engine/Source/Runtime/CoreUObject/Private/Serialization/Formatters/JsonArchiveInputFormatter.cpp
--- a/Engine/Source/Runtime/CoreUObject/Private/Serialization/Formatters/JsonArchiveInputFormatter.cpp
+++ b/Engine/Source/Runtime/CoreUObject/Private/Serialization/Formatters/JsonArchiveInputFormatter.cpp
@@ -10,6 +10,7 @@
 #include "UObject/SoftObjectPtr.h"
 #include "UObject/WeakObjectPtr.h"
 #include "UObject/LazyObjectPtr.h"
+#include "UObject/ObjectPtr.h"
 
 #if WITH_TEXT_ARCHIVE_SUPPORT
 
@@ -439,6 +440,20 @@ void FJsonArchiveInputFormatter::Serialize(FLazyObjectPtr& Value)
 	}
 }
 
+void FJsonArchiveInputFormatter::Serialize(FObjectPtr& Value)
+{
+	FString StringValue;
+	const TCHAR Prefix[] = TEXT("Object:");
+	if (ValueStack.Top()->TryGetString(StringValue) && StringValue.StartsWith(Prefix))
+	{
+		Value = ResolveObjectName(*StringValue + UE_ARRAY_COUNT(Prefix) - 1);
+	}
+	else
+	{
+		Value = nullptr;
+	}
+}
+
 void FJsonArchiveInputFormatter::Serialize(TArray<uint8>& Data)
 {
 	const FJsonValue& Value = *ValueStack.Top();
diff --git a/Engine/Source/Runtime/CoreUObject/Private/Serialization/Formatters/JsonArchiveOutputFormatter.cpp b/Engine/Source/Runtime/CoreUObject/Private/Serialization/Formatters/JsonArchiveOutputFormatter.cpp
--- a/Engine/Source/Runtime/CoreUObject/Private/Serialization/Formatters/JsonArchiveOutputFormatter.cpp
+++ b/Engine/Source/Runtime/CoreUObject/Private/Serialization/Formatters/JsonArchiveOutputFormatter.cpp
@@ -9,6 +9,7 @@
 #include "UObject/SoftObjectPtr.h"
 #include "UObject/SoftObjectPath.h"
 #include "UObject/LazyObjectPtr.h"
+#include "UObject/ObjectPtr.h"
 
 #if WITH_TEXT_ARCHIVE_SUPPORT
 
@@ -402,6 +403,19 @@ void FJsonArchiveOutputFormatter::Serialize(FLazyObjectPtr& Value)
 	}
 }
 
+void FJsonArchiveOutputFormatter::Serialize(FObjectPtr& Value)
+{
+	if (!Value.IsNull() && IsObjectAllowed(Value.Get()))
+	{
+		FString FullObjectName = Value->GetFullName(nullptr, EObjectFullNameFlags::IncludeClassPackage);
+		SerializeStringInternal(FString::Printf(TEXT("Object:%s"), *FullObjectName));
+	}
+	else
+	{
+		WriteValue(TEXT("null"));
+	}
+}
+
 void FJsonArchiveOutputFormatter::Serialize(TArray<uint8>& Data)
 {
 	Serialize(Data.GetData(), Data.Num());
diff --git a/Engine/Source/Runtime/CoreUObject/Private/Serialization/ObjectReader.cpp b/Engine/Source/Runtime/CoreUObject/Private/Serialization/ObjectReader.cpp
--- a/Engine/Source/Runtime/CoreUObject/Private/Serialization/ObjectReader.cpp
+++ b/Engine/Source/Runtime/CoreUObject/Private/Serialization/ObjectReader.cpp
@@ -2,6 +2,7 @@
 
 #include "Serialization/ObjectReader.h"
 #include "UObject/LazyObjectPtr.h"
+#include "UObject/ObjectPtr.h"
 #include "UObject/SoftObjectPtr.h"
 
 ///////////////////////////////////////////////////////
@@ -26,6 +27,12 @@ FArchive& FObjectReader::operator<<(UObject*& Res)
 	return *this;
 }
 
+FArchive& FObjectReader::operator<<(FObjectPtr& Value)
+{
+	ByteOrderSerialize(&Value, sizeof(Value));
+	return *this;
+}
+
 FArchive& FObjectReader::operator<<(FLazyObjectPtr& Value)
 {
 	FArchive& Ar = *this;
diff --git a/Engine/Source/Runtime/CoreUObject/Private/Serialization/ObjectWriter.cpp b/Engine/Source/Runtime/CoreUObject/Private/Serialization/ObjectWriter.cpp
--- a/Engine/Source/Runtime/CoreUObject/Private/Serialization/ObjectWriter.cpp
+++ b/Engine/Source/Runtime/CoreUObject/Private/Serialization/ObjectWriter.cpp
@@ -2,6 +2,7 @@
 
 #include "Serialization/ObjectWriter.h"
 #include "UObject/LazyObjectPtr.h"
+#include "UObject/ObjectPtr.h"
 #include "UObject/SoftObjectPtr.h"
 
 ///////////////////////////////////////////////////////
@@ -24,6 +25,12 @@ FArchive& FObjectWriter::operator<<( class UObject*& Res )
 	return *this;
 }
 
+FArchive& FObjectWriter::operator<<(FObjectPtr& Value)
+{
+	ByteOrderSerialize(&Value, sizeof(Value));
+	return *this;
+}
+
 FArchive& FObjectWriter::operator<<(FLazyObjectPtr& Value)
 {
 	FUniqueObjectGuid ID = Value.GetUniqueID();
diff --git a/Engine/Source/Runtime/CoreUObject/Private/Tests/ObjectHandleTest.cpp b/Engine/Source/Runtime/CoreUObject/Private/Tests/ObjectHandleTest.cpp
new file mode 100644
--- /dev/null
+++ b/Engine/Source/Runtime/CoreUObject/Private/Tests/ObjectHandleTest.cpp
@@ -0,0 +1,179 @@
+// Copyright Epic Games, Inc. All Rights Reserved.
+
+#include "UObject/ObjectHandle.h"
+
+#include "HAL/PlatformProperties.h"
+#include "ObjectRefTrackingTestBase.h"
+
+static_assert(sizeof(FObjectHandle) == sizeof(void*), "FObjectHandle type must always compile to something equivalent to a pointer size.");
+
+#if WITH_DEV_AUTOMATION_TESTS
+
+class FObjectHandleTestBase : public FObjectRefTrackingTestBase
+{
+public:
+	FObjectHandleTestBase(const FString& InName, const bool bInComplexTask)
+	: FObjectRefTrackingTestBase(InName, bInComplexTask)
+	{
+	}
+
+protected:
+	UObject* ConstructAndResolveHandle(const ANSICHAR* PackageName, const ANSICHAR* ObjectName, const ANSICHAR* ClassPackageName = nullptr, const ANSICHAR* ClassName = nullptr)
+	{
+		FObjectRef TargetRef{FName(PackageName), FName(ClassPackageName), FName(ClassName), FObjectPathId(ObjectName)};
+		if (!TestFalse(TEXT("Reference to target is null"), IsObjectRefNull(TargetRef)))
+		{
+			return nullptr;
+		}
+
+		FObjectHandle TargetHandle = MakeObjectHandle(TargetRef);
+
+	#if UE_WITH_OBJECT_HANDLE_LATE_RESOLVE
+		// Late resolved handles cannot be null or resolved at this point
+		if (!TestFalse(TEXT("Handle to target is null"), IsObjectHandleNull(TargetHandle)))
+		{
+			return nullptr;
+		}
+		if (!TestFalse(TEXT("Handle to target is resolved"), IsObjectHandleResolved(TargetHandle)))
+		{
+			return nullptr;
+		}
+	#else
+		// Immediately resolved handles may be null (if the target is invalid) and must be resolved at this point
+		if (!TestTrue(TEXT("Handle to target is not resolved"), IsObjectHandleResolved(TargetHandle)))
+		{
+			return nullptr;
+		}
+	#endif
+
+		return ResolveObjectHandle(TargetHandle);
+	}
+
+	bool TestResolvableNonNull(const ANSICHAR* PackageName, const ANSICHAR* ObjectName, const ANSICHAR* ClassPackageName = nullptr, const ANSICHAR* ClassName = nullptr)
+	{
+		FSnapshotObjectRefMetrics ObjectRefMetrics(*this);
+		UObject* ResolvedObject = ConstructAndResolveHandle(PackageName, ObjectName, ClassPackageName, ClassName);
+		ObjectRefMetrics.TestNumResolves(TEXT("NumResolves should be incremented by one after a resolve attempt"), 1);
+		ObjectRefMetrics.TestNumReads(TEXT("NumReads should be incremented by one after a resolve attempt"), 1);
+
+		if (!ResolvedObject)
+		{
+			AddError(FString::Printf(TEXT("Expected '%s.%s' to resolve to non null."), ANSI_TO_TCHAR(PackageName), ANSI_TO_TCHAR(ObjectName)), 1);
+			return false;
+		}
+		ObjectRefMetrics.TestNumFailedResolves(TEXT("NumFailedResolves should not change after a successful resolve attempt"), 0);
+
+		return true;
+	}
+
+	bool TestResolveFailure(const ANSICHAR* PackageName, const ANSICHAR* ObjectName, const ANSICHAR* ClassPackageName = nullptr, const ANSICHAR* ClassName = nullptr)
+	{
+		FSnapshotObjectRefMetrics ObjectRefMetrics(*this);
+		UObject* ResolvedObject = ConstructAndResolveHandle(PackageName, ObjectName, ClassPackageName, ClassName);
+		ObjectRefMetrics.TestNumResolves(TEXT("NumResolves should be incremented by one after a resolve attempt"), 1);
+		ObjectRefMetrics.TestNumReads(TEXT("NumReads should be incremented by one after a resolve attempt"), 1);
+
+		if (ResolvedObject)
+		{
+			AddError(FString::Printf(TEXT("Expected '%s.%s' to resolve to null."), ANSI_TO_TCHAR(PackageName), ANSI_TO_TCHAR(ObjectName)), 1);
+			return false;
+		}
+		ObjectRefMetrics.TestNumFailedResolves(TEXT("NumFailedResolves should be incremented by one after a failed resolve attempt"), 1);
+		return true;
+	}
+};
+
+#define TEST_NAME_ROOT TEXT("System.CoreUObject.ObjectHandle")
+constexpr const uint32 ObjectHandleTestFlags = EAutomationTestFlags::ApplicationContextMask | EAutomationTestFlags::EngineFilter;
+
+IMPLEMENT_CUSTOM_SIMPLE_AUTOMATION_TEST(FObjectHandleTestNullBehavior, FObjectHandleTestBase, TEST_NAME_ROOT TEXT(".NullBehavior"), ObjectHandleTestFlags)
+bool FObjectHandleTestNullBehavior::RunTest(const FString& Parameters)
+{
+	FObjectHandle TargetHandle = MakeObjectHandle(nullptr);
+
+	TestTrue(TEXT("Handle to target is null"), IsObjectHandleNull(TargetHandle));
+	TestTrue(TEXT("Handle to target is resolved"), IsObjectHandleResolved(TargetHandle));
+
+	FSnapshotObjectRefMetrics ObjectRefMetrics(*this);
+	UObject* ResolvedObject = ResolveObjectHandle(TargetHandle);
+
+	TestEqual(TEXT("Resolved object is equal to original object"), (UObject*)nullptr, ResolvedObject);
+
+	ObjectRefMetrics.TestNumFailedResolves(TEXT("NumFailedResolves should not change after a resolve attempt on a null handle"), 0);
+	ObjectRefMetrics.TestNumResolves(TEXT("NumResolves should not change after a resolve attempt on a null handle"), 0);
+	ObjectRefMetrics.TestNumReads(TEXT("NumReads should be incremented by one after a resolve attempt on a null handle"), 1);
+
+	return true;
+}
+
+IMPLEMENT_CUSTOM_SIMPLE_AUTOMATION_TEST(FObjectHandleTestPointerBehavior, FObjectHandleTestBase, TEST_NAME_ROOT TEXT(".PointerBehavior"), ObjectHandleTestFlags)
+bool FObjectHandleTestPointerBehavior::RunTest(const FString& Parameters)
+{
+	FObjectHandle TargetHandle = MakeObjectHandle((UObject*)0x0042);
+
+	TestFalse(TEXT("Handle to target is null"), IsObjectHandleNull(TargetHandle));
+	TestTrue(TEXT("Handle to target is resolved"), IsObjectHandleResolved(TargetHandle));
+
+	FSnapshotObjectRefMetrics ObjectRefMetrics(*this);
+	UObject* ResolvedObject = ResolveObjectHandle(TargetHandle);
+
+	TestEqual(TEXT("Resolved object is equal to original object"), (UObject*)0x0042, ResolvedObject);
+
+	ObjectRefMetrics.TestNumResolves(TEXT("NumResolves should not change after a resolve attempt on a pointer handle"), 0);
+	ObjectRefMetrics.TestNumFailedResolves(TEXT("NumFailedResolves should not change after a resolve attempt on a pointer handle"), 0);
+	ObjectRefMetrics.TestNumReads(TEXT("NumReads should be incremented by one after a resolve attempt on a pointer handle"),1);
+
+	return true;
+}
+
+IMPLEMENT_CUSTOM_SIMPLE_AUTOMATION_TEST(FObjectHandleTestResolveEngineContentTarget, FObjectHandleTestBase, TEST_NAME_ROOT TEXT(".ResolveEngineContentTarget"), ObjectHandleTestFlags)
+bool FObjectHandleTestResolveEngineContentTarget::RunTest(const FString& Parameters)
+{
+	// Confirm we successfully resolve a correct reference to engine content
+	TestResolvableNonNull("/Engine/EngineResources/DefaultTexture", "DefaultTexture");
+
+	// @TODO: OBJPTR: These assets aren't in a standard cook of EngineTest, so avoid testing them when using cooked content.  Should look for other assets to use instead.
+	if (!FPlatformProperties::RequiresCookedData())
+	{
+		// Confirm we successfully resolve a correct reference to a subobject in engine content
+		TestResolvableNonNull("/Engine/FunctionalTesting/Blueprints/AITesting_MoveGoal", "AITesting_MoveGoal.EventGraph.K2Node_VariableGet_142");
+
+		// Attempt to load something that uses a User Defined Enum
+		TestResolvableNonNull("/Engine/ArtTools/RenderToTexture/Macros/RenderToTextureMacros", "RenderToTextureMacros:Array to HLSL Float Array.K2Node_Select_1");
+	}
+
+	return true;
+}
+
+IMPLEMENT_CUSTOM_SIMPLE_AUTOMATION_TEST(FObjectHandleTestFailResolveInvalidTarget, FObjectHandleTestBase, TEST_NAME_ROOT TEXT(".FailResolveInvalidTarget"), ObjectHandleTestFlags)
+bool FObjectHandleTestFailResolveInvalidTarget::RunTest(const FString& Parameters)
+{
+	if (FPlatformProperties::RequiresCookedData())
+	{
+		AddExpectedError(TEXT("Couldn't find file for package"));
+		AddExpectedError(TEXT("Found 0 dependent packages..."));
+	}
+
+	// Confirm we don't successfully resolve an incorrect reference to engine content
+	TestResolveFailure("/Engine/EngineResources/NonExistentPackageName_0", "DefaultTexture");
+	TestResolveFailure("/Engine/EngineResources/DefaultTexture", "NonExistentObject_0");
+
+	return true;
+}
+
+
+
+IMPLEMENT_CUSTOM_SIMPLE_AUTOMATION_TEST(FObjectHandleTestResolveScriptTarget, FObjectHandleTestBase, TEST_NAME_ROOT TEXT(".ResolveScriptTarget"), ObjectHandleTestFlags)
+bool FObjectHandleTestResolveScriptTarget::RunTest(const FString& Parameters)
+{
+	// Confirm we successfully resolve a correct reference to engine content
+	TestResolvableNonNull("/Script/Engine", "Default__Actor");
+
+	TestResolvableNonNull("/Script/Engine", "DefaultPawn");
+
+	return true;
+}
+
+#undef TEST_NAME_ROOT
+
+#endif // WITH_DEV_AUTOMATION_TESTS
