Changed indentation of code in namespaces

#rnx

diff --git a/Engine/Source/Runtime/Core/Private/String/BytesToHex.cpp b/Engine/Source/Runtime/Core/Private/String/BytesToHex.cpp
--- a/Engine/Source/Runtime/Core/Private/String/BytesToHex.cpp
+++ b/Engine/Source/Runtime/Core/Private/String/BytesToHex.cpp
@@ -8,39 +8,41 @@ namespace UE
 {
 namespace String
 {
-	template <typename CharType>
-	void BytesToHexImpl(TArrayView<const uint8> Bytes, CharType* OutHex)
-	{
-		auto NibbleToHex = [](uint8 Value) -> CharType { return Value + (Value > 9 ? 'A' - 10 : '0'); };
-		const uint8* Data = Bytes.GetData();
-		
-		for (const uint8* DataEnd = Data + Bytes.Num(); Data != DataEnd; ++Data)
-		{
-			*OutHex++ = NibbleToHex(*Data >> 4);
-			*OutHex++ = NibbleToHex(*Data & 15);
-		}
-	}
 
-	void BytesToHex(TArrayView<const uint8> Bytes, WIDECHAR* OutHex)
+template <typename CharType>
+void BytesToHexImpl(TArrayView<const uint8> Bytes, CharType* OutHex)
+{
+	auto NibbleToHex = [](uint8 Value) -> CharType { return Value + (Value > 9 ? 'A' - 10 : '0'); };
+	const uint8* Data = Bytes.GetData();
+	
+	for (const uint8* DataEnd = Data + Bytes.Num(); Data != DataEnd; ++Data)
 	{
-		BytesToHexImpl(Bytes, OutHex);
+		*OutHex++ = NibbleToHex(*Data >> 4);
+		*OutHex++ = NibbleToHex(*Data & 15);
 	}
+}
 
-	void BytesToHex(TArrayView<const uint8> Bytes, ANSICHAR* OutHex)
-	{
-		BytesToHexImpl(Bytes, OutHex);
-	}
+void BytesToHex(TArrayView<const uint8> Bytes, WIDECHAR* OutHex)
+{
+	BytesToHexImpl(Bytes, OutHex);
+}
 
-	void BytesToHex(TArrayView<const uint8> Bytes, FStringBuilderBase& Builder)
-	{
-		const int32 Offset = Builder.AddUninitialized(Bytes.Num() * 2);
-		BytesToHexImpl(Bytes, GetData(Builder) + Offset);
-	}
+void BytesToHex(TArrayView<const uint8> Bytes, ANSICHAR* OutHex)
+{
+	BytesToHexImpl(Bytes, OutHex);
+}
+
+void BytesToHex(TArrayView<const uint8> Bytes, FStringBuilderBase& Builder)
+{
+	const int32 Offset = Builder.AddUninitialized(Bytes.Num() * 2);
+	BytesToHexImpl(Bytes, GetData(Builder) + Offset);
+}
+
+void BytesToHex(TArrayView<const uint8> Bytes, FAnsiStringBuilderBase& Builder)
+{
+	const int32 Offset = Builder.AddUninitialized(Bytes.Num() * 2);
+	BytesToHexImpl(Bytes, GetData(Builder) + Offset);
+}
 
-	void BytesToHex(TArrayView<const uint8> Bytes, FAnsiStringBuilderBase& Builder)
-	{
-		const int32 Offset = Builder.AddUninitialized(Bytes.Num() * 2);
-		BytesToHexImpl(Bytes, GetData(Builder) + Offset);
-	}
 }
 }
diff --git a/Engine/Source/Runtime/Core/Private/String/Find.cpp b/Engine/Source/Runtime/Core/Private/String/Find.cpp
--- a/Engine/Source/Runtime/Core/Private/String/Find.cpp
+++ b/Engine/Source/Runtime/Core/Private/String/Find.cpp
@@ -10,386 +10,390 @@ namespace UE
 {
 namespace String
 {
-	namespace Private
-	{
-		// These are naive implementations that take time proportional to View.Len() * TotalSearchLen.
-		// If these functions become a bottleneck, they can be specialized separately for one and many search patterns.
-		// There are algorithms for each that are linear or sub-linear in the length of the string to search.
 
-		template <typename CharType>
-		static inline typename TStringView<CharType>::SizeType FindFirst(TStringView<CharType> View, TStringView<CharType> Search, ESearchCase::Type SearchCase)
+namespace Private
+{
+
+// These are naive implementations that take time proportional to View.Len() * TotalSearchLen.
+// If these functions become a bottleneck, they can be specialized separately for one and many search patterns.
+// There are algorithms for each that are linear or sub-linear in the length of the string to search.
+
+template <typename CharType>
+static inline typename TStringView<CharType>::SizeType FindFirst(TStringView<CharType> View, TStringView<CharType> Search, ESearchCase::Type SearchCase)
+{
+	check(!Search.IsEmpty());
+	const typename TStringView<CharType>::SizeType SearchLen = Search.Len();
+	if (SearchLen == 1)
+	{
+		return String::FindFirstChar(View, Search[0], SearchCase);
+	}
+	const CharType* const SearchData = Search.GetData();
+	const CharType* const ViewBegin = View.GetData();
+	const CharType* const ViewEnd = ViewBegin + View.Len() - SearchLen;
+	if (SearchCase == ESearchCase::CaseSensitive)
+	{
+		for (const CharType* ViewIt = ViewBegin; ViewIt <= ViewEnd; ++ViewIt)
 		{
-			check(!Search.IsEmpty());
-			const typename TStringView<CharType>::SizeType SearchLen = Search.Len();
-			if (SearchLen == 1)
+			if (TCString<CharType>::Strncmp(ViewIt, SearchData, SearchLen) == 0)
 			{
-				return String::FindFirstChar(View, Search[0], SearchCase);
+				return typename TStringView<CharType>::SizeType(ViewIt - ViewBegin);
 			}
-			const CharType* const SearchData = Search.GetData();
-			const CharType* const ViewBegin = View.GetData();
-			const CharType* const ViewEnd = ViewBegin + View.Len() - SearchLen;
-			if (SearchCase == ESearchCase::CaseSensitive)
-			{
-				for (const CharType* ViewIt = ViewBegin; ViewIt <= ViewEnd; ++ViewIt)
-				{
-					if (TCString<CharType>::Strncmp(ViewIt, SearchData, SearchLen) == 0)
-					{
-						return typename TStringView<CharType>::SizeType(ViewIt - ViewBegin);
-					}
-				}
-			}
-			else
+		}
+	}
+	else
+	{
+		for (const CharType* ViewIt = ViewBegin; ViewIt <= ViewEnd; ++ViewIt)
+		{
+			if (TCString<CharType>::Strnicmp(ViewIt, SearchData, SearchLen) == 0)
 			{
-				for (const CharType* ViewIt = ViewBegin; ViewIt <= ViewEnd; ++ViewIt)
-				{
-					if (TCString<CharType>::Strnicmp(ViewIt, SearchData, SearchLen) == 0)
-					{
-						return typename TStringView<CharType>::SizeType(ViewIt - ViewBegin);
-					}
-				}
+				return typename TStringView<CharType>::SizeType(ViewIt - ViewBegin);
 			}
-			return INDEX_NONE;
 		}
+	}
+	return INDEX_NONE;
+}
 
-		template <typename CharType>
-		static inline typename TStringView<CharType>::SizeType FindLast(TStringView<CharType> View, TStringView<CharType> Search, ESearchCase::Type SearchCase)
+template <typename CharType>
+static inline typename TStringView<CharType>::SizeType FindLast(TStringView<CharType> View, TStringView<CharType> Search, ESearchCase::Type SearchCase)
+{
+	check(!Search.IsEmpty());
+	const typename TStringView<CharType>::SizeType SearchLen = Search.Len();
+	if (SearchLen == 1)
+	{
+		return String::FindLastChar(View, Search[0], SearchCase);
+	}
+	const CharType* const SearchData = Search.GetData();
+	const CharType* const ViewBegin = View.GetData();
+	const CharType* const ViewEnd = ViewBegin + View.Len() - SearchLen;
+	if (SearchCase == ESearchCase::CaseSensitive)
+	{
+		for (const CharType* ViewIt = ViewEnd; ViewIt >= ViewBegin; --ViewIt)
 		{
-			check(!Search.IsEmpty());
-			const typename TStringView<CharType>::SizeType SearchLen = Search.Len();
-			if (SearchLen == 1)
+			if (TCString<CharType>::Strncmp(ViewIt, SearchData, SearchLen) == 0)
 			{
-				return String::FindLastChar(View, Search[0], SearchCase);
+				return typename TStringView<CharType>::SizeType(ViewIt - ViewBegin);
 			}
-			const CharType* const SearchData = Search.GetData();
-			const CharType* const ViewBegin = View.GetData();
-			const CharType* const ViewEnd = ViewBegin + View.Len() - SearchLen;
-			if (SearchCase == ESearchCase::CaseSensitive)
-			{
-				for (const CharType* ViewIt = ViewEnd; ViewIt >= ViewBegin; --ViewIt)
-				{
-					if (TCString<CharType>::Strncmp(ViewIt, SearchData, SearchLen) == 0)
-					{
-						return typename TStringView<CharType>::SizeType(ViewIt - ViewBegin);
-					}
-				}
-			}
-			else
+		}
+	}
+	else
+	{
+		for (const CharType* ViewIt = ViewEnd; ViewIt >= ViewBegin; --ViewIt)
+		{
+			if (TCString<CharType>::Strnicmp(ViewIt, SearchData, SearchLen) == 0)
 			{
-				for (const CharType* ViewIt = ViewEnd; ViewIt >= ViewBegin; --ViewIt)
-				{
-					if (TCString<CharType>::Strnicmp(ViewIt, SearchData, SearchLen) == 0)
-					{
-						return typename TStringView<CharType>::SizeType(ViewIt - ViewBegin);
-					}
-				}
+				return typename TStringView<CharType>::SizeType(ViewIt - ViewBegin);
 			}
-			return INDEX_NONE;
 		}
+	}
+	return INDEX_NONE;
+}
 
-		template <typename CharType>
-		static inline typename TStringView<CharType>::SizeType FindFirstOfAny(TStringView<CharType> View, TConstArrayView<TStringView<CharType>> Search, ESearchCase::Type SearchCase)
+template <typename CharType>
+static inline typename TStringView<CharType>::SizeType FindFirstOfAny(TStringView<CharType> View, TConstArrayView<TStringView<CharType>> Search, ESearchCase::Type SearchCase)
+{
+	check(Algo::NoneOf(Search, &TStringView<CharType>::IsEmpty));
+	switch (Search.Num())
+	{
+	case 0:
+		return INDEX_NONE;
+	case 1:
+		return String::FindFirst(View, Search[0], SearchCase);
+	default:
+		if (Algo::AllOf(Search, [](const TStringView<CharType>& Pattern) { return Pattern.Len() == 1; }))
 		{
-			check(Algo::NoneOf(Search, &TStringView<CharType>::IsEmpty));
-			switch (Search.Num())
+			TArray<CharType, TInlineAllocator<32>> SearchChars;
+			SearchChars.Reserve(Search.Num());
+			for (const TStringView<CharType>& Pattern : Search)
 			{
-			case 0:
-				return INDEX_NONE;
-			case 1:
-				return String::FindFirst(View, Search[0], SearchCase);
-			default:
-				if (Algo::AllOf(Search, [](const TStringView<CharType>& Pattern) { return Pattern.Len() == 1; }))
-				{
-					TArray<CharType, TInlineAllocator<32>> SearchChars;
-					SearchChars.Reserve(Search.Num());
-					for (const TStringView<CharType>& Pattern : Search)
-					{
-						SearchChars.Add(Pattern[0]);
-					}
-					return String::FindFirstOfAnyChar(View, SearchChars, SearchCase);
-				}
-				break;
+				SearchChars.Add(Pattern[0]);
 			}
+			return String::FindFirstOfAnyChar(View, SearchChars, SearchCase);
+		}
+		break;
+	}
 
-			const CharType* const ViewBegin = View.GetData();
-			const typename TStringView<CharType>::SizeType ViewLen = View.Len();
-			for (typename TStringView<CharType>::SizeType ViewIndex = 0; ViewIndex != ViewLen; ++ViewIndex)
-			{
-				const TStringView<CharType> RemainingView(ViewBegin + ViewIndex, ViewLen - ViewIndex);
-				auto MatchPattern = [&RemainingView, SearchCase](const TStringView<CharType>& Pattern) { return RemainingView.StartsWith(Pattern, SearchCase); };
-				if (Algo::FindByPredicate(Search, MatchPattern))
-				{
-					return ViewIndex;
-				}
-			}
-			return INDEX_NONE;
+	const CharType* const ViewBegin = View.GetData();
+	const typename TStringView<CharType>::SizeType ViewLen = View.Len();
+	for (typename TStringView<CharType>::SizeType ViewIndex = 0; ViewIndex != ViewLen; ++ViewIndex)
+	{
+		const TStringView<CharType> RemainingView(ViewBegin + ViewIndex, ViewLen - ViewIndex);
+		auto MatchPattern = [&RemainingView, SearchCase](const TStringView<CharType>& Pattern) { return RemainingView.StartsWith(Pattern, SearchCase); };
+		if (Algo::FindByPredicate(Search, MatchPattern))
+		{
+			return ViewIndex;
 		}
+	}
+	return INDEX_NONE;
+}
 
-		template <typename CharType>
-		static inline typename TStringView<CharType>::SizeType FindLastOfAny(TStringView<CharType> View, TConstArrayView<TStringView<CharType>> Search, ESearchCase::Type SearchCase)
+template <typename CharType>
+static inline typename TStringView<CharType>::SizeType FindLastOfAny(TStringView<CharType> View, TConstArrayView<TStringView<CharType>> Search, ESearchCase::Type SearchCase)
+{
+	check(Algo::NoneOf(Search, &TStringView<CharType>::IsEmpty));
+	switch (Search.Num())
+	{
+	case 0:
+		return INDEX_NONE;
+	case 1:
+		return String::FindLast(View, Search[0], SearchCase);
+	default:
+		if (Algo::AllOf(Search, [](const TStringView<CharType>& Pattern) { return Pattern.Len() == 1; }))
 		{
-			check(Algo::NoneOf(Search, &TStringView<CharType>::IsEmpty));
-			switch (Search.Num())
+			TArray<CharType, TInlineAllocator<32>> SearchChars;
+			SearchChars.Reserve(Search.Num());
+			for (const TStringView<CharType>& Pattern : Search)
 			{
-			case 0:
-				return INDEX_NONE;
-			case 1:
-				return String::FindLast(View, Search[0], SearchCase);
-			default:
-				if (Algo::AllOf(Search, [](const TStringView<CharType>& Pattern) { return Pattern.Len() == 1; }))
-				{
-					TArray<CharType, TInlineAllocator<32>> SearchChars;
-					SearchChars.Reserve(Search.Num());
-					for (const TStringView<CharType>& Pattern : Search)
-					{
-						SearchChars.Add(Pattern[0]);
-					}
-					return String::FindLastOfAnyChar(View, SearchChars, SearchCase);
-				}
-				break;
+				SearchChars.Add(Pattern[0]);
 			}
+			return String::FindLastOfAnyChar(View, SearchChars, SearchCase);
+		}
+		break;
+	}
 
 
-			const CharType* const ViewBegin = View.GetData();
-			const typename TStringView<CharType>::SizeType ViewLen = View.Len();
-			for (typename TStringView<CharType>::SizeType ViewIndex = ViewLen - 1; ViewIndex >= 0; --ViewIndex)
-			{
-				const TStringView<CharType> RemainingView(ViewBegin + ViewIndex, ViewLen - ViewIndex);
-				auto MatchPattern = [&RemainingView, SearchCase](const TStringView<CharType>& Pattern) { return RemainingView.StartsWith(Pattern, SearchCase); };
-				if (Algo::FindByPredicate(Search, MatchPattern))
-				{
-					return ViewIndex;
-				}
-			}
-			return INDEX_NONE;
+	const CharType* const ViewBegin = View.GetData();
+	const typename TStringView<CharType>::SizeType ViewLen = View.Len();
+	for (typename TStringView<CharType>::SizeType ViewIndex = ViewLen - 1; ViewIndex >= 0; --ViewIndex)
+	{
+		const TStringView<CharType> RemainingView(ViewBegin + ViewIndex, ViewLen - ViewIndex);
+		auto MatchPattern = [&RemainingView, SearchCase](const TStringView<CharType>& Pattern) { return RemainingView.StartsWith(Pattern, SearchCase); };
+		if (Algo::FindByPredicate(Search, MatchPattern))
+		{
+			return ViewIndex;
 		}
+	}
+	return INDEX_NONE;
+}
 
-		template <typename CharType>
-		static inline typename TStringView<CharType>::SizeType FindFirstChar(TStringView<CharType> View, CharType Search, ESearchCase::Type SearchCase)
+template <typename CharType>
+static inline typename TStringView<CharType>::SizeType FindFirstChar(TStringView<CharType> View, CharType Search, ESearchCase::Type SearchCase)
+{
+	const CharType* const ViewBegin = View.GetData();
+	const CharType* const ViewEnd = ViewBegin + View.Len();
+	if (SearchCase == ESearchCase::CaseSensitive)
+	{
+		for (const CharType* ViewIt = ViewBegin; ViewIt < ViewEnd; ++ViewIt)
 		{
-			const CharType* const ViewBegin = View.GetData();
-			const CharType* const ViewEnd = ViewBegin + View.Len();
-			if (SearchCase == ESearchCase::CaseSensitive)
-			{
-				for (const CharType* ViewIt = ViewBegin; ViewIt < ViewEnd; ++ViewIt)
-				{
-					if (*ViewIt == Search)
-					{
-						return typename TStringView<CharType>::SizeType(ViewIt - ViewBegin);
-					}
-				}
-			}
-			else
+			if (*ViewIt == Search)
 			{
-				const CharType SearchUpper = TChar<CharType>::ToUpper(Search);
-				for (const CharType* ViewIt = ViewBegin; ViewIt < ViewEnd; ++ViewIt)
-				{
-					if (TChar<CharType>::ToUpper(*ViewIt) == SearchUpper)
-					{
-						return typename TStringView<CharType>::SizeType(ViewIt - ViewBegin);
-					}
-				}
+				return typename TStringView<CharType>::SizeType(ViewIt - ViewBegin);
 			}
-			return INDEX_NONE;
 		}
-
-		template <typename CharType>
-		static inline typename TStringView<CharType>::SizeType FindLastChar(TStringView<CharType> View, CharType Search, ESearchCase::Type SearchCase)
+	}
+	else
+	{
+		const CharType SearchUpper = TChar<CharType>::ToUpper(Search);
+		for (const CharType* ViewIt = ViewBegin; ViewIt < ViewEnd; ++ViewIt)
 		{
-			const CharType* const ViewBegin = View.GetData();
-			const CharType* const ViewEnd = ViewBegin + View.Len();
-			if (SearchCase == ESearchCase::CaseSensitive)
-			{
-				for (const CharType* ViewIt = ViewEnd - 1; ViewIt >= ViewBegin; --ViewIt)
-				{
-					if (*ViewIt == Search)
-					{
-						return typename TStringView<CharType>::SizeType(ViewIt - ViewBegin);
-					}
-				}
-			}
-			else
+			if (TChar<CharType>::ToUpper(*ViewIt) == SearchUpper)
 			{
-				const CharType SearchUpper = TChar<CharType>::ToUpper(Search);
-				for (const CharType* ViewIt = ViewEnd - 1; ViewIt >= ViewBegin; --ViewIt)
-				{
-					if (TChar<CharType>::ToUpper(*ViewIt) == SearchUpper)
-					{
-						return typename TStringView<CharType>::SizeType(ViewIt - ViewBegin);
-					}
-				}
+				return typename TStringView<CharType>::SizeType(ViewIt - ViewBegin);
 			}
-			return INDEX_NONE;
 		}
+	}
+	return INDEX_NONE;
+}
 
-		template <typename CharType>
-		static inline typename TStringView<CharType>::SizeType FindFirstOfAnyChar(TStringView<CharType> View, TConstArrayView<CharType> Search, ESearchCase::Type SearchCase)
+template <typename CharType>
+static inline typename TStringView<CharType>::SizeType FindLastChar(TStringView<CharType> View, CharType Search, ESearchCase::Type SearchCase)
+{
+	const CharType* const ViewBegin = View.GetData();
+	const CharType* const ViewEnd = ViewBegin + View.Len();
+	if (SearchCase == ESearchCase::CaseSensitive)
+	{
+		for (const CharType* ViewIt = ViewEnd - 1; ViewIt >= ViewBegin; --ViewIt)
 		{
-			switch (Search.Num())
+			if (*ViewIt == Search)
 			{
-			case 0:
-				return INDEX_NONE;
-			case 1:
-				return String::FindFirstChar(View, Search[0], SearchCase);
-			default:
-				break;
-			}
-
-			const CharType* const ViewBegin = View.GetData();
-			const CharType* const ViewEnd = ViewBegin + View.Len();
-			if (SearchCase == ESearchCase::CaseSensitive)
-			{
-				for (const CharType* ViewIt = ViewBegin; ViewIt < ViewEnd; ++ViewIt)
-				{
-					if (Algo::Find(Search, *ViewIt))
-					{
-						return typename TStringView<CharType>::SizeType(ViewIt - ViewBegin);
-					}
-				}
-				return INDEX_NONE;
-			}
-			else
-			{
-				TArray<CharType, TInlineAllocator<32>> SearchUpper;
-				SearchUpper.Reserve(Search.Num());
-				for (const CharType Pattern : Search)
-				{
-					SearchUpper.Add(TChar<CharType>::ToUpper(Pattern));
-				}
-				for (const CharType* ViewIt = ViewBegin; ViewIt < ViewEnd; ++ViewIt)
-				{
-					if (Algo::Find(SearchUpper, TChar<CharType>::ToUpper(*ViewIt)))
-					{
-						return typename TStringView<CharType>::SizeType(ViewIt - ViewBegin);
-					}
-				}
-				return INDEX_NONE;
+				return typename TStringView<CharType>::SizeType(ViewIt - ViewBegin);
 			}
 		}
-
-		template <typename CharType>
-		static inline typename TStringView<CharType>::SizeType FindLastOfAnyChar(TStringView<CharType> View, TConstArrayView<CharType> Search, ESearchCase::Type SearchCase)
+	}
+	else
+	{
+		const CharType SearchUpper = TChar<CharType>::ToUpper(Search);
+		for (const CharType* ViewIt = ViewEnd - 1; ViewIt >= ViewBegin; --ViewIt)
 		{
-			switch (Search.Num())
-			{
-			case 0:
-				return INDEX_NONE;
-			case 1:
-				return String::FindLastChar(View, Search[0], SearchCase);
-			default:
-				break;
-			}
-
-			const CharType* const ViewBegin = View.GetData();
-			const CharType* const ViewEnd = ViewBegin + View.Len();
-			if (SearchCase == ESearchCase::CaseSensitive)
-			{
-				for (const CharType* ViewIt = ViewEnd - 1; ViewIt >= ViewBegin; --ViewIt)
-				{
-					if (Algo::Find(Search, *ViewIt))
-					{
-						return typename TStringView<CharType>::SizeType(ViewIt - ViewBegin);
-					}
-				}
-				return INDEX_NONE;
-			}
-			else
+			if (TChar<CharType>::ToUpper(*ViewIt) == SearchUpper)
 			{
-				TArray<CharType, TInlineAllocator<32>> SearchUpper;
-				SearchUpper.Reserve(Search.Num());
-				for (const CharType Pattern : Search)
-				{
-					SearchUpper.Add(TChar<CharType>::ToUpper(Pattern));
-				}
-				for (const CharType* ViewIt = ViewEnd - 1; ViewIt >= ViewBegin; --ViewIt)
-				{
-					if (Algo::Find(SearchUpper, TChar<CharType>::ToUpper(*ViewIt)))
-					{
-						return typename TStringView<CharType>::SizeType(ViewIt - ViewBegin);
-					}
-				}
-				return INDEX_NONE;
+				return typename TStringView<CharType>::SizeType(ViewIt - ViewBegin);
 			}
 		}
 	}
+	return INDEX_NONE;
+}
 
-	FAnsiStringView::SizeType FindFirst(FAnsiStringView View, FAnsiStringView Search, ESearchCase::Type SearchCase)
+template <typename CharType>
+static inline typename TStringView<CharType>::SizeType FindFirstOfAnyChar(TStringView<CharType> View, TConstArrayView<CharType> Search, ESearchCase::Type SearchCase)
+{
+	switch (Search.Num())
 	{
-		return Private::FindFirst(View, Search, SearchCase);
+	case 0:
+		return INDEX_NONE;
+	case 1:
+		return String::FindFirstChar(View, Search[0], SearchCase);
+	default:
+		break;
 	}
 
-	FWideStringView::SizeType FindFirst(FWideStringView View, FWideStringView Search, ESearchCase::Type SearchCase)
+	const CharType* const ViewBegin = View.GetData();
+	const CharType* const ViewEnd = ViewBegin + View.Len();
+	if (SearchCase == ESearchCase::CaseSensitive)
 	{
-		return Private::FindFirst(View, Search, SearchCase);
+		for (const CharType* ViewIt = ViewBegin; ViewIt < ViewEnd; ++ViewIt)
+		{
+			if (Algo::Find(Search, *ViewIt))
+			{
+				return typename TStringView<CharType>::SizeType(ViewIt - ViewBegin);
+			}
+		}
+		return INDEX_NONE;
 	}
-
-	FAnsiStringView::SizeType FindLast(FAnsiStringView View, FAnsiStringView Search, ESearchCase::Type SearchCase)
+	else
 	{
-		return Private::FindLast(View, Search, SearchCase);
+		TArray<CharType, TInlineAllocator<32>> SearchUpper;
+		SearchUpper.Reserve(Search.Num());
+		for (const CharType Pattern : Search)
+		{
+			SearchUpper.Add(TChar<CharType>::ToUpper(Pattern));
+		}
+		for (const CharType* ViewIt = ViewBegin; ViewIt < ViewEnd; ++ViewIt)
+		{
+			if (Algo::Find(SearchUpper, TChar<CharType>::ToUpper(*ViewIt)))
+			{
+				return typename TStringView<CharType>::SizeType(ViewIt - ViewBegin);
+			}
+		}
+		return INDEX_NONE;
 	}
+}
 
-	FWideStringView::SizeType FindLast(FWideStringView View, FWideStringView Search, ESearchCase::Type SearchCase)
+template <typename CharType>
+static inline typename TStringView<CharType>::SizeType FindLastOfAnyChar(TStringView<CharType> View, TConstArrayView<CharType> Search, ESearchCase::Type SearchCase)
+{
+	switch (Search.Num())
 	{
-		return Private::FindLast(View, Search, SearchCase);
+	case 0:
+		return INDEX_NONE;
+	case 1:
+		return String::FindLastChar(View, Search[0], SearchCase);
+	default:
+		break;
 	}
 
-	FAnsiStringView::SizeType FindFirstOfAny(FAnsiStringView View, TConstArrayView<FAnsiStringView> Search, ESearchCase::Type SearchCase)
+	const CharType* const ViewBegin = View.GetData();
+	const CharType* const ViewEnd = ViewBegin + View.Len();
+	if (SearchCase == ESearchCase::CaseSensitive)
 	{
-		return Private::FindFirstOfAny(View, Search, SearchCase);
+		for (const CharType* ViewIt = ViewEnd - 1; ViewIt >= ViewBegin; --ViewIt)
+		{
+			if (Algo::Find(Search, *ViewIt))
+			{
+				return typename TStringView<CharType>::SizeType(ViewIt - ViewBegin);
+			}
+		}
+		return INDEX_NONE;
 	}
-
-	FWideStringView::SizeType FindFirstOfAny(FWideStringView View, TConstArrayView<FWideStringView> Search, ESearchCase::Type SearchCase)
+	else
 	{
-		return Private::FindFirstOfAny(View, Search, SearchCase);
+		TArray<CharType, TInlineAllocator<32>> SearchUpper;
+		SearchUpper.Reserve(Search.Num());
+		for (const CharType Pattern : Search)
+		{
+			SearchUpper.Add(TChar<CharType>::ToUpper(Pattern));
+		}
+		for (const CharType* ViewIt = ViewEnd - 1; ViewIt >= ViewBegin; --ViewIt)
+		{
+			if (Algo::Find(SearchUpper, TChar<CharType>::ToUpper(*ViewIt)))
+			{
+				return typename TStringView<CharType>::SizeType(ViewIt - ViewBegin);
+			}
+		}
+		return INDEX_NONE;
 	}
+}
 
-	FAnsiStringView::SizeType FindLastOfAny(FAnsiStringView View, TConstArrayView<FAnsiStringView> Search, ESearchCase::Type SearchCase)
-	{
-		return Private::FindLastOfAny(View, Search, SearchCase);
-	}
+}
 
-	FWideStringView::SizeType FindLastOfAny(FWideStringView View, TConstArrayView<FWideStringView> Search, ESearchCase::Type SearchCase)
-	{
-		return Private::FindLastOfAny(View, Search, SearchCase);
-	}
+FAnsiStringView::SizeType FindFirst(FAnsiStringView View, FAnsiStringView Search, ESearchCase::Type SearchCase)
+{
+	return Private::FindFirst(View, Search, SearchCase);
+}
 
-	FAnsiStringView::SizeType FindFirstChar(FAnsiStringView View, ANSICHAR Search, ESearchCase::Type SearchCase)
-	{
-		return Private::FindFirstChar(View, Search, SearchCase);
-	}
+FWideStringView::SizeType FindFirst(FWideStringView View, FWideStringView Search, ESearchCase::Type SearchCase)
+{
+	return Private::FindFirst(View, Search, SearchCase);
+}
 
-	FWideStringView::SizeType FindFirstChar(FWideStringView View, WIDECHAR Search, ESearchCase::Type SearchCase)
-	{
-		return Private::FindFirstChar(View, Search, SearchCase);
-	}
+FAnsiStringView::SizeType FindLast(FAnsiStringView View, FAnsiStringView Search, ESearchCase::Type SearchCase)
+{
+	return Private::FindLast(View, Search, SearchCase);
+}
 
-	FAnsiStringView::SizeType FindLastChar(FAnsiStringView View, ANSICHAR Search, ESearchCase::Type SearchCase)
-	{
-		return Private::FindLastChar(View, Search, SearchCase);
-	}
+FWideStringView::SizeType FindLast(FWideStringView View, FWideStringView Search, ESearchCase::Type SearchCase)
+{
+	return Private::FindLast(View, Search, SearchCase);
+}
 
-	FWideStringView::SizeType FindLastChar(FWideStringView View, WIDECHAR Search, ESearchCase::Type SearchCase)
-	{
-		return Private::FindLastChar(View, Search, SearchCase);
-	}
+FAnsiStringView::SizeType FindFirstOfAny(FAnsiStringView View, TConstArrayView<FAnsiStringView> Search, ESearchCase::Type SearchCase)
+{
+	return Private::FindFirstOfAny(View, Search, SearchCase);
+}
 
-	FAnsiStringView::SizeType FindFirstOfAnyChar(FAnsiStringView View, TConstArrayView<ANSICHAR> Search, ESearchCase::Type SearchCase)
-	{
-		return Private::FindFirstOfAnyChar(View, Search, SearchCase);
-	}
+FWideStringView::SizeType FindFirstOfAny(FWideStringView View, TConstArrayView<FWideStringView> Search, ESearchCase::Type SearchCase)
+{
+	return Private::FindFirstOfAny(View, Search, SearchCase);
+}
 
-	FWideStringView::SizeType FindFirstOfAnyChar(FWideStringView View, TConstArrayView<WIDECHAR> Search, ESearchCase::Type SearchCase)
-	{
-		return Private::FindFirstOfAnyChar(View, Search, SearchCase);
-	}
+FAnsiStringView::SizeType FindLastOfAny(FAnsiStringView View, TConstArrayView<FAnsiStringView> Search, ESearchCase::Type SearchCase)
+{
+	return Private::FindLastOfAny(View, Search, SearchCase);
+}
 
-	FAnsiStringView::SizeType FindLastOfAnyChar(FAnsiStringView View, TConstArrayView<ANSICHAR> Search, ESearchCase::Type SearchCase)
-	{
-		return Private::FindLastOfAnyChar(View, Search, SearchCase);
-	}
+FWideStringView::SizeType FindLastOfAny(FWideStringView View, TConstArrayView<FWideStringView> Search, ESearchCase::Type SearchCase)
+{
+	return Private::FindLastOfAny(View, Search, SearchCase);
+}
+
+FAnsiStringView::SizeType FindFirstChar(FAnsiStringView View, ANSICHAR Search, ESearchCase::Type SearchCase)
+{
+	return Private::FindFirstChar(View, Search, SearchCase);
+}
+
+FWideStringView::SizeType FindFirstChar(FWideStringView View, WIDECHAR Search, ESearchCase::Type SearchCase)
+{
+	return Private::FindFirstChar(View, Search, SearchCase);
+}
+
+FAnsiStringView::SizeType FindLastChar(FAnsiStringView View, ANSICHAR Search, ESearchCase::Type SearchCase)
+{
+	return Private::FindLastChar(View, Search, SearchCase);
+}
+
+FWideStringView::SizeType FindLastChar(FWideStringView View, WIDECHAR Search, ESearchCase::Type SearchCase)
+{
+	return Private::FindLastChar(View, Search, SearchCase);
+}
+
+FAnsiStringView::SizeType FindFirstOfAnyChar(FAnsiStringView View, TConstArrayView<ANSICHAR> Search, ESearchCase::Type SearchCase)
+{
+	return Private::FindFirstOfAnyChar(View, Search, SearchCase);
+}
+
+FWideStringView::SizeType FindFirstOfAnyChar(FWideStringView View, TConstArrayView<WIDECHAR> Search, ESearchCase::Type SearchCase)
+{
+	return Private::FindFirstOfAnyChar(View, Search, SearchCase);
+}
+
+FAnsiStringView::SizeType FindLastOfAnyChar(FAnsiStringView View, TConstArrayView<ANSICHAR> Search, ESearchCase::Type SearchCase)
+{
+	return Private::FindLastOfAnyChar(View, Search, SearchCase);
+}
+
+FWideStringView::SizeType FindLastOfAnyChar(FWideStringView View, TConstArrayView<WIDECHAR> Search, ESearchCase::Type SearchCase)
+{
+	return Private::FindLastOfAnyChar(View, Search, SearchCase);
+}
 
-	FWideStringView::SizeType FindLastOfAnyChar(FWideStringView View, TConstArrayView<WIDECHAR> Search, ESearchCase::Type SearchCase)
-	{
-		return Private::FindLastOfAnyChar(View, Search, SearchCase);
-	}
 }
 }
diff --git a/Engine/Source/Runtime/Core/Private/String/HexToBytes.cpp b/Engine/Source/Runtime/Core/Private/String/HexToBytes.cpp
--- a/Engine/Source/Runtime/Core/Private/String/HexToBytes.cpp
+++ b/Engine/Source/Runtime/Core/Private/String/HexToBytes.cpp
@@ -9,33 +9,35 @@ namespace UE
 {
 namespace String
 {
-	template <typename CharType>
-	int32 HexToBytesImpl(const TStringView<CharType>& Hex, uint8* const OutBytes)
-	{
-		const int32 HexCount = Hex.Len();
-		const CharType* HexPos = Hex.GetData();
-		const CharType* HexEnd = HexPos + HexCount;
-		uint8* OutPos = OutBytes;
-		if (const bool bPadNibble = (HexCount % 2) == 1)
-		{
-			*OutPos++ = TCharToNibble(*HexPos++);
-		}
-		while (HexPos != HexEnd)
-		{
-			const uint8 HiNibble = uint8(TCharToNibble(*HexPos++) << 4);
-			*OutPos++ = HiNibble | TCharToNibble(*HexPos++);
-		}
-		return static_cast<int32>(OutPos - OutBytes);
-	}
 
-	int32 HexToBytes(const FStringView& Hex, uint8* OutBytes)
+template <typename CharType>
+int32 HexToBytesImpl(const TStringView<CharType>& Hex, uint8* const OutBytes)
+{
+	const int32 HexCount = Hex.Len();
+	const CharType* HexPos = Hex.GetData();
+	const CharType* HexEnd = HexPos + HexCount;
+	uint8* OutPos = OutBytes;
+	if (const bool bPadNibble = (HexCount % 2) == 1)
 	{
-		return HexToBytesImpl<TCHAR>(Hex, OutBytes);
+		*OutPos++ = TCharToNibble(*HexPos++);
 	}
-
-	int32 HexToBytes(const FAnsiStringView& Hex, uint8* OutBytes)
+	while (HexPos != HexEnd)
 	{
-		return HexToBytesImpl<ANSICHAR>(Hex, OutBytes);
+		const uint8 HiNibble = uint8(TCharToNibble(*HexPos++) << 4);
+		*OutPos++ = HiNibble | TCharToNibble(*HexPos++);
 	}
+	return static_cast<int32>(OutPos - OutBytes);
+}
+
+int32 HexToBytes(const FStringView& Hex, uint8* OutBytes)
+{
+	return HexToBytesImpl<TCHAR>(Hex, OutBytes);
+}
+
+int32 HexToBytes(const FAnsiStringView& Hex, uint8* OutBytes)
+{
+	return HexToBytesImpl<ANSICHAR>(Hex, OutBytes);
+}
+
 }
 }
diff --git a/Engine/Source/Runtime/Core/Private/String/ParseLines.cpp b/Engine/Source/Runtime/Core/Private/String/ParseLines.cpp
--- a/Engine/Source/Runtime/Core/Private/String/ParseLines.cpp
+++ b/Engine/Source/Runtime/Core/Private/String/ParseLines.cpp
@@ -8,36 +8,38 @@ namespace UE
 {
 namespace String
 {
-	void ParseLines(const FStringView& View, const TFunctionRef<void(FStringView)>& Visitor)
+
+void ParseLines(const FStringView& View, const TFunctionRef<void(FStringView)>& Visitor)
+{
+	const TCHAR* ViewIt = View.GetData();
+	const TCHAR* ViewEnd = ViewIt + View.Len();
+	do
 	{
-		const TCHAR* ViewIt = View.GetData();
-		const TCHAR* ViewEnd = ViewIt + View.Len();
-		do
+		const TCHAR* LineStart = ViewIt;
+		const TCHAR* LineEnd = ViewEnd;
+		for (; ViewIt != ViewEnd; ++ViewIt)
 		{
-			const TCHAR* LineStart = ViewIt;
-			const TCHAR* LineEnd = ViewEnd;
-			for (; ViewIt != ViewEnd; ++ViewIt)
+			const TCHAR CurrentChar = *ViewIt;
+			if (CurrentChar == TEXT('\n'))
 			{
-				const TCHAR CurrentChar = *ViewIt;
-				if (CurrentChar == TEXT('\n'))
-				{
-					LineEnd = ViewIt++;
-					break;
-				}
-				if (CurrentChar == TEXT('\r'))
+				LineEnd = ViewIt++;
+				break;
+			}
+			if (CurrentChar == TEXT('\r'))
+			{
+				LineEnd = ViewIt++;
+				if (ViewIt != ViewEnd && *ViewIt == TEXT('\n'))
 				{
-					LineEnd = ViewIt++;
-					if (ViewIt != ViewEnd && *ViewIt == TEXT('\n'))
-					{
-						++ViewIt;
-					}
-					break;
+					++ViewIt;
 				}
+				break;
 			}
-
-			Visitor(FStringView(LineStart, static_cast<FStringView::SizeType>(LineEnd - LineStart)));
 		}
-		while (ViewIt != ViewEnd);
+
+		Visitor(FStringView(LineStart, static_cast<FStringView::SizeType>(LineEnd - LineStart)));
 	}
+	while (ViewIt != ViewEnd);
+}
+
 }
 }
diff --git a/Engine/Source/Runtime/Core/Private/String/ParseTokens.cpp b/Engine/Source/Runtime/Core/Private/String/ParseTokens.cpp
--- a/Engine/Source/Runtime/Core/Private/String/ParseTokens.cpp
+++ b/Engine/Source/Runtime/Core/Private/String/ParseTokens.cpp
@@ -12,164 +12,166 @@ namespace UE
 {
 namespace String
 {
-	/** Parse tokens with one single-character delimiter. */
-	inline static void ParseTokens1Delim1Char(const FStringView& View, const TCHAR Delimiter, TFunctionRef<void(FStringView)> Visitor)
-	{
-		const TCHAR* ViewIt = View.GetData();
-		const TCHAR* const ViewEnd = ViewIt + View.Len();
-		const TCHAR* NextToken = ViewIt;
 
-		for (;;)
+/** Parse tokens with one single-character delimiter. */
+inline static void ParseTokens1Delim1Char(const FStringView& View, const TCHAR Delimiter, TFunctionRef<void(FStringView)> Visitor)
+{
+	const TCHAR* ViewIt = View.GetData();
+	const TCHAR* const ViewEnd = ViewIt + View.Len();
+	const TCHAR* NextToken = ViewIt;
+
+	for (;;)
+	{
+		if (ViewIt == ViewEnd)
 		{
-			if (ViewIt == ViewEnd)
-			{
-				break;
-			}
-			if (*ViewIt != Delimiter)
-			{
-				++ViewIt;
-				continue;
-			}
-			Visitor(FStringView(NextToken, static_cast<FStringView::SizeType>(ViewIt - NextToken)));
-			NextToken = ++ViewIt;
+			break;
+		}
+		if (*ViewIt != Delimiter)
+		{
+			++ViewIt;
+			continue;
 		}
-
 		Visitor(FStringView(NextToken, static_cast<FStringView::SizeType>(ViewIt - NextToken)));
+		NextToken = ++ViewIt;
 	}
 
-	/** Parse tokens with multiple single-character Basic Latin delimiters. */
-	inline static void ParseTokensNDelim1CharBasicLatin(const FStringView& View, TArrayView<const TCHAR> Delimiters, TFunctionRef<void(FStringView)> Visitor)
+	Visitor(FStringView(NextToken, static_cast<FStringView::SizeType>(ViewIt - NextToken)));
+}
+
+/** Parse tokens with multiple single-character Basic Latin delimiters. */
+inline static void ParseTokensNDelim1CharBasicLatin(const FStringView& View, TArrayView<const TCHAR> Delimiters, TFunctionRef<void(FStringView)> Visitor)
+{
+	TBitArray<> DelimiterMask(false, 128);
+	for (uint32 Delimiter : Delimiters)
 	{
-		TBitArray<> DelimiterMask(false, 128);
-		for (uint32 Delimiter : Delimiters)
-		{
-			DelimiterMask[Delimiter] = true;
-		}
+		DelimiterMask[Delimiter] = true;
+	}
 
-		const TCHAR* ViewIt = View.GetData();
-		const TCHAR* const ViewEnd = ViewIt + View.Len();
-		const TCHAR* NextToken = ViewIt;
+	const TCHAR* ViewIt = View.GetData();
+	const TCHAR* const ViewEnd = ViewIt + View.Len();
+	const TCHAR* NextToken = ViewIt;
 
-		for (;;)
+	for (;;)
+	{
+		if (ViewIt == ViewEnd)
 		{
-			if (ViewIt == ViewEnd)
-			{
-				break;
-			}
-			const uint32 CodePoint = *ViewIt;
-			if (CodePoint >= 128 || !DelimiterMask[CodePoint])
-			{
-				++ViewIt;
-				continue;
-			}
-			Visitor(FStringView(NextToken, static_cast<FStringView::SizeType>(ViewIt - NextToken)));
-			NextToken = ++ViewIt;
+			break;
+		}
+		const uint32 CodePoint = *ViewIt;
+		if (CodePoint >= 128 || !DelimiterMask[CodePoint])
+		{
+			++ViewIt;
+			continue;
 		}
-
 		Visitor(FStringView(NextToken, static_cast<FStringView::SizeType>(ViewIt - NextToken)));
+		NextToken = ++ViewIt;
 	}
 
-	/** Parse tokens with multiple single-character delimiters in the Basic Multilingual Plane. */
-	inline static void ParseTokensNDelim1Char(const FStringView& View, TArrayView<const TCHAR> Delimiters, TFunctionRef<void(FStringView)> Visitor)
+	Visitor(FStringView(NextToken, static_cast<FStringView::SizeType>(ViewIt - NextToken)));
+}
+
+/** Parse tokens with multiple single-character delimiters in the Basic Multilingual Plane. */
+inline static void ParseTokensNDelim1Char(const FStringView& View, TArrayView<const TCHAR> Delimiters, TFunctionRef<void(FStringView)> Visitor)
+{
+	if (Algo::AllOf(Delimiters, [](TCHAR Delimiter) { return Delimiter < 128; }))
 	{
-		if (Algo::AllOf(Delimiters, [](TCHAR Delimiter) { return Delimiter < 128; }))
-		{
-			return ParseTokensNDelim1CharBasicLatin(View, Delimiters, Visitor);
-		}
+		return ParseTokensNDelim1CharBasicLatin(View, Delimiters, Visitor);
+	}
 
-		const TCHAR* ViewIt = View.GetData();
-		const TCHAR* const ViewEnd = ViewIt + View.Len();
-		const TCHAR* NextToken = ViewIt;
+	const TCHAR* ViewIt = View.GetData();
+	const TCHAR* const ViewEnd = ViewIt + View.Len();
+	const TCHAR* NextToken = ViewIt;
 
-		for (;;)
+	for (;;)
+	{
+		if (ViewIt == ViewEnd)
 		{
-			if (ViewIt == ViewEnd)
-			{
-				break;
-			}
-			if (!Algo::Find(Delimiters, *ViewIt))
-			{
-				++ViewIt;
-				continue;
-			}
-			Visitor(FStringView(NextToken, static_cast<FStringView::SizeType>(ViewIt - NextToken)));
-			NextToken = ++ViewIt;
+			break;
+		}
+		if (!Algo::Find(Delimiters, *ViewIt))
+		{
+			++ViewIt;
+			continue;
 		}
-
 		Visitor(FStringView(NextToken, static_cast<FStringView::SizeType>(ViewIt - NextToken)));
+		NextToken = ++ViewIt;
 	}
 
-	/** Parse tokens with multiple multi-character delimiters. */
-	inline static void ParseTokensNDelimNChar(const FStringView& View, TArrayView<const FStringView> Delimiters, TFunctionRef<void(FStringView)> Visitor)
-	{
-		// This is a naive implementation that takes time proportional to View.Len() * TotalDelimiterLen.
-		// If this function becomes a bottleneck, it can be specialized separately for one and many delimiters.
-		// There are algorithms for each are linear or sub-linear in the length of string to search.
+	Visitor(FStringView(NextToken, static_cast<FStringView::SizeType>(ViewIt - NextToken)));
+}
+
+/** Parse tokens with multiple multi-character delimiters. */
+inline static void ParseTokensNDelimNChar(const FStringView& View, TArrayView<const FStringView> Delimiters, TFunctionRef<void(FStringView)> Visitor)
+{
+	// This is a naive implementation that takes time proportional to View.Len() * TotalDelimiterLen.
+	// If this function becomes a bottleneck, it can be specialized separately for one and many delimiters.
+	// There are algorithms for each are linear or sub-linear in the length of string to search.
 
-		const FStringView::SizeType ViewLen = View.Len();
-		FStringView::SizeType NextTokenIndex = 0;
+	const FStringView::SizeType ViewLen = View.Len();
+	FStringView::SizeType NextTokenIndex = 0;
 
-		for (FStringView::SizeType ViewIndex = 0; ViewIndex != ViewLen;)
+	for (FStringView::SizeType ViewIndex = 0; ViewIndex != ViewLen;)
+	{
+		const FStringView RemainingView(View.GetData() + ViewIndex, ViewLen - ViewIndex);
+		auto MatchDelimiter = [&RemainingView](const FStringView& Delimiter) { return RemainingView.StartsWith(Delimiter, ESearchCase::CaseSensitive); };
+		if (const FStringView* Delimiter = Algo::FindByPredicate(Delimiters, MatchDelimiter))
 		{
-			const FStringView RemainingView(View.GetData() + ViewIndex, ViewLen - ViewIndex);
-			auto MatchDelimiter = [&RemainingView](const FStringView& Delimiter) { return RemainingView.StartsWith(Delimiter, ESearchCase::CaseSensitive); };
-			if (const FStringView* Delimiter = Algo::FindByPredicate(Delimiters, MatchDelimiter))
-			{
-				Visitor(FStringView(View.GetData() + NextTokenIndex, ViewIndex - NextTokenIndex));
-				ViewIndex += Delimiter->Len();
-				NextTokenIndex = ViewIndex;
-			}
-			else
-			{
-				++ViewIndex;
-			}
+			Visitor(FStringView(View.GetData() + NextTokenIndex, ViewIndex - NextTokenIndex));
+			ViewIndex += Delimiter->Len();
+			NextTokenIndex = ViewIndex;
+		}
+		else
+		{
+			++ViewIndex;
 		}
-
-		Visitor(FStringView(View.GetData() + NextTokenIndex, ViewLen - NextTokenIndex));
 	}
 
-	void ParseTokensMultiple(const FStringView& View, TArrayView<const FStringView> Delimiters, TFunctionRef<void(FStringView)> Visitor)
+	Visitor(FStringView(View.GetData() + NextTokenIndex, ViewLen - NextTokenIndex));
+}
+
+void ParseTokensMultiple(const FStringView& View, TArrayView<const FStringView> Delimiters, TFunctionRef<void(FStringView)> Visitor)
+{
+	check(Algo::NoneOf(Delimiters, &FStringView::IsEmpty));
+	switch (Delimiters.Num())
 	{
-		check(Algo::NoneOf(Delimiters, &FStringView::IsEmpty));
-		switch (Delimiters.Num())
+	case 0:
+		return Visitor(View);
+	case 1:
+		if (Delimiters[0].Len() == 1)
+		{
+			return ParseTokens1Delim1Char(View, Delimiters[0][0], MoveTemp(Visitor));
+		}
+		return ParseTokensNDelimNChar(View, Delimiters, MoveTemp(Visitor));
+	default:
+		if (Algo::AllOf(Delimiters, [](const FStringView& Delimiter) { return Delimiter.Len() == 1; }))
 		{
-		case 0:
-			return Visitor(View);
-		case 1:
-			if (Delimiters[0].Len() == 1)
+			TArray<TCHAR, TInlineAllocator<32>> DelimiterChars;
+			DelimiterChars.Reserve(Delimiters.Num());
+			for (const FStringView& Delimiter : Delimiters)
 			{
-				return ParseTokens1Delim1Char(View, Delimiters[0][0], MoveTemp(Visitor));
+				DelimiterChars.Add(Delimiter[0]);
 			}
+			return ParseTokensNDelim1Char(View, DelimiterChars, MoveTemp(Visitor));
+		}
+		else
+		{
 			return ParseTokensNDelimNChar(View, Delimiters, MoveTemp(Visitor));
-		default:
-			if (Algo::AllOf(Delimiters, [](const FStringView& Delimiter) { return Delimiter.Len() == 1; }))
-			{
-				TArray<TCHAR, TInlineAllocator<32>> DelimiterChars;
-				DelimiterChars.Reserve(Delimiters.Num());
-				for (const FStringView& Delimiter : Delimiters)
-				{
-					DelimiterChars.Add(Delimiter[0]);
-				}
-				return ParseTokensNDelim1Char(View, DelimiterChars, MoveTemp(Visitor));
-			}
-			else
-			{
-				return ParseTokensNDelimNChar(View, Delimiters, MoveTemp(Visitor));
-			}
 		}
 	}
+}
 
-	void ParseTokensMultiple(const FStringView& View, TArrayView<const TCHAR> Delimiters, TFunctionRef<void(FStringView)> Visitor)
+void ParseTokensMultiple(const FStringView& View, TArrayView<const TCHAR> Delimiters, TFunctionRef<void(FStringView)> Visitor)
+{
+	switch (Delimiters.Num())
 	{
-		switch (Delimiters.Num())
-		{
-		case 0:
-			return Visitor(View);
-		case 1:
-			return ParseTokens1Delim1Char(View, Delimiters[0], MoveTemp(Visitor));
-		default:
-			return ParseTokensNDelim1Char(View, Delimiters, MoveTemp(Visitor));
-		}
+	case 0:
+		return Visitor(View);
+	case 1:
+		return ParseTokens1Delim1Char(View, Delimiters[0], MoveTemp(Visitor));
+	default:
+		return ParseTokensNDelim1Char(View, Delimiters, MoveTemp(Visitor));
 	}
 }
+
+}
 }
