Split FPreloadableFile into FPreloadableArchive base class and FPreloadableFile subclass so that we can use FPreloadableArchive with the planned IPackageResourceManager alternative to IFileManager.
#rnx

diff --git a/Engine/Source/Runtime/Core/Private/Misc/PreloadableFile.cpp b/Engine/Source/Runtime/Core/Private/Misc/PreloadableFile.cpp
--- a/Engine/Source/Runtime/Core/Private/Misc/PreloadableFile.cpp
+++ b/Engine/Source/Runtime/Core/Private/Misc/PreloadableFile.cpp
@@ -17,7 +17,7 @@
 #define PRELOADABLEFILE_COOK_STATS_ENABLED 0 && ENABLE_COOK_STATS
 
 #if PRELOADABLEFILE_COOK_STATS_ENABLED
-namespace FPreloadableFileImpl
+namespace FPreloadableArchiveImpl
 {
 	static int64 NumNonPreloadedPages = 0;
 	static int64 NumPreloadedPages = 0;
@@ -34,48 +34,9 @@ namespace FPreloadableFileImpl
 }
 #endif
 
-class FPreloadableFileProxy : public FArchive
-{
-public:
-	explicit FPreloadableFileProxy(const TSharedPtr<FPreloadableFile>& InArchive)
-		:Archive(InArchive)
-	{
-		check(Archive);
-	}
-	virtual void Seek(int64 InPos) final
-	{
-		Archive->Seek(InPos);
-	}
-	virtual int64 Tell() final
-	{
-		return Archive->Tell();
-	}
-	virtual int64 TotalSize() final
-	{
-		return Archive->TotalSize();
-	}
-	virtual bool Close() final
-	{
-		return Archive->Close();
-	}
-	virtual void Serialize(void* V, int64 Length) final
-	{
-		Archive->Serialize(V, Length);
-	}
-	virtual FString GetArchiveName() const final
-	{
-		return Archive->GetArchiveName();
-	}
-
-private:
-	TSharedPtr<FPreloadableFile> Archive;
-};
-
-TMap<FString, TSharedPtr<FPreloadableFile>> FPreloadableFile::RegisteredFiles;
-
-FPreloadableFile::FPreloadableFile(const TCHAR* InFileName)
+FPreloadableArchive::FPreloadableArchive(FStringView InArchiveName)
 	: FArchive()
-	, FileName(InFileName)
+	, ArchiveName(InArchiveName)
 	, bInitialized(false)
 	, bIsPreloading(false)
 	, bIsPreloadingPaused(false)
@@ -83,22 +44,21 @@ FPreloadableFile::FPreloadableFile(const TCHAR* InFileName)
 {
 	PendingAsyncComplete = FPlatformProcess::GetSynchEventFromPool(true);
 	PendingAsyncComplete->Trigger();
-	FPaths::MakeStandardFilename(FileName);
 	this->SetIsLoading(true);
 	this->SetIsPersistent(true);
 }
 
-FPreloadableFile::~FPreloadableFile()
+FPreloadableArchive::~FPreloadableArchive()
 {
 	Close();
 	// It is possible we set a flag indicating that an async event is done, but we haven't yet called Trigger in the task thread; Trigger is always the last memory-accessing instruction on the task thread
-	// This happens for example at the end of FPreloadableFile::InitializeInternal
+	// This happens for example at the end of FPreloadableArchive::InitializeInternal
 	// Wait for the trigger call to be made before deleting PendingAsyncComplete.
 	PendingAsyncComplete->Wait();
 	FPlatformProcess::ReturnSynchEventToPool(PendingAsyncComplete);
 }
 
-void FPreloadableFile::SetPageSize(int64 InPageSize)
+void FPreloadableArchive::SetPageSize(int64 InPageSize)
 {
 	if (!bInitialized)
 	{
@@ -112,27 +72,24 @@ void FPreloadableFile::SetPageSize(int64 InPageSize)
 	PageSize = InPageSize;
 }
 
-void FPreloadableFile::InitializeAsync(uint32 InFlags, int64 PrimeSize)
+void FPreloadableArchive::InitializeAsync(FCreateArchive&& InCreateArchiveFunction, uint32 InFlags, int64 PrimeSize)
 {
-	if (!bInitialized)
-	{
-		PendingAsyncComplete->Wait();
-	}
-	if (bInitialized)
-	{
-		return;
-	}
-	check(PendingAsyncComplete->Wait());
-	PendingAsyncComplete->Reset();
-	Async(EAsyncExecution::TaskGraph, [this, InFlags, PrimeSize] { InitializeInternal(InFlags, PrimeSize); });
+	InFlags = (InFlags & (~Flags::ModeBits)) | Flags::PreloadHandle;
+	InitializeInternalAsync(MoveTemp(InCreateArchiveFunction), FCreateAsyncArchive(), InFlags, PrimeSize);
 }
 
-bool FPreloadableFile::IsInitialized() const
+void FPreloadableArchive::InitializeAsync(FCreateAsyncArchive&& InCreateAsyncArchiveFunction, uint32 InFlags, int64 PrimeSize)
+{
+	InFlags = (InFlags & (~Flags::ModeBits)) | Flags::PreloadBytes;
+	InitializeInternalAsync(FCreateArchive(), MoveTemp(InCreateAsyncArchiveFunction), InFlags, PrimeSize);
+}
+
+bool FPreloadableArchive::IsInitialized() const
 {
 	return bInitialized;
 }
 
-void FPreloadableFile::WaitForInitialization() const
+void FPreloadableArchive::WaitForInitialization() const
 {
 	if (bInitialized)
 	{
@@ -141,17 +98,39 @@ void FPreloadableFile::WaitForInitialization() const
 	PendingAsyncComplete->Wait();
 }
 
-void FPreloadableFile::InitializeInternal(uint32 InFlags, int64 PrimeSize)
+void FPreloadableArchive::InitializeInternalAsync(FCreateArchive&& InCreateArchiveFunction, FCreateAsyncArchive&& InCreateAsyncArchiveFunction, uint32 InFlags, int64 PrimeSize)
+{
+	if (!bInitialized)
+	{
+		PendingAsyncComplete->Wait();
+	}
+	if (bInitialized)
+	{
+		return;
+	}
+	check(PendingAsyncComplete->Wait());
+	PendingAsyncComplete->Reset();
+	Async(EAsyncExecution::TaskGraph, [this, InFlags, PrimeSize, PassedArchiveFunction = MoveTemp(InCreateArchiveFunction), PassedAsyncArchiveFunction = MoveTemp(InCreateAsyncArchiveFunction)]
+		() mutable
+		{
+			InitializeInternal(MoveTemp(PassedArchiveFunction), MoveTemp(PassedAsyncArchiveFunction), InFlags, PrimeSize);
+		});
+}
+
+
+void FPreloadableArchive::InitializeInternal(FCreateArchive&& InCreateArchiveFunction, FCreateAsyncArchive&& InCreateAsyncArchiveFunction, uint32 InFlags, int64 PrimeSize)
 {
 	check(!bInitialized);
 
-	IPlatformFile& PlatformFile = FPlatformFileManager::Get().GetPlatformFile();
 	uint32 Mode = InFlags & Flags::ModeBits;
 	switch (Mode)
 	{
 	case Flags::PreloadBytes:
 	{
-		AsynchronousHandle.Reset(PlatformFile.OpenAsyncRead(*FileName));
+		if (InCreateAsyncArchiveFunction)
+		{
+			AsynchronousHandle.Reset(InCreateAsyncArchiveFunction());
+		}
 		if (AsynchronousHandle)
 		{
 			IAsyncReadRequest* SizeRequest = AsynchronousHandle->SizeRequest();
@@ -171,7 +150,15 @@ void FPreloadableFile::InitializeInternal(uint32 InFlags, int64 PrimeSize)
 	}
 	case Flags::PreloadHandle:
 	{
-		ConstructSynchronousArchive();
+		if (InCreateArchiveFunction)
+		{
+#if PRELOADABLEFILE_COOK_STATS_ENABLED
+			FScopeLock OpenFileTimeScopeLock(&FPreloadableArchiveImpl::OpenFileTimeLock);
+			FScopedDurationTimer ScopedDurationTimer(FPreloadableArchiveImpl::OpenFileTime);
+#endif
+			SynchronousArchive.Reset(InCreateArchiveFunction());
+		}
+
 		if (SynchronousArchive)
 		{
 			Size = SynchronousArchive->TotalSize();
@@ -190,10 +177,10 @@ void FPreloadableFile::InitializeInternal(uint32 InFlags, int64 PrimeSize)
 	}
 
 #if FPRELOADABLEFILE_TEST_ENABLED
-	if (Size != -1)
+	if (Size != -1 && InCreateArchiveFunction)
 	{
-		TestHandle.Reset(PlatformFile.OpenRead(*FileName, false));
-		check(TestHandle);
+		TestArchive.Reset(InCreateArchiveFunction());
+		check(TestArchive);
 	}
 #endif
 
@@ -203,81 +190,7 @@ void FPreloadableFile::InitializeInternal(uint32 InFlags, int64 PrimeSize)
 	PendingAsyncComplete->Trigger();
 }
 
-bool FPreloadableFile::TryRegister(const TSharedPtr<FPreloadableFile>& PreloadableFile)
-{
-	if (!PreloadableFile || !PreloadableFile->IsInitialized() || PreloadableFile->TotalSize() < 0)
-	{
-		return false;
-	}
-
-	TSharedPtr<FPreloadableFile>& ExistingFile = RegisteredFiles.FindOrAdd(PreloadableFile->FileName);
-	if (ExistingFile)
-	{
-		return ExistingFile.Get() == PreloadableFile.Get();
-	}
-
-	ExistingFile = PreloadableFile;
-	return true;
-}
-
-FArchive* FPreloadableFile::TryTakeArchive(const TCHAR* FileName)
-{
-	if (RegisteredFiles.Num() == 0)
-	{
-		return nullptr;
-	}
-
-	FString StandardFileName(FileName);
-	FPaths::MakeStandardFilename(StandardFileName);
-	TSharedPtr<FPreloadableFile> ExistingFile;
-	if (!RegisteredFiles.RemoveAndCopyValue(*StandardFileName, ExistingFile))
-	{
-		return nullptr;
-	}
-	if (!ExistingFile->IsInitialized())
-	{
-		// Someone has called Close on the FPreloadableFile. Unregister it, and behave as if it had not been registered.
-		return nullptr;
-	}
-	if (!ExistingFile->AsynchronousHandle)
-	{
-		// The PreloadableFile is in PreloadHandle mode; it is not preloading bytes, but instead is only providing a pre-opened (and possibly primed) sync handle
-		// Detach the SynchronousHandle from the preloadable file and return that
-		// The SynchronousArchive may be nullptr, which will indicate TryTakeArchive has nothing to offer for the given FileName.
-		return ExistingFile->DetachLowerLevel();
-	}
-	else
-	{
-		// Return a proxy to the Preloadable file; it will use its cache to service serialize requests
-		return new FPreloadableFileProxy(ExistingFile);
-	}
-}
-
-bool FPreloadableFile::UnRegister(const TSharedPtr<FPreloadableFile>& PreloadableFile)
-{
-	if (!PreloadableFile)
-	{
-		return false;
-	}
-
-	TSharedPtr<FPreloadableFile> ExistingFile;
-	if (!RegisteredFiles.RemoveAndCopyValue(PreloadableFile->FileName, ExistingFile))
-	{
-		return false;
-	}
-
-	if (ExistingFile.Get() != PreloadableFile.Get())
-	{
-		// Some other FPreloadableFile was registered for the same FileName. We removed it in the RemoveAndCopyValue above (which we do to optimize the common case).
-		// Add it back, and notify the caller that their PreloadableFile was not registered.
-		RegisteredFiles.Add(PreloadableFile->FileName, MoveTemp(ExistingFile));
-		return false;
-	}
-
-	return true;
-}
-
-bool FPreloadableFile::StartPreload()
+bool FPreloadableArchive::StartPreload()
 {
 	if (bIsPreloading)
 	{
@@ -285,7 +198,7 @@ bool FPreloadableFile::StartPreload()
 	}
 	if (!bInitialized)
 	{
-		UE_LOG(LogCore, Error, TEXT("Attempted FPreloadableFile::StartPreload when uninitialized. Call will be ignored."));
+		UE_LOG(LogCore, Error, TEXT("Attempted FPreloadableArchive::StartPreload when uninitialized. Call will be ignored."));
 		return false;
 	}
 	if (!AllocateCache())
@@ -304,7 +217,7 @@ bool FPreloadableFile::StartPreload()
 	return true;
 }
 
-void FPreloadableFile::StopPreload()
+void FPreloadableArchive::StopPreload()
 {
 	if (!bIsPreloading)
 	{
@@ -317,14 +230,14 @@ void FPreloadableFile::StopPreload()
 	bIsPreloadingPaused = false;
 }
 
-bool FPreloadableFile::IsPreloading() const
+bool FPreloadableArchive::IsPreloading() const
 {
 	// Note that this function is for public use only, and a true result does not indicate we have a currently pending Preload operation;
 	// we may be paused even if bIsPreloading is true.
 	return bIsPreloading;
 }
 
-bool FPreloadableFile::AllocateCache()
+bool FPreloadableArchive::AllocateCache()
 {
 	if (IsCacheAllocated())
 	{
@@ -332,7 +245,7 @@ bool FPreloadableFile::AllocateCache()
 	}
 	if (!bInitialized)
 	{
-		UE_LOG(LogCore, Error, TEXT("Attempted FPreloadableFile::AllocateCache when uninitialized. Call will be ignored."));
+		UE_LOG(LogCore, Error, TEXT("Attempted FPreloadableArchive::AllocateCache when uninitialized. Call will be ignored."));
 		return false;
 	}
 	if (Size < 0)
@@ -349,7 +262,7 @@ bool FPreloadableFile::AllocateCache()
 	return true;
 }
 
-void FPreloadableFile::ReleaseCache()
+void FPreloadableArchive::ReleaseCache()
 {
 	if (!IsCacheAllocated())
 	{
@@ -358,8 +271,8 @@ void FPreloadableFile::ReleaseCache()
 
 	StopPreload();
 #if PRELOADABLEFILE_COOK_STATS_ENABLED
-	FPreloadableFileImpl::NumPreloadedPages += CacheEnd / PageSize;
-	FPreloadableFileImpl::NumNonPreloadedPages += (Size - CacheEnd + PageSize - 1) / PageSize;
+	FPreloadableArchiveImpl::NumPreloadedPages += CacheEnd / PageSize;
+	FPreloadableArchiveImpl::NumNonPreloadedPages += (Size - CacheEnd + PageSize - 1) / PageSize;
 #endif
 	FMemory::Free(CacheBytes);
 	CacheBytes = nullptr;
@@ -367,18 +280,18 @@ void FPreloadableFile::ReleaseCache()
 	RetiredRequests.Shrink();
 }
 
-bool FPreloadableFile::IsCacheAllocated() const
+bool FPreloadableArchive::IsCacheAllocated() const
 {
 	return CacheBytes != nullptr;
 }
 
-FArchive* FPreloadableFile::DetachLowerLevel()
+FArchive* FPreloadableArchive::DetachLowerLevel()
 {
 	WaitForInitialization();
 	return SynchronousArchive.Release();
 }
 
-void FPreloadableFile::PausePreload()
+void FPreloadableArchive::PausePreload()
 {
 	bIsPreloadingPaused = true;
 	PendingAsyncComplete->Wait();
@@ -389,7 +302,7 @@ void FPreloadableFile::PausePreload()
 	}
 }
 
-void FPreloadableFile::ResumePreload()
+void FPreloadableArchive::ResumePreload()
 {
 	// Contract: This function is only called when inside the PreloadLock CriticalSection
 	// Contract: this function is only called when already initialized and no async reads are pending
@@ -417,7 +330,7 @@ void FPreloadableFile::ResumePreload()
 	}
 }
 
-bool FPreloadableFile::ResumePreloadNonRecursive()
+bool FPreloadableArchive::ResumePreloadNonRecursive()
 {
 	check(!PendingAsyncComplete->Wait(0)); // Caller should have set this before calling
 	int64 RemainingSize = Size - CacheEnd;
@@ -453,7 +366,7 @@ bool FPreloadableFile::ResumePreloadNonRecursive()
 	return false;
 }
 
-void FPreloadableFile::OnReadComplete(bool bCanceled, IAsyncReadRequest* ReadRequest)
+void FPreloadableArchive::OnReadComplete(bool bCanceled, IAsyncReadRequest* ReadRequest)
 {
 	TArray<IAsyncReadRequest*> LocalRetired;
 	while (true)
@@ -483,7 +396,7 @@ void FPreloadableFile::OnReadComplete(bool bCanceled, IAsyncReadRequest* ReadReq
 		uint8* ReadResults = ReadRequest->GetReadResults();
 		if (bCanceled || !ReadResults)
 		{
-			UE_LOG(LogCore, Warning, TEXT("Precaching failed for %s: %s."), *FileName, (bCanceled ? TEXT("Canceled") : TEXT("GetReadResults returned null")));
+			UE_LOG(LogCore, Warning, TEXT("Precaching failed for %s: %s."), *GetArchiveName(), (bCanceled ? TEXT("Canceled") : TEXT("GetReadResults returned null")));
 			RetiredRequests.Append(MoveTemp(LocalRetired));
 			FPlatformMisc::MemoryBarrier(); // Make sure we have fully written any values of CacheEnd written earlier in the loop before we set the thread-safe bIsPreloading value to false
 			bIsPreloading = false;
@@ -520,7 +433,7 @@ void FPreloadableFile::OnReadComplete(bool bCanceled, IAsyncReadRequest* ReadReq
 	}
 }
 
-void FPreloadableFile::FreeRetiredRequests()
+void FPreloadableArchive::FreeRetiredRequests()
 {
 	for (IAsyncReadRequest* Retired : RetiredRequests)
 	{
@@ -531,13 +444,13 @@ void FPreloadableFile::FreeRetiredRequests()
 }
 
 
-void FPreloadableFile::Serialize(void* V, int64 Length)
+void FPreloadableArchive::Serialize(void* V, int64 Length)
 {
 #if PRELOADABLEFILE_COOK_STATS_ENABLED
-	FScopedDurationTimer ScopeTimer(FPreloadableFileImpl::SerializeTime);
+	FScopedDurationTimer ScopeTimer(FPreloadableArchiveImpl::SerializeTime);
 #endif
 #if FPRELOADABLEFILE_TEST_ENABLED
-	if (!TestHandle)
+	if (!TestArchive)
 	{
 		SerializeInternal(V, Length);
 		return;
@@ -546,46 +459,46 @@ void FPreloadableFile::Serialize(void* V, int64 Length)
 	int64 SavedPos = Pos;
 
 	bool bWasPreloading = IsPreloading();
-	TestHandle->Seek(Pos);
+	TestArchive->Seek(Pos);
 	TArray64<uint8> TestBytes;
 	TestBytes.AddUninitialized(Length);
-	TestHandle->Read(TestBytes.GetData(), Length);
+	TestArchive->Serialize(TestBytes.GetData(), Length);
 
 	SerializeInternal(V, Length);
 
 	bool bBytesMatch = FMemory::Memcmp(V, TestBytes.GetData(), Length) == 0;
-	bool bPosMatch = Pos == TestHandle->Tell();
+	bool bPosMatch = Pos == TestArchive->Tell();
 	if (!bBytesMatch || !bPosMatch)
 	{
-		UE_LOG(LogCore, Warning, TEXT("FPreloadableFile::Serialize Mismatch on %s. BytesMatch=%s, PosMatch=%s, WasPreloading=%s"),
-			*FileName, (bBytesMatch ? TEXT("true") : TEXT("false")), (bPosMatch ? TEXT("true") : TEXT("false")),
+		UE_LOG(LogCore, Warning, TEXT("FPreloadableArchive::Serialize Mismatch on %s. BytesMatch=%s, PosMatch=%s, WasPreloading=%s"),
+			*GetArchiveName(), (bBytesMatch ? TEXT("true") : TEXT("false")), (bPosMatch ? TEXT("true") : TEXT("false")),
 			(bWasPreloading ? TEXT("true") : TEXT("false")));
 		Seek(SavedPos);
-		TestHandle->Seek(SavedPos);
+		TestArchive->Seek(SavedPos);
 		SerializeInternal(V, Length);
-		TestHandle->Read(TestBytes.GetData(), Length);
+		TestArchive->Serialize(TestBytes.GetData(), Length);
 	}
 }
 
-void FPreloadableFile::SerializeInternal(void* V, int64 Length)
+void FPreloadableArchive::SerializeInternal(void* V, int64 Length)
 {
 #endif
 	if (!bInitialized)
 	{
 		SetError();
-		UE_LOG(LogCore, Error, TEXT("Attempted to Serialize from FPreloadableFile when not initialized."));
+		UE_LOG(LogCore, Error, TEXT("Attempted to Serialize from FPreloadableArchive when not initialized."));
 		return;
 	}
 	if (Pos + Length > Size)
 	{
 		SetError();
-		UE_LOG(LogCore, Error, TEXT("Requested read of %d bytes when %d bytes remain (file=%s, size=%d)"), Length, Size - Pos, *FileName, Size);
+		UE_LOG(LogCore, Error, TEXT("Requested read of %d bytes when %d bytes remain (file=%s, size=%d)"), Length, Size - Pos, *GetArchiveName(), Size);
 		return;
 	}
 
 	if (!IsCacheAllocated())
 	{
-		SerializeFromSynchronousArchive(V, Length);
+		SerializeSynchronously(V, Length);
 		return;
 	}
 
@@ -628,78 +541,101 @@ void FPreloadableFile::SerializeInternal(void* V, int64 Length)
 			}
 
 			int64 ReadLength = EndPos - Pos;
-			SerializeFromSynchronousArchive(V, ReadLength);
+			SerializeSynchronously(V, ReadLength);
 			V = ((uint8*)V) + ReadLength;
 			// SerializeBuffered incremented Pos
 		}
 	}
 }
 
-void FPreloadableFile::SerializeFromSynchronousArchive(void* V, int64 Length)
+void FPreloadableArchive::SerializeSynchronously(void* V, int64 Length)
 {
-	if (!SynchronousArchive)
+	if (SynchronousArchive)
 	{
-		ConstructSynchronousArchive();
-		if (!SynchronousArchive)
+		SynchronousArchive->Seek(Pos);
+		if (SynchronousArchive->IsError())
 		{
-			UE_LOG(LogCore, Warning, TEXT("Failed to open file for %s"), *FileName);
-			SetError();
-			Pos += Length;
-			return;
+			if (!IsError())
+			{
+				UE_LOG(LogCore, Warning, TEXT("Failed to seek to offset %ld in %s."), Pos, *GetArchiveName());
+				SetError();
+			}
+		}
+		else
+		{
+			SynchronousArchive->Serialize(V, Length);
+			if (SynchronousArchive->IsError() && !IsError())
+			{
+				UE_LOG(LogCore, Warning, TEXT("Failed to read %ld bytes at offset %ld in %s."), Length, Pos, *GetArchiveName());
+				SetError();
+			}
 		}
+		Pos += Length;
 	}
-	SynchronousArchive->Seek(Pos);
-	if (SynchronousArchive->IsError())
+	else if (AsynchronousHandle)
 	{
-		if (!IsError())
+		if (Length == 0)
 		{
-			UE_LOG(LogCore, Warning, TEXT("Failed to seek to offset %ld in %s."), Pos, *FileName);
-			SetError();
+			// 0 length ReadRequests are not allowed on IAsyncReadFileHandle
+			return;
+		}
+		IAsyncReadRequest* ReadRequest = AsynchronousHandle->ReadRequest(Pos, Length, AIOP_Normal, nullptr, reinterpret_cast<uint8*>(V));
+		if (!ReadRequest)
+		{
+			if (!IsError())
+			{
+				UE_LOG(LogCore, Warning, TEXT("Failed to create ReadRequest to offset %ld in %s."), Pos, *GetArchiveName());
+				SetError();
+			}
+		}
+		else
+		{
+			if (!ReadRequest->WaitCompletion())
+			{
+				if (!IsError())
+				{
+					UE_LOG(LogCore, Warning, TEXT("Failed to WaitCompletion on ReadRequest to offset %ld in %s."), Pos, *GetArchiveName());
+					SetError();
+				}
+			}
+			else
+			{
+				ensure(ReadRequest->GetReadResults() == V);
+				Pos += Length;
+			}
+
+			delete ReadRequest;
 		}
 	}
 	else
 	{
-		SynchronousArchive->Serialize(V, Length);
-		if (SynchronousArchive->IsError() && !IsError())
+		if (!IsError())
 		{
-			UE_LOG(LogCore, Warning, TEXT("Failed to read %ld bytes at offset %ld in %s."), Length, Pos, *FileName);
+			UE_LOG(LogCore, Warning, TEXT("No InnerArchive available for serialization in %s"), *GetArchiveName());
 			SetError();
 		}
 	}
-	Pos += Length;
 	return;
 }
 
-void FPreloadableFile::ConstructSynchronousArchive()
+void FPreloadableArchive::Seek(int64 InPos)
 {
-	check(!SynchronousArchive);
-	IPlatformFile& PlatformFile = FPlatformFileManager::Get().GetPlatformFile();
-#if PRELOADABLEFILE_COOK_STATS_ENABLED
-	FScopeLock OpenFileTimeScopeLock(&FPreloadableFileImpl::OpenFileTimeLock);
-	FScopedDurationTimer ScopedDurationTimer(FPreloadableFileImpl::OpenFileTime);
-#endif
-
-	SynchronousArchive.Reset(IFileManager::Get().CreateFileReader(*FileName));
-}
-
-void FPreloadableFile::Seek(int64 InPos)
-{
-	checkf(InPos >= 0, TEXT("Attempted to seek to a negative location (%lld/%lld), file: %s. The file is most likely corrupt."), InPos, Size, *FileName);
-	checkf(InPos <= Size, TEXT("Attempted to seek past the end of file (%lld/%lld), file: %s. The file is most likely corrupt."), InPos, Size, *FileName);
+	checkf(InPos >= 0, TEXT("Attempted to seek to a negative location (%lld/%lld), file: %s. The file is most likely corrupt."), InPos, Size, *GetArchiveName());
+	checkf(InPos <= Size, TEXT("Attempted to seek past the end of file (%lld/%lld), file: %s. The file is most likely corrupt."), InPos, Size, *GetArchiveName());
 	Pos = InPos;
 }
 
-int64 FPreloadableFile::Tell()
+int64 FPreloadableArchive::Tell()
 {
 	return Pos;
 }
 
-int64 FPreloadableFile::TotalSize()
+int64 FPreloadableArchive::TotalSize()
 {
 	return Size;
 }
 
-bool FPreloadableFile::Close()
+bool FPreloadableArchive::Close()
 {
 	if (!bInitialized)
 	{
@@ -710,14 +646,130 @@ bool FPreloadableFile::Close()
 	AsynchronousHandle.Reset();
 	SynchronousArchive.Reset();
 #if FPRELOADABLEFILE_TEST_ENABLED
-	TestHandle.Reset();
+	TestArchive.Reset();
 #endif
 
 	bInitialized = false;
 	return !IsError();
 }
 
-FString FPreloadableFile::GetArchiveName() const
+FString FPreloadableArchive::GetArchiveName() const
+{
+	if (!ArchiveName.IsEmpty())
+	{
+		return ArchiveName;
+	}
+	else if (SynchronousArchive)
+	{
+		return SynchronousArchive->GetArchiveName();
+	}
+	else
+	{
+		return TEXT("FPreloadableArchive");
+	}
+}
+
+TMap<FString, TSharedPtr<FPreloadableFile>> FPreloadableFile::RegisteredFiles;
+
+FPreloadableFile::FPreloadableFile(FStringView FileName)
+	: FPreloadableArchive(FileName)
+{
+	FPaths::MakeStandardFilename(ArchiveName);
+}
+
+void FPreloadableFile::InitializeAsync(uint32 InFlags, int64 PrimeSize)
 {
-	return FileName;
+	uint32 Mode = InFlags & Flags::ModeBits;
+	switch (Mode)
+	{
+	case Flags::PreloadBytes:
+		InitializeInternalAsync(
+#if FPRELOADABLEFILE_TEST_ENABLED
+			[this]() { return IFileManager::Get().CreateFileReader(*ArchiveName); },
+#else
+			FCreateArchive(),
+#endif
+			[this]() { return FPlatformFileManager::Get().GetPlatformFile().OpenAsyncRead(*ArchiveName); },
+			InFlags, PrimeSize);
+		break;
+	case Flags::PreloadHandle:
+		InitializeInternalAsync(
+			[this]() { return IFileManager::Get().CreateFileReader(*ArchiveName); },
+			FCreateAsyncArchive(),
+			InFlags, PrimeSize);
+		break;
+	default:
+		checkf(false, TEXT("Invalid mode %u."), Mode);
+		break;
+	}
+}
+
+bool FPreloadableFile::TryRegister(const TSharedPtr<FPreloadableFile>& PreloadableFile)
+{
+	if (!PreloadableFile || !PreloadableFile->IsInitialized() || PreloadableFile->TotalSize() < 0)
+	{
+		return false;
+	}
+
+	TSharedPtr<FPreloadableFile>& ExistingFile = RegisteredFiles.FindOrAdd(PreloadableFile->ArchiveName);
+	if (ExistingFile)
+	{
+		return ExistingFile.Get() == PreloadableFile.Get();
+	}
+
+	ExistingFile = PreloadableFile;
+	return true;
+}
+
+FArchive* FPreloadableFile::TryTakeArchive(const TCHAR* FileName)
+{
+	if (RegisteredFiles.Num() == 0)
+	{
+		return nullptr;
+	}
+
+	FString StandardFileName(FileName);
+	FPaths::MakeStandardFilename(StandardFileName);
+	TSharedPtr<FPreloadableFile> ExistingFile;
+	if (!RegisteredFiles.RemoveAndCopyValue(*StandardFileName, ExistingFile))
+	{
+		return nullptr;
+	}
+	if (!ExistingFile->IsInitialized())
+	{
+		// Someone has called Close on the archive already.
+		return nullptr;
+	}
+	FArchive* Result = ExistingFile->DetachLowerLevel();
+	if (Result)
+	{
+		// the PreloadableArchive is in PreloadHandle mode; it is not preloading bytes, but instead is only providing a pre-opened (and possibly primed) sync handle
+		return Result;
+	}
+	// Otherwise the archive is in PreloadBytes mode, and we need to return a proxy to it
+	return new FPreloadableArchiveProxy(ExistingFile);
+}
+
+bool FPreloadableFile::UnRegister(const TSharedPtr<FPreloadableFile>& PreloadableFile)
+{
+	if (!PreloadableFile)
+	{
+		return false;
+	}
+
+	TSharedPtr<FPreloadableFile> ExistingFile;
+	if (!RegisteredFiles.RemoveAndCopyValue(PreloadableFile->ArchiveName, ExistingFile))
+	{
+		return false;
+	}
+
+	if (ExistingFile.Get() != PreloadableFile.Get())
+	{
+		// Some other FPreloadableFile was registered for the same FileName. We removed it in the RemoveAndCopyValue above (which we do to optimize the common case).
+		// Add it back, and notify the caller that their PreloadableFile was not registered.
+		RegisteredFiles.Add(PreloadableFile->ArchiveName, MoveTemp(ExistingFile));
+		return false;
+	}
+
+	return true;
 }
diff --git a/Engine/Source/Runtime/Core/Public/Misc/PreloadableFile.h b/Engine/Source/Runtime/Core/Public/Misc/PreloadableFile.h
--- a/Engine/Source/Runtime/Core/Public/Misc/PreloadableFile.h
+++ b/Engine/Source/Runtime/Core/Public/Misc/PreloadableFile.h
@@ -5,6 +5,7 @@
 #include "CoreTypes.h"
 #include "Containers/Array.h"
 #include "Containers/Map.h"
+#include "Containers/StringView.h"
 #include "Containers/UnrealString.h"
 #include "HAL/CriticalSection.h"
 #include "Serialization/Archive.h"
@@ -21,35 +22,33 @@ class FEvent;
 #define FPRELOADABLEFILE_TEST_ENABLED 0
 
 /**
- * A read-only archive that provides access to a File on disk, similar to FArchiveFileReaderGeneric provided by FFileManagerGeneric, but with support
- * for asynchronous preloading and priming.
+ * A read-only archive that adds support for asynchronous preloading and priming to an inner archive.
  *
  * This class supports two mutually-exclusive modes:
  *   PreloadBytes:
- *     A lower-level asynchronous archive is opened during initialization and size is read.
- *     After initialization, when StartPreload is called, an array of bytes equal in size to the file's size is allocated,
+ *     An asynchronous inner archive is opened using the passed-in CreateAsyncArchive function; this call is made asynchronously on a TaskGraph thread.
+ *     The size is read during initialization.
+ *     After initialization, when StartPreload is called, an array of bytes equal in size to the inner archive's size is allocated,
  *       and an asynchronous ReadRequest is sent to the IAsyncReadFileHandle to read the first <PageSize> bytes of the file.
  *     Upon completion of each in-flight ReadRequest, another asynchronous ReadRequest is issued, until the entire file has been read.
  *     If serialize functions are called beyond the bytes of the file that have been cached so far, they requests are satisfied by synchronous reads.
  *
- *     Activate this mode by passing Flags::PreloadBytes to InitializeAsync.
+ *     Activate this mode by passing an FCreateArchive to InitializeAsync.
  *  PreloadHandle:
- *     A lower-level FArchive is opened for the file using IFileManager::Get().CreateFileReader; this call is made asynchronously on a TaskGraph thread.
- *     Optionally, a precache request is sent to the lower-level FArchive for the first <PrimeSize> bytes; this call is also made asynchronously.
+ *     A synchronous inner archive is opened using the passed-in CreateArchive function; this call is made asynchronously on a TaskGraph thread.
+ *     Optionally, a precache request is sent to the inner archive for the first <PrimeSize> bytes; this call is also made asynchronously.
  *     The created and primed lower-level FArchive can then be detached from this class and handed off to a new owner.
  *
- *    Activate this mode by passing Flags::OpenHandle and (optionally, for the precache request) Flags::Prime to InitializeAsync.
- *
- * This class also supports registration of the members of this class by filename, which other systems in the engine can use to request
- *   an FArchive for the preload file, if it exists, replacing a call they would otherwise make to IFileManager::Get().CreateFileReader.
+ *    Activate this mode by passing an FCreateAsyncArchive and (optionally, for the precache request) Flags::Prime to InitializeAsync.
  *
  * This class is not threadsafe. The public interface can be used at the same time as internal asynchronous tasks are executing, but the
  * public interface can not be used from multiple threads at once.
  */
-class CORE_API FPreloadableFile : public FArchive
+class CORE_API FPreloadableArchive : public FArchive
 {
 public:
-	typedef TUniqueFunction<bool(FPreloadableFile*)> FOnInitialized;
+	typedef TUniqueFunction<FArchive* ()> FCreateArchive;
+	typedef TUniqueFunction<IAsyncReadFileHandle* ()> FCreateAsyncArchive;
 
 	enum Flags
 	{
@@ -69,37 +68,20 @@ public:
 		DefaultPageSize = 64*1024	// The default size of read requests made to the LowerLevel Archive in PreloadBytes mode when reading bytes into the In-Memory cache.
 	};
 
-	FPreloadableFile(const TCHAR* FileName);
-	~FPreloadableFile();
+	FPreloadableArchive(FStringView ArchiveName);
+	virtual ~FPreloadableArchive();
 
 	// Initialization
 	/** Set the PageSize used read requests made to the LowerLevel Archive in PreloadBytes mode when reading bytes into the In-Memory cache. Invalid to set after Initialization; PageSize must be constant during use. */
 	void SetPageSize(int64 PageSize);
 	/** Initialize the FPreloadableFile asynchronously, performing FileOpen operations on another thread. Use IsInitialized or WaitForInitialization to check progress. */
-	void InitializeAsync(uint32 InFlags = Flags::None, int64 PrimeSize=DefaultPrimeSize);
+	void InitializeAsync(FCreateArchive&& InCreateArchiveFunction, uint32 InFlags = Flags::None, int64 PrimeSize = DefaultPrimeSize);
+	void InitializeAsync(FCreateAsyncArchive&& InCreateAsyncArchiveFunction, uint32 InFlags = Flags::None, int64 PrimeSize = DefaultPrimeSize);
 	/** Return whether InitializeAsync has completed. If Close is called, state returns to false until the next call to InitializeAsync. */
 	bool IsInitialized() const;
 	/** Wait for InitializeAsync to complete if it is running, otherwise return immediately. */
 	void WaitForInitialization() const;
 
-	// Registration
-	/**
-	 * Try to register the given FPreloadableFile instance to handle the next call to TryTakeArchive for its FileName.
-	 * Will fail if the instance has not been initialized or if another instance has already registered for the Filename.
-	 * Return whether the instance is currently registered. Returns true if the instance was already registered.
-	 * Registered files are referenced-counted, and the reference will not be dropped until (1) (TryTakeArchive or UnRegister is called) and (2) (PreloadBytes mode only) the archive returned from TryTakeArchive is deleted.
-	 */
-	static bool TryRegister(const TSharedPtr<FPreloadableFile>& PreloadableFile);
-	/**
-	 * Look up an FPreloadableFile instance registered for the given FileName, and return an FArchive from it.
-	 * If found, removes the registration so no future call to TryTakeArchive can sue the same FArchive.
-	 * If the instance is in PreloadHandle mode, the Lower-Level FArchive will be detached from the FPreloadableFile and returned using DetachLowerLevel.
-	 * If the instance is in PreloadBytes mode, a ProxyArchive will be returned that forwards call to the FPreloadableFile instance.
-	 */
-	static FArchive* TryTakeArchive(const TCHAR* FileName);
-	/** Remove the FPreloadableFile instance if it is registered for its FileName. Returns whether the instance was registered. */
-	static bool UnRegister(const TSharedPtr<FPreloadableFile>& PreloadableFile);
-
 	// Preloading
 	/** When in PreloadBytes mode, if not already preloading, allocate if necessary the memory for the preloaded bytes and start the chain of asynchronous ReadRequests for the bytes. Returns whether preloading is now active. */
 	bool StartPreload();
@@ -125,9 +107,11 @@ public:
 	virtual bool Close() final;
 	virtual FString GetArchiveName() const final;
 
-private:
+protected:
+	/** Helper function for InitializeAsync, sets up the asynchronous call to InitializeInternal */
+	void InitializeInternalAsync(FCreateArchive&& InCreateArchiveFunction, FCreateAsyncArchive&& InCreateAsyncArchiveFunction, uint32 InFlags, int64 PrimeSize);
 	/** Helper function for InitializeAsync, called from a TaskGraph thread. */
-	void InitializeInternal(uint32 Flags, int64 PrimeSize);
+	void InitializeInternal(FCreateArchive&& InCreateArchiveFunction, FCreateAsyncArchive&& InCreateAsyncArchiveFunction, uint32 Flags, int64 PrimeSize);
 #if FPRELOADABLEFILE_TEST_ENABLED
 	void SerializeInternal(void* V, int64 Length);
 #endif
@@ -136,10 +120,9 @@ private:
 	bool ResumePreloadNonRecursive();
 	void OnReadComplete(bool bCanceled, IAsyncReadRequest* ReadRequest);
 	void FreeRetiredRequests();
-	void SerializeFromSynchronousArchive(void* V, int64 Length);
-	void ConstructSynchronousArchive();
+	void SerializeSynchronously(void* V, int64 Length);
 
-	FString FileName;
+	FString ArchiveName;
 	/** The Offset into the file or preloaded bytes that will be used in the next call to Serialize. */
 	int64 Pos = 0;
 	/** The number of bytes in the file. */
@@ -169,7 +152,7 @@ private:
 	TUniquePtr<FArchive> SynchronousArchive;
 #if FPRELOADABLEFILE_TEST_ENABLED
 	/** A duplicate handle used in serialize to validate the returned bytes are correct. */
-	TUniquePtr<IFileHandle> TestHandle;
+	TUniquePtr<FArchive> TestArchive;
 #endif
 	/** ReadRequests that have completed but we have not yet deleted. */
 	TArray<IAsyncReadRequest*> RetiredRequests;
@@ -204,7 +187,91 @@ private:
 	}
 	SavedReadCompleteArguments;
 
+	friend class FPreloadableFileProxy;
+};
+
+/**
+ * An FPreloadableArchive that is customized for reading files from IFileManager.
+ *
+ * This class also supports registration of instances of this class by filename, which other systems in the engine can use to request
+ *   an FArchive for the preload file, if it exists, replacing a call they would otherwise make to IFileManager::Get().CreateFileReader.
+ *
+ * As with the base class, the preloading can work in either PreloadBytes or PreloadHandle mode.
+ *
+ * Activate PreloadBytes mode by passing Flags::PreloadBytes to InitializeAsync.
+ * Activate PreloadHandle mode by passing Flags::PreloadHandle to InitializeAsync, optionally or'd with Flags::Prime.
+ *
+ */
+class CORE_API FPreloadableFile : public FPreloadableArchive
+{
+public:
+	FPreloadableFile(FStringView FileName);
+
+	/** Initialize the FPreloadableFile asynchronously, performing FileOpen operations on another thread. Use IsInitialized or WaitForInitialization to check progress. */
+	void InitializeAsync(uint32 InFlags = Flags::None, int64 PrimeSize=DefaultPrimeSize);
+
+	// Registration
+	/**
+	 * Try to register the given FPreloadableFile instance to handle the next call to TryTakeArchive for its FileName.
+	 * Will fail if the instance has not been initialized or if another instance has already registered for the Filename.
+	 * Return whether the instance is currently registered. Returns true if the instance was already registered.
+	 * Registered files are referenced-counted, and the reference will not be dropped until (1) (TryTakeArchive or UnRegister is called) and (2) (PreloadBytes mode only) the archive returned from TryTakeArchive is deleted.
+	 */
+	static bool TryRegister(const TSharedPtr<FPreloadableFile>& PreloadableFile);
+	/**
+	 * Look up an FPreloadableFile instance registered for the given FileName, and return an FArchive from it.
+	 * If found, removes the registration so no future call to TryTakeArchive can sue the same FArchive.
+	 * If the instance is in PreloadHandle mode, the Lower-Level FArchive will be detached from the FPreloadableFile and returned using DetachLowerLevel.
+	 * If the instance is in PreloadBytes mode, a ProxyArchive will be returned that forwards call to the FPreloadableFile instance.
+	 */
+	static FArchive* TryTakeArchive(const TCHAR* FileName);
+	/** Remove the FPreloadableFile instance if it is registered for its FileName. Returns whether the instance was registered. */
+	static bool UnRegister(const TSharedPtr<FPreloadableFile>& PreloadableFile);
+
+private:
 	/** Map used for TryTakeArchive registration. */
 	static TMap<FString, TSharedPtr<FPreloadableFile>> RegisteredFiles;
-	friend class FPreloadableFileProxy;
 };
+
+/**
+ * A helper class for systems that want to make their own registration system.
+ * A proxy archive that keeps a shared pointer to the inner FPreloadableArchive, so that the inner preloadablearchive will remain
+ * alive until at least as long as the proxy is destroyed.
+ */
+class FPreloadableArchiveProxy : public FArchive
+{
+public:
+	explicit FPreloadableArchiveProxy(const TSharedPtr<FPreloadableArchive>& InArchive)
+		:Archive(InArchive)
+	{
+		check(Archive);
+	}
+	virtual void Seek(int64 InPos) final
+	{
+		Archive->Seek(InPos);
+	}
+	virtual int64 Tell() final
+	{
+		return Archive->Tell();
+	}
+	virtual int64 TotalSize() final
+	{
+		return Archive->TotalSize();
+	}
+	virtual bool Close() final
+	{
+		return Archive->Close();
+	}
+	virtual void Serialize(void* V, int64 Length) final
+	{
+		Archive->Serialize(V, Length);
+	}
+	virtual FString GetArchiveName() const final
+	{
+		return Archive->GetArchiveName();
+	}
+
+private:
+	TSharedPtr<FPreloadableArchive> Archive;
+};
+
