SharedBuffer: Rewrote SharedBuffer to make it easier to use and support immutable buffers

#rnx

diff --git a/Engine/Source/Runtime/Core/Private/Memory/SharedBuffer.cpp b/Engine/Source/Runtime/Core/Private/Memory/SharedBuffer.cpp
--- a/Engine/Source/Runtime/Core/Private/Memory/SharedBuffer.cpp
+++ b/Engine/Source/Runtime/Core/Private/Memory/SharedBuffer.cpp
@@ -2,40 +2,88 @@
 
 #include "Memory/SharedBuffer.h"
 
-FSharedBuffer::FSharedBuffer(const uint64 InSize)
+FSharedBufferRef FSharedBuffer::NewBuffer(
+	void* const Data,
+	const uint64 Size,
+	const ESharedBufferFlags Flags,
+	const uint32 DeleterSize)
 {
-	checkf(InSize < (uint64(1) << 48), TEXT("Size %" UINT64_FMT " exceeds maximum shared buffer size of 256 TiB"), InSize);
-	Data = FMemory::Malloc(InSize);
-	DataSize = InSize;
-	Flags = uint16(ESharedBufferFlags::Owned);
+	checkf((DeleterSize == 0) == !(Flags & ESharedBufferFlags::HasDeleter), TEXT("Mismatch between DeleterSize %u and Flags."), DeleterSize);
+
+	void* const BufferMemory = FMemory::Malloc(sizeof(FSharedBuffer) + DeleterSize, alignof(FSharedBuffer));
+
+	FSharedBuffer* const Buffer = new(BufferMemory) FSharedBuffer;
+	Buffer->Data = Data;
+	Buffer->Size = Size;
+	Buffer->ReferenceCountAndFlags = SetFlags(Flags);
+
+	return FSharedBufferRef(SharedBufferPrivate::TSharedBufferPtr</*bAllowNull*/ false, /*bIsConst*/ false, /*bIsWeak*/ false>(Buffer));
 }
 
-FSharedBuffer::FSharedBuffer(EAssumeOwnershipTag, void* const InData, const uint64 InSize)
+void FSharedBuffer::DeleteBuffer(FSharedBuffer* const Buffer)
 {
-	checkf(InSize < (uint64(1) << 48), TEXT("Size %" UINT64_FMT " exceeds maximum shared buffer size of 256 TiB"), InSize);
-	Data = InData;
-	DataSize = InSize;
-	Flags = uint16(ESharedBufferFlags::Owned);
+	checkSlow(!Data && !Size);
+	checkSlow(ReferenceCountAndFlags == 0);
+	Buffer->~FSharedBuffer();
+	FMemory::Free(Buffer);
 }
 
-FSharedBuffer::FSharedBuffer(ECloneTag, const void* const InData, const uint64 InSize)
-	: FSharedBuffer(InSize)
+void FSharedBuffer::ReleaseData()
 {
-	FMemory::Memcpy(Data, InData, InSize);
+	if (IsOwned())
+	{
+		if (GetFlags(ReferenceCountAndFlags) & ESharedBufferFlags::HasDeleter)
+		{
+			reinterpret_cast<FSharedBufferDeleter*>(this + 1)->Free(Data);
+		}
+		else
+		{
+			FMemory::Free(Data);
+		}
+	}
+	Data = nullptr;
+	Size = 0;
+	ReferenceCountAndFlags.fetch_and(~SetFlags(~ESharedBufferFlags::None));
 }
 
-FSharedBuffer::FSharedBuffer(EWrapTag, void* const InData, const uint64 InSize)
+bool FSharedBuffer::TryMakeImmutable() const
 {
-	checkf(InSize < (uint64(1) << 48), TEXT("Size %" UINT64_FMT " exceeds maximum shared buffer size of 256 TiB"), InSize);
-	Data = InData;
-	DataSize = InSize;
-	Flags = uint16(ESharedBufferFlags::None);
+	for (uint64 Value = ReferenceCountAndFlags.load(std::memory_order_relaxed);;)
+	{
+		if (EnumHasAnyFlags(GetFlags(Value), ESharedBufferFlags::Immutable))
+		{
+			return true;
+		}
+		if (GetSharedRefCount(Value) != 1 || GetWeakRefCount(Value) != 1 ||
+			!EnumHasAnyFlags(GetFlags(Value), ESharedBufferFlags::Owned))
+		{
+			return false;
+		}
+		if (ReferenceCountAndFlags.compare_exchange_weak(Value, Value | SetFlags(ESharedBufferFlags::Immutable),
+			std::memory_order_relaxed, std::memory_order_relaxed))
+		{
+			return true;
+		}
+	}
 }
 
-FSharedBuffer::~FSharedBuffer()
+bool FSharedBuffer::TryMakeMutable() const
 {
-	if (IsOwned())
+	for (uint64 Value = ReferenceCountAndFlags.load(std::memory_order_relaxed);;)
 	{
-		FMemory::Free(Data);
+		if (GetSharedRefCount(Value) != 1 || GetWeakRefCount(Value) != 1 ||
+			!EnumHasAnyFlags(GetFlags(Value), ESharedBufferFlags::Owned))
+		{
+			return false;
+		}
+		if (!EnumHasAnyFlags(GetFlags(Value), ESharedBufferFlags::Immutable))
+		{
+			return true;
+		}
+		if (ReferenceCountAndFlags.compare_exchange_weak(Value, Value & ~SetFlags(ESharedBufferFlags::Immutable),
+			std::memory_order_relaxed, std::memory_order_relaxed))
+		{
+			return true;
+		}
 	}
 }
diff --git a/Engine/Source/Runtime/Core/Private/Serialization/CompactBinary.cpp b/Engine/Source/Runtime/Core/Private/Serialization/CompactBinary.cpp
--- a/Engine/Source/Runtime/Core/Private/Serialization/CompactBinary.cpp
+++ b/Engine/Source/Runtime/Core/Private/Serialization/CompactBinary.cpp
@@ -92,7 +92,7 @@ FConstMemoryView FCbField::AsMemoryView() const
 
 uint64 FCbField::Size() const
 {
-	return AsMemoryView().Size();
+	return AsMemoryView().GetSize();
 }
 
 uint64 FCbField::PayloadSize() const
@@ -416,7 +416,7 @@ FTimespan FCbField::AsTimeSpan(FTimespan Default)
 void FCbField::CopyTo(const FMutableMemoryView Buffer) const
 {
 	const FConstMemoryView Source = AsMemoryView();
-	FMemory::Memcpy(Buffer.GetData(), Source.GetData(), FMath::Min(Buffer.Size(), Source.Size()));
+	FMemory::Memcpy(Buffer.GetData(), Source.GetData(), FMath::Min(Buffer.GetSize(), Source.GetSize()));
 }
 
 ///////////////////////////////////////////////////////////////////////////////////////////////////////////////////////
@@ -448,7 +448,7 @@ FConstMemoryView FCbArray::AsMemoryView() const
 
 uint64 FCbArray::Size() const
 {
-	return sizeof(ECbFieldType) + AsMemoryView().Size();
+	return sizeof(ECbFieldType) + AsMemoryView().GetSize();
 }
 
 uint64 FCbArray::Num() const
@@ -486,7 +486,7 @@ void FCbArray::CopyTo(const FMutableMemoryView Buffer) const
 	uint8* BufferBytes = static_cast<uint8*>(Buffer.GetData());
 	*BufferBytes++ = uint8(FCbFieldType::GetType(Type));
 	const FConstMemoryView Source = AsMemoryView();
-	FMemory::Memcpy(BufferBytes, Source.GetData(), FMath::Min(Buffer.Size() - 1, Source.Size()));
+	FMemory::Memcpy(BufferBytes, Source.GetData(), FMath::Min(Buffer.GetSize() - 1, Source.GetSize()));
 }
 
 ///////////////////////////////////////////////////////////////////////////////////////////////////////////////////////
@@ -518,7 +518,7 @@ FConstMemoryView FCbObject::AsMemoryView() const
 
 uint64 FCbObject::Size() const
 {
-	return sizeof(ECbFieldType) + AsMemoryView().Size();
+	return sizeof(ECbFieldType) + AsMemoryView().GetSize();
 }
 
 FCbFieldIterator FCbObject::CreateIterator() const
@@ -571,7 +571,7 @@ void FCbObject::CopyTo(const FMutableMemoryView Buffer) const
 	uint8* BufferBytes = static_cast<uint8*>(Buffer.GetData());
 	*BufferBytes++ = uint8(FCbFieldType::GetType(Type));
 	const FConstMemoryView Source = AsMemoryView();
-	FMemory::Memcpy(BufferBytes, Source.GetData(), FMath::Min(Buffer.Size() - 1, Source.Size()));
+	FMemory::Memcpy(BufferBytes, Source.GetData(), FMath::Min(Buffer.GetSize() - 1, Source.GetSize()));
 }
 
 ///////////////////////////////////////////////////////////////////////////////////////////////////////////////////////
@@ -595,7 +595,7 @@ static ECbFieldType ValidateCbFieldType(FConstMemoryView& View, ECbValidateMode
 {
 	if (FCbFieldType::HasFieldType(Type))
 	{
-		if (View.Size() >= 1)
+		if (View.GetSize() >= 1)
 		{
 			Type = *static_cast<const ECbFieldType*>(View.GetData());
 			View += 1;
@@ -628,7 +628,7 @@ static ECbFieldType ValidateCbFieldType(FConstMemoryView& View, ECbValidateMode
  */
 static uint64 ValidateCbUInt(FConstMemoryView& View, ECbValidateMode Mode, ECbValidateError& Error)
 {
-	if (View.Size() > 0 && View.Size() >= MeasureVarUInt(View.GetData()))
+	if (View.GetSize() > 0 && View.GetSize() >= MeasureVarUInt(View.GetData()))
 	{
 		uint32 ValueByteCount;
 		const uint64 Value = ReadVarUInt(View.GetData(), ValueByteCount);
@@ -654,7 +654,7 @@ static uint64 ValidateCbUInt(FConstMemoryView& View, ECbValidateMode Mode, ECbVa
  */
 static void ValidateCbFloat64(FConstMemoryView& View, ECbValidateMode Mode, ECbValidateError& Error)
 {
-	if (View.Size() >= sizeof(double))
+	if (View.GetSize() >= sizeof(double))
 	{
 		if (EnumHasAnyFlags(Mode, ECbValidateMode::Format))
 		{
@@ -682,7 +682,7 @@ static void ValidateCbFloat64(FConstMemoryView& View, ECbValidateMode Mode, ECbV
 static FAnsiStringView ValidateCbString(FConstMemoryView& View, ECbValidateMode Mode, ECbValidateError& Error)
 {
 	const uint64 NameSize = ValidateCbUInt(View, Mode, Error);
-	if (View.Size() >= NameSize)
+	if (View.GetSize() >= NameSize)
 	{
 		const FAnsiStringView Name(static_cast<const ANSICHAR*>(View.GetData()), static_cast<int32>(NameSize));
 		View += NameSize;
@@ -801,7 +801,7 @@ static void ValidateCbArray(FConstMemoryView& View, ECbValidateMode Mode, ECbVal
 	View += Size;
 
 	const uint64 Count = ValidateCbUInt(ArrayView, Mode, Error);
-	const uint64 FieldsSize = ArrayView.Size();
+	const uint64 FieldsSize = ArrayView.GetSize();
 	const bool bUniformArray = FCbFieldType::GetType(ArrayType) == ECbFieldType::UniformArray;
 	const ECbFieldType ExternalType = bUniformArray ? ValidateCbFieldType(ArrayView, Mode, Error) : ECbFieldType::HasFieldType;
 	FCbUniformFieldsValidator UniformValidator(ExternalType);
@@ -831,7 +831,7 @@ static FCbField ValidateCbField(FConstMemoryView& View, ECbValidateMode Mode, EC
 
 	auto ValidateFixedPayload = [&View, &Error](uint32 PayloadSize)
 	{
-		if (View.Size() >= PayloadSize)
+		if (View.GetSize() >= PayloadSize)
 		{
 			View += PayloadSize;
 		}
@@ -875,7 +875,7 @@ static FCbField ValidateCbField(FConstMemoryView& View, ECbValidateMode Mode, EC
 	case ECbFieldType::Binary:
 	{
 		const uint64 ValueSize = ValidateCbUInt(View, Mode, Error);
-		if (View.Size() < ValueSize)
+		if (View.GetSize() < ValueSize)
 		{
 			AddError(Error, ECbValidateError::OutOfBounds);
 			View.Reset();
@@ -951,7 +951,7 @@ uint64 MeasureCompactBinary(FConstMemoryView View, ECbFieldType Type)
 
 	if (FCbFieldType::HasFieldType(Type))
 	{
-		if (View.Size() == 0)
+		if (View.GetSize() == 0)
 		{
 			return 0;
 		}
@@ -1003,7 +1003,7 @@ uint64 MeasureCompactBinary(FConstMemoryView View, ECbFieldType Type)
 
 	if (FCbFieldType::HasFieldName(Type))
 	{
-		if (View.Size() == 0 || View.Size() < MeasureVarUInt(View.GetData()))
+		if (View.GetSize() == 0 || View.GetSize() < MeasureVarUInt(View.GetData()))
 		{
 			return 0;
 		}
@@ -1012,7 +1012,7 @@ uint64 MeasureCompactBinary(FConstMemoryView View, ECbFieldType Type)
 		const uint64 NameLen = ReadVarUInt(View.GetData(), NameLenByteCount);
 		const uint64 NameSize = NameLen + NameLenByteCount;
 
-		if (bVariableSize && View.Size() < NameSize)
+		if (bVariableSize && View.GetSize() < NameSize)
 		{
 			return 0;
 		}
@@ -1029,7 +1029,7 @@ uint64 MeasureCompactBinary(FConstMemoryView View, ECbFieldType Type)
 	case ECbFieldType::UniformArray:
 	case ECbFieldType::Binary:
 	case ECbFieldType::String:
-		if (View.Size() == 0 || View.Size() < MeasureVarUInt(View.GetData()))
+		if (View.GetSize() == 0 || View.GetSize() < MeasureVarUInt(View.GetData()))
 		{
 			return 0;
 		}
@@ -1041,7 +1041,7 @@ uint64 MeasureCompactBinary(FConstMemoryView View, ECbFieldType Type)
 		}
 	case ECbFieldType::IntegerPositive:
 	case ECbFieldType::IntegerNegative:
-		if (View.Size() == 0)
+		if (View.GetSize() == 0)
 		{
 			return 0;
 		}
diff --git a/Engine/Source/Runtime/Core/Private/Serialization/CompactBinaryWriter.cpp b/Engine/Source/Runtime/Core/Private/Serialization/CompactBinaryWriter.cpp
--- a/Engine/Source/Runtime/Core/Private/Serialization/CompactBinaryWriter.cpp
+++ b/Engine/Source/Runtime/Core/Private/Serialization/CompactBinaryWriter.cpp
@@ -70,7 +70,7 @@ void FCbWriter::Reset()
 FCbFieldRefIterator FCbWriter::Save() const
 {
 	const uint64 Size = SaveSize();
-	const FSharedBufferRef Buffer = MakeSharedBuffer(Size);
+	const FSharedBufferRef Buffer = FSharedBuffer::Alloc(Size);
 	const FMutableMemoryView View = Buffer->GetView();
 	SaveToMemory(View);
 	return FCbFieldRefIterator(FCbFieldRef(Buffer), View.RightChop(Size).GetData());
@@ -81,10 +81,10 @@ FCbFieldIterator FCbWriter::SaveToMemory(const FMutableMemoryView Buffer) const
 	checkf(States.Num() == 1 && States.Last().Flags == EStateFlags::None,
 		TEXT("It is invalid to save while there are incomplete write operations."));
 	checkf(Data.Num() > 0, TEXT("It is invalid to save when nothing has been written."));
-	checkf(Buffer.Size() == Data.Num(),
-		TEXT("Buffer is %" UINT64_FMT " bytes but %" INT64_FMT " is required."), Buffer.Size(), Data.Num());
+	checkf(Buffer.GetSize() == Data.Num(),
+		TEXT("Buffer is %" UINT64_FMT " bytes but %" INT64_FMT " is required."), Buffer.GetSize(), Data.Num());
 	FMemory::Memcpy(Buffer.GetData(), Data.GetData(), Data.Num());
-	return FCbFieldIterator(FCbField(Buffer.GetData()), Buffer.RightChop(Buffer.Size()).GetData());
+	return FCbFieldIterator(FCbField(Buffer.GetData()), Buffer.RightChop(Buffer.GetSize()).GetData());
 }
 
 uint64 FCbWriter::SaveSize() const
@@ -202,7 +202,7 @@ void FCbWriter::MakeFieldsUniform(const int64 FieldBeginOffset, const int64 Fiel
 	}
 	if (!TargetView.IsEmpty())
 	{
-		Data.RemoveAt(FieldEndOffset - TargetView.Size(), TargetView.Size());
+		Data.RemoveAt(FieldEndOffset - TargetView.GetSize(), TargetView.GetSize());
 	}
 }
 
@@ -276,8 +276,8 @@ void FCbWriter::Object(const FCbObject& Value)
 	};
 	const FCbObjectCopy& ValueCopy = static_cast<const FCbObjectCopy&>(Value);
 	const FConstMemoryView SourceView = ValueCopy.AsMemoryView();
-	const int64 TargetOffset = Data.AddUninitialized(SourceView.Size());
-	FMemory::Memcpy(Data.GetData() + TargetOffset, SourceView.GetData(), SourceView.Size());
+	const int64 TargetOffset = Data.AddUninitialized(SourceView.GetSize());
+	FMemory::Memcpy(Data.GetData() + TargetOffset, SourceView.GetData(), SourceView.GetSize());
 	EndField(FCbFieldType::GetType(ValueCopy.GetCopyType()));
 }
 
@@ -335,8 +335,8 @@ void FCbWriter::Array(const FCbArray& Value)
 	};
 	const FCbArrayCopy& ValueCopy = static_cast<const FCbArrayCopy&>(Value);
 	const FConstMemoryView SourceView = ValueCopy.AsMemoryView();
-	const int64 TargetOffset = Data.AddUninitialized(SourceView.Size());
-	FMemory::Memcpy(Data.GetData() + TargetOffset, SourceView.GetData(), SourceView.Size());
+	const int64 TargetOffset = Data.AddUninitialized(SourceView.GetSize());
+	FMemory::Memcpy(Data.GetData() + TargetOffset, SourceView.GetData(), SourceView.GetSize());
 	EndField(FCbFieldType::GetType(ValueCopy.GetCopyType()));
 }
 
diff --git a/Engine/Source/Runtime/Core/Private/Tests/Memory/MemoryViewTest.cpp b/Engine/Source/Runtime/Core/Private/Tests/Memory/MemoryViewTest.cpp
--- a/Engine/Source/Runtime/Core/Private/Tests/Memory/MemoryViewTest.cpp
+++ b/Engine/Source/Runtime/Core/Private/Tests/Memory/MemoryViewTest.cpp
@@ -35,14 +35,14 @@ bool FMemoryViewTest::RunTest(const FString& Parameters)
 	auto TestMutableMemoryView = [this](const FMutableMemoryView& View, void* Data, uint64 Size)
 	{
 		TestEqual(TEXT("MemoryView.GetData()"), View.GetData(), Data);
-		TestEqual(TEXT("MemoryView.Size()"), View.Size(), Size);
+		TestEqual(TEXT("MemoryView.GetSize()"), View.GetSize(), Size);
 		TestEqual(TEXT("MemoryView.IsEmpty()"), View.IsEmpty(), Size == 0);
 	};
 
 	auto TestConstMemoryView = [this](const FConstMemoryView& View, const void* Data, uint64 Size)
 	{
 		TestEqual(TEXT("MemoryView.GetData()"), View.GetData(), Data);
-		TestEqual(TEXT("MemoryView.Size()"), View.Size(), Size);
+		TestEqual(TEXT("MemoryView.GetSize()"), View.GetSize(), Size);
 		TestEqual(TEXT("MemoryView.IsEmpty()"), View.IsEmpty(), Size == 0);
 	};
 
diff --git a/Engine/Source/Runtime/Core/Private/Tests/Memory/SharedBufferTest.cpp b/Engine/Source/Runtime/Core/Private/Tests/Memory/SharedBufferTest.cpp
--- a/Engine/Source/Runtime/Core/Private/Tests/Memory/SharedBufferTest.cpp
+++ b/Engine/Source/Runtime/Core/Private/Tests/Memory/SharedBufferTest.cpp
@@ -8,7 +8,10 @@
 
 ///////////////////////////////////////////////////////////////////////////////////////////////////////////////////////
 
+static_assert(!std::is_constructible<FSharedBuffer>::value, "Invalid constructor");
+
 static_assert(!std::is_constructible<FSharedBufferRef>::value, "Invalid constructor");
+static_assert( std::is_constructible<FSharedBufferRef, const FSharedBufferRef&>::value, "Invalid constructor");
 static_assert(!std::is_constructible<FSharedBufferRef, FSharedBufferPtr>::value, "Invalid constructor");
 static_assert(!std::is_constructible<FSharedBufferRef, FSharedBufferWeakPtr>::value, "Invalid constructor");
 static_assert(!std::is_constructible<FSharedBufferRef, FSharedBufferConstRef>::value, "Invalid constructor");
@@ -16,29 +19,37 @@ static_assert(!std::is_constructible<FSharedBufferRef, FSharedBufferConstPtr>::v
 static_assert(!std::is_constructible<FSharedBufferRef, FSharedBufferConstWeakPtr>::value, "Invalid constructor");
 
 static_assert(!std::is_constructible<FSharedBufferConstRef>::value, "Invalid constructor");
-static_assert(std::is_constructible<FSharedBufferConstRef, FSharedBufferRef>::value, "Missing constructor");
+static_assert( std::is_constructible<FSharedBufferConstRef, const FSharedBufferRef&>::value, "Missing constructor");
 static_assert(!std::is_constructible<FSharedBufferConstRef, FSharedBufferPtr>::value, "Invalid constructor");
 static_assert(!std::is_constructible<FSharedBufferConstRef, FSharedBufferWeakPtr>::value, "Invalid constructor");
+static_assert( std::is_constructible<FSharedBufferConstRef, const FSharedBufferConstRef&>::value, "Invalid constructor");
 static_assert(!std::is_constructible<FSharedBufferConstRef, FSharedBufferConstPtr>::value, "Invalid constructor");
 static_assert(!std::is_constructible<FSharedBufferConstRef, FSharedBufferConstWeakPtr>::value, "Invalid constructor");
 
-static_assert(std::is_constructible<FSharedBufferPtr>::value, "Missing constructor");
-static_assert(std::is_constructible<FSharedBufferPtr, FSharedBufferRef>::value, "Missing constructor");
+static_assert( std::is_constructible<FSharedBufferPtr>::value, "Missing constructor");
+static_assert( std::is_constructible<FSharedBufferPtr, const FSharedBufferRef&>::value, "Missing constructor");
+static_assert( std::is_constructible<FSharedBufferPtr, const FSharedBufferPtr&>::value, "Missing constructor");
+static_assert( std::is_constructible<FSharedBufferPtr, FSharedBufferPtr&&>::value, "Missing constructor");
 static_assert(!std::is_constructible<FSharedBufferPtr, FSharedBufferWeakPtr>::value, "Invalid constructor");
 static_assert(!std::is_constructible<FSharedBufferPtr, FSharedBufferConstRef>::value, "Invalid constructor");
 static_assert(!std::is_constructible<FSharedBufferPtr, FSharedBufferConstPtr>::value, "Invalid constructor");
 static_assert(!std::is_constructible<FSharedBufferPtr, FSharedBufferConstWeakPtr>::value, "Invalid constructor");
 
-static_assert(std::is_constructible<FSharedBufferConstPtr>::value, "Missing constructor");
-static_assert(std::is_constructible<FSharedBufferConstPtr, FSharedBufferRef>::value, "Missing constructor");
-static_assert(std::is_constructible<FSharedBufferConstPtr, FSharedBufferPtr>::value, "Missing constructor");
+static_assert( std::is_constructible<FSharedBufferConstPtr>::value, "Missing constructor");
+static_assert( std::is_constructible<FSharedBufferConstPtr, const FSharedBufferRef&>::value, "Missing constructor");
+static_assert( std::is_constructible<FSharedBufferConstPtr, const FSharedBufferPtr&>::value, "Missing constructor");
+static_assert( std::is_constructible<FSharedBufferConstPtr, FSharedBufferPtr&&>::value, "Missing constructor");
 static_assert(!std::is_constructible<FSharedBufferConstPtr, FSharedBufferWeakPtr>::value, "Invalid constructor");
-static_assert(std::is_constructible<FSharedBufferConstPtr, FSharedBufferConstRef>::value, "Missing constructor");
+static_assert( std::is_constructible<FSharedBufferConstPtr, const FSharedBufferConstRef&>::value, "Missing constructor");
+static_assert( std::is_constructible<FSharedBufferConstPtr, const FSharedBufferConstPtr&>::value, "Missing constructor");
+static_assert( std::is_constructible<FSharedBufferConstPtr, FSharedBufferConstPtr&&>::value, "Missing constructor");
 static_assert(!std::is_constructible<FSharedBufferConstPtr, FSharedBufferConstWeakPtr>::value, "Invalid constructor");
 
-static_assert(std::is_constructible<FSharedBufferWeakPtr>::value, "Missing constructor");
-static_assert(std::is_constructible<FSharedBufferWeakPtr, FSharedBufferRef>::value, "Missing constructor");
-static_assert(std::is_constructible<FSharedBufferWeakPtr, FSharedBufferPtr>::value, "Missing constructor");
+static_assert( std::is_constructible<FSharedBufferWeakPtr>::value, "Missing constructor");
+static_assert( std::is_constructible<FSharedBufferWeakPtr, FSharedBufferRef>::value, "Missing constructor");
+static_assert( std::is_constructible<FSharedBufferWeakPtr, FSharedBufferPtr>::value, "Missing constructor");
+static_assert( std::is_constructible<FSharedBufferWeakPtr, const FSharedBufferWeakPtr&>::value, "Missing constructor");
+static_assert( std::is_constructible<FSharedBufferWeakPtr, FSharedBufferWeakPtr&&>::value, "Missing constructor");
 static_assert(!std::is_constructible<FSharedBufferWeakPtr, FSharedBufferConstRef>::value, "Invalid constructor");
 static_assert(!std::is_constructible<FSharedBufferWeakPtr, FSharedBufferConstPtr>::value, "Invalid constructor");
 static_assert(!std::is_constructible<FSharedBufferWeakPtr, FSharedBufferConstWeakPtr>::value, "Invalid constructor");
@@ -46,79 +57,86 @@ static_assert(!std::is_constructible<FSharedBufferWeakPtr, FSharedBufferConstWea
 static_assert(std::is_constructible<FSharedBufferConstWeakPtr>::value, "Missing constructor");
 static_assert(std::is_constructible<FSharedBufferConstWeakPtr, FSharedBufferRef>::value, "Missing constructor");
 static_assert(std::is_constructible<FSharedBufferConstWeakPtr, FSharedBufferPtr>::value, "Missing constructor");
-static_assert(std::is_constructible<FSharedBufferConstWeakPtr, FSharedBufferWeakPtr>::value, "Missing constructor");
+static_assert(std::is_constructible<FSharedBufferConstWeakPtr, const FSharedBufferWeakPtr&>::value, "Missing constructor");
+static_assert(std::is_constructible<FSharedBufferConstWeakPtr, FSharedBufferWeakPtr&&>::value, "Missing constructor");
 static_assert(std::is_constructible<FSharedBufferConstWeakPtr, FSharedBufferConstRef>::value, "Missing constructor");
 static_assert(std::is_constructible<FSharedBufferConstWeakPtr, FSharedBufferConstPtr>::value, "Missing constructor");
+static_assert(std::is_constructible<FSharedBufferConstWeakPtr, const FSharedBufferConstWeakPtr&>::value, "Missing constructor");
+static_assert(std::is_constructible<FSharedBufferConstWeakPtr, FSharedBufferConstWeakPtr&&>::value, "Missing constructor");
 
 ///////////////////////////////////////////////////////////////////////////////////////////////////////////////////////
 
-static_assert(std::is_assignable<FSharedBufferRef, FSharedBufferRef>::value, "Missing assignment");
+static_assert( std::is_assignable<FSharedBufferRef, const FSharedBufferRef&>::value, "Missing assignment");
 static_assert(!std::is_assignable<FSharedBufferRef, FSharedBufferPtr>::value, "Invalid assignment");
 static_assert(!std::is_assignable<FSharedBufferRef, FSharedBufferWeakPtr>::value, "Invalid assignment");
 static_assert(!std::is_assignable<FSharedBufferRef, FSharedBufferConstRef>::value, "Invalid assignment");
 static_assert(!std::is_assignable<FSharedBufferRef, FSharedBufferConstPtr>::value, "Invalid assignment");
 static_assert(!std::is_assignable<FSharedBufferRef, FSharedBufferConstWeakPtr>::value, "Invalid assignment");
 
-static_assert(std::is_assignable<FSharedBufferConstRef, FSharedBufferRef>::value, "Missing assignment");
+static_assert( std::is_assignable<FSharedBufferConstRef, const FSharedBufferRef&>::value, "Missing assignment");
 static_assert(!std::is_assignable<FSharedBufferConstRef, FSharedBufferPtr>::value, "Invalid assignment");
 static_assert(!std::is_assignable<FSharedBufferConstRef, FSharedBufferWeakPtr>::value, "Invalid assignment");
-static_assert(std::is_assignable<FSharedBufferConstRef, FSharedBufferConstRef>::value, "Missing assignment");
+static_assert( std::is_assignable<FSharedBufferConstRef, const FSharedBufferConstRef&>::value, "Missing assignment");
 static_assert(!std::is_assignable<FSharedBufferConstRef, FSharedBufferConstPtr>::value, "Invalid assignment");
 static_assert(!std::is_assignable<FSharedBufferConstRef, FSharedBufferConstWeakPtr>::value, "Invalid assignment");
 
-static_assert(std::is_assignable<FSharedBufferPtr, FSharedBufferRef>::value, "Missing assignment");
-static_assert(std::is_assignable<FSharedBufferPtr, FSharedBufferPtr>::value, "Missing assignment");
+static_assert( std::is_assignable<FSharedBufferPtr, const FSharedBufferRef&>::value, "Missing assignment");
+static_assert( std::is_assignable<FSharedBufferPtr, const FSharedBufferPtr&>::value, "Missing assignment");
+static_assert( std::is_assignable<FSharedBufferPtr, FSharedBufferPtr&&>::value, "Missing assignment");
 static_assert(!std::is_assignable<FSharedBufferPtr, FSharedBufferWeakPtr>::value, "Invalid assignment");
 static_assert(!std::is_assignable<FSharedBufferPtr, FSharedBufferConstRef>::value, "Invalid assignment");
 static_assert(!std::is_assignable<FSharedBufferPtr, FSharedBufferConstPtr>::value, "Invalid assignment");
 static_assert(!std::is_assignable<FSharedBufferPtr, FSharedBufferConstWeakPtr>::value, "Invalid assignment");
 
-static_assert(std::is_assignable<FSharedBufferConstPtr, FSharedBufferRef>::value, "Missing assignment");
-static_assert(std::is_assignable<FSharedBufferConstPtr, FSharedBufferPtr>::value, "Missing assignment");
+static_assert( std::is_assignable<FSharedBufferConstPtr, const FSharedBufferRef&>::value, "Missing assignment");
+static_assert( std::is_assignable<FSharedBufferConstPtr, const FSharedBufferPtr&>::value, "Missing assignment");
+static_assert( std::is_assignable<FSharedBufferConstPtr, FSharedBufferPtr&&>::value, "Missing assignment");
 static_assert(!std::is_assignable<FSharedBufferConstPtr, FSharedBufferWeakPtr>::value, "Invalid assignment");
-static_assert(std::is_assignable<FSharedBufferConstPtr, FSharedBufferConstRef>::value, "Missing assignment");
-static_assert(std::is_assignable<FSharedBufferConstPtr, FSharedBufferConstPtr>::value, "Missing assignment");
+static_assert( std::is_assignable<FSharedBufferConstPtr, const FSharedBufferConstRef&>::value, "Missing assignment");
+static_assert( std::is_assignable<FSharedBufferConstPtr, const FSharedBufferConstPtr&>::value, "Missing assignment");
+static_assert( std::is_assignable<FSharedBufferConstPtr, FSharedBufferConstPtr&&>::value, "Missing assignment");
 static_assert(!std::is_assignable<FSharedBufferConstPtr, FSharedBufferConstWeakPtr>::value, "Invalid assignment");
 
-static_assert(std::is_assignable<FSharedBufferWeakPtr, FSharedBufferRef>::value, "Missing assignment");
-static_assert(std::is_assignable<FSharedBufferWeakPtr, FSharedBufferPtr>::value, "Missing assignment");
-static_assert(std::is_assignable<FSharedBufferWeakPtr, FSharedBufferWeakPtr>::value, "Missing assignment");
+static_assert( std::is_assignable<FSharedBufferWeakPtr, const FSharedBufferRef&>::value, "Missing assignment");
+static_assert( std::is_assignable<FSharedBufferWeakPtr, const FSharedBufferPtr&>::value, "Missing assignment");
+static_assert( std::is_assignable<FSharedBufferWeakPtr, const FSharedBufferWeakPtr&>::value, "Missing assignment");
+static_assert( std::is_assignable<FSharedBufferWeakPtr, FSharedBufferWeakPtr&&>::value, "Missing assignment");
 static_assert(!std::is_assignable<FSharedBufferWeakPtr, FSharedBufferConstRef>::value, "Invalid assignment");
 static_assert(!std::is_assignable<FSharedBufferWeakPtr, FSharedBufferConstPtr>::value, "Invalid assignment");
 static_assert(!std::is_assignable<FSharedBufferWeakPtr, FSharedBufferConstWeakPtr>::value, "Invalid assignment");
 
-static_assert(std::is_assignable<FSharedBufferConstWeakPtr, FSharedBufferRef>::value, "Missing assignment");
-static_assert(std::is_assignable<FSharedBufferConstWeakPtr, FSharedBufferPtr>::value, "Missing assignment");
-static_assert(std::is_assignable<FSharedBufferConstWeakPtr, FSharedBufferWeakPtr>::value, "Missing assignment");
-static_assert(std::is_assignable<FSharedBufferConstWeakPtr, FSharedBufferConstRef>::value, "Missing assignment");
-static_assert(std::is_assignable<FSharedBufferConstWeakPtr, FSharedBufferConstPtr>::value, "Missing assignment");
-static_assert(std::is_assignable<FSharedBufferConstWeakPtr, FSharedBufferConstWeakPtr>::value, "Missing assignment");
+static_assert(std::is_assignable<FSharedBufferConstWeakPtr, const FSharedBufferRef&>::value, "Missing assignment");
+static_assert(std::is_assignable<FSharedBufferConstWeakPtr, const FSharedBufferPtr&>::value, "Missing assignment");
+static_assert(std::is_assignable<FSharedBufferConstWeakPtr, const FSharedBufferWeakPtr&>::value, "Missing assignment");
+static_assert(std::is_assignable<FSharedBufferConstWeakPtr, FSharedBufferWeakPtr&&>::value, "Missing assignment");
+static_assert(std::is_assignable<FSharedBufferConstWeakPtr, const FSharedBufferConstRef&>::value, "Missing assignment");
+static_assert(std::is_assignable<FSharedBufferConstWeakPtr, const FSharedBufferConstPtr&>::value, "Missing assignment");
+static_assert(std::is_assignable<FSharedBufferConstWeakPtr, FSharedBufferConstPtr&&>::value, "Missing assignment");
+static_assert(std::is_assignable<FSharedBufferConstWeakPtr, const FSharedBufferConstWeakPtr&>::value, "Missing assignment");
+static_assert(std::is_assignable<FSharedBufferConstWeakPtr, FSharedBufferConstWeakPtr&&>::value, "Missing assignment");
 
 ///////////////////////////////////////////////////////////////////////////////////////////////////////////////////////
 
-static_assert(std::is_same<FSharedBufferRef, decltype(MakeSharedBuffer())>::value, "Invalid constructor");
-static_assert(std::is_same<FSharedBufferRef, decltype(MakeSharedBuffer(DeclVal<uint64>()))>::value, "Invalid constructor");
-
-static_assert(std::is_same<FSharedBufferRef, decltype(MakeSharedBuffer(FSharedBuffer::AssumeOwnership, DeclVal<void*>(), 0))>::value, "Invalid constructor");
-static_assert(std::is_same<FSharedBufferConstRef, decltype(MakeSharedBuffer(FSharedBuffer::AssumeOwnership, DeclVal<const void*>(), 0))>::value, "Invalid constructor");
-
-static_assert(std::is_same<FSharedBufferRef, decltype(MakeSharedBuffer(FSharedBuffer::Clone, DeclVal<void*>(), 0))>::value, "Invalid constructor");
-static_assert(std::is_same<FSharedBufferRef, decltype(MakeSharedBuffer(FSharedBuffer::Clone, DeclVal<const void*>(), 0))>::value, "Invalid constructor");
-
-static_assert(std::is_same<FSharedBufferRef, decltype(MakeSharedBuffer(FSharedBuffer::Wrap, DeclVal<void*>(), 0))>::value, "Invalid constructor");
-static_assert(std::is_same<FSharedBufferConstRef, decltype(MakeSharedBuffer(FSharedBuffer::Wrap, DeclVal<const void*>(), 0))>::value, "Invalid constructor");
+static_assert(std::is_same<FSharedBufferRef, decltype(FSharedBuffer::Alloc(DeclVal<uint64>()))>::value, "Invalid constructor");
 
-static_assert(std::is_same<FSharedBufferRef, decltype(MakeSharedBuffer(FSharedBuffer::AssumeOwnership, FMutableMemoryView()))>::value, "Invalid constructor");
-static_assert(std::is_same<FSharedBufferConstRef, decltype(MakeSharedBuffer(FSharedBuffer::AssumeOwnership, FConstMemoryView()))>::value, "Invalid constructor");
+static_assert(std::is_same<FSharedBufferRef, decltype(FSharedBuffer::Clone(DeclVal<void*>(), 0))>::value, "Invalid constructor");
+static_assert(std::is_same<FSharedBufferRef, decltype(FSharedBuffer::Clone(DeclVal<const void*>(), 0))>::value, "Invalid constructor");
+static_assert(std::is_same<FSharedBufferRef, decltype(FSharedBuffer::Clone(*DeclVal<FSharedBufferRef>()))>::value, "Invalid constructor");
+static_assert(std::is_same<FSharedBufferRef, decltype(FSharedBuffer::Clone(*DeclVal<FSharedBufferConstRef>()))>::value, "Invalid constructor");
+static_assert(std::is_same<FSharedBufferRef, decltype(FSharedBuffer::Clone(DeclVal<FMutableMemoryView>()))>::value, "Invalid constructor");
+static_assert(std::is_same<FSharedBufferRef, decltype(FSharedBuffer::Clone(DeclVal<FConstMemoryView>()))>::value, "Invalid constructor");
 
-static_assert(std::is_same<FSharedBufferRef, decltype(MakeSharedBuffer(FSharedBuffer::Clone, FMutableMemoryView()))>::value, "Invalid constructor");
-static_assert(std::is_same<FSharedBufferRef, decltype(MakeSharedBuffer(FSharedBuffer::Clone, FConstMemoryView()))>::value, "Invalid constructor");
+static_assert(std::is_same<FSharedBufferRef, decltype(FSharedBuffer::TakeOwnership(DeclVal<void*>(), 0, FMemory::Free))>::value, "Invalid constructor");
+static_assert(std::is_same<FSharedBufferConstRef, decltype(FSharedBuffer::TakeOwnership(DeclVal<const void*>(), 0, FMemory::Free))>::value, "Invalid constructor");
 
-static_assert(std::is_same<FSharedBufferRef, decltype(MakeSharedBuffer(FSharedBuffer::Wrap, FMutableMemoryView()))>::value, "Invalid constructor");
-static_assert(std::is_same<FSharedBufferConstRef, decltype(MakeSharedBuffer(FSharedBuffer::Wrap, FConstMemoryView()))>::value, "Invalid constructor");
+static_assert(std::is_same<FSharedBufferRef, decltype(FSharedBuffer::MakeView(DeclVal<void*>(), 0))>::value, "Invalid constructor");
+static_assert(std::is_same<FSharedBufferConstRef, decltype(FSharedBuffer::MakeView(DeclVal<const void*>(), 0))>::value, "Invalid constructor");
+static_assert(std::is_same<FSharedBufferRef, decltype(FSharedBuffer::MakeView(DeclVal<FMutableMemoryView>()))>::value, "Invalid constructor");
+static_assert(std::is_same<FSharedBufferConstRef, decltype(FSharedBuffer::MakeView(DeclVal<FConstMemoryView>()))>::value, "Invalid constructor");
 
 static_assert(std::is_same<void*, decltype(DeclVal<FSharedBuffer>().GetData())>::value, "Invalid accessor");
 static_assert(std::is_same<const void*, decltype(DeclVal<const FSharedBuffer>().GetData())>::value, "Invalid accessor");
+static_assert(std::is_same<uint64, decltype(DeclVal<FSharedBuffer>().GetSize())>::value, "Invalid accessor");
 
 static_assert(std::is_same<FMutableMemoryView, decltype(DeclVal<FSharedBuffer>().GetView())>::value, "Invalid accessor");
 static_assert(std::is_same<FConstMemoryView, decltype(DeclVal<const FSharedBuffer>().GetView())>::value, "Invalid accessor");
@@ -130,13 +148,13 @@ static_assert(std::is_convertible<const FSharedBuffer, FConstMemoryView>::value,
 
 ///////////////////////////////////////////////////////////////////////////////////////////////////////////////////////
 
-static_assert(std::is_same<FSharedBufferRef, decltype(MakeSharedBufferOwned(DeclVal<const FSharedBufferRef&>()))>::value, "Invalid conversion");
-static_assert(std::is_same<FSharedBufferPtr, decltype(MakeSharedBufferOwned(DeclVal<const FSharedBufferPtr&>()))>::value, "Invalid conversion");
-static_assert(std::is_same<FSharedBufferPtr, decltype(MakeSharedBufferOwned(DeclVal<FSharedBufferPtr&&>()))>::value, "Invalid conversion");
+static_assert(std::is_same<FSharedBufferRef, decltype(FSharedBuffer::MakeOwned(DeclVal<const FSharedBufferRef&>()))>::value, "Invalid conversion");
+static_assert(std::is_same<FSharedBufferPtr, decltype(FSharedBuffer::MakeOwned(DeclVal<const FSharedBufferPtr&>()))>::value, "Invalid conversion");
+static_assert(std::is_same<FSharedBufferPtr, decltype(FSharedBuffer::MakeOwned(DeclVal<FSharedBufferPtr&&>()))>::value, "Invalid conversion");
 
-static_assert(std::is_same<FSharedBufferConstRef, decltype(MakeSharedBufferOwned(DeclVal<const FSharedBufferConstRef&>()))>::value, "Invalid conversion");
-static_assert(std::is_same<FSharedBufferConstPtr, decltype(MakeSharedBufferOwned(DeclVal<const FSharedBufferConstPtr&>()))>::value, "Invalid conversion");
-static_assert(std::is_same<FSharedBufferConstPtr, decltype(MakeSharedBufferOwned(DeclVal<FSharedBufferConstPtr&&>()))>::value, "Invalid conversion");
+static_assert(std::is_same<FSharedBufferConstRef, decltype(FSharedBuffer::MakeOwned(DeclVal<const FSharedBufferConstRef&>()))>::value, "Invalid conversion");
+static_assert(std::is_same<FSharedBufferConstPtr, decltype(FSharedBuffer::MakeOwned(DeclVal<const FSharedBufferConstPtr&>()))>::value, "Invalid conversion");
+static_assert(std::is_same<FSharedBufferConstPtr, decltype(FSharedBuffer::MakeOwned(DeclVal<FSharedBufferConstPtr&&>()))>::value, "Invalid conversion");
 
 ///////////////////////////////////////////////////////////////////////////////////////////////////////////////////////
 
@@ -154,77 +172,213 @@ static_assert(std::is_same<uint32, decltype(GetTypeHash(DeclVal<const FSharedBuf
 IMPLEMENT_SIMPLE_AUTOMATION_TEST(FSharedBufferTest, "System.Core.Memory.SharedBuffer", EAutomationTestFlags::ApplicationContextMask | EAutomationTestFlags::SmokeFilter)
 bool FSharedBufferTest::RunTest(const FString& Parameters)
 {
-	// Test Default
+	// Test Size
 	{
-		FSharedBufferRef Ref = MakeSharedBuffer();
-		TestFalse(TEXT("FSharedBuffer().IsOwned()"), Ref->IsOwned());
-		TestEqual(TEXT("FSharedBuffer().Size()"), Ref->Size(), uint64(0));
+		constexpr uint64 Size = 64;
+		FSharedBufferRef Ref = FSharedBuffer::Alloc(64);
+		TestTrue(TEXT("FSharedBuffer(Size).IsOwned()"), Ref->IsOwned());
+		TestEqual(TEXT("FSharedBuffer(Size).GetSize()"), Ref->GetSize(), Size);
 	}
 
-	// Test Size
+	// Test Clone
 	{
 		constexpr uint64 Size = 64;
-		FSharedBufferRef Ref = MakeSharedBuffer(64);
-		TestTrue(TEXT("FSharedBuffer(Size).IsOwned()"), Ref->IsOwned());
-		TestEqual(TEXT("FSharedBuffer(Size).Size()"), Ref->Size(), Size);
+		const uint8 Data[Size]{};
+		FSharedBufferRef Ref = FSharedBuffer::Clone(Data, Size);
+		TestTrue(TEXT("FSharedBuffer::Clone().IsOwned()"), Ref->IsOwned());
+		TestEqual(TEXT("FSharedBuffer::Clone().GetSize()"), Ref->GetSize(), Size);
+		TestNotEqual(TEXT("FSharedBuffer::Clone().GetData()"), static_cast<const void*>(Ref->GetData()), static_cast<const void*>(Data));
 	}
 
-	// Test AssumeOwnership
+	// Test MakeView
+	{
+		constexpr uint64 Size = 64;
+		uint8 Data[Size]{};
+		FSharedBufferRef Ref = FSharedBuffer::MakeView(Data, Size);
+		TestFalse(TEXT("FSharedBuffer::MakeView().IsOwned()"), Ref->IsOwned());
+		TestEqual(TEXT("FSharedBuffer::MakeView().GetSize()"), Ref->GetSize(), Size);
+		TestEqual(TEXT("FSharedBuffer::MakeView().GetData()"), Ref->GetData(), static_cast<void*>(Data));
+	}
+
+	// Test TakeOwnership FMemory::Free
 	{
 		constexpr uint64 Size = 64;
 		void* const Data = FMemory::Malloc(Size);
-		FSharedBufferRef Ref = MakeSharedBuffer(FSharedBuffer::AssumeOwnership, Data, Size);
-		TestTrue(TEXT("FSharedBuffer(AssumeOwnership).IsOwned()"), Ref->IsOwned());
-		TestEqual(TEXT("FSharedBuffer(AssumeOwnership).Size()"), Ref->Size(), Size);
-		TestEqual(TEXT("FSharedBuffer(AssumeOwnership).GetData()"), Ref->GetData(), Data);
+		FSharedBufferRef Ref = FSharedBuffer::TakeOwnership(Data, Size, FMemory::Free);
+		TestTrue(TEXT("FSharedBuffer::TakeOwnership(FMemory::Free).IsOwned()"), Ref->IsOwned());
+		TestEqual(TEXT("FSharedBuffer::TakeOwnership(FMemory::Free).GetSize()"), Ref->GetSize(), Size);
+		TestEqual(TEXT("FSharedBuffer::TakeOwnership(FMemory::Free).GetData()"), Ref->GetData(), Data);
 	}
 
-	// Test Clone
+	// Test TakeOwnership Lambda
 	{
+		bool bDeleted = false;
+		auto Deleter = [&bDeleted](void* Data)
+		{
+			bDeleted = true;
+			delete[] static_cast<uint8*>(Data);
+		};
 		constexpr uint64 Size = 64;
-		const uint8 Data[Size]{};
-		FSharedBufferRef Ref = MakeSharedBuffer(FSharedBuffer::Clone, Data, Size);
-		TestTrue(TEXT("FSharedBuffer(Clone).IsOwned()"), Ref->IsOwned());
-		TestEqual(TEXT("FSharedBuffer(Clone).Size()"), Ref->Size(), Size);
-		TestNotEqual(TEXT("FSharedBuffer(Clone).GetData()"), static_cast<const void*>(Ref->GetData()), static_cast<const void*>(Data));
+		FSharedBuffer::TakeOwnership(new uint8[Size], Size, Deleter);
+		TestTrue(TEXT("FSharedBuffer::TakeOwnership(Lambda) Deleted"), bDeleted);
+	}
+
+	// Test MakeOwned
+	{
+		constexpr uint64 Size = 64;
+		uint8 Data[Size]{};
+		FSharedBufferRef Ref = FSharedBuffer::MakeOwned(FSharedBuffer::MakeView(Data, Size));
+		TestTrue(TEXT("FSharedBuffer::MakeOwned(Wrap).IsOwned()"), Ref->IsOwned());
+		TestEqual(TEXT("FSharedBuffer::MakeOwned(Wrap).GetSize()"), Ref->GetSize(), Size);
+		TestNotEqual(TEXT("FSharedBuffer::MakeOwned(Wrap).GetData()"), Ref->GetData(), static_cast<void*>(Data));
+		FSharedBufferRef RefCopy = FSharedBuffer::MakeOwned(Ref);
+		TestTrue(TEXT("FSharedBuffer::MakeOwned(Owned).IsOwned()"), RefCopy->IsOwned());
+		TestEqual(TEXT("FSharedBuffer::MakeOwned(Owned).GetSize()"), RefCopy->GetSize(), Size);
+		TestEqual(TEXT("FSharedBuffer::MakeOwned(Owned).GetData()"), RefCopy->GetData(), Ref->GetData());
+	}
+	{
+		FSharedBufferPtr Ptr;
+		TestFalse(TEXT("FSharedBuffer::MakeOwned(Null)"), FSharedBuffer::MakeOwned(Ptr).IsValid());
+		TestFalse(TEXT("FSharedBuffer::MakeOwned(Null)"), FSharedBuffer::MakeOwned(AsConst(Ptr)).IsValid());
+		TestFalse(TEXT("FSharedBuffer::MakeOwned(Null)"), FSharedBuffer::MakeOwned(FSharedBufferPtr()).IsValid());
+		TestFalse(TEXT("FSharedBuffer::MakeOwned(Null)"), FSharedBuffer::MakeOwned(FSharedBufferConstPtr()).IsValid());
 	}
 
-	// Test Wrap
+	// Test MakeImmutable
 	{
+		// MakeImmutable from a new reference.
+		constexpr uint64 Size = 64;
+		FSharedBufferConstRef Ref = FSharedBuffer::MakeImmutable(FSharedBuffer::Alloc(Size));
+		TestTrue(TEXT("FSharedBuffer::MakeImmutable(AllocRef).IsOwned()"), Ref->IsOwned());
+		TestTrue(TEXT("FSharedBuffer::MakeImmutable(AllocRef).IsImmutable()"), Ref->IsImmutable());
+		TestEqual(TEXT("FSharedBuffer::MakeImmutable(AllocRef).GetSize()"), Ref->GetSize(), Size);
+		FSharedBufferConstRef OtherRef = Ref;
+		FSharedBufferConstRef ImmutableRef = FSharedBuffer::MakeImmutable(MoveTemp(Ref));
+		TestEqual(TEXT("FSharedBuffer::MakeImmutable(AllocRef).GetData()"), Ref->GetData(), ImmutableRef->GetData());
+	}
+	{
+		// MakeImmutable from a new pointer.
+		constexpr uint64 Size = 64;
+		FSharedBufferPtr MutablePtr = FSharedBuffer::Alloc(Size);
+		TestFalse(TEXT("FSharedBuffer::MakeImmutable(MovePtr).IsImmutable()"), MutablePtr->IsImmutable());
+		TestTrue(TEXT("FSharedBuffer::MakeImmutable(MovePtr).IsOwned()"), MutablePtr->IsOwned());
+		const void* const Data = MutablePtr->GetData();
+		FSharedBufferConstPtr Ptr = FSharedBuffer::MakeImmutable(MoveTemp(MutablePtr));
+		TestTrue(TEXT("FSharedBuffer::MakeImmutable(MovePtr).IsOwned()"), Ptr->IsOwned());
+		TestTrue(TEXT("FSharedBuffer::MakeImmutable(MovePtr).IsImmutable()"), Ptr->IsImmutable());
+		TestEqual(TEXT("FSharedBuffer::MakeImmutable(MovePtr).GetSize()"), Ptr->GetSize(), Size);
+		TestEqual(TEXT("FSharedBuffer::MakeImmutable(MovePtr).GetData()"), Ptr->GetData(), Data);
+	}
+	{
+		// MakeImmutable from a view.
 		constexpr uint64 Size = 64;
 		uint8 Data[Size]{};
-		FSharedBufferRef Ref = MakeSharedBuffer(FSharedBuffer::Wrap, Data, Size);
-		TestFalse(TEXT("FSharedBuffer(Wrap).IsOwned()"), Ref->IsOwned());
-		TestEqual(TEXT("FSharedBuffer(Wrap).Size()"), Ref->Size(), Size);
-		TestEqual(TEXT("FSharedBuffer(Wrap).GetData()"), Ref->GetData(), static_cast<void*>(Data));
+		FSharedBufferPtr MutablePtr = FSharedBuffer::MakeView(Data, Size);
+		TestFalse(TEXT("FSharedBuffer::MakeImmutable(View).IsOwned()"), MutablePtr->IsOwned());
+		TestFalse(TEXT("FSharedBuffer::MakeImmutable(View).IsImmutable()"), MutablePtr->IsImmutable());
+		FSharedBufferConstPtr Ptr = FSharedBuffer::MakeImmutable(MoveTemp(MutablePtr));
+		TestTrue(TEXT("FSharedBuffer::MakeImmutable(View).IsOwned()"), Ptr->IsOwned());
+		TestTrue(TEXT("FSharedBuffer::MakeImmutable(View).IsImmutable()"), Ptr->IsImmutable());
+		TestEqual(TEXT("FSharedBuffer::MakeImmutable(View).GetSize()"), Ptr->GetSize(), Size);
+		TestNotEqual(TEXT("FSharedBuffer::MakeImmutable(View).GetData()"), Ptr->GetData(), static_cast<const void*>(Data));
+	}
+	{
+		// MakeImmutable with another shared reference.
+		constexpr uint64 Size = 64;
+		FSharedBufferPtr MutablePtr = FSharedBuffer::Alloc(Size);
+		FSharedBufferConstPtr SharedPtr = MutablePtr;
+		TestFalse(TEXT("FSharedBuffer::MakeImmutable(SharedRef).IsImmutable()"), MutablePtr->IsImmutable());
+		const void* const Data = MutablePtr->GetData();
+		FSharedBufferConstPtr Ptr = FSharedBuffer::MakeImmutable(MoveTemp(MutablePtr));
+		TestTrue(TEXT("FSharedBuffer::MakeImmutable(SharedRef).IsOwned()"), Ptr->IsOwned());
+		TestTrue(TEXT("FSharedBuffer::MakeImmutable(SharedRef).IsImmutable()"), Ptr->IsImmutable());
+		TestNotEqual(TEXT("FSharedBuffer::MakeImmutable(SharedRef).GetData()"), Ptr->GetData(), Data);
+	}
+	{
+		// MakeImmutable with another weak reference.
+		constexpr uint64 Size = 64;
+		FSharedBufferPtr MutablePtr = FSharedBuffer::Alloc(Size);
+		FSharedBufferConstWeakPtr WeakPtr = MutablePtr;
+		TestFalse(TEXT("FSharedBuffer::MakeImmutable(WeakRef).IsImmutable()"), MutablePtr->IsImmutable());
+		const void* const Data = MutablePtr->GetData();
+		FSharedBufferConstPtr Ptr = FSharedBuffer::MakeImmutable(MoveTemp(MutablePtr));
+		TestTrue(TEXT("FSharedBuffer::MakeImmutable(WeakRef).IsOwned()"), Ptr->IsOwned());
+		TestTrue(TEXT("FSharedBuffer::MakeImmutable(WeakRef).IsImmutable()"), Ptr->IsImmutable());
+		TestNotEqual(TEXT("FSharedBuffer::MakeImmutable(WeakRef).GetData()"), Ptr->GetData(), Data);
 	}
 
-	// Test MakeSharedBufferOwned
+	// Test MakeMutable
 	{
+		// MakeMutable from a new reference.
+		constexpr uint64 Size = 64;
+		FSharedBufferConstRef ConstRef = FSharedBuffer::MakeImmutable(FSharedBuffer::Alloc(Size));
+		FSharedBufferRef Ref = FSharedBuffer::MakeMutable(MoveTemp(ConstRef));
+		TestTrue(TEXT("FSharedBuffer::MakeMutable(AllocRef).IsOwned()"), Ref->IsOwned());
+		TestFalse(TEXT("FSharedBuffer::MakeMutable(AllocRef).IsImmutable()"), Ref->IsImmutable());
+		TestEqual(TEXT("FSharedBuffer::MakeMutable(AllocRef).GetSize()"), Ref->GetSize(), Size);
+		TestEqual(TEXT("FSharedBuffer::MakeMutable(AllocRef).GetData()"), const_cast<const void*>(Ref->GetData()), ConstRef->GetData());
+		FSharedBufferRef MutableRef = FSharedBuffer::MakeMutable(MoveTemp(Ref));
+		TestEqual(TEXT("FSharedBuffer::MakeMutable(AllocRef).GetData()"), Ref->GetData(), MutableRef->GetData());
+	}
+	{
+		// MakeMutable from a new pointer.
+		constexpr uint64 Size = 64;
+		FSharedBufferConstPtr ConstPtr = FSharedBuffer::MakeImmutable(FSharedBuffer::Alloc(Size));
+		const void* const Data = ConstPtr->GetData();
+		FSharedBufferPtr Ptr = FSharedBuffer::MakeMutable(MoveTemp(ConstPtr));
+		TestFalse(TEXT("FSharedBuffer::MakeImmutable(MovePtr).IsMutable()"), Ptr->IsImmutable());
+		TestTrue(TEXT("FSharedBuffer::MakeImmutable(MovePtr).IsOwned()"), Ptr->IsOwned());
+		TestEqual(TEXT("FSharedBuffer::MakeImmutable(MovePtr).GetSize()"), Ptr->GetSize(), Size);
+		TestEqual(TEXT("FSharedBuffer::MakeImmutable(MovePtr).GetData()"), const_cast<const void*>(Ptr->GetData()), Data);
+	}
+	{
+		// MakeMutable from a mutable view.
 		constexpr uint64 Size = 64;
 		uint8 Data[Size]{};
-		FSharedBufferRef Ref = MakeSharedBufferOwned(MakeSharedBuffer(FSharedBuffer::Wrap, Data, Size));
-		TestTrue(TEXT("MakeSharedBufferOwned(Wrap).IsOwned()"), Ref->IsOwned());
-		TestEqual(TEXT("MakeSharedBufferOwned(Wrap).Size()"), Ref->Size(), Size);
-		TestNotEqual(TEXT("MakeSharedBufferOwned(Wrap).GetData()"), Ref->GetData(), static_cast<void*>(Data));
-		FSharedBufferRef RefCopy = MakeSharedBufferOwned(Ref);
-		TestTrue(TEXT("MakeSharedBufferOwned(Owned).IsOwned()"), RefCopy->IsOwned());
-		TestEqual(TEXT("MakeSharedBufferOwned(Owned).Size()"), RefCopy->Size(), Size);
-		TestEqual(TEXT("MakeSharedBufferOwned(Owned).GetData()"), RefCopy->GetData(), Ref->GetData());
+		FSharedBufferPtr Ptr = FSharedBuffer::MakeMutable(FSharedBuffer::MakeView(Data, Size));
+		TestFalse(TEXT("FSharedBuffer::MakeMutable(View).IsOwned()"), Ptr->IsOwned());
+		TestFalse(TEXT("FSharedBuffer::MakeMutable(View).IsImmutable()"), Ptr->IsImmutable());
+		TestEqual(TEXT("FSharedBuffer::MakeMutable(View).GetSize()"), Ptr->GetSize(), Size);
+		TestEqual(TEXT("FSharedBuffer::MakeMutable(View).GetData()"), const_cast<const void*>(Ptr->GetData()), static_cast<const void*>(Data));
 	}
 	{
-		FSharedBufferPtr Ptr;
-		TestFalse(TEXT("MakeSharedBufferOwned(Null)"), MakeSharedBufferOwned(Ptr).IsValid());
-		TestFalse(TEXT("MakeSharedBufferOwned(Null)"), MakeSharedBufferOwned(AsConst(Ptr)).IsValid());
-		TestFalse(TEXT("MakeSharedBufferOwned(Null)"), MakeSharedBufferOwned(FSharedBufferPtr()).IsValid());
-		TestFalse(TEXT("MakeSharedBufferOwned(Null)"), MakeSharedBufferOwned(FSharedBufferConstPtr()).IsValid());
+		// MakeMutable from a const view.
+		constexpr uint64 Size = 64;
+		const uint8 Data[Size]{};
+		FSharedBufferPtr Ptr = FSharedBuffer::MakeMutable(FSharedBuffer::MakeView(Data, Size));
+		TestTrue(TEXT("FSharedBuffer::MakeMutable(View).IsOwned()"), Ptr->IsOwned());
+		TestFalse(TEXT("FSharedBuffer::MakeMutable(View).IsImmutable()"), Ptr->IsImmutable());
+		TestEqual(TEXT("FSharedBuffer::MakeMutable(View).GetSize()"), Ptr->GetSize(), Size);
+		TestNotEqual(TEXT("FSharedBuffer::MakeMutable(View).GetData()"), const_cast<const void*>(Ptr->GetData()), static_cast<const void*>(Data));
+	}
+	{
+		// MakeMutable with another shared reference.
+		constexpr uint64 Size = 64;
+		FSharedBufferConstPtr ConstPtr = FSharedBuffer::MakeImmutable(FSharedBuffer::Alloc(Size));
+		FSharedBufferConstPtr SharedPtr = ConstPtr;
+		const void* const Data = ConstPtr->GetData();
+		FSharedBufferConstPtr Ptr = FSharedBuffer::MakeMutable(MoveTemp(ConstPtr));
+		TestTrue(TEXT("FSharedBuffer::MakeMutable(SharedRef).IsOwned()"), Ptr->IsOwned());
+		TestFalse(TEXT("FSharedBuffer::MakeMutable(SharedRef).IsImmutable()"), Ptr->IsImmutable());
+		TestNotEqual(TEXT("FSharedBuffer::MakeMutable(SharedRef).GetData()"), const_cast<const void*>(Ptr->GetData()), Data);
+	}
+	{
+		// MakeMutable with another weak reference.
+		constexpr uint64 Size = 64;
+		FSharedBufferConstPtr ConstPtr = FSharedBuffer::MakeImmutable(FSharedBuffer::Alloc(Size));
+		FSharedBufferConstWeakPtr WeakPtr = ConstPtr;
+		const void* const Data = ConstPtr->GetData();
+		FSharedBufferConstPtr Ptr = FSharedBuffer::MakeMutable(MoveTemp(ConstPtr));
+		TestTrue(TEXT("FSharedBuffer::MakeMutable(WeakRef).IsOwned()"), Ptr->IsOwned());
+		TestFalse(TEXT("FSharedBuffer::MakeMutable(WeakRef).IsImmutable()"), Ptr->IsImmutable());
+		TestNotEqual(TEXT("FSharedBuffer::MakeMutable(WeakRef).GetData()"), const_cast<const void*>(Ptr->GetData()), Data);
 	}
 
 	// Test WeakPtr
 	{
 		FSharedBufferWeakPtr WeakPtr;
 		{
-			FSharedBufferRef Ref = MakeSharedBuffer();
+			FSharedBufferRef Ref = FSharedBuffer::Alloc(0);
 			WeakPtr = Ref;
 			TestTrue(TEXT("FSharedBufferWeakPtr(Ref).Pin().IsValid()"), WeakPtr.Pin().IsValid());
 		}
diff --git a/Engine/Source/Runtime/Core/Private/Tests/Serialization/CompactBinaryTest.cpp b/Engine/Source/Runtime/Core/Private/Tests/Serialization/CompactBinaryTest.cpp
--- a/Engine/Source/Runtime/Core/Private/Tests/Serialization/CompactBinaryTest.cpp
+++ b/Engine/Source/Runtime/Core/Private/Tests/Serialization/CompactBinaryTest.cpp
@@ -931,8 +931,8 @@ bool FCbFieldRefTest::RunTest(const FString& Parameters)
 	// Test Field w/ Type from Shared Buffer
 	{
 		uint8 Payload[] = { uint8(ECbFieldType::Binary), 3, 4, 5, 6 }; // Size: 3, Data: 4/5/6
-		FSharedBufferRef WrapBuffer = MakeSharedBuffer(FSharedBuffer::Wrap, MakeMemoryView(Payload));
-		FSharedBufferRef OwnedBuffer = MakeSharedBufferOwned(WrapBuffer);
+		FSharedBufferRef WrapBuffer = FSharedBuffer::MakeView(MakeMemoryView(Payload));
+		FSharedBufferRef OwnedBuffer = FSharedBuffer::MakeOwned(WrapBuffer);
 		FSharedBufferConstPtr WrapBufferPtr = WrapBuffer;
 		FSharedBufferConstPtr OwnedBufferPtr = OwnedBuffer;
 
diff --git a/Engine/Source/Runtime/Core/Public/Memory/MemoryFwd.h b/Engine/Source/Runtime/Core/Public/Memory/MemoryFwd.h
new file mode 100644
--- /dev/null
+++ b/Engine/Source/Runtime/Core/Public/Memory/MemoryFwd.h
@@ -0,0 +1,23 @@
+// Copyright Epic Games, Inc. All Rights Reserved.
+
+#pragma once
+
+// MemoryView
+
+template <typename DataType>
+class TMemoryView;
+
+/** A non-owning mutable view of a contiguous region of memory. */
+using FMutableMemoryView = TMemoryView<void>;
+
+/** A non-owning const view of a contiguous region of memory. */
+using FConstMemoryView = TMemoryView<const void>;
+
+// SharedBuffer
+
+class FSharedBufferRef;
+class FSharedBufferConstRef;
+class FSharedBufferPtr;
+class FSharedBufferConstPtr;
+class FSharedBufferWeakPtr;
+class FSharedBufferConstWeakPtr;
diff --git a/Engine/Source/Runtime/Core/Public/Memory/MemoryView.h b/Engine/Source/Runtime/Core/Public/Memory/MemoryView.h
--- a/Engine/Source/Runtime/Core/Public/Memory/MemoryView.h
+++ b/Engine/Source/Runtime/Core/Public/Memory/MemoryView.h
@@ -5,6 +5,7 @@
 #include "CoreTypes.h"
 #include "Math/NumericLimits.h"
 #include "Math/UnrealMathUtility.h"
+#include "Memory/MemoryFwd.h"
 #include "Templates/AndOrNot.h"
 #include "Templates/ChooseClass.h"
 #include "Templates/EnableIf.h"
@@ -34,19 +35,19 @@ public:
 	constexpr TMemoryView() = default;
 
 	/** Construct a view of InSize bytes starting at InData. */
-	constexpr inline TMemoryView(DataType* InData, uint64 InSize) : Data(InData), DataSize(InSize) {}
+	constexpr inline TMemoryView(DataType* InData, uint64 InSize) : Data(InData), Size(InSize) {}
 
 	template <typename OtherDataType, typename = typename TEnableIf<TPointerIsConvertibleFromTo<OtherDataType, DataType>::Value>::Type>
-	constexpr inline TMemoryView(const TMemoryView<OtherDataType>& InView) : Data(InView.Data), DataSize(InView.DataSize) {}
+	constexpr inline TMemoryView(const TMemoryView<OtherDataType>& InView) : Data(InView.Data), Size(InView.Size) {}
 
 	/** Returns a pointer to the start of the view. */
 	constexpr inline DataType* GetData() const { return Data; }
 
 	/** Returns the number of bytes in the view. */
-	constexpr inline uint64 Size() const { return DataSize; }
+	constexpr inline uint64 GetSize() const { return Size; }
 
 	/** Returns whether the view has a size of 0 regardless of its data pointer. */
-	constexpr inline bool IsEmpty() const { return DataSize == 0; }
+	constexpr inline bool IsEmpty() const { return Size == 0; }
 
 	/** Resets to an empty view. */
 	constexpr inline void Reset() { *this = TMemoryView(); }
@@ -94,30 +95,30 @@ public:
 	/** Modifies the view to be the given number of bytes from the left. */
 	constexpr inline void LeftInline(uint64 InSize)
 	{
-		DataSize = FMath::Min(DataSize, InSize);
+		Size = FMath::Min(Size, InSize);
 	}
 
 	/** Modifies the view by chopping the given number of bytes from the right. */
 	constexpr inline void LeftChopInline(uint64 InSize)
 	{
-		DataSize -= FMath::Min(DataSize, InSize);
+		Size -= FMath::Min(Size, InSize);
 	}
 
 	/** Modifies the view to be the given number of bytes from the right. */
 	inline void RightInline(uint64 InSize)
 	{
-		const uint64 OldSize = DataSize;
+		const uint64 OldSize = Size;
 		const uint64 NewSize = FMath::Min(OldSize, InSize);
 		Data = GetDataAtOffsetNoCheck(OldSize - NewSize);
-		DataSize = NewSize;
+		Size = NewSize;
 	}
 
 	/** Modifies the view by chopping the given number of bytes from the left. */
 	inline void RightChopInline(uint64 InSize)
 	{
-		const uint64 Offset = FMath::Min(DataSize, InSize);
+		const uint64 Offset = FMath::Min(Size, InSize);
 		Data = GetDataAtOffsetNoCheck(Offset);
-		DataSize -= Offset;
+		Size -= Offset;
 	}
 
 	/** Modifies the view to be the middle part by taking up to the given number of bytes from the given offset. */
@@ -131,36 +132,36 @@ public:
 	template <typename OtherDataType>
 	inline bool Contains(const TMemoryView<OtherDataType>& InView) const
 	{
-		return Data <= InView.Data && GetDataAtOffsetNoCheck(DataSize) >= InView.GetDataAtOffsetNoCheck(InView.DataSize);
+		return Data <= InView.Data && GetDataAtOffsetNoCheck(Size) >= InView.GetDataAtOffsetNoCheck(InView.Size);
 	}
 
 	/** Returns whether this view intersects the other view. */
 	template <typename OtherDataType>
 	inline bool Intersects(const TMemoryView<OtherDataType>& InView) const
 	{
-		return Data < InView.GetDataAtOffsetNoCheck(InView.DataSize) && InView.Data < GetDataAtOffsetNoCheck(DataSize);
+		return Data < InView.GetDataAtOffsetNoCheck(InView.Size) && InView.Data < GetDataAtOffsetNoCheck(Size);
 	}
 
 	/** Returns whether the bytes of this view are equal or less/greater than the bytes of the other view. */
 	template <typename OtherDataType>
 	inline int32 CompareBytes(const TMemoryView<OtherDataType>& InView) const
 	{
-		const int32 Compare = Data == InView.Data ? 0 : FMemory::Memcmp(Data, InView.Data, FMath::Min(DataSize, InView.DataSize));
-		return Compare || DataSize == InView.DataSize ? Compare : DataSize < InView.DataSize ? -1 : 1;
+		const int32 Compare = Data == InView.Data ? 0 : FMemory::Memcmp(Data, InView.Data, FMath::Min(Size, InView.Size));
+		return Compare || Size == InView.Size ? Compare : Size < InView.Size ? -1 : 1;
 	}
 
 	/** Returns whether the bytes of this views are equal to the bytes of the other view. */
 	template <typename OtherDataType>
 	inline bool EqualBytes(const TMemoryView<OtherDataType>& InView) const
 	{
-		return DataSize == InView.DataSize && (Data == InView.Data || FMemory::Memcmp(Data, InView.Data, DataSize) == 0);
+		return Size == InView.Size && (Data == InView.Data || FMemory::Memcmp(Data, InView.Data, Size) == 0);
 	}
 
 	/** Returns whether the data pointers and sizes of this view and the other view are equal. */
 	template <typename OtherDataType>
 	constexpr inline bool Equals(const TMemoryView<OtherDataType>& InView) const
 	{
-		return DataSize == InView.DataSize && (DataSize == 0 || Data == InView.Data);
+		return Size == InView.Size && (Size == 0 || Data == InView.Data);
 	}
 
 	/** Returns whether the data pointers and sizes of this view and the other view are equal. */
@@ -197,7 +198,7 @@ private:
 
 private:
 	DataType* Data = nullptr;
-	uint64 DataSize = 0;
+	uint64 Size = 0;
 };
 
 /** Advances the start of the view by an offset, which is clamped to stay within the view. */
@@ -246,9 +247,3 @@ constexpr inline auto MakeMemoryView(ContainerType&& Container)
 	using DataType = typename TChooseClass<bIsConst, const void, void>::Result;
 	return TMemoryView<DataType>(GetData(Container), GetNum(Container) * sizeof(ElementType));
 }
-
-/** A non-owning mutable view of a contiguous region of memory. */
-using FMutableMemoryView = TMemoryView<void>;
-
-/** A non-owning const view of a contiguous region of memory. */
-using FConstMemoryView = TMemoryView<const void>;
diff --git a/Engine/Source/Runtime/Core/Public/Memory/SharedBuffer.h b/Engine/Source/Runtime/Core/Public/Memory/SharedBuffer.h
--- a/Engine/Source/Runtime/Core/Public/Memory/SharedBuffer.h
+++ b/Engine/Source/Runtime/Core/Public/Memory/SharedBuffer.h
@@ -3,133 +3,217 @@
 #pragma once
 
 #include "CoreTypes.h"
+#include "HAL/UnrealMemory.h"
 #include "Memory/MemoryView.h"
-#include "Templates/Decay.h"
-#include "Templates/EnableIf.h"
-#include "Templates/SharedPointer.h"
+#include "Misc/EnumClassFlags.h"
+#include "Templates/TypeHash.h"
 #include "Templates/UnrealTemplate.h"
 
+#include <atomic>
+#include <type_traits>
+
 ///////////////////////////////////////////////////////////////////////////////////////////////////////////////////////
 
-/**
- * A reference-counted shared buffer type.
- *
- * A shared buffer may either own its memory or provide a view into memory. When used as a view,
- * it is the responsibility of the user to guarantee that the lifetime of the wrapped memory will
- * exceed the lifetime of the shared buffer. \ref MakeSharedBufferOwned can be used to guarantee
- * that the memory is owned by the buffer.
- *
- * A shared buffer must be constructed using \ref MakeSharedBuffer to invoke one of the provided
- * constructors to allocate, assume ownership of, clone, or wrap a buffer.
- *
- * A shared buffer must be stored and accessed through one of its six reference types:
- * FSharedBuffer[Const]Ref or FSharedBuffer[Const][Weak]Ptr.
- */
-class FSharedBuffer
+class FSharedBuffer;
+
+namespace SharedBufferPrivate
 {
-public:
-	enum EAssumeOwnershipTag { AssumeOwnership };
-	enum ECloneTag { Clone };
-	enum EWrapTag { Wrap };
 
-	FSharedBuffer() = default;
+enum ESharedBufferFlags : uint16
+{
+	None = 0,
+	Owned = 1 << 0,
+	Immutable = 1 << 1,
+	HasDeleter = 1 << 2,
+};
 
-	/** Use via MakeSharedBuffer. Allocate an owned buffer of the given size. */
-	CORE_API explicit FSharedBuffer(uint64 Size);
+ENUM_CLASS_FLAGS(ESharedBufferFlags);
 
-	/** Use via MakeSharedBuffer. Assumes ownership of Data and calls FMemory::Free upon destruction. */
-	CORE_API FSharedBuffer(EAssumeOwnershipTag, void* Data, uint64 Size);
-	/** Use via MakeSharedBuffer. Clone the given buffer to an owned copy. */
-	CORE_API FSharedBuffer(ECloneTag, const void* Data, uint64 Size);
-	/** Use via MakeSharedBuffer. Wrap the given buffer without taking ownership. */
-	CORE_API FSharedBuffer(EWrapTag, void* Data, uint64 Size);
+struct FSharedRefOps
+{
+	static bool TryAddRef(const FSharedBuffer* Buffer);
+	static void AddRef(const FSharedBuffer* Buffer);
+	static void Release(const FSharedBuffer* Buffer);
+};
 
-	CORE_API ~FSharedBuffer();
+struct FSharedPtrOps
+{
+	static bool TryAddRef(const FSharedBuffer* Buffer);
+	static void AddRef(const FSharedBuffer* Buffer);
+	static void Release(const FSharedBuffer* Buffer);
+};
 
-	FSharedBuffer(const FSharedBuffer&) = delete;
-	FSharedBuffer& operator=(const FSharedBuffer&) = delete;
+struct FWeakPtrOps
+{
+	static bool TryAddRef(const FSharedBuffer* Buffer);
+	static void AddRef(const FSharedBuffer* Buffer);
+	static void Release(const FSharedBuffer* Buffer);
+};
 
-	inline void* GetData() { return Data; }
-	inline const void* GetData() const { return Data; }
+template <bool bInAllowNull, bool bInIsConst, bool bInIsWeak>
+class TSharedBufferPtr
+{
+	static constexpr bool bAllowNull = bInAllowNull;
+	static constexpr bool bIsConst = bInIsConst;
+	static constexpr bool bIsWeak = bInIsWeak;
+	using BufferType = std::conditional_t<bIsConst, const FSharedBuffer, FSharedBuffer>;
+	using Ops = std::conditional_t<bIsWeak, FWeakPtrOps, std::conditional_t<bAllowNull, FSharedPtrOps, FSharedRefOps>>;
 
-	inline uint64 Size() const { return DataSize; }
+	template <bool bOtherAllowNull, bool bOtherIsConst, bool bOtherIsWeak>
+	friend class TSharedBufferPtr;
 
-	/** Whether the shared buffer owns the memory that it provides a view of. */
-	inline bool IsOwned() const { return (Flags & uint16(ESharedBufferFlags::Owned)) != 0; }
+	template <bool bOtherAllowNull, bool bOtherIsConst, bool bOtherIsWeak>
+	static inline BufferType* CopyFrom(const TSharedBufferPtr<bOtherAllowNull, bOtherIsConst, bOtherIsWeak>& Ref)
+	{
+		BufferType* NewBuffer = Ref.Buffer;
+		if (bOtherIsWeak && !bIsWeak)
+		{
+			if (!Ops::TryAddRef(NewBuffer))
+			{
+				NewBuffer = nullptr;
+			}
+		}
+		else
+		{
+			Ops::AddRef(NewBuffer);
+		}
+		return NewBuffer;
+	}
 
-	inline FMutableMemoryView GetView() { return FMutableMemoryView(GetData(), Size()); }
-	inline FConstMemoryView GetView() const { return FConstMemoryView(GetData(), Size()); }
+	template <bool bOtherAllowNull, bool bOtherIsConst, bool bOtherIsWeak>
+	static inline BufferType* MoveFrom(TSharedBufferPtr<bOtherAllowNull, bOtherIsConst, bOtherIsWeak>&& Ref)
+	{
+		BufferType* NewBuffer = Ref.Buffer;
+		if (bOtherIsWeak)
+		{
+			if (bIsWeak)
+			{
+				Ref.Buffer = nullptr;
+			}
+			else
+			{
+				if (!Ops::TryAddRef(NewBuffer))
+				{
+					NewBuffer = nullptr;
+				}
+			}
+		}
+		else
+		{
+			if (bIsWeak)
+			{
+				Ops::AddRef(NewBuffer);
+			}
+			else
+			{
+				if (bOtherAllowNull)
+				{
+					Ref.Buffer = nullptr;
+				}
+				else
+				{
+					Ops::AddRef(NewBuffer);
+				}
+			}
+		}
+		return NewBuffer;
+	}
 
-	inline operator FMutableMemoryView() { return GetView(); }
-	inline operator FConstMemoryView() const { return GetView(); }
+public:
+	inline explicit TSharedBufferPtr(BufferType* const InBuffer)
+		: Buffer(InBuffer)
+	{
+		Ops::AddRef(InBuffer);
+	}
 
-private:
-	enum class ESharedBufferFlags : uint16
+	inline TSharedBufferPtr(const TSharedBufferPtr& Ref)
+		: Buffer(CopyFrom(Ref))
 	{
-		None = 0,
-		Owned = 1 << 0,
-	};
+	}
 
-	void* Data = nullptr;
-	union
+	inline TSharedBufferPtr(TSharedBufferPtr&& Ref)
+		: Buffer(MoveFrom(MoveTemp(Ref)))
 	{
-		uint64 DataSizeAndFlags = 0;
-		struct
-		{
-			uint64 DataSize : 48;
-			uint64 Flags : 16;
-		};
-	};
-};
+	}
 
-///////////////////////////////////////////////////////////////////////////////////////////////////////////////////////
+	template <bool bOtherAllowNull, bool bOtherIsConst, bool bOtherIsWeak>
+	inline explicit TSharedBufferPtr(const TSharedBufferPtr<bOtherAllowNull, bOtherIsConst, bOtherIsWeak>& Ref)
+		: Buffer(CopyFrom(Ref))
+	{
+	}
 
-class FSharedBufferRef;
-class FSharedBufferConstRef;
+	template <bool bOtherAllowNull, bool bOtherIsConst, bool bOtherIsWeak>
+	inline explicit TSharedBufferPtr(TSharedBufferPtr<bOtherAllowNull, bOtherIsConst, bOtherIsWeak>&& Ref)
+		: Buffer(MoveFrom(MoveTemp(Ref)))
+	{
+	}
 
-namespace SharedBufferPrivate
-{
-	/** Wrapper to allow GetTypeHash to resolve from within a scope with another overload. */
-	template <typename... ArgTypes>
-	auto WrapGetTypeHash(ArgTypes&&... Args) -> decltype(GetTypeHash(Forward<ArgTypes>(Args)...))
+	inline ~TSharedBufferPtr()
 	{
-		return GetTypeHash(Forward<ArgTypes>(Args)...);
+		Ops::Release(Buffer);
 	}
 
-	FSharedBufferRef MakeSharedBuffer(TSharedRef<FSharedBuffer, ESPMode::ThreadSafe>&& Ref);
-	FSharedBufferConstRef MakeSharedBuffer(TSharedRef<const FSharedBuffer, ESPMode::ThreadSafe>&& Ref);
-}
+	inline TSharedBufferPtr& operator=(const TSharedBufferPtr& Ref)
+	{
+		BufferType* const OldBuffer = Buffer;
+		Buffer = CopyFrom(Ref);
+		Ops::Release(OldBuffer);
+		return *this;
+	}
 
-///////////////////////////////////////////////////////////////////////////////////////////////////////////////////////
+	inline TSharedBufferPtr& operator=(TSharedBufferPtr&& Ref)
+	{
+		BufferType* const OldBuffer = Buffer;
+		Buffer = MoveFrom(MoveTemp(Ref));
+		Ops::Release(OldBuffer);
+		return *this;
+	}
+
+	template <bool bOtherAllowNull, bool bOtherIsConst, bool bOtherIsWeak>
+	inline TSharedBufferPtr& operator=(const TSharedBufferPtr<bOtherAllowNull, bOtherIsConst, bOtherIsWeak>& Ref)
+	{
+		BufferType* const OldBuffer = Buffer;
+		Buffer = CopyFrom(Ref);
+		Ops::Release(OldBuffer);
+		return *this;
+	}
+
+	template <bool bOtherAllowNull, bool bOtherIsConst, bool bOtherIsWeak>
+	inline TSharedBufferPtr& operator=(TSharedBufferPtr<bOtherAllowNull, bOtherIsConst, bOtherIsWeak>&& Ref)
+	{
+		BufferType* const OldBuffer = Buffer;
+		Buffer = MoveFrom(MoveTemp(Ref));
+		Ops::Release(OldBuffer);
+		return *this;
+	}
+
+	BufferType* Get() const { return Buffer; }
 
-// The reference types defined below provide an interface that is a subset of the shared pointer
-// types provided in Templates/SharedPointer.h. These types will be switched to alias the generic
-// types once the generic types have been updated to always provide thread safety and to deprecate
-// operations that are not safe in the presence of multi-threaded usage.
+private:
+	BufferType* Buffer = nullptr;
+};
+
+}
 
 ///////////////////////////////////////////////////////////////////////////////////////////////////////////////////////
 
-/**
- * A non-nullable, thread-safe, mutable reference to a shared buffer.
- *
- * Equivalent to TSharedRef<FSharedBuffer, ESPMode::ThreadSafe>
- */
+/** A non-nullable, thread-safe, reference to a mutable shared buffer. */
 class FSharedBufferRef
 {
 public:
-	inline FSharedBuffer& Get() const { return Ref.Get(); }
+	using PtrType = SharedBufferPrivate::TSharedBufferPtr</*bAllowNull*/ false, /*bIsConst*/ false, /*bIsWeak*/ false>;
 
-	inline FSharedBuffer& operator*() const { return Ref.operator*(); }
-	inline FSharedBuffer* operator->() const { return Ref.operator->(); }
+	inline explicit FSharedBufferRef(PtrType&& InPtr) : Ptr(MoveTemp(InPtr)) {}
 
-	inline friend uint32 GetTypeHash(const FSharedBufferRef& InRef) { return SharedBufferPrivate::WrapGetTypeHash(InRef.Ref); }
+	inline FSharedBuffer& Get() const { return *Ptr.Get(); }
 
-private:
-	constexpr inline friend bool IsValid(const FSharedBufferRef&) { return true; }
+	inline FSharedBuffer& operator*() const { return *Ptr.Get(); }
+	inline FSharedBuffer* operator->() const { return Ptr.Get(); }
 
-	inline explicit FSharedBufferRef(TSharedRef<FSharedBuffer, ESPMode::ThreadSafe>&& InRef) : Ref(InRef) {}
+	inline friend uint32 GetTypeHash(const FSharedBufferRef& Ref) { return PointerHash(Ref.Ptr.Get()); }
 
-	friend FSharedBufferRef SharedBufferPrivate::MakeSharedBuffer(TSharedRef<FSharedBuffer, ESPMode::ThreadSafe>&& Ref);
+private:
+	constexpr inline friend bool IsValid(const FSharedBufferRef&) { return true; }
 
 	friend class FSharedBufferConstRef;
 	friend class FSharedBufferPtr;
@@ -137,175 +221,158 @@ private:
 	friend class FSharedBufferWeakPtr;
 	friend class FSharedBufferConstWeakPtr;
 
-	TSharedRef<FSharedBuffer, ESPMode::ThreadSafe> Ref;
+	PtrType Ptr;
 };
 
-/**
- * A non-nullable, thread-safe, const reference to a shared buffer.
- *
- * Equivalent to TSharedRef<const FSharedBuffer, ESPMode::ThreadSafe>
- */
+/** A non-nullable, thread-safe, reference to a const shared buffer. */
 class FSharedBufferConstRef
 {
 public:
-	inline FSharedBufferConstRef(const FSharedBufferRef& InRef) : Ref(InRef.Ref) {}
+	using PtrType = SharedBufferPrivate::TSharedBufferPtr</*bAllowNull*/ false, /*bIsConst*/ true, /*bIsWeak*/ false>;
 
-	inline const FSharedBuffer& Get() const { return Ref.Get(); }
+	inline explicit FSharedBufferConstRef(PtrType&& InPtr) : Ptr(MoveTemp(InPtr)) {}
 
-	inline const FSharedBuffer& operator*() const { return Ref.operator*(); }
-	inline const FSharedBuffer* operator->() const { return Ref.operator->(); }
+	inline FSharedBufferConstRef(const FSharedBufferRef& Ref) : Ptr(Ref.Ptr) {}
 
-	inline friend uint32 GetTypeHash(const FSharedBufferConstRef& InRef) { return SharedBufferPrivate::WrapGetTypeHash(InRef.Ref); }
+	inline const FSharedBuffer& Get() const { return *Ptr.Get(); }
 
-private:
-	constexpr inline friend bool IsValid(const FSharedBufferConstRef&) { return true; }
+	inline const FSharedBuffer& operator*() const { return *Ptr.Get(); }
+	inline const FSharedBuffer* operator->() const { return Ptr.Get(); }
 
-	inline explicit FSharedBufferConstRef(TSharedRef<const FSharedBuffer, ESPMode::ThreadSafe>&& InRef) : Ref(InRef) {}
+	inline friend uint32 GetTypeHash(const FSharedBufferConstRef& Ref) { return PointerHash(Ref.Ptr.Get()); }
 
-	friend FSharedBufferConstRef SharedBufferPrivate::MakeSharedBuffer(TSharedRef<const FSharedBuffer, ESPMode::ThreadSafe>&& Ref);
+private:
+	constexpr inline friend bool IsValid(const FSharedBufferConstRef&) { return true; }
 
 	friend class FSharedBufferConstPtr;
 	friend class FSharedBufferConstWeakPtr;
 
-	TSharedRef<const FSharedBuffer, ESPMode::ThreadSafe> Ref;
+	PtrType Ptr;
 };
 
 ///////////////////////////////////////////////////////////////////////////////////////////////////////////////////////
 
-/**
- * A non-nullable, thread-safe, mutable pointer to a shared buffer.
- *
- * Equivalent to TSharedPtr<FSharedBuffer, ESPMode::ThreadSafe>
- */
+/** A non-nullable, thread-safe, pointer to a mutable shared buffer. */
 class FSharedBufferPtr
 {
 public:
-	FSharedBufferPtr() = default;
+	using PtrType = SharedBufferPrivate::TSharedBufferPtr</*bAllowNull*/ true, /*bIsConst*/ false, /*bIsWeak*/ false>;
+
+	inline explicit FSharedBufferPtr(PtrType&& InPtr) : Ptr(MoveTemp(InPtr)) {}
 
-	inline FSharedBufferPtr(const FSharedBufferRef& InRef) : Ptr(InRef.Ref) {}
+	inline FSharedBufferPtr() : Ptr(nullptr) {}
+	inline FSharedBufferPtr(const FSharedBufferRef& Ref) : Ptr(Ref.Ptr) {}
 
-	inline FSharedBufferRef ToSharedRef() const& { return FSharedBufferRef(Ptr.ToSharedRef()); }
-	inline FSharedBufferRef ToSharedRef() && { return FSharedBufferRef(MoveTemp(Ptr).ToSharedRef()); }
+	inline FSharedBufferRef ToSharedRef() const& { return FSharedBufferRef(FSharedBufferRef::PtrType(Ptr)); }
+	inline FSharedBufferRef ToSharedRef() && { return FSharedBufferRef(FSharedBufferRef::PtrType(MoveTemp(Ptr))); }
 
-	inline bool IsValid() const { return Ptr.IsValid(); }
+	inline bool IsValid() const { return Ptr.Get() != nullptr; }
 	inline FSharedBuffer* Get() const { return Ptr.Get(); }
 
-	inline explicit operator bool() const { return Ptr.IsValid(); }
-	inline FSharedBuffer& operator*() const { return Ptr.operator*(); }
-	inline FSharedBuffer* operator->() const { return Ptr.operator->(); }
+	inline explicit operator bool() const { return Ptr.Get() != nullptr; }
+	inline FSharedBuffer& operator*() const { return *Ptr.Get(); }
+	inline FSharedBuffer* operator->() const { return Ptr.Get(); }
 
-	inline void Reset() { Ptr.Reset(); }
+	inline void Reset() { *this = FSharedBufferPtr(); }
 
-	inline friend uint32 GetTypeHash(const FSharedBufferPtr& InPtr) { return SharedBufferPrivate::WrapGetTypeHash(InPtr.Ptr); }
+	inline friend uint32 GetTypeHash(const FSharedBufferPtr& InPtr) { return PointerHash(InPtr.Ptr.Get()); }
 
 private:
 	inline friend bool IsValid(const FSharedBufferPtr& InPtr) { return InPtr.IsValid(); }
 
-	inline FSharedBufferPtr(TSharedPtr<FSharedBuffer, ESPMode::ThreadSafe>&& InPtr) : Ptr(InPtr) {}
-
 	friend class FSharedBufferConstPtr;
 	friend class FSharedBufferWeakPtr;
 	friend class FSharedBufferConstWeakPtr;
 
-	TSharedPtr<FSharedBuffer, ESPMode::ThreadSafe> Ptr;
+	PtrType Ptr;
 };
 
-/**
- * A non-nullable, thread-safe, const pointer to a shared buffer.
- *
- * Equivalent to TSharedPtr<const FSharedBuffer, ESPMode::ThreadSafe>
- */
+/** A non-nullable, thread-safe, pointer to a const shared buffer. */
 class FSharedBufferConstPtr
 {
 public:
-	FSharedBufferConstPtr() = default;
+	using PtrType = SharedBufferPrivate::TSharedBufferPtr</*bAllowNull*/ true, /*bIsConst*/ true, /*bIsWeak*/ false>;
 
-	inline FSharedBufferConstPtr(const FSharedBufferRef& InRef) : Ptr(InRef.Ref) {}
-	inline FSharedBufferConstPtr(const FSharedBufferConstRef& InRef) : Ptr(InRef.Ref) {}
+	inline explicit FSharedBufferConstPtr(PtrType&& InPtr) : Ptr(MoveTemp(InPtr)) {}
+
+	inline FSharedBufferConstPtr() : Ptr(nullptr) {}
+	inline FSharedBufferConstPtr(const FSharedBufferRef& Ref) : Ptr(Ref.Ptr) {}
+	inline FSharedBufferConstPtr(const FSharedBufferConstRef& Ref) : Ptr(Ref.Ptr) {}
 	inline FSharedBufferConstPtr(const FSharedBufferPtr& InPtr) : Ptr(InPtr.Ptr) {}
 	inline FSharedBufferConstPtr(FSharedBufferPtr&& InPtr) : Ptr(MoveTemp(InPtr.Ptr)) {}
 
-	inline FSharedBufferConstRef ToSharedRef() const& { return FSharedBufferConstRef(Ptr.ToSharedRef()); }
-	inline FSharedBufferConstRef ToSharedRef() && { return FSharedBufferConstRef(MoveTemp(Ptr).ToSharedRef()); }
+	inline FSharedBufferConstRef ToSharedRef() const& { return FSharedBufferConstRef(FSharedBufferConstRef::PtrType(Ptr)); }
+	inline FSharedBufferConstRef ToSharedRef() && { return FSharedBufferConstRef(FSharedBufferConstRef::PtrType(MoveTemp(Ptr))); }
 
-	inline bool IsValid() const { return Ptr.IsValid(); }
+	inline bool IsValid() const { return Ptr.Get() != nullptr; }
 	inline const FSharedBuffer* Get() const { return Ptr.Get(); }
 
-	inline explicit operator bool() const { return Ptr.IsValid(); }
-	inline const FSharedBuffer& operator*() const { return Ptr.operator*(); }
-	inline const FSharedBuffer* operator->() const { return Ptr.operator->(); }
+	inline explicit operator bool() const { return Ptr.Get() != nullptr; }
+	inline const FSharedBuffer& operator*() const { return *Ptr.Get(); }
+	inline const FSharedBuffer* operator->() const { return Ptr.Get(); }
 
-	inline void Reset() { Ptr.Reset(); }
+	inline void Reset() { *this = FSharedBufferConstPtr(); }
 
-	inline friend uint32 GetTypeHash(const FSharedBufferConstPtr& InPtr) { return SharedBufferPrivate::WrapGetTypeHash(InPtr.Ptr); }
+	inline friend uint32 GetTypeHash(const FSharedBufferConstPtr& InPtr) { return PointerHash(InPtr.Ptr.Get()); }
 
 private:
 	inline friend bool IsValid(const FSharedBufferConstPtr& InPtr) { return InPtr.IsValid(); }
 
-	inline FSharedBufferConstPtr(TSharedPtr<const FSharedBuffer, ESPMode::ThreadSafe>&& InPtr) : Ptr(InPtr) {}
-
 	friend class FSharedBufferConstWeakPtr;
 
-	TSharedPtr<const FSharedBuffer, ESPMode::ThreadSafe> Ptr;
+	PtrType Ptr;
 };
 
 ///////////////////////////////////////////////////////////////////////////////////////////////////////////////////////
 
-/**
- * A non-nullable, thread-safe, weak, mutable pointer to a shared buffer.
- *
- * Equivalent to TWeakPtr<FSharedBuffer, ESPMode::ThreadSafe>
- */
+/** A non-nullable, thread-safe, weak, pointer to a mutable shared buffer. */
 class FSharedBufferWeakPtr
 {
 public:
-	FSharedBufferWeakPtr() = default;
+	using PtrType = SharedBufferPrivate::TSharedBufferPtr</*bAllowNull*/ true, /*bIsConst*/ false, /*bIsWeak*/ true>;
+
+	inline explicit FSharedBufferWeakPtr(PtrType&& InPtr) : Ptr(MoveTemp(InPtr)) {}
 
-	inline FSharedBufferWeakPtr(const FSharedBufferRef& InRef) : Ptr(InRef.Ref) {}
+	inline FSharedBufferWeakPtr() : Ptr(nullptr) {}
+	inline FSharedBufferWeakPtr(const FSharedBufferRef& Ref) : Ptr(Ref.Ptr) {}
 	inline FSharedBufferWeakPtr(const FSharedBufferPtr& InPtr) : Ptr(InPtr.Ptr) {}
-	inline FSharedBufferWeakPtr(FSharedBufferPtr&& InPtr) : Ptr(MoveTemp(InPtr.Ptr)) {}
 
-	inline FSharedBufferPtr Pin() const& { return FSharedBufferPtr(Ptr.Pin()); }
-	inline FSharedBufferPtr Pin() && { return FSharedBufferPtr(MoveTemp(Ptr).Pin()); }
+	inline FSharedBufferPtr Pin() { return FSharedBufferPtr(FSharedBufferPtr::PtrType(Ptr)); }
 
-	inline void Reset() { Ptr.Reset(); }
+	inline void Reset() { *this = FSharedBufferWeakPtr(); }
 
-	inline friend uint32 GetTypeHash(const FSharedBufferWeakPtr& InPtr) { return SharedBufferPrivate::WrapGetTypeHash(InPtr.Ptr); }
+	inline friend uint32 GetTypeHash(const FSharedBufferWeakPtr& InPtr) { return PointerHash(InPtr.Ptr.Get()); }
 
 private:
 	friend class FSharedBufferConstWeakPtr;
 
-	TWeakPtr<FSharedBuffer, ESPMode::ThreadSafe> Ptr;
+	PtrType Ptr;
 };
 
-/**
- * A non-nullable, thread-safe, weak, const pointer to a shared buffer.
- *
- * Equivalent to TWeakPtr<const FSharedBuffer, ESPMode::ThreadSafe>
- */
+/** A non-nullable, thread-safe, weak, pointer to a const shared buffer. */
 class FSharedBufferConstWeakPtr
 {
 public:
-	FSharedBufferConstWeakPtr() = default;
+	using PtrType = SharedBufferPrivate::TSharedBufferPtr</*bAllowNull*/ true, /*bIsConst*/ true, /*bIsWeak*/ true>;
+
+	inline explicit FSharedBufferConstWeakPtr(PtrType&& InPtr) : Ptr(MoveTemp(InPtr)) {}
 
-	inline FSharedBufferConstWeakPtr(const FSharedBufferRef& InRef) : Ptr(InRef.Ref) {}
-	inline FSharedBufferConstWeakPtr(const FSharedBufferConstRef& InRef) : Ptr(InRef.Ref) {}
+	inline FSharedBufferConstWeakPtr() : Ptr(nullptr) {}
+	inline FSharedBufferConstWeakPtr(const FSharedBufferRef& Ref) : Ptr(Ref.Ptr) {}
+	inline FSharedBufferConstWeakPtr(const FSharedBufferConstRef& Ref) : Ptr(Ref.Ptr) {}
 	inline FSharedBufferConstWeakPtr(const FSharedBufferPtr& InPtr) : Ptr(InPtr.Ptr) {}
 	inline FSharedBufferConstWeakPtr(const FSharedBufferConstPtr& InPtr) : Ptr(InPtr.Ptr) {}
 	inline FSharedBufferConstWeakPtr(const FSharedBufferWeakPtr& InPtr) : Ptr(InPtr.Ptr) {}
-	inline FSharedBufferConstWeakPtr(FSharedBufferPtr&& InPtr) : Ptr(MoveTemp(InPtr.Ptr)) {}
-	inline FSharedBufferConstWeakPtr(FSharedBufferConstPtr&& InPtr) : Ptr(MoveTemp(InPtr.Ptr)) {}
 	inline FSharedBufferConstWeakPtr(FSharedBufferWeakPtr&& InPtr) : Ptr(MoveTemp(InPtr.Ptr)) {}
 
-	inline FSharedBufferConstPtr Pin() const& { return FSharedBufferConstPtr(Ptr.Pin()); }
-	inline FSharedBufferConstPtr Pin() && { return FSharedBufferConstPtr(MoveTemp(Ptr).Pin()); }
+	inline FSharedBufferConstPtr Pin() { return FSharedBufferConstPtr(FSharedBufferConstPtr::PtrType(Ptr)); }
 
-	inline void Reset() { Ptr.Reset(); }
+	inline void Reset() { *this = FSharedBufferConstWeakPtr(); }
 
-	inline friend uint32 GetTypeHash(const FSharedBufferConstWeakPtr& InPtr) { return SharedBufferPrivate::WrapGetTypeHash(InPtr.Ptr); }
+	inline friend uint32 GetTypeHash(const FSharedBufferConstWeakPtr& InPtr) { return PointerHash(InPtr.Ptr.Get()); }
 
 private:
-	TWeakPtr<const FSharedBuffer, ESPMode::ThreadSafe> Ptr;
+	PtrType Ptr;
 };
 
 ///////////////////////////////////////////////////////////////////////////////////////////////////////////////////////
@@ -322,69 +389,463 @@ template<> struct TIsWeakPointerType<FSharedBufferConstWeakPtr> { enum { Value =
 
 ///////////////////////////////////////////////////////////////////////////////////////////////////////////////////////
 
-namespace SharedBufferPrivate
+/**
+ * A reference-counted shared buffer type.
+ *
+ * A shared buffer may either own its memory or provide a view into memory owned externally. When
+ * used as to wrap a externally-owned memory, the user must guarantee the lifetime of the wrapped
+ * memory will exceed the lifetime of the shared buffer.
+ *
+ * A shared buffer must be stored and accessed through one of its six reference types:
+ * FSharedBuffer[Const]Ref or FSharedBuffer[Const][Weak]Ptr.
+ *
+ * Make an owned sized buffer with optional alignment.
+ *
+ *     FSharedBuffer::Alloc(128)
+ *     FSharedBuffer::Alloc(128, 16)
+ *
+ * Make an owned buffer by cloning the buffer pointed to by the arguments.
+ *
+ *     FSharedBuffer::Clone(Data, Size)
+ *     FSharedBuffer::Clone(MakeMemoryView(...))
+ *     FSharedBuffer::Clone(*SharedBufferRefOrPtr)
+ *
+ * Make a non-owned buffer by creating a view of the buffer pointed to by the arguments.
+ *
+ *     FSharedBuffer::MakeView(Data, Size)
+ *     FSharedBuffer::MakeView(MakeMemoryView(...))
+ *
+ * Make an owned buffer by taking ownership of the buffer pointed to by the arguments.
+ *
+ *     FSharedBuffer::TakeOwnership(Data, Size, FMemory::Free)
+ *     FSharedBuffer::TakeOwnership(Data, Size, [](void* Data) { delete[] static_cast<uint8*>(Data); })
+ */
+class FSharedBuffer
 {
-	inline FSharedBufferRef MakeSharedBuffer(TSharedRef<FSharedBuffer, ESPMode::ThreadSafe>&& Ref)
+	using ESharedBufferFlags = SharedBufferPrivate::ESharedBufferFlags;
+
+public:
+	/** Make an owned writable buffer of Size bytes. */
+	static inline FSharedBufferRef Alloc(const uint64 Size, const uint32 Alignment = DEFAULT_ALIGNMENT)
 	{
-		return FSharedBufferRef(MoveTemp(Ref));
+		return NewBuffer(FMemory::Malloc(Size, Alignment), Size, ESharedBufferFlags::Owned);
 	}
 
-	inline FSharedBufferConstRef MakeSharedBuffer(TSharedRef<const FSharedBuffer, ESPMode::ThreadSafe>&& Ref)
+	/** Make an owned writable clone of the memory view. */
+	static inline FSharedBufferRef Clone(const FConstMemoryView View)
 	{
-		return FSharedBufferConstRef(MoveTemp(Ref));
+		return Clone(View.GetData(), View.GetSize());
+	}
+
+	/** Make an owned writable clone of the buffer. */
+	static inline FSharedBufferRef Clone(const void* const Data, const uint64 Size)
+	{
+		const FSharedBufferRef Buffer = Alloc(Size);
+		FMemory::Memcpy(Buffer->GetData(), Data, Size);
+		return Buffer;
+	}
+
+	/** Make a non-owned writable view of the memory view. */
+	static inline FSharedBufferRef MakeView(FMutableMemoryView View)
+	{
+		return MakeView(View.GetData(), View.GetSize());
 	}
-}
+
+	/** Make a non-owned read-only view of the memory view. */
+	static inline FSharedBufferConstRef MakeView(FConstMemoryView View)
+	{
+		return MakeView(View.GetData(), View.GetSize());
+	}
+
+	/** Make a non-owned writable view of the buffer. */
+	static inline FSharedBufferRef MakeView(void* Data, uint64 Size)
+	{
+		return NewBuffer(Data, Size, ESharedBufferFlags::None);
+	}
+
+	/** Make a non-owned read-only view of the buffer. */
+	static inline FSharedBufferConstRef MakeView(const void* Data, uint64 Size)
+	{
+		// const_cast is safe here because the return value only allows accessing data through a const void*
+		return MakeView(const_cast<void*>(Data), Size);
+	}
+
+	/**
+	 * Make an owned writable buffer by taking ownership of the provided memory.
+	 *
+	 * @param Deleter Called with the data pointer to free memory when there are no shared references.
+	 */
+	template <typename DeleterType>
+	static inline FSharedBufferRef TakeOwnership(void* Data, uint64 Size, DeleterType&& Deleter)
+	{
+		return NewBufferWithDeleter(Data, Size, ESharedBufferFlags::Owned, Forward<DeleterType>(Deleter));
+	}
+
+	/**
+	 * Make an owned read-only buffer by taking ownership of the provided memory.
+	 *
+	 * @param Deleter Called with the data pointer to free memory when there are no shared references.
+	 */
+	template <typename DeleterType>
+	static inline FSharedBufferConstRef TakeOwnership(const void* Data, uint64 Size, DeleterType&& Deleter)
+	{
+		// const_cast is safe here because the return value only allows accessing data through a const void*
+		return TakeOwnership(const_cast<void*>(Data), Size, Forward<DeleterType>(Deleter));
+	}
+
+	/**
+	 * Return the buffer if it is owned, or a clone otherwise.
+	 *
+	 * @param Buffer A FSharedBuffer[Const]Ref or FSharedBuffer[Const]Ptr to make owned.
+	 * @return An owned copy of the input buffer in the same reference type as the input.
+	 */
+	template <typename T,
+		std::enable_if_t<std::is_same<FSharedBuffer, std::decay_t<decltype(*std::declval<T>())>>::value>* = nullptr
+		>
+	static inline auto MakeOwned(T&& Buffer) -> decltype(std::decay_t<T>(Clone(Buffer->GetData(), Buffer->GetSize())))
+	{
+		return !IsValid(Buffer) || Buffer->IsOwned() ? Forward<T>(Buffer) : Clone(Buffer->GetData(), Buffer->GetSize());
+	}
+
+	/** Return the buffer if it is owned and either immutable or this is the only reference, or a clone otherwise. */
+	static inline FSharedBufferConstRef MakeImmutable(FSharedBufferRef&& Buffer)
+	{
+		const FSharedBuffer& BufferRef = *Buffer;
+		return BufferRef.TryMakeImmutable() ? MoveTemp(Buffer) : MakeImmutable(Clone(*Buffer));
+	}
+
+	/** Return the buffer if it is owned and either immutable or this is the only reference, or a clone otherwise. */
+	static inline FSharedBufferConstRef MakeImmutable(FSharedBufferConstRef&& Buffer)
+	{
+		return Buffer->TryMakeImmutable() ? MoveTemp(Buffer) : MakeImmutable(Clone(*Buffer));
+	}
+
+	/** Return the buffer if it is owned and either immutable or this is the only reference, or a clone otherwise. */
+	static inline FSharedBufferConstPtr MakeImmutable(FSharedBufferPtr&& Buffer)
+	{
+		return MakeImmutable(FSharedBufferConstPtr(MoveTemp(Buffer)));
+	}
+
+	/** Return the buffer if it is owned and either immutable or this is the only reference, or a clone otherwise. */
+	static inline FSharedBufferConstPtr MakeImmutable(FSharedBufferConstPtr&& Buffer)
+	{
+		return Buffer ? Buffer->TryMakeImmutable() ? MoveTemp(Buffer) : MakeImmutable(Clone(*Buffer)) : FSharedBufferConstPtr();
+	}
+
+	// Disable overloads that would allow the argument to be copied.
+	static void MakeImmutable(const FSharedBufferRef&) = delete;
+	static void MakeImmutable(const FSharedBufferConstRef&) = delete;
+	static void MakeImmutable(const FSharedBufferPtr&) = delete;
+	static void MakeImmutable(const FSharedBufferConstPtr&) = delete;
+
+	/** Return the buffer because it is already mutable. */
+	static inline FSharedBufferRef MakeMutable(FSharedBufferRef&& Buffer) { return MoveTemp(Buffer); }
+
+	/** Return the buffer if it is owned and either mutable or this is the only reference, or a clone otherwise. */
+	static inline FSharedBufferRef MakeMutable(FSharedBufferConstRef&& Buffer)
+	{
+		if (Buffer->TryMakeMutable())
+		{
+			return FSharedBufferRef(FSharedBufferRef::PtrType(const_cast<FSharedBuffer*>(&Buffer.Get())));
+		}
+		return Clone(*Buffer);
+	}
+
+	/** Return the buffer because it is already mutable. */
+	static inline FSharedBufferPtr MakeMutable(FSharedBufferPtr&& Buffer) { return MoveTemp(Buffer); }
+
+	/** Return the buffer if it is owned and either mutable or this is the only reference, or a clone otherwise. */
+	static inline FSharedBufferPtr MakeMutable(FSharedBufferConstPtr&& Buffer)
+	{
+		if (Buffer->TryMakeMutable())
+		{
+			return FSharedBufferPtr(FSharedBufferPtr::PtrType(const_cast<FSharedBuffer*>(Buffer.Get())));
+		}
+		return Clone(*Buffer);
+	}
+
+	// Disable overloads that would allow the argument to be copied.
+	static void MakeMutable(const FSharedBufferRef&) = delete;
+	static void MakeMutable(const FSharedBufferConstRef&) = delete;
+	static void MakeMutable(const FSharedBufferPtr&) = delete;
+	static void MakeMutable(const FSharedBufferConstPtr&) = delete;
+
+	/** A pointer to the start of the buffer. */
+	inline void* GetData() { return Data; }
+	/** A pointer to the start of the buffer. */
+	inline const void* GetData() const { return Data; }
+
+	/** The size of the buffer in bytes. */
+	inline uint64 GetSize() const { return Size; }
+
+	/** Whether the shared buffer owns the memory that it provides a view of. */
+	inline bool IsOwned() const { return (GetFlags(ReferenceCountAndFlags) & ESharedBufferFlags::Owned) != 0; }
+
+	/**
+	 * Whether the shared buffer is immutable.
+	 *
+	 * An immutable shared buffer is owned and every reference or pointer to it is const.
+	 */
+	inline bool IsImmutable() const { return (GetFlags(ReferenceCountAndFlags) & ESharedBufferFlags::Immutable) != 0; }
+
+	/**
+	 * Try to make the shared buffer immutable.
+	 *
+	 * A shared buffer can be made immutable if it is owned and there is only one reference to it.
+	 *
+	 * @return true if the buffer was immutable or was made immutable, otherwise false.
+	 */
+	CORE_API bool TryMakeImmutable() const;
+
+	// Disable this overload because it would allow an immutable buffer to be mutated.
+	bool TryMakeImmutable() = delete;
+
+	/**
+	 * Try to make the shared buffer mutable.
+	 *
+	 * A shared buffer can be made mutable if it is owned and there is only one reference to it.
+	 *
+	 * @return true if the buffer was mutable or was made mutable, otherwise false.
+	 */
+	CORE_API bool TryMakeMutable() const;
+
+	/** The shared buffer is already mutable. */
+	inline bool TryMakeMutable() { return true; }
+
+	/** A mutable view of buffer. */
+	inline FMutableMemoryView GetView() { return FMutableMemoryView(GetData(), GetSize()); }
+	/** A const view of the buffer. */
+	inline FConstMemoryView GetView() const { return FConstMemoryView(GetData(), GetSize()); }
+
+	inline operator FMutableMemoryView() { return GetView(); }
+	inline operator FConstMemoryView() const { return GetView(); }
+
+private:
+	/**
+	 * Construct a new shared buffer with a custom deleter.
+	 *
+	 * @param Data A pointer to the start of the data buffer.
+	 * @param Size The size of the data buffer.
+	 * @param Flags The flags associated with the data buffer. Internally adds HasDeleter.
+	 * @param Delete A callable to delete the data buffer when it is released.
+	 */
+	template <typename DeleterType>
+	static inline FSharedBufferRef NewBufferWithDeleter(
+		void* const Data,
+		const uint64 Size,
+		ESharedBufferFlags Flags,
+		DeleterType&& Deleter)
+	{
+		Flags |= ESharedBufferFlags::HasDeleter;
+		static_assert(alignof(TSharedBufferDeleter<DeleterType>) <= alignof(FSharedBuffer), "Required alignment is too high.");
+		FSharedBufferRef Buffer = NewBuffer(Data, Size, Flags, sizeof(TSharedBufferDeleter<DeleterType>));
+		new(&Buffer.Get() + 1) TSharedBufferDeleter<DeleterType>(Forward<DeleterType>(Deleter));
+		return Buffer;
+	}
+
+	/**
+	 * Construct a new shared buffer.
+	 *
+	 * @param Data A pointer to the start of the data buffer.
+	 * @param Size The size of the data buffer.
+	 * @param Flags The flags associated with the data buffer.
+	 * @param DeleterSize The size of the optional deleter to allocate after the FSharedBuffer.
+	 */
+	CORE_API static FSharedBufferRef NewBuffer(void* Data, uint64 Size, ESharedBufferFlags Flags, uint32 DeleterSize = 0);
+
+	/**
+	 * Delete a shared buffer. Requires ReleaseData to have been called already.
+	 */
+	CORE_API static void DeleteBuffer(FSharedBuffer* Buffer);
+
+	/**
+	 * Free the data if it is owned and release the pointer to it.
+	 *
+	 * This does not delete the shared buffer object, which survives until the last shared or weak
+	 * reference to it has been released.
+	 */
+	CORE_API void ReleaseData();
+
+	FSharedBuffer() = default;
+	FSharedBuffer(const FSharedBuffer&) = delete;
+	FSharedBuffer& operator=(const FSharedBuffer&) = delete;
+	~FSharedBuffer() = default;
+
+private:
+	class FSharedBufferDeleter
+	{
+	public:
+		/** Free the memory for the buffer and destroy the deleter. */
+		virtual void Free(void* Data) = 0;
+
+	protected:
+		FSharedBufferDeleter() = default;
+		~FSharedBufferDeleter() = default;
+	};
+
+	template <typename DeleterType>
+	class TSharedBufferDeleter final : public FSharedBufferDeleter
+	{
+	public:
+		explicit TSharedBufferDeleter(DeleterType&& InDeleter)
+			: Deleter(Forward<DeleterType>(InDeleter))
+		{
+		}
+
+	private:
+		~TSharedBufferDeleter() = default;
+
+		virtual void Free(void* const InData) final
+		{
+			Deleter(InData);
+			this->~TSharedBufferDeleter();
+		}
+
+		std::decay_t<DeleterType> Deleter;
+	};
+
+private:
+	static uint32 GetSharedRefCount(uint64 RefCountsAndFlags) { return uint32(RefCountsAndFlags >> 0) & 0xffffff; }
+	static uint64 SetSharedRefCount(uint32 RefCount) { return uint64(RefCount) << 0; }
+	static uint32 GetWeakRefCount(uint64 RefCountsAndFlags) { return uint32(RefCountsAndFlags >> 24) & 0xffffff; }
+	static uint64 SetWeakRefCount(uint32 RefCount) { return uint64(RefCount) << 24; }
+	static ESharedBufferFlags GetFlags(uint64 RefCountsAndFlags) { return ESharedBufferFlags(RefCountsAndFlags >> 48); }
+	static uint64 SetFlags(ESharedBufferFlags Flags) { return uint64(Flags) << 48; }
+
+	inline void AddSharedReference() const
+	{
+		const uint64 PreviousValue = ReferenceCountAndFlags.fetch_add(SetSharedRefCount(1), std::memory_order_relaxed);
+		checkSlow(GetSharedRefCount(PreviousValue) < 0xffffff);
+		if (GetSharedRefCount(PreviousValue) == 0)
+		{
+			AddWeakReference();
+		}
+	}
+
+	inline void ReleaseSharedReference() const
+	{
+		const uint64 PreviousValue = ReferenceCountAndFlags.fetch_sub(SetSharedRefCount(1), std::memory_order_acq_rel);
+		checkSlow(GetSharedRefCount(PreviousValue) > 0);
+		if (GetSharedRefCount(PreviousValue) == 1)
+		{
+			const_cast<FSharedBuffer&>(*this).ReleaseData();
+			ReleaseWeakReference();
+		}
+	}
+
+	inline bool TryAddSharedReference() const
+	{
+		for (uint64 Value = ReferenceCountAndFlags.load(std::memory_order_relaxed);;)
+		{
+			if (GetSharedRefCount(Value) == 0)
+			{
+				return false;
+			}
+			if (ReferenceCountAndFlags.compare_exchange_weak(Value, Value + SetSharedRefCount(1),
+				std::memory_order_relaxed, std::memory_order_relaxed))
+			{
+				return true;
+			}
+		}
+	}
+
+	inline void AddWeakReference() const
+	{
+		const uint64 PreviousValue = ReferenceCountAndFlags.fetch_add(SetWeakRefCount(1), std::memory_order_relaxed);
+		checkSlow(GetWeakRefCount(PreviousValue) < 0xffffff);
+	}
+
+	inline void ReleaseWeakReference() const
+	{
+		const uint64 PreviousValue = ReferenceCountAndFlags.fetch_sub(SetWeakRefCount(1), std::memory_order_acq_rel);
+		checkSlow(GetWeakRefCount(PreviousValue) > 0);
+		if (GetWeakRefCount(PreviousValue) == 1)
+		{
+			DeleteBuffer(const_cast<FSharedBuffer*>(this));
+		}
+	}
+
+	friend struct SharedBufferPrivate::FSharedRefOps;
+	friend struct SharedBufferPrivate::FSharedPtrOps;
+	friend struct SharedBufferPrivate::FWeakPtrOps;
+
+private:
+	void* Data;
+	uint64 Size;
+	mutable std::atomic<uint64> ReferenceCountAndFlags;
+};
 
 ///////////////////////////////////////////////////////////////////////////////////////////////////////////////////////
 
-/**
- * Construct a shared buffer. \see \ref FSharedBuffer for more details.
- *
- * MakeSharedBuffer(); an empty non-owned buffer
- * MakeSharedBuffer(uint64 Size); an uninitialized buffer of Size bytes
- * MakeSharedBuffer(FSharedBuffer::AssumeOwnership, Data, Size); take ownership of the buffer (Data, Size)
- * MakeSharedBuffer(FSharedBuffer::Clone, Data, Size); an owned clone the buffer (Data, Size)
- * MakeSharedBuffer(FSharedBuffer::Wrap, Data, Size); non-owning wrapper of the buffer (Data, Size)
- */
-template <typename... ArgTypes, typename = decltype(FSharedBuffer(Forward<ArgTypes>(DeclVal<ArgTypes>())...))>
-inline FSharedBufferRef MakeSharedBuffer(ArgTypes&&... Args)
+namespace SharedBufferPrivate
+{
+
+inline bool FSharedRefOps::TryAddRef(const FSharedBuffer* const Buffer)
 {
-	return SharedBufferPrivate::MakeSharedBuffer(MakeShared<FSharedBuffer, ESPMode::ThreadSafe>(Forward<ArgTypes>(Args)...));
+	check(Buffer);
+	verify(Buffer->TryAddSharedReference());
+	return true;
 }
 
-/** Construct a const buffer to take ownership of the buffer (Data, Size). \see \ref FSharedBuffer for details. */
-inline FSharedBufferConstRef MakeSharedBuffer(FSharedBuffer::EAssumeOwnershipTag, const void* Data, uint64 Size)
+inline void FSharedRefOps::AddRef(const FSharedBuffer* const Buffer)
 {
-	// The const_cast is safe because it is returned as a const buffer that disallows mutable access.
-	return SharedBufferPrivate::MakeSharedBuffer(MakeShared<const FSharedBuffer, ESPMode::ThreadSafe>(FSharedBuffer::AssumeOwnership, const_cast<void*>(Data), Size));
+	check(Buffer);
+	Buffer->AddSharedReference();
 }
 
-/** Construct a const non-owning wrapper of the buffer (Data, Size) \see \ref FSharedBuffer for details. */
-inline FSharedBufferConstRef MakeSharedBuffer(FSharedBuffer::EWrapTag, const void* Data, uint64 Size)
+inline void FSharedRefOps::Release(const FSharedBuffer* const Buffer)
 {
-	// The const_cast is safe because it is returned as a const buffer that disallows mutable access.
-	return SharedBufferPrivate::MakeSharedBuffer(MakeShared<const FSharedBuffer, ESPMode::ThreadSafe>(FSharedBuffer::Wrap, const_cast<void*>(Data), Size));
+	check(Buffer);
+	Buffer->ReleaseSharedReference();
 }
 
-/** Construct a shared buffer from a memory view. \see \ref FSharedBuffer for details. */
-template <typename TagType, typename ViewType>
-inline auto MakeSharedBuffer(TagType Tag, TMemoryView<ViewType> View) -> decltype(MakeSharedBuffer(DeclVal<TagType>(), View.GetData(), View.Size()))
+inline bool FSharedPtrOps::TryAddRef(const FSharedBuffer* const Buffer)
 {
-	return MakeSharedBuffer(Tag, View.GetData(), View.Size());
+	return Buffer && Buffer->TryAddSharedReference();
 }
 
-///////////////////////////////////////////////////////////////////////////////////////////////////////////////////////
+inline void FSharedPtrOps::AddRef(const FSharedBuffer* const Buffer)
+{
+	if (Buffer)
+	{
+		Buffer->AddSharedReference();
+	}
+}
 
-/**
- * Return the buffer if it is owned, or a clone otherwise.
- *
- * @param Buffer A FSharedBuffer[Const]Ref or FSharedBuffer[Const][Weak]Ptr to make owned.
- * @return An owned copy of the input buffer in the same reference type as the input.
- */
-template <typename T, typename TEnableIf<TIsSame<FSharedBuffer, typename TDecay<decltype(*DeclVal<T>())>::Type>::Value>::Type* = nullptr>
-inline auto MakeSharedBufferOwned(T&& Buffer) -> decltype(typename TDecay<T>::Type(MakeSharedBuffer(FSharedBuffer::Clone, DeclVal<T>()->GetData(), DeclVal<T>()->Size())))
+inline void FSharedPtrOps::Release(const FSharedBuffer* const Buffer)
+{
+	if (Buffer)
+	{
+		Buffer->ReleaseSharedReference();
+	}
+}
+
+inline bool FWeakPtrOps::TryAddRef(const FSharedBuffer* const Buffer)
+{
+	if (Buffer)
+	{
+		Buffer->AddWeakReference();
+		return true;
+	}
+	return false;
+}
+
+inline void FWeakPtrOps::AddRef(const FSharedBuffer* const Buffer)
+{
+	if (Buffer)
+	{
+		Buffer->AddWeakReference();
+	}
+}
+
+inline void FWeakPtrOps::Release(const FSharedBuffer* const Buffer)
 {
-	return !IsValid(Buffer) || Buffer->IsOwned() ? Forward<T>(Buffer) : MakeSharedBuffer(FSharedBuffer::Clone, Buffer->GetData(), Buffer->Size());
+	if (Buffer)
+	{
+		Buffer->ReleaseWeakReference();
+	}
+}
+
 }
 
 ///////////////////////////////////////////////////////////////////////////////////////////////////////////////////////
diff --git a/Engine/Source/Runtime/Core/Public/Serialization/CompactBinary.h b/Engine/Source/Runtime/Core/Public/Serialization/CompactBinary.h
--- a/Engine/Source/Runtime/Core/Public/Serialization/CompactBinary.h
+++ b/Engine/Source/Runtime/Core/Public/Serialization/CompactBinary.h
@@ -832,7 +832,7 @@ public:
 	/** Construct a value and take ownership its memory. */
 	inline TCbBufferRef(EAssumeOwnershipTag, const void* Value)
 		: BaseType(Value)
-		, BufferPtr(MakeSharedBuffer(FSharedBuffer::AssumeOwnership, Value, BaseType::Size()))
+		, BufferPtr(FSharedBuffer::TakeOwnership(Value, BaseType::Size(), FMemory::Free))
 	{
 	}
 
@@ -846,7 +846,7 @@ public:
 	inline TCbBufferRef(ECloneTag, const BaseType& Value)
 		: BaseType(Value)
 	{
-		const FSharedBufferRef BufferRef = MakeSharedBuffer(BaseType::Size());
+		const FSharedBufferRef BufferRef = FSharedBuffer::Alloc(BaseType::Size());
 		BaseType::CopyTo(*BufferRef);
 		BufferPtr = BufferRef;
 		static_cast<BaseType&>(*this) = BaseType(BufferPtr->GetData(), BaseType::GetCopyType());
diff --git a/Engine/Source/Runtime/Core/Public/Serialization/CompactBinaryWriter.h b/Engine/Source/Runtime/Core/Public/Serialization/CompactBinaryWriter.h
--- a/Engine/Source/Runtime/Core/Public/Serialization/CompactBinaryWriter.h
+++ b/Engine/Source/Runtime/Core/Public/Serialization/CompactBinaryWriter.h
@@ -121,7 +121,7 @@ public:
 	/** Write a binary field by copying Size bytes from Value. */
 	CORE_API void Binary(const void* Value, uint64 Size);
 	/** Write a binary field by copying the view. */
-	inline void Binary(FConstMemoryView Value) { Binary(Value.GetData(), Value.Size()); }
+	inline void Binary(FConstMemoryView Value) { Binary(Value.GetData(), Value.GetSize()); }
 
 	/** Write a string field by copying the UTF-8 value. */
 	CORE_API void String(FAnsiStringView Value);
diff --git a/Engine/Source/Runtime/CoreUObject/Private/Tests/Virtualization/VirtualizedBulkDataTests.cpp b/Engine/Source/Runtime/CoreUObject/Private/Tests/Virtualization/VirtualizedBulkDataTests.cpp
--- a/Engine/Source/Runtime/CoreUObject/Private/Tests/Virtualization/VirtualizedBulkDataTests.cpp
+++ b/Engine/Source/Runtime/CoreUObject/Private/Tests/Virtualization/VirtualizedBulkDataTests.cpp
@@ -32,18 +32,18 @@ bool FVirtualizationWrapperTestBasic::RunTest(const FString& Parameters)
 	auto ValidateBulkData = [this, &SourceBuffer, BufferSize](const FVirtualizedUntypedBulkData& BulkDataToValidate, const TCHAR* Label)
 		{
 			FSharedBufferConstPtr RetrievedBuffer = BulkDataToValidate.GetData().Get();
-			TestEqual(FString::Printf(TEXT("%s buffer length"),Label), (int64)RetrievedBuffer->Size(), BufferSize);
+			TestEqual(FString::Printf(TEXT("%s buffer length"),Label), (int64)RetrievedBuffer->GetSize(), BufferSize);
 			TestTrue(FString::Printf(TEXT("SourceBuffer values == %s values"),Label), FMemory::Memcmp(SourceBuffer.Get(), RetrievedBuffer->GetData(), BufferSize) == 0);
 		};
 
 	// Create a basic bulkdata (but retain ownership of the buffer!)
 	FByteVirtualizedBulkData BulkData;
-	BulkData.UpdatePayload(MakeSharedBuffer(FSharedBuffer::Wrap, SourceBuffer.Get(), BufferSize));
+	BulkData.UpdatePayload(FSharedBuffer::MakeView(SourceBuffer.Get(), BufferSize));
 
 	// Test access via callback
 	BulkData.GetData([this,BufferSize,&SourceBuffer](FSharedBufferConstPtr Payload)
 		{
-			TestEqual(TEXT("Callback buffer length"), (int64)Payload->Size(), BufferSize);
+			TestEqual(TEXT("Callback buffer length"), (int64)Payload->GetSize(), BufferSize);
 			TestTrue(TEXT("SourceBuffer values == Callback values"), FMemory::Memcmp(SourceBuffer.Get(), Payload->GetData(), BufferSize) == 0);
 		});
 
@@ -55,7 +55,7 @@ bool FVirtualizationWrapperTestBasic::RunTest(const FString& Parameters)
 
 	// Create a new bulkdata object and copy by assignment (note we assign some junk data that will get overwritten)
 	FByteVirtualizedBulkData BulkDataAssignment;
-	BulkDataAssignment.UpdatePayload(MakeSharedBuffer(128));
+	BulkDataAssignment.UpdatePayload(FSharedBuffer::Alloc(128));
 	BulkDataAssignment = BulkData;
 
 	// Test both bulkdata objects
@@ -111,7 +111,7 @@ bool FVirtualizationWrapperTestUpdatePayload::RunTest(const FString& Parameters)
 
 	// Pass the buffer to to bulkdata but retain ownership
 	FByteVirtualizedBulkData BulkData;
-	BulkData.UpdatePayload(MakeSharedBuffer(FSharedBuffer::Wrap, OriginalData.Get(), BufferSize));
+	BulkData.UpdatePayload(FSharedBuffer::MakeView(OriginalData.Get(), BufferSize));
 
 	// Access the payload, edit it and push it back into the bulkdata object
 	{
@@ -119,24 +119,24 @@ bool FVirtualizationWrapperTestUpdatePayload::RunTest(const FString& Parameters)
 		// memory address since we retained ownership in the TUniquePtr, so the bulkdata object should 
 		// have created it's own copy.
 		FSharedBufferConstPtr Payload = BulkData.GetData().Get();
-		TestEqual(TEXT("Payload length"), (int64)Payload->Size(), BufferSize);
+		TestEqual(TEXT("Payload length"), (int64)Payload->GetSize(), BufferSize);
 		TestNotEqual(TEXT("OriginalData and the payload should have different memory addresses"), (uint8*)OriginalData.Get(), (uint8*)Payload->GetData());
-		TestTrue(TEXT("Orginal buffer == Payload data"), FMemory::Memcmp(OriginalData.Get(), Payload->GetData(), Payload->Size()) == 0);
+		TestTrue(TEXT("Orginal buffer == Payload data"), FMemory::Memcmp(OriginalData.Get(), Payload->GetData(), Payload->GetSize()) == 0);
 
 		// Make a copy of the payload that we can edit
-		FSharedBufferRef EditablePayload = MakeSharedBuffer(FSharedBuffer::Clone, Payload->GetView());
+		FSharedBufferRef EditablePayload = FSharedBuffer::Clone(*Payload);
 
 		const uint8 NewValue = 255;
-		FMemory::Memset(EditablePayload->GetData(), NewValue, EditablePayload->Size());
+		FMemory::Memset(EditablePayload->GetData(), NewValue, EditablePayload->GetSize());
 
 		// Update the bulkdata object with the new edited payload
 		BulkData.UpdatePayload(EditablePayload);
 
 		Payload = BulkData.GetData().Get();
-		TestEqual(TEXT("Updated payload length"), (int64)Payload->Size(), BufferSize);
+		TestEqual(TEXT("Updated payload length"), (int64)Payload->GetSize(), BufferSize);
 		TestEqual(TEXT("Payload and EditablePayload should have the same memory addresses"), (uint8*)Payload->GetData(), (uint8*)EditablePayload->GetData());
 
-		const bool bAllElementsCorrect = Algo::AllOf(TArrayView<uint8>((uint8*)Payload->GetData(), Payload->Size()), [NewValue](uint8 Val)
+		const bool bAllElementsCorrect = Algo::AllOf(TArrayView<uint8>((uint8*)Payload->GetData(), Payload->GetSize()), [NewValue](uint8 Val)
 			{ 
 				return Val == NewValue; 
 			});
@@ -151,14 +151,14 @@ bool FVirtualizationWrapperTestUpdatePayload::RunTest(const FString& Parameters)
 
 		// Update the bulkdata object with the original data but this time we give ownership of the buffer to
 		// the bulkdata object.
-		BulkData.UpdatePayload(MakeSharedBuffer(FSharedBuffer::AssumeOwnership, OriginalData.Release(), BufferSize));
+		BulkData.UpdatePayload(FSharedBuffer::TakeOwnership(OriginalData.Release(), BufferSize, FMemory::Free));
 
 		FSharedBufferConstPtr Payload = BulkData.GetData().Get();
-		TestEqual(TEXT("Updated payload length"), (int64)Payload->Size(), BufferSize);
+		TestEqual(TEXT("Updated payload length"), (int64)Payload->GetSize(), BufferSize);
 		TestEqual(TEXT("Payload and OriginalDataPtr should have the same memory addresses"), (uint8*)Payload->GetData(), OriginalDataPtr);
 
 		// The original data was all zeros, so we can test for that to make sure that the contents are correct.
-		const bool bAllElementsCorrect = Algo::AllOf(TArrayView<uint8>((uint8*)Payload->GetData(), Payload->Size()), [](uint8 Val)
+		const bool bAllElementsCorrect = Algo::AllOf(TArrayView<uint8>((uint8*)Payload->GetData(), Payload->GetSize()), [](uint8 Val)
 			{
 				return Val == 0;
 			});
diff --git a/Engine/Source/Runtime/CoreUObject/Private/Virtualization/VirtualizationManager.cpp b/Engine/Source/Runtime/CoreUObject/Private/Virtualization/VirtualizationManager.cpp
--- a/Engine/Source/Runtime/CoreUObject/Private/Virtualization/VirtualizationManager.cpp
+++ b/Engine/Source/Runtime/CoreUObject/Private/Virtualization/VirtualizationManager.cpp
@@ -33,13 +33,13 @@ bool FVirtualizationManager::PushData(const FSharedBufferConstPtr& Payload, cons
 	}
 
 	// Early out if we have no payload
-	if (!Payload.IsValid() || Payload->Size() == 0)
+	if (!Payload.IsValid() || Payload->GetSize() == 0)
 	{
 		return false;
 	}
 
 	// Early out if the payload length is below our minimum required length
-	if ((int64)Payload->Size() < MinPayloadLength)
+	if ((int64)Payload->GetSize() < MinPayloadLength)
 	{
 		return false;
 	}
@@ -52,7 +52,7 @@ bool FVirtualizationManager::PushData(const FSharedBufferConstPtr& Payload, cons
 		TUniquePtr<FArchive> FileAr(IFileManager::Get().CreateFileWriter(*FilePath));
 		check(FileAr);
 
-		uint64 PayloadLength = Payload->Size();
+		uint64 PayloadLength = Payload->GetSize();
 		// Cast away the const because FArchive requires a non-const pointer even if we are only reading from it (saving)
 		void* PayloadPtr = const_cast<void*>(Payload->GetData());
 
@@ -81,7 +81,7 @@ FSharedBufferConstPtr FVirtualizationManager::PullData(const FGuid& Guid)
 	uint8* DataPtr = (uint8*)FMemory::Malloc(DataLength, DEFAULT_ALIGNMENT);
 	FileAr->Serialize(DataPtr, DataLength);
 
-	return MakeSharedBuffer(FSharedBuffer::AssumeOwnership, DataPtr, DataLength);
+	return FSharedBuffer::TakeOwnership(DataPtr, DataLength, FMemory::Free);
 }
 
 void FVirtualizationManager::LoadSettingsFromConfigFiles()
diff --git a/Engine/Source/Runtime/CoreUObject/Private/Virtualization/VirtualizedBulkData.cpp b/Engine/Source/Runtime/CoreUObject/Private/Virtualization/VirtualizedBulkData.cpp
--- a/Engine/Source/Runtime/CoreUObject/Private/Virtualization/VirtualizedBulkData.cpp
+++ b/Engine/Source/Runtime/CoreUObject/Private/Virtualization/VirtualizedBulkData.cpp
@@ -219,7 +219,7 @@ bool FVirtualizedUntypedBulkData::SerializeData(FArchive& Ar, FSharedBufferConst
 {
 	TRACE_CPUPROFILER_EVENT_SCOPE(FVirtualizedUntypedBulkData::SerializeData);
 
-	const int64 PayloadSize = Ar.IsSaving() ? InPayload->Size() : GetBulkDataSize();
+	const int64 PayloadSize = Ar.IsSaving() ? InPayload->GetSize() : GetBulkDataSize();
 
 	if (PayloadSize > 0)
 	{
@@ -243,7 +243,7 @@ bool FVirtualizedUntypedBulkData::SerializeData(FArchive& Ar, FSharedBufferConst
 
 		if (Ar.IsLoading())
 		{
-			InPayload = MakeSharedBuffer(FSharedBuffer::AssumeOwnership, DataPtr, PayloadSize);
+			InPayload = FSharedBuffer::TakeOwnership(DataPtr, PayloadSize, FMemory::Free);
 		}
 
 		return true;
@@ -291,10 +291,10 @@ FSharedBufferConstPtr FVirtualizedUntypedBulkData::PullData() const
 
 	if (PulledPayload.IsValid())
 	{
-		checkf(	PayloadLength == PulledPayload->Size(),	
+		checkf(	PayloadLength == PulledPayload->GetSize(),	
 				TEXT("Mismatch between serialized length (%d) and virtualized data length (%d)"), 
 				PayloadLength,
-				PulledPayload->Size());
+				PulledPayload->GetSize());
 	}
 	else
 	{
@@ -399,8 +399,8 @@ void FVirtualizedUntypedBulkData::UpdatePayload(const FSharedBufferConstRef& InP
 	UnloadData();
 
 	// Make sure that we own the memory in the shared buffer
-	Payload = MakeSharedBufferOwned(InPayload);
-	PayloadLength = (int64)InPayload->Size();
+	Payload = FSharedBuffer::MakeOwned(InPayload);
+	PayloadLength = (int64)InPayload->GetSize();
 
 	bIsVirtualized = false;
 	bIsDataStoredAsCompressed = false;
diff --git a/Engine/Source/Runtime/CoreUObject/Public/Virtualization/VirtualizationUtilities.h b/Engine/Source/Runtime/CoreUObject/Public/Virtualization/VirtualizationUtilities.h
--- a/Engine/Source/Runtime/CoreUObject/Public/Virtualization/VirtualizationUtilities.h
+++ b/Engine/Source/Runtime/CoreUObject/Public/Virtualization/VirtualizationUtilities.h
@@ -21,7 +21,7 @@ public:
 
 		if (Payload.IsValid())
 		{
-			const int64 CurrentDataLength = Payload->Size();
+			const int64 CurrentDataLength = Payload->GetSize();
 
 			// Clone the payload so that we have a local copy that we can
 			// append additional data to.
@@ -48,7 +48,7 @@ public:
 	{
 		// Remove the slack from the allocated bulk data
 		Buffer = FMemory::Realloc(Buffer, DataLength, DEFAULT_ALIGNMENT);
-		BulkData.UpdatePayload(MakeSharedBuffer(FSharedBuffer::AssumeOwnership, Buffer, DataLength));
+		BulkData.UpdatePayload(FSharedBuffer::TakeOwnership(Buffer, DataLength, FMemory::Free));
 	}
 
 	virtual void Serialize(void* Data, int64 Num)
@@ -151,7 +151,7 @@ namespace UE4VirtualizedBulkData_Private
 
 		int64 GetDataLength() const
 		{
-			return Payload.IsValid() ? Payload->Size() : 0;
+			return Payload.IsValid() ? Payload->GetSize() : 0;
 		}
 
 	private:
