SharedBuffer: Fixed compile errors

#rnx

diff --git a/Engine/Source/Runtime/Core/Private/Memory/SharedBuffer.cpp b/Engine/Source/Runtime/Core/Private/Memory/SharedBuffer.cpp
--- a/Engine/Source/Runtime/Core/Private/Memory/SharedBuffer.cpp
+++ b/Engine/Source/Runtime/Core/Private/Memory/SharedBuffer.cpp
@@ -22,8 +22,8 @@ FSharedBufferRef FSharedBuffer::NewBuffer(
 
 void FSharedBuffer::DeleteBuffer(FSharedBuffer* const Buffer)
 {
-	checkSlow(!Data && !Size);
-	checkSlow(ReferenceCountAndFlags == 0);
+	checkSlow(!Buffer->Data && !Buffer->Size);
+	checkSlow(Buffer->ReferenceCountAndFlags == 0);
 	Buffer->~FSharedBuffer();
 	FMemory::Free(Buffer);
 }
diff --git a/Engine/Source/Runtime/Core/Public/Memory/SharedBuffer.h b/Engine/Source/Runtime/Core/Public/Memory/SharedBuffer.h
--- a/Engine/Source/Runtime/Core/Public/Memory/SharedBuffer.h
+++ b/Engine/Source/Runtime/Core/Public/Memory/SharedBuffer.h
@@ -695,7 +695,7 @@ private:
 	private:
 		~TSharedBufferDeleter() = default;
 
-		virtual void Free(void* const InData) final
+		virtual void Free(void* InData) final
 		{
 			Deleter(InData);
 			this->~TSharedBufferDeleter();
