Fix warnings in SherwoodTable implementation.

diff --git a/Engine/Source/Runtime/Core/Public/Experimental/Containers/SherwoodHashTable.h b/Engine/Source/Runtime/Core/Public/Experimental/Containers/SherwoodHashTable.h
--- a/Engine/Source/Runtime/Core/Public/Experimental/Containers/SherwoodHashTable.h
+++ b/Engine/Source/Runtime/Core/Public/Experimental/Containers/SherwoodHashTable.h
@@ -50,7 +50,7 @@ struct TSherwoodHashTable
 	static constexpr uint32 MinNumSlots = 4;
 
 	// Ratio between number of stored elements and allocated capacity beyond which the container will be grown (doubled in size).
-	static constexpr double MaxLoadFactor = 0.9;
+	static constexpr float MaxLoadFactor = 0.9f;
 	static_assert(MaxLoadFactor >= 0.5 && MaxLoadFactor <= 0.9, "MaxLoadFactor must be in range [0.5 .. 0.9]");
 
 	// NOTE: Non-trivial type support, move and copy ops are not implemented yet, but can be.
@@ -149,10 +149,10 @@ struct TSherwoodHashTable
 		return NumSlotsMinusOne ? NumSlotsMinusOne + 1 : 0;
 	}
 
-	static uint32 ComputeMaxLookups(uint32 InNumSlots)
+	static uint8 ComputeMaxLookups(uint32 InNumSlots)
 	{
 		uint32 Desired = (32-FMath::CountLeadingZeros(InNumSlots)); // FloorLog2
-		return FMath::Max(MinNumLookups, Desired);
+		return uint8(FMath::Max(MinNumLookups, Desired));
 	}
 
 	static FData AllocateData(uint32 Count)
@@ -300,7 +300,7 @@ struct TSherwoodHashTable
 			return;
 		}
 
-		const uint32 NewMaxLookups = ComputeMaxLookups(DesiredNumSlots);
+		const uint8 NewMaxLookups = ComputeMaxLookups(DesiredNumSlots);
 
 		FData OldData = CurrentData;
 		const uint32 OldNumSlotsMinusOne = NumSlotsMinusOne;
