Rename LLM trace channel to "MemTag" and add "Memory" channel preset.

diff --git a/Engine/Config/BaseEngine.ini b/Engine/Config/BaseEngine.ini
--- a/Engine/Config/BaseEngine.ini
+++ b/Engine/Config/BaseEngine.ini
@@ -1930,6 +1930,7 @@ DrawColor=(R=140,G=255,B=0,A=164)
 [Trace.ChannelPresets]
 Default=cpu,frame,log,bookmark
 Rendering=gpu,cpu,frame,log,bookmark
+Memory=memtag,memalloc,callstack,module
 
 [RemoteConfiguration]
 Enabled=false
diff --git a/Engine/Source/Runtime/Core/Private/HAL/LowLevelMemTracker.cpp b/Engine/Source/Runtime/Core/Private/HAL/LowLevelMemTracker.cpp
--- a/Engine/Source/Runtime/Core/Private/HAL/LowLevelMemTracker.cpp
+++ b/Engine/Source/Runtime/Core/Private/HAL/LowLevelMemTracker.cpp
@@ -19,7 +19,7 @@
 
 
 
-UE_TRACE_CHANNEL(MemoryChannel, "Memory overview", true)
+UE_TRACE_CHANNEL(MemTagChannel, "Memory overview", true)
 
 UE_TRACE_EVENT_BEGIN(LLM, TagsSpec, NoSync|Important)
 	UE_TRACE_EVENT_FIELD(const void*, TagId)
@@ -947,7 +947,7 @@ void FLowLevelMemTracker::ProcessCommandLine(const TCHAR* CmdLine)
 #endif
 
 	bool bLocalCsvWriterEnabled = FParse::Param(CmdLine, TEXT("LLMCSV"));
-	bool bLocalTraceWriterEnabled = UE_TRACE_CHANNELEXPR_IS_ENABLED(MemoryChannel);
+	bool bLocalTraceWriterEnabled = UE_TRACE_CHANNELEXPR_IS_ENABLED(MemTagChannel);
 	// automatically enable LLM if only csv or trace output is active
 	if (bLocalCsvWriterEnabled || bLocalTraceWriterEnabled)
 	{
@@ -3693,7 +3693,7 @@ namespace LLMPrivate
 
 	void FLLMTraceWriter::Publish(FLowLevelMemTracker& LLMRef, const FTrackerTagSizeMap& TagSizes, const FTagData* OverrideTrackedTotalTagData, const FTagData* OverrideUntaggedTagData, int64 TrackedTotal, bool bTrackPeaks)
 	{
-		if (!UE_TRACE_CHANNELEXPR_IS_ENABLED(MemoryChannel))
+		if (!UE_TRACE_CHANNELEXPR_IS_ENABLED(MemTagChannel))
 		{
 			return;
 		}
@@ -3710,7 +3710,7 @@ namespace LLMPrivate
 			};
 			static_assert(UE_ARRAY_COUNT(TrackerNames) == int(ELLMTracker::Max), "");
 			uint32 NameLen = TrackerNames[(uint8)Tracker].Len;
-			UE_TRACE_LOG(LLM, TrackerSpec, MemoryChannel, NameLen * sizeof(ANSICHAR))
+			UE_TRACE_LOG(LLM, TrackerSpec, MemTagChannel, NameLen * sizeof(ANSICHAR))
 				<< TrackerSpec.TrackerId((uint8)Tracker)
 				<< TrackerSpec.Name(TrackerNames[(uint8)Tracker].Name, NameLen);
 		}
@@ -3727,7 +3727,7 @@ namespace LLMPrivate
 			const FTagData* Parent = TagData->GetParent();
 			NameBuffer.Reset();
 			TagData->AppendDisplayPath(NameBuffer);
-			UE_TRACE_LOG(LLM, TagsSpec, MemoryChannel, NameBuffer.Len() * sizeof(ANSICHAR))
+			UE_TRACE_LOG(LLM, TagsSpec, MemTagChannel, NameBuffer.Len() * sizeof(ANSICHAR))
 				<< TagsSpec.TagId(GetTagId(TagData))
 				<< TagsSpec.ParentId(GetTagId(Parent))
 				<< TagsSpec.Name(*NameBuffer, NameBuffer.Len());
@@ -3786,7 +3786,7 @@ namespace LLMPrivate
 		int TagCount = TagIds.Num();
 		LLMCheck(TagCount == TagValues.Num());
 		const uint64 Cycle = FPlatformTime::Cycles64();
-		UE_TRACE_LOG(LLM, TagValue, MemoryChannel)
+		UE_TRACE_LOG(LLM, TagValue, MemTagChannel)
 			<< TagValue.TrackerId((uint8)Tracker)
 			<< TagValue.Cycle(Cycle)
 			<< TagValue.Tags(TagIds.GetData(), TagCount)
@@ -3799,7 +3799,7 @@ namespace LLMPrivate
 		{
 			const FTagData* TagData = LLMRef.FindTagData(static_cast<ELLMTag>(GenericTagIndex));
 			FString TagName = TagData->GetDisplayPath();
-			UE_TRACE_LOG(LLM, TagsSpec, MemoryChannel, TagName.Len() * sizeof(ANSICHAR))
+			UE_TRACE_LOG(LLM, TagsSpec, MemTagChannel, TagName.Len() * sizeof(ANSICHAR))
 				<< TagsSpec.TagId(GetTagId(TagData))
 				<< TagsSpec.ParentId(GetTagId(TagData->GetParent()))
 				<< TagsSpec.Name(*TagName, TagName.Len());
