Memory tracing.

* Trace all allocation events on channel "MemTrack"
* Index all loaded modules and traces on "Diagnostics" channel.
* Trace all unique callstacks on channel "Callstacks".

diff --git a/Engine/Source/Runtime/Core/Private/HAL/UnrealMemory.cpp b/Engine/Source/Runtime/Core/Private/HAL/UnrealMemory.cpp
--- a/Engine/Source/Runtime/Core/Private/HAL/UnrealMemory.cpp
+++ b/Engine/Source/Runtime/Core/Private/HAL/UnrealMemory.cpp
@@ -19,6 +19,7 @@
 -----------------------------------------------------------------------------*/
 
 #include "ProfilingDebugging/MallocProfiler.h"
+#include "ProfilingDebugging/MemoryTrace.h"
 #include "HAL/MallocThreadSafeProxy.h"
 #include "HAL/MallocVerify.h"
 #include "HAL/MallocLeakDetectionProxy.h"
@@ -338,16 +339,28 @@ static int FMemory_GCreateMalloc_ThreadUnsafe()
 		UE_LOG(LogMemory, Fatal, TEXT("PLATFORM_USES_FIXED_GMalloc_CLASS only makes sense for allocators that are internally threadsafe."));
 	}
 #else
-// so now check to see if we are using a Mem Profiler which wraps the GMalloc
+
+#if UE_MEMORY_TRACE_ENABLED
+	FMalloc* TraceMalloc = MemoryTrace_Create(GMalloc);
+	if (TraceMalloc != GMalloc)
+	{
+		GMalloc = TraceMalloc;
+		MemoryTrace_Initialize();
+	}
+	else
+#endif // UE_MEMORY_TRACE_ENABLED
+	{
+	// so now check to see if we are using a Mem Profiler which wraps the GMalloc
 #if USE_MALLOC_PROFILER
-	#if WITH_ENGINE && IS_MONOLITHIC
-		GMallocProfiler = new FMallocProfilerEx( GMalloc );
-	#else
-		GMallocProfiler = new FMallocProfiler( GMalloc );
-	#endif
-	GMallocProfiler->BeginProfiling();
-	GMalloc = GMallocProfiler;
+		#if WITH_ENGINE && IS_MONOLITHIC
+			GMallocProfiler = new FMallocProfilerEx( GMalloc );
+		#else
+			GMallocProfiler = new FMallocProfiler( GMalloc );
+		#endif
+		GMallocProfiler->BeginProfiling();
+		GMalloc = GMallocProfiler;
 #endif
+	}
 
 	// if the allocator is already thread safe, there is no need for the thread safe proxy
 	if (!GMalloc->IsInternallyThreadSafe())
diff --git a/Engine/Source/Runtime/Core/Private/ProfilingDebugging/MemoryTrace.h b/Engine/Source/Runtime/Core/Private/ProfilingDebugging/MemoryTrace.h
new file mode 100644
--- /dev/null
+++ b/Engine/Source/Runtime/Core/Private/ProfilingDebugging/MemoryTrace.h
@@ -0,0 +1,37 @@
+// Copyright Epic Games, Inc. All Rights Reserved.
+
+#include "Trace/Config.h"
+
+//#define UE_MEMORY_TRACE_ENABLED 0
+
+#if !defined(UE_MEMORY_TRACE_ENABLED) && UE_TRACE_ENABLED
+#	if PLATFORM_WINDOWS
+#		if !PLATFORM_USES_FIXED_GMalloc_CLASS && PLATFORM_64BITS
+#			define UE_MEMORY_TRACE_ENABLED !UE_BUILD_SHIPPING
+#		endif
+#	endif
+#endif
+
+#if !defined(UE_MEMORY_TRACE_ENABLED)
+#	define UE_MEMORY_TRACE_ENABLED 0
+#endif
+
+#if UE_MEMORY_TRACE_ENABLED
+
+UE_TRACE_CHANNEL_EXTERN(MemTrackChannel);
+
+class FMalloc*	MemoryTrace_Create(class FMalloc* InMalloc);
+void			MemoryTrace_Initialize();
+#if 0
+CORE_API void	MemoryTrace_CoreAdd(void* Base, SIZE_T Size);
+CORE_API void	MemoryTrace_CoreRemove(void* Base, SIZE_T Size);
+#endif // 0
+
+#else
+
+#if 0
+inline void		MemoryTrace_CoreAdd(void* Base, SIZE_T Size) {}
+inline void		MemoryTrace_CoreRemove(void* Base, SIZE_T Size) {}
+#endif // 0
+
+#endif // UE_MEMORY_TRACE_ENABLED
diff --git a/Engine/Source/Runtime/Core/Private/ProfilingDebugging/Windows/CallstackTrace.cpp b/Engine/Source/Runtime/Core/Private/ProfilingDebugging/Windows/CallstackTrace.cpp
new file mode 100644
--- /dev/null
+++ b/Engine/Source/Runtime/Core/Private/ProfilingDebugging/Windows/CallstackTrace.cpp
@@ -0,0 +1,753 @@
+// Copyright Epic Games, Inc. All Rights Reserved.
+
+#include "Containers/Array.h"
+#include "Containers/ArrayView.h"
+#include "Experimental/Containers/SherwoodHashTable.h"
+#include "CoreTypes.h"
+#include "HAL/CriticalSection.h"
+#include "HAL/MemoryBase.h"
+#include "MemoryTrace.inl"
+#include "Misc/ScopeRWLock.h"
+#include "Containers/CircularQueue.h"
+#include "HAL/RunnableThread.h"
+#include "HAL/LowLevelMemTracker.h"
+#include "ProfilingDebugging/MemoryTrace.h"
+
+#include "Windows/AllowWindowsPlatformTypes.h"
+#	include <winnt.h>
+#	include <winternl.h>
+#include "Windows/HideWindowsPlatformTypes.h"
+#include "Misc\CoreDelegates.h"
+#include "Trace/Trace.inl"
+#include "AtomicQueue/AtomicQueue.h"
+
+// 0=off, 1=stats, 2=validation, 3=truth_compare
+#define BACKTRACE_DBGLVL 0
+
+/*
+ * Windows' x64 binaries contain a ".pdata" section that describes the location
+ * and size of its functions and details on how to unwind them. The unwind
+ * information includes descriptions about a function's stack frame size and
+ * the non-volatile registers it pushes onto the stack. From this we can
+ * calculate where a call instruction wrote its return address. This is enough
+ * to walk the callstack and by caching this information it can be done
+ * efficiently.
+ *
+ * Some functions need a variable amount of stack (such as those that use
+ * alloc() for example) will use a frame pointer. Frame pointers involve saving
+ * and restoring the stack pointer in the function's prologue/epilogue. This
+ * frees the function up to modify the stack pointer arbitrarily. This
+ * significantly complicates establishing where a return address is, so this
+ * pdata scheme of walking the stack just doesn't support functions like this.
+ * Walking stops if it encounters such a function. Fortunately there are
+ * usually very few such functions, saving us from having to read and track
+ * non-volatile registers which adds a significant amount of work.
+ *
+ * A further optimisation is to to assume we are only interested methods that
+ * are part of engine or game code. As such we only build lookup tables for
+ * such modules and never accept OS or third party modules. Backtracing stops
+ * if an address is encountered which doesn't map to a known module.
+ */
+
+////////////////////////////////////////////////////////////////////////////////
+static uint32 AddressToId(UPTRINT Address)
+{
+	return uint32(Address >> 16);
+}
+
+static UPTRINT IdToAddress(uint32 Id)
+{
+	return UPTRINT(Id << 16);
+}
+
+struct FIdPredicate
+{
+	template <class T> bool operator () (uint32 Id, const T& Item) const { return Id < Item.Id; }
+	template <class T> bool operator () (const T& Item, uint32 Id) const { return Item.Id < Id; }
+};
+
+////////////////////////////////////////////////////////////////////////////////
+struct FUnwindInfo
+{
+	uint8	Version : 3;
+	uint8	Flags : 5;
+	uint8	PrologBytes;
+	uint8	NumUnwindCodes;
+	uint8	FrameReg	: 4;
+	uint8	FrameRspBias : 4;
+};
+
+struct FUnwindCode
+{
+	uint8	PrologOffset;
+	uint8	OpCode : 4;
+	uint8	OpInfo : 4;
+	uint16	Params[];
+};
+
+enum
+{
+	UWOP_PUSH_NONVOL		= 0,	// 1 node
+	UWOP_ALLOC_LARGE		= 1,	// 2 or 3 nodes
+	UWOP_ALLOC_SMALL		= 2,	// 1 node
+	UWOP_SET_FPREG			= 3,	// 1 node
+	UWOP_SAVE_NONVOL		= 4,	// 2 nodes
+	UWOP_SAVE_NONVOL_FAR	= 5,	// 3 nodes
+	UWOP_SAVE_XMM128		= 8,	// 2 nodes
+	UWOP_SAVE_XMM128_FAR	= 9,	// 3 nodes
+	UWOP_PUSH_MACHFRAME		= 10,	// 1 node
+};
+
+////////////////////////////////////////////////////////////////////////////////
+namespace {
+
+	using namespace Experimental;
+	using namespace atomic_queue;
+
+	class FCallstackProcWorker : public FRunnable 
+	{
+		public:
+		struct FBacktraceEntry
+		{
+			enum {MaxStackDepth = 256};
+			uint64	Id = 0;
+			uint32	FrameCount = 0;
+			uint64	Frames[MaxStackDepth] = { 0 };
+		};
+
+		FCallstackProcWorker()
+			: bRun(true)
+		{
+			KnownSet.Reserve(1024 * 1024 * 2);
+		}
+		
+		uint32 Run() override;
+		void Stop() override { bRun = false; }
+		void AddCallstack(const FBacktraceEntry& Entry);
+		void AddWork(const FBacktraceEntry& Entry);
+
+		TSherwoodSet<uint64> 				KnownSet;
+		FCriticalSection					ProducerCs;
+		AtomicQueue2<FBacktraceEntry, 256>	Queue;
+		TAtomic<bool>						bRun;
+		TAtomic<bool>						bStarted;
+		static FRunnableThread*				Thread;
+	};
+}
+
+FRunnableThread* FCallstackProcWorker::Thread = nullptr;
+
+////////////////////////////////////////////////////////////////////////////////
+class FBacktracer
+{
+public:
+							FBacktracer(FMalloc* InMalloc);
+							~FBacktracer();
+	static FBacktracer*		Get();
+	static void				StartWorker();
+	void					AddModule(UPTRINT Base, const TCHAR* Name);
+	void					RemoveModule(UPTRINT Base);
+	void*					GetBacktraceId(void* AddressOfReturnAddress) const;
+	int32					GetFrameSize(void* FunctionAddress) const;
+
+private:
+	struct FFunction
+	{
+		uint32				Id;
+		int32				RspBias;
+#if BACKTRACE_DBGLVL >= 2
+		uint32				Size;
+		const FUnwindInfo*	UnwindInfo;
+#endif
+	};
+
+	struct FModule
+	{
+		uint32				Id;
+		uint32				IdSize = 0;
+		uint32				NumFunctions;
+#if BACKTRACE_DBGLVL >= 1
+		uint16				NumFpTypes;
+		//uint16			*padding*
+#else
+		//uint32			*padding*
+#endif
+		FFunction*			Functions;
+	};
+
+	struct FLookupState
+	{
+		FModule				Module;
+	};
+
+	void					StartWorkerThread();
+	const FFunction*		LookupFunction(UPTRINT Address, FLookupState& State) const;
+	static FBacktracer*		Instance;
+	mutable FRWLock			Lock;
+	TMiniArray<FModule>		Modules;
+	FMalloc*				Malloc;
+	FCallstackProcWorker*	ProcessingThreadRunnable;
+#if BACKTRACE_DBGLVL >= 1
+	mutable uint32			NumFpTruncations = 0;
+	mutable uint32			TotalFunctions = 0;
+#endif
+};
+
+////////////////////////////////////////////////////////////////////////////////
+FBacktracer* FBacktracer::Instance = nullptr;
+
+////////////////////////////////////////////////////////////////////////////////
+FBacktracer::FBacktracer(FMalloc* InMalloc)
+: Modules(InMalloc)
+, Malloc(InMalloc)
+{
+	Modules.MakeRoom();
+	Instance = this;
+	ProcessingThreadRunnable = (FCallstackProcWorker*) Malloc->Malloc(sizeof(FCallstackProcWorker));
+	ProcessingThreadRunnable = new(ProcessingThreadRunnable) FCallstackProcWorker();
+
+	// We cannot start worker threads directly on creation. Delay them until the
+	// engine has gotten a little bit further.
+	FCoreDelegates::GetPreMainInitDelegate().AddStatic(FBacktracer::StartWorker);
+
+}
+
+////////////////////////////////////////////////////////////////////////////////
+FBacktracer::~FBacktracer()
+{
+	for (FModule& Module : Modules)
+	{
+		Malloc->Free(Module.Functions);
+	}
+
+	ProcessingThreadRunnable->Stop();
+	if (FCallstackProcWorker::Thread)
+	{
+		FCallstackProcWorker::Thread->WaitForCompletion();
+	} 
+}
+
+////////////////////////////////////////////////////////////////////////////////
+FBacktracer* FBacktracer::Get()
+{
+	return Instance;
+}
+////////////////////////////////////////////////////////////////////////////////
+void FBacktracer::StartWorker()
+{
+	if (Instance)
+	{
+		Instance->StartWorkerThread();
+	}
+}
+
+////////////////////////////////////////////////////////////////////////////////
+void FBacktracer::StartWorkerThread()
+{
+	if (ProcessingThreadRunnable)
+	{
+		FCallstackProcWorker::Thread = FRunnableThread::Create(ProcessingThreadRunnable, TEXT("TraceMemCallstacks"), 0, TPri_BelowNormal);
+	}
+}
+
+////////////////////////////////////////////////////////////////////////////////
+void FBacktracer::AddModule(UPTRINT ModuleBase, const TCHAR* Name)
+{
+	const auto* DosHeader = (IMAGE_DOS_HEADER*)ModuleBase;
+	const auto* NtHeader = (IMAGE_NT_HEADERS*)(ModuleBase + DosHeader->e_lfanew);
+	const IMAGE_FILE_HEADER* FileHeader = &(NtHeader->FileHeader);
+
+	if (FCString::Strfind(Name, TEXT("Binaries")) == nullptr || FCString::Strfind(Name, TEXT("ThirdParty")) != nullptr)
+	{
+		return;
+	}
+
+	uint32 NumSections = FileHeader->NumberOfSections;
+	const auto* Sections = (IMAGE_SECTION_HEADER*)(UPTRINT(&(NtHeader->OptionalHeader)) + FileHeader->SizeOfOptionalHeader);
+
+	// Find ".pdata" section
+	UPTRINT PdataBase = 0;
+	for (uint32 i = 0; i < NumSections; ++i)
+	{
+		const IMAGE_SECTION_HEADER* Section = Sections + i;
+		if (*(uint64*)(Section->Name) == 0x61'74'61'64'70'2eull) // Sections names are eight bytes and zero padded. This constant is '.pdata'
+		{
+			PdataBase = ModuleBase + Section->VirtualAddress;
+			break;
+		}
+	}
+
+	if (PdataBase == 0)
+	{
+		return;
+	}
+
+	// Count the number of functions.
+	const auto* FunctionTables = (RUNTIME_FUNCTION*)PdataBase;
+	uint32 NumFunctions = 0;
+	do
+	{
+		++NumFunctions;
+	}
+	while (FunctionTables[NumFunctions].BeginAddress);
+
+	// Allocate some space for the module's function-to-frame-size table
+	auto* OutTable = (FFunction*)Malloc->Malloc(sizeof(FFunction) * NumFunctions);
+	FFunction* OutTableCursor = OutTable;
+
+	// Extract frame size for each function from pdata's unwind codes.
+	uint32 NumFpFuncs = 0;
+	for (uint32 i = 0; i < NumFunctions; ++i)
+	{
+		const RUNTIME_FUNCTION* FunctionTable = FunctionTables + i;
+
+		UPTRINT UnwindInfoAddr = ModuleBase + FunctionTable->UnwindInfoAddress;
+		const auto* UnwindInfo = (FUnwindInfo*)UnwindInfoAddr;
+
+		if (UnwindInfo->Version != 1)
+		{
+			/* some v2s have been seen in msvc. Always seem to be assembly
+			 * routines (memset, memcpy, etc) */
+			continue;
+		}
+
+		int32 FpInfo = 0;
+		int32 RspBias = 0;
+
+#if BACKTRACE_DBGLVL >= 2
+		uint32 PrologVerify = UnwindInfo->PrologBytes;
+#endif
+
+		const auto* Code = (FUnwindCode*)(UnwindInfo + 1);
+		const auto* EndCode = Code + UnwindInfo->NumUnwindCodes;
+		while (Code < EndCode)
+		{
+#if BACKTRACE_DBGLVL >= 2
+			if (Code->PrologOffset > PrologVerify)
+			{
+				PLATFORM_BREAK();
+			}
+			PrologVerify = Code->PrologOffset;
+#endif
+
+			switch (Code->OpCode)
+			{
+			case UWOP_PUSH_NONVOL:
+				RspBias += 8;
+				Code += 1;
+				break;
+
+			case UWOP_ALLOC_LARGE:
+				if (Code->OpInfo)
+				{
+					RspBias += *(uint32*)(Code->Params);
+					Code += 3;
+				}
+				else
+				{
+					RspBias += Code->Params[0] * 8;
+					Code += 2;
+				}
+				break;
+
+			case UWOP_ALLOC_SMALL:
+				RspBias += (Code->OpInfo * 8) + 8;
+				Code += 1;
+				break;
+
+			case UWOP_SET_FPREG:
+				// Function will adjust RSP (e.g. through use of alloca()) so it
+				// uses a frame pointer register. There's instructions like;
+				//
+				//   push FRAME_REG
+				//   lea FRAME_REG, [rsp + (FRAME_RSP_BIAS * 16)]
+				//   ...
+				//   add rsp, rax
+				//   ...
+				//   sub rsp, FRAME_RSP_BIAS * 16
+				//   pop FRAME_REG
+				//   ret
+				//
+				// To recover the stack frame we would need to track non-volatile
+				// registers which adds a lot of overhead for a small subset of
+				// functions. Instead we'll end backtraces at these functions.
+
+
+				// MSB is set to detect variable sized frames that we can't proceed
+				// past when back-tracing.
+				NumFpFuncs++;
+				FpInfo |= 0x80000000 | (uint32(UnwindInfo->FrameReg) << 27) | (uint32(UnwindInfo->FrameRspBias) << 23);
+				Code += 1;
+				break;
+
+			case UWOP_PUSH_MACHFRAME:
+				RspBias = Code->OpInfo ? 48 : 40;
+				Code += 1;
+				break;
+
+			case UWOP_SAVE_NONVOL:		Code += 2; break; /* saves are movs instead of pushes */
+			case UWOP_SAVE_NONVOL_FAR:	Code += 3; break;
+			case UWOP_SAVE_XMM128:		Code += 2; break;
+			case UWOP_SAVE_XMM128_FAR:	Code += 3; break;
+
+			default:
+#if BACKTRACE_DBGLVL >= 2
+				PLATFORM_BREAK();
+#endif
+				break;
+			}
+		}
+
+		// "Chained" simply means that multiple RUNTIME_FUNCTIONs pertains to a
+		// single actual function in the .text segment.
+		bool bIsChained = (UnwindInfo->Flags & UNW_FLAG_CHAININFO);
+
+		RspBias /= sizeof(void*);	// stack push/popds in units of one machine word
+		RspBias += !bIsChained;		// and one extra push for the ret address
+		RspBias |= FpInfo;			// pack in details about possible frame pointer
+
+		if (bIsChained)
+		{
+			OutTableCursor[-1].RspBias += RspBias;
+#if BACKTRACE_DBGLVL >= 2
+			OutTableCursor[-1].Size += (FunctionTable->EndAddress - FunctionTable->BeginAddress);
+#endif
+		}
+		else
+		{
+			*OutTableCursor = {
+				FunctionTable->BeginAddress,
+				RspBias,
+#if BACKTRACE_DBGLVL >= 2
+				FunctionTable->EndAddress - FunctionTable->BeginAddress,
+				UnwindInfo,
+#endif
+			};
+
+			++OutTableCursor;
+		}
+	}
+
+	UPTRINT ModuleSize = NtHeader->OptionalHeader.SizeOfImage;
+	ModuleSize += 0xffff; // to align up to next 64K page. it'll get shifted by AddressToId()
+
+	FModule Module = {
+		AddressToId(ModuleBase),
+		AddressToId(ModuleSize),
+		uint32(UPTRINT(OutTableCursor - OutTable)),
+#if BACKTRACE_DBGLVL >= 1
+		uint16(NumFpFuncs),
+#endif
+		OutTable,
+	};
+
+	TArrayView<FModule> ModulesView(Modules.Data, Modules.Num);
+	int32 Index = Algo::UpperBound(ModulesView, Module.Id, FIdPredicate());
+	{
+		FWriteScopeLock _(Lock);
+		Modules.Insert(Module, Index);
+	}
+
+#if BACKTRACE_DBGLVL >= 1
+	NumFpTruncations += NumFpFuncs;
+	TotalFunctions += NumFunctions;
+#endif
+}
+
+////////////////////////////////////////////////////////////////////////////////
+void FBacktracer::RemoveModule(UPTRINT ModuleBase)
+{
+	uint32 ModuleId = AddressToId(ModuleBase);
+	TArrayView<FModule> ModulesView(Modules.Data, Modules.Num);
+	int32 Index = Algo::LowerBound(ModulesView, ModuleId, FIdPredicate());
+	if (Index >= Modules.Num)
+	{
+		return;
+	}
+
+	const FModule& Module = Modules.Data[Index];
+	if (Module.Id != ModuleId)
+	{
+		return;
+	}
+
+#if BACKTRACE_DBGLVL >= 1
+	NumFpTruncations -= Module.NumFpTypes;
+	TotalFunctions -= Module.NumFunctions;
+#endif
+
+	// no code should be executing at this point so we can safely free the
+	// table knowing know one is looking at it.
+	Malloc->Free(Module.Functions);
+
+	{
+		FWriteScopeLock _(Lock);
+		Modules.RemoveAt(Index);
+	}
+}
+
+////////////////////////////////////////////////////////////////////////////////
+const FBacktracer::FFunction* FBacktracer::LookupFunction(UPTRINT Address, FLookupState& State) const
+{
+	// This function caches the previous module look up. The theory here is that
+	// a series of return address in a backtrace often cluster around one module
+
+	FIdPredicate IdPredicate;
+
+	// Look up the module that Address belongs to.
+	uint32 AddressId = AddressToId(Address);
+	if ((AddressId - State.Module.Id) >= State.Module.IdSize)
+	{
+		TArrayView<FModule> ModulesView(Modules.Data, Modules.Num);
+		uint32 Index = Algo::UpperBound(ModulesView, AddressId, IdPredicate);
+		if (Index == 0)
+		{
+			return nullptr;
+		}
+
+		State.Module = Modules.Data[Index - 1];
+	}
+
+	// Check that the address is within the address space of the best-found module
+	const FModule* Module = &(State.Module);
+	if ((AddressId - Module->Id) >= Module->IdSize)
+	{
+		return nullptr;
+	}
+
+	// Now we've a module we have a table of functions and their stack sizes so
+	// we can get the frame size for Address
+	uint32 FuncId = uint32(Address - IdToAddress(Module->Id)); 
+	TArrayView<FFunction> FuncsView(Module->Functions, Module->NumFunctions);
+	uint32 Index = Algo::UpperBound(FuncsView, FuncId, IdPredicate);
+
+	const FFunction* Function = Module->Functions + (Index - 1);
+#if BACKTRACE_DBGLVL >= 2
+	if ((FuncId - Function->Id) >= Function->Size)
+	{
+		PLATFORM_BREAK();
+		return nullptr;
+	}
+#endif
+	return Function;
+}
+
+////////////////////////////////////////////////////////////////////////////////
+void* FBacktracer::GetBacktraceId(void* AddressOfReturnAddress) const
+{
+	FLookupState LookupState;
+	FCallstackProcWorker::FBacktraceEntry BacktraceEntry;
+
+	UPTRINT* StackPointer = (UPTRINT*)AddressOfReturnAddress;
+
+#if BACKTRACE_DBGLVL >= 3
+	UPTRINT TruthBacktrace[1024];
+	uint32 NumTruth = RtlCaptureStackBackTrace(0, 1024, (void**)TruthBacktrace, nullptr);
+	UPTRINT* TruthCursor = TruthBacktrace;
+	for (; *TruthCursor != *StackPointer; ++TruthCursor);
+#endif
+
+#if BACKTRACE_DBGLVL >= 2
+	struct { void* Sp; void* Ip; const FFunction* Function; } Backtrace[1024] = {};
+	uint32 NumBacktrace = 0;
+#endif
+
+	uint64 BacktraceId = 0;
+	uint32 FrameIdx = 0;
+
+	FReadScopeLock _(Lock);
+	do
+	{
+		UPTRINT RetAddr = *StackPointer;
+		
+		BacktraceEntry.Frames[FrameIdx++] = RetAddr;
+		BacktraceEntry.FrameCount = FrameIdx;
+
+		// This is a simple order-dependent LCG. Should be sufficient enough
+		BacktraceId += RetAddr;
+		BacktraceId *= 0x30be8efa499c249dull;
+
+		const FFunction* Function = LookupFunction(RetAddr, LookupState);
+		if (Function == nullptr)
+		{
+			break;
+		}
+
+#if BACKTRACE_DBGLVL >= 2
+		if (NumBacktrace < 1024)
+		{
+			Backtrace[NumBacktrace++] = {
+				StackPointer,
+				(void*)RetAddr,
+				Function,
+			};
+		}
+#endif
+
+		if (Function->RspBias < 0)
+		{
+			// This is a frame with a variable-sized stack pointer. We don't
+			// track enough information to proceed.
+#if BACKTRACE_DBGLVL >= 1
+			NumFpTruncations++;
+#endif
+			break;
+		}
+
+		StackPointer += Function->RspBias;
+	}
+	// Trunkate callstacks longer than MaxStackDepth
+	while (*StackPointer && FrameIdx < FCallstackProcWorker::FBacktraceEntry::MaxStackDepth);
+
+	// Save the collected id
+	BacktraceEntry.Id = BacktraceId;
+
+#if BACKTRACE_DBGLVL >= 3
+	for (uint32 i = 0; i < NumBacktrace; ++i)
+	{
+		if ((void*)TruthCursor[i] != Backtrace[i].Ip)
+		{
+			PLATFORM_BREAK();
+			break;
+		}
+	}
+#endif
+
+	// Add to queue to be processed. This might block until there is room in the
+	// queue (i.e. the processing thread has caught up processing).
+	ProcessingThreadRunnable->AddWork(BacktraceEntry);
+	return (void*)(BacktraceId & ((1ull << 47) - 1));
+}
+
+////////////////////////////////////////////////////////////////////////////////
+int32 FBacktracer::GetFrameSize(void* FunctionAddress) const
+{
+	FReadScopeLock _(Lock);
+
+	FLookupState LookupState;
+	const FFunction* Function = LookupFunction(UPTRINT(FunctionAddress), LookupState);
+	if (FunctionAddress == nullptr)
+	{
+		return -1;
+	}
+
+	if (Function->RspBias < 0)
+	{
+		return -1;
+	}
+
+	return Function->RspBias;
+}
+
+////////////////////////////////////////////////////////////////////////////////
+UE_TRACE_CHANNEL(CallstackChannel);
+
+UE_TRACE_EVENT_BEGIN(Memory, CallstackSpec, NoSync)
+	UE_TRACE_EVENT_FIELD(uint64, Id)
+	UE_TRACE_EVENT_FIELD(uint64[], Frames)
+UE_TRACE_EVENT_END()
+
+////////////////////////////////////////////////////////////////////////////////
+uint32 FCallstackProcWorker::Run()
+{
+	while (bRun)
+	{
+		TRACE_CPUPROFILER_EVENT_SCOPE("ProfilerThreadWork")
+		while (!Queue.was_empty()) 
+		{
+			FBacktraceEntry Entry = Queue.pop();
+			AddCallstack(Entry);
+		}
+ 		FPlatformProcess::Yield();
+	}
+	return 0;
+}
+
+////////////////////////////////////////////////////////////////////////////////
+void FCallstackProcWorker::AddCallstack(const FBacktraceEntry& Entry)
+{
+	bool bAlreadyAdded = false;
+	KnownSet.Add(Entry.Id, &bAlreadyAdded);
+	if (!bAlreadyAdded)
+	{
+		UE_TRACE_LOG(Memory, CallstackSpec, CallstackChannel)
+			<< CallstackSpec.Id(Entry.Id)
+			<< CallstackSpec.Frames(Entry.Frames, Entry.FrameCount);
+	}
+}
+
+////////////////////////////////////////////////////////////////////////////////
+void FCallstackProcWorker::AddWork(const FBacktraceEntry& Entry)
+{
+	// The queue only supports single producer, single consumer. So we use a lock
+	// on the producer side.
+	FScopeLock _(&ProducerCs);
+
+	if (Thread) 
+	{
+		while(!Queue.try_push(Entry))
+		{
+			FPlatformProcess::Yield();
+		}
+	}
+	else
+	{
+		// Worker thread has not yet been started, manually process callstack for now.
+		AddCallstack(Entry);
+	}
+}
+
+////////////////////////////////////////////////////////////////////////////////
+void Modules_Create(FMalloc*);
+void Modules_Subscribe(void (*)(bool, void*, const TCHAR*));
+void Modules_Initialize();
+
+////////////////////////////////////////////////////////////////////////////////
+void Backtracer_Create(FMalloc* Malloc)
+{
+	if (FBacktracer::Get() != nullptr)
+	{
+		return;
+	}
+
+	static FBacktracer Instance(Malloc);
+
+	Modules_Create(Malloc);
+	Modules_Subscribe(
+		[] (bool bLoad, void* Module, const TCHAR* Name)
+		{
+			bLoad
+				? Instance.AddModule(UPTRINT(Module), Name)
+				: Instance.RemoveModule(UPTRINT(Module));
+		}
+	);
+}
+
+////////////////////////////////////////////////////////////////////////////////
+void Backtracer_Initialize()
+{
+	Modules_Initialize();
+}
+
+////////////////////////////////////////////////////////////////////////////////
+int32 Backtracer_GetFrameSize(void* FunctionAddress)
+{
+	if (FBacktracer* Instance = FBacktracer::Get())
+	{
+		return Instance->GetFrameSize(FunctionAddress);
+	}
+
+	return -1;
+}
+
+////////////////////////////////////////////////////////////////////////////////
+void* Backtracer_GetBacktraceId(void* AddressOfReturnAddress)
+{
+	if (FBacktracer* Instance = FBacktracer::Get())
+	{
+		return Instance->GetBacktraceId(AddressOfReturnAddress);
+	}
+
+	return nullptr;
+}
+
diff --git a/Engine/Source/Runtime/Core/Private/ProfilingDebugging/Windows/MemoryTrace.cpp b/Engine/Source/Runtime/Core/Private/ProfilingDebugging/Windows/MemoryTrace.cpp
new file mode 100644
--- /dev/null
+++ b/Engine/Source/Runtime/Core/Private/ProfilingDebugging/Windows/MemoryTrace.cpp
@@ -0,0 +1,792 @@
+// Copyright Epic Games, Inc. All Rights Reserved.
+
+#include "ProfilingDebugging/MemoryTrace.h"
+
+#if UE_MEMORY_TRACE_ENABLED
+
+#include "Containers/StringView.h"
+#include "CoreTypes.h"
+#include "HAL/MemoryBase.h"
+#include "HAL/Platform.h"
+#include "Misc/CoreDelegates.h"
+#include "Misc/CString.h"
+#include "Trace/Trace.inl"
+
+#include "Windows/AllowWindowsPlatformTypes.h"
+#include <winternl.h>
+
+#include <atomic>
+#include <intrin.h>
+
+PRAGMA_DISABLE_OPTIMIZATION
+
+////////////////////////////////////////////////////////////////////////////////
+void	Backtracer_Create(FMalloc*);
+void	Backtracer_Initialize();
+void*	Backtracer_GetBacktraceId(void*);
+
+////////////////////////////////////////////////////////////////////////////////
+#if defined(_MSC_VER)
+#	define UE_RETURN_ADDRESS_ADDRESS()	(void*)_AddressOfReturnAddress()
+#else if defined(__clang__) || defined(__GNUC__)
+#	define UE_RETURN_ADDRESS_ADDRESS()
+#	error TODO
+#endif
+
+FORCEINLINE void* GetOwner(bool bLight)
+{
+	void* RetAddrAddr = UE_RETURN_ADDRESS_ADDRESS();
+	if (bLight)
+	{
+		return *(void**)RetAddrAddr;
+	}
+
+	return Backtracer_GetBacktraceId(RetAddrAddr);
+}
+
+////////////////////////////////////////////////////////////////////////////////
+template <class T>
+class alignas(alignof(T)) FUndestructed
+{
+public:
+	template <typename... ArgTypes>
+	void Construct(ArgTypes... Args)
+	{
+		::new (Buffer) T(Args...);
+	}
+
+	T* operator & ()	{ return (T*)Buffer; }
+	T* operator -> ()	{ return (T*)Buffer; }
+
+protected:
+	uint8 Buffer[sizeof(T)];
+};
+
+////////////////////////////////////////////////////////////////////////////////
+struct FAddrPack
+{
+			FAddrPack() = default;
+			FAddrPack(UPTRINT Addr, uint16 Value) { Set(Addr, Value); }
+	void	Set(UPTRINT Addr, uint16 Value) { Inner = uint64(Addr) | (uint64(Value) << 48ull); }
+	uint64	Inner;
+};
+static_assert(sizeof(FAddrPack) == sizeof(uint64), "");
+
+
+
+////////////////////////////////////////////////////////////////////////////////
+UE_TRACE_CHANNEL(MemSummaryChannel)
+UE_TRACE_CHANNEL_DEFINE(MemAllocChannel)
+
+UE_TRACE_EVENT_BEGIN(Memory, Init)
+	UE_TRACE_EVENT_FIELD(uint8, MinAlignment)
+	UE_TRACE_EVENT_FIELD(uint8, SizeShift)
+	UE_TRACE_EVENT_FIELD(uint8, SummarySizeShift)
+	UE_TRACE_EVENT_FIELD(uint8, Mode)
+UE_TRACE_EVENT_END()
+
+UE_TRACE_EVENT_BEGIN(Memory, Summary)
+	UE_TRACE_EVENT_FIELD(uint32, Bytes)
+	UE_TRACE_EVENT_FIELD(uint32, ActiveAllocs)
+	UE_TRACE_EVENT_FIELD(uint32, TotalAllocs)
+	UE_TRACE_EVENT_FIELD(uint32, TotalReallocs)
+	UE_TRACE_EVENT_FIELD(uint32, TotalFrees)
+UE_TRACE_EVENT_END()
+
+UE_TRACE_EVENT_BEGIN(Memory, CoreAdd)
+	UE_TRACE_EVENT_FIELD(void*, Owner)
+	UE_TRACE_EVENT_FIELD(void*, Base)
+	UE_TRACE_EVENT_FIELD(uint32, Size)
+UE_TRACE_EVENT_END()
+
+UE_TRACE_EVENT_BEGIN(Memory, CoreRemove)
+	UE_TRACE_EVENT_FIELD(void*, Owner)
+	UE_TRACE_EVENT_FIELD(void*, Base)
+	UE_TRACE_EVENT_FIELD(uint32, Size)
+UE_TRACE_EVENT_END()
+
+UE_TRACE_EVENT_BEGIN(Memory, Alloc)
+	UE_TRACE_EVENT_FIELD(void*, Owner)
+	UE_TRACE_EVENT_FIELD(void*, Address)
+	UE_TRACE_EVENT_FIELD(uint32, Size)
+	UE_TRACE_EVENT_FIELD(uint8, Alignment)
+	UE_TRACE_EVENT_FIELD(uint8, Waste)
+UE_TRACE_EVENT_END()
+
+UE_TRACE_EVENT_BEGIN(Memory, Realloc)
+	UE_TRACE_EVENT_FIELD(void*, Owner)
+	UE_TRACE_EVENT_FIELD(void*, FreeAddress)	// compatible with both Alloc
+	UE_TRACE_EVENT_FIELD(void*, Address)		// and Free events.
+	UE_TRACE_EVENT_FIELD(uint32, Size)
+	UE_TRACE_EVENT_FIELD(uint8, Alignment)
+	UE_TRACE_EVENT_FIELD(uint8, Waste)
+UE_TRACE_EVENT_END()
+
+UE_TRACE_EVENT_BEGIN(Memory, Free)
+	UE_TRACE_EVENT_FIELD(void*, Owner)
+	UE_TRACE_EVENT_FIELD(void*, FreeAddress)
+UE_TRACE_EVENT_END()
+
+
+
+#define USE_OVERVIEW_TRACE 0 // Disabled for now, as it occasionally conflicts with MallocBinned2's canary poison logic.
+#if USE_OVERVIEW_TRACE
+
+////////////////////////////////////////////////////////////////////////////////
+class FSummaryTrace
+{
+public:
+	static const uint32			SizeShift = 4;
+	void						Initialize();
+	uint32						SetTag(uint32 Tag);
+	uint32						Alloc(SIZE_T Size, bool IsRealloc);
+	void						Free(uint32 Tag, SIZE_T Size);
+	void						Flush();
+
+private:
+	template <typename T>
+	struct FItem
+	{
+		std::atomic<T>			Current = 0;
+		T						Max = 0;
+		T						Min = 0;
+		void					Add(T Operand);
+	};
+
+	struct alignas(PLATFORM_CACHE_LINE_SIZE) FInfo
+	{
+		uint32					Tag;
+		FItem<int64>			Bytes;
+		FItem<int32>			ActiveAllocs;
+		std::atomic<uint32>		TotalAllocs = 0;
+		std::atomic<uint32>		TotalReallocs = 0;
+		std::atomic<uint32>		TotalFrees = 0;
+	};
+
+	FInfo&						GetInfo();
+	static thread_local FInfo*	InfoPtr;
+	std::atomic<uint32>			InfoIndex;
+	FInfo						InfoPool[1];
+};
+
+////////////////////////////////////////////////////////////////////////////////
+static FUndestructed<FSummaryTrace>	GSummaryTrace;
+thread_local FSummaryTrace::FInfo*	FSummaryTrace::InfoPtr = nullptr;
+
+////////////////////////////////////////////////////////////////////////////////
+template <typename T>
+void FSummaryTrace::FItem<T>::Add(T Size)
+{
+	T Value = Current.fetch_add(T(Size), std::memory_order_relaxed);
+	Value += Size;
+	Max = (Current > Max) ? Value : Max;
+	Min = (Current < Min) ? Value : Min;
+}
+
+////////////////////////////////////////////////////////////////////////////////
+void FSummaryTrace::Initialize()
+{
+	FCoreDelegates::OnEndFrame.AddStatic([] () { GSummaryTrace->Flush(); });
+}
+
+////////////////////////////////////////////////////////////////////////////////
+uint32 FSummaryTrace::SetTag(uint32 Tag)
+{
+	FInfo& Info = GetInfo();
+	uint32 PrevTag = Info.Tag;
+	Info.Tag = Tag;
+	return PrevTag;
+}
+
+////////////////////////////////////////////////////////////////////////////////
+uint32 FSummaryTrace::Alloc(SIZE_T Size, bool IsRealloc)
+{
+	FInfo& Info = GetInfo();
+	Info.Bytes.Add(Size);
+	Info.ActiveAllocs.Add(1);
+	Info.TotalAllocs.fetch_add(1, std::memory_order_relaxed);
+	if (IsRealloc)
+	{
+		Info.TotalReallocs.fetch_add(1, std::memory_order_relaxed);
+	}
+	return Info.Tag;
+}
+
+////////////////////////////////////////////////////////////////////////////////
+void FSummaryTrace::Free(uint32 Tag, SIZE_T Size)
+{
+	FInfo& Info = GetInfo();
+	Info.Bytes.Add(-int64(Size));
+	Info.ActiveAllocs.Add(-1);
+	Info.TotalFrees.fetch_add(1, std::memory_order_relaxed);
+}
+
+////////////////////////////////////////////////////////////////////////////////
+FSummaryTrace::FInfo& FSummaryTrace::GetInfo()
+{
+	FInfo* Info = InfoPtr;
+	if (Info == nullptr)
+	{
+		uint32 Index = InfoIndex.fetch_add(1, std::memory_order_relaxed);
+		Index &= UE_ARRAY_COUNT(InfoPool) - 1;
+		Info = InfoPtr = InfoPool + Index;
+	}
+	return *Info;
+}
+
+////////////////////////////////////////////////////////////////////////////////
+void FSummaryTrace::Flush()
+{
+	int64 Bytes = 0;
+	int32 Active = 0;
+	int32 TotalAllocs = 0;
+	int32 TotalReallocs = 0;
+	int32 TotalFrees = 0;
+	for (const FInfo& Info : InfoPool)
+	{
+		Bytes += Info.Bytes.Current.load(std::memory_order_relaxed);
+		Active += Info.ActiveAllocs.Current.load(std::memory_order_relaxed);
+		TotalAllocs += Info.TotalAllocs.load(std::memory_order_relaxed);
+		TotalReallocs += Info.TotalReallocs.load(std::memory_order_relaxed);
+		TotalFrees += Info.TotalFrees.load(std::memory_order_relaxed);
+	}
+	
+	UE_TRACE_LOG(Memory, Summary, MemSummaryChannel)
+		<< Summary.Bytes(uint32(Bytes >> SizeShift))
+		<< Summary.ActiveAllocs(Active)
+		<< Summary.TotalAllocs(TotalAllocs)
+		<< Summary.TotalReallocs(TotalReallocs)
+		<< Summary.TotalFrees(TotalFrees);
+}
+
+#endif // USE_OVERVIEW_TRACE
+
+
+
+////////////////////////////////////////////////////////////////////////////////
+class FAllocationTrace
+{
+public:
+	void	Initialize();
+	void	CoreAdd(void* Base, size_t Size, void* Owner);
+	void	CoreRemove(void* Base, size_t Size, void* Owner);
+	void	Alloc(void* Address, size_t Size, uint32 Alignment, void* Owner);
+	void	Realloc(void* PrevAddress, void* Address, size_t NewSize, uint32 Alignment, void* Owner);
+	void	Free(void* Address, void* Owner);
+
+private:
+	static const uint32 SizeShift = 3;
+	static_assert(MIN_ALIGNMENT >= (1 << SizeShift), "");
+};
+
+////////////////////////////////////////////////////////////////////////////////
+static FUndestructed<FAllocationTrace> GAllocationTrace;
+
+////////////////////////////////////////////////////////////////////////////////
+void FAllocationTrace::Initialize()
+{
+	UE_TRACE_LOG(Memory, Init, MemAllocChannel)
+		<< Init.MinAlignment(uint8(MIN_ALIGNMENT))
+#if USE_OVERVIEW_TRACE
+		<< Init.SummarySizeShift(uint8(FSummaryTrace::SizeShift))
+#endif
+		<< Init.SizeShift(uint8(SizeShift));
+}
+
+////////////////////////////////////////////////////////////////////////////////
+void FAllocationTrace::CoreAdd(void* Base, size_t Size, void* Owner)
+{
+	UE_TRACE_LOG(Memory, CoreAdd, MemAllocChannel)
+		<< CoreAdd.Base(Base)
+		<< CoreAdd.Size(uint32(Size >> SizeShift));
+}
+
+////////////////////////////////////////////////////////////////////////////////
+void FAllocationTrace::CoreRemove(void* Base, size_t Size, void* Owner)
+{
+	UE_TRACE_LOG(Memory, CoreRemove, MemAllocChannel)
+		<< CoreRemove.Base(Base)
+		<< CoreRemove.Size(uint32(Size >> SizeShift));
+}
+
+////////////////////////////////////////////////////////////////////////////////
+void FAllocationTrace::Alloc(void* Address, size_t Size, uint32 Alignment, void* Owner)
+{
+	uint32 ActualAlignment = Alignment > uint32(MIN_ALIGNMENT) ? Alignment : uint32(MIN_ALIGNMENT);
+	uint32 Waste = ActualAlignment - (Size & (ActualAlignment - 1));
+
+	UE_TRACE_LOG(Memory, Alloc, MemAllocChannel)
+		<< Alloc.Address(Address)
+		<< Alloc.Size(uint32(Size >> 3))
+		<< Alloc.Alignment(uint8(Alignment))
+		<< Alloc.Waste(uint8(Waste));
+}
+
+////////////////////////////////////////////////////////////////////////////////
+void FAllocationTrace::Realloc(void* PrevAddress, void* Address, size_t NewSize, uint32 Alignment, void* Owner)
+{
+	uint32 ActualAlignment = Alignment > uint32(MIN_ALIGNMENT) ? Alignment : uint32(MIN_ALIGNMENT);
+	uint32 Waste = ActualAlignment - (NewSize & (ActualAlignment - 1));
+
+	UE_TRACE_LOG(Memory, Realloc, MemAllocChannel)
+		<< Realloc.FreeAddress(PrevAddress)
+		<< Realloc.Address(Address)
+		<< Realloc.Size(uint32(NewSize >> 3))
+		<< Realloc.Alignment(uint8(Alignment))
+		<< Realloc.Waste(uint8(Waste));
+}
+
+////////////////////////////////////////////////////////////////////////////////
+void FAllocationTrace::Free(void* Address, void* Owner)
+{
+	UE_TRACE_LOG(Memory, Free, MemAllocChannel)
+		<< Free.FreeAddress(Address);
+}
+
+
+
+////////////////////////////////////////////////////////////////////////////////
+class FMallocWrapper
+	: public FMalloc
+{
+public:
+							FMallocWrapper(FMalloc* InMalloc, bool bInLight);
+
+private:
+	struct FCookie
+	{
+		uint64				Tag  : 16;
+		uint64				Bias : 8;
+		uint64				Size : 40;
+	};
+
+	static uint32			GetActualAlignment(SIZE_T Size, uint32 Alignment);
+	virtual void*			Malloc(SIZE_T Size, uint32 Alignment) override;
+	virtual void*			Realloc(void* PrevAddress, SIZE_T NewSize, uint32 Alignment) override;
+	virtual void			Free(void* Address) override;
+	virtual bool			IsInternallyThreadSafe() const override						{ return InnerMalloc->IsInternallyThreadSafe(); }
+	virtual void			UpdateStats() override										{ InnerMalloc->UpdateStats(); }
+	virtual void			GetAllocatorStats(FGenericMemoryStats& Out) override		{ InnerMalloc->GetAllocatorStats(Out); }
+	virtual void			DumpAllocatorStats(FOutputDevice& Ar) override				{ InnerMalloc->DumpAllocatorStats(Ar); }
+	virtual bool			ValidateHeap() override										{ return InnerMalloc->ValidateHeap(); }
+	virtual bool			GetAllocationSize(void* Address, SIZE_T &SizeOut) override	{ return InnerMalloc->GetAllocationSize(Address, SizeOut); }
+	FMalloc*				InnerMalloc;
+	bool					bLight;
+};
+
+////////////////////////////////////////////////////////////////////////////////
+FMallocWrapper::FMallocWrapper(FMalloc* InMalloc, bool bInLight)
+: InnerMalloc(InMalloc)
+, bLight(bInLight)
+{
+}
+
+////////////////////////////////////////////////////////////////////////////////
+uint32 FMallocWrapper::GetActualAlignment(SIZE_T Size, uint32 Alignment)
+{
+	// Defaults; if size is < 16 then alignment is 8 else 16.
+	uint32 DefaultAlignment = 8 << uint32(Size >= 16);
+	return (Alignment < DefaultAlignment) ? DefaultAlignment : Alignment;
+}
+
+////////////////////////////////////////////////////////////////////////////////
+void* FMallocWrapper::Malloc(SIZE_T Size, uint32 Alignment)
+{
+	if (Size == 0)
+	{
+		return nullptr;
+	}
+
+	uint32 ActualAlignment = GetActualAlignment(Size, Alignment);
+
+	SIZE_T InnerSize = Size;
+#if USE_OVERVIEW_TRACE
+	static_assert(sizeof(FCookie) <= 8, "Unexpected FCookie size");
+	InnerSize += ActualAlignment;
+#endif 
+
+	void* Address = InnerMalloc->Malloc(InnerSize, Alignment);
+
+#if USE_OVERVIEW_TRACE
+	Address = (uint8*)Address + ActualAlignment;
+	FCookie* Cookie = (FCookie*)Address - 1;
+	Cookie->Size = Size;
+	Cookie->Tag = GSummaryTrace->Alloc(Size, false);
+	Cookie->Bias = ActualAlignment >> 3;
+#endif
+
+	void* Owner = GetOwner(bLight);
+	GAllocationTrace->Alloc(Address, Size, ActualAlignment, Owner);
+
+	return Address;
+}
+
+////////////////////////////////////////////////////////////////////////////////
+void* FMallocWrapper::Realloc(void* PrevAddress, SIZE_T NewSize, uint32 Alignment)
+{
+	// This simplifies things and means reallocs trace events are true reallocs
+	if (PrevAddress == nullptr)
+	{
+		return Malloc(NewSize, Alignment);
+	}
+
+	if (NewSize == 0)
+	{
+		Free(PrevAddress);
+		return nullptr;
+	}
+
+	uint32 HeaderSize = 0;
+	void* InnerAddress = PrevAddress;
+#if USE_OVERVIEW_TRACE
+	FCookie* Cookie = (FCookie*)PrevAddress - 1;
+	HeaderSize = uint32(Cookie->Bias << 3);
+	InnerAddress = (uint8*)PrevAddress - HeaderSize;
+	GSummaryTrace->Free(Cookie->Tag, Cookie->Size);
+#endif
+
+	SIZE_T InnerSize = NewSize + HeaderSize;
+	void* RetAddress = InnerMalloc->Realloc(InnerAddress, InnerSize, Alignment);
+
+#if USE_OVERVIEW_TRACE
+	RetAddress = (uint8*)RetAddress + HeaderSize;
+	Cookie = (FCookie*)RetAddress - 1;
+	Cookie->Size = NewSize;
+	GSummaryTrace->Alloc(NewSize, true);
+#endif
+
+	void* Owner = GetOwner(bLight);
+	Alignment = GetActualAlignment(NewSize, Alignment);
+	GAllocationTrace->Realloc(PrevAddress, RetAddress, NewSize, Alignment, Owner);
+
+	return RetAddress;
+}
+
+////////////////////////////////////////////////////////////////////////////////
+void FMallocWrapper::Free(void* Address)
+{
+	if (Address == nullptr)
+	{
+		return;
+	}
+
+	void* Owner = GetOwner(bLight);
+	GAllocationTrace->Free(Address, Owner);
+
+	void* InnerAddress = Address;
+#if USE_OVERVIEW_TRACE
+	const FCookie* Cookie = (const FCookie*)Address - 1;
+	InnerAddress = (uint8*)Address - (Cookie->Bias << 3);
+	GSummaryTrace->Free(Cookie->Tag, Cookie->Size);
+#endif
+
+	return InnerMalloc->Free(InnerAddress);
+}
+
+
+
+////////////////////////////////////////////////////////////////////////////////
+class FTextSectionEditor
+{
+public:
+							FTextSectionEditor(void* InBase);
+							~FTextSectionEditor();
+	template <typename T>
+	T*						Hook(T* Target, T* HookFunction);
+
+private:
+	static void*			GetActualAddress(void* Function);
+	uint8*					AllocateTrampoline(unsigned int Size);
+	void*					HookImpl(void* Target, void* HookFunction);
+	uint8*					TrampolineTail;
+	void*					Base;
+	size_t					Size;
+	DWORD					Protection;
+};
+
+////////////////////////////////////////////////////////////////////////////////
+FTextSectionEditor::FTextSectionEditor(void* InBase)
+{
+	InBase = GetActualAddress(InBase);
+
+	MEMORY_BASIC_INFORMATION MemInfo;
+	VirtualQuery(InBase, &MemInfo, sizeof(MemInfo));
+	Base = MemInfo.BaseAddress;
+	Size = MemInfo.RegionSize;
+
+	VirtualProtect(Base, Size, PAGE_EXECUTE_READWRITE, &Protection);
+
+	TrampolineTail = (uint8*)Base + Size;
+}
+
+////////////////////////////////////////////////////////////////////////////////
+FTextSectionEditor::~FTextSectionEditor()
+{
+	VirtualProtect(Base, Size, Protection, &Protection);
+	FlushInstructionCache(GetCurrentProcess(), Base, Size);
+}
+
+////////////////////////////////////////////////////////////////////////////////
+void* FTextSectionEditor::GetActualAddress(void* Function)
+{
+	uint8* Addr = (uint8*)Function;
+	int Offset = unsigned(Addr[0] & 0xf0) == 0x40; // REX prefix
+	if (Addr[Offset + 0] == 0xff && Addr[Offset + 1] == 0x25)
+	{
+		Addr += Offset;
+		Addr = *(uint8**)(Addr + 6 + *(uint32*)(Addr + 2));
+	}
+	return Addr;
+}
+
+////////////////////////////////////////////////////////////////////////////////
+uint8* FTextSectionEditor::AllocateTrampoline(unsigned int PatchSize)
+{
+	static const int TrampolineSize = 24;
+	uint8* NextTail = TrampolineTail - TrampolineSize;
+	do 
+	{
+		check(*NextTail == 0);
+		++NextTail;
+	}
+	while (NextTail != TrampolineTail);
+
+	TrampolineTail -= TrampolineSize;
+	return TrampolineTail;
+}
+
+////////////////////////////////////////////////////////////////////////////////
+template <typename T>
+T* FTextSectionEditor::Hook(T* Target, T* HookFunction)
+{
+	return (T*)HookImpl((void*)Target, (void*)HookFunction);
+}
+
+////////////////////////////////////////////////////////////////////////////////
+void* FTextSectionEditor::HookImpl(void* Target, void* HookFunction)
+{
+	Target = GetActualAddress(Target);
+
+	const uint8* Start = (uint8*)Target;
+	const uint8* Read = Start;
+	do
+	{
+		Read += (Read[0] & 0xf0) == 0x40; // REX prefix
+		uint8 Inst = *Read++;
+		if (unsigned(Inst - 0x80) < 0x0cu)
+		{
+			uint8 ModRm = *Read++;
+			Read += ((ModRm & 0300) < 0300) & ((ModRm & 0007) == 0004); // SIB
+			switch (ModRm & 0300) // Disp[8|32]
+			{
+			case 0100: Read += 1; break;
+			case 0200: Read += 5; break;
+			}
+			Read += (Inst == 0x83);
+		}
+		else if (unsigned(Inst - 0x50) >= 0x10u)
+			check(!"Unknown instruction");
+	}
+	while (Read - Start < 6);
+
+	int PatchSize = int(Read - Start);
+	uint8* TrampolinePtr = AllocateTrampoline(PatchSize);
+
+	*(void**)TrampolinePtr = HookFunction;
+
+	uint8* PatchJmp = TrampolinePtr + sizeof(void*);
+	memcpy(PatchJmp, Start, PatchSize);
+
+	PatchJmp += PatchSize;
+	*PatchJmp = 0xe9;
+	*(int32*)(PatchJmp + 1) = int32(intptr_t(Start + PatchSize) - intptr_t(PatchJmp)) - 5;
+
+	uint16* HookJmp = (uint16*)Target;
+	HookJmp[0] = 0x25ff;
+	*(int32*)(HookJmp + 1) = int32(intptr_t(TrampolinePtr) - intptr_t(HookJmp + 3));
+
+	return PatchJmp - PatchSize;
+}
+
+
+
+////////////////////////////////////////////////////////////////////////////////
+class FVirtualWinApiHooks
+{
+public:
+	static void				Initialize(bool bInLight);
+
+private:
+							FVirtualWinApiHooks();
+	static bool				bLight;
+	static LPVOID WINAPI	VmAlloc(LPVOID Address, SIZE_T Size, DWORD Type, DWORD Protect);
+	static BOOL WINAPI		VmFree(LPVOID Address, SIZE_T Size, DWORD Type);
+	static BOOL WINAPI		VmFreeEx(HANDLE Process, LPVOID Address, SIZE_T Size, DWORD Type);
+	static LPVOID WINAPI	VmAllocEx(HANDLE Process, LPVOID Address, SIZE_T Size, DWORD Type, DWORD Protect);
+	static LPVOID			(WINAPI *VmAllocOrig)(LPVOID, SIZE_T, DWORD, DWORD);
+	static LPVOID			(WINAPI *VmAllocExOrig)(HANDLE, LPVOID, SIZE_T, DWORD, DWORD);
+	static BOOL				(WINAPI *VmFreeOrig)(LPVOID, SIZE_T, DWORD);
+	static BOOL				(WINAPI *VmFreeExOrig)(HANDLE, LPVOID, SIZE_T, DWORD);
+};
+
+////////////////////////////////////////////////////////////////////////////////
+bool	FVirtualWinApiHooks::bLight;
+LPVOID	(WINAPI *FVirtualWinApiHooks::VmAllocOrig)(LPVOID, SIZE_T, DWORD, DWORD);
+LPVOID	(WINAPI *FVirtualWinApiHooks::VmAllocExOrig)(HANDLE, LPVOID, SIZE_T, DWORD, DWORD);
+BOOL	(WINAPI *FVirtualWinApiHooks::VmFreeOrig)(LPVOID, SIZE_T, DWORD);
+BOOL	(WINAPI *FVirtualWinApiHooks::VmFreeExOrig)(HANDLE, LPVOID, SIZE_T, DWORD);
+
+////////////////////////////////////////////////////////////////////////////////
+void FVirtualWinApiHooks::Initialize(bool bInLight)
+{
+	bLight = bInLight;
+
+	{
+		FTextSectionEditor Editor((void*)VirtualAlloc);
+		VmAllocOrig = Editor.Hook(VirtualAlloc, &FVirtualWinApiHooks::VmAlloc);
+		VmFreeOrig = Editor.Hook(VirtualFree, &FVirtualWinApiHooks::VmFree); 
+		
+		{
+			FTextSectionEditor EditorEx((void*)VirtualAllocEx);
+			VmAllocExOrig = Editor.Hook(VirtualAllocEx, &FVirtualWinApiHooks::VmAllocEx);
+			VmFreeExOrig = Editor.Hook(VirtualFreeEx, &FVirtualWinApiHooks::VmFreeEx);
+		}
+	}
+
+	
+}
+
+////////////////////////////////////////////////////////////////////////////////
+LPVOID WINAPI FVirtualWinApiHooks::VmAlloc(LPVOID Address, SIZE_T Size, DWORD Type, DWORD Protect)
+{
+	LPVOID Ret = VmAllocOrig(Address, Size, Type, Protect);
+	if (Ret != nullptr && (Type & MEM_COMMIT))
+	{
+		void* Owner = GetOwner(bLight);
+		GAllocationTrace->CoreAdd(Ret, Size, Owner);
+	}
+
+	return Ret;
+}
+
+////////////////////////////////////////////////////////////////////////////////
+BOOL WINAPI FVirtualWinApiHooks::VmFree(LPVOID Address, SIZE_T Size, DWORD Type)
+{
+	void* Owner = GetOwner(bLight);
+	GAllocationTrace->CoreRemove(Address, Size, Owner);
+	return VmFreeOrig(Address, Size, Type);
+}
+
+////////////////////////////////////////////////////////////////////////////////
+LPVOID WINAPI FVirtualWinApiHooks::VmAllocEx(HANDLE Process, LPVOID Address, SIZE_T Size, DWORD Type, DWORD Protect)
+{
+	LPVOID Ret = VmAllocExOrig(Process, Address, Size, Type, Protect);
+	if (Process == GetCurrentProcess() && Ret != nullptr && (Type & MEM_COMMIT))
+	{
+		void* Owner = GetOwner(bLight);
+		GAllocationTrace->CoreAdd(Ret, Size, Owner);
+	}
+
+	return Ret;
+}
+
+////////////////////////////////////////////////////////////////////////////////
+BOOL WINAPI FVirtualWinApiHooks::VmFreeEx(HANDLE Process, LPVOID Address, SIZE_T Size, DWORD Type)
+{
+	if (Process == GetCurrentProcess())
+	{
+		void* Owner = GetOwner(bLight);
+		GAllocationTrace->CoreRemove(Address, Size, Owner);
+	}
+
+	return VmFreeExOrig(Process, Address, Size, Type);
+}
+
+
+
+////////////////////////////////////////////////////////////////////////////////
+FMalloc* MemoryTrace_Create(FMalloc* InMalloc)
+{
+	int32 Mode = 0;
+	const TCHAR* CmdLine = ::GetCommandLineW();
+	if (const TCHAR* TraceArg = FCString::Strstr(CmdLine, TEXT("-trace=")))
+	{
+		const TCHAR* Start = TraceArg + 7;
+		const TCHAR* End = Start;
+		for (; *End != ' ' && *End != '\0'; ++End);
+
+		for (const TCHAR* c = Start; c < End + 1; ++c)
+		{
+			if (*c == ' ' || *c == '\0' || *c == ',')
+			{
+				FStringView View(Start, uint32(c - Start));
+				if (View.Equals(TEXT("memalloc")))
+				{
+					Mode = 2;
+					break;
+				}
+
+				Start = c;
+			}
+		}
+	}
+
+	if (Mode > 0)
+	{
+		GAllocationTrace->Initialize();
+#if USE_OVERVIEW_TRACE
+		GSummaryTrace->Initialize();
+#endif
+
+		bool bLight = (Mode == 1);
+
+		if (!bLight)
+		{
+			Backtracer_Create(InMalloc);
+		}
+
+		FVirtualWinApiHooks::Initialize(bLight);
+
+		static FUndestructed<FMallocWrapper> MemoryTrace;
+		MemoryTrace.Construct(InMalloc, bLight);
+
+		GAllocationTrace.Construct();
+#if USE_OVERVIEW_TRACE
+		GSummaryTrace.Construct();
+#endif
+
+		return &MemoryTrace;
+	}
+
+	return InMalloc;
+}
+
+////////////////////////////////////////////////////////////////////////////////
+void MemoryTrace_Initialize()
+{
+	// Allocators aren't completely ready in _Create() so we have an extra step
+	// where any initialisation that may allocate can go.
+	Backtracer_Initialize();
+}
+
+#if 0
+////////////////////////////////////////////////////////////////////////////////
+void MemoryTrace_CoreAdd(void* Base, SIZE_T Size)
+{
+	void* Owner = GetOwner(bLight);
+	GAllocationTrace->CoreAdd(Base, Size, Owner);
+}
+
+////////////////////////////////////////////////////////////////////////////////
+void MemoryTrace_CoreRemove(void* Base, SIZE_T Size)
+{
+	void* Owner = GetOwner(bLight);
+	GAllocationTrace->CoreRemove(Base, Size, Owner);
+}
+#endif // 0
+
+#undef USE_OVERVIEW_TRACE
+
+#include "Windows/HideWindowsPlatformTypes.h"
+
+#endif // UE_MEMORY_TRACE_ENABLED
diff --git a/Engine/Source/Runtime/Core/Private/ProfilingDebugging/Windows/MemoryTrace.inl b/Engine/Source/Runtime/Core/Private/ProfilingDebugging/Windows/MemoryTrace.inl
new file mode 100644
--- /dev/null
+++ b/Engine/Source/Runtime/Core/Private/ProfilingDebugging/Windows/MemoryTrace.inl
@@ -0,0 +1,67 @@
+// Copyright Epic Games, Inc. All Rights Reserved.
+
+#pragma once
+
+#include "HAL/MemoryBase.h"
+
+////////////////////////////////////////////////////////////////////////////////
+template <typename T>
+struct TMiniArray
+{
+				TMiniArray(FMalloc* InMalloc) : Malloc(InMalloc) {}
+				~TMiniArray() { Malloc->Free(Data); }
+	void		Add(T const& Item);
+	void		Insert(T const& Item, int32 Index);
+	void		RemoveAt(int32 Index);
+	void		MakeRoom();
+	T*			begin()		{ return Data; }
+	T*			end()		{ return Data + Num; }
+	FMalloc*	Malloc;
+	int32		Num = 0;
+	int32		Max = 0;
+	T*			Data = nullptr;
+};
+
+////////////////////////////////////////////////////////////////////////////////
+template <typename T>
+inline void TMiniArray<T>::MakeRoom()   
+{
+	if (Num + 1 >= Max)
+	{
+		Max += 8;
+		Data = (T*)(Malloc->Realloc(Data, sizeof(T) * Max));
+	}
+}
+
+////////////////////////////////////////////////////////////////////////////////
+template <typename T>
+inline void TMiniArray<T>::Add(T const& Item)
+{
+	MakeRoom();
+	Data[Num] = Item;
+	++Num;
+}
+
+////////////////////////////////////////////////////////////////////////////////
+template <typename T>
+inline void TMiniArray<T>::Insert(T const& Item, int32 Index)
+{
+	MakeRoom();
+	for (int i = Num - 1; i > Index; --i)
+	{
+		Data[i] = Data[i - 1];
+	}
+	Data[Index] = Item;
+	++Num;
+}       
+
+////////////////////////////////////////////////////////////////////////////////
+template <typename T>
+inline void TMiniArray<T>::RemoveAt(int32 Index)
+{
+	for (int i = Index + 1; i < Num; ++i)
+	{
+		Data[i - 1] = Data[i];
+	}
+	--Num;
+}
diff --git a/Engine/Source/Runtime/Core/Private/ProfilingDebugging/Windows/ModuleDiagnostics.cpp b/Engine/Source/Runtime/Core/Private/ProfilingDebugging/Windows/ModuleDiagnostics.cpp
new file mode 100644
--- /dev/null
+++ b/Engine/Source/Runtime/Core/Private/ProfilingDebugging/Windows/ModuleDiagnostics.cpp
@@ -0,0 +1,216 @@
+// Copyright Epic Games, Inc. All Rights Reserved.
+
+#include "CoreTypes.h"
+#include "MemoryTrace.inl"
+#include "Trace/Trace.inl"
+
+#include "Windows/AllowWindowsPlatformTypes.h"
+#	include <winternl.h>
+#include "Windows/HideWindowsPlatformTypes.h"
+
+////////////////////////////////////////////////////////////////////////////////
+struct FNtDllFunction
+{
+	FARPROC Addr;
+
+	FNtDllFunction(const char* Name)
+	{
+		HMODULE NtDll = LoadLibraryW(L"ntdll.dll");
+		Addr = GetProcAddress(NtDll, Name);
+	}
+
+	template <typename... ArgTypes>
+	unsigned int operator () (ArgTypes... Args)
+	{
+		typedef unsigned int (NTAPI *Prototype)(ArgTypes...);
+		return (Prototype((void*)Addr))(Args...);
+	}
+};
+
+
+
+////////////////////////////////////////////////////////////////////////////////
+UE_TRACE_CHANNEL(ModuleChannel, "Module information needed for symbols resolution", true)
+
+UE_TRACE_EVENT_BEGIN(Diagnostics, ModuleLoad, NoSync|Important)
+	UE_TRACE_EVENT_FIELD(Trace::WideString, Name)
+	UE_TRACE_EVENT_FIELD(uint32, Base)
+	UE_TRACE_EVENT_FIELD(uint32, Size)
+UE_TRACE_EVENT_END()
+
+UE_TRACE_EVENT_BEGIN(Diagnostics, ModuleUnload, NoSync|Important)
+	UE_TRACE_EVENT_FIELD(uint32, Base)
+UE_TRACE_EVENT_END()
+
+////////////////////////////////////////////////////////////////////////////////
+class FModuleTrace
+{
+public:
+	typedef void				(*SubscribeFunc)(bool, void*, const TCHAR*);
+
+								FModuleTrace(FMalloc* InMalloc);
+								~FModuleTrace();
+	static FModuleTrace*		Get();
+	void						Initialize();
+	void						Subscribe(SubscribeFunc Function);
+
+private:
+	void						OnDllLoaded(const UNICODE_STRING& Name, UPTRINT Base);
+	void						OnDllUnloaded(UPTRINT Base);
+	void						OnDllNotification(unsigned int Reason, const void* DataPtr);
+	static FModuleTrace*		Instance;
+	TMiniArray<SubscribeFunc>	Subscribers;
+	void*						CallbackCookie;
+};
+
+////////////////////////////////////////////////////////////////////////////////
+FModuleTrace* FModuleTrace::Instance = nullptr;
+
+////////////////////////////////////////////////////////////////////////////////
+FModuleTrace::FModuleTrace(FMalloc* InMalloc)
+: Subscribers(InMalloc)
+{
+	Instance = this;
+}
+
+////////////////////////////////////////////////////////////////////////////////
+FModuleTrace::~FModuleTrace()
+{
+	FNtDllFunction UnregisterFunc("LdrUnregisterDllNotification");
+	UnregisterFunc(CallbackCookie);
+}
+
+////////////////////////////////////////////////////////////////////////////////
+FModuleTrace* FModuleTrace::Get()
+{
+	return Instance;
+}
+
+////////////////////////////////////////////////////////////////////////////////
+void FModuleTrace::Initialize()
+{
+	using namespace Trace;
+
+	// Register for DLL load/unload notifications.
+	auto Thunk = [] (ULONG Reason, const void* Data, void* Context)
+	{
+		auto* Self = (FModuleTrace*)Context;
+		Self->OnDllNotification(Reason, Data);
+	};
+
+	typedef void (CALLBACK *ThunkType)(ULONG, const void*, void*);
+	auto ThunkImpl = ThunkType(Thunk);
+
+	FNtDllFunction RegisterFunc("LdrRegisterDllNotification");
+	RegisterFunc(0, ThunkImpl, this, &CallbackCookie);
+
+	// Enumerate already loaded modules.
+	const TEB* ThreadEnvBlock = NtCurrentTeb();
+	const PEB* ProcessEnvBlock = ThreadEnvBlock->ProcessEnvironmentBlock;
+	const LIST_ENTRY* ModuleIter = ProcessEnvBlock->Ldr->InMemoryOrderModuleList.Flink;
+	const LIST_ENTRY* ModuleIterEnd = ModuleIter->Blink;
+	do
+	{
+		const auto& ModuleData = *(LDR_DATA_TABLE_ENTRY*)(ModuleIter - 1);
+		if (ModuleData.DllBase == 0)
+		{
+			break;
+		}
+
+		OnDllLoaded(ModuleData.FullDllName, UPTRINT(ModuleData.DllBase));
+		ModuleIter = ModuleIter->Flink;
+	}
+	while (ModuleIter != ModuleIterEnd);
+}
+
+////////////////////////////////////////////////////////////////////////////////
+void FModuleTrace::Subscribe(SubscribeFunc Function)
+{
+	Subscribers.Add(Function);
+}
+
+////////////////////////////////////////////////////////////////////////////////
+void FModuleTrace::OnDllNotification(unsigned int Reason, const void* DataPtr)
+{
+	enum
+	{
+		LDR_DLL_NOTIFICATION_REASON_LOADED = 1,
+		LDR_DLL_NOTIFICATION_REASON_UNLOADED = 2,
+	};
+
+	struct FNotificationData
+	{
+		uint32					Flags;
+		const UNICODE_STRING&	FullPath;
+		const UNICODE_STRING&	BaseName;
+		UPTRINT					Base;
+	};
+	const auto& Data = *(FNotificationData*)DataPtr;
+
+	switch (Reason)
+	{
+	case LDR_DLL_NOTIFICATION_REASON_LOADED:	OnDllLoaded(Data.FullPath, Data.Base);	break;
+	case LDR_DLL_NOTIFICATION_REASON_UNLOADED:	OnDllUnloaded(Data.Base);				break;
+	}
+}
+
+////////////////////////////////////////////////////////////////////////////////
+void FModuleTrace::OnDllLoaded(const UNICODE_STRING& Name, UPTRINT Base)
+{
+	const auto* DosHeader = (IMAGE_DOS_HEADER*)Base;
+	const auto* NtHeaders = (IMAGE_NT_HEADERS*)(Base + DosHeader->e_lfanew);
+	const IMAGE_OPTIONAL_HEADER& OptionalHeader = NtHeaders->OptionalHeader;
+
+	UE_TRACE_LOG(Diagnostics, ModuleLoad, ModuleChannel, Name.Length * sizeof(TCHAR))
+		<< ModuleLoad.Name(Name.Buffer, Name.Length / 2)
+		<< ModuleLoad.Base(uint32(Base >> 16)) // Windows' DLLs are on 64K page boundaries
+		<< ModuleLoad.Size(OptionalHeader.SizeOfImage);
+
+	for (SubscribeFunc Subscriber : Subscribers)
+	{
+		Subscriber(true, (void*)Base, Name.Buffer);
+	}
+}
+
+////////////////////////////////////////////////////////////////////////////////
+void FModuleTrace::OnDllUnloaded(UPTRINT Base)
+{
+	UE_TRACE_LOG(Diagnostics, ModuleUnload, ModuleChannel)
+		<< ModuleUnload.Base(uint32(Base >> 16));
+
+	for (SubscribeFunc Subscriber : Subscribers)
+	{
+		Subscriber(false, (void*)Base, nullptr);
+	}
+}
+
+
+
+////////////////////////////////////////////////////////////////////////////////
+void Modules_Create(FMalloc* Malloc)
+{
+	if (FModuleTrace::Get() != nullptr)
+	{
+		return;
+	}
+
+	static FModuleTrace Instance(Malloc);
+}
+
+////////////////////////////////////////////////////////////////////////////////
+void Modules_Initialize()
+{
+	if (FModuleTrace* Instance = FModuleTrace::Get())
+	{
+		Instance->Initialize();
+	}
+}
+
+////////////////////////////////////////////////////////////////////////////////
+void Modules_Subscribe(void (*Function)(bool, void*, const TCHAR*))
+{
+	if (FModuleTrace* Instance = FModuleTrace::Get())
+	{
+		Instance->Subscribe(Function);
+	}
+}
diff --git a/Engine/Source/Runtime/Core/Public/HAL/UnrealMemory.h b/Engine/Source/Runtime/Core/Public/HAL/UnrealMemory.h
--- a/Engine/Source/Runtime/Core/Public/HAL/UnrealMemory.h
+++ b/Engine/Source/Runtime/Core/Public/HAL/UnrealMemory.h
@@ -171,11 +171,13 @@ struct CORE_API FMemory
 	//
 	static FORCEINLINE void* SystemMalloc(SIZE_T Size)
 	{
+		/* TODO: Trace! */
 		return ::malloc(Size);
 	}
 
 	static FORCEINLINE void SystemFree(void* Ptr)
 	{
+		/* TODO: Trace! */
 		::free(Ptr);
 	}
 
