FMath SegmentTriangleIntersection - Fix hit case when segment intersects with edge of triangle ( was reporting a miss in this case ).

Should have minimal impact: this code is only being used within the mesh painting toolset and the interactive toolset as part of LineTraceComponent calls.

diff --git a/Engine/Source/Runtime/Core/Private/Math/UnrealMath.cpp b/Engine/Source/Runtime/Core/Private/Math/UnrealMath.cpp
--- a/Engine/Source/Runtime/Core/Private/Math/UnrealMath.cpp
+++ b/Engine/Source/Runtime/Core/Private/Math/UnrealMath.cpp
@@ -1753,7 +1753,14 @@ bool FMath::SegmentTriangleIntersection(const FVector& StartPoint, const FVector
 	}
 
 	FVector BaryCentric = FMath::ComputeBaryCentric2D(OutIntersectPoint, A, B, C);
-	if (BaryCentric.X > 0.0f && BaryCentric.Y > 0.0f && BaryCentric.Z > 0.0f)
+
+	// ComputeBaryCenteric2D returns ZeroVector when the triangle is too small.
+	if (BaryCentric == FVector::ZeroVector)
+	{
+		return false;
+	}
+
+	if (BaryCentric.X >= 0.0f && BaryCentric.Y >= 0.0f && BaryCentric.Z >= 0.0f)
 	{
 		OutTriangleNormal = TriNormal;
 		return true;
