Remove pragma optimize macro

diff --git a/Engine/Source/Runtime/Core/Private/ProfilingDebugging/Windows/MemoryTrace.cpp b/Engine/Source/Runtime/Core/Private/ProfilingDebugging/Windows/MemoryTrace.cpp
--- a/Engine/Source/Runtime/Core/Private/ProfilingDebugging/Windows/MemoryTrace.cpp
+++ b/Engine/Source/Runtime/Core/Private/ProfilingDebugging/Windows/MemoryTrace.cpp
@@ -18,8 +18,6 @@
 #include <atomic>
 #include <intrin.h>
 
-PRAGMA_DISABLE_OPTIMIZATION
-
 ////////////////////////////////////////////////////////////////////////////////
 void	Backtracer_Create(FMalloc*);
 void	Backtracer_Initialize();
