CompactBinary: Updated TCbBufferRef for compact binary to match the new interface for FSharedBuffer

#rnx

diff --git a/Engine/Source/Runtime/Core/Private/Tests/Serialization/CompactBinaryTest.cpp b/Engine/Source/Runtime/Core/Private/Tests/Serialization/CompactBinaryTest.cpp
--- a/Engine/Source/Runtime/Core/Private/Tests/Serialization/CompactBinaryTest.cpp
+++ b/Engine/Source/Runtime/Core/Private/Tests/Serialization/CompactBinaryTest.cpp
@@ -89,7 +89,7 @@ protected:
 	void TestFieldAsType(FCbField& Field, T (FCbField::*AsTypeFn)(Default), T ExpectedValue = T(), T DefaultValue = T(), ECbFieldError ExpectedError = ECbFieldError::None)
 	{
 		TestFieldAsTypeNoClone(Field, AsTypeFn, ExpectedValue, DefaultValue, ExpectedError);
-		FCbFieldRef FieldClone(FCbFieldRef::Clone, Field);
+		FCbFieldRef FieldClone = FCbFieldRef::Clone(Field);
 		TestFieldAsTypeNoClone(FieldClone, AsTypeFn, ExpectedValue, DefaultValue, ExpectedError);
 		TestTrue(TEXT("FCbField::Equals()"), Field.Equals(FieldClone));
 	}
@@ -107,7 +107,7 @@ protected:
 	void TestFieldAsType(FCbField& Field, T (FCbField::*AsTypeFn)(), T ExpectedValue = T(), T DefaultValue = T(), ECbFieldError ExpectedError = ECbFieldError::None)
 	{
 		TestFieldAsTypeNoClone(Field, AsTypeFn, ExpectedValue, DefaultValue, ExpectedError);
-		FCbFieldRef FieldClone(FCbFieldRef::Clone, Field);
+		FCbFieldRef FieldClone = FCbFieldRef::Clone(Field);
 		TestFieldAsTypeNoClone(FieldClone, AsTypeFn, ExpectedValue, DefaultValue, ExpectedError);
 		TestTrue(TEXT("FCbField::Equals()"), Field.Equals(FieldClone));
 	}
@@ -316,7 +316,7 @@ bool FCbFieldObjectTest::RunTest(const FString& Parameters)
 		const uint8 Payload[] = { 12, IntType, 1, 'A', 1, IntType, 1, 'B', 2, IntType, 1, 'C', 3 };
 		FCbField Field(Payload, ECbFieldType::Object);
 		TestFieldAsType(Field, &FCbField::AsObject);
-		FCbObjectRef Object(FCbObjectRef::Clone, Field.AsObject());
+		FCbObjectRef Object = FCbObjectRef::Clone(Field.AsObject());
 		TestIntObject(Object, 3, sizeof(Payload));
 		TestIntObject(Field.AsObject(), 3, sizeof(Payload));
 		TestTrue(TEXT("FCbObject::Equals()"), Object.Equals(Field.AsObject()));
@@ -334,7 +334,7 @@ bool FCbFieldObjectTest::RunTest(const FString& Parameters)
 		const uint8 Payload[] = { 10, IntType, 1, 'A', 1, 1, 'B', 2, 1, 'C', 3 };
 		FCbField Field(Payload, ECbFieldType::UniformObject);
 		TestFieldAsType(Field, &FCbField::AsObject);
-		FCbObjectRef Object(FCbObjectRef::Clone, Field.AsObject());
+		FCbObjectRef Object = FCbObjectRef::Clone(Field.AsObject());
 		TestIntObject(Object, 3, sizeof(Payload));
 		TestIntObject(Field.AsObject(), 3, sizeof(Payload));
 		TestTrue(TEXT("FCbObject{Uniform}::Equals()"), Object.Equals(Field.AsObject()));
@@ -348,14 +348,14 @@ bool FCbFieldObjectTest::RunTest(const FString& Parameters)
 		TestEqual(TEXT("FCbObject{Uniform}::FindIgnoreCase()"), Object.FindRefIgnoreCase("b"_ASV).AsInt32(), 2);
 		TestEqual(TEXT("FCbObject{Uniform}::operator[]"), Object["B"_ASV].AsInt32(), 2);
 		TestEqual(TEXT("FCbObject{Uniform}::operator[]"), Object["b"_ASV].AsInt32(4), 4);
-		Object = FCbFieldRef(FCbFieldRef::Wrap, Field).AsObjectRef();
+		Object = FCbFieldRef::MakeView(Field).AsObjectRef();
 	}
 
 	// Test FCbField(None) as Object
 	{
 		FCbField Field;
 		TestFieldError<ECbFieldType::Object>(Field, ECbFieldError::TypeError);
-		FCbFieldRef(FCbFieldRef::Wrap, Field).AsObjectRef();
+		FCbFieldRef::MakeView(Field).AsObjectRef();
 	}
 
 	// Test FCbObject(ObjectWithName) and CreateRefIterator
@@ -364,7 +364,7 @@ bool FCbFieldObjectTest::RunTest(const FString& Parameters)
 		const uint8 Buffer[] = { ObjectType, 3, 'K', 'e', 'y', 4, uint8(ECbFieldType::HasFieldName | ECbFieldType::IntegerPositive), 1, 'F', 8 };
 		const FCbObject Object(Buffer);
 		TestEqual(TEXT("FCbObject(ObjectWithName)::Size()"), Object.Size(), uint64(6));
-		const FCbObjectRef ObjectClone(FCbObjectRef::Clone, Object);
+		const FCbObjectRef ObjectClone = FCbObjectRef::Clone(Object);
 		TestEqual(TEXT("FCbObjectRef(ObjectWithName)::Size()"), ObjectClone.Size(), uint64(6));
 		TestTrue(TEXT("FCbObject::Equals()"), Object.Equals(ObjectClone));
 		for (FCbFieldRefIterator It = ObjectClone.CreateRefIterator(); It; ++It)
@@ -434,7 +434,7 @@ bool FCbFieldArrayTest::RunTest(const FString& Parameters)
 		const uint8 Payload[] = { 7, 3, IntType, 1, IntType, 2, IntType, 3 };
 		FCbField Field(Payload, ECbFieldType::Array);
 		TestFieldAsType(Field, &FCbField::AsArray);
-		FCbArrayRef Array(FCbArrayRef::Clone, Field.AsArray());
+		FCbArrayRef Array = FCbArrayRef::Clone(Field.AsArray());
 		TestIntArray(Array, 3, sizeof(Payload));
 		TestIntArray(Field.AsArray(), 3, sizeof(Payload));
 		TestTrue(TEXT("FCbArray::Equals()"), Array.Equals(Field.AsArray()));
@@ -446,18 +446,18 @@ bool FCbFieldArrayTest::RunTest(const FString& Parameters)
 		const uint8 Payload[] = { 5, 3, IntType, 1, 2, 3 };
 		FCbField Field(Payload, ECbFieldType::UniformArray);
 		TestFieldAsType(Field, &FCbField::AsArray);
-		FCbArrayRef Array(FCbArrayRef::Clone, Field.AsArray());
+		FCbArrayRef Array = FCbArrayRef::Clone(Field.AsArray());
 		TestIntArray(Array, 3, sizeof(Payload));
 		TestIntArray(Field.AsArray(), 3, sizeof(Payload));
 		TestTrue(TEXT("FCbArray::Equals()"), Array.Equals(Field.AsArray()));
-		Array = FCbFieldRef(FCbFieldRef::Wrap, Field).AsArrayRef();
+		Array = FCbFieldRef::MakeView(Field).AsArrayRef();
 	}
 
 	// Test FCbField(None) as Array
 	{
 		FCbField Field;
 		TestFieldError<ECbFieldType::Array>(Field, ECbFieldError::TypeError);
-		FCbFieldRef(FCbFieldRef::Wrap, Field).AsArrayRef();
+		FCbFieldRef::MakeView(Field).AsArrayRef();
 	}
 
 	// Test FCbArray(ArrayWithName) and CreateRefIterator
@@ -466,7 +466,7 @@ bool FCbFieldArrayTest::RunTest(const FString& Parameters)
 		const uint8 Buffer[] = { ArrayType, 3, 'K', 'e', 'y', 3, 1, uint8(ECbFieldType::IntegerPositive), 8 };
 		const FCbArray Array(Buffer);
 		TestEqual(TEXT("Array(ArrayWithName)::Size()"), Array.Size(), uint64(5));
-		const FCbArrayRef ArrayClone(FCbArrayRef::Clone, Array);
+		const FCbArrayRef ArrayClone = FCbArrayRef::Clone(Array);
 		TestEqual(TEXT("FCbArrayRef(ArrayWithName)::Size()"), ArrayClone.Size(), uint64(5));
 		TestTrue(TEXT("FCbArray::Equals()"), Array.Equals(ArrayClone));
 		for (FCbFieldRefIterator It = ArrayClone.CreateRefIterator(); It; ++It)
@@ -916,9 +916,6 @@ bool FCbFieldRefTest::RunTest(const FString& Parameters)
 	static_assert(std::is_constructible<FCbFieldRef, const FCbField&, FCbArrayRef&&>::value, "Missing constructor for FCbFieldRef");
 	static_assert(std::is_constructible<FCbFieldRef, const FCbField&, FCbObjectRef&&>::value, "Missing constructor for FCbFieldRef");
 
-	static_assert(std::is_constructible<FCbFieldRef, TCbBufferRef<FCbField>&&>::value, "Missing constructor for FCbFieldRef");
-	static_assert(std::is_assignable<FCbFieldRef, TCbBufferRef<FCbField>&&>::value, "Missing assignment for FCbFieldRef");
-
 	// Test FCbFieldRef()
 	{
 		FCbFieldRef DefaultField;
@@ -931,16 +928,16 @@ bool FCbFieldRefTest::RunTest(const FString& Parameters)
 	// Test Field w/ Type from Shared Buffer
 	{
 		uint8 Payload[] = { uint8(ECbFieldType::Binary), 3, 4, 5, 6 }; // Size: 3, Data: 4/5/6
-		FSharedBufferRef WrapBuffer = FSharedBuffer::MakeView(MakeMemoryView(Payload));
-		FSharedBufferRef OwnedBuffer = FSharedBuffer::MakeOwned(WrapBuffer);
-		FSharedBufferConstPtr WrapBufferPtr = WrapBuffer;
+		FSharedBufferRef ViewBuffer = FSharedBuffer::MakeView(MakeMemoryView(Payload));
+		FSharedBufferRef OwnedBuffer = FSharedBuffer::MakeOwned(ViewBuffer);
+		FSharedBufferConstPtr ViewBufferPtr = ViewBuffer;
 		FSharedBufferConstPtr OwnedBufferPtr = OwnedBuffer;
 
-		FCbFieldRef WrapRef(WrapBuffer);
-		FCbFieldRef WrapPtr(WrapBufferPtr);
-		FCbFieldRef WrapPtrMove{FSharedBufferConstPtr(WrapBufferPtr)};
-		FCbFieldRef WrapOuterFieldRef(ImplicitConv<FCbField>(WrapPtr), WrapPtrMove);
-		FCbFieldRef WrapOuterBufferRef(ImplicitConv<FCbField>(WrapRef), WrapPtr);
+		FCbFieldRef ViewRef(ViewBuffer);
+		FCbFieldRef ViewPtr(ViewBufferPtr);
+		FCbFieldRef ViewPtrMove{FSharedBufferConstPtr(ViewBufferPtr)};
+		FCbFieldRef ViewOuterFieldRef(ImplicitConv<FCbField>(ViewPtr), ViewPtrMove);
+		FCbFieldRef ViewOuterBufferRef(ImplicitConv<FCbField>(ViewRef), ViewPtr);
 		FCbFieldRef OwnedRef(OwnedBuffer);
 		FCbFieldRef OwnedPtr(OwnedBufferPtr);
 		FCbFieldRef OwnedPtrMove{FSharedBufferConstPtr(OwnedBufferPtr)};
@@ -948,38 +945,38 @@ bool FCbFieldRefTest::RunTest(const FString& Parameters)
 		FCbFieldRef OwnedOuterBufferRef(ImplicitConv<FCbField>(OwnedRef), OwnedPtr);
 
 		// These lines are expected to assert when uncommented.
-		//FCbFieldRef InvalidOuterBuffer(ImplicitConv<FCbField>(OwnedRef), WrapBufferPtr);
-		//FCbFieldRef InvalidOuterBufferMove(ImplicitConv<FCbField>(OwnedRef), FSharedBufferConstPtr(WrapBufferPtr));
+		//FCbFieldRef InvalidOuterBuffer(ImplicitConv<FCbField>(OwnedRef), ViewBufferPtr);
+		//FCbFieldRef InvalidOuterBufferMove(ImplicitConv<FCbField>(OwnedRef), FSharedBufferConstPtr(ViewBufferPtr));
 
 		Payload[UE_ARRAY_COUNT(Payload) - 1] = 4;
 
-		TestEqualBytes(TEXT("FCbFieldRef(WrapBufferRef)"), WrapRef.AsBinary(), {4, 5, 4});
-		TestEqualBytes(TEXT("FCbFieldRef(WrapBufferPtr)"), WrapPtr.AsBinary(), {4, 5, 4});
-		TestEqualBytes(TEXT("FCbFieldRef(WrapBufferPtr&&)"), WrapPtrMove.AsBinary(), {4, 5, 4});
-		TestEqualBytes(TEXT("FCbFieldRef(WrapOuterFieldRef)"), WrapOuterFieldRef.AsBinary(), {4, 5, 4});
-		TestEqualBytes(TEXT("FCbFieldRef(WrapOuterBufferRef)"), WrapOuterBufferRef.AsBinary(), {4, 5, 4});
+		TestEqualBytes(TEXT("FCbFieldRef(ViewBufferRef)"), ViewRef.AsBinary(), {4, 5, 4});
+		TestEqualBytes(TEXT("FCbFieldRef(ViewBufferPtr)"), ViewPtr.AsBinary(), {4, 5, 4});
+		TestEqualBytes(TEXT("FCbFieldRef(ViewBufferPtr&&)"), ViewPtrMove.AsBinary(), {4, 5, 4});
+		TestEqualBytes(TEXT("FCbFieldRef(ViewOuterFieldRef)"), ViewOuterFieldRef.AsBinary(), {4, 5, 4});
+		TestEqualBytes(TEXT("FCbFieldRef(ViewOuterBufferRef)"), ViewOuterBufferRef.AsBinary(), {4, 5, 4});
 		TestEqualBytes(TEXT("FCbFieldRef(OwnedBufferRef)"), OwnedRef.AsBinary(), {4, 5, 6});
 		TestEqualBytes(TEXT("FCbFieldRef(OwnedBufferPtr)"), OwnedPtr.AsBinary(), {4, 5, 6});
 		TestEqualBytes(TEXT("FCbFieldRef(OwnedBufferPtr&&)"), OwnedPtrMove.AsBinary(), {4, 5, 6});
 		TestEqualBytes(TEXT("FCbFieldRef(OwnedOuterFieldRef)"), OwnedOuterFieldRef.AsBinary(), {4, 5, 6});
 		TestEqualBytes(TEXT("FCbFieldRef(OwnedOuterBufferRef)"), OwnedOuterBufferRef.AsBinary(), {4, 5, 6});
 
-		TestFalse(TEXT("FCbFieldRef(WrapBufferRef).IsOwned()"), WrapRef.IsOwned());
-		TestFalse(TEXT("FCbFieldRef(WrapBufferPtr).IsOwned()"), WrapPtr.IsOwned());
-		TestFalse(TEXT("FCbFieldRef(WrapBufferPtr&&).IsOwned()"), WrapPtrMove.IsOwned());
-		TestFalse(TEXT("FCbFieldRef(WrapOuterFieldRef).IsOwned()"), WrapOuterFieldRef.IsOwned());
-		TestFalse(TEXT("FCbFieldRef(WrapOuterBufferRef).IsOwned()"), WrapOuterBufferRef.IsOwned());
+		TestFalse(TEXT("FCbFieldRef(ViewBufferRef).IsOwned()"), ViewRef.IsOwned());
+		TestFalse(TEXT("FCbFieldRef(ViewBufferPtr).IsOwned()"), ViewPtr.IsOwned());
+		TestFalse(TEXT("FCbFieldRef(ViewBufferPtr&&).IsOwned()"), ViewPtrMove.IsOwned());
+		TestFalse(TEXT("FCbFieldRef(ViewOuterFieldRef).IsOwned()"), ViewOuterFieldRef.IsOwned());
+		TestFalse(TEXT("FCbFieldRef(ViewOuterBufferRef).IsOwned()"), ViewOuterBufferRef.IsOwned());
 		TestTrue(TEXT("FCbFieldRef(OwnedBufferRef).IsOwned()"), OwnedRef.IsOwned());
 		TestTrue(TEXT("FCbFieldRef(OwnedBufferPtr).IsOwned()"), OwnedPtr.IsOwned());
 		TestTrue(TEXT("FCbFieldRef(OwnedBufferPtr&&).IsOwned()"), OwnedPtrMove.IsOwned());
 		TestTrue(TEXT("FCbFieldRef(OwnedOuterFieldRef).IsOwned()"), OwnedOuterFieldRef.IsOwned());
 		TestTrue(TEXT("FCbFieldRef(OwnedOuterBufferRef).IsOwned()"), OwnedOuterBufferRef.IsOwned());
 
-		WrapRef.MakeOwned();
+		ViewRef.MakeOwned();
 		OwnedRef.MakeOwned();
 		static_cast<uint8*>(OwnedBuffer->GetData())[UE_ARRAY_COUNT(Payload) - 1] = 5;
-		TestEqualBytes(TEXT("FCbFieldRef(Wrap).MakeOwned()"), WrapRef.AsBinary(), {4, 5, 4});
-		TestTrue(TEXT("FCbFieldRef(Wrap).MakeOwned().IsOwned()"), WrapRef.IsOwned());
+		TestEqualBytes(TEXT("FCbFieldRef(View).MakeOwned()"), ViewRef.AsBinary(), {4, 5, 4});
+		TestTrue(TEXT("FCbFieldRef(View).MakeOwned().IsOwned()"), ViewRef.IsOwned());
 		TestEqualBytes(TEXT("FCbFieldRef(Owned).MakeOwned()"), OwnedRef.AsBinary(), {4, 5, 5});
 		TestTrue(TEXT("FCbFieldRef(Owned).MakeOwned().IsOwned()"), OwnedRef.IsOwned());
 	}
@@ -992,28 +989,28 @@ bool FCbFieldRefTest::RunTest(const FString& Parameters)
 
 		FCbField Field(Payload);
 
-		FCbFieldRef VoidAssume(FCbFieldRef::AssumeOwnership, ImplicitConv<const void*>(PayloadCopy));
-		FCbFieldRef VoidWrap(FCbFieldRef::Wrap, Payload);
-		FCbFieldRef VoidClone(FCbFieldRef::Clone, Payload);
-		FCbFieldRef FieldWrap(FCbFieldRef::Wrap, Field);
-		FCbFieldRef FieldClone(FCbFieldRef::Clone, Field);
-		FCbFieldRef FieldRefClone(FCbFieldRef::Clone, FieldWrap);
+		FCbFieldRef VoidTake = FCbFieldRef::TakeOwnership(ImplicitConv<const void*>(PayloadCopy), FMemory::Free);
+		FCbFieldRef VoidView = FCbFieldRef::MakeView(Payload);
+		FCbFieldRef VoidClone = FCbFieldRef::Clone(Payload);
+		FCbFieldRef FieldView = FCbFieldRef::MakeView(Field);
+		FCbFieldRef FieldClone = FCbFieldRef::Clone(Field);
+		FCbFieldRef FieldRefClone = FCbFieldRef::Clone(FieldView);
 
 		Payload[UE_ARRAY_COUNT(Payload) - 1] = 4;
 
-		TestEqualBytes(TEXT("FCbFieldRef(AssumeOwnership, Void)"), VoidAssume.AsBinary(), {4, 5, 6});
-		TestEqualBytes(TEXT("FCbFieldRef(Wrap, Void)"), VoidWrap.AsBinary(), {4, 5, 4});
-		TestEqualBytes(TEXT("FCbFieldRef(Clone, Void)"), VoidClone.AsBinary(), {4, 5, 6});
-		TestEqualBytes(TEXT("FCbFieldRef(Wrap, Field)"), FieldWrap.AsBinary(), {4, 5, 4});
-		TestEqualBytes(TEXT("FCbFieldRef(Clone, Field)"), FieldClone.AsBinary(), {4, 5, 6});
-		TestEqualBytes(TEXT("FCbFieldRef(Clone, FieldRef)"), FieldRefClone.AsBinary(), {4, 5, 6});
-
-		TestTrue(TEXT("FCbFieldRef(AssumeOwnership, Void).IsOwned()"), VoidAssume.IsOwned());
-		TestFalse(TEXT("FCbFieldRef(Wrap, Void).IsOwned()"), VoidWrap.IsOwned());
-		TestTrue(TEXT("FCbFieldRef(Clone, Void).IsOwned()"), VoidClone.IsOwned());
-		TestFalse(TEXT("FCbFieldRef(Wrap, Field).IsOwned()"), FieldWrap.IsOwned());
-		TestTrue(TEXT("FCbFieldRef(Clone, Field).IsOwned()"), FieldClone.IsOwned());
-		TestTrue(TEXT("FCbFieldRef(Clone, FieldRef).IsOwned()"), FieldRefClone.IsOwned());
+		TestEqualBytes(TEXT("FCbFieldRef::TakeOwnership(Void)"), VoidTake.AsBinary(), {4, 5, 6});
+		TestEqualBytes(TEXT("FCbFieldRef::MakeView(Void)"), VoidView.AsBinary(), {4, 5, 4});
+		TestEqualBytes(TEXT("FCbFieldRef::Clone(Void)"), VoidClone.AsBinary(), {4, 5, 6});
+		TestEqualBytes(TEXT("FCbFieldRef::MakeView(Field)"), FieldView.AsBinary(), {4, 5, 4});
+		TestEqualBytes(TEXT("FCbFieldRef::Clone(Field)"), FieldClone.AsBinary(), {4, 5, 6});
+		TestEqualBytes(TEXT("FCbFieldRef::Clone(FieldRef)"), FieldRefClone.AsBinary(), {4, 5, 6});
+
+		TestTrue(TEXT("FCbFieldRef::TakeOwnership(Void).IsOwned()"), VoidTake.IsOwned());
+		TestFalse(TEXT("FCbFieldRef::MakeView(Void).IsOwned()"), VoidView.IsOwned());
+		TestTrue(TEXT("FCbFieldRef::Clone(Void).IsOwned()"), VoidClone.IsOwned());
+		TestFalse(TEXT("FCbFieldRef::MakeView(Field).IsOwned()"), FieldView.IsOwned());
+		TestTrue(TEXT("FCbFieldRef::Clone(Field).IsOwned()"), FieldClone.IsOwned());
+		TestTrue(TEXT("FCbFieldRef::Clone(FieldRef).IsOwned()"), FieldRefClone.IsOwned());
 	}
 
 	// Test Field w/o Type
@@ -1021,23 +1018,23 @@ bool FCbFieldRefTest::RunTest(const FString& Parameters)
 		uint8 Payload[] = { 3, 4, 5, 6 }; // Size: 3, Data: 4/5/6
 		FCbField Field(Payload, ECbFieldType::Binary);
 
-		FCbFieldRef FieldWrap(FCbFieldRef::Wrap, Field);
-		FCbFieldRef FieldClone(FCbFieldRef::Clone, Field);
-		FCbFieldRef FieldRefClone(FCbFieldRef::Clone, FieldWrap);
+		FCbFieldRef FieldView = FCbFieldRef::MakeView(Field);
+		FCbFieldRef FieldClone = FCbFieldRef::Clone(Field);
+		FCbFieldRef FieldRefClone = FCbFieldRef::Clone(FieldView);
 
 		Payload[UE_ARRAY_COUNT(Payload) - 1] = 4;
 
-		TestEqualBytes(TEXT("FCbFieldRef(Wrap, Field, NoType)"), FieldWrap.AsBinary(), {4, 5, 4});
-		TestEqualBytes(TEXT("FCbFieldRef(Clone, Field, NoType)"), FieldClone.AsBinary(), {4, 5, 6});
-		TestEqualBytes(TEXT("FCbFieldRef(Clone, FieldRef, NoType)"), FieldRefClone.AsBinary(), {4, 5, 6});
+		TestEqualBytes(TEXT("FCbFieldRef::MakeView(Field, NoType)"), FieldView.AsBinary(), {4, 5, 4});
+		TestEqualBytes(TEXT("FCbFieldRef::Clone(Field, NoType)"), FieldClone.AsBinary(), {4, 5, 6});
+		TestEqualBytes(TEXT("FCbFieldRef::Clone(FieldRef, NoType)"), FieldRefClone.AsBinary(), {4, 5, 6});
 
-		TestFalse(TEXT("FCbFieldRef(Wrap, Field, NoType).IsOwned()"), FieldWrap.IsOwned());
-		TestTrue(TEXT("FCbFieldRef(Clone, Field, NoType).IsOwned()"), FieldClone.IsOwned());
-		TestTrue(TEXT("FCbFieldRef(Clone, FieldRef, NoType).IsOwned()"), FieldRefClone.IsOwned());
+		TestFalse(TEXT("FCbFieldRef::MakeView(Field, NoType).IsOwned()"), FieldView.IsOwned());
+		TestTrue(TEXT("FCbFieldRef::Clone(Field, NoType).IsOwned()"), FieldClone.IsOwned());
+		TestTrue(TEXT("FCbFieldRef::Clone(FieldRef, NoType).IsOwned()"), FieldRefClone.IsOwned());
 
-		FieldWrap.MakeOwned();
-		TestEqualBytes(TEXT("FCbFieldRef(Wrap, NoType).MakeOwned()"), FieldWrap.AsBinary(), {4, 5, 4});
-		TestTrue(TEXT("FCbFieldRef(Wrap, NoType).MakeOwned().IsOwned()"), FieldWrap.IsOwned());
+		FieldView.MakeOwned();
+		TestEqualBytes(TEXT("FCbFieldRef::MakeView(NoType).MakeOwned()"), FieldView.AsBinary(), {4, 5, 4});
+		TestTrue(TEXT("FCbFieldRef::MakeView(NoType).MakeOwned().IsOwned()"), FieldView.IsOwned());
 	}
 
 	return true;
@@ -1058,9 +1055,6 @@ bool FCbArrayRefTest::RunTest(const FString& Parameters)
 	static_assert(std::is_constructible<FCbArrayRef, const FCbArray&, FCbArrayRef&&>::value, "Missing constructor for FCbArrayRef");
 	static_assert(std::is_constructible<FCbArrayRef, const FCbArray&, FCbObjectRef&&>::value, "Missing constructor for FCbArrayRef");
 
-	static_assert(std::is_constructible<FCbArrayRef, TCbBufferRef<FCbArray>&&>::value, "Missing constructor for FCbArrayRef");
-	static_assert(std::is_assignable<FCbArrayRef, TCbBufferRef<FCbArray>&&>::value, "Missing assignment for FCbArrayRef");
-
 	// Test FCbArrayRef()
 	{
 		FCbArrayRef DefaultArray;
@@ -1087,9 +1081,6 @@ bool FCbObjectRefTest::RunTest(const FString& Parameters)
 	static_assert(std::is_constructible<FCbObjectRef, const FCbObject&, FCbArrayRef&&>::value, "Missing constructor for FCbObjectRef");
 	static_assert(std::is_constructible<FCbObjectRef, const FCbObject&, FCbObjectRef&&>::value, "Missing constructor for FCbObjectRef");
 
-	static_assert(std::is_constructible<FCbObjectRef, TCbBufferRef<FCbObject>&&>::value, "Missing constructor for FCbObjectRef");
-	static_assert(std::is_assignable<FCbObjectRef, TCbBufferRef<FCbObject>&&>::value, "Missing assignment for FCbObjectRef");
-
 	// Test FCbObjectRef()
 	{
 		FCbObjectRef DefaultObject;
diff --git a/Engine/Source/Runtime/Core/Public/Serialization/CompactBinary.h b/Engine/Source/Runtime/Core/Public/Serialization/CompactBinary.h
--- a/Engine/Source/Runtime/Core/Public/Serialization/CompactBinary.h
+++ b/Engine/Source/Runtime/Core/Public/Serialization/CompactBinary.h
@@ -225,31 +225,31 @@ public:
 	/** The type with transient flags removed. */
 	static constexpr inline ECbFieldType GetSerializedType(ECbFieldType Type)   { return Type & SerializedTypeMask; }
 
-	static constexpr inline bool HasFieldType(ECbFieldType Type)                { return EnumHasAnyFlags(Type, ECbFieldType::HasFieldType); }
-	static constexpr inline bool HasFieldName(ECbFieldType Type)                { return EnumHasAnyFlags(Type, ECbFieldType::HasFieldName); }
+	static constexpr inline bool HasFieldType(ECbFieldType Type) { return EnumHasAnyFlags(Type, ECbFieldType::HasFieldType); }
+	static constexpr inline bool HasFieldName(ECbFieldType Type) { return EnumHasAnyFlags(Type, ECbFieldType::HasFieldName); }
 
-	static constexpr inline bool IsNone(ECbFieldType Type)                      { return GetType(Type) == ECbFieldType::None; }
-	static constexpr inline bool IsNull(ECbFieldType Type)                      { return GetType(Type) == ECbFieldType::Null; }
+	static constexpr inline bool IsNone(ECbFieldType Type)       { return GetType(Type) == ECbFieldType::None; }
+	static constexpr inline bool IsNull(ECbFieldType Type)       { return GetType(Type) == ECbFieldType::Null; }
 
-	static constexpr inline bool IsObject(ECbFieldType Type)                    { return (Type & ObjectMask) == ObjectBase; }
-	static constexpr inline bool IsArray(ECbFieldType Type)                     { return (Type & ArrayMask) == ArrayBase; }
+	static constexpr inline bool IsObject(ECbFieldType Type)     { return (Type & ObjectMask) == ObjectBase; }
+	static constexpr inline bool IsArray(ECbFieldType Type)      { return (Type & ArrayMask) == ArrayBase; }
 
-	static constexpr inline bool IsBinary(ECbFieldType Type)                    { return GetType(Type) == ECbFieldType::Binary; }
-	static constexpr inline bool IsString(ECbFieldType Type)                    { return GetType(Type) == ECbFieldType::String; }
+	static constexpr inline bool IsBinary(ECbFieldType Type)     { return GetType(Type) == ECbFieldType::Binary; }
+	static constexpr inline bool IsString(ECbFieldType Type)     { return GetType(Type) == ECbFieldType::String; }
 
-	static constexpr inline bool IsInteger(ECbFieldType Type)                   { return (Type & IntegerMask) == IntegerBase; }
+	static constexpr inline bool IsInteger(ECbFieldType Type)    { return (Type & IntegerMask) == IntegerBase; }
 	/** Whether the field is a float, or integer due to implicit conversion. */
-	static constexpr inline bool IsFloat(ECbFieldType Type)                     { return (Type & FloatMask) == FloatBase; }
-	static constexpr inline bool IsBool(ECbFieldType Type)                      { return (Type & BoolMask) == BoolBase; }
+	static constexpr inline bool IsFloat(ECbFieldType Type)      { return (Type & FloatMask) == FloatBase; }
+	static constexpr inline bool IsBool(ECbFieldType Type)       { return (Type & BoolMask) == BoolBase; }
 
-	static constexpr inline bool IsExternHash(ECbFieldType Type)                { return (Type & ExternHashMask) == ExternHashBase; }
-	static constexpr inline bool IsBinaryHash(ECbFieldType Type)                { return GetType(Type) == ECbFieldType::BinaryHash; }
-	static constexpr inline bool IsFieldHash(ECbFieldType Type)                 { return GetType(Type) == ECbFieldType::FieldHash; }
+	static constexpr inline bool IsExternHash(ECbFieldType Type) { return (Type & ExternHashMask) == ExternHashBase; }
+	static constexpr inline bool IsBinaryHash(ECbFieldType Type) { return GetType(Type) == ECbFieldType::BinaryHash; }
+	static constexpr inline bool IsFieldHash(ECbFieldType Type)  { return GetType(Type) == ECbFieldType::FieldHash; }
 
-	static constexpr inline bool IsUuid(ECbFieldType Type)                      { return GetType(Type) == ECbFieldType::Uuid; }
+	static constexpr inline bool IsUuid(ECbFieldType Type)       { return GetType(Type) == ECbFieldType::Uuid; }
 
-	static constexpr inline bool IsDateTime(ECbFieldType Type)                  { return GetType(Type) == ECbFieldType::DateTime; }
-	static constexpr inline bool IsTimeSpan(ECbFieldType Type)                  { return GetType(Type) == ECbFieldType::TimeSpan; }
+	static constexpr inline bool IsDateTime(ECbFieldType Type)   { return GetType(Type) == ECbFieldType::DateTime; }
+	static constexpr inline bool IsTimeSpan(ECbFieldType Type)   { return GetType(Type) == ECbFieldType::TimeSpan; }
 };
 
 ///////////////////////////////////////////////////////////////////////////////////////////////////////////////////////
@@ -378,7 +378,10 @@ public:
 	CORE_API explicit FCbField(const void* Data, ECbFieldType Type = ECbFieldType::HasFieldType);
 
 	/** Name of the field if it has one, otherwise empty. */
-	constexpr inline FAnsiStringView Name() const { return FAnsiStringView(static_cast<const ANSICHAR*>(Payload) - NameLen, NameLen); }
+	constexpr inline FAnsiStringView Name() const
+	{
+		return FAnsiStringView(static_cast<const ANSICHAR*>(Payload) - NameLen, NameLen);
+	}
 
 	/** Size of the field in bytes. */
 	CORE_API uint64 Size() const;
@@ -738,10 +741,6 @@ template <typename BaseType>
 class TCbBufferRef : public BaseType
 {
 public:
-	enum EAssumeOwnershipTag { AssumeOwnership };
-	enum ECloneTag           { Clone };
-	enum EWrapTag            { Wrap };
-
 	/** Construct a default value. */
 	TCbBufferRef() = default;
 
@@ -829,41 +828,6 @@ public:
 	{
 	}
 
-	/** Construct a value and take ownership its memory. */
-	inline TCbBufferRef(EAssumeOwnershipTag, const void* Value)
-		: BaseType(Value)
-		, BufferPtr(FSharedBuffer::TakeOwnership(Value, BaseType::Size(), FMemory::Free))
-	{
-	}
-
-	/** Construct a value from an owned clone of its memory. */
-	inline TCbBufferRef(ECloneTag, const void* Value)
-		: TCbBufferRef(Clone, BaseType(Value))
-	{
-	}
-
-	/** Construct a value from an owned clone of its memory. */
-	inline TCbBufferRef(ECloneTag, const BaseType& Value)
-		: BaseType(Value)
-	{
-		const FSharedBufferRef BufferRef = FSharedBuffer::Alloc(BaseType::Size());
-		BaseType::CopyTo(*BufferRef);
-		BufferPtr = BufferRef;
-		static_cast<BaseType&>(*this) = BaseType(BufferPtr->GetData(), BaseType::GetCopyType());
-	}
-
-	/** Construct a value by wrapping its memory without taking ownership of it. */
-	inline TCbBufferRef(EWrapTag, const void* Value)
-		: BaseType(Value)
-	{
-	}
-
-	/** Construct a value by wrapping its memory without taking ownership of it. */
-	inline TCbBufferRef(EWrapTag, const BaseType& Value)
-		: BaseType(Value)
-	{
-	}
-
 	/** Whether this reference has ownership of the memory in its buffer. */
 	inline bool IsOwned() const { return BufferPtr && BufferPtr->IsOwned(); }
 
@@ -872,7 +836,10 @@ public:
 	{
 		if (!IsOwned())
 		{
-			*this = TCbBufferRef(Clone, ImplicitConv<BaseType>(*this));
+			FSharedBufferPtr Buffer = FSharedBuffer::Alloc(BaseType::Size());
+			BaseType::CopyTo(*Buffer);
+			static_cast<BaseType&>(*this) = BaseType(Buffer->GetData(), BaseType::GetCopyType());
+			BufferPtr = MoveTemp(Buffer);
 		}
 	}
 
@@ -888,24 +855,64 @@ private:
 
 ///////////////////////////////////////////////////////////////////////////////////////////////////////////////////////
 
+/**
+ * Factory functions for types derived from TCbBufferRef.
+ *
+ * This uses the curiously recurring template pattern to construct the correct type of reference.
+ * The derived type inherits from TCbBufferRef and this type to expose the factory functions.
+ */
+template <typename RefType, typename BaseType>
+class TCbBufferRefFactory
+{
+public:
+	/** Construct a value from an owned clone of its memory. */
+	static inline RefType Clone(const void* const Data)
+	{
+		return Clone(BaseType(Data));
+	}
+
+	/** Construct a value from an owned clone of its memory. */
+	static inline RefType Clone(const BaseType& Value)
+	{
+		RefType Ref = MakeView(Value);
+		Ref.MakeOwned();
+		return Ref;
+	}
+
+	/** Construct a value from a read-only view of its memory. */
+	static inline RefType MakeView(const void* const Data)
+	{
+		return MakeView(BaseType(Data));
+	}
+
+	/** Construct a value from a read-only view of its memory. */
+	static inline RefType MakeView(const BaseType& Value)
+	{
+		return RefType(Value, FSharedBufferConstPtr());
+	}
+
+	/** Construct a value and take ownership its memory. */
+	template <typename DeleterType>
+	static inline RefType TakeOwnership(const void* const Data, DeleterType&& Deleter)
+	{
+		const BaseType Value(Data);
+		return RefType(Value, FSharedBuffer::TakeOwnership(Data, Value.Size(), Forward<DeleterType>(Deleter)));
+	}
+};
+
+///////////////////////////////////////////////////////////////////////////////////////////////////////////////////////
+
 class FCbArrayRef;
 class FCbObjectRef;
 
 ///////////////////////////////////////////////////////////////////////////////////////////////////////////////////////
 
 /** A field that can hold a reference to the memory that contains it. \see \ref TCbBufferRef */
-class FCbFieldRef : public TCbBufferRef<FCbField>
+class FCbFieldRef : public TCbBufferRef<FCbField>, public TCbBufferRefFactory<FCbFieldRef, FCbField>
 {
 public:
 	using TCbBufferRef::TCbBufferRef;
 
-	FCbFieldRef() = default;
-
-	inline FCbFieldRef(TCbBufferRef<FCbField>&& Ref)
-		: TCbBufferRef(MoveTemp(Ref))
-	{
-	}
-
 	/** Access the field as an object. Defaults to an empty object on error. */
 	inline FCbObjectRef AsObjectRef();
 
@@ -918,30 +925,16 @@ class FCbFieldRefIterator : public TCbFieldIterator<FCbFieldRef>
 {
 public:
 	using TCbFieldIterator::TCbFieldIterator;
-
-	FCbFieldRefIterator() = default;
-
-	inline FCbFieldRefIterator(TCbFieldIterator<FCbFieldRef>&& Field)
-		: TCbFieldIterator(MoveTemp(Field))
-	{
-	}
 };
 
 ///////////////////////////////////////////////////////////////////////////////////////////////////////////////////////
 
 /** An array that can hold a reference to the memory that contains it. \see \ref TCbBufferRef */
-class FCbArrayRef : public TCbBufferRef<FCbArray>
+class FCbArrayRef : public TCbBufferRef<FCbArray>, public TCbBufferRefFactory<FCbArrayRef, FCbArray>
 {
 public:
 	using TCbBufferRef::TCbBufferRef;
 
-	FCbArrayRef() = default;
-
-	inline FCbArrayRef(TCbBufferRef<FCbArray>&& Ref)
-		: TCbBufferRef(MoveTemp(Ref))
-	{
-	}
-
 	/** Create an iterator for the fields of this array. */
 	inline FCbFieldRefIterator CreateRefIterator() const
 	{
@@ -953,18 +946,11 @@ public:
 ///////////////////////////////////////////////////////////////////////////////////////////////////////////////////////
 
 /** An object that can hold a reference to the memory that contains it. \see \ref TCbBufferRef */
-class FCbObjectRef : public TCbBufferRef<FCbObject>
+class FCbObjectRef : public TCbBufferRef<FCbObject>, public TCbBufferRefFactory<FCbObjectRef, FCbObject>
 {
 public:
 	using TCbBufferRef::TCbBufferRef;
 
-	FCbObjectRef() = default;
-
-	inline FCbObjectRef(TCbBufferRef<FCbObject>&& Ref)
-		: TCbBufferRef(MoveTemp(Ref))
-	{
-	}
-
 	/** Create an iterator for the fields of this object. */
 	inline FCbFieldRefIterator CreateRefIterator() const
 	{
