UE5 - NonUnity compiler fix

diff --git a/Engine/Source/Runtime/Core/Private/ProfilingDebugging/Windows/CallstackTrace.cpp b/Engine/Source/Runtime/Core/Private/ProfilingDebugging/Windows/CallstackTrace.cpp
--- a/Engine/Source/Runtime/Core/Private/ProfilingDebugging/Windows/CallstackTrace.cpp
+++ b/Engine/Source/Runtime/Core/Private/ProfilingDebugging/Windows/CallstackTrace.cpp
@@ -20,6 +20,9 @@
 #include "Misc\CoreDelegates.h"
 #include "Trace/Trace.inl"
 #include "AtomicQueue/AtomicQueue.h"
+#include "HAL/Runnable.h"
+#include "HAL/PlatformProcess.h"
+#include "Misc/ScopeLock.h"
 
 // 0=off, 1=stats, 2=validation, 3=truth_compare
 #define BACKTRACE_DBGLVL 0
