Allow UnfrozenCopy to copy data that was frozen with FPlatformLayoutParameters that are different from the current/local platform
- This is used to allow editor to load cooked materials/shader maps
- Remove UE_FORCE_64BIT_MEMORY_IMAGE_POINTERS, and instead always force memory image ptrs to be 64bits...this simplifies some code, and plan is to remove 32bit support at some point anyway
- Frozen memory image pointers only need 32bits for offset, which leaves another 32bits to stash a type identifier...this allows us to interpret that frozen data being referenced, even it's packed for a different platform

diff --git a/Engine/Plugins/Enterprise/LidarPointCloud/Source/LidarPointCloudRuntime/Private/Rendering/LidarPointCloudRenderBuffers.cpp b/Engine/Plugins/Enterprise/LidarPointCloud/Source/LidarPointCloudRuntime/Private/Rendering/LidarPointCloudRenderBuffers.cpp
--- a/Engine/Plugins/Enterprise/LidarPointCloud/Source/LidarPointCloudRuntime/Private/Rendering/LidarPointCloudRenderBuffers.cpp
+++ b/Engine/Plugins/Enterprise/LidarPointCloud/Source/LidarPointCloudRuntime/Private/Rendering/LidarPointCloudRenderBuffers.cpp
@@ -237,6 +237,8 @@ void FLidarPointCloudVertexFactory::ReleaseRHI()
 	VertexBuffer.ReleaseResource();
 }
 
+IMPLEMENT_TYPE_LAYOUT(FLidarPointCloudVertexFactoryShaderParameters);
+
 IMPLEMENT_VERTEX_FACTORY_PARAMETER_TYPE(FLidarPointCloudVertexFactory, SF_Vertex, FLidarPointCloudVertexFactoryShaderParameters);
 
 IMPLEMENT_VERTEX_FACTORY_TYPE(FLidarPointCloudVertexFactory, "/Plugin/LidarPointCloud/Private/LidarPointCloudVertexFactory.ush", /* bUsedWithMaterials */ true, /* bSupportsStaticLighting */ false, /* bSupportsDynamicLighting */ true, /* bPrecisePrevWorldPos */ false, /* bSupportsPositionOnly */ true);
diff --git a/Engine/Plugins/Enterprise/LidarPointCloud/Source/LidarPointCloudRuntime/Private/Rendering/LidarPointCloudRenderBuffers.h b/Engine/Plugins/Enterprise/LidarPointCloud/Source/LidarPointCloudRuntime/Private/Rendering/LidarPointCloudRenderBuffers.h
--- a/Engine/Plugins/Enterprise/LidarPointCloud/Source/LidarPointCloudRuntime/Private/Rendering/LidarPointCloudRenderBuffers.h
+++ b/Engine/Plugins/Enterprise/LidarPointCloud/Source/LidarPointCloudRuntime/Private/Rendering/LidarPointCloudRenderBuffers.h
@@ -102,7 +102,7 @@ struct FLidarPointCloudBatchElementUserData
  */
 class FLidarPointCloudVertexFactoryShaderParameters : public FVertexFactoryShaderParameters
 {
-	DECLARE_INLINE_TYPE_LAYOUT(FLidarPointCloudVertexFactoryShaderParameters, NonVirtual);
+	DECLARE_TYPE_LAYOUT(FLidarPointCloudVertexFactoryShaderParameters, NonVirtual);
 
 public:
 	void Bind(const FShaderParameterMap& ParameterMap);
diff --git a/Engine/Plugins/Experimental/VirtualHeightfieldMesh/Source/VirtualHeightfieldMesh/Private/VirtualHeightfieldMeshVertexFactory.cpp b/Engine/Plugins/Experimental/VirtualHeightfieldMesh/Source/VirtualHeightfieldMesh/Private/VirtualHeightfieldMeshVertexFactory.cpp
--- a/Engine/Plugins/Experimental/VirtualHeightfieldMesh/Source/VirtualHeightfieldMesh/Private/VirtualHeightfieldMeshVertexFactory.cpp
+++ b/Engine/Plugins/Experimental/VirtualHeightfieldMesh/Source/VirtualHeightfieldMesh/Private/VirtualHeightfieldMeshVertexFactory.cpp
@@ -94,7 +94,7 @@ void FVirtualHeightfieldMeshIndexBuffer::InitRHI()
  */
 class FVirtualHeightfieldMeshVertexFactoryShaderParameters : public FVertexFactoryShaderParameters
 {
-	DECLARE_INLINE_TYPE_LAYOUT(FVirtualHeightfieldMeshVertexFactoryShaderParameters, NonVirtual);
+	DECLARE_TYPE_LAYOUT(FVirtualHeightfieldMeshVertexFactoryShaderParameters, NonVirtual);
 
 public:
 	void Bind(const FShaderParameterMap& ParameterMap)
@@ -175,6 +175,8 @@ protected:
 	LAYOUT_FIELD(FShaderParameter, LodDistancesParameter);
 };
 
+IMPLEMENT_TYPE_LAYOUT(FVirtualHeightfieldMeshVertexFactoryShaderParameters);
+
 IMPLEMENT_VERTEX_FACTORY_PARAMETER_TYPE(FVirtualHeightfieldMeshVertexFactory, SF_Vertex, FVirtualHeightfieldMeshVertexFactoryShaderParameters);
 IMPLEMENT_VERTEX_FACTORY_PARAMETER_TYPE(FVirtualHeightfieldMeshVertexFactory, SF_Pixel, FVirtualHeightfieldMeshVertexFactoryShaderParameters);
 
diff --git a/Engine/Plugins/FX/Niagara/Source/NiagaraShader/Private/NiagaraShader.cpp b/Engine/Plugins/FX/Niagara/Source/NiagaraShader/Private/NiagaraShader.cpp
--- a/Engine/Plugins/FX/Niagara/Source/NiagaraShader/Private/NiagaraShader.cpp
+++ b/Engine/Plugins/FX/Niagara/Source/NiagaraShader/Private/NiagaraShader.cpp
@@ -102,9 +102,9 @@ void* FNiagaraShaderMapPointerTable::GetIndexedPointer(const FTypeLayoutDesc& Ty
 	return Super::GetIndexedPointer(TypeDesc, i);
 }
 
-void FNiagaraShaderMapPointerTable::SaveToArchive(FArchive& Ar, void* FrozenContent, bool bInlineShaderResources) const
+void FNiagaraShaderMapPointerTable::SaveToArchive(FArchive& Ar, const FPlatformTypeLayoutParameters& LayoutParams, const void* FrozenObject) const
 {
-	Super::SaveToArchive(Ar, FrozenContent, bInlineShaderResources);
+	Super::SaveToArchive(Ar, LayoutParams, FrozenObject);
 
 	int32 NumDITypes = DITypes.Num();
 
@@ -118,9 +118,9 @@ void FNiagaraShaderMapPointerTable::SaveToArchive(FArchive& Ar, void* FrozenCont
 	}
 }
 
-void FNiagaraShaderMapPointerTable::LoadFromArchive(FArchive& Ar, void* FrozenContent, bool bInlineShaderResources, bool bLoadedByCookedMaterial)
+bool FNiagaraShaderMapPointerTable::LoadFromArchive(FArchive& Ar, const FPlatformTypeLayoutParameters& LayoutParams, void* FrozenObject)
 {
-	Super::LoadFromArchive(Ar, FrozenContent, bInlineShaderResources, bLoadedByCookedMaterial);
+	const bool bResult = Super::LoadFromArchive(Ar, LayoutParams, FrozenObject);
 	INiagaraShaderModule * Module = INiagaraShaderModule::Get();
 
 	int32 NumDITypes = 0;
@@ -135,6 +135,8 @@ void FNiagaraShaderMapPointerTable::LoadFromArchive(FArchive& Ar, void* FrozenCo
 		UNiagaraDataInterfaceBase* DIType = Module->RequestDefaultDataInterface(*DIClassName);
 		DITypes.LoadIndexedPointer(DIType);
 	}
+
+	return bResult;
 }
 
 
diff --git a/Engine/Plugins/FX/Niagara/Source/NiagaraShader/Public/NiagaraShared.h b/Engine/Plugins/FX/Niagara/Source/NiagaraShader/Public/NiagaraShared.h
--- a/Engine/Plugins/FX/Niagara/Source/NiagaraShader/Public/NiagaraShared.h
+++ b/Engine/Plugins/FX/Niagara/Source/NiagaraShader/Public/NiagaraShared.h
@@ -87,8 +87,8 @@ public:
 	virtual int32 AddIndexedPointer(const FTypeLayoutDesc& TypeDesc, void* Ptr) override;
 	virtual void* GetIndexedPointer(const FTypeLayoutDesc& TypeDesc, uint32 i) const override;
 
-	virtual void SaveToArchive(FArchive& Ar, void* FrozenContent, bool bInlineShaderResources) const override;
-	virtual void LoadFromArchive(FArchive& Ar, void* FrozenContent, bool bInlineShaderResources, bool bLoadedByCookedMaterial) override;
+	virtual void SaveToArchive(FArchive& Ar, const FPlatformTypeLayoutParameters& LayoutParams, const void* FrozenObject) const override;
+	virtual bool LoadFromArchive(FArchive& Ar, const FPlatformTypeLayoutParameters& LayoutParams, void* FrozenObject) override;
 
 	TPtrTable<UNiagaraDataInterfaceBase> DITypes;
 };
diff --git a/Engine/Plugins/FX/Niagara/Source/NiagaraVertexFactories/Private/NiagaraMeshVertexFactory.cpp b/Engine/Plugins/FX/Niagara/Source/NiagaraVertexFactories/Private/NiagaraMeshVertexFactory.cpp
--- a/Engine/Plugins/FX/Niagara/Source/NiagaraVertexFactories/Private/NiagaraMeshVertexFactory.cpp
+++ b/Engine/Plugins/FX/Niagara/Source/NiagaraVertexFactories/Private/NiagaraMeshVertexFactory.cpp
@@ -14,7 +14,7 @@ IMPLEMENT_GLOBAL_SHADER_PARAMETER_STRUCT(FNiagaraMeshUniformParameters, "Niagara
 
 class FNiagaraMeshVertexFactoryShaderParametersVS : public FNiagaraVertexFactoryShaderParametersBase
 {
-	DECLARE_INLINE_TYPE_LAYOUT(FNiagaraMeshVertexFactoryShaderParametersVS, NonVirtual);
+	DECLARE_TYPE_LAYOUT(FNiagaraMeshVertexFactoryShaderParametersVS, NonVirtual);
 public:
 	void Bind(const FShaderParameterMap& ParameterMap)
 	{
@@ -54,19 +54,21 @@ public:
 private:
 
 
-		//LAYOUT_FIELD(FShaderResourceParameter, PrevTransformBuffer)
+	//LAYOUT_FIELD(FShaderResourceParameter, PrevTransformBuffer)
 
-		// 	LAYOUT_FIELD(FShaderResourceParameter, NiagaraParticleDataInt)
-		// 	LAYOUT_FIELD(FShaderParameter, Int32DataOffset)
-		// 	LAYOUT_FIELD(FShaderParameter, Int32DataStride)
-		LAYOUT_FIELD(FShaderResourceParameter, SortedIndices)
-		LAYOUT_FIELD(FShaderParameter, SortedIndicesOffset)
+	// 	LAYOUT_FIELD(FShaderResourceParameter, NiagaraParticleDataInt)
+	// 	LAYOUT_FIELD(FShaderParameter, Int32DataOffset)
+	// 	LAYOUT_FIELD(FShaderParameter, Int32DataStride)
+	LAYOUT_FIELD(FShaderResourceParameter, SortedIndices);
+	LAYOUT_FIELD(FShaderParameter, SortedIndicesOffset);
 
 };
 
+IMPLEMENT_TYPE_LAYOUT(FNiagaraMeshVertexFactoryShaderParametersVS);
+
 class FNiagaraMeshVertexFactoryShaderParametersPS : public FNiagaraVertexFactoryShaderParametersBase
 {
-	DECLARE_INLINE_TYPE_LAYOUT(FNiagaraMeshVertexFactoryShaderParametersPS, NonVirtual);
+	DECLARE_TYPE_LAYOUT(FNiagaraMeshVertexFactoryShaderParametersPS, NonVirtual);
 public:
 	void GetElementShaderBindings(
 		const FSceneInterface* Scene,
@@ -89,6 +91,8 @@ public:
 	
 };
 
+IMPLEMENT_TYPE_LAYOUT(FNiagaraMeshVertexFactoryShaderParametersPS);
+
 void FNiagaraMeshVertexFactory::InitRHI()
 {
 	FVertexDeclarationElementList Elements;
diff --git a/Engine/Plugins/FX/Niagara/Source/NiagaraVertexFactories/Private/NiagaraRibbonVertexFactory.cpp b/Engine/Plugins/FX/Niagara/Source/NiagaraVertexFactories/Private/NiagaraRibbonVertexFactory.cpp
--- a/Engine/Plugins/FX/Niagara/Source/NiagaraVertexFactories/Private/NiagaraRibbonVertexFactory.cpp
+++ b/Engine/Plugins/FX/Niagara/Source/NiagaraVertexFactories/Private/NiagaraRibbonVertexFactory.cpp
@@ -18,7 +18,7 @@ IMPLEMENT_GLOBAL_SHADER_PARAMETER_STRUCT(FNiagaraRibbonVFLooseParameters, "Niaga
 */
 class FNiagaraRibbonVertexFactoryShaderParametersVS : public FNiagaraVertexFactoryShaderParametersBase
 {
-	DECLARE_INLINE_TYPE_LAYOUT(FNiagaraRibbonVertexFactoryShaderParametersVS, NonVirtual);
+	DECLARE_TYPE_LAYOUT(FNiagaraRibbonVertexFactoryShaderParametersVS, NonVirtual);
 public:
 	void Bind(const FShaderParameterMap& ParameterMap)
 	{
@@ -47,13 +47,15 @@ public:
 private:
 };
 
+IMPLEMENT_TYPE_LAYOUT(FNiagaraRibbonVertexFactoryShaderParametersVS);
+
 
 /**
 * Shader parameters for the beam/trail vertex factory.
 */
 class FNiagaraRibbonVertexFactoryShaderParametersPS : public FNiagaraVertexFactoryShaderParametersBase
 {
-	DECLARE_INLINE_TYPE_LAYOUT(FNiagaraRibbonVertexFactoryShaderParametersPS, NonVirtual);
+	DECLARE_TYPE_LAYOUT(FNiagaraRibbonVertexFactoryShaderParametersPS, NonVirtual);
 public:
 	 void GetElementShaderBindings(
 		const FSceneInterface* Scene,
@@ -73,6 +75,8 @@ public:
 	}
 };
 
+IMPLEMENT_TYPE_LAYOUT(FNiagaraRibbonVertexFactoryShaderParametersPS);
+
 
 ///////////////////////////////////////////////////////////////////////////////
 /**
diff --git a/Engine/Plugins/FX/Niagara/Source/NiagaraVertexFactories/Private/NiagaraSpriteVertexFactory.cpp b/Engine/Plugins/FX/Niagara/Source/NiagaraVertexFactories/Private/NiagaraSpriteVertexFactory.cpp
--- a/Engine/Plugins/FX/Niagara/Source/NiagaraVertexFactories/Private/NiagaraSpriteVertexFactory.cpp
+++ b/Engine/Plugins/FX/Niagara/Source/NiagaraVertexFactories/Private/NiagaraSpriteVertexFactory.cpp
@@ -18,7 +18,7 @@ TGlobalResource<FNullDynamicParameterVertexBuffer> GNullNiagaraDynamicParameterV
 
 class FNiagaraSpriteVertexFactoryShaderParametersVS : public FNiagaraVertexFactoryShaderParametersBase
 {
-	DECLARE_INLINE_TYPE_LAYOUT(FNiagaraSpriteVertexFactoryShaderParametersVS, NonVirtual);
+	DECLARE_TYPE_LAYOUT(FNiagaraSpriteVertexFactoryShaderParametersVS, NonVirtual);
 public:
 
 	void Bind(const FShaderParameterMap& ParameterMap)
@@ -28,9 +28,9 @@ public:
 		NumCutoutVerticesPerFrame.Bind(ParameterMap, TEXT("NumCutoutVerticesPerFrame"));
 		CutoutGeometry.Bind(ParameterMap, TEXT("CutoutGeometry"));
 
-//  		NiagaraParticleDataInt.Bind(ParameterMap, TEXT("NiagaraParticleDataInt"));
-//  		Int32DataOffset.Bind(ParameterMap, TEXT("NiagaraInt32DataOffset"));
-//  		Int32DataStride.Bind(ParameterMap, TEXT("NiagaraInt3DataStride"));
+		//  		NiagaraParticleDataInt.Bind(ParameterMap, TEXT("NiagaraParticleDataInt"));
+		//  		Int32DataOffset.Bind(ParameterMap, TEXT("NiagaraInt32DataOffset"));
+		//  		Int32DataStride.Bind(ParameterMap, TEXT("NiagaraInt3DataStride"));
 
 		ParticleAlignmentMode.Bind(ParameterMap, TEXT("ParticleAlignmentMode"));
 		ParticleFacingMode.Bind(ParameterMap, TEXT("ParticleFacingMode"));
@@ -53,10 +53,10 @@ public:
 		FNiagaraVertexFactoryShaderParametersBase::GetElementShaderBindings(Scene, View, Shader, VertexStreamType, FeatureLevel, VertexFactory, BatchElement, ShaderBindings, VertexStreams);
 
 		const FNiagaraSpriteVertexFactory* SpriteVF = static_cast<const FNiagaraSpriteVertexFactory*>(VertexFactory);
-		ShaderBindings.Add(Shader->GetUniformBufferParameter<FNiagaraSpriteUniformParameters>(), SpriteVF->GetSpriteUniformBuffer() );
+		ShaderBindings.Add(Shader->GetUniformBufferParameter<FNiagaraSpriteUniformParameters>(), SpriteVF->GetSpriteUniformBuffer());
 
 		ShaderBindings.Add(Shader->GetUniformBufferParameter<FNiagaraSpriteVFLooseParameters>(), SpriteVF->LooseParameterUniformBuffer);
-		
+
 		ShaderBindings.Add(NumCutoutVerticesPerFrame, SpriteVF->GetNumCutoutVerticesPerFrame());
 		FRHIShaderResourceView* NullSRV = GFNiagaraNullCutoutVertexBuffer.VertexBufferSRV;
 		ShaderBindings.Add(CutoutGeometry, SpriteVF->GetCutoutGeometrySRV() ? SpriteVF->GetCutoutGeometrySRV() : NullSRV);
@@ -69,27 +69,29 @@ public:
 	}
 
 private:
-	
-		LAYOUT_FIELD(FShaderParameter, NumCutoutVerticesPerFrame);
 
-		LAYOUT_FIELD(FShaderParameter, ParticleAlignmentMode);
-		LAYOUT_FIELD(FShaderParameter, ParticleFacingMode);
+	LAYOUT_FIELD(FShaderParameter, NumCutoutVerticesPerFrame);
 
-		LAYOUT_FIELD(FShaderResourceParameter, CutoutGeometry);
+	LAYOUT_FIELD(FShaderParameter, ParticleAlignmentMode);
+	LAYOUT_FIELD(FShaderParameter, ParticleFacingMode);
+
+	LAYOUT_FIELD(FShaderResourceParameter, CutoutGeometry);
 
 
 	//  	LAYOUT_FIELD(FShaderResourceParameter, NiagaraParticleDataInt);
 	//  	LAYOUT_FIELD(FShaderParameter, Int32DataOffset);
 	//  	LAYOUT_FIELD(FShaderParameter, Int32DataStride);
-	
-		LAYOUT_FIELD(FShaderResourceParameter, SortedIndices);
-		LAYOUT_FIELD(FShaderParameter, SortedIndicesOffset);
-	
+
+	LAYOUT_FIELD(FShaderResourceParameter, SortedIndices);
+	LAYOUT_FIELD(FShaderParameter, SortedIndicesOffset);
+
 };
 
+IMPLEMENT_TYPE_LAYOUT(FNiagaraSpriteVertexFactoryShaderParametersVS);
+
 class FNiagaraSpriteVertexFactoryShaderParametersPS : public FNiagaraVertexFactoryShaderParametersBase
 {
-	DECLARE_INLINE_TYPE_LAYOUT(FNiagaraSpriteVertexFactoryShaderParametersPS, NonVirtual);
+	DECLARE_TYPE_LAYOUT(FNiagaraSpriteVertexFactoryShaderParametersPS, NonVirtual);
 public:
 	void GetElementShaderBindings(
 		const FSceneInterface* Scene,
@@ -109,6 +111,8 @@ public:
 	}
 };
 
+IMPLEMENT_TYPE_LAYOUT(FNiagaraSpriteVertexFactoryShaderParametersPS);
+
 /**
  * The particle system vertex declaration resource type.
  */
diff --git a/Engine/Plugins/FX/Niagara/Source/NiagaraVertexFactories/Private/NiagaraVertexFactory.cpp b/Engine/Plugins/FX/Niagara/Source/NiagaraVertexFactories/Private/NiagaraVertexFactory.cpp
--- a/Engine/Plugins/FX/Niagara/Source/NiagaraVertexFactories/Private/NiagaraVertexFactory.cpp
+++ b/Engine/Plugins/FX/Niagara/Source/NiagaraVertexFactories/Private/NiagaraVertexFactory.cpp
@@ -4,6 +4,8 @@
 #include "ShaderParameterUtils.h"
 #include "MeshMaterialShader.h"
 
+IMPLEMENT_TYPE_LAYOUT(FNiagaraVertexFactoryShaderParametersBase);
+
 void FNiagaraVertexFactoryShaderParametersBase::Bind(const FShaderParameterMap& ParameterMap)
 {
 
diff --git a/Engine/Plugins/FX/Niagara/Source/NiagaraVertexFactories/Public/NiagaraVertexFactory.h b/Engine/Plugins/FX/Niagara/Source/NiagaraVertexFactories/Public/NiagaraVertexFactory.h
--- a/Engine/Plugins/FX/Niagara/Source/NiagaraVertexFactories/Public/NiagaraVertexFactory.h
+++ b/Engine/Plugins/FX/Niagara/Source/NiagaraVertexFactories/Public/NiagaraVertexFactory.h
@@ -137,7 +137,7 @@ private:
 */
 class FNiagaraVertexFactoryShaderParametersBase : public FVertexFactoryShaderParameters
 {
-	DECLARE_INLINE_TYPE_LAYOUT(FNiagaraVertexFactoryShaderParametersBase, NonVirtual);
+	DECLARE_TYPE_LAYOUT(FNiagaraVertexFactoryShaderParametersBase, NonVirtual);
 
 public:
 	void Bind(const FShaderParameterMap& ParameterMap);
diff --git a/Engine/Shaders/Public/ShaderVersion.ush b/Engine/Shaders/Public/ShaderVersion.ush
--- a/Engine/Shaders/Public/ShaderVersion.ush
+++ b/Engine/Shaders/Public/ShaderVersion.ush
@@ -3,4 +3,4 @@
 // in Platform.ush (which should be included in any shader) it allows to invalidate the shader DDC.
 // 
 // If you are merging streams and there is a conflict with this GUID you should make a new GUID rather than taking one or the other.
-// GUID = 6988BCD6-722E-4AB1-A73E-2ABF58B306CF
+// GUID = ECD05EA2-D8F5-4874-B259-8109E2FEDEA9
diff --git a/Engine/Source/Runtime/Core/Private/Containers/MemoryImageAllocator.cpp b/Engine/Source/Runtime/Core/Private/Containers/MemoryImageAllocator.cpp
--- a/Engine/Source/Runtime/Core/Private/Containers/MemoryImageAllocator.cpp
+++ b/Engine/Source/Runtime/Core/Private/Containers/MemoryImageAllocator.cpp
@@ -55,14 +55,13 @@ void FMemoryImageAllocatorBase::WriteMemoryImage(FMemoryImageWriter& Writer, con
 	{
 		const void* RawPtr = GetAllocation();
 		check(RawPtr);
-		FMemoryImageWriter ArrayWriter = Writer.WritePointer(FString::Printf(TEXT("FMemoryImageAllocator<%s>"), TypeDesc.Name));
-		ArrayWriter.AddDependency(TypeDesc);
+		FMemoryImageWriter ArrayWriter = Writer.WritePointer(TypeDesc);
 		ArrayWriter.WriteAlignment(Alignment);
 		ArrayWriter.WriteObjectArray(RawPtr, TypeDesc, NumAllocatedElements);
 	}
 	else
 	{
-		Writer.WriteMemoryImagePointerSizedBytes(0u);
+		Writer.WriteNullPointer();
 	}
 }
 
@@ -70,23 +69,28 @@ void FMemoryImageAllocatorBase::CopyUnfrozen(const FMemoryUnfreezeContent& Conte
 {
 	if (NumAllocatedElements > 0)
 	{
-		const void* RawPtr = GetAllocation();
+		const uint8* FrozenObject = (uint8*)GetAllocation();
+		uint8* UnfrozenDst = (uint8*)OutDst;
+
 		FTypeLayoutDesc::FUnfrozenCopyFunc* Func = TypeDesc.UnfrozenCopyFunc;
 		const uint32 ElementSize = TypeDesc.Size;
+		const uint32 ElementAlign = TypeDesc.Alignment;
+		const uint32 FrozenElementAlign = Freeze::GetTargetAlignment(TypeDesc, Context.FrozenLayoutParameters);
+		uint32 FrozenOffset = 0u;
+		uint32 UnfrozenOffset = 0u;
 
 		for (int32 i = 0; i < NumAllocatedElements; ++i)
 		{
-			Func(Context,
-				(uint8*)RawPtr + ElementSize * i,
-				TypeDesc,
-				(uint8*)OutDst + ElementSize * i);
+			const uint32 FrozenElementSize = Func(Context, FrozenObject + FrozenOffset, TypeDesc, UnfrozenDst + UnfrozenOffset);
+			FrozenOffset = Align(FrozenOffset + FrozenElementSize, FrozenElementAlign);
+			UnfrozenOffset = Align(UnfrozenOffset + ElementSize, ElementAlign);
 		}
 	}
 }
 
-void FMemoryImageAllocatorBase::ToString(const FTypeLayoutDesc& TypeDesc, int32 NumAllocatedElements, const FPlatformTypeLayoutParameters& LayoutParams, FMemoryToStringContext& OutContext) const
+void FMemoryImageAllocatorBase::ToString(const FTypeLayoutDesc& TypeDesc, int32 NumAllocatedElements, int32 MaxAllocatedElements, const FPlatformTypeLayoutParameters& LayoutParams, FMemoryToStringContext& OutContext) const
 {
-	OutContext.String->Appendf(TEXT("TArray<%s>, Num: %d\n"), TypeDesc.Name, NumAllocatedElements);
+	OutContext.String->Appendf(TEXT("TArray<%s>, Num: %d, Max: %d\n"), TypeDesc.Name, NumAllocatedElements, MaxAllocatedElements);
 	++OutContext.Indent;
 
 	const void* RawPtr = GetAllocation();
diff --git a/Engine/Source/Runtime/Core/Private/Serialization/MemoryImage.cpp b/Engine/Source/Runtime/Core/Private/Serialization/MemoryImage.cpp
--- a/Engine/Source/Runtime/Core/Private/Serialization/MemoryImage.cpp
+++ b/Engine/Source/Runtime/Core/Private/Serialization/MemoryImage.cpp
@@ -64,7 +64,6 @@ void FPlatformTypeLayoutParameters::InitializeForPlatform(const FString& Platfor
 	if (bHasEditorOnlyData) Flags |= Flag_WithEditorOnly;
 	if (PlatformInfo.Freezing_bWithRayTracing) Flags |= Flag_WithRaytracing;
 	if (PlatformInfo.Freezing_b32Bit) Flags |= Flag_Is32Bit;
-	if (PlatformInfo.Freezing_bForce64BitMemoryImagePointers) Flags |= Flag_Force64BitMemoryImagePointers;
 	if (PlatformInfo.Freezing_bAlignBases) Flags |= Flag_AlignBases;
 
 	MaxFieldAlignment = PlatformInfo.Freezing_MaxFieldAlignment;
@@ -76,10 +75,8 @@ void FPlatformTypeLayoutParameters::InitializeForCurrent()
 	if (WITH_EDITORONLY_DATA) Flags |= Flag_WithEditorOnly;
 	if (WITH_RAYTRACING) Flags |= Flag_WithRaytracing;
 	if (PLATFORM_32BITS) Flags |= Flag_Is32Bit;
-	if (UE_FORCE_64BIT_MEMORY_IMAGE_POINTERS) Flags |= Flag_Force64BitMemoryImagePointers;
 
 	check(GetRawPointerSize() == sizeof(void*));
-	check(GetMemoryImagePointerSize() == sizeof(FMemoryImagePtrInt));
 
 	// clang for Windows matches the MSVC ABI
 #if defined(__clang__) && !PLATFORM_WINDOWS
@@ -111,10 +108,6 @@ FArchive& FPlatformTypeLayoutParameters::Serialize(FArchive& Ar)
 
 void FPlatformTypeLayoutParameters::AppendKeyString(FString& KeyString) const
 {
-	if (Is32Bit() && HasForce64BitMemoryImagePointers())
-	{
-		KeyString += TEXT("FIX_");
-	}
 }
 
 // evaluated during static-initialization, so logging from regular check() macros won't work correctly
@@ -329,8 +322,8 @@ uint32 FTypeLayoutDesc::GetOffsetToBase(const FTypeLayoutDesc& BaseTypeDesc) con
 
 bool FTypeLayoutDesc::IsDerivedFrom(const FTypeLayoutDesc& BaseTypeDesc) const
 {
-	uint32 Offset = 0u;
-	return TryGetOffsetToBase(*this, BaseTypeDesc, Offset);
+	uint32 UnusedOffset = 0u;
+	return TryGetOffsetToBase(*this, BaseTypeDesc, UnusedOffset);
 }
 
 void Freeze::ExtractBitFieldValue(const void* Value, uint32 SrcBitOffset, uint32 DestBitOffset, uint32 NumBits, uint64& InOutValue)
@@ -377,12 +370,6 @@ void Freeze::DefaultWriteMemoryImage(FMemoryImageWriter& Writer, const void* Obj
 {
 	const FPlatformTypeLayoutParameters& TargetLayoutParams = Writer.GetTargetLayoutParams();
 
-	if (TypeDesc.NameHash == FHashedName(TEXT("FGlobalShaderMapContent")).GetHash())
-	{
-		int a = 0;
-	}
-
-
 	// VTable will be shared with any base class vtable, if present
 	if (ETypeLayoutInterface::HasVTable(TypeDesc.Interface) && TypeDesc.NumVirtualBases == 0u)
 	{
@@ -549,11 +536,6 @@ uint32 Freeze::DefaultAppendHash(const FTypeLayoutDesc& TypeLayout, const FPlatf
 		return TypeLayout.SizeFromFields;
 	}
 
-	if (TypeLayout.NameHash == FHashedName(TEXT("FShader")).GetHash())
-	{
-		int a = 0;
-	}
-
 	const FTypeLayoutDesc* CurrentBitFieldType = nullptr;
 	uint32 CurrentNumBits = 0u;
 	uint32 Offset = 0u;
@@ -644,10 +626,9 @@ uint32 Freeze::DefaultAppendHash(const FTypeLayoutDesc& TypeLayout, const FPlatf
 
 				Hasher.Update((uint8*)&Field->BitFieldSize, sizeof(Field->BitFieldSize));
 			}
-
-			++FieldIndex;
 		}
 
+		++FieldIndex;
 		Field = Field->Next;
 	}
 
@@ -691,6 +672,16 @@ void FMemoryToStringContext::AppendNullptr()
 	String->Append(TEXT("nullptr\n"));
 }
 
+void FMemoryToStringContext::AppendUnfrozenPointer(const FTypeLayoutDesc& StaticTypeDesc)
+{
+	String->Appendf(TEXT("TMemoryImagePtr<%s> (Unfrozen), Target: "), StaticTypeDesc.Name);
+}
+
+void FMemoryToStringContext::AppendFrozenPointer(const FTypeLayoutDesc& StaticTypeDesc, int32 FrozenTypeIndex)
+{
+	String->Appendf(TEXT("TMemoryImagePtr<%s>, FrozenTypeIndex: %d, Target: "), StaticTypeDesc.Name, FrozenTypeIndex);
+}
+
 void FMemoryToStringContext::AppendIndent()
 {
 	for (int32 i = 0; i < Indent; ++i)
@@ -867,35 +858,102 @@ uint32 Freeze::HashLayouts(const TArray<const FTypeLayoutDesc*>& TypeLayouts, co
 	return Size;
 }
 
-void Freeze::DefaultUnfrozenCopy(const FMemoryUnfreezeContent& Context, const void* Object, const FTypeLayoutDesc& TypeDesc, void* OutDst)
+uint32 Freeze::DefaultUnfrozenCopy(const FMemoryUnfreezeContent& Context, const void* FrozenObject, const FTypeLayoutDesc& TypeDesc, void* OutDst)
 {
+	uint32 FrozenOffset = 0u;
+	uint32 FieldIndex = 0u;
+	uint32 NumEmptyBases = 0u;
+	uint32 PrevOffset = ~0u;
+
 	if (ETypeLayoutInterface::HasVTable(TypeDesc.Interface) && TypeDesc.NumVirtualBases == 0u)
 	{
-		// Copy v-table
-		FMemory::Memcpy(OutDst, Object, sizeof(void*));
+		// Copy v-table, TODO
+		FMemory::Memcpy(OutDst, FrozenObject, sizeof(void*));
+		FrozenOffset += Context.FrozenLayoutParameters.GetRawPointerSize();
 	}
 
 	const FFieldLayoutDesc* FieldDesc = TypeDesc.Fields;
 	while (FieldDesc)
 	{
-		if (FieldDesc->BitFieldSize == 0u || FieldDesc->Offset != ~0u)
-		{
-			const FTypeLayoutDesc& FieldType = *FieldDesc->Type;
-			FTypeLayoutDesc::FUnfrozenCopyFunc* Func = FieldType.UnfrozenCopyFunc;
-			const uint32 FieldTypeSize = FieldType.Size;
+		const FTypeLayoutDesc& FieldType = *FieldDesc->Type;
+		const uint32 FieldTypeSize = FieldType.Size;
+		uint8* FieldDst = (uint8*)OutDst + FieldDesc->Offset;
 
-			const uint8* FieldObject = (uint8*)Object + FieldDesc->Offset;
-			uint8* FieldDst = (uint8*)OutDst + FieldDesc->Offset;
-			for (uint32 i = 0u; i < FieldDesc->NumArray; ++i)
+		if (IncludeField(FieldDesc, Context.FrozenLayoutParameters))
+		{
+			if (FieldDesc->BitFieldSize == 0u || FieldDesc->Offset != PrevOffset)
 			{
-				Func(Context, FieldObject, FieldType, FieldDst);
+				const bool bIsBase = (FieldIndex < TypeDesc.NumBases);
+				const uint32 FieldTypeAlignment = GetTargetAlignment(FieldType, Context.FrozenLayoutParameters);
+				const uint32 FieldAlignment = FMath::Min(FieldTypeAlignment, Context.FrozenLayoutParameters.MaxFieldAlignment);
+
+				FrozenOffset = Align(FrozenOffset, FieldAlignment);
+
+				FTypeLayoutDesc::FUnfrozenCopyFunc* Func = FieldType.UnfrozenCopyFunc;
+				const uint8* FrozenFieldObject = (uint8*)FrozenObject + FrozenOffset;
+
+				uint32 PaddedFrozenFieldSize = Func(Context, FrozenFieldObject, FieldType, FieldDst);
+				if (PaddedFrozenFieldSize == 0u && bIsBase)
+				{
+					if (NumEmptyBases > 0u)
+					{
+						PaddedFrozenFieldSize = 1u;
+					}
+					else
+					{
+						// Empty bases are allowed to have Offset 0, since they logically overlap
+						++NumEmptyBases;
+					}
+				}
+				else
+				{
+					if (PaddedFrozenFieldSize == 0u)
+					{
+						PaddedFrozenFieldSize = 1u;
+					}
+				}
+
+				if (PaddedFrozenFieldSize > 0u)
+				{
+					if (!bIsBase || Context.FrozenLayoutParameters.HasAlignBases())
+					{
+						PaddedFrozenFieldSize = Align(PaddedFrozenFieldSize, FieldTypeAlignment);
+					}
+				}
+
 				FieldDst += FieldTypeSize;
-				FieldObject += FieldTypeSize;
+				FrozenFieldObject += PaddedFrozenFieldSize;
+				FrozenOffset += PaddedFrozenFieldSize;
+
+				for (uint32 i = 1u; i < FieldDesc->NumArray; ++i)
+				{
+					Func(Context, FrozenFieldObject, FieldType, FieldDst);
+					FieldDst += FieldTypeSize;
+					FrozenFieldObject += PaddedFrozenFieldSize;
+					FrozenOffset += PaddedFrozenFieldSize;
+				}
+			}
+		}
+		else
+		{
+			const void* DefaultObject = FieldType.GetDefaultObjectFunc ? FieldType.GetDefaultObjectFunc() : nullptr;
+			if (DefaultObject)
+			{
+				FMemory::Memcpy(FieldDst, DefaultObject, FieldTypeSize);
+			}
+			else
+			{
+				// If no 'GetDefaultObjectFunc', assume default object is simply zero-initialized
+				FMemory::Memzero(FieldDst, FieldTypeSize);
 			}
 		}
 
+		++FieldIndex;
+		PrevOffset = FieldDesc->Offset;
 		FieldDesc = FieldDesc->Next;
 	}
+
+	return FrozenOffset;
 }
 
 void Freeze::IntrinsicWriteMemoryImage(FMemoryImageWriter& Writer, const void* Object, uint32 Size)
@@ -1032,7 +1090,7 @@ public:
 void Freeze::IntrinsicWriteMemoryImage(FMemoryImageWriter& Writer, const FHashedNameDebugString& Object, const FTypeLayoutDesc&)
 {
 	const char* Data = Object.String.Get();
-	FMemoryImageWriter StringWriter = Writer.WritePointer("String");
+	FMemoryImageWriter StringWriter = Writer.WritePointer(StaticGetTypeLayoutDesc<char>());
 	if (Data)
 	{
 		const int32 Length = FCStringAnsi::Strlen(Data);
@@ -1044,11 +1102,12 @@ void Freeze::IntrinsicWriteMemoryImage(FMemoryImageWriter& Writer, const FHashed
 	}
 }
 
-void Freeze::IntrinsicUnfrozenCopy(const FMemoryUnfreezeContent& Context, const FHashedNameDebugString& Object, void* OutDst)
+uint32 Freeze::IntrinsicUnfrozenCopy(const FMemoryUnfreezeContent& Context, const FHashedNameDebugString& Object, void* OutDst)
 {
 	const FName Name(Object.String.Get());
 	const FHashedName HashedName(Name);
 	FHashedNameDebugString* Result = new(OutDst) FHashedNameDebugString(HashedName.GetDebugString());
+	return sizeof(Object);
 }
 
 #endif // WITH_EDITORONLY_DATA
@@ -1120,6 +1179,118 @@ FHashedName::FHashedName(const FName& InName)
 	}
 }
 
+void FPointerTableBase::SaveToArchive(FArchive& Ar, const FPlatformTypeLayoutParameters& LayoutParams, const void* FrozenObject) const
+{
+	int32 NumDependencies = 0;
+#if UE_MEMORYIMAGE_TRACK_TYPE_DEPENDENCIES
+	NumDependencies = TypeDependencies.Num();
+	Ar << NumDependencies;
+	for (int32 i = 0; i < TypeDependencies.Num(); ++i)
+	{
+		const FTypeLayoutDesc* DependencyTypeDesc = TypeDependencies[i];
+		uint64 NameHash = DependencyTypeDesc->NameHash;
+		FSHAHash LayoutHash;
+		uint32 LayoutSize = Freeze::HashLayout(*DependencyTypeDesc, LayoutParams, LayoutHash);
+		Ar << NameHash;
+		Ar << LayoutSize;
+		Ar << LayoutHash;
+	}
+#else
+	Ar << NumDependencies;
+#endif
+}
+
+bool FPointerTableBase::LoadFromArchive(FArchive& Ar, const FPlatformTypeLayoutParameters& LayoutParams, void* FrozenObject)
+{
+	int32 NumDependencies = 0;
+	Ar << NumDependencies;
+#if UE_MEMORYIMAGE_TRACK_TYPE_DEPENDENCIES
+	TypeDependencies.Empty(NumDependencies);
+#endif
+
+	// Waste a bit of time even in shipping builds skipping over this stuff
+	// Could add a cook-time option to exclude dependencies completely
+	bool bValid = true;
+	for (int32 i = 0u; i < NumDependencies; ++i)
+	{
+		uint64 NameHash = 0u;
+		uint32 SavedLayoutSize = 0u;
+		FSHAHash SavedLayoutHash;
+		Ar << NameHash;
+		Ar << SavedLayoutSize;
+		Ar << SavedLayoutHash;
+
+#if UE_MEMORYIMAGE_TRACK_TYPE_DEPENDENCIES
+		const FTypeLayoutDesc* DependencyType = FTypeLayoutDesc::Find(NameHash);
+		TypeDependencies.Add(DependencyType);
+		if (DependencyType)
+		{
+			FSHAHash CheckLayoutHash;
+			const uint32 CheckLayoutSize = Freeze::HashLayout(*DependencyType, LayoutParams, CheckLayoutHash);
+			if (CheckLayoutSize != SavedLayoutSize)
+			{
+				UE_LOG(LogMemoryImage, Error, TEXT("Mismatch size for type %s, compiled size is %d, loaded size is %d"), DependencyType->Name, CheckLayoutSize, SavedLayoutSize);
+				bValid = false;
+			}
+			else if (CheckLayoutHash != SavedLayoutHash)
+			{
+				UE_LOG(LogMemoryImage, Error, TEXT("Mismatch hash for type %s"), DependencyType->Name);
+				bValid = false;
+			}
+		}
+#endif // UE_MEMORYIMAGE_TRACK_TYPE_DEPENDENCIES
+	}
+
+	return bValid;
+}
+
+#if UE_MEMORYIMAGE_TRACK_TYPE_DEPENDENCIES
+int32 FPointerTableBase::AddTypeDependency(const FTypeLayoutDesc& TypeDesc)
+{
+	if (TypeDesc.NameHash != 0u)
+	{
+		for (int32 Index = 0; Index < TypeDependencies.Num(); ++Index)
+		{
+			if (TypeDependencies[Index]->NameHash == TypeDesc.NameHash)
+			{
+				return Index;
+			}
+		}
+
+		return TypeDependencies.Add(&TypeDesc);
+	}
+	return INDEX_NONE;
+}
+#endif // UE_MEMORYIMAGE_TRACK_TYPE_DEPENDENCIES
+
+FMemoryImageObject FreezeMemoryImageObject(const void* Object, const FTypeLayoutDesc& TypeDesc, FPointerTableBase* PointerTable)
+{
+	FMemoryImage MemoryImage;
+	MemoryImage.TargetLayoutParameters.InitializeForCurrent();
+	MemoryImage.PointerTable = PointerTable;
+	FMemoryImageWriter Writer(MemoryImage);
+
+	Writer.WriteRootObject(Object, TypeDesc);
+
+	FMemoryImageResult MemoryImageResult;
+	MemoryImage.Flatten(MemoryImageResult, true);
+
+	const uint32 FrozenSize = MemoryImageResult.Bytes.Num();
+	check(FrozenSize > 0u);
+	void* FrozenObject = FMemory::Malloc(FrozenSize);
+	FMemory::Memcpy(FrozenObject, MemoryImageResult.Bytes.GetData(), FrozenSize);
+	MemoryImageResult.ApplyPatches(FrozenObject);
+	return FMemoryImageObject(TypeDesc, FrozenObject, FrozenSize);
+}
+
+void* UnfreezeMemoryImageObject(const void* FrozenObject, const FTypeLayoutDesc& TypeDesc, const FPointerTableBase* PointerTable)
+{
+	void* UnfrozenMemory = ::operator new(TypeDesc.Size);
+	FMemoryUnfreezeContent Context(PointerTable);
+	Context.UnfreezeObject(FrozenObject, TypeDesc, UnfrozenMemory);
+	return UnfrozenMemory;
+}
+
 static void CountNumNames(const TArray<FMemoryImageNamePointer>& Names, TArray<uint32>& OutNameCounts)
 {
 	FName CurrentName;
@@ -1166,6 +1337,16 @@ static void SerializeNames(const TArray<FMemoryImageNamePointer>& Names, const T
 
 void FMemoryImageResult::SaveToArchive(FArchive& Ar) const
 {
+	uint32 FrozenSize = Bytes.Num();
+	uint8* FrozenObject = const_cast<uint8*>(Bytes.GetData());
+
+	Ar << const_cast<FMemoryImageResult*>(this)->TargetLayoutParameters;
+
+	Ar << FrozenSize;
+	Ar.Serialize(FrozenObject, FrozenSize);
+
+	PointerTable->SaveToArchive(Ar, TargetLayoutParameters, FrozenObject);
+
 	TArray<uint32> VTableCounts;
 	{
 		uint64 CurrentTypeNameHash = 0u;
@@ -1267,9 +1448,19 @@ void FMemoryImageResult::ApplyPatches(void* FrozenObject) const
 	}
 }
 
-void FMemoryImageResult::ApplyPatchesFromArchive(void* FrozenObject, FArchive& Ar)
+FMemoryImageObject FMemoryImageResult::LoadFromArchive(FArchive& Ar, const FTypeLayoutDesc& TypeDesc, FPointerTableBase* PointerTable)
 {
-	SCOPED_LOADTIMER(FMemoryImageResult_ApplyPatchesFromArchive);
+	SCOPED_LOADTIMER(FMemoryImageResult_LoadFromArchive);
+
+	FPlatformTypeLayoutParameters LayoutParameters;
+	Ar << LayoutParameters;
+
+	uint32 FrozenSize = 0u;
+	Ar << FrozenSize;
+	void* FrozenObject = FMemory::Malloc(FrozenSize);
+	Ar.Serialize(FrozenObject, FrozenSize);
+
+	const bool bFrozenObjectValid = PointerTable->LoadFromArchive(Ar, LayoutParameters, FrozenObject);
 
 	uint32 NumVTables = 0u;
 	uint32 NumScriptNames = 0u;
@@ -1327,6 +1518,31 @@ void FMemoryImageResult::ApplyPatchesFromArchive(void* FrozenObject, FArchive& A
 			ApplyMinimalNamePatch(FrozenObject, NameToMinimalName(Name), Offset);
 		}
 	}
+
+	if (!bFrozenObjectValid)
+	{
+		FMemory::Free(FrozenObject);
+		return FMemoryImageObject();
+	}
+	else if (LayoutParameters.IsCurrentPlatform())
+	{
+		// Loading content that matches the current platform, just use frozen content directly (expected case)
+		return FMemoryImageObject(TypeDesc, FrozenObject, FrozenSize);
+	}
+	else
+	{
+		// Loading content that was frozen for a different platform/configuration, need to convert
+		// This path is only supported if we're tracking type dependencies, otherwise we won't have any TypeDependencies stored in the pointer table, which is required for conversion
+		void* UnfrozenObject = nullptr;
+#if UE_MEMORYIMAGE_TRACK_TYPE_DEPENDENCIES
+		UnfrozenObject = ::operator new(TypeDesc.Size);
+
+		FMemoryUnfreezeContent UnfreezeContext(PointerTable, LayoutParameters);
+		UnfreezeContext.UnfreezeObject(FrozenObject, TypeDesc, UnfrozenObject);
+#endif // UE_MEMORYIMAGE_TRACK_TYPE_DEPENDENCIES
+		FMemory::Free(FrozenObject);
+		return FMemoryImageObject(TypeDesc, UnfrozenObject, 0u);
+	}
 }
 
 void FPtrTableBase::SavePatchesToArchive(FArchive& Ar, uint32 PtrIndex) const
@@ -1367,13 +1583,35 @@ void FPtrTableBase::AddPatchedPointerBase(uint32 PtrIndex, uint64 Offset)
 	++List.NumOffsets;
 }
 
-FMemoryImageSection* FMemoryImageSection::WritePointer(const FString& SectionName, uint32 Offset)
+FMemoryImageSection* FMemoryImageSection::WritePointer(const FTypeLayoutDesc& StaticTypeDesc, const FTypeLayoutDesc& DerivedTypeDesc, uint32* OutOffsetToBase)
 {
 	FSectionPointer& SectionPointer = Pointers.AddDefaulted_GetRef();
-	SectionPointer.Offset = WriteMemoryImagePointerSizedBytes(0u); // write dummy value
-	SectionPointer.PointerOffset = Offset;
+
+	const uint32 OffsetToBase = DerivedTypeDesc.GetOffsetToBase(StaticTypeDesc);
+	const int32 TypeDependencyIndex = ParentImage->PointerTable->AddTypeDependency(DerivedTypeDesc);
+	const bool bStaticTypeMatchesDerived = (StaticTypeDesc == DerivedTypeDesc);
+
+	if (OutOffsetToBase)
+	{
+		*OutOffsetToBase = OffsetToBase;
+	}
+
+	checkf(OutOffsetToBase || bStaticTypeMatchesDerived, TEXT("Must consider OffsetToBase if static/derived types are different, %s/%s"), DerivedTypeDesc.Name, StaticTypeDesc.Name);
+	checkf(TypeDependencyIndex != INDEX_NONE || bStaticTypeMatchesDerived,
+		TEXT("Unable to store pointer to derived type %s, different from static type %s\n")
+		TEXT("Ensure UE_MEMORYIMAGE_TRACK_TYPE_DEPENDENCIES is set\n")
+		TEXT("Make sure derived type is not declared using DECLARE_INLINE_TYPE_LAYOUT()"),
+		DerivedTypeDesc.Name, StaticTypeDesc.Name);
+
+	FFrozenMemoryImagePtr FrozenPtr;
+	FrozenPtr.OffsetFromThis = 0; // dummy value
+	// If derived type matches the static type, store INDEX_NONE to indicate the static type may be used
+	FrozenPtr.TypeIndex = bStaticTypeMatchesDerived ? INDEX_NONE : TypeDependencyIndex;
+
+	SectionPointer.Offset = WriteBytes(FrozenPtr);
+	SectionPointer.PointerOffset = OffsetToBase;
 	SectionPointer.SectionIndex = ParentImage->Sections.Num();
-	return ParentImage->AllocateSection(SectionName);
+	return ParentImage->AllocateSection();
 }
 
 uint32 FMemoryImageSection::WriteRawPointerSizedBytes(uint64 PointerValue)
@@ -1388,18 +1626,6 @@ uint32 FMemoryImageSection::WriteRawPointerSizedBytes(uint64 PointerValue)
 	}
 }
 
-uint32 FMemoryImageSection::WriteMemoryImagePointerSizedBytes(uint64 PointerValue)
-{
-	const FPlatformTypeLayoutParameters& TargetParameters = ParentImage->TargetLayoutParameters;
-	if (TargetParameters.Has32BitMemoryImagePointers())
-	{
-		return WriteBytes((uint32)PointerValue);
-	}
-	else
-	{
-		return WriteBytes(PointerValue);
-	}
-}
 
 uint32 FMemoryImageSection::WriteVTable(const FTypeLayoutDesc& TypeDesc, const FTypeLayoutDesc& DerivedTypeDesc)
 {
@@ -1507,18 +1733,6 @@ void FMemoryImageSection::ComputeHash()
 	HashState.GetHash(Hash.Hash);
 }
 
-void FMemoryImage::AddDependency(const FTypeLayoutDesc& TypeDesc)
-{
-	if (TypeDesc.NameHash != 0u)
-	{
-		const int32 SortedIndex = Algo::LowerBoundBy(TypeDependencies, TypeDesc.NameHash, [](const FTypeLayoutDesc* InTypeDesc) { return InTypeDesc->NameHash; });
-		if (SortedIndex >= TypeDependencies.Num() || TypeDependencies[SortedIndex] != &TypeDesc)
-		{
-			TypeDependencies.Insert(&TypeDesc, SortedIndex);
-		}
-	}
-}
-
 void FMemoryImage::Flatten(FMemoryImageResult& OutResult, bool bMergeDuplicateSections)
 {
 	TArray<FMemoryImageSection*> UniqueSections;
@@ -1570,39 +1784,27 @@ void FMemoryImage::Flatten(FMemoryImageResult& OutResult, bool bMergeDuplicateSe
 	for (int32 SectionIndex = 0; SectionIndex < UniqueSections.Num(); ++SectionIndex)
 	{
 		const FMemoryImageSection* Section = UniqueSections[SectionIndex];
-		if (TargetLayoutParameters.Has32BitMemoryImagePointers())
-		{
-			for (const FMemoryImageSection::FSectionPointer& Pointer : Section->Pointers)
-			{
-				const int32 OffsetToPointer = SectionOffset[SectionIndex] + Pointer.Offset;
-				const int32 RemapSectionIndex = SectionIndexRemap[Pointer.SectionIndex];
-				int32* PointerData = (int32*)(OutResult.Bytes.GetData() + OffsetToPointer);
-				check(*PointerData == 0);
-				const int32 OffsetFromPointer = (int32)(SectionOffset[RemapSectionIndex] + Pointer.PointerOffset) - OffsetToPointer;
-				*PointerData = (OffsetFromPointer << 1) | 1;
-			}
-		}
-		else
+		for (const FMemoryImageSection::FSectionPointer& Pointer : Section->Pointers)
 		{
-			for (const FMemoryImageSection::FSectionPointer& Pointer : Section->Pointers)
-			{
-				const int64 OffsetToPointer = SectionOffset[SectionIndex] + Pointer.Offset;
-				const int32 RemapSectionIndex = SectionIndexRemap[Pointer.SectionIndex];
-				int64* PointerData = (int64*)(OutResult.Bytes.GetData() + OffsetToPointer);
-				check(*PointerData == 0);
-				const int64 OffsetFromPointer = (int64)(SectionOffset[RemapSectionIndex] + Pointer.PointerOffset) - OffsetToPointer;
-				*PointerData = (OffsetFromPointer << 1) | 1;
-			}
+			const int32 OffsetToPointer = SectionOffset[SectionIndex] + Pointer.Offset;
+			const int32 RemapSectionIndex = SectionIndexRemap[Pointer.SectionIndex];
+			FFrozenMemoryImagePtr* FrozenPtr = (FFrozenMemoryImagePtr*)(OutResult.Bytes.GetData() + OffsetToPointer);
+			check(FrozenPtr->OffsetFromThis == 0);
+			const int32 OffsetFromPointer = (int32)(SectionOffset[RemapSectionIndex] + Pointer.PointerOffset) - OffsetToPointer;
+			FrozenPtr->OffsetFromThis = (OffsetFromPointer << 1) | 1;
 		}
 	}
 
+	OutResult.PointerTable = PointerTable;
+	OutResult.TargetLayoutParameters = TargetLayoutParameters;
+
 	// Sort to group runs of the same FName/VTable
 	OutResult.VTables.Sort();
 	OutResult.ScriptNames.Sort();
 	OutResult.MinimalNames.Sort();
 }
 
-FMemoryImageWriter::FMemoryImageWriter(FMemoryImage& InImage) : Section(InImage.AllocateSection(TEXT("ROOT"))) {}
+FMemoryImageWriter::FMemoryImageWriter(FMemoryImage& InImage) : Section(InImage.AllocateSection()) {}
 FMemoryImageWriter::FMemoryImageWriter(FMemoryImageSection* InSection) : Section(InSection) {}
 
 FMemoryImageWriter::~FMemoryImageWriter()
@@ -1615,14 +1817,22 @@ const FPlatformTypeLayoutParameters& FMemoryImageWriter::GetTargetLayoutParams()
 FPointerTableBase& FMemoryImageWriter::GetPointerTable() const { return GetImage().GetPointerTable(); }
 const FPointerTableBase* FMemoryImageWriter::TryGetPrevPointerTable() const { return GetImage().PrevPointerTable; }
 
-void FMemoryImageWriter::AddDependency(const FTypeLayoutDesc& TypeDesc)
+int32 FMemoryImageWriter::AddTypeDependency(const FTypeLayoutDesc& TypeDesc)
 {
-	Section->ParentImage->AddDependency(TypeDesc);
+	return GetPointerTable().AddTypeDependency(TypeDesc);
 }
 
 void FMemoryImageWriter::WriteObject(const void* Object, const FTypeLayoutDesc& TypeDesc)
 {
-	AddDependency(TypeDesc);
+	TypeDesc.WriteFrozenMemoryImageFunc(*this, Object, TypeDesc, TypeDesc);
+}
+
+void FMemoryImageWriter::WriteRootObject(const void* Object, const FTypeLayoutDesc& TypeDesc)
+{
+	const int32 RootTypeIndex = AddTypeDependency(TypeDesc);
+#if UE_MEMORYIMAGE_TRACK_TYPE_DEPENDENCIES
+	checkf(RootTypeIndex == 0, TEXT("Can't add root object of type %s, already have a root object"), TypeDesc.Name);
+#endif
 	TypeDesc.WriteFrozenMemoryImageFunc(*this, Object, TypeDesc, TypeDesc);
 }
 
@@ -1661,19 +1871,24 @@ uint32 FMemoryImageWriter::WriteBytes(const void* Data, uint32 Size)
 	return Section->WriteBytes(Data, Size);
 }
 
-FMemoryImageWriter FMemoryImageWriter::WritePointer(const FString& SectionName, uint32 Offset)
+uint32 FMemoryImageWriter::WriteNullPointer()
 {
-	return FMemoryImageWriter(Section->WritePointer(SectionName, Offset));
+	return WriteBytes(FFrozenMemoryImagePtr{ 0, 0u });
 }
 
-uint32 FMemoryImageWriter::WriteRawPointerSizedBytes(uint64 PointerValue)
+FMemoryImageWriter FMemoryImageWriter::WritePointer(const FTypeLayoutDesc& StaticTypeDesc, const FTypeLayoutDesc& DerivedTypeDesc, uint32* OutOffsetToBase)
 {
-	return Section->WriteRawPointerSizedBytes(PointerValue);
+	return FMemoryImageWriter(Section->WritePointer(StaticTypeDesc, DerivedTypeDesc, OutOffsetToBase));
 }
 
-uint32 FMemoryImageWriter::WriteMemoryImagePointerSizedBytes(uint64 PointerValue)
+FMemoryImageWriter FMemoryImageWriter::WritePointer(const FTypeLayoutDesc& TypeDesc)
 {
-	return Section->WriteMemoryImagePointerSizedBytes(PointerValue);
+	return FMemoryImageWriter(Section->WritePointer(TypeDesc, TypeDesc, nullptr));
+}
+
+uint32 FMemoryImageWriter::WriteRawPointerSizedBytes(uint64 PointerValue)
+{
+	return Section->WriteRawPointerSizedBytes(PointerValue);
 }
 
 uint32 FMemoryImageWriter::WriteVTable(const FTypeLayoutDesc& TypeDesc, const FTypeLayoutDesc& DerivedTypeDesc)
@@ -1696,6 +1911,20 @@ uint32 FMemoryImageWriter::WriteFScriptName(const FScriptName& Name)
 	return Section->WriteFScriptName(Name);
 }
 
+const FTypeLayoutDesc* FMemoryUnfreezeContent::GetDerivedTypeDesc(const FTypeLayoutDesc& StaticTypeDesc, int32 TypeIndex) const
+{
+	if (TypeIndex == INDEX_NONE)
+	{
+		// TypeIndex of INDEX_NONE means the derived type was the same as the static type
+		return &StaticTypeDesc;
+	}
+	if (PrevPointerTable)
+	{
+		return PrevPointerTable->GetTypeDependency(TypeIndex);
+	}
+	return nullptr;
+}
+
 // Finds the length of the field name, omitting any _DEPRECATED suffix
 uint8 Freeze::FindFieldNameLength(const TCHAR* Name)
 {
diff --git a/Engine/Source/Runtime/Core/Public/Containers/Array.h b/Engine/Source/Runtime/Core/Public/Containers/Array.h
--- a/Engine/Source/Runtime/Core/Public/Containers/Array.h
+++ b/Engine/Source/Runtime/Core/Public/Containers/Array.h
@@ -2685,7 +2685,7 @@ private:
 		}
 		static void ToString(const FPlatformTypeLayoutParameters& LayoutParams, FMemoryToStringContext& OutContext, const TArray& Object)
 		{
-			Object.AllocatorInstance.ToString(StaticGetTypeLayoutDesc<ElementType>(), Object.ArrayNum, LayoutParams, OutContext);
+			Object.AllocatorInstance.ToString(StaticGetTypeLayoutDesc<ElementType>(), Object.ArrayNum, Object.ArrayMax, LayoutParams, OutContext);
 		}
 	};
 
@@ -3006,26 +3006,23 @@ namespace Freeze
 	}
 
 	template<typename T, typename AllocatorType>
-	void IntrinsicUnfrozenCopy(const FMemoryUnfreezeContent& Context, const TArray<T, AllocatorType>& Object, void* OutDst)
+	uint32 IntrinsicUnfrozenCopy(const FMemoryUnfreezeContent& Context, const TArray<T, AllocatorType>& Object, void* OutDst)
 	{
 		Object.CopyUnfrozen(Context, OutDst);
+		return sizeof(Object);
 	}
 
 	template<typename T, typename AllocatorType>
 	uint32 IntrinsicAppendHash(const TArray<T, AllocatorType>* DummyObject, const FTypeLayoutDesc& TypeDesc, const FPlatformTypeLayoutParameters& LayoutParams, FSHA1& Hasher)
 	{
-		// sizeof(TArray) changes depending on target platform 32bit vs 64bit
-		// For now, calculate the size manually
-		static_assert(sizeof(TArray<T, AllocatorType>) == sizeof(FMemoryImageUPtrInt) + sizeof(int32) + sizeof(int32), "Unexpected TArray size");
-		const uint32 SizeFromFields = LayoutParams.GetMemoryImagePointerSize() + sizeof(int32) + sizeof(int32);
-		return AppendHashForNameAndSize(TypeDesc.Name, SizeFromFields, Hasher);;
+		return AppendHashForNameAndSize(TypeDesc.Name, sizeof(TArray<T, AllocatorType>), Hasher);
 	}
 
 	template<typename T, typename AllocatorType>
 	uint32 IntrinsicGetTargetAlignment(const TArray<T, AllocatorType>* DummyObject, const FTypeLayoutDesc& TypeDesc, const FPlatformTypeLayoutParameters& LayoutParams)
 	{
 		// Assume alignment of array is drive by pointer
-		return FMath::Min(LayoutParams.GetMemoryImagePointerSize(), LayoutParams.MaxFieldAlignment);
+		return FMath::Min(8u, LayoutParams.MaxFieldAlignment);
 	}
 
 	template<typename T, typename AllocatorType>
diff --git a/Engine/Source/Runtime/Core/Public/Containers/HashTable.h b/Engine/Source/Runtime/Core/Public/Containers/HashTable.h
--- a/Engine/Source/Runtime/Core/Public/Containers/HashTable.h
+++ b/Engine/Source/Runtime/Core/Public/Containers/HashTable.h
@@ -591,26 +591,23 @@ namespace Freeze
 	}
 
 	template<typename InAllocator>
-	void IntrinsicUnfrozenCopy(const FMemoryUnfreezeContent& Context, const THashTable<InAllocator>& Object, void* OutDst)
+	uint32 IntrinsicUnfrozenCopy(const FMemoryUnfreezeContent& Context, const THashTable<InAllocator>& Object, void* OutDst)
 	{
 		Object.CopyUnfrozen(Context, OutDst);
+		return sizeof(Object);
 	}
 
 	template<typename InAllocator>
 	uint32 IntrinsicAppendHash(const THashTable<InAllocator>* DummyObject, const FTypeLayoutDesc& TypeDesc, const FPlatformTypeLayoutParameters& LayoutParams, FSHA1& Hasher)
 	{
-		// sizeof(TArray) changes depending on target platform 32bit vs 64bit
-		// For now, calculate the size manually
-		static_assert(sizeof(THashTable<InAllocator>) == sizeof(FMemoryImageUPtrInt) * 2 + sizeof(uint32) * 2, "Unexpected THashTable size");
-		const uint32 SizeFromFields = LayoutParams.GetMemoryImagePointerSize() * 2u + sizeof(uint32) * 2u;
-		return AppendHashForNameAndSize(TypeDesc.Name, SizeFromFields, Hasher);
+		return AppendHashForNameAndSize(TypeDesc.Name, sizeof(THashTable<InAllocator>), Hasher);
 	}
 
 	template<typename InAllocator>
 	uint32 IntrinsicGetTargetAlignment(const THashTable<InAllocator>* DummyObject, const FTypeLayoutDesc& TypeDesc, const FPlatformTypeLayoutParameters& LayoutParams)
 	{
-		// Assume alignment of array is drive by pointer
-		return FMath::Min(LayoutParams.GetMemoryImagePointerSize(), LayoutParams.MaxFieldAlignment);
+		// Assume alignment of hash-table is drive by pointer
+		return FMath::Min(8u, LayoutParams.MaxFieldAlignment);
 	}
 }
 
diff --git a/Engine/Source/Runtime/Core/Public/Containers/Map.h b/Engine/Source/Runtime/Core/Public/Containers/Map.h
--- a/Engine/Source/Runtime/Core/Public/Containers/Map.h
+++ b/Engine/Source/Runtime/Core/Public/Containers/Map.h
@@ -1257,9 +1257,10 @@ namespace Freeze
 	}
 
 	template <typename KeyType, typename ValueType, typename SetAllocator, typename KeyFuncs>
-	void IntrinsicUnfrozenCopy(const FMemoryUnfreezeContent& Context, const TMap<KeyType, ValueType, SetAllocator, KeyFuncs>& Object, void* OutDst)
+	uint32 IntrinsicUnfrozenCopy(const FMemoryUnfreezeContent& Context, const TMap<KeyType, ValueType, SetAllocator, KeyFuncs>& Object, void* OutDst)
 	{
 		Object.CopyUnfrozen(Context, OutDst);
+		return sizeof(Object);
 	}
 
 	template <typename KeyType, typename ValueType, typename SetAllocator, typename KeyFuncs>
@@ -1591,9 +1592,10 @@ namespace Freeze
 	}
 
 	template <typename KeyType, typename ValueType, typename SetAllocator, typename KeyFuncs>
-	void IntrinsicUnfrozenCopy(const FMemoryUnfreezeContent& Context, const TMultiMap<KeyType, ValueType, SetAllocator, KeyFuncs>& Object, void* OutDst)
+	uint32 IntrinsicUnfrozenCopy(const FMemoryUnfreezeContent& Context, const TMultiMap<KeyType, ValueType, SetAllocator, KeyFuncs>& Object, void* OutDst)
 	{
 		Object.CopyUnfrozen(Context, OutDst);
+		return sizeof(Object);
 	}
 
 	template <typename KeyType, typename ValueType, typename SetAllocator, typename KeyFuncs>
diff --git a/Engine/Source/Runtime/Core/Public/Containers/Set.h b/Engine/Source/Runtime/Core/Public/Containers/Set.h
--- a/Engine/Source/Runtime/Core/Public/Containers/Set.h
+++ b/Engine/Source/Runtime/Core/Public/Containers/Set.h
@@ -1647,9 +1647,10 @@ namespace Freeze
 	}
 
 	template<typename ElementType, typename KeyFuncs, typename Allocator>
-	void IntrinsicUnfrozenCopy(const FMemoryUnfreezeContent& Context, const TSet<ElementType, KeyFuncs, Allocator>& Object, void* OutDst)
+	uint32 IntrinsicUnfrozenCopy(const FMemoryUnfreezeContent& Context, const TSet<ElementType, KeyFuncs, Allocator>& Object, void* OutDst)
 	{
 		Object.CopyUnfrozen(Context, OutDst);
+		return sizeof(Object);
 	}
 
 	template<typename ElementType, typename KeyFuncs, typename Allocator>
diff --git a/Engine/Source/Runtime/Core/Public/Containers/SparseArray.h b/Engine/Source/Runtime/Core/Public/Containers/SparseArray.h
--- a/Engine/Source/Runtime/Core/Public/Containers/SparseArray.h
+++ b/Engine/Source/Runtime/Core/Public/Containers/SparseArray.h
@@ -1122,7 +1122,7 @@ private:
 			if (NumElements > 0)
 			{
 				const FTypeLayoutDesc& ElementTypeDesc = StaticGetTypeLayoutDesc<ElementType>();
-				FMemoryImageWriter ArrayWriter = Writer.WritePointer(FString::Printf(TEXT("TSparseArray<%s>"), ElementTypeDesc.Name));
+				FMemoryImageWriter ArrayWriter = Writer.WritePointer(ElementTypeDesc);
 				for (int32 i = 0; i < NumElements; ++i)
 				{
 					const FElementOrFreeListLink& Elem = Object.Data[i];
@@ -1141,7 +1141,7 @@ private:
 			}
 			else
 			{
-				Writer.WriteMemoryImagePointerSizedBytes(0u);
+				Writer.WriteNullPointer();
 			}
 			Writer.WriteBytes(NumElements);
 			Writer.WriteBytes(NumElements);
@@ -1213,9 +1213,10 @@ namespace Freeze
 	}
 
 	template<typename ElementType, typename Allocator>
-	void IntrinsicUnfrozenCopy(const FMemoryUnfreezeContent& Context, const TSparseArray<ElementType, Allocator>& Object, void* OutDst)
+	uint32 IntrinsicUnfrozenCopy(const FMemoryUnfreezeContent& Context, const TSparseArray<ElementType, Allocator>& Object, void* OutDst)
 	{
 		Object.CopyUnfrozen(Context, OutDst);
+		return sizeof(Object);
 	}
 
 	template<typename ElementType, typename Allocator>
diff --git a/Engine/Source/Runtime/Core/Public/Serialization/MemoryImage.h b/Engine/Source/Runtime/Core/Public/Serialization/MemoryImage.h
--- a/Engine/Source/Runtime/Core/Public/Serialization/MemoryImage.h
+++ b/Engine/Source/Runtime/Core/Public/Serialization/MemoryImage.h
@@ -5,6 +5,7 @@
 #include "CoreTypes.h"
 #include "Containers/Array.h"
 #include "Containers/Set.h"
+#include "Containers/HashTable.h"
 #include "Containers/UnrealString.h"
 #include "Misc/SecureHash.h"
 #include "Serialization/Archive.h"
@@ -19,13 +20,118 @@
 class FMemoryImage;
 class FMemoryImageString;
 
-class FPointerTableBase
+// Store type dependencies for loaded frozen memory images
+// This uses a bit of memory and adds small CPU cost to loading, but is required to support unfreezing memory images that target platforms other than the current
+// This is also required for creating frozen memory images
+#define UE_MEMORYIMAGE_TRACK_TYPE_DEPENDENCIES (WITH_EDITORONLY_DATA)
+
+class CORE_API FPointerTableBase
 {
 public:
 	virtual ~FPointerTableBase() {}
 	virtual int32 AddIndexedPointer(const FTypeLayoutDesc& TypeDesc, void* Ptr) = 0;
 	virtual void* GetIndexedPointer(const FTypeLayoutDesc& TypeDesc, uint32 i) const = 0;
+
+	virtual void SaveToArchive(FArchive& Ar, const FPlatformTypeLayoutParameters& LayoutParams, const void* FrozenObject) const;
+	virtual bool LoadFromArchive(FArchive& Ar, const FPlatformTypeLayoutParameters& LayoutParams, void* FrozenObject);
+
+#if UE_MEMORYIMAGE_TRACK_TYPE_DEPENDENCIES
+	int32 AddTypeDependency(const FTypeLayoutDesc& TypeDesc);
+	inline const FTypeLayoutDesc* GetTypeDependency(int32 Index) const { return TypeDependencies[Index]; }
+private:
+	TArray<const FTypeLayoutDesc*> TypeDependencies;
+#else
+	int32 AddTypeDependency(const FTypeLayoutDesc& TypeDesc) { return INDEX_NONE; }
+	inline const FTypeLayoutDesc* GetTypeDependency(int32 Index) const { return nullptr; }
+#endif
+};
+
+template<typename T>
+struct TMemoryImageObject
+{
+	TMemoryImageObject() : TypeDesc(nullptr), Object(nullptr), FrozenSize(0u) {}
+
+	TMemoryImageObject(const FTypeLayoutDesc& InTypeDesc, T* InObject, uint32 InFrozenSize)
+		: TypeDesc(&InTypeDesc)
+		, Object(InObject)
+		, FrozenSize(InFrozenSize)
+	{}
+
+	template<typename TOther>
+	TMemoryImageObject(TOther* InObject)
+		: TypeDesc(InObject ? &GetTypeLayoutDesc(nullptr, *InObject) : nullptr)
+		, Object(InObject)
+		, FrozenSize(0u)
+	{}
+
+	template<typename TOther>
+	explicit TMemoryImageObject(const TMemoryImageObject<TOther>& Rhs)
+		: TypeDesc(Rhs.TypeDesc)
+		, Object(static_cast<T*>(Rhs.Object))
+		, FrozenSize(Rhs.FrozenSize)
+	{}
+
+	void Destroy(const FPointerTableBase* PointerTable);
+
+	// Returns true if the frozen/unfrozen state of the object was changed
+	bool Freeze(FPointerTableBase* PointerTable);
+	bool Unfreeze(const FPointerTableBase* PointerTable);
+
+	const FTypeLayoutDesc* TypeDesc;
+	T* Object;
+	uint32 FrozenSize;
 };
+using FMemoryImageObject = TMemoryImageObject<void>;
+
+CORE_API FMemoryImageObject FreezeMemoryImageObject(const void* Object, const FTypeLayoutDesc& TypeDesc, FPointerTableBase* PointerTable);
+CORE_API void* UnfreezeMemoryImageObject(const void* FrozenObject, const FTypeLayoutDesc& TypeDesc, const FPointerTableBase* PointerTable);
+
+template<typename T>
+inline void TMemoryImageObject<T>::Destroy(const FPointerTableBase* PointerTable)
+{
+	if (Object)
+	{
+		InternalDeleteObjectFromLayout(Object, *TypeDesc, PointerTable, FrozenSize > 0u);
+		if (FrozenSize > 0u)
+		{
+			// InternalDeleteObjectFromLayout will delete unfrozen objects,
+			// but won't free frozen objects, since that's not safe for internal object pointers
+			// Here we are working with a root-level object, so it's safe to free it
+			FMemory::Free(Object);
+		}
+		Object = nullptr;
+		FrozenSize = 0u;
+	}
+}
+
+
+template<typename T>
+inline bool TMemoryImageObject<T>::Freeze(FPointerTableBase* PointerTable)
+{
+	if (FrozenSize == 0u && Object)
+	{
+		const FMemoryImageObject FrozenContent = FreezeMemoryImageObject(Object, *TypeDesc, PointerTable);
+		Destroy(nullptr);
+		Object = static_cast<T*>(FrozenContent.Object);
+		FrozenSize = FrozenContent.FrozenSize;
+		return true;
+	}
+	return false;
+}
+
+template<typename T>
+inline bool TMemoryImageObject<T>::Unfreeze(const FPointerTableBase* PointerTable)
+{
+	if (FrozenSize > 0u)
+	{
+		void* UnfrozenObject = UnfreezeMemoryImageObject(Object, *TypeDesc, PointerTable);
+		Destroy(PointerTable);
+		Object = static_cast<T*>(UnfrozenObject);
+		FrozenSize = 0u;
+		return true;
+	}
+	return false;
+}
 
 struct FMemoryImageVTablePointer
 {
@@ -73,13 +179,15 @@ inline bool operator<(const FMemoryImageNamePointer& Lhs, const FMemoryImageName
 struct FMemoryImageResult
 {
 	TArray<uint8> Bytes;
+	FPointerTableBase* PointerTable = nullptr;
+	FPlatformTypeLayoutParameters TargetLayoutParameters;
 	TArray<FMemoryImageVTablePointer> VTables;
 	TArray<FMemoryImageNamePointer> ScriptNames;
 	TArray<FMemoryImageNamePointer> MinimalNames;
 
 	CORE_API void SaveToArchive(FArchive& Ar) const;
 	CORE_API void ApplyPatches(void* FrozenObject) const;
-	CORE_API static void ApplyPatchesFromArchive(void* FrozenObject, FArchive& Ar);
+	CORE_API static FMemoryImageObject LoadFromArchive(FArchive& Ar, const FTypeLayoutDesc& TypeDesc, FPointerTableBase* PointerTable);
 };
 
 class CORE_API FMemoryImageSection : public FRefCountedObject
@@ -92,9 +200,8 @@ public:
 		uint32 Offset;
 	};
 
-	FMemoryImageSection(FMemoryImage* InImage, FString InName)
+	FMemoryImageSection(FMemoryImage* InImage)
 		: ParentImage(InImage)
-		, DebugName(InName)
 		, MaxAlignment(1u)
 	{}
 
@@ -126,8 +233,7 @@ public:
 	template<typename T>
 	uint32 WriteBytes(const T& Data) { return WriteBytes(&Data, sizeof(T)); }
 
-	FMemoryImageSection* WritePointer(const FString& SectionName, uint32 Offset = 0u);
-	uint32 WriteMemoryImagePointerSizedBytes(uint64 PointerValue);
+	FMemoryImageSection* WritePointer(const FTypeLayoutDesc& StaticTypeDesc, const FTypeLayoutDesc& DerivedTypeDesc, uint32* OutOffsetToBase = nullptr);
 	uint32 WriteRawPointerSizedBytes(uint64 PointerValue);
 	uint32 WriteVTable(const FTypeLayoutDesc& TypeDesc, const FTypeLayoutDesc& DerivedTypeDesc);
 	uint32 WriteFName(const FName& Name);
@@ -138,7 +244,6 @@ public:
 	void ComputeHash();
 
 	FMemoryImage* ParentImage;
-	FString DebugName;
 	TArray<uint8> Bytes;
 	TArray<FSectionPointer> Pointers;
 	TArray<FMemoryImageVTablePointer> VTables;
@@ -162,22 +267,19 @@ public:
 	FPointerTableBase& GetPointerTable() const { check(PointerTable); return *PointerTable; }
 	const FPointerTableBase& GetPrevPointerTable() const { check(PrevPointerTable); return *PrevPointerTable; }
 
-	FMemoryImageSection* AllocateSection(const FString& Name)
+	FMemoryImageSection* AllocateSection()
 	{
-		FMemoryImageSection* Section = new FMemoryImageSection(this, Name);
+		FMemoryImageSection* Section = new FMemoryImageSection(this);
 		Sections.Add(Section);
 		return Section;
 	}
 
-	void AddDependency(const FTypeLayoutDesc& TypeDesc);
-
 	/** Merging duplicate sections will make the resulting memory image smaller.
 	 * This will only work for data that is expected to be read-only after freezing.  Merging sections will break any manual fix-ups applied to the frozen data
 	 */
 	void Flatten(FMemoryImageResult& OutResult, bool bMergeDuplicateSections = false);
 
 	TArray<TRefCountPtr<FMemoryImageSection>> Sections;
-	TArray<const FTypeLayoutDesc*> TypeDependencies;
 	FPointerTableBase* PointerTable;
 	const FPointerTableBase* PrevPointerTable;
 	FPlatformTypeLayoutParameters HostLayoutParameters;
@@ -185,29 +287,36 @@ public:
 	const class UStruct* CurrentStruct;
 };
 
+struct FFrozenMemoryImagePtr
+{
+	int32 OffsetFromThis;
+	int32 TypeIndex; // store 0-based index or INDEX_NONE, don't really need 32bit for this, but nothing else to store here for now
+};
+
 template<typename T>
 class TMemoryImagePtr
 {
 public:
-	inline bool IsFrozen() const { return OffsetFromThis & IsFrozenMask; }
-	inline bool IsValid() const { return Ptr != nullptr; }
-	inline bool IsNull() const { return Ptr == nullptr; }
+	inline bool IsFrozen() const { return Frozen.OffsetFromThis & IsFrozenMask; }
+	inline bool IsValid() const { return UnfrozenPtr != nullptr; }
+	inline bool IsNull() const { return UnfrozenPtr == nullptr; }
 
-	inline TMemoryImagePtr(T* InPtr = nullptr) : Ptr(InPtr) {}
-	inline TMemoryImagePtr(const TMemoryImagePtr<T>& InPtr) : Ptr(InPtr.Get()) {}
-	inline TMemoryImagePtr& operator=(T* InPtr) { Ptr = InPtr; check((OffsetFromThis & AllFlags) == 0u); return *this; }
-	inline TMemoryImagePtr& operator=(const TMemoryImagePtr<T>& InPtr) { Ptr = InPtr.Get(); check((OffsetFromThis & AllFlags) == 0u); return *this; }
+	inline TMemoryImagePtr(T* InPtr = nullptr) : UnfrozenPtr(InPtr) {}
+	inline TMemoryImagePtr(const TMemoryImagePtr<T>& InPtr) : UnfrozenPtr(InPtr.Get()) {}
+	inline TMemoryImagePtr& operator=(T* InPtr) { UnfrozenPtr = InPtr; check((Frozen.OffsetFromThis & AllFlags) == 0u); return *this; }
+	inline TMemoryImagePtr& operator=(const TMemoryImagePtr<T>& InPtr) { UnfrozenPtr = InPtr.Get(); check((Frozen.OffsetFromThis & AllFlags) == 0u); return *this; }
 
 	inline ~TMemoryImagePtr() 
 	{
 
 	}
 
-	inline FMemoryImagePtrInt GetFrozenOffsetFromThis() const { check(IsFrozen()); return (OffsetFromThis >> OffsetShift); }
+	inline int32 GetFrozenOffsetFromThis() const { check(IsFrozen()); return (Frozen.OffsetFromThis >> OffsetShift); }
+	inline uint32 GetFrozenTypeIndex() const { check(IsFrozen()); return Frozen.TypeIndex; }
 
 	inline T* Get() const
 	{
-		return IsFrozen() ? GetFrozenPtrInternal() : Ptr;
+		return IsFrozen() ? GetFrozenPtrInternal() : UnfrozenPtr;
 	}
 
 	inline T* GetChecked() const { T* Value = Get(); check(Value); return Value; }
@@ -221,7 +330,7 @@ public:
 		if (RawPtr)
 		{
 			DeleteObjectFromLayout(RawPtr, PtrTable, IsFrozen());
-			Ptr = nullptr;
+			UnfrozenPtr = nullptr;
 		}
 	}
 
@@ -231,24 +340,20 @@ public:
 		if (RawPtr)
 		{
 			check(DerivedTypeDesc);
-			// Compile-time type of the thing we're pointing to
-			const FTypeLayoutDesc& StaticTypeDesc = StaticGetTypeLayoutDesc<T>();
-			// 'this' offset to adjust from the compile-time type to the run-time type
-			const uint32 OffsetToBase = DerivedTypeDesc->GetOffsetToBase(StaticTypeDesc);
-
-			FMemoryImageWriter PointerWriter = Writer.WritePointer(FString::Printf(TEXT("TMemoryImagePtr<%s>"), DerivedTypeDesc->Name), OffsetToBase);
+			uint32 OffsetToBase = 0u;
+			FMemoryImageWriter PointerWriter = Writer.WritePointer(StaticGetTypeLayoutDesc<T>(), *DerivedTypeDesc, &OffsetToBase);
 			PointerWriter.WriteObject((uint8*)RawPtr - OffsetToBase, *DerivedTypeDesc);
 		}
 		else
 		{
-			Writer.WriteMemoryImagePointerSizedBytes(0u);
+			Writer.WriteNullPointer();
 		}
 	}
 
 private:
 	inline T* GetFrozenPtrInternal() const
 	{
-		return (T*)((char*)this + (OffsetFromThis >> OffsetShift));
+		return (T*)((char*)this + (Frozen.OffsetFromThis >> OffsetShift));
 	}
 
 	enum
@@ -261,8 +366,9 @@ private:
 protected:
 	union
 	{
-		T* Ptr;
-		FMemoryImagePtrInt OffsetFromThis;
+		uint64 Packed;
+		FFrozenMemoryImagePtr Frozen;
+		T* UnfrozenPtr;
 	};
 };
 
@@ -274,56 +380,63 @@ namespace Freeze
 		T* RawPtr = Object.Get();
 		if (RawPtr)
 		{
-			// Compile-time type of the thing we're pointing to
-			const FTypeLayoutDesc& StaticTypeDesc = StaticGetTypeLayoutDesc<T>();
-			// Actual run-time type of the thing we're pointing to
 			const FTypeLayoutDesc& DerivedTypeDesc = GetTypeLayoutDesc(Writer.TryGetPrevPointerTable(), *RawPtr);
-			// 'this' offset to adjust from the compile-time type to the run-time type
-			const uint32 OffsetToBase = DerivedTypeDesc.GetOffsetToBase(StaticTypeDesc);
 
-			FMemoryImageWriter PointerWriter = Writer.WritePointer(FString::Printf(TEXT("TMemoryImagePtr<%s>"), DerivedTypeDesc.Name), OffsetToBase);
+			uint32 OffsetToBase = 0u;
+			FMemoryImageWriter PointerWriter = Writer.WritePointer(StaticGetTypeLayoutDesc<T>(), DerivedTypeDesc, &OffsetToBase);
 			PointerWriter.WriteObject((uint8*)RawPtr - OffsetToBase, DerivedTypeDesc);
 		}
 		else
 		{
-			Writer.WriteMemoryImagePointerSizedBytes(0u);
+			Writer.WriteNullPointer();
 		}
 	}
 
 	template<typename T>
-	void IntrinsicUnfrozenCopy(const FMemoryUnfreezeContent& Context, const TMemoryImagePtr<T>& Object, void* OutDst)
+	uint32 IntrinsicUnfrozenCopy(const FMemoryUnfreezeContent& Context, const TMemoryImagePtr<T>& Object, void* OutDst)
 	{
-		T* RawPtr = Object.Get();
+		const T* RawPtr = Object.Get();
 		if (RawPtr)
 		{
 			// Compile-time type of the thing we're pointing to
 			const FTypeLayoutDesc& StaticTypeDesc = StaticGetTypeLayoutDesc<T>();
+
 			// Actual run-time type of the thing we're pointing to
-			const FTypeLayoutDesc& DerivedTypeDesc = GetTypeLayoutDesc(Context.TryGetPrevPointerTable(), *RawPtr);
+			const FTypeLayoutDesc* DerivedTypeDesc = Context.GetDerivedTypeDesc(StaticTypeDesc, Object.GetFrozenTypeIndex());
+			if (!DerivedTypeDesc && Context.bIsFrozenForCurrentPlatform)
+			{
+				// It's possible we won't be able to get derived type desc from the context, if we're not storing type dependencies
+				// In this case, if we're unfreezing data for the current platform, we can just grab the derived type directly from the frozen object
+				// If we're NOT unfreezing data for current platform, we can't access the frozen object, so we'll fail in that case
+				DerivedTypeDesc = &GetTypeLayoutDesc(Context.PrevPointerTable, *RawPtr);
+			}
+			checkf(DerivedTypeDesc, TEXT("Can't get DerivedTypeDesc for TMemoryImagePtr<%s>, unfreezing for a different platform, and type dependencies are not availiable"), StaticTypeDesc.Name);
+
 			// 'this' offset to adjust from the compile-time type to the run-time type
-			const uint32 OffsetToBase = DerivedTypeDesc.GetOffsetToBase(StaticTypeDesc);
+			const uint32 OffsetToBase = DerivedTypeDesc->GetOffsetToBase(StaticTypeDesc);
 
-			void* UnfrozenMemory = FMemory::Malloc(DerivedTypeDesc.Size, DerivedTypeDesc.Alignment);
-			Context.UnfreezeObject((uint8*)RawPtr - OffsetToBase, DerivedTypeDesc, UnfrozenMemory);
-			T* UnfrozenPtr = (T*)((uint8*)UnfrozenMemory + OffsetToBase);
-			new(OutDst) TMemoryImagePtr<T>(UnfrozenPtr);
+			void* UnfrozenMemory = ::operator new(DerivedTypeDesc->Size);
+			Context.UnfreezeObject((uint8*)RawPtr - OffsetToBase, *DerivedTypeDesc, UnfrozenMemory);
+			T* UnfrozenObject = (T*)((uint8*)UnfrozenMemory + OffsetToBase);
+			new(OutDst) TMemoryImagePtr<T>(UnfrozenObject);
 		}
 		else
 		{
 			new(OutDst) TMemoryImagePtr<T>(nullptr);
 		}
+		return sizeof(Object);
 	}
 
 	template<typename T>
 	uint32 IntrinsicAppendHash(const TMemoryImagePtr<T>* DummyObject, const FTypeLayoutDesc& TypeDesc, const FPlatformTypeLayoutParameters& LayoutParams, FSHA1& Hasher)
 	{
-		return AppendHashForNameAndSize(TypeDesc.Name, LayoutParams.GetMemoryImagePointerSize(), Hasher);
+		return AppendHashForNameAndSize(TypeDesc.Name, sizeof(TMemoryImagePtr<T>), Hasher);
 	}
 
 	template<typename T>
 	inline uint32 IntrinsicGetTargetAlignment(const TMemoryImagePtr<T>* DummyObject, const FTypeLayoutDesc& TypeDesc, const FPlatformTypeLayoutParameters& LayoutParams)
 	{
-		return FMath::Min(LayoutParams.GetMemoryImagePointerSize(), LayoutParams.MaxFieldAlignment);
+		return FMath::Min(8u, LayoutParams.MaxFieldAlignment);
 	}
 
 	template<typename T>
@@ -339,6 +452,14 @@ namespace Freeze
 			// 'this' offset to adjust from the compile-time type to the run-time type
 			const uint32 OffsetToBase = DerivedTypeDesc.GetOffsetToBase(StaticTypeDesc);
 
+			if (Object.IsFrozen())
+			{
+				OutContext.AppendFrozenPointer(StaticTypeDesc, Object.GetFrozenTypeIndex());
+			}
+			else
+			{
+				OutContext.AppendUnfrozenPointer(StaticTypeDesc);
+			}
 			DerivedTypeDesc.ToStringFunc((uint8*)RawPtr - OffsetToBase, DerivedTypeDesc, LayoutParams, OutContext);
 		}
 		else
@@ -421,11 +542,11 @@ public:
 	{
 		return Data.IsValid();
 	}
-	FORCEINLINE FMemoryImagePtrInt GetFrozenOffsetFromThis() const { return Data.GetFrozenOffsetFromThis(); }
+	FORCEINLINE int32 GetFrozenOffsetFromThis() const { return Data.GetFrozenOffsetFromThis(); }
 
 	void ResizeAllocation(int32 PreviousNumElements, int32 NumElements, SIZE_T NumBytesPerElement, uint32 Alignment);
 	void WriteMemoryImage(FMemoryImageWriter& Writer, const FTypeLayoutDesc& TypeDesc, int32 NumAllocatedElements, uint32 Alignment) const;
-	void ToString(const FTypeLayoutDesc& TypeDesc, int32 NumAllocatedElements, const FPlatformTypeLayoutParameters& LayoutParams, FMemoryToStringContext& OutContext) const;
+	void ToString(const FTypeLayoutDesc& TypeDesc, int32 NumAllocatedElements, int32 MaxAllocatedElements, const FPlatformTypeLayoutParameters& LayoutParams, FMemoryToStringContext& OutContext) const;
 	void CopyUnfrozen(const FMemoryUnfreezeContent& Context, const FTypeLayoutDesc& TypeDesc, int32 NumAllocatedElements, void* Dst) const;
 
 private:
@@ -623,7 +744,7 @@ struct FHashedNameDebugString
 namespace Freeze
 {
 	void IntrinsicWriteMemoryImage(FMemoryImageWriter& Writer, const FHashedNameDebugString& Object, const FTypeLayoutDesc&);
-	void IntrinsicUnfrozenCopy(const FMemoryUnfreezeContent& Context, const FHashedNameDebugString& Object, void* OutDst);
+	uint32 IntrinsicUnfrozenCopy(const FMemoryUnfreezeContent& Context, const FHashedNameDebugString& Object, void* OutDst);
 }
 
 DECLARE_INTRINSIC_TYPE_LAYOUT(FHashedNameDebugString);
@@ -914,11 +1035,11 @@ private:
 		InPtr.SafeRelease();
 	}
 
-	static_assert(sizeof(PtrType) <= sizeof(FMemoryImageUPtrInt), "PtrType must fit within a standard pointer");
+	static_assert(sizeof(PtrType) <= sizeof(uint64), "PtrType must fit within a standard pointer");
 	union
 	{
 		PtrType Ptr;
-		FMemoryImageUPtrInt PackedIndex;
+		uint64 PackedIndex;
 	};
 };
 
@@ -943,30 +1064,31 @@ namespace Freeze
 			const uint32 Index = Writer.GetPointerTable().AddIndexedPointer(TypeDesc, RawPtr);
 			check(Index != (uint32)INDEX_NONE);
 			const uint64 FrozenPackedIndex = ((uint64)Index << 1u) | 1u;
-			Writer.WriteMemoryImagePointerSizedBytes(FrozenPackedIndex);
+			Writer.WriteBytes(FrozenPackedIndex);
 		}
 		else
 		{
-			Writer.WriteMemoryImagePointerSizedBytes(0u);
+			Writer.WriteBytes(uint64(0u));
 		}
 	}
 
 	template<typename T, typename PtrType>
-	void IntrinsicUnfrozenCopy(const FMemoryUnfreezeContent& Context, const TIndexedPtrBase<T, PtrType>& Object, void* OutDst)
+	uint32 IntrinsicUnfrozenCopy(const FMemoryUnfreezeContent& Context, const TIndexedPtrBase<T, PtrType>& Object, void* OutDst)
 	{
 		new(OutDst) TIndexedPtrBase<T, PtrType>(Object.Get(Context.TryGetPrevPointerTable()));
+		return sizeof(Object);
 	}
 
 	template<typename T, typename PtrType>
 	uint32 IntrinsicAppendHash(const TIndexedPtrBase<T, PtrType>* DummyObject, const FTypeLayoutDesc& TypeDesc, const FPlatformTypeLayoutParameters& LayoutParams, FSHA1& Hasher)
 	{
-		return AppendHashForNameAndSize(TypeDesc.Name, LayoutParams.GetMemoryImagePointerSize(), Hasher);
+		return AppendHashForNameAndSize(TypeDesc.Name, sizeof(TIndexedPtrBase<T, PtrType>), Hasher);
 	}
 
 	template<typename T, typename PtrType>
 	inline uint32 IntrinsicGetTargetAlignment(const TIndexedPtrBase<T, PtrType>* DummyObject, const FTypeLayoutDesc& TypeDesc, const FPlatformTypeLayoutParameters& LayoutParams)
 	{
-		return FMath::Min(LayoutParams.GetMemoryImagePointerSize(), LayoutParams.MaxFieldAlignment);
+		return FMath::Min(8u, LayoutParams.MaxFieldAlignment);
 	}
 }
 
diff --git a/Engine/Source/Runtime/Core/Public/Serialization/MemoryImageWriter.h b/Engine/Source/Runtime/Core/Public/Serialization/MemoryImageWriter.h
--- a/Engine/Source/Runtime/Core/Public/Serialization/MemoryImageWriter.h
+++ b/Engine/Source/Runtime/Core/Public/Serialization/MemoryImageWriter.h
@@ -27,34 +27,37 @@ public:
 	inline bool Is32BitTarget() const { return GetTargetLayoutParams().Is32Bit(); }
 	inline bool Is64BitTarget() const { return !Is32BitTarget(); }
 
-	void AddDependency(const FTypeLayoutDesc& TypeDesc);
+	int32 AddTypeDependency(const FTypeLayoutDesc& TypeDesc);
 
 	void WriteObject(const void* Object, const FTypeLayoutDesc& TypeDesc);
 	void WriteObjectArray(const void* Object, const FTypeLayoutDesc& TypeDesc, uint32_t NumArray);
+	void WriteRootObject(const void* Object, const FTypeLayoutDesc& TypeDesc);
 
 	uint32 GetOffset() const;
 	uint32 WriteAlignment(uint32 Alignment);
 	void WritePaddingToSize(uint32 Offset);
 	uint32 WriteBytes(const void* Data, uint32 Size);
-	FMemoryImageWriter WritePointer(const FString& SectionName, uint32 Offset = 0u);
+	FMemoryImageWriter WritePointer(const FTypeLayoutDesc& StaticTypeDesc, const FTypeLayoutDesc& DerivedTypeDesc, uint32* OutOffsetToBase = nullptr);
+	FMemoryImageWriter WritePointer(const FTypeLayoutDesc& TypeDesc);
+	uint32 WriteNullPointer();
 	uint32 WriteRawPointerSizedBytes(uint64 PointerValue);
-	uint32 WriteMemoryImagePointerSizedBytes(uint64 PointerValue);
 	uint32 WriteVTable(const FTypeLayoutDesc& TypeDesc, const FTypeLayoutDesc& DerivedTypeDesc);
 	uint32 WriteFName(const FName& Name);
 	uint32 WriteFMinimalName(const FMinimalName& Name);
 	uint32 WriteFScriptName(const FScriptName& Name);
 
 	template<typename T>
-	void WriteObjectArray(const T* Object, uint32 NumArray)
+	void WriteObject(const T& Object)
 	{
-		const FTypeLayoutDesc& TypeDesc = StaticGetTypeLayoutDesc<T>();
-		WriteObjectArray(Object, TypeDesc, TypeDesc, NumArray);
+		const FTypeLayoutDesc& TypeDesc = GetTypeLayoutDesc(TryGetPrevPointerTable(), Object);
+		WriteObject(&Object, TypeDesc);
 	}
 
 	template<typename T>
-	void WriteObject(const T& Object)
+	void WriteRootObject(const T& Object)
 	{
-		WriteObject(&Object, GetTypeLayoutDesc(TryGetPrevPointerTable(), Object));
+		const FTypeLayoutDesc& TypeDesc = GetTypeLayoutDesc(TryGetPrevPointerTable(), Object);
+		WriteRootObject(&Object, TypeDesc);
 	}
 
 	template<typename T>
@@ -76,21 +79,39 @@ public:
 class FMemoryUnfreezeContent
 {
 public:
-	FMemoryUnfreezeContent() : PrevPointerTable(nullptr) {}
+	explicit FMemoryUnfreezeContent(const FPointerTableBase* InPointerTable)
+		: PrevPointerTable(InPointerTable)
+		, bIsFrozenForCurrentPlatform(true)
+	{
+		FrozenLayoutParameters.InitializeForCurrent();
+	}
 
-	const FPointerTableBase* TryGetPrevPointerTable() const { return PrevPointerTable; }
+	FMemoryUnfreezeContent(const FPointerTableBase* InPointerTable, const FPlatformTypeLayoutParameters& InLayoutParams)
+		: PrevPointerTable(InPointerTable)
+		, FrozenLayoutParameters(InLayoutParams)
+		, bIsFrozenForCurrentPlatform(InLayoutParams.IsCurrentPlatform())
+	{
+	}
+
+	inline const FPointerTableBase* TryGetPrevPointerTable() const { return PrevPointerTable; }
 
-	void UnfreezeObject(const void* Object, const FTypeLayoutDesc& TypeDesc, void* OutDst) const
+	CORE_API const FTypeLayoutDesc* GetDerivedTypeDesc(const FTypeLayoutDesc& StaticTypeDesc, int32 TypeIndex) const;
+
+	inline uint32 UnfreezeObject(const void* Object, const FTypeLayoutDesc& TypeDesc, void* OutDst) const
 	{
-		TypeDesc.UnfrozenCopyFunc(*this, Object, TypeDesc, OutDst);
+		return TypeDesc.UnfrozenCopyFunc(*this, Object, TypeDesc, OutDst);
 	}
 
 	template<typename T>
-	void UnfreezeObject(const T& Object, void* OutDst) const
+	inline uint32 UnfreezeObject(const T& Object, void* OutDst) const
 	{
 		const FTypeLayoutDesc& TypeDesc = GetTypeLayoutDesc(TryGetPrevPointerTable(), Object);
-		UnfreezeObject(&Object, TypeDesc, OutDst);
+		return UnfreezeObject(&Object, TypeDesc, OutDst);
 	}
 
-	const FPointerTableBase* PrevPointerTable;
+	const FPointerTableBase* PrevPointerTable = nullptr;
+
+	// Layout of the frozen data
+	FPlatformTypeLayoutParameters FrozenLayoutParameters;
+	bool bIsFrozenForCurrentPlatform; // FrozenLayoutParameters.IsCurrentPlatform
 };
diff --git a/Engine/Source/Runtime/Core/Public/Serialization/MemoryLayout.h b/Engine/Source/Runtime/Core/Public/Serialization/MemoryLayout.h
--- a/Engine/Source/Runtime/Core/Public/Serialization/MemoryLayout.h
+++ b/Engine/Source/Runtime/Core/Public/Serialization/MemoryLayout.h
@@ -28,25 +28,6 @@ struct FPlatformTypeLayoutParameters;
 #endif
 #endif
 
-/**
- * If this is set, TMemoryImagePtr, TIndexedPtr, and other wrapped pointer types used for memory images will be forced to 64bits, even when building 32bit targets
- * This is intended to facilitate sharing packs between 32/64bit builds
- * Android requires sharing data between 32/64bit exe
- */
-#if PLATFORM_ANDROID
-#define UE_FORCE_64BIT_MEMORY_IMAGE_POINTERS 1
-#else
-#define UE_FORCE_64BIT_MEMORY_IMAGE_POINTERS 0
-#endif
-
-#if UE_FORCE_64BIT_MEMORY_IMAGE_POINTERS
-using FMemoryImagePtrInt = int64;
-using FMemoryImageUPtrInt = uint64;
-#else
-using FMemoryImagePtrInt = PTRINT;
-using FMemoryImageUPtrInt = UPTRINT;
-#endif
-
 /*#define UE_STATIC_ONLY(T) \
 	T() = delete; \
 	~T() = delete; \
@@ -70,8 +51,10 @@ template<typename T> const FTypeLayoutDesc& StaticGetTypeLayoutDesc();
 
 struct FMemoryToStringContext
 {
-	CORE_API void AppendNullptr();
 	CORE_API void AppendIndent();
+	CORE_API void AppendNullptr();
+	CORE_API void AppendUnfrozenPointer(const FTypeLayoutDesc& StaticTypeDesc);
+	CORE_API void AppendFrozenPointer(const FTypeLayoutDesc& StaticTypeDesc, int32 FrozenTypeIndex);
 
 	const FPointerTableBase* TryGetPrevPointerTable() const { return PrevPointerTable; }
 
@@ -127,7 +110,7 @@ struct FTypeLayoutDesc
 {
 	typedef void (FDestroyFunc)(void* Object, const FTypeLayoutDesc& TypeDesc, const FPointerTableBase* PtrTable);
 	typedef void (FWriteFrozenMemoryImageFunc)(FMemoryImageWriter& Writer, const void* Object, const FTypeLayoutDesc& TypeDesc, const FTypeLayoutDesc& DerivedTypeDesc);
-	typedef void (FUnfrozenCopyFunc)(const FMemoryUnfreezeContent& Context, const void* Object, const FTypeLayoutDesc& TypeDesc, void* OutDst);
+	typedef uint32 (FUnfrozenCopyFunc)(const FMemoryUnfreezeContent& Context, const void* Object, const FTypeLayoutDesc& TypeDesc, void* OutDst);
 	typedef uint32 (FAppendHashFunc)(const FTypeLayoutDesc& TypeDesc, const FPlatformTypeLayoutParameters& LayoutParams, FSHA1& Hasher);
 	typedef uint32 (FGetTargetAlignmentFunc)(const FTypeLayoutDesc& TypeDesc, const FPlatformTypeLayoutParameters& LayoutParams);
 	typedef void (FToStringFunc)(const void* Object, const FTypeLayoutDesc& TypeDesc, const FPlatformTypeLayoutParameters& LayoutParams, FMemoryToStringContext& OutContext);
@@ -253,7 +236,7 @@ namespace Freeze
 {
 	CORE_API void DefaultWriteMemoryImageField(FMemoryImageWriter& Writer, const void* Object, const void* FieldObject, const FTypeLayoutDesc& TypeDesc, const FTypeLayoutDesc& DerivedTypeDesc);
 	CORE_API void DefaultWriteMemoryImage(FMemoryImageWriter& Writer, const void* Object, const FTypeLayoutDesc& TypeDesc, const FTypeLayoutDesc& DerivedTypeDesc);
-	CORE_API void DefaultUnfrozenCopy(const FMemoryUnfreezeContent& Context, const void* Object, const FTypeLayoutDesc& TypeDesc, void* OutDst);
+	CORE_API uint32 DefaultUnfrozenCopy(const FMemoryUnfreezeContent& Context, const void* Object, const FTypeLayoutDesc& TypeDesc, void* OutDst);
 	CORE_API uint32 DefaultAppendHash(const FTypeLayoutDesc& TypeDesc, const FPlatformTypeLayoutParameters& LayoutParams, FSHA1& Hasher);
 	CORE_API uint32 DefaultGetTargetAlignment(const FTypeLayoutDesc& TypeDesc, const FPlatformTypeLayoutParameters& LayoutParams);
 	CORE_API void DefaultToString(const void* Object, const FTypeLayoutDesc& TypeDesc, const FPlatformTypeLayoutParameters& LayoutParams, FMemoryToStringContext& OutContext);
@@ -274,9 +257,10 @@ namespace Freeze
 	}
 
 	template<typename T>
-	FORCEINLINE void IntrinsicUnfrozenCopy(const FMemoryUnfreezeContent& Context, const T& Object, void* OutDst)
+	FORCEINLINE uint32 IntrinsicUnfrozenCopy(const FMemoryUnfreezeContent& Context, const T& Object, void* OutDst)
 	{
 		new(OutDst) T(Object);
+		return sizeof(T);
 	}
 
 	template<typename T>
@@ -702,8 +686,8 @@ inline void DeleteObjectFromLayout(T* Object, const FPointerTableBase* PtrTable
 		static void CallWriteMemoryImage(FMemoryImageWriter& Writer, const void* Object, const FTypeLayoutDesc& TypeDesc, const FTypeLayoutDesc& DerivedTypeDesc) { \
 			Freeze::IntrinsicWriteMemoryImage(Writer, *static_cast<PREPROCESSOR_REMOVE_OPTIONAL_PARENS(T) const*>(Object), TypeDesc); \
 		} \
-		static void CallUnfrozenCopy(const FMemoryUnfreezeContent& Context, const void* Object, const FTypeLayoutDesc& TypeDesc, void* OutDst) { \
-			Freeze::IntrinsicUnfrozenCopy(Context, *static_cast<PREPROCESSOR_REMOVE_OPTIONAL_PARENS(T) const*>(Object), OutDst); \
+		static uint32 CallUnfrozenCopy(const FMemoryUnfreezeContent& Context, const void* Object, const FTypeLayoutDesc& TypeDesc, void* OutDst) { \
+			return Freeze::IntrinsicUnfrozenCopy(Context, *static_cast<PREPROCESSOR_REMOVE_OPTIONAL_PARENS(T) const*>(Object), OutDst); \
 		} \
 		static uint32 CallAppendHash(const FTypeLayoutDesc& TypeDesc, const FPlatformTypeLayoutParameters& LayoutParams, FSHA1& Hasher) { \
 			return Freeze::IntrinsicAppendHash(static_cast<PREPROCESSOR_REMOVE_OPTIONAL_PARENS(T) const*>(nullptr), TypeDesc, LayoutParams, Hasher); \
@@ -746,7 +730,7 @@ inline void DeleteObjectFromLayout(T* Object, const FPointerTableBase* PtrTable
 	PREPROCESSOR_REMOVE_OPTIONAL_PARENS(TemplatePrefix) struct TStaticGetTypeLayoutHelper<PREPROCESSOR_REMOVE_OPTIONAL_PARENS(T)> { \
 		UE_STATIC_ONLY(TStaticGetTypeLayoutHelper); \
 		static void CallWriteMemoryImage(FMemoryImageWriter& Writer, const void* Object, const FTypeLayoutDesc& TypeDesc, const FTypeLayoutDesc& DerivedTypeDesc); \
-		static void CallUnfrozenCopy(const FMemoryUnfreezeContent& Context, const void* Object, const FTypeLayoutDesc& TypeDesc, void* OutDst); \
+		static uint32 CallUnfrozenCopy(const FMemoryUnfreezeContent& Context, const void* Object, const FTypeLayoutDesc& TypeDesc, void* OutDst); \
 		static uint32 CallAppendHash(const FTypeLayoutDesc& TypeDesc, const FPlatformTypeLayoutParameters& LayoutParams, FSHA1& Hasher); \
 		static uint32 CallGetTargetAlignment(const FTypeLayoutDesc& TypeDesc, const FPlatformTypeLayoutParameters& LayoutParams); \
 		static void CallToString(const void* Object, const FTypeLayoutDesc& TypeDesc, const FPlatformTypeLayoutParameters& LayoutParams, FMemoryToStringContext& OutContext); \
@@ -760,8 +744,8 @@ inline void DeleteObjectFromLayout(T* Object, const FPointerTableBase* PtrTable
 		void TStaticGetTypeLayoutHelper<PREPROCESSOR_REMOVE_OPTIONAL_PARENS(T)>::CallWriteMemoryImage(FMemoryImageWriter& Writer, const void* Object, const FTypeLayoutDesc& TypeDesc, const FTypeLayoutDesc& DerivedTypeDesc) { \
 			Freeze::IntrinsicWriteMemoryImage(Writer, *static_cast<PREPROCESSOR_REMOVE_OPTIONAL_PARENS(T) const*>(Object), TypeDesc); \
 		} \
-		void TStaticGetTypeLayoutHelper<PREPROCESSOR_REMOVE_OPTIONAL_PARENS(T)>::CallUnfrozenCopy(const FMemoryUnfreezeContent& Context, const void* Object, const FTypeLayoutDesc& TypeDesc, void* OutDst) { \
-			Freeze::IntrinsicUnfrozenCopy(Context, *static_cast<PREPROCESSOR_REMOVE_OPTIONAL_PARENS(T) const*>(Object), OutDst); \
+		uint32 TStaticGetTypeLayoutHelper<PREPROCESSOR_REMOVE_OPTIONAL_PARENS(T)>::CallUnfrozenCopy(const FMemoryUnfreezeContent& Context, const void* Object, const FTypeLayoutDesc& TypeDesc, void* OutDst) { \
+			return Freeze::IntrinsicUnfrozenCopy(Context, *static_cast<PREPROCESSOR_REMOVE_OPTIONAL_PARENS(T) const*>(Object), OutDst); \
 		} \
 		uint32 TStaticGetTypeLayoutHelper<PREPROCESSOR_REMOVE_OPTIONAL_PARENS(T)>::CallAppendHash(const FTypeLayoutDesc& TypeDesc, const FPlatformTypeLayoutParameters& LayoutParams, FSHA1& Hasher) { \
 			return Freeze::IntrinsicAppendHash(static_cast<PREPROCESSOR_REMOVE_OPTIONAL_PARENS(T) const*>(nullptr), TypeDesc, LayoutParams, Hasher); \
@@ -823,6 +807,7 @@ DECLARE_INTRINSIC_TYPE_LAYOUT(double);
 DECLARE_INTRINSIC_TYPE_LAYOUT(wchar_t);
 DECLARE_INTRINSIC_TYPE_LAYOUT(char16_t);
 DECLARE_INTRINSIC_TYPE_LAYOUT(void*);
+DECLARE_INTRINSIC_TYPE_LAYOUT(ETypeLayoutInterface::Type);
 
 DECLARE_INTRINSIC_TYPE_LAYOUT(FThreadSafeCounter);
 DECLARE_TEMPLATE_INTRINSIC_TYPE_LAYOUT(template<typename T>, TEnumAsByte<T>);
@@ -841,10 +826,9 @@ struct FPlatformTypeLayoutParameters
 	{
 		Flag_Initialized = (1 << 0),
 		Flag_Is32Bit = (1 << 1),
-		Flag_Force64BitMemoryImagePointers = (1 << 2),
-		Flag_AlignBases = (1 << 3),
-		Flag_WithEditorOnly = (1 << 4),
-		Flag_WithRaytracing = (1 << 5),
+		Flag_AlignBases = (1 << 2),
+		Flag_WithEditorOnly = (1 << 3),
+		Flag_WithRaytracing = (1 << 4),
 	};
 
 	LAYOUT_FIELD_INITIALIZED(uint32, MaxFieldAlignment, 0xffffffff);
@@ -852,7 +836,6 @@ struct FPlatformTypeLayoutParameters
 
 	inline bool IsInitialized() const { return (Flags & Flag_Initialized) != 0u; }
 	inline bool Is32Bit() const { return (Flags & Flag_Is32Bit) != 0u; }
-	inline bool HasForce64BitMemoryImagePointers() const { return (Flags & Flag_Force64BitMemoryImagePointers) != 0u; }
 	inline bool HasAlignBases() const { return (Flags & Flag_AlignBases) != 0u; }
 	inline bool WithEditorOnly() const { return (Flags & Flag_WithEditorOnly) != 0u; }
 	inline bool WithRaytracing() const { return (Flags & Flag_WithRaytracing) != 0u; }
@@ -860,11 +843,7 @@ struct FPlatformTypeLayoutParameters
 	// May need dedicated flag for this, if we need to support case-preserving names in non-editor builds
 	inline bool WithCasePreservingFName() const { return WithEditorOnly(); }
 
-	inline bool Has32BitMemoryImagePointers() const { return Is32Bit() && !HasForce64BitMemoryImagePointers(); }
-	inline bool Has64BitMemoryImagePointers() const { return !Has32BitMemoryImagePointers(); }
-
 	inline uint32 GetRawPointerSize() const { return Is32Bit() ? sizeof(uint32) : sizeof(uint64); }
-	inline uint32 GetMemoryImagePointerSize() const { return Has32BitMemoryImagePointers() ? sizeof(uint32) : sizeof(uint64); }
 
 	friend inline FArchive& operator<<(FArchive& Ar, FPlatformTypeLayoutParameters& Ref)
 	{
diff --git a/Engine/Source/Runtime/Core/Public/Templates/Tuple.h b/Engine/Source/Runtime/Core/Public/Templates/Tuple.h
--- a/Engine/Source/Runtime/Core/Public/Templates/Tuple.h
+++ b/Engine/Source/Runtime/Core/Public/Templates/Tuple.h
@@ -730,11 +730,12 @@ namespace Freeze
 	}
 
 	template<typename KeyType, typename ValueType>
-	void IntrinsicUnfrozenCopy(const FMemoryUnfreezeContent& Context, const TTuple<KeyType, ValueType>& Object, void* OutDst)
+	uint32 IntrinsicUnfrozenCopy(const FMemoryUnfreezeContent& Context, const TTuple<KeyType, ValueType>& Object, void* OutDst)
 	{
 		TTuple<KeyType, ValueType>* DstObject = (TTuple<KeyType, ValueType>*)OutDst;
 		Context.UnfreezeObject(Object.Key, &DstObject->Key);
 		Context.UnfreezeObject(Object.Value, &DstObject->Value);
+		return sizeof(Object);
 	}
 
 	template<typename KeyType, typename ValueType>
diff --git a/Engine/Source/Runtime/Engine/Private/GPUSkinVertexFactory.cpp b/Engine/Source/Runtime/Engine/Private/GPUSkinVertexFactory.cpp
--- a/Engine/Source/Runtime/Engine/Private/GPUSkinVertexFactory.cpp
+++ b/Engine/Source/Runtime/Engine/Private/GPUSkinVertexFactory.cpp
@@ -577,7 +577,7 @@ TGPUSkinVertexFactoryShaderParameters
 /** Shader parameters for use with TGPUSkinVertexFactory */
 class FGPUSkinVertexFactoryShaderParameters : public FVertexFactoryShaderParameters
 {
-	DECLARE_INLINE_TYPE_LAYOUT(FGPUSkinVertexFactoryShaderParameters, NonVirtual);
+	DECLARE_TYPE_LAYOUT(FGPUSkinVertexFactoryShaderParameters, NonVirtual);
 public:
 	/**
 	* Bind shader constants by name
@@ -656,6 +656,8 @@ private:
 
 };
 
+IMPLEMENT_TYPE_LAYOUT(FGPUSkinVertexFactoryShaderParameters);
+
 IMPLEMENT_GPUSKINNING_VERTEX_FACTORY_PARAMETER_TYPE(TGPUSkinVertexFactory, SF_Vertex, FGPUSkinVertexFactoryShaderParameters);
 
 /** bind gpu skin vertex factory to its shader file and its shader parameters */
@@ -668,7 +670,7 @@ TGPUSkinVertexFactoryShaderParameters
 /** Shader parameters for use with TGPUSkinVertexFactory */
 class FGPUSkinVertexPassthroughFactoryShaderParameters : public FLocalVertexFactoryShaderParametersBase
 {
-	DECLARE_INLINE_TYPE_LAYOUT(FGPUSkinVertexPassthroughFactoryShaderParameters, NonVirtual);
+	DECLARE_TYPE_LAYOUT(FGPUSkinVertexPassthroughFactoryShaderParameters, NonVirtual);
 public:
 	/**
 	* Bind shader constants by name
@@ -706,12 +708,12 @@ public:
 	}
 
 private:
-	
-		LAYOUT_FIELD(FShaderResourceParameter, GPUSkinCachePositionBuffer)
-		LAYOUT_FIELD(FShaderResourceParameter, GPUSkinCachePreviousPositionBuffer)
-	
+	LAYOUT_FIELD(FShaderResourceParameter, GPUSkinCachePositionBuffer);
+	LAYOUT_FIELD(FShaderResourceParameter, GPUSkinCachePreviousPositionBuffer);
 };
 
+IMPLEMENT_TYPE_LAYOUT(FGPUSkinVertexPassthroughFactoryShaderParameters);
+
 /*-----------------------------------------------------------------------------
 FGPUSkinPassthroughVertexFactory
 -----------------------------------------------------------------------------*/
@@ -877,7 +879,7 @@ IMPLEMENT_GPUSKINNING_VERTEX_FACTORY_TYPE(TGPUSkinMorphVertexFactory, "/Engine/P
 /** Shader parameters for use with TGPUSkinAPEXClothVertexFactory */
 class TGPUSkinAPEXClothVertexFactoryShaderParameters : public FGPUSkinVertexFactoryShaderParameters
 {
-	DECLARE_INLINE_TYPE_LAYOUT(TGPUSkinAPEXClothVertexFactoryShaderParameters, NonVirtual);
+	DECLARE_TYPE_LAYOUT(TGPUSkinAPEXClothVertexFactoryShaderParameters, NonVirtual);
 public:
 
 	/**
@@ -940,16 +942,18 @@ public:
 
 protected:
 	
-		LAYOUT_FIELD(FShaderResourceParameter, ClothSimulVertsPositionsNormalsParameter)
-		LAYOUT_FIELD(FShaderResourceParameter, PreviousClothSimulVertsPositionsNormalsParameter)
-		LAYOUT_FIELD(FShaderParameter, ClothLocalToWorldParameter)
-		LAYOUT_FIELD(FShaderParameter, PreviousClothLocalToWorldParameter)
-		LAYOUT_FIELD(FShaderParameter, ClothBlendWeightParameter)
-		LAYOUT_FIELD(FShaderResourceParameter, GPUSkinApexClothParameter)
-		LAYOUT_FIELD(FShaderParameter, GPUSkinApexClothStartIndexOffsetParameter)
+	LAYOUT_FIELD(FShaderResourceParameter, ClothSimulVertsPositionsNormalsParameter);
+	LAYOUT_FIELD(FShaderResourceParameter, PreviousClothSimulVertsPositionsNormalsParameter);
+	LAYOUT_FIELD(FShaderParameter, ClothLocalToWorldParameter);
+	LAYOUT_FIELD(FShaderParameter, PreviousClothLocalToWorldParameter);
+	LAYOUT_FIELD(FShaderParameter, ClothBlendWeightParameter);
+	LAYOUT_FIELD(FShaderResourceParameter, GPUSkinApexClothParameter);
+	LAYOUT_FIELD(FShaderParameter, GPUSkinApexClothStartIndexOffsetParameter);
 	
 };
 
+IMPLEMENT_TYPE_LAYOUT(TGPUSkinAPEXClothVertexFactoryShaderParameters);
+
 /*-----------------------------------------------------------------------------
 	TGPUSkinAPEXClothVertexFactory::ClothShaderType
 -----------------------------------------------------------------------------*/
diff --git a/Engine/Source/Runtime/Engine/Private/GeometryCacheVertexFactory.cpp b/Engine/Source/Runtime/Engine/Private/GeometryCacheVertexFactory.cpp
--- a/Engine/Source/Runtime/Engine/Private/GeometryCacheVertexFactory.cpp
+++ b/Engine/Source/Runtime/Engine/Private/GeometryCacheVertexFactory.cpp
@@ -21,7 +21,7 @@ IMPLEMENT_GLOBAL_SHADER_PARAMETER_STRUCT(FGeometryCacheManualVertexFetchUniformB
 /** Shader parameters for use with TGPUSkinVertexFactory */
 class FGeometryCacheVertexFactoryShaderParameters : public FVertexFactoryShaderParameters
 {
-	DECLARE_INLINE_TYPE_LAYOUT(FGeometryCacheVertexFactoryShaderParameters, NonVirtual);
+	DECLARE_TYPE_LAYOUT(FGeometryCacheVertexFactoryShaderParameters, NonVirtual);
 public:
 
 	/**
@@ -79,6 +79,8 @@ private:
 	LAYOUT_FIELD(FShaderParameter, MotionBlurPositionScale);
 };
 
+IMPLEMENT_TYPE_LAYOUT(FGeometryCacheVertexFactoryShaderParameters);
+
 /*-----------------------------------------------------------------------------
 FGPUSkinPassthroughVertexFactory
 -----------------------------------------------------------------------------*/
diff --git a/Engine/Source/Runtime/Engine/Private/Materials/MaterialShader.cpp b/Engine/Source/Runtime/Engine/Private/Materials/MaterialShader.cpp
--- a/Engine/Source/Runtime/Engine/Private/Materials/MaterialShader.cpp
+++ b/Engine/Source/Runtime/Engine/Private/Materials/MaterialShader.cpp
@@ -2488,9 +2488,7 @@ FMaterialShaderMap* FMaterialShaderMap::AcquireFinalizedClone()
 	Clone->ShaderMapId = ShaderMapId;
 	Clone->bCompilationFinalized = bCompilationFinalized;
 	Clone->bCompiledSuccessfully = bCompiledSuccessfully;
-	Clone->CopyResourceCode(*LocalCode);
-	Clone->AssignAndFreezeContent(LocalContent);
-	Clone->InitResource();
+	Clone->AssignCopy(*this);
 #if WITH_EDITOR
 	Clone->AssociateWithAssets(GetAssociatedAssets());
 #endif
diff --git a/Engine/Source/Runtime/Engine/Private/MeshParticleVertexFactory.cpp b/Engine/Source/Runtime/Engine/Private/MeshParticleVertexFactory.cpp
--- a/Engine/Source/Runtime/Engine/Private/MeshParticleVertexFactory.cpp
+++ b/Engine/Source/Runtime/Engine/Private/MeshParticleVertexFactory.cpp
@@ -11,7 +11,7 @@
 
 class FMeshParticleVertexFactoryShaderParameters : public FVertexFactoryShaderParameters
 {
-	DECLARE_INLINE_TYPE_LAYOUT(FMeshParticleVertexFactoryShaderParameters, NonVirtual);
+	DECLARE_TYPE_LAYOUT(FMeshParticleVertexFactoryShaderParameters, NonVirtual);
 public:
 	void Bind(const FShaderParameterMap& ParameterMap)
 	{
@@ -43,6 +43,8 @@ private:
 	LAYOUT_FIELD(FShaderResourceParameter, PrevTransformBuffer);
 };
 
+IMPLEMENT_TYPE_LAYOUT(FMeshParticleVertexFactoryShaderParameters);
+
 class FDummyPrevTransformBuffer : public FRenderResource
 {
 public:
diff --git a/Engine/Source/Runtime/Engine/Private/Particles/ParticleBeamTrailVertexFactory.cpp b/Engine/Source/Runtime/Engine/Private/Particles/ParticleBeamTrailVertexFactory.cpp
--- a/Engine/Source/Runtime/Engine/Private/Particles/ParticleBeamTrailVertexFactory.cpp
+++ b/Engine/Source/Runtime/Engine/Private/Particles/ParticleBeamTrailVertexFactory.cpp
@@ -17,7 +17,7 @@ IMPLEMENT_GLOBAL_SHADER_PARAMETER_STRUCT(FParticleBeamTrailUniformParameters,"Be
  */
 class FParticleBeamTrailVertexFactoryShaderParameters : public FVertexFactoryShaderParameters
 {
-	DECLARE_INLINE_TYPE_LAYOUT(FParticleBeamTrailVertexFactoryShaderParameters, NonVirtual);
+	DECLARE_TYPE_LAYOUT(FParticleBeamTrailVertexFactoryShaderParameters, NonVirtual);
 public:
 	void GetElementShaderBindings(
 		const FSceneInterface* Scene,
@@ -33,11 +33,10 @@ public:
 		FParticleBeamTrailVertexFactory* BeamTrailVF = (FParticleBeamTrailVertexFactory*)VertexFactory;
 		ShaderBindings.Add(Shader->GetUniformBufferParameter<FParticleBeamTrailUniformParameters>(), BeamTrailVF->GetBeamTrailUniformBuffer() );
 	}
-
-	
-	
 };
 
+IMPLEMENT_TYPE_LAYOUT(FParticleBeamTrailVertexFactoryShaderParameters);
+
 ///////////////////////////////////////////////////////////////////////////////
 /**
  * The particle system beam trail vertex declaration resource type.
diff --git a/Engine/Source/Runtime/Engine/Private/Particles/ParticleGpuSimulation.cpp b/Engine/Source/Runtime/Engine/Private/Particles/ParticleGpuSimulation.cpp
--- a/Engine/Source/Runtime/Engine/Private/Particles/ParticleGpuSimulation.cpp
+++ b/Engine/Source/Runtime/Engine/Private/Particles/ParticleGpuSimulation.cpp
@@ -539,9 +539,9 @@ typedef TUniformBufferRef<FGPUSpriteEmitterDynamicUniformParameters> FGPUSpriteE
  */
 class FGPUSpriteVertexFactoryShaderParametersVS : public FVertexFactoryShaderParameters
 {
-	DECLARE_INLINE_TYPE_LAYOUT(FGPUSpriteVertexFactoryShaderParametersVS, NonVirtual);
+	DECLARE_TYPE_LAYOUT(FGPUSpriteVertexFactoryShaderParametersVS, NonVirtual);
 public:
-	void Bind( const FShaderParameterMap& ParameterMap )
+	void Bind(const FShaderParameterMap& ParameterMap)
 	{
 		ParticleIndices.Bind(ParameterMap, TEXT("ParticleIndices"));
 		ParticleIndicesOffset.Bind(ParameterMap, TEXT("ParticleIndicesOffset"));
@@ -566,33 +566,32 @@ public:
 		class FMeshDrawSingleShaderBindings& ShaderBindings,
 		FVertexInputStreamArray& VertexStreams) const;
 private:
-
-	
-		/** Buffer containing particle indices. */
-		LAYOUT_FIELD(FShaderResourceParameter, ParticleIndices)
-		/** Offset in to the particle indices buffer. */
-		LAYOUT_FIELD(FShaderParameter, ParticleIndicesOffset)
-		/** Texture containing positions for all particles. */
-		LAYOUT_FIELD(FShaderResourceParameter, PositionTexture)
-		LAYOUT_FIELD(FShaderResourceParameter, PositionTextureSampler)
-		/** Texture containing velocities for all particles. */
-		LAYOUT_FIELD(FShaderResourceParameter, VelocityTexture)
-		LAYOUT_FIELD(FShaderResourceParameter, VelocityTextureSampler)
-		/** Texture containint attributes for all particles. */
-		LAYOUT_FIELD(FShaderResourceParameter, AttributesTexture)
-		LAYOUT_FIELD(FShaderResourceParameter, AttributesTextureSampler)
-		/** Texture containing curves from which attributes are sampled. */
-		LAYOUT_FIELD(FShaderResourceParameter, CurveTexture)
-		LAYOUT_FIELD(FShaderResourceParameter, CurveTextureSampler)
-	
+	/** Buffer containing particle indices. */
+	LAYOUT_FIELD(FShaderResourceParameter, ParticleIndices);
+	/** Offset in to the particle indices buffer. */
+	LAYOUT_FIELD(FShaderParameter, ParticleIndicesOffset);
+	/** Texture containing positions for all particles. */
+	LAYOUT_FIELD(FShaderResourceParameter, PositionTexture);
+	LAYOUT_FIELD(FShaderResourceParameter, PositionTextureSampler);
+	/** Texture containing velocities for all particles. */
+	LAYOUT_FIELD(FShaderResourceParameter, VelocityTexture);
+	LAYOUT_FIELD(FShaderResourceParameter, VelocityTextureSampler);
+	/** Texture containint attributes for all particles. */
+	LAYOUT_FIELD(FShaderResourceParameter, AttributesTexture);
+	LAYOUT_FIELD(FShaderResourceParameter, AttributesTextureSampler);
+	/** Texture containing curves from which attributes are sampled. */
+	LAYOUT_FIELD(FShaderResourceParameter, CurveTexture);
+	LAYOUT_FIELD(FShaderResourceParameter, CurveTextureSampler);
 };
 
+IMPLEMENT_TYPE_LAYOUT(FGPUSpriteVertexFactoryShaderParametersVS);
+
 /**
  * Pixel shader parameters for the particle vertex factory.
  */
 class FGPUSpriteVertexFactoryShaderParametersPS : public FVertexFactoryShaderParameters
 {
-	DECLARE_INLINE_TYPE_LAYOUT(FGPUSpriteVertexFactoryShaderParametersPS, NonVirtual);
+	DECLARE_TYPE_LAYOUT(FGPUSpriteVertexFactoryShaderParametersPS, NonVirtual);
 public:
 	void GetElementShaderBindings(
 		const FSceneInterface* Scene,
@@ -604,11 +603,10 @@ public:
 		const FMeshBatchElement& BatchElement,
 		class FMeshDrawSingleShaderBindings& ShaderBindings,
 		FVertexInputStreamArray& VertexStreams) const;
-
-	
-	
 };
 
+IMPLEMENT_TYPE_LAYOUT(FGPUSpriteVertexFactoryShaderParametersPS);
+
 /**
  * GPU Sprite vertex factory vertex declaration.
  */
diff --git a/Engine/Source/Runtime/Engine/Private/Particles/ParticleVertexFactory.cpp b/Engine/Source/Runtime/Engine/Private/Particles/ParticleVertexFactory.cpp
--- a/Engine/Source/Runtime/Engine/Private/Particles/ParticleVertexFactory.cpp
+++ b/Engine/Source/Runtime/Engine/Private/Particles/ParticleVertexFactory.cpp
@@ -50,15 +50,15 @@ TGlobalResource<FNullSubUVCutoutVertexBuffer> GFNullSubUVCutoutVertexBuffer;
  */
 class FParticleSpriteVertexFactoryShaderParameters : public FVertexFactoryShaderParameters
 {
-	DECLARE_INLINE_TYPE_LAYOUT(FParticleSpriteVertexFactoryShaderParameters, NonVirtual);
+	DECLARE_TYPE_LAYOUT(FParticleSpriteVertexFactoryShaderParameters, NonVirtual);
 public:
-	
-	
 };
 
+IMPLEMENT_TYPE_LAYOUT(FParticleSpriteVertexFactoryShaderParameters);
+
 class FParticleSpriteVertexFactoryShaderParametersVS : public FParticleSpriteVertexFactoryShaderParameters
 {
-	DECLARE_INLINE_TYPE_LAYOUT(FParticleSpriteVertexFactoryShaderParametersVS, NonVirtual);
+	DECLARE_TYPE_LAYOUT(FParticleSpriteVertexFactoryShaderParametersVS, NonVirtual);
 public:
 	void Bind(const FShaderParameterMap& ParameterMap)
 	{
@@ -86,15 +86,15 @@ public:
 	}
 
 private:
-	
-		LAYOUT_FIELD(FShaderParameter, NumCutoutVerticesPerFrame)
-		LAYOUT_FIELD(FShaderResourceParameter, CutoutGeometry)
-	
+	LAYOUT_FIELD(FShaderParameter, NumCutoutVerticesPerFrame);
+	LAYOUT_FIELD(FShaderResourceParameter, CutoutGeometry);
 };
 
+IMPLEMENT_TYPE_LAYOUT(FParticleSpriteVertexFactoryShaderParametersVS);
+
 class FParticleSpriteVertexFactoryShaderParametersPS : public FParticleSpriteVertexFactoryShaderParameters
 {
-	DECLARE_INLINE_TYPE_LAYOUT(FParticleSpriteVertexFactoryShaderParametersPS, NonVirtual);
+	DECLARE_TYPE_LAYOUT(FParticleSpriteVertexFactoryShaderParametersPS, NonVirtual);
 public:
 	void GetElementShaderBindings(
 		const FSceneInterface* Scene,
@@ -110,11 +110,10 @@ public:
 		FParticleSpriteVertexFactory* SpriteVF = (FParticleSpriteVertexFactory*)VertexFactory;
 		ShaderBindings.Add(Shader->GetUniformBufferParameter<FParticleSpriteUniformParameters>(), SpriteVF->GetSpriteUniformBuffer() );
 	}
-
-	
-	
 };
 
+IMPLEMENT_TYPE_LAYOUT(FParticleSpriteVertexFactoryShaderParametersPS);
+
 /**
  * The particle system vertex declaration resource type.
  */
diff --git a/Engine/Source/Runtime/Engine/Private/VectorFieldVisualization.cpp b/Engine/Source/Runtime/Engine/Private/VectorFieldVisualization.cpp
--- a/Engine/Source/Runtime/Engine/Private/VectorFieldVisualization.cpp
+++ b/Engine/Source/Runtime/Engine/Private/VectorFieldVisualization.cpp
@@ -26,9 +26,9 @@ IMPLEMENT_GLOBAL_SHADER_PARAMETER_STRUCT(FVectorFieldVisualizationParameters,"Ve
  */
 class FVectorFieldVisualizationVertexFactoryShaderParameters : public FVertexFactoryShaderParameters
 {
-	DECLARE_INLINE_TYPE_LAYOUT(FVectorFieldVisualizationVertexFactoryShaderParameters, NonVirtual);
+	DECLARE_TYPE_LAYOUT(FVectorFieldVisualizationVertexFactoryShaderParameters, NonVirtual);
 public:
-	void Bind( const FShaderParameterMap& ParameterMap )
+	void Bind(const FShaderParameterMap& ParameterMap)
 	{
 		VectorFieldTexture.Bind(ParameterMap, TEXT("VectorFieldTexture"));
 		VectorFieldTextureSampler.Bind(ParameterMap, TEXT("VectorFieldTextureSampler"));
@@ -46,13 +46,14 @@ public:
 		FVertexInputStreamArray& VertexStreams) const;
 
 private:
-	
-		/** The vector field texture parameter. */
-		LAYOUT_FIELD(FShaderResourceParameter, VectorFieldTexture)
-		LAYOUT_FIELD(FShaderResourceParameter, VectorFieldTextureSampler)
-	
+
+	/** The vector field texture parameter. */
+	LAYOUT_FIELD(FShaderResourceParameter, VectorFieldTexture);
+	LAYOUT_FIELD(FShaderResourceParameter, VectorFieldTextureSampler);
 };
 
+IMPLEMENT_TYPE_LAYOUT(FVectorFieldVisualizationVertexFactoryShaderParameters);
+
 /**
  * Vertex declaration for visualizing vector fields.
  */
diff --git a/Engine/Source/Runtime/Landscape/Private/LandscapeRender.cpp b/Engine/Source/Runtime/Landscape/Private/LandscapeRender.cpp
--- a/Engine/Source/Runtime/Landscape/Private/LandscapeRender.cpp
+++ b/Engine/Source/Runtime/Landscape/Private/LandscapeRender.cpp
@@ -3342,7 +3342,7 @@ FLandscapeSharedAdjacencyIndexBuffer::~FLandscapeSharedAdjacencyIndexBuffer()
 /** Shader parameters for use with FLandscapeVertexFactory */
 class FLandscapeVertexFactoryVertexShaderParameters : public FVertexFactoryShaderParameters
 {
-	DECLARE_INLINE_TYPE_LAYOUT(FLandscapeVertexFactoryVertexShaderParameters, NonVirtual);
+	DECLARE_TYPE_LAYOUT(FLandscapeVertexFactoryVertexShaderParameters, NonVirtual);
 public:
 	/**
 	* Bind shader constants by name
@@ -3391,13 +3391,15 @@ public:
 	}
 };
 
+IMPLEMENT_TYPE_LAYOUT(FLandscapeVertexFactoryVertexShaderParameters);
+
 /** 
   * Shader parameters for use with FLandscapeFixedGridVertexFactory
   * Simple grid rendering (without dynamic lod blend) needs a simpler fixed setup.
   */
 class FLandscapeFixedGridVertexFactoryVertexShaderParameters : public FLandscapeVertexFactoryVertexShaderParameters
 {
-	DECLARE_INLINE_TYPE_LAYOUT(FLandscapeFixedGridVertexFactoryVertexShaderParameters, NonVirtual);
+	DECLARE_TYPE_LAYOUT(FLandscapeFixedGridVertexFactoryVertexShaderParameters, NonVirtual);
 public:
 	void GetElementShaderBindings(
 		const class FSceneInterface* Scene,
@@ -3428,6 +3430,8 @@ public:
 	}
 };
 
+IMPLEMENT_TYPE_LAYOUT(FLandscapeFixedGridVertexFactoryVertexShaderParameters);
+
 //
 // FLandscapeVertexFactoryPixelShaderParameters
 //
diff --git a/Engine/Source/Runtime/Landscape/Private/LandscapeRenderMobile.cpp b/Engine/Source/Runtime/Landscape/Private/LandscapeRenderMobile.cpp
--- a/Engine/Source/Runtime/Landscape/Private/LandscapeRenderMobile.cpp
+++ b/Engine/Source/Runtime/Landscape/Private/LandscapeRenderMobile.cpp
@@ -51,7 +51,7 @@ void FLandscapeVertexFactoryMobile::InitRHI()
 /** Shader parameters for use with FLandscapeVertexFactory */
 class FLandscapeVertexFactoryMobileVertexShaderParameters : public FVertexFactoryShaderParameters
 {
-	DECLARE_INLINE_TYPE_LAYOUT(FLandscapeVertexFactoryMobileVertexShaderParameters, NonVirtual);
+	DECLARE_TYPE_LAYOUT(FLandscapeVertexFactoryMobileVertexShaderParameters, NonVirtual);
 public:
 	/**
 	* Bind shader constants by name
@@ -107,10 +107,12 @@ protected:
 	LAYOUT_FIELD(FShaderParameter, TexCoordOffsetParameter);
 };
 
+IMPLEMENT_TYPE_LAYOUT(FLandscapeVertexFactoryMobileVertexShaderParameters);
+
 /** Shader parameters for use with FLandscapeVertexFactory */
 class FLandscapeVertexFactoryMobilePixelShaderParameters : public FLandscapeVertexFactoryPixelShaderParameters
 {
-	DECLARE_INLINE_TYPE_LAYOUT(FLandscapeVertexFactoryMobilePixelShaderParameters, NonVirtual);
+	DECLARE_TYPE_LAYOUT(FLandscapeVertexFactoryMobilePixelShaderParameters, NonVirtual);
 public:
 	/**
 	* Bind shader constants by name
@@ -139,12 +141,15 @@ public:
 	}
 };
 
+IMPLEMENT_TYPE_LAYOUT(FLandscapeVertexFactoryMobilePixelShaderParameters);
+
 /**
   * Shader parameters for use with FLandscapeFixedGridVertexFactory
   * Simple grid rendering (without dynamic lod blend) needs a simpler fixed setup.
   */
 class FLandscapeFixedGridVertexFactoryMobileVertexShaderParameters : public FLandscapeVertexFactoryMobileVertexShaderParameters
 {
+	DECLARE_TYPE_LAYOUT(FLandscapeFixedGridVertexFactoryMobileVertexShaderParameters, NonVirtual);
 public:
 	void GetElementShaderBindings(
 		const class FSceneInterface* Scene,
@@ -172,6 +177,8 @@ public:
 	}
 };
 
+IMPLEMENT_TYPE_LAYOUT(FLandscapeFixedGridVertexFactoryMobileVertexShaderParameters);
+
 IMPLEMENT_VERTEX_FACTORY_PARAMETER_TYPE(FLandscapeVertexFactoryMobile, SF_Vertex, FLandscapeVertexFactoryMobileVertexShaderParameters);
 IMPLEMENT_VERTEX_FACTORY_PARAMETER_TYPE(FLandscapeVertexFactoryMobile, SF_Pixel, FLandscapeVertexFactoryMobilePixelShaderParameters);
 
diff --git a/Engine/Source/Runtime/RenderCore/Private/Shader.cpp b/Engine/Source/Runtime/RenderCore/Private/Shader.cpp
--- a/Engine/Source/Runtime/RenderCore/Private/Shader.cpp
+++ b/Engine/Source/Runtime/RenderCore/Private/Shader.cpp
@@ -427,8 +427,10 @@ void* FShaderMapPointerTable::GetIndexedPointer(const FTypeLayoutDesc& TypeDesc,
 	return Ptr;
 }
 
-void FShaderMapPointerTable::SaveToArchive(FArchive& Ar, void* FrozenContent, bool bInlineShaderResources) const
+void FShaderMapPointerTable::SaveToArchive(FArchive& Ar, const FPlatformTypeLayoutParameters& LayoutParams, const void* FrozenObject) const
 {
+	FPointerTableBase::SaveToArchive(Ar, LayoutParams, FrozenObject);
+
 	int32 NumTypes = ShaderTypes.Num();
 	int32 NumVFTypes = VFTypes.Num();
 
@@ -450,10 +452,12 @@ void FShaderMapPointerTable::SaveToArchive(FArchive& Ar, void* FrozenContent, bo
 	}
 }
 
-void FShaderMapPointerTable::LoadFromArchive(FArchive& Ar, void* FrozenContent, bool bInlineShaderResources, bool bLoadedByCookedMaterial)
+bool FShaderMapPointerTable::LoadFromArchive(FArchive& Ar, const FPlatformTypeLayoutParameters& LayoutParams, void* FrozenObject)
 {
 	SCOPED_LOADTIMER(FShaderMapPointerTable_LoadFromArchive);
 
+	const bool bResult = FPointerTableBase::LoadFromArchive(Ar, LayoutParams, FrozenObject);
+
 	int32 NumTypes = 0;
 	int32 NumVFTypes = 0;
 		
@@ -477,6 +481,8 @@ void FShaderMapPointerTable::LoadFromArchive(FArchive& Ar, void* FrozenContent,
 		FVertexFactoryType* VFType = FVertexFactoryType::GetVFByName(TypeName);
 		VFTypes.LoadIndexedPointer(VFType);
 	}
+
+	return bResult;
 }
 
 FShaderCompiledShaderInitializerType::FShaderCompiledShaderInitializerType(
diff --git a/Engine/Source/Runtime/RenderCore/Private/ShaderMap.cpp b/Engine/Source/Runtime/RenderCore/Private/ShaderMap.cpp
--- a/Engine/Source/Runtime/RenderCore/Private/ShaderMap.cpp
+++ b/Engine/Source/Runtime/RenderCore/Private/ShaderMap.cpp
@@ -28,11 +28,8 @@
 #include "Interfaces/IShaderFormat.h"
 #endif
 
-FShaderMapBase::FShaderMapBase(const FTypeLayoutDesc& InContentTypeLayout)
-	: ContentTypeLayout(InContentTypeLayout)
-	, PointerTable(nullptr)
-	, Content(nullptr)
-	, FrozenContentSize(0u)
+FShaderMapBase::FShaderMapBase()
+	: PointerTable(nullptr)
 	, NumFrozenShaders(0u)
 {}
 
@@ -54,17 +51,43 @@ FShaderMapResourceCode* FShaderMapBase::GetResourceCode()
 	return Code;
 }
 
-void FShaderMapBase::CopyResourceCode(const FShaderMapResourceCode& Source)
+void FShaderMapBase::AssignContent(TMemoryImageObject<FShaderMapContent> InContent)
 {
-	Code = new FShaderMapResourceCode(Source);
+	check(!Content.Object);
+	check(!PointerTable);
+	const FTypeLayoutDesc& ExpectedTypeDesc = GetContentTypeDesc();
+	checkf(*InContent.TypeDesc == ExpectedTypeDesc, TEXT("FShaderMapBase expected content of type %s, got %s"), ExpectedTypeDesc.Name, InContent.TypeDesc->Name);
+
+	Content = InContent;
+	PointerTable = CreatePointerTable();
 }
 
-void FShaderMapBase::AssignContent(FShaderMapContent* InContent)
+void FShaderMapBase::AssignCopy(const FShaderMapBase& Source)
 {
-	check(!Content);
 	check(!PointerTable);
-	Content = InContent;
-	PointerTable = CreatePointerTable();
+	check(!Code);
+	check(Source.Content.Object);
+
+	if (Source.Content.FrozenSize == 0u)
+	{
+		PointerTable = CreatePointerTable();
+		Content = TMemoryImageObject<FShaderMapContent>(FreezeMemoryImageObject(Source.Content.Object, *Source.Content.TypeDesc, PointerTable));
+	}
+	else
+	{
+		PointerTable = Source.PointerTable->Clone();
+		Content.TypeDesc = Source.Content.TypeDesc;
+		Content.FrozenSize = Source.Content.FrozenSize;
+		Content.Object = static_cast<FShaderMapContent*>(FMemory::Malloc(Content.FrozenSize));
+		FMemory::Memcpy(Content.Object, Source.Content.Object, Content.FrozenSize);
+	}
+
+	NumFrozenShaders = Content.Object->GetNumShaders();
+	INC_DWORD_STAT_BY(STAT_Shaders_ShaderMemory, Content.FrozenSize);
+	INC_DWORD_STAT_BY(STAT_Shaders_NumShadersLoaded, NumFrozenShaders);
+
+	Code = new FShaderMapResourceCode(*Source.Code);
+	InitResource();
 }
 
 void FShaderMapBase::InitResource()
@@ -78,61 +101,23 @@ void FShaderMapBase::InitResource()
 	}
 }
 
-void FShaderMapBase::AssignAndFreezeContent(const FShaderMapContent* InContent)
+void FShaderMapBase::FinalizeContent()
 {
-	FShaderMapPointerTable* LocalPointerTable = nullptr;
-	void* LocalContentMemory = nullptr;
-	uint32 LocalContentSize = 0u;
-	if (InContent)
+	if (Content.Freeze(PointerTable))
 	{
-		LocalPointerTable = CreatePointerTable();
-
-		FMemoryImage MemoryImage;
-		MemoryImage.TargetLayoutParameters.InitializeForCurrent();
-		MemoryImage.PointerTable = LocalPointerTable;
-		FMemoryImageWriter Writer(MemoryImage);
-
-		Writer.WriteObject(InContent, ContentTypeLayout);
-
-		FMemoryImageResult MemoryImageResult;
-		MemoryImage.Flatten(MemoryImageResult, true);
-
-		LocalContentSize = MemoryImageResult.Bytes.Num();
-		check(LocalContentSize > 0u);
-		LocalContentMemory = FMemory::Malloc(LocalContentSize);
-		FMemory::Memcpy(LocalContentMemory, MemoryImageResult.Bytes.GetData(), LocalContentSize);
-		MemoryImageResult.ApplyPatches(LocalContentMemory);
-	}
-
-	DestroyContent();
-
-	if (LocalContentMemory)
-	{
-		PointerTable = LocalPointerTable;
-		Content = static_cast<FShaderMapContent*>(LocalContentMemory);
-		FrozenContentSize = LocalContentSize;
-		NumFrozenShaders = Content->GetNumShaders();
-
-		INC_DWORD_STAT_BY(STAT_Shaders_ShaderMemory, FrozenContentSize);
+		NumFrozenShaders = Content.Object->GetNumShaders();
+		INC_DWORD_STAT_BY(STAT_Shaders_ShaderMemory, Content.FrozenSize);
 		INC_DWORD_STAT_BY(STAT_Shaders_NumShadersLoaded, NumFrozenShaders);
 	}
+	InitResource();
 }
 
 void FShaderMapBase::UnfreezeContent()
 {
-	if (Content && FrozenContentSize > 0u)
-	{
-		// Invoke 'operator new' rather than malloc, as unfrozen memory is expected to be allocate via 'new'
-		void* UnfrozenMemory = ::operator new(ContentTypeLayout.Size);
-
-		FMemoryUnfreezeContent Context;
-		Context.PrevPointerTable = PointerTable;
-		Context.UnfreezeObject(Content, ContentTypeLayout, UnfrozenMemory);
-
-		DestroyContent();
-
-		Content = static_cast<FShaderMapContent*>(UnfrozenMemory);
-	}
+	DEC_DWORD_STAT_BY(STAT_Shaders_ShaderMemory, Content.FrozenSize);
+	DEC_DWORD_STAT_BY(STAT_Shaders_NumShadersLoaded, NumFrozenShaders);
+	Content.Unfreeze(PointerTable);
+	NumFrozenShaders = 0u;
 }
 
 #define CHECK_SHADERMAP_DEPENDENCIES (WITH_EDITOR || !(UE_BUILD_SHIPPING || UE_BUILD_TEST))
@@ -140,45 +125,26 @@ void FShaderMapBase::UnfreezeContent()
 bool FShaderMapBase::Serialize(FArchive& Ar, bool bInlineShaderResources, bool bLoadedByCookedMaterial, bool bInlineShaderCode)
 {
 	LLM_SCOPE(ELLMTag::Shaders);
-	bool bContentValid = true;
 	if (Ar.IsSaving())
 	{
-		check(Content);
-		Content->Validate(*this);
-
-		FShaderMapPointerTable* SavePointerTable = CreatePointerTable();
+		check(Content.Object);
+		Content.Object->Validate(*this);
 
-		FMemoryImage MemoryImage;
-		MemoryImage.PrevPointerTable = PointerTable;
-		MemoryImage.PointerTable = SavePointerTable;
-		MemoryImage.TargetLayoutParameters.InitializeForArchive(Ar);
-
-		FMemoryImageWriter Writer(MemoryImage);
+		{
+			TUniquePtr<FShaderMapPointerTable> SavePointerTable(CreatePointerTable());
 
-		Writer.WriteObject(Content, ContentTypeLayout);
+			FMemoryImage MemoryImage;
+			MemoryImage.PrevPointerTable = PointerTable;
+			MemoryImage.PointerTable = SavePointerTable.Get();
+			MemoryImage.TargetLayoutParameters.InitializeForArchive(Ar);
 
-		FMemoryImageResult MemoryImageResult;
-		MemoryImage.Flatten(MemoryImageResult, true);
+			FMemoryImageWriter Writer(MemoryImage);
+			Writer.WriteRootObject(Content.Object, *Content.TypeDesc);
 
-		void* SaveFrozenContent = MemoryImageResult.Bytes.GetData();
-		uint32 SaveFrozenContentSize = MemoryImageResult.Bytes.Num();
-		check(SaveFrozenContentSize > 0u);
-		Ar << SaveFrozenContentSize;
-		Ar.Serialize(SaveFrozenContent, SaveFrozenContentSize);
-		MemoryImageResult.SaveToArchive(Ar);
-		SavePointerTable->SaveToArchive(Ar, SaveFrozenContent, bInlineShaderResources);
-		delete SavePointerTable;
+			FMemoryImageResult MemoryImageResult;
+			MemoryImage.Flatten(MemoryImageResult, true);
 
-		int32 NumDependencies = MemoryImage.TypeDependencies.Num();
-		Ar << NumDependencies;
-		for (const FTypeLayoutDesc* DependencyTypeDesc : MemoryImage.TypeDependencies)
-		{
-			uint64 NameHash = DependencyTypeDesc->NameHash;
-			FSHAHash LayoutHash;
-			uint32 LayoutSize = Freeze::HashLayout(*DependencyTypeDesc, MemoryImage.TargetLayoutParameters, LayoutHash);
-			Ar << NameHash;
-			Ar << LayoutSize;
-			Ar << LayoutHash;
+			MemoryImageResult.SaveToArchive(Ar);
 		}
 
 		bool bShareCode = false;
@@ -187,6 +153,12 @@ bool FShaderMapBase::Serialize(FArchive& Ar, bool bInlineShaderResources, bool b
 #endif // WITH_EDITOR
 		Ar << bShareCode;
 #if WITH_EDITOR
+
+		// Serialize a copy of ShaderPlatform directly into the archive
+		// This will allow us to correctly deserialize the stream, even if we're not able to load the frozen content
+		TEnumAsByte<EShaderPlatform> ShaderPlatform = GetShaderPlatform();
+		Ar << ShaderPlatform;
+
 		if (Ar.IsCooking())
 		{
 			Code->NotifyShadersCooked(Ar.CookingTarget());
@@ -209,58 +181,14 @@ bool FShaderMapBase::Serialize(FArchive& Ar, bool bInlineShaderResources, bool b
 		check(!PointerTable);
 		PointerTable = CreatePointerTable();
 
-		Ar << FrozenContentSize;
-		// ensure frozen content is at least as big as our FShaderMapContent-derived class
-		checkf(FrozenContentSize >= ContentTypeLayout.Size, TEXT("Invalid FrozenContentSize for %s, got %d, expected at least %d"), ContentTypeLayout.Name, FrozenContentSize, ContentTypeLayout.Size);
-
-		void* ContentMemory = FMemory::Malloc(FrozenContentSize);
-		Ar.Serialize(ContentMemory, FrozenContentSize);
-		Content = static_cast<FShaderMapContent*>(ContentMemory);
-		FMemoryImageResult::ApplyPatchesFromArchive(Content, Ar);
-		PointerTable->LoadFromArchive(Ar, Content, bInlineShaderResources, bLoadedByCookedMaterial);
+		FMemoryImageObject LoadedContent = FMemoryImageResult::LoadFromArchive(Ar, GetContentTypeDesc(), PointerTable);
 
-		int32 NumDependencies = 0;
-		Ar << NumDependencies;
-		if(NumDependencies > 0)
-		{
-#if CHECK_SHADERMAP_DEPENDENCIES
-			FPlatformTypeLayoutParameters LayoutParams;
-			LayoutParams.InitializeForCurrent();
-#endif // CHECK_SHADERMAP_DEPENDENCIES
-
-			// Waste a bit of time even in shipping builds skipping over this stuff
-			// Could add a cook-time option to exclude dependencies completely
-			for (int32 i = 0u; i < NumDependencies; ++i)
-			{
-				uint64 NameHash = 0u;
-				uint32 SavedLayoutSize = 0u;
-				FSHAHash SavedLayoutHash;
-				Ar << NameHash;
-				Ar << SavedLayoutSize;
-				Ar << SavedLayoutHash;
-#if CHECK_SHADERMAP_DEPENDENCIES
-				const FTypeLayoutDesc* DependencyType = FTypeLayoutDesc::Find(NameHash);
-				if (DependencyType)
-				{
-					FSHAHash CheckLayoutHash;
-					const uint32 CheckLayoutSize = Freeze::HashLayout(*DependencyType, LayoutParams, CheckLayoutHash);
-					if (CheckLayoutSize != SavedLayoutSize)
-					{
-						UE_LOG(LogShaders, Error, TEXT("Mismatch size for type %s, compiled size is %d, loaded size is %d"), DependencyType->Name, CheckLayoutSize, SavedLayoutSize);
-						bContentValid = false;
-					}
-					else if (CheckLayoutHash != SavedLayoutHash)
-					{
-						UE_LOG(LogShaders, Error, TEXT("Mismatch hash for type %s"), DependencyType->Name);
-						bContentValid = false;
-					}
-				}
-#endif // CHECK_SHADERMAP_DEPENDENCIES
-			}
-		}
-		
 		bool bShareCode = false;
 		Ar << bShareCode;
+
+		TEnumAsByte<EShaderPlatform> ShaderPlatform = SP_NumPlatforms;
+		Ar << ShaderPlatform;
+
 		if (bShareCode)
 		{
 			FSHAHash ResourceHash;
@@ -270,42 +198,46 @@ bool FShaderMapBase::Serialize(FArchive& Ar, bool bInlineShaderResources, bool b
 			{
 				// do not warn when running -nullrhi (the resource cannot be created as the shader library will not be uninitialized),
 				// also do not warn for shader platforms other than current (if the game targets more than one RHI)
-				if (FApp::CanEverRender() && GetShaderPlatform() == GMaxRHIShaderPlatform)
+				if (FApp::CanEverRender() && ShaderPlatform == GMaxRHIShaderPlatform)
 				{
-					UE_LOG(LogShaders, Error, TEXT("Missing shader resource for hash '%s' for shader platform %d in the shader library"), *ResourceHash.ToString(), GetShaderPlatform());
+					UE_LOG(LogShaders, Error, TEXT("Missing shader resource for hash '%s' for shader platform %d in the shader library"), *ResourceHash.ToString(), ShaderPlatform);
 				}
-				bContentValid = false;
 			}
 		}
 		else
 		{
 			Code = new FShaderMapResourceCode();
 			Code->Serialize(Ar, bLoadedByCookedMaterial);
-			Resource = new FShaderMapResource_InlineCode(GetShaderPlatform(), Code);
+			Resource = new FShaderMapResource_InlineCode(ShaderPlatform, Code);
 		}
 
-		if (bContentValid)
+		if (LoadedContent.Object && Resource)
 		{
-			check(Resource);
-			NumFrozenShaders = Content->GetNumShaders();
+			Content = TMemoryImageObject<FShaderMapContent>(LoadedContent);
 
-			BeginInitResource(Resource);
+			// Possible we've loaded/converted unfrozen content, make sure it's frozen for the current platform before trying to render anything
+			if (Content.FrozenSize == 0u)
+			{
+				Content.Freeze(PointerTable);
+			}
 
-			INC_DWORD_STAT_BY(STAT_Shaders_ShaderResourceMemory, Resource->GetSizeBytes());
-			INC_DWORD_STAT_BY(STAT_Shaders_ShaderMemory, FrozenContentSize);
+			NumFrozenShaders = Content.Object->GetNumShaders();
+			INC_DWORD_STAT_BY(STAT_Shaders_ShaderMemory, Content.FrozenSize);
 			INC_DWORD_STAT_BY(STAT_Shaders_NumShadersLoaded, NumFrozenShaders);
+
+			BeginInitResource(Resource);
+			INC_DWORD_STAT_BY(STAT_Shaders_ShaderResourceMemory, Resource->GetSizeBytes());
 		}
 		else
 		{
+			// Missing either content and/or resource
+			// In either case, shader map has failed to load
+			LoadedContent.Destroy(PointerTable);
 			Resource.SafeRelease();
-
-			// Don't call destructors here, this is basically unknown/invalid memory at this point
-			FMemory::Free(Content);
-			Content = nullptr;
 		}
 	}
 
-	return bContentValid;
+	return (bool)Content.Object;
 }
 
 FString FShaderMapBase::ToString() const
@@ -319,7 +251,7 @@ FString FShaderMapBase::ToString() const
 		FPlatformTypeLayoutParameters LayoutParams;
 		LayoutParams.InitializeForCurrent();
 
-		ContentTypeLayout.ToStringFunc(Content, ContentTypeLayout, LayoutParams, Context);
+		Content.TypeDesc->ToStringFunc(Content.Object, *Content.TypeDesc, LayoutParams, Context);
 	}
 
 	if (Code)
@@ -332,21 +264,10 @@ FString FShaderMapBase::ToString() const
 
 void FShaderMapBase::DestroyContent()
 {
-	if (Content)
-	{
-		DEC_DWORD_STAT_BY(STAT_Shaders_ShaderMemory, FrozenContentSize);
-		DEC_DWORD_STAT_BY(STAT_Shaders_NumShadersLoaded, NumFrozenShaders);
-
-		InternalDeleteObjectFromLayout(Content, ContentTypeLayout, PointerTable, FrozenContentSize > 0u);
-		if (FrozenContentSize > 0u)
-		{
-			FMemory::Free(Content);
-		}
-
-		FrozenContentSize = 0u;
-		NumFrozenShaders = 0u;
-		Content = nullptr;
-	}
+	DEC_DWORD_STAT_BY(STAT_Shaders_ShaderMemory, Content.FrozenSize);
+	DEC_DWORD_STAT_BY(STAT_Shaders_NumShadersLoaded, NumFrozenShaders);
+	Content.Destroy(PointerTable);
+	NumFrozenShaders = 0u;
 }
 
 static uint16 MakeShaderHash(const FHashedName& TypeName, int32 PermutationId)
diff --git a/Engine/Source/Runtime/RenderCore/Public/Shader.h b/Engine/Source/Runtime/RenderCore/Public/Shader.h
--- a/Engine/Source/Runtime/RenderCore/Public/Shader.h
+++ b/Engine/Source/Runtime/RenderCore/Public/Shader.h
@@ -394,8 +394,8 @@ public:
 
 	virtual FShaderMapPointerTable* Clone() const { return new FShaderMapPointerTable(*this); }
 
-	virtual void SaveToArchive(FArchive& Ar, void* FrozenContent, bool bInlineShaderResources) const;
-	virtual void LoadFromArchive(FArchive& Ar, void* FrozenContent, bool bInlineShaderResources, bool bLoadedByCookedMaterial);
+	virtual void SaveToArchive(FArchive& Ar, const FPlatformTypeLayoutParameters& LayoutParams, const void* FrozenObject) const override;
+	virtual bool LoadFromArchive(FArchive& Ar, const FPlatformTypeLayoutParameters& LayoutParams, void* FrozenObject) override;
 
 	TPtrTable<FShaderType> ShaderTypes;
 	TPtrTable<FVertexFactoryType> VFTypes;
@@ -1987,23 +1987,23 @@ public:
 	virtual ~FShaderMapBase();
 
 	FShaderMapResourceCode* GetResourceCode();
-	void CopyResourceCode(const FShaderMapResourceCode& Source);
 
 	inline FShaderMapResource* GetResource() const { return Resource; }
 	inline FShaderMapResource* GetResourceChecked() const { check(Resource); return Resource; }
 	inline const FShaderMapPointerTable& GetPointerTable() const { check(PointerTable); return *PointerTable; }
-	inline const FShaderMapContent* GetContent() const { return Content; }
+	inline const FShaderMapContent* GetContent() const { return Content.Object; }
 	inline FShaderMapContent* GetMutableContent()
 	{
 		UnfreezeContent();
-		return Content;
+		return Content.Object;
 	}
 
-	inline EShaderPlatform GetShaderPlatform() const { return Content ? Content->GetShaderPlatform() : SP_NumPlatforms; }
-	inline uint32 GetFrozenContentSize() const { return FrozenContentSize; }
+	inline EShaderPlatform GetShaderPlatform() const { return Content.Object ? Content.Object->GetShaderPlatform() : SP_NumPlatforms; }
+	inline uint32 GetFrozenContentSize() const { return Content.FrozenSize; }
 
-	void AssignContent(FShaderMapContent* InContent);
+	void AssignContent(TMemoryImageObject<FShaderMapContent> InContent);
 
+	void FinalizeContent();
 	void UnfreezeContent();
 	bool Serialize(FArchive& Ar, bool bInlineShaderResources, bool bLoadedByCookedMaterial, bool bInlineShaderCode=false);
 
@@ -2012,11 +2012,11 @@ public:
 #if WITH_EDITOR
 	inline void GetOutdatedTypes(TArray<const FShaderType*>& OutdatedShaderTypes, TArray<const FShaderPipelineType*>& OutdatedShaderPipelineTypes, TArray<const FVertexFactoryType*>& OutdatedFactoryTypes) const
 	{
-		Content->GetOutdatedTypes(*this, OutdatedShaderTypes, OutdatedShaderPipelineTypes, OutdatedFactoryTypes);
+		Content.Object->GetOutdatedTypes(*this, OutdatedShaderTypes, OutdatedShaderPipelineTypes, OutdatedFactoryTypes);
 	}
 	void SaveShaderStableKeys(EShaderPlatform TargetShaderPlatform, const struct FStableShaderKeyAndValue& SaveKeyVal)
 	{
-		Content->SaveShaderStableKeys(*this, TargetShaderPlatform, SaveKeyVal);
+		Content.Object->SaveShaderStableKeys(*this, TargetShaderPlatform, SaveKeyVal);
 	}
 
 	/** Associates a shadermap with an asset (note: one shadermap can be used by several assets, e.g. MIs). 
@@ -2039,13 +2039,15 @@ public:
 #endif // WITH_EDITOR
 
 protected:
-	explicit FShaderMapBase(const FTypeLayoutDesc& InContentTypeLayout);
+	FShaderMapBase();
+
+	void AssignCopy(const FShaderMapBase& Source);
 
-	void AssignAndFreezeContent(const FShaderMapContent* InContent);
 	void InitResource();
 	void DestroyContent();
 
 protected:
+	virtual const FTypeLayoutDesc& GetContentTypeDesc() const = 0;
 	virtual FShaderMapPointerTable* CreatePointerTable() const = 0;
 
 private:
@@ -2053,12 +2055,10 @@ private:
 	/** List of the assets that are using this shadermap. This is only available in the editor (cooker) to influence ordering of shader libraries. */
 	FShaderMapAssetPaths AssociatedAssets;
 #endif
-	const FTypeLayoutDesc& ContentTypeLayout;
 	TRefCountPtr<FShaderMapResource> Resource;
 	TRefCountPtr<FShaderMapResourceCode> Code;
 	FShaderMapPointerTable* PointerTable;
-	FShaderMapContent* Content;
-	uint32 FrozenContentSize;
+	TMemoryImageObject<FShaderMapContent> Content;
 	uint32 NumFrozenShaders;
 };
 
@@ -2076,12 +2076,11 @@ public:
 		check(LocalContent);
 		LocalContent->Finalize(this->GetResourceCode());
 		LocalContent->Validate(*this);
-		this->AssignAndFreezeContent(LocalContent);
-		this->InitResource();
+		FShaderMapBase::FinalizeContent();
 	}
 
 protected:
-	TShaderMap() : FShaderMapBase(StaticGetTypeLayoutDesc<ContentType>()) {}
+	virtual const FTypeLayoutDesc& GetContentTypeDesc() const final override { return StaticGetTypeLayoutDesc<ContentType>(); }
 	virtual FShaderMapPointerTable* CreatePointerTable() const final override { return new PointerTableType(); }
 };
 
diff --git a/Engine/Source/Runtime/Renderer/Private/ShaderBaseClasses.cpp b/Engine/Source/Runtime/Renderer/Private/ShaderBaseClasses.cpp
--- a/Engine/Source/Runtime/Renderer/Private/ShaderBaseClasses.cpp
+++ b/Engine/Source/Runtime/Renderer/Private/ShaderBaseClasses.cpp
@@ -110,6 +110,12 @@ void FMaterialShader::VerifyExpressionAndShaderMaps(const FMaterialRenderProxy*
 {
 	// Validate that the shader is being used for a material that matches the uniform expression set the shader was compiled for.
 	FMaterialShaderMap* ShaderMap = Material.GetRenderingThreadShaderMap();
+	if (ShaderMap->GetShaderMapId().IsCookedId())
+	{
+		// If this is a cooked shadermap, the debug uniform expression set won't be properly initialized...so skip the verification
+		return;
+	}
+
 	const FUniformExpressionSet& MaterialUniformExpressionSet = ShaderMap->GetUniformExpressionSet();
 	bool bUniformExpressionSetMismatch = !DebugUniformExpressionSet.Matches(MaterialUniformExpressionSet)
 		|| UniformExpressionCache->CachedUniformExpressionShaderMap != ShaderMap;
