Fixed ClangEditor compile error

diff --git a/Engine/Source/Runtime/Core/Private/ProfilingDebugging/MemoryTrace.h b/Engine/Source/Runtime/Core/Private/ProfilingDebugging/MemoryTrace.h
--- a/Engine/Source/Runtime/Core/Private/ProfilingDebugging/MemoryTrace.h
+++ b/Engine/Source/Runtime/Core/Private/ProfilingDebugging/MemoryTrace.h
@@ -1,8 +1,6 @@
 // Copyright Epic Games, Inc. All Rights Reserved.
 
-#include "Trace/Config.h"
-
-//#define UE_MEMORY_TRACE_ENABLED 0
+#include "Trace/Trace.h"
 
 #if !defined(UE_MEMORY_TRACE_ENABLED) && UE_TRACE_ENABLED
 #	if PLATFORM_WINDOWS
diff --git a/Engine/Source/Runtime/Core/Private/ProfilingDebugging/Windows/MemoryTrace.cpp b/Engine/Source/Runtime/Core/Private/ProfilingDebugging/Windows/MemoryTrace.cpp
--- a/Engine/Source/Runtime/Core/Private/ProfilingDebugging/Windows/MemoryTrace.cpp
+++ b/Engine/Source/Runtime/Core/Private/ProfilingDebugging/Windows/MemoryTrace.cpp
@@ -26,9 +26,8 @@ void*	Backtracer_GetBacktraceId(void*);
 ////////////////////////////////////////////////////////////////////////////////
 #if defined(_MSC_VER)
 #	define UE_RETURN_ADDRESS_ADDRESS()	(void*)_AddressOfReturnAddress()
-#else if defined(__clang__) || defined(__GNUC__)
-#	define UE_RETURN_ADDRESS_ADDRESS()
-#	error TODO
+#elif defined(__clang__) || defined(__GNUC__)
+#	define UE_RETURN_ADDRESS_ADDRESS()	__builtin_frame_address(0)
 #endif
 
 FORCEINLINE void* GetOwner(bool bLight)
