Fix header guards

diff --git a/Engine/Source/Runtime/Core/Private/ProfilingDebugging/MemoryTrace.h b/Engine/Source/Runtime/Core/Private/ProfilingDebugging/MemoryTrace.h
--- a/Engine/Source/Runtime/Core/Private/ProfilingDebugging/MemoryTrace.h
+++ b/Engine/Source/Runtime/Core/Private/ProfilingDebugging/MemoryTrace.h
@@ -1,4 +1,5 @@
 // Copyright Epic Games, Inc. All Rights Reserved.
+#pragma once
 
 #include "Trace/Trace.h"
 
