Add performance tags on compression and report back error message on ZLIB errors

diff --git a/Engine/Source/Runtime/Core/Private/Misc/Compression.cpp b/Engine/Source/Runtime/Core/Private/Misc/Compression.cpp
--- a/Engine/Source/Runtime/Core/Private/Misc/Compression.cpp
+++ b/Engine/Source/Runtime/Core/Private/Misc/Compression.cpp
@@ -67,6 +67,7 @@ static uint32 appGZIPVersion()
  */
 static bool appCompressMemoryZLIB(void* CompressedBuffer, int32& CompressedSize, const void* UncompressedBuffer, int32 UncompressedSize, int32 BitWindow, int32 CompLevel)
 {
+	TRACE_CPUPROFILER_EVENT_SCOPE(appCompressMemoryZLIB);
 	DECLARE_SCOPE_CYCLE_COUNTER(TEXT("Compress Memory ZLIB"), STAT_appCompressMemoryZLIB, STATGROUP_Compression);
 
 	ensureMsgf(CompLevel >= Z_DEFAULT_COMPRESSION, TEXT("CompLevel must be >= Z_DEFAULT_COMPRESSION"));
@@ -120,6 +121,7 @@ static bool appCompressMemoryZLIB(void* CompressedBuffer, int32& CompressedSize,
 
 static bool appCompressMemoryGZIP(void* CompressedBuffer, int32& CompressedSize, const void* UncompressedBuffer, int32 UncompressedSize)
 {
+	TRACE_CPUPROFILER_EVENT_SCOPE(appCompressMemoryGZIP);
 	DECLARE_SCOPE_CYCLE_COUNTER( TEXT( "Compress Memory GZIP" ), STAT_appCompressMemoryGZIP, STATGROUP_Compression );
 
 	z_stream gzipstream;
@@ -165,6 +167,7 @@ static bool appCompressMemoryGZIP(void* CompressedBuffer, int32& CompressedSize,
 
 static int appCompressMemoryBoundGZIP(int32 UncompressedSize)
 {
+	TRACE_CPUPROFILER_EVENT_SCOPE(appCompressMemoryBoundGZIP);
 	DECLARE_SCOPE_CYCLE_COUNTER(TEXT("Compress Memory Bound GZIP"), STAT_appCompressMemoryBoundGZIP, STATGROUP_Compression);
 	z_stream gzipstream;
 	gzipstream.zalloc = &zalloc;
@@ -199,6 +202,7 @@ static int appCompressMemoryBoundGZIP(int32 UncompressedSize)
  */
 bool appUncompressMemoryGZIP(void* UncompressedBuffer, int32 UncompressedSize, const void* CompressedBuffer, int32 CompressedSize)
 {
+	TRACE_CPUPROFILER_EVENT_SCOPE(appUncompressMemoryGZIP);
 	DECLARE_SCOPE_CYCLE_COUNTER(TEXT("Uncompress Memory GZIP"), STAT_appUncompressMemoryGZIP, STATGROUP_Compression);
 
 	// Zlib wants to use unsigned long.
@@ -260,6 +264,7 @@ bool appUncompressMemoryGZIP(void* UncompressedBuffer, int32 UncompressedSize, c
  */
 bool appUncompressMemoryZLIB( void* UncompressedBuffer, int32 UncompressedSize, const void* CompressedBuffer, int32 CompressedSize, int32 BitWindow )
 {
+	TRACE_CPUPROFILER_EVENT_SCOPE(appUncompressMemoryZLIB);
 	DECLARE_SCOPE_CYCLE_COUNTER( TEXT( "Uncompress Memory ZLIB" ), STAT_appUncompressMemoryZLIB, STATGROUP_Compression );
 
 	// Zlib wants to use unsigned long.
@@ -299,9 +304,9 @@ bool appUncompressMemoryZLIB( void* UncompressedBuffer, int32 UncompressedSize,
 	}
 
 	// These warnings will be compiled out in shipping.
-	UE_CLOG(Result == Z_MEM_ERROR, LogCompression, Warning, TEXT("appUncompressMemoryZLIB failed: Error: Z_MEM_ERROR, not enough memory!"));
-	UE_CLOG(Result == Z_BUF_ERROR, LogCompression, Warning, TEXT("appUncompressMemoryZLIB failed: Error: Z_BUF_ERROR, not enough room in the output buffer!"));
-	UE_CLOG(Result == Z_DATA_ERROR, LogCompression, Warning, TEXT("appUncompressMemoryZLIB failed: Error: Z_DATA_ERROR, input data was corrupted or incomplete!"));
+	UE_CLOG(Result == Z_MEM_ERROR, LogCompression, Warning, TEXT("appUncompressMemoryZLIB failed: Error: Z_MEM_ERROR, not enough memory! (%s)"), UTF8_TO_TCHAR(stream.msg));
+	UE_CLOG(Result == Z_BUF_ERROR, LogCompression, Warning, TEXT("appUncompressMemoryZLIB failed: Error: Z_BUF_ERROR, not enough room in the output buffer! (%s)"), UTF8_TO_TCHAR(stream.msg));
+	UE_CLOG(Result == Z_DATA_ERROR, LogCompression, Warning, TEXT("appUncompressMemoryZLIB failed: Error: Z_DATA_ERROR, input data was corrupted or incomplete! (%s)"), UTF8_TO_TCHAR(stream.msg));
 
 	bool bOperationSucceeded = (Result == Z_OK);
 
@@ -316,6 +321,7 @@ bool appUncompressMemoryZLIB( void* UncompressedBuffer, int32 UncompressedSize,
 
 bool appUncompressMemoryStreamZLIB(void* UncompressedBuffer, int32 UncompressedSize, IMemoryReadStream* Stream, int64 StreamOffset, int32 CompressedSize, int32 BitWindow)
 {
+	TRACE_CPUPROFILER_EVENT_SCOPE(appUncompressMemoryStreamZLIB);
 	DECLARE_SCOPE_CYCLE_COUNTER(TEXT("Uncompress Memory ZLIB"), STAT_appUncompressMemoryZLIB, STATGROUP_Compression);
 
 	int64 ChunkOffset = 0;
@@ -484,6 +490,7 @@ FName FCompression::GetCompressionFormatFromDeprecatedFlags(ECompressionFlags Fl
 
 int32 FCompression::CompressMemoryBound(FName FormatName, int32 UncompressedSize, ECompressionFlags Flags, int32 CompressionData)
 {
+	TRACE_CPUPROFILER_EVENT_SCOPE(FCompression::CompressMemoryBound);
 	int32 CompressionBound = UncompressedSize;
 
 	if (FormatName == NAME_Zlib)
@@ -524,6 +531,7 @@ int32 FCompression::CompressMemoryBound(FName FormatName, int32 UncompressedSize
 
 bool FCompression::CompressMemory(FName FormatName, void* CompressedBuffer, int32& CompressedSize, const void* UncompressedBuffer, int32 UncompressedSize, ECompressionFlags Flags, int32 CompressionData)
 {
+	TRACE_CPUPROFILER_EVENT_SCOPE(FCompression::CompressMemory);
 	uint64 CompressorStartTime = FPlatformTime::Cycles64();
 
 	bool bCompressSucceeded = false;
@@ -602,6 +610,7 @@ DECLARE_FLOAT_ACCUMULATOR_STAT(TEXT("Uncompressor total time"),STAT_Uncompressor
 
 bool FCompression::UncompressMemory(FName FormatName, void* UncompressedBuffer, int32 UncompressedSize, const void* CompressedBuffer, int32 CompressedSize, ECompressionFlags Flags, int32 CompressionData)
 {
+	TRACE_CPUPROFILER_EVENT_SCOPE(FCompression::UncompressMemory);
 	SCOPED_NAMED_EVENT(FCompression_UncompressMemory, FColor::Cyan);
 	// Keep track of time spent uncompressing memory.
 	STAT(double UncompressorStartTime = FPlatformTime::Seconds();)
@@ -666,6 +675,8 @@ bool FCompression::UncompressMemory(FName FormatName, void* UncompressedBuffer,
 
 bool FCompression::UncompressMemoryStream(FName FormatName, void* UncompressedBuffer, int32 UncompressedSize, IMemoryReadStream* Stream, int64 StreamOffset, int32 CompressedSize, ECompressionFlags Flags, int32 CompressionData)
 {
+	TRACE_CPUPROFILER_EVENT_SCOPE(FCompression::UncompressMemoryStream);
+
 	int64 ContiguousChunkSize = 0;
 	const void* ContiguousMemory = Stream->Read(ContiguousChunkSize, StreamOffset, CompressedSize);
 	bool bUncompressResult = false;
