fixed FMallocLeakDetection destruction on exit: if leak detection is active, destruction of the FMallocLeakDetection singleton leads freeing memory of its members, which calls back to half destroyed singleton instance. The fix cleans insternal state and disables detection in the destructor, while the instance still usable

diff --git a/Engine/Source/Runtime/Core/Private/HAL/MallocLeakDetection.cpp b/Engine/Source/Runtime/Core/Private/HAL/MallocLeakDetection.cpp
--- a/Engine/Source/Runtime/Core/Private/HAL/MallocLeakDetection.cpp
+++ b/Engine/Source/Runtime/Core/Private/HAL/MallocLeakDetection.cpp
@@ -65,6 +65,8 @@ FMallocLeakDetection& FMallocLeakDetection::Get()
 
 FMallocLeakDetection::~FMallocLeakDetection()
 {	
+	OpenPointers.Empty(); // clean up the state
+	SetAllocationCollection(false); // disable colletion to avoid a call back to this instance when its members are destroyed
 }
 
 
