Static analysis fixes

diff --git a/Engine/Source/Runtime/Core/Private/ProfilingDebugging/Windows/CallstackTrace.cpp b/Engine/Source/Runtime/Core/Private/ProfilingDebugging/Windows/CallstackTrace.cpp
--- a/Engine/Source/Runtime/Core/Private/ProfilingDebugging/Windows/CallstackTrace.cpp
+++ b/Engine/Source/Runtime/Core/Private/ProfilingDebugging/Windows/CallstackTrace.cpp
@@ -60,7 +60,7 @@ static uint32 AddressToId(UPTRINT Address)
 
 static UPTRINT IdToAddress(uint32 Id)
 {
-	return UPTRINT(Id << 16);
+	return static_cast<uint32>(UPTRINT(Id) << 16);
 }
 
 struct FIdPredicate
diff --git a/Engine/Source/Runtime/Core/Private/ProfilingDebugging/Windows/ModuleDiagnostics.cpp b/Engine/Source/Runtime/Core/Private/ProfilingDebugging/Windows/ModuleDiagnostics.cpp
--- a/Engine/Source/Runtime/Core/Private/ProfilingDebugging/Windows/ModuleDiagnostics.cpp
+++ b/Engine/Source/Runtime/Core/Private/ProfilingDebugging/Windows/ModuleDiagnostics.cpp
@@ -16,6 +16,7 @@ struct FNtDllFunction
 	FNtDllFunction(const char* Name)
 	{
 		HMODULE NtDll = LoadLibraryW(L"ntdll.dll");
+		check(NtDll);
 		Addr = GetProcAddress(NtDll, Name);
 	}
 
