Output owner

diff --git a/Engine/Source/Runtime/Core/Private/ProfilingDebugging/Windows/MemoryTrace.cpp b/Engine/Source/Runtime/Core/Private/ProfilingDebugging/Windows/MemoryTrace.cpp
--- a/Engine/Source/Runtime/Core/Private/ProfilingDebugging/Windows/MemoryTrace.cpp
+++ b/Engine/Source/Runtime/Core/Private/ProfilingDebugging/Windows/MemoryTrace.cpp
@@ -120,7 +120,7 @@ UE_TRACE_EVENT_BEGIN(Memory, Realloc)
 UE_TRACE_EVENT_END()
 
 UE_TRACE_EVENT_BEGIN(Memory, Free)
-	UE_TRACE_EVENT_FIELD(void*, Owner)
+	UE_TRACE_EVENT_FIELD(uint64, Owner)
 	UE_TRACE_EVENT_FIELD(void*, FreeAddress)
 UE_TRACE_EVENT_END()
 
@@ -337,7 +337,8 @@ void FAllocationTrace::Realloc(void* PrevAddress, void* Address, size_t NewSize,
 void FAllocationTrace::Free(void* Address, void* Owner)
 {
 	UE_TRACE_LOG(Memory, Free, MemAllocChannel)
-		<< Free.FreeAddress(Address);
+		<< Free.FreeAddress(Address)
+		<< Free.Owner(uint64(Owner));
 }
 
 
