Reduce log level of Mac platform GPU check.

#rnx

diff --git a/Engine/Source/Runtime/Core/Private/Mac/MacPlatformMisc.cpp b/Engine/Source/Runtime/Core/Private/Mac/MacPlatformMisc.cpp
--- a/Engine/Source/Runtime/Core/Private/Mac/MacPlatformMisc.cpp
+++ b/Engine/Source/Runtime/Core/Private/Mac/MacPlatformMisc.cpp
@@ -2317,7 +2317,7 @@ bool FMacPlatformMisc::HasPlatformFeature(const TCHAR* FeatureName)
                             NSString* GPUName = [GPU.GPUName length] > 0 ? GPU.GPUName : @"Unnamed GPU";
                             NSString* GPUMetalBundle = [GPU.GPUMetalBundle length] > 0 ? GPU.GPUMetalBundle : @"null";
                             
-                            UE_LOG(LogMac, Error, TEXT("Checking GPU %s (MetalBundle %s"), *FString(GPUName), *FString(GPUMetalBundle));
+                            UE_LOG(LogMac, Display, TEXT("Checking GPU: %s (MetalBundle: %s)"), *FString(GPUName), *FString(GPUMetalBundle));
 
                             if ([GPU.GPUMetalBundle length] > 0)
                             {
