SharedBuffer: Renamed immutable to read-only in the shared buffer

#rnx

diff --git a/Engine/Source/Runtime/Core/Private/Memory/SharedBuffer.cpp b/Engine/Source/Runtime/Core/Private/Memory/SharedBuffer.cpp
--- a/Engine/Source/Runtime/Core/Private/Memory/SharedBuffer.cpp
+++ b/Engine/Source/Runtime/Core/Private/Memory/SharedBuffer.cpp
@@ -46,11 +46,11 @@ void FSharedBuffer::ReleaseData()
 	ReferenceCountAndFlags.fetch_and(~SetFlags(~ESharedBufferFlags::None));
 }
 
-bool FSharedBuffer::TryMakeImmutable() const
+bool FSharedBuffer::TryMakeReadOnly() const
 {
 	for (uint64 Value = ReferenceCountAndFlags.load(std::memory_order_relaxed);;)
 	{
-		if (EnumHasAnyFlags(GetFlags(Value), ESharedBufferFlags::Immutable))
+		if (EnumHasAnyFlags(GetFlags(Value), ESharedBufferFlags::ReadOnly))
 		{
 			return true;
 		}
@@ -59,7 +59,7 @@ bool FSharedBuffer::TryMakeImmutable() const
 		{
 			return false;
 		}
-		if (ReferenceCountAndFlags.compare_exchange_weak(Value, Value | SetFlags(ESharedBufferFlags::Immutable),
+		if (ReferenceCountAndFlags.compare_exchange_weak(Value, Value | SetFlags(ESharedBufferFlags::ReadOnly),
 			std::memory_order_relaxed, std::memory_order_relaxed))
 		{
 			return true;
@@ -67,7 +67,7 @@ bool FSharedBuffer::TryMakeImmutable() const
 	}
 }
 
-bool FSharedBuffer::TryMakeMutable() const
+bool FSharedBuffer::TryMakeWritable() const
 {
 	for (uint64 Value = ReferenceCountAndFlags.load(std::memory_order_relaxed);;)
 	{
@@ -76,11 +76,11 @@ bool FSharedBuffer::TryMakeMutable() const
 		{
 			return false;
 		}
-		if (!EnumHasAnyFlags(GetFlags(Value), ESharedBufferFlags::Immutable))
+		if (!EnumHasAnyFlags(GetFlags(Value), ESharedBufferFlags::ReadOnly))
 		{
 			return true;
 		}
-		if (ReferenceCountAndFlags.compare_exchange_weak(Value, Value & ~SetFlags(ESharedBufferFlags::Immutable),
+		if (ReferenceCountAndFlags.compare_exchange_weak(Value, Value & ~SetFlags(ESharedBufferFlags::ReadOnly),
 			std::memory_order_relaxed, std::memory_order_relaxed))
 		{
 			return true;
diff --git a/Engine/Source/Runtime/Core/Private/Tests/Memory/SharedBufferTest.cpp b/Engine/Source/Runtime/Core/Private/Tests/Memory/SharedBufferTest.cpp
--- a/Engine/Source/Runtime/Core/Private/Tests/Memory/SharedBufferTest.cpp
+++ b/Engine/Source/Runtime/Core/Private/Tests/Memory/SharedBufferTest.cpp
@@ -244,134 +244,134 @@ bool FSharedBufferTest::RunTest(const FString& Parameters)
 		TestFalse(TEXT("FSharedBuffer::MakeOwned(Null)"), FSharedBuffer::MakeOwned(FSharedBufferConstPtr()).IsValid());
 	}
 
-	// Test MakeImmutable
+	// Test MakeReadOnly
 	{
-		// MakeImmutable from a new reference.
+		// MakeReadOnly from a new reference.
 		constexpr uint64 Size = 64;
-		FSharedBufferConstRef Ref = FSharedBuffer::MakeImmutable(FSharedBuffer::Alloc(Size));
-		TestTrue(TEXT("FSharedBuffer::MakeImmutable(AllocRef).IsOwned()"), Ref->IsOwned());
-		TestTrue(TEXT("FSharedBuffer::MakeImmutable(AllocRef).IsImmutable()"), Ref->IsImmutable());
-		TestEqual(TEXT("FSharedBuffer::MakeImmutable(AllocRef).GetSize()"), Ref->GetSize(), Size);
+		FSharedBufferConstRef Ref = FSharedBuffer::MakeReadOnly(FSharedBuffer::Alloc(Size));
+		TestTrue(TEXT("FSharedBuffer::MakeReadOnly(AllocRef).IsOwned()"), Ref->IsOwned());
+		TestTrue(TEXT("FSharedBuffer::MakeReadOnly(AllocRef).IsReadOnly()"), Ref->IsReadOnly());
+		TestEqual(TEXT("FSharedBuffer::MakeReadOnly(AllocRef).GetSize()"), Ref->GetSize(), Size);
 		FSharedBufferConstRef OtherRef = Ref;
-		FSharedBufferConstRef ImmutableRef = FSharedBuffer::MakeImmutable(MoveTemp(Ref));
-		TestEqual(TEXT("FSharedBuffer::MakeImmutable(AllocRef).GetData()"), Ref->GetData(), ImmutableRef->GetData());
+		FSharedBufferConstRef ReadOnlyRef = FSharedBuffer::MakeReadOnly(MoveTemp(Ref));
+		TestEqual(TEXT("FSharedBuffer::MakeReadOnly(AllocRef).GetData()"), Ref->GetData(), ReadOnlyRef->GetData());
 	}
 	{
-		// MakeImmutable from a new pointer.
+		// MakeReadOnly from a new pointer.
 		constexpr uint64 Size = 64;
 		FSharedBufferPtr MutablePtr = FSharedBuffer::Alloc(Size);
-		TestFalse(TEXT("FSharedBuffer::MakeImmutable(MovePtr).IsImmutable()"), MutablePtr->IsImmutable());
-		TestTrue(TEXT("FSharedBuffer::MakeImmutable(MovePtr).IsOwned()"), MutablePtr->IsOwned());
+		TestFalse(TEXT("FSharedBuffer::MakeReadOnly(MovePtr).IsReadOnly()"), MutablePtr->IsReadOnly());
+		TestTrue(TEXT("FSharedBuffer::MakeReadOnly(MovePtr).IsOwned()"), MutablePtr->IsOwned());
 		const void* const Data = MutablePtr->GetData();
-		FSharedBufferConstPtr Ptr = FSharedBuffer::MakeImmutable(MoveTemp(MutablePtr));
-		TestTrue(TEXT("FSharedBuffer::MakeImmutable(MovePtr).IsOwned()"), Ptr->IsOwned());
-		TestTrue(TEXT("FSharedBuffer::MakeImmutable(MovePtr).IsImmutable()"), Ptr->IsImmutable());
-		TestEqual(TEXT("FSharedBuffer::MakeImmutable(MovePtr).GetSize()"), Ptr->GetSize(), Size);
-		TestEqual(TEXT("FSharedBuffer::MakeImmutable(MovePtr).GetData()"), Ptr->GetData(), Data);
+		FSharedBufferConstPtr Ptr = FSharedBuffer::MakeReadOnly(MoveTemp(MutablePtr));
+		TestTrue(TEXT("FSharedBuffer::MakeReadOnly(MovePtr).IsOwned()"), Ptr->IsOwned());
+		TestTrue(TEXT("FSharedBuffer::MakeReadOnly(MovePtr).IsReadOnly()"), Ptr->IsReadOnly());
+		TestEqual(TEXT("FSharedBuffer::MakeReadOnly(MovePtr).GetSize()"), Ptr->GetSize(), Size);
+		TestEqual(TEXT("FSharedBuffer::MakeReadOnly(MovePtr).GetData()"), Ptr->GetData(), Data);
 	}
 	{
-		// MakeImmutable from a view.
+		// MakeReadOnly from a view.
 		constexpr uint64 Size = 64;
 		uint8 Data[Size]{};
 		FSharedBufferPtr MutablePtr = FSharedBuffer::MakeView(Data, Size);
-		TestFalse(TEXT("FSharedBuffer::MakeImmutable(View).IsOwned()"), MutablePtr->IsOwned());
-		TestFalse(TEXT("FSharedBuffer::MakeImmutable(View).IsImmutable()"), MutablePtr->IsImmutable());
-		FSharedBufferConstPtr Ptr = FSharedBuffer::MakeImmutable(MoveTemp(MutablePtr));
-		TestTrue(TEXT("FSharedBuffer::MakeImmutable(View).IsOwned()"), Ptr->IsOwned());
-		TestTrue(TEXT("FSharedBuffer::MakeImmutable(View).IsImmutable()"), Ptr->IsImmutable());
-		TestEqual(TEXT("FSharedBuffer::MakeImmutable(View).GetSize()"), Ptr->GetSize(), Size);
-		TestNotEqual(TEXT("FSharedBuffer::MakeImmutable(View).GetData()"), Ptr->GetData(), static_cast<const void*>(Data));
+		TestFalse(TEXT("FSharedBuffer::MakeReadOnly(View).IsOwned()"), MutablePtr->IsOwned());
+		TestFalse(TEXT("FSharedBuffer::MakeReadOnly(View).IsReadOnly()"), MutablePtr->IsReadOnly());
+		FSharedBufferConstPtr Ptr = FSharedBuffer::MakeReadOnly(MoveTemp(MutablePtr));
+		TestTrue(TEXT("FSharedBuffer::MakeReadOnly(View).IsOwned()"), Ptr->IsOwned());
+		TestTrue(TEXT("FSharedBuffer::MakeReadOnly(View).IsReadOnly()"), Ptr->IsReadOnly());
+		TestEqual(TEXT("FSharedBuffer::MakeReadOnly(View).GetSize()"), Ptr->GetSize(), Size);
+		TestNotEqual(TEXT("FSharedBuffer::MakeReadOnly(View).GetData()"), Ptr->GetData(), static_cast<const void*>(Data));
 	}
 	{
-		// MakeImmutable with another shared reference.
+		// MakeReadOnly with another shared reference.
 		constexpr uint64 Size = 64;
 		FSharedBufferPtr MutablePtr = FSharedBuffer::Alloc(Size);
 		FSharedBufferConstPtr SharedPtr = MutablePtr;
-		TestFalse(TEXT("FSharedBuffer::MakeImmutable(SharedRef).IsImmutable()"), MutablePtr->IsImmutable());
+		TestFalse(TEXT("FSharedBuffer::MakeReadOnly(SharedRef).IsReadOnly()"), MutablePtr->IsReadOnly());
 		const void* const Data = MutablePtr->GetData();
-		FSharedBufferConstPtr Ptr = FSharedBuffer::MakeImmutable(MoveTemp(MutablePtr));
-		TestTrue(TEXT("FSharedBuffer::MakeImmutable(SharedRef).IsOwned()"), Ptr->IsOwned());
-		TestTrue(TEXT("FSharedBuffer::MakeImmutable(SharedRef).IsImmutable()"), Ptr->IsImmutable());
-		TestNotEqual(TEXT("FSharedBuffer::MakeImmutable(SharedRef).GetData()"), Ptr->GetData(), Data);
+		FSharedBufferConstPtr Ptr = FSharedBuffer::MakeReadOnly(MoveTemp(MutablePtr));
+		TestTrue(TEXT("FSharedBuffer::MakeReadOnly(SharedRef).IsOwned()"), Ptr->IsOwned());
+		TestTrue(TEXT("FSharedBuffer::MakeReadOnly(SharedRef).IsReadOnly()"), Ptr->IsReadOnly());
+		TestNotEqual(TEXT("FSharedBuffer::MakeReadOnly(SharedRef).GetData()"), Ptr->GetData(), Data);
 	}
 	{
-		// MakeImmutable with another weak reference.
+		// MakeReadOnly with another weak reference.
 		constexpr uint64 Size = 64;
 		FSharedBufferPtr MutablePtr = FSharedBuffer::Alloc(Size);
 		FSharedBufferConstWeakPtr WeakPtr = MutablePtr;
-		TestFalse(TEXT("FSharedBuffer::MakeImmutable(WeakRef).IsImmutable()"), MutablePtr->IsImmutable());
+		TestFalse(TEXT("FSharedBuffer::MakeReadOnly(WeakRef).IsReadOnly()"), MutablePtr->IsReadOnly());
 		const void* const Data = MutablePtr->GetData();
-		FSharedBufferConstPtr Ptr = FSharedBuffer::MakeImmutable(MoveTemp(MutablePtr));
-		TestTrue(TEXT("FSharedBuffer::MakeImmutable(WeakRef).IsOwned()"), Ptr->IsOwned());
-		TestTrue(TEXT("FSharedBuffer::MakeImmutable(WeakRef).IsImmutable()"), Ptr->IsImmutable());
-		TestNotEqual(TEXT("FSharedBuffer::MakeImmutable(WeakRef).GetData()"), Ptr->GetData(), Data);
+		FSharedBufferConstPtr Ptr = FSharedBuffer::MakeReadOnly(MoveTemp(MutablePtr));
+		TestTrue(TEXT("FSharedBuffer::MakeReadOnly(WeakRef).IsOwned()"), Ptr->IsOwned());
+		TestTrue(TEXT("FSharedBuffer::MakeReadOnly(WeakRef).IsReadOnly()"), Ptr->IsReadOnly());
+		TestNotEqual(TEXT("FSharedBuffer::MakeReadOnly(WeakRef).GetData()"), Ptr->GetData(), Data);
 	}
 
-	// Test MakeMutable
+	// Test MakeWritable
 	{
-		// MakeMutable from a new reference.
+		// MakeWritable from a new reference.
 		constexpr uint64 Size = 64;
-		FSharedBufferConstRef ConstRef = FSharedBuffer::MakeImmutable(FSharedBuffer::Alloc(Size));
-		FSharedBufferRef Ref = FSharedBuffer::MakeMutable(MoveTemp(ConstRef));
-		TestTrue(TEXT("FSharedBuffer::MakeMutable(AllocRef).IsOwned()"), Ref->IsOwned());
-		TestFalse(TEXT("FSharedBuffer::MakeMutable(AllocRef).IsImmutable()"), Ref->IsImmutable());
-		TestEqual(TEXT("FSharedBuffer::MakeMutable(AllocRef).GetSize()"), Ref->GetSize(), Size);
-		TestEqual(TEXT("FSharedBuffer::MakeMutable(AllocRef).GetData()"), const_cast<const void*>(Ref->GetData()), ConstRef->GetData());
-		FSharedBufferRef MutableRef = FSharedBuffer::MakeMutable(MoveTemp(Ref));
-		TestEqual(TEXT("FSharedBuffer::MakeMutable(AllocRef).GetData()"), Ref->GetData(), MutableRef->GetData());
+		FSharedBufferConstRef ConstRef = FSharedBuffer::MakeReadOnly(FSharedBuffer::Alloc(Size));
+		FSharedBufferRef Ref = FSharedBuffer::MakeWritable(MoveTemp(ConstRef));
+		TestTrue(TEXT("FSharedBuffer::MakeWritable(AllocRef).IsOwned()"), Ref->IsOwned());
+		TestFalse(TEXT("FSharedBuffer::MakeWritable(AllocRef).IsReadOnly()"), Ref->IsReadOnly());
+		TestEqual(TEXT("FSharedBuffer::MakeWritable(AllocRef).GetSize()"), Ref->GetSize(), Size);
+		TestEqual(TEXT("FSharedBuffer::MakeWritable(AllocRef).GetData()"), const_cast<const void*>(Ref->GetData()), ConstRef->GetData());
+		FSharedBufferRef MutableRef = FSharedBuffer::MakeWritable(MoveTemp(Ref));
+		TestEqual(TEXT("FSharedBuffer::MakeWritable(AllocRef).GetData()"), Ref->GetData(), MutableRef->GetData());
 	}
 	{
-		// MakeMutable from a new pointer.
+		// MakeWritable from a new pointer.
 		constexpr uint64 Size = 64;
-		FSharedBufferConstPtr ConstPtr = FSharedBuffer::MakeImmutable(FSharedBuffer::Alloc(Size));
+		FSharedBufferConstPtr ConstPtr = FSharedBuffer::MakeReadOnly(FSharedBuffer::Alloc(Size));
 		const void* const Data = ConstPtr->GetData();
-		FSharedBufferPtr Ptr = FSharedBuffer::MakeMutable(MoveTemp(ConstPtr));
-		TestFalse(TEXT("FSharedBuffer::MakeImmutable(MovePtr).IsMutable()"), Ptr->IsImmutable());
-		TestTrue(TEXT("FSharedBuffer::MakeImmutable(MovePtr).IsOwned()"), Ptr->IsOwned());
-		TestEqual(TEXT("FSharedBuffer::MakeImmutable(MovePtr).GetSize()"), Ptr->GetSize(), Size);
-		TestEqual(TEXT("FSharedBuffer::MakeImmutable(MovePtr).GetData()"), const_cast<const void*>(Ptr->GetData()), Data);
+		FSharedBufferPtr Ptr = FSharedBuffer::MakeWritable(MoveTemp(ConstPtr));
+		TestFalse(TEXT("FSharedBuffer::MakeReadOnly(MovePtr).IsReadOnly()"), Ptr->IsReadOnly());
+		TestTrue(TEXT("FSharedBuffer::MakeReadOnly(MovePtr).IsOwned()"), Ptr->IsOwned());
+		TestEqual(TEXT("FSharedBuffer::MakeReadOnly(MovePtr).GetSize()"), Ptr->GetSize(), Size);
+		TestEqual(TEXT("FSharedBuffer::MakeReadOnly(MovePtr).GetData()"), const_cast<const void*>(Ptr->GetData()), Data);
 	}
 	{
-		// MakeMutable from a mutable view.
+		// MakeWritable from a mutable view.
 		constexpr uint64 Size = 64;
 		uint8 Data[Size]{};
-		FSharedBufferPtr Ptr = FSharedBuffer::MakeMutable(FSharedBuffer::MakeView(Data, Size));
-		TestFalse(TEXT("FSharedBuffer::MakeMutable(View).IsOwned()"), Ptr->IsOwned());
-		TestFalse(TEXT("FSharedBuffer::MakeMutable(View).IsImmutable()"), Ptr->IsImmutable());
-		TestEqual(TEXT("FSharedBuffer::MakeMutable(View).GetSize()"), Ptr->GetSize(), Size);
-		TestEqual(TEXT("FSharedBuffer::MakeMutable(View).GetData()"), const_cast<const void*>(Ptr->GetData()), static_cast<const void*>(Data));
+		FSharedBufferPtr Ptr = FSharedBuffer::MakeWritable(FSharedBuffer::MakeView(Data, Size));
+		TestFalse(TEXT("FSharedBuffer::MakeWritable(View).IsOwned()"), Ptr->IsOwned());
+		TestFalse(TEXT("FSharedBuffer::MakeWritable(View).IsReadOnly()"), Ptr->IsReadOnly());
+		TestEqual(TEXT("FSharedBuffer::MakeWritable(View).GetSize()"), Ptr->GetSize(), Size);
+		TestEqual(TEXT("FSharedBuffer::MakeWritable(View).GetData()"), const_cast<const void*>(Ptr->GetData()), static_cast<const void*>(Data));
 	}
 	{
-		// MakeMutable from a const view.
+		// MakeWritable from a const view.
 		constexpr uint64 Size = 64;
 		const uint8 Data[Size]{};
-		FSharedBufferPtr Ptr = FSharedBuffer::MakeMutable(FSharedBuffer::MakeView(Data, Size));
-		TestTrue(TEXT("FSharedBuffer::MakeMutable(View).IsOwned()"), Ptr->IsOwned());
-		TestFalse(TEXT("FSharedBuffer::MakeMutable(View).IsImmutable()"), Ptr->IsImmutable());
-		TestEqual(TEXT("FSharedBuffer::MakeMutable(View).GetSize()"), Ptr->GetSize(), Size);
-		TestNotEqual(TEXT("FSharedBuffer::MakeMutable(View).GetData()"), const_cast<const void*>(Ptr->GetData()), static_cast<const void*>(Data));
+		FSharedBufferPtr Ptr = FSharedBuffer::MakeWritable(FSharedBuffer::MakeView(Data, Size));
+		TestTrue(TEXT("FSharedBuffer::MakeWritable(View).IsOwned()"), Ptr->IsOwned());
+		TestFalse(TEXT("FSharedBuffer::MakeWritable(View).IsReadOnly()"), Ptr->IsReadOnly());
+		TestEqual(TEXT("FSharedBuffer::MakeWritable(View).GetSize()"), Ptr->GetSize(), Size);
+		TestNotEqual(TEXT("FSharedBuffer::MakeWritable(View).GetData()"), const_cast<const void*>(Ptr->GetData()), static_cast<const void*>(Data));
 	}
 	{
-		// MakeMutable with another shared reference.
+		// MakeWritable with another shared reference.
 		constexpr uint64 Size = 64;
-		FSharedBufferConstPtr ConstPtr = FSharedBuffer::MakeImmutable(FSharedBuffer::Alloc(Size));
+		FSharedBufferConstPtr ConstPtr = FSharedBuffer::MakeReadOnly(FSharedBuffer::Alloc(Size));
 		FSharedBufferConstPtr SharedPtr = ConstPtr;
 		const void* const Data = ConstPtr->GetData();
-		FSharedBufferConstPtr Ptr = FSharedBuffer::MakeMutable(MoveTemp(ConstPtr));
-		TestTrue(TEXT("FSharedBuffer::MakeMutable(SharedRef).IsOwned()"), Ptr->IsOwned());
-		TestFalse(TEXT("FSharedBuffer::MakeMutable(SharedRef).IsImmutable()"), Ptr->IsImmutable());
-		TestNotEqual(TEXT("FSharedBuffer::MakeMutable(SharedRef).GetData()"), const_cast<const void*>(Ptr->GetData()), Data);
+		FSharedBufferConstPtr Ptr = FSharedBuffer::MakeWritable(MoveTemp(ConstPtr));
+		TestTrue(TEXT("FSharedBuffer::MakeWritable(SharedRef).IsOwned()"), Ptr->IsOwned());
+		TestFalse(TEXT("FSharedBuffer::MakeWritable(SharedRef).IsReadOnly()"), Ptr->IsReadOnly());
+		TestNotEqual(TEXT("FSharedBuffer::MakeWritable(SharedRef).GetData()"), const_cast<const void*>(Ptr->GetData()), Data);
 	}
 	{
-		// MakeMutable with another weak reference.
+		// MakeWritable with another weak reference.
 		constexpr uint64 Size = 64;
-		FSharedBufferConstPtr ConstPtr = FSharedBuffer::MakeImmutable(FSharedBuffer::Alloc(Size));
+		FSharedBufferConstPtr ConstPtr = FSharedBuffer::MakeReadOnly(FSharedBuffer::Alloc(Size));
 		FSharedBufferConstWeakPtr WeakPtr = ConstPtr;
 		const void* const Data = ConstPtr->GetData();
-		FSharedBufferConstPtr Ptr = FSharedBuffer::MakeMutable(MoveTemp(ConstPtr));
-		TestTrue(TEXT("FSharedBuffer::MakeMutable(WeakRef).IsOwned()"), Ptr->IsOwned());
-		TestFalse(TEXT("FSharedBuffer::MakeMutable(WeakRef).IsImmutable()"), Ptr->IsImmutable());
-		TestNotEqual(TEXT("FSharedBuffer::MakeMutable(WeakRef).GetData()"), const_cast<const void*>(Ptr->GetData()), Data);
+		FSharedBufferConstPtr Ptr = FSharedBuffer::MakeWritable(MoveTemp(ConstPtr));
+		TestTrue(TEXT("FSharedBuffer::MakeWritable(WeakRef).IsOwned()"), Ptr->IsOwned());
+		TestFalse(TEXT("FSharedBuffer::MakeWritable(WeakRef).IsReadOnly()"), Ptr->IsReadOnly());
+		TestNotEqual(TEXT("FSharedBuffer::MakeWritable(WeakRef).GetData()"), const_cast<const void*>(Ptr->GetData()), Data);
 	}
 
 	// Test WeakPtr
diff --git a/Engine/Source/Runtime/Core/Public/Memory/SharedBuffer.h b/Engine/Source/Runtime/Core/Public/Memory/SharedBuffer.h
--- a/Engine/Source/Runtime/Core/Public/Memory/SharedBuffer.h
+++ b/Engine/Source/Runtime/Core/Public/Memory/SharedBuffer.h
@@ -23,7 +23,7 @@ enum ESharedBufferFlags : uint16
 {
 	None = 0,
 	Owned = 1 << 0,
-	Immutable = 1 << 1,
+	ReadOnly = 1 << 1,
 	HasDeleter = 1 << 2,
 };
 
@@ -197,7 +197,7 @@ private:
 
 ///////////////////////////////////////////////////////////////////////////////////////////////////////////////////////
 
-/** A non-nullable, thread-safe, reference to a mutable shared buffer. */
+/** A non-nullable, thread-safe, reference to a writable shared buffer. */
 class FSharedBufferRef
 {
 public:
@@ -252,7 +252,7 @@ private:
 
 ///////////////////////////////////////////////////////////////////////////////////////////////////////////////////////
 
-/** A non-nullable, thread-safe, pointer to a mutable shared buffer. */
+/** A non-nullable, thread-safe, pointer to a writable shared buffer. */
 class FSharedBufferPtr
 {
 public:
@@ -325,7 +325,7 @@ private:
 
 ///////////////////////////////////////////////////////////////////////////////////////////////////////////////////////
 
-/** A non-nullable, thread-safe, weak, pointer to a mutable shared buffer. */
+/** A non-nullable, thread-safe, weak, pointer to a writable shared buffer. */
 class FSharedBufferWeakPtr
 {
 public:
@@ -507,57 +507,57 @@ public:
 		return !IsValid(Buffer) || Buffer->IsOwned() ? Forward<T>(Buffer) : Clone(Buffer->GetData(), Buffer->GetSize());
 	}
 
-	/** Return the buffer if it is owned and either immutable or this is the only reference, or a clone otherwise. */
-	static inline FSharedBufferConstRef MakeImmutable(FSharedBufferRef&& Buffer)
+	/** Return the buffer if it is owned and either read-only or this is the only reference, or a clone otherwise. */
+	static inline FSharedBufferConstRef MakeReadOnly(FSharedBufferRef&& Buffer)
 	{
 		const FSharedBuffer& BufferRef = *Buffer;
-		return BufferRef.TryMakeImmutable() ? MoveTemp(Buffer) : MakeImmutable(Clone(*Buffer));
+		return BufferRef.TryMakeReadOnly() ? MoveTemp(Buffer) : MakeReadOnly(Clone(*Buffer));
 	}
 
-	/** Return the buffer if it is owned and either immutable or this is the only reference, or a clone otherwise. */
-	static inline FSharedBufferConstRef MakeImmutable(FSharedBufferConstRef&& Buffer)
+	/** Return the buffer if it is owned and either read-only or this is the only reference, or a clone otherwise. */
+	static inline FSharedBufferConstRef MakeReadOnly(FSharedBufferConstRef&& Buffer)
 	{
-		return Buffer->TryMakeImmutable() ? MoveTemp(Buffer) : MakeImmutable(Clone(*Buffer));
+		return Buffer->TryMakeReadOnly() ? MoveTemp(Buffer) : MakeReadOnly(Clone(*Buffer));
 	}
 
-	/** Return the buffer if it is owned and either immutable or this is the only reference, or a clone otherwise. */
-	static inline FSharedBufferConstPtr MakeImmutable(FSharedBufferPtr&& Buffer)
+	/** Return the buffer if it is owned and either read-only or this is the only reference, or a clone otherwise. */
+	static inline FSharedBufferConstPtr MakeReadOnly(FSharedBufferPtr&& Buffer)
 	{
-		return MakeImmutable(FSharedBufferConstPtr(MoveTemp(Buffer)));
+		return MakeReadOnly(FSharedBufferConstPtr(MoveTemp(Buffer)));
 	}
 
-	/** Return the buffer if it is owned and either immutable or this is the only reference, or a clone otherwise. */
-	static inline FSharedBufferConstPtr MakeImmutable(FSharedBufferConstPtr&& Buffer)
+	/** Return the buffer if it is owned and either read-only or this is the only reference, or a clone otherwise. */
+	static inline FSharedBufferConstPtr MakeReadOnly(FSharedBufferConstPtr&& Buffer)
 	{
-		return Buffer ? Buffer->TryMakeImmutable() ? MoveTemp(Buffer) : MakeImmutable(Clone(*Buffer)) : FSharedBufferConstPtr();
+		return Buffer ? Buffer->TryMakeReadOnly() ? MoveTemp(Buffer) : MakeReadOnly(Clone(*Buffer)) : FSharedBufferConstPtr();
 	}
 
 	// Disable overloads that would allow the argument to be copied.
-	static void MakeImmutable(const FSharedBufferRef&) = delete;
-	static void MakeImmutable(const FSharedBufferConstRef&) = delete;
-	static void MakeImmutable(const FSharedBufferPtr&) = delete;
-	static void MakeImmutable(const FSharedBufferConstPtr&) = delete;
+	static void MakeReadOnly(const FSharedBufferRef&) = delete;
+	static void MakeReadOnly(const FSharedBufferConstRef&) = delete;
+	static void MakeReadOnly(const FSharedBufferPtr&) = delete;
+	static void MakeReadOnly(const FSharedBufferConstPtr&) = delete;
 
-	/** Return the buffer because it is already mutable. */
-	static inline FSharedBufferRef MakeMutable(FSharedBufferRef&& Buffer) { return MoveTemp(Buffer); }
+	/** Return the buffer because it is already writable. */
+	static inline FSharedBufferRef MakeWritable(FSharedBufferRef&& Buffer) { return MoveTemp(Buffer); }
 
-	/** Return the buffer if it is owned and either mutable or this is the only reference, or a clone otherwise. */
-	static inline FSharedBufferRef MakeMutable(FSharedBufferConstRef&& Buffer)
+	/** Return the buffer if it is owned and either writable or this is the only reference, or a clone otherwise. */
+	static inline FSharedBufferRef MakeWritable(FSharedBufferConstRef&& Buffer)
 	{
-		if (Buffer->TryMakeMutable())
+		if (Buffer->TryMakeWritable())
 		{
 			return FSharedBufferRef(FSharedBufferRef::PtrType(const_cast<FSharedBuffer*>(&Buffer.Get())));
 		}
 		return Clone(*Buffer);
 	}
 
-	/** Return the buffer because it is already mutable. */
-	static inline FSharedBufferPtr MakeMutable(FSharedBufferPtr&& Buffer) { return MoveTemp(Buffer); }
+	/** Return the buffer because it is already writable. */
+	static inline FSharedBufferPtr MakeWritable(FSharedBufferPtr&& Buffer) { return MoveTemp(Buffer); }
 
-	/** Return the buffer if it is owned and either mutable or this is the only reference, or a clone otherwise. */
-	static inline FSharedBufferPtr MakeMutable(FSharedBufferConstPtr&& Buffer)
+	/** Return the buffer if it is owned and either writable or this is the only reference, or a clone otherwise. */
+	static inline FSharedBufferPtr MakeWritable(FSharedBufferConstPtr&& Buffer)
 	{
-		if (Buffer->TryMakeMutable())
+		if (Buffer->TryMakeWritable())
 		{
 			return FSharedBufferPtr(FSharedBufferPtr::PtrType(const_cast<FSharedBuffer*>(Buffer.Get())));
 		}
@@ -565,10 +565,10 @@ public:
 	}
 
 	// Disable overloads that would allow the argument to be copied.
-	static void MakeMutable(const FSharedBufferRef&) = delete;
-	static void MakeMutable(const FSharedBufferConstRef&) = delete;
-	static void MakeMutable(const FSharedBufferPtr&) = delete;
-	static void MakeMutable(const FSharedBufferConstPtr&) = delete;
+	static void MakeWritable(const FSharedBufferRef&) = delete;
+	static void MakeWritable(const FSharedBufferConstRef&) = delete;
+	static void MakeWritable(const FSharedBufferPtr&) = delete;
+	static void MakeWritable(const FSharedBufferConstPtr&) = delete;
 
 	/** A pointer to the start of the buffer. */
 	inline void* GetData() { return Data; }
@@ -582,39 +582,39 @@ public:
 	inline bool IsOwned() const { return (GetFlags(ReferenceCountAndFlags) & ESharedBufferFlags::Owned) != 0; }
 
 	/**
-	 * Whether the shared buffer is immutable.
+	 * Whether the shared buffer is read-only.
 	 *
-	 * An immutable shared buffer is owned and every reference or pointer to it is const.
+	 * A read-only shared buffer is owned and every reference or pointer to it is const.
 	 */
-	inline bool IsImmutable() const { return (GetFlags(ReferenceCountAndFlags) & ESharedBufferFlags::Immutable) != 0; }
+	inline bool IsReadOnly() const { return (GetFlags(ReferenceCountAndFlags) & ESharedBufferFlags::ReadOnly) != 0; }
 
 	/**
-	 * Try to make the shared buffer immutable.
+	 * Try to make the shared buffer read-only.
 	 *
-	 * A shared buffer can be made immutable if it is owned and there is only one reference to it.
+	 * A shared buffer can be made read-only if it is owned and there is only one reference to it.
 	 *
-	 * @return true if the buffer was immutable or was made immutable, otherwise false.
+	 * @return true if the buffer was read-only or was made read-only, otherwise false.
 	 */
-	CORE_API bool TryMakeImmutable() const;
+	CORE_API bool TryMakeReadOnly() const;
 
-	// Disable this overload because it would allow an immutable buffer to be mutated.
-	bool TryMakeImmutable() = delete;
+	// Disable this overload because it would allow a read-only buffer to be mutated.
+	bool TryMakeReadOnly() = delete;
 
 	/**
-	 * Try to make the shared buffer mutable.
+	 * Try to make the shared buffer writable.
 	 *
-	 * A shared buffer can be made mutable if it is owned and there is only one reference to it.
+	 * A shared buffer can be made writable if it is owned and there is only one reference to it.
 	 *
-	 * @return true if the buffer was mutable or was made mutable, otherwise false.
+	 * @return true if the buffer was writable or was made writable, otherwise false.
 	 */
-	CORE_API bool TryMakeMutable() const;
+	CORE_API bool TryMakeWritable() const;
 
-	/** The shared buffer is already mutable. */
-	inline bool TryMakeMutable() { return true; }
+	/** The shared buffer is already writable. */
+	inline bool TryMakeWritable() { return true; }
 
-	/** A mutable view of buffer. */
+	/** A writable view of buffer. */
 	inline FMutableMemoryView GetView() { return FMutableMemoryView(GetData(), GetSize()); }
-	/** A const view of the buffer. */
+	/** A read-only view of the buffer. */
 	inline FConstMemoryView GetView() const { return FConstMemoryView(GetData(), GetSize()); }
 
 	inline operator FMutableMemoryView() { return GetView(); }
