Removing unused cvar.

diff --git a/Engine/Source/Runtime/Core/Private/HAL/ConsoleManager.cpp b/Engine/Source/Runtime/Core/Private/HAL/ConsoleManager.cpp
--- a/Engine/Source/Runtime/Core/Private/HAL/ConsoleManager.cpp
+++ b/Engine/Source/Runtime/Core/Private/HAL/ConsoleManager.cpp
@@ -2528,14 +2528,6 @@ static TAutoConsoleVariable<int32> CVarCompileShadersForDevelopment(
 	TEXT(" 1: on (Default)"),
 	ECVF_ReadOnly);
 
-static TAutoConsoleVariable<int32> CVarPS4MixeedModeShaderDebugInfo(
-	TEXT("r.PS4MixedModeShaderDebugInfo"),
-	0,
-	TEXT("Whether to compile shaders to allow mixed mode shader debugging. This will currently generate slower code.\n")
-	TEXT(" 0: Normal mode\n")
-	TEXT(" 1: Mixed mode)"),
-	ECVF_ReadOnly);
-
 static TAutoConsoleVariable<int32> CVarDontLimitOnBattery(
 	TEXT("r.DontLimitOnBattery"),
 	0,
