Add protection against copying a RWLock by value
Add check similar to pthread_rw on Mac/Linux to make sure Windows FRWLock is not held when destroyed
Fix the one place that we were actually copying a RWLock
Introduce a new TDontCopy wrapper that can be used to make a class member an exclusive entity of that class. (Contributed by [REDACTED_PERSON])

#rnx

diff --git a/Engine/Source/Editor/Sequencer/Private/SSequencerDebugVisualizer.cpp b/Engine/Source/Editor/Sequencer/Private/SSequencerDebugVisualizer.cpp
--- a/Engine/Source/Editor/Sequencer/Private/SSequencerDebugVisualizer.cpp
+++ b/Engine/Source/Editor/Sequencer/Private/SSequencerDebugVisualizer.cpp
@@ -382,7 +382,7 @@ void SSequencerEntityComponentSystemDebugSlot::Refresh()
 		}
 
 		FEntityInfo EntityInfo = Linker->EntityManager.GetEntity(CachedEntityID);
-		for (FComponentHeader ComponentHeader : EntityInfo.Data.Allocation->GetComponentHeaders())
+		for (const FComponentHeader& ComponentHeader : EntityInfo.Data.Allocation->GetComponentHeaders())
 		{
 			const FComponentTypeInfo& ComponentTypeInfo = Linker->GetComponents()->GetComponentTypeChecked(ComponentHeader.ComponentType);
 			const int32 ComponentBitIndex = ComponentHeader.ComponentType.BitIndex();
diff --git a/Engine/Source/Runtime/Core/Private/Windows/MinimalWindowsApi.cpp b/Engine/Source/Runtime/Core/Private/Windows/MinimalWindowsApi.cpp
--- a/Engine/Source/Runtime/Core/Private/Windows/MinimalWindowsApi.cpp
+++ b/Engine/Source/Runtime/Core/Private/Windows/MinimalWindowsApi.cpp
@@ -80,6 +80,16 @@ namespace Windows
 		::AcquireSRWLockShared((::PSRWLOCK)SRWLock);
 	}
 
+	MINIMAL_WINDOWS_API BOOLEAN WINAPI TryAcquireSRWLockShared(PSRWLOCK SRWLock)
+	{
+		return ::TryAcquireSRWLockShared((::PSRWLOCK)SRWLock);
+	}
+
+	MINIMAL_WINDOWS_API BOOLEAN WINAPI TryAcquireSRWLockExclusive(PSRWLOCK SRWLock)
+	{
+		::TryAcquireSRWLockExclusive((::PSRWLOCK)SRWLock);
+	}
+
 	MINIMAL_WINDOWS_API void WINAPI ReleaseSRWLockShared(PSRWLOCK SRWLock)
 	{
 		::ReleaseSRWLockShared((::PSRWLOCK)SRWLock);
diff --git a/Engine/Source/Runtime/Core/Public/GenericPlatform/GenericPlatformCriticalSection.h b/Engine/Source/Runtime/Core/Public/GenericPlatform/GenericPlatformCriticalSection.h
--- a/Engine/Source/Runtime/Core/Public/GenericPlatform/GenericPlatformCriticalSection.h
+++ b/Engine/Source/Runtime/Core/Public/GenericPlatform/GenericPlatformCriticalSection.h
@@ -39,6 +39,9 @@ template<class CriticalSection>
 class TGenericPlatformRWLock
 {
 public:
+	TGenericPlatformRWLock(const TGenericPlatformRWLock&) = delete;
+	TGenericPlatformRWLock& operator=(const TGenericPlatformRWLock&) = delete;
+
 	FORCEINLINE TGenericPlatformRWLock()
 	{
 	}
diff --git a/Engine/Source/Runtime/Core/Public/HAL/PThreadRWLock.h b/Engine/Source/Runtime/Core/Public/HAL/PThreadRWLock.h
--- a/Engine/Source/Runtime/Core/Public/HAL/PThreadRWLock.h
+++ b/Engine/Source/Runtime/Core/Public/HAL/PThreadRWLock.h
@@ -14,6 +14,9 @@
 class FPThreadsRWLock
 {
 public:
+	FPThreadsRWLock(const FPThreadsRWLock&) = delete;
+	FPThreadsRWLock& operator=(const FPThreadsRWLock&) = delete;
+
 	FPThreadsRWLock()
 	{
 		int Err = pthread_rwlock_init(&Mutex, nullptr);
diff --git a/Engine/Source/Runtime/Core/Public/HoloLens/HoloLensCriticalSection.h b/Engine/Source/Runtime/Core/Public/HoloLens/HoloLensCriticalSection.h
--- a/Engine/Source/Runtime/Core/Public/HoloLens/HoloLensCriticalSection.h
+++ b/Engine/Source/Runtime/Core/Public/HoloLens/HoloLensCriticalSection.h
@@ -81,6 +81,9 @@ public:
 class FHoloLensRWLock
 {
 public:
+	FHoloLensRWLock(const FHoloLensRWLock&) = delete;
+	FHoloLensRWLock& operator=(const FHoloLensRWLock&) = delete;
+
 	FORCEINLINE FHoloLensRWLock(uint32 Level = 0)
 	{
 		Windows::InitializeSRWLock(&Mutex);
diff --git a/Engine/Source/Runtime/Core/Public/Templates/DontCopy.h b/Engine/Source/Runtime/Core/Public/Templates/DontCopy.h
new file mode 100644
--- /dev/null
+++ b/Engine/Source/Runtime/Core/Public/Templates/DontCopy.h
@@ -0,0 +1,32 @@
+// Copyright Epic Games, Inc. All Rights Reserved.
+
+#pragma once
+
+#include "CoreTypes.h"
+
+// Can be used as a wrapper over non-copyable objects (i.e. FCriticalSection) if you still want the class containing your
+// object to keep its copyable property without compromising integrity of the object that doesn't support being copied.
+// It makes each TDontCopy member of a class an exclusive entity of that class.
+
+template <typename T>
+struct TDontCopy
+{
+	TDontCopy() = default;
+	~TDontCopy() = default;
+	TDontCopy(TDontCopy&&) {}
+	TDontCopy(const TDontCopy&) {}
+	TDontCopy& operator=(TDontCopy&&) { return *this; }
+	TDontCopy& operator=(const TDontCopy&) { return *this; }
+
+	T& Get() { return Value; }
+	const T& Get() const { return Value; }
+
+	T* operator->()       { return &Value; }
+	const T* operator->() const { return &Value; }
+	
+	T& operator*() { return Value; }
+	const T& operator*() const { return Value; }
+
+private:
+    T Value;
+};
diff --git a/Engine/Source/Runtime/Core/Public/Windows/MinimalWindowsApi.h b/Engine/Source/Runtime/Core/Public/Windows/MinimalWindowsApi.h
--- a/Engine/Source/Runtime/Core/Public/Windows/MinimalWindowsApi.h
+++ b/Engine/Source/Runtime/Core/Public/Windows/MinimalWindowsApi.h
@@ -62,6 +62,8 @@ namespace Windows
 {
 	// Typedefs for basic Windows types
 	typedef int32 BOOL;
+	typedef unsigned char BYTE;
+	typedef BYTE  BOOLEAN;
 	typedef unsigned long DWORD;
 	typedef DWORD* LPDWORD;
 	typedef long LONG;
@@ -120,12 +122,24 @@ namespace Windows
 	MINIMAL_WINDOWS_API void WINAPI ReleaseSRWLockShared(PSRWLOCK SRWLock);
 	MINIMAL_WINDOWS_API void WINAPI AcquireSRWLockExclusive(PSRWLOCK SRWLock);
 	MINIMAL_WINDOWS_API void WINAPI ReleaseSRWLockExclusive(PSRWLOCK SRWLock);
+	MINIMAL_WINDOWS_API BOOLEAN WINAPI TryAcquireSRWLockExclusive(PSRWLOCK SRWLock);
+	MINIMAL_WINDOWS_API BOOLEAN WINAPI TryAcquireSRWLockShared(PSRWLOCK SRWLock);
 
 	FORCEINLINE void WINAPI InitializeSRWLock(SRWLOCK* SRWLock)
 	{
 		InitializeSRWLock((PSRWLOCK)SRWLock);
 	}
 
+	FORCEINLINE BOOLEAN WINAPI TryAcquireSRWLockExclusive(SRWLOCK* SRWLock)
+	{
+		return TryAcquireSRWLockExclusive((PSRWLOCK)SRWLock);
+	}
+
+	FORCEINLINE BOOLEAN WINAPI TryAcquireSRWLockShared(SRWLOCK* SRWLock)
+	{
+		return TryAcquireSRWLockShared((PSRWLOCK)SRWLock);
+	}
+	
 	FORCEINLINE void WINAPI AcquireSRWLockShared(SRWLOCK* SRWLock)
 	{
 		AcquireSRWLockShared((PSRWLOCK)SRWLock);
diff --git a/Engine/Source/Runtime/Core/Public/Windows/WindowsCriticalSection.h b/Engine/Source/Runtime/Core/Public/Windows/WindowsCriticalSection.h
--- a/Engine/Source/Runtime/Core/Public/Windows/WindowsCriticalSection.h
+++ b/Engine/Source/Runtime/Core/Public/Windows/WindowsCriticalSection.h
@@ -14,12 +14,9 @@ class FString;
  */
 class FWindowsCriticalSection
 {
-	/**
-	 * The windows specific critical section
-	 */
-	Windows::CRITICAL_SECTION CriticalSection;
-
 public:
+	FWindowsCriticalSection(const FWindowsCriticalSection&) = delete;
+	FWindowsCriticalSection& operator=(const FWindowsCriticalSection&) = delete;
 
 	/**
 	 * Constructor that initializes the aggregated critical section
@@ -70,8 +67,10 @@ public:
 	}
 
 private:
-	FWindowsCriticalSection(const FWindowsCriticalSection&);
-	FWindowsCriticalSection& operator=(const FWindowsCriticalSection&);
+	/**
+	 * The windows specific critical section
+	 */
+	Windows::CRITICAL_SECTION CriticalSection;
 };
 
 /** System-Wide Critical Section for windows using mutex */
@@ -110,6 +109,9 @@ private:
 class FWindowsRWLock
 {
 public:
+	FWindowsRWLock(const FWindowsRWLock&) = delete;
+	FWindowsRWLock& operator=(const FWindowsRWLock&) = delete;
+
 	FORCEINLINE FWindowsRWLock(uint32 Level = 0)
 	{
 		Windows::InitializeSRWLock(&Mutex);
@@ -117,6 +119,7 @@ public:
 	
 	FORCEINLINE ~FWindowsRWLock()
 	{
+		checkf(Windows::TryAcquireSRWLockExclusive(&Mutex), TEXT("Destroying a lock that is still held!"));
 	}
 	
 	FORCEINLINE void ReadLock()
diff --git a/Engine/Source/Runtime/RawMesh/Private/RawMesh.cpp b/Engine/Source/Runtime/RawMesh/Private/RawMesh.cpp
--- a/Engine/Source/Runtime/RawMesh/Private/RawMesh.cpp
+++ b/Engine/Source/Runtime/RawMesh/Private/RawMesh.cpp
@@ -271,7 +271,7 @@ void FRawMeshBulkData::LoadRawMesh(FRawMesh& OutMesh)
 	{
 #if WITH_EDITOR
 		// A lock is required so we can safely load the raw data from multiple threads
-		FRWScopeLock ScopeLock(BulkDataLock, SLT_Write);
+		FRWScopeLock ScopeLock(BulkDataLock.Get(), SLT_Write);
 
 		// This allows any thread to be able to deserialize from the RawMesh directly
 		// from disk so we can unload bulk data from memory.
diff --git a/Engine/Source/Runtime/RawMesh/Public/RawMesh.h b/Engine/Source/Runtime/RawMesh/Public/RawMesh.h
--- a/Engine/Source/Runtime/RawMesh/Public/RawMesh.h
+++ b/Engine/Source/Runtime/RawMesh/Public/RawMesh.h
@@ -5,6 +5,7 @@
 #include "CoreMinimal.h"
 #include "Misc/Guid.h"
 #include "Serialization/BulkData.h"
+#include "Templates/DontCopy.h"
 
 enum
 {
@@ -101,7 +102,7 @@ class FRawMeshBulkData
 {
 #if WITH_EDITOR
 	/** Protects simultaneous access to BulkData */
-	FRWLock       BulkDataLock;
+	TDontCopy<FRWLock> BulkDataLock;
 #endif
 	/** Internally store bulk data as bytes. */
 	FByteBulkData BulkData;
