SharedBuffer: Extended shared buffer with an optional outer buffer and optional buffer owner

MakeView now has overloads that accept an outer shared buffer that it will hold a reference to and inherit the read-only state from. The view must be contained by the outer shared buffer.

TakeOwnership now has overloads that accept an instance of a type derived from FSharedBuffer::FBufferOwner that will be moved and stored to manage the lifetime of the buffer when a delete function is insufficient.

Functions that create buffers now return pointers instead of references, for more optimal reference counting. ToSharedRef can be called on the return value to convert it to a reference in cases where that is useful.

#rnx

diff --git a/Engine/Source/Runtime/Core/Private/Memory/SharedBuffer.cpp b/Engine/Source/Runtime/Core/Private/Memory/SharedBuffer.cpp
--- a/Engine/Source/Runtime/Core/Private/Memory/SharedBuffer.cpp
+++ b/Engine/Source/Runtime/Core/Private/Memory/SharedBuffer.cpp
@@ -6,11 +6,12 @@ FSharedBufferRef FSharedBuffer::NewBuffer(
 	void* const Data,
 	const uint64 Size,
 	const ESharedBufferFlags Flags,
-	const uint32 DeleterSize)
+	const uint32 OwnerSize)
 {
-	checkf((DeleterSize == 0) == !(Flags & ESharedBufferFlags::HasDeleter), TEXT("Mismatch between DeleterSize %u and Flags."), DeleterSize);
+	checkf((OwnerSize > 0) == EnumHasAnyFlags(Flags, ESharedBufferFlags::HasBufferOwner),
+		TEXT("Mismatch between OwnerSize %u and Flags."), OwnerSize);
 
-	void* const BufferMemory = FMemory::Malloc(sizeof(FSharedBuffer) + DeleterSize, alignof(FSharedBuffer));
+	void* const BufferMemory = FMemory::Malloc(sizeof(FSharedBuffer) + OwnerSize, alignof(FSharedBuffer));
 
 	FSharedBuffer* const Buffer = new(BufferMemory) FSharedBuffer;
 	Buffer->Data = Data;
@@ -30,16 +31,16 @@ void FSharedBuffer::DeleteBuffer(FSharedBuffer* const Buffer)
 
 void FSharedBuffer::ReleaseData()
 {
-	if (IsOwned())
+	const ESharedBufferFlags Flags = GetFlags(ReferenceCountAndFlags);
+	if (EnumHasAnyFlags(Flags, ESharedBufferFlags::HasBufferOwner))
 	{
-		if (GetFlags(ReferenceCountAndFlags) & ESharedBufferFlags::HasDeleter)
-		{
-			reinterpret_cast<FSharedBufferDeleter*>(this + 1)->Free(Data);
-		}
-		else
-		{
-			FMemory::Free(Data);
-		}
+		FBufferOwner& BufferOwner = *reinterpret_cast<FBufferOwner*>(this + 1);
+		BufferOwner.Free(Data, Size);
+		BufferOwner.~FBufferOwner();
+	}
+	else if (EnumHasAnyFlags(Flags, ESharedBufferFlags::Owned))
+	{
+		FMemory::Free(Data);
 	}
 	Data = nullptr;
 	Size = 0;
diff --git a/Engine/Source/Runtime/Core/Private/Serialization/CompactBinaryWriter.cpp b/Engine/Source/Runtime/Core/Private/Serialization/CompactBinaryWriter.cpp
--- a/Engine/Source/Runtime/Core/Private/Serialization/CompactBinaryWriter.cpp
+++ b/Engine/Source/Runtime/Core/Private/Serialization/CompactBinaryWriter.cpp
@@ -70,7 +70,7 @@ void FCbWriter::Reset()
 FCbFieldRefIterator FCbWriter::Save() const
 {
 	const uint64 Size = SaveSize();
-	const FSharedBufferRef Buffer = FSharedBuffer::Alloc(Size);
+	const FSharedBufferPtr Buffer = FSharedBuffer::Alloc(Size);
 	const FMutableMemoryView View = Buffer->GetView();
 	SaveToMemory(View);
 	return FCbFieldRefIterator(FCbFieldRef(Buffer), View.RightChop(Size).GetData());
diff --git a/Engine/Source/Runtime/Core/Private/Tests/Memory/SharedBufferTest.cpp b/Engine/Source/Runtime/Core/Private/Tests/Memory/SharedBufferTest.cpp
--- a/Engine/Source/Runtime/Core/Private/Tests/Memory/SharedBufferTest.cpp
+++ b/Engine/Source/Runtime/Core/Private/Tests/Memory/SharedBufferTest.cpp
@@ -117,22 +117,24 @@ static_assert(std::is_assignable<FSharedBufferConstWeakPtr, FSharedBufferConstWe
 
 ///////////////////////////////////////////////////////////////////////////////////////////////////////////////////////
 
-static_assert(std::is_same<FSharedBufferRef, decltype(FSharedBuffer::Alloc(DeclVal<uint64>()))>::value, "Invalid constructor");
+static_assert(std::is_same<FSharedBufferPtr, decltype(FSharedBuffer::Alloc(DeclVal<uint64>()))>::value, "Invalid constructor");
 
-static_assert(std::is_same<FSharedBufferRef, decltype(FSharedBuffer::Clone(DeclVal<void*>(), 0))>::value, "Invalid constructor");
-static_assert(std::is_same<FSharedBufferRef, decltype(FSharedBuffer::Clone(DeclVal<const void*>(), 0))>::value, "Invalid constructor");
-static_assert(std::is_same<FSharedBufferRef, decltype(FSharedBuffer::Clone(*DeclVal<FSharedBufferRef>()))>::value, "Invalid constructor");
-static_assert(std::is_same<FSharedBufferRef, decltype(FSharedBuffer::Clone(*DeclVal<FSharedBufferConstRef>()))>::value, "Invalid constructor");
-static_assert(std::is_same<FSharedBufferRef, decltype(FSharedBuffer::Clone(DeclVal<FMutableMemoryView>()))>::value, "Invalid constructor");
-static_assert(std::is_same<FSharedBufferRef, decltype(FSharedBuffer::Clone(DeclVal<FConstMemoryView>()))>::value, "Invalid constructor");
+static_assert(std::is_same<FSharedBufferPtr, decltype(FSharedBuffer::Clone(DeclVal<void*>(), 0))>::value, "Invalid constructor");
+static_assert(std::is_same<FSharedBufferPtr, decltype(FSharedBuffer::Clone(DeclVal<const void*>(), 0))>::value, "Invalid constructor");
+static_assert(std::is_same<FSharedBufferPtr, decltype(FSharedBuffer::Clone(*DeclVal<FSharedBufferRef>()))>::value, "Invalid constructor");
+static_assert(std::is_same<FSharedBufferPtr, decltype(FSharedBuffer::Clone(*DeclVal<FSharedBufferConstRef>()))>::value, "Invalid constructor");
+static_assert(std::is_same<FSharedBufferPtr, decltype(FSharedBuffer::Clone(DeclVal<FMutableMemoryView>()))>::value, "Invalid constructor");
+static_assert(std::is_same<FSharedBufferPtr, decltype(FSharedBuffer::Clone(DeclVal<FConstMemoryView>()))>::value, "Invalid constructor");
 
-static_assert(std::is_same<FSharedBufferRef, decltype(FSharedBuffer::TakeOwnership(DeclVal<void*>(), 0, FMemory::Free))>::value, "Invalid constructor");
-static_assert(std::is_same<FSharedBufferConstRef, decltype(FSharedBuffer::TakeOwnership(DeclVal<const void*>(), 0, FMemory::Free))>::value, "Invalid constructor");
+static_assert(std::is_same<FSharedBufferPtr, decltype(FSharedBuffer::MakeView(DeclVal<FMutableMemoryView>()))>::value, "Invalid constructor");
+static_assert(std::is_same<FSharedBufferConstPtr, decltype(FSharedBuffer::MakeView(DeclVal<FConstMemoryView>()))>::value, "Invalid constructor");
+static_assert(std::is_same<FSharedBufferPtr, decltype(FSharedBuffer::MakeView(DeclVal<void*>(), 0))>::value, "Invalid constructor");
+static_assert(std::is_same<FSharedBufferConstPtr, decltype(FSharedBuffer::MakeView(DeclVal<const void*>(), 0))>::value, "Invalid constructor");
+static_assert(std::is_same<FSharedBufferConstPtr, decltype(FSharedBuffer::MakeView(DeclVal<FConstMemoryView>(), DeclVal<const FSharedBuffer&>()))>::value, "Invalid constructor");
+static_assert(std::is_same<FSharedBufferConstPtr, decltype(FSharedBuffer::MakeView(DeclVal<const void*>(), 0, DeclVal<const FSharedBuffer&>()))>::value, "Invalid constructor");
 
-static_assert(std::is_same<FSharedBufferRef, decltype(FSharedBuffer::MakeView(DeclVal<void*>(), 0))>::value, "Invalid constructor");
-static_assert(std::is_same<FSharedBufferConstRef, decltype(FSharedBuffer::MakeView(DeclVal<const void*>(), 0))>::value, "Invalid constructor");
-static_assert(std::is_same<FSharedBufferRef, decltype(FSharedBuffer::MakeView(DeclVal<FMutableMemoryView>()))>::value, "Invalid constructor");
-static_assert(std::is_same<FSharedBufferConstRef, decltype(FSharedBuffer::MakeView(DeclVal<FConstMemoryView>()))>::value, "Invalid constructor");
+static_assert(std::is_same<FSharedBufferPtr, decltype(FSharedBuffer::TakeOwnership(DeclVal<void*>(), 0, FMemory::Free))>::value, "Invalid constructor");
+static_assert(std::is_same<FSharedBufferConstPtr, decltype(FSharedBuffer::TakeOwnership(DeclVal<const void*>(), 0, FMemory::Free))>::value, "Invalid constructor");
 
 static_assert(std::is_same<void*, decltype(DeclVal<FSharedBuffer>().GetData())>::value, "Invalid accessor");
 static_assert(std::is_same<const void*, decltype(DeclVal<const FSharedBuffer>().GetData())>::value, "Invalid accessor");
@@ -149,10 +151,12 @@ static_assert(std::is_convertible<const FSharedBuffer, FConstMemoryView>::value,
 ///////////////////////////////////////////////////////////////////////////////////////////////////////////////////////
 
 static_assert(std::is_same<FSharedBufferRef, decltype(FSharedBuffer::MakeOwned(DeclVal<const FSharedBufferRef&>()))>::value, "Invalid conversion");
+static_assert(std::is_same<FSharedBufferRef, decltype(FSharedBuffer::MakeOwned(DeclVal<FSharedBufferRef&&>()))>::value, "Invalid conversion");
 static_assert(std::is_same<FSharedBufferPtr, decltype(FSharedBuffer::MakeOwned(DeclVal<const FSharedBufferPtr&>()))>::value, "Invalid conversion");
 static_assert(std::is_same<FSharedBufferPtr, decltype(FSharedBuffer::MakeOwned(DeclVal<FSharedBufferPtr&&>()))>::value, "Invalid conversion");
 
 static_assert(std::is_same<FSharedBufferConstRef, decltype(FSharedBuffer::MakeOwned(DeclVal<const FSharedBufferConstRef&>()))>::value, "Invalid conversion");
+static_assert(std::is_same<FSharedBufferConstRef, decltype(FSharedBuffer::MakeOwned(DeclVal<FSharedBufferConstRef&&>()))>::value, "Invalid conversion");
 static_assert(std::is_same<FSharedBufferConstPtr, decltype(FSharedBuffer::MakeOwned(DeclVal<const FSharedBufferConstPtr&>()))>::value, "Invalid conversion");
 static_assert(std::is_same<FSharedBufferConstPtr, decltype(FSharedBuffer::MakeOwned(DeclVal<FSharedBufferConstPtr&&>()))>::value, "Invalid conversion");
 
@@ -175,7 +179,7 @@ bool FSharedBufferTest::RunTest(const FString& Parameters)
 	// Test Size
 	{
 		constexpr uint64 Size = 64;
-		FSharedBufferRef Ref = FSharedBuffer::Alloc(64);
+		FSharedBufferRef Ref = FSharedBuffer::Alloc(64).ToSharedRef();
 		TestTrue(TEXT("FSharedBuffer(Size).IsOwned()"), Ref->IsOwned());
 		TestEqual(TEXT("FSharedBuffer(Size).GetSize()"), Ref->GetSize(), Size);
 	}
@@ -184,7 +188,7 @@ bool FSharedBufferTest::RunTest(const FString& Parameters)
 	{
 		constexpr uint64 Size = 64;
 		const uint8 Data[Size]{};
-		FSharedBufferRef Ref = FSharedBuffer::Clone(Data, Size);
+		FSharedBufferRef Ref = FSharedBuffer::Clone(Data, Size).ToSharedRef();
 		TestTrue(TEXT("FSharedBuffer::Clone().IsOwned()"), Ref->IsOwned());
 		TestEqual(TEXT("FSharedBuffer::Clone().GetSize()"), Ref->GetSize(), Size);
 		TestNotEqual(TEXT("FSharedBuffer::Clone().GetData()"), static_cast<const void*>(Ref->GetData()), static_cast<const void*>(Data));
@@ -194,17 +198,47 @@ bool FSharedBufferTest::RunTest(const FString& Parameters)
 	{
 		constexpr uint64 Size = 64;
 		uint8 Data[Size]{};
-		FSharedBufferRef Ref = FSharedBuffer::MakeView(Data, Size);
+		FSharedBufferRef Ref = FSharedBuffer::MakeView(Data, Size).ToSharedRef();
 		TestFalse(TEXT("FSharedBuffer::MakeView().IsOwned()"), Ref->IsOwned());
 		TestEqual(TEXT("FSharedBuffer::MakeView().GetSize()"), Ref->GetSize(), Size);
 		TestEqual(TEXT("FSharedBuffer::MakeView().GetData()"), Ref->GetData(), static_cast<void*>(Data));
 	}
 
+	// Test MakeView Outer
+	{
+		constexpr uint64 Size = 64;
+		FSharedBufferWeakPtr WeakPtr;
+		{
+			FSharedBufferConstPtr Ptr;
+			{
+				FSharedBufferRef OuterRef = FSharedBuffer::Alloc(Size).ToSharedRef();
+				Ptr = FSharedBuffer::MakeView(OuterRef->GetData(), OuterRef->GetSize() / 2, *OuterRef);
+				WeakPtr = OuterRef;
+			}
+			TestFalse(TEXT("FSharedBuffer::MakeView(Outer).IsOwned()"), Ptr->IsOwned());
+			TestFalse(TEXT("FSharedBuffer::MakeView(Outer).IsReadOnly()"), Ptr->IsReadOnly());
+			TestTrue(TEXT("FSharedBuffer::MakeView(Outer) Outer Valid"), WeakPtr.Pin().IsValid());
+		}
+		TestFalse(TEXT("FSharedBuffer::MakeView(Outer) Outer Not Valid"), WeakPtr.Pin().IsValid());
+	}
+
+	// Test MakeView Outer Read-Only
+	{
+		constexpr uint64 Size = 64;
+		FSharedBufferConstPtr Ptr;
+		{
+			FSharedBufferConstRef OuterRef = FSharedBuffer::MakeReadOnly(FSharedBuffer::Alloc(Size)).ToSharedRef();
+			Ptr = FSharedBuffer::MakeView(OuterRef->GetData(), OuterRef->GetSize() / 2, *OuterRef);
+		}
+		TestFalse(TEXT("FSharedBuffer::MakeView(Outer ReadOnly).IsOwned()"), Ptr->IsOwned());
+		TestTrue(TEXT("FSharedBuffer::MakeView(Outer ReadOnly).IsReadOnly()"), Ptr->IsReadOnly());
+	}
+
 	// Test TakeOwnership FMemory::Free
 	{
 		constexpr uint64 Size = 64;
 		void* const Data = FMemory::Malloc(Size);
-		FSharedBufferRef Ref = FSharedBuffer::TakeOwnership(Data, Size, FMemory::Free);
+		FSharedBufferRef Ref = FSharedBuffer::TakeOwnership(Data, Size, FMemory::Free).ToSharedRef();
 		TestTrue(TEXT("FSharedBuffer::TakeOwnership(FMemory::Free).IsOwned()"), Ref->IsOwned());
 		TestEqual(TEXT("FSharedBuffer::TakeOwnership(FMemory::Free).GetSize()"), Ref->GetSize(), Size);
 		TestEqual(TEXT("FSharedBuffer::TakeOwnership(FMemory::Free).GetData()"), Ref->GetData(), Data);
@@ -223,11 +257,35 @@ bool FSharedBufferTest::RunTest(const FString& Parameters)
 		TestTrue(TEXT("FSharedBuffer::TakeOwnership(Lambda) Deleted"), bDeleted);
 	}
 
+	// Test TakeOwnership BufferOwner
+	{
+		bool bDeleted = false;
+		class FTestBufferOwner final : public FSharedBuffer::FBufferOwner
+		{
+		public:
+			FTestBufferOwner(bool& bInDeletedRef)
+				: bDeletedRef(bInDeletedRef)
+			{
+			}
+
+		private:
+			virtual void Free(void* Data, uint64 Size) final
+			{
+				bDeletedRef = true;
+			}
+
+			bool& bDeletedRef;
+		};
+		constexpr uint64 Size = 64;
+		FSharedBuffer::TakeOwnership(nullptr, 0, FTestBufferOwner(bDeleted));
+		TestTrue(TEXT("FSharedBuffer::TakeOwnership(BufferOwner) Deleted"), bDeleted);
+	}
+
 	// Test MakeOwned
 	{
 		constexpr uint64 Size = 64;
 		uint8 Data[Size]{};
-		FSharedBufferRef Ref = FSharedBuffer::MakeOwned(FSharedBuffer::MakeView(Data, Size));
+		FSharedBufferRef Ref = FSharedBuffer::MakeOwned(FSharedBuffer::MakeView(Data, Size).ToSharedRef());
 		TestTrue(TEXT("FSharedBuffer::MakeOwned(Wrap).IsOwned()"), Ref->IsOwned());
 		TestEqual(TEXT("FSharedBuffer::MakeOwned(Wrap).GetSize()"), Ref->GetSize(), Size);
 		TestNotEqual(TEXT("FSharedBuffer::MakeOwned(Wrap).GetData()"), Ref->GetData(), static_cast<void*>(Data));
@@ -248,7 +306,7 @@ bool FSharedBufferTest::RunTest(const FString& Parameters)
 	{
 		// MakeReadOnly from a new reference.
 		constexpr uint64 Size = 64;
-		FSharedBufferConstRef Ref = FSharedBuffer::MakeReadOnly(FSharedBuffer::Alloc(Size));
+		FSharedBufferConstRef Ref = FSharedBuffer::MakeReadOnly(FSharedBuffer::Alloc(Size)).ToSharedRef();
 		TestTrue(TEXT("FSharedBuffer::MakeReadOnly(AllocRef).IsOwned()"), Ref->IsOwned());
 		TestTrue(TEXT("FSharedBuffer::MakeReadOnly(AllocRef).IsReadOnly()"), Ref->IsReadOnly());
 		TestEqual(TEXT("FSharedBuffer::MakeReadOnly(AllocRef).GetSize()"), Ref->GetSize(), Size);
@@ -292,6 +350,7 @@ bool FSharedBufferTest::RunTest(const FString& Parameters)
 		FSharedBufferConstPtr Ptr = FSharedBuffer::MakeReadOnly(MoveTemp(MutablePtr));
 		TestTrue(TEXT("FSharedBuffer::MakeReadOnly(SharedRef).IsOwned()"), Ptr->IsOwned());
 		TestTrue(TEXT("FSharedBuffer::MakeReadOnly(SharedRef).IsReadOnly()"), Ptr->IsReadOnly());
+		TestFalse(TEXT("FSharedBuffer::MakeReadOnly(SharedRef).IsReadOnly(SharedPtr)"), SharedPtr->IsReadOnly());
 		TestNotEqual(TEXT("FSharedBuffer::MakeReadOnly(SharedRef).GetData()"), Ptr->GetData(), Data);
 	}
 	{
@@ -305,13 +364,14 @@ bool FSharedBufferTest::RunTest(const FString& Parameters)
 		TestTrue(TEXT("FSharedBuffer::MakeReadOnly(WeakRef).IsOwned()"), Ptr->IsOwned());
 		TestTrue(TEXT("FSharedBuffer::MakeReadOnly(WeakRef).IsReadOnly()"), Ptr->IsReadOnly());
 		TestNotEqual(TEXT("FSharedBuffer::MakeReadOnly(WeakRef).GetData()"), Ptr->GetData(), Data);
+		TestFalse(TEXT("FSharedBuffer::MakeReadOnly(WeakRef).Pin().IsValid()"), WeakPtr.Pin().IsValid());
 	}
 
 	// Test MakeWritable
 	{
 		// MakeWritable from a new reference.
 		constexpr uint64 Size = 64;
-		FSharedBufferConstRef ConstRef = FSharedBuffer::MakeReadOnly(FSharedBuffer::Alloc(Size));
+		FSharedBufferConstRef ConstRef = FSharedBuffer::MakeReadOnly(FSharedBuffer::Alloc(Size)).ToSharedRef();
 		FSharedBufferRef Ref = FSharedBuffer::MakeWritable(MoveTemp(ConstRef));
 		TestTrue(TEXT("FSharedBuffer::MakeWritable(AllocRef).IsOwned()"), Ref->IsOwned());
 		TestFalse(TEXT("FSharedBuffer::MakeWritable(AllocRef).IsReadOnly()"), Ref->IsReadOnly());
@@ -357,7 +417,7 @@ bool FSharedBufferTest::RunTest(const FString& Parameters)
 		FSharedBufferConstPtr ConstPtr = FSharedBuffer::MakeReadOnly(FSharedBuffer::Alloc(Size));
 		FSharedBufferConstPtr SharedPtr = ConstPtr;
 		const void* const Data = ConstPtr->GetData();
-		FSharedBufferConstPtr Ptr = FSharedBuffer::MakeWritable(MoveTemp(ConstPtr));
+		FSharedBufferPtr Ptr = FSharedBuffer::MakeWritable(MoveTemp(ConstPtr));
 		TestTrue(TEXT("FSharedBuffer::MakeWritable(SharedRef).IsOwned()"), Ptr->IsOwned());
 		TestFalse(TEXT("FSharedBuffer::MakeWritable(SharedRef).IsReadOnly()"), Ptr->IsReadOnly());
 		TestNotEqual(TEXT("FSharedBuffer::MakeWritable(SharedRef).GetData()"), const_cast<const void*>(Ptr->GetData()), Data);
@@ -368,17 +428,18 @@ bool FSharedBufferTest::RunTest(const FString& Parameters)
 		FSharedBufferConstPtr ConstPtr = FSharedBuffer::MakeReadOnly(FSharedBuffer::Alloc(Size));
 		FSharedBufferConstWeakPtr WeakPtr = ConstPtr;
 		const void* const Data = ConstPtr->GetData();
-		FSharedBufferConstPtr Ptr = FSharedBuffer::MakeWritable(MoveTemp(ConstPtr));
+		FSharedBufferPtr Ptr = FSharedBuffer::MakeWritable(MoveTemp(ConstPtr));
 		TestTrue(TEXT("FSharedBuffer::MakeWritable(WeakRef).IsOwned()"), Ptr->IsOwned());
 		TestFalse(TEXT("FSharedBuffer::MakeWritable(WeakRef).IsReadOnly()"), Ptr->IsReadOnly());
 		TestNotEqual(TEXT("FSharedBuffer::MakeWritable(WeakRef).GetData()"), const_cast<const void*>(Ptr->GetData()), Data);
+		TestTrue(TEXT("FSharedBuffer::MakeWritable(WeakRef).Pin().IsValid()"), WeakPtr.Pin().IsValid());
 	}
 
 	// Test WeakPtr
 	{
 		FSharedBufferWeakPtr WeakPtr;
 		{
-			FSharedBufferRef Ref = FSharedBuffer::Alloc(0);
+			FSharedBufferRef Ref = FSharedBuffer::Alloc(0).ToSharedRef();
 			WeakPtr = Ref;
 			TestTrue(TEXT("FSharedBufferWeakPtr(Ref).Pin().IsValid()"), WeakPtr.Pin().IsValid());
 		}
diff --git a/Engine/Source/Runtime/Core/Private/Tests/Serialization/CompactBinaryTest.cpp b/Engine/Source/Runtime/Core/Private/Tests/Serialization/CompactBinaryTest.cpp
--- a/Engine/Source/Runtime/Core/Private/Tests/Serialization/CompactBinaryTest.cpp
+++ b/Engine/Source/Runtime/Core/Private/Tests/Serialization/CompactBinaryTest.cpp
@@ -928,8 +928,8 @@ bool FCbFieldRefTest::RunTest(const FString& Parameters)
 	// Test Field w/ Type from Shared Buffer
 	{
 		uint8 Payload[] = { uint8(ECbFieldType::Binary), 3, 4, 5, 6 }; // Size: 3, Data: 4/5/6
-		FSharedBufferRef ViewBuffer = FSharedBuffer::MakeView(MakeMemoryView(Payload));
-		FSharedBufferRef OwnedBuffer = FSharedBuffer::MakeOwned(ViewBuffer);
+		FSharedBufferPtr ViewBuffer = FSharedBuffer::MakeView(MakeMemoryView(Payload));
+		FSharedBufferPtr OwnedBuffer = FSharedBuffer::MakeOwned(ViewBuffer);
 		FSharedBufferConstPtr ViewBufferPtr = ViewBuffer;
 		FSharedBufferConstPtr OwnedBufferPtr = OwnedBuffer;
 
diff --git a/Engine/Source/Runtime/Core/Public/Memory/SharedBuffer.h b/Engine/Source/Runtime/Core/Public/Memory/SharedBuffer.h
--- a/Engine/Source/Runtime/Core/Public/Memory/SharedBuffer.h
+++ b/Engine/Source/Runtime/Core/Public/Memory/SharedBuffer.h
@@ -6,6 +6,7 @@
 #include "HAL/UnrealMemory.h"
 #include "Memory/MemoryView.h"
 #include "Misc/EnumClassFlags.h"
+#include "Templates/Invoke.h"
 #include "Templates/TypeHash.h"
 #include "Templates/UnrealTemplate.h"
 
@@ -24,26 +25,26 @@ enum ESharedBufferFlags : uint16
 	None = 0,
 	Owned = 1 << 0,
 	ReadOnly = 1 << 1,
-	HasDeleter = 1 << 2,
+	HasBufferOwner = 1 << 2,
 };
 
 ENUM_CLASS_FLAGS(ESharedBufferFlags);
 
-struct FSharedRefOps
+struct FSharedRefOps final
 {
 	static bool TryAddRef(const FSharedBuffer* Buffer);
 	static void AddRef(const FSharedBuffer* Buffer);
 	static void Release(const FSharedBuffer* Buffer);
 };
 
-struct FSharedPtrOps
+struct FSharedPtrOps final
 {
 	static bool TryAddRef(const FSharedBuffer* Buffer);
 	static void AddRef(const FSharedBuffer* Buffer);
 	static void Release(const FSharedBuffer* Buffer);
 };
 
-struct FWeakPtrOps
+struct FWeakPtrOps final
 {
 	static bool TryAddRef(const FSharedBuffer* Buffer);
 	static void AddRef(const FSharedBuffer* Buffer);
@@ -51,7 +52,7 @@ struct FWeakPtrOps
 };
 
 template <bool bInAllowNull, bool bInIsConst, bool bInIsWeak>
-class TSharedBufferPtr
+class TSharedBufferPtr final
 {
 	static constexpr bool bAllowNull = bInAllowNull;
 	static constexpr bool bIsConst = bInIsConst;
@@ -198,7 +199,7 @@ private:
 ///////////////////////////////////////////////////////////////////////////////////////////////////////////////////////
 
 /** A non-nullable, thread-safe, reference to a writable shared buffer. */
-class FSharedBufferRef
+class FSharedBufferRef final
 {
 public:
 	using PtrType = SharedBufferPrivate::TSharedBufferPtr</*bAllowNull*/ false, /*bIsConst*/ false, /*bIsWeak*/ false>;
@@ -214,6 +215,7 @@ public:
 
 private:
 	constexpr inline friend bool IsValid(const FSharedBufferRef&) { return true; }
+	inline friend FSharedBufferRef ToSharedRef(const FSharedBufferRef& Ref) { return Ref; }
 
 	friend class FSharedBufferConstRef;
 	friend class FSharedBufferPtr;
@@ -224,8 +226,8 @@ private:
 	PtrType Ptr;
 };
 
-/** A non-nullable, thread-safe, reference to a const shared buffer. */
-class FSharedBufferConstRef
+/** A non-nullable, thread-safe, const reference to a shared buffer. */
+class FSharedBufferConstRef final
 {
 public:
 	using PtrType = SharedBufferPrivate::TSharedBufferPtr</*bAllowNull*/ false, /*bIsConst*/ true, /*bIsWeak*/ false>;
@@ -243,6 +245,7 @@ public:
 
 private:
 	constexpr inline friend bool IsValid(const FSharedBufferConstRef&) { return true; }
+	inline friend FSharedBufferConstRef ToSharedRef(const FSharedBufferConstRef& Ref) { return Ref; }
 
 	friend class FSharedBufferConstPtr;
 	friend class FSharedBufferConstWeakPtr;
@@ -253,7 +256,7 @@ private:
 ///////////////////////////////////////////////////////////////////////////////////////////////////////////////////////
 
 /** A non-nullable, thread-safe, pointer to a writable shared buffer. */
-class FSharedBufferPtr
+class FSharedBufferPtr final
 {
 public:
 	using PtrType = SharedBufferPrivate::TSharedBufferPtr</*bAllowNull*/ true, /*bIsConst*/ false, /*bIsWeak*/ false>;
@@ -279,6 +282,8 @@ public:
 
 private:
 	inline friend bool IsValid(const FSharedBufferPtr& InPtr) { return InPtr.IsValid(); }
+	inline friend FSharedBufferRef ToSharedRef(const FSharedBufferPtr& InPtr) { return InPtr.ToSharedRef(); }
+	inline friend FSharedBufferRef ToSharedRef(FSharedBufferPtr&& InPtr) { return MoveTemp(InPtr).ToSharedRef(); }
 
 	friend class FSharedBufferConstPtr;
 	friend class FSharedBufferWeakPtr;
@@ -287,8 +292,8 @@ private:
 	PtrType Ptr;
 };
 
-/** A non-nullable, thread-safe, pointer to a const shared buffer. */
-class FSharedBufferConstPtr
+/** A non-nullable, thread-safe, const pointer to a shared buffer. */
+class FSharedBufferConstPtr final
 {
 public:
 	using PtrType = SharedBufferPrivate::TSharedBufferPtr</*bAllowNull*/ true, /*bIsConst*/ true, /*bIsWeak*/ false>;
@@ -317,6 +322,8 @@ public:
 
 private:
 	inline friend bool IsValid(const FSharedBufferConstPtr& InPtr) { return InPtr.IsValid(); }
+	inline friend FSharedBufferConstRef ToSharedRef(const FSharedBufferConstPtr& InPtr) { return InPtr.ToSharedRef(); }
+	inline friend FSharedBufferConstRef ToSharedRef(FSharedBufferConstPtr&& InPtr) { return MoveTemp(InPtr).ToSharedRef(); }
 
 	friend class FSharedBufferConstWeakPtr;
 
@@ -326,7 +333,7 @@ private:
 ///////////////////////////////////////////////////////////////////////////////////////////////////////////////////////
 
 /** A non-nullable, thread-safe, weak, pointer to a writable shared buffer. */
-class FSharedBufferWeakPtr
+class FSharedBufferWeakPtr final
 {
 public:
 	using PtrType = SharedBufferPrivate::TSharedBufferPtr</*bAllowNull*/ true, /*bIsConst*/ false, /*bIsWeak*/ true>;
@@ -349,8 +356,8 @@ private:
 	PtrType Ptr;
 };
 
-/** A non-nullable, thread-safe, weak, pointer to a const shared buffer. */
-class FSharedBufferConstWeakPtr
+/** A non-nullable, thread-safe, weak, const pointer to a shared buffer. */
+class FSharedBufferConstWeakPtr final
 {
 public:
 	using PtrType = SharedBufferPrivate::TSharedBufferPtr</*bAllowNull*/ true, /*bIsConst*/ true, /*bIsWeak*/ true>;
@@ -414,83 +421,150 @@ template<> struct TIsWeakPointerType<FSharedBufferConstWeakPtr> { enum { Value =
  *
  *     FSharedBuffer::MakeView(Data, Size)
  *     FSharedBuffer::MakeView(MakeMemoryView(...))
+ *     FSharedBuffer::MakeView(Data, Size, OuterBuffer)
+ *     FSharedBuffer::MakeView(MakeMemoryView(...), OuterBuffer)
  *
  * Make an owned buffer by taking ownership of the buffer pointed to by the arguments.
  *
  *     FSharedBuffer::TakeOwnership(Data, Size, FMemory::Free)
  *     FSharedBuffer::TakeOwnership(Data, Size, [](void* Data) { delete[] static_cast<uint8*>(Data); })
+ *     FSharedBuffer::TakeOwnership(Data, Size, FCustomBufferOwner())
  */
-class FSharedBuffer
+class FSharedBuffer final
 {
 	using ESharedBufferFlags = SharedBufferPrivate::ESharedBufferFlags;
 
 public:
+	/**
+	 * Base for customizing lifetime management for the buffer.
+	 *
+	 * Types derived from this must be movable. Prefer an arbitrary callable (including a lambda) to
+	 * handle deletion in place of the extra code required to derive from this type.
+	 */
+	class FBufferOwner
+	{
+	protected:
+		FBufferOwner() = default;
+
+	public:
+		virtual ~FBufferOwner() = default;
+
+		/** Free the memory for the buffer. */
+		virtual void Free(void* Data, uint64 Size) = 0;
+	};
+
 	/** Make an owned writable buffer of Size bytes. */
-	static inline FSharedBufferRef Alloc(const uint64 Size, const uint32 Alignment = DEFAULT_ALIGNMENT)
+	static inline FSharedBufferPtr Alloc(const uint64 Size, const uint32 Alignment = DEFAULT_ALIGNMENT)
 	{
 		return NewBuffer(FMemory::Malloc(Size, Alignment), Size, ESharedBufferFlags::Owned);
 	}
 
 	/** Make an owned writable clone of the memory view. */
-	static inline FSharedBufferRef Clone(const FConstMemoryView View)
+	static inline FSharedBufferPtr Clone(const FConstMemoryView View)
 	{
 		return Clone(View.GetData(), View.GetSize());
 	}
 
 	/** Make an owned writable clone of the buffer. */
-	static inline FSharedBufferRef Clone(const void* const Data, const uint64 Size)
+	static inline FSharedBufferPtr Clone(const void* const Data, const uint64 Size)
 	{
-		const FSharedBufferRef Buffer = Alloc(Size);
+		const FSharedBufferPtr Buffer = Alloc(Size);
 		FMemory::Memcpy(Buffer->GetData(), Data, Size);
 		return Buffer;
 	}
 
 	/** Make a non-owned writable view of the memory view. */
-	static inline FSharedBufferRef MakeView(FMutableMemoryView View)
+	static inline FSharedBufferPtr MakeView(FMutableMemoryView View)
 	{
 		return MakeView(View.GetData(), View.GetSize());
 	}
 
 	/** Make a non-owned read-only view of the memory view. */
-	static inline FSharedBufferConstRef MakeView(FConstMemoryView View)
+	static inline FSharedBufferConstPtr MakeView(FConstMemoryView View)
 	{
 		return MakeView(View.GetData(), View.GetSize());
 	}
 
 	/** Make a non-owned writable view of the buffer. */
-	static inline FSharedBufferRef MakeView(void* Data, uint64 Size)
+	static inline FSharedBufferPtr MakeView(void* Data, uint64 Size)
 	{
 		return NewBuffer(Data, Size, ESharedBufferFlags::None);
 	}
 
 	/** Make a non-owned read-only view of the buffer. */
-	static inline FSharedBufferConstRef MakeView(const void* Data, uint64 Size)
+	static inline FSharedBufferConstPtr MakeView(const void* Data, uint64 Size)
 	{
 		// const_cast is safe here because the return value only allows accessing data through a const void*
 		return MakeView(const_cast<void*>(Data), Size);
 	}
 
+	/** Make a non-owned read-only view of the outer buffer and hold a reference to it. */
+	static inline FSharedBufferConstPtr MakeView(FConstMemoryView View, const FSharedBuffer& OuterBuffer)
+	{
+		return MakeView(View.GetData(), View.GetSize(), OuterBuffer);
+	}
+
+	/** Make a non-owned read-only view of the outer buffer and hold a reference to it. */
+	static inline FSharedBufferConstPtr MakeView(const void* Data, uint64 Size, const FSharedBuffer& OuterBuffer)
+	{
+		check(OuterBuffer.GetView().Contains(MakeMemoryView(Data, Size)));
+		// Create a reference now to prevent the read-only state from changing before we add a reference later.
+		FSharedBufferConstPtr OuterBufferPtr{FSharedBufferConstPtr::PtrType(&OuterBuffer)};
+		// This buffer is not owned, but is read-only if its outer is read-only.
+		const ESharedBufferFlags Flags = GetFlags(OuterBuffer.ReferenceCountAndFlags) & ESharedBufferFlags::ReadOnly;
+		// const_cast is safe here because the return value only allows accessing data through a const void*
+		return NewBufferWithOwner<FOuterBufferOwner>(const_cast<void*>(Data), Size, Flags, MoveTemp(OuterBufferPtr));
+	}
+
 	/**
 	 * Make an owned writable buffer by taking ownership of the provided memory.
 	 *
-	 * @param Deleter Called with the data pointer to free memory when there are no shared references.
+	 * @param DeleteFunction Called with the data pointer to free memory when there are no shared references.
+	 */
+	template <typename DeleteFunctionType,
+		decltype(Invoke(std::declval<DeleteFunctionType>(), std::declval<void*>()))* = nullptr>
+	static inline FSharedBufferPtr TakeOwnership(void* Data, uint64 Size, DeleteFunctionType&& DeleteFunction)
+	{
+		return NewBufferWithOwner<TDeleteFunctionOwner<DeleteFunctionType>>(Data, Size, ESharedBufferFlags::Owned,
+			Forward<DeleteFunctionType>(DeleteFunction));
+	}
+
+	/**
+	 * Make an owned read-only buffer by taking ownership of the provided memory.
+	 *
+	 * @param DeleteFunction Called with the data pointer to free memory when there are no shared references.
+	 */
+	template <typename DeleteFunctionType,
+		decltype(Invoke(std::declval<DeleteFunctionType>(), std::declval<void*>()))* = nullptr>
+	static inline FSharedBufferConstPtr TakeOwnership(const void* Data, uint64 Size, DeleteFunctionType&& DeleteFunction)
+	{
+		// const_cast is safe here because the return value only allows accessing data through a const void*
+		return TakeOwnership(const_cast<void*>(Data), Size, Forward<DeleteFunctionType>(DeleteFunction));
+	}
+
+	/**
+	 * Make an owned read-only buffer by taking ownership of the provided memory.
+	 *
+	 * @param BufferOwner An owner for the buffer that is derived from FBufferOwner.
 	 */
-	template <typename DeleterType>
-	static inline FSharedBufferRef TakeOwnership(void* Data, uint64 Size, DeleterType&& Deleter)
+	template <typename BufferOwnerType,
+		decltype(ImplicitConv<FBufferOwner*>(std::declval<std::decay_t<BufferOwnerType>*>()))* = nullptr>
+	static inline FSharedBufferPtr TakeOwnership(void* Data, uint64 Size, BufferOwnerType&& BufferOwner)
 	{
-		return NewBufferWithDeleter(Data, Size, ESharedBufferFlags::Owned, Forward<DeleterType>(Deleter));
+		return NewBufferWithOwner<BufferOwnerType>(Data, Size, ESharedBufferFlags::Owned, Forward<BufferOwnerType>(BufferOwner));
 	}
 
 	/**
 	 * Make an owned read-only buffer by taking ownership of the provided memory.
 	 *
-	 * @param Deleter Called with the data pointer to free memory when there are no shared references.
+	 * @param BufferOwner An owner for the buffer that is derived from FBufferOwner.
 	 */
-	template <typename DeleterType>
-	static inline FSharedBufferConstRef TakeOwnership(const void* Data, uint64 Size, DeleterType&& Deleter)
+	template <typename BufferOwnerType,
+		decltype(&ImplicitConv<FBufferOwner*>(std::declval<std::decay_t<BufferOwnerType>*>()))* = nullptr>
+	static inline FSharedBufferConstPtr TakeOwnership(const void* Data, uint64 Size, BufferOwnerType&& BufferOwner)
 	{
 		// const_cast is safe here because the return value only allows accessing data through a const void*
-		return TakeOwnership(const_cast<void*>(Data), Size, Forward<DeleterType>(Deleter));
+		return TakeOwnership(const_cast<void*>(Data), Size, Forward<BufferOwnerType>(BufferOwner));
 	}
 
 	/**
@@ -502,22 +576,15 @@ public:
 	template <typename T,
 		std::enable_if_t<std::is_same<FSharedBuffer, std::decay_t<decltype(*std::declval<T>())>>::value>* = nullptr
 		>
-	static inline auto MakeOwned(T&& Buffer) -> decltype(std::decay_t<T>(Clone(Buffer->GetData(), Buffer->GetSize())))
+	static inline auto MakeOwned(T&& Buffer) -> decltype(std::decay_t<T>(ToSharedRef(Clone(Buffer->GetData(), Buffer->GetSize()))))
 	{
-		return !IsValid(Buffer) || Buffer->IsOwned() ? Forward<T>(Buffer) : Clone(Buffer->GetData(), Buffer->GetSize());
-	}
-
-	/** Return the buffer if it is owned and either read-only or this is the only reference, or a clone otherwise. */
-	static inline FSharedBufferConstRef MakeReadOnly(FSharedBufferRef&& Buffer)
-	{
-		const FSharedBuffer& BufferRef = *Buffer;
-		return BufferRef.TryMakeReadOnly() ? MoveTemp(Buffer) : MakeReadOnly(Clone(*Buffer));
+		return !IsValid(Buffer) || Buffer->IsOwned() ? Forward<T>(Buffer) : ToSharedRef(Clone(Buffer->GetData(), Buffer->GetSize()));
 	}
 
 	/** Return the buffer if it is owned and either read-only or this is the only reference, or a clone otherwise. */
 	static inline FSharedBufferConstRef MakeReadOnly(FSharedBufferConstRef&& Buffer)
 	{
-		return Buffer->TryMakeReadOnly() ? MoveTemp(Buffer) : MakeReadOnly(Clone(*Buffer));
+		return Buffer->TryMakeReadOnly() ? MoveTemp(Buffer) : ToSharedRef(MakeReadOnly(Clone(*Buffer)));
 	}
 
 	/** Return the buffer if it is owned and either read-only or this is the only reference, or a clone otherwise. */
@@ -533,6 +600,7 @@ public:
 	}
 
 	// Disable overloads that would allow the argument to be copied.
+	static void MakeReadOnly(FSharedBufferRef&&) = delete;
 	static void MakeReadOnly(const FSharedBufferRef&) = delete;
 	static void MakeReadOnly(const FSharedBufferConstRef&) = delete;
 	static void MakeReadOnly(const FSharedBufferPtr&) = delete;
@@ -548,7 +616,7 @@ public:
 		{
 			return FSharedBufferRef(FSharedBufferRef::PtrType(const_cast<FSharedBuffer*>(&Buffer.Get())));
 		}
-		return Clone(*Buffer);
+		return Clone(*Buffer).ToSharedRef();
 	}
 
 	/** Return the buffer because it is already writable. */
@@ -578,15 +646,21 @@ public:
 	/** The size of the buffer in bytes. */
 	inline uint64 GetSize() const { return Size; }
 
-	/** Whether the shared buffer owns the memory that it provides a view of. */
-	inline bool IsOwned() const { return (GetFlags(ReferenceCountAndFlags) & ESharedBufferFlags::Owned) != 0; }
+	/** Whether this shared buffer owns the memory that it provides a view of. */
+	inline bool IsOwned() const
+	{
+		return EnumHasAnyFlags(GetFlags(ReferenceCountAndFlags), ESharedBufferFlags::Owned);
+	}
 
 	/**
 	 * Whether the shared buffer is read-only.
 	 *
-	 * A read-only shared buffer is owned and every reference or pointer to it is const.
+	 * A read-only shared buffer (or its outer) is owned and every reference to it is const.
 	 */
-	inline bool IsReadOnly() const { return (GetFlags(ReferenceCountAndFlags) & ESharedBufferFlags::ReadOnly) != 0; }
+	inline bool IsReadOnly() const
+	{
+		return EnumHasAnyFlags(GetFlags(ReferenceCountAndFlags), ESharedBufferFlags::ReadOnly);
+	}
 
 	/**
 	 * Try to make the shared buffer read-only.
@@ -597,7 +671,7 @@ public:
 	 */
 	CORE_API bool TryMakeReadOnly() const;
 
-	// Disable this overload because it would allow a read-only buffer to be mutated.
+	// Disable this overload because it would allow a read-only buffer to be written.
 	bool TryMakeReadOnly() = delete;
 
 	/**
@@ -622,24 +696,24 @@ public:
 
 private:
 	/**
-	 * Construct a new shared buffer with a custom deleter.
+	 * Construct a new shared buffer with a buffer owner.
 	 *
 	 * @param Data A pointer to the start of the data buffer.
 	 * @param Size The size of the data buffer.
-	 * @param Flags The flags associated with the data buffer. Internally adds HasDeleter.
-	 * @param Delete A callable to delete the data buffer when it is released.
+	 * @param Flags The flags associated with the data buffer. Internally adds HasBufferOwner.
+	 * @param OwnerArgs Arguments to the constructor of the buffer owner.
 	 */
-	template <typename DeleterType>
-	static inline FSharedBufferRef NewBufferWithDeleter(
+	template <typename OwnerType, typename... OwnerArgTypes>
+	static inline FSharedBufferRef NewBufferWithOwner(
 		void* const Data,
 		const uint64 Size,
 		ESharedBufferFlags Flags,
-		DeleterType&& Deleter)
+		OwnerArgTypes&&... OwnerArgs)
 	{
-		Flags |= ESharedBufferFlags::HasDeleter;
-		static_assert(alignof(TSharedBufferDeleter<DeleterType>) <= alignof(FSharedBuffer), "Required alignment is too high.");
-		FSharedBufferRef Buffer = NewBuffer(Data, Size, Flags, sizeof(TSharedBufferDeleter<DeleterType>));
-		new(&Buffer.Get() + 1) TSharedBufferDeleter<DeleterType>(Forward<DeleterType>(Deleter));
+		Flags |= ESharedBufferFlags::HasBufferOwner;
+		static_assert(alignof(OwnerType) <= alignof(FSharedBuffer), "Required alignment of OwnerType is too high.");
+		FSharedBufferRef Buffer = NewBuffer(Data, Size, Flags, sizeof(OwnerType));
+		new(&Buffer.Get() + 1) OwnerType(Forward<OwnerArgTypes>(OwnerArgs)...);
 		return Buffer;
 	}
 
@@ -649,9 +723,9 @@ private:
 	 * @param Data A pointer to the start of the data buffer.
 	 * @param Size The size of the data buffer.
 	 * @param Flags The flags associated with the data buffer.
-	 * @param DeleterSize The size of the optional deleter to allocate after the FSharedBuffer.
+	 * @param OwnerSize The size of the optional buffer owner to allocate after the shared buffer.
 	 */
-	CORE_API static FSharedBufferRef NewBuffer(void* Data, uint64 Size, ESharedBufferFlags Flags, uint32 DeleterSize = 0);
+	CORE_API static FSharedBufferRef NewBuffer(void* Data, uint64 Size, ESharedBufferFlags Flags, uint32 OwnerSize = 0);
 
 	/**
 	 * Delete a shared buffer. Requires ReleaseData to have been called already.
@@ -671,37 +745,41 @@ private:
 	FSharedBuffer& operator=(const FSharedBuffer&) = delete;
 	~FSharedBuffer() = default;
 
-private:
-	class FSharedBufferDeleter
+	/** A buffer owner that holds a reference to an outer shared buffer. */
+	class FOuterBufferOwner final : public FBufferOwner
 	{
 	public:
-		/** Free the memory for the buffer and destroy the deleter. */
-		virtual void Free(void* Data) = 0;
+		explicit FOuterBufferOwner(FSharedBufferConstPtr&& InOuterBuffer)
+			: OuterBuffer(MoveTemp(InOuterBuffer))
+		{
+		}
 
-	protected:
-		FSharedBufferDeleter() = default;
-		~FSharedBufferDeleter() = default;
+	private:
+		virtual void Free(void* InData, uint64 InSize) final
+		{
+			// Destruction of the OuterBuffer reference will release the buffer.
+		}
+
+		FSharedBufferConstPtr OuterBuffer;
 	};
 
-	template <typename DeleterType>
-	class TSharedBufferDeleter final : public FSharedBufferDeleter
+	/** A buffer owner that wraps an arbitrary callable delete function. */
+	template <typename DeleteFunctionType>
+	class TDeleteFunctionOwner final : public FBufferOwner
 	{
 	public:
-		explicit TSharedBufferDeleter(DeleterType&& InDeleter)
-			: Deleter(Forward<DeleterType>(InDeleter))
+		explicit TDeleteFunctionOwner(DeleteFunctionType&& InDeleteFunction)
+			: InDeleteFunction(Forward<DeleteFunctionType>(InDeleteFunction))
 		{
 		}
 
 	private:
-		~TSharedBufferDeleter() = default;
-
-		virtual void Free(void* InData) final
+		virtual void Free(void* InData, uint64 InSize) final
 		{
-			Deleter(InData);
-			this->~TSharedBufferDeleter();
+			Invoke(InDeleteFunction, InData);
 		}
 
-		std::decay_t<DeleterType> Deleter;
+		std::decay_t<DeleteFunctionType> InDeleteFunction;
 	};
 
 private:
diff --git a/Engine/Source/Runtime/CoreUObject/Private/Tests/Virtualization/VirtualizedBulkDataTests.cpp b/Engine/Source/Runtime/CoreUObject/Private/Tests/Virtualization/VirtualizedBulkDataTests.cpp
--- a/Engine/Source/Runtime/CoreUObject/Private/Tests/Virtualization/VirtualizedBulkDataTests.cpp
+++ b/Engine/Source/Runtime/CoreUObject/Private/Tests/Virtualization/VirtualizedBulkDataTests.cpp
@@ -38,7 +38,7 @@ bool FVirtualizationWrapperTestBasic::RunTest(const FString& Parameters)
 
 	// Create a basic bulkdata (but retain ownership of the buffer!)
 	FByteVirtualizedBulkData BulkData;
-	BulkData.UpdatePayload(FSharedBuffer::MakeView(SourceBuffer.Get(), BufferSize));
+	BulkData.UpdatePayload(FSharedBuffer::MakeView(SourceBuffer.Get(), BufferSize).ToSharedRef());
 
 	// Test access via callback
 	BulkData.GetData([this,BufferSize,&SourceBuffer](FSharedBufferConstPtr Payload)
@@ -55,7 +55,7 @@ bool FVirtualizationWrapperTestBasic::RunTest(const FString& Parameters)
 
 	// Create a new bulkdata object and copy by assignment (note we assign some junk data that will get overwritten)
 	FByteVirtualizedBulkData BulkDataAssignment;
-	BulkDataAssignment.UpdatePayload(FSharedBuffer::Alloc(128));
+	BulkDataAssignment.UpdatePayload(FSharedBuffer::Alloc(128).ToSharedRef());
 	BulkDataAssignment = BulkData;
 
 	// Test both bulkdata objects
@@ -111,7 +111,7 @@ bool FVirtualizationWrapperTestUpdatePayload::RunTest(const FString& Parameters)
 
 	// Pass the buffer to to bulkdata but retain ownership
 	FByteVirtualizedBulkData BulkData;
-	BulkData.UpdatePayload(FSharedBuffer::MakeView(OriginalData.Get(), BufferSize));
+	BulkData.UpdatePayload(FSharedBuffer::MakeView(OriginalData.Get(), BufferSize).ToSharedRef());
 
 	// Access the payload, edit it and push it back into the bulkdata object
 	{
@@ -124,7 +124,7 @@ bool FVirtualizationWrapperTestUpdatePayload::RunTest(const FString& Parameters)
 		TestTrue(TEXT("Orginal buffer == Payload data"), FMemory::Memcmp(OriginalData.Get(), Payload->GetData(), Payload->GetSize()) == 0);
 
 		// Make a copy of the payload that we can edit
-		FSharedBufferRef EditablePayload = FSharedBuffer::Clone(*Payload);
+		FSharedBufferRef EditablePayload = FSharedBuffer::Clone(*Payload).ToSharedRef();
 
 		const uint8 NewValue = 255;
 		FMemory::Memset(EditablePayload->GetData(), NewValue, EditablePayload->GetSize());
@@ -151,7 +151,7 @@ bool FVirtualizationWrapperTestUpdatePayload::RunTest(const FString& Parameters)
 
 		// Update the bulkdata object with the original data but this time we give ownership of the buffer to
 		// the bulkdata object.
-		BulkData.UpdatePayload(FSharedBuffer::TakeOwnership(OriginalData.Release(), BufferSize, FMemory::Free));
+		BulkData.UpdatePayload(FSharedBuffer::TakeOwnership(OriginalData.Release(), BufferSize, FMemory::Free).ToSharedRef());
 
 		FSharedBufferConstPtr Payload = BulkData.GetData().Get();
 		TestEqual(TEXT("Updated payload length"), (int64)Payload->GetSize(), BufferSize);
diff --git a/Engine/Source/Runtime/CoreUObject/Public/Virtualization/VirtualizationUtilities.h b/Engine/Source/Runtime/CoreUObject/Public/Virtualization/VirtualizationUtilities.h
--- a/Engine/Source/Runtime/CoreUObject/Public/Virtualization/VirtualizationUtilities.h
+++ b/Engine/Source/Runtime/CoreUObject/Public/Virtualization/VirtualizationUtilities.h
@@ -48,7 +48,7 @@ public:
 	{
 		// Remove the slack from the allocated bulk data
 		Buffer = FMemory::Realloc(Buffer, DataLength, DEFAULT_ALIGNMENT);
-		BulkData.UpdatePayload(FSharedBuffer::TakeOwnership(Buffer, DataLength, FMemory::Free));
+		BulkData.UpdatePayload(FSharedBuffer::TakeOwnership(Buffer, DataLength, FMemory::Free).ToSharedRef());
 	}
 
 	virtual void Serialize(void* Data, int64 Num)
