Fix missing return

diff --git a/Engine/Source/Runtime/Core/Private/Windows/MinimalWindowsApi.cpp b/Engine/Source/Runtime/Core/Private/Windows/MinimalWindowsApi.cpp
--- a/Engine/Source/Runtime/Core/Private/Windows/MinimalWindowsApi.cpp
+++ b/Engine/Source/Runtime/Core/Private/Windows/MinimalWindowsApi.cpp
@@ -87,7 +87,7 @@ namespace Windows
 
 	MINIMAL_WINDOWS_API BOOLEAN WINAPI TryAcquireSRWLockExclusive(PSRWLOCK SRWLock)
 	{
-		::TryAcquireSRWLockExclusive((::PSRWLOCK)SRWLock);
+		return ::TryAcquireSRWLockExclusive((::PSRWLOCK)SRWLock);
 	}
 
 	MINIMAL_WINDOWS_API void WINAPI ReleaseSRWLockShared(PSRWLOCK SRWLock)
