CompactBinary: Changed compact binary to create read-only buffers

#rnx

diff --git a/Engine/Source/Runtime/Core/Private/Serialization/CompactBinaryWriter.cpp b/Engine/Source/Runtime/Core/Private/Serialization/CompactBinaryWriter.cpp
--- a/Engine/Source/Runtime/Core/Private/Serialization/CompactBinaryWriter.cpp
+++ b/Engine/Source/Runtime/Core/Private/Serialization/CompactBinaryWriter.cpp
@@ -70,10 +70,10 @@ void FCbWriter::Reset()
 FCbFieldRefIterator FCbWriter::Save() const
 {
 	const uint64 Size = SaveSize();
-	const FSharedBufferPtr Buffer = FSharedBuffer::Alloc(Size);
-	const FMutableMemoryView View = Buffer->GetView();
-	SaveToMemory(View);
-	return FCbFieldRefIterator(FCbFieldRef(Buffer), View.RightChop(Size).GetData());
+	FSharedBufferPtr MutableBuffer = FSharedBuffer::Alloc(Size);
+	SaveToMemory(MutableBuffer->GetView());
+	FSharedBufferConstPtr ConstBuffer = FSharedBuffer::MakeReadOnly(MoveTemp(MutableBuffer));
+	return FCbFieldRefIterator(FCbFieldRef(MoveTemp(ConstBuffer)), ConstBuffer->GetView().RightChop(Size).GetData());
 }
 
 FCbFieldIterator FCbWriter::SaveToMemory(const FMutableMemoryView Buffer) const
diff --git a/Engine/Source/Runtime/Core/Private/Tests/Serialization/CompactBinaryTest.cpp b/Engine/Source/Runtime/Core/Private/Tests/Serialization/CompactBinaryTest.cpp
--- a/Engine/Source/Runtime/Core/Private/Tests/Serialization/CompactBinaryTest.cpp
+++ b/Engine/Source/Runtime/Core/Private/Tests/Serialization/CompactBinaryTest.cpp
@@ -921,8 +921,10 @@ bool FCbFieldRefTest::RunTest(const FString& Parameters)
 		FCbFieldRef DefaultField;
 		TestFalse(TEXT("FCbFieldRef().HasValue()"), DefaultField.HasValue());
 		TestFalse(TEXT("FCbFieldRef().IsOwned()"), DefaultField.IsOwned());
+		TestFalse(TEXT("FCbFieldRef().IsReadOnly()"), DefaultField.IsReadOnly());
 		DefaultField.MakeOwned();
 		TestTrue(TEXT("FCbFieldRef().MakeOwned().IsOwned()"), DefaultField.IsOwned());
+		TestTrue(TEXT("FCbFieldRef().MakeOwned().IsReadOnly()"), DefaultField.IsReadOnly());
 	}
 
 	// Test Field w/ Type from Shared Buffer
@@ -977,8 +979,15 @@ bool FCbFieldRefTest::RunTest(const FString& Parameters)
 		static_cast<uint8*>(OwnedBuffer->GetData())[UE_ARRAY_COUNT(Payload) - 1] = 5;
 		TestEqualBytes(TEXT("FCbFieldRef(View).MakeOwned()"), ViewRef.AsBinary(), {4, 5, 4});
 		TestTrue(TEXT("FCbFieldRef(View).MakeOwned().IsOwned()"), ViewRef.IsOwned());
+		TestTrue(TEXT("FCbFieldRef(View).MakeOwned().IsReadOnly()"), ViewRef.IsReadOnly());
 		TestEqualBytes(TEXT("FCbFieldRef(Owned).MakeOwned()"), OwnedRef.AsBinary(), {4, 5, 5});
 		TestTrue(TEXT("FCbFieldRef(Owned).MakeOwned().IsOwned()"), OwnedRef.IsOwned());
+		TestFalse(TEXT("FCbFieldRef(Owned).MakeOwned().IsReadOnly()"), OwnedRef.IsReadOnly());
+
+		OwnedRef.MakeReadOnly();
+		TestEqualBytes(TEXT("FCbFieldRef(Owned).MakeReadOnly()"), OwnedRef.AsBinary(), {4, 5, 5});
+		TestTrue(TEXT("FCbFieldRef(Owned).MakeReadOnly().IsOwned()"), OwnedRef.IsOwned());
+		TestTrue(TEXT("FCbFieldRef(Owned).MakeReadOnly().IsReadOnly()"), OwnedRef.IsReadOnly());
 	}
 
 	// Test Field w/ Type
@@ -1006,11 +1015,17 @@ bool FCbFieldRefTest::RunTest(const FString& Parameters)
 		TestEqualBytes(TEXT("FCbFieldRef::Clone(FieldRef)"), FieldRefClone.AsBinary(), {4, 5, 6});
 
 		TestTrue(TEXT("FCbFieldRef::TakeOwnership(Void).IsOwned()"), VoidTake.IsOwned());
+		TestTrue(TEXT("FCbFieldRef::TakeOwnership(Void).IsReadOnly()"), VoidTake.IsReadOnly());
 		TestFalse(TEXT("FCbFieldRef::MakeView(Void).IsOwned()"), VoidView.IsOwned());
+		TestFalse(TEXT("FCbFieldRef::MakeView(Void).IsReadOnly()"), VoidView.IsReadOnly());
 		TestTrue(TEXT("FCbFieldRef::Clone(Void).IsOwned()"), VoidClone.IsOwned());
+		TestTrue(TEXT("FCbFieldRef::Clone(Void).IsReadOnly()"), VoidClone.IsReadOnly());
 		TestFalse(TEXT("FCbFieldRef::MakeView(Field).IsOwned()"), FieldView.IsOwned());
+		TestFalse(TEXT("FCbFieldRef::MakeView(Field).IsReadOnly()"), FieldView.IsReadOnly());
 		TestTrue(TEXT("FCbFieldRef::Clone(Field).IsOwned()"), FieldClone.IsOwned());
+		TestTrue(TEXT("FCbFieldRef::Clone(Field).IsReadOnly()"), FieldClone.IsReadOnly());
 		TestTrue(TEXT("FCbFieldRef::Clone(FieldRef).IsOwned()"), FieldRefClone.IsOwned());
+		TestTrue(TEXT("FCbFieldRef::Clone(FieldRef).IsReadOnly()"), FieldRefClone.IsReadOnly());
 	}
 
 	// Test Field w/o Type
@@ -1032,9 +1047,14 @@ bool FCbFieldRefTest::RunTest(const FString& Parameters)
 		TestTrue(TEXT("FCbFieldRef::Clone(Field, NoType).IsOwned()"), FieldClone.IsOwned());
 		TestTrue(TEXT("FCbFieldRef::Clone(FieldRef, NoType).IsOwned()"), FieldRefClone.IsOwned());
 
+		TestFalse(TEXT("FCbFieldRef::MakeView(Field, NoType).IsReadOnly()"), FieldView.IsReadOnly());
+		TestTrue(TEXT("FCbFieldRef::Clone(Field, NoType).IsReadOnly()"), FieldClone.IsReadOnly());
+		TestTrue(TEXT("FCbFieldRef::Clone(FieldRef, NoType).IsReadOnly()"), FieldRefClone.IsReadOnly());
+
 		FieldView.MakeOwned();
 		TestEqualBytes(TEXT("FCbFieldRef::MakeView(NoType).MakeOwned()"), FieldView.AsBinary(), {4, 5, 4});
 		TestTrue(TEXT("FCbFieldRef::MakeView(NoType).MakeOwned().IsOwned()"), FieldView.IsOwned());
+		TestTrue(TEXT("FCbFieldRef::MakeView(NoType).MakeOwned().IsReadOnly()"), FieldView.IsReadOnly());
 	}
 
 	return true;
@@ -1059,8 +1079,10 @@ bool FCbArrayRefTest::RunTest(const FString& Parameters)
 	{
 		FCbArrayRef DefaultArray;
 		TestFalse(TEXT("FCbArrayRef().IsOwned()"), DefaultArray.IsOwned());
+		TestFalse(TEXT("FCbArrayRef().IsReadOnly()"), DefaultArray.IsReadOnly());
 		DefaultArray.MakeOwned();
 		TestTrue(TEXT("FCbArrayRef().MakeOwned().IsOwned()"), DefaultArray.IsOwned());
+		TestTrue(TEXT("FCbArrayRef().MakeOwned().IsReadOnly()"), DefaultArray.IsReadOnly());
 	}
 
 	return true;
@@ -1085,8 +1107,10 @@ bool FCbObjectRefTest::RunTest(const FString& Parameters)
 	{
 		FCbObjectRef DefaultObject;
 		TestFalse(TEXT("FCbObjectRef().IsOwned()"), DefaultObject.IsOwned());
+		TestFalse(TEXT("FCbObjectRef().IsReadOnly()"), DefaultObject.IsReadOnly());
 		DefaultObject.MakeOwned();
 		TestTrue(TEXT("FCbObjectRef().MakeOwned().IsOwned()"), DefaultObject.IsOwned());
+		TestTrue(TEXT("FCbObjectRef().MakeOwned().IsReadOnly()"), DefaultObject.IsReadOnly());
 	}
 
 	return true;
diff --git a/Engine/Source/Runtime/Core/Private/Tests/Serialization/CompactBinaryWriterTest.cpp b/Engine/Source/Runtime/Core/Private/Tests/Serialization/CompactBinaryWriterTest.cpp
--- a/Engine/Source/Runtime/Core/Private/Tests/Serialization/CompactBinaryWriterTest.cpp
+++ b/Engine/Source/Runtime/Core/Private/Tests/Serialization/CompactBinaryWriterTest.cpp
@@ -882,6 +882,21 @@ bool FCbWriterComplexTest::RunTest(const FString& Parameters)
 
 ///////////////////////////////////////////////////////////////////////////////////////////////////////////////////////
 
+IMPLEMENT_SIMPLE_AUTOMATION_TEST(FCbWriterOwnedReadOnlyTest, "System.Core.Serialization.CbWriter.OwnedReadOnly", CompactBinaryWriterTestFlags)
+bool FCbWriterOwnedReadOnlyTest::RunTest(const FString& Parameters)
+{
+	FCbWriter Writer;
+	Writer.BeginObject();
+	Writer.EndObject();
+	FCbObjectRef Object = Writer.Save().AsObjectRef();
+	TestTrue(TEXT("FCbWriter().Save().IsOwned()"), Object.IsOwned());
+	TestTrue(TEXT("FCbWriter().Save().IsReadOnly()"), Object.IsReadOnly());
+
+	return true;
+}
+
+///////////////////////////////////////////////////////////////////////////////////////////////////////////////////////
+
 IMPLEMENT_SIMPLE_AUTOMATION_TEST(FCbWriterStreamTest, "System.Core.Serialization.CbWriter.Stream", CompactBinaryWriterTestFlags)
 bool FCbWriterStreamTest::RunTest(const FString& Parameters)
 {
diff --git a/Engine/Source/Runtime/Core/Public/Serialization/CompactBinary.h b/Engine/Source/Runtime/Core/Public/Serialization/CompactBinary.h
--- a/Engine/Source/Runtime/Core/Public/Serialization/CompactBinary.h
+++ b/Engine/Source/Runtime/Core/Public/Serialization/CompactBinary.h
@@ -835,11 +835,23 @@ public:
 	inline void MakeOwned()
 	{
 		if (!IsOwned())
+		{
+			MakeReadOnly();
+		}
+	}
+
+	/** Whether this reference is backed by a read-only buffer. */
+	inline bool IsReadOnly() const { return BufferPtr && BufferPtr->IsReadOnly(); }
+
+	/** Clone the value, if necessary, to a read-only buffer. */
+	inline void MakeReadOnly()
+	{
+		if (!IsReadOnly())
 		{
 			FSharedBufferPtr Buffer = FSharedBuffer::Alloc(BaseType::Size());
 			BaseType::CopyTo(*Buffer);
 			static_cast<BaseType&>(*this) = BaseType(Buffer->GetData(), BaseType::GetCopyType());
-			BufferPtr = MoveTemp(Buffer);
+			BufferPtr = FSharedBuffer::MakeReadOnly(MoveTemp(Buffer));
 		}
 	}
 
@@ -892,11 +904,13 @@ public:
 	}
 
 	/** Construct a value and take ownership its memory. */
-	template <typename DeleterType>
-	static inline RefType TakeOwnership(const void* const Data, DeleterType&& Deleter)
+	template <typename DeleteFunctionOrBufferOwnerType>
+	static inline RefType TakeOwnership(const void* const Data, DeleteFunctionOrBufferOwnerType&& DeleteFunctionOrBufferOwner)
 	{
 		const BaseType Value(Data);
-		return RefType(Value, FSharedBuffer::TakeOwnership(Data, Value.Size(), Forward<DeleterType>(Deleter)));
+		FSharedBufferConstPtr Buffer = FSharedBuffer::TakeOwnership(Data, Value.Size(),
+			Forward<DeleteFunctionOrBufferOwnerType>(DeleteFunctionOrBufferOwner));
+		return RefType(Value, FSharedBuffer::MakeReadOnly(MoveTemp(Buffer)));
 	}
 };
 
