CompactBinary: Renamed accessors to include a Get prefix to be consistent

#rnx

diff --git a/Engine/Source/Runtime/Core/Private/Serialization/CompactBinary.cpp b/Engine/Source/Runtime/Core/Private/Serialization/CompactBinary.cpp
--- a/Engine/Source/Runtime/Core/Private/Serialization/CompactBinary.cpp
+++ b/Engine/Source/Runtime/Core/Private/Serialization/CompactBinary.cpp
@@ -86,16 +86,16 @@ FConstMemoryView FCbField::AsMemoryView() const
 {
 	const uint32 TypeSize = FCbFieldType::HasFieldType(Type) ? sizeof(ECbFieldType) : 0;
 	const uint32 NameSize = FCbFieldType::HasFieldName(Type) ? NameLen + MeasureVarUInt(NameLen) : 0;
-	const uint64 LocalPayloadSize = PayloadSize();
-	return MakeMemoryView(static_cast<const uint8*>(Payload) - TypeSize - NameSize, TypeSize + NameSize + PayloadSize());
+	const uint64 LocalPayloadSize = GetPayloadSize();
+	return MakeMemoryView(static_cast<const uint8*>(Payload) - TypeSize - NameSize, TypeSize + NameSize + GetPayloadSize());
 }
 
-uint64 FCbField::Size() const
+uint64 FCbField::GetSize() const
 {
 	return AsMemoryView().GetSize();
 }
 
-uint64 FCbField::PayloadSize() const
+uint64 FCbField::GetPayloadSize() const
 {
 	switch (FCbFieldType::GetType(Type))
 	{
@@ -446,7 +446,7 @@ FConstMemoryView FCbArray::AsMemoryView() const
 	return MakeMemoryView(Payload, PayloadSizeByteCount + PayloadSize);
 }
 
-uint64 FCbArray::Size() const
+uint64 FCbArray::GetSize() const
 {
 	return sizeof(ECbFieldType) + AsMemoryView().GetSize();
 }
@@ -516,7 +516,7 @@ FConstMemoryView FCbObject::AsMemoryView() const
 	return MakeMemoryView(Payload, PayloadSizeByteCount + PayloadSize);
 }
 
-uint64 FCbObject::Size() const
+uint64 FCbObject::GetSize() const
 {
 	return sizeof(ECbFieldType) + AsMemoryView().GetSize();
 }
@@ -546,7 +546,7 @@ FCbField FCbObject::Find(const FAnsiStringView Name) const
 {
 	for (const FCbField Field : *this)
 	{
-		if (Name.Equals(Field.Name(), ESearchCase::CaseSensitive))
+		if (Name.Equals(Field.GetName(), ESearchCase::CaseSensitive))
 		{
 			return Field;
 		}
@@ -558,7 +558,7 @@ FCbField FCbObject::FindIgnoreCase(const FAnsiStringView Name) const
 {
 	for (const FCbField Field : *this)
 	{
-		if (Name.Equals(Field.Name(), ESearchCase::IgnoreCase))
+		if (Name.Equals(Field.GetName(), ESearchCase::IgnoreCase))
 		{
 			return Field;
 		}
@@ -763,7 +763,7 @@ static void ValidateCbObject(FConstMemoryView& View, ECbValidateMode Mode, ECbVa
 				{
 					if (Field.HasName())
 					{
-						Names.Add(Field.Name());
+						Names.Add(Field.GetName());
 					}
 					else
 					{
diff --git a/Engine/Source/Runtime/Core/Private/Tests/Serialization/CompactBinaryTest.cpp b/Engine/Source/Runtime/Core/Private/Tests/Serialization/CompactBinaryTest.cpp
--- a/Engine/Source/Runtime/Core/Private/Tests/Serialization/CompactBinaryTest.cpp
+++ b/Engine/Source/Runtime/Core/Private/Tests/Serialization/CompactBinaryTest.cpp
@@ -123,7 +123,7 @@ protected:
 	void TestField(TArrayView<const uint8> Payload, T ExpectedValue = T(), T DefaultValue = T())
 	{
 		FCbField Field(Payload.GetData(), FieldType);
-		TestEqual(TEXT("FCbField::Size()"), Field.Size(), uint64(Payload.Num()));
+		TestEqual(TEXT("FCbField::GetSize()"), Field.GetSize(), uint64(Payload.Num()));
 		TestTrue(TEXT("FCbField::HasValue()"), Field.HasValue());
 		TestFalse(TEXT("FCbField::HasError() == false"), Field.HasError());
 		TestEqual(TEXT("FCbField::GetError() == None"), Field.GetError(), ECbFieldError::None);
@@ -181,8 +181,8 @@ bool FCbFieldNoneTest::RunTest(const FString& Parameters)
 		static_assert(!DefaultField.HasValue(), "Error in HasValue()");
 		static_assert(!DefaultField.HasError(), "Error in HasError()");
 		static_assert(DefaultField.GetError() == ECbFieldError::None, "Error in GetError()");
-		TestEqual(TEXT("FCbField()::Size() == 0"), DefaultField.Size(), uint64(0));
-		TestEqual(TEXT("FCbField()::Name().Len() == 0"), DefaultField.Name().Len(), 0);
+		TestEqual(TEXT("FCbField()::GetSize() == 0"), DefaultField.GetSize(), uint64(0));
+		TestEqual(TEXT("FCbField()::GetName().Len() == 0"), DefaultField.GetName().Len(), 0);
 		TestFalse(TEXT("!FCbField()::HasName()"), DefaultField.HasName());
 		TestFalse(TEXT("!FCbField()::HasValue()"), DefaultField.HasValue());
 		TestFalse(TEXT("!FCbField()::HasError()"), DefaultField.HasError());
@@ -192,8 +192,8 @@ bool FCbFieldNoneTest::RunTest(const FString& Parameters)
 	// Test FCbField(None)
 	{
 		FCbField NoneField(nullptr, ECbFieldType::None);
-		TestEqual(TEXT("FCbField(None)::Size() == 0"), NoneField.Size(), uint64(0));
-		TestEqual(TEXT("FCbField(None)::Name().Len() == 0"), NoneField.Name().Len(), 0);
+		TestEqual(TEXT("FCbField(None)::GetSize() == 0"), NoneField.GetSize(), uint64(0));
+		TestEqual(TEXT("FCbField(None)::GetName().Len() == 0"), NoneField.GetName().Len(), 0);
 		TestFalse(TEXT("!FCbField(None)::HasName()"), NoneField.HasName());
 		TestFalse(TEXT("!FCbField(None)::HasValue()"), NoneField.HasValue());
 		TestFalse(TEXT("!FCbField(None)::HasError()"), NoneField.HasError());
@@ -205,8 +205,8 @@ bool FCbFieldNoneTest::RunTest(const FString& Parameters)
 		constexpr ECbFieldType FieldType = ECbFieldType::None | ECbFieldType::HasFieldName;
 		constexpr const ANSICHAR NoneBytes[] = { ANSICHAR(FieldType), 4, 'N', 'a', 'm', 'e' };
 		FCbField NoneField(NoneBytes);
-		TestEqual(TEXT("FCbField(None|Type|Name)::Size()"), NoneField.Size(), uint64(sizeof(NoneBytes)));
-		TestEqual(TEXT("FCbField(None|Type|Name)::Name()"), NoneField.Name(), "Name"_ASV);
+		TestEqual(TEXT("FCbField(None|Type|Name)::GetSize()"), NoneField.GetSize(), uint64(sizeof(NoneBytes)));
+		TestEqual(TEXT("FCbField(None|Type|Name)::GetName()"), NoneField.GetName(), "Name"_ASV);
 		TestTrue(TEXT("FCbField(None|Type|Name)::HasName()"), NoneField.HasName());
 		TestFalse(TEXT("!FCbField(None|Type|Name)::HasValue()"), NoneField.HasValue());
 	}
@@ -216,8 +216,8 @@ bool FCbFieldNoneTest::RunTest(const FString& Parameters)
 		constexpr ECbFieldType FieldType = ECbFieldType::None;
 		constexpr const ANSICHAR NoneBytes[] = { ANSICHAR(FieldType) };
 		FCbField NoneField(NoneBytes);
-		TestEqual(TEXT("FCbField(None|Type)::Size()"), NoneField.Size(), uint64(sizeof(NoneBytes)));
-		TestEqual(TEXT("FCbField(None|Type)::Name()"), NoneField.Name().Len(), 0);
+		TestEqual(TEXT("FCbField(None|Type)::GetSize()"), NoneField.GetSize(), uint64(sizeof(NoneBytes)));
+		TestEqual(TEXT("FCbField(None|Type)::GetName()"), NoneField.GetName().Len(), 0);
 		TestFalse(TEXT("FCbField(None|Type)::HasName()"), NoneField.HasName());
 		TestFalse(TEXT("!FCbField(None|Type)::HasValue()"), NoneField.HasValue());
 	}
@@ -227,8 +227,8 @@ bool FCbFieldNoneTest::RunTest(const FString& Parameters)
 		constexpr ECbFieldType FieldType = ECbFieldType::None | ECbFieldType::HasFieldName;
 		constexpr const ANSICHAR NoneBytes[] = { 4, 'N', 'a', 'm', 'e' };
 		FCbField NoneField(NoneBytes, FieldType);
-		TestEqual(TEXT("FCbField(None|Name)::Size()"), NoneField.Size(), uint64(sizeof(NoneBytes)));
-		TestEqual(TEXT("FCbField(None|Name)::Name()"), NoneField.Name(), "Name"_ASV);
+		TestEqual(TEXT("FCbField(None|Name)::GetSize()"), NoneField.GetSize(), uint64(sizeof(NoneBytes)));
+		TestEqual(TEXT("FCbField(None|Name)::GetName()"), NoneField.GetName(), "Name"_ASV);
 		TestTrue(TEXT("FCbField(None|Name)::HasName()"), NoneField.HasName());
 		TestFalse(TEXT("!FCbField(None|Name)::HasValue()"), NoneField.HasValue());
 	}
@@ -242,7 +242,7 @@ bool FCbFieldNullTest::RunTest(const FString& Parameters)
 	// Test FCbField(Null)
 	{
 		FCbField NullField(nullptr, ECbFieldType::Null);
-		TestEqual(TEXT("FCbField(Null)::Size() == 0"), NullField.Size(), uint64(0));
+		TestEqual(TEXT("FCbField(Null)::GetSize() == 0"), NullField.GetSize(), uint64(0));
 		TestTrue(TEXT("FCbField(Null)::IsNull()"), NullField.IsNull());
 		TestTrue(TEXT("FCbField(Null)::HasValue()"), NullField.HasValue());
 		TestFalse(TEXT("!FCbField(Null)::HasError()"), NullField.HasError());
@@ -263,13 +263,13 @@ bool FCbFieldObjectTest::RunTest(const FString& Parameters)
 {
 	auto TestIntObject = [this](const FCbObject& Object, int32 ExpectedNum, uint64 ExpectedPayloadSize)
 	{
-		TestEqual(TEXT("FCbField(Object)::AsObject().Size()"), Object.Size(), ExpectedPayloadSize + sizeof(ECbFieldType));
+		TestEqual(TEXT("FCbField(Object)::AsObject().GetSize()"), Object.GetSize(), ExpectedPayloadSize + sizeof(ECbFieldType));
 
 		int32 ActualNum = 0;
 		for (FCbFieldIterator It = Object.CreateIterator(); It; ++It)
 		{
 			++ActualNum;
-			TestNotEqual(TEXT("FCbField(Object) Iterator Name"), It->Name().Len(), 0);
+			TestNotEqual(TEXT("FCbField(Object) Iterator Name"), It->GetName().Len(), 0);
 			TestEqual(TEXT("FCbField(Object) Iterator"), It->AsInt32(), ActualNum);
 		}
 		TestEqual(TEXT("FCbField(Object)::AsObject().CreateIterator() -> Count"), ActualNum, ExpectedNum);
@@ -278,7 +278,7 @@ bool FCbFieldObjectTest::RunTest(const FString& Parameters)
 		for (FCbField Field : Object)
 		{
 			++ActualNum;
-			TestNotEqual(TEXT("FCbField(Object) Iterator Name"), Field.Name().Len(), 0);
+			TestNotEqual(TEXT("FCbField(Object) Iterator Name"), Field.GetName().Len(), 0);
 			TestEqual(TEXT("FCbField(Object) Range"), Field.AsInt32(), ActualNum);
 		}
 		TestEqual(TEXT("FCbField(Object)::AsObject() Range -> Count"), ActualNum, ExpectedNum);
@@ -363,14 +363,14 @@ bool FCbFieldObjectTest::RunTest(const FString& Parameters)
 		const uint8 ObjectType = uint8(ECbFieldType::Object | ECbFieldType::HasFieldName);
 		const uint8 Buffer[] = { ObjectType, 3, 'K', 'e', 'y', 4, uint8(ECbFieldType::HasFieldName | ECbFieldType::IntegerPositive), 1, 'F', 8 };
 		const FCbObject Object(Buffer);
-		TestEqual(TEXT("FCbObject(ObjectWithName)::Size()"), Object.Size(), uint64(6));
+		TestEqual(TEXT("FCbObject(ObjectWithName)::GetSize()"), Object.GetSize(), uint64(6));
 		const FCbObjectRef ObjectClone = FCbObjectRef::Clone(Object);
-		TestEqual(TEXT("FCbObjectRef(ObjectWithName)::Size()"), ObjectClone.Size(), uint64(6));
+		TestEqual(TEXT("FCbObjectRef(ObjectWithName)::GetSize()"), ObjectClone.GetSize(), uint64(6));
 		TestTrue(TEXT("FCbObject::Equals()"), Object.Equals(ObjectClone));
 		for (FCbFieldRefIterator It = ObjectClone.CreateRefIterator(); It; ++It)
 		{
 			FCbFieldRef Field = *It;
-			TestEqual(TEXT("FCbObjectRef::CreateRefIterator().Name()"), Field.Name(), "F"_ASV);
+			TestEqual(TEXT("FCbObjectRef::CreateRefIterator().GetName()"), Field.GetName(), "F"_ASV);
 			TestEqual(TEXT("FCbObjectRef::CreateRefIterator().AsInt32()"), Field.AsInt32(), 8);
 			TestTrue(TEXT("FCbObjectRef::CreateRefIterator().IsOwned()"), Field.IsOwned());
 		}
@@ -387,7 +387,7 @@ bool FCbFieldArrayTest::RunTest(const FString& Parameters)
 {
 	auto TestIntArray = [this](FCbArray Array, int32 ExpectedNum, uint64 ExpectedPayloadSize)
 	{
-		TestEqual(TEXT("FCbField(Array)::AsArray().Size()"), Array.Size(), ExpectedPayloadSize + sizeof(ECbFieldType));
+		TestEqual(TEXT("FCbField(Array)::AsArray().GetSize()"), Array.GetSize(), ExpectedPayloadSize + sizeof(ECbFieldType));
 		TestEqual(TEXT("FCbField(Array)::AsArray().Num()"), Array.Num(), uint64(ExpectedNum));
 
 		int32 ActualNum = 0;
@@ -465,9 +465,9 @@ bool FCbFieldArrayTest::RunTest(const FString& Parameters)
 		const uint8 ArrayType = uint8(ECbFieldType::Array | ECbFieldType::HasFieldName);
 		const uint8 Buffer[] = { ArrayType, 3, 'K', 'e', 'y', 3, 1, uint8(ECbFieldType::IntegerPositive), 8 };
 		const FCbArray Array(Buffer);
-		TestEqual(TEXT("Array(ArrayWithName)::Size()"), Array.Size(), uint64(5));
+		TestEqual(TEXT("Array(ArrayWithName)::GetSize()"), Array.GetSize(), uint64(5));
 		const FCbArrayRef ArrayClone = FCbArrayRef::Clone(Array);
-		TestEqual(TEXT("FCbArrayRef(ArrayWithName)::Size()"), ArrayClone.Size(), uint64(5));
+		TestEqual(TEXT("FCbArrayRef(ArrayWithName)::GetSize()"), ArrayClone.GetSize(), uint64(5));
 		TestTrue(TEXT("FCbArray::Equals()"), Array.Equals(ArrayClone));
 		for (FCbFieldRefIterator It = ArrayClone.CreateRefIterator(); It; ++It)
 		{
@@ -1465,22 +1465,22 @@ bool FCbFieldParseTest::RunTest(const FString& Parameters)
 		for (FCbFieldIterator It = Object.CreateIterator(); It;)
 		{
 			const FCbFieldIterator Last = It;
-			if (It.Name().Equals("A"_ASV))
+			if (It.GetName().Equals("A"_ASV))
 			{
 				A = It.AsUInt32();
 				++It;
 			}
-			if (It.Name().Equals("B"_ASV))
+			if (It.GetName().Equals("B"_ASV))
 			{
 				B = It.AsUInt32();
 				++It;
 			}
-			if (It.Name().Equals("C"_ASV))
+			if (It.GetName().Equals("C"_ASV))
 			{
 				C = It.AsUInt32();
 				++It;
 			}
-			if (It.Name().Equals("D"_ASV))
+			if (It.GetName().Equals("D"_ASV))
 			{
 				D = It.AsUInt32();
 				++It;
diff --git a/Engine/Source/Runtime/Core/Private/Tests/Serialization/CompactBinaryWriterTest.cpp b/Engine/Source/Runtime/Core/Private/Tests/Serialization/CompactBinaryWriterTest.cpp
--- a/Engine/Source/Runtime/Core/Private/Tests/Serialization/CompactBinaryWriterTest.cpp
+++ b/Engine/Source/Runtime/Core/Private/Tests/Serialization/CompactBinaryWriterTest.cpp
@@ -26,7 +26,7 @@ bool FCbWriterObjectTest::RunTest(const FString& Parameters)
 		Writer.BeginObject();
 		Writer.EndObject();
 		FCbFieldRef Field = Writer.Save();
-		if (TestEqual(TEXT("FCbWriter(Object, Empty) Validate"), ValidateCompactBinary(*Field.Buffer(), ECbValidateMode::All), ECbValidateError::None))
+		if (TestEqual(TEXT("FCbWriter(Object, Empty) Validate"), ValidateCompactBinary(*Field.GetBuffer(), ECbValidateMode::All), ECbValidateError::None))
 		{
 			TestTrue(TEXT("FCbWriter(Object, Empty).IsObject()"), Field.IsObject());
 			TestFalse(TEXT("FCbWriter(Object, Empty).AsObject()"), Field.AsObject().CreateIterator().HasValue());
@@ -40,7 +40,7 @@ bool FCbWriterObjectTest::RunTest(const FString& Parameters)
 		Writer.BeginObject();
 		Writer.EndObject();
 		FCbFieldRef Field = Writer.Save();
-		if (TestEqual(TEXT("FCbWriter(Object, Empty, Name) Validate"), ValidateCompactBinary(*Field.Buffer(), ECbValidateMode::All), ECbValidateError::None))
+		if (TestEqual(TEXT("FCbWriter(Object, Empty, Name) Validate"), ValidateCompactBinary(*Field.GetBuffer(), ECbValidateMode::All), ECbValidateError::None))
 		{
 			TestTrue(TEXT("FCbWriter(Object, Empty, Name).IsObject()"), Field.IsObject());
 			TestFalse(TEXT("FCbWriter(Object, Empty, Name).AsObject()"), Field.AsObject().CreateIterator().HasValue());
@@ -55,7 +55,7 @@ bool FCbWriterObjectTest::RunTest(const FString& Parameters)
 		Writer.Name("Float"_ASV).Float(0.0f);
 		Writer.EndObject();
 		FCbFieldRef Field = Writer.Save();
-		if (TestEqual(TEXT("FCbWriter(Object, Basic) Validate"), ValidateCompactBinary(*Field.Buffer(), ECbValidateMode::All), ECbValidateError::None))
+		if (TestEqual(TEXT("FCbWriter(Object, Basic) Validate"), ValidateCompactBinary(*Field.GetBuffer(), ECbValidateMode::All), ECbValidateError::None))
 		{
 			TestTrue(TEXT("FCbWriter(Object, Basic).IsObject()"), Field.IsObject());
 			FCbObject Object = Field.AsObject();
@@ -72,7 +72,7 @@ bool FCbWriterObjectTest::RunTest(const FString& Parameters)
 		Writer.Name("Field2"_ASV).Integer(1);
 		Writer.EndObject();
 		FCbFieldRef Field = Writer.Save();
-		if (TestEqual(TEXT("FCbWriter(Object, Uniform) Validate"), ValidateCompactBinary(*Field.Buffer(), ECbValidateMode::All), ECbValidateError::None))
+		if (TestEqual(TEXT("FCbWriter(Object, Uniform) Validate"), ValidateCompactBinary(*Field.GetBuffer(), ECbValidateMode::All), ECbValidateError::None))
 		{
 			TestTrue(TEXT("FCbWriter(Object, Uniform).IsObject()"), Field.IsObject());
 			FCbObject Object = Field.AsObject();
@@ -97,7 +97,7 @@ bool FCbWriterArrayTest::RunTest(const FString& Parameters)
 		Writer.BeginArray();
 		Writer.EndArray();
 		FCbFieldRef Field = Writer.Save();
-		if (TestEqual(TEXT("FCbWriter(Array, Empty) Validate"), ValidateCompactBinary(*Field.Buffer(), ECbValidateMode::All), ECbValidateError::None))
+		if (TestEqual(TEXT("FCbWriter(Array, Empty) Validate"), ValidateCompactBinary(*Field.GetBuffer(), ECbValidateMode::All), ECbValidateError::None))
 		{
 			TestTrue(TEXT("FCbWriter(Array, Empty).IsArray()"), Field.IsArray());
 			TestEqual(TEXT("FCbWriter(Array, Empty).AsArray()"), Field.AsArray().Num(), uint64(0));
@@ -111,7 +111,7 @@ bool FCbWriterArrayTest::RunTest(const FString& Parameters)
 		Writer.BeginArray();
 		Writer.EndArray();
 		FCbFieldRef Field = Writer.Save();
-		if (TestEqual(TEXT("FCbWriter(Array, Empty, Name) Validate"), ValidateCompactBinary(*Field.Buffer(), ECbValidateMode::All), ECbValidateError::None))
+		if (TestEqual(TEXT("FCbWriter(Array, Empty, Name) Validate"), ValidateCompactBinary(*Field.GetBuffer(), ECbValidateMode::All), ECbValidateError::None))
 		{
 			TestTrue(TEXT("FCbWriter(Array, Empty, Name).IsArray()"), Field.IsArray());
 			TestEqual(TEXT("FCbWriter(Array, Empty, Name).AsArray()"), Field.AsArray().Num(), uint64(0));
@@ -126,7 +126,7 @@ bool FCbWriterArrayTest::RunTest(const FString& Parameters)
 		Writer.Float(0.0f);
 		Writer.EndArray();
 		FCbFieldRef Field = Writer.Save();
-		if (TestEqual(TEXT("FCbWriter(Array, Basic) Validate"), ValidateCompactBinary(*Field.Buffer(), ECbValidateMode::All), ECbValidateError::None))
+		if (TestEqual(TEXT("FCbWriter(Array, Basic) Validate"), ValidateCompactBinary(*Field.GetBuffer(), ECbValidateMode::All), ECbValidateError::None))
 		{
 			TestTrue(TEXT("FCbWriter(Array, Basic).IsArray()"), Field.IsArray());
 			FCbFieldIterator Iterator = Field.AsArray().CreateIterator();
@@ -146,7 +146,7 @@ bool FCbWriterArrayTest::RunTest(const FString& Parameters)
 		Writer.Integer(1);
 		Writer.EndArray();
 		FCbFieldRef Field = Writer.Save();
-		if (TestEqual(TEXT("FCbWriter(Array, Uniform) Validate"), ValidateCompactBinary(*Field.Buffer(), ECbValidateMode::All), ECbValidateError::None))
+		if (TestEqual(TEXT("FCbWriter(Array, Uniform) Validate"), ValidateCompactBinary(*Field.GetBuffer(), ECbValidateMode::All), ECbValidateError::None))
 		{
 			TestTrue(TEXT("FCbWriter(Array, Uniform).IsArray()"), Field.IsArray());
 			FCbFieldIterator Iterator = Field.AsArray().CreateIterator();
@@ -173,7 +173,7 @@ bool FCbWriterNullTest::RunTest(const FString& Parameters)
 		Writer.Reset();
 		Writer.Null();
 		FCbFieldRef Field = Writer.Save();
-		if (TestEqual(TEXT("FCbWriter(Null) Validate"), ValidateCompactBinary(*Field.Buffer(), ECbValidateMode::All), ECbValidateError::None))
+		if (TestEqual(TEXT("FCbWriter(Null) Validate"), ValidateCompactBinary(*Field.GetBuffer(), ECbValidateMode::All), ECbValidateError::None))
 		{
 			TestFalse(TEXT("FCbWriter(Null).HasName()"), Field.HasName());
 			TestTrue(TEXT("FCbWriter(Null).IsNull()"), Field.IsNull());
@@ -186,9 +186,9 @@ bool FCbWriterNullTest::RunTest(const FString& Parameters)
 		Writer.Name("Null"_ASV);
 		Writer.Null();
 		FCbFieldRef Field = Writer.Save();
-		if (TestEqual(TEXT("FCbWriter(Null, Name) Validate"), ValidateCompactBinary(*Field.Buffer(), ECbValidateMode::All), ECbValidateError::None))
+		if (TestEqual(TEXT("FCbWriter(Null, Name) Validate"), ValidateCompactBinary(*Field.GetBuffer(), ECbValidateMode::All), ECbValidateError::None))
 		{
-			TestEqual(TEXT("FCbWriter(Null, Name).Name()"), Field.Name(), "Null"_ASV);
+			TestEqual(TEXT("FCbWriter(Null, Name).GetName()"), Field.GetName(), "Null"_ASV);
 			TestTrue(TEXT("FCbWriter(Null, Name).HasName()"), Field.HasName());
 			TestTrue(TEXT("FCbWriter(Null, Name).IsNull()"), Field.IsNull());
 		}
@@ -211,7 +211,7 @@ bool FCbWriterNullTest::RunTest(const FString& Parameters)
 		Writer.EndObject();
 
 		FCbFieldRefIterator Fields = Writer.Save();
-		TestEqual(TEXT("FCbWriter(Null, Uniform) Validate"), ValidateCompactBinaryRange(*Fields.Buffer(), ECbValidateMode::All), ECbValidateError::None);
+		TestEqual(TEXT("FCbWriter(Null, Uniform) Validate"), ValidateCompactBinaryRange(*Fields.GetBuffer(), ECbValidateMode::All), ECbValidateError::None);
 	}
 
 	// Test Null with SaveToMemory
@@ -263,7 +263,7 @@ bool FCbWriterBinaryTest::RunTest(const FString& Parameters)
 		Writer.Reset();
 		Writer.Binary(nullptr, 0);
 		FCbFieldRef Field = Writer.Save();
-		if (TestEqual(TEXT("FCbWriter(Binary, Empty) Validate"), ValidateCompactBinary(*Field.Buffer(), ECbValidateMode::All), ECbValidateError::None))
+		if (TestEqual(TEXT("FCbWriter(Binary, Empty) Validate"), ValidateCompactBinary(*Field.GetBuffer(), ECbValidateMode::All), ECbValidateError::None))
 		{
 			TestFalse(TEXT("FCbWriter(Binary, Empty).HasName()"), Field.HasName());
 			TestTrue(TEXT("FCbWriter(Binary, Empty).IsBinary()"), Field.IsBinary());
@@ -278,9 +278,9 @@ bool FCbWriterBinaryTest::RunTest(const FString& Parameters)
 		Writer.Name("Binary"_ASV);
 		Writer.Binary(BinaryValue, sizeof(BinaryValue));
 		FCbFieldRef Field = Writer.Save();
-		if (TestEqual(TEXT("FCbWriter(Binary, Array) Validate"), ValidateCompactBinary(*Field.Buffer(), ECbValidateMode::All), ECbValidateError::None))
+		if (TestEqual(TEXT("FCbWriter(Binary, Array) Validate"), ValidateCompactBinary(*Field.GetBuffer(), ECbValidateMode::All), ECbValidateError::None))
 		{
-			TestEqual(TEXT("FCbWriter(Binary, Array).Name()"), Field.Name(), "Binary"_ASV);
+			TestEqual(TEXT("FCbWriter(Binary, Array).GetName()"), Field.GetName(), "Binary"_ASV);
 			TestTrue(TEXT("FCbWriter(Binary, Array).HasName()"), Field.HasName());
 			TestTrue(TEXT("FCbWriter(Binary, Array).IsBinary()"), Field.IsBinary());
 			TestTrue(TEXT("FCbWriter(Binary, Array).AsBinary()"), Field.AsBinary().EqualBytes(MakeMemoryView(BinaryValue)));
@@ -303,7 +303,7 @@ bool FCbWriterStringTest::RunTest(const FString& Parameters)
 		Writer.String(FAnsiStringView());
 		Writer.String(FWideStringView());
 		FCbFieldRefIterator Fields = Writer.Save();
-		if (TestEqual(TEXT("FCbWriter(String, Empty) Validate"), ValidateCompactBinaryRange(*Fields.Buffer(), ECbValidateMode::All), ECbValidateError::None))
+		if (TestEqual(TEXT("FCbWriter(String, Empty) Validate"), ValidateCompactBinaryRange(*Fields.GetBuffer(), ECbValidateMode::All), ECbValidateError::None))
 		{
 			for (FCbField Field : Fields)
 			{
@@ -320,11 +320,11 @@ bool FCbWriterStringTest::RunTest(const FString& Parameters)
 		Writer.Name("String"_ASV).String("Value"_ASV);
 		Writer.Name("String"_ASV).String(TEXT("Value"_SV));
 		FCbFieldRefIterator Fields = Writer.Save();
-		if (TestEqual(TEXT("FCbWriter(String, Basic) Validate"), ValidateCompactBinaryRange(*Fields.Buffer(), ECbValidateMode::All), ECbValidateError::None))
+		if (TestEqual(TEXT("FCbWriter(String, Basic) Validate"), ValidateCompactBinaryRange(*Fields.GetBuffer(), ECbValidateMode::All), ECbValidateError::None))
 		{
 			for (FCbField Field : Fields)
 			{
-				TestEqual(TEXT("FCbWriter(String, Basic).Name()"), Field.Name(), "String"_ASV);
+				TestEqual(TEXT("FCbWriter(String, Basic).GetName()"), Field.GetName(), "String"_ASV);
 				TestTrue(TEXT("FCbWriter(String, Basic).HasName()"), Field.HasName());
 				TestTrue(TEXT("FCbWriter(String, Basic).IsString()"), Field.IsString());
 				TestEqual(TEXT("FCbWriter(String, Basic).AsString()"), Field.AsString(), "Value"_ASV);
@@ -344,7 +344,7 @@ bool FCbWriterStringTest::RunTest(const FString& Parameters)
 		Writer.String(Dots);
 		Writer.String(FString::ChrN(DotCount, TEXT('.')));
 		FCbFieldRefIterator Fields = Writer.Save();
-		if (TestEqual(TEXT("FCbWriter(String, Long) Validate"), ValidateCompactBinaryRange(*Fields.Buffer(), ECbValidateMode::All), ECbValidateError::None))
+		if (TestEqual(TEXT("FCbWriter(String, Long) Validate"), ValidateCompactBinaryRange(*Fields.GetBuffer(), ECbValidateMode::All), ECbValidateError::None))
 		{
 			for (FCbField Field : Fields)
 			{
@@ -360,7 +360,7 @@ bool FCbWriterStringTest::RunTest(const FString& Parameters)
 		Writer.String("\xf0\x9f\x98\x80"_ASV);
 		Writer.String(FWideStringView(Value, UE_ARRAY_COUNT(Value)));
 		FCbFieldRefIterator Fields = Writer.Save();
-		if (TestEqual(TEXT("FCbWriter(String, Unicode) Validate"), ValidateCompactBinaryRange(*Fields.Buffer(), ECbValidateMode::All), ECbValidateError::None))
+		if (TestEqual(TEXT("FCbWriter(String, Unicode) Validate"), ValidateCompactBinaryRange(*Fields.GetBuffer(), ECbValidateMode::All), ECbValidateError::None))
 		{
 			for (FCbField Field : Fields)
 			{
@@ -384,7 +384,7 @@ bool FCbWriterIntegerTest::RunTest(const FString& Parameters)
 		Writer.Reset();
 		Writer.Integer(Value);
 		FCbFieldRef Field = Writer.Save();
-		if (TestEqual(TEXT("FCbWriter(Integer, Int32) Validate"), ValidateCompactBinary(*Field.Buffer(), ECbValidateMode::All), ECbValidateError::None))
+		if (TestEqual(TEXT("FCbWriter(Integer, Int32) Validate"), ValidateCompactBinary(*Field.GetBuffer(), ECbValidateMode::All), ECbValidateError::None))
 		{
 			TestEqual(TEXT("FCbWriter(Integer, Int32) Value"), Field.AsInt32(), Value);
 			TestFalse(TEXT("FCbWriter(Integer, Int32) Error"), Field.HasError());
@@ -396,7 +396,7 @@ bool FCbWriterIntegerTest::RunTest(const FString& Parameters)
 		Writer.Reset();
 		Writer.Integer(Value);
 		FCbFieldRef Field = Writer.Save();
-		if (TestEqual(TEXT("FCbWriter(Integer, UInt32) Validate"), ValidateCompactBinary(*Field.Buffer(), ECbValidateMode::All), ECbValidateError::None))
+		if (TestEqual(TEXT("FCbWriter(Integer, UInt32) Validate"), ValidateCompactBinary(*Field.GetBuffer(), ECbValidateMode::All), ECbValidateError::None))
 		{
 			TestEqual(TEXT("FCbWriter(Integer, UInt32) Value"), Field.AsUInt32(), Value);
 			TestFalse(TEXT("FCbWriter(Integer, UInt32) Error"), Field.HasError());
@@ -408,7 +408,7 @@ bool FCbWriterIntegerTest::RunTest(const FString& Parameters)
 		Writer.Reset();
 		Writer.Integer(Value);
 		FCbFieldRef Field = Writer.Save();
-		if (TestEqual(TEXT("FCbWriter(Integer, Int64) Validate"), ValidateCompactBinary(*Field.Buffer(), ECbValidateMode::All), ECbValidateError::None))
+		if (TestEqual(TEXT("FCbWriter(Integer, Int64) Validate"), ValidateCompactBinary(*Field.GetBuffer(), ECbValidateMode::All), ECbValidateError::None))
 		{
 			TestEqual(TEXT("FCbWriter(Integer, Int64) Value"), Field.AsInt64(), Value);
 			TestFalse(TEXT("FCbWriter(Integer, Int64) Error"), Field.HasError());
@@ -420,7 +420,7 @@ bool FCbWriterIntegerTest::RunTest(const FString& Parameters)
 		Writer.Reset();
 		Writer.Integer(Value);
 		FCbFieldRef Field = Writer.Save();
-		if (TestEqual(TEXT("FCbWriter(Integer, UInt64) Validate"), ValidateCompactBinary(*Field.Buffer(), ECbValidateMode::All), ECbValidateError::None))
+		if (TestEqual(TEXT("FCbWriter(Integer, UInt64) Validate"), ValidateCompactBinary(*Field.GetBuffer(), ECbValidateMode::All), ECbValidateError::None))
 		{
 			TestEqual(TEXT("FCbWriter(Integer, UInt64) Value"), Field.AsUInt64(), Value);
 			TestFalse(TEXT("FCbWriter(Integer, UInt64) Error"), Field.HasError());
@@ -479,7 +479,7 @@ bool FCbWriterFloatTest::RunTest(const FString& Parameters)
 			Writer.Float(Value);
 		}
 		FCbFieldRefIterator Fields = Writer.Save();
-		if (TestEqual(TEXT("FCbWriter(Float, Single) Validate"), ValidateCompactBinaryRange(*Fields.Buffer(), ECbValidateMode::All), ECbValidateError::None))
+		if (TestEqual(TEXT("FCbWriter(Float, Single) Validate"), ValidateCompactBinaryRange(*Fields.GetBuffer(), ECbValidateMode::All), ECbValidateError::None))
 		{
 			const float* CheckValue = Values;
 			for (FCbField Field : Fields)
@@ -499,7 +499,7 @@ bool FCbWriterFloatTest::RunTest(const FString& Parameters)
 			Writer.Float(Value);
 		}
 		FCbFieldRefIterator Fields = Writer.Save();
-		if (TestEqual(TEXT("FCbWriter(Float, Double) Validate"), ValidateCompactBinaryRange(*Fields.Buffer(), ECbValidateMode::All), ECbValidateError::None))
+		if (TestEqual(TEXT("FCbWriter(Float, Double) Validate"), ValidateCompactBinaryRange(*Fields.GetBuffer(), ECbValidateMode::All), ECbValidateError::None))
 		{
 			const double* CheckValue = Values;
 			for (FCbField Field : Fields)
@@ -526,7 +526,7 @@ bool FCbWriterBoolTest::RunTest(const FString& Parameters)
 		Writer.Bool(false);
 
 		FCbFieldRefIterator Fields = Writer.Save();
-		if (TestEqual(TEXT("FCbWriter(Bool) Validate"), ValidateCompactBinaryRange(*Fields.Buffer(), ECbValidateMode::All), ECbValidateError::None))
+		if (TestEqual(TEXT("FCbWriter(Bool) Validate"), ValidateCompactBinaryRange(*Fields.GetBuffer(), ECbValidateMode::All), ECbValidateError::None))
 		{
 			TestTrue(TEXT("FCbWriter(Bool).AsBool()"), Fields.AsBool());
 			TestFalse(TEXT("FCbWriter(Bool) Error"), Fields.HasError());
@@ -553,7 +553,7 @@ bool FCbWriterBoolTest::RunTest(const FString& Parameters)
 		Writer.EndObject();
 
 		FCbFieldRefIterator Fields = Writer.Save();
-		TestEqual(TEXT("FCbWriter(Bool, Uniform) Validate"), ValidateCompactBinaryRange(*Fields.Buffer(), ECbValidateMode::All), ECbValidateError::None);
+		TestEqual(TEXT("FCbWriter(Bool, Uniform) Validate"), ValidateCompactBinaryRange(*Fields.GetBuffer(), ECbValidateMode::All), ECbValidateError::None);
 	}
 
 	return true;
@@ -578,7 +578,7 @@ bool FCbWriterBinaryHashTest::RunTest(const FString& Parameters)
 	}
 
 	FCbFieldRefIterator Fields = Writer.Save();
-	if (TestEqual(TEXT("FCbWriter(BinaryHash) Validate"), ValidateCompactBinaryRange(*Fields.Buffer(), ECbValidateMode::All), ECbValidateError::None))
+	if (TestEqual(TEXT("FCbWriter(BinaryHash) Validate"), ValidateCompactBinaryRange(*Fields.GetBuffer(), ECbValidateMode::All), ECbValidateError::None))
 	{
 		const FBlake3Hash* CheckValue = Values;
 		for (FCbField Field : Fields)
@@ -608,7 +608,7 @@ bool FCbWriterFieldHashTest::RunTest(const FString& Parameters)
 	}
 
 	FCbFieldRefIterator Fields = Writer.Save();
-	if (TestEqual(TEXT("FCbWriter(FieldHash) Validate"), ValidateCompactBinaryRange(*Fields.Buffer(), ECbValidateMode::All), ECbValidateError::None))
+	if (TestEqual(TEXT("FCbWriter(FieldHash) Validate"), ValidateCompactBinaryRange(*Fields.GetBuffer(), ECbValidateMode::All), ECbValidateError::None))
 	{
 		const FBlake3Hash* CheckValue = Values;
 		for (FCbField Field : Fields)
@@ -635,7 +635,7 @@ bool FCbWriterUuidTest::RunTest(const FString& Parameters)
 	}
 
 	FCbFieldRefIterator Fields = Writer.Save();
-	if (TestEqual(TEXT("FCbWriter(Uuid) Validate"), ValidateCompactBinaryRange(*Fields.Buffer(), ECbValidateMode::All), ECbValidateError::None))
+	if (TestEqual(TEXT("FCbWriter(Uuid) Validate"), ValidateCompactBinaryRange(*Fields.GetBuffer(), ECbValidateMode::All), ECbValidateError::None))
 	{
 		const FGuid* CheckValue = Values;
 		for (FCbField Field : Fields)
@@ -662,7 +662,7 @@ bool FCbWriterDateTimeTest::RunTest(const FString& Parameters)
 	}
 
 	FCbFieldRefIterator Fields = Writer.Save();
-	if (TestEqual(TEXT("FCbWriter(DateTime) Validate"), ValidateCompactBinaryRange(*Fields.Buffer(), ECbValidateMode::All), ECbValidateError::None))
+	if (TestEqual(TEXT("FCbWriter(DateTime) Validate"), ValidateCompactBinaryRange(*Fields.GetBuffer(), ECbValidateMode::All), ECbValidateError::None))
 	{
 		const FDateTime* CheckValue = Values;
 		for (FCbField Field : Fields)
@@ -689,7 +689,7 @@ bool FCbWriterTimeSpanTest::RunTest(const FString& Parameters)
 	}
 
 	FCbFieldRefIterator Fields = Writer.Save();
-	if (TestEqual(TEXT("FCbWriter(TimeSpan) Validate"), ValidateCompactBinaryRange(*Fields.Buffer(), ECbValidateMode::All), ECbValidateError::None))
+	if (TestEqual(TEXT("FCbWriter(TimeSpan) Validate"), ValidateCompactBinaryRange(*Fields.GetBuffer(), ECbValidateMode::All), ECbValidateError::None))
 	{
 		const FTimespan* CheckValue = Values;
 		for (FCbField Field : Fields)
@@ -875,7 +875,7 @@ bool FCbWriterComplexTest::RunTest(const FString& Parameters)
 		Object = Writer.Save().AsObjectRef();
 	}
 
-	TestEqual(TEXT("FCbWriter(Complex) Validate"), ValidateCompactBinary(*Object.Buffer(), ECbValidateMode::All), ECbValidateError::None);
+	TestEqual(TEXT("FCbWriter(Complex) Validate"), ValidateCompactBinary(*Object.GetBuffer(), ECbValidateMode::All), ECbValidateError::None);
 
 	return true;
 }
@@ -945,7 +945,7 @@ bool FCbWriterStreamTest::RunTest(const FString& Parameters)
 		Object = Writer.Save().AsObjectRef();
 	}
 
-	TestEqual(TEXT("FCbWriter(Stream) Validate"), ValidateCompactBinary(*Object.Buffer(), ECbValidateMode::All), ECbValidateError::None);
+	TestEqual(TEXT("FCbWriter(Stream) Validate"), ValidateCompactBinary(*Object.GetBuffer(), ECbValidateMode::All), ECbValidateError::None);
 
 	return true;
 }
diff --git a/Engine/Source/Runtime/Core/Public/Serialization/CompactBinary.h b/Engine/Source/Runtime/Core/Public/Serialization/CompactBinary.h
--- a/Engine/Source/Runtime/Core/Public/Serialization/CompactBinary.h
+++ b/Engine/Source/Runtime/Core/Public/Serialization/CompactBinary.h
@@ -378,13 +378,13 @@ public:
 	CORE_API explicit FCbField(const void* Data, ECbFieldType Type = ECbFieldType::HasFieldType);
 
 	/** Name of the field if it has one, otherwise empty. */
-	constexpr inline FAnsiStringView Name() const
+	constexpr inline FAnsiStringView GetName() const
 	{
 		return FAnsiStringView(static_cast<const ANSICHAR*>(Payload) - NameLen, NameLen);
 	}
 
 	/** Size of the field in bytes. */
-	CORE_API uint64 Size() const;
+	CORE_API uint64 GetSize() const;
 
 	/**
 	 * Whether this field is identical to the other field.
@@ -506,7 +506,7 @@ protected:
 	constexpr inline ECbFieldType GetType() const { return Type; }
 	constexpr inline ECbFieldType GetCopyType() const { return Type; }
 	constexpr inline const void* GetPayload() const { return Payload; }
-	inline const void* GetPayloadEnd() const { return static_cast<const uint8*>(Payload) + PayloadSize(); }
+	inline const void* GetPayloadEnd() const { return static_cast<const uint8*>(Payload) + GetPayloadSize(); }
 
 	inline void Assign(const void* InData, const ECbFieldType InType)
 	{
@@ -516,7 +516,7 @@ protected:
 	}
 
 	/**
-	 * Copy the field into a buffer of at least Size() bytes.
+	 * Copy the field into a buffer of at least GetSize() bytes.
 	 *
 	 * The field type will only be copied if HasFieldType is set.
 	 */
@@ -527,7 +527,7 @@ protected:
 
 private:
 	/** Size of the field payload in bytes. That is the field excluding its type and name. */
-	CORE_API uint64 PayloadSize() const;
+	CORE_API uint64 GetPayloadSize() const;
 
 	/** Parameters for converting to an integer. */
 	struct FIntegerParams
@@ -612,7 +612,7 @@ public:
 	CORE_API explicit FCbArray(const void* Data, ECbFieldType Type = ECbFieldType::HasFieldType);
 
 	/** Size of the array in bytes if serialized by itself with no name. */
-	CORE_API uint64 Size() const;
+	CORE_API uint64 GetSize() const;
 
 	/** Number of items in the array. */
 	CORE_API uint64 Num() const;
@@ -634,7 +634,7 @@ public:
 protected:
 	inline ECbFieldType GetCopyType() const { return FCbFieldType::GetType(Type) | ECbFieldType::HasFieldType; }
 
-	/** Copy the array into a buffer of at least Size() bytes. */
+	/** Copy the array into a buffer of at least GetSize() bytes. */
 	CORE_API void CopyTo(FMutableMemoryView Buffer) const;
 
 	/** Create a view of the payload, excluding the type or the name of the outer field. */
@@ -677,7 +677,7 @@ public:
 	CORE_API explicit FCbObject(const void* Data, ECbFieldType Type = ECbFieldType::HasFieldType);
 
 	/** Size of the object in bytes if serialized by itself with no name. */
-	CORE_API uint64 Size() const;
+	CORE_API uint64 GetSize() const;
 
 	/** Create an iterator for the fields of this object. */
 	CORE_API FCbFieldIterator CreateIterator() const;
@@ -713,7 +713,7 @@ public:
 protected:
 	inline ECbFieldType GetCopyType() const { return FCbFieldType::GetType(Type) | ECbFieldType::HasFieldType; }
 
-	/** Copy the object into a buffer of at least Size() bytes. */
+	/** Copy the object into a buffer of at least GetSize() bytes. */
 	CORE_API void CopyTo(FMutableMemoryView Buffer) const;
 
 	/** Create a view of the payload, excluding the type or the name of the outer field. */
@@ -750,7 +750,7 @@ public:
 	{
 		if (ValueBuffer->IsOwned())
 		{
-			BufferPtr = ValueBuffer;
+			Buffer = ValueBuffer;
 		}
 	}
 
@@ -762,7 +762,7 @@ public:
 			static_cast<BaseType&>(*this) = BaseType(ValueBuffer->GetData());
 			if (ValueBuffer->IsOwned())
 			{
-				BufferPtr = ValueBuffer;
+				Buffer = ValueBuffer;
 			}
 		}
 	}
@@ -775,7 +775,7 @@ public:
 			static_cast<BaseType&>(*this) = BaseType(ValueBuffer->GetData());
 			if (ValueBuffer->IsOwned())
 			{
-				BufferPtr = MoveTemp(ValueBuffer);
+				Buffer = MoveTemp(ValueBuffer);
 			}
 		}
 	}
@@ -795,7 +795,7 @@ public:
 			check(ValueBuffer->GetView().Contains(BaseType::AsMemoryView()));
 			if (ValueBuffer->IsOwned())
 			{
-				BufferPtr = ValueBuffer;
+				Buffer = ValueBuffer;
 			}
 		}
 	}
@@ -809,7 +809,7 @@ public:
 			check(ValueBuffer->GetView().Contains(BaseType::AsMemoryView()));
 			if (ValueBuffer->IsOwned())
 			{
-				BufferPtr = MoveTemp(ValueBuffer);
+				Buffer = MoveTemp(ValueBuffer);
 			}
 		}
 	}
@@ -817,19 +817,19 @@ public:
 	/** Construct a value that holds a reference to the buffer of the outer that contains it. */
 	template <typename OtherBaseType>
 	inline TCbBufferRef(const BaseType& Value, const TCbBufferRef<OtherBaseType>& OuterRef)
-		: TCbBufferRef(Value, OuterRef.BufferPtr)
+		: TCbBufferRef(Value, OuterRef.Buffer)
 	{
 	}
 
 	/** Construct a value that holds a reference to the buffer of the outer that contains it. */
 	template <typename OtherBaseType>
 	inline TCbBufferRef(const BaseType& Value, TCbBufferRef<OtherBaseType>&& OuterRef)
-		: TCbBufferRef(Value, MoveTemp(OuterRef.BufferPtr))
+		: TCbBufferRef(Value, MoveTemp(OuterRef.Buffer))
 	{
 	}
 
 	/** Whether this reference has ownership of the memory in its buffer. */
-	inline bool IsOwned() const { return BufferPtr && BufferPtr->IsOwned(); }
+	inline bool IsOwned() const { return Buffer && Buffer->IsOwned(); }
 
 	/** Clone the value, if necessary, to a buffer that this reference has ownership of. */
 	inline void MakeOwned()
@@ -841,28 +841,28 @@ public:
 	}
 
 	/** Whether this reference is backed by a read-only buffer. */
-	inline bool IsReadOnly() const { return BufferPtr && BufferPtr->IsReadOnly(); }
+	inline bool IsReadOnly() const { return Buffer && Buffer->IsReadOnly(); }
 
 	/** Clone the value, if necessary, to a read-only buffer. */
 	inline void MakeReadOnly()
 	{
 		if (!IsReadOnly())
 		{
-			FSharedBufferPtr Buffer = FSharedBuffer::Alloc(BaseType::Size());
-			BaseType::CopyTo(*Buffer);
-			static_cast<BaseType&>(*this) = BaseType(Buffer->GetData(), BaseType::GetCopyType());
-			BufferPtr = FSharedBuffer::MakeReadOnly(MoveTemp(Buffer));
+			FSharedBufferPtr MutableBuffer = FSharedBuffer::Alloc(BaseType::GetSize());
+			BaseType::CopyTo(*MutableBuffer);
+			static_cast<BaseType&>(*this) = BaseType(MutableBuffer->GetData(), BaseType::GetCopyType());
+			Buffer = FSharedBuffer::MakeReadOnly(MoveTemp(MutableBuffer));
 		}
 	}
 
 	/** The buffer (if any) that contains this value. */
-	inline const FSharedBufferConstPtr& Buffer() const { return BufferPtr; }
+	inline const FSharedBufferConstPtr& GetBuffer() const { return Buffer; }
 
 private:
 	template <typename OtherType>
 	friend class TCbBufferRef;
 
-	FSharedBufferConstPtr BufferPtr;
+	FSharedBufferConstPtr Buffer;
 };
 
 ///////////////////////////////////////////////////////////////////////////////////////////////////////////////////////
@@ -908,7 +908,7 @@ public:
 	static inline RefType TakeOwnership(const void* const Data, DeleteFunctionOrBufferOwnerType&& DeleteFunctionOrBufferOwner)
 	{
 		const BaseType Value(Data);
-		FSharedBufferConstPtr Buffer = FSharedBuffer::TakeOwnership(Data, Value.Size(),
+		FSharedBufferConstPtr Buffer = FSharedBuffer::TakeOwnership(Data, Value.GetSize(),
 			Forward<DeleteFunctionOrBufferOwnerType>(DeleteFunctionOrBufferOwner));
 		return RefType(Value, FSharedBuffer::MakeReadOnly(MoveTemp(Buffer)));
 	}
@@ -953,7 +953,7 @@ public:
 	inline FCbFieldRefIterator CreateRefIterator() const
 	{
 		FCbFieldIterator It = CreateIterator();
-		return FCbFieldRefIterator(FCbFieldRef(It, Buffer()), It);
+		return FCbFieldRefIterator(FCbFieldRef(It, GetBuffer()), It);
 	}
 };
 
@@ -969,7 +969,7 @@ public:
 	inline FCbFieldRefIterator CreateRefIterator() const
 	{
 		FCbFieldIterator It = CreateIterator();
-		return FCbFieldRefIterator(FCbFieldRef(It, Buffer()), It);
+		return FCbFieldRefIterator(FCbFieldRef(It, GetBuffer()), It);
 	}
 
 	/** Find a field by case-sensitive name comparison. */
