SharedBuffer: Added equality operators for shared buffers and fixed pinning of const weak pointers

#rnx

diff --git a/Engine/Source/Runtime/Core/Private/Tests/Memory/SharedBufferTest.cpp b/Engine/Source/Runtime/Core/Private/Tests/Memory/SharedBufferTest.cpp
--- a/Engine/Source/Runtime/Core/Private/Tests/Memory/SharedBufferTest.cpp
+++ b/Engine/Source/Runtime/Core/Private/Tests/Memory/SharedBufferTest.cpp
@@ -162,6 +162,94 @@ static_assert(std::is_same<FSharedBufferConstPtr, decltype(FSharedBuffer::MakeOw
 
 ///////////////////////////////////////////////////////////////////////////////////////////////////////////////////////
 
+static_assert(std::is_same<bool, decltype(DeclVal<const FSharedBufferRef&>() == DeclVal<const FSharedBufferRef&>())>::value, "Invalid equality");
+static_assert(std::is_same<bool, decltype(DeclVal<const FSharedBufferRef&>() == DeclVal<const FSharedBufferPtr&>())>::value, "Invalid equality");
+static_assert(std::is_same<bool, decltype(DeclVal<const FSharedBufferRef&>() == DeclVal<const FSharedBufferWeakPtr&>())>::value, "Invalid equality");
+static_assert(std::is_same<bool, decltype(DeclVal<const FSharedBufferRef&>() == DeclVal<const FSharedBufferConstRef&>())>::value, "Invalid equality");
+static_assert(std::is_same<bool, decltype(DeclVal<const FSharedBufferRef&>() == DeclVal<const FSharedBufferConstPtr&>())>::value, "Invalid equality");
+static_assert(std::is_same<bool, decltype(DeclVal<const FSharedBufferRef&>() == DeclVal<const FSharedBufferConstWeakPtr&>())>::value, "Invalid equality");
+
+static_assert(std::is_same<bool, decltype(DeclVal<const FSharedBufferPtr&>() == DeclVal<const FSharedBufferRef&>())>::value, "Invalid equality");
+static_assert(std::is_same<bool, decltype(DeclVal<const FSharedBufferPtr&>() == DeclVal<const FSharedBufferPtr&>())>::value, "Invalid equality");
+static_assert(std::is_same<bool, decltype(DeclVal<const FSharedBufferPtr&>() == DeclVal<const FSharedBufferWeakPtr&>())>::value, "Invalid equality");
+static_assert(std::is_same<bool, decltype(DeclVal<const FSharedBufferPtr&>() == DeclVal<const FSharedBufferConstRef&>())>::value, "Invalid equality");
+static_assert(std::is_same<bool, decltype(DeclVal<const FSharedBufferPtr&>() == DeclVal<const FSharedBufferConstPtr&>())>::value, "Invalid equality");
+static_assert(std::is_same<bool, decltype(DeclVal<const FSharedBufferPtr&>() == DeclVal<const FSharedBufferConstWeakPtr&>())>::value, "Invalid equality");
+
+static_assert(std::is_same<bool, decltype(DeclVal<const FSharedBufferWeakPtr&>() == DeclVal<const FSharedBufferRef&>())>::value, "Invalid equality");
+static_assert(std::is_same<bool, decltype(DeclVal<const FSharedBufferWeakPtr&>() == DeclVal<const FSharedBufferPtr&>())>::value, "Invalid equality");
+static_assert(std::is_same<bool, decltype(DeclVal<const FSharedBufferWeakPtr&>() == DeclVal<const FSharedBufferWeakPtr&>())>::value, "Invalid equality");
+static_assert(std::is_same<bool, decltype(DeclVal<const FSharedBufferWeakPtr&>() == DeclVal<const FSharedBufferConstRef&>())>::value, "Invalid equality");
+static_assert(std::is_same<bool, decltype(DeclVal<const FSharedBufferWeakPtr&>() == DeclVal<const FSharedBufferConstPtr&>())>::value, "Invalid equality");
+static_assert(std::is_same<bool, decltype(DeclVal<const FSharedBufferWeakPtr&>() == DeclVal<const FSharedBufferConstWeakPtr&>())>::value, "Invalid equality");
+
+static_assert(std::is_same<bool, decltype(DeclVal<const FSharedBufferConstRef&>() == DeclVal<const FSharedBufferRef&>())>::value, "Invalid equality");
+static_assert(std::is_same<bool, decltype(DeclVal<const FSharedBufferConstRef&>() == DeclVal<const FSharedBufferPtr&>())>::value, "Invalid equality");
+static_assert(std::is_same<bool, decltype(DeclVal<const FSharedBufferConstRef&>() == DeclVal<const FSharedBufferWeakPtr&>())>::value, "Invalid equality");
+static_assert(std::is_same<bool, decltype(DeclVal<const FSharedBufferConstRef&>() == DeclVal<const FSharedBufferConstRef&>())>::value, "Invalid equality");
+static_assert(std::is_same<bool, decltype(DeclVal<const FSharedBufferConstRef&>() == DeclVal<const FSharedBufferConstPtr&>())>::value, "Invalid equality");
+static_assert(std::is_same<bool, decltype(DeclVal<const FSharedBufferConstRef&>() == DeclVal<const FSharedBufferConstWeakPtr&>())>::value, "Invalid equality");
+
+static_assert(std::is_same<bool, decltype(DeclVal<const FSharedBufferConstPtr&>() == DeclVal<const FSharedBufferRef&>())>::value, "Invalid equality");
+static_assert(std::is_same<bool, decltype(DeclVal<const FSharedBufferConstPtr&>() == DeclVal<const FSharedBufferPtr&>())>::value, "Invalid equality");
+static_assert(std::is_same<bool, decltype(DeclVal<const FSharedBufferConstPtr&>() == DeclVal<const FSharedBufferWeakPtr&>())>::value, "Invalid equality");
+static_assert(std::is_same<bool, decltype(DeclVal<const FSharedBufferConstPtr&>() == DeclVal<const FSharedBufferConstRef&>())>::value, "Invalid equality");
+static_assert(std::is_same<bool, decltype(DeclVal<const FSharedBufferConstPtr&>() == DeclVal<const FSharedBufferConstPtr&>())>::value, "Invalid equality");
+static_assert(std::is_same<bool, decltype(DeclVal<const FSharedBufferConstPtr&>() == DeclVal<const FSharedBufferConstWeakPtr&>())>::value, "Invalid equality");
+
+static_assert(std::is_same<bool, decltype(DeclVal<const FSharedBufferConstWeakPtr&>() == DeclVal<const FSharedBufferRef&>())>::value, "Invalid equality");
+static_assert(std::is_same<bool, decltype(DeclVal<const FSharedBufferConstWeakPtr&>() == DeclVal<const FSharedBufferPtr&>())>::value, "Invalid equality");
+static_assert(std::is_same<bool, decltype(DeclVal<const FSharedBufferConstWeakPtr&>() == DeclVal<const FSharedBufferWeakPtr&>())>::value, "Invalid equality");
+static_assert(std::is_same<bool, decltype(DeclVal<const FSharedBufferConstWeakPtr&>() == DeclVal<const FSharedBufferConstRef&>())>::value, "Invalid equality");
+static_assert(std::is_same<bool, decltype(DeclVal<const FSharedBufferConstWeakPtr&>() == DeclVal<const FSharedBufferConstPtr&>())>::value, "Invalid equality");
+static_assert(std::is_same<bool, decltype(DeclVal<const FSharedBufferConstWeakPtr&>() == DeclVal<const FSharedBufferConstWeakPtr&>())>::value, "Invalid equality");
+
+///////////////////////////////////////////////////////////////////////////////////////////////////////////////////////
+
+static_assert(std::is_same<bool, decltype(DeclVal<const FSharedBufferRef&>() != DeclVal<const FSharedBufferRef&>())>::value, "Invalid inequality");
+static_assert(std::is_same<bool, decltype(DeclVal<const FSharedBufferRef&>() != DeclVal<const FSharedBufferPtr&>())>::value, "Invalid inequality");
+static_assert(std::is_same<bool, decltype(DeclVal<const FSharedBufferRef&>() != DeclVal<const FSharedBufferWeakPtr&>())>::value, "Invalid inequality");
+static_assert(std::is_same<bool, decltype(DeclVal<const FSharedBufferRef&>() != DeclVal<const FSharedBufferConstRef&>())>::value, "Invalid inequality");
+static_assert(std::is_same<bool, decltype(DeclVal<const FSharedBufferRef&>() != DeclVal<const FSharedBufferConstPtr&>())>::value, "Invalid inequality");
+static_assert(std::is_same<bool, decltype(DeclVal<const FSharedBufferRef&>() != DeclVal<const FSharedBufferConstWeakPtr&>())>::value, "Invalid inequality");
+
+static_assert(std::is_same<bool, decltype(DeclVal<const FSharedBufferPtr&>() != DeclVal<const FSharedBufferRef&>())>::value, "Invalid inequality");
+static_assert(std::is_same<bool, decltype(DeclVal<const FSharedBufferPtr&>() != DeclVal<const FSharedBufferPtr&>())>::value, "Invalid inequality");
+static_assert(std::is_same<bool, decltype(DeclVal<const FSharedBufferPtr&>() != DeclVal<const FSharedBufferWeakPtr&>())>::value, "Invalid inequality");
+static_assert(std::is_same<bool, decltype(DeclVal<const FSharedBufferPtr&>() != DeclVal<const FSharedBufferConstRef&>())>::value, "Invalid inequality");
+static_assert(std::is_same<bool, decltype(DeclVal<const FSharedBufferPtr&>() != DeclVal<const FSharedBufferConstPtr&>())>::value, "Invalid inequality");
+static_assert(std::is_same<bool, decltype(DeclVal<const FSharedBufferPtr&>() != DeclVal<const FSharedBufferConstWeakPtr&>())>::value, "Invalid inequality");
+
+static_assert(std::is_same<bool, decltype(DeclVal<const FSharedBufferWeakPtr&>() != DeclVal<const FSharedBufferRef&>())>::value, "Invalid inequality");
+static_assert(std::is_same<bool, decltype(DeclVal<const FSharedBufferWeakPtr&>() != DeclVal<const FSharedBufferPtr&>())>::value, "Invalid inequality");
+static_assert(std::is_same<bool, decltype(DeclVal<const FSharedBufferWeakPtr&>() != DeclVal<const FSharedBufferWeakPtr&>())>::value, "Invalid inequality");
+static_assert(std::is_same<bool, decltype(DeclVal<const FSharedBufferWeakPtr&>() != DeclVal<const FSharedBufferConstRef&>())>::value, "Invalid inequality");
+static_assert(std::is_same<bool, decltype(DeclVal<const FSharedBufferWeakPtr&>() != DeclVal<const FSharedBufferConstPtr&>())>::value, "Invalid inequality");
+static_assert(std::is_same<bool, decltype(DeclVal<const FSharedBufferWeakPtr&>() != DeclVal<const FSharedBufferConstWeakPtr&>())>::value, "Invalid inequality");
+
+static_assert(std::is_same<bool, decltype(DeclVal<const FSharedBufferConstRef&>() != DeclVal<const FSharedBufferRef&>())>::value, "Invalid inequality");
+static_assert(std::is_same<bool, decltype(DeclVal<const FSharedBufferConstRef&>() != DeclVal<const FSharedBufferPtr&>())>::value, "Invalid inequality");
+static_assert(std::is_same<bool, decltype(DeclVal<const FSharedBufferConstRef&>() != DeclVal<const FSharedBufferWeakPtr&>())>::value, "Invalid inequality");
+static_assert(std::is_same<bool, decltype(DeclVal<const FSharedBufferConstRef&>() != DeclVal<const FSharedBufferConstRef&>())>::value, "Invalid inequality");
+static_assert(std::is_same<bool, decltype(DeclVal<const FSharedBufferConstRef&>() != DeclVal<const FSharedBufferConstPtr&>())>::value, "Invalid inequality");
+static_assert(std::is_same<bool, decltype(DeclVal<const FSharedBufferConstRef&>() != DeclVal<const FSharedBufferConstWeakPtr&>())>::value, "Invalid inequality");
+
+static_assert(std::is_same<bool, decltype(DeclVal<const FSharedBufferConstPtr&>() != DeclVal<const FSharedBufferRef&>())>::value, "Invalid inequality");
+static_assert(std::is_same<bool, decltype(DeclVal<const FSharedBufferConstPtr&>() != DeclVal<const FSharedBufferPtr&>())>::value, "Invalid inequality");
+static_assert(std::is_same<bool, decltype(DeclVal<const FSharedBufferConstPtr&>() != DeclVal<const FSharedBufferWeakPtr&>())>::value, "Invalid inequality");
+static_assert(std::is_same<bool, decltype(DeclVal<const FSharedBufferConstPtr&>() != DeclVal<const FSharedBufferConstRef&>())>::value, "Invalid inequality");
+static_assert(std::is_same<bool, decltype(DeclVal<const FSharedBufferConstPtr&>() != DeclVal<const FSharedBufferConstPtr&>())>::value, "Invalid inequality");
+static_assert(std::is_same<bool, decltype(DeclVal<const FSharedBufferConstPtr&>() != DeclVal<const FSharedBufferConstWeakPtr&>())>::value, "Invalid inequality");
+
+static_assert(std::is_same<bool, decltype(DeclVal<const FSharedBufferConstWeakPtr&>() != DeclVal<const FSharedBufferRef&>())>::value, "Invalid inequality");
+static_assert(std::is_same<bool, decltype(DeclVal<const FSharedBufferConstWeakPtr&>() != DeclVal<const FSharedBufferPtr&>())>::value, "Invalid inequality");
+static_assert(std::is_same<bool, decltype(DeclVal<const FSharedBufferConstWeakPtr&>() != DeclVal<const FSharedBufferWeakPtr&>())>::value, "Invalid inequality");
+static_assert(std::is_same<bool, decltype(DeclVal<const FSharedBufferConstWeakPtr&>() != DeclVal<const FSharedBufferConstRef&>())>::value, "Invalid inequality");
+static_assert(std::is_same<bool, decltype(DeclVal<const FSharedBufferConstWeakPtr&>() != DeclVal<const FSharedBufferConstPtr&>())>::value, "Invalid inequality");
+static_assert(std::is_same<bool, decltype(DeclVal<const FSharedBufferConstWeakPtr&>() != DeclVal<const FSharedBufferConstWeakPtr&>())>::value, "Invalid inequality");
+
+///////////////////////////////////////////////////////////////////////////////////////////////////////////////////////
+
 static_assert(std::is_same<uint32, decltype(GetTypeHash(DeclVal<const FSharedBufferRef>()))>::value, "Missing or invalid hash function");
 static_assert(std::is_same<uint32, decltype(GetTypeHash(DeclVal<const FSharedBufferConstRef>()))>::value, "Missing or invalid hash function");
 static_assert(std::is_same<uint32, decltype(GetTypeHash(DeclVal<const FSharedBufferPtr>()))>::value, "Missing or invalid hash function");
diff --git a/Engine/Source/Runtime/Core/Public/Memory/SharedBuffer.h b/Engine/Source/Runtime/Core/Public/Memory/SharedBuffer.h
--- a/Engine/Source/Runtime/Core/Public/Memory/SharedBuffer.h
+++ b/Engine/Source/Runtime/Core/Public/Memory/SharedBuffer.h
@@ -64,9 +64,9 @@ class TSharedBufferPtr final
 	friend class TSharedBufferPtr;
 
 	template <bool bOtherAllowNull, bool bOtherIsConst, bool bOtherIsWeak>
-	static inline BufferType* CopyFrom(const TSharedBufferPtr<bOtherAllowNull, bOtherIsConst, bOtherIsWeak>& Ref)
+	static inline BufferType* CopyFrom(const TSharedBufferPtr<bOtherAllowNull, bOtherIsConst, bOtherIsWeak>& Ptr)
 	{
-		BufferType* NewBuffer = Ref.Buffer;
+		BufferType* NewBuffer = Ptr.Buffer;
 		if (bOtherIsWeak && !bIsWeak)
 		{
 			if (!Ops::TryAddRef(NewBuffer))
@@ -82,14 +82,14 @@ class TSharedBufferPtr final
 	}
 
 	template <bool bOtherAllowNull, bool bOtherIsConst, bool bOtherIsWeak>
-	static inline BufferType* MoveFrom(TSharedBufferPtr<bOtherAllowNull, bOtherIsConst, bOtherIsWeak>&& Ref)
+	static inline BufferType* MoveFrom(TSharedBufferPtr<bOtherAllowNull, bOtherIsConst, bOtherIsWeak>&& Ptr)
 	{
-		BufferType* NewBuffer = Ref.Buffer;
+		BufferType* NewBuffer = Ptr.Buffer;
 		if (bOtherIsWeak)
 		{
 			if (bIsWeak)
 			{
-				Ref.Buffer = nullptr;
+				Ptr.Buffer = nullptr;
 			}
 			else
 			{
@@ -109,7 +109,7 @@ class TSharedBufferPtr final
 			{
 				if (bOtherAllowNull)
 				{
-					Ref.Buffer = nullptr;
+					Ptr.Buffer = nullptr;
 				}
 				else
 				{
@@ -127,25 +127,25 @@ public:
 		Ops::AddRef(InBuffer);
 	}
 
-	inline TSharedBufferPtr(const TSharedBufferPtr& Ref)
-		: Buffer(CopyFrom(Ref))
+	inline TSharedBufferPtr(const TSharedBufferPtr& Ptr)
+		: Buffer(CopyFrom(Ptr))
 	{
 	}
 
-	inline TSharedBufferPtr(TSharedBufferPtr&& Ref)
-		: Buffer(MoveFrom(MoveTemp(Ref)))
+	inline TSharedBufferPtr(TSharedBufferPtr&& Ptr)
+		: Buffer(MoveFrom(MoveTemp(Ptr)))
 	{
 	}
 
 	template <bool bOtherAllowNull, bool bOtherIsConst, bool bOtherIsWeak>
-	inline explicit TSharedBufferPtr(const TSharedBufferPtr<bOtherAllowNull, bOtherIsConst, bOtherIsWeak>& Ref)
-		: Buffer(CopyFrom(Ref))
+	inline explicit TSharedBufferPtr(const TSharedBufferPtr<bOtherAllowNull, bOtherIsConst, bOtherIsWeak>& Ptr)
+		: Buffer(CopyFrom(Ptr))
 	{
 	}
 
 	template <bool bOtherAllowNull, bool bOtherIsConst, bool bOtherIsWeak>
-	inline explicit TSharedBufferPtr(TSharedBufferPtr<bOtherAllowNull, bOtherIsConst, bOtherIsWeak>&& Ref)
-		: Buffer(MoveFrom(MoveTemp(Ref)))
+	inline explicit TSharedBufferPtr(TSharedBufferPtr<bOtherAllowNull, bOtherIsConst, bOtherIsWeak>&& Ptr)
+		: Buffer(MoveFrom(MoveTemp(Ptr)))
 	{
 	}
 
@@ -154,42 +154,54 @@ public:
 		Ops::Release(Buffer);
 	}
 
-	inline TSharedBufferPtr& operator=(const TSharedBufferPtr& Ref)
+	inline TSharedBufferPtr& operator=(const TSharedBufferPtr& Ptr)
 	{
 		BufferType* const OldBuffer = Buffer;
-		Buffer = CopyFrom(Ref);
+		Buffer = CopyFrom(Ptr);
 		Ops::Release(OldBuffer);
 		return *this;
 	}
 
-	inline TSharedBufferPtr& operator=(TSharedBufferPtr&& Ref)
+	inline TSharedBufferPtr& operator=(TSharedBufferPtr&& Ptr)
 	{
 		BufferType* const OldBuffer = Buffer;
-		Buffer = MoveFrom(MoveTemp(Ref));
+		Buffer = MoveFrom(MoveTemp(Ptr));
 		Ops::Release(OldBuffer);
 		return *this;
 	}
 
 	template <bool bOtherAllowNull, bool bOtherIsConst, bool bOtherIsWeak>
-	inline TSharedBufferPtr& operator=(const TSharedBufferPtr<bOtherAllowNull, bOtherIsConst, bOtherIsWeak>& Ref)
+	inline TSharedBufferPtr& operator=(const TSharedBufferPtr<bOtherAllowNull, bOtherIsConst, bOtherIsWeak>& Ptr)
 	{
 		BufferType* const OldBuffer = Buffer;
-		Buffer = CopyFrom(Ref);
+		Buffer = CopyFrom(Ptr);
 		Ops::Release(OldBuffer);
 		return *this;
 	}
 
 	template <bool bOtherAllowNull, bool bOtherIsConst, bool bOtherIsWeak>
-	inline TSharedBufferPtr& operator=(TSharedBufferPtr<bOtherAllowNull, bOtherIsConst, bOtherIsWeak>&& Ref)
+	inline TSharedBufferPtr& operator=(TSharedBufferPtr<bOtherAllowNull, bOtherIsConst, bOtherIsWeak>&& Ptr)
 	{
 		BufferType* const OldBuffer = Buffer;
-		Buffer = MoveFrom(MoveTemp(Ref));
+		Buffer = MoveFrom(MoveTemp(Ptr));
 		Ops::Release(OldBuffer);
 		return *this;
 	}
 
 	BufferType* Get() const { return Buffer; }
 
+	template <bool bOtherAllowNull, bool bOtherIsConst, bool bOtherIsWeak>
+	inline bool operator==(const TSharedBufferPtr<bOtherAllowNull, bOtherIsConst, bOtherIsWeak>& Ptr) const
+	{
+		return Buffer == Ptr.Buffer;
+	}
+
+	template <bool bOtherAllowNull, bool bOtherIsConst, bool bOtherIsWeak>
+	inline bool operator!=(const TSharedBufferPtr<bOtherAllowNull, bOtherIsConst, bOtherIsWeak>& Ptr) const
+	{
+		return Buffer != Ptr.Buffer;
+	}
+
 private:
 	BufferType* Buffer = nullptr;
 };
@@ -216,6 +228,7 @@ public:
 private:
 	constexpr inline friend bool IsValid(const FSharedBufferRef&) { return true; }
 	inline friend FSharedBufferRef ToSharedRef(const FSharedBufferRef& Ref) { return Ref; }
+	inline friend const PtrType& ToPrivateSharedBufferPtr(const FSharedBufferRef& Ref) { return Ref.Ptr; }
 
 	friend class FSharedBufferConstRef;
 	friend class FSharedBufferPtr;
@@ -246,6 +259,7 @@ public:
 private:
 	constexpr inline friend bool IsValid(const FSharedBufferConstRef&) { return true; }
 	inline friend FSharedBufferConstRef ToSharedRef(const FSharedBufferConstRef& Ref) { return Ref; }
+	inline friend const PtrType& ToPrivateSharedBufferPtr(const FSharedBufferConstRef& Ref) { return Ref.Ptr; }
 
 	friend class FSharedBufferConstPtr;
 	friend class FSharedBufferConstWeakPtr;
@@ -284,6 +298,7 @@ private:
 	inline friend bool IsValid(const FSharedBufferPtr& InPtr) { return InPtr.IsValid(); }
 	inline friend FSharedBufferRef ToSharedRef(const FSharedBufferPtr& InPtr) { return InPtr.ToSharedRef(); }
 	inline friend FSharedBufferRef ToSharedRef(FSharedBufferPtr&& InPtr) { return MoveTemp(InPtr).ToSharedRef(); }
+	inline friend const PtrType& ToPrivateSharedBufferPtr(const FSharedBufferPtr& InPtr) { return InPtr.Ptr; }
 
 	friend class FSharedBufferConstPtr;
 	friend class FSharedBufferWeakPtr;
@@ -324,6 +339,7 @@ private:
 	inline friend bool IsValid(const FSharedBufferConstPtr& InPtr) { return InPtr.IsValid(); }
 	inline friend FSharedBufferConstRef ToSharedRef(const FSharedBufferConstPtr& InPtr) { return InPtr.ToSharedRef(); }
 	inline friend FSharedBufferConstRef ToSharedRef(FSharedBufferConstPtr&& InPtr) { return MoveTemp(InPtr).ToSharedRef(); }
+	inline friend const PtrType& ToPrivateSharedBufferPtr(const FSharedBufferConstPtr& InPtr) { return InPtr.Ptr; }
 
 	friend class FSharedBufferConstWeakPtr;
 
@@ -344,7 +360,7 @@ public:
 	inline FSharedBufferWeakPtr(const FSharedBufferRef& Ref) : Ptr(Ref.Ptr) {}
 	inline FSharedBufferWeakPtr(const FSharedBufferPtr& InPtr) : Ptr(InPtr.Ptr) {}
 
-	inline FSharedBufferPtr Pin() { return FSharedBufferPtr(FSharedBufferPtr::PtrType(Ptr)); }
+	inline FSharedBufferPtr Pin() const { return FSharedBufferPtr(FSharedBufferPtr::PtrType(Ptr)); }
 
 	inline void Reset() { *this = FSharedBufferWeakPtr(); }
 
@@ -352,6 +368,7 @@ public:
 
 private:
 	friend class FSharedBufferConstWeakPtr;
+	inline friend const PtrType& ToPrivateSharedBufferPtr(const FSharedBufferWeakPtr& InPtr) { return InPtr.Ptr; }
 
 	PtrType Ptr;
 };
@@ -372,18 +389,36 @@ public:
 	inline FSharedBufferConstWeakPtr(const FSharedBufferWeakPtr& InPtr) : Ptr(InPtr.Ptr) {}
 	inline FSharedBufferConstWeakPtr(FSharedBufferWeakPtr&& InPtr) : Ptr(MoveTemp(InPtr.Ptr)) {}
 
-	inline FSharedBufferConstPtr Pin() { return FSharedBufferConstPtr(FSharedBufferConstPtr::PtrType(Ptr)); }
+	inline FSharedBufferConstPtr Pin() const { return FSharedBufferConstPtr(FSharedBufferConstPtr::PtrType(Ptr)); }
 
 	inline void Reset() { *this = FSharedBufferConstWeakPtr(); }
 
 	inline friend uint32 GetTypeHash(const FSharedBufferConstWeakPtr& InPtr) { return PointerHash(InPtr.Ptr.Get()); }
 
 private:
+	inline friend const PtrType& ToPrivateSharedBufferPtr(const FSharedBufferConstWeakPtr& InPtr) { return InPtr.Ptr; }
+
 	PtrType Ptr;
 };
 
 ///////////////////////////////////////////////////////////////////////////////////////////////////////////////////////
 
+template <typename TypeA, typename TypeB>
+inline auto operator==(const TypeA& BufferA, const TypeB& BufferB)
+	-> decltype(ToPrivateSharedBufferPtr(BufferA) == ToPrivateSharedBufferPtr(BufferB))
+{
+	return ToPrivateSharedBufferPtr(BufferA) == ToPrivateSharedBufferPtr(BufferB);
+}
+
+template <typename TypeA, typename TypeB>
+inline auto operator!=(const TypeA& BufferA, const TypeB& BufferB)
+	-> decltype(ToPrivateSharedBufferPtr(BufferA) != ToPrivateSharedBufferPtr(BufferB))
+{
+	return ToPrivateSharedBufferPtr(BufferA) != ToPrivateSharedBufferPtr(BufferB);
+}
+
+///////////////////////////////////////////////////////////////////////////////////////////////////////////////////////
+
 template<> struct TIsZeroConstructType<FSharedBufferRef> { enum { Value = true }; };
 template<> struct TIsZeroConstructType<FSharedBufferConstRef> { enum { Value = true }; };
 template<> struct TIsZeroConstructType<FSharedBufferPtr> { enum { Value = true }; };
