CompactBinary: Allowed objects and arrays to convert back to fields, and fields to field iterators

#rnx

diff --git a/Engine/Source/Runtime/Core/Private/Tests/Serialization/CompactBinaryTest.cpp b/Engine/Source/Runtime/Core/Private/Tests/Serialization/CompactBinaryTest.cpp
--- a/Engine/Source/Runtime/Core/Private/Tests/Serialization/CompactBinaryTest.cpp
+++ b/Engine/Source/Runtime/Core/Private/Tests/Serialization/CompactBinaryTest.cpp
@@ -348,7 +348,16 @@ bool FCbFieldObjectTest::RunTest(const FString& Parameters)
 		TestEqual(TEXT("FCbObject{Uniform}::FindIgnoreCase()"), Object.FindRefIgnoreCase("b"_ASV).AsInt32(), 2);
 		TestEqual(TEXT("FCbObject{Uniform}::operator[]"), Object["B"_ASV].AsInt32(), 2);
 		TestEqual(TEXT("FCbObject{Uniform}::operator[]"), Object["b"_ASV].AsInt32(4), 4);
+
+		TestTrue(TEXT("FCbObjectRef::AsFieldRef()"), Object.GetBuffer() == Object.AsFieldRef().AsObjectRef().GetBuffer());
+
 		Object = FCbFieldRef::MakeView(Field).AsObjectRef();
+
+		const uint8 NamedPayload[] = { 1, 'O', 10, IntType, 1, 'A', 1, 1, 'B', 2, 1, 'C', 3 };
+		FCbField NamedField(NamedPayload, ECbFieldType::UniformObject | ECbFieldType::HasFieldName);
+		TestTrue(TEXT("FCbObject::Equals()"), Object.Equals(NamedField.AsObject()));
+		TestTrue(TEXT("FCbObject::AsField().Equals()"), Field.Equals(NamedField.AsObject().AsField()));
+		TestFalse(TEXT("FCbObject::AsField().Equals()"), NamedField.Equals(NamedField.AsObject().AsField()));
 	}
 
 	// Test FCbField(None) as Object
@@ -379,6 +388,31 @@ bool FCbFieldObjectTest::RunTest(const FString& Parameters)
 		}
 	}
 
+	// Test FCbObject as FCbFieldIterator
+	{
+		uint32 Count = 0;
+		FCbObject Object;
+		for (FCbField Field : FCbFieldIterator(Object.AsField()))
+		{
+			TestTrue(TEXT("FCbObject::AsField() as Iterator"), Field.IsObject());
+			++Count;
+		}
+		TestEqual(TEXT("FCbObject::AsField() as Iterator Count"), Count, 1u);
+	}
+
+	// Test FCbObjectRef as FCbFieldRefIterator
+	{
+		uint32 Count = 0;
+		FCbObjectRef Object;
+		Object.MakeOwned();
+		for (FCbFieldRef Field : FCbFieldRefIterator(Object.AsFieldRef()))
+		{
+			TestTrue(TEXT("FCbObjectRef::AsField() as Iterator"), Field.IsObject());
+			++Count;
+		}
+		TestEqual(TEXT("FCbObjectRef::AsField() as Iterator Count"), Count, 1u);
+	}
+
 	return true;
 }
 
@@ -450,7 +484,16 @@ bool FCbFieldArrayTest::RunTest(const FString& Parameters)
 		TestIntArray(Array, 3, sizeof(Payload));
 		TestIntArray(Field.AsArray(), 3, sizeof(Payload));
 		TestTrue(TEXT("FCbArray::Equals()"), Array.Equals(Field.AsArray()));
+
+		TestTrue(TEXT("FCbArrayRef::AsFieldRef()"), Array.GetBuffer() == Array.AsFieldRef().AsArrayRef().GetBuffer());
+
 		Array = FCbFieldRef::MakeView(Field).AsArrayRef();
+
+		const uint8 NamedPayload[] = { 1, 'A', 5, 3, IntType, 1, 2, 3 };
+		FCbField NamedField(NamedPayload, ECbFieldType::UniformArray | ECbFieldType::HasFieldName);
+		TestTrue(TEXT("FCbArray::Equals()"), Array.Equals(NamedField.AsArray()));
+		TestTrue(TEXT("FCbArray::AsField().Equals()"), Field.Equals(NamedField.AsArray().AsField()));
+		TestFalse(TEXT("FCbArray::AsField().Equals()"), NamedField.Equals(NamedField.AsArray().AsField()));
 	}
 
 	// Test FCbField(None) as Array
@@ -480,6 +523,31 @@ bool FCbFieldArrayTest::RunTest(const FString& Parameters)
 		}
 	}
 
+	// Test FCbArray as FCbFieldIterator
+	{
+		uint32 Count = 0;
+		FCbArray Array;
+		for (FCbField Field : FCbFieldIterator(Array.AsField()))
+		{
+			TestTrue(TEXT("FCbArray::AsField() as Iterator"), Field.IsArray());
+			++Count;
+		}
+		TestEqual(TEXT("FCbArray::AsField() as Iterator Count"), Count, 1u);
+	}
+
+	// Test FCbArrayRef as FCbFieldRefIterator
+	{
+		uint32 Count = 0;
+		FCbArrayRef Array;
+		Array.MakeOwned();
+		for (FCbFieldRef Field : FCbFieldRefIterator(Array.AsFieldRef()))
+		{
+			TestTrue(TEXT("FCbArrayRef::AsField() as Iterator"), Field.IsArray());
+			++Count;
+		}
+		TestEqual(TEXT("FCbArrayRef::AsField() as Iterator Count"), Count, 1u);
+	}
+
 	return true;
 }
 
@@ -904,6 +972,7 @@ bool FCbFieldRefTest::RunTest(const FString& Parameters)
 	static_assert(std::is_constructible<FCbFieldRef, const FCbField&, const FSharedBufferPtr&>::value, "Missing constructor for FCbFieldRef");
 	static_assert(std::is_constructible<FCbFieldRef, const FCbField&, const FSharedBufferConstRef&>::value, "Missing constructor for FCbFieldRef");
 	static_assert(std::is_constructible<FCbFieldRef, const FCbField&, const FSharedBufferConstPtr&>::value, "Missing constructor for FCbFieldRef");
+	static_assert(std::is_constructible<FCbFieldRef, const FCbField&, const FCbFieldRefIterator&>::value, "Missing constructor for FCbFieldRef");
 	static_assert(std::is_constructible<FCbFieldRef, const FCbField&, const FCbFieldRef&>::value, "Missing constructor for FCbFieldRef");
 	static_assert(std::is_constructible<FCbFieldRef, const FCbField&, const FCbArrayRef&>::value, "Missing constructor for FCbFieldRef");
 	static_assert(std::is_constructible<FCbFieldRef, const FCbField&, const FCbObjectRef&>::value, "Missing constructor for FCbFieldRef");
@@ -912,6 +981,7 @@ bool FCbFieldRefTest::RunTest(const FString& Parameters)
 	static_assert(std::is_constructible<FCbFieldRef, const FCbField&, FSharedBufferPtr&&>::value, "Missing constructor for FCbFieldRef");
 	static_assert(std::is_constructible<FCbFieldRef, const FCbField&, FSharedBufferConstRef&&>::value, "Missing constructor for FCbFieldRef");
 	static_assert(std::is_constructible<FCbFieldRef, const FCbField&, FSharedBufferConstPtr&&>::value, "Missing constructor for FCbFieldRef");
+	static_assert(std::is_constructible<FCbFieldRef, const FCbField&, FCbFieldRefIterator&&>::value, "Missing constructor for FCbFieldRef");
 	static_assert(std::is_constructible<FCbFieldRef, const FCbField&, FCbFieldRef&&>::value, "Missing constructor for FCbFieldRef");
 	static_assert(std::is_constructible<FCbFieldRef, const FCbField&, FCbArrayRef&&>::value, "Missing constructor for FCbFieldRef");
 	static_assert(std::is_constructible<FCbFieldRef, const FCbField&, FCbObjectRef&&>::value, "Missing constructor for FCbFieldRef");
@@ -1066,11 +1136,13 @@ IMPLEMENT_SIMPLE_AUTOMATION_TEST(FCbArrayRefTest, "System.Core.Serialization.CbA
 bool FCbArrayRefTest::RunTest(const FString& Parameters)
 {
 	static_assert(std::is_constructible<FCbArrayRef, const FCbArray&, const FSharedBufferConstPtr&>::value, "Missing constructor for FCbArrayRef");
+	static_assert(std::is_constructible<FCbArrayRef, const FCbArray&, const FCbFieldRefIterator&>::value, "Missing constructor for FCbArrayRef");
 	static_assert(std::is_constructible<FCbArrayRef, const FCbArray&, const FCbFieldRef&>::value, "Missing constructor for FCbArrayRef");
 	static_assert(std::is_constructible<FCbArrayRef, const FCbArray&, const FCbArrayRef&>::value, "Missing constructor for FCbArrayRef");
 	static_assert(std::is_constructible<FCbArrayRef, const FCbArray&, const FCbObjectRef&>::value, "Missing constructor for FCbArrayRef");
 
 	static_assert(std::is_constructible<FCbArrayRef, const FCbArray&, FSharedBufferConstPtr&&>::value, "Missing constructor for FCbArrayRef");
+	static_assert(std::is_constructible<FCbArrayRef, const FCbArray&, FCbFieldRefIterator&&>::value, "Missing constructor for FCbArrayRef");
 	static_assert(std::is_constructible<FCbArrayRef, const FCbArray&, FCbFieldRef&&>::value, "Missing constructor for FCbArrayRef");
 	static_assert(std::is_constructible<FCbArrayRef, const FCbArray&, FCbArrayRef&&>::value, "Missing constructor for FCbArrayRef");
 	static_assert(std::is_constructible<FCbArrayRef, const FCbArray&, FCbObjectRef&&>::value, "Missing constructor for FCbArrayRef");
@@ -1094,11 +1166,13 @@ IMPLEMENT_SIMPLE_AUTOMATION_TEST(FCbObjectRefTest, "System.Core.Serialization.Cb
 bool FCbObjectRefTest::RunTest(const FString& Parameters)
 {
 	static_assert(std::is_constructible<FCbObjectRef, const FCbObject&, const FSharedBufferConstPtr&>::value, "Missing constructor for FCbObjectRef");
+	static_assert(std::is_constructible<FCbObjectRef, const FCbObject&, const FCbFieldRefIterator&>::value, "Missing constructor for FCbObjectRef");
 	static_assert(std::is_constructible<FCbObjectRef, const FCbObject&, const FCbFieldRef&>::value, "Missing constructor for FCbObjectRef");
 	static_assert(std::is_constructible<FCbObjectRef, const FCbObject&, const FCbArrayRef&>::value, "Missing constructor for FCbObjectRef");
 	static_assert(std::is_constructible<FCbObjectRef, const FCbObject&, const FCbObjectRef&>::value, "Missing constructor for FCbObjectRef");
 
 	static_assert(std::is_constructible<FCbObjectRef, const FCbObject&, FSharedBufferConstPtr&&>::value, "Missing constructor for FCbObjectRef");
+	static_assert(std::is_constructible<FCbObjectRef, const FCbObject&, FCbFieldRefIterator&&>::value, "Missing constructor for FCbObjectRef");
 	static_assert(std::is_constructible<FCbObjectRef, const FCbObject&, FCbFieldRef&&>::value, "Missing constructor for FCbObjectRef");
 	static_assert(std::is_constructible<FCbObjectRef, const FCbObject&, FCbArrayRef&&>::value, "Missing constructor for FCbObjectRef");
 	static_assert(std::is_constructible<FCbObjectRef, const FCbObject&, FCbObjectRef&&>::value, "Missing constructor for FCbObjectRef");
diff --git a/Engine/Source/Runtime/Core/Public/Serialization/CompactBinary.h b/Engine/Source/Runtime/Core/Public/Serialization/CompactBinary.h
--- a/Engine/Source/Runtime/Core/Public/Serialization/CompactBinary.h
+++ b/Engine/Source/Runtime/Core/Public/Serialization/CompactBinary.h
@@ -292,6 +292,28 @@ public:
 	{
 	}
 
+	/**
+	 * Construct a field iterator for one field.
+	 *
+	 * @param InField The single field to be iterated.
+	 */
+	constexpr inline explicit TCbFieldIterator(const FieldType& InField)
+		: FieldType(InField)
+		, FieldsEnd(FieldType::GetPayloadEnd())
+	{
+	}
+
+	/**
+	 * Construct a field iterator for one field.
+	 *
+	 * @param InField The single field to be iterated.
+	 */
+	constexpr inline explicit TCbFieldIterator(FieldType&& InField)
+		: FieldType(MoveTemp(InField))
+		, FieldsEnd(FieldType::GetPayloadEnd())
+	{
+	}
+
 	inline TCbFieldIterator& operator++()
 	{
 		const void* const PayloadEnd = FieldType::GetPayloadEnd();
@@ -393,7 +415,7 @@ public:
 	 * assumes that both fields are valid and are written in the canonical format. Fields must be
 	 * written in the same order in arrays and objects, and name comparison is case sensitive. If
 	 * these assumptions do not hold, this may return false for equivalent inputs. Validation can
-	 * be done with the `All` mode to check these assumptions about the format of the inputs.
+	 * be performed with \ref ValidateCompactBinary, except for field order and field name case.
 	 */
 	CORE_API bool Equals(const FCbField& Other) const;
 
@@ -631,6 +653,12 @@ public:
 	 */
 	CORE_API bool Equals(const FCbArray& Other) const;
 
+	/** Access the array as an array field. */
+	inline FCbField AsField() const
+	{
+		return FCbField(Payload, FCbFieldType::GetType(Type));
+	}
+
 protected:
 	inline ECbFieldType GetCopyType() const { return FCbFieldType::GetType(Type) | ECbFieldType::HasFieldType; }
 
@@ -710,6 +738,12 @@ public:
 	/** Find a field by case-sensitive name comparison. \see \ref FCbObject::Find */
 	inline FCbField operator[](FAnsiStringView Name) const { return Find(Name); }
 
+	/** Access the object as an object field. */
+	inline FCbField AsField() const
+	{
+		return FCbField(Payload, FCbFieldType::GetType(Type));
+	}
+
 protected:
 	inline ECbFieldType GetCopyType() const { return FCbFieldType::GetType(Type) | ECbFieldType::HasFieldType; }
 
@@ -928,10 +962,16 @@ public:
 	using TCbBufferRef::TCbBufferRef;
 
 	/** Access the field as an object. Defaults to an empty object on error. */
-	inline FCbObjectRef AsObjectRef();
+	inline FCbObjectRef AsObjectRef() &;
+
+	/** Access the field as an object. Defaults to an empty object on error. */
+	inline FCbObjectRef AsObjectRef() &&;
+
+	/** Access the field as an array. Defaults to an empty array on error. */
+	inline FCbArrayRef AsArrayRef() &;
 
 	/** Access the field as an array. Defaults to an empty array on error. */
-	inline FCbArrayRef AsArrayRef();
+	inline FCbArrayRef AsArrayRef() &&;
 };
 
 /** Iterator for \ref FCbFieldRef. \see \ref TCbFieldIterator */
@@ -955,6 +995,18 @@ public:
 		FCbFieldIterator It = CreateIterator();
 		return FCbFieldRefIterator(FCbFieldRef(It, GetBuffer()), It);
 	}
+
+	/** Access the array as an array field. */
+	inline FCbFieldRef AsFieldRef() const &
+	{
+		return FCbFieldRef(FCbArray::AsField(), *this);
+	}
+
+	/** Access the array as an array field. */
+	inline FCbFieldRef AsFieldRef() &&
+	{
+		return FCbFieldRef(FCbArray::AsField(), MoveTemp(*this));
+	}
 };
 
 ///////////////////////////////////////////////////////////////////////////////////////////////////////////////////////
@@ -991,20 +1043,42 @@ public:
 		}
 		return FCbFieldRef();
 	}
+
+	/** Access the object as an object field. */
+	inline FCbFieldRef AsFieldRef() const &
+	{
+		return FCbFieldRef(FCbObject::AsField(), *this);
+	}
+
+	/** Access the object as an object field. */
+	inline FCbFieldRef AsFieldRef() &&
+	{
+		return FCbFieldRef(FCbObject::AsField(), MoveTemp(*this));
+	}
 };
 
 ///////////////////////////////////////////////////////////////////////////////////////////////////////////////////////
 
-inline FCbObjectRef FCbFieldRef::AsObjectRef()
+inline FCbObjectRef FCbFieldRef::AsObjectRef() &
 {
 	return IsObject() ? FCbObjectRef(AsObject(), *this) : FCbObjectRef();
 }
 
-inline FCbArrayRef FCbFieldRef::AsArrayRef()
+inline FCbObjectRef FCbFieldRef::AsObjectRef() &&
+{
+	return IsObject() ? FCbObjectRef(AsObject(), MoveTemp(*this)) : FCbObjectRef();
+}
+
+inline FCbArrayRef FCbFieldRef::AsArrayRef() &
 {
 	return IsArray() ? FCbArrayRef(AsArray(), *this) : FCbArrayRef();
 }
 
+inline FCbArrayRef FCbFieldRef::AsArrayRef() &&
+{
+	return IsArray() ? FCbArrayRef(AsArray(), MoveTemp(*this)) : FCbArrayRef();
+}
+
 ///////////////////////////////////////////////////////////////////////////////////////////////////////////////////////
 
 /** Flags for validating compact binary data. */
