CompactBinary: Added Hash type and renamed BinaryHash and FieldHash to BinaryReference and Reference

#rnx

diff --git a/Engine/Source/Runtime/Core/Private/Serialization/CompactBinary.cpp b/Engine/Source/Runtime/Core/Private/Serialization/CompactBinary.cpp
--- a/Engine/Source/Runtime/Core/Private/Serialization/CompactBinary.cpp
+++ b/Engine/Source/Runtime/Core/Private/Serialization/CompactBinary.cpp
@@ -61,10 +61,10 @@ void FCbFieldType::StaticAssertTypeConstants()
 	static_assert(!(BoolMask & (BoolBase ^ ECbFieldType::BoolTrue)), "BoolMask or BoolBase is invalid!");
 	static_assert(TypeMask == (BoolMask | (BoolBase ^ ECbFieldType::BoolTrue)), "BoolMask or BoolBase is invalid!");
 
-	static_assert(ExternHashBase == ECbFieldType::BinaryHash, "ExternHashBase is invalid!");
-	static_assert((ExternHashMask & (AllFlags | ECbFieldType::FieldHash)) == ECbFieldType::BinaryHash, "ExternHashMask is invalid!");
-	static_assert(!(ExternHashMask & (ExternHashBase ^ ECbFieldType::FieldHash)), "ExternHashMask or ExternHashBase is invalid!");
-	static_assert(TypeMask == (ExternHashMask | (ExternHashBase ^ ECbFieldType::FieldHash)), "ExternHashMask or ExternHashBase is invalid!");
+	static_assert(AnyReferenceBase == ECbFieldType::Reference, "AnyReferenceBase is invalid!");
+	static_assert((AnyReferenceMask & (AllFlags | ECbFieldType::BinaryReference)) == ECbFieldType::Reference, "AnyReferenceMask is invalid!");
+	static_assert(!(AnyReferenceMask & (AnyReferenceBase ^ ECbFieldType::BinaryReference)), "AnyReferenceMask or AnyReferenceBase is invalid!");
+	static_assert(TypeMask == (AnyReferenceMask | (AnyReferenceBase ^ ECbFieldType::BinaryReference)), "AnyReferenceMask or AnyReferenceBase is invalid!");
 }
 
 ///////////////////////////////////////////////////////////////////////////////////////////////////////////////////////
@@ -123,8 +123,9 @@ uint64 FCbField::GetPayloadSize() const
 	case ECbFieldType::BoolFalse:
 	case ECbFieldType::BoolTrue:
 		return 0;
-	case ECbFieldType::BinaryHash:
-	case ECbFieldType::FieldHash:
+	case ECbFieldType::Reference:
+	case ECbFieldType::BinaryReference:
+	case ECbFieldType::Hash:
 		return 32;
 	case ECbFieldType::Uuid:
 		return 16;
@@ -312,9 +313,9 @@ bool FCbField::AsBool(const bool bDefault)
 	return (uint8(bIsBool) & uint8(LocalType) & 1) | ((!bIsBool) & bDefault);
 }
 
-FBlake3Hash FCbField::AsBinaryHash(const FBlake3Hash& Default)
+FBlake3Hash FCbField::AsReference(const FBlake3Hash& Default)
 {
-	if (FCbFieldType::IsBinaryHash(Type))
+	if (FCbFieldType::IsReference(Type))
 	{
 		Error = ECbFieldError::None;
 		return FBlake3Hash(*static_cast<const FBlake3Hash::ByteArray*>(Payload));
@@ -326,9 +327,23 @@ FBlake3Hash FCbField::AsBinaryHash(const FBlake3Hash& Default)
 	}
 }
 
-FBlake3Hash FCbField::AsFieldHash(const FBlake3Hash& Default)
+FBlake3Hash FCbField::AsBinaryReference(const FBlake3Hash& Default)
 {
-	if (FCbFieldType::IsFieldHash(Type))
+	if (FCbFieldType::IsBinaryReference(Type))
+	{
+		Error = ECbFieldError::None;
+		return FBlake3Hash(*static_cast<const FBlake3Hash::ByteArray*>(Payload));
+	}
+	else
+	{
+		Error = ECbFieldError::TypeError;
+		return Default;
+	}
+}
+
+FBlake3Hash FCbField::AsHash(const FBlake3Hash& Default)
+{
+	if (FCbFieldType::IsHash(Type))
 	{
 		Error = ECbFieldError::None;
 		return FBlake3Hash(*static_cast<const FBlake3Hash::ByteArray*>(Payload));
@@ -901,8 +916,9 @@ static FCbField ValidateCbField(FConstMemoryView& View, ECbValidateMode Mode, EC
 	case ECbFieldType::Float64:
 		ValidateCbFloat64(View, Mode, Error);
 		break;
-	case ECbFieldType::BinaryHash:
-	case ECbFieldType::FieldHash:
+	case ECbFieldType::Reference:
+	case ECbFieldType::BinaryReference:
+	case ECbFieldType::Hash:
 		ValidateFixedPayload(32);
 		break;
 	case ECbFieldType::Uuid:
@@ -986,8 +1002,9 @@ uint64 MeasureCompactBinary(FConstMemoryView View, ECbFieldType Type)
 	case ECbFieldType::BoolFalse:
 	case ECbFieldType::BoolTrue:
 		break;
-	case ECbFieldType::BinaryHash:
-	case ECbFieldType::FieldHash:
+	case ECbFieldType::Reference:
+	case ECbFieldType::BinaryReference:
+	case ECbFieldType::Hash:
 		Size += 32;
 		break;
 	case ECbFieldType::Uuid:
diff --git a/Engine/Source/Runtime/Core/Private/Serialization/CompactBinaryWriter.cpp b/Engine/Source/Runtime/Core/Private/Serialization/CompactBinaryWriter.cpp
--- a/Engine/Source/Runtime/Core/Private/Serialization/CompactBinaryWriter.cpp
+++ b/Engine/Source/Runtime/Core/Private/Serialization/CompactBinaryWriter.cpp
@@ -461,18 +461,25 @@ void FCbWriter::Bool(const bool bValue)
 	EndField(bValue ? ECbFieldType::BoolTrue : ECbFieldType::BoolFalse);
 }
 
-void FCbWriter::BinaryHash(const FBlake3Hash& Value)
+void FCbWriter::Reference(const FBlake3Hash& Value)
 {
 	BeginField();
 	Data.Append(Value.Bytes(), sizeof(decltype(Value.Bytes())));
-	EndField(ECbFieldType::BinaryHash);
+	EndField(ECbFieldType::Reference);
 }
 
-void FCbWriter::FieldHash(const FBlake3Hash& Value)
+void FCbWriter::BinaryReference(const FBlake3Hash& Value)
 {
 	BeginField();
 	Data.Append(Value.Bytes(), sizeof(decltype(Value.Bytes())));
-	EndField(ECbFieldType::FieldHash);
+	EndField(ECbFieldType::BinaryReference);
+}
+
+void FCbWriter::Hash(const FBlake3Hash& Value)
+{
+	BeginField();
+	Data.Append(Value.Bytes(), sizeof(decltype(Value.Bytes())));
+	EndField(ECbFieldType::Hash);
 }
 
 void FCbWriter::Uuid(const FGuid& Value)
diff --git a/Engine/Source/Runtime/Core/Private/Tests/Serialization/CompactBinaryTest.cpp b/Engine/Source/Runtime/Core/Private/Tests/Serialization/CompactBinaryTest.cpp
--- a/Engine/Source/Runtime/Core/Private/Tests/Serialization/CompactBinaryTest.cpp
+++ b/Engine/Source/Runtime/Core/Private/Tests/Serialization/CompactBinaryTest.cpp
@@ -57,8 +57,9 @@ UE_CBFIELD_TYPE_ACCESSOR(Float32, IsFloat, AsFloat, float);
 UE_CBFIELD_TYPE_ACCESSOR(Float64, IsFloat, AsDouble, double);
 UE_CBFIELD_TYPE_ACCESSOR(BoolFalse, IsBool, AsBool, bool);
 UE_CBFIELD_TYPE_ACCESSOR(BoolTrue, IsBool, AsBool, bool);
-UE_CBFIELD_TYPE_ACCESSOR(BinaryHash, IsBinaryHash, AsBinaryHash, FBlake3Hash);
-UE_CBFIELD_TYPE_ACCESSOR(FieldHash, IsFieldHash, AsFieldHash, FBlake3Hash);
+UE_CBFIELD_TYPE_ACCESSOR(Reference, IsReference, AsReference, FBlake3Hash);
+UE_CBFIELD_TYPE_ACCESSOR(BinaryReference, IsBinaryReference, AsBinaryReference, FBlake3Hash);
+UE_CBFIELD_TYPE_ACCESSOR(Hash, IsHash, AsHash, FBlake3Hash);
 UE_CBFIELD_TYPE_ACCESSOR_TYPED(Uuid, IsUuid, AsUuid, FGuid, const FGuid&);
 UE_CBFIELD_TYPE_ACCESSOR(DateTime, IsDateTime, AsDateTimeTicks, int64);
 UE_CBFIELD_TYPE_ACCESSOR(TimeSpan, IsTimeSpan, AsTimeSpanTicks, int64);
@@ -835,43 +836,64 @@ bool FCbFieldBoolTest::RunTest(const FString& Parameters)
 	return true;
 }
 
-IMPLEMENT_CUSTOM_SIMPLE_AUTOMATION_TEST(FCbFieldBinaryHashTest, FCbFieldTestBase, "System.Core.Serialization.CbField.BinaryHash", CompactBinaryTestFlags)
-bool FCbFieldBinaryHashTest::RunTest(const FString& Parameters)
+IMPLEMENT_CUSTOM_SIMPLE_AUTOMATION_TEST(FCbFieldReferenceTest, FCbFieldTestBase, "System.Core.Serialization.CbField.Reference", CompactBinaryTestFlags)
+bool FCbFieldReferenceTest::RunTest(const FString& Parameters)
 {
 	const FBlake3Hash::ByteArray ZeroBytes{};
 	const FBlake3Hash::ByteArray SequentialBytes{1, 2, 3, 4, 5, 6, 7, 8, 9, 10, 11, 12, 13, 14, 15, 16, 17, 18, 19, 20, 21, 22, 23, 24, 25, 26, 27, 28, 29, 30, 31, 32};
 
-	// Test FCbField(BinaryHash, Zero)
-	TestField<ECbFieldType::BinaryHash>(ZeroBytes);
+	// Test FCbField(Reference, Zero)
+	TestField<ECbFieldType::Reference>(ZeroBytes);
 
-	// Test FCbField(BinaryHash, NonZero)
-	TestField<ECbFieldType::BinaryHash>(SequentialBytes, FBlake3Hash(SequentialBytes));
+	// Test FCbField(Reference, NonZero)
+	TestField<ECbFieldType::Reference>(SequentialBytes, FBlake3Hash(SequentialBytes));
 
-	// Test FCbField(None) as BinaryHash
+	// Test FCbField(None) as Reference
 	{
 		FCbField DefaultField;
-		TestFieldError<ECbFieldType::BinaryHash>(DefaultField, ECbFieldError::TypeError, FBlake3Hash(SequentialBytes));
+		TestFieldError<ECbFieldType::Reference>(DefaultField, ECbFieldError::TypeError, FBlake3Hash(SequentialBytes));
 	}
 
 	return true;
 }
 
-IMPLEMENT_CUSTOM_SIMPLE_AUTOMATION_TEST(FCbFieldFieldHashTest, FCbFieldTestBase, "System.Core.Serialization.CbField.FieldHash", CompactBinaryTestFlags)
-bool FCbFieldFieldHashTest::RunTest(const FString& Parameters)
+IMPLEMENT_CUSTOM_SIMPLE_AUTOMATION_TEST(FCbFieldBinaryReferenceTest, FCbFieldTestBase, "System.Core.Serialization.CbField.BinaryReference", CompactBinaryTestFlags)
+bool FCbFieldBinaryReferenceTest::RunTest(const FString& Parameters)
 {
 	const FBlake3Hash::ByteArray ZeroBytes{};
 	const FBlake3Hash::ByteArray SequentialBytes{1, 2, 3, 4, 5, 6, 7, 8, 9, 10, 11, 12, 13, 14, 15, 16, 17, 18, 19, 20, 21, 22, 23, 24, 25, 26, 27, 28, 29, 30, 31, 32};
 
-	// Test FCbField(FieldHash, Zero)
-	TestField<ECbFieldType::FieldHash>(ZeroBytes);
+	// Test FCbField(BinaryReference, Zero)
+	TestField<ECbFieldType::BinaryReference>(ZeroBytes);
 
-	// Test FCbField(FieldHash, NonZero)
-	TestField<ECbFieldType::FieldHash>(SequentialBytes, FBlake3Hash(SequentialBytes));
+	// Test FCbField(BinaryReference, NonZero)
+	TestField<ECbFieldType::BinaryReference>(SequentialBytes, FBlake3Hash(SequentialBytes));
 
-	// Test FCbField(None) as FieldHash
+	// Test FCbField(None) as BinaryReference
 	{
 		FCbField DefaultField;
-		TestFieldError<ECbFieldType::FieldHash>(DefaultField, ECbFieldError::TypeError, FBlake3Hash(SequentialBytes));
+		TestFieldError<ECbFieldType::BinaryReference>(DefaultField, ECbFieldError::TypeError, FBlake3Hash(SequentialBytes));
+	}
+
+	return true;
+}
+
+IMPLEMENT_CUSTOM_SIMPLE_AUTOMATION_TEST(FCbFieldHashTest, FCbFieldTestBase, "System.Core.Serialization.CbField.Hash", CompactBinaryTestFlags)
+bool FCbFieldHashTest::RunTest(const FString& Parameters)
+{
+	const FBlake3Hash::ByteArray ZeroBytes{};
+	const FBlake3Hash::ByteArray SequentialBytes{1, 2, 3, 4, 5, 6, 7, 8, 9, 10, 11, 12, 13, 14, 15, 16, 17, 18, 19, 20, 21, 22, 23, 24, 25, 26, 27, 28, 29, 30, 31, 32};
+
+	// Test FCbField(Hash, Zero)
+	TestField<ECbFieldType::Hash>(ZeroBytes);
+
+	// Test FCbField(Hash, NonZero)
+	TestField<ECbFieldType::Hash>(SequentialBytes, FBlake3Hash(SequentialBytes));
+
+	// Test FCbField(None) as Hash
+	{
+		FCbField DefaultField;
+		TestFieldError<ECbFieldType::Hash>(DefaultField, ECbFieldError::TypeError, FBlake3Hash(SequentialBytes));
 	}
 
 	return true;
@@ -1314,15 +1336,20 @@ bool FCbValidateTest::RunTest(const FString& Parameters)
 	TestEqual(TEXT("ValidateCompactBinary(Valid, BoolFalse)"), Validate({uint8(ECbFieldType::BoolFalse)}), ECbValidateError::None);
 	TestEqual(TEXT("ValidateCompactBinary(Valid, BoolTrue)"), Validate({uint8(ECbFieldType::BoolTrue)}), ECbValidateError::None);
 
-	TestEqual(TEXT("ValidateCompactBinary(Valid, BinaryHash)"), Validate({uint8(ECbFieldType::BinaryHash), 0, 0, 0, 0, 0, 0, 0, 0, 0, 0, 0, 0, 0, 0, 0, 0, 0, 0, 0, 0, 0, 0, 0, 0, 0, 0, 0, 0, 0, 0, 0, 0}), ECbValidateError::None);
-	TestEqual(TEXT("ValidateCompactBinary(Valid, BinaryHash, NoType)"), Validate({0, 0, 0, 0, 0, 0, 0, 0, 0, 0, 0, 0, 0, 0, 0, 0, 0, 0, 0, 0, 0, 0, 0, 0, 0, 0, 0, 0, 0, 0, 0, 0}, ECbFieldType::BinaryHash), ECbValidateError::None);
-	TestEqual(TEXT("ValidateCompactBinary(OutOfBounds, BinaryHash)"), Validate({uint8(ECbFieldType::BinaryHash), 0, 0, 0, 0, 0, 0, 0, 0, 0, 0, 0, 0, 0, 0, 0, 0, 0, 0, 0, 0, 0, 0, 0, 0, 0, 0, 0, 0, 0, 0, 0}), ECbValidateError::OutOfBounds);
-	TestEqual(TEXT("ValidateCompactBinary(OutOfBounds, BinaryHash, NoType)"), Validate({0, 0, 0, 0, 0, 0, 0, 0, 0, 0, 0, 0, 0, 0, 0, 0, 0, 0, 0, 0, 0, 0, 0, 0, 0, 0, 0, 0, 0, 0, 0}, ECbFieldType::BinaryHash), ECbValidateError::OutOfBounds);
+	TestEqual(TEXT("ValidateCompactBinary(Valid, Reference)"), Validate({uint8(ECbFieldType::Reference), 0, 0, 0, 0, 0, 0, 0, 0, 0, 0, 0, 0, 0, 0, 0, 0, 0, 0, 0, 0, 0, 0, 0, 0, 0, 0, 0, 0, 0, 0, 0, 0}), ECbValidateError::None);
+	TestEqual(TEXT("ValidateCompactBinary(Valid, Reference, NoType)"), Validate({0, 0, 0, 0, 0, 0, 0, 0, 0, 0, 0, 0, 0, 0, 0, 0, 0, 0, 0, 0, 0, 0, 0, 0, 0, 0, 0, 0, 0, 0, 0, 0}, ECbFieldType::Reference), ECbValidateError::None);
+	TestEqual(TEXT("ValidateCompactBinary(OutOfBounds, Reference)"), Validate({uint8(ECbFieldType::Reference), 0, 0, 0, 0, 0, 0, 0, 0, 0, 0, 0, 0, 0, 0, 0, 0, 0, 0, 0, 0, 0, 0, 0, 0, 0, 0, 0, 0, 0, 0, 0}), ECbValidateError::OutOfBounds);
+	TestEqual(TEXT("ValidateCompactBinary(OutOfBounds, Reference, NoType)"), Validate({0, 0, 0, 0, 0, 0, 0, 0, 0, 0, 0, 0, 0, 0, 0, 0, 0, 0, 0, 0, 0, 0, 0, 0, 0, 0, 0, 0, 0, 0, 0}, ECbFieldType::Reference), ECbValidateError::OutOfBounds);
+
+	TestEqual(TEXT("ValidateCompactBinary(Valid, BinaryReference)"), Validate({uint8(ECbFieldType::BinaryReference), 0, 0, 0, 0, 0, 0, 0, 0, 0, 0, 0, 0, 0, 0, 0, 0, 0, 0, 0, 0, 0, 0, 0, 0, 0, 0, 0, 0, 0, 0, 0, 0}), ECbValidateError::None);
+	TestEqual(TEXT("ValidateCompactBinary(Valid, BinaryReference, NoType)"), Validate({0, 0, 0, 0, 0, 0, 0, 0, 0, 0, 0, 0, 0, 0, 0, 0, 0, 0, 0, 0, 0, 0, 0, 0, 0, 0, 0, 0, 0, 0, 0, 0}, ECbFieldType::BinaryReference), ECbValidateError::None);
+	TestEqual(TEXT("ValidateCompactBinary(OutOfBounds, BinaryReference)"), Validate({uint8(ECbFieldType::BinaryReference), 0, 0, 0, 0, 0, 0, 0, 0, 0, 0, 0, 0, 0, 0, 0, 0, 0, 0, 0, 0, 0, 0, 0, 0, 0, 0, 0, 0, 0, 0, 0}), ECbValidateError::OutOfBounds);
+	TestEqual(TEXT("ValidateCompactBinary(OutOfBounds, BinaryReference, NoType)"), Validate({0, 0, 0, 0, 0, 0, 0, 0, 0, 0, 0, 0, 0, 0, 0, 0, 0, 0, 0, 0, 0, 0, 0, 0, 0, 0, 0, 0, 0, 0, 0}, ECbFieldType::BinaryReference), ECbValidateError::OutOfBounds);
 
-	TestEqual(TEXT("ValidateCompactBinary(Valid, FieldHash)"), Validate({uint8(ECbFieldType::FieldHash), 0, 0, 0, 0, 0, 0, 0, 0, 0, 0, 0, 0, 0, 0, 0, 0, 0, 0, 0, 0, 0, 0, 0, 0, 0, 0, 0, 0, 0, 0, 0, 0}), ECbValidateError::None);
-	TestEqual(TEXT("ValidateCompactBinary(Valid, FieldHash, NoType)"), Validate({0, 0, 0, 0, 0, 0, 0, 0, 0, 0, 0, 0, 0, 0, 0, 0, 0, 0, 0, 0, 0, 0, 0, 0, 0, 0, 0, 0, 0, 0, 0, 0}, ECbFieldType::FieldHash), ECbValidateError::None);
-	TestEqual(TEXT("ValidateCompactBinary(OutOfBounds, FieldHash)"), Validate({uint8(ECbFieldType::FieldHash), 0, 0, 0, 0, 0, 0, 0, 0, 0, 0, 0, 0, 0, 0, 0, 0, 0, 0, 0, 0, 0, 0, 0, 0, 0, 0, 0, 0, 0, 0, 0}), ECbValidateError::OutOfBounds);
-	TestEqual(TEXT("ValidateCompactBinary(OutOfBounds, FieldHash, NoType)"), Validate({0, 0, 0, 0, 0, 0, 0, 0, 0, 0, 0, 0, 0, 0, 0, 0, 0, 0, 0, 0, 0, 0, 0, 0, 0, 0, 0, 0, 0, 0, 0}, ECbFieldType::FieldHash), ECbValidateError::OutOfBounds);
+	TestEqual(TEXT("ValidateCompactBinary(Valid, Hash)"), Validate({uint8(ECbFieldType::Hash), 0, 0, 0, 0, 0, 0, 0, 0, 0, 0, 0, 0, 0, 0, 0, 0, 0, 0, 0, 0, 0, 0, 0, 0, 0, 0, 0, 0, 0, 0, 0, 0}), ECbValidateError::None);
+	TestEqual(TEXT("ValidateCompactBinary(Valid, Hash, NoType)"), Validate({0, 0, 0, 0, 0, 0, 0, 0, 0, 0, 0, 0, 0, 0, 0, 0, 0, 0, 0, 0, 0, 0, 0, 0, 0, 0, 0, 0, 0, 0, 0, 0}, ECbFieldType::Hash), ECbValidateError::None);
+	TestEqual(TEXT("ValidateCompactBinary(OutOfBounds, Hash)"), Validate({uint8(ECbFieldType::Hash), 0, 0, 0, 0, 0, 0, 0, 0, 0, 0, 0, 0, 0, 0, 0, 0, 0, 0, 0, 0, 0, 0, 0, 0, 0, 0, 0, 0, 0, 0, 0}), ECbValidateError::OutOfBounds);
+	TestEqual(TEXT("ValidateCompactBinary(OutOfBounds, Hash, NoType)"), Validate({0, 0, 0, 0, 0, 0, 0, 0, 0, 0, 0, 0, 0, 0, 0, 0, 0, 0, 0, 0, 0, 0, 0, 0, 0, 0, 0, 0, 0, 0, 0}, ECbFieldType::Hash), ECbValidateError::OutOfBounds);
 
 	TestEqual(TEXT("ValidateCompactBinary(Valid, Uuid)"), Validate({uint8(ECbFieldType::Uuid), 0, 0, 0, 0, 0, 0, 0, 0, 0, 0, 0, 0, 0, 0, 0, 0}), ECbValidateError::None);
 	TestEqual(TEXT("ValidateCompactBinary(Valid, Uuid, NoType)"), Validate({0, 0, 0, 0, 0, 0, 0, 0, 0, 0, 0, 0, 0, 0, 0, 0}, ECbFieldType::Uuid), ECbValidateError::None);
@@ -1514,9 +1541,10 @@ bool FCbMeasureTest::RunTest(const FString& Parameters)
 	TestEqual(TEXT("MeasureCompactBinary(BoolFalse)"), Measure({uint8(ECbFieldType::BoolFalse)}), uint64(1));
 	TestEqual(TEXT("MeasureCompactBinary(BoolTrue)"), Measure({uint8(ECbFieldType::BoolTrue)}), uint64(1));
 
-	TestEqual(TEXT("MeasureCompactBinary(BinaryHash)"), Measure({uint8(ECbFieldType::BinaryHash)}), uint64(33));
-	TestEqual(TEXT("MeasureCompactBinary(FieldHash)"), Measure({uint8(ECbFieldType::FieldHash)}), uint64(33));
+	TestEqual(TEXT("MeasureCompactBinary(Reference)"), Measure({uint8(ECbFieldType::Reference)}), uint64(33));
+	TestEqual(TEXT("MeasureCompactBinary(BinaryReference)"), Measure({uint8(ECbFieldType::BinaryReference)}), uint64(33));
 
+	TestEqual(TEXT("MeasureCompactBinary(Hash)"), Measure({uint8(ECbFieldType::Hash)}), uint64(33));
 	TestEqual(TEXT("MeasureCompactBinary(Uuid)"), Measure({uint8(ECbFieldType::Uuid)}), uint64(17));
 
 	TestEqual(TEXT("MeasureCompactBinary(DateTime)"), Measure({uint8(ECbFieldType::DateTime)}), uint64(9));
diff --git a/Engine/Source/Runtime/Core/Private/Tests/Serialization/CompactBinaryWriterTest.cpp b/Engine/Source/Runtime/Core/Private/Tests/Serialization/CompactBinaryWriterTest.cpp
--- a/Engine/Source/Runtime/Core/Private/Tests/Serialization/CompactBinaryWriterTest.cpp
+++ b/Engine/Source/Runtime/Core/Private/Tests/Serialization/CompactBinaryWriterTest.cpp
@@ -11,7 +11,7 @@
 
 #if WITH_DEV_AUTOMATION_TESTS
 
-static constexpr EAutomationTestFlags::Type CompactBinaryWriterTestFlags = EAutomationTestFlags::Type(EAutomationTestFlags::ApplicationContextMask | EAutomationTestFlags::EngineFilter);
+static constexpr EAutomationTestFlags::Type CompactBinaryWriterTestFlags = EAutomationTestFlags::Type(EAutomationTestFlags::ApplicationContextMask | EAutomationTestFlags::SmokeFilter);
 
 ///////////////////////////////////////////////////////////////////////////////////////////////////////////////////////
 
@@ -561,10 +561,38 @@ bool FCbWriterBoolTest::RunTest(const FString& Parameters)
 
 ///////////////////////////////////////////////////////////////////////////////////////////////////////////////////////
 
+IMPLEMENT_SIMPLE_AUTOMATION_TEST(FCbWriterReferenceTest, "System.Core.Serialization.CbWriter.Reference", CompactBinaryWriterTestFlags)
+bool FCbWriterReferenceTest::RunTest(const FString& Parameters)
+{
+	TCbWriter<256> Writer;
+
+	const FBlake3Hash::ByteArray ZeroBytes{};
+	const FBlake3Hash::ByteArray SequentialBytes{1, 2, 3, 4, 5, 6, 7, 8, 9, 10, 11, 12, 13, 14, 15, 16, 17, 18, 19, 20, 21, 22, 23, 24, 25, 26, 27, 28, 29, 30, 31, 32};
+
+	const FBlake3Hash Values[] = { FBlake3Hash(ZeroBytes), FBlake3Hash(SequentialBytes) };
+	for (const FBlake3Hash& Value : Values)
+	{
+		Writer.Reference(Value);
+	}
+
+	FCbFieldRefIterator Fields = Writer.Save();
+	if (TestEqual(TEXT("FCbWriter(Reference) Validate"), ValidateCompactBinaryRange(*Fields.GetBuffer(), ECbValidateMode::All), ECbValidateError::None))
+	{
+		const FBlake3Hash* CheckValue = Values;
+		for (FCbField Field : Fields)
+		{
+			TestEqual(TEXT("FCbWriter(Reference).AsReference()"), Field.AsReference(), *CheckValue++);
+			TestFalse(TEXT("FCbWriter(Reference) Error"), Field.HasError());
+		}
+	}
+
+	return true;
+}
+
 ///////////////////////////////////////////////////////////////////////////////////////////////////////////////////////
 
-IMPLEMENT_SIMPLE_AUTOMATION_TEST(FCbWriterBinaryHashTest, "System.Core.Serialization.CbWriter.BinaryHash", CompactBinaryWriterTestFlags)
-bool FCbWriterBinaryHashTest::RunTest(const FString& Parameters)
+IMPLEMENT_SIMPLE_AUTOMATION_TEST(FCbWriterBinaryReferenceTest, "System.Core.Serialization.CbWriter.BinaryReference", CompactBinaryWriterTestFlags)
+bool FCbWriterBinaryReferenceTest::RunTest(const FString& Parameters)
 {
 	TCbWriter<256> Writer;
 
@@ -574,17 +602,17 @@ bool FCbWriterBinaryHashTest::RunTest(const FString& Parameters)
 	const FBlake3Hash Values[] = { FBlake3Hash(ZeroBytes), FBlake3Hash(SequentialBytes) };
 	for (const FBlake3Hash& Value : Values)
 	{
-		Writer.BinaryHash(Value);
+		Writer.BinaryReference(Value);
 	}
 
 	FCbFieldRefIterator Fields = Writer.Save();
-	if (TestEqual(TEXT("FCbWriter(BinaryHash) Validate"), ValidateCompactBinaryRange(*Fields.GetBuffer(), ECbValidateMode::All), ECbValidateError::None))
+	if (TestEqual(TEXT("FCbWriter(BinaryReference) Validate"), ValidateCompactBinaryRange(*Fields.GetBuffer(), ECbValidateMode::All), ECbValidateError::None))
 	{
 		const FBlake3Hash* CheckValue = Values;
 		for (FCbField Field : Fields)
 		{
-			TestEqual(TEXT("FCbWriter(BinaryHash).AsBinaryHash()"), Field.AsBinaryHash(), *CheckValue++);
-			TestFalse(TEXT("FCbWriter(BinaryHash) Error"), Field.HasError());
+			TestEqual(TEXT("FCbWriter(BinaryReference).AsBinaryReference()"), Field.AsBinaryReference(), *CheckValue++);
+			TestFalse(TEXT("FCbWriter(BinaryReference) Error"), Field.HasError());
 		}
 	}
 
@@ -593,8 +621,8 @@ bool FCbWriterBinaryHashTest::RunTest(const FString& Parameters)
 
 ///////////////////////////////////////////////////////////////////////////////////////////////////////////////////////
 
-IMPLEMENT_SIMPLE_AUTOMATION_TEST(FCbWriterFieldHashTest, "System.Core.Serialization.CbWriter.FieldHash", CompactBinaryWriterTestFlags)
-bool FCbWriterFieldHashTest::RunTest(const FString& Parameters)
+IMPLEMENT_SIMPLE_AUTOMATION_TEST(FCbWriterHashTest, "System.Core.Serialization.CbWriter.Hash", CompactBinaryWriterTestFlags)
+bool FCbWriterHashTest::RunTest(const FString& Parameters)
 {
 	TCbWriter<256> Writer;
 
@@ -604,17 +632,17 @@ bool FCbWriterFieldHashTest::RunTest(const FString& Parameters)
 	const FBlake3Hash Values[] = { FBlake3Hash(ZeroBytes), FBlake3Hash(SequentialBytes) };
 	for (const FBlake3Hash& Value : Values)
 	{
-		Writer.FieldHash(Value);
+		Writer.Hash(Value);
 	}
 
 	FCbFieldRefIterator Fields = Writer.Save();
-	if (TestEqual(TEXT("FCbWriter(FieldHash) Validate"), ValidateCompactBinaryRange(*Fields.GetBuffer(), ECbValidateMode::All), ECbValidateError::None))
+	if (TestEqual(TEXT("FCbWriter(Hash) Validate"), ValidateCompactBinaryRange(*Fields.GetBuffer(), ECbValidateMode::All), ECbValidateError::None))
 	{
 		const FBlake3Hash* CheckValue = Values;
 		for (FCbField Field : Fields)
 		{
-			TestEqual(TEXT("FCbWriter(FieldHash).AsFieldHash()"), Field.AsFieldHash(), *CheckValue++);
-			TestFalse(TEXT("FCbWriter(FieldHash) Error"), Field.HasError());
+			TestEqual(TEXT("FCbWriter(Hash).AsHash()"), Field.AsHash(), *CheckValue++);
+			TestFalse(TEXT("FCbWriter(Hash) Error"), Field.HasError());
 		}
 	}
 
@@ -779,9 +807,10 @@ bool FCbWriterComplexTest::RunTest(const FString& Parameters)
 		Writer.Name("False"_ASV).Bool(false);
 		Writer.Name("True"_ASV).Bool(true);
 
-		Writer.Name("BinaryHash"_ASV).BinaryHash(FBlake3Hash());
-		Writer.Name("FieldHash"_ASV).FieldHash(FBlake3Hash());
+		Writer.Name("Reference"_ASV).Reference(FBlake3Hash());
+		Writer.Name("BinaryReference"_ASV).BinaryReference(FBlake3Hash());
 
+		Writer.Name("Hash"_ASV).Hash(FBlake3Hash());
 		Writer.Name("Uuid"_ASV).Uuid(FGuid());
 
 		Writer.Name("DateTimeZero"_ASV).DateTimeTicks(0);
@@ -936,6 +965,7 @@ bool FCbWriterStreamTest::RunTest(const FString& Parameters)
 
 		Writer << "False"_ASV << false << "True"_ASV << true;
 
+		Writer << "Hash"_ASV << FBlake3Hash();
 		Writer << "Uuid"_ASV << FGuid();
 
 		Writer << "DateTime"_ASV << FDateTime(2020, 5, 13, 15, 10);
diff --git a/Engine/Source/Runtime/Core/Public/Serialization/CompactBinary.h b/Engine/Source/Runtime/Core/Public/Serialization/CompactBinary.h
--- a/Engine/Source/Runtime/Core/Public/Serialization/CompactBinary.h
+++ b/Engine/Source/Runtime/Core/Public/Serialization/CompactBinary.h
@@ -44,7 +44,7 @@
  * \code
  * void BeginBuild(FCbObjectRef Params)
  * {
- *     if (FSharedBufferConstPtr Data = Storage().Load(Params["Data"].AsBinaryHash()))
+ *     if (FSharedBufferConstPtr Data = Storage().Load(Params["Data"].AsBinaryReference()))
  *     {
  *         SetData(Data);
  *     }
@@ -149,33 +149,35 @@ enum class ECbFieldType : uint8
 	BoolTrue        = 0x0d,
 
 	/**
-	 * BinaryHash is a reference to a blob stored externally.
+	 * Reference is a reference to compact binary stored externally.
 	 *
-	 * Payload is a 256-bit hash digest of the blob.
+	 * Payload is a 256-bit hash digest of the compact binary.
 	 */
-	BinaryHash      = 0x0e,
+	Reference       = 0x0e,
 	/**
-	 * FieldHash is a reference to a field/array/object stored externally.
+	 * BinaryReference is a reference to a blob stored externally.
 	 *
-	 * Payload is a 256-bit hash digest of the field.
+	 * Payload is a 256-bit hash digest of the blob.
 	 */
-	FieldHash       = 0x0f,
+	BinaryReference = 0x0f,
 
+	/** Hash. Payload is a 256-bit hash digest. */
+	Hash            = 0x10,
 	/** UUID/GUID. Payload is a 128-bit UUID as defined by RFC 4122. */
-	Uuid            = 0x10,
+	Uuid            = 0x11,
 
 	/**
 	 * Date and time between 0001-01-01 00:00:00.0000000 and 9999-12-31 23:59:59.9999999.
 	 *
 	 * Payload is a big endian int64 count of 100ns ticks since 0001-01-01 00:00:00.0000000.
 	 */
-	DateTime        = 0x11,
+	DateTime        = 0x12,
 	/**
 	 * Difference between two date/time values.
 	 *
 	 * Payload is a big endian int64 count of 100ns ticks in the span, and may be negative.
 	 */
-	TimeSpan        = 0x12,
+	TimeSpan        = 0x13,
 
 	/**
 	 * A transient flag which indicates that the object or array containing this field has stored
@@ -214,8 +216,8 @@ class FCbFieldType
 	static constexpr ECbFieldType BoolMask              = ECbFieldType(0b0011'1110);
 	static constexpr ECbFieldType BoolBase              = ECbFieldType(0b0000'1100);
 
-	static constexpr ECbFieldType ExternHashMask        = ECbFieldType(0b0011'1110);
-	static constexpr ECbFieldType ExternHashBase        = ECbFieldType(0b0000'1110);
+	static constexpr ECbFieldType AnyReferenceMask      = ECbFieldType(0b0011'1110);
+	static constexpr ECbFieldType AnyReferenceBase      = ECbFieldType(0b0000'1110);
 
 	static void StaticAssertTypeConstants();
 
@@ -242,10 +244,11 @@ public:
 	static constexpr inline bool IsFloat(ECbFieldType Type)      { return (Type & FloatMask) == FloatBase; }
 	static constexpr inline bool IsBool(ECbFieldType Type)       { return (Type & BoolMask) == BoolBase; }
 
-	static constexpr inline bool IsExternHash(ECbFieldType Type) { return (Type & ExternHashMask) == ExternHashBase; }
-	static constexpr inline bool IsBinaryHash(ECbFieldType Type) { return GetType(Type) == ECbFieldType::BinaryHash; }
-	static constexpr inline bool IsFieldHash(ECbFieldType Type)  { return GetType(Type) == ECbFieldType::FieldHash; }
+	static constexpr inline bool IsReference(ECbFieldType Type)       { return GetType(Type) == ECbFieldType::Reference; }
+	static constexpr inline bool IsBinaryReference(ECbFieldType Type) { return GetType(Type) == ECbFieldType::BinaryReference; }
+	static constexpr inline bool IsAnyReference(ECbFieldType Type)    { return (Type & AnyReferenceMask) == AnyReferenceBase; }
 
+	static constexpr inline bool IsHash(ECbFieldType Type)       { return GetType(Type) == ECbFieldType::Hash; }
 	static constexpr inline bool IsUuid(ECbFieldType Type)       { return GetType(Type) == ECbFieldType::Uuid; }
 
 	static constexpr inline bool IsDateTime(ECbFieldType Type)   { return GetType(Type) == ECbFieldType::DateTime; }
@@ -456,10 +459,13 @@ public:
 	/** Access the field as a bool. Returns the provided default on error. */
 	CORE_API bool AsBool(bool bDefault = false);
 
+	/** Access the field as a hash referencing compact binary. Returns the provided default on error. */
+	CORE_API FBlake3Hash AsReference(const FBlake3Hash& Default = FBlake3Hash());
 	/** Access the field as a hash referencing a blob. Returns the provided default on error. */
-	CORE_API FBlake3Hash AsBinaryHash(const FBlake3Hash& Default = FBlake3Hash());
-	/** Access the field as a hash referencing a field/array/object. Returns the provided default on error. */
-	CORE_API FBlake3Hash AsFieldHash(const FBlake3Hash& Default = FBlake3Hash());
+	CORE_API FBlake3Hash AsBinaryReference(const FBlake3Hash& Default = FBlake3Hash());
+
+	/** Access the field as a hash. Returns the provided default on error. */
+	CORE_API FBlake3Hash AsHash(const FBlake3Hash& Default = FBlake3Hash());
 
 	/** Access the field as a UUID. Returns a nil UUID on error. */
 	CORE_API FGuid AsUuid();
@@ -499,10 +505,11 @@ public:
 	constexpr inline bool IsFloat() const           { return FCbFieldType::IsFloat(Type); }
 	constexpr inline bool IsBool() const            { return FCbFieldType::IsBool(Type); }
 
-	constexpr inline bool IsExternHash() const      { return FCbFieldType::IsExternHash(Type); }
-	constexpr inline bool IsBinaryHash() const      { return FCbFieldType::IsBinaryHash(Type); }
-	constexpr inline bool IsFieldHash() const       { return FCbFieldType::IsFieldHash(Type); }
+	constexpr inline bool IsReference() const       { return FCbFieldType::IsReference(Type); }
+	constexpr inline bool IsBinaryReference() const { return FCbFieldType::IsBinaryReference(Type); }
+	constexpr inline bool IsAnyReference() const    { return FCbFieldType::IsAnyReference(Type); }
 
+	constexpr inline bool IsHash() const            { return FCbFieldType::IsHash(Type); }
 	constexpr inline bool IsUuid() const            { return FCbFieldType::IsUuid(Type); }
 
 	constexpr inline bool IsDateTime() const        { return FCbFieldType::IsDateTime(Type); }
diff --git a/Engine/Source/Runtime/Core/Public/Serialization/CompactBinaryWriter.h b/Engine/Source/Runtime/Core/Public/Serialization/CompactBinaryWriter.h
--- a/Engine/Source/Runtime/Core/Public/Serialization/CompactBinaryWriter.h
+++ b/Engine/Source/Runtime/Core/Public/Serialization/CompactBinaryWriter.h
@@ -145,11 +145,13 @@ public:
 	/** Write a bool field. */
 	CORE_API void Bool(bool bValue);
 
-	/** Write a hash field referencing a blob. */
-	CORE_API void BinaryHash(const FBlake3Hash& Value);
-	/** Write a hash field referencing a field/array/object. */
-	CORE_API void FieldHash(const FBlake3Hash& Value);
+	/** Write a field referencing compact binary. */
+	CORE_API void Reference(const FBlake3Hash& Value);
+	/** Write a field referencing a blob. */
+	CORE_API void BinaryReference(const FBlake3Hash& Value);
 
+	/** Write a hash field. */
+	CORE_API void Hash(const FBlake3Hash& Value);
 	/** Write a UUID field. */
 	CORE_API void Uuid(const FGuid& Value);
 
@@ -320,6 +322,12 @@ inline FCbWriter& operator<<(FCbWriter& Writer, bool Value)
 	return Writer;
 }
 
+inline FCbWriter& operator<<(FCbWriter& Writer, const FBlake3Hash& Value)
+{
+	Writer.Hash(Value);
+	return Writer;
+}
+
 inline FCbWriter& operator<<(FCbWriter& Writer, const FGuid& Value)
 {
 	Writer.Uuid(Value);
