CompactBinary: Fixed an error in the compact binary writer

#rnx

diff --git a/Engine/Source/Runtime/Core/Private/Serialization/CompactBinaryWriter.cpp b/Engine/Source/Runtime/Core/Private/Serialization/CompactBinaryWriter.cpp
--- a/Engine/Source/Runtime/Core/Private/Serialization/CompactBinaryWriter.cpp
+++ b/Engine/Source/Runtime/Core/Private/Serialization/CompactBinaryWriter.cpp
@@ -70,10 +70,9 @@ void FCbWriter::Reset()
 FCbFieldRefIterator FCbWriter::Save() const
 {
 	const uint64 Size = SaveSize();
-	FSharedBufferPtr MutableBuffer = FSharedBuffer::Alloc(Size);
-	SaveToMemory(MutableBuffer->GetView());
-	FSharedBufferConstPtr ConstBuffer = FSharedBuffer::MakeReadOnly(MoveTemp(MutableBuffer));
-	return FCbFieldRefIterator(FCbFieldRef(MoveTemp(ConstBuffer)), ConstBuffer->GetView().RightChop(Size).GetData());
+	FSharedBufferPtr Buffer = FSharedBuffer::Alloc(Size);
+	const FCbFieldIterator Output = SaveToMemory(Buffer->GetView());
+	return FCbFieldRefIterator(FCbFieldRef(Output, FSharedBuffer::MakeReadOnly(MoveTemp(Buffer))), Output);
 }
 
 FCbFieldIterator FCbWriter::SaveToMemory(const FMutableMemoryView Buffer) const
