Fixed compile error if load-time tracing is disabled
#rnx

diff --git a/Engine/Source/Runtime/Core/Public/ProfilingDebugging/LoadTimeTracker.h b/Engine/Source/Runtime/Core/Public/ProfilingDebugging/LoadTimeTracker.h
--- a/Engine/Source/Runtime/Core/Public/ProfilingDebugging/LoadTimeTracker.h
+++ b/Engine/Source/Runtime/Core/Public/ProfilingDebugging/LoadTimeTracker.h
@@ -236,11 +236,17 @@ struct CORE_API FScopedLoadTimeAccumulatorTimer : public FScopedDurationTimer
 
 #define CUSTOM_LOADTIMER_LOG Cpu
 
+#if LOADTIMEPROFILERTRACE_ENABLED
 #define SCOPED_LOADTIMER(TimerName) TRACE_CPUPROFILER_EVENT_SCOPE_ON_CHANNEL(TimerName, LoadTimeChannel)
 #define SCOPED_CUSTOM_LOADTIMER(TimerName) UE_TRACE_LOG_SCOPED_T(CUSTOM_LOADTIMER_LOG, TimerName, LoadTimeChannel)
-#define SCOPED_LOADTIMER_CNT(TimerName)
-
 #define ADD_CUSTOM_LOADTIMER_META(TimerName, Key, Value) << TimerName.Key(Value)
+#define SCOPED_LOADTIMER_CNT(TimerName)
+#else
+#define SCOPED_LOADTIMER(TimerName)
+#define SCOPED_CUSTOM_LOADTIMER(TimerName)
+#define ADD_CUSTOM_LOADTIMER_META(TimerName, Key, Value)
+#define SCOPED_LOADTIMER_CNT(TimerName)
+#endif
 #endif
 
 #if ENABLE_LOADTIME_TRACKING_WITH_STATS && STATS
