Fixed StatsTrace compile errors.

#rnx

diff --git a/Engine/Source/Developer/TraceServices/Private/Analyzers/StatsTraceAnalysis.cpp b/Engine/Source/Developer/TraceServices/Private/Analyzers/StatsTraceAnalysis.cpp
--- a/Engine/Source/Developer/TraceServices/Private/Analyzers/StatsTraceAnalysis.cpp
+++ b/Engine/Source/Developer/TraceServices/Private/Analyzers/StatsTraceAnalysis.cpp
@@ -35,15 +35,27 @@ bool FStatsAnalyzer::OnEvent(uint16 RouteId, EStyle Style, const FOnEventContext
 			Counter = CounterProvider.CreateCounter();
 			CountersMap.Add(StatId, Counter);
 		}
-		const ANSICHAR* Name = reinterpret_cast<const ANSICHAR*>(EventData.GetAttachment());
-		const TCHAR* Description = reinterpret_cast<const TCHAR*>(EventData.GetAttachment() + strlen(Name) + 1);
+
+		FString Name;
+		FString Description;
+		if (EventData.GetString("Name", Name))
+		{
+			EventData.GetString("Description", Description);
+		}
+		else
+		{
+			Name = reinterpret_cast<const ANSICHAR*>(EventData.GetAttachment());
+			Description = reinterpret_cast<const TCHAR*>(EventData.GetAttachment() + Name.Len() + 1);
+		}
+
+
 		Trace::ECounterDisplayHint DisplayHint = Trace::CounterDisplayHint_None;
 		if (EventData.GetValue<bool>("IsMemory"))
 		{
 			DisplayHint = Trace::CounterDisplayHint_Memory;
 		}
-		Counter->SetName(Session.StoreString(ANSI_TO_TCHAR(Name)));
-		Counter->SetDescription(Session.StoreString(Description));
+		Counter->SetName(Session.StoreString(*Name));
+		Counter->SetDescription(Session.StoreString(*Description));
 		Counter->SetIsFloatingPoint(EventData.GetValue<bool>("IsFloatingPoint"));
 		Counter->SetDisplayHint(DisplayHint);
 		break;
diff --git a/Engine/Source/Runtime/Core/Private/Stats/StatsTrace.cpp b/Engine/Source/Runtime/Core/Private/Stats/StatsTrace.cpp
--- a/Engine/Source/Runtime/Core/Private/Stats/StatsTrace.cpp
+++ b/Engine/Source/Runtime/Core/Private/Stats/StatsTrace.cpp
@@ -17,6 +17,8 @@ UE_TRACE_EVENT_BEGIN(Stats, Spec, NoSync|Important)
 	UE_TRACE_EVENT_FIELD(bool, IsFloatingPoint)
 	UE_TRACE_EVENT_FIELD(bool, IsMemory)
 	UE_TRACE_EVENT_FIELD(bool, ShouldClearEveryFrame)
+	UE_TRACE_EVENT_FIELD(Trace::AnsiString, Name)
+	UE_TRACE_EVENT_FIELD(Trace::AnsiString, Description)
 UE_TRACE_EVENT_END()
 
 UE_TRACE_EVENT_BEGIN(Stats, EventBatch)
@@ -102,19 +104,16 @@ void FStatsTraceInternal::EndEncodeOp(FThreadState* ThreadState, uint8* BufferPt
 
 void FStatsTrace::DeclareStat(const FName& Stat, const ANSICHAR* Name, const TCHAR* Description, bool IsFloatingPoint, bool IsMemory, bool ShouldClearEveryFrame)
 {
-	uint16 NameSize = (strlen(Name) + 1) * sizeof(ANSICHAR);
-	uint16 DescriptionSize = (FCString::Strlen(Description) + 1) * sizeof(TCHAR);
-	auto Attachment = [Name, NameSize, Description, DescriptionSize](uint8* Buffer)
-	{
-		memcpy(Buffer, Name, NameSize);
-		memcpy(Buffer + NameSize, Description, DescriptionSize);
-	};
-	UE_TRACE_LOG(Stats, Spec, StatsChannel, NameSize + DescriptionSize)
+	uint32 NameLen = FCStringAnsi::Strlen(Name);
+	uint32 DescriptionLen = FCString::Strlen(Description);
+
+	UE_TRACE_LOG(Stats, Spec, StatsChannel, NameLen + DescriptionLen)
 		<< Spec.Id(Stat.GetComparisonIndex().ToUnstableInt())
 		<< Spec.IsFloatingPoint(IsFloatingPoint)
 		<< Spec.IsMemory(IsMemory)
 		<< Spec.ShouldClearEveryFrame(ShouldClearEveryFrame)
-		<< Spec.Attachment(Attachment);
+		<< Spec.Name(Name, NameLen)
+		<< Spec.Description(Description, DescriptionLen);
 }
 
 void FStatsTrace::Increment(const FName& Stat)
