Off-by-one error when parsing command line arguments.

#rnx

diff --git a/Engine/Source/Runtime/Core/Private/ProfilingDebugging/Windows/MemoryTrace.cpp b/Engine/Source/Runtime/Core/Private/ProfilingDebugging/Windows/MemoryTrace.cpp
--- a/Engine/Source/Runtime/Core/Private/ProfilingDebugging/Windows/MemoryTrace.cpp
+++ b/Engine/Source/Runtime/Core/Private/ProfilingDebugging/Windows/MemoryTrace.cpp
@@ -724,7 +724,7 @@ FMalloc* MemoryTrace_Create(FMalloc* InMalloc)
 					break;
 				}
 
-				Start = c;
+				Start = c + 1;
 			}
 		}
 	}
