Added bitcasting utility functions to go from uint32/uint64 to float/double respectively.

diff --git a/Engine/Source/Runtime/Core/Private/Tests/Math/UnrealMathTest.cpp b/Engine/Source/Runtime/Core/Private/Tests/Math/UnrealMathTest.cpp
--- a/Engine/Source/Runtime/Core/Private/Tests/Math/UnrealMathTest.cpp
+++ b/Engine/Source/Runtime/Core/Private/Tests/Math/UnrealMathTest.cpp
@@ -2232,4 +2232,26 @@ bool FNanInfVerificationTest::RunTest(const FString& Parameters)
 	return true;
 }
 
+IMPLEMENT_SIMPLE_AUTOMATION_TEST(FBitCastTest, "System.Core.Math.Bitcast", EAutomationTestFlags::ApplicationContextMask | EAutomationTestFlags::SmokeFilter)
+bool FBitCastTest::RunTest(const FString& Parameters)
+{
+	TestTrue(TEXT("CastFloatToInt32_0"),  FPlatformMath::AsUInt( 0.0f) == 0x00000000U);
+	TestTrue(TEXT("CastFloatToInt32_P1"), FPlatformMath::AsUInt(+1.0f) == 0x3f800000U);
+	TestTrue(TEXT("CastFloatToInt32_N1"), FPlatformMath::AsUInt(-1.0f) == 0xbf800000U);
+
+	TestTrue(TEXT("CastFloatToInt64_0"),  FPlatformMath::AsUInt( 0.0) == 0x0000000000000000ULL);
+	TestTrue(TEXT("CastFloatToInt64_P1"), FPlatformMath::AsUInt(+1.0) == 0x3ff0000000000000ULL);
+	TestTrue(TEXT("CastFloatToInt64_N1"), FPlatformMath::AsUInt(-1.0) == 0xbff0000000000000ULL);
+
+	TestTrue(TEXT("CastIntToFloat32_0"),  FPlatformMath::AsFloat(static_cast<uint32>(0x00000000U)) == 0.0f);
+	TestTrue(TEXT("CastIntToFloat32_P1"), FPlatformMath::AsFloat(static_cast<uint32>(0x3f800000U)) == +1.0f);
+	TestTrue(TEXT("CastIntToFloat32_N1"), FPlatformMath::AsFloat(static_cast<uint32>(0xbf800000U)) == -1.0f);
+
+	TestTrue(TEXT("CastIntToFloat64_0"),  FPlatformMath::AsFloat(static_cast<uint64>(0x0000000000000000ULL)) == 0.0);
+	TestTrue(TEXT("CastIntToFloat64_P1"), FPlatformMath::AsFloat(static_cast<uint64>(0x3ff0000000000000ULL)) == +1.0);
+	TestTrue(TEXT("CastIntToFloat64_N1"), FPlatformMath::AsFloat(static_cast<uint64>(0xbff0000000000000ULL)) == -1.0);
+
+	return true;
+}
+
 #endif //WITH_DEV_AUTOMATION_TESTS
diff --git a/Engine/Source/Runtime/Core/Public/GenericPlatform/GenericPlatformMath.h b/Engine/Source/Runtime/Core/Public/GenericPlatform/GenericPlatformMath.h
--- a/Engine/Source/Runtime/Core/Public/GenericPlatform/GenericPlatformMath.h
+++ b/Engine/Source/Runtime/Core/Public/GenericPlatform/GenericPlatformMath.h
@@ -132,6 +132,59 @@ struct FGenericPlatformMath
 		VectorStoreHalf(Dst + 4, Src + 4);
 	}
 
+	/** 
+	 * Performs a bit cast of the given float to an unsigned int of the same bit width.
+	 * @param F The float to bit cast to an unsigned integer.
+	*  @return A bitwise copy of the float in a 32-bit unsigned integer value.
+	 */
+	static inline uint32 AsUInt(float F)
+	{
+		uint32 U{};
+		static_assert(sizeof(F) == sizeof(U), "The float and uint sizes must be equal");
+		::memcpy(&U, &F, sizeof(F));
+		return U;
+	}
+
+	/** 
+	 * Performs a bit cast of the given double to an unsigned int of the same bit width.
+	 * @param D The double to bit cast to an unsigned integer.
+	*  @return A bitwise copy of the double in a 64-bit unsigned integer value.
+	 */
+	static inline uint64 AsUInt(double F)
+	{
+		uint64 U{};
+		static_assert(sizeof(F) == sizeof(U), "The float and uint sizes must be equal");
+		::memcpy(&U, &F, sizeof(F));
+		return U;
+	}
+
+	/** 
+	 * Performs a bit cast of the given unsigned int to float of the same bit width.
+	 * @param U The 32-bit unsigned int to bit cast to a 32-bit float.
+	*  @return A bitwise copy of the 32-bit float in a 32-bit unsigned integer value.
+	 */
+	static inline float AsFloat(uint32 U)
+	{
+		float F{};
+		static_assert(sizeof(F) == sizeof(U), "The float and uint sizes must be equal");
+		::memcpy(&F, &U, sizeof(U));
+		return F;
+	}
+
+	/** 
+	 * Performs a bit cast of the given unsigned int to float of the same bit width.
+	 * @param U The 64-bit unsigned int to bit cast to a 64-bit float.
+	*  @return A bitwise copy of the 64-bit float in a 64-bit unsigned integer value.
+	 */
+	static inline double AsFloat(uint64 U)
+	{
+		double F{};
+		static_assert(sizeof(F) == sizeof(U), "The float and uint sizes must be equal");
+		::memcpy(&F, &U, sizeof(F));
+		return F;
+	}
+
+
 	/**
 	 * Converts a float to an integer with truncation towards zero.
 	 * @param F		Floating point value to convert
diff --git a/Engine/Source/Runtime/Core/Public/Math/UnrealMathUtility.h b/Engine/Source/Runtime/Core/Public/Math/UnrealMathUtility.h
--- a/Engine/Source/Runtime/Core/Public/Math/UnrealMathUtility.h
+++ b/Engine/Source/Runtime/Core/Public/Math/UnrealMathUtility.h
@@ -284,15 +284,8 @@ private:
 			}
 		};
 
-		union FFloatToInt { FloatType F; IntegralType I; };
-		FFloatToInt FloatA;
-		FFloatToInt FloatB;
-
-		FloatA.F = A;
-		FloatB.F = B;
-
-		IntegralType SNA = FloatToSignedNumber(FloatA.I);
-		IntegralType SNB = FloatToSignedNumber(FloatB.I);
+		IntegralType SNA = FloatToSignedNumber(FMath::AsUInt(A));
+		IntegralType SNB = FloatToSignedNumber(FMath::AsUInt(B));
 		IntegralType Distance = (SNA >= SNB) ? (SNA - SNB) : (SNB - SNA);
 		return Distance <= IntegralType(MaxUlps);
 	}
